[AutoRTFM] Cut more AutoRTFM ties to Core & general clean-up.

Core-decoupling:
* Replace all UE branch-predict and assertion macros with AutoRTFM versions.
* Replace all UE integer typedefs with cstdint equivalents.
* Replace a handful of uses of FMath and FPlatformTime with AutoRTFM versions.

Clean-up:
* Replace system-style includes (<foo>) with regular includes ("foo"). AutoRTFM is not in a system include directory.
* Remove "AutoRTFM/" prefix on includes to headers that are in the same directory.
* Sort includes and newline-separate regular includes from system includes.
* Remove unused includes. Add explicit includes that were being transitively included.

diff --git a/Engine/Source/Programs/AutoRTFMTests/Private/API.cpp b/Engine/Source/Programs/AutoRTFMTests/Private/API.cpp
--- a/Engine/Source/Programs/AutoRTFMTests/Private/API.cpp
+++ b/Engine/Source/Programs/AutoRTFMTests/Private/API.cpp
@@ -1,9 +1,10 @@
 // Copyright Epic Games, Inc. All Rights Reserved.
 
 #include "API.h"
+
+#include "AutoRTFM/AutoRTFMTesting.h"
 #include "Catch2Includes.h"
 
-#include <AutoRTFM/AutoRTFMTesting.h>
 #include <memory>
 #include <thread>
 
diff --git a/Engine/Source/Programs/AutoRTFMTests/Private/APINoAutoRTFM.cpp b/Engine/Source/Programs/AutoRTFMTests/Private/APINoAutoRTFM.cpp
--- a/Engine/Source/Programs/AutoRTFMTests/Private/APINoAutoRTFM.cpp
+++ b/Engine/Source/Programs/AutoRTFMTests/Private/APINoAutoRTFM.cpp
@@ -1,7 +1,7 @@
 // Copyright Epic Games, Inc. All Rights Reserved.
 
 #include "API.h"
-#include <AutoRTFM/AutoRTFM.h>
+#include "AutoRTFM/AutoRTFM.h"
 
 UE_AUTORTFM_NOAUTORTFM int NoAutoRTFM::DoSomethingC(int I)
 {
diff --git a/Engine/Source/Programs/AutoRTFMTests/Private/Abort.cpp b/Engine/Source/Programs/AutoRTFMTests/Private/Abort.cpp
--- a/Engine/Source/Programs/AutoRTFMTests/Private/Abort.cpp
+++ b/Engine/Source/Programs/AutoRTFMTests/Private/Abort.cpp
@@ -1,8 +1,9 @@
 // Copyright Epic Games, Inc. All Rights Reserved.
 
+#include "AutoRTFM/AutoRTFMTesting.h"
 #include "Catch2Includes.h"
 #include "Utils.h"
-#include <AutoRTFM/AutoRTFMTesting.h>
+
 #include <map>
 #include <vector>
 
diff --git a/Engine/Source/Programs/AutoRTFMTests/Private/Allocators.cpp b/Engine/Source/Programs/AutoRTFMTests/Private/Allocators.cpp
--- a/Engine/Source/Programs/AutoRTFMTests/Private/Allocators.cpp
+++ b/Engine/Source/Programs/AutoRTFMTests/Private/Allocators.cpp
@@ -1,7 +1,8 @@
 // Copyright Epic Games, Inc. All Rights Reserved.
 
+#include "AutoRTFM/AutoRTFM.h"
 #include "Catch2Includes.h"
-#include <AutoRTFM/AutoRTFM.h>
+
 #include <new>
 #include <stdlib.h>
 
diff --git a/Engine/Source/Programs/AutoRTFMTests/Private/AutoRTFMTests.cpp b/Engine/Source/Programs/AutoRTFMTests/Private/AutoRTFMTests.cpp
--- a/Engine/Source/Programs/AutoRTFMTests/Private/AutoRTFMTests.cpp
+++ b/Engine/Source/Programs/AutoRTFMTests/Private/AutoRTFMTests.cpp
@@ -1,6 +1,7 @@
 // Copyright Epic Games, Inc. All Rights Reserved.
 
 #include "AutoRTFM/AutoRTFM.h"
+#include "AutoRTFM/AutoRTFMTesting.h"
 #include "CoreMinimal.h"
 #include "Misc/EventPool.h"
 #include "Misc/LazySingleton.h"
@@ -34,6 +35,11 @@ public:
 
 CATCH_REGISTER_LISTENER(FListener)
 
+void AutoRTFM::Testing::AssertionFailure(const char* Expression, const char* File, int Line)
+{
+	FAIL(File << ":" << Line << ": " << Expression);
+}
+
 // Checks for memory leaks. Returns true if no leaks were found, otherwise all
 // leaks are printed to stderr and false is returned.
 bool CheckNoMemoryLeaks();
diff --git a/Engine/Source/Programs/AutoRTFMTests/Private/Benchmarks.cpp b/Engine/Source/Programs/AutoRTFMTests/Private/Benchmarks.cpp
--- a/Engine/Source/Programs/AutoRTFMTests/Private/Benchmarks.cpp
+++ b/Engine/Source/Programs/AutoRTFMTests/Private/Benchmarks.cpp
@@ -1,7 +1,8 @@
 // Copyright Epic Games, Inc. All Rights Reserved.
 
+#include "AutoRTFM/AutoRTFMTesting.h"
 #include "Catch2Includes.h"
-#include <AutoRTFM/AutoRTFMTesting.h>
+
 #include <vector>
 #include <unordered_set>
 
diff --git a/Engine/Source/Programs/AutoRTFMTests/Private/CString.cpp b/Engine/Source/Programs/AutoRTFMTests/Private/CString.cpp
--- a/Engine/Source/Programs/AutoRTFMTests/Private/CString.cpp
+++ b/Engine/Source/Programs/AutoRTFMTests/Private/CString.cpp
@@ -1,11 +1,12 @@
 // Copyright Epic Games, Inc. All Rights Reserved.
 
+#include "AutoRTFM/AutoRTFM.h"
 #include "AutoRTFM/AutoRTFMTesting.h"
 #include "Catch2Includes.h"
 #include "CoreGlobals.h"
 #include "Logging/LogVerbosity.h"
 #include "Misc/FeedbackContext.h"
-#include <AutoRTFM/AutoRTFM.h>
+
 #include <catch_amalgamated.hpp>
 #include <charconv>
 #include <cstring>
diff --git a/Engine/Source/Programs/AutoRTFMTests/Private/Commit.cpp b/Engine/Source/Programs/AutoRTFMTests/Private/Commit.cpp
--- a/Engine/Source/Programs/AutoRTFMTests/Private/Commit.cpp
+++ b/Engine/Source/Programs/AutoRTFMTests/Private/Commit.cpp
@@ -1,7 +1,8 @@
 // Copyright Epic Games, Inc. All Rights Reserved.
 
+#include "AutoRTFM/AutoRTFMTesting.h"
 #include "Catch2Includes.h"
-#include <AutoRTFM/AutoRTFMTesting.h>
+
 #include <map>
 #include <vector>
 
diff --git a/Engine/Source/Programs/AutoRTFMTests/Private/CopyAddInt64.cpp b/Engine/Source/Programs/AutoRTFMTests/Private/CopyAddInt64.cpp
--- a/Engine/Source/Programs/AutoRTFMTests/Private/CopyAddInt64.cpp
+++ b/Engine/Source/Programs/AutoRTFMTests/Private/CopyAddInt64.cpp
@@ -1,7 +1,7 @@
 // Copyright Epic Games, Inc. All Rights Reserved.
 
+#include "AutoRTFM/AutoRTFM.h"
 #include "Catch2Includes.h"
-#include <AutoRTFM/AutoRTFM.h>
 
 TEST_CASE("CopyAddInt64")
 {
diff --git a/Engine/Source/Programs/AutoRTFMTests/Private/CrossCUTests.cpp b/Engine/Source/Programs/AutoRTFMTests/Private/CrossCUTests.cpp
--- a/Engine/Source/Programs/AutoRTFMTests/Private/CrossCUTests.cpp
+++ b/Engine/Source/Programs/AutoRTFMTests/Private/CrossCUTests.cpp
@@ -1,11 +1,10 @@
 // Copyright Epic Games, Inc. All Rights Reserved.
 
+#include "AutoRTFM/AutoRTFMTesting.h"
 #include "Catch2Includes.h"
 #include "CoreTypes.h"
 #include "CrossCUTests.h"
 
-#include <AutoRTFM/AutoRTFMTesting.h>
-
 UE_DISABLE_OPTIMIZATION_SHIP
 
 TEST_CASE("CrossCU.Call")
diff --git a/Engine/Source/Programs/AutoRTFMTests/Private/CsvProfilerTests.cpp b/Engine/Source/Programs/AutoRTFMTests/Private/CsvProfilerTests.cpp
--- a/Engine/Source/Programs/AutoRTFMTests/Private/CsvProfilerTests.cpp
+++ b/Engine/Source/Programs/AutoRTFMTests/Private/CsvProfilerTests.cpp
@@ -1,15 +1,14 @@
 // Copyright Epic Games, Inc. All Rights Reserved.
 
-#include "Async/TaskGraphInterfaces.h"
 #include "API.h"
+#include "Async/TaskGraphInterfaces.h"
+#include "AutoRTFM/AutoRTFMTesting.h"
 #include "Catch2Includes.h"
 #include "CoreTypes.h"
 #include "ProfilingDebugging/CsvProfiler.h"
 #include "Trace/Trace.h"
 #include "Trace/Trace.inl"
 
-#include <AutoRTFM/AutoRTFMTesting.h>
-
 #if !CSV_PROFILER
 #error Test relies on CsvProfiler!
 #endif
diff --git a/Engine/Source/Programs/AutoRTFMTests/Private/Delegate.cpp b/Engine/Source/Programs/AutoRTFMTests/Private/Delegate.cpp
--- a/Engine/Source/Programs/AutoRTFMTests/Private/Delegate.cpp
+++ b/Engine/Source/Programs/AutoRTFMTests/Private/Delegate.cpp
@@ -1,11 +1,11 @@
 // Copyright Epic Games, Inc. All Rights Reserved.
 
+#include "AutoRTFM/AutoRTFMTesting.h"
 #include "Catch2Includes.h"
-#include <AutoRTFM/AutoRTFMTesting.h>
 #include "Delegates/Delegate.h"
-#include "Utils.h"
 #include "MyAutoRTFMTestObject.h"
 #include "UObject/GCObject.h"
+#include "Utils.h"
 
 TEST_CASE("Delegate.Broadcast")
 {
diff --git a/Engine/Source/Programs/AutoRTFMTests/Private/DependencyTests.cpp b/Engine/Source/Programs/AutoRTFMTests/Private/DependencyTests.cpp
--- a/Engine/Source/Programs/AutoRTFMTests/Private/DependencyTests.cpp
+++ b/Engine/Source/Programs/AutoRTFMTests/Private/DependencyTests.cpp
@@ -1,6 +1,6 @@
 // Copyright Epic Games, Inc. All Rights Reserved.
 
-#include <AutoRTFM/AutoRTFM.h>
+#include "AutoRTFM/AutoRTFM.h"
 
 // We can't use Catch2Includes.h here, because it #includes more Unreal headers.
 // Fortunately, we don't even need to run these tests; if this code compiles and links, 
diff --git a/Engine/Source/Programs/AutoRTFMTests/Private/FakeVM.cpp b/Engine/Source/Programs/AutoRTFMTests/Private/FakeVM.cpp
--- a/Engine/Source/Programs/AutoRTFMTests/Private/FakeVM.cpp
+++ b/Engine/Source/Programs/AutoRTFMTests/Private/FakeVM.cpp
@@ -1,8 +1,9 @@
 // Copyright Epic Games, Inc. All Rights Reserved.
 
+#include "AutoRTFM/AutoRTFM.h"
 #include "Catch2Includes.h"
 #include "Utils.h"
-#include <AutoRTFM/AutoRTFM.h>
+
 #include <vector>
 
 enum class EOpcodeType
diff --git a/Engine/Source/Programs/AutoRTFMTests/Private/FunctionPointerTests.cpp b/Engine/Source/Programs/AutoRTFMTests/Private/FunctionPointerTests.cpp
--- a/Engine/Source/Programs/AutoRTFMTests/Private/FunctionPointerTests.cpp
+++ b/Engine/Source/Programs/AutoRTFMTests/Private/FunctionPointerTests.cpp
@@ -1,9 +1,9 @@
 // Copyright Epic Games, Inc. All Rights Reserved.
 
+#include "AutoRTFM/AutoRTFM.h"
 #include "Catch2Includes.h"
 #include "Templates/Function.h"
 
-#include <AutoRTFM/AutoRTFM.h>
 #include <cmath>
 
 namespace
diff --git a/Engine/Source/Programs/AutoRTFMTests/Private/Hash.cpp b/Engine/Source/Programs/AutoRTFMTests/Private/Hash.cpp
--- a/Engine/Source/Programs/AutoRTFMTests/Private/Hash.cpp
+++ b/Engine/Source/Programs/AutoRTFMTests/Private/Hash.cpp
@@ -1,10 +1,10 @@
 // Copyright Epic Games, Inc. All Rights Reserved.
 
+#include "AutoRTFM/AutoRTFM.h"
 #include "Catch2Includes.h"
 #include "Hash/Blake3.h"
 #include "Memory/CompositeBuffer.h"
 #include "Memory/MemoryFwd.h"
-#include <AutoRTFM/AutoRTFM.h>
 
 TEST_CASE("Hash.Blake3")
 {
diff --git a/Engine/Source/Programs/AutoRTFMTests/Private/InsightsTests.cpp b/Engine/Source/Programs/AutoRTFMTests/Private/InsightsTests.cpp
--- a/Engine/Source/Programs/AutoRTFMTests/Private/InsightsTests.cpp
+++ b/Engine/Source/Programs/AutoRTFMTests/Private/InsightsTests.cpp
@@ -1,14 +1,13 @@
 // Copyright Epic Games, Inc. All Rights Reserved.
 
 #include "API.h"
+#include "AutoRTFM/AutoRTFMTesting.h"
 #include "Catch2Includes.h"
 #include "CoreTypes.h"
 #include "ProfilingDebugging/CpuProfilerTrace.h"
 #include "Trace/Trace.h"
 #include "Trace/Trace.inl"
 
-#include <AutoRTFM/AutoRTFMTesting.h>
-
 UE_TRACE_CHANNEL_DEFINE(AutoRTFMInsightsChannel)
 
 struct AutoRTFMInsightsDeferer final
diff --git a/Engine/Source/Programs/AutoRTFMTests/Private/MTTSAccessDetector.cpp b/Engine/Source/Programs/AutoRTFMTests/Private/MTTSAccessDetector.cpp
--- a/Engine/Source/Programs/AutoRTFMTests/Private/MTTSAccessDetector.cpp
+++ b/Engine/Source/Programs/AutoRTFMTests/Private/MTTSAccessDetector.cpp
@@ -1,8 +1,9 @@
 // Copyright Epic Games, Inc. All Rights Reserved.
 
+#include "AutoRTFM/AutoRTFMTesting.h"
 #include "Catch2Includes.h"
 #include "Misc/MTTransactionallySafeAccessDetector.h"
-#include <AutoRTFM/AutoRTFMTesting.h>
+
 #include <atomic>
 #include <thread>
 
diff --git a/Engine/Source/Programs/AutoRTFMTests/Private/Math.cpp b/Engine/Source/Programs/AutoRTFMTests/Private/Math.cpp
--- a/Engine/Source/Programs/AutoRTFMTests/Private/Math.cpp
+++ b/Engine/Source/Programs/AutoRTFMTests/Private/Math.cpp
@@ -1,7 +1,8 @@
 // Copyright Epic Games, Inc. All Rights Reserved.
 
+#include "AutoRTFM/AutoRTFM.h"
 #include "Catch2Includes.h"
-#include <AutoRTFM/AutoRTFM.h>
+
 #include <math.h>
 #include <utility>
 
diff --git a/Engine/Source/Programs/AutoRTFMTests/Private/Open.cpp b/Engine/Source/Programs/AutoRTFMTests/Private/Open.cpp
--- a/Engine/Source/Programs/AutoRTFMTests/Private/Open.cpp
+++ b/Engine/Source/Programs/AutoRTFMTests/Private/Open.cpp
@@ -1,7 +1,8 @@
 // Copyright Epic Games, Inc. All Rights Reserved.
 
+#include "AutoRTFM/AutoRTFM.h"
 #include "Catch2Includes.h"
-#include <AutoRTFM/AutoRTFM.h>
+
 #include <atomic>
 #include <catch_amalgamated.hpp>
 #include <map>
diff --git a/Engine/Source/Programs/AutoRTFMTests/Private/OpenAPI.cpp b/Engine/Source/Programs/AutoRTFMTests/Private/OpenAPI.cpp
--- a/Engine/Source/Programs/AutoRTFMTests/Private/OpenAPI.cpp
+++ b/Engine/Source/Programs/AutoRTFMTests/Private/OpenAPI.cpp
@@ -1,8 +1,8 @@
 // Copyright Epic Games, Inc. All Rights Reserved.
 
+#include "AutoRTFM/AutoRTFM.h"
 #include "Catch2Includes.h"
 #include "Utils.h"
-#include <AutoRTFM/AutoRTFM.h>
 
 #include <atomic>
 #include <map>
diff --git a/Engine/Source/Programs/AutoRTFMTests/Private/PackedPointer.cpp b/Engine/Source/Programs/AutoRTFMTests/Private/PackedPointer.cpp
--- a/Engine/Source/Programs/AutoRTFMTests/Private/PackedPointer.cpp
+++ b/Engine/Source/Programs/AutoRTFMTests/Private/PackedPointer.cpp
@@ -1,7 +1,8 @@
 // Copyright Epic Games, Inc. All Rights Reserved.
 
+#include "AutoRTFM/AutoRTFM.h"
 #include "Catch2Includes.h"
-#include <AutoRTFM/AutoRTFM.h>
+
 #include <inttypes.h>
 
 #pragma pack(push,1)
diff --git a/Engine/Source/Programs/AutoRTFMTests/Private/RefCountPtrTests.cpp b/Engine/Source/Programs/AutoRTFMTests/Private/RefCountPtrTests.cpp
--- a/Engine/Source/Programs/AutoRTFMTests/Private/RefCountPtrTests.cpp
+++ b/Engine/Source/Programs/AutoRTFMTests/Private/RefCountPtrTests.cpp
@@ -1,8 +1,8 @@
 // Copyright Epic Games, Inc. All Rights Reserved.
 
-#include "Catch2Includes.h"
-#include <AutoRTFM/AutoRTFM.h>
+#include "AutoRTFM/AutoRTFM.h"
 #include "AutoRTFM/AutoRTFMTesting.h"
+#include "Catch2Includes.h"
 #include "Templates/RefCounting.h"
 
 namespace
@@ -256,4 +256,4 @@ TEMPLATE_TEST_CASE("OnAbortCapturingRefCountPtr", "[RefCountPtr]", RefCountedInt
 
 		REQUIRE(Foo->Value == 13);
 	}
-}
\ No newline at end of file
+}
diff --git a/Engine/Source/Programs/AutoRTFMTests/Private/SharedPointerTests.cpp b/Engine/Source/Programs/AutoRTFMTests/Private/SharedPointerTests.cpp
--- a/Engine/Source/Programs/AutoRTFMTests/Private/SharedPointerTests.cpp
+++ b/Engine/Source/Programs/AutoRTFMTests/Private/SharedPointerTests.cpp
@@ -1,9 +1,9 @@
 // Copyright Epic Games, Inc. All Rights Reserved.
 
+#include "AutoRTFM/AutoRTFM.h"
 #include "Catch2Includes.h"
-#include <AutoRTFM/AutoRTFM.h>
-#include "Templates/SharedPointer.h"
 #include "HAL/UnrealMemory.h"
+#include "Templates/SharedPointer.h"
 
 TEST_CASE("SharedPointer.PreviouslyAllocated")
 {
diff --git a/Engine/Source/Programs/AutoRTFMTests/Private/Statics.cpp b/Engine/Source/Programs/AutoRTFMTests/Private/Statics.cpp
--- a/Engine/Source/Programs/AutoRTFMTests/Private/Statics.cpp
+++ b/Engine/Source/Programs/AutoRTFMTests/Private/Statics.cpp
@@ -1,12 +1,13 @@
 // Copyright Epic Games, Inc. All Rights Reserved.
 
 #include "AutoRTFM/AutoRTFM.h"
+#include "AutoRTFM/AutoRTFMTesting.h"
 #include "AutoRTFM/Context.h"
 #include "AutoRTFM/ContextInlines.h"
 #include "AutoRTFM/Transaction.h"
 #include "Catch2Includes.h"
 #include "CrossCUTests.h"
-#include <AutoRTFM/AutoRTFMTesting.h>
+
 #include <chrono>
 #include <condition_variable>
 #include <mutex>
diff --git a/Engine/Source/Programs/AutoRTFMTests/Private/TaskArrayTests.cpp b/Engine/Source/Programs/AutoRTFMTests/Private/TaskArrayTests.cpp
--- a/Engine/Source/Programs/AutoRTFMTests/Private/TaskArrayTests.cpp
+++ b/Engine/Source/Programs/AutoRTFMTests/Private/TaskArrayTests.cpp
@@ -2,7 +2,7 @@
 
 #include "AutoRTFM/TaskArray.h"
 #include "Catch2Includes.h"
-#include "Templates/UnrealTemplate.h"
+
 #include <memory>
 #include <random>
 #include <unordered_map>
@@ -267,7 +267,7 @@ TEST_CASE("TTaskArray")
 			Populate(Config);
 			std::swap(TaskArray, Second);
 
-			TaskArray->AddAll(MoveTemp(*Second));
+			TaskArray->AddAll(std::move(*Second));
 
 			Check();
 		});
diff --git a/Engine/Source/Programs/AutoRTFMTests/Private/Tests.cpp b/Engine/Source/Programs/AutoRTFMTests/Private/Tests.cpp
--- a/Engine/Source/Programs/AutoRTFMTests/Private/Tests.cpp
+++ b/Engine/Source/Programs/AutoRTFMTests/Private/Tests.cpp
@@ -1,16 +1,17 @@
 // Copyright Epic Games, Inc. All Rights Reserved.
 
-#include "Catch2Includes.h"
-#include <AutoRTFM/AutoRTFM.h>
+#include "AutoRTFM/AutoRTFM.h"
 #include "AutoRTFM/AutoRTFMTesting.h"
-#include <thread>
+#include "Catch2Includes.h"
+#include "Logging/LogMacros.h"
+
 #include <atomic>
-#include <vector>
-#include <cstring>
 #include <cmath>
 #include <cstdio>
+#include <cstring>
 #include <stdio.h>
-#include "Logging/LogMacros.h"
+#include <thread>
+#include <vector>
 
 #if PLATFORM_CPU_X86_FAMILY
 #include <immintrin.h>
@@ -543,4 +544,4 @@ TEST_CASE("Tests.thread_local")
 		REQUIRE(TLSInt() == 42);
 		REQUIRE(TLSString() == "Cat");
 	}
-}
\ No newline at end of file
+}
diff --git a/Engine/Source/Programs/AutoRTFMTests/Private/UObjectTests.cpp b/Engine/Source/Programs/AutoRTFMTests/Private/UObjectTests.cpp
--- a/Engine/Source/Programs/AutoRTFMTests/Private/UObjectTests.cpp
+++ b/Engine/Source/Programs/AutoRTFMTests/Private/UObjectTests.cpp
@@ -1,20 +1,20 @@
 // Copyright Epic Games, Inc. All Rights Reserved.
 
+#include "AutoRTFM/AutoRTFM.h"
 #include "AutoRTFM/AutoRTFMTesting.h"
 #include "AutoRTFM/ScopedGuard.h"
 #include "GenericPlatform/GenericPlatformMisc.h"
 #include "HAL/MallocLeakDetection.h"
+#include "MyAutoRTFMTestObject.h"
+#include "UObject/GCObject.h"
+#include "UObject/ReachabilityAnalysis.h"
+#include "UObject/UObjectAnnotation.h"
 #include "UObject/UObjectThreadContext.h"
 
 THIRD_PARTY_INCLUDES_START
 #include "Catch2Includes.h"
 THIRD_PARTY_INCLUDES_END
 
-#include "MyAutoRTFMTestObject.h"
-#include <AutoRTFM/AutoRTFM.h>
-#include "UObject/GCObject.h"
-#include "UObject/ReachabilityAnalysis.h"
-#include "UObject/UObjectAnnotation.h"
 
 TEST_CASE("UObject.NewObject")
 {
diff --git a/Engine/Source/Programs/AutoRTFMTests/Private/UnorderedSet.cpp b/Engine/Source/Programs/AutoRTFMTests/Private/UnorderedSet.cpp
--- a/Engine/Source/Programs/AutoRTFMTests/Private/UnorderedSet.cpp
+++ b/Engine/Source/Programs/AutoRTFMTests/Private/UnorderedSet.cpp
@@ -1,7 +1,8 @@
 // Copyright Epic Games, Inc. All Rights Reserved.
 
+#include "AutoRTFM/AutoRTFM.h"
 #include "Catch2Includes.h"
-#include <AutoRTFM/AutoRTFM.h>
+
 #include <unordered_set>
 
 using namespace std;
diff --git a/Engine/Source/Programs/AutoRTFMTests/Private/Utils.h b/Engine/Source/Programs/AutoRTFMTests/Private/Utils.h
--- a/Engine/Source/Programs/AutoRTFMTests/Private/Utils.h
+++ b/Engine/Source/Programs/AutoRTFMTests/Private/Utils.h
@@ -2,8 +2,9 @@
 
 #pragma once
 
+#include "AutoRTFM/AutoRTFM.h"
 #include "Misc/FeedbackContext.h"
-#include <AutoRTFM/AutoRTFM.h>
+
 #include <vector>
 
 namespace AutoRTFMTestUtils
diff --git a/Engine/Source/Programs/AutoRTFMTests/Private/UtilsTests.cpp b/Engine/Source/Programs/AutoRTFMTests/Private/UtilsTests.cpp
--- a/Engine/Source/Programs/AutoRTFMTests/Private/UtilsTests.cpp
+++ b/Engine/Source/Programs/AutoRTFMTests/Private/UtilsTests.cpp
@@ -89,7 +89,6 @@ TEST_CASE("AlignUp")
 	REQUIRE(AutoRTFM::AlignUp(8, 8) == 8);
 }
 
-
 TEST_CASE("RoundDown")
 {
 	REQUIRE(AutoRTFM::RoundDown(0, 1) == 0);
@@ -255,3 +254,12 @@ TEST_CASE("RoundUp")
 	REQUIRE(AutoRTFM::RoundUp(7, 8) == 8);
 	REQUIRE(AutoRTFM::RoundUp(8, 8) == 8);
 }
+
+TEST_CASE("Lerp")
+{
+	REQUIRE_THAT(AutoRTFM::Lerp(2.0, 4.0, -0.5), Catch::Matchers::WithinRel(1.0, 1e-5));
+	REQUIRE_THAT(AutoRTFM::Lerp(2.0, 4.0, +0.0), Catch::Matchers::WithinRel(2.0, 1e-5));
+	REQUIRE_THAT(AutoRTFM::Lerp(2.0, 4.0, +0.5), Catch::Matchers::WithinRel(3.0, 1e-5));
+	REQUIRE_THAT(AutoRTFM::Lerp(2.0, 4.0, +1.0), Catch::Matchers::WithinRel(4.0, 1e-5));
+	REQUIRE_THAT(AutoRTFM::Lerp(2.0, 4.0, +1.5), Catch::Matchers::WithinRel(5.0, 1e-5));
+}
diff --git a/Engine/Source/Runtime/Core/Private/AutoRTFM/API.cpp b/Engine/Source/Runtime/Core/Private/AutoRTFM/API.cpp
--- a/Engine/Source/Runtime/Core/Private/AutoRTFM/API.cpp
+++ b/Engine/Source/Runtime/Core/Private/AutoRTFM/API.cpp
@@ -1,13 +1,13 @@
 // Copyright Epic Games, Inc. All Rights Reserved.
 
 #include "AutoRTFM/AutoRTFM.h"
-#include "AutoRTFM/Context.h"
+#include "Context.h"
 #include "CoreGlobals.h"
 #include "ExternAPI.h"
-#include "HAL/Platform.h"
 #include "Transaction.h"
 
 #include <random>
+#include <tuple>
 
 #if UE_AUTORTFM
 static_assert(UE_AUTORTFM_ENABLED, "AutoRTFM/API.cpp requires the compiler flag '-fautortfm'");
@@ -389,8 +389,6 @@ namespace AutoRTFM
 #include "Toggles.h"
 #include "Utils.h"
 
-#include "Templates/Tuple.h"
-
 // This is the implementation of the AutoRTFM.h API. Ideally, functions here should just delegate to some internal API.
 // For now, I have these functions also perform some error checking.
 
@@ -466,7 +464,7 @@ AUTORTFM_CORE_API autortfm_status RTFM_autortfm_close(void (*UninstrumentedWork)
 {
 	FContext* const Context = FContext::Get();
 
-	if (LIKELY(InstrumentedWork))
+	if (AUTORTFM_LIKELY(InstrumentedWork))
 	{
 		InstrumentedWork(Arg);
 	}
@@ -487,7 +485,7 @@ extern "C" UE_AUTORTFM_NOAUTORTFM UE_AUTORTFM_API void autortfm_pre_open(autortf
 			= MemoryValidationLevel == autortfm_memory_validation_level_default
 			? ForTheRuntime::GetMemoryValidationLevel()
 			: static_cast<AutoRTFM::EMemoryValidationLevel>(MemoryValidationLevel);
-		Transaction->SetOpenActive(Level, PLATFORM_RETURN_ADDRESS());
+		Transaction->SetOpenActive(Level, __builtin_return_address(0));
 	}
 }
 
@@ -742,10 +740,10 @@ extern "C" autortfm_result autortfm_transact_then_open(void (*UninstrumentedWork
 {
 	if (autortfm_is_closed())
 	{
-		return RTFM_autortfm_transact_then_open(UninstrumentedWork, InstrumentedWork, Arg, PLATFORM_RETURN_ADDRESS());
+		return RTFM_autortfm_transact_then_open(UninstrumentedWork, InstrumentedWork, Arg, __builtin_return_address(0));
 	}
 
-	return TransactThenOpenImpl(UninstrumentedWork, Arg, PLATFORM_RETURN_ADDRESS());
+	return TransactThenOpenImpl(UninstrumentedWork, Arg, __builtin_return_address(0));
 }
 
 extern "C" void autortfm_commit(void (*UninstrumentedWork)(void*), void (*InstrumentedWork)(void*), void* Arg)
@@ -886,7 +884,7 @@ extern "C" autortfm_status autortfm_close(void (*UninstrumentedWork)(void*), voi
 
 		FContext* const Context = FContext::Get();
 
-		if (LIKELY(InstrumentedWork))
+		if (AUTORTFM_LIKELY(InstrumentedWork))
 		{
 			Result = static_cast<autortfm_status>(Context->CallClosedNest(InstrumentedWork, Arg));
 		}
@@ -1140,7 +1138,7 @@ extern "C" void autortfm_did_free(void* Ptr)
 	// We only need to process did free if we need to track allocation locations.
 	if constexpr (bTrackAllocationLocations)
 	{
-		if (UNLIKELY(GIsCriticalError))
+		if (AUTORTFM_UNLIKELY(GIsCriticalError))
 		{
 			return;
 		}
@@ -1177,7 +1175,7 @@ extern "C" UE_AUTORTFM_NOAUTORTFM void autortfm_check_abi(void* const Ptr, const
 	private:
 		auto Tied() const
 		{
-			return Tie(Major, Minor, Patch);
+			return std::make_tuple(Major, Minor, Patch);
 		}
 
 	public:
diff --git a/Engine/Source/Runtime/Core/Private/AutoRTFM/AutoRTFMUE.cpp b/Engine/Source/Runtime/Core/Private/AutoRTFM/AutoRTFMUE.cpp
--- a/Engine/Source/Runtime/Core/Private/AutoRTFM/AutoRTFMUE.cpp
+++ b/Engine/Source/Runtime/Core/Private/AutoRTFM/AutoRTFMUE.cpp
@@ -6,19 +6,18 @@
 
 #include "AutoRTFM/AutoRTFM.h"
 #include "Containers/StringConv.h"
-#include "Logging/LogVerbosity.h"
-#include "Misc/AssertionMacros.h"
-#include <cstdarg>
-#include <cstdio>
-#include <memory>
-
 #include "Containers/UnrealString.h"
 #include "GenericPlatform/GenericPlatformCrashContext.h"
 #include "HAL/IConsoleManager.h"
 #include "HAL/UnrealMemory.h"
 #include "Logging/LogMacros.h"
+#include "Logging/LogVerbosity.h"
 #include "Misc/AssertionMacros.h"
 
+#include <cstdarg>
+#include <cstdio>
+#include <memory>
+
 DECLARE_LOG_CATEGORY_EXTERN(LogAutoRTFM, Display, All)
 
 DEFINE_LOG_CATEGORY(LogAutoRTFM)
diff --git a/Engine/Source/Runtime/Core/Private/AutoRTFM/BlockAllocator.h b/Engine/Source/Runtime/Core/Private/AutoRTFM/BlockAllocator.h
--- a/Engine/Source/Runtime/Core/Private/AutoRTFM/BlockAllocator.h
+++ b/Engine/Source/Runtime/Core/Private/AutoRTFM/BlockAllocator.h
@@ -4,8 +4,9 @@
 
 #if (defined(__AUTORTFM) && __AUTORTFM)
 
-#include "AutoRTFM/Utils.h"
+#include "Utils.h"
 #include "ExternAPI.h"
+
 #include <cstddef>
 
 namespace AutoRTFM
@@ -75,15 +76,13 @@ public:
 	}
 
 private:
-	UE_NONCOPYABLE(TBlockAllocator);
+	TBlockAllocator(TBlockAllocator&&) = delete;
+	TBlockAllocator(const TBlockAllocator&) = delete;
+	TBlockAllocator& operator=(const TBlockAllocator&) = delete;
+	TBlockAllocator& operator=(TBlockAllocator&&) = delete;
 	
 	static constexpr size_t BlockAlignment = std::max<size_t>(DataAlignment, 8);
 
-	static inline size_t AlignDown(size_t Value, size_t Alignment)
-	{
-		return Value & ~(Alignment-1);
-	}
-
 	struct alignas(BlockAlignment) FBlockHeader
 	{
 		// The previous FBlock in the singly linked-list
diff --git a/Engine/Source/Runtime/Core/Private/AutoRTFM/CompilerInterface.cpp b/Engine/Source/Runtime/Core/Private/AutoRTFM/CompilerInterface.cpp
--- a/Engine/Source/Runtime/Core/Private/AutoRTFM/CompilerInterface.cpp
+++ b/Engine/Source/Runtime/Core/Private/AutoRTFM/CompilerInterface.cpp
@@ -3,8 +3,8 @@
 #if (defined(__AUTORTFM) && __AUTORTFM)
 
 #include "AutoRTFM/AutoRTFM.h"
-#include "ExternAPI.h"
 #include "ContextInlines.h"
+#include "ExternAPI.h"
 #include "FunctionMapInlines.h"
 #include "Memcpy.h"
 
@@ -34,14 +34,14 @@ void CheckAlignment(FContext* Context, void* Ptr, size_t AlignmentMask)
 
 // Check for writes to null in development code, so that the inevitable crash will occur
 // in the caller's code rather than in the AutoRTFM runtime.
-#define UE_AUTORTFM_HANDLE_NULL_WRITE(Ptr)  \
-    do                                      \
-    {                                       \
-        if (UNLIKELY(Ptr == nullptr))       \
-        {                                   \
-            return;                         \
-        }                                   \
-    }                                       \
+#define UE_AUTORTFM_HANDLE_NULL_WRITE(Ptr)     \
+    do                                         \
+    {                                          \
+        if (AUTORTFM_UNLIKELY(Ptr == nullptr)) \
+        {                                      \
+            return;                            \
+        }                                      \
+    }                                          \
     while (0)
 
 #else
diff --git a/Engine/Source/Runtime/Core/Private/AutoRTFM/Context.cpp b/Engine/Source/Runtime/Core/Private/AutoRTFM/Context.cpp
--- a/Engine/Source/Runtime/Core/Private/AutoRTFM/Context.cpp
+++ b/Engine/Source/Runtime/Core/Private/AutoRTFM/Context.cpp
@@ -3,18 +3,18 @@
 #if (defined(__AUTORTFM) && __AUTORTFM)
 
 #include "Context.h"
+#include "ContextInlines.h"
 
 #include "AutoRTFM/AutoRTFMMetrics.h"
 #include "CallNestInlines.h"
-#include "ContextInlines.h"
 #include "ExternAPI.h"
 #include "FunctionMap.h"
-#include "Misc/AssertionMacros.h"
 #include "ScopedGuard.h"
 #include "StackRange.h"
 #include "Stats.h"
 #include "Transaction.h"
 #include "TransactionInlines.h"
+#include "Utils.h"
 
 #if PLATFORM_WINDOWS
 extern "C" __declspec(dllimport) void __stdcall GetCurrentThreadStackLimits(void**, void**);
@@ -32,7 +32,7 @@ FContext* FContext::Instance = nullptr;
 
 FContext* FContext::Create()
 {
-	ensure(Instance == nullptr);
+	AUTORTFM_ENSURE(Instance == nullptr);
 	void* Memory = AutoRTFM::Allocate(sizeof(FContext), alignof(FContext));
 	Instance = new (Memory) FContext();
 	return Instance;
@@ -68,9 +68,9 @@ bool FContext::IsCommittingOrAborting() const
 
 void FContext::MaterializeDeferredTransactions()
 {
-	uint64 NumToAllocate = GetNumDeferredTransactions();
+	uint64_t NumToAllocate = GetNumDeferredTransactions();
 	NumDeferredTransactions = 0;
-	for (uint64 I = 0; I < NumToAllocate; ++I)
+	for (uint64_t I = 0; I < NumToAllocate; ++I)
 	{
 		StartNonDeferredTransaction(EMemoryValidationLevel::Disabled);
 	}
@@ -400,12 +400,12 @@ ETransactionResult FContext::Transact(void (*UninstrumentedFunction)(void*), voi
 		return ETransactionResult::AbortedByLanguage;
 	}
 
-    if (UNLIKELY(EContextStatus::Committing == Status))
+    if (AUTORTFM_UNLIKELY(EContextStatus::Committing == Status))
     {
     	return ETransactionResult::AbortedByTransactInOnCommit;
     }
 
-    if (UNLIKELY(IsAborting()))
+    if (AUTORTFM_UNLIKELY(IsAborting()))
     {
     	return ETransactionResult::AbortedByTransactInOnAbort;
     }
@@ -482,7 +482,7 @@ ETransactionResult FContext::Transact(void (*UninstrumentedFunction)(void*), voi
 				DumpState();
 				AUTORTFM_VERBOSE("Committing...");
 
-				if (UNLIKELY(!bTriedToRunOnce && AutoRTFM::ForTheRuntime::ShouldRetryNonNestedTransactions()))
+				if (AUTORTFM_UNLIKELY(!bTriedToRunOnce && AutoRTFM::ForTheRuntime::ShouldRetryNonNestedTransactions()))
 				{
 					// We skip trying to commit this time, and instead re-run the transaction.
 					Status = EContextStatus::AbortedByFailedLockAcquisition;
@@ -617,7 +617,7 @@ ETransactionResult FContext::Transact(void (*UninstrumentedFunction)(void*), voi
 
 			if (Status == EContextStatus::OnTrack)
 			{
-				if (UNLIKELY(!bTriedToRunOnce && AutoRTFM::ForTheRuntime::ShouldRetryNestedTransactionsToo()))
+				if (AUTORTFM_UNLIKELY(!bTriedToRunOnce && AutoRTFM::ForTheRuntime::ShouldRetryNestedTransactionsToo()))
 				{
 					// We skip trying to commit this time, and instead re-run the transaction.
 					Status = EContextStatus::AbortedByFailedLockAcquisition;
diff --git a/Engine/Source/Runtime/Core/Private/AutoRTFM/Context.h b/Engine/Source/Runtime/Core/Private/AutoRTFM/Context.h
--- a/Engine/Source/Runtime/Core/Private/AutoRTFM/Context.h
+++ b/Engine/Source/Runtime/Core/Private/AutoRTFM/Context.h
@@ -5,8 +5,8 @@
 #if (defined(__AUTORTFM) && __AUTORTFM)
 
 #include "IntrusivePool.h"
-#include "StackRange.h"
 #include "OpenHashThrottler.h"
+#include "StackRange.h"
 #include "TaskArray.h"
 #include "ThreadID.h"
 #include "Transaction.h"
@@ -59,13 +59,13 @@ public:
     // The rest of this is internalish.
     [[noreturn]] void AbortByLanguageAndThrow();
 
-	inline uint64 GetNumDeferredTransactions() const
+	inline uint64_t GetNumDeferredTransactions() const
 	{
 		return NumDeferredTransactions;
 	}
 	inline bool MustMaterializeDeferredTransactions() const
 	{
-		return (UNLIKELY(0 < GetNumDeferredTransactions()));
+		return (AUTORTFM_UNLIKELY(0 < GetNumDeferredTransactions()));
 	}
 	inline FTransaction* GetCurrentTransaction()
 	{
@@ -101,7 +101,7 @@ public:
 
 	inline void EnteringStaticLocalInitializer()
 	{
-		if (LIKELY(GetStatus() == EContextStatus::Idle))
+		if (AUTORTFM_LIKELY(GetStatus() == EContextStatus::Idle))
 		{
 			return;
 		}
@@ -120,7 +120,7 @@ public:
 
 	inline void LeavingStaticLocalInitializer()
 	{
-		if (LIKELY(GetStatus() == EContextStatus::Idle))
+		if (AUTORTFM_LIKELY(GetStatus() == EContextStatus::Idle))
 		{
 			return;
 		}
@@ -164,7 +164,7 @@ private:
 
 	// We defer allocating FTransactions at the top of the transaction stack.
 	// This allows us to make starting a transaction in the open a load, some math, and a store.
-	uint64 NumDeferredTransactions{0};
+	uint64_t NumDeferredTransactions{0};
     FTransaction* CurrentTransaction{nullptr}; 
 	FCallNest* CurrentNest{nullptr};
 
@@ -172,9 +172,9 @@ private:
 	void* ClosedStackAddress = nullptr;
     EContextStatus Status{EContextStatus::Idle};
 	FThreadID CurrentThreadId;
-	uint32 StackLocalInitializerDepth = 0;
+	uint32_t StackLocalInitializerDepth = 0;
 	TTask<void()> PostAbortCallback;
-	uint32 NumAbortAndRetries = 0;
+	uint32_t NumAbortAndRetries = 0;
 	FTaskPool TaskPool;
 	FTransactionPool TransactionPool;
 	FOpenHashThrottler OpenHashThrottler{
diff --git a/Engine/Source/Runtime/Core/Private/AutoRTFM/ContextInlines.h b/Engine/Source/Runtime/Core/Private/AutoRTFM/ContextInlines.h
--- a/Engine/Source/Runtime/Core/Private/AutoRTFM/ContextInlines.h
+++ b/Engine/Source/Runtime/Core/Private/AutoRTFM/ContextInlines.h
@@ -49,7 +49,7 @@ UE_AUTORTFM_FORCEINLINE void FContext::DidFree(void* LogicalAddress)
     // We can do free's in the open within a transaction *during* when the
     // transaction itself is being destroyed, so we need to check for that case.
 	FTransaction* Transaction = GetCurrentTransaction();
-	if (LIKELY(Transaction))
+	if (AUTORTFM_LIKELY(Transaction))
 	{
 		Transaction->DidFree(LogicalAddress);
 	}
diff --git a/Engine/Source/Runtime/Core/Private/AutoRTFM/ExternAPI.cpp b/Engine/Source/Runtime/Core/Private/AutoRTFM/ExternAPI.cpp
--- a/Engine/Source/Runtime/Core/Private/AutoRTFM/ExternAPI.cpp
+++ b/Engine/Source/Runtime/Core/Private/AutoRTFM/ExternAPI.cpp
@@ -4,6 +4,7 @@
 
 #include "ExternAPI.h"
 
+#include <cstdio>
 #include <memory>
 
 namespace AutoRTFM
diff --git a/Engine/Source/Runtime/Core/Private/AutoRTFM/FunctionMap.cpp b/Engine/Source/Runtime/Core/Private/AutoRTFM/FunctionMap.cpp
--- a/Engine/Source/Runtime/Core/Private/AutoRTFM/FunctionMap.cpp
+++ b/Engine/Source/Runtime/Core/Private/AutoRTFM/FunctionMap.cpp
@@ -3,16 +3,16 @@
 #if (defined(__AUTORTFM) && __AUTORTFM)
 
 #include "FunctionMap.h"
-#include "FunctionMapInlines.h"
-#include "ContextInlines.h"
+
 #include "Atomic.h"
 #include "FastLock.h"
+#include "FunctionMapInlines.h"
 #include "LockGuard.h"
 #include "Utils.h"
+
 #include <stdlib.h>
 #include <string.h>
 
-#include "AutoRTFM/AutoRTFMConstants.h"
 
 namespace AutoRTFM
 {
@@ -29,11 +29,11 @@ static inline void* FollowRelocation(void* Function)
 	// This handles the only relocation mode currently observed, but may need to
 	// support more modes in the future.
 	std::byte* Ptr = reinterpret_cast<std::byte*>(Function);
-	uint16 U16;
+	uint16_t U16;
 	memcpy(&U16, Ptr, 2);
 	if (U16 == 0x25ff) // jmp qword ptr [rip+<relative-address>]
 	{
-		uint32 RelativeAddress;
+		uint32_t RelativeAddress;
 		memcpy(&RelativeAddress, Ptr+2, 4);
 		Ptr += RelativeAddress + 6; // apply relative offset
 		memcpy(&Ptr, Ptr, 8);	   // load pointer
@@ -132,7 +132,7 @@ void FunctionMapAdd(void* OldFunction, void* NewFunction)
 
 void* FunctionMapLookupTryFollowingRelocation(void* Function, const char* Where)
 {
-	if (void* Ptr = FollowRelocation(Function); LIKELY(Ptr))
+	if (void* Ptr = FollowRelocation(Function); AUTORTFM_LIKELY(Ptr))
 	{
 		AUTORTFM_MUST_TAIL return FunctionMapLookup(Ptr, Where);
 	}
diff --git a/Engine/Source/Runtime/Core/Private/AutoRTFM/FunctionMapInlines.h b/Engine/Source/Runtime/Core/Private/AutoRTFM/FunctionMapInlines.h
--- a/Engine/Source/Runtime/Core/Private/AutoRTFM/FunctionMapInlines.h
+++ b/Engine/Source/Runtime/Core/Private/AutoRTFM/FunctionMapInlines.h
@@ -5,7 +5,6 @@
 #if (defined(__AUTORTFM) && __AUTORTFM)
 
 #include "AutoRTFM/AutoRTFMConstants.h"
-#include "Context.h"
 #include "FunctionMap.h"
 #include "Utils.h"
 
@@ -17,7 +16,7 @@ inline void* FunctionMapLookupSlow(void* OldFunction, const char* Where)
 	// Instead fall back to the slower function map lookup.
     void* const Result = FunctionMapTryLookup(OldFunction);
 
-	if (UNLIKELY(!Result))
+	if (AUTORTFM_UNLIKELY(!Result))
 	{
 		AUTORTFM_MUST_TAIL return FunctionMapLookupTryFollowingRelocation(OldFunction, Where);
 	}
@@ -33,9 +32,9 @@ inline void* FunctionMapLookup(void* OldFunction, const char* Where)
 	// 16-bits of the function pointer address as a magic constant check
 	// to give us a much higher confidence that there is actually a closed
 	// variant pointer residing 8-bytes before our function address.
-	uint64 PrefixData;
-	memcpy(&PrefixData, reinterpret_cast<char*>(OldFunction) - sizeof(uint64), sizeof(uint64));
-	if (LIKELY(Constants::MagicMike == (PrefixData & 0xffff000000000000)))
+	uint64_t PrefixData;
+	memcpy(&PrefixData, reinterpret_cast<char*>(OldFunction) - sizeof(uint64_t), sizeof(uint64_t));
+	if (AUTORTFM_LIKELY(Constants::MagicMike == (PrefixData & 0xffff000000000000)))
 	{
 		return reinterpret_cast<void*>(PrefixData & 0x0000ffffffffffff);
 	}
@@ -43,15 +42,15 @@ inline void* FunctionMapLookup(void* OldFunction, const char* Where)
 	// ends preceding our Magic Mike prefix. They use 0xc105cafe in the
 	// lower 32-bits to distinguish the 64-bit word containing their type
 	// hash. If we see it, check the preceding 64-bit word for our prefix.
-	else if (UNLIKELY(0xc105cafe == (PrefixData & 0x00000000ffffffff)))
+	else if (AUTORTFM_UNLIKELY(0xc105cafe == (PrefixData & 0x00000000ffffffff)))
 	{
-		memcpy(&PrefixData, reinterpret_cast<char*>(OldFunction) - sizeof(uint64) * 2, sizeof(uint64));
-		if (LIKELY(Constants::MagicMike == (PrefixData & 0xffff000000000000)))
+		memcpy(&PrefixData, reinterpret_cast<char*>(OldFunction) - sizeof(uint64_t) * 2, sizeof(uint64_t));
+		if (AUTORTFM_LIKELY(Constants::MagicMike == (PrefixData & 0xffff000000000000)))
 		{
 			return reinterpret_cast<void*>(PrefixData & 0x0000ffffffffffff);
 		}
 	}
-	
+
 	AUTORTFM_MUST_TAIL return FunctionMapLookupSlow(OldFunction, Where);
 }
 
diff --git a/Engine/Source/Runtime/Core/Private/AutoRTFM/GlobalData.cpp b/Engine/Source/Runtime/Core/Private/AutoRTFM/GlobalData.cpp
--- a/Engine/Source/Runtime/Core/Private/AutoRTFM/GlobalData.cpp
+++ b/Engine/Source/Runtime/Core/Private/AutoRTFM/GlobalData.cpp
@@ -4,6 +4,7 @@
 
 #include "GlobalData.h"
 #include "Utils.h"
+
 #include <inttypes.h>
 
 #ifdef _MSC_VER
@@ -27,7 +28,7 @@ void InitializeGlobalDataIfNecessary()
 #pragma warning(disable : 4996)
 #endif
 
-    UE_CALL_ONCE([]
+    static bool bInitialized = []
     {
         // AutoRTFM is intended to be used primarily from monolithic UE binaries, where having
         // some additional DLL/so/dylib would just be annoying. However, we guard against the
@@ -55,7 +56,9 @@ void InitializeGlobalDataIfNecessary()
             snprintf(Buffer, sizeof(Buffer), "%s=%p", EnvName, GlobalData);
             putenv(Buffer);
         }
-    });
+        return true;
+    }();
+
 #ifdef _MSC_VER
 #pragma warning(pop)
 #endif
diff --git a/Engine/Source/Runtime/Core/Private/AutoRTFM/IntervalTree.h b/Engine/Source/Runtime/Core/Private/AutoRTFM/IntervalTree.h
--- a/Engine/Source/Runtime/Core/Private/AutoRTFM/IntervalTree.h
+++ b/Engine/Source/Runtime/Core/Private/AutoRTFM/IntervalTree.h
@@ -22,14 +22,14 @@ struct FIntervalTree final
 	FORCENOINLINE bool Insert(const void* const Address, const size_t Size)
     {        
         FRange NewRange(Address, Size);
-        return Insert(MoveTemp(NewRange));
+        return Insert(std::move(NewRange));
     }
 
     FORCENOINLINE bool Contains(const void* const Address, const size_t Size) const
     {
         const FRange NewRange(Address, Size);
 
-        if (UNLIKELY(IntervalTreeNodeIndexNone == Root))
+        if (AUTORTFM_UNLIKELY(IntervalTreeNodeIndexNone == Root))
         {
             return false;
         }
@@ -88,7 +88,7 @@ struct FIntervalTree final
 
             FRange Range = Other.NodeRanges[Current];
 
-            Insert(MoveTemp(Range));
+            Insert(std::move(Range));
 
             const FIntervalTreeNodeIndex Left = Other.NodeLefts[Current];
             const FIntervalTreeNodeIndex Right = Other.NodeRights[Current];
@@ -162,10 +162,10 @@ private:
 
     FORCENOINLINE bool Insert(FRange&& NewRange)
     {
-        if (UNLIKELY(IntervalTreeNodeIndexNone == Root))
+        if (AUTORTFM_UNLIKELY(IntervalTreeNodeIndexNone == Root))
         {
             AUTORTFM_ASSERT(NodeRanges.IsEmpty());
-			Root = AddNode(MoveTemp(NewRange));
+			Root = AddNode(std::move(NewRange));
             return true;
         }
 
@@ -175,14 +175,14 @@ private:
         {
             const FRange Range = NodeRanges[Current];
 
-            if (UNLIKELY((NewRange.Start < Range.End) && (Range.Start < NewRange.End)))
+            if (AUTORTFM_UNLIKELY((NewRange.Start < Range.End) && (Range.Start < NewRange.End)))
             {
                 return false;
             }
 
             if (NewRange.Start < Range.Start)
             {
-                if (UNLIKELY(NewRange.End == Range.Start))
+                if (AUTORTFM_UNLIKELY(NewRange.End == Range.Start))
                 {
                     AUTORTFM_ASSERT(NewRange.Start < Range.Start);
 
@@ -192,7 +192,7 @@ private:
                 }
                 else if (IntervalTreeNodeIndexNone == NodeLefts[Current])
                 {
-					const FIntervalTreeNodeIndex Index = AddNode(MoveTemp(NewRange), Current);
+					const FIntervalTreeNodeIndex Index = AddNode(std::move(NewRange), Current);
 					NodeLefts[Current] = Index;
                     Current = Index;
                     break;
@@ -202,7 +202,7 @@ private:
             }
             else
             {
-                if (UNLIKELY(NewRange.Start == Range.End))
+                if (AUTORTFM_UNLIKELY(NewRange.Start == Range.End))
                 {
                     AUTORTFM_ASSERT(NewRange.End > Range.End);
 
@@ -212,7 +212,7 @@ private:
                 }
                 else if (IntervalTreeNodeIndexNone == NodeRights[Current])
                 {
-                    const FIntervalTreeNodeIndex Index = AddNode(MoveTemp(NewRange), Current);
+                    const FIntervalTreeNodeIndex Index = AddNode(std::move(NewRange), Current);
 					NodeRights[Current] = Index;
                     Current = Index;
                     break;
diff --git a/Engine/Source/Runtime/Core/Private/AutoRTFM/OpenHashThrottler.cpp b/Engine/Source/Runtime/Core/Private/AutoRTFM/OpenHashThrottler.cpp
--- a/Engine/Source/Runtime/Core/Private/AutoRTFM/OpenHashThrottler.cpp
+++ b/Engine/Source/Runtime/Core/Private/AutoRTFM/OpenHashThrottler.cpp
@@ -3,27 +3,43 @@
 #if (defined(__AUTORTFM) && __AUTORTFM)
 
 #include "OpenHashThrottler.h"
-#include "AutoRTFM/Utils.h"
-#include "Math/UnrealMathUtility.h"
-#include "Misc/AssertionMacros.h"
-#include "HAL/PlatformTime.h"
+
+#include "Utils.h"
 #include "WriteLog.h"
-#include <limits>
+
+#include <algorithm>
+#include <chrono>
+#include <cmath>
+#include <cstdlib>
 
 namespace AutoRTFM
 {
 
+namespace
+{
+
+double TimeInSeconds()
+{
+	using Clock = std::chrono::high_resolution_clock;
+	using Duration = std::chrono::duration<double>;
+	using TimePoint = std::chrono::time_point<Clock, Duration>;
+	TimePoint Now = Clock::now();
+	return Now.time_since_epoch().count();
+}
+
+}
+
 FOpenHashThrottler::FHashScope::FHashScope(FOpenHashThrottler& Throttler, const void* OpenReturnAddress, const FWriteLog& WriteLog)
 	: Throttler(Throttler)
 	, WriteLog(WriteLog)
 	, OpenReturnAddress(OpenReturnAddress)
-	, StartTime(FPlatformTime::Seconds())
+	, StartTime(TimeInSeconds())
 {
 }
 
 FOpenHashThrottler::FHashScope::~FHashScope()
 {
-	const FSeconds EndTime = FPlatformTime::Seconds();
+	const FSeconds EndTime = TimeInSeconds();
 	Throttler.OnHash(StartTime, EndTime, OpenReturnAddress, WriteLog);
 }
 
@@ -66,8 +82,9 @@ double FOpenHashThrottler::HashProbabilityFor(const void* OpenReturnAddress)
 
 bool FOpenHashThrottler::ShouldHashFor(const void* OpenReturnAddress)
 {
-	// Multiply the results of two calls to FMath::FRand() to get decent fractional precision.
-	const double RandSqr = static_cast<double>(FMath::FRand()) * static_cast<double>(FMath::FRand());
+	auto FRand = [] { return static_cast<double>(rand()) / static_cast<double>(RAND_MAX); };
+	// Multiply the results of two calls to FRand() to get decent fractional precision.
+	const double RandSqr = FRand() * FRand();
 	// Negate the probability as we want to bias towards not-hashing, and if
 	// either call to FRand() returns zero then RandSqr will be zero.
 	const double SkipProbability = 1.0 - HashProbabilityFor(OpenReturnAddress);
@@ -82,7 +99,7 @@ bool FOpenHashThrottler::ShouldHashFor(const void* OpenReturnAddress)
 
 void FOpenHashThrottler::Update(FSeconds DeltaTime /* = 0 */)
 {
-	const FSeconds Now = FPlatformTime::Seconds();
+	const FSeconds Now = TimeInSeconds();
 	if (DeltaTime == 0)
 	{
 		DeltaTime = (LastUpdateTimestamp > 0) ? (Now - LastUpdateTimestamp) : 0;
@@ -123,7 +140,7 @@ void FOpenHashThrottler::UpdateThrottlingData()
 	}
 
 	// Active is a superset of those that hashed.
-	ensure(NumActiveOpens >= NumOpensThatHashed);
+	AUTORTFM_ENSURE(NumActiveOpens >= NumOpensThatHashed);
 
 	if (NumActiveOpens == 0)
 	{
@@ -149,17 +166,17 @@ void FOpenHashThrottler::UpdateThrottlingData()
 		// exceed the budgeted time.
 		for (auto& It : ThrottlingData.Opens)
 		{
-			It.Value.Probability = FMath::Max(It.Value.Probability * TotalGain, MinProbability);
+			It.Value.Probability = std::max(It.Value.Probability * TotalGain, MinProbability);
 			It.Value.TimeSpentHashing = 0;
 			It.Value.bActive = false;
 		}
-		ThrottlingData.DefaultHashProbability = FMath::Max(ThrottlingData.DefaultHashProbability * TotalGain, MinProbability);
+		ThrottlingData.DefaultHashProbability = std::max(ThrottlingData.DefaultHashProbability * TotalGain, MinProbability);
 	}
 	else
 	{
 		// Average time spent hashing per open
 		const double AverageTimeSpentHashingPerOpen =
-			ThrottlingData.TotalTimeSpentHashing / static_cast<double>(FMath::Max<size_t>(NumOpensThatHashed, 1));
+			ThrottlingData.TotalTimeSpentHashing / static_cast<double>(std::max<size_t>(NumOpensThatHashed, 1));
 
 		// The new lowest probability across all active opens.
 		double LowestProbability = 1;
@@ -179,7 +196,7 @@ void FOpenHashThrottler::UpdateThrottlingData()
 			const FSeconds TimeSpentHashing = It.Value.TimeSpentHashing;
 			
 			// Target probability starts with the current probability
-			double TargetProbability = FMath::Max(It.Value.Probability, MinProbability);
+			double TargetProbability = std::max(It.Value.Probability, MinProbability);
 
 			// Adjust for the relative time spent hashing this open compared to the others.
 			if (TimeSpentHashing > 0)
@@ -195,7 +212,7 @@ void FOpenHashThrottler::UpdateThrottlingData()
 			}
 
 			// Finally clamp between 0..1
-			TargetProbability = FMath::Clamp(TargetProbability, 0, 1);
+			TargetProbability = std::clamp(TargetProbability, 0.0, 1.0);
 			
 			if (TargetProbability < It.Value.Probability)
 			{
@@ -207,11 +224,11 @@ void FOpenHashThrottler::UpdateThrottlingData()
 			{
 				// Probability is being increased.
 				// Interpolate to the new target probability at 15% per second.
-				It.Value.Probability = FMath::Lerp(TargetProbability, It.Value.Probability, FMath::Pow(0.85, ThrottlingData.TimeSinceLastUpdate));
+				It.Value.Probability = Lerp(TargetProbability, It.Value.Probability, std::pow(0.85, ThrottlingData.TimeSinceLastUpdate));
 			}
 
 			// Update LowestProbability if this is the new lowest probability.
-			LowestProbability = FMath::Min(LowestProbability, It.Value.Probability);
+			LowestProbability = std::min(LowestProbability, It.Value.Probability);
 
 			// Reset the TimeSpentHashing and bActive state for this open.
 			It.Value.TimeSpentHashing = 0;
diff --git a/Engine/Source/Runtime/Core/Private/AutoRTFM/Stats.cpp b/Engine/Source/Runtime/Core/Private/AutoRTFM/Stats.cpp
--- a/Engine/Source/Runtime/Core/Private/AutoRTFM/Stats.cpp
+++ b/Engine/Source/Runtime/Core/Private/AutoRTFM/Stats.cpp
@@ -3,7 +3,9 @@
 #if (defined(__AUTORTFM) && __AUTORTFM)
 
 #include "Stats.h"
+
 #include "Utils.h"
+
 #include <inttypes.h>
 
 template<> void AutoRTFM::FStats::Report<AutoRTFM::EStatsKind::Transaction>(const uint64_t Data) const
@@ -151,8 +153,6 @@ void AutoRTFM::FStats::Report() const
 #undef REPORT_CASE
 			}
 		}
-
-		UE_DEBUG_BREAK();
 	}
 }
 
diff --git a/Engine/Source/Runtime/Core/Private/AutoRTFM/Stats.h b/Engine/Source/Runtime/Core/Private/AutoRTFM/Stats.h
--- a/Engine/Source/Runtime/Core/Private/AutoRTFM/Stats.h
+++ b/Engine/Source/Runtime/Core/Private/AutoRTFM/Stats.h
@@ -7,6 +7,8 @@
 #include "Toggles.h"
 #include "Utils.h"
 
+#include <cstring>
+
 namespace AutoRTFM
 {
 
@@ -111,7 +113,7 @@ struct FStats final
 {
 	FStats()
 	{
-		FMemory::Memset(Datas, 0, sizeof(Datas));
+		memset(Datas, 0, sizeof(Datas));
 	}
 
     ~FStats()
diff --git a/Engine/Source/Runtime/Core/Private/AutoRTFM/StdLibSupport.cpp b/Engine/Source/Runtime/Core/Private/AutoRTFM/StdLibSupport.cpp
--- a/Engine/Source/Runtime/Core/Private/AutoRTFM/StdLibSupport.cpp
+++ b/Engine/Source/Runtime/Core/Private/AutoRTFM/StdLibSupport.cpp
@@ -5,16 +5,16 @@
 #include "AutoRTFM/AutoRTFM.h"
 #include "ContextInlines.h"
 #include "Utils.h"
+#include "Memcpy.h"
 
 #include <algorithm>
 #include <charconv>
+#include <float.h>
 #include <functional> // note: introduces additional math overloads
-#include "Memcpy.h"
 #include <stdarg.h>
 #include <stdio.h>
 #include <string.h>
 #include <wchar.h>
-#include <float.h>
 
 #ifdef __APPLE__
 #include <malloc/malloc.h>
@@ -33,7 +33,8 @@
 #include <unistd.h>
 #endif
 
-#if PLATFORM_HAS_ASAN_INCLUDE
+#if __has_include(<sanitizer/asan_interface.h>)
+#	include <sanitizer/asan_interface.h>
 #	if defined(__SANITIZE_ADDRESS__)
 #		define AUTORTFM_ASAN_ENABLED 1
 #	elif defined(__has_feature)
@@ -47,10 +48,6 @@
 #define AUTORTFM_ASAN_ENABLED 0
 #endif
 
-#if AUTORTFM_ASAN_ENABLED
-#include <sanitizer/asan_interface.h>
-#endif
-
 #ifdef _MSC_VER
 // BEGIN: Disable warning about deprecated STD C functions.
 #pragma warning(push)
diff --git a/Engine/Source/Runtime/Core/Private/AutoRTFM/TaskArray.h b/Engine/Source/Runtime/Core/Private/AutoRTFM/TaskArray.h
--- a/Engine/Source/Runtime/Core/Private/AutoRTFM/TaskArray.h
+++ b/Engine/Source/Runtime/Core/Private/AutoRTFM/TaskArray.h
@@ -8,6 +8,7 @@
 #include "Pool.h"
 #include "Stack.h"
 #include "Utils.h"
+
 #include <cstddef>
 
 namespace AutoRTFM
@@ -141,7 +142,7 @@ public:
 		{
 			KeyType Key = Iterator.Key;
 			TEntryStack& EntryStack = Iterator.Value;
-			Keyed.FindOrAdd(Key).PushAll(MoveTemp(EntryStack));
+			Keyed.FindOrAdd(Key).PushAll(std::move(EntryStack));
 		}
 
 		// Add in OtherArray's count.
diff --git a/Engine/Source/Runtime/Core/Private/AutoRTFM/Transaction.cpp b/Engine/Source/Runtime/Core/Private/AutoRTFM/Transaction.cpp
--- a/Engine/Source/Runtime/Core/Private/AutoRTFM/Transaction.cpp
+++ b/Engine/Source/Runtime/Core/Private/AutoRTFM/Transaction.cpp
@@ -3,16 +3,16 @@
 #if (defined(__AUTORTFM) && __AUTORTFM)
 
 #include "Transaction.h"
-#include "ContextStatus.h"
+#include "TransactionInlines.h"
+
 #include "CallNestInlines.h"
+#include "ContextStatus.h"
 #include "GlobalData.h"
-#include "TransactionInlines.h"
 #include "HitSet.h"
 #include "Utils.h"
 
 #include <cstddef>
 #include <cstdint>
-#include <type_traits>
 
 namespace AutoRTFM
 {
@@ -188,7 +188,7 @@ void FTransaction::Undo()
     {
 		FWriteLogEntry Entry = *Iter;
 		// No write records should be within the transaction's stack range.
-		ensure(!IsOnStack(Entry.LogicalAddress));
+		AUTORTFM_ENSURE(!IsOnStack(Entry.LogicalAddress));
 
 		memcpy(Entry.LogicalAddress, Entry.Data, Entry.Size);
 
@@ -262,8 +262,8 @@ void FTransaction::CommitNested()
 	}
 	DeferredPopAllOnAbortHandlers.Reset();
 
-    Parent->CommitTasks.AddAll(MoveTemp(CommitTasks));
-    Parent->AbortTasks.AddAll(MoveTemp(AbortTasks));
+    Parent->CommitTasks.AddAll(std::move(CommitTasks));
+    Parent->AbortTasks.AddAll(std::move(AbortTasks));
 
     Parent->NewMemoryTracker.Merge(NewMemoryTracker);
 }
diff --git a/Engine/Source/Runtime/Core/Private/AutoRTFM/Transaction.h b/Engine/Source/Runtime/Core/Private/AutoRTFM/Transaction.h
--- a/Engine/Source/Runtime/Core/Private/AutoRTFM/Transaction.h
+++ b/Engine/Source/Runtime/Core/Private/AutoRTFM/Transaction.h
@@ -5,8 +5,6 @@
 #if (defined(__AUTORTFM) && __AUTORTFM)
 
 #include "AutoRTFM/AutoRTFM.h"
-#include "AutoRTFM/AutoRTFMTask.h"
-#include "HAL/Platform.h"
 #include "HitSet.h"
 #include "IntervalTree.h"
 #include "IntrusivePool.h"
diff --git a/Engine/Source/Runtime/Core/Private/AutoRTFM/TransactionInlines.h b/Engine/Source/Runtime/Core/Private/AutoRTFM/TransactionInlines.h
--- a/Engine/Source/Runtime/Core/Private/AutoRTFM/TransactionInlines.h
+++ b/Engine/Source/Runtime/Core/Private/AutoRTFM/TransactionInlines.h
@@ -5,15 +5,16 @@
 #if (defined(__AUTORTFM) && __AUTORTFM)
 
 #include "Context.h"
-#include "ScopedGuard.h"
 #include "HitSet.h"
+#include "ScopedGuard.h"
 #include "Transaction.h"
-
-#include "HAL/Platform.h"
 #include "Utils.h"
 #include "WriteLog.h"
 
-#if PLATFORM_HAS_ASAN_INCLUDE
+#include <utility>
+
+#if __has_include(<sanitizer/asan_interface.h>)
+#	include <sanitizer/asan_interface.h>
 #	if defined(__SANITIZE_ADDRESS__)
 #		define AUTORTFM_CHECK_ASAN_FAKE_STACKS 1
 #	elif defined(__has_feature)
@@ -27,10 +28,6 @@
 #define AUTORTFM_CHECK_ASAN_FAKE_STACKS 0
 #endif
 
-#if AUTORTFM_CHECK_ASAN_FAKE_STACKS
-#include <sanitizer/asan_interface.h>
-#endif
-
 namespace AutoRTFM
 {
 
@@ -103,7 +100,7 @@ UE_AUTORTFM_FORCEINLINE bool FTransaction::ShouldRecordWrite(void* LogicalAddres
 AUTORTFM_NO_ASAN UE_AUTORTFM_FORCEINLINE
 void FTransaction::RecordWrite(void* LogicalAddress, size_t Size, bool bNoMemoryValidation /* = false */)
 {
-	if (UNLIKELY(0 == Size))
+	if (AUTORTFM_UNLIKELY(0 == Size))
 	{
 		return;
 	}
@@ -142,7 +139,7 @@ void FTransaction::RecordWrite(void* LogicalAddress, size_t Size, bool bNoMemory
 	WriteLogEntry.LogicalAddress = static_cast<std::byte*>(LogicalAddress);
 	WriteLogEntry.Data = static_cast<std::byte*>(LogicalAddress);
 	WriteLogEntry.bNoMemoryValidation = bNoMemoryValidation;
-	if (UNLIKELY(Size > FWriteLogEntry::MaxSize))
+	if (AUTORTFM_UNLIKELY(Size > FWriteLogEntry::MaxSize))
 	{
 		WriteLogEntry.Size = FWriteLogEntry::MaxSize;
 		while (Size > FWriteLogEntry::MaxSize)
@@ -249,7 +246,7 @@ UE_AUTORTFM_FORCEINLINE void FTransaction::DeferUntilCommit(TTask<void()>&& Call
 	// within a transactional context, and thus the memory is allocating under
 	// transactionalized conditions. By copying, we create an open copy of the callback.
 	TTask<void()> Copy(Callback);
-    CommitTasks.Add(MoveTemp(Copy));
+    CommitTasks.Add(std::move(Copy));
 }
 
 UE_AUTORTFM_FORCEINLINE void FTransaction::DeferUntilAbort(TTask<void()>&& Callback)
@@ -258,7 +255,7 @@ UE_AUTORTFM_FORCEINLINE void FTransaction::DeferUntilAbort(TTask<void()>&& Callb
 	// within a transactional context, and thus the memory is allocating under
 	// transactionalized conditions. By copying, we create an open copy of the callback.
 	TTask<void()> Copy(Callback);
-    AbortTasks.Add(MoveTemp(Copy));
+    AbortTasks.Add(std::move(Copy));
 }
 
 UE_AUTORTFM_FORCEINLINE void FTransaction::PushDeferUntilCommitHandler(const void* Key, TTask<void()>&& Callback)
@@ -267,12 +264,12 @@ UE_AUTORTFM_FORCEINLINE void FTransaction::PushDeferUntilCommitHandler(const voi
 	// within a transactional context, and thus the memory was allocated under
 	// transactionalized conditions. By copying, we create an open copy of the callback.
 	TTask<void()> Copy(Callback);
-	CommitTasks.AddKeyed(Key, MoveTemp(Copy));
+	CommitTasks.AddKeyed(Key, std::move(Copy));
 }
 
 UE_AUTORTFM_FORCEINLINE void FTransaction::PopDeferUntilCommitHandler(const void* Key)
 {
-	if (LIKELY(CommitTasks.DeleteKey(Key)))
+	if (AUTORTFM_LIKELY(CommitTasks.DeleteKey(Key)))
 	{
 		return;
 	}
@@ -294,12 +291,12 @@ UE_AUTORTFM_FORCEINLINE void FTransaction::PushDeferUntilAbortHandler(const void
 	// within a transactional context, and thus the memory is allocating under
 	// transactionalized conditions. By copying, we create an open copy of the callback.
 	TTask<void()> Copy(Callback);
-    AbortTasks.AddKeyed(Key, MoveTemp(Copy));
+    AbortTasks.AddKeyed(Key, std::move(Copy));
 }
 
 UE_AUTORTFM_FORCEINLINE void FTransaction::PopDeferUntilAbortHandler(const void* Key)
 {
-	if (LIKELY(AbortTasks.DeleteKey(Key)))
+	if (AUTORTFM_LIKELY(AbortTasks.DeleteKey(Key)))
 	{
 		return;
 	}
diff --git a/Engine/Source/Runtime/Core/Private/AutoRTFM/Utils.cpp b/Engine/Source/Runtime/Core/Private/AutoRTFM/Utils.cpp
--- a/Engine/Source/Runtime/Core/Private/AutoRTFM/Utils.cpp
+++ b/Engine/Source/Runtime/Core/Private/AutoRTFM/Utils.cpp
@@ -3,8 +3,11 @@
 #if (defined(__AUTORTFM) && __AUTORTFM)
 
 #include "Utils.h"
+
 #include "ContextInlines.h"
 
+#include <utility>
+
 #if PLATFORM_WINDOWS
 #include "Windows/WindowsHWrapper.h"
 #include <dbghelp.h>
@@ -89,7 +92,7 @@ UE_AUTORTFM_API FORCENOINLINE void ReportError(const char* Format, ...)
 	else
 	{
 		TTask<void()> Task = [] { AUTORTFM_ENSURE(ForTheRuntime::SetAutoRTFMRuntime(ForTheRuntime::AutoRTFM_ForcedDisabled)); };
-		Context->AbortTransactionWithPostAbortCallback(EContextStatus::AbortedByCascadingRetry, MoveTemp(Task));
+		Context->AbortTransactionWithPostAbortCallback(EContextStatus::AbortedByCascadingRetry, std::move(Task));
 	}
 }
 
diff --git a/Engine/Source/Runtime/Core/Private/AutoRTFM/Utils.h b/Engine/Source/Runtime/Core/Private/AutoRTFM/Utils.h
--- a/Engine/Source/Runtime/Core/Private/AutoRTFM/Utils.h
+++ b/Engine/Source/Runtime/Core/Private/AutoRTFM/Utils.h
@@ -2,14 +2,16 @@
 
 #pragma once
 
-#include <type_traits>
 #if (defined(__AUTORTFM) && __AUTORTFM)
 
 #include "AutoRTFM/AutoRTFM.h"
 #include "ExternAPI.h"
 
+#include <algorithm>
 #include <cstdarg>
 #include <string>
+#include <type_traits>
+#include <utility>
 
 // If 0 then verbose logging is compiled-out.
 #define AUTORTFM_VERBOSE_ENABLED 0
@@ -20,11 +22,7 @@ namespace AutoRTFM
 __attribute__((__format__(__printf__, 1, 2)))
 UE_AUTORTFM_API void ReportError(const char* Format, ...);
 
-[[noreturn]] inline void InternalUnreachable() //-V1082 Silence PVS [[noreturn]] false positive warning
-{
-	LowLevelFatalError(TEXT("Unreachable encountered!"));
-	__builtin_unreachable();
-}
+[[noreturn]] inline void InternalUnreachable(); //-V1082 Silence PVS [[noreturn]] false positive warning
 
 std::string GetFunctionDescription(void* FunctionPtr);
 
@@ -140,6 +138,13 @@ inline constexpr T RoundUp(T Value, T Multiple)
 	return RoundDown(Value + Multiple - 1, Multiple);
 }
 
+// Returns the linear interpolation between Start and End using the coefficient Fraction.
+template<typename T>
+inline constexpr T Lerp(T Start, T End, T Fraction)
+{
+	return Start + Fraction * (End - Start);
+}
+
 } // namespace AutoRTFM
 
 #define AUTORTFM_LIKELY(x) __builtin_expect(!!(x), 1)
@@ -227,7 +232,7 @@ inline constexpr T RoundUp(T Value, T Multiple)
 // Same funkiness as AUTORTFM_ASSERT, except that it doesn't cause a fatal error.
 // TODO: Maybe upgrade the Info -> Warning?
 #define AUTORTFM_EXPECT(exp) \
-	if (UNLIKELY(!(exp))) \
+	if (AUTORTFM_UNLIKELY(!(exp))) \
 	{ \
 		auto Lambda = [] \
 		{ \
@@ -254,4 +259,14 @@ inline constexpr T RoundUp(T Value, T Multiple)
 #define AUTORTFM_MUST_TAIL
 #endif
 
+#if !defined(CA_SUPPRESS)
+#define CA_SUPPRESS( WarningNumber )
+#endif
+
+[[noreturn]] inline void AutoRTFM::InternalUnreachable() //-V1082 Silence PVS [[noreturn]] false positive warning
+{
+	AUTORTFM_FATAL("Unreachable encountered!");
+	__builtin_unreachable();
+}
+
 #endif // (defined(__AUTORTFM) && __AUTORTFM)
diff --git a/Engine/Source/Runtime/Core/Private/AutoRTFM/WriteLog.h b/Engine/Source/Runtime/Core/Private/AutoRTFM/WriteLog.h
--- a/Engine/Source/Runtime/Core/Private/AutoRTFM/WriteLog.h
+++ b/Engine/Source/Runtime/Core/Private/AutoRTFM/WriteLog.h
@@ -5,9 +5,8 @@
 #if (defined(__AUTORTFM) && __AUTORTFM)
 
 #include "AutoRTFM/AutoRTFM.h"
-#include "HAL/Platform.h"
-#include "Misc/CoreMiscDefines.h"
 #include "Utils.h"
+
 #include <cstddef>
 #include <cstdint>
 #include <cstring>
@@ -73,7 +72,7 @@ namespace AutoRTFM
 			// the FBlock header.
 			explicit FBlock(size_t TotalSize)
 			{ 
-				ensure((TotalSize & (alignof(FRecord) - 1)) == 0);
+				AUTORTFM_ENSURE((TotalSize & (alignof(FRecord) - 1)) == 0);
 				std::byte* End = reinterpret_cast<std::byte*>(this) + TotalSize;
 				DataEnd = DataStart();
 				// Note: The initial empty state has LastRecord pointing one
diff --git a/Engine/Source/Runtime/Core/Public/AutoRTFM/AutoRTFM.h b/Engine/Source/Runtime/Core/Public/AutoRTFM/AutoRTFM.h
--- a/Engine/Source/Runtime/Core/Public/AutoRTFM/AutoRTFM.h
+++ b/Engine/Source/Runtime/Core/Public/AutoRTFM/AutoRTFM.h
@@ -2,6 +2,20 @@
 
 #pragma once
 
+#include "AutoRTFMConstants.h"
+
+#include <memory.h>
+#include <stdlib.h>
+
+#ifdef __cplusplus
+#include "AutoRTFMTask.h"
+#include <algorithm>
+#include <cstdarg>
+#include <tuple>
+#include <type_traits>
+#include <utility>
+#endif // __cplusplus
+
 #if (defined(__AUTORTFM) && __AUTORTFM)
 #define UE_AUTORTFM 1  // Compiler is 'verse-clang'
 #else
@@ -54,30 +68,13 @@
 #define UE_AUTORTFM_ASSUME_SAFE
 #endif
 
-#include "AutoRTFMConstants.h"
-
 #if defined(UE_AUTORTFM_STANDALONE)
-#include <stdarg.h> // for va_list
-#include <stdlib.h>
-#include <memory.h>
-#include <type_traits>
-#include <utility>
 #define UE_AUTORTFM_API
 #define UE_AUTORTFM_FORCEINLINE inline
-#define UE_AUTORTFM_MEMCPY ::memcpy
-#define UE_AUTORTFM_MOVE std::move
-#define PRAGMA_DISABLE_UNREACHABLE_CODE_WARNINGS
-#define PRAGMA_RESTORE_UNREACHABLE_CODE_WARNINGS
-#define LIKELY(EXPR) __builtin_expect(!!(EXPR), 1)
-#define PLATFORM_RETURN_ADDRESS() __builtin_return_address(0)
 #else
 #include <HAL/Platform.h>
-#include <HAL/PlatformMemory.h>
-#include <Templates/UnrealTemplate.h>
 #define UE_AUTORTFM_API CORE_API
 #define UE_AUTORTFM_FORCEINLINE FORCEINLINE
-#define UE_AUTORTFM_MEMCPY FPlatformMemory::Memcpy
-#define UE_AUTORTFM_MOVE MoveTemp
 #endif
 
 #if defined(UE_BUILD_DEBUG) && UE_BUILD_DEBUG
@@ -88,6 +85,17 @@
 #define UE_AUTORTFM_FORCEINLINE inline
 #endif
 
+#ifdef _MSC_VER
+#define AUTORTFM_DISABLE_UNREACHABLE_CODE_WARNINGS \
+	__pragma (warning(push)) \
+	__pragma (warning(disable: 4702)) /* unreachable code */
+#define AUTORTFM_RESTORE_UNREACHABLE_CODE_WARNINGS \
+	__pragma (warning(pop))
+#else
+#define AUTORTFM_DISABLE_UNREACHABLE_CODE_WARNINGS
+#define AUTORTFM_RESTORE_UNREACHABLE_CODE_WARNINGS
+#endif
+
 #define UE_AUTORTFM_UNUSED(UNUSEDVAR) (void)UNUSEDVAR
 
 #ifdef __cplusplus
@@ -341,16 +349,16 @@ static UE_AUTORTFM_FORCEINLINE void autortfm_open_explicit_validation(autortfm_m
 #if UE_AUTORTFM_ENABLED
 [[nodiscard]] UE_AUTORTFM_API autortfm_status autortfm_close(void (*UninstrumentedWork)(void*), void (*InstrumentedWork)(void*), void* Arg);
 #else
-PRAGMA_DISABLE_UNREACHABLE_CODE_WARNINGS
+AUTORTFM_DISABLE_UNREACHABLE_CODE_WARNINGS
 [[nodiscard]] static UE_AUTORTFM_FORCEINLINE autortfm_status autortfm_close(void (*UninstrumentedWork)(void*), void (*InstrumentedWork)(void*), void* Arg)
 {
 	UE_AUTORTFM_UNUSED(UninstrumentedWork);
 	UE_AUTORTFM_UNUSED(InstrumentedWork);
 	UE_AUTORTFM_UNUSED(Arg);
-    abort();
+	abort();
 	return autortfm_status_aborted_by_language;
 }
-PRAGMA_RESTORE_UNREACHABLE_CODE_WARNINGS
+AUTORTFM_RESTORE_UNREACHABLE_CODE_WARNINGS
 #endif
 
 #if UE_AUTORTFM_ENABLED
@@ -382,7 +390,7 @@ struct autortfm_open_to_closed_mapping
 
 // autortfm_open_to_closed_table holds a pointer to a null-terminated list of
 // autortfm_open_to_closed_mapping, and an intrusive linked-list pointer to the
-// next registered autortfm_open_to_closed_table.
+// previous and next registered autortfm_open_to_closed_table.
 struct autortfm_open_to_closed_table
 {
 	// Null-terminated open function to closed function mapping table.
@@ -501,10 +509,6 @@ static UE_AUTORTFM_FORCEINLINE void autortfm_unreachable(const char* Message)
 #ifdef __cplusplus
 }
 
-#include "AutoRTFMTask.h"
-#include <algorithm>
-#include <tuple>
-
 namespace AutoRTFM
 {
 
@@ -851,7 +855,7 @@ static UE_AUTORTFM_FORCEINLINE ReturnType Open(const FunctorType& Functor)
 {
 	static_assert(IsSafeToReturnFromOpen<ReturnType>,
 		"function return type is not safe to return from Open()");
-
+#if UE_AUTORTFM
 	if (!autortfm_is_closed())
 	{
 		return Functor();
@@ -872,7 +876,7 @@ static UE_AUTORTFM_FORCEINLINE ReturnType Open(const FunctorType& Functor)
 
 			if constexpr (VALIDATION_LEVEL == EMemoryValidationLevel::Default)
 			{
-				autortfm_open(&FCallHelper::Call, const_cast<void*>(static_cast<const void*>(&Functor)), PLATFORM_RETURN_ADDRESS());
+				autortfm_open(&FCallHelper::Call, const_cast<void*>(static_cast<const void*>(&Functor)), __builtin_return_address(0));
 			}
 			else
 			{
@@ -880,7 +884,7 @@ static UE_AUTORTFM_FORCEINLINE ReturnType Open(const FunctorType& Functor)
 					static_cast<autortfm_memory_validation_level>(VALIDATION_LEVEL),
 					&FCallHelper::Call,
 					const_cast<void*>(static_cast<const void*>(&Functor)),
-					PLATFORM_RETURN_ADDRESS());
+					__builtin_return_address(0));
 			}
 		}
 		else
@@ -899,7 +903,7 @@ static UE_AUTORTFM_FORCEINLINE ReturnType Open(const FunctorType& Functor)
 			FCallHelper Helper{Functor};
 			if constexpr (VALIDATION_LEVEL == EMemoryValidationLevel::Default)
 			{
-				autortfm_open(&FCallHelper::Call, reinterpret_cast<void*>(&Helper), PLATFORM_RETURN_ADDRESS());
+				autortfm_open(&FCallHelper::Call, reinterpret_cast<void*>(&Helper), __builtin_return_address(0));
 			}
 			else
 			{
@@ -907,11 +911,14 @@ static UE_AUTORTFM_FORCEINLINE ReturnType Open(const FunctorType& Functor)
 					static_cast<autortfm_memory_validation_level>(VALIDATION_LEVEL),
 					&FCallHelper::Call,
 					reinterpret_cast<void*>(&Helper),
-					PLATFORM_RETURN_ADDRESS());
+					__builtin_return_address(0));
 			}
 			return Helper.ReturnValue;
 		}
 	}
+#else // UE_AUTORTFM
+	return Functor();
+#endif // UE_AUTORTFM
 }
 
 // Always executes the given code transactionally when called from a transaction nest
@@ -1439,7 +1446,7 @@ namespace AutoRTFM::Private
 		template<typename FunctorType>
 		void operator+(FunctorType F)
 		{
-			AutoRTFM::Open(UE_AUTORTFM_MOVE(F));
+			AutoRTFM::Open(std::move(F));
 		}
 	};
 	struct FOpenNoMemoryValidationHelper
@@ -1447,7 +1454,7 @@ namespace AutoRTFM::Private
 		template<typename FunctorType>
 		void operator+(FunctorType F)
 		{
-			AutoRTFM::Open<EMemoryValidationLevel::Disabled>(UE_AUTORTFM_MOVE(F));
+			AutoRTFM::Open<EMemoryValidationLevel::Disabled>(std::move(F));
 		}
 	};
 	struct FOnAbortHelper
@@ -1455,7 +1462,7 @@ namespace AutoRTFM::Private
 		template<typename FunctorType>
 		void operator+(FunctorType F)
 		{
-			AutoRTFM::OnAbort(UE_AUTORTFM_MOVE(F));
+			AutoRTFM::OnAbort(std::move(F));
 		}
 	};
 	struct FOnCommitHelper
@@ -1463,7 +1470,7 @@ namespace AutoRTFM::Private
 		template<typename FunctorType>
 		void operator+(FunctorType F)
 		{
-			AutoRTFM::OnCommit(UE_AUTORTFM_MOVE(F));
+			AutoRTFM::OnCommit(std::move(F));
 		}
 	};
 	struct FTransactHelper
@@ -1471,7 +1478,7 @@ namespace AutoRTFM::Private
 		template<typename FunctorType>
 		void operator+(FunctorType F)
 		{
-			AutoRTFM::Transact(UE_AUTORTFM_MOVE(F));
+			AutoRTFM::Transact(std::move(F));
 		}
 	};
 	namespace /* must have internal linkage */
diff --git a/Engine/Source/Runtime/Core/Public/AutoRTFM/AutoRTFMMetrics.h b/Engine/Source/Runtime/Core/Public/AutoRTFM/AutoRTFMMetrics.h
--- a/Engine/Source/Runtime/Core/Public/AutoRTFM/AutoRTFMMetrics.h
+++ b/Engine/Source/Runtime/Core/Public/AutoRTFM/AutoRTFMMetrics.h
@@ -2,19 +2,19 @@
 
 #pragma once
 
-#include "AutoRTFM/AutoRTFM.h"
+#include <cstdint>
 
 namespace AutoRTFM
 {
 
 struct FAutoRTFMMetrics
 {
-	int64 NumTransactionsStarted;
-	int64 NumTransactionsCommitted;
+	int64_t NumTransactionsStarted;
+	int64_t NumTransactionsCommitted;
 
-	int64 NumTransactionsAborted;
-	int64 NumTransactionsAbortedByRequest;
-	int64 NumTransactionsAbortedByLanguage;
+	int64_t NumTransactionsAborted;
+	int64_t NumTransactionsAbortedByRequest;
+	int64_t NumTransactionsAbortedByLanguage;
 };
 
 // reset the internal metrics tracking back to zero
diff --git a/Engine/Source/Runtime/Core/Public/AutoRTFM/AutoRTFMTesting.h b/Engine/Source/Runtime/Core/Public/AutoRTFM/AutoRTFMTesting.h
--- a/Engine/Source/Runtime/Core/Public/AutoRTFM/AutoRTFMTesting.h
+++ b/Engine/Source/Runtime/Core/Public/AutoRTFM/AutoRTFMTesting.h
@@ -4,48 +4,46 @@
 
 #include "AutoRTFM.h"
 
-#if defined(UE_AUTORTFM_STANDALONE)
-#include <assert.h>
-#define UE_AUTORTFM_CHECK(x) assert(x)
-#else
-#include <Misc/AssertionMacros.h>
-#define UE_AUTORTFM_CHECK(x) check(x)
-#endif
-
-namespace AutoRTFM
+namespace AutoRTFM::Testing
 {
-namespace Testing
+
+// Implemented by the test framework, and called when AUTORTFM_TESTING_ASSERT() fails.
+void AssertionFailure(const char* Expression, const char* File, int Line);
+
+#define AUTORTFM_TESTING_ASSERT(Condition) \
+	do { if (!(Condition)) { AssertionFailure("AUTORTFM_TESTING_ASSERT(" #Condition ")", __FILE__, __LINE__); }} while (false)
+
+// Run the callback in a transaction like Transact, but abort program
+// execution if the result is anything other than autortfm_committed.
+// Useful for testing.
+template<typename TFunctor>
+static UE_AUTORTFM_FORCEINLINE void Commit(const TFunctor& Functor)
+{
+	const ETransactionResult Result = Transact(Functor);
+	AUTORTFM_TESTING_ASSERT(ETransactionResult::Committed == Result);
+}
+
+// Run the callback in a transaction like Transact, but abort program
+// execution if the result is anything other than abort.
+template<typename TFunctor>
+static UE_AUTORTFM_FORCEINLINE void Abort(const TFunctor& Functor)
+{
+	const ETransactionResult Result = Transact(Functor);
+	AUTORTFM_TESTING_ASSERT(ETransactionResult::Committed != Result);
+}
+
+// Force set the AutoRTFM runtime state. Used in testing. Returns the old value.
+UE_AUTORTFM_API ForTheRuntime::EAutoRTFMEnabledState ForceSetAutoRTFMRuntime(ForTheRuntime::EAutoRTFMEnabledState State);
+
+struct FEnabledStateResetterScoped final
 {
-	// Run the callback in a transaction like Transact, but abort program
-	// execution if the result is anything other than autortfm_committed.
-	// Useful for testing.
-	template<typename TFunctor>
-	static UE_AUTORTFM_FORCEINLINE void Commit(const TFunctor& Functor)
-	{
-		const ETransactionResult Result = Transact(Functor);
-		UE_AUTORTFM_CHECK(ETransactionResult::Committed == Result);
-	}
-
-	// Run the callback in a transaction like Transact, but abort program
-	// execution if the result is anything other than abort.
-	template<typename TFunctor>
-	static UE_AUTORTFM_FORCEINLINE void Abort(const TFunctor& Functor)
-	{
-		const ETransactionResult Result = Transact(Functor);
-		UE_AUTORTFM_CHECK(ETransactionResult::Committed != Result);
-	}
-
-	// Force set the AutoRTFM runtime state. Used in testing. Returns the old value.
-	UE_AUTORTFM_API ForTheRuntime::EAutoRTFMEnabledState ForceSetAutoRTFMRuntime(ForTheRuntime::EAutoRTFMEnabledState State);
-
-	struct FEnabledStateResetterScoped final
-	{
-		FEnabledStateResetterScoped(ForTheRuntime::EAutoRTFMEnabledState State) : Original(ForceSetAutoRTFMRuntime(State)) {}
-		~FEnabledStateResetterScoped() { ForceSetAutoRTFMRuntime(Original); }
-
-	private:
-		const ForTheRuntime::EAutoRTFMEnabledState Original;
-	};
-} // namespace Testing
-
-} // namespace AutoRTFM
+	FEnabledStateResetterScoped(ForTheRuntime::EAutoRTFMEnabledState State) : Original(ForceSetAutoRTFMRuntime(State)) {}
+	~FEnabledStateResetterScoped() { ForceSetAutoRTFMRuntime(Original); }
+
+private:
+	const ForTheRuntime::EAutoRTFMEnabledState Original;
+};
+
+#undef AUTORTFM_TESTING_ASSERT
+
+} // namespace AutoRTFM::Testing
diff --git a/Engine/Source/Runtime/Engine/Private/Tests/AutoRTFM/AutoRTFMTestingAssertionFailure.cpp b/Engine/Source/Runtime/Engine/Private/Tests/AutoRTFM/AutoRTFMTestingAssertionFailure.cpp
new file mode 100644
--- /dev/null
+++ b/Engine/Source/Runtime/Engine/Private/Tests/AutoRTFM/AutoRTFMTestingAssertionFailure.cpp
@@ -0,0 +1,13 @@
+// Copyright Epic Games, Inc. All Rights Reserved.
+
+#include "AutoRTFM/AutoRTFMTesting.h"
+#include "Misc/AssertionMacros.h"
+
+#if WITH_DEV_AUTOMATION_TESTS
+
+void AutoRTFM::Testing::AssertionFailure(const char* Expression, const char* File, int Line)
+{
+	FDebug::AssertFailed(Expression, File, Line);
+}
+
+#endif
