Fix `GetOSVersion` under transactions.

diff --git a/Engine/Source/Programs/AutoRTFMTests/Private/UECoreTests.cpp b/Engine/Source/Programs/AutoRTFMTests/Private/UECoreTests.cpp
--- a/Engine/Source/Programs/AutoRTFMTests/Private/UECoreTests.cpp
+++ b/Engine/Source/Programs/AutoRTFMTests/Private/UECoreTests.cpp
@@ -2965,3 +2965,12 @@ TEST_CASE("UECore.HeartBeat")
 		FDisableHitchDetectorScope _;
 	});
 }
+
+TEST_CASE("UECore.GetOSVersion")
+{
+	FString Version;
+	AutoRTFM::Testing::Commit([&]
+	{
+		Version = FPlatformMisc::GetOSVersion();
+	});
+}
diff --git a/Engine/Source/Runtime/Core/Private/Unix/UnixPlatformMisc.cpp b/Engine/Source/Runtime/Core/Private/Unix/UnixPlatformMisc.cpp
--- a/Engine/Source/Runtime/Core/Private/Unix/UnixPlatformMisc.cpp
+++ b/Engine/Source/Runtime/Core/Private/Unix/UnixPlatformMisc.cpp
@@ -551,6 +551,7 @@ void FUnixPlatformMisc::GetOSVersions(FString& out_OSVersionLabel, FString& out_
 	}
 }
 
+UE_AUTORTFM_ALWAYS_OPEN
 FString FUnixPlatformMisc::GetOSVersion()
 {
 	TMap<FString, FString> OsInfo = ReadConfigurationFile(TEXT("/etc/os-release"));
