FBox ctor from FSphere

Box.h already includes Sphere.h so might as well.
Start setting up some very basic tests for FBox while I'm there.

diff --git a/Engine/Source/Runtime/Core/Public/Math/Box.h b/Engine/Source/Runtime/Core/Public/Math/Box.h
--- a/Engine/Source/Runtime/Core/Public/Math/Box.h
+++ b/Engine/Source/Runtime/Core/Public/Math/Box.h
@@ -165,6 +165,13 @@ public:
 	{
 	}
 
+	/**
+	 * Creates and initializes a box that tightly encloses the given sphere.
+	 *
+	 * @param Sphere The sphere being enclosed.
+	 */
+	explicit TBox(const TSphere<T>& Sphere) : TBox<T>(Sphere.Center - TVector<T>(Sphere.W), Sphere.Center + TVector<T>(Sphere.W)) {};
+
 public:
 
 	/**
diff --git a/Engine/Source/Runtime/Core/Tests/Math/UnrealMathTest.cpp b/Engine/Source/Runtime/Core/Tests/Math/UnrealMathTest.cpp
--- a/Engine/Source/Runtime/Core/Tests/Math/UnrealMathTest.cpp
+++ b/Engine/Source/Runtime/Core/Tests/Math/UnrealMathTest.cpp
@@ -19,6 +19,7 @@
 #include "Math/Quat.h"
 #include "Math/QuatRotationTranslationMatrix.h"
 #include "Math/Color.h"
+#include "Math/Box.h"
 #include "Async/ParallelFor.h"
 #include "Misc/ScopeLock.h"
 #include <limits>
@@ -5126,4 +5127,30 @@ TEST_CASE_NAMED(FColorConversionHeavyTest, "System::Core::Math::ColorConversionH
 	
 }
 
-#endif //WITH_TESTS
\ No newline at end of file
+TEST_CASE_NAMED(FBoxTest, "System::Core::Math::Box Test", "[ApplicationContextMask][SmokeFilter]")
+{
+	// Default ctor
+	CHECK(!FBox3f {}.IsValid);
+	CHECK(!FBox3d {}.IsValid);
+
+	// Constructors from explicit coordinates and from spheres
+	// as well as operator ==
+	CHECK(FBox3f { FVector3f { -1.f, -1.f, -1.f }, FVector3f { 1.f, 1.f, 1.f} } == FBox3f { FSphere3f { FVector3f { 0.f, 0.f, 0.f, }, 1.f } });
+	CHECK(FBox3d { FVector3d { -1.0,  0.0,  1.0 }, FVector3d { 1.0, 2.0, 3.0} } == FBox3d { FSphere3d { FVector3d { 0.0, 1.0, 2.0, }, 1.0 } });
+	
+	// Start with empty box then add points add one by one
+	{
+		FBox3f BoxF;
+		BoxF += FVector3f { 1.f, 0.f, 1.f };
+		BoxF = BoxF + FVector3f { 0.f, 1.f, 0.f };
+		CHECK(BoxF == FBox3f { FVector3f { 0.f, 0.f, 0.f, }, FVector3f { 1.f, 1.f, 1.f } });
+		CHECK(BoxF[0] == BoxF.Min);
+		CHECK(BoxF[1] == BoxF.Max);
+
+		FVector3d DoublePoints[2] { FVector3d { 1.0, -2.0, 3.0 }, FVector3d { -4.0, 5.0, -6.0 } };
+		FBox3d BoxD(DoublePoints, 2);
+		CHECK(BoxD == FBox3d { FVector3d { -4.0, -2.0, -6.0 }, FVector3d { 1.0, 5.0, 3.0 } });
+	}
+}
+
+#endif //WITH_TESTS
diff --git a/Engine/Source/Runtime/Engine/Private/Components/SpotLightComponent.cpp b/Engine/Source/Runtime/Engine/Private/Components/SpotLightComponent.cpp
--- a/Engine/Source/Runtime/Engine/Private/Components/SpotLightComponent.cpp
+++ b/Engine/Source/Runtime/Engine/Private/Components/SpotLightComponent.cpp
@@ -291,7 +291,7 @@ void USpotLightComponent::SetLightBrightness(float InBrightness)
 FBox USpotLightComponent::GetStreamingBounds() const
 {
 	const FSphere BoundingSphere = GetBoundingSphere();
-	return FBox(BoundingSphere.Center - BoundingSphere.W, BoundingSphere.Center + BoundingSphere.W);
+	return FBox(BoundingSphere);
 }
 #endif // WITH_EDITOR
 
diff --git a/Engine/Source/Runtime/Renderer/Private/TranslucentLighting.cpp b/Engine/Source/Runtime/Renderer/Private/TranslucentLighting.cpp
--- a/Engine/Source/Runtime/Renderer/Private/TranslucentLighting.cpp
+++ b/Engine/Source/Runtime/Renderer/Private/TranslucentLighting.cpp
@@ -275,7 +275,7 @@ void FViewInfo::CalcTranslucencyLightingVolumeBounds(FBox* InOutCascadeBoundsArr
 			SphereBounds.Center.Y = SphereBounds.Center.Y - FMath::Fmod(SphereBounds.Center.Y, SphereBounds.W * 2 / TranslucencyLightingVolumeDim);
 			SphereBounds.Center.Z = SphereBounds.Center.Z - FMath::Fmod(SphereBounds.Center.Z, SphereBounds.W * 2 / TranslucencyLightingVolumeDim);
 
-			InOutCascadeBoundsArray[CascadeIndex] = FBox(SphereBounds.Center - SphereBounds.W, SphereBounds.Center + SphereBounds.W);
+			InOutCascadeBoundsArray[CascadeIndex] = FBox(SphereBounds);
 		}
 		else
 		{
