Tighten up timeout logging code guarding GetTests().

Previously, there were two minor flaws:
- We counted the time spent in `TestInfo.Append(TestsToAdd);` against the test.
- We ran the timeout check even if no tests were being added.

In practice these changes didn't actually affect the bug I was looking at, but still seemed worth cleaning up.

diff --git a/Engine/Source/Runtime/Core/Private/Misc/AutomationTest.cpp b/Engine/Source/Runtime/Core/Private/Misc/AutomationTest.cpp
--- a/Engine/Source/Runtime/Core/Private/Misc/AutomationTest.cpp
+++ b/Engine/Source/Runtime/Core/Private/Misc/AutomationTest.cpp
@@ -805,24 +805,24 @@ void FAutomationTestFramework::GetValidTestNames( TArray<FAutomationTestInfo>& T
 		const bool bPassesFeatureRequirements = !CurTestFeatureFlags || !!(CurTestFeatureFlags & FeatureSupportFlags);
 
 		const bool bEnabled = !(CurTestFlags & EAutomationTestFlags::Disabled);
-
-		const double GenerateTestNamesStartTime = FPlatformTime::Seconds();
-		
 		if (bEnabled && bPassesApplicationRequirements && bPassesFeatureRequirements && bPassesFilterRequirement)
 		{
+			// Make sure people are not writing complex tests that take forever to return the names of the tests
+			// otherwise the session frontend locks up when looking at your local tests.
+			const double GenerateTestNamesStartTime = FPlatformTime::Seconds();
+
 			TArray<FAutomationTestInfo> TestsToAdd;
 			CurTest->GenerateTestNames(TestsToAdd);
-			TestInfo.Append(TestsToAdd);
-		}
+		
+			const double GenerateTestNamesEndTime = FPlatformTime::Seconds();
+			const double TimeForGetTests = GenerateTestNamesEndTime - GenerateTestNamesStartTime;
+			if (TimeForGetTests > 10.0)
+			{
+				// Emit a warning if GetTests() takes too long.
+				UE_LOG(LogAutomationTest, Warning, TEXT("Automation Test '%s' took > 10 seconds to return from GetTests(...): %.2fs"), *CurTest->GetTestName(), (float)TimeForGetTests);
+			}
 
-		// Make sure people are not writing complex tests that take forever to return the names of the tests
-		// otherwise the session frontend locks up when looking at your local tests.
-		const double GenerateTestNamesEndTime = FPlatformTime::Seconds();
-		const double TimeForGetTests = static_cast<float>(GenerateTestNamesEndTime - GenerateTestNamesStartTime);
-		if (TimeForGetTests > 10.0f)
-		{
-			//force a failure if a smoke test takes too long
-			UE_LOG(LogAutomationTest, Warning, TEXT("Automation Test '%s' took > 10 seconds to return from GetTests(...): %.2fs"), *CurTest->GetTestName(), (float)TimeForGetTests);
+			TestInfo.Append(TestsToAdd);
 		}
 	}
 }
