Fix Windows' IsRunningOnBattery to account for all cases.

GetSystemPowerStatus's BatteryFlag field is set to 0 in the case where the battery is not charging and the battery is between 33% and 66%. This case was not handled correctly and the function would return false as if there was no battery.

Also treat this field as a flag that could have multiple flags set (although in testing this doesn't appear to be always the case).

#rnx

diff --git a/Engine/Source/Runtime/Core/Private/Windows/WindowsPlatformMisc.cpp b/Engine/Source/Runtime/Core/Private/Windows/WindowsPlatformMisc.cpp
--- a/Engine/Source/Runtime/Core/Private/Windows/WindowsPlatformMisc.cpp
+++ b/Engine/Source/Runtime/Core/Private/Windows/WindowsPlatformMisc.cpp
@@ -4188,19 +4188,22 @@ FText FWindowsPlatformMisc::GetFileManagerName()
 bool FWindowsPlatformMisc::IsRunningOnBattery()
 {
 	SYSTEM_POWER_STATUS status;
-	GetSystemPowerStatus(&status);
-	switch(status.BatteryFlag)
+	if (GetSystemPowerStatus(&status) == 0)
 	{
-	case 4://	"Critical-the battery capacity is at less than five percent"
-	case 2://	"Low-the battery capacity is at less than 33 percent"
-	case 1://	"High-the battery capacity is at more than 66 percent"
-	case 8://	"Charging"
-		return true;
-	case 128://	"No system battery" - desktop, NB: UPS don't count as batteries under Windows
-	case 255://	"Unknown status-unable to read the battery flag information"
-	default:
 		return false;
 	}
+
+	BYTE BatteryPresentFlags = 4  // Critical-the battery capacity is at less than five percent.
+		| 2  // Low-the battery capacity is at less than 33 percent.
+		| 1  // High-the battery capacity is at more than 66 percent.
+		| 8; // Charging.
+
+	BYTE BatteryNotPresentFlags = 128; // No system battery.
+
+	bool BatteryPresent = ((status.BatteryFlag & BatteryPresentFlags) != 0) && ((status.BatteryFlag & BatteryNotPresentFlags) == 0);
+
+	// We also need to check for 0, which means "the battery is not being charged and the battery capacity is between low and high".
+	return BatteryPresent || status.BatteryFlag == 0;
 }
 
 FString FWindowsPlatformMisc::GetOperatingSystemId()
