[StripedMap]
- Fix transactionally safe mutex scopes

diff --git a/Engine/Source/Runtime/Core/Public/Containers/StripedMap.h b/Engine/Source/Runtime/Core/Public/Containers/StripedMap.h
--- a/Engine/Source/Runtime/Core/Public/Containers/StripedMap.h
+++ b/Engine/Source/Runtime/Core/Public/Containers/StripedMap.h
@@ -27,8 +27,8 @@ struct FSharedMutexStripedMapLockingPolicy
 struct FTransactionallySafeStripedMapLockingPolicy
 {
 	typedef FTransactionallySafeRWLock        MutexType;
-	typedef FWriteScopeLock                   ExclusiveLockType;
-	typedef FReadScopeLock                    SharedLockType;
+	typedef UE::TWriteScopeLock<MutexType>    ExclusiveLockType;
+	typedef UE::TReadScopeLock<MutexType>     SharedLockType;
 };
 
 // Use FTransactionallySafeRWLock by default for now because there is no transactionally-safe
