Use UE::FSourceLocation to identify AutomationTest file and line numbers.

We are seeing warnings on Horde that test initialization is taking longer than 10 seconds, and I believe repeated stack crawls were the culprit. Locally, I was seeing approximately 0.5 seconds for initialization to occur on my machine, with no other system load.

C++20 added a way to get the caller's filename and line number directly from the compiler, and UE supports this. This is much, much more efficient than performing a stack crawl to obtain this information. Now, these tests initialize locally almost instantly. I believe this should run faster on Horde as well.

diff --git a/Engine/Source/Runtime/Core/Public/Misc/AutomationTest.h b/Engine/Source/Runtime/Core/Public/Misc/AutomationTest.h
--- a/Engine/Source/Runtime/Core/Public/Misc/AutomationTest.h
+++ b/Engine/Source/Runtime/Core/Public/Misc/AutomationTest.h
@@ -40,6 +40,8 @@
 #include "Misc/Guid.h"
 #include "Misc/Optional.h"
 #include "Misc/OutputDevice.h"
+#include "Misc/SourceLocation.h"
+#include "Misc/TextFilterExpressionEvaluator.h"
 #include "Misc/Timespan.h"
 #include "Templates/Function.h"
 #include "Templates/SharedPointer.h"
@@ -3296,73 +3298,66 @@ public:
 		//disabled
 	}
 
-	void It(const FString& InDescription, TFunction<void()> DoWork)
+	void It(const FString& InDescription, TFunction<void()> DoWork, UE::FSourceLocation Location = UE::FSourceLocation::Current())
 	{
 		const TSharedRef<FSpecDefinitionScope> CurrentScope = DefinitionScopeStack.Last();
-		const auto Stack = GetStack();
 
 		PushDescription(InDescription);
-		CurrentScope->It.Push(MakeShareable(new FSpecIt(GetDescription(), GetId(), Stack.Get()[0].Filename, Stack.Get()[0].LineNumber, MakeShareable(new FSingleExecuteLatentCommand(this, DoWork, bEnableSkipIfError)))));
+		CurrentScope->It.Push(MakeShareable(new FSpecIt(GetDescription(), GetId(), Location.GetFileName(), Location.GetLine(), MakeShareable(new FSingleExecuteLatentCommand(this, DoWork, bEnableSkipIfError)))));
 		PopDescription(InDescription);
 	}
 
-	void It(const FString& InDescription, EAsyncExecution Execution, TFunction<void()> DoWork)
+	void It(const FString& InDescription, EAsyncExecution Execution, TFunction<void()> DoWork, UE::FSourceLocation Location = UE::FSourceLocation::Current())
 	{
 		const TSharedRef<FSpecDefinitionScope> CurrentScope = DefinitionScopeStack.Last();
-		const auto Stack = GetStack();
 
 		PushDescription(InDescription);
-		CurrentScope->It.Push(MakeShareable(new FSpecIt(GetDescription(), GetId(), Stack.Get()[0].Filename, Stack.Get()[0].LineNumber, MakeShareable(new FAsyncLatentCommand(this, Execution, DoWork, DefaultTimeout, bEnableSkipIfError)))));
+		CurrentScope->It.Push(MakeShareable(new FSpecIt(GetDescription(), GetId(), Location.GetFileName(), Location.GetLine(), MakeShareable(new FAsyncLatentCommand(this, Execution, DoWork, DefaultTimeout, bEnableSkipIfError)))));
 		PopDescription(InDescription);
 	}
 
-	void It(const FString& InDescription, EAsyncExecution Execution, const FTimespan& Timeout, TFunction<void()> DoWork)
+	void It(const FString& InDescription, EAsyncExecution Execution, const FTimespan& Timeout, TFunction<void()> DoWork, UE::FSourceLocation Location = UE::FSourceLocation::Current())
 	{
 		const TSharedRef<FSpecDefinitionScope> CurrentScope = DefinitionScopeStack.Last();
-		const auto Stack = GetStack();
 
 		PushDescription(InDescription);
-		CurrentScope->It.Push(MakeShareable(new FSpecIt(GetDescription(), GetId(), Stack.Get()[0].Filename, Stack.Get()[0].LineNumber, MakeShareable(new FAsyncLatentCommand(this, Execution, DoWork, Timeout, bEnableSkipIfError)))));
+		CurrentScope->It.Push(MakeShareable(new FSpecIt(GetDescription(), GetId(), Location.GetFileName(), Location.GetLine(), MakeShareable(new FAsyncLatentCommand(this, Execution, DoWork, Timeout, bEnableSkipIfError)))));
 		PopDescription(InDescription);
 	}
 
-	void LatentIt(const FString& InDescription, TFunction<void(const FDoneDelegate&)> DoWork)
+	void LatentIt(const FString& InDescription, TFunction<void(const FDoneDelegate&)> DoWork, UE::FSourceLocation Location = UE::FSourceLocation::Current())
 	{
 		const TSharedRef<FSpecDefinitionScope> CurrentScope = DefinitionScopeStack.Last();
-		const auto Stack = GetStack();
 
 		PushDescription(InDescription);
-		CurrentScope->It.Push(MakeShareable(new FSpecIt(GetDescription(), GetId(), Stack.Get()[0].Filename, Stack.Get()[0].LineNumber, MakeShareable(new FUntilDoneLatentCommand(this, DoWork, DefaultTimeout, bEnableSkipIfError)))));
+		CurrentScope->It.Push(MakeShareable(new FSpecIt(GetDescription(), GetId(), Location.GetFileName(), Location.GetLine(), MakeShareable(new FUntilDoneLatentCommand(this, DoWork, DefaultTimeout, bEnableSkipIfError)))));
 		PopDescription(InDescription);
 	}
 
-	void LatentIt(const FString& InDescription, const FTimespan& Timeout, TFunction<void(const FDoneDelegate&)> DoWork)
+	void LatentIt(const FString& InDescription, const FTimespan& Timeout, TFunction<void(const FDoneDelegate&)> DoWork, UE::FSourceLocation Location = UE::FSourceLocation::Current())
 	{
 		const TSharedRef<FSpecDefinitionScope> CurrentScope = DefinitionScopeStack.Last();
-		const auto Stack = GetStack();
 
 		PushDescription(InDescription);
-		CurrentScope->It.Push(MakeShareable(new FSpecIt(GetDescription(), GetId(), Stack.Get()[0].Filename, Stack.Get()[0].LineNumber, MakeShareable(new FUntilDoneLatentCommand(this, DoWork, Timeout, bEnableSkipIfError)))));
+		CurrentScope->It.Push(MakeShareable(new FSpecIt(GetDescription(), GetId(), Location.GetFileName(), Location.GetLine(), MakeShareable(new FUntilDoneLatentCommand(this, DoWork, Timeout, bEnableSkipIfError)))));
 		PopDescription(InDescription);
 	}
 
-	void LatentIt(const FString& InDescription, EAsyncExecution Execution, TFunction<void(const FDoneDelegate&)> DoWork)
+	void LatentIt(const FString& InDescription, EAsyncExecution Execution, TFunction<void(const FDoneDelegate&)> DoWork, UE::FSourceLocation Location = UE::FSourceLocation::Current())
 	{
 		const TSharedRef<FSpecDefinitionScope> CurrentScope = DefinitionScopeStack.Last();
-		const auto Stack = GetStack();
 
 		PushDescription(InDescription);
-		CurrentScope->It.Push(MakeShareable(new FSpecIt(GetDescription(), GetId(), Stack.Get()[0].Filename, Stack.Get()[0].LineNumber, MakeShareable(new FAsyncUntilDoneLatentCommand(this, Execution, DoWork, DefaultTimeout, bEnableSkipIfError)))));
+		CurrentScope->It.Push(MakeShareable(new FSpecIt(GetDescription(), GetId(), Location.GetFileName(), Location.GetLine(), MakeShareable(new FAsyncUntilDoneLatentCommand(this, Execution, DoWork, DefaultTimeout, bEnableSkipIfError)))));
 		PopDescription(InDescription);
 	}
 
-	void LatentIt(const FString& InDescription, EAsyncExecution Execution, const FTimespan& Timeout, TFunction<void(const FDoneDelegate&)> DoWork)
+	void LatentIt(const FString& InDescription, EAsyncExecution Execution, const FTimespan& Timeout, TFunction<void(const FDoneDelegate&)> DoWork, UE::FSourceLocation Location = UE::FSourceLocation::Current())
 	{
 		const TSharedRef<FSpecDefinitionScope> CurrentScope = DefinitionScopeStack.Last();
-		const auto Stack = GetStack();
 
 		PushDescription(InDescription);
-		CurrentScope->It.Push(MakeShareable(new FSpecIt(GetDescription(), GetId(), Stack.Get()[0].Filename, Stack.Get()[0].LineNumber, MakeShareable(new FAsyncUntilDoneLatentCommand(this, Execution, DoWork, Timeout, bEnableSkipIfError)))));
+		CurrentScope->It.Push(MakeShareable(new FSpecIt(GetDescription(), GetId(), Location.GetFileName(), Location.GetLine(), MakeShareable(new FAsyncUntilDoneLatentCommand(this, Execution, DoWork, Timeout, bEnableSkipIfError)))));
 		PopDescription(InDescription);
 	}
 
@@ -3697,45 +3692,7 @@ private:
 		return CompleteId;
 	}
 
-	static TArray<FProgramCounterSymbolInfo> StackWalk(int32 IgnoreCount, int32 MaxDepth)
-	{
-		QUICK_SCOPE_CYCLE_COUNTER(STAT_FAutomationSpecBase_StackWalk);
-
-		LLM_SCOPE_BYNAME(TEXT("AutomationTest/Framework"));
-		SAFE_GETSTACK(Stack, IgnoreCount, MaxDepth);
-		return Stack;
-	}
-
-	static TArray<FProgramCounterSymbolInfo> SkipStackWalk()
-	{
-		QUICK_SCOPE_CYCLE_COUNTER(STAT_FAutomationSpecBase_SkipStackWalk);
-
-		LLM_SCOPE_BYNAME(TEXT("AutomationTest/Framework"));
-		TArray<FProgramCounterSymbolInfo> Stack;
-		FProgramCounterSymbolInfo First;
-		TCString<ANSICHAR>::Strncpy(First.Filename, "Unknown", FProgramCounterSymbolInfo::MAX_NAME_LENGTH);
-		First.LineNumber = 0;
-		Stack.Add(First);
-
-		return Stack;
-	}
-
-	static TSharedRef<TArray<FProgramCounterSymbolInfo>> GetStack()
-	{
-		QUICK_SCOPE_CYCLE_COUNTER(STAT_FAutomationSpecBase_GetStack);
-
-		const bool NeedSkipStackWalk(FAutomationTestFramework::NeedSkipStackWalk());
-		constexpr int32 IgnoreCount(3);
-		constexpr int32 MaxDepth(1);
-		
-		TSharedRef<TArray<FProgramCounterSymbolInfo>> Stack = MakeShared<TArray<FProgramCounterSymbolInfo>>(
-			NeedSkipStackWalk ? SkipStackWalk() : StackWalk(IgnoreCount, MaxDepth));
-
-		return Stack;
-	}
-
 private:
-
 	TArray<FString> Description;
 	TMap<FString, TSharedRef<FSpec>> IdToSpecMap;
 	TSharedPtr<FSpecDefinitionScope> RootDefinitionScope;
