Fixed some comments and names in FString.

diff --git a/Engine/Source/Runtime/Core/Private/Containers/String.cpp.inl b/Engine/Source/Runtime/Core/Private/Containers/String.cpp.inl
--- a/Engine/Source/Runtime/Core/Private/Containers/String.cpp.inl
+++ b/Engine/Source/Runtime/Core/Private/Containers/String.cpp.inl
@@ -255,17 +255,17 @@ UE_STRING_CLASS UE_STRING_CLASS::ConstructFromPtrSizeWithSlack(const WIDECHAR* S
 UE_STRING_CLASS UE_STRING_CLASS::ConstructFromPtrSizeWithSlack(const UTF8CHAR* Str, int32 Len, int32 ExtraSlack)	{ UE_STRING_CLASS Result; UE::Core::Private::ConstructWithLength(Result.Data, Str, Len, ExtraSlack); return Result; }
 UE_STRING_CLASS UE_STRING_CLASS::ConstructFromPtrSizeWithSlack(const UCS2CHAR* Str, int32 Len, int32 ExtraSlack)	{ UE_STRING_CLASS Result; UE::Core::Private::ConstructWithLength(Result.Data, Str, Len, ExtraSlack); return Result; }
 
-UE_STRING_CLASS& UE_STRING_CLASS::operator=( const ElementType* Other )
+UE_STRING_CLASS& UE_STRING_CLASS::operator=(const ElementType* Other)
 {
 	if (Data.GetData() != Other)
 	{
-		int32 Len = (Other && *Other) ? TCString<ElementType>::Strlen(Other)+1 : 0;
+		int32 Len = (Other && *Other) ? TCString<ElementType>::Strlen(Other) + 1 : 0;
 		Data.Empty(Len);
 		Data.AddUninitialized(Len);
-			
-		if( Len )
+
+		if (Len)
 		{
-			FMemory::Memcpy( Data.GetData(), Other, Len * sizeof(ElementType) );
+			FMemory::Memcpy(Data.GetData(), Other, Len * sizeof(ElementType));
 		}
 	}
 	return *this;
@@ -365,7 +365,7 @@ UE_STRING_CLASS& UE_STRING_CLASS::AppendChar(ElementType InChar)
 {
 	CheckInvariants();
 
-	if ( InChar != 0 )
+	if (InChar != 0)
 	{
 		// position to insert the character.  
 		// At the end of the string if we have existing characters, otherwise at the 0 position
@@ -408,7 +408,7 @@ void UE_STRING_CLASS::AppendChars(const UTF8CHAR* Str, int32 Count)
 
 void UE_STRING_CLASS::TrimToNullTerminator()
 {
-	if( Data.Num() )
+	if (Data.Num())
 	{
 		int32 DataLen = TCString<ElementType>::Strlen(Data.GetData());
 		check(DataLen == 0 || DataLen < Data.Num());
@@ -444,7 +444,7 @@ int32 UE_STRING_CLASS::Find(const ElementType* SubStr, int32 SubStrLen, ESearchC
 	{
 		// if ignoring, do a onetime ToUpper on both strings, to avoid ToUppering multiple
 		// times in the loop below
-		if ( SearchCase == ESearchCase::IgnoreCase)
+		if (SearchCase == ESearchCase::IgnoreCase)
 		{
 			return ToUpper().Find(UE_STRING_CLASS::ConstructFromPtrSize(SubStr, SubStrLen).ToUpper(), ESearchCase::CaseSensitive, SearchDir, StartPosition);
 		}
@@ -478,39 +478,48 @@ int32 UE_STRING_CLASS::Find(const ElementType* SubStr, int32 SubStrLen, ESearchC
 	}
 }
 
-bool UE_STRING_CLASS::Split(const UE_STRING_CLASS& InS, UE_STRING_CLASS* LeftS, UE_STRING_CLASS* RightS, ESearchCase::Type SearchCase, ESearchDir::Type SearchDir) const
+bool UE_STRING_CLASS::Split(const UE_STRING_CLASS& InStr, UE_STRING_CLASS* LeftStr, UE_STRING_CLASS* RightStr, ESearchCase::Type SearchCase, ESearchDir::Type SearchDir) const
 {
-	check(LeftS != RightS || LeftS == nullptr);
+	check(LeftStr != RightStr || LeftStr == nullptr);
 
-	int32 InPos = Find(InS, SearchCase, SearchDir);
+	int32 InPos = Find(InStr, SearchCase, SearchDir);
 
-	if (InPos < 0) { return false; }
+	if (InPos < 0)
+	{
+		return false;
+	}
 
-	if (LeftS)
+	if (LeftStr)
 	{
-		if (LeftS != this)
+		if (LeftStr != this)
 		{
-			*LeftS = Left(InPos);
-			if (RightS) { *RightS = Mid(InPos + InS.Len()); }
+			*LeftStr = Left(InPos);
+			if (RightStr)
+			{
+				*RightStr = Mid(InPos + InStr.Len());
+			}
 		}
 		else
 		{
-			// we know that RightS can't be this so we can safely modify it before we deal with LeftS
-			if (RightS) { *RightS = Mid(InPos + InS.Len()); }
-			*LeftS = Left(InPos);
+			// we know that RightStr can't be this so we can safely modify it before we deal with LeftStr
+			if (RightStr)
+			{
+				*RightStr = Mid(InPos + InStr.Len());
+			}
+			*LeftStr = Left(InPos);
 		}
 	}
-	else if (RightS)
+	else if (RightStr)
 	{
-		*RightS = Mid(InPos + InS.Len());
+		*RightStr = Mid(InPos + InStr.Len());
 	}
 
 	return true;
 }
 
-bool UE_STRING_CLASS::Split(const UE_STRING_CLASS& InS, UE_STRING_CLASS* LeftS, UE_STRING_CLASS* RightS) const
+bool UE_STRING_CLASS::Split(const UE_STRING_CLASS& InStr, UE_STRING_CLASS* LeftStr, UE_STRING_CLASS* RightStr) const
 {
-	return Split(InS, LeftS, RightS, ESearchCase::IgnoreCase);
+	return Split(InStr, LeftStr, RightStr, ESearchCase::IgnoreCase);
 }
 
 UE_STRING_CLASS UE_STRING_CLASS::ToUpper() const &
@@ -599,19 +608,19 @@ bool UE_STRING_CLASS::StartsWith(const ElementType* InPrefix, int32 InPrefixLen,
 	}
 }
 
-bool UE_STRING_CLASS::EndsWith(const ElementType* InSuffix, int32 InSuffixLen, ESearchCase::Type SearchCase ) const
+bool UE_STRING_CLASS::EndsWith(const ElementType* InSuffix, int32 InSuffixLen, ESearchCase::Type SearchCase) const
 {
 	if (SearchCase == ESearchCase::IgnoreCase)
 	{
 		return InSuffixLen > 0 &&
 			Len() >= InSuffixLen &&
-			!TCString<ElementType>::Stricmp( &(*this)[ Len() - InSuffixLen ], InSuffix );
+			!TCString<ElementType>::Stricmp(&(*this)[ Len() - InSuffixLen ], InSuffix);
 	}
 	else
 	{
 		return InSuffixLen > 0 &&
 			Len() >= InSuffixLen &&
-			!TCString<ElementType>::Strcmp( &(*this)[ Len() - InSuffixLen ], InSuffix );
+			!TCString<ElementType>::Strcmp(&(*this)[ Len() - InSuffixLen ], InSuffix);
 	}
 }
 
@@ -661,12 +670,12 @@ void UE_STRING_CLASS::RemoveAt(int32 Index, int32 Count, EAllowShrinking AllowSh
 
 bool UE_STRING_CLASS::RemoveFromStart(const ElementType* InPrefix, int32 InPrefixLen, ESearchCase::Type SearchCase)
 {
-	if (InPrefixLen == 0 )
+	if (InPrefixLen == 0)
 	{
 		return false;
 	}
 
-	if (StartsWith(InPrefix, InPrefixLen, SearchCase ))
+	if (StartsWith(InPrefix, InPrefixLen, SearchCase))
 	{
 		RemoveAt(0, InPrefixLen);
 		return true;
@@ -684,7 +693,7 @@ bool UE_STRING_CLASS::RemoveFromEnd(const ElementType* InSuffix, int32 InSuffixL
 
 	if (EndsWith(InSuffix, InSuffixLen, SearchCase))
 	{
-		RemoveAt( Len() - InSuffixLen, InSuffixLen );
+		RemoveAt(Len() - InSuffixLen, InSuffixLen);
 		return true;
 	}
 
@@ -861,19 +870,19 @@ UE_STRING_CLASS UE_STRING_CLASS::Mid(int32 Start, int32 Count) &&
 	return MoveTemp(*this);
 }
 
-void UE_STRING_CLASS::ReplaceCharInlineCaseSensitive(const ElementType SearchChar, const ElementType ReplacementChar)
+void UE_STRING_CLASS::ReplaceCharInlineCaseSensitive(const ElementType From, const ElementType To)
 {
 	for (ElementType& Character : Data)
 	{
-		Character = Character == SearchChar ? ReplacementChar : Character;
+		Character = Character == From ? To : Character;
 	}
 }
 
-void UE_STRING_CLASS::ReplaceCharInlineIgnoreCase(const ElementType SearchChar, const ElementType ReplacementChar)
+void UE_STRING_CLASS::ReplaceCharInlineIgnoreCase(const ElementType From, const ElementType To)
 {
-	ElementType OtherCaseSearchChar = TChar<ElementType>::IsUpper(SearchChar) ? TChar<ElementType>::ToLower(SearchChar) : TChar<ElementType>::ToUpper(SearchChar);
-	ReplaceCharInlineCaseSensitive(OtherCaseSearchChar, ReplacementChar);
-	ReplaceCharInlineCaseSensitive(SearchChar, ReplacementChar);
+	ElementType OtherCaseFrom = TChar<ElementType>::IsUpper(From) ? TChar<ElementType>::ToLower(From) : TChar<ElementType>::ToUpper(From);
+	ReplaceCharInlineCaseSensitive(OtherCaseFrom, To);
+	ReplaceCharInlineCaseSensitive(From, To);
 }
 
 void UE_STRING_CLASS::TrimStartAndEndInline()
@@ -945,23 +954,23 @@ void UE_STRING_CLASS::TrimCharInline(const ElementType CharacterToTrim, bool* bC
 {
 	bool bQuotesWereRemoved=false;
 	int32 Start = 0, Count = Len();
-	if ( Count > 0 )
+	if (Count > 0)
 	{
-		if ( (*this)[0] == CharacterToTrim )
+		if ((*this)[0] == CharacterToTrim)
 		{
 			Start++;
 			Count--;
 			bQuotesWereRemoved=true;
 		}
 
-		if ( Len() > 1 && (*this)[Len() - 1] == CharacterToTrim )
+		if (Len() > 1 && (*this)[Len() - 1] == CharacterToTrim)
 		{
 			Count--;
 			bQuotesWereRemoved=true;
 		}
 	}
 
-	if ( bCharRemoved != nullptr )
+	if (bCharRemoved != nullptr)
 	{
 		*bCharRemoved = bQuotesWereRemoved;
 	}
@@ -999,12 +1008,12 @@ UE_STRING_CLASS UE_STRING_CLASS::TrimChar(const ElementType CharacterToTrim, boo
 	return MoveTemp(*this);
 }
 
-int32 UE_STRING_CLASS::CullArray( TArray<UE_STRING_CLASS>* InArray )
+int32 UE_STRING_CLASS::CullArray(TArray<UE_STRING_CLASS>* InOutArray)
 {
-	check(InArray);
+	check(InOutArray);
 	UE_STRING_CLASS Empty;
-	InArray->Remove(Empty);
-	return InArray->Num();
+	InOutArray->Remove(Empty);
+	return InOutArray->Num();
 }
 
 UE_STRING_CLASS UE_STRING_CLASS::Reverse() const &
@@ -1022,7 +1031,7 @@ UE_STRING_CLASS UE_STRING_CLASS::Reverse() &&
 
 void UE_STRING_CLASS::ReverseString()
 {
-	if ( Len() > 0 )
+	if (Len() > 0)
 	{
 		ElementType* StartChar = &(*this)[0];
 		ElementType* EndChar = &(*this)[Len()-1];
@@ -1036,24 +1045,24 @@ void UE_STRING_CLASS::ReverseString()
 			StartChar++;
 			EndChar--;
 
-		} while( StartChar < EndChar );	// repeat until we've reached the midpoint of the string
+		} while (StartChar < EndChar);	// repeat until we've reached the midpoint of the string
 	}
 }
 
-UE_STRING_CLASS UE_STRING_CLASS::FormatAsNumber( int32 InNumber )
+UE_STRING_CLASS UE_STRING_CLASS::FormatAsNumber(int32 InNumber)
 {
-	UE_STRING_CLASS Number = UE_STRING_CLASS::FromInt( InNumber ), Result;
+	UE_STRING_CLASS Number = UE_STRING_CLASS::FromInt(InNumber), Result;
 
-	int32 dec = 0;
-	for( int32 x = Number.Len()-1 ; x > -1 ; --x )
+	int32 Dec = 0;
+	for (int32 x = Number.Len()-1 ; x > -1 ; --x)
 	{
 		Result += Number.Mid(x,1);
 
-		dec++;
-		if( dec == 3 && x > 0 )
+		Dec++;
+		if (Dec == 3 && x > 0)
 		{
 			Result += CHARTEXT(ElementType, ',');
-			dec = 0;
+			Dec = 0;
 		}
 	}
 
@@ -1067,26 +1076,26 @@ UE_STRING_CLASS UE_STRING_CLASS::FormatAsNumber( int32 InNumber )
  * @param	Ar				Archive to serialize with
  * @param	MinCharacters	Minimum number of characters to serialize.
  */
-void UE_STRING_CLASS::SerializeAsANSICharArray( FArchive& Ar, int32 MinCharacters ) const
+void UE_STRING_CLASS::SerializeAsANSICharArray(FArchive& Ar, int32 MinCharacters) const
 {
-	int32	Length = FMath::Max( Len(), MinCharacters );
+	int32 Length = FMath::Max(Len(), MinCharacters);
 	Ar << Length;
-	
-	for( int32 CharIndex=0; CharIndex<Len(); CharIndex++ )
+
+	for (int32 CharIndex=0; CharIndex<Len(); CharIndex++)
 	{
-		ANSICHAR AnsiChar = CharCast<ANSICHAR>( (*this)[CharIndex] );
+		ANSICHAR AnsiChar = CharCast<ANSICHAR>((*this)[CharIndex]);
 		Ar << AnsiChar;
 	}
 
 	// Zero pad till minimum number of characters are written.
-	for( int32 i=Len(); i<Length; i++ )
+	for (int32 i = Len(); i < Length; i++)
 	{
 		ANSICHAR NullChar = 0;
 		Ar << NullChar;
 	}
 }
 
-void UE_STRING_CLASS::AppendInt( int32 Num )
+void UE_STRING_CLASS::AppendInt(int32 Num)
 {
 	const ElementType* DigitToChar	= CHARTEXT(ElementType, "9876543210123456789");
 	constexpr int32 ZeroDigitIndex	= 9;
@@ -1100,9 +1109,9 @@ void UE_STRING_CLASS::AppendInt( int32 Num )
 	{
 		TempNum[--TempAt] = DigitToChar[ZeroDigitIndex + (Num % 10)];
 		Num /= 10;
-	} while( Num );
+	} while (Num);
 
-	if( bIsNumberNegative )
+	if (bIsNumberNegative)
 	{
 		TempNum[--TempAt] = CHARTEXT(ElementType, '-');
 	}
@@ -1118,10 +1127,10 @@ bool UE_STRING_CLASS::ToBool() const
 	return TCString<ElementType>::ToBool(**this);
 }
 
-UE_STRING_CLASS UE_STRING_CLASS::FromBlob(const uint8* SrcBuffer,const uint32 SrcSize)
+UE_STRING_CLASS UE_STRING_CLASS::FromBlob(const uint8* SrcBuffer, const uint32 SrcSize)
 {
 	UE_STRING_CLASS Result;
-	Result.Reserve( SrcSize * 3 );
+	Result.Reserve(SrcSize * 3);
 	// Convert and append each byte in the buffer
 	for (uint32 Count = 0; Count < SrcSize; Count++)
 	{
@@ -1130,7 +1139,7 @@ UE_STRING_CLASS UE_STRING_CLASS::FromBlob(const uint8* SrcBuffer,const uint32 Sr
 	return Result;
 }
 
-bool UE_STRING_CLASS::ToBlob(const UE_STRING_CLASS& Source,uint8* DestBuffer,const uint32 DestSize)
+bool UE_STRING_CLASS::ToBlob(const UE_STRING_CLASS& Source, uint8* DestBuffer, const uint32 DestSize)
 {
 	// Make sure the buffer is at least half the size and that the string is an
 	// even number of characters long
@@ -1153,10 +1162,10 @@ bool UE_STRING_CLASS::ToBlob(const UE_STRING_CLASS& Source,uint8* DestBuffer,con
 	return false;
 }
 
-UE_STRING_CLASS UE_STRING_CLASS::FromHexBlob( const uint8* SrcBuffer, const uint32 SrcSize )
+UE_STRING_CLASS UE_STRING_CLASS::FromHexBlob(const uint8* SrcBuffer, const uint32 SrcSize)
 {
 	UE_STRING_CLASS Result;
-	Result.Reserve( SrcSize * 2 );
+	Result.Reserve(SrcSize * 2);
 	// Convert and append each byte in the buffer
 	for (uint32 Count = 0; Count < SrcSize; Count++)
 	{
@@ -1165,7 +1174,7 @@ UE_STRING_CLASS UE_STRING_CLASS::FromHexBlob( const uint8* SrcBuffer, const uint
 	return Result;
 }
 
-bool UE_STRING_CLASS::ToHexBlob( const UE_STRING_CLASS& Source, uint8* DestBuffer, const uint32 DestSize )
+bool UE_STRING_CLASS::ToHexBlob(const UE_STRING_CLASS& Source, uint8* DestBuffer, const uint32 DestSize)
 {
 	// Make sure the buffer is at least half the size and that the string is an
 	// even number of characters long
@@ -1173,7 +1182,7 @@ bool UE_STRING_CLASS::ToHexBlob( const UE_STRING_CLASS& Source, uint8* DestBuffe
 		 (Source.Len() % 2) == 0)
 	{
 		ElementType ConvBuffer[3];
-		ConvBuffer[2] = CHARTEXT(ElementType, '\0' );
+		ConvBuffer[2] = CHARTEXT(ElementType, '\0');
 		int32 WriteIndex = 0;
 		// Walk the string 2 chars at a time
 		ElementType* End = nullptr;
@@ -1181,14 +1190,14 @@ bool UE_STRING_CLASS::ToHexBlob( const UE_STRING_CLASS& Source, uint8* DestBuffe
 		{
 			ConvBuffer[0] = Source[Index];
 			ConvBuffer[1] = Source[Index + 1];
-			DestBuffer[WriteIndex] = (uint8)TCString<ElementType>::Strtoi( ConvBuffer, &End, 16 );
+			DestBuffer[WriteIndex] = (uint8)TCString<ElementType>::Strtoi(ConvBuffer, &End, 16);
 		}
 		return true;
 	}
 	return false;
 }
 
-UE_STRING_CLASS UE_STRING_CLASS::SanitizeFloat( double InFloat, const int32 InMinFractionalDigits )
+UE_STRING_CLASS UE_STRING_CLASS::SanitizeFloat(double InFloat, const int32 InMinFractionalDigits)
 {
 	// Avoids negative zero
 	UE::Core::Private::StripNegativeZero(InFloat);
@@ -1252,13 +1261,13 @@ UE_STRING_CLASS UE_STRING_CLASS::Chr(ElementType Ch)
 }
 
 
-UE_STRING_CLASS UE_STRING_CLASS::ChrN( int32 NumCharacters, ElementType Char )
+UE_STRING_CLASS UE_STRING_CLASS::ChrN(int32 NumCharacters, ElementType Char)
 {
-	check( NumCharacters >= 0 );
+	check(NumCharacters >= 0);
 
 	UE_STRING_CLASS Temp;
 	Temp.Data.AddUninitialized(NumCharacters+1);
-	for( int32 Cx = 0; Cx < NumCharacters; ++Cx )
+	for (int32 Cx = 0; Cx < NumCharacters; ++Cx)
 	{
 		Temp[Cx] = Char;
 	}
@@ -1266,7 +1275,7 @@ UE_STRING_CLASS UE_STRING_CLASS::ChrN( int32 NumCharacters, ElementType Char )
 	return Temp;
 }
 
-UE_STRING_CLASS UE_STRING_CLASS::LeftPad( int32 ChCount ) const
+UE_STRING_CLASS UE_STRING_CLASS::LeftPad(int32 ChCount) const
 {
 	int32 Pad = ChCount - Len();
 
@@ -1279,7 +1288,7 @@ UE_STRING_CLASS UE_STRING_CLASS::LeftPad( int32 ChCount ) const
 		return *this;
 	}
 }
-UE_STRING_CLASS UE_STRING_CLASS::RightPad( int32 ChCount ) const
+UE_STRING_CLASS UE_STRING_CLASS::RightPad(int32 ChCount) const
 {
 	int32 Pad = ChCount - Len();
 
@@ -1303,7 +1312,7 @@ bool UE_STRING_CLASS::IsNumeric() const
 	return TCString<ElementType>::IsNumeric(Data.GetData());
 }
 
-int32 UE_STRING_CLASS::ParseIntoArray( TArray<UE_STRING_CLASS>& OutArray, const ElementType* pchDelim, const bool InCullEmpty ) const
+int32 UE_STRING_CLASS::ParseIntoArray(TArray<UE_STRING_CLASS>& OutArray, const ElementType* pchDelim, const bool bInCullEmpty) const
 {
 	check(pchDelim);
 	OutArray.Reset();
@@ -1311,7 +1320,7 @@ int32 UE_STRING_CLASS::ParseIntoArray( TArray<UE_STRING_CLASS>& OutArray, const
 	if (pchDelim[0] != '\0')
 	{
 		UE::String::EParseTokensOptions ParseOptions = UE::String::EParseTokensOptions::IgnoreCase |
-			(InCullEmpty ? UE::String::EParseTokensOptions::SkipEmpty : UE::String::EParseTokensOptions::None);
+			(bInCullEmpty ? UE::String::EParseTokensOptions::SkipEmpty : UE::String::EParseTokensOptions::None);
 		UE::String::ParseTokens(TStringView<ElementType>(*this), TStringView<ElementType>(pchDelim),
 			[&OutArray](TStringView<ElementType> Token) { OutArray.Emplace(Token); },
 			ParseOptions);
@@ -1354,7 +1363,7 @@ bool UE_STRING_CLASS::MatchesWildcard(const ElementType* InWildcard, int32 InWil
 
 
 /** Caution!! this routine is O(N^2) allocations...use it for parsing very short text or not at all */
-int32 UE_STRING_CLASS::ParseIntoArrayWS( TArray<UE_STRING_CLASS>& OutArray, const ElementType* pchExtraDelim, bool InCullEmpty ) const
+int32 UE_STRING_CLASS::ParseIntoArrayWS(TArray<UE_STRING_CLASS>& OutArray, const ElementType* pchExtraDelim, bool bInCullEmpty) const
 {
 	// default array of White Spaces, the last entry can be replaced with the optional pchExtraDelim string
 	// (if you want to split on white space and another character)
@@ -1375,10 +1384,10 @@ int32 UE_STRING_CLASS::ParseIntoArrayWS( TArray<UE_STRING_CLASS>& OutArray, cons
 		WhiteSpace[NumWhiteSpaces++] = pchExtraDelim;
 	}
 
-	return ParseIntoArray(OutArray, WhiteSpace, NumWhiteSpaces, InCullEmpty);
+	return ParseIntoArray(OutArray, WhiteSpace, NumWhiteSpaces, bInCullEmpty);
 }
 
-int32 UE_STRING_CLASS::ParseIntoArrayLines(TArray<UE_STRING_CLASS>& OutArray, bool InCullEmpty) const
+int32 UE_STRING_CLASS::ParseIntoArrayLines(TArray<UE_STRING_CLASS>& OutArray, bool bInCullEmpty) const
 {
 	// default array of LineEndings
 	static const ElementType* LineEndings[] =
@@ -1390,10 +1399,10 @@ int32 UE_STRING_CLASS::ParseIntoArrayLines(TArray<UE_STRING_CLASS>& OutArray, bo
 
 	// start with just the standard line endings
 	int32 NumLineEndings = UE_ARRAY_COUNT(LineEndings);	
-	return ParseIntoArray(OutArray, LineEndings, NumLineEndings, InCullEmpty);
+	return ParseIntoArray(OutArray, LineEndings, NumLineEndings, bInCullEmpty);
 }
 
-int32 UE_STRING_CLASS::ParseIntoArray(TArray<UE_STRING_CLASS>& OutArray, const ElementType* const * DelimArray, int32 NumDelims, bool InCullEmpty) const
+int32 UE_STRING_CLASS::ParseIntoArray(TArray<UE_STRING_CLASS>& OutArray, const ElementType* const * DelimArray, int32 NumDelims, bool bInCullEmpty) const
 {
 	// Make sure the delimit string is not null or empty
 	check(DelimArray);
@@ -1426,7 +1435,7 @@ int32 UE_STRING_CLASS::ParseIntoArray(TArray<UE_STRING_CLASS>& OutArray, const E
 		{
 			const int32 SubstringLength = SubstringEndIndex - SubstringBeginIndex;
 			// If we're not culling empty strings or if we are but the string isn't empty anyways...
-			if(!InCullEmpty || SubstringLength != 0)
+			if(!bInCullEmpty || SubstringLength != 0)
 			{
 				// ... add new string from substring beginning up to the beginning of this delimiter.
 				OutArray.Add(UE_STRING_CLASS::ConstructFromPtrSize(Start + SubstringBeginIndex, SubstringEndIndex - SubstringBeginIndex));
@@ -1444,7 +1453,7 @@ int32 UE_STRING_CLASS::ParseIntoArray(TArray<UE_STRING_CLASS>& OutArray, const E
 	// Add any remaining characters after the last delimiter.
 	const int32 SubstringLength = Length - SubstringBeginIndex;
 	// If we're not culling empty strings or if we are but the string isn't empty anyways...
-	if(!InCullEmpty || SubstringLength != 0)
+	if(!bInCullEmpty || SubstringLength != 0)
 	{
 		// ... add new string from substring beginning up to the beginning of this delimiter.
 		OutArray.Emplace(TStringView<ElementType>(Start + SubstringBeginIndex, SubstringLength));
@@ -1501,20 +1510,20 @@ UE_STRING_CLASS UE_STRING_CLASS::Replace(const ElementType* From, const ElementT
 	return MoveTemp(*this);
 }
 
-int32 UE_STRING_CLASS::ReplaceInline(const ElementType* SearchText, const ElementType* ReplacementText, ESearchCase::Type SearchCase)
+int32 UE_STRING_CLASS::ReplaceInline(const ElementType* From, const ElementType* To, ESearchCase::Type SearchCase)
 {
 	int32 ReplacementCount = 0;
 
 	if (Len() > 0
-		&& SearchText != nullptr && *SearchText != 0
-		&& ReplacementText != nullptr && (SearchCase == ESearchCase::IgnoreCase || TCString<ElementType>::Strcmp(SearchText, ReplacementText) != 0))
+		&& From != nullptr && *From != 0
+		&& To != nullptr && (SearchCase == ESearchCase::IgnoreCase || TCString<ElementType>::Strcmp(From, To) != 0))
 	{
-		const int32 NumCharsToReplace = TCString<ElementType>::Strlen(SearchText);
-		const int32 NumCharsToInsert = TCString<ElementType>::Strlen(ReplacementText);
+		const int32 NumCharsToReplace = TCString<ElementType>::Strlen(From);
+		const int32 NumCharsToInsert = TCString<ElementType>::Strlen(To);
 
 		if (NumCharsToInsert == NumCharsToReplace)
 		{
-			ElementType* Pos = SearchCase == ESearchCase::IgnoreCase ? TCString<ElementType>::Stristr(&(*this)[0], SearchText) : TCString<ElementType>::Strstr(&(*this)[0], SearchText);
+			ElementType* Pos = SearchCase == ESearchCase::IgnoreCase ? TCString<ElementType>::Stristr(&(*this)[0], From) : TCString<ElementType>::Strstr(&(*this)[0], From);
 			while (Pos != nullptr)
 			{
 				ReplacementCount++;
@@ -1522,12 +1531,12 @@ int32 UE_STRING_CLASS::ReplaceInline(const ElementType* SearchText, const Elemen
 				// TCString<ElementType>::Strcpy now inserts a terminating zero so can't use that
 				for (int32 i = 0; i < NumCharsToInsert; i++)
 				{
-					Pos[i] = ReplacementText[i];
+					Pos[i] = To[i];
 				}
 
 				if (Pos + NumCharsToReplace - **this < Len())
 				{
-					Pos = SearchCase == ESearchCase::IgnoreCase ? TCString<ElementType>::Stristr(Pos + NumCharsToReplace, SearchText) : TCString<ElementType>::Strstr(Pos + NumCharsToReplace, SearchText);
+					Pos = SearchCase == ESearchCase::IgnoreCase ? TCString<ElementType>::Stristr(Pos + NumCharsToReplace, From) : TCString<ElementType>::Strstr(Pos + NumCharsToReplace, From);
 				}
 				else
 				{
@@ -1535,14 +1544,14 @@ int32 UE_STRING_CLASS::ReplaceInline(const ElementType* SearchText, const Elemen
 				}
 			}
 		}
-		else if (Contains(SearchText, SearchCase))
+		else if (Contains(From, SearchCase))
 		{
 			UE_STRING_CLASS Copy(MoveTemp(*this));
 
 			// get a pointer into the character data
 			ElementType* WritePosition = (ElementType*)Copy.Data.GetData();
 			// look for From in the remaining string
-			ElementType* SearchPosition = SearchCase == ESearchCase::IgnoreCase ? TCString<ElementType>::Stristr(WritePosition, SearchText) : TCString<ElementType>::Strstr(WritePosition, SearchText);
+			ElementType* SearchPosition = SearchCase == ESearchCase::IgnoreCase ? TCString<ElementType>::Stristr(WritePosition, From) : TCString<ElementType>::Strstr(WritePosition, From);
 			while (SearchPosition != nullptr)
 			{
 				ReplacementCount++;
@@ -1553,14 +1562,14 @@ int32 UE_STRING_CLASS::ReplaceInline(const ElementType* SearchText, const Elemen
 				// copy everything up to the SearchPosition
 				(*this) += WritePosition;
 
-				// copy over the ReplacementText
-				(*this) += ReplacementText;
+				// copy over the replacement string
+				(*this) += To;
 
 				// restore the letter, just so we don't have 0's in the string
-				*SearchPosition = *SearchText;
+				*SearchPosition = *From;
 
 				WritePosition = SearchPosition + NumCharsToReplace;
-				SearchPosition = SearchCase == ESearchCase::IgnoreCase ? TCString<ElementType>::Stristr(WritePosition, SearchText) : TCString<ElementType>::Strstr(WritePosition, SearchText);
+				SearchPosition = SearchCase == ESearchCase::IgnoreCase ? TCString<ElementType>::Stristr(WritePosition, From) : TCString<ElementType>::Strstr(WritePosition, From);
 			}
 
 			// copy anything left over
@@ -1584,17 +1593,17 @@ UE_STRING_CLASS UE_STRING_CLASS::ReplaceQuotesWithEscapedQuotes() &&
 		const ElementType* pChar = *Copy;
 
 		bool bEscaped = false;
-		while ( *pChar != 0 )
+		while (*pChar != 0)
 		{
-			if ( bEscaped )
+			if (bEscaped)
 			{
 				bEscaped = false;
 			}
-			else if ( *pChar == ElementType('\\') )
+			else if (*pChar == ElementType('\\'))
 			{
 				bEscaped = true;
 			}
-			else if ( *pChar == ElementType('"') )
+			else if (*pChar == ElementType('"'))
 			{
 				*this += ElementType('\\');
 			}
@@ -1619,13 +1628,13 @@ static const UE_STRING_CHARTYPE* PREPROCESSOR_JOIN(CharToEscapeSeqMap_, UE_STRIN
 
 void UE_STRING_CLASS::ReplaceCharWithEscapedCharInline(const TArray<ElementType>* Chars/*=nullptr*/)
 {
-	if ( Len() > 0 && (Chars == nullptr || Chars->Num() > 0) )
+	if (Len() > 0 && (Chars == nullptr || Chars->Num() > 0))
 	{
 		const auto& CharToEscapeSeqMap = PREPROCESSOR_JOIN(CharToEscapeSeqMap_, UE_STRING_CLASS);
 
-		for ( uint32 ChIdx = 0; ChIdx < UE_ARRAY_COUNT(CharToEscapeSeqMap); ChIdx++ )
+		for (uint32 ChIdx = 0; ChIdx < UE_ARRAY_COUNT(CharToEscapeSeqMap); ChIdx++)
 		{
-			if ( Chars == nullptr || Chars->Contains(*(CharToEscapeSeqMap[ChIdx][0])) )
+			if (Chars == nullptr || Chars->Contains(*(CharToEscapeSeqMap[ChIdx][0])))
 			{
 				// use ReplaceInline as that won't create a copy of the string if the character isn't found
 				ReplaceInline(CharToEscapeSeqMap[ChIdx][0], CharToEscapeSeqMap[ChIdx][1]);
@@ -1636,16 +1645,16 @@ void UE_STRING_CLASS::ReplaceCharWithEscapedCharInline(const TArray<ElementType>
 
 void UE_STRING_CLASS::ReplaceEscapedCharWithCharInline(const TArray<ElementType>* Chars/*=nullptr*/)
 {
-	if ( Len() > 0 && (Chars == nullptr || Chars->Num() > 0) )
+	if (Len() > 0 && (Chars == nullptr || Chars->Num() > 0))
 	{
 		const auto& CharToEscapeSeqMap = PREPROCESSOR_JOIN(CharToEscapeSeqMap_, UE_STRING_CLASS);
 
 		// Spin CharToEscapeSeqMap backwards to ensure we're doing the inverse of ReplaceCharWithEscapedChar
-		for ( int32 ChIdx = UE_ARRAY_COUNT(CharToEscapeSeqMap); ChIdx > 0; )
+		for (int32 ChIdx = UE_ARRAY_COUNT(CharToEscapeSeqMap); ChIdx > 0; )
 		{
 			--ChIdx;
 
-			if ( Chars == nullptr || Chars->Contains(*(CharToEscapeSeqMap[ChIdx][0])) )
+			if (Chars == nullptr || Chars->Contains(*(CharToEscapeSeqMap[ChIdx][0])))
 			{
 				// use ReplaceInline as that won't create a copy of the string if the character isn't found
 				ReplaceInline(CharToEscapeSeqMap[ChIdx][1], CharToEscapeSeqMap[ChIdx][0]);
@@ -1664,7 +1673,7 @@ void UE_STRING_CLASS::ConvertTabsToSpacesInline(const int32 InSpacesPerTab)
 	check(InSpacesPerTab > 0);
 
 	int32 TabIndex = 0;
-	while ((TabIndex = Find(CHARTEXT(ElementType, "\t"), ESearchCase::CaseSensitive)) != INDEX_NONE )
+	while ((TabIndex = Find(CHARTEXT(ElementType, "\t"), ESearchCase::CaseSensitive)) != INDEX_NONE)
 	{
 		UE_STRING_CLASS RightSide = Mid(TabIndex+1);
 		LeftInline(TabIndex, EAllowShrinking::No);
@@ -1697,17 +1706,17 @@ UE_STRING_CLASS UE_STRING_CLASS::PrintfImpl(const ElementType* Fmt, ...)
 	int32		Result		= -1;
 
 	// First try to print to a stack allocated location 
-	GET_TYPED_VARARGS_RESULT( ElementType, Buffer, BufferSize, BufferSize-1, Fmt, Fmt, Result );
+	GET_TYPED_VARARGS_RESULT(ElementType, Buffer, BufferSize, BufferSize-1, Fmt, Fmt, Result);
 
 	// If that fails, start allocating regular memory
-	if( Result == -1 )
+	if (Result == -1)
 	{
 		Buffer = nullptr;
 		while(Result == -1)
 		{
 			BufferSize *= 2;
-			Buffer = (ElementType*) FMemory::Realloc( Buffer, BufferSize * sizeof(ElementType) );
-			GET_TYPED_VARARGS_RESULT( ElementType, Buffer, BufferSize, BufferSize-1, Fmt, Fmt, Result );
+			Buffer = (ElementType*) FMemory::Realloc(Buffer, BufferSize * sizeof(ElementType));
+			GET_TYPED_VARARGS_RESULT(ElementType, Buffer, BufferSize, BufferSize-1, Fmt, Fmt, Result);
 		};
 	}
 
@@ -1715,9 +1724,9 @@ UE_STRING_CLASS UE_STRING_CLASS::PrintfImpl(const ElementType* Fmt, ...)
 
 	UE_STRING_CLASS ResultString(Buffer);
 
-	if( BufferSize != STARTING_BUFFER_SIZE )
+	if (BufferSize != STARTING_BUFFER_SIZE)
 	{
-		FMemory::Free( Buffer );
+		FMemory::Free(Buffer);
 	}
 
 	return ResultString;
@@ -1757,7 +1766,7 @@ void UE_STRING_CLASS::AppendfImpl(UE_STRING_CLASS& AppendToMe, const ElementType
 
 static_assert(PLATFORM_LITTLE_ENDIAN, PREPROCESSOR_TO_STRING(UE_STRING_CLASS) " serialization needs updating to support big-endian platforms!");
 
-FArchive& operator<<( FArchive& Ar, UE_STRING_CLASS& A )
+FArchive& operator<<(FArchive& Ar, UE_STRING_CLASS& Str)
 {
 	using ElementType = UE_STRING_CLASS::ElementType;
 
@@ -1806,15 +1815,15 @@ FArchive& operator<<( FArchive& Ar, UE_STRING_CLASS& A )
 
 		#if UE_STRING_CHARTYPE_IS_TCHAR
 			// Resize the array only if it passes the above tests to prevent rogue packets from crashing
-			A.Data.Empty(SaveNum);
-			A.Data.AddUninitialized(SaveNum);
+			Str.Data.Empty(SaveNum);
+			Str.Data.AddUninitialized(SaveNum);
 
 			if (SaveNum)
 			{
 				if (bLoadUnicodeChar)
 				{
 					// read in the unicode string
-					auto Passthru = StringMemoryPassthru<UCS2CHAR>(A.Data.GetData(), SaveNum, SaveNum);
+					auto Passthru = StringMemoryPassthru<UCS2CHAR>(Str.Data.GetData(), SaveNum, SaveNum);
 					Ar.Serialize(Passthru.Get(), SaveNum * sizeof(UCS2CHAR));
 					if (Ar.IsByteSwapping())
 					{
@@ -1828,21 +1837,21 @@ FArchive& operator<<( FArchive& Ar, UE_STRING_CLASS& A )
 					Passthru.Apply();
 
 					// Inline combine any surrogate pairs in the data when loading into a UTF-32 string
-					StringConv::InlineCombineSurrogates(A);
+					StringConv::InlineCombineSurrogates(Str);
 
 					// Since Microsoft's vsnwprintf implementation raises an invalid parameter warning
 					// with a character of 0xffff, scan for it and terminate the string there.
 					// 0xffff isn't an actual Unicode character anyway.
 					int Index = 0;
-					if (A.FindChar(0xffff, Index))
+					if (Str.FindChar(0xffff, Index))
 					{
-						A[Index] = CHARTEXT(ElementType, '\0');
-						A.TrimToNullTerminator();
+						Str[Index] = CHARTEXT(ElementType, '\0');
+						Str.TrimToNullTerminator();
 					}
 				}
 				else
 				{
-					auto Passthru = StringMemoryPassthru<ANSICHAR>(A.Data.GetData(), SaveNum, SaveNum);
+					auto Passthru = StringMemoryPassthru<ANSICHAR>(Str.Data.GetData(), SaveNum, SaveNum);
 					Ar.Serialize(Passthru.Get(), SaveNum * sizeof(ANSICHAR));
 					// Ensure the string has a null terminator
 					Passthru.Get()[SaveNum - 1] = '\0';
@@ -1852,16 +1861,16 @@ FArchive& operator<<( FArchive& Ar, UE_STRING_CLASS& A )
 				// Throw away empty string.
 				if (SaveNum == 1)
 				{
-					A.Data.Empty();
+					Str.Data.Empty();
 				}
 			}
 		#else
 			if (SaveNum)
 			{
-				A.Data.Empty(SaveNum + 1);
-				A.Data.AddUninitialized(SaveNum + 1);
+				Str.Data.Empty(SaveNum + 1);
+				Str.Data.AddUninitialized(SaveNum + 1);
 
-				auto Passthru = StringMemoryPassthru<ElementType>(A.Data.GetData(), SaveNum, SaveNum);
+				auto Passthru = StringMemoryPassthru<ElementType>(Str.Data.GetData(), SaveNum, SaveNum);
 				Ar.Serialize(Passthru.Get(), SaveNum * sizeof(ElementType));
 				// Ensure the string has a null terminator
 				Passthru.Get()[SaveNum] = UTF8TEXT('\0');
@@ -1871,20 +1880,20 @@ FArchive& operator<<( FArchive& Ar, UE_STRING_CLASS& A )
 			}
 			else
 			{
-				A.Empty();
+				Str.Empty();
 			}
 		#endif
 	}
 	else
 	{
-		A.Data.CountBytes(Ar);
+		Str.Data.CountBytes(Ar);
 
 		#if UE_STRING_CHARTYPE_IS_TCHAR
-			const bool bSaveUnicodeChar = Ar.IsForcingUnicode() || !TCString<ElementType>::IsPureAnsi(*A);
+			const bool bSaveUnicodeChar = Ar.IsForcingUnicode() || !TCString<ElementType>::IsPureAnsi(*Str);
 			if (bSaveUnicodeChar)
 			{
 				// This preprocessor block should not be necessary when the StrCast below understands UTF16CHAR.
-				FTCHARToUTF16 UTF16String(*A, A.Len() + 1); // include the null terminator
+				FTCHARToUTF16 UTF16String(*Str, Str.Len() + 1); // include the null terminator
 				int32 Num = UTF16String.Length() + 1; // include the null terminator
 
 				int32 SaveNum = -Num;
@@ -1909,22 +1918,22 @@ FArchive& operator<<( FArchive& Ar, UE_STRING_CLASS& A )
 			}
 			else
 			{
-				int32 Num = A.Data.Num();
+				int32 Num = Str.Data.Num();
 				Ar << Num;
 
 				if (Num)
 				{
-					Ar.Serialize((void*)StrCast<ANSICHAR>(A.Data.GetData(), Num).Get(), sizeof(ANSICHAR) * Num);
+					Ar.Serialize((void*)StrCast<ANSICHAR>(Str.Data.GetData(), Num).Get(), sizeof(ANSICHAR) * Num);
 				}
 			}
 		#else
 			// Unlike the TCHAR case, we don't bother to save the null terminator
-			int32 SaveNum = A.Len();
+			int32 SaveNum = Str.Len();
 			Ar << SaveNum;
 
 			if (SaveNum)
 			{
-				auto CompactString = StrCast<UTF8CHAR>(A.Data.GetData(), SaveNum);
+				auto CompactString = StrCast<UTF8CHAR>(Str.Data.GetData(), SaveNum);
 				Ar.Serialize((void*)CompactString.Get(), sizeof(UTF8CHAR) * CompactString.Length());
 			}
 		#endif
@@ -1992,7 +2001,7 @@ UE_STRING_CLASS SlugStringForValidName(const UE_STRING_CLASS& DisplayString, con
 	// Convert the display label, which may consist of just about any possible character, into a
 	// suitable name for a UObject (remove whitespace, certain symbols, etc.)
 	{
-		for ( int32 BadCharacterIndex = 0; BadCharacterIndex < UE_ARRAY_COUNT(INVALID_OBJECTNAME_CHARACTERS) - 1; ++BadCharacterIndex )
+		for (int32 BadCharacterIndex = 0; BadCharacterIndex < UE_ARRAY_COUNT(INVALID_OBJECTNAME_CHARACTERS) - 1; ++BadCharacterIndex)
 		{
 			const ElementType TestChar[2] = { (ElementType)INVALID_OBJECTNAME_CHARACTERS[BadCharacterIndex], CHARTEXT(ElementType, '\0') };
 			const int32 NumReplacedChars = GeneratedName.ReplaceInline(TestChar, ReplaceWith);
diff --git a/Engine/Source/Runtime/Core/Public/Containers/UnrealString.h.inl b/Engine/Source/Runtime/Core/Public/Containers/UnrealString.h.inl
--- a/Engine/Source/Runtime/Core/Public/Containers/UnrealString.h.inl
+++ b/Engine/Source/Runtime/Core/Public/Containers/UnrealString.h.inl
@@ -250,7 +250,7 @@ public:
 	 * @param Index into string
 	 * @return Character at Index
 	 */
-	[[nodiscard]] FORCEINLINE ElementType& operator[]( int32 Index ) UE_LIFETIMEBOUND
+	[[nodiscard]] FORCEINLINE ElementType& operator[](int32 Index) UE_LIFETIMEBOUND
 	{
 		checkf(IsValidIndex(Index), TEXT("String index out of bounds: Index %i from a string with a length of %i"), Index, Len());
 		return Data.GetData()[Index];
@@ -262,7 +262,7 @@ public:
 	 * @param Index into string
 	 * @return const Character at Index
 	 */
-	[[nodiscard]] FORCEINLINE const ElementType& operator[]( int32 Index ) const UE_LIFETIMEBOUND
+	[[nodiscard]] FORCEINLINE const ElementType& operator[](int32 Index) const UE_LIFETIMEBOUND
 	{
 		checkf(IsValidIndex(Index), TEXT("String index out of bounds: Index %i from a string with a length of %i"), Index, Len());
 		return Data.GetData()[Index];
@@ -494,7 +494,8 @@ public:
 	/**
 	 * Removes the text from the start of the string if it exists.
 	 *
-	 * @param InPrefix the prefix to search for at the start of the string to remove.
+	 * @param InPrefix    The prefix to search for at the start of the string to remove.
+	 * @param SearchCase  Indicates whether the search is case sensitive or not (defaults to ESearchCase::IgnoreCase)
 	 * @return true if the prefix was removed, otherwise false.
 	 */
 	template <
@@ -516,7 +517,8 @@ public:
 	/**
 	 * Removes the text from the start of the string if it exists.
 	 *
-	 * @param InPrefix the prefix to search for at the start of the string to remove.
+	 * @param InPrefix    The prefix to search for at the start of the string to remove.
+	 * @param SearchCase  Indicates whether the search is case sensitive or not (defaults to ESearchCase::IgnoreCase)
 	 * @return true if the prefix was removed, otherwise false.
 	 */
 	bool RemoveFromStart(const ElementType* InPrefix, ESearchCase::Type SearchCase = ESearchCase::IgnoreCase)
@@ -527,7 +529,8 @@ public:
 	/**
 	 * Removes the text from the start of the string if it exists.
 	 *
-	 * @param InPrefix the prefix to search for at the start of the string to remove.
+	 * @param InPrefix    The prefix to search for at the start of the string to remove.
+	 * @param SearchCase  Indicates whether the search is case sensitive or not (defaults to ESearchCase::IgnoreCase)
 	 * @return true if the prefix was removed, otherwise false.
 	 */
 	bool RemoveFromStart(const UE_STRING_CLASS& InPrefix, ESearchCase::Type SearchCase = ESearchCase::IgnoreCase)
@@ -538,8 +541,8 @@ public:
 	/**
 	 * Removes the text from the start of the string if it exists.
 	 *
-	 * @param InPrefix the prefix to search for at the start of the string to remove.
-	 * @param InPrefixLen length of InPrefix
+	 * @param InPrefix     The prefix to search for at the start of the string to remove.
+	 * @param InPrefixLen  The length of InPrefix
 	 * @return true if the prefix was removed, otherwise false.
 	 */
 	CORE_API bool RemoveFromStart(const ElementType* InPrefix, int32 InPrefixLen, ESearchCase::Type SearchCase = ESearchCase::IgnoreCase);
@@ -547,7 +550,8 @@ public:
 	/**
 	 * Removes the text from the end of the string if it exists.
 	 *
-	 * @param InSuffix the suffix to search for at the end of the string to remove.
+	 * @param InSuffix    The suffix to search for at the end of the string to remove.
+	 * @param SearchCase  Indicates whether the search is case sensitive or not (defaults to ESearchCase::IgnoreCase)
 	 * @return true if the suffix was removed, otherwise false.
 	 */
 	template <
@@ -569,7 +573,8 @@ public:
 	/**
 	 * Removes the text from the end of the string if it exists.
 	 *
-	 * @param InSuffix the suffix to search for at the end of the string to remove.
+	 * @param InSuffix    The suffix to search for at the end of the string to remove.
+	 * @param SearchCase  Indicates whether the search is case sensitive or not (defaults to ESearchCase::IgnoreCase)
 	 * @return true if the suffix was removed, otherwise false.
 	 */
 	bool RemoveFromEnd(const ElementType* InSuffix, ESearchCase::Type SearchCase = ESearchCase::IgnoreCase)
@@ -580,7 +585,8 @@ public:
 	/**
 	 * Removes the text from the end of the string if it exists.
 	 *
-	 * @param InSuffix the suffix to search for at the end of the string to remove.
+	 * @param InSuffix    The suffix to search for at the end of the string to remove.
+	 * @param SearchCase  Indicates whether the search is case sensitive or not (defaults to ESearchCase::IgnoreCase)
 	 * @return true if the suffix was removed, otherwise false.
 	 */
 	bool RemoveFromEnd(const UE_STRING_CLASS& InSuffix, ESearchCase::Type SearchCase = ESearchCase::IgnoreCase)
@@ -591,8 +597,9 @@ public:
 	/**
 	 * Removes the text from the end of the string if it exists.
 	 *
-	 * @param InSuffix the suffix to search for at the end of the string to remove.
-	 * @param InSuffixLen length of InSuffix
+	 * @param InSuffix     The suffix to search for at the end of the string to remove.
+	 * @param InSuffixLen  The length of InSuffix
+	 * @param SearchCase   Indicates whether the search is case sensitive or not (defaults to ESearchCase::IgnoreCase)
 	 * @return true if the suffix was removed, otherwise false.
 	 */
 	CORE_API bool RemoveFromEnd(const ElementType* InSuffix, int32 InSuffixLen, ESearchCase::Type SearchCase = ESearchCase::IgnoreCase);
@@ -736,7 +743,7 @@ public:
 	 * @param Str path array of characters to be concatenated onto the end of this
 	 * @return reference to path
 	 */
-	FORCEINLINE UE_STRING_CLASS& operator/=( const ElementType* Str )
+	FORCEINLINE UE_STRING_CLASS& operator/=(const ElementType* Str)
 	{
 		checkSlow(Str);
 
@@ -747,7 +754,7 @@ public:
 	/**
 	* Concatenate this path with given path ensuring the / character is used between them
 	* 
-	* @param Str path CharRangeType (string class/string view/string builder etc) to be concatenated onto the end of this
+	* @param Str path range of CharType (string class/string view/string builder etc) to be concatenated onto the end of this
 	* @return reference to path
 	*/
 	template <
@@ -1131,9 +1138,9 @@ public:
 	}
 
 	/** Returns the left most given number of characters */
-	[[nodiscard]] FORCEINLINE UE_STRING_CLASS Left( int32 Count ) const &
+	[[nodiscard]] FORCEINLINE UE_STRING_CLASS Left(int32 Count) const &
 	{
-		return UE_STRING_CLASS::ConstructFromPtrSize(**this, FMath::Clamp(Count,0,Len()) );
+		return UE_STRING_CLASS::ConstructFromPtrSize(**this, FMath::Clamp(Count, 0, Len()));
 	}
 
 	[[nodiscard]] FORCEINLINE UE_STRING_CLASS Left(int32 Count) &&
@@ -1156,13 +1163,13 @@ public:
 	}
 
 	/** Returns the left most characters from the string chopping the given number of characters from the end */
-	[[nodiscard]] FORCEINLINE UE_STRING_CLASS LeftChop( int32 Count ) const &
+	[[nodiscard]] FORCEINLINE UE_STRING_CLASS LeftChop(int32 Count) const &
 	{
 		const int32 Length = Len();
-		return UE_STRING_CLASS::ConstructFromPtrSize( **this, FMath::Clamp(Length-Count,0, Length) );
+		return UE_STRING_CLASS::ConstructFromPtrSize(**this, FMath::Clamp(Length - Count, 0, Length));
 	}
 
-	[[nodiscard]] FORCEINLINE UE_STRING_CLASS LeftChop(int32 Count)&&
+	[[nodiscard]] FORCEINLINE UE_STRING_CLASS LeftChop(int32 Count) &&
 	{
 		LeftChopInline(Count, EAllowShrinking::No);
 		return MoveTemp(*this);
@@ -1181,10 +1188,10 @@ public:
 	}
 
 	/** Returns the string to the right of the specified location, counting back from the right (end of the word). */
-	[[nodiscard]] FORCEINLINE UE_STRING_CLASS Right( int32 Count ) const &
+	[[nodiscard]] FORCEINLINE UE_STRING_CLASS Right(int32 Count) const &
 	{
 		const int32 Length = Len();
-		return UE_STRING_CLASS( **this + Length-FMath::Clamp(Count,0,Length) );
+		return UE_STRING_CLASS(**this + Length - FMath::Clamp(Count, 0, Length));
 	}
 
 	[[nodiscard]] FORCEINLINE UE_STRING_CLASS Right(int32 Count) &&
@@ -1206,7 +1213,7 @@ public:
 	}
 
 	/** Returns the string to the right of the specified location, counting forward from the left (from the beginning of the word). */
-	[[nodiscard]] CORE_API UE_STRING_CLASS RightChop( int32 Count ) const &;
+	[[nodiscard]] CORE_API UE_STRING_CLASS RightChop(int32 Count) const &;
 
 	[[nodiscard]] FORCEINLINE UE_STRING_CLASS RightChop(int32 Count) &&
 	{
@@ -1249,16 +1256,17 @@ public:
 	}
 
 	/**
-	 * Searches the string for a substring, and returns index into this string of the first found instance. Can search
-	 * from beginning or end, and ignore case or not. If substring is empty, returns clamped StartPosition.
+	 * Searches the string for a substring. Can search from beginning or end, and ignore case or not.
 	 *
 	 * @param SubStr			The string to search for
+	 * @param SearchCase		Indicates whether the search is case sensitive or not (defaults to ESearchCase::IgnoreCase)
+	 * @param SearchDir			Indicates whether the search starts at the beginning or at the end (defaults to ESearchDir::FromStart)
 	 * @param StartPosition		The start character position to search from
-	 * @param SearchCase		Indicates whether the search is case sensitive or not ( defaults to ESearchCase::IgnoreCase )
-	 * @param SearchDir			Indicates whether the search starts at the beginning or at the end ( defaults to ESearchDir::FromStart )
+	 *
+	 * @return  Index into this string of the first found instance. If SubStr is empty, returns clamped StartPosition.
 	 *
 	 * @note  When searching backwards where a StartPosition is provided, searching will actually begin from
-	 *        StartPosition - SubStr.Len(), therefore:
+	 *        StartPosition - GetNum(SubStr), therefore:
 	 *
 	 *        FString("X").Find("X", ESearchCase::CaseSensitive, ESearchDir::FromEnd, 0) == INDEX_NONE
 	 *
@@ -1282,13 +1290,21 @@ public:
 	}
 
 	/**
-	 * Searches the string for a substring, and returns index into this string of the first found instance. Can search
-	 * from beginning or end, and ignore case or not. If substring is empty, returns clamped StartPosition.
+	 * Searches the string for a substring. Can search from beginning or end, and ignore case or not.
 	 *
 	 * @param SubStr			The string array of characters to search for
-	 * @param StartPosition		The start character position to search from.  See note below.
 	 * @param SearchCase		Indicates whether the search is case sensitive or not
 	 * @param SearchDir			Indicates whether the search starts at the beginning or at the end.
+	 * @param StartPosition		The start character position to search from.  See note below.
+	 *
+	 * @return  Index into this string of the first found instance. If SubStr is empty, returns clamped StartPosition.
+	 *
+	 * @note  When searching backwards where a StartPosition is provided, searching will actually begin from
+	 *        StartPosition - Strlen(SubStr), therefore:
+	 *
+	 *        FString("X").Find("X", ESearchCase::CaseSensitive, ESearchDir::FromEnd, 0) == INDEX_NONE
+	 *
+	 *        Consider using UE::String::FindLast() as an alternative.
 	 */
 	[[nodiscard]] int32 Find(const ElementType* SubStr, ESearchCase::Type SearchCase = ESearchCase::IgnoreCase,
 		ESearchDir::Type SearchDir = ESearchDir::FromStart, int32 StartPosition = INDEX_NONE) const
@@ -1297,14 +1313,15 @@ public:
 	}
 
 	/**
-	 * Searches the string for a substring, and returns index into this string of the first found instance. Can search
-	 * from beginning or end, and ignore case or not. If substring is empty, returns clamped StartPosition.
+	 * Searches the string for a substring. Can search from beginning or end, and ignore case or not.
 	 *
 	 * @param SubStr			The string array of characters to search for
 	 * @param StartPosition		The start character position to search from.  See note below.
 	 * @param SearchCase		Indicates whether the search is case sensitive or not
 	 * @param SearchDir			Indicates whether the search starts at the beginning or at the end.
 	 *
+	 * @return  Index into this string of the first found instance. If SubStr is empty, returns clamped StartPosition.
+	 *
 	 * @note  When searching backwards where a StartPosition is provided, searching will actually begin from
 	 *        StartPosition - SubStr.Len(), therefore:
 	 *
@@ -1319,17 +1336,18 @@ public:
 	}
 
 	/**
-	 * Searches the string for a substring, and returns index into this string of the first found instance. Can search
-	 * from beginning or end, and ignore case or not. If substring is empty, returns clamped StartPosition.
+	 * Searches the string for a substring. Can search from beginning or end, and ignore case or not.
 	 *
 	 * @param SubStr			The string array of characters to search for
 	 * @param SubStrLen			The length of the SubStr array
-	 * @param StartPosition		The start character position to search from.  See note below.
 	 * @param SearchCase		Indicates whether the search is case sensitive or not
 	 * @param SearchDir			Indicates whether the search starts at the beginning or at the end.
+	 * @param StartPosition		The start character position to search from.  See note below.
+	 *
+	 * @return  Index into this string of the first found instance. If SubStr is empty, returns clamped StartPosition.
 	 *
 	 * @note  When searching backwards where a StartPosition is provided, searching will actually begin from
-	 *        StartPosition - SubStr.Len(), therefore:
+	 *        StartPosition - InSubStrLen, therefore:
 	 *
 	 *        FString("X").Find("X", ESearchCase::CaseSensitive, ESearchDir::FromEnd, 0) == INDEX_NONE
 	 *
@@ -1343,8 +1361,8 @@ public:
 	 * Returns whether this string contains the specified substring.
 	 *
 	 * @param SubStr			Text to search for
-	 * @param SearchCase		Indicates whether the search is case sensitive or not ( defaults to ESearchCase::IgnoreCase )
-	 * @param SearchDir			Indicates whether the search starts at the beginning or at the end ( defaults to ESearchDir::FromStart )
+	 * @param SearchCase		Indicates whether the search is case sensitive or not (defaults to ESearchCase::IgnoreCase)
+	 * @param SearchDir			Indicates whether the search starts at the beginning or at the end (defaults to ESearchDir::FromStart)
 	 * @return					Returns whether the string contains the substring. If the substring is empty, returns true.
 	 **/
 	template <
@@ -1368,8 +1386,8 @@ public:
 	 * Returns whether this string contains the specified substring.
 	 *
 	 * @param SubStr			Text to search for
-	 * @param SearchCase		Indicates whether the search is case sensitive or not ( defaults to ESearchCase::IgnoreCase )
-	 * @param SearchDir			Indicates whether the search starts at the beginning or at the end ( defaults to ESearchDir::FromStart )
+	 * @param SearchCase		Indicates whether the search is case sensitive or not (defaults to ESearchCase::IgnoreCase)
+	 * @param SearchDir			Indicates whether the search starts at the beginning or at the end (defaults to ESearchDir::FromStart)
 	 * @return					Returns whether the string contains the substring. If the substring is empty, returns true.
 	 **/
 	[[nodiscard]] FORCEINLINE bool Contains(const ElementType* SubStr, ESearchCase::Type SearchCase = ESearchCase::IgnoreCase,
@@ -1382,12 +1400,12 @@ public:
 	 * Returns whether this string contains the specified substring.
 	 *
 	 * @param SubStr			Text to search for
-	 * @param SearchCase		Indicates whether the search is case sensitive or not ( defaults to ESearchCase::IgnoreCase )
-	 * @param SearchDir			Indicates whether the search starts at the beginning or at the end ( defaults to ESearchDir::FromStart )
+	 * @param SearchCase		Indicates whether the search is case sensitive or not (defaults to ESearchCase::IgnoreCase)
+	 * @param SearchDir			Indicates whether the search starts at the beginning or at the end (defaults to ESearchDir::FromStart)
 	 * @return					Returns whether the string contains the substring. If the substring is empty, returns true.
 	 **/
 	[[nodiscard]] FORCEINLINE bool Contains(const UE_STRING_CLASS& SubStr, ESearchCase::Type SearchCase = ESearchCase::IgnoreCase,
-							  ESearchDir::Type SearchDir = ESearchDir::FromStart ) const
+							  ESearchDir::Type SearchDir = ESearchDir::FromStart) const
 	{
 		return Find(*SubStr, SearchCase, SearchDir) != INDEX_NONE;
 	}
@@ -1397,8 +1415,8 @@ public:
 	 *
 	 * @param SubStr			Text to search for
 	 * @param SubStrLen			Length of the Text
-	 * @param SearchCase		Indicates whether the search is case sensitive or not ( defaults to ESearchCase::IgnoreCase )
-	 * @param SearchDir			Indicates whether the search starts at the beginning or at the end ( defaults to ESearchDir::FromStart )
+	 * @param SearchCase		Indicates whether the search is case sensitive or not (defaults to ESearchCase::IgnoreCase)
+	 * @param SearchDir			Indicates whether the search starts at the beginning or at the end (defaults to ESearchDir::FromStart)
 	 * @return					Returns whether the string contains the substring. If the substring is empty, returns true.
 	 **/
 	[[nodiscard]] FORCEINLINE bool Contains(const ElementType* SubStr, int32 SubStrLen,
@@ -1411,24 +1429,24 @@ public:
 	 * Searches the string for a character
 	 *
 	 * @param InChar the character to search for
-	 * @param Index out the position the character was found at, INDEX_NONE if return is false
+	 * @param OutIndex out the position the character was found at, INDEX_NONE if return is false
 	 * @return true if character was found in this string, otherwise false
 	 */
-	FORCEINLINE bool FindChar(ElementType InChar, int32& Index ) const
+	FORCEINLINE bool FindChar(ElementType InChar, int32& OutIndex) const
 	{
-		return Data.Find(InChar, Index);
+		return Data.Find(InChar, OutIndex);
 	}
 
 	/**
 	 * Searches the string for the last occurrence of a character
 	 *
 	 * @param InChar the character to search for
-	 * @param Index out the position the character was found at, INDEX_NONE if return is false
+	 * @param OutIndex out the position the character was found at, INDEX_NONE if return is false
 	 * @return true if character was found in this string, otherwise false
 	 */
-	FORCEINLINE bool FindLastChar( ElementType InChar, int32& Index ) const
+	FORCEINLINE bool FindLastChar(ElementType InChar, int32& OutIndex) const
 	{
-		return Data.FindLast(InChar, Index);
+		return Data.FindLast(InChar, OutIndex);
 	}
 
 	/**
@@ -1450,7 +1468,6 @@ public:
 	 * Searches the string for the last occurrence of a character which matches the specified predicate.
 	 *
 	 * @param Pred Predicate that takes a character and returns true if it matches search criteria, false otherwise.
-	 * @param StartIndex Index of element from which to start searching. Defaults to the last character in string.
 	 *
 	 * @return Index of the found character, INDEX_NONE otherwise.
 	 */
@@ -1499,15 +1516,15 @@ public:
 	 * @param SearchCase 	Whether or not the comparison should ignore case
 	 * @return 0 if equal, negative if less than, positive if greater than
 	 */
-	[[nodiscard]] FORCEINLINE int32 Compare( const UE_STRING_CLASS& Other, ESearchCase::Type SearchCase = ESearchCase::CaseSensitive ) const
+	[[nodiscard]] FORCEINLINE int32 Compare(const UE_STRING_CLASS& Other, ESearchCase::Type SearchCase = ESearchCase::CaseSensitive) const
 	{
-		if( SearchCase == ESearchCase::CaseSensitive )
+		if (SearchCase == ESearchCase::CaseSensitive)
 		{
-			return TCString<ElementType>::Strcmp( **this, *Other ); 
+			return TCString<ElementType>::Strcmp(**this, *Other); 
 		}
 		else
 		{
-			return TCString<ElementType>::Stricmp( **this, *Other );
+			return TCString<ElementType>::Stricmp(**this, *Other);
 		}
 	}
 
@@ -1517,15 +1534,15 @@ public:
 	 * @param InStr The string to search and split at
 	 * @param LeftS out the string to the left of InStr, not updated if return is false. LeftS must not point to the same location as RightS, but can point to this.
 	 * @param RightS out the string to the right of InStr, not updated if return is false. RightS must not point to the same location as LeftS, but can point to this.
-	 * @param SearchCase		Indicates whether the search is case sensitive or not ( defaults to ESearchCase::IgnoreCase )
-	 * @param SearchDir			Indicates whether the search starts at the beginning or at the end ( defaults to ESearchDir::FromStart )
+	 * @param SearchCase		Indicates whether the search is case sensitive or not (defaults to ESearchCase::IgnoreCase)
+	 * @param SearchDir			Indicates whether the search starts at the beginning or at the end (defaults to ESearchDir::FromStart)
 	 * @return true if string is split, otherwise false
 	 */
-	CORE_API bool Split(const UE_STRING_CLASS& InS, UE_STRING_CLASS* LeftS, UE_STRING_CLASS* RightS, ESearchCase::Type SearchCase,
+	CORE_API bool Split(const UE_STRING_CLASS& InStr, UE_STRING_CLASS* LeftS, UE_STRING_CLASS* RightS, ESearchCase::Type SearchCase,
 		ESearchDir::Type SearchDir = ESearchDir::FromStart) const;
 
 	/** Split with ESearchCase::IgnoreCase and ESearchDir::FromStart. Allows compiler to avoid branches w/o inlining code. */
-	CORE_API bool Split(const UE_STRING_CLASS& InS, UE_STRING_CLASS* LeftS, UE_STRING_CLASS* RightS) const;
+	CORE_API bool Split(const UE_STRING_CLASS& InStr, UE_STRING_CLASS* LeftS, UE_STRING_CLASS* RightS) const;
 
 	/** Returns a new string with the characters of this converted to uppercase */
 	[[nodiscard]] CORE_API UE_STRING_CLASS ToUpper() const &;
@@ -1552,10 +1569,10 @@ public:
 	CORE_API void ToLowerInline();
 
 	/** Pad the left of this string for ChCount characters */
-	[[nodiscard]] CORE_API UE_STRING_CLASS LeftPad( int32 ChCount ) const;
+	[[nodiscard]] CORE_API UE_STRING_CLASS LeftPad(int32 ChCount) const;
 
 	/** Pad the right of this string for ChCount characters */
-	[[nodiscard]] CORE_API UE_STRING_CLASS RightPad( int32 ChCount ) const;
+	[[nodiscard]] CORE_API UE_STRING_CLASS RightPad(int32 ChCount) const;
 	
 	/** Returns true if the string only contains numeric characters */
 	[[nodiscard]] CORE_API bool IsNumeric() const;
@@ -1566,8 +1583,8 @@ public:
 	/**
 	 * Constructs a string similarly to how classic sprintf works.
 	 *
-	 * @param Format	Format string that specifies how the string should be built optionally using additional args. Refer to standard printf format.
-	 * @param ...		Depending on format function may require additional arguments to build output object.
+	 * @param Fmt  Format string that specifies how the string should be built optionally using additional args. Refer to standard printf format.
+	 * @param ...  Depending on format function may require additional arguments to build output object.
 	 *
 	 * @returns A string that was constructed using format and additional parameters.
 	 */
@@ -1648,7 +1665,7 @@ private:
 public:
 
 	/** Returns a string containing only the Ch character */
-	[[nodiscard]] static CORE_API UE_STRING_CLASS Chr( ElementType Ch );
+	[[nodiscard]] static CORE_API UE_STRING_CLASS Chr(ElementType Ch);
 
 	/**
 	 * Returns a string that is full of a variable number of characters
@@ -1658,23 +1675,24 @@ public:
 	 * 
 	 * @return The string of NumCharacters characters.
 	 */
-	[[nodiscard]] static CORE_API UE_STRING_CLASS ChrN( int32 NumCharacters, ElementType Char );
+	[[nodiscard]] static CORE_API UE_STRING_CLASS ChrN(int32 NumCharacters, ElementType Char);
 
 	/**
 	 * Serializes the string.
 	 *
 	 * @param Ar Reference to the serialization archive.
-	 * @param S Reference to the string being serialized.
+	 * @param Str Reference to the string being serialized.
 	 *
 	 * @return Reference to the Archive after serialization.
 	 */
-	friend CORE_API FArchive& operator<<( FArchive& Ar, UE_STRING_CLASS& S );
+	friend CORE_API FArchive& operator<<(FArchive& Ar, UE_STRING_CLASS& Str);
 
 
 	/**
 	 * Test whether this string starts with given prefix.
 	 *
-	 * @param SearchCase		Indicates whether the search is case sensitive or not ( defaults to ESearchCase::IgnoreCase )
+	 * @param InPrefix    The prefix to search for at the start of the string.
+	 * @param SearchCase  Indicates whether the search is case sensitive or not (defaults to ESearchCase::IgnoreCase)
 	 * @return true if this string begins with specified text, false otherwise
 	 */
 	template <
@@ -1696,7 +1714,8 @@ public:
 	/**
 	 * Test whether this string starts with given prefix.
 	 *
-	 * @param SearchCase		Indicates whether the search is case sensitive or not ( defaults to ESearchCase::IgnoreCase )
+	 * @param InPrefix    The prefix to search for at the start of the string.
+	 * @param SearchCase  Indicates whether the search is case sensitive or not (defaults to ESearchCase::IgnoreCase)
 	 * @return true if this string begins with specified text, false otherwise
 	 */
 	[[nodiscard]] bool StartsWith(const ElementType* InPrefix, ESearchCase::Type SearchCase = ESearchCase::IgnoreCase) const
@@ -1707,7 +1726,8 @@ public:
 	/**
 	 * Test whether this string starts with given prefix.
 	 *
-	 * @param SearchCase		Indicates whether the search is case sensitive or not ( defaults to ESearchCase::IgnoreCase )
+	 * @param InPrefix    The prefix to search for at the start of the string.
+	 * @param SearchCase  Indicates whether the search is case sensitive or not (defaults to ESearchCase::IgnoreCase)
 	 * @return true if this string begins with specified text, false otherwise
 	 */
 	[[nodiscard]] bool StartsWith(const UE_STRING_CLASS& InPrefix, ESearchCase::Type SearchCase = ESearchCase::IgnoreCase) const
@@ -1718,7 +1738,8 @@ public:
 	/**
 	 * Test whether this string starts with given prefix.
 	 *
-	 * @param SearchCase		Indicates whether the search is case sensitive or not ( defaults to ESearchCase::IgnoreCase )
+	 * @param InPrefix    The prefix to search for at the start of the string.
+	 * @param SearchCase  Indicates whether the search is case sensitive or not (defaults to ESearchCase::IgnoreCase)
 	 * @return true if this string begins with specified text, false otherwise
 	 */
 	[[nodiscard]] CORE_API bool StartsWith(const ElementType* InPrefix, int32 InPrefixLen, ESearchCase::Type SearchCase = ESearchCase::IgnoreCase) const;
@@ -1726,7 +1747,8 @@ public:
 	/**
 	 * Test whether this string ends with given suffix.
 	 *
-	 * @param SearchCase		Indicates whether the search is case sensitive or not ( defaults to ESearchCase::IgnoreCase )
+	 * @param InSuffix    The suffix to search for at the end of the string.
+	 * @param SearchCase  Indicates whether the search is case sensitive or not (defaults to ESearchCase::IgnoreCase)
 	 * @return true if this string ends with specified text, false otherwise
 	 */
 	template <
@@ -1748,7 +1770,8 @@ public:
 	/**
 	 * Test whether this string ends with given suffix.
 	 *
-	 * @param SearchCase		Indicates whether the search is case sensitive or not ( defaults to ESearchCase::IgnoreCase )
+	 * @param InSuffix    The suffix to search for at the end of the string.
+	 * @param SearchCase  Indicates whether the search is case sensitive or not (defaults to ESearchCase::IgnoreCase)
 	 * @return true if this string ends with specified text, false otherwise
 	 */
 	[[nodiscard]] bool EndsWith(const ElementType* InSuffix, ESearchCase::Type SearchCase = ESearchCase::IgnoreCase) const
@@ -1759,7 +1782,8 @@ public:
 	/**
 	 * Test whether this string ends with given suffix.
 	 *
-	 * @param SearchCase		Indicates whether the search is case sensitive or not ( defaults to ESearchCase::IgnoreCase )
+	 * @param InSuffix    The suffix to search for at the end of the string.
+	 * @param SearchCase  Indicates whether the search is case sensitive or not (defaults to ESearchCase::IgnoreCase)
 	 * @return true if this string ends with specified text, false otherwise
 	 */
 	[[nodiscard]] bool EndsWith(const UE_STRING_CLASS& InSuffix, ESearchCase::Type SearchCase = ESearchCase::IgnoreCase) const
@@ -1770,16 +1794,17 @@ public:
 	/**
 	 * Test whether this string ends with given suffix.
 	 *
-	 * @param SearchCase		Indicates whether the search is case sensitive or not ( defaults to ESearchCase::IgnoreCase )
+	 * @param InSuffix    The suffix to search for at the end of the string.
+	 * @param SearchCase  Indicates whether the search is case sensitive or not (defaults to ESearchCase::IgnoreCase)
 	 * @return true if this string ends with specified text, false otherwise
 	 */
-	[[nodiscard]] CORE_API bool EndsWith(const ElementType* InSuffix, int32 InSuffixLen, ESearchCase::Type SearchCase = ESearchCase::IgnoreCase ) const;
+	[[nodiscard]] CORE_API bool EndsWith(const ElementType* InSuffix, int32 InSuffixLen, ESearchCase::Type SearchCase = ESearchCase::IgnoreCase) const;
 
 	/**
 	 * Searches this string for a given wild card
 	 *
 	 * @param Wildcard		*?-type wildcard
-	 * @param SearchCase	Indicates whether the search is case sensitive or not ( defaults to ESearchCase::IgnoreCase )
+	 * @param SearchCase	Indicates whether the search is case sensitive or not (defaults to ESearchCase::IgnoreCase)
 	 * @return true if this string matches the *?-type wildcard given. 
 	 * @warning This is a simple, SLOW routine. Use with caution
 	 */
@@ -1803,7 +1828,7 @@ public:
 	 * Searches this string for a given wild card
 	 *
 	 * @param Wildcard		*?-type wildcard
-	 * @param SearchCase	Indicates whether the search is case sensitive or not ( defaults to ESearchCase::IgnoreCase )
+	 * @param SearchCase	Indicates whether the search is case sensitive or not (defaults to ESearchCase::IgnoreCase)
 	 * @return true if this string matches the *?-type wildcard given.
 	 * @warning This is a simple, SLOW routine. Use with caution
 	 */
@@ -1816,7 +1841,7 @@ public:
 	 * Searches this string for a given wild card
 	 *
 	 * @param Wildcard		*?-type wildcard
-	 * @param SearchCase	Indicates whether the search is case sensitive or not ( defaults to ESearchCase::IgnoreCase )
+	 * @param SearchCase	Indicates whether the search is case sensitive or not (defaults to ESearchCase::IgnoreCase)
 	 * @return true if this string matches the *?-type wildcard given.
 	 * @warning This is a simple, SLOW routine. Use with caution
 	 */
@@ -1828,8 +1853,9 @@ public:
 	/**
 	 * Searches this string for a given wild card
 	 *
-	 * @param Wildcard		*?-type wildcard
-	 * @param SearchCase	Indicates whether the search is case sensitive or not ( defaults to ESearchCase::IgnoreCase )
+	 * @param Wildcard     *?-type wildcard
+	 * @param WildcardLen  The length of Wildcard
+	 * @param SearchCase   Indicates whether the search is case sensitive or not (defaults to ESearchCase::IgnoreCase)
 	 * @return true if this string matches the *?-type wildcard given.
 	 * @warning This is a simple, SLOW routine. Use with caution
 	 */
@@ -1902,11 +1928,11 @@ public:
 	* @see TrimChar for a variant that returns a modified copy of the string
 	*/
 	CORE_API void TrimCharInline(ElementType CharacterToTrim, bool* bCharRemoved);
-	
+
 	/**
 	 * Returns a copy of this string with wrapping quotation marks removed.
 	 */
-	[[nodiscard]] CORE_API UE_STRING_CLASS TrimQuotes(bool* bQuotesRemoved = nullptr ) const &;
+	[[nodiscard]] CORE_API UE_STRING_CLASS TrimQuotes(bool* bQuotesRemoved = nullptr) const &;
 
 	/**
 	 * Returns this string with wrapping quotation marks removed.
@@ -1917,66 +1943,69 @@ public:
 	* Returns a copy of this string with wrapping CharacterToTrim removed (removes at max one instance in the beginning and end of the string).
 	* @see TrimCharInline for an inline variant
 	*/
-	[[nodiscard]] CORE_API UE_STRING_CLASS TrimChar(ElementType CharacterToTrim, bool* bCharRemoved = nullptr ) const &;
+	[[nodiscard]] CORE_API UE_STRING_CLASS TrimChar(ElementType CharacterToTrim, bool* bCharRemoved = nullptr) const &;
 
 	/**
 	* Returns a copy of this string with wrapping CharacterToTrim removed (removes at max one instance in the beginning and end of the string).
 	*/
 	[[nodiscard]] CORE_API UE_STRING_CLASS TrimChar(ElementType CharacterToTrim, bool* bCharRemoved = nullptr) &&;
-	
+
 	/**
 	 * Breaks up a delimited string into elements of a string array.
 	 *
-	 * @param	InArray		The array to fill with the string pieces
-	 * @param	pchDelim	The string to delimit on
-	 * @param	InCullEmpty	If 1, empty strings are not added to the array
+	 * @param OutArray      The array to fill with the string pieces
+	 * @param pchDelim      The string to delimit on
+	 * @param bInCullEmpty  If true, empty strings are not added to the array
 	 *
-	 * @return	The number of elements in InArray
+	 * @return	The number of elements in OutArray
 	 */
-	CORE_API int32 ParseIntoArray( TArray<UE_STRING_CLASS>& OutArray, const ElementType* pchDelim, bool InCullEmpty = true ) const;
+	CORE_API int32 ParseIntoArray(TArray<UE_STRING_CLASS>& OutArray, const ElementType* pchDelim, bool bInCullEmpty = true) const;
 
 	/**
 	 * Breaks up a delimited string into elements of a string array, using any whitespace and an 
 	 * optional extra delimter, like a ","
 	 * @warning Caution!! this routine is O(N^2) allocations...use it for parsing very short text or not at all! 
 	 *
-	 * @param	InArray			The array to fill with the string pieces
-	 * @param	pchExtraDelim	The string to delimit on
+	 * @param OutArray       The array to fill with the string pieces
+	 * @param pchExtraDelim  The string to delimit on
+	 * @param bInCullEmpty   If true, empty strings are not added to the array
 	 *
-	 * @return	The number of elements in InArray
+	 * @return	The number of elements in OutArray
 	 */
-	CORE_API int32 ParseIntoArrayWS( TArray<UE_STRING_CLASS>& OutArray, const ElementType* pchExtraDelim = nullptr, bool InCullEmpty = true ) const;
+	CORE_API int32 ParseIntoArrayWS(TArray<UE_STRING_CLASS>& OutArray, const ElementType* pchExtraDelim = nullptr, bool bInCullEmpty = true) const;
 
 	/**
 	* Breaks up a delimited string into elements of a string array, using line ending characters
 	* @warning Caution!! this routine is O(N^2) allocations...use it for parsing very short text or not at all!
 	*
-	* @param	InArray			The array to fill with the string pieces	
+	* @param OutArray      The array to fill with the string pieces
+	* @param bInCullEmpty  If true, empty strings are not added to the array
 	*
-	* @return	The number of elements in InArray
+	* @return	The number of elements in OutArray
 	*/
-	CORE_API int32 ParseIntoArrayLines(TArray<UE_STRING_CLASS>& OutArray, bool InCullEmpty = true) const;
+	CORE_API int32 ParseIntoArrayLines(TArray<UE_STRING_CLASS>& OutArray, bool bInCullEmpty = true) const;
 
 	/**
 	* Breaks up a delimited string into elements of a string array, using the given delimiters
 	* @warning Caution!! this routine is O(N^2) allocations...use it for parsing very short text or not at all!
 	*
-	* @param	InArray			The array to fill with the string pieces
-	* @param	DelimArray		The strings to delimit on
-	* @param	NumDelims		The number of delimiters.
+	* @param OutArray      The array to fill with the string pieces
+	* @param DelimArray    The strings to delimit on
+	* @param NumDelims     The number of delimiters.
+	* @param bInCullEmpty  If true, empty strings are not added to the array
 	*
-	* @return	The number of elements in InArray
+	* @return	The number of elements in OutArray
 	*/
-	CORE_API int32 ParseIntoArray(TArray<UE_STRING_CLASS>& OutArray, const ElementType*const* DelimArray, int32 NumDelims, bool InCullEmpty = true) const;
+	CORE_API int32 ParseIntoArray(TArray<UE_STRING_CLASS>& OutArray, const ElementType*const* DelimArray, int32 NumDelims, bool bInCullEmpty = true) const;
 
 	/**
 	 * Takes an array of strings and removes any zero length entries.
 	 *
-	 * @param	InArray	The array to cull
+	 * @param	InOutArray  The array to cull
 	 *
-	 * @return	The number of elements left in InArray
+	 * @return	The number of elements left in InOutArray
 	 */
-	static CORE_API int32 CullArray( TArray<UE_STRING_CLASS>* InArray );
+	static CORE_API int32 CullArray(TArray<UE_STRING_CLASS>* InOutArray);
 
 	/**
 	 * Returns a copy of this string, with the characters in reverse order
@@ -1994,60 +2023,59 @@ public:
 	CORE_API void ReverseString();
 
 	/**
-	 * Replace all occurrences of a substring in this string
+	 * Returns a copy of the string with all occurrences of a substring replaced.
 	 *
-	 * @param From substring to replace
-	 * @param To substring to replace From with
-	 * @param SearchCase	Indicates whether the search is case sensitive or not ( defaults to ESearchCase::IgnoreCase )
+	 * @param From        The substring to replace.
+	 * @param To          The substring to insert in its place.
+	 * @param SearchCase  Indicates whether the search is case sensitive or not (defaults to ESearchCase::IgnoreCase)
 	 * @return a copy of this string with the replacement made
 	 */
 	[[nodiscard]] CORE_API UE_STRING_CLASS Replace(const ElementType* From, const ElementType* To, ESearchCase::Type SearchCase = ESearchCase::IgnoreCase) const &;
 
 	/**
-	 * Replace all occurrences of a substring in this string
+	 * Returns a copy of the string with all occurrences of a substring replaced.
 	 *
-	 * @param From substring to replace
-	 * @param To substring to replace From with
-	 * @param SearchCase	Indicates whether the search is case sensitive or not ( defaults to ESearchCase::IgnoreCase )
+	 * @param From        The substring to replace.
+	 * @param To          The substring to insert in its place.
+	 * @param SearchCase  Indicates whether the search is case sensitive or not (defaults to ESearchCase::IgnoreCase)
 	 * @return a copy of this string with the replacement made
 	 */
 	[[nodiscard]] CORE_API UE_STRING_CLASS Replace(const ElementType* From, const ElementType* To, ESearchCase::Type SearchCase = ESearchCase::IgnoreCase) &&;
 
 	/**
-	 * Replace all occurrences of SearchText with ReplacementText in this string.
+	 * Replaces all occurrences of a substring in this string.
 	 *
-	 * @param	SearchText	the text that should be removed from this string
-	 * @param	ReplacementText		the text to insert in its place
-	 * @param SearchCase	Indicates whether the search is case sensitive or not ( defaults to ESearchCase::IgnoreCase )
+	 * @param From        The substring to replace.
+	 * @param To          The substring to insert in its place.
+	 * @param SearchCase  Indicates whether the search is case sensitive or not (defaults to ESearchCase::IgnoreCase)
 	 *
 	 * @return	the number of occurrences of SearchText that were replaced.
 	 */
-	CORE_API int32 ReplaceInline( const ElementType* SearchText, const ElementType* ReplacementText, ESearchCase::Type SearchCase = ESearchCase::IgnoreCase );
+	CORE_API int32 ReplaceInline(const ElementType* SearchText, const ElementType* ReplacementText, ESearchCase::Type SearchCase = ESearchCase::IgnoreCase);
 
 	/**
-	 * Replace all occurrences of a character with another.
+	 * Replaces all occurrences of a character in this string.
 	 *
 	 * @param SearchChar      Character to remove from this string
 	 * @param ReplacementChar Replacement character
-	 * @param SearchCase      Indicates whether the search is case sensitive or not ( defaults to ESearchCase::IgnoreCase )
+	 * @param SearchCase      Indicates whether the search is case sensitive or not (defaults to ESearchCase::IgnoreCase)
 	 * @note no dynamic allocation
 	 */
-	void ReplaceCharInline(ElementType SearchChar, ElementType ReplacementChar, ESearchCase::Type SearchCase = ESearchCase::IgnoreCase)
+	void ReplaceCharInline(ElementType From, ElementType To, ESearchCase::Type SearchCase = ESearchCase::IgnoreCase)
 	{
-		if (SearchCase == ESearchCase::IgnoreCase && TChar<ElementType>::IsAlpha(SearchChar))
+		if (SearchCase == ESearchCase::IgnoreCase && TChar<ElementType>::IsAlpha(From))
 		{
-			ReplaceCharInlineIgnoreCase(SearchChar, ReplacementChar);
+			ReplaceCharInlineIgnoreCase(From, To);
 		}
 		else
 		{
-			ReplaceCharInlineCaseSensitive(SearchChar, ReplacementChar);
+			ReplaceCharInlineCaseSensitive(From, To);
 		}
 	}
 
 private:
-	
-	CORE_API void ReplaceCharInlineCaseSensitive(ElementType SearchChar, ElementType ReplacementChar);
-	CORE_API void ReplaceCharInlineIgnoreCase(ElementType SearchChar, ElementType ReplacementChar);
+	CORE_API void ReplaceCharInlineCaseSensitive(ElementType From, ElementType To);
+	CORE_API void ReplaceCharInlineIgnoreCase(ElementType From, ElementType To);
 
 public:
 
@@ -2071,7 +2099,7 @@ public:
 	 *
 	 * @param	Chars	by default, replaces all supported characters; this parameter allows you to limit the replacement to a subset.
 	 */
-	CORE_API void ReplaceCharWithEscapedCharInline( const TArray<ElementType>* Chars = nullptr );
+	CORE_API void ReplaceCharWithEscapedCharInline(const TArray<ElementType>* Chars = nullptr);
 
 	/**
 	 * Replaces certain characters with the "escaped" version of that character (i.e. replaces "\n" with "\\n").
@@ -2081,7 +2109,7 @@ public:
 	 *
 	 * @return	a string with all control characters replaced by the escaped version.
 	 */
-	[[nodiscard]] UE_STRING_CLASS ReplaceCharWithEscapedChar( const TArray<ElementType>* Chars = nullptr ) const &
+	[[nodiscard]] UE_STRING_CLASS ReplaceCharWithEscapedChar(const TArray<ElementType>* Chars = nullptr) const &
 	{
 		UE_STRING_CLASS Result(*this);
 		Result.ReplaceCharWithEscapedCharInline(Chars);
@@ -2096,7 +2124,7 @@ public:
 	 *
 	 * @return	a string with all control characters replaced by the escaped version.
 	 */
-	[[nodiscard]] UE_STRING_CLASS ReplaceCharWithEscapedChar( const TArray<ElementType>* Chars = nullptr ) &&
+	[[nodiscard]] UE_STRING_CLASS ReplaceCharWithEscapedChar(const TArray<ElementType>* Chars = nullptr) &&
 	{
 		ReplaceCharWithEscapedCharInline(Chars);
 		return MoveTemp(*this);
@@ -2106,14 +2134,14 @@ public:
 	 * Removes the escape backslash for all supported characters, replacing the escape and character with the non-escaped version.  (i.e.
 	 * replaces "\\n" with "\n".  Counterpart to ReplaceCharWithEscapedCharInline().
 	 */
-	CORE_API void ReplaceEscapedCharWithCharInline( const TArray<ElementType>* Chars = nullptr );
+	CORE_API void ReplaceEscapedCharWithCharInline(const TArray<ElementType>* Chars = nullptr);
 
 	/**
 	 * Removes the escape backslash for all supported characters, replacing the escape and character with the non-escaped version.  (i.e.
 	 * replaces "\\n" with "\n".  Counterpart to ReplaceCharWithEscapedChar().
 	 * @return copy of this string with replacement made
 	 */
-	[[nodiscard]] UE_STRING_CLASS ReplaceEscapedCharWithChar( const TArray<ElementType>* Chars = nullptr ) const &
+	[[nodiscard]] UE_STRING_CLASS ReplaceEscapedCharWithChar(const TArray<ElementType>* Chars = nullptr) const &
 	{
 		UE_STRING_CLASS Result(*this);
 		Result.ReplaceEscapedCharWithCharInline(Chars);
@@ -2125,7 +2153,7 @@ public:
 	 * replaces "\\n" with "\n".  Counterpart to ReplaceCharWithEscapedChar().
 	 * @return copy of this string with replacement made
 	 */
-	[[nodiscard]] UE_STRING_CLASS ReplaceEscapedCharWithChar( const TArray<ElementType>* Chars = nullptr ) &&
+	[[nodiscard]] UE_STRING_CLASS ReplaceEscapedCharWithChar(const TArray<ElementType>* Chars = nullptr) &&
 	{
 		ReplaceEscapedCharWithCharInline(Chars);
 		return MoveTemp(*this);
@@ -2160,8 +2188,8 @@ public:
 		return MoveTemp(*this);
 	}
 
-	// Takes the number passed in and formats the string in comma format ( 12345 becomes "12,345")
-	[[nodiscard]] static CORE_API UE_STRING_CLASS FormatAsNumber( int32 InNumber );
+	// Takes the number passed in and formats the string in comma format (12345 becomes "12,345")
+	[[nodiscard]] static CORE_API UE_STRING_CLASS FormatAsNumber(int32 InNumber);
 
 	// To allow more efficient memory handling, automatically adds one for the string termination.
 	CORE_API void Reserve(int32 CharacterCount);
@@ -2173,11 +2201,11 @@ public:
 	 * @param	Ar				Archive to serialize with
 	 * @param	MinCharacters	Minimum number of characters to serialize.
 	 */
-	CORE_API void SerializeAsANSICharArray( FArchive& Ar, int32 MinCharacters=0 ) const;
+	CORE_API void SerializeAsANSICharArray(FArchive& Ar, int32 MinCharacters = 0) const;
 
 
 	/** Converts an integer to a string. */
-	[[nodiscard]] static FORCEINLINE UE_STRING_CLASS FromInt( int32 Num )
+	[[nodiscard]] static FORCEINLINE UE_STRING_CLASS FromInt(int32 Num)
 	{
 		UE_STRING_CLASS Ret;
 		Ret.AppendInt(Num); 
@@ -2185,7 +2213,7 @@ public:
 	}
 
 	/** appends the integer InNum to this string */
-	CORE_API void AppendInt( int32 InNum );
+	CORE_API void AppendInt(int32 InNum);
 
 	/**
 	 * Converts a string into a boolean value
@@ -2204,7 +2232,7 @@ public:
 	 *
 	 * @return the blob in string form
 	 */
-	[[nodiscard]] static CORE_API UE_STRING_CLASS FromBlob(const uint8* SrcBuffer,const uint32 SrcSize);
+	[[nodiscard]] static CORE_API UE_STRING_CLASS FromBlob(const uint8* SrcBuffer, const uint32 SrcSize);
 
 	/**
 	 * Converts a string into a buffer
@@ -2214,7 +2242,7 @@ public:
 	 *
 	 * @return true if the conversion happened, false otherwise
 	 */
-	static CORE_API bool ToBlob(const UE_STRING_CLASS& Source,uint8* DestBuffer,const uint32 DestSize);
+	static CORE_API bool ToBlob(const UE_STRING_CLASS& Source,uint8* DestBuffer, const uint32 DestSize);
 
 	/**
 	 * Converts a buffer to a string by hex-ifying the elements
@@ -2224,7 +2252,7 @@ public:
 	 *
 	 * @return the blob in string form
 	 */
-	[[nodiscard]] static CORE_API UE_STRING_CLASS FromHexBlob(const uint8* SrcBuffer,const uint32 SrcSize);
+	[[nodiscard]] static CORE_API UE_STRING_CLASS FromHexBlob(const uint8* SrcBuffer, const uint32 SrcSize);
 
 	/**
 	 * Converts a string into a buffer
@@ -2234,7 +2262,7 @@ public:
 	 *
 	 * @return true if the conversion happened, false otherwise
 	 */
-	static CORE_API bool ToHexBlob(const UE_STRING_CLASS& Source,uint8* DestBuffer,const uint32 DestSize);
+	static CORE_API bool ToHexBlob(const UE_STRING_CLASS& Source, uint8* DestBuffer, const uint32 DestSize);
 
 	/**
 	 * Converts a float string with the trailing zeros stripped
@@ -2245,7 +2273,7 @@ public:
 	 *
 	 * @return sanitized string version of float
 	 */
-	[[nodiscard]] static CORE_API UE_STRING_CLASS SanitizeFloat( double InFloat, const int32 InMinFractionalDigits = 1 );
+	[[nodiscard]] static CORE_API UE_STRING_CLASS SanitizeFloat(double InFloat, const int32 InMinFractionalDigits = 1);
 
 	/**
 	 * Joins a range of 'something that can be concatentated to strings with +=' together into a single string with separators.
