Added PRAGMA_DISABLE_CPP20_EXTENSION_WARNINGS for allowing asm in constexpr blocks (supported in Clang, but not officially until C++20).

diff --git a/Engine/Source/Runtime/Core/Public/Clang/ClangPlatformCompilerPreSetup.h b/Engine/Source/Runtime/Core/Public/Clang/ClangPlatformCompilerPreSetup.h
--- a/Engine/Source/Runtime/Core/Public/Clang/ClangPlatformCompilerPreSetup.h
+++ b/Engine/Source/Runtime/Core/Public/Clang/ClangPlatformCompilerPreSetup.h
@@ -52,6 +52,17 @@
 		_Pragma("clang diagnostic pop")
 #endif // PRAGMA_ENABLE_SHADOW_VARIABLE_WARNINGS
 
+#ifndef PRAGMA_DISABLE_CPP20_EXTENSION_WARNINGS
+	#define PRAGMA_DISABLE_CPP20_EXTENSION_WARNINGS \
+		_Pragma("clang diagnostic push") \
+		_Pragma("clang diagnostic ignored \"-Wc++20-extensions\"")
+#endif // PRAGMA_DISABLE_CPP20_EXTENSION_WARNINGS
+
+#ifndef PRAGMA_ENABLE_CPP20_EXTENSION_WARNINGS
+	#define PRAGMA_ENABLE_CPP20_EXTENSION_WARNINGS \
+		_Pragma("clang diagnostic pop")
+#endif // PRAGMA_ENABLE_CPP20_EXTENSION_WARNINGS
+
 // Unreachable not being evaluated on clang... yet
 #ifndef PRAGMA_DISABLE_UNREACHABLE_CODE_WARNINGS
 	#define PRAGMA_DISABLE_UNREACHABLE_CODE_WARNINGS
