Update TStaticArray to be constexpr along with Swap and Exchange.

Also added some basic compile time tests for TStaticArray.

[RN] Update TStaticArray to be constexpr along with Swap and Exchange.

diff --git a/Engine/Plugins/Runtime/GameFeatures/Source/GameFeatures/Private/GameFeaturePluginStateMachine.cpp b/Engine/Plugins/Runtime/GameFeatures/Source/GameFeatures/Private/GameFeaturePluginStateMachine.cpp
--- a/Engine/Plugins/Runtime/GameFeatures/Source/GameFeatures/Private/GameFeaturePluginStateMachine.cpp
+++ b/Engine/Plugins/Runtime/GameFeatures/Source/GameFeatures/Private/GameFeaturePluginStateMachine.cpp
@@ -3844,6 +3844,8 @@ namespace UE::GameFeatures
 		using FImplementsBatchProcessPtr = bool(*)();
 		using FBatchProcessPtr = void (*)(TConstArrayView<UGameFeaturePluginStateMachine*>);
 
+		FBatchProcessHelperFunctors() = default;
+
 		FBatchProcessHelperFunctors(FImplementsBatchProcessPtr InImplementsBatchProcessPtr, FBatchProcessPtr InBatchProcessPtr)
 			: ImplementsBatchProcess(InImplementsBatchProcessPtr)
 			, BatchProcess(InBatchProcessPtr)
diff --git a/Engine/Source/Runtime/Core/Public/Containers/StaticArray.h b/Engine/Source/Runtime/Core/Public/Containers/StaticArray.h
--- a/Engine/Source/Runtime/Core/Public/Containers/StaticArray.h
+++ b/Engine/Source/Runtime/Core/Public/Containers/StaticArray.h
@@ -25,11 +25,11 @@ class alignas(Alignment) TStaticArray
 public:
 	using ElementType = InElementType;
 
-	[[nodiscard]] TStaticArray() = default;
+	[[nodiscard]] constexpr TStaticArray() = default;
 
 	// Constructs each element with Args
 	template <typename... ArgTypes>
-	[[nodiscard]] explicit TStaticArray(EInPlace, ArgTypes&&... Args)
+	[[nodiscard]] constexpr explicit TStaticArray(EInPlace, ArgTypes&&... Args)
 		: Storage(InPlace, TMakeIntegerSequence<uint32, NumElements>(), Forward<ArgTypes>(Args)...)
 	{
 	}
@@ -39,31 +39,31 @@ public:
 		typename... ArgTypes
 		UE_REQUIRES((sizeof...(ArgTypes) > 0 && sizeof...(ArgTypes) <= NumElements) && UE::Core::Private::TCanBeConvertedToFromAll_V<InElementType, ArgTypes...>)
 	>
-	[[nodiscard]] TStaticArray(ArgTypes&&... Args)
+	[[nodiscard]] constexpr TStaticArray(ArgTypes&&... Args)
 		: Storage(PerElement, Forward<ArgTypes>(Args)...)
 	{
 	}
 
-	[[nodiscard]] TStaticArray(TStaticArray&& Other) = default;
-	[[nodiscard]] TStaticArray(const TStaticArray& Other) = default;
-	TStaticArray& operator=(TStaticArray&& Other) = default;
-	TStaticArray& operator=(const TStaticArray& Other) = default;
+	[[nodiscard]] constexpr TStaticArray(TStaticArray&& Other) = default;
+	[[nodiscard]] constexpr TStaticArray(const TStaticArray& Other) = default;
+	constexpr TStaticArray& operator=(TStaticArray&& Other) = default;
+	constexpr TStaticArray& operator=(const TStaticArray& Other) = default;
 
 	// Accessors.
-	[[nodiscard]] FORCEINLINE_DEBUGGABLE InElementType& operator[](uint32 Index)
+	[[nodiscard]] FORCEINLINE_DEBUGGABLE constexpr InElementType& operator[](uint32 Index)
 	{
 		checkSlow(Index < NumElements);
 		return Storage.Elements[Index].Element;
 	}
 
-	[[nodiscard]] FORCEINLINE_DEBUGGABLE const InElementType& operator[](uint32 Index) const
+	[[nodiscard]] FORCEINLINE_DEBUGGABLE constexpr const InElementType& operator[](uint32 Index) const
 	{
 		checkSlow(Index < NumElements);
 		return Storage.Elements[Index].Element;
 	}
 
 	// Comparisons.
-	[[nodiscard]] friend bool operator==(const TStaticArray& A,const TStaticArray& B)
+	[[nodiscard]] constexpr friend bool operator==(const TStaticArray& A,const TStaticArray& B)
 	{
 		for(uint32 ElementIndex = 0;ElementIndex < NumElements;++ElementIndex)
 		{
@@ -75,7 +75,7 @@ public:
 		return true;
 	}
 
-	[[nodiscard]] bool operator!=(const TStaticArray& B) const
+	[[nodiscard]] constexpr bool operator!=(const TStaticArray& B) const
 	{
 		for(uint32 ElementIndex = 0;ElementIndex < NumElements;++ElementIndex)
 		{
@@ -93,24 +93,24 @@ public:
 	 * @returns True if the array is empty.
 	 * @see Num
 	 */
-	[[nodiscard]] bool IsEmpty() const
+	[[nodiscard]] constexpr bool IsEmpty() const
 	{
 		return NumElements == 0;
 	}
 
 	/** The number of elements in the array. */
-	[[nodiscard]] FORCEINLINE_DEBUGGABLE int32 Num() const
+	[[nodiscard]] FORCEINLINE_DEBUGGABLE constexpr int32 Num() const
 	{
 		return NumElements;
 	}
 
 	/** A pointer to the first element of the array */
-	[[nodiscard]] FORCEINLINE_DEBUGGABLE InElementType* GetData()
+	[[nodiscard]] FORCEINLINE_DEBUGGABLE constexpr InElementType* GetData()
 	{
 		static_assert((alignof(ElementType) % Alignment) == 0, "GetData() cannot be called on a TStaticArray with non-standard alignment");
 		return &Storage.Elements[0].Element;
 	}
-	[[nodiscard]] FORCEINLINE_DEBUGGABLE const InElementType* GetData() const
+	[[nodiscard]] FORCEINLINE_DEBUGGABLE constexpr const InElementType* GetData() const
 	{
 		return const_cast<TStaticArray*>(this)->GetData();
 	}
@@ -119,11 +119,11 @@ private:
 
 	struct alignas(Alignment) TArrayStorageElementAligned
 	{
-		[[nodiscard]] TArrayStorageElementAligned() = default;
+		[[nodiscard]] constexpr TArrayStorageElementAligned() = default;
 
 		// Index is used to achieve pack expansion in TArrayStorage's first constructor, but is unused here
 		template <typename... ArgTypes>
-		[[nodiscard]] explicit TArrayStorageElementAligned(EInPlace, uint32 /*Index*/, ArgTypes&&... Args)
+		[[nodiscard]] constexpr explicit TArrayStorageElementAligned(EInPlace, uint32 /*Index*/, ArgTypes&&... Args)
 			: Element(Forward<ArgTypes>(Args)...)
 		{
 		}
@@ -133,10 +133,10 @@ private:
 
 	struct TArrayStorage
 	{
-		[[nodiscard]] TArrayStorage() = default;
+		[[nodiscard]] constexpr TArrayStorage() = default;
 
 		template<uint32... Indices, typename... ArgTypes>
-		[[nodiscard]] explicit TArrayStorage(EInPlace, TIntegerSequence<uint32, Indices...>, ArgTypes&&... Args)
+		[[nodiscard]] constexpr explicit TArrayStorage(EInPlace, TIntegerSequence<uint32, Indices...>, ArgTypes&&... Args)
 			: Elements{ TArrayStorageElementAligned(InPlace, Indices, Args...)... }
 		{
 			// The arguments are deliberately not forwarded arguments here, because we're initializing multiple elements
@@ -149,7 +149,7 @@ private:
 		}
 
 		template<typename... ArgTypes>
-		[[nodiscard]] explicit TArrayStorage(EPerElement, ArgTypes&&... Args)
+		[[nodiscard]] constexpr explicit TArrayStorage(EPerElement, ArgTypes&&... Args)
 			: Elements{ TArrayStorageElementAligned(InPlace, 0 /* dummy index */, Forward<ArgTypes>(Args))... }
 		{
 		}
@@ -165,12 +165,12 @@ public:
 	template <typename StorageElementType, bool bReverse = false>
 	struct FRangedForIterator
 	{
-		[[nodiscard]] explicit FRangedForIterator(StorageElementType* InPtr)
+		[[nodiscard]] constexpr explicit FRangedForIterator(StorageElementType* InPtr)
 			: Ptr(InPtr)
 		{
 		}
 
-		[[nodiscard]] auto& operator*() const
+		[[nodiscard]] constexpr auto& operator*() const
 		{
 			if constexpr (bReverse)
 			{
@@ -182,7 +182,7 @@ public:
 			}
 		}
 
-		FRangedForIterator& operator++()
+		constexpr FRangedForIterator& operator++()
 		{
 			if constexpr (bReverse)
 			{
@@ -195,7 +195,7 @@ public:
 			return *this;
 		}
 
-		[[nodiscard]] bool operator!=(const FRangedForIterator& B) const
+		[[nodiscard]] constexpr bool operator!=(const FRangedForIterator& B) const
 		{
 			return Ptr != B.Ptr;
 		}
@@ -210,19 +210,19 @@ public:
 	using RangedForConstReverseIteratorType = FRangedForIterator<const TArrayStorageElementAligned, true>;
 
 	/** STL-like iterators to enable range-based for loop support. */
-	[[nodiscard]] FORCEINLINE RangedForIteratorType             begin()        { return RangedForIteratorType(Storage.Elements); }
-	[[nodiscard]] FORCEINLINE RangedForConstIteratorType        begin()  const { return RangedForConstIteratorType(Storage.Elements); }
-	[[nodiscard]] FORCEINLINE RangedForIteratorType             end()          { return RangedForIteratorType(Storage.Elements + NumElements); }
-	[[nodiscard]] FORCEINLINE RangedForConstIteratorType        end()    const { return RangedForConstIteratorType(Storage.Elements + NumElements); }
-	[[nodiscard]] FORCEINLINE RangedForReverseIteratorType      rbegin()       { return RangedForReverseIteratorType(Storage.Elements + NumElements); }
-	[[nodiscard]] FORCEINLINE RangedForConstReverseIteratorType rbegin() const { return RangedForConstReverseIteratorType(Storage.Elements + NumElements); }
-	[[nodiscard]] FORCEINLINE RangedForReverseIteratorType      rend()         { return RangedForReverseIteratorType(Storage.Elements); }
-	[[nodiscard]] FORCEINLINE RangedForConstReverseIteratorType rend()   const { return RangedForConstReverseIteratorType(Storage.Elements); }
+	[[nodiscard]] FORCEINLINE RangedForIteratorType             constexpr begin()        { return RangedForIteratorType(Storage.Elements); }
+	[[nodiscard]] FORCEINLINE RangedForConstIteratorType        constexpr begin()  const { return RangedForConstIteratorType(Storage.Elements); }
+	[[nodiscard]] FORCEINLINE RangedForIteratorType             constexpr end()          { return RangedForIteratorType(Storage.Elements + NumElements); }
+	[[nodiscard]] FORCEINLINE RangedForConstIteratorType        constexpr end()    const { return RangedForConstIteratorType(Storage.Elements + NumElements); }
+	[[nodiscard]] FORCEINLINE RangedForReverseIteratorType      constexpr rbegin()       { return RangedForReverseIteratorType(Storage.Elements + NumElements); }
+	[[nodiscard]] FORCEINLINE RangedForConstReverseIteratorType constexpr rbegin() const { return RangedForConstReverseIteratorType(Storage.Elements + NumElements); }
+	[[nodiscard]] FORCEINLINE RangedForReverseIteratorType      constexpr rend()         { return RangedForReverseIteratorType(Storage.Elements); }
+	[[nodiscard]] FORCEINLINE RangedForConstReverseIteratorType constexpr rend()   const { return RangedForConstReverseIteratorType(Storage.Elements); }
 };
 
 /** Creates a static array filled with the specified value. */
 template <typename InElementType, uint32 NumElements>
-[[nodiscard]] TStaticArray<InElementType,NumElements> MakeUniformStaticArray(typename TCallTraits<InElementType>::ParamType InValue)
+[[nodiscard]] constexpr TStaticArray<InElementType,NumElements> MakeUniformStaticArray(typename TCallTraits<InElementType>::ParamType InValue)
 {
 	TStaticArray<InElementType,NumElements> Result;
 	for(uint32 ElementIndex = 0;ElementIndex < NumElements;++ElementIndex)
diff --git a/Engine/Source/Runtime/Core/Public/Templates/UnrealTemplate.h b/Engine/Source/Runtime/Core/Public/Templates/UnrealTemplate.h
--- a/Engine/Source/Runtime/Core/Public/Templates/UnrealTemplate.h
+++ b/Engine/Source/Runtime/Core/Public/Templates/UnrealTemplate.h
@@ -625,7 +625,7 @@ UE_INTRINSIC_CAST FORCEINLINE constexpr T&& Forward(std::remove_reference_t<T>&&
  * Swap two values.  Assumes the types are trivially relocatable.
  */
 template <typename T>
-inline void Swap(T& A, T& B)
+constexpr inline void Swap(T& A, T& B)
 {
 	// std::is_swappable isn't correct here, because we allow bitwise swapping of types containing e.g. const and reference members,
 	// but we don't want to allow swapping of types which are UE_NONCOPYABLE or equivalent.  We also allow bitwise swapping of arrays, so
@@ -653,7 +653,7 @@ inline void Swap(T& A, T& B)
 }
 
 template <typename T>
-inline void Exchange(T& A, T& B)
+constexpr void Exchange(T& A, T& B)
 {
 	Swap(A, B);
 }
diff --git a/Engine/Source/Runtime/Core/Tests/Containers/StaticArrayTests.cpp b/Engine/Source/Runtime/Core/Tests/Containers/StaticArrayTests.cpp
new file mode 100644
--- /dev/null
+++ b/Engine/Source/Runtime/Core/Tests/Containers/StaticArrayTests.cpp
@@ -0,0 +1,158 @@
+// Copyright Epic Games, Inc. All Rights Reserved.
+
+#if WITH_TESTS
+
+#include "Containers/StaticArray.h"
+#include "Templates/UnrealTemplate.h"
+
+// Currently only doing compile time tests to verify constexpr functionality.
+namespace Test
+{
+	struct FNonTrivialStruct
+	{
+		constexpr FNonTrivialStruct() = default;
+		constexpr ~FNonTrivialStruct() = default;
+
+		constexpr FNonTrivialStruct(int32 InValue) noexcept
+			: Value(InValue)
+		{
+
+		}
+
+		constexpr FNonTrivialStruct(const FNonTrivialStruct& Other) noexcept
+		{
+			Value = Other.Value;
+		}
+
+		constexpr FNonTrivialStruct(FNonTrivialStruct&& Other) noexcept
+		{
+			Swap(Value, Other.Value);
+		}
+
+		constexpr FNonTrivialStruct& operator=(const FNonTrivialStruct& Other) noexcept
+		{
+			if (this != &Other)
+			{
+				Value = Other.Value;
+			}
+			return *this;
+		}
+
+		constexpr FNonTrivialStruct& operator=(FNonTrivialStruct&& Other) noexcept
+		{
+			if (this != &Other)
+			{
+				Swap(Value, Other.Value);
+			}
+			return *this;
+		}
+
+		constexpr bool operator==(const FNonTrivialStruct& Other) const = default;
+		constexpr bool operator!=(const FNonTrivialStruct& Other) const = default;
+
+		constexpr int32 GetValue() const
+		{
+			return Value;
+		}
+
+		int32 Value = 0;
+	};
+
+	static_assert(!std::is_trivial<FNonTrivialStruct>::value);
+
+	using FTestType = FNonTrivialStruct;
+
+	// Test creating an array and assigning elements
+	template<uint32 Count>
+	constexpr TStaticArray<FTestType, Count> CreateAscendingArray()
+	{
+		TStaticArray<FTestType, Count> Array;
+		for (int32 Idx = 0; Idx < Array.Num(); ++Idx)
+		{
+			Array[Idx] = Idx;
+		}
+		return Array;
+	}
+
+	static_assert(CreateAscendingArray<4>().Num() == 4);
+	static_assert(CreateAscendingArray<4>()[2].GetValue() == 2);
+
+	// Test variadic constructor
+	template<typename ElementType, typename... Elements>
+	constexpr auto CreateStaticArray(Elements... InElements)
+	{
+		return TStaticArray<ElementType, sizeof...(InElements)>(InElements...);
+	}
+
+	static_assert(CreateStaticArray<FTestType>(1, 2, 3, 4)[2].GetValue() == 3);
+
+	// Test In-place constructor
+	static_assert(TStaticArray<FTestType, 4>(InPlace, 42)[2].GetValue() == 42);
+
+	// Test copying
+	template<uint32 Count>
+	constexpr TStaticArray<FTestType, Count> CloneStaticArray(const TStaticArray<FTestType, Count>& InArray)
+	{
+		TStaticArray<FTestType, Count> OutArray = InArray;
+		return OutArray;
+	}
+
+	static_assert(CloneStaticArray(CreateAscendingArray<4>())[2].GetValue() == 2);
+
+	// Test moving
+	template<uint32 Count>
+	constexpr TStaticArray<FTestType, Count> StealStaticArray(TStaticArray<FTestType, Count>&& InArray)
+	{
+		TStaticArray<FTestType, Count> OutArray(MoveTemp(InArray));
+		return OutArray;
+	}
+
+	static_assert(StealStaticArray(CreateAscendingArray<4>())[2].GetValue() == 2);
+
+	template<uint32 Count>
+	constexpr TStaticArray<FTestType, Count> StealStaticArrayViaMoveAssignment(TStaticArray<FTestType, Count>&& InArray)
+	{
+		TStaticArray<FTestType, Count> OutArray;
+		OutArray = MoveTemp(InArray);
+		return OutArray;
+	}
+
+	static_assert(StealStaticArray(CreateAscendingArray<4>())[2].GetValue() == 2);
+
+	template<uint32 Count>
+	constexpr TStaticArray<FTestType, Count> ReturnMovedFromArray(TStaticArray<FTestType, Count>&& InArray)
+	{
+		TStaticArray<FTestType, Count> OutArray = MoveTemp(InArray);
+		return InArray;
+	}
+
+	// This is only safe as we know the moved-from state is valid.
+	static_assert(ReturnMovedFromArray<4>(CreateAscendingArray<4>())[2] == 0);
+
+	// Test MakeUniformStaticArray
+	static_assert(MakeUniformStaticArray<FTestType, 8>(42).Num() == 8);
+	static_assert(MakeUniformStaticArray<FTestType, 8>(42)[3].GetValue() == 42);
+
+	// Test comparisons
+	static_assert(CreateAscendingArray<4>() == CreateAscendingArray<4>());
+	static_assert(CreateAscendingArray<4>() != MakeUniformStaticArray<FTestType, 4>(10));
+
+	// Test IsEmpty
+	static_assert(TStaticArray<FTestType, 0>{}.IsEmpty());
+
+	// Test Iteration
+	template<uint32 Count>
+	constexpr TStaticArray<FTestType, Count> AddOne(TStaticArray<FTestType, Count> InArray)
+	{
+		for (FTestType& Value : InArray)
+		{
+			++Value.Value;
+		}
+		return InArray;
+	}
+
+	static_assert(AddOne(MakeUniformStaticArray<FTestType, 4>(10)) == MakeUniformStaticArray<FTestType, 4>(11));
+
+} // namespace Test
+
+#endif // WITH_TESTS
