FOrientedBox cleanup

Project() can be done much more cheaply, CalcVertices() a bit cheaper.

diff --git a/Engine/Source/Runtime/Core/Public/Math/OrientedBox.h b/Engine/Source/Runtime/Core/Public/Math/OrientedBox.h
--- a/Engine/Source/Runtime/Core/Public/Math/OrientedBox.h
+++ b/Engine/Source/Runtime/Core/Public/Math/OrientedBox.h
@@ -80,41 +80,32 @@ FORCEINLINE void FOrientedBox::CalcVertices( FVector* Verts ) const
 		{
 			for (int32 k = 0; k < 2; k++)
 			{
-				*Verts++ = Center + Signs[i] * AxisX * ExtentX + Signs[j] * AxisY * ExtentY + Signs[k] * AxisZ * ExtentZ;
+				*Verts++ = Center + (Signs[i] * ExtentX) * AxisX + (Signs[j] * ExtentY) * AxisY + (Signs[k] * ExtentZ) * AxisZ;
 			}
 		}
 	}
 }
 
-
 FORCEINLINE FFloatInterval FOrientedBox::Project( const FVector& Axis ) const
 {
-	const FVector::FReal Signs[] = {-1.0f, 1.0f};
-
-	// calculate the projections of the box center and the extent-scaled axes.
+	// Consider:
+	//   max { dot(Center +- ExtentX * AxisX +- ExtentY * AxisY +- ExtentZ * AxisZ, Axis) }
+	//  =max { dot(Center, Axis) +- ExtentX * dot(AxisX, Axis) +- ExtentY * dot(AxisY, Axis) +- ExtentZ * dot(AxisZ, Axis) }
+	//
+	// These individual terms can be maximized separately and are clearly maximal when their effective signs
+	// are all positive. Analogous for the min with all-negative signs.
+
+	// Calculate the projections of the box center and the extent-scaled axes.
 	FVector::FReal ProjectedCenter = Axis | Center;
-	FVector::FReal ProjectedAxisX = Axis | (ExtentX * AxisX);
-	FVector::FReal ProjectedAxisY = Axis | (ExtentY * AxisY);
-	FVector::FReal ProjectedAxisZ = Axis | (ExtentZ * AxisZ);
-
-	FFloatInterval ProjectionInterval;
-
-	for (int32 i = 0; i < 2; i++)
-	{
-		for (int32 j = 0; j < 2; j++)
-		{
-			for (int32 k = 0; k < 2; k++)
-			{
-				// project the box vertex onto the axis.
-				float ProjectedVertex = float(ProjectedCenter + Signs[i] * ProjectedAxisX + Signs[j] * ProjectedAxisY + Signs[k] * ProjectedAxisZ);	// LWC_TODO: Precision loss
-
-				// if necessary, expand the projection interval to include the box vertex projection.
-				ProjectionInterval.Include(ProjectedVertex);
-			}
-		}
-	}
-
-	return ProjectionInterval;
+	FVector::FReal AbsProjectedAxisX = FMath::Abs(ExtentX * (Axis | AxisX));
+	FVector::FReal AbsProjectedAxisY = FMath::Abs(ExtentY * (Axis | AxisY));
+	FVector::FReal AbsProjectedAxisZ = FMath::Abs(ExtentZ * (Axis | AxisZ));
+
+	FVector::FReal AbsProjectedExtent = AbsProjectedAxisX + AbsProjectedAxisY + AbsProjectedAxisZ;
+	return FFloatInterval {
+		float(ProjectedCenter - AbsProjectedExtent),
+		float(ProjectedCenter + AbsProjectedExtent)
+	};
 }
 
 template <> struct TIsPODType<FOrientedBox> { enum { Value = true }; };
