Add IsTaskEvent to FBaseGraphTask and fix it so Unlock and DispatchSubsequents are safe to call on either type.
Header cleanup of TaskGraphInterfaces.h to clarify the public API now that the old task graph frontend has been removed.

diff --git a/Engine/Source/Runtime/Core/Public/Async/TaskGraphInterfaces.h b/Engine/Source/Runtime/Core/Public/Async/TaskGraphInterfaces.h
--- a/Engine/Source/Runtime/Core/Public/Async/TaskGraphInterfaces.h
+++ b/Engine/Source/Runtime/Core/Public/Async/TaskGraphInterfaces.h
@@ -452,13 +452,11 @@ struct FTaskGraphBlockAllocationTag : FDefaultBlockAllocationTag
 	using Allocator = TBlockAllocationCache<BlockSize, FAlignedAllocator>;
 };
 
-/** 
- *	Base class for all tasks. A replacement for `FBaseGraphTask` and `FGraphEvent` from the old API, based on `Tasks::Private::FTaskBase` functionality
- **/
-
+/** Base class for all graph tasks, used for both TGraphTask and simple graph events. This is a wrapper around the `Tasks::Private::FTaskBase` functionality */
 class FBaseGraphTask : public UE::Tasks::Private::FTaskBase
 {
 public:
+	// These functions handle integration with FTaskBase and the named thread executor and should not be called directly
 	explicit FBaseGraphTask(const FGraphEventArray* InPrerequisites)
 		: FTaskBase(/*InitRefCount=*/ 1, false /* bUnlockPrerequisites */)
 	{
@@ -475,11 +473,6 @@ public:
 		FTaskBase::Init(InDebugName, InPriority, InExtendedPriority, InTaskFlags);
 	}
 
-	void Unlock(ENamedThreads::Type CurrentThreadIfKnown = ENamedThreads::AnyThread)
-	{
-		TryLaunch(0);
-	}
-
 	FORCEINLINE void Execute(TArray<FBaseGraphTask*>& NewTasks, ENamedThreads::Type CurrentThread, bool bDeleteOnCompletion)
 	{	// only called for named thread tasks, normal tasks are executed using `FTaskBase` API directly (see `TGraphTask`)
 		checkSlow(NewTasks.Num() == 0);
@@ -489,11 +482,17 @@ public:
 		ReleaseInternalReference(); // named tasks are executed by named threads, outside of the scheduler
 	}
 
+
+	/** Returns a reference to this task that can be used as a prerequisite for new tasks or passed to DontCompleteUntil */
 	FGraphEventRef GetCompletionEvent()
 	{
 		return this;
 	}
 
+	/** 
+	 * Call on a currently active task to add a nested task, which will delay any subsequent tasks until the nested task completes.
+	 * AddPrerequisites should be used instead if the task has not yet been launched.
+	 */
 	void DontCompleteUntil(FGraphEventRef NestedTask)
 	{
 		if (!NestedTask)
@@ -501,8 +500,8 @@ public:
 			return;
 		}
 
-		if (GetExtendedPriority() == UE::Tasks::EExtendedTaskPriority::TaskEvent)
-		{	// TaskEvent can't have nested tasks, add it as a prerequisite instead
+		if (IsTaskEvent())
+		{	// TaskEvent can't have nested tasks, add it as a prerequisite instead to retain backward compatibility
 			AddPrerequisites(*NestedTask);
 		}
 		else
@@ -512,20 +511,43 @@ public:
 		}
 	}
 
+	/** Return true if this task has finished executing, this wrapper exists for backward compatibility */
 	bool IsComplete() const
 	{
-		return IsCompleted(); // the new API uses a slightly different name
+		return IsCompleted();
 	}
 
+	/** Create a simple task event that can be used as a prerequisite for other tasks, and then manually triggered with DispatchSubsequents */
 	static FGraphEventRef CreateGraphEvent();
 
-	void DispatchSubsequents(ENamedThreads::Type CurrentThreadIfKnown = ENamedThreads::AnyThread)
+	/** Returns true if this is a simple task event, which cannot execute code and must be manually dispatched */
+	bool IsTaskEvent() const
 	{
-		//check(CurrentThreadIfKnown == ENamedThreads::AnyThread); // the feature is not used
-		AddRef(); // scheduler's reference
+		return GetExtendedPriority() == UE::Tasks::EExtendedTaskPriority::TaskEvent;
+	}
+
+	/** 
+	 * Unlocks a task that was returned from ConstructAndHold or CreateGraphEvent, which may execute immediately if prerequisites are already completed.
+	 * This can only be safely called exactly once and the caller is responsible for tracking that.
+	 */
+	void Unlock(ENamedThreads::Type CurrentThreadIfKnown = ENamedThreads::AnyThread)
+	{
+		if (IsTaskEvent())
+		{
+			// An event is not "in the system" until it's triggered, and should be kept alive only by external references. Once it's triggered it's in the system 
+			// and can outlive external references, so we need to keep it alive by holding an internal reference. It will be released when the event is signaled
+			AddRef();
+		}
 		TryLaunch(0);
 	}
 
+	/** Backward compatibility wrapper for Unlock */
+	void DispatchSubsequents(ENamedThreads::Type CurrentThreadIfKnown = ENamedThreads::AnyThread)
+	{
+		Unlock(CurrentThreadIfKnown);
+	}
+
+	UE_DEPRECATED(5.6, "Call AddPrerequisites separately if you need to add new tasks before dispatch")
 	void DispatchSubsequents(TArray<FBaseGraphTask*>& NewTasks, ENamedThreads::Type CurrentThreadIfKnown = ENamedThreads::AnyThread)
 	{
 		check(NewTasks.Num() == 0); // the feature is not used
@@ -537,6 +559,7 @@ public:
 		// to the old API recently, is used only in a couple of places and will be fixed manually by switching to the new API
 	}
 
+	/** Blocks the current thread and waits for the completion of this task */
 	void Wait(ENamedThreads::Type CurrentThreadIfKnown = ENamedThreads::AnyThread)
 	{
 		// local queue have to be handled by the original TaskGraph implementation. Tasks System doesn't support local queues
@@ -548,13 +571,14 @@ public:
 		FTaskBase::WaitWithNamedThreadsSupport();
 	}
 
+	/** Returns the old-style named thread this task will be executed on */
 	ENamedThreads::Type GetThreadToExecuteOn() const
 	{
 		return UE::Tasks::Private::TranslatePriority(GetPriority(), GetExtendedPriority());
 	}
 };
 
-// the new task implementation integrated into the old task API
+/** Templated graph task that is created to execute a specific function */
 template<typename TTask>
 class TGraphTask final : public TConcurrentLinearObject<TGraphTask<TTask>, FTaskGraphBlockAllocationTag>, public FBaseGraphTask
 {
@@ -568,7 +592,10 @@ public:
 	public:
 		UE_NONCOPYABLE(FConstructor);
 
-		/** Passthrough internal task constructor and dispatch. Note! Generally speaking references will not pass through; use pointers */
+		/**
+		 * Constructs a task and immediately dispatches the task for possible execution if prerequisites have completed.
+		 * Note! Generally speaking references will not pass through; use pointers
+		 */
 		template<typename...T>
 		FORCEINLINE_DEBUGGABLE FGraphEventRef ConstructAndDispatchWhenReady(T&&... Args)
 		{
@@ -577,7 +604,10 @@ public:
 			return Ref;
 		}
 
-		/** Passthrough internal task constructor and hold. */
+		/**
+		 * Constructs a task and holds it for later dispatch by calling Unlock.
+		 * Note! Generally speaking references will not pass through; use pointers
+		 */
 		template<typename...T>
 		FORCEINLINE_DEBUGGABLE TGraphTask* ConstructAndHold(T&&... Args)
 		{
