[AutoRTFM] Make Serialize logic easier to understand.

Previously, there weren't comments and it wasn't easy to understand why it was structured the way it was.
I've added comments and separated apart some of the logic that used to be mixed together into separate blocks.
Functionally this should still work as it did before.

diff --git a/Engine/Source/Programs/AutoRTFMTests/Private/ArchiveWriterTests.cpp b/Engine/Source/Programs/AutoRTFMTests/Private/ArchiveWriterTests.cpp
--- a/Engine/Source/Programs/AutoRTFMTests/Private/ArchiveWriterTests.cpp
+++ b/Engine/Source/Programs/AutoRTFMTests/Private/ArchiveWriterTests.cpp
@@ -316,6 +316,25 @@ TEST_CASE("TransactionallySafeArchiveWriter.OnCommitWorks")
 	FString("Two").SerializeAsANSICharArray(MemoryWriter);
 	FString("Three").SerializeAsANSICharArray(MemoryWriter);
 
+	SECTION("Commit(OnCommit)")
+	{
+		TArray<uint8> TSStorage;
+		FTransactionallySafeArchiveWriter SafeWriter(MakeUnique<FCheckingMemoryWriter>(TSStorage));
+
+		AutoRTFM::Testing::Commit([&]
+		{
+			AutoRTFM::OnCommit([&]
+			{
+				FString("One").SerializeAsANSICharArray(SafeWriter);
+				FString("Two").SerializeAsANSICharArray(SafeWriter);
+				FString("Three").SerializeAsANSICharArray(SafeWriter);
+			});
+		});
+
+		REQUIRE(TSStorage == NormalStorage);
+	}
+
+
 	SECTION("Write, Commit(Write, OnCommit)")
 	{
 		TArray<uint8> TSStorage;
diff --git a/Engine/Source/Runtime/Core/Public/Serialization/TransactionallySafeArchiveWriter.h b/Engine/Source/Runtime/Core/Public/Serialization/TransactionallySafeArchiveWriter.h
--- a/Engine/Source/Runtime/Core/Public/Serialization/TransactionallySafeArchiveWriter.h
+++ b/Engine/Source/Runtime/Core/Public/Serialization/TransactionallySafeArchiveWriter.h
@@ -75,32 +75,38 @@ public:
 	{
 		check(InnerArchive);
 
-		if (bRegisteredCommitHandler || AutoRTFM::IsTransactional())
+		if (bRegisteredCommitHandler)
 		{
+			// If we already have a registered commit handler, we are either within a transaction or are committing/aborting.
+			check(AutoRTFM::IsTransactional() || AutoRTFM::IsCommittingOrAborting());
+			DeferredWrites.Append(static_cast<uint8*>(Data), IntCastChecked<DeferredWriteBuffer::SizeType>(Num));
+		}
+		else if (AutoRTFM::IsTransactional())
+		{
+			// If we don't have a commit handler set up yet, but find ourselves running transactionally, set up the commit handler now.
+			RegisterCommitHandler();
 			DeferredWrites.Append(static_cast<uint8*>(Data), IntCastChecked<DeferredWriteBuffer::SizeType>(Num));
-			MaybeRegisterCommitHandler();
 		}
 		else
 		{
+			// We aren't in a transaction and don't have a commit handler, so nothing needs to be done. Forward to the inner archive.
 			check(DeferredWrites.IsEmpty());
 			InnerArchive->Serialize(Data, Num);
 		}
 	}
 
 private:
-	void MaybeRegisterCommitHandler()
+	void RegisterCommitHandler()
 	{
-		if (!bRegisteredCommitHandler)
-		{
-			bRegisteredCommitHandler = true;
+		check(!bRegisteredCommitHandler);
+		bRegisteredCommitHandler = true;
 
-			AutoRTFM::PushOnCommitHandler(this, [this]
-			{
-				InnerArchive->Serialize(DeferredWrites.GetData(), DeferredWrites.NumBytes());
-				DeferredWrites.Reset();
-				bRegisteredCommitHandler = false;
-			});
-		}
+		AutoRTFM::PushOnCommitHandler(this, [this]
+		{
+			InnerArchive->Serialize(DeferredWrites.GetData(), DeferredWrites.NumBytes());
+			DeferredWrites.Reset();
+			bRegisteredCommitHandler = false;
+		});
 	}
 
 	using DeferredWriteBuffer = TArray<uint8>;
