Unix: Indicate OutOfMemory when appropriate during signal-based termination

[RN] OutOfMemory crash type is now reported when appropriate on Linux

diff --git a/Engine/Source/Runtime/Core/Private/Unix/UnixPlatformCrashContext.cpp b/Engine/Source/Runtime/Core/Private/Unix/UnixPlatformCrashContext.cpp
--- a/Engine/Source/Runtime/Core/Private/Unix/UnixPlatformCrashContext.cpp
+++ b/Engine/Source/Runtime/Core/Private/Unix/UnixPlatformCrashContext.cpp
@@ -1000,6 +1000,12 @@ void PlatformCrashHandler(int32 Signal, siginfo_t* Info, void* Context)
 		ErrorProgramCounter = GCrashErrorProgramCounter;
 	}
 
+	// Expose OOM root cause during generic crash (eg. SIGSEGV on nullptr).
+	if (Type == ECrashContextType::Crash && FPlatformMemory::bIsOOM)
+	{
+		Type = ECrashContextType::OutOfMemory;
+	}
+
 	FUnixCrashContext CrashContext(Type, ErrorMessage);
 	CrashContext.InitFromSignal(Signal, Info, Context);
 	CrashContext.FirstCrashHandlerFrame = static_cast<uint64*>(__builtin_return_address(0));
