Improve logging for library loading

#rnx

diff --git a/Engine/Source/Runtime/Core/Private/Modules/ModuleManager.cpp b/Engine/Source/Runtime/Core/Private/Modules/ModuleManager.cpp
--- a/Engine/Source/Runtime/Core/Private/Modules/ModuleManager.cpp
+++ b/Engine/Source/Runtime/Core/Private/Modules/ModuleManager.cpp
@@ -522,7 +522,7 @@ void FModuleManager::RefreshModuleFilenameFromManifest(const FName InModuleName)
 
 void* FModuleManager::InternalLoadLibrary(FName ModuleName, const FString& LibraryToLoad)
 {
-	UE_LOG(LogModuleManager, Verbose, TEXT("InternalLoadLibrary: '%s'"), *ModuleName.ToString());
+	UE_LOG(LogModuleManager, Verbose, TEXT("InternalLoadLibrary: '%s' ('%s')"), *ModuleName.ToString(), *LibraryToLoad);
 
 	void* Handle = nullptr;
 
