Make `RecordCustomStat` work in transactions.

#rnx

diff --git a/Engine/Source/Programs/AutoRTFMTests/Private/CsvProfilerTests.cpp b/Engine/Source/Programs/AutoRTFMTests/Private/CsvProfilerTests.cpp
--- a/Engine/Source/Programs/AutoRTFMTests/Private/CsvProfilerTests.cpp
+++ b/Engine/Source/Programs/AutoRTFMTests/Private/CsvProfilerTests.cpp
@@ -16,6 +16,10 @@
 CSV_DEFINE_CATEGORY(MyFalseCategory, false);
 CSV_DEFINE_CATEGORY(MyTrueCategory, true);
 
+CSV_DEFINE_CATEGORY(MyIntCategory, true);
+CSV_DEFINE_CATEGORY(MyFloatCategory, true);
+CSV_DEFINE_CATEGORY(MyDoubleCategory, true);
+
 TEST_CASE("CsvProfiler")
 {
 	struct RAII final
@@ -100,4 +104,37 @@ TEST_CASE("CsvProfiler")
 	{
 		CSV_CUSTOM_STAT_GLOBAL(SomeStat, 42, ECsvCustomStatOp::Max);
 	});
+
+	AutoRTFM::Testing::Commit([&]
+	{
+		FName Name(TEXT("Wowwee"));
+		FCsvProfiler::RecordCustomStat(Name, CSV_CATEGORY_INDEX(MyIntCategory), 42, ECsvCustomStatOp::Accumulate);
+	});
+
+	AutoRTFM::Testing::Commit([&]
+	{
+		FCsvProfiler::RecordCustomStat("Wowwee", CSV_CATEGORY_INDEX(MyIntCategory), 42, ECsvCustomStatOp::Accumulate);
+	});
+
+	AutoRTFM::Testing::Commit([&]
+	{
+		FName Name(TEXT("Wowwee"));
+		FCsvProfiler::RecordCustomStat(Name, CSV_CATEGORY_INDEX(MyFloatCategory), static_cast<float>(42), ECsvCustomStatOp::Accumulate);
+	});
+
+	AutoRTFM::Testing::Commit([&]
+	{
+		FCsvProfiler::RecordCustomStat("Wowwee", CSV_CATEGORY_INDEX(MyFloatCategory), static_cast<float>(42), ECsvCustomStatOp::Accumulate);
+	});
+
+	AutoRTFM::Testing::Commit([&]
+	{
+		FName Name(TEXT("Wowwee"));
+		FCsvProfiler::RecordCustomStat(Name, CSV_CATEGORY_INDEX(MyDoubleCategory), static_cast<double>(42), ECsvCustomStatOp::Accumulate);
+	});
+
+	AutoRTFM::Testing::Commit([&]
+	{
+		FCsvProfiler::RecordCustomStat("Wowwee", CSV_CATEGORY_INDEX(MyDoubleCategory), static_cast<double>(42), ECsvCustomStatOp::Accumulate);
+	});
 }
diff --git a/Engine/Source/Runtime/Core/Public/ProfilingDebugging/CsvProfiler.h b/Engine/Source/Runtime/Core/Public/ProfilingDebugging/CsvProfiler.h
--- a/Engine/Source/Runtime/Core/Public/ProfilingDebugging/CsvProfiler.h
+++ b/Engine/Source/Runtime/Core/Public/ProfilingDebugging/CsvProfiler.h
@@ -380,12 +380,12 @@ public:
 	CORE_API static void BeginExclusiveStat(const char * StatName, const char* NamedEventName = nullptr);
 	CORE_API static void EndExclusiveStat(const char * StatName);
 
-	CORE_API static void RecordCustomStat(const char * StatName, uint32 CategoryIndex, float Value, const ECsvCustomStatOp CustomStatOp);
-	CORE_API static void RecordCustomStat(const FName& StatName, uint32 CategoryIndex, float Value, const ECsvCustomStatOp CustomStatOp);
-	CORE_API static void RecordCustomStat(const char* StatName, uint32 CategoryIndex, double Value, const ECsvCustomStatOp CustomStatOp);
-	CORE_API static void RecordCustomStat(const FName& StatName, uint32 CategoryIndex, double Value, const ECsvCustomStatOp CustomStatOp);
-	CORE_API static void RecordCustomStat(const char * StatName, uint32 CategoryIndex, int32 Value, const ECsvCustomStatOp CustomStatOp);
-	CORE_API static void RecordCustomStat(const FName& StatName, uint32 CategoryIndex, int32 Value, const ECsvCustomStatOp CustomStatOp);
+	CORE_API UE_AUTORTFM_ALWAYS_OPEN static void RecordCustomStat(const char * StatName, uint32 CategoryIndex, float Value, const ECsvCustomStatOp CustomStatOp);
+	CORE_API UE_AUTORTFM_ALWAYS_OPEN static void RecordCustomStat(const FName& StatName, uint32 CategoryIndex, float Value, const ECsvCustomStatOp CustomStatOp);
+	CORE_API UE_AUTORTFM_ALWAYS_OPEN static void RecordCustomStat(const char* StatName, uint32 CategoryIndex, double Value, const ECsvCustomStatOp CustomStatOp);
+	CORE_API UE_AUTORTFM_ALWAYS_OPEN static void RecordCustomStat(const FName& StatName, uint32 CategoryIndex, double Value, const ECsvCustomStatOp CustomStatOp);
+	CORE_API UE_AUTORTFM_ALWAYS_OPEN static void RecordCustomStat(const char * StatName, uint32 CategoryIndex, int32 Value, const ECsvCustomStatOp CustomStatOp);
+	CORE_API UE_AUTORTFM_ALWAYS_OPEN static void RecordCustomStat(const FName& StatName, uint32 CategoryIndex, int32 Value, const ECsvCustomStatOp CustomStatOp);
 
 	CORE_API static void RecordEndOfPipeCustomStat(const FName& StatName, uint32 CategoryIndex, double Value, const ECsvCustomStatOp CustomStatOp);
 
