Suppress PVS warning in StaticArrayTests.

The sub expressions are intentionally identical to test the comparison operator.

#rnx

diff --git a/Engine/Source/Runtime/Core/Tests/Containers/StaticArrayTests.cpp b/Engine/Source/Runtime/Core/Tests/Containers/StaticArrayTests.cpp
--- a/Engine/Source/Runtime/Core/Tests/Containers/StaticArrayTests.cpp
+++ b/Engine/Source/Runtime/Core/Tests/Containers/StaticArrayTests.cpp
@@ -134,7 +134,7 @@ namespace Test
 	static_assert(MakeUniformStaticArray<FTestType, 8>(42)[3].GetValue() == 42);
 
 	// Test comparisons
-	static_assert(CreateAscendingArray<4>() == CreateAscendingArray<4>());
+	static_assert(CreateAscendingArray<4>() == CreateAscendingArray<4>()); //-V501
 	static_assert(CreateAscendingArray<4>() != MakeUniformStaticArray<FTestType, 4>(10));
 
 	// Test IsEmpty
