Fix the compile fail on non-monolithic + AutoRTFM.

#rnx

diff --git a/Engine/Source/Runtime/Core/Public/ProfilingDebugging/CsvProfiler.h b/Engine/Source/Runtime/Core/Public/ProfilingDebugging/CsvProfiler.h
--- a/Engine/Source/Runtime/Core/Public/ProfilingDebugging/CsvProfiler.h
+++ b/Engine/Source/Runtime/Core/Public/ProfilingDebugging/CsvProfiler.h
@@ -380,12 +380,12 @@ public:
 	CORE_API static void BeginExclusiveStat(const char * StatName, const char* NamedEventName = nullptr);
 	CORE_API static void EndExclusiveStat(const char * StatName);
 
-	CORE_API UE_AUTORTFM_ALWAYS_OPEN static void RecordCustomStat(const char * StatName, uint32 CategoryIndex, float Value, const ECsvCustomStatOp CustomStatOp);
-	CORE_API UE_AUTORTFM_ALWAYS_OPEN static void RecordCustomStat(const FName& StatName, uint32 CategoryIndex, float Value, const ECsvCustomStatOp CustomStatOp);
-	CORE_API UE_AUTORTFM_ALWAYS_OPEN static void RecordCustomStat(const char* StatName, uint32 CategoryIndex, double Value, const ECsvCustomStatOp CustomStatOp);
-	CORE_API UE_AUTORTFM_ALWAYS_OPEN static void RecordCustomStat(const FName& StatName, uint32 CategoryIndex, double Value, const ECsvCustomStatOp CustomStatOp);
-	CORE_API UE_AUTORTFM_ALWAYS_OPEN static void RecordCustomStat(const char * StatName, uint32 CategoryIndex, int32 Value, const ECsvCustomStatOp CustomStatOp);
-	CORE_API UE_AUTORTFM_ALWAYS_OPEN static void RecordCustomStat(const FName& StatName, uint32 CategoryIndex, int32 Value, const ECsvCustomStatOp CustomStatOp);
+	UE_AUTORTFM_ALWAYS_OPEN CORE_API static void RecordCustomStat(const char * StatName, uint32 CategoryIndex, float Value, const ECsvCustomStatOp CustomStatOp);
+	UE_AUTORTFM_ALWAYS_OPEN CORE_API static void RecordCustomStat(const FName& StatName, uint32 CategoryIndex, float Value, const ECsvCustomStatOp CustomStatOp);
+	UE_AUTORTFM_ALWAYS_OPEN CORE_API static void RecordCustomStat(const char* StatName, uint32 CategoryIndex, double Value, const ECsvCustomStatOp CustomStatOp);
+	UE_AUTORTFM_ALWAYS_OPEN CORE_API static void RecordCustomStat(const FName& StatName, uint32 CategoryIndex, double Value, const ECsvCustomStatOp CustomStatOp);
+	UE_AUTORTFM_ALWAYS_OPEN CORE_API static void RecordCustomStat(const char * StatName, uint32 CategoryIndex, int32 Value, const ECsvCustomStatOp CustomStatOp);
+	UE_AUTORTFM_ALWAYS_OPEN CORE_API static void RecordCustomStat(const FName& StatName, uint32 CategoryIndex, int32 Value, const ECsvCustomStatOp CustomStatOp);
 
 	CORE_API static void RecordEndOfPipeCustomStat(const FName& StatName, uint32 CategoryIndex, double Value, const ECsvCustomStatOp CustomStatOp);
 
