PR #12918: FrameValue copy constructor fix

#rnx

diff --git a/Engine/Source/Runtime/Core/Public/Misc/FrameValue.h b/Engine/Source/Runtime/Core/Public/Misc/FrameValue.h
--- a/Engine/Source/Runtime/Core/Public/Misc/FrameValue.h
+++ b/Engine/Source/Runtime/Core/Public/Misc/FrameValue.h
@@ -41,7 +41,7 @@ public:
 	TFrameValue(const TFrameValue& InValue)
 		: FrameSet(GFrameCounter)
 	{
-		Value = InValue;
+		Value = InValue.Value;
 	}
 
 	TFrameValue(TFrameValue&& InValue)
