Add note to TVector::GetSafeNormal() that Tolerance must be positive.

The TVector impl of GetSafeNormal() can produce NaNs if that is not the case.
Does not apply to Vector2D or Vector4 variants which do not have this problem.

Prompted by [REDACTED_PERSON]190084. (That one is rejected, but clarifying the docs is good.)

diff --git a/Engine/Source/Runtime/Core/Public/Math/Vector.h b/Engine/Source/Runtime/Core/Public/Math/Vector.h
--- a/Engine/Source/Runtime/Core/Public/Math/Vector.h
+++ b/Engine/Source/Runtime/Core/Public/Math/Vector.h
@@ -643,7 +643,7 @@ public:
      * Gets a normalized copy of the vector, checking it is safe to do so based on the length.
      * Returns zero vector by default if vector length is too small to safely normalize.
      *
-     * @param Tolerance Minimum squared vector length.
+     * @param Tolerance Minimum squared vector length. Must be positive.
      * @return A normalized copy if safe, ResultIfZero otherwise.
      */
     [[nodiscard]] TVector<T> GetSafeNormal(T Tolerance=UE_SMALL_NUMBER, const TVector<T>& ResultIfZero = ZeroVector) const;
@@ -652,7 +652,7 @@ public:
      * Gets a normalized copy of the 2D components of the vector, checking it is safe to do so. Z is set to zero. 
      * Returns zero vector by default if vector length is too small to normalize.
      *
-     * @param Tolerance Minimum squared vector length.
+     * @param Tolerance Minimum squared vector length. Must be positive.
      * @return Normalized copy if safe, otherwise returns ResultIfZero.
      */
     [[nodiscard]] TVector<T> GetSafeNormal2D(T Tolerance=UE_SMALL_NUMBER, const TVector<T>& ResultIfZero = ZeroVector) const;
