[AutoRTFM] Maintain a thread-local cache of each thread's stack range.

This dramatically improves performance of top-level transactions on Linux, making an empty AutoRTFM::Transact() run 54x faster than before. The added benchmark has a mean duration of 207.6ms before the fix, and 3.8ms with the fix.

Without this change, some complex maps were too slow to be playable on Linux servers due to the high cost of object duplication/serialization.

For a deep dive on the Linux performance flaw, see [REDACTED_PROFILE_URL]/go/issues/68587

diff --git a/Engine/Source/Programs/AutoRTFMTests/Private/Benchmarks.cpp b/Engine/Source/Programs/AutoRTFMTests/Private/Benchmarks.cpp
--- a/Engine/Source/Programs/AutoRTFMTests/Private/Benchmarks.cpp
+++ b/Engine/Source/Programs/AutoRTFMTests/Private/Benchmarks.cpp
@@ -718,3 +718,35 @@ TEST_CASE("Benchmarks.PopAllOnCommitHandlers")
 		REQUIRE(!bHit);
 	};
 }
+
+TEST_CASE("Benchmarks.TopLevelTransaction")
+{
+	constexpr int Iterations = 100;
+
+	BENCHMARK("Commit")
+	{
+		int Counter = 0;
+		for (int Index = 0; Index < Iterations; ++Index)
+		{
+			AutoRTFM::Testing::Commit([&]
+			{
+				++Counter;
+			});
+		}
+		REQUIRE(Counter == Iterations);
+	};
+
+	BENCHMARK("Abort")
+	{
+		int Counter = 0;
+		for (int Index = 0; Index < Iterations; ++Index)
+		{
+			AutoRTFM::Testing::Abort([&]
+			{
+				++Counter;
+				AutoRTFM::AbortTransaction();
+			});
+		}
+		REQUIRE(Counter == 0);
+	};
+}
\ No newline at end of file
diff --git a/Engine/Source/Runtime/Core/Private/AutoRTFM/Context.cpp b/Engine/Source/Runtime/Core/Private/AutoRTFM/Context.cpp
--- a/Engine/Source/Runtime/Core/Private/AutoRTFM/Context.cpp
+++ b/Engine/Source/Runtime/Core/Private/AutoRTFM/Context.cpp
@@ -390,6 +390,37 @@ ETransactionResult FContext::ResolveNestedTransaction(FTransaction* NewTransacti
 	}
 }
 
+AutoRTFM::FStackRange FContext::GetThreadStackRange()
+{
+	// On some platforms, looking up the stack range is quite expensive, so caching it
+	// is important for performance. Linux glibc is particularly bad--see
+	// [REDACTED_PROFILE_URL]/go/issues/68587 for a deep dive.
+	thread_local FStackRange CachedStackRange = []
+	{
+		FStackRange Stack;
+
+#if PLATFORM_WINDOWS
+		GetCurrentThreadStackLimits(&Stack.Low, &Stack.High);
+#elif defined(__APPLE__)         
+		Stack.High = pthread_get_stackaddr_np(pthread_self());
+		size_t StackSize = pthread_get_stacksize_np(pthread_self());
+		Stack.Low = static_cast<char*>(Stack.High) - StackSize;
+#else
+		pthread_attr_t Attr{};
+		pthread_getattr_np(pthread_self(), &Attr);
+		Stack.Low = 0;
+		size_t StackSize = 0;
+		pthread_attr_getstack(&Attr, &Stack.Low, &StackSize);
+		Stack.High = static_cast<char*>(Stack.Low) + StackSize;
+#endif
+
+		AUTORTFM_ASSERT(Stack.High > Stack.Low);
+		return Stack;
+	}();
+
+	return CachedStackRange;
+}
+
 ETransactionResult FContext::Transact(void (*UninstrumentedFunction)(void*), void (*InstrumentedFunction)(void*), void* Arg)
 {
     constexpr bool bVerbose = false;
@@ -438,22 +469,7 @@ ETransactionResult FContext::Transact(void (*UninstrumentedFunction)(void*), voi
 		CurrentThreadId = FThreadID::GetCurrent();
 
 		AUTORTFM_ASSERT(Stack == FStackRange{});
-
-#if PLATFORM_WINDOWS
-		GetCurrentThreadStackLimits(&Stack.Low, &Stack.High);
-#elif defined(__APPLE__)         
-		Stack.High = pthread_get_stackaddr_np(pthread_self());
-		size_t StackSize = pthread_get_stacksize_np(pthread_self());
-		Stack.Low = static_cast<char*>(Stack.High) - StackSize;
-#else
-		pthread_attr_t Attr{};
-		pthread_getattr_np(pthread_self(), &Attr);
-		Stack.Low = 0;
-		size_t StackSize = 0;
-		pthread_attr_getstack(&Attr, &Stack.Low, &StackSize);
-		Stack.High = static_cast<char*>(Stack.Low) + StackSize;
-#endif
-		AUTORTFM_ASSERT(Stack.High > Stack.Low);
+		Stack = GetThreadStackRange();
 
 		AUTORTFM_ASSERT(Stack.Contains(TransactStackStart));
 
diff --git a/Engine/Source/Runtime/Core/Private/AutoRTFM/Context.h b/Engine/Source/Runtime/Core/Private/AutoRTFM/Context.h
--- a/Engine/Source/Runtime/Core/Private/AutoRTFM/Context.h
+++ b/Engine/Source/Runtime/Core/Private/AutoRTFM/Context.h
@@ -162,6 +162,8 @@ private:
     // All of this other stuff ought to be private?
     void Reset();
 
+	static FStackRange GetThreadStackRange();
+
 	// We defer allocating FTransactions at the top of the transaction stack.
 	// This allows us to make starting a transaction in the open a load, some math, and a store.
 	uint64_t NumDeferredTransactions{0};
