Core: Add copy sign function.

diff --git a/Engine/Source/Runtime/Core/Public/GenericPlatform/GenericPlatformMath.h b/Engine/Source/Runtime/Core/Public/GenericPlatform/GenericPlatformMath.h
--- a/Engine/Source/Runtime/Core/Public/GenericPlatform/GenericPlatformMath.h
+++ b/Engine/Source/Runtime/Core/Public/GenericPlatform/GenericPlatformMath.h
@@ -559,6 +559,11 @@ struct FGenericPlatformMath
 	[[nodiscard]] static FORCEINLINE float InvSqrtEst( float F ) { return InvSqrt( F ); }
 	[[nodiscard]] static FORCEINLINE double InvSqrtEst( double F ) { return InvSqrt( F ); }
 
+	/** Composes a floating-point value with the magnitude of X and the sign of Y. */
+	[[nodiscard]] static FORCEINLINE float CopySign(float X, float Y) { return copysignf(X, Y); }
+	[[nodiscard]] static FORCEINLINE double CopySign(double X, double Y) { return copysign(X, Y); }
+	RESOLVE_FLOAT_AMBIGUITY_2_ARGS(CopySign);
+
 	/** Return true if value is NaN (not a number). */
 	[[nodiscard]] static FORCEINLINE bool IsNaN( float A ) 
 	{
