Allow us to force `USE_HITCH_DETECTION` on.

diff --git a/Engine/Source/Programs/AutoRTFMTests/AutoRTFMTests.Target.cs b/Engine/Source/Programs/AutoRTFMTests/AutoRTFMTests.Target.cs
--- a/Engine/Source/Programs/AutoRTFMTests/AutoRTFMTests.Target.cs
+++ b/Engine/Source/Programs/AutoRTFMTests/AutoRTFMTests.Target.cs
@@ -76,6 +76,9 @@ public class AutoRTFMTestsTarget : TargetRules
 		GlobalDefinitions.Add("CSV_PROFILER=1");
 		GlobalDefinitions.Add("CSV_PROFILER_ENABLE_IN_SHIPPING=1");
 
+		// Need to force the heartbeat on so we can test with it!
+		GlobalDefinitions.Add("USE_HITCH_DETECTION=1");
+
 		// For testing insights, this sets `UE_TRACE_ENABLED`.
 		bEnableTrace = true;
 	}
diff --git a/Engine/Source/Runtime/Core/Public/Misc/Build.h b/Engine/Source/Runtime/Core/Public/Misc/Build.h
--- a/Engine/Source/Runtime/Core/Public/Misc/Build.h
+++ b/Engine/Source/Runtime/Core/Public/Misc/Build.h
@@ -446,7 +446,9 @@
 	#define USE_PER_MODULE_UOBJECT_BOOTSTRAP 0
 #endif
 
-#define USE_HITCH_DETECTION (ALLOW_HITCH_DETECTION && !WITH_EDITORONLY_DATA && !IS_PROGRAM && !UE_BUILD_DEBUG && !USING_THREAD_SANITISER && !USING_INSTRUMENTATION)
+#ifndef USE_HITCH_DETECTION
+	#define USE_HITCH_DETECTION (ALLOW_HITCH_DETECTION && !WITH_EDITORONLY_DATA && !IS_PROGRAM && !UE_BUILD_DEBUG && !USING_THREAD_SANITISER && !USING_INSTRUMENTATION)
+#endif
 
 // Controls whether shipping builds create backups of the most recent log file.
 // All other configurations always create backups.
