Made TIsCharEncodingCompatibleWith accept non-char types but return false, which makes it easier to use in SFINAE contexts.

diff --git a/Engine/Source/Runtime/Core/Public/Traits/IsCharEncodingCompatibleWith.h b/Engine/Source/Runtime/Core/Public/Traits/IsCharEncodingCompatibleWith.h
--- a/Engine/Source/Runtime/Core/Public/Traits/IsCharEncodingCompatibleWith.h
+++ b/Engine/Source/Runtime/Core/Public/Traits/IsCharEncodingCompatibleWith.h
@@ -14,7 +14,11 @@ namespace UE::Core::Private
 	template <typename SrcEncoding, typename DestEncoding>
 	constexpr bool IsCharEncodingCompatibleWithImpl()
 	{
-		if constexpr (std::is_same_v<SrcEncoding, DestEncoding>)
+		if constexpr (!TIsCharType_V<SrcEncoding> || !TIsCharType_V<DestEncoding>)
+		{
+			return false;
+		}
+		else if constexpr (std::is_same_v<SrcEncoding, DestEncoding>)
 		{
 			return true;
 		}
@@ -59,9 +63,6 @@ namespace UE::Core::Private
 template <typename SrcEncoding, typename DestEncoding>
 struct TIsCharEncodingCompatibleWith
 {
-	static_assert(TIsCharType<SrcEncoding >::Value, "SrcEncoding is not a character encoding type");
-	static_assert(TIsCharType<DestEncoding>::Value, "DestEncoding is not a character encoding type");
-
 	static constexpr bool Value = UE::Core::Private::IsCharEncodingCompatibleWithImpl<std::remove_cv_t<SrcEncoding>, std::remove_cv_t<DestEncoding>>();
 };
 
