Add a small optimization to `AutoRTFM::ForTheRuntime::IsAutoRTFMRuntimeEnabled` for a 1.4% speed bump.

It speeds up `execStmEnabled` by 47% as it was previously bound on this function.

#rnx

diff --git a/Engine/Source/Runtime/Core/Private/AutoRTFM/API.cpp b/Engine/Source/Runtime/Core/Private/AutoRTFM/API.cpp
--- a/Engine/Source/Runtime/Core/Private/AutoRTFM/API.cpp
+++ b/Engine/Source/Runtime/Core/Private/AutoRTFM/API.cpp
@@ -151,7 +151,7 @@ namespace AutoRTFM
 #endif
 		}
 
-		bool IsAutoRTFMRuntimeEnabled()
+		bool IsAutoRTFMRuntimeEnabledInternal()
 		{
 			// #noop if AutoRTFM is not compiled in
 #if UE_AUTORTFM
diff --git a/Engine/Source/Runtime/Core/Public/AutoRTFM/AutoRTFM.h b/Engine/Source/Runtime/Core/Public/AutoRTFM/AutoRTFM.h
--- a/Engine/Source/Runtime/Core/Public/AutoRTFM/AutoRTFM.h
+++ b/Engine/Source/Runtime/Core/Public/AutoRTFM/AutoRTFM.h
@@ -1231,8 +1231,19 @@ namespace ForTheRuntime
 	// successfully.
 	UE_AUTORTFM_API bool SetAutoRTFMRuntime(EAutoRTFMEnabledState State);
 
+	UE_AUTORTFM_API bool IsAutoRTFMRuntimeEnabledInternal();
+
 	// Query whether the AutoRTFM runtime is enabled.
-	UE_AUTORTFM_API bool IsAutoRTFMRuntimeEnabled();
+	static UE_AUTORTFM_FORCEINLINE bool IsAutoRTFMRuntimeEnabled()
+	{
+		// If we are already in the closed nest of a transaction, we must have our runtime enabled!
+		if (AutoRTFM::IsClosed())
+		{
+			return true;
+		}
+
+		return IsAutoRTFMRuntimeEnabledInternal();
+	}
 
 	// Set the percentage [0..100] chance that a call to `CoinTossDisable` will end up disabling AutoRTFM.
 	// 100% means never disable via coin-toss, 0% means always disable. So passing `0.1` means disable
