Deprecated bool AllowShrinking parameters for everyone.

diff --git a/Engine/Source/Runtime/Core/Public/Containers/AllowShrinking.h b/Engine/Source/Runtime/Core/Public/Containers/AllowShrinking.h
--- a/Engine/Source/Runtime/Core/Public/Containers/AllowShrinking.h
+++ b/Engine/Source/Runtime/Core/Public/Containers/AllowShrinking.h
@@ -13,4 +13,4 @@ enum class EAllowShrinking : uint8
 	Default = Yes
 };
 
-#define UE_ALLOWSHRINKING_BOOL_DEPRECATED(FunctionName) UE_DEPRECATED_FORENGINE(5.5, FunctionName " with a boolean bAllowShrinking has been deprecated - please use the EAllowShrinking enum instead")
+#define UE_ALLOWSHRINKING_BOOL_DEPRECATED(FunctionName) UE_DEPRECATED(5.6, FunctionName " with a boolean bAllowShrinking has been deprecated - please use the EAllowShrinking enum instead")
