Wrap PLATFORM_BREAK on Unix platforms with PRAGMA_DISABLE_CPP20_EXTENSION_WARNINGS for allowing asm in constexpr blocks (supported in Clang, but not officially until C++20).

#rnx

diff --git a/Engine/Source/Runtime/Core/Public/Unix/UnixPlatform.h b/Engine/Source/Runtime/Core/Public/Unix/UnixPlatform.h
--- a/Engine/Source/Runtime/Core/Public/Unix/UnixPlatform.h
+++ b/Engine/Source/Runtime/Core/Public/Unix/UnixPlatform.h
@@ -79,7 +79,7 @@ typedef FUnixPlatformTypes FPlatformTypes;
 #endif
 
 #if PLATFORM_CPU_X86_FAMILY
-	#define PLATFORM_BREAK()							__asm__ volatile("int $0x03")
+	#define PLATFORM_BREAK()							{ PRAGMA_DISABLE_CPP20_EXTENSION_WARNINGS __asm__ volatile("int $0x03"); PRAGMA_ENABLE_CPP20_EXTENSION_WARNINGS }
 #else
 	#define PLATFORM_BREAK()							raise(SIGTRAP)
 #endif // PLATFORM_CPU_X86_FAMILY
