[AutoRTFM] Tweaks to allow AutoRTFM to compile as a standalone module.

* Use UE_AUTORTFM_UNUSED() for all unused parameters. This will become a warning / error once AutoRTFM is in its own module.
* Replace use of MoveTemp() with std::move(). We can't depend on Core.
* Provide AutoRTFM dllexport / dllimport definitions for UE_AUTORTFM_API so we don't have to depend on Core's definitions of DLLEXPORT and DLLIMPORT.
* Add BuildMacros.h - declares AutoRTFM platform & OS defines. UBA will #define the particular platform & OS, but leave the others undefined. This ensures there's a definition for all AutoRTFM cares about. Use these instead of the UE ones.
* Consistently use the AUTORTFM macro variants for FORCEINLINE / FORCENOINLINE / UE_ASSUME.
* Add `inline` to UE_AUTORTFM_FORCEINLINE to fix a duplicate symbol error when building modular builds.
* Replace use of PLATFORM_BREAK() with __builtin_debugtrap(). We only compile with clang, so we don't need to abstract this.
* Replace a few remaining uses of the UE integer types with the standard library types.

diff --git a/Engine/Source/Runtime/Core/Private/AutoRTFM/API.cpp b/Engine/Source/Runtime/Core/Private/AutoRTFM/API.cpp
--- a/Engine/Source/Runtime/Core/Private/AutoRTFM/API.cpp
+++ b/Engine/Source/Runtime/Core/Private/AutoRTFM/API.cpp
@@ -8,6 +8,7 @@
 
 #include <random>
 #include <tuple>
+#include <utility>
 
 #if UE_AUTORTFM
 static_assert(UE_AUTORTFM_ENABLED, "AutoRTFM/API.cpp requires the compiler flag '-fautortfm'");
@@ -147,6 +148,7 @@ namespace AutoRTFM
 
 			return true;
 #else
+			UE_AUTORTFM_UNUSED(State);
 			return false;
 #endif
 		}
@@ -175,6 +177,8 @@ namespace AutoRTFM
 #if UE_AUTORTFM
 			AUTORTFM_ASSERT(Chance >= 0.0f && Chance <= 100.0f);
 			GAutoRTFMEnabledProbability = Chance;
+#else
+			UE_AUTORTFM_UNUSED(Chance);
 #endif
 		}
 
@@ -246,6 +250,8 @@ namespace AutoRTFM
 					GExternAPI.OnRetryTransactionsChanged();
 				}
 			}
+#else
+			UE_AUTORTFM_UNUSED(State);
 #endif
 		}
 
@@ -313,6 +319,8 @@ namespace AutoRTFM
 					GExternAPI.OnMemoryValidationLevelChanged();
 				}
 			}
+#else
+			UE_AUTORTFM_UNUSED(Level);
 #endif
 		}
 
@@ -336,6 +344,8 @@ namespace AutoRTFM
 					GExternAPI.OnMemoryValidationThrottlingChanged();
 				}
 			}
+#else
+			UE_AUTORTFM_UNUSED(bEnabled);
 #endif
 		}
 
@@ -359,6 +369,8 @@ namespace AutoRTFM
 					GExternAPI.OnMemoryValidationStatisticsChanged();
 				}
 			}
+#else
+			UE_AUTORTFM_UNUSED(bEnabled);
 #endif
 		}
 
@@ -551,35 +563,35 @@ AUTORTFM_CORE_API void RTFM_CascadingAbortTransactionInternal(TTask<void()>&& Ru
 {
 	FContext* const Context = FContext::Get();
 	AUTORTFM_ASSERT(Context->GetStatus() == EContextStatus::OnTrack);
-	Context->AbortTransactionWithPostAbortCallback(EContextStatus::AbortedByCascadingAbort, MoveTemp(RunAfterAbort));
+	Context->AbortTransactionWithPostAbortCallback(EContextStatus::AbortedByCascadingAbort, std::move(RunAfterAbort));
 }
 
 AUTORTFM_CORE_API void RTFM_CascadingRetryTransactionInternal(TTask<void()> && RunAfterAbortBeforeRetryWork)
 {
 	FContext* Context = FContext::Get();
 	AUTORTFM_ASSERT(Context->GetStatus() == EContextStatus::OnTrack);
-	Context->AbortTransactionWithPostAbortCallback(EContextStatus::AbortedByCascadingRetry, MoveTemp(RunAfterAbortBeforeRetryWork));
+	Context->AbortTransactionWithPostAbortCallback(EContextStatus::AbortedByCascadingRetry, std::move(RunAfterAbortBeforeRetryWork));
 }
 
 AUTORTFM_CORE_API void RTFM_OnCommitInternal(TTask<void()> && Work)
 {
 	FContext* Context = FContext::Get();
 	AUTORTFM_ASSERT(Context->GetStatus() == EContextStatus::OnTrack);
-	Context->GetCurrentTransaction()->DeferUntilCommit(MoveTemp(Work));
+	Context->GetCurrentTransaction()->DeferUntilCommit(std::move(Work));
 }
 
 AUTORTFM_CORE_API void RTFM_OnAbortInternal(TTask<void()>&& Work)
 {
 	FContext* Context = FContext::Get();
 	AUTORTFM_ASSERT(Context->GetStatus() == EContextStatus::OnTrack);
-	Context->GetCurrentTransaction()->DeferUntilAbort(MoveTemp(Work));
+	Context->GetCurrentTransaction()->DeferUntilAbort(std::move(Work));
 }
 
 AUTORTFM_CORE_API void RTFM_PushOnCommitHandlerInternal(const void* Key, TTask<void()>&& Work)
 {
 	FContext* Context = FContext::Get();
 	AUTORTFM_ASSERT(Context->GetStatus() == EContextStatus::OnTrack);
-	Context->GetCurrentTransaction()->PushDeferUntilCommitHandler(Key, MoveTemp(Work));
+	Context->GetCurrentTransaction()->PushDeferUntilCommitHandler(Key, std::move(Work));
 }
 
 AUTORTFM_CORE_API void RTFM_PopOnCommitHandlerInternal(const void* Key)
@@ -600,7 +612,7 @@ AUTORTFM_CORE_API void RTFM_PushOnAbortHandlerInternal(const void* Key, TTask<vo
 {
 	FContext* Context = FContext::Get();
 	AUTORTFM_ASSERT(Context->GetStatus() == EContextStatus::OnTrack);
-	Context->GetCurrentTransaction()->PushDeferUntilAbortHandler(Key, MoveTemp(Work));
+	Context->GetCurrentTransaction()->PushDeferUntilAbortHandler(Key, std::move(Work));
 }
 
 AUTORTFM_CORE_API void RTFM_PopOnAbortHandlerInternal(const void* Key)
@@ -1004,14 +1016,14 @@ void ForTheRuntime::Initialize(const ForTheRuntime::FExternAPI& ExternAPI)
 void ForTheRuntime::CascadingAbortTransactionInternal(TTask<void()>&& RunAfterAbort)
 {
 	AUTORTFM_ASSERT(autortfm_is_closed());
-	RTFM_CascadingAbortTransactionInternal(MoveTemp(RunAfterAbort));
+	RTFM_CascadingAbortTransactionInternal(std::move(RunAfterAbort));
 }
 
 void ForTheRuntime::CascadingRetryTransactionInternal(TTask<void()>&& RunAfterAbortBeforeRetryWork)
 {
 	if (autortfm_is_closed())
 	{
-		return RTFM_CascadingRetryTransactionInternal(MoveTemp(RunAfterAbortBeforeRetryWork));
+		return RTFM_CascadingRetryTransactionInternal(std::move(RunAfterAbortBeforeRetryWork));
 	}
 }
 
@@ -1019,7 +1031,7 @@ void ForTheRuntime::OnCommitInternal(TTask<void()> && Work)
 {
 	if (autortfm_is_closed())
 	{
-		return RTFM_OnCommitInternal(MoveTemp(Work));
+		return RTFM_OnCommitInternal(std::move(Work));
 	}
 
 	Work();
@@ -1029,7 +1041,7 @@ void ForTheRuntime::OnAbortInternal(TTask<void()> && Work)
 {
 	if (autortfm_is_closed())
 	{
-		return RTFM_OnAbortInternal(MoveTemp(Work));
+		return RTFM_OnAbortInternal(std::move(Work));
 	}
 }
 
@@ -1037,7 +1049,7 @@ void ForTheRuntime::PushOnCommitHandlerInternal(const void* Key, TTask<void()>&&
 {
 	if (autortfm_is_closed())
 	{
-		return RTFM_PushOnCommitHandlerInternal(Key, MoveTemp(Work));
+		return RTFM_PushOnCommitHandlerInternal(Key, std::move(Work));
 	}
 }
 
@@ -1061,7 +1073,7 @@ void ForTheRuntime::PushOnAbortHandlerInternal(const void* Key, TTask<void()> &&
 {
 	if (autortfm_is_closed())
 	{
-		return RTFM_PushOnAbortHandlerInternal(Key, MoveTemp(Work));
+		return RTFM_PushOnAbortHandlerInternal(Key, std::move(Work));
 	}
 }
 
diff --git a/Engine/Source/Runtime/Core/Private/AutoRTFM/BuildMacros.h b/Engine/Source/Runtime/Core/Private/AutoRTFM/BuildMacros.h
new file mode 100644
--- /dev/null
+++ b/Engine/Source/Runtime/Core/Private/AutoRTFM/BuildMacros.h
@@ -0,0 +1,61 @@
+// Copyright Epic Games, Inc. All Rights Reserved.
+
+#pragma once
+
+#ifndef AUTORTFM_BUILD_DEBUG
+	#if defined(UE_BUILD_DEBUG) && UE_BUILD_DEBUG
+		#define AUTORTFM_BUILD_DEBUG 1
+	#else
+		#define AUTORTFM_BUILD_DEBUG 0
+	#endif
+#endif
+
+#ifndef AUTORTFM_BUILD_DEVELOPMENT
+	#if defined(UE_BUILD_DEVELOPMENT) && UE_BUILD_DEVELOPMENT
+		#define AUTORTFM_BUILD_DEVELOPMENT 1
+	#else
+		#define AUTORTFM_BUILD_DEVELOPMENT 0
+	#endif
+#endif
+
+#ifndef AUTORTFM_BUILD_TEST
+	#if defined(UE_BUILD_TEST) && UE_BUILD_TEST
+		#define AUTORTFM_BUILD_TEST 1
+	#else
+		#define AUTORTFM_BUILD_TEST 0
+	#endif
+#endif
+
+#ifndef AUTORTFM_BUILD_SHIPPING
+	#if defined(UE_BUILD_SHIPPING) && UE_BUILD_SHIPPING
+		#define AUTORTFM_BUILD_SHIPPING 1
+	#else
+		#define AUTORTFM_BUILD_SHIPPING 0
+	#endif
+#endif
+
+static_assert(AUTORTFM_BUILD_DEBUG + AUTORTFM_BUILD_DEVELOPMENT + AUTORTFM_BUILD_TEST + AUTORTFM_BUILD_SHIPPING == 1);
+
+#ifndef AUTORTFM_PLATFORM_WINDOWS
+	#if defined(PLATFORM_WINDOWS) && PLATFORM_WINDOWS
+		#define AUTORTFM_PLATFORM_WINDOWS 1
+	#else
+		#define AUTORTFM_PLATFORM_WINDOWS 0
+	#endif
+#endif
+
+#ifndef AUTORTFM_PLATFORM_MAC
+	#if defined(PLATFORM_MAC) && PLATFORM_MAC
+		#define AUTORTFM_PLATFORM_MAC 1
+	#else
+		#define AUTORTFM_PLATFORM_MAC 0
+	#endif
+#endif
+
+#ifndef AUTORTFM_PLATFORM_LINUX
+	#if defined(PLATFORM_LINUX) && PLATFORM_LINUX
+		#define AUTORTFM_PLATFORM_LINUX 1
+	#else
+		#define AUTORTFM_PLATFORM_LINUX 0
+	#endif
+#endif
diff --git a/Engine/Source/Runtime/Core/Private/AutoRTFM/CompilerInterface.cpp b/Engine/Source/Runtime/Core/Private/AutoRTFM/CompilerInterface.cpp
--- a/Engine/Source/Runtime/Core/Private/AutoRTFM/CompilerInterface.cpp
+++ b/Engine/Source/Runtime/Core/Private/AutoRTFM/CompilerInterface.cpp
@@ -3,6 +3,7 @@
 #if (defined(__AUTORTFM) && __AUTORTFM)
 
 #include "AutoRTFM/AutoRTFM.h"
+#include "BuildMacros.h"
 #include "ContextInlines.h"
 #include "ExternAPI.h"
 #include "FunctionMapInlines.h"
@@ -30,7 +31,7 @@ void CheckAlignment(FContext* Context, void* Ptr, size_t AlignmentMask)
     }
 }
 
-#if !UE_BUILD_SHIPPING
+#if !AUTORTFM_BUILD_SHIPPING
 
 // Check for writes to null in development code, so that the inevitable crash will occur
 // in the caller's code rather than in the AutoRTFM runtime.
diff --git a/Engine/Source/Runtime/Core/Private/AutoRTFM/Context.cpp b/Engine/Source/Runtime/Core/Private/AutoRTFM/Context.cpp
--- a/Engine/Source/Runtime/Core/Private/AutoRTFM/Context.cpp
+++ b/Engine/Source/Runtime/Core/Private/AutoRTFM/Context.cpp
@@ -16,7 +16,7 @@
 #include "TransactionInlines.h"
 #include "Utils.h"
 
-#if PLATFORM_WINDOWS
+#if AUTORTFM_PLATFORM_WINDOWS
 extern "C" __declspec(dllimport) void __stdcall GetCurrentThreadStackLimits(void**, void**);
 #endif
 
@@ -399,7 +399,7 @@ AutoRTFM::FStackRange FContext::GetThreadStackRange()
 	{
 		FStackRange Stack;
 
-#if PLATFORM_WINDOWS
+#if AUTORTFM_PLATFORM_WINDOWS
 		GetCurrentThreadStackLimits(&Stack.Low, &Stack.High);
 #elif defined(__APPLE__)         
 		Stack.High = pthread_get_stackaddr_np(pthread_self());
diff --git a/Engine/Source/Runtime/Core/Private/AutoRTFM/ContextInlines.h b/Engine/Source/Runtime/Core/Private/AutoRTFM/ContextInlines.h
--- a/Engine/Source/Runtime/Core/Private/AutoRTFM/ContextInlines.h
+++ b/Engine/Source/Runtime/Core/Private/AutoRTFM/ContextInlines.h
@@ -32,7 +32,7 @@ template<unsigned SIZE> AUTORTFM_NO_ASAN UE_AUTORTFM_FORCEINLINE void FContext::
 	}
 }
 
-template<unsigned SIZE> AUTORTFM_NO_ASAN FORCENOINLINE void FContext::RecordWriteSlow(void* LogicalAddress)
+template<unsigned SIZE> AUTORTFM_NO_ASAN UE_AUTORTFM_FORCENOINLINE void FContext::RecordWriteSlow(void* LogicalAddress)
 {
 	// Going through get current transaction will do the materialization of deferred transactions, but it is slow
 	// so that is why we use this carve out.
diff --git a/Engine/Source/Runtime/Core/Private/AutoRTFM/FunctionMap.cpp b/Engine/Source/Runtime/Core/Private/AutoRTFM/FunctionMap.cpp
--- a/Engine/Source/Runtime/Core/Private/AutoRTFM/FunctionMap.cpp
+++ b/Engine/Source/Runtime/Core/Private/AutoRTFM/FunctionMap.cpp
@@ -5,6 +5,7 @@
 #include "FunctionMap.h"
 
 #include "Atomic.h"
+#include "BuildMacros.h"
 #include "FastLock.h"
 #include "FunctionMapInlines.h"
 #include "LockGuard.h"
@@ -22,7 +23,7 @@ namespace AutoRTFM
 // be resolved (or is not an import thunk).
 static inline void* FollowRelocation(void* Function)
 {
-#if PLATFORM_WINDOWS
+#if AUTORTFM_PLATFORM_WINDOWS
 #if defined(_x86_64) || defined(__x86_64__) || defined(_M_X64) || defined(_M_AMD64)
 	// Note: Windows has multiple ways to perform relocations:
 	// https://learn.microsoft.com/en-us/windows/win32/debug/pe-format#base-relocation-types
@@ -40,7 +41,7 @@ static inline void* FollowRelocation(void* Function)
 		return Ptr;
 	}
 #endif  // is-x64
-#endif  // PLATFORM_WINDOWS
+#endif  // AUTORTFM_PLATFORM_WINDOWS
 	return nullptr;
 }
 
diff --git a/Engine/Source/Runtime/Core/Private/AutoRTFM/HitSet.h b/Engine/Source/Runtime/Core/Private/AutoRTFM/HitSet.h
--- a/Engine/Source/Runtime/Core/Private/AutoRTFM/HitSet.h
+++ b/Engine/Source/Runtime/Core/Private/AutoRTFM/HitSet.h
@@ -136,7 +136,7 @@ namespace AutoRTFM
 
 		UE_AUTORTFM_FORCEINLINE uint64_t Capacity() const
 		{
-			UE_ASSUME(SixtyFourMinusLogCapacity < 64);
+			UE_AUTORTFM_ASSUME(SixtyFourMinusLogCapacity < 64);
 			return static_cast<uint64_t>(1) << (64 - SixtyFourMinusLogCapacity);
 		}
 
diff --git a/Engine/Source/Runtime/Core/Private/AutoRTFM/IntervalTree.h b/Engine/Source/Runtime/Core/Private/AutoRTFM/IntervalTree.h
--- a/Engine/Source/Runtime/Core/Private/AutoRTFM/IntervalTree.h
+++ b/Engine/Source/Runtime/Core/Private/AutoRTFM/IntervalTree.h
@@ -19,13 +19,13 @@ struct FIntervalTree final
     FIntervalTree(const FIntervalTree&) = delete;
     FIntervalTree& operator=(const FIntervalTree&) = delete;
 
-	FORCENOINLINE bool Insert(const void* const Address, const size_t Size)
+	UE_AUTORTFM_FORCENOINLINE bool Insert(const void* const Address, const size_t Size)
     {        
         FRange NewRange(Address, Size);
         return Insert(std::move(NewRange));
     }
 
-    FORCENOINLINE bool Contains(const void* const Address, const size_t Size) const
+    UE_AUTORTFM_FORCENOINLINE bool Contains(const void* const Address, const size_t Size) const
     {
         const FRange NewRange(Address, Size);
 
@@ -138,7 +138,7 @@ private:
 
 	static constexpr bool bExtraDebugging = false;
 
-	FORCEINLINE FIntervalTreeNodeIndex AddNode(FRange&& NewRange)
+	UE_AUTORTFM_FORCEINLINE FIntervalTreeNodeIndex AddNode(FRange&& NewRange)
 	{
 		const FIntervalTreeNodeIndex Index = NodeRanges.Num();
         NodeRanges.Push(NewRange);
@@ -149,7 +149,7 @@ private:
 		return Index;
 	}
 
-	FORCEINLINE FIntervalTreeNodeIndex AddNode(FRange&& NewRange, FIntervalTreeNodeIndex Parent)
+	UE_AUTORTFM_FORCEINLINE FIntervalTreeNodeIndex AddNode(FRange&& NewRange, FIntervalTreeNodeIndex Parent)
 	{
 		const FIntervalTreeNodeIndex Index = NodeRanges.Num();
         NodeRanges.Push(NewRange);
@@ -160,7 +160,7 @@ private:
 		return Index;
 	}
 
-    FORCENOINLINE bool Insert(FRange&& NewRange)
+    UE_AUTORTFM_FORCENOINLINE bool Insert(FRange&& NewRange)
     {
         if (AUTORTFM_UNLIKELY(IntervalTreeNodeIndexNone == Root))
         {
@@ -233,7 +233,7 @@ private:
         {
             FIntervalTreeNodeIndex Parent = NodeParents[Current];
 
-            UE_ASSUME(Current != Parent);
+            UE_AUTORTFM_ASSUME(Current != Parent);
 
             // The root will always have a black parent, so this check covers both.
             if (Parent == Root)
@@ -248,14 +248,14 @@ private:
 
             const FIntervalTreeNodeIndex GrandParent = NodeParents[Parent];
 
-            UE_ASSUME((Parent != GrandParent) && (Current != GrandParent));
+            UE_AUTORTFM_ASSUME((Parent != GrandParent) && (Current != GrandParent));
 
             const bool bParentIsLeft = NodeLefts[GrandParent] == Parent;
 
             // The uncle is the other node of our parent.
             const FIntervalTreeNodeIndex Uncle = bParentIsLeft ? NodeRights[GrandParent] : NodeLefts[GrandParent];
 
-            UE_ASSUME((GrandParent != Uncle) && (Parent != Uncle) && (Current != Uncle));
+            UE_AUTORTFM_ASSUME((GrandParent != Uncle) && (Parent != Uncle) && (Current != Uncle));
 
             if (!IsBlack(Uncle))
             {
@@ -383,7 +383,7 @@ private:
 		}
 	}
 
-	FORCENOINLINE void AssertNodeIsOk(FIntervalTreeNodeIndex Index) const
+	UE_AUTORTFM_FORCENOINLINE void AssertNodeIsOk(FIntervalTreeNodeIndex Index) const
 	{
 		// We need to use recursion to check because we cannot have any
 		// allocations within the checker itself!
diff --git a/Engine/Source/Runtime/Core/Private/AutoRTFM/LongJump.h b/Engine/Source/Runtime/Core/Private/AutoRTFM/LongJump.h
--- a/Engine/Source/Runtime/Core/Private/AutoRTFM/LongJump.h
+++ b/Engine/Source/Runtime/Core/Private/AutoRTFM/LongJump.h
@@ -4,6 +4,7 @@
 
 #if (defined(__AUTORTFM) && __AUTORTFM)
 
+#include "BuildMacros.h"
 #include "Utils.h"
 
 #include <setjmp.h>
@@ -36,11 +37,11 @@ template<typename TTryFunctor, typename TCatchFunctor>
 void FLongJump::TryCatch(const TTryFunctor& TryFunctor, const TCatchFunctor& CatchFunctor)
 {
     AUTORTFM_ASSERT(!bIsSet);
-#if PLATFORM_WINDOWS
+#if AUTORTFM_PLATFORM_WINDOWS
 	if (!setjmp(JmpBuf))
 #else
 	if (!_setjmp(JmpBuf))
-#endif // PLATFORM_WINDOWS
+#endif // AUTORTFM_PLATFORM_WINDOWS
     {
         bIsSet = true;
         TryFunctor();
@@ -57,11 +58,11 @@ void FLongJump::TryCatch(const TTryFunctor& TryFunctor, const TCatchFunctor& Cat
 inline void FLongJump::Throw()
 {
     AUTORTFM_ASSERT(bIsSet);
-#if PLATFORM_WINDOWS
+#if AUTORTFM_PLATFORM_WINDOWS
 	longjmp(JmpBuf, 1);
 #else
 	_longjmp(JmpBuf, 1);
-#endif // PLATFORM_WINDOWS
+#endif // AUTORTFM_PLATFORM_WINDOWS
 }
 
 } // namespace AutoRTFM
diff --git a/Engine/Source/Runtime/Core/Private/AutoRTFM/StdLibSupport.cpp b/Engine/Source/Runtime/Core/Private/AutoRTFM/StdLibSupport.cpp
--- a/Engine/Source/Runtime/Core/Private/AutoRTFM/StdLibSupport.cpp
+++ b/Engine/Source/Runtime/Core/Private/AutoRTFM/StdLibSupport.cpp
@@ -3,9 +3,10 @@
 #if (defined(__AUTORTFM) && __AUTORTFM)
 
 #include "AutoRTFM/AutoRTFM.h"
+#include "BuildMacros.h"
 #include "ContextInlines.h"
-#include "Utils.h"
 #include "Memcpy.h"
+#include "Utils.h"
 
 #include <algorithm>
 #include <charconv>
@@ -22,11 +23,11 @@
 #include <malloc.h>
 #endif
 
-#if PLATFORM_WINDOWS
-#include "Windows/WindowsHWrapper.h"
+#if AUTORTFM_PLATFORM_WINDOWS
+#include "Windows/WindowsHWrapper.h"  // TODO(SOL-7652): change to "WindowsHeader.h"
 #endif
 
-#if PLATFORM_LINUX
+#if AUTORTFM_PLATFORM_LINUX
 #include <sys/stat.h>
 #include <sys/time.h>
 #include <sys/types.h>
@@ -122,8 +123,8 @@ static void ThrowIfFormatContainsPercentN(const wchar_t* Format)
     }
 }
 
-#if PLATFORM_WINDOWS
-// vvvvvvvvvvvvvvvvvvvvvvvvvvvvvv PLATFORM_WINDOWS vvvvvvvvvvvvvvvvvvvvvvvvvvvvv
+#if AUTORTFM_PLATFORM_WINDOWS
+// vvvvvvvvvvvvvvvvvvvvvvvvvvvvvv AUTORTFM_PLATFORM_WINDOWS vvvvvvvvvvvvvvvvvvvvvvvvvvvvv
 template<typename RETURN, typename CHAR, RETURN FN(const CHAR*, CHAR**, _locale_t)>
 RETURN RTFM_StringToFloat(const CHAR* String, CHAR** EndPtr, _locale_t Locale)
 {
@@ -267,15 +268,15 @@ BOOL RTFM_TlsSetValue(DWORD dwTlsIndex, LPVOID lpTlsValue)
 	return TlsSetValue(dwTlsIndex, lpTlsValue);
 }
 
-// ^^^^^^^^^^^^^^^^^^^^^^^^^^^^^^ PLATFORM_WINDOWS ^^^^^^^^^^^^^^^^^^^^^^^^^^^^^
+// ^^^^^^^^^^^^^^^^^^^^^^^^^^^^^^ AUTORTFM_PLATFORM_WINDOWS ^^^^^^^^^^^^^^^^^^^^^^^^^^^^^
 #else
-// vvvvvvvvvvvvvvvvvvvvvvvvvvvvv !PLATFORM_WINDOWS vvvvvvvvvvvvvvvvvvvvvvvvvvvvv
+// vvvvvvvvvvvvvvvvvvvvvvvvvvvvv !AUTORTFM_PLATFORM_WINDOWS vvvvvvvvvvvvvvvvvvvvvvvvvvvvv
 extern "C" size_t _ZNSt3__112__next_primeEm(size_t N) __attribute__((weak));
-// ^^^^^^^^^^^^^^^^^^^^^^^^^^^^^ !PLATFORM_WINDOWS ^^^^^^^^^^^^^^^^^^^^^^^^^^^^^
+// ^^^^^^^^^^^^^^^^^^^^^^^^^^^^^ !AUTORTFM_PLATFORM_WINDOWS ^^^^^^^^^^^^^^^^^^^^^^^^^^^^^
 #endif
 
-#if PLATFORM_LINUX
-// vvvvvvvvvvvvvvvvvvvvvvvvvvvvvv PLATFORM_LINUX vvvvvvvvvvvvvvvvvvvvvvvvvvvvvv
+#if AUTORTFM_PLATFORM_LINUX
+// vvvvvvvvvvvvvvvvvvvvvvvvvvvvvv AUTORTFM_PLATFORM_LINUX vvvvvvvvvvvvvvvvvvvvvvvvvvvvvv
 int RTFM_stat(const char* Path, struct stat* StatBuf) throw()
 {
 	FContext* Context = FContext::Get();
@@ -307,7 +308,7 @@ int RTFM___fxstat(int Ver, int Fd, struct stat* StatBuf) throw()
 
 	return __fxstat(Ver, Fd, StatBuf);
 }
-// ^^^^^^^^^^^^^^^^^^^^^^^^^^^^^^ PLATFORM_LINUX ^^^^^^^^^^^^^^^^^^^^^^^^^^^^^^
+// ^^^^^^^^^^^^^^^^^^^^^^^^^^^^^^ AUTORTFM_PLATFORM_LINUX ^^^^^^^^^^^^^^^^^^^^^^^^^^^^^^
 #endif
 
 void* RTFM_memcpy(void* Dst, const void* Src, size_t Size) throw()
@@ -497,7 +498,7 @@ int RTFM_vswprintf(wchar_t* Str, size_t Size, const wchar_t* Format, va_list Arg
         va_list ArgList2;
         va_copy(ArgList2, ArgList);
 
-#if PLATFORM_WINDOWS
+#if AUTORTFM_PLATFORM_WINDOWS
         int Count = vswprintf(nullptr, 0, Format, ArgList2);
 #else
         // vswprintf(nullptr, 0, ...) will return -1.
@@ -604,7 +605,7 @@ UE_AUTORTFM_REGISTER_OPEN_TO_CLOSED_FUNCTIONS(
 	UE_AUTORTFM_MAP_OPEN_TO_SELF(__asan_get_current_fake_stack),
 #endif
 
-#if PLATFORM_WINDOWS
+#if AUTORTFM_PLATFORM_WINDOWS
 	UE_AUTORTFM_MAP_OPEN_TO_CLOSED(_strtoi64, (RTFM_StringToIntNoLocale<long long, char, _strtoi64>)),
 	UE_AUTORTFM_MAP_OPEN_TO_CLOSED(_wcstoi64, (RTFM_StringToInt<long long, wchar_t, _wcstoi64>)),
 	UE_AUTORTFM_MAP_OPEN_TO_CLOSED(_wcstoui64, (RTFM_StringToInt<unsigned long long, wchar_t, _wcstoui64>)),
@@ -655,7 +656,7 @@ UE_AUTORTFM_REGISTER_OPEN_TO_CLOSED_FUNCTIONS(
 	UE_AUTORTFM_MAP_OPEN_TO_SELF(_ZNSt3__112__next_primeEm),
 #endif
 
-#if PLATFORM_LINUX
+#if AUTORTFM_PLATFORM_LINUX
 	UE_AUTORTFM_MAP_OPEN_TO_SELF(getpid),
 	UE_AUTORTFM_MAP_OPEN_TO_SELF(clock_gettime),
 	UE_AUTORTFM_MAP_OPEN_TO_SELF(gettimeofday),
diff --git a/Engine/Source/Runtime/Core/Private/AutoRTFM/Toggles.h b/Engine/Source/Runtime/Core/Private/AutoRTFM/Toggles.h
--- a/Engine/Source/Runtime/Core/Private/AutoRTFM/Toggles.h
+++ b/Engine/Source/Runtime/Core/Private/AutoRTFM/Toggles.h
@@ -2,7 +2,7 @@
 
 #pragma once
 
-#include "Misc/Build.h"
+#include "BuildMacros.h"
 
 namespace AutoRTFM
 {
@@ -11,7 +11,7 @@ namespace AutoRTFM
 static constexpr bool bCollectStats = false;
 
 // Tracking the locations of allocations is helpful for catching invalid frees.
-#if UE_BUILD_SHIPPING
+#if AUTORTFM_BUILD_SHIPPING
 static constexpr bool bTrackAllocationLocations = false;
 #else
 static constexpr bool bTrackAllocationLocations = true;
diff --git a/Engine/Source/Runtime/Core/Private/AutoRTFM/Transaction.cpp b/Engine/Source/Runtime/Core/Private/AutoRTFM/Transaction.cpp
--- a/Engine/Source/Runtime/Core/Private/AutoRTFM/Transaction.cpp
+++ b/Engine/Source/Runtime/Core/Private/AutoRTFM/Transaction.cpp
@@ -307,7 +307,7 @@ void FTransaction::SetOpenActiveValidatorEnabled(EMemoryValidationLevel NewMemor
 
 void FTransaction::SetOpenActive(EMemoryValidationLevel NewMemoryValidationLevel, const void* ReturnAddress)
 {
-	if (UNLIKELY(NewMemoryValidationLevel != EMemoryValidationLevel::Disabled))
+	if (AUTORTFM_UNLIKELY(NewMemoryValidationLevel != EMemoryValidationLevel::Disabled))
 	{
 		AUTORTFM_MUST_TAIL return SetOpenActiveValidatorEnabled(NewMemoryValidationLevel, ReturnAddress);
 	}
@@ -450,7 +450,7 @@ void FTransaction::DebugBreakIfMemoryValidationFails()
 		if (OldHash != NewHash)
 		{
 			AUTORTFM_WARN("DebugBreakIfInvalidMemoryHash() detected a change in hash");
-			PLATFORM_BREAK();
+			__builtin_debugtrap();
 		}
 	}
 }
diff --git a/Engine/Source/Runtime/Core/Private/AutoRTFM/TransactionInlines.h b/Engine/Source/Runtime/Core/Private/AutoRTFM/TransactionInlines.h
--- a/Engine/Source/Runtime/Core/Private/AutoRTFM/TransactionInlines.h
+++ b/Engine/Source/Runtime/Core/Private/AutoRTFM/TransactionInlines.h
@@ -180,7 +180,7 @@ template<unsigned SIZE> AUTORTFM_NO_ASAN UE_AUTORTFM_FORCEINLINE void FTransacti
 	}
 }
 
-template<unsigned SIZE> AUTORTFM_NO_ASAN FORCENOINLINE void FTransaction::RecordWriteNotInsertedSlow(void* LogicalAddress)
+template<unsigned SIZE> AUTORTFM_NO_ASAN UE_AUTORTFM_FORCENOINLINE void FTransaction::RecordWriteNotInsertedSlow(void* LogicalAddress)
 {
     FHitSetEntry Entry{};
     Entry.Address = reinterpret_cast<uintptr_t>(LogicalAddress);
@@ -197,7 +197,7 @@ template<unsigned SIZE> AUTORTFM_NO_ASAN FORCENOINLINE void FTransaction::Record
 	return RecordWriteInsertedSlow<SIZE>(LogicalAddress);
 }
 
-template<unsigned SIZE> AUTORTFM_NO_ASAN FORCENOINLINE void FTransaction::RecordWriteInsertedSlow(void* LogicalAddress)
+template<unsigned SIZE> AUTORTFM_NO_ASAN UE_AUTORTFM_FORCENOINLINE void FTransaction::RecordWriteInsertedSlow(void* LogicalAddress)
 {
 	if (NewMemoryTracker.Contains(LogicalAddress, SIZE))
 	{
diff --git a/Engine/Source/Runtime/Core/Private/AutoRTFM/Utils.cpp b/Engine/Source/Runtime/Core/Private/AutoRTFM/Utils.cpp
--- a/Engine/Source/Runtime/Core/Private/AutoRTFM/Utils.cpp
+++ b/Engine/Source/Runtime/Core/Private/AutoRTFM/Utils.cpp
@@ -5,11 +5,12 @@
 #include "Utils.h"
 
 #include "ContextInlines.h"
+#include "BuildMacros.h"
 
 #include <utility>
 
-#if PLATFORM_WINDOWS
-#include "Windows/WindowsHWrapper.h"
+#if AUTORTFM_PLATFORM_WINDOWS
+#include "Windows/WindowsHWrapper.h"  // TODO(SOL-7652): change to "WindowsHeader.h"
 #include <dbghelp.h>
 #else
 #include <execinfo.h>
@@ -21,7 +22,7 @@ namespace AutoRTFM
 
 std::string GetFunctionDescription(void* FunctionPtr)
 {
-#if PLATFORM_WINDOWS
+#if AUTORTFM_PLATFORM_WINDOWS
     // This is gross, but it works. It's possible for someone to have SymInitialized before. But if they had, then this
     // will just fail. Also, this function is called in cases where we're failing, so it's ok if we do dirty things.
     SymInitialize(GetCurrentProcess(), nullptr, true);
@@ -40,15 +41,15 @@ std::string GetFunctionDescription(void* FunctionPtr)
     {
         return "<error getting description>";
     }
-#else // PLATFORM_WINDOWS -> so !PLATFORM_WINDOWS
+#else // AUTORTFM_PLATFORM_WINDOWS -> so !AUTORTFM_PLATFORM_WINDOWS
     char** const symbols = backtrace_symbols(&FunctionPtr, 1);
     std::string Name(*symbols);
     free(symbols);
     return Name;
-#endif // PLATFORM_WINDOWS -> so !PLATFORM_WINDOWS
+#endif // AUTORTFM_PLATFORM_WINDOWS -> so !AUTORTFM_PLATFORM_WINDOWS
 }
 
-UE_AUTORTFM_API FORCENOINLINE void ReportError(const char* Format, ...)
+UE_AUTORTFM_API UE_AUTORTFM_FORCENOINLINE void ReportError(const char* Format, ...)
 {
 	FContext* const Context = FContext::Get();
 
diff --git a/Engine/Source/Runtime/Core/Private/AutoRTFM/Utils.h b/Engine/Source/Runtime/Core/Private/AutoRTFM/Utils.h
--- a/Engine/Source/Runtime/Core/Private/AutoRTFM/Utils.h
+++ b/Engine/Source/Runtime/Core/Private/AutoRTFM/Utils.h
@@ -38,7 +38,7 @@ std::string GetFunctionDescription(TReturnType (*FunctionPtr)(TParameterTypes...
 // fatal process destroying explosions happen here).
 [[clang::preserve_most]]
 [[noreturn]]
-FORCENOINLINE inline void DoAssert(void(*Logger)()) //-V1082 Silence PVS [[noreturn]] false positive warning
+UE_AUTORTFM_FORCENOINLINE inline void DoAssert(void(*Logger)()) //-V1082 Silence PVS [[noreturn]] false positive warning
 {
 	Logger();
 	__builtin_unreachable();
@@ -49,7 +49,7 @@ FORCENOINLINE inline void DoAssert(void(*Logger)()) //-V1082 Silence PVS [[noret
 // be called because it is the expecting true codepath, EG. only
 // things we don't expect to happen get here).
 [[clang::preserve_most]]
-FORCENOINLINE inline void DoExpect(void(*Logger)())
+UE_AUTORTFM_FORCENOINLINE inline void DoExpect(void(*Logger)())
 {
 	Logger();
 }
diff --git a/Engine/Source/Runtime/Core/Private/AutoRTFM/WriteLog.cpp b/Engine/Source/Runtime/Core/Private/AutoRTFM/WriteLog.cpp
--- a/Engine/Source/Runtime/Core/Private/AutoRTFM/WriteLog.cpp
+++ b/Engine/Source/Runtime/Core/Private/AutoRTFM/WriteLog.cpp
@@ -4,12 +4,13 @@
 
 #include "WriteLog.h"
 
+#include "BuildMacros.h"
 #include "Utils.h"
 
 namespace AutoRTFM
 {
 
-#if PLATFORM_CPU_X86_FAMILY
+#if defined(__x86_64__) || defined(_M_X64)
 #define WRITELOG_HASH_MAYBE_HAS_AVX2 1
 #include <immintrin.h>
 #else
@@ -48,7 +49,7 @@ bool SmallHash(const std::byte* Data, size_t Size, FWriteLog::FHash& Hash)
 		case 2:
 			if (AUTORTFM_LIKELY((reinterpret_cast<uintptr_t>(Data) & 1) == 0))
 			{
-				uint16 Load = *reinterpret_cast<uint16 const*>(Data);
+				uint16_t Load = *reinterpret_cast<uint16_t const*>(Data);
 				Hash = (Hash * 31) ^ static_cast<FWriteLog::FHash>(Load);
 				return true;
 			}
diff --git a/Engine/Source/Runtime/Core/Private/AutoRTFM/WriteLog.h b/Engine/Source/Runtime/Core/Private/AutoRTFM/WriteLog.h
--- a/Engine/Source/Runtime/Core/Private/AutoRTFM/WriteLog.h
+++ b/Engine/Source/Runtime/Core/Private/AutoRTFM/WriteLog.h
@@ -5,6 +5,7 @@
 #if (defined(__AUTORTFM) && __AUTORTFM)
 
 #include "AutoRTFM/AutoRTFM.h"
+#include "BuildMacros.h"
 #include "Utils.h"
 
 #include <cstddef>
@@ -168,7 +169,7 @@ namespace AutoRTFM
 
 				DataEnd += Entry.Size;
 
-#if UE_BUILD_DEBUG
+#if AUTORTFM_BUILD_DEBUG
 				AUTORTFM_ASSERT(DataEnd <= reinterpret_cast<std::byte*>(LastRecord));
 #endif
 
@@ -214,7 +215,7 @@ namespace AutoRTFM
 			{
 				NumEntries++;
 			}
-			else if (UNLIKELY(PushResult == FBlock::EPushResult::Full))
+			else if (AUTORTFM_UNLIKELY(PushResult == FBlock::EPushResult::Full))
 			{
 				const size_t RequiredSize = AlignUp(sizeof(FBlock) + Entry.Size, alignof(FRecord)) + sizeof(FRecord);
 				FBlock* NewBlock = FBlock::Allocate(std::max(RequiredSize, FBlock::DefaultSize));
diff --git a/Engine/Source/Runtime/Core/Public/AutoRTFM/AutoRTFM.h b/Engine/Source/Runtime/Core/Public/AutoRTFM/AutoRTFM.h
--- a/Engine/Source/Runtime/Core/Public/AutoRTFM/AutoRTFM.h
+++ b/Engine/Source/Runtime/Core/Public/AutoRTFM/AutoRTFM.h
@@ -68,16 +68,48 @@
 #define UE_AUTORTFM_ASSUME_SAFE
 #endif
 
+#define UE_AUTORTFM_CONCAT_IMPL(A, B) A ## B
+#define UE_AUTORTFM_CONCAT(A, B) UE_AUTORTFM_CONCAT_IMPL(A, B)
+
+#if defined(PLATFORM_MICROSOFT) && PLATFORM_MICROSOFT
+#define AUTORTFM_DLLEXPORT __declspec(dllexport)
+#define AUTORTFM_DLLIMPORT __declspec(dllimport)
+#else
+#define AUTORTFM_DLLEXPORT __attribute__((visibility("default")))
+#define AUTORTFM_DLLIMPORT __attribute__((visibility("default")))
+#endif
+
 #if defined(UE_AUTORTFM_STANDALONE)
 #define UE_AUTORTFM_API
-#define UE_AUTORTFM_FORCEINLINE inline
 #else
-#include <HAL/Platform.h>
-#define UE_AUTORTFM_API CORE_API
-#define UE_AUTORTFM_FORCEINLINE FORCEINLINE
+#include <HAL/Platform.h>  // TODO(SOL-7652): Remove
+// CORE_API will be defined to "", "DLLEXPORT" or "DLLIMPORT".
+// AutoRTFM cannot depend on Core definitions for DLLEXPORT or DLLIMPORT, and it
+// would be risky to have AutoRTFM provide its own definitions for these as they
+// will leak into code that includes this header, so use macro concatenation to
+// map UE_AUTORTFM_API to "", "AUTORTFM_DLLEXPORT", "AUTORTFM_DLLIMPORT",
+// respectively.
+#define UE_AUTORTFM_PUBLIC_API_
+#define UE_AUTORTFM_PUBLIC_API_DLLEXPORT AUTORTFM_DLLEXPORT
+#define UE_AUTORTFM_PUBLIC_API_DLLIMPORT AUTORTFM_DLLIMPORT
+#define UE_AUTORTFM_API UE_AUTORTFM_CONCAT(UE_AUTORTFM_PUBLIC_API_, CORE_API)
+#endif
+
+#if defined(_MSC_VER)
+#define UE_AUTORTFM_FORCEINLINE __forceinline
+#define UE_AUTORTFM_FORCENOINLINE __declspec(noinline)
+#define UE_AUTORTFM_ASSUME(x) __assume(x)
+#elif defined(__clang__)
+#define UE_AUTORTFM_FORCEINLINE __attribute__((always_inline)) inline
+#define UE_AUTORTFM_FORCENOINLINE __attribute__((noinline))
+#define UE_AUTORTFM_ASSUME(x) __builtin_assume(x)
+#else
+#define UE_AUTORTFM_FORCEINLINE inline
+#define UE_AUTORTFM_FORCENOINLINE
+#define UE_AUTORTFM_ASSUME(x)
 #endif
 
-#if defined(UE_BUILD_DEBUG) && UE_BUILD_DEBUG
+#if (defined(UE_BUILD_DEBUG) && UE_BUILD_DEBUG) || (defined(AUTORTFM_BUILD_DEBUG) && AUTORTFM_BUILD_DEBUG)
 // Force-inlining can make debugging glitchy. Disable this if we're running a debug build.
 #undef UE_AUTORTFM_CALLSITE_FORCEINLINE
 #define UE_AUTORTFM_CALLSITE_FORCEINLINE
@@ -228,6 +260,7 @@ UE_AUTORTFM_API void autortfm_initialize(const autortfm_extern_api* ExternAPI);
 #else
 static UE_AUTORTFM_FORCEINLINE void autortfm_initialize(const autortfm_extern_api* ExternAPI)
 {
+	UE_AUTORTFM_UNUSED(ExternAPI);
 }
 #endif
 
@@ -422,6 +455,7 @@ UE_AUTORTFM_API bool autortfm_is_on_current_transaction_stack(void* Ptr);
 #else
 static UE_AUTORTFM_FORCEINLINE bool autortfm_is_on_current_transaction_stack(void* Ptr)
 {
+	UE_AUTORTFM_UNUSED(Ptr);
 	return false;
 }
 #endif
@@ -1072,7 +1106,10 @@ static UE_AUTORTFM_FORCEINLINE void PopOnAbortHandler(const void* Key)
 }
 #else
 // Unregister the most recently pushed handler (via `PushOnAbortHandler`) for the given key.
-static UE_AUTORTFM_FORCEINLINE void PopOnAbortHandler(const void* Key) {}
+static UE_AUTORTFM_FORCEINLINE void PopOnAbortHandler(const void* Key)
+{
+	UE_AUTORTFM_UNUSED(Key);
+}
 #endif
 
 #if UE_AUTORTFM
@@ -1086,7 +1123,10 @@ static UE_AUTORTFM_FORCEINLINE void PopAllOnAbortHandlers(const void* Key)
 }
 #else
 // Unregister all pushed handlers (via `PushOnAbortHandler`) for the given key.
-static UE_AUTORTFM_FORCEINLINE void PopAllOnAbortHandlers(const void* Key) {}
+static UE_AUTORTFM_FORCEINLINE void PopAllOnAbortHandlers(const void* Key)
+{
+	UE_AUTORTFM_UNUSED(Key);
+}
 #endif
 
 // Inform the runtime that we have performed a new object allocation. It's only
@@ -1560,9 +1600,6 @@ namespace AutoRTFM::Private
 // Calls should be written like this: UE_AUTORTFM_TRANSACT { ... code ... };
 #define UE_AUTORTFM_TRANSACT  UE_AUTORTFM_TRANSACT_IMPL
 
-#define UE_AUTORTFM_CONCAT_IMPL(A, B) A ## B
-#define UE_AUTORTFM_CONCAT(A, B) UE_AUTORTFM_CONCAT_IMPL(A, B)
-
 #if UE_AUTORTFM
 #define UE_AUTORTFM_REGISTER_OPEN_TO_CLOSED_FUNCTIONS(...) \
 	static const ::AutoRTFM::ForTheRuntime::TAutoRegisterOpenToClosedFunctions<__VA_ARGS__> UE_AUTORTFM_CONCAT(AutoRTFMFunctionRegistration, __COUNTER__)
