Remove `UE_AUTORTFM_ALWAYS_OPEN` and handle the transactional codepath manually to stop the perf regression in the CsvProfiler.

Using the `UE_AUTORTFM_ALWAYS_OPEN` attribute causes the function to be marked as no-inline too (a quirk that requires much compiler fun to fix). This was causing a perf regression noticed by the perfmem team.

#rnx

diff --git a/Engine/Source/Runtime/Core/Private/ProfilingDebugging/CsvProfiler.cpp b/Engine/Source/Runtime/Core/Private/ProfilingDebugging/CsvProfiler.cpp
--- a/Engine/Source/Runtime/Core/Private/ProfilingDebugging/CsvProfiler.cpp
+++ b/Engine/Source/Runtime/Core/Private/ProfilingDebugging/CsvProfiler.cpp
@@ -4410,7 +4410,10 @@ void FCsvProfiler::RecordCustomStat(const char * StatName, uint32 CategoryIndex,
 #if !CSV_PROFILER_MINIMAL
 	if (GCsvProfilerIsCapturing && GCsvCategoriesEnabled[CategoryIndex])
 	{
-		FCsvProfilerThreadData::Get().AddCustomStat(StatName, CategoryIndex, Value, CustomStatOp);
+		UE_AUTORTFM_OPEN
+		{
+			FCsvProfilerThreadData::Get().AddCustomStat(StatName, CategoryIndex, Value, CustomStatOp);
+		};
 	}
 #endif
 }
@@ -4428,8 +4431,11 @@ void FCsvProfiler::RecordCustomStat(const FName& StatName, uint32 CategoryIndex,
 #if !CSV_PROFILER_MINIMAL
 	if (GCsvProfilerIsCapturing && GCsvCategoriesEnabled[CategoryIndex])
 	{
-		ValidateFName(StatName);
-		FCsvProfilerThreadData::Get().AddCustomStat(StatName, CategoryIndex, Value, CustomStatOp);
+		UE_AUTORTFM_OPEN
+		{
+			ValidateFName(StatName);
+			FCsvProfilerThreadData::Get().AddCustomStat(StatName, CategoryIndex, Value, CustomStatOp);
+		};
 	}
 #endif
 }
@@ -4448,7 +4454,10 @@ void FCsvProfiler::RecordCustomStat(const char * StatName, uint32 CategoryIndex,
 #if !CSV_PROFILER_MINIMAL
 	if (GCsvProfilerIsCapturing && GCsvCategoriesEnabled[CategoryIndex])
 	{
-		FCsvProfilerThreadData::Get().AddCustomStat(StatName, CategoryIndex, Value, CustomStatOp);
+		UE_AUTORTFM_OPEN
+		{
+			FCsvProfilerThreadData::Get().AddCustomStat(StatName, CategoryIndex, Value, CustomStatOp);
+		};
 	}
 #endif
 }
@@ -4458,8 +4467,11 @@ void FCsvProfiler::RecordCustomStat(const FName& StatName, uint32 CategoryIndex,
 #if !CSV_PROFILER_MINIMAL
 	if (GCsvProfilerIsCapturing && GCsvCategoriesEnabled[CategoryIndex])
 	{
-		ValidateFName(StatName);
-		FCsvProfilerThreadData::Get().AddCustomStat(StatName, CategoryIndex, Value, CustomStatOp);
+		UE_AUTORTFM_OPEN
+		{
+			ValidateFName(StatName);
+			FCsvProfilerThreadData::Get().AddCustomStat(StatName, CategoryIndex, Value, CustomStatOp);
+		};
 	}
 #endif
 }
diff --git a/Engine/Source/Runtime/Core/Public/ProfilingDebugging/CsvProfiler.h b/Engine/Source/Runtime/Core/Public/ProfilingDebugging/CsvProfiler.h
--- a/Engine/Source/Runtime/Core/Public/ProfilingDebugging/CsvProfiler.h
+++ b/Engine/Source/Runtime/Core/Public/ProfilingDebugging/CsvProfiler.h
@@ -380,12 +380,12 @@ public:
 	CORE_API static void BeginExclusiveStat(const char * StatName, const char* NamedEventName = nullptr);
 	CORE_API static void EndExclusiveStat(const char * StatName);
 
-	UE_AUTORTFM_ALWAYS_OPEN CORE_API static void RecordCustomStat(const char * StatName, uint32 CategoryIndex, float Value, const ECsvCustomStatOp CustomStatOp);
-	UE_AUTORTFM_ALWAYS_OPEN CORE_API static void RecordCustomStat(const FName& StatName, uint32 CategoryIndex, float Value, const ECsvCustomStatOp CustomStatOp);
-	UE_AUTORTFM_ALWAYS_OPEN CORE_API static void RecordCustomStat(const char* StatName, uint32 CategoryIndex, double Value, const ECsvCustomStatOp CustomStatOp);
-	UE_AUTORTFM_ALWAYS_OPEN CORE_API static void RecordCustomStat(const FName& StatName, uint32 CategoryIndex, double Value, const ECsvCustomStatOp CustomStatOp);
-	UE_AUTORTFM_ALWAYS_OPEN CORE_API static void RecordCustomStat(const char * StatName, uint32 CategoryIndex, int32 Value, const ECsvCustomStatOp CustomStatOp);
-	UE_AUTORTFM_ALWAYS_OPEN CORE_API static void RecordCustomStat(const FName& StatName, uint32 CategoryIndex, int32 Value, const ECsvCustomStatOp CustomStatOp);
+	CORE_API static void RecordCustomStat(const char * StatName, uint32 CategoryIndex, float Value, const ECsvCustomStatOp CustomStatOp);
+	CORE_API static void RecordCustomStat(const FName& StatName, uint32 CategoryIndex, float Value, const ECsvCustomStatOp CustomStatOp);
+	CORE_API static void RecordCustomStat(const char* StatName, uint32 CategoryIndex, double Value, const ECsvCustomStatOp CustomStatOp);
+	CORE_API static void RecordCustomStat(const FName& StatName, uint32 CategoryIndex, double Value, const ECsvCustomStatOp CustomStatOp);
+	CORE_API static void RecordCustomStat(const char * StatName, uint32 CategoryIndex, int32 Value, const ECsvCustomStatOp CustomStatOp);
+	CORE_API static void RecordCustomStat(const FName& StatName, uint32 CategoryIndex, int32 Value, const ECsvCustomStatOp CustomStatOp);
 
 	CORE_API static void RecordEndOfPipeCustomStat(const FName& StatName, uint32 CategoryIndex, double Value, const ECsvCustomStatOp CustomStatOp);
 
