A single implementation of Concatenate supports invocations with more than two parameters, and removes the limit of five parameters.
TransformCast returns a reference, rather than a copy, when TrasformConverter returns a reference.

diff --git a/Engine/Source/Runtime/Core/Public/Math/TransformCalculus.h b/Engine/Source/Runtime/Core/Public/Math/TransformCalculus.h
--- a/Engine/Source/Runtime/Core/Public/Math/TransformCalculus.h
+++ b/Engine/Source/Runtime/Core/Public/Math/TransformCalculus.h
@@ -259,7 +259,7 @@ struct TransformConverter
  * Return type uses decltype to support classes that can return by const-ref more efficiently than returning a new value.
  */
 template<typename ResultType, typename TransformType>
-[[nodiscard]] inline auto TransformCast(const TransformType& Transform) -> decltype(TransformConverter<ResultType>::Convert(Transform))
+[[nodiscard]] inline decltype(auto) TransformCast(const TransformType& Transform)
 {
 	return TransformConverter<ResultType>::Convert(Transform);
 }
@@ -324,51 +324,18 @@ template<typename TransformType>
 }
 
 /**
- * Concatenates three transforms. Uses two-argument Concatenate to do its work,
+ * Concatenates three or more transforms. Uses two-argument Concatenate to do its work,
  * and infers the return type using decltype.
  * 
  * @param TransformAToB Transformation that goes from space A to space B
  * @param TransformBToC Transformation that goes from space B to space C.
- * @param TransformCToD Transformation that goes from space C to space D.
- * @return a new Transform representing the transformation from space A to space D.
+ * @param TransformCToRest Transformation that goes from space C onwards.
+ * @return a new Transform representing the complete transformation.
  */
-template<typename TransformType1, typename TransformType2, typename TransformType3>
-[[nodiscard]] inline auto Concatenate(const TransformType1& TransformAToB, const TransformType2& TransformBToC, const TransformType3& TransformCToD) -> decltype(Concatenate(Concatenate(TransformAToB, TransformBToC), TransformCToD)) 
+template<typename TransformType1, typename TransformType2, typename... TransformTypeRest>
+[[nodiscard]] inline decltype(auto) Concatenate(const TransformType1& TransformAToB, const TransformType2& TransformBToC, TransformTypeRest&&... TransformCToRest)
 {
-	return Concatenate(Concatenate(TransformAToB, TransformBToC), TransformCToD);
-}
-
-/**
- * Concatenates four transforms. Uses two-argument Concatenate to do its work,
- * and infers the return type using decltype.
- * 
- * @param TransformAToB Transformation that goes from space A to space B
- * @param TransformBToC Transformation that goes from space B to space C.
- * @param TransformCToD Transformation that goes from space C to space D.
- * @param TransformDToE Transformation that goes from space D to space E.
- * @return a new Transform representing the transformation from space A to space E.
- */
-template<typename TransformType1, typename TransformType2, typename TransformType3, typename TransformType4>
-[[nodiscard]] inline auto Concatenate(const TransformType1& TransformAToB, const TransformType2& TransformBToC, const TransformType3& TransformCToD, const TransformType4& TransformDToE) -> decltype(Concatenate(Concatenate(TransformAToB, TransformBToC, TransformCToD), TransformDToE)) 
-{
-	return Concatenate(Concatenate(TransformAToB, TransformBToC, TransformCToD), TransformDToE);
-}
-
-/**
- * Concatenates five transforms. Uses two-argument Concatenate to do its work,
- * and infers the return type using decltype.
- * 
- * @param TransformAToB Transformation that goes from space A to space B
- * @param TransformBToC Transformation that goes from space B to space C.
- * @param TransformCToD Transformation that goes from space C to space D.
- * @param TransformDToE Transformation that goes from space D to space E.
- * @param TransformEToF Transformation that goes from space E to space F.
- * @return a new Transform representing the transformation from space A to space F.
- */
-template<typename TransformType1, typename TransformType2, typename TransformType3, typename TransformType4, typename TransformType5>
-[[nodiscard]] inline auto Concatenate(const TransformType1& TransformAToB, const TransformType2& TransformBToC, const TransformType3& TransformCToD, const TransformType4& TransformDToE, const TransformType5& TransformEToF) -> decltype(Concatenate(Concatenate(TransformAToB, TransformBToC, TransformCToD, TransformDToE), TransformEToF)) 
-{
-	return Concatenate(Concatenate(TransformAToB, TransformBToC, TransformCToD, TransformDToE), TransformEToF);
+	return Concatenate(Concatenate(TransformAToB, TransformBToC), Forward<TransformTypeRest>(TransformCToRest)...);
 }
 
 /**
