Add a define to allow disabling -basedir and -BaseFromWorkingDir.

diff --git a/Engine/Source/Runtime/Core/Private/Windows/WindowsPlatformProcess.cpp b/Engine/Source/Runtime/Core/Private/Windows/WindowsPlatformProcess.cpp
--- a/Engine/Source/Runtime/Core/Private/Windows/WindowsPlatformProcess.cpp
+++ b/Engine/Source/Runtime/Core/Private/Windows/WindowsPlatformProcess.cpp
@@ -44,6 +44,10 @@
 
 #pragma comment(lib, "psapi.lib")
 
+#ifndef WINDOWS_ALLOW_BASEDIR_OVERRIDE
+#define WINDOWS_ALLOW_BASEDIR_OVERRIDE 1
+#endif
+
 PRAGMA_DISABLE_UNSAFE_TYPECAST_WARNINGS
 
 // static variables
@@ -1162,6 +1166,7 @@ const TCHAR* FWindowsPlatformProcess::BaseDir()
 	static TCHAR Result[512]=TEXT("");
 	if( !Result[0] )
 	{
+#if WINDOWS_ALLOW_BASEDIR_OVERRIDE
 		// Normally the BaseDir is determined from the path of the running process module, 
 		// but for debugging, particularly client or server, it can be useful to point the
 		// code at an existing cooked directory. If using -BaseFromWorkingDir set the
@@ -1193,6 +1198,7 @@ const TCHAR* FWindowsPlatformProcess::BaseDir()
 			FCString::Strcpy(Result, *TempResult);
 		}
 		else
+#endif // WINDOWS_ALLOW_BASEDIR_OVERRIDE
 		{
 			// Get the directory containing the current module if possible, or use the directory containing the executable if not
 			HMODULE hCurrentModule;
