Resubmitting CL 40325683 in Main (was skipped due to merge conflict)

Added support for custom memory storage for FNames.
Removed support for WITH_CUSTOM_NAME_ENCODING (should use new logic instead, which is more flexible)

#rnx

diff --git a/Engine/Source/Runtime/Core/Private/UObject/UnrealNames.cpp b/Engine/Source/Runtime/Core/Private/UObject/UnrealNames.cpp
--- a/Engine/Source/Runtime/Core/Private/UObject/UnrealNames.cpp
+++ b/Engine/Source/Runtime/Core/Private/UObject/UnrealNames.cpp
@@ -127,11 +127,6 @@ const TCHAR* LexToString(EName Ename)
 	}
 }
 
-int32 FNameEntry::GetDataOffset()
-{
-	return STRUCT_OFFSET(FNameEntry, AnsiName);
-}
-
 struct FNameEntryIds
 {
 	FNameEntryId ComparisonId, DisplayId;
@@ -237,15 +232,8 @@ FORCEINLINE bool Equals(FNameStringView A, FNameStringView B)
 	return (A.Len == B.Len & A.IsAnsi() == B.IsAnsi()) && EqualsSameDimensions<Sensitivity>(A, B);
 }
 
-// Minimize stack lifetime of large decode buffers
-#ifdef WITH_CUSTOM_NAME_ENCODING
-#define OUTLINE_DECODE_BUFFER FORCENOINLINE
-#else
-#define OUTLINE_DECODE_BUFFER
-#endif
-
 template<ENameCase Sensitivity>
-OUTLINE_DECODE_BUFFER bool EqualsSameDimensions(const FNameEntry& Entry, FNameStringView Name)
+bool EqualsSameDimensions(const FNameEntry& Entry, FNameStringView Name)
 {
 	FNameBuffer DecodeBuffer;
 	return EqualsSameDimensions<Sensitivity>(Entry.MakeView(DecodeBuffer), Name);
@@ -259,6 +247,71 @@ static constexpr uint32 FNameBlockOffsets = 1 << FNameBlockOffsetBits;
 static constexpr uint32 FNameEntryIdBits = FNameBlockOffsetBits + FNameMaxBlockBits;
 static constexpr uint32 FNameEntryIdMask = (1 << FNameEntryIdBits ) - 1;
 
+// Stores unterminated strings in FNameEntry::AnsiName/WideName
+struct FDefaultNameStore
+{
+	// Decides if entry header + Name data fits in OutSize. Called inside allocator and shard lock.
+	//
+	// @return Total entry aligned size if Name fits, 0 otherwise
+	static uint32 TryPlace(uint8* OutEntry, uint32 OutSize, FNameStringView Name)
+	{
+		uint32 Bytes = GetDefaultSize(Name.Len, Name.bIsWide);
+		return Bytes <= OutSize ? Bytes : 0;
+	}
+
+	// Called inside shard lock but outside allocator lock to avoid contention
+	static void Finalize(uint8* OutNameData, FNameStringView Name)
+	{
+		FMemory::Memcpy(OutNameData, Name.Data, Name.BytesWithoutTerminator());
+	}
+	
+	static void Copy(ANSICHAR* Out, const uint8* NameData, uint32 NameLen)
+	{
+		FMemory::Memcpy(Out, NameData, sizeof(ANSICHAR) * NameLen);
+	}
+
+	static void Copy(WIDECHAR* Out, const uint8* NameData, uint32 NameLen)
+	{
+		FMemory::Memcpy(Out, NameData, sizeof(WIDECHAR) * NameLen);
+	}
+
+	static uint32 GetSize(const FNameEntry& Entry)
+	{
+		return GetDefaultSize(uint32(Entry.GetNameLength()), Entry.IsWide());
+	}
+
+	static const ANSICHAR* TryGetAnsi(const void* NameData)
+	{
+		return static_cast<const ANSICHAR*>(NameData);
+	}
+	
+	static const WIDECHAR* TryGetWide(const void* NameData)
+	{
+		return static_cast<const WIDECHAR*>(NameData);
+	}
+	
+private:
+	static uint32 GetDefaultSize(uint32 Len, bool bIsWide)
+	{
+		uint32 HeaderSize = FNameEntry::GetDataOffset();
+		uint32 FooterSize =  Len * (bIsWide ? sizeof(WIDECHAR) : sizeof(ANSICHAR));
+		return Align(HeaderSize + FooterSize, alignof(FNameEntry));
+	}
+};
+
+// To customize how FNames are encoded in memory:
+// - Define UE_FNAME_CUSTOM_STORE_HEADER to point to a .inl source file
+// - Implement a class that can override FDefaultNameStore in that file
+// - End the included file with using FNameStore = FMyCustomNameStore
+#ifdef UE_FNAME_CUSTOM_STORE_HEADER
+#include UE_FNAME_CUSTOM_STORE_HEADER
+#else
+using FNameStore = FDefaultNameStore;
+#endif
+#ifdef WITH_CUSTOM_NAME_ENCODING
+#pragma error "WITH_CUSTOM_NAME_ENCODING macro is deprecated. Use UE_FNAME_CUSTOM_STORE_HEADER to change how FName data is saved to memory."
+#endif
+
 /** An unpacked FNameEntryId */
 struct FNameEntryHandle
 {
@@ -525,48 +578,63 @@ public:
 			Blocks[Idx] = AllocBlock();
 		}
 	}
-
-
-	/**
-	 * Allocates the requested amount of bytes and returns an id that can be used to access them
-	 *
-	 * @param   Size  Size in bytes to allocate, 
-	 * @return  Allocation of passed in size cast to a FNameEntry pointer.
-	 */
-	template <class ScopeLock, bool bNumbered = false>
-	FNameEntryHandle Allocate(uint32 Bytes)
+ 
+	// @param Bytes is aligned and fits within current block
+	inline FNameEntryHandle Allocate(uint32 Bytes)
 	{
-		Bytes = Align(Bytes, alignof(FNameEntry));
-		check(Bytes <= BlockSizeBytes);
-
-		ScopeLock _(Lock);
-
-		AlignCurrentByteCursor<bNumbered>();
-
-		// Allocate a new pool if current one is exhausted. We don't worry about a little bit
-		// of waste at the end given the relative size of pool to average and max allocation.
-		if (BlockSizeBytes - CurrentByteCursor < Bytes)
-		{
-			AllocateNewBlock();
-			AlignCurrentByteCursor<bNumbered>();
-		}
-
+		checkSlow(Bytes % Stride == 0);
+		checkSlow(CurrentByteCursor % Stride == 0);
+		checkSlow(CurrentByteCursor + Bytes <= BlockSizeBytes);
+ 
 		// Use current cursor position for this allocation and increment cursor for next allocation
 		uint32 ByteOffset = CurrentByteCursor;
 		CurrentByteCursor += Bytes;
-		
-		check(!(bNumbered && IsUnalignedNumberedEntry(ByteOffset)));
-		check(ByteOffset % Stride == 0 && ByteOffset / Stride < FNameBlockOffsets);
-
 		return FNameEntryHandle(CurrentBlock, ByteOffset / Stride);
 	}
+ 
+#if UE_FNAME_OUTLINE_NUMBER
+	template <class ScopeLock>
+	FNameEntryHandle AllocateNumbered()
+	{
+		ScopeLock _(Lock);
+ 
+		PadNumberedEntry();
+		if (CurrentByteCursor + NumberedEntrySize > BlockSizeBytes)
+		{
+			AllocateNewBlock();
+			PadNumberedEntry();
+		}
+ 
+		return Allocate(NumberedEntrySize);
+	}
+#endif
+ 
+	// Allocate and store Name
+	template <class ScopeLock>
+	FNameEntryHandle AllocateRegular(FNameStringView Name)
+	{
+		ScopeLock _(Lock);
+ 
+		uint32 Bytes = FNameStore::TryPlace(Blocks[CurrentBlock] + CurrentByteCursor, BlockSizeBytes - CurrentByteCursor, Name);
+ 
+		// Allocate a new block if current one is exhausted. We don't worry about a little bit
+		// of waste at the end given the relative size of block to average and max allocation.
+		if (Bytes == 0)
+		{
+			AllocateNewBlock();
+			Bytes = FNameStore::TryPlace(Blocks[CurrentBlock] + CurrentByteCursor, BlockSizeBytes - CurrentByteCursor, Name);
+			check(Bytes);
+		}
+ 
+		return Allocate(Bytes); 
+	}
 
 	template<class ScopeLock>
 	FNameEntryHandle Create(FNameStringView Name, TOptional<FNameEntryId> ComparisonId, FNameEntryHeader Header)
 	{
 		FPlatformMisc::Prefetch(Blocks[CurrentBlock]);
 
-		FNameEntryHandle Handle = Allocate<ScopeLock>(FNameEntry::GetDataOffset() + Name.BytesWithoutTerminator());
+		FNameEntryHandle Handle = AllocateRegular<ScopeLock>(Name);
 		FNameEntry& Entry = Resolve(Handle);
 
 #if WITH_CASE_PRESERVING_NAME
@@ -574,15 +642,7 @@ public:
 #endif
 
 		Entry.Header = Header;
-		
-		if (Name.bIsWide)
-		{
-			Entry.StoreName(Name.Wide, Name.Len);
-		}
-		else
-		{
-			Entry.StoreName(Name.Ansi, Name.Len);
-		}
+		FNameStore::Finalize(Entry.NameData, Name);
 
 		return Handle;
 	}
@@ -596,7 +656,7 @@ public:
 	FNameEntryHandle CreateWithNumber(const FNameEntryId& StringPart, uint32 NumberPart, TOptional<FNameEntryId> ComparisonId, FNameEntryHeader Header)
 	{
 		FPlatformMisc::Prefetch(Blocks[CurrentBlock]);
-		FNameEntryHandle Handle = Allocate<ScopeLock, true>(NumberedEntrySize);
+		FNameEntryHandle Handle = AllocateNumbered<ScopeLock>();
 		FNameEntry& Entry = Resolve(Handle);
 
 #if WITH_CASE_PRESERVING_NAME
@@ -648,28 +708,26 @@ public:
 	}
 
 private:
+
+#if UE_FNAME_OUTLINE_NUMBER
 	static constexpr bool IsUnalignedNumberedEntry(uint32 BlockOffset)
 	{
-#if UE_FNAME_OUTLINE_NUMBER
 		constexpr bool bRuntimeAlign = alignof(FNameEntry) < alignof(FNumberedEntry);
 		return bRuntimeAlign && (BlockOffset + offsetof(FNameEntry, NumberedName) + offsetof(FNameEntry::FNumberedData, Id)) % alignof(FNumberedEntry) != 0;
-#else
-		return false;
-#endif
 	}
 	
 	static constexpr FNameEntryHeader NumberPadHeader = {};
 
 	// Pad up FNameEntry::NumberedName to be 4B-aligned if FNameEntry is only 2B-aligned
-	template <bool bNumbered>
-	void AlignCurrentByteCursor()
+	void PadNumberedEntry()
 	{
-		if (bNumbered && IsUnalignedNumberedEntry(CurrentByteCursor) && (BlockSizeBytes - CurrentByteCursor) >= sizeof(FNameEntryHeader))
+		if (IsUnalignedNumberedEntry(CurrentByteCursor) && (BlockSizeBytes - CurrentByteCursor) >= sizeof(FNameEntryHeader))
 		{
 			(FNameEntryHeader&)Blocks[CurrentBlock][CurrentByteCursor] = NumberPadHeader;
 			CurrentByteCursor += sizeof(NumberPadHeader);
 		}	
 	}
+#endif
 
 	static void DebugDumpBlock(const uint8* Block, uint32 BlockSize, TArray<const FNameEntry*>& Out)
 	{
@@ -681,7 +739,7 @@ private:
 			if (uint32 Len = Entry->Header.Len)
 			{
 				Out.Add(Entry);
-				It += FNameEntry::GetSize(Len, !Entry->IsWide());
+				It += Entry->GetSizeInBytes();
 			}
 #if !UE_FNAME_OUTLINE_NUMBER
 			else
@@ -768,7 +826,7 @@ private:
 			}
 		}
 #endif
-		return (uint8*)FMemory::Malloc(BlockSizeBytes, alignof(FNameEntry));
+		return (uint8*)FMemory::Malloc(BlockSizeBytes, Stride);
 	}
 	
 	void AllocateNewBlock()
@@ -913,8 +971,7 @@ struct FNameHash
 		EntryProbeHeader.bIsWide = bIsWide;
 
 		// When we always use lowercase hashing, we can store parts of the hash in the entry
-		// to avoid copying and decoding entries needlessly. WITH_CUSTOM_NAME_ENCODING
-		// that makes this important is normally on when WITH_CASE_PRESERVING_NAME is off.
+		// to avoid copying and decoding entries needlessly.
 #if !WITH_CASE_PRESERVING_NAME		
 		static constexpr uint32 EntryProbeMask = (1u << FNameEntryHeader::ProbeHashBits) - 1; 
 		EntryProbeHeader.LowercaseProbeHash = static_cast<uint16>((Hi >> FNamePoolShardBits) & EntryProbeMask);
@@ -2177,11 +2234,9 @@ static int32 CompareDifferentIdsAlphabetically(FNameEntryId AId, FNameEntryId BI
 	{
 		FNameStringView& AnsiView = AView.bIsWide ? BView : AView;
 		FNameBuffer& AnsiBuffer =	AView.bIsWide ? BBuffer : ABuffer;
-
-#ifndef WITH_CUSTOM_NAME_ENCODING
+		
 		FPlatformMemory::Memcpy(AnsiBuffer.AnsiName, AnsiView.Ansi, AnsiView.Len * sizeof(ANSICHAR));
-		AnsiView.Ansi = AnsiBuffer.AnsiName;
-#endif
+        AnsiView.Ansi = AnsiBuffer.AnsiName;
 
 		ConvertInPlace<ANSICHAR, WIDECHAR>(AnsiBuffer.AnsiName, AnsiView.Len);
 		AnsiView.bIsWide = true;
@@ -2289,48 +2344,35 @@ static uint16 GetRawNonCasePreservingHash(const TCharType* Source)
 struct FClangKeepDebugInfo {};
 FNameEntry::FNameEntry(FClangKeepDebugInfo) {}
 
-void FNameEntry::StoreName(const ANSICHAR* InName, uint32 Len)
-{
-	FPlatformMemory::Memcpy(AnsiName, InName, sizeof(ANSICHAR) * Len);
-	Encode(AnsiName, Len);
-}
-
-void FNameEntry::StoreName(const WIDECHAR* InName, uint32 Len)
-{
-	FPlatformMemory::Memcpy(WideName, InName, sizeof(WIDECHAR) * Len);
-	Encode(WideName, Len);
-}
-
 void FNameEntry::CopyUnterminatedName(ANSICHAR* Out, uint32 Len) const
 {
-	FPlatformMemory::Memcpy(Out, AnsiName, sizeof(ANSICHAR) * Len);
-	Decode(Out, Len);
+	FNameStore::Copy(Out, NameData, Len);
 }
 
 void FNameEntry::CopyUnterminatedName(WIDECHAR* Out, uint32 Len) const
 {
-	FPlatformMemory::Memcpy(Out, WideName, sizeof(WIDECHAR) * Len);
-	Decode(Out, Len);
+	FNameStore::Copy(Out, NameData, Len);
 }
 
 FORCEINLINE const WIDECHAR* FNameEntry::GetUnterminatedName(WIDECHAR(&OptionalDecodeBuffer)[NAME_SIZE]) const
 {
-#ifdef WITH_CUSTOM_NAME_ENCODING
-	CopyUnterminatedName(OptionalDecodeBuffer, Header.Len);
+	if (const WIDECHAR* StoredPtr = FNameStore::TryGetWide(NameData))
+	{
+		return StoredPtr;
+	}
+	FNameStore::Copy(OptionalDecodeBuffer, NameData, Header.Len);
 	return OptionalDecodeBuffer;
-#else
-	return WideName;
-#endif
 }
 
 FORCEINLINE ANSICHAR const* FNameEntry::GetUnterminatedName(ANSICHAR(&OptionalDecodeBuffer)[NAME_SIZE]) const
 {
-#ifdef WITH_CUSTOM_NAME_ENCODING
-	CopyUnterminatedName(OptionalDecodeBuffer, Header.Len);
+	if (const ANSICHAR* StoredPtr = FNameStore::TryGetAnsi(NameData))
+	{
+		return StoredPtr;
+	}
+	
+	FNameStore::Copy(OptionalDecodeBuffer, NameData, Header.Len);
 	return OptionalDecodeBuffer;
-#else
-	return AnsiName;
-#endif
 }
 
 int32 FNameEntry::GetNameLengthUtf8() const
@@ -2509,17 +2551,6 @@ void FNameEntry::AppendNameToPathString(FString& Out) const
 	Out.PathAppend(EntryToCString(*this, Temp), Header.Len);
 }
 
-int32 FNameEntry::GetSize(const TCHAR* Name)
-{
-	return FNameEntry::GetSize(FCString::Strlen(Name), FCString::IsPureAnsi(Name));
-}
-
-int32 FNameEntry::GetSize(int32 Length, bool bIsPureAnsi)
-{
-	int32 Bytes = GetDataOffset() + Length * (bIsPureAnsi ? sizeof(ANSICHAR) : sizeof(WIDECHAR));
-	return Align(Bytes, alignof(FNameEntry));
-}
-
 #if UE_FNAME_OUTLINE_NUMBER
 uint32 FNameEntry::GetNumber() const
 {
@@ -2538,7 +2569,7 @@ int32 FNameEntry::GetSizeInBytes() const
 	else
 #endif
 	{
-		return GetSize(GetNameLength(), !IsWide());
+		return FNameStore::GetSize(*this);
 	}
 }
 
@@ -3279,21 +3310,6 @@ private:
 		}
 	}
 #endif // UE_FNAME_OUTLINE_NUMBER
-
-	static void ReplaceName(FNameEntry& Existing, FNameStringView Updated)
-	{
-		check(Existing.Header.bIsWide == Updated.bIsWide);
-		check(Existing.Header.Len == Updated.Len);
-
-		if (Updated.bIsWide)
-		{
-			Existing.StoreName(Updated.Wide, Updated.Len);
-		}
-		else
-		{
-			Existing.StoreName(Updated.Ansi, Updated.Len);
-		}
-	}
 };
 
 
diff --git a/Engine/Source/Runtime/Core/Public/UObject/NameTypes.h b/Engine/Source/Runtime/Core/Public/UObject/NameTypes.h
--- a/Engine/Source/Runtime/Core/Public/UObject/NameTypes.h
+++ b/Engine/Source/Runtime/Core/Public/UObject/NameTypes.h
@@ -245,6 +245,7 @@ private:
 	{
 		ANSICHAR			AnsiName[NAME_SIZE];
 		WIDECHAR			WideName[NAME_SIZE];
+		uint8				NameData[0];
 		FNumberedData		NumberedName;
 	};
 
@@ -300,21 +301,11 @@ public:
 
 	CORE_API void DebugDump(FOutputDevice& Out) const;
 
-	/**
-	 * Returns the size in bytes for FNameEntry structure. This is != sizeof(FNameEntry) as we only allocated as needed.
-	 *
-	 * @param	Length			Length of name
-	 * @param	bIsPureAnsi		Whether name is pure ANSI or not
-	 * @return	required size of FNameEntry structure to hold this string (might be wide or ansi)
-	 */
-	static int32 GetSize( int32 Length, bool bIsPureAnsi );
-	static CORE_API int32 GetSize(const TCHAR* Name);
-
 	CORE_API int32 GetSizeInBytes() const;
 
 	CORE_API void Write(FArchive& Ar) const;
 
-	static CORE_API int32 GetDataOffset();
+	static constexpr int32 GetDataOffset() { return offsetof(FNameEntry, NameData); }
 	struct CORE_API FNameStringView MakeView(union FNameBuffer& OptionalDecodeBuffer) const;
 private:
 	friend class FName;
@@ -322,14 +313,6 @@ private:
 	friend class FNameEntryAllocator;
 	friend class FNamePoolShardBase;
 	friend class FNamePool;
-
-	static void Encode(ANSICHAR* Name, uint32 Len);
-	static void Encode(WIDECHAR* Name, uint32 Len);
-	static void Decode(ANSICHAR* Name, uint32 Len);
-	static void Decode(WIDECHAR* Name, uint32 Len);
-
-	void StoreName(const ANSICHAR* InName, uint32 Len);
-	void StoreName(const WIDECHAR* InName, uint32 Len);
 	void CopyUnterminatedName(ANSICHAR* OutName, uint32 Len) const;
 	void CopyUnterminatedName(WIDECHAR* OutName, uint32 Len) const;
 	void CopyAndConvertUnterminatedName(TCHAR* OutName, uint32 Len) const;
@@ -1655,13 +1638,6 @@ struct FNameLexicalLess
 	}
 };
 
-#ifndef WITH_CUSTOM_NAME_ENCODING
-inline void FNameEntry::Encode(ANSICHAR*, uint32) {}
-inline void FNameEntry::Encode(WIDECHAR*, uint32) {}
-inline void FNameEntry::Decode(ANSICHAR*, uint32) {}
-inline void FNameEntry::Decode(WIDECHAR*, uint32) {}
-#endif
-
 struct FNameDebugVisualizer
 {
 	CORE_API static uint8** GetBlocks();
