Fixed TSparseArray::Reserve() to not use growth semantics for the storage of elements.

diff --git a/Engine/Source/Runtime/Core/Public/Containers/SparseArray.h b/Engine/Source/Runtime/Core/Public/Containers/SparseArray.h
--- a/Engine/Source/Runtime/Core/Public/Containers/SparseArray.h
+++ b/Engine/Source/Runtime/Core/Public/Containers/SparseArray.h
@@ -440,7 +440,8 @@ public:
 		{
 			const int32 ElementsToAdd = ExpectedNumElements - Data.Num();
 
-			// allocate memory in the array itself
+			// allocate memory in the array itself - reserve first to ensure that resizing-for-growth doesn't occur
+			Data.Reserve(ExpectedNumElements);
 			int32 ElementIndex = Data.AddUninitialized(ElementsToAdd);
 
 			FElementOrFreeListLink* DataPtr = (FElementOrFreeListLink*)Data.GetData();
