SavePackage: Add helper TryHashFile to support upcoming calculation of SavedPackageHash from the hash of the final package bytes.
FileHelper: Add LoadFileInBlocks for TryHashFile and other similar performant reading of files.
[RN] Minor, Core

diff --git a/Engine/Source/Runtime/Core/Private/Misc/FileHelper.cpp b/Engine/Source/Runtime/Core/Private/Misc/FileHelper.cpp
--- a/Engine/Source/Runtime/Core/Private/Misc/FileHelper.cpp
+++ b/Engine/Source/Runtime/Core/Private/Misc/FileHelper.cpp
@@ -5,6 +5,7 @@
 #include "Containers/StringConv.h"
 #include "Logging/LogMacros.h"
 #include "CoreGlobals.h"
+#include "HAL/PlatformMath.h"
 #include "Memory/MemoryView.h"
 #include "Memory/SharedBuffer.h"
 #include "Misc/ByteSwap.h"
@@ -12,10 +13,13 @@
 #include "Misc/Paths.h"
 #include "Math/IntRect.h"
 #include "Misc/OutputDeviceFile.h"
+#include "Misc/StringBuilder.h"
 #include "ProfilingDebugging/ProfilingHelpers.h"
 #include "Misc/ConfigCacheIni.h"
 #include "Misc/SecureHash.h"
 #include "HAL/FileManagerGeneric.h"
+#include "Templates/Function.h"
+#include "Templates/UniquePtr.h"
 
 #define LOCTEXT_NAMESPACE "FileHelper"
 
@@ -91,6 +95,47 @@ bool FFileHelper::LoadFileToArray(TArray64<uint8>& Result, const TCHAR* Filename
 	return bSuccess;
 }
 
+bool FFileHelper::LoadFileInBlocks(FStringView Filename, TFunctionRef<void(FMemoryView)> BlockVisitor,
+	int64 Offset, int64 Size, uint32 Flags, int64 BlockSize)
+{
+	TUniquePtr<FArchive> FileReader(IFileManager::Get().CreateFileReader(*WriteToString<256>(Filename), Flags));
+	if (!FileReader)
+	{
+		return false;
+	}
+
+	const int64 DiskSize = FileReader->TotalSize();;
+	Offset = FMath::Clamp(Offset, (int64)0, DiskSize);
+	int64 RemainingLength = DiskSize - Offset;
+	RemainingLength = Size < 0 ? RemainingLength : FMath::Min(Size, RemainingLength);
+
+	if (RemainingLength > 0)
+	{
+		if (Offset > 0)
+		{
+			FileReader->Seek(Offset);
+		}
+
+		constexpr int64 DefaultMaxBufferSize = 1024 * 1024;
+		const int64 MaxBufferSize = BlockSize <= 0 ? DefaultMaxBufferSize : BlockSize;
+		int64 BufferSize = FMath::Min(MaxBufferSize, RemainingLength);
+		TUniquePtr<uint8[]> Buffer(new uint8[BufferSize]);
+		while (RemainingLength > 0)
+		{
+			int64 ReadLength = FMath::Min(RemainingLength, BufferSize);
+			FileReader->Serialize(Buffer.Get(), ReadLength);
+			RemainingLength -= ReadLength;
+			BlockVisitor(FMemoryView(Buffer.Get(), ReadLength));
+		}
+		if (FileReader->IsError())
+		{
+			return false;
+		}
+	}
+
+	return true;
+}
+
 /**
  * Converts an arbitrary text buffer to an FString.
  * Supports all combination of ANSI/Unicode files and platforms.
diff --git a/Engine/Source/Runtime/Core/Public/Misc/FileHelper.h b/Engine/Source/Runtime/Core/Public/Misc/FileHelper.h
--- a/Engine/Source/Runtime/Core/Public/Misc/FileHelper.h
+++ b/Engine/Source/Runtime/Core/Public/Misc/FileHelper.h
@@ -13,6 +13,7 @@
 #include "HAL/PlatformCrt.h"
 #include "Math/Color.h"
 #include "Math/MathFwd.h"
+#include "Memory/MemoryFwd.h"
 #include "Misc/EnumClassFlags.h"
 #include "Templates/UnrealTemplate.h"
 
@@ -77,6 +78,30 @@ struct FFileHelper
 	*/
 	static CORE_API bool LoadFileToArray( TArray64<uint8>& Result, const TCHAR* Filename, uint32 Flags = 0 );
 
+	/**
+	 * Load a binary file a block at a time into a buffer and call the visitor each time the buffer is populated
+	 * with a new block. This uses less memory and has higher cache coherency than loading the entire file into
+	 * a block of bytes, and is more performant than reading the file a byte at a time. Useful for e.g. hashing or
+	 * parsing.
+	 * 
+	 * @param Filename The file to read
+	 * @param BlockVisitor Function called repeatedly with each sequential block.
+	 *        Size of the input might change for performance and should not be assumed. Possibly the entire file will
+	 *        be passed in a single block. If LoadFileInBlocks returns true, the sum of sizes in each call to
+	 *        BlockVisitor equals the size of the requested range of the the file on disk. If the requested range
+	 *        has size 0, LoadFileInBlocks will return true and the BlockVisitor will not be called.
+	 * @param Offset Start of the range of bytes to read from the file. If less than zero or larger than filesize,
+	 *        it will be silently clamped to the [0,FileSize] range. This may result in LoadFileInBlocks returning true
+	 *        but using an empty range and not calling BlockVisitor.
+	 * @param Length of the range of bytes to read from the file. If less than zero, all bytes between Offset and end
+	 *        of file are read. If larger than the size of the range between Offset and end of file, it will be
+	 *        silently shortened to the size of the range.
+	 * @param Flags Flags to pass to IFileManager::CreateFileReader
+	 * @param BlockSize Size of the internally-allocated buffer to read bytes into and pass to each call to
+	 *        BlockVisitor. Sizes <= 0 use the default-selected block size for the current platform.
+	 */
+	static CORE_API bool LoadFileInBlocks(FStringView Filename, TFunctionRef<void(FMemoryView)> BlockVisitor,
+		int64 Offset = 0, int64 Size = -1, uint32 Flags = 0, int64 BlockSize = 0);
 	/**
 	 * Loads a text file from an FArchive to an FString. Supports all combination of ANSI/Unicode files and platforms.
 	 *
diff --git a/Engine/Source/Runtime/CoreUObject/Private/UObject/SavePackage/SavePackageUtilities.cpp b/Engine/Source/Runtime/CoreUObject/Private/UObject/SavePackage/SavePackageUtilities.cpp
--- a/Engine/Source/Runtime/CoreUObject/Private/UObject/SavePackage/SavePackageUtilities.cpp
+++ b/Engine/Source/Runtime/CoreUObject/Private/UObject/SavePackage/SavePackageUtilities.cpp
@@ -11,7 +11,9 @@
 #include "Cooker/CookEvents.h"
 #include "CoreMinimal.h"
 #include "HAL/FileManager.h"
+#include "Hash/Blake3.h"
 #include "Interfaces/ITargetPlatform.h"
+#include "IO/IoHash.h"
 #include "Logging/StructuredLog.h"
 #include "Misc/AssetRegistryInterface.h"
 #include "Misc/CommandLine.h"
@@ -2407,6 +2409,19 @@ void SaveWorldLevelInfo(UPackage* InOuter, FLinkerSave* Linker, FStructuredArchi
 	}
 }
 
+bool TryHashFile(FStringView Filename, FIoHash& OutHash, int64 Offset, int64 Size)
+{
+	FBlake3 Builder;
+	if (!TryHashFile(Filename, Builder, Offset, Size))
+	{
+		OutHash = FIoHash::Zero;
+		return false;
+	}
+
+	OutHash = FIoHash(Builder.Finalize());
+	return true;
+}
+
 } // end namespace UE::SavePackageUtilities
 
 void UPackage::WaitForAsyncFileWrites()
diff --git a/Engine/Source/Runtime/CoreUObject/Private/UObject/SavePackage/SavePackageUtilities.h b/Engine/Source/Runtime/CoreUObject/Private/UObject/SavePackage/SavePackageUtilities.h
--- a/Engine/Source/Runtime/CoreUObject/Private/UObject/SavePackage/SavePackageUtilities.h
+++ b/Engine/Source/Runtime/CoreUObject/Private/UObject/SavePackage/SavePackageUtilities.h
@@ -7,6 +7,8 @@
 #include "Containers/Set.h"
 #include "Containers/StringFwd.h"
 #include "HAL/ThreadSingleton.h"
+#include "Memory/MemoryView.h"
+#include "Misc/FileHelper.h"
 #include "ProfilingDebugging/CookStats.h"
 #include "Serialization/ArchiveObjectCrc32.h"
 #include "Serialization/ArchiveStackTrace.h"
@@ -28,6 +30,7 @@ class FSavePackageContext;
 class IPackageWriter;
 class ITargetPlatform;
 enum class ESavePackageResult;
+struct FIoHash;
 
 // Save Time trace
 #if UE_TRACE_ENABLED && !UE_BUILD_SHIPPING
@@ -364,6 +367,10 @@ void WriteToFile(const FString& Filename, const uint8* InDataPtr, int64 InDataSi
 void AsyncWriteFile(FLargeMemoryPtr Data, const int64 DataSize, const TCHAR* Filename, EAsyncWriteOptions Options, TArrayView<const FFileRegion> InFileRegions);
 void AsyncWriteFile(EAsyncWriteOptions Options, FSavePackageOutputFile& File);
 
+template <typename HashBuilderType>
+bool TryHashFile(FStringView Filename, HashBuilderType& Builder, int64 Offset = 0, int64 Size = -1);
+bool TryHashFile(FStringView Filename, FIoHash& OutHash, int64 Offset = 0, int64 Size = -1);
+
 void GetCDOSubobjects(UObject* CDO, TArray<UObject*>& Subobjects);
 
 enum class EEditorOnlyObjectFlags
@@ -414,3 +421,23 @@ struct FSavePackageStats
 	static void MergeStats(const TMap<FName, FArchiveDiffStats>& ToMerge);
 };
 #endif
+
+
+///////////////////////////////////////////////////////
+// Inline implementations
+///////////////////////////////////////////////////////
+
+
+namespace UE::SavePackageUtilities
+{
+
+template <typename HashBuilderType>
+bool TryHashFile(FStringView Filename, HashBuilderType& Builder, int64 Offset, int64 Size)
+{
+	return FFileHelper::LoadFileInBlocks(Filename, [&Builder](FMemoryView Block)
+		{
+			Builder.Update(Block);
+		}, Offset, Size);
+}
+
+}
\ No newline at end of file
