Introduced utility enum for iterations. Applied to EntitySerializedDataAccess visitors.

diff --git a/Engine/Source/Runtime/Core/Public/Misc/Iteration.h b/Engine/Source/Runtime/Core/Public/Misc/Iteration.h
new file mode 100644
--- /dev/null
+++ b/Engine/Source/Runtime/Core/Public/Misc/Iteration.h
@@ -0,0 +1,13 @@
+// Copyright Epic Games, Inc. All Rights Reserved.
+
+#pragma once
+
+/** Enum that can be used as an explicit result for iterations. E.g. as the return value of visitor functions. */
+enum class EIteration
+{
+	/** The loop should be stopped. */
+	Break,
+
+	/** The loop should continue. */
+	Continue,
+};
