[Core] - Fixes to be able to compile with header units
* Added inline in front of global constexpr variables to make them available outside TU
* Changed Array allocation functions to use uint32 for element size and element alignment instead of changing all allocators or adding casts everywhere.. We can be pretty sure an element size is not larger than 4gb and alignment is definitely not that big
* Fixed BitArray using correct signed/unsigned for internal calculation
* Removed explicit template instantiation that are not needed since function implementation is not visible in header anyway and methods have dll storage
* Changed math types SerializeFromMismatchedTag to use if constexpr instead of template specialization.

diff --git a/Engine/Source/Runtime/Core/Private/Math/UnrealMath.cpp b/Engine/Source/Runtime/Core/Private/Math/UnrealMath.cpp
--- a/Engine/Source/Runtime/Core/Private/Math/UnrealMath.cpp
+++ b/Engine/Source/Runtime/Core/Private/Math/UnrealMath.cpp
@@ -3705,9 +3705,9 @@ template struct UE::Math::TVector4<float>;
 template struct UE::Math::TVector4<double>;
 template struct UE::Math::TMatrix<float>;
 template struct UE::Math::TMatrix<double>;
-template void FMath::SegmentDistToSegmentSafe<float>(UE::Math::TVector<float> A1, UE::Math::TVector<float> B1, UE::Math::TVector<float> A2, UE::Math::TVector<float> B2, UE::Math::TVector<float>& OutP1, UE::Math::TVector<float>& OutP2);
-template void FMath::SegmentDistToSegmentSafe<double>(UE::Math::TVector<double> A1, UE::Math::TVector<double> B1, UE::Math::TVector<double> A2, UE::Math::TVector<double> B2, UE::Math::TVector<double>& OutP1, UE::Math::TVector<double>& OutP2);
-template UE::Math::TVector<float> FMath::ClosestPointOnLine(const UE::Math::TVector<float>& LineStart, const UE::Math::TVector<float>& LineEnd, const UE::Math::TVector<float>& Point);
-template UE::Math::TVector<double> FMath::ClosestPointOnLine(const UE::Math::TVector<double>& LineStart, const UE::Math::TVector<double>& LineEnd, const UE::Math::TVector<double>& Point);
-template UE::Math::TVector<float> FMath::ClosestPointOnSegment(const UE::Math::TVector<float>& Point, const UE::Math::TVector<float>& StartPoint, const UE::Math::TVector<float>& EndPoint);
-template UE::Math::TVector<double> FMath::ClosestPointOnSegment(const UE::Math::TVector<double>& Point, const UE::Math::TVector<double>& StartPoint, const UE::Math::TVector<double>& EndPoint);
+template CORE_API void FMath::SegmentDistToSegmentSafe<float>(UE::Math::TVector<float> A1, UE::Math::TVector<float> B1, UE::Math::TVector<float> A2, UE::Math::TVector<float> B2, UE::Math::TVector<float>& OutP1, UE::Math::TVector<float>& OutP2);
+template CORE_API void FMath::SegmentDistToSegmentSafe<double>(UE::Math::TVector<double> A1, UE::Math::TVector<double> B1, UE::Math::TVector<double> A2, UE::Math::TVector<double> B2, UE::Math::TVector<double>& OutP1, UE::Math::TVector<double>& OutP2);
+template CORE_API UE::Math::TVector<float> FMath::ClosestPointOnLine(const UE::Math::TVector<float>& LineStart, const UE::Math::TVector<float>& LineEnd, const UE::Math::TVector<float>& Point);
+template CORE_API UE::Math::TVector<double> FMath::ClosestPointOnLine(const UE::Math::TVector<double>& LineStart, const UE::Math::TVector<double>& LineEnd, const UE::Math::TVector<double>& Point);
+template CORE_API UE::Math::TVector<float> FMath::ClosestPointOnSegment(const UE::Math::TVector<float>& Point, const UE::Math::TVector<float>& StartPoint, const UE::Math::TVector<float>& EndPoint);
+template CORE_API UE::Math::TVector<double> FMath::ClosestPointOnSegment(const UE::Math::TVector<double>& Point, const UE::Math::TVector<double>& StartPoint, const UE::Math::TVector<double>& EndPoint);
diff --git a/Engine/Source/Runtime/Core/Private/Misc/ExpressionParser.cpp b/Engine/Source/Runtime/Core/Private/Misc/ExpressionParser.cpp
--- a/Engine/Source/Runtime/Core/Private/Misc/ExpressionParser.cpp
+++ b/Engine/Source/Runtime/Core/Private/Misc/ExpressionParser.cpp
@@ -849,25 +849,25 @@ namespace ExpressionParser
 		return MakeError(LOCTEXT("SyntaxError_InvalidExpression", "Could not evaluate expression"));
 	}
 
-	template TLexResultType<ANSICHAR> Lex<ANSICHAR>(const ANSICHAR* InExpression, const TTokenDefinitions<ANSICHAR>& TokenDefinitions);
-	template TLexResultType<UTF8CHAR> Lex<UTF8CHAR>(const UTF8CHAR* InExpression, const TTokenDefinitions<UTF8CHAR>& TokenDefinitions);
-	template TLexResultType<WIDECHAR> Lex<WIDECHAR>(const WIDECHAR* InExpression, const TTokenDefinitions<WIDECHAR>& TokenDefinitions);
+	template CORE_API TLexResultType<ANSICHAR> Lex<ANSICHAR>(const ANSICHAR* InExpression, const TTokenDefinitions<ANSICHAR>& TokenDefinitions);
+	template CORE_API TLexResultType<UTF8CHAR> Lex<UTF8CHAR>(const UTF8CHAR* InExpression, const TTokenDefinitions<UTF8CHAR>& TokenDefinitions);
+	template CORE_API TLexResultType<WIDECHAR> Lex<WIDECHAR>(const WIDECHAR* InExpression, const TTokenDefinitions<WIDECHAR>& TokenDefinitions);
 
-	template TCompileResultType<ANSICHAR> Compile<ANSICHAR>(const ANSICHAR* InExpression, const TTokenDefinitions<ANSICHAR>& TokenDefinitions, const FExpressionGrammar& InGrammar);
-	template TCompileResultType<UTF8CHAR> Compile<UTF8CHAR>(const UTF8CHAR* InExpression, const TTokenDefinitions<UTF8CHAR>& TokenDefinitions, const FExpressionGrammar& InGrammar);
-	template TCompileResultType<WIDECHAR> Compile<WIDECHAR>(const WIDECHAR* InExpression, const TTokenDefinitions<WIDECHAR>& TokenDefinitions, const FExpressionGrammar& InGrammar);
+	template CORE_API TCompileResultType<ANSICHAR> Compile<ANSICHAR>(const ANSICHAR* InExpression, const TTokenDefinitions<ANSICHAR>& TokenDefinitions, const FExpressionGrammar& InGrammar);
+	template CORE_API TCompileResultType<UTF8CHAR> Compile<UTF8CHAR>(const UTF8CHAR* InExpression, const TTokenDefinitions<UTF8CHAR>& TokenDefinitions, const FExpressionGrammar& InGrammar);
+	template CORE_API TCompileResultType<WIDECHAR> Compile<WIDECHAR>(const WIDECHAR* InExpression, const TTokenDefinitions<WIDECHAR>& TokenDefinitions, const FExpressionGrammar& InGrammar);
 
-	template TCompileResultType<ANSICHAR> Compile<ANSICHAR>(TArray<TExpressionToken<ANSICHAR>> InTokens, const FExpressionGrammar& InGrammar);
-	template TCompileResultType<UTF8CHAR> Compile<UTF8CHAR>(TArray<TExpressionToken<UTF8CHAR>> InTokens, const FExpressionGrammar& InGrammar);
-	template TCompileResultType<WIDECHAR> Compile<WIDECHAR>(TArray<TExpressionToken<WIDECHAR>> InTokens, const FExpressionGrammar& InGrammar);
+	template CORE_API TCompileResultType<ANSICHAR> Compile<ANSICHAR>(TArray<TExpressionToken<ANSICHAR>> InTokens, const FExpressionGrammar& InGrammar);
+	template CORE_API TCompileResultType<UTF8CHAR> Compile<UTF8CHAR>(TArray<TExpressionToken<UTF8CHAR>> InTokens, const FExpressionGrammar& InGrammar);
+	template CORE_API TCompileResultType<WIDECHAR> Compile<WIDECHAR>(TArray<TExpressionToken<WIDECHAR>> InTokens, const FExpressionGrammar& InGrammar);
 
-	template FExpressionResult Evaluate<ANSICHAR>(const ANSICHAR* InExpression, const TTokenDefinitions<ANSICHAR>& InTokenDefinitions, const FExpressionGrammar& InGrammar, const TIOperatorEvaluationEnvironment<ANSICHAR>& InEnvironment);
-	template FExpressionResult Evaluate<UTF8CHAR>(const UTF8CHAR* InExpression, const TTokenDefinitions<UTF8CHAR>& InTokenDefinitions, const FExpressionGrammar& InGrammar, const TIOperatorEvaluationEnvironment<UTF8CHAR>& InEnvironment);
-	template FExpressionResult Evaluate<WIDECHAR>(const WIDECHAR* InExpression, const TTokenDefinitions<WIDECHAR>& InTokenDefinitions, const FExpressionGrammar& InGrammar, const TIOperatorEvaluationEnvironment<WIDECHAR>& InEnvironment);
+	template CORE_API FExpressionResult Evaluate<ANSICHAR>(const ANSICHAR* InExpression, const TTokenDefinitions<ANSICHAR>& InTokenDefinitions, const FExpressionGrammar& InGrammar, const TIOperatorEvaluationEnvironment<ANSICHAR>& InEnvironment);
+	template CORE_API FExpressionResult Evaluate<UTF8CHAR>(const UTF8CHAR* InExpression, const TTokenDefinitions<UTF8CHAR>& InTokenDefinitions, const FExpressionGrammar& InGrammar, const TIOperatorEvaluationEnvironment<UTF8CHAR>& InEnvironment);
+	template CORE_API FExpressionResult Evaluate<WIDECHAR>(const WIDECHAR* InExpression, const TTokenDefinitions<WIDECHAR>& InTokenDefinitions, const FExpressionGrammar& InGrammar, const TIOperatorEvaluationEnvironment<WIDECHAR>& InEnvironment);
 
-	template FExpressionResult Evaluate<ANSICHAR>(const TArray<TCompiledToken<ANSICHAR>>& CompiledTokens, const TIOperatorEvaluationEnvironment<ANSICHAR>& InEnvironment);
-	template FExpressionResult Evaluate<UTF8CHAR>(const TArray<TCompiledToken<UTF8CHAR>>& CompiledTokens, const TIOperatorEvaluationEnvironment<UTF8CHAR>& InEnvironment);
-	template FExpressionResult Evaluate<WIDECHAR>(const TArray<TCompiledToken<WIDECHAR>>& CompiledTokens, const TIOperatorEvaluationEnvironment<WIDECHAR>& InEnvironment);
+	template CORE_API FExpressionResult Evaluate<ANSICHAR>(const TArray<TCompiledToken<ANSICHAR>>& CompiledTokens, const TIOperatorEvaluationEnvironment<ANSICHAR>& InEnvironment);
+	template CORE_API FExpressionResult Evaluate<UTF8CHAR>(const TArray<TCompiledToken<UTF8CHAR>>& CompiledTokens, const TIOperatorEvaluationEnvironment<UTF8CHAR>& InEnvironment);
+	template CORE_API FExpressionResult Evaluate<WIDECHAR>(const TArray<TCompiledToken<WIDECHAR>>& CompiledTokens, const TIOperatorEvaluationEnvironment<WIDECHAR>& InEnvironment);
 }
 
 template class TTokenStream<ANSICHAR>;
diff --git a/Engine/Source/Runtime/Core/Public/ColorManagement/ColorManagementDefines.h b/Engine/Source/Runtime/Core/Public/ColorManagement/ColorManagementDefines.h
--- a/Engine/Source/Runtime/Core/Public/ColorManagement/ColorManagementDefines.h
+++ b/Engine/Source/Runtime/Core/Public/ColorManagement/ColorManagementDefines.h
@@ -7,10 +7,10 @@
 namespace UE { namespace Color {
 
 /** Increment upon breaking changes to the EEncoding enum. Note that changing this forces a rebuild of textures that rely on it. */
-constexpr uint32 ENCODING_TYPES_VER = 4;
+inline constexpr uint32 ENCODING_TYPES_VER = 4;
 
 /** Increment upon breaking changes to the EColorSpace and EChromaticAdaptationMethod enums. Note that changing this forces a rebuild of textures that rely on it.*/
-constexpr uint32 COLORSPACE_VER = 2;
+inline constexpr uint32 COLORSPACE_VER = 2;
 
 /** List of available encodings/transfer functions.
 * 
diff --git a/Engine/Source/Runtime/Core/Public/Containers/Array.h b/Engine/Source/Runtime/Core/Public/Containers/Array.h
--- a/Engine/Source/Runtime/Core/Public/Containers/Array.h
+++ b/Engine/Source/Runtime/Core/Public/Containers/Array.h
@@ -380,8 +380,8 @@ namespace UE::Core::Private
 	// This should be used for repeated growing operations when reallocations are to be amortized over multiple inserts.
 	template <uint32 Flags, typename AllocatorInstanceType>
 	FORCENOINLINE TAllocatorSizeType_T<AllocatorInstanceType> ReallocGrow(
-		std::size_t                                  ElementSize,
-		std::size_t                                  ElementAlignment,
+		uint32                                       ElementSize,
+		uint32                                       ElementAlignment,
 		TAllocatorSizeType_T<AllocatorInstanceType>  Count,
 		AllocatorInstanceType&                       AllocatorInstance,
 		TAllocatorSizeType_T<AllocatorInstanceType>& ArrayNum,
@@ -442,8 +442,8 @@ namespace UE::Core::Private
 	// This should be used for repeated shrinking operations when reallocations are to be amortized over multiple removals.
 	template <uint32 Flags, typename AllocatorInstanceType>
 	FORCENOINLINE void ReallocShrink(
-		size_t                                       ElementSize,
-		std::size_t                                  ElementAlignment,
+		uint32                                       ElementSize,
+		uint32                                       ElementAlignment,
 		AllocatorInstanceType&                       AllocatorInstance,
 		TAllocatorSizeType_T<AllocatorInstanceType>  ArrayNum,
 		TAllocatorSizeType_T<AllocatorInstanceType>& ArrayMax
@@ -477,8 +477,8 @@ namespace UE::Core::Private
 	// Precondition: NewMax >= ArrayNum.
 	template <uint32 Flags, typename AllocatorInstanceType>
 	FORCENOINLINE void ReallocTo(
-		size_t                                       ElementSize,
-		std::size_t                                  ElementAlignment,
+		uint32                                       ElementSize,
+		uint32                                       ElementAlignment,
 		TAllocatorSizeType_T<AllocatorInstanceType>  NewMax,
 		AllocatorInstanceType&                       AllocatorInstance,
 		TAllocatorSizeType_T<AllocatorInstanceType>  ArrayNum,
@@ -513,8 +513,8 @@ namespace UE::Core::Private
 
 	template <uint32 Flags, typename AllocatorInstanceType>
 	FORCENOINLINE void ReallocForCopy(
-		size_t                                       ElementSize,
-		std::size_t                                  ElementAlignment,
+		uint32                                       ElementSize,
+		uint32                                       ElementAlignment,
 		TAllocatorSizeType_T<AllocatorInstanceType>  NewMax,
 		TAllocatorSizeType_T<AllocatorInstanceType>  PrevMax,
 		AllocatorInstanceType&                       AllocatorInstance,
diff --git a/Engine/Source/Runtime/Core/Public/Containers/BitArray.h b/Engine/Source/Runtime/Core/Public/Containers/BitArray.h
--- a/Engine/Source/Runtime/Core/Public/Containers/BitArray.h
+++ b/Engine/Source/Runtime/Core/Public/Containers/BitArray.h
@@ -1546,7 +1546,7 @@ public:
 			if (!this->Mask)
 			{
 				// Advance to the next uint32.
-				this->Mask = (1 << (NumBitsPerDWORD - 1));
+				this->Mask = (1u << (NumBitsPerDWORD - 1u));
 				--this->WordIndex;
 			}
 			return *this;
diff --git a/Engine/Source/Runtime/Core/Public/Math/Matrix.h b/Engine/Source/Runtime/Core/Public/Math/Matrix.h
--- a/Engine/Source/Runtime/Core/Public/Math/Matrix.h
+++ b/Engine/Source/Runtime/Core/Public/Math/Matrix.h
@@ -405,7 +405,17 @@ public:
 		//return false;
 	}
 	
-	bool SerializeFromMismatchedTag(FName StructTag, FArchive& Ar);
+	bool SerializeFromMismatchedTag(FName StructTag, FArchive& Ar)
+	{
+		if constexpr (std::is_same_v<T, float>)
+		{
+			return UE_SERIALIZE_VARIANT_FROM_MISMATCHED_TAG(Ar, Matrix, Matrix44f, Matrix44d);
+		}
+		else
+		{
+			return UE_SERIALIZE_VARIANT_FROM_MISMATCHED_TAG(Ar, Matrix, Matrix44d, Matrix44f);
+		}
+	}
 
 	// Conversion to other type.
 	template<typename FArg UE_REQUIRES(!std::is_same_v<T, FArg>)>
diff --git a/Engine/Source/Runtime/Core/Public/Math/Matrix.inl b/Engine/Source/Runtime/Core/Public/Math/Matrix.inl
--- a/Engine/Source/Runtime/Core/Public/Math/Matrix.inl
+++ b/Engine/Source/Runtime/Core/Public/Math/Matrix.inl
@@ -985,18 +985,3 @@ FORCEINLINE TLookAtMatrix<T>::TLookAtMatrix(const TVector<T>& EyePosition, const
 
 } // namespace UE::Core
 } // namespace UE
-
-template<>
-inline bool FMatrix44f::SerializeFromMismatchedTag(FName StructTag, FArchive& Ar)
-{	
-	return UE_SERIALIZE_VARIANT_FROM_MISMATCHED_TAG(Ar, Matrix, Matrix44f, Matrix44d);
-}
-
-template<>
-inline bool FMatrix44d::SerializeFromMismatchedTag(FName StructTag, FArchive& Ar)
-{
-	return UE_SERIALIZE_VARIANT_FROM_MISMATCHED_TAG(Ar, Matrix, Matrix44d, Matrix44f);
-}
-
-
-
diff --git a/Engine/Source/Runtime/Core/Public/Math/Quat.h b/Engine/Source/Runtime/Core/Public/Math/Quat.h
--- a/Engine/Source/Runtime/Core/Public/Math/Quat.h
+++ b/Engine/Source/Runtime/Core/Public/Math/Quat.h
@@ -710,7 +710,17 @@ public:
 		return true;
 	}
 
-	bool SerializeFromMismatchedTag(FName StructTag, FArchive& Ar);
+	bool SerializeFromMismatchedTag(FName StructTag, FArchive& Ar)
+	{
+		if constexpr (std::is_same_v<T, float>)
+		{
+			return UE_SERIALIZE_VARIANT_FROM_MISMATCHED_TAG(Ar, Quat, Quat4f, Quat4d);
+		}
+		else
+		{
+			return UE_SERIALIZE_VARIANT_FROM_MISMATCHED_TAG(Ar, Quat, Quat4d, Quat4f);
+		}
+	}
 
 	// Conversion to other type.
 	template<typename FArg UE_REQUIRES(!std::is_same_v<T, FArg>)>
@@ -1419,20 +1429,6 @@ template<> CORE_API FRotator3f FQuat4f::Rotator() const;
 template<> CORE_API FRotator3d FQuat4d::Rotator() const;
 
 
-
-template<>
-inline bool FQuat4f::SerializeFromMismatchedTag(FName StructTag, FArchive& Ar)
-{
-	return UE_SERIALIZE_VARIANT_FROM_MISMATCHED_TAG(Ar, Quat, Quat4f, Quat4d);
-}
-
-template<>
-inline bool FQuat4d::SerializeFromMismatchedTag(FName StructTag, FArchive& Ar)
-{
-	return UE_SERIALIZE_VARIANT_FROM_MISMATCHED_TAG(Ar, Quat, Quat4d, Quat4f);
-}
-
-
 /* FMath inline functions
  *****************************************************************************/
  
diff --git a/Engine/Source/Runtime/Core/Public/Math/Rotator.h b/Engine/Source/Runtime/Core/Public/Math/Rotator.h
--- a/Engine/Source/Runtime/Core/Public/Math/Rotator.h
+++ b/Engine/Source/Runtime/Core/Public/Math/Rotator.h
@@ -482,7 +482,17 @@ public:
 		return true;
 	}
 
-	bool SerializeFromMismatchedTag(FName StructTag, FArchive& Ar);
+	bool SerializeFromMismatchedTag(FName StructTag, FArchive& Ar)
+	{
+		if constexpr (std::is_same_v<T, float>)
+		{
+			return UE_SERIALIZE_VARIANT_FROM_MISMATCHED_TAG(Ar, Rotator, Rotator3f, Rotator3d);
+		}
+		else
+		{
+			return UE_SERIALIZE_VARIANT_FROM_MISMATCHED_TAG(Ar, Rotator, Rotator3d, Rotator3f);
+		}
+	}
 
 	// Conversion from other type.
 	template<typename FArg UE_REQUIRES(!std::is_same_v<T, FArg>)>
@@ -909,20 +919,6 @@ template<> CORE_API FQuat4f FRotator3f::Quaternion() const;
 template<> CORE_API FQuat4d FRotator3d::Quaternion() const;
 
 
-
-template<>
-inline bool FRotator3f::SerializeFromMismatchedTag(FName StructTag, FArchive& Ar)
-{	
-	return UE_SERIALIZE_VARIANT_FROM_MISMATCHED_TAG(Ar, Rotator, Rotator3f, Rotator3d);
-}
-
-template<>
-inline bool FRotator3d::SerializeFromMismatchedTag(FName StructTag, FArchive& Ar)
-{
-	return UE_SERIALIZE_VARIANT_FROM_MISMATCHED_TAG(Ar, Rotator, Rotator3d, Rotator3f);
-}
-
-
 /* FMath inline functions
  *****************************************************************************/
 
diff --git a/Engine/Source/Runtime/Core/Public/Math/Transform.h b/Engine/Source/Runtime/Core/Public/Math/Transform.h
--- a/Engine/Source/Runtime/Core/Public/Math/Transform.h
+++ b/Engine/Source/Runtime/Core/Public/Math/Transform.h
@@ -39,28 +39,6 @@ template<> struct TCanBulkSerialize<FTransform3d> { enum { Value = false }; };
 DECLARE_INTRINSIC_TYPE_LAYOUT(FTransform3f);
 DECLARE_INTRINSIC_TYPE_LAYOUT(FTransform3d);
 
-template<>
-inline bool FTransform3f::SerializeFromMismatchedTag(FName StructTag, FArchive& Ar)
-{
-	// FTransform3f/d are immutable, but FTransform is not and needs a per-property conversion
-	if(StructTag == NAME_Transform)
-	{
-		return false;
-	}
-	return UE_SERIALIZE_VARIANT_FROM_MISMATCHED_TAG(Ar, Transform, Transform3f, Transform3d);
-}
-
-template<>
-inline bool FTransform3d::SerializeFromMismatchedTag(FName StructTag, FArchive& Ar)
-{
-	// FTransform3f/d are immutable, but FTransform is not and needs a per-property conversion
-	if(StructTag == NAME_Transform)
-	{
-		return false;
-	}
-	return UE_SERIALIZE_VARIANT_FROM_MISMATCHED_TAG(Ar, Transform, Transform3d, Transform3f);
-}
-
 namespace UE::Math
 {
 	/**
diff --git a/Engine/Source/Runtime/Core/Public/Math/TransformNonVectorized.h b/Engine/Source/Runtime/Core/Public/Math/TransformNonVectorized.h
--- a/Engine/Source/Runtime/Core/Public/Math/TransformNonVectorized.h
+++ b/Engine/Source/Runtime/Core/Public/Math/TransformNonVectorized.h
@@ -629,7 +629,24 @@ public:
 		Ar << *this;
 		return true;
 	}
-	bool SerializeFromMismatchedTag(FName StructTag, FArchive& Ar);
+
+	bool SerializeFromMismatchedTag(FName StructTag, FArchive& Ar)
+	{
+		// FTransform3f/d are immutable, but FTransform is not and needs a per-property conversion
+		if(StructTag == NAME_Transform)
+		{
+			return false;
+		}
+
+		if constexpr (std::is_same_v<T, float>)
+		{
+			return UE_SERIALIZE_VARIANT_FROM_MISMATCHED_TAG(Ar, Transform, Transform3f, Transform3d);
+		}
+		else
+		{
+			return UE_SERIALIZE_VARIANT_FROM_MISMATCHED_TAG(Ar, Transform, Transform3d, Transform3f);
+		}
+	}
 
 	// Binary comparison operators.
 	/*
diff --git a/Engine/Source/Runtime/Core/Public/Math/TransformVectorized.h b/Engine/Source/Runtime/Core/Public/Math/TransformVectorized.h
--- a/Engine/Source/Runtime/Core/Public/Math/TransformVectorized.h
+++ b/Engine/Source/Runtime/Core/Public/Math/TransformVectorized.h
@@ -687,7 +687,23 @@ public:
 		return true;
 	}
 	
-	bool SerializeFromMismatchedTag(FName StructTag, FArchive& Ar);
+	bool SerializeFromMismatchedTag(FName StructTag, FArchive& Ar)
+	{
+		// FTransform3f/d are immutable, but FTransform is not and needs a per-property conversion
+		if(StructTag == NAME_Transform)
+		{
+			return false;
+		}
+
+		if constexpr (std::is_same_v<T, float>)
+		{
+			return UE_SERIALIZE_VARIANT_FROM_MISMATCHED_TAG(Ar, Transform, Transform3f, Transform3d);
+		}
+		else
+		{
+			return UE_SERIALIZE_VARIANT_FROM_MISMATCHED_TAG(Ar, Transform, Transform3d, Transform3f);
+		}
+	}
 	
 	// Binary comparison operators.
 	/*
diff --git a/Engine/Source/Runtime/Core/Public/Math/UnrealMathUtility.h b/Engine/Source/Runtime/Core/Public/Math/UnrealMathUtility.h
--- a/Engine/Source/Runtime/Core/Public/Math/UnrealMathUtility.h
+++ b/Engine/Source/Runtime/Core/Public/Math/UnrealMathUtility.h
@@ -2609,14 +2609,6 @@ public:
 
 };
 
-extern template CORE_API void FMath::SegmentDistToSegmentSafe<float>(UE::Math::TVector<float> A1, UE::Math::TVector<float> B1, UE::Math::TVector<float> A2, UE::Math::TVector<float> B2, UE::Math::TVector<float>& OutP1, UE::Math::TVector<float>& OutP2);
-extern template CORE_API void FMath::SegmentDistToSegmentSafe<double>(UE::Math::TVector<double> A1, UE::Math::TVector<double> B1, UE::Math::TVector<double> A2, UE::Math::TVector<double> B2, UE::Math::TVector<double>& OutP1, UE::Math::TVector<double>& OutP2);
-extern template CORE_API UE::Math::TVector<float> FMath::ClosestPointOnLine(const UE::Math::TVector<float>& LineStart, const UE::Math::TVector<float>& LineEnd, const UE::Math::TVector<float>& Point);
-extern template CORE_API UE::Math::TVector<double> FMath::ClosestPointOnLine(const UE::Math::TVector<double>& LineStart, const UE::Math::TVector<double>& LineEnd, const UE::Math::TVector<double>& Point);
-extern template CORE_API UE::Math::TVector<float> FMath::ClosestPointOnSegment(const UE::Math::TVector<float>& Point, const UE::Math::TVector<float>& StartPoint, const UE::Math::TVector<float>& EndPoint);
-extern template CORE_API UE::Math::TVector<double> FMath::ClosestPointOnSegment(const UE::Math::TVector<double>& Point, const UE::Math::TVector<double>& StartPoint, const UE::Math::TVector<double>& EndPoint);
-
-
 // LWC Conversion helpers
 namespace UE
 {
diff --git a/Engine/Source/Runtime/Core/Public/Math/Vector.h b/Engine/Source/Runtime/Core/Public/Math/Vector.h
--- a/Engine/Source/Runtime/Core/Public/Math/Vector.h
+++ b/Engine/Source/Runtime/Core/Public/Math/Vector.h
@@ -1174,7 +1174,17 @@ public:
         return true;
     }
 
-	bool SerializeFromMismatchedTag(FName StructTag, FStructuredArchive::FSlot Slot);
+	bool SerializeFromMismatchedTag(FName StructTag, FStructuredArchive::FSlot Slot)
+	{
+		if constexpr (std::is_same_v<T, float>)
+		{
+			return UE_SERIALIZE_VARIANT_FROM_MISMATCHED_TAG(Slot, Vector, Vector3f, Vector3d);
+		}
+		else
+		{
+			return UE_SERIALIZE_VARIANT_FROM_MISMATCHED_TAG(Slot, Vector, Vector3d, Vector3f);
+		}
+	}
 	
     /** 
      * Network serialization function.
@@ -2711,20 +2721,6 @@ inline FVector FMath::VRand()
     return Result * (1.0f / Sqrt(L));
 }
 
-template<>
-inline bool FVector3f::SerializeFromMismatchedTag(FName StructTag, FStructuredArchive::FSlot Slot)
-{
-
-	return UE_SERIALIZE_VARIANT_FROM_MISMATCHED_TAG(Slot, Vector, Vector3f, Vector3d);
-}
-
-template<>
-inline bool FVector3d::SerializeFromMismatchedTag(FName StructTag, FStructuredArchive::FSlot Slot)
-{
-	return UE_SERIALIZE_VARIANT_FROM_MISMATCHED_TAG(Slot, Vector, Vector3d, Vector3f);
-}
-
-
 /* TVector2<T> inline functions
  *****************************************************************************/
 namespace UE {
diff --git a/Engine/Source/Runtime/Core/Public/Math/Vector4.h b/Engine/Source/Runtime/Core/Public/Math/Vector4.h
--- a/Engine/Source/Runtime/Core/Public/Math/Vector4.h
+++ b/Engine/Source/Runtime/Core/Public/Math/Vector4.h
@@ -562,7 +562,22 @@ public:
 		return true;
 	}
 
-	bool SerializeFromMismatchedTag(FName StructTag, FArchive& Ar);
+	bool SerializeFromMismatchedTag(FName StructTag, FArchive& Ar)
+	{
+		if(SerializeFromVector3(StructTag, Ar))
+		{
+			return true;
+		}
+
+		if constexpr (std::is_same_v<T, float>)
+		{
+			return UE_SERIALIZE_VARIANT_FROM_MISMATCHED_TAG(Ar, Vector4, Vector4f, Vector4d);
+		}
+		else
+		{
+			return UE_SERIALIZE_VARIANT_FROM_MISMATCHED_TAG(Ar, Vector4, Vector4d, Vector4f);
+		}
+	}
 	
 	// Conversion from other type: double->float
 	template<typename FArg UE_REQUIRES(std::is_same_v<FArg, double> && std::is_same_v<T, float>)>
@@ -978,30 +993,6 @@ template<> struct TCanBulkSerialize<FVector4d> { enum { Value = true }; };
 DECLARE_INTRINSIC_TYPE_LAYOUT(FVector4f);
 DECLARE_INTRINSIC_TYPE_LAYOUT(FVector4d);
 
-
-template<>
-inline bool FVector4f::SerializeFromMismatchedTag(FName StructTag, FArchive& Ar)
-{	
-	if(SerializeFromVector3(StructTag, Ar))
-	{
-		return true;
-	}
-
-	return UE_SERIALIZE_VARIANT_FROM_MISMATCHED_TAG(Ar, Vector4, Vector4f, Vector4d);
-}
-
-template<>
-inline bool FVector4d::SerializeFromMismatchedTag(FName StructTag, FArchive& Ar)
-{
-	if(SerializeFromVector3(StructTag, Ar))
-	{
-		return true;
-	}
-
-	return UE_SERIALIZE_VARIANT_FROM_MISMATCHED_TAG(Ar, Vector4, Vector4d, Vector4f);
-}
-
-
 /**
  * Calculates 3D Dot product of two 4D vectors.
  *
diff --git a/Engine/Source/Runtime/Core/Public/Misc/ExpressionParser.h b/Engine/Source/Runtime/Core/Public/Misc/ExpressionParser.h
--- a/Engine/Source/Runtime/Core/Public/Misc/ExpressionParser.h
+++ b/Engine/Source/Runtime/Core/Public/Misc/ExpressionParser.h
@@ -62,24 +62,4 @@ namespace ExpressionParser
 		TOperatorEvaluationEnvironment<ContextType> Env(InJumpTable, InContext);
 		return Evaluate(CompiledTokens, Env);
 	}
-
-	extern template CORE_API TLexResultType<ANSICHAR> Lex<ANSICHAR>(const ANSICHAR* InExpression, const TTokenDefinitions<ANSICHAR>& TokenDefinitions);
-	extern template CORE_API TLexResultType<UTF8CHAR> Lex<UTF8CHAR>(const UTF8CHAR* InExpression, const TTokenDefinitions<UTF8CHAR>& TokenDefinitions);
-	extern template CORE_API TLexResultType<WIDECHAR> Lex<WIDECHAR>(const WIDECHAR* InExpression, const TTokenDefinitions<WIDECHAR>& TokenDefinitions);
-
-	extern template CORE_API TCompileResultType<ANSICHAR> Compile<ANSICHAR>(const ANSICHAR* InExpression, const TTokenDefinitions<ANSICHAR>& TokenDefinitions, const FExpressionGrammar& InGrammar);
-	extern template CORE_API TCompileResultType<UTF8CHAR> Compile<UTF8CHAR>(const UTF8CHAR* InExpression, const TTokenDefinitions<UTF8CHAR>& TokenDefinitions, const FExpressionGrammar& InGrammar);
-	extern template CORE_API TCompileResultType<WIDECHAR> Compile<WIDECHAR>(const WIDECHAR* InExpression, const TTokenDefinitions<WIDECHAR>& TokenDefinitions, const FExpressionGrammar& InGrammar);
-
-	extern template CORE_API TCompileResultType<ANSICHAR> Compile<ANSICHAR>(TArray<TExpressionToken<ANSICHAR>> InTokens, const FExpressionGrammar& InGrammar);
-	extern template CORE_API TCompileResultType<UTF8CHAR> Compile<UTF8CHAR>(TArray<TExpressionToken<UTF8CHAR>> InTokens, const FExpressionGrammar& InGrammar);
-	extern template CORE_API TCompileResultType<WIDECHAR> Compile<WIDECHAR>(TArray<TExpressionToken<WIDECHAR>> InTokens, const FExpressionGrammar& InGrammar);
-
-	extern template CORE_API FExpressionResult Evaluate<ANSICHAR>(const ANSICHAR* InExpression, const TTokenDefinitions<ANSICHAR>& InTokenDefinitions, const FExpressionGrammar& InGrammar, const TIOperatorEvaluationEnvironment<ANSICHAR>& InEnvironment);
-	extern template CORE_API FExpressionResult Evaluate<UTF8CHAR>(const UTF8CHAR* InExpression, const TTokenDefinitions<UTF8CHAR>& InTokenDefinitions, const FExpressionGrammar& InGrammar, const TIOperatorEvaluationEnvironment<UTF8CHAR>& InEnvironment);
-	extern template CORE_API FExpressionResult Evaluate<WIDECHAR>(const WIDECHAR* InExpression, const TTokenDefinitions<WIDECHAR>& InTokenDefinitions, const FExpressionGrammar& InGrammar, const TIOperatorEvaluationEnvironment<WIDECHAR>& InEnvironment);
-
-	extern template CORE_API FExpressionResult Evaluate<ANSICHAR>(const TArray<TCompiledToken<ANSICHAR>>& CompiledTokens, const TIOperatorEvaluationEnvironment<ANSICHAR>& InEnvironment);
-	extern template CORE_API FExpressionResult Evaluate<UTF8CHAR>(const TArray<TCompiledToken<UTF8CHAR>>& CompiledTokens, const TIOperatorEvaluationEnvironment<UTF8CHAR>& InEnvironment);
-	extern template CORE_API FExpressionResult Evaluate<WIDECHAR>(const TArray<TCompiledToken<WIDECHAR>>& CompiledTokens, const TIOperatorEvaluationEnvironment<WIDECHAR>& InEnvironment);
 }
