[Core]
* Changed AutoRTFM.h to use a define for the static keyword in front of functions and also an explanation why they are static

diff --git a/Engine/Source/Programs/UnrealBuildTool/Configuration/UEBuildTarget.cs b/Engine/Source/Programs/UnrealBuildTool/Configuration/UEBuildTarget.cs
--- a/Engine/Source/Programs/UnrealBuildTool/Configuration/UEBuildTarget.cs
+++ b/Engine/Source/Programs/UnrealBuildTool/Configuration/UEBuildTarget.cs
@@ -5288,6 +5288,7 @@ namespace UnrealBuildTool
 			if (Rules.bUseHeaderUnitsForPch)
 			{
 				GlobalCompileEnvironment.Definitions.Add("UE_DEBUG_SECTION=");
+				GlobalCompileEnvironment.Definitions.Add("UE_HEADER_UNITS=");
 			}
 
 			// Toggle to enable vorbis for audio streaming where available
diff --git a/Engine/Source/Runtime/Core/Public/AutoRTFM/AutoRTFM.h b/Engine/Source/Runtime/Core/Public/AutoRTFM/AutoRTFM.h
--- a/Engine/Source/Runtime/Core/Public/AutoRTFM/AutoRTFM.h
+++ b/Engine/Source/Runtime/Core/Public/AutoRTFM/AutoRTFM.h
@@ -130,6 +130,14 @@
 
 #define UE_AUTORTFM_UNUSED(UNUSEDVAR) (void)UNUSEDVAR
 
+// It is critical that these functions are both static and forceinline to prevent binary sizes to explode
+// This is a trick to ensure that there will never be a non-inlined version of these functions that the linker can decide to use
+#ifndef UE_HEADER_UNITS
+#define UE_AUTORTFM_CRITICAL_INLINE static UE_AUTORTFM_FORCEINLINE
+#else
+#define UE_AUTORTFM_CRITICAL_INLINE UE_AUTORTFM_FORCEINLINE // TODO: This needs to be revisited. we don't want bloated executables when modules are enabled
+#endif
+
 #ifdef __cplusplus
 extern "C"
 {
@@ -258,7 +266,7 @@ typedef struct
 //               Must be non-null.
 UE_AUTORTFM_API void autortfm_initialize(const autortfm_extern_api* ExternAPI);
 #else
-static UE_AUTORTFM_FORCEINLINE void autortfm_initialize(const autortfm_extern_api* ExternAPI)
+UE_AUTORTFM_CRITICAL_INLINE void autortfm_initialize(const autortfm_extern_api* ExternAPI)
 {
 	UE_AUTORTFM_UNUSED(ExternAPI);
 }
@@ -269,7 +277,7 @@ static UE_AUTORTFM_FORCEINLINE void autortfm_initialize(const autortfm_extern_ap
 // The AutoRTFM compiler will replace all calls to this function with a constant boolean value.
 UE_AUTORTFM_API bool autortfm_is_closed(void);
 #else
-static UE_AUTORTFM_FORCEINLINE bool autortfm_is_closed(void)
+UE_AUTORTFM_CRITICAL_INLINE bool autortfm_is_closed(void)
 {
     return false;
 }
@@ -278,7 +286,7 @@ static UE_AUTORTFM_FORCEINLINE bool autortfm_is_closed(void)
 #if UE_AUTORTFM_ENABLED
 UE_AUTORTFM_API bool autortfm_is_transactional(void);
 #else
-static UE_AUTORTFM_FORCEINLINE bool autortfm_is_transactional(void)
+UE_AUTORTFM_CRITICAL_INLINE bool autortfm_is_transactional(void)
 {
     return false;
 }
@@ -287,7 +295,7 @@ static UE_AUTORTFM_FORCEINLINE bool autortfm_is_transactional(void)
 #if UE_AUTORTFM_ENABLED
 UE_AUTORTFM_API bool autortfm_is_committing_or_aborting(void);
 #else
-static UE_AUTORTFM_FORCEINLINE bool autortfm_is_committing_or_aborting(void)
+UE_AUTORTFM_CRITICAL_INLINE bool autortfm_is_committing_or_aborting(void)
 {
 	return false;
 }
@@ -296,7 +304,7 @@ static UE_AUTORTFM_FORCEINLINE bool autortfm_is_committing_or_aborting(void)
 #if UE_AUTORTFM_ENABLED
 UE_AUTORTFM_API autortfm_result autortfm_transact(void (*UninstrumentedWork)(void*), void (*InstrumentedWork)(void*), void* Arg);
 #else
-static UE_AUTORTFM_FORCEINLINE autortfm_result autortfm_transact(void (*UninstrumentedWork)(void*), void (*InstrumentedWork)(void*), void* Arg)
+UE_AUTORTFM_CRITICAL_INLINE autortfm_result autortfm_transact(void (*UninstrumentedWork)(void*), void (*InstrumentedWork)(void*), void* Arg)
 {
 	UE_AUTORTFM_UNUSED(InstrumentedWork);
 	UninstrumentedWork(Arg);
@@ -307,7 +315,7 @@ static UE_AUTORTFM_FORCEINLINE autortfm_result autortfm_transact(void (*Uninstru
 #if UE_AUTORTFM_ENABLED
 UE_AUTORTFM_API autortfm_result autortfm_transact_then_open(void (*UninstrumentedWork)(void*), void (*InstrumentedWork)(void*), void* Arg);
 #else
-static UE_AUTORTFM_FORCEINLINE autortfm_result autortfm_transact_then_open(void (*UninstrumentedWork)(void*), void (*InstrumentedWork)(void*), void* Arg)
+UE_AUTORTFM_CRITICAL_INLINE autortfm_result autortfm_transact_then_open(void (*UninstrumentedWork)(void*), void (*InstrumentedWork)(void*), void* Arg)
 {
 	UE_AUTORTFM_UNUSED(InstrumentedWork);
 	UninstrumentedWork(Arg);
@@ -318,7 +326,7 @@ static UE_AUTORTFM_FORCEINLINE autortfm_result autortfm_transact_then_open(void
 #if UE_AUTORTFM_ENABLED
 UE_AUTORTFM_API void autortfm_commit(void (*UninstrumentedWork)(void*), void (*InstrumentedWork)(void*), void* Arg);
 #else
-static UE_AUTORTFM_FORCEINLINE void autortfm_commit(void (*UninstrumentedWork)(void*), void (*InstrumentedWork)(void*), void* Arg)
+UE_AUTORTFM_CRITICAL_INLINE void autortfm_commit(void (*UninstrumentedWork)(void*), void (*InstrumentedWork)(void*), void* Arg)
 {
 	UE_AUTORTFM_UNUSED(InstrumentedWork);
 	UninstrumentedWork(Arg);
@@ -328,62 +336,62 @@ static UE_AUTORTFM_FORCEINLINE void autortfm_commit(void (*UninstrumentedWork)(v
 #if UE_AUTORTFM_ENABLED
 UE_AUTORTFM_API void autortfm_abort_transaction();
 #else
-static UE_AUTORTFM_FORCEINLINE void autortfm_abort_transaction() {}
+UE_AUTORTFM_CRITICAL_INLINE void autortfm_abort_transaction() {}
 #endif
 
 #if UE_AUTORTFM_ENABLED
 UE_AUTORTFM_API autortfm_result autortfm_rollback_transaction();
 #else
-static UE_AUTORTFM_FORCEINLINE autortfm_result autortfm_rollback_transaction() { return autortfm_aborted_by_request; }
+UE_AUTORTFM_CRITICAL_INLINE autortfm_result autortfm_rollback_transaction() { return autortfm_aborted_by_request; }
 #endif
 
 #if UE_AUTORTFM_ENABLED
 UE_AUTORTFM_API autortfm_result autortfm_cascading_rollback_transaction();
 #else
-static UE_AUTORTFM_FORCEINLINE autortfm_result autortfm_cascading_rollback_transaction() { return autortfm_aborted_by_cascade; }
+UE_AUTORTFM_CRITICAL_INLINE autortfm_result autortfm_cascading_rollback_transaction() { return autortfm_aborted_by_cascade; }
 #endif
 
 #if UE_AUTORTFM_ENABLED
 UE_AUTORTFM_API bool autortfm_start_transaction();
 #else
-static UE_AUTORTFM_FORCEINLINE bool autortfm_start_transaction() { return false; }
+UE_AUTORTFM_CRITICAL_INLINE bool autortfm_start_transaction() { return false; }
 #endif
 
 #if UE_AUTORTFM_ENABLED
 UE_AUTORTFM_API autortfm_result autortfm_commit_transaction();
 #else
-static UE_AUTORTFM_FORCEINLINE autortfm_result autortfm_commit_transaction() { return autortfm_aborted_by_language; }
+UE_AUTORTFM_CRITICAL_INLINE autortfm_result autortfm_commit_transaction() { return autortfm_aborted_by_language; }
 #endif
 
 #if UE_AUTORTFM_ENABLED
 UE_AUTORTFM_API void autortfm_clear_transaction_status();
 #else
-static UE_AUTORTFM_FORCEINLINE void autortfm_clear_transaction_status() {}
+UE_AUTORTFM_CRITICAL_INLINE void autortfm_clear_transaction_status() {}
 #endif
 
 #if UE_AUTORTFM_ENABLED
 UE_AUTORTFM_API autortfm_status autortfm_get_context_status();
 #else
-static UE_AUTORTFM_FORCEINLINE autortfm_status autortfm_get_context_status() { return autortfm_status::autortfm_status_idle; }
+UE_AUTORTFM_CRITICAL_INLINE autortfm_status autortfm_get_context_status() { return autortfm_status::autortfm_status_idle; }
 #endif
 
 #if UE_AUTORTFM_ENABLED
 UE_AUTORTFM_API void autortfm_open(void (*work)(void* arg), void* arg, const void* return_address);
 #else
-static UE_AUTORTFM_FORCEINLINE void autortfm_open(void (*work)(void* arg), void* arg, const void* /* return_address */) { work(arg); }
+UE_AUTORTFM_CRITICAL_INLINE void autortfm_open(void (*work)(void* arg), void* arg, const void* /* return_address */) { work(arg); }
 #endif
 
 #if UE_AUTORTFM_ENABLED
 UE_AUTORTFM_API void autortfm_open_explicit_validation(autortfm_memory_validation_level, void (*work)(void* arg), void* arg, const void* return_address);
 #else
-static UE_AUTORTFM_FORCEINLINE void autortfm_open_explicit_validation(autortfm_memory_validation_level, void (*work)(void* arg), void* arg, const void* /* return_address */) { work(arg); }
+UE_AUTORTFM_CRITICAL_INLINE void autortfm_open_explicit_validation(autortfm_memory_validation_level, void (*work)(void* arg), void* arg, const void* /* return_address */) { work(arg); }
 #endif
 
 #if UE_AUTORTFM_ENABLED
 [[nodiscard]] UE_AUTORTFM_API autortfm_status autortfm_close(void (*UninstrumentedWork)(void*), void (*InstrumentedWork)(void*), void* Arg);
 #else
 AUTORTFM_DISABLE_UNREACHABLE_CODE_WARNINGS
-[[nodiscard]] static UE_AUTORTFM_FORCEINLINE autortfm_status autortfm_close(void (*UninstrumentedWork)(void*), void (*InstrumentedWork)(void*), void* Arg)
+[[nodiscard]] UE_AUTORTFM_CRITICAL_INLINE autortfm_status autortfm_close(void (*UninstrumentedWork)(void*), void (*InstrumentedWork)(void*), void* Arg)
 {
 	UE_AUTORTFM_UNUSED(UninstrumentedWork);
 	UE_AUTORTFM_UNUSED(InstrumentedWork);
@@ -397,7 +405,7 @@ AUTORTFM_RESTORE_UNREACHABLE_CODE_WARNINGS
 #if UE_AUTORTFM_ENABLED
 UE_AUTORTFM_API void autortfm_record_open_write(void* Ptr, size_t Size);
 #else
-static UE_AUTORTFM_FORCEINLINE void autortfm_record_open_write(void* Ptr, size_t Size)
+UE_AUTORTFM_CRITICAL_INLINE void autortfm_record_open_write(void* Ptr, size_t Size)
 {
 	UE_AUTORTFM_UNUSED(Ptr);
 	UE_AUTORTFM_UNUSED(Size);
@@ -407,7 +415,7 @@ static UE_AUTORTFM_FORCEINLINE void autortfm_record_open_write(void* Ptr, size_t
 #if UE_AUTORTFM_ENABLED
 UE_AUTORTFM_API void autortfm_record_open_write_no_memory_validation(void* Ptr, size_t Size);
 #else
-static UE_AUTORTFM_FORCEINLINE void autortfm_record_open_write_no_memory_validation(void* Ptr, size_t Size)
+UE_AUTORTFM_CRITICAL_INLINE void autortfm_record_open_write_no_memory_validation(void* Ptr, size_t Size)
 {
 	UE_AUTORTFM_UNUSED(Ptr);
 	UE_AUTORTFM_UNUSED(Size);
@@ -440,11 +448,11 @@ struct autortfm_open_to_closed_table
 UE_AUTORTFM_API void autortfm_register_open_to_closed_functions(autortfm_open_to_closed_table* Table);
 UE_AUTORTFM_API void autortfm_unregister_open_to_closed_functions(autortfm_open_to_closed_table* Table);
 #else
-static UE_AUTORTFM_FORCEINLINE void autortfm_register_open_to_closed_functions(autortfm_open_to_closed_table* Table)
+UE_AUTORTFM_CRITICAL_INLINE void autortfm_register_open_to_closed_functions(autortfm_open_to_closed_table* Table)
 {
 	UE_AUTORTFM_UNUSED(Table);
 }
-static UE_AUTORTFM_FORCEINLINE void autortfm_unregister_open_to_closed_functions(autortfm_open_to_closed_table* Table)
+UE_AUTORTFM_CRITICAL_INLINE void autortfm_unregister_open_to_closed_functions(autortfm_open_to_closed_table* Table)
 {
 	UE_AUTORTFM_UNUSED(Table);
 }
@@ -453,7 +461,7 @@ static UE_AUTORTFM_FORCEINLINE void autortfm_unregister_open_to_closed_functions
 #if UE_AUTORTFM_ENABLED
 UE_AUTORTFM_API bool autortfm_is_on_current_transaction_stack(void* Ptr);
 #else
-static UE_AUTORTFM_FORCEINLINE bool autortfm_is_on_current_transaction_stack(void* Ptr)
+UE_AUTORTFM_CRITICAL_INLINE bool autortfm_is_on_current_transaction_stack(void* Ptr)
 {
 	UE_AUTORTFM_UNUSED(Ptr);
 	return false;
@@ -463,7 +471,7 @@ static UE_AUTORTFM_FORCEINLINE bool autortfm_is_on_current_transaction_stack(voi
 #if UE_AUTORTFM_ENABLED
 UE_AUTORTFM_API void autortfm_on_commit(void (*work)(void* arg), void* arg);
 #else
-static UE_AUTORTFM_FORCEINLINE void autortfm_on_commit(void (*work)(void* arg), void* arg)
+UE_AUTORTFM_CRITICAL_INLINE void autortfm_on_commit(void (*work)(void* arg), void* arg)
 {
     work(arg);
 }
@@ -472,7 +480,7 @@ static UE_AUTORTFM_FORCEINLINE void autortfm_on_commit(void (*work)(void* arg),
 #if UE_AUTORTFM_ENABLED
 UE_AUTORTFM_API void autortfm_on_abort(void (*work)(void* arg), void* arg);
 #else
-static UE_AUTORTFM_FORCEINLINE void autortfm_on_abort(void (*work)(void* arg), void* arg)
+UE_AUTORTFM_CRITICAL_INLINE void autortfm_on_abort(void (*work)(void* arg), void* arg)
 {
 	UE_AUTORTFM_UNUSED(work);
 	UE_AUTORTFM_UNUSED(arg);
@@ -483,14 +491,14 @@ static UE_AUTORTFM_FORCEINLINE void autortfm_on_abort(void (*work)(void* arg), v
 UE_AUTORTFM_API void autortfm_push_on_abort_handler(const void* key, void (*work)(void* arg), void* arg);
 UE_AUTORTFM_API void autortfm_pop_on_abort_handler(const void* key);
 #else
-static UE_AUTORTFM_FORCEINLINE void autortfm_push_on_abort_handler(const void* key, void (*work)(void* arg), void* arg)
+UE_AUTORTFM_CRITICAL_INLINE void autortfm_push_on_abort_handler(const void* key, void (*work)(void* arg), void* arg)
 {
 	UE_AUTORTFM_UNUSED(key);
 	UE_AUTORTFM_UNUSED(work);
 	UE_AUTORTFM_UNUSED(arg);
 }
 
-static UE_AUTORTFM_FORCEINLINE void autortfm_pop_on_abort_handler(const void* key)
+UE_AUTORTFM_CRITICAL_INLINE void autortfm_pop_on_abort_handler(const void* key)
 {
 	UE_AUTORTFM_UNUSED(key);
 }
@@ -499,7 +507,7 @@ static UE_AUTORTFM_FORCEINLINE void autortfm_pop_on_abort_handler(const void* ke
 #if UE_AUTORTFM_ENABLED
 UE_AUTORTFM_API void* autortfm_did_allocate(void* ptr, size_t size);
 #else
-static UE_AUTORTFM_FORCEINLINE void* autortfm_did_allocate(void* ptr, size_t size)
+UE_AUTORTFM_CRITICAL_INLINE void* autortfm_did_allocate(void* ptr, size_t size)
 {
 	UE_AUTORTFM_UNUSED(size);
     return ptr;
@@ -509,7 +517,7 @@ static UE_AUTORTFM_FORCEINLINE void* autortfm_did_allocate(void* ptr, size_t siz
 #if UE_AUTORTFM_ENABLED
 UE_AUTORTFM_API void autortfm_did_free(void* ptr);
 #else
-static UE_AUTORTFM_FORCEINLINE void autortfm_did_free(void* ptr)
+UE_AUTORTFM_CRITICAL_INLINE void autortfm_did_free(void* ptr)
 {
 	UE_AUTORTFM_UNUSED(ptr);
 }
@@ -524,7 +532,7 @@ static UE_AUTORTFM_FORCEINLINE void autortfm_did_free(void* ptr)
 #if UE_AUTORTFM_ENABLED
 UE_AUTORTFM_API void autortfm_check_abi(void* ptr, size_t size);
 #else
-static UE_AUTORTFM_FORCEINLINE void autortfm_check_abi(void* ptr, size_t size)
+UE_AUTORTFM_CRITICAL_INLINE void autortfm_check_abi(void* ptr, size_t size)
 {
 	UE_AUTORTFM_UNUSED(ptr);
 	UE_AUTORTFM_UNUSED(size);
@@ -534,7 +542,7 @@ static UE_AUTORTFM_FORCEINLINE void autortfm_check_abi(void* ptr, size_t size)
 #if UE_AUTORTFM_ENABLED
 UE_AUTORTFM_API void autortfm_unreachable(const char* Message);
 #else
-static UE_AUTORTFM_FORCEINLINE void autortfm_unreachable(const char* Message)
+UE_AUTORTFM_CRITICAL_INLINE void autortfm_unreachable(const char* Message)
 {
 	UE_AUTORTFM_UNUSED(Message);
 }
@@ -675,22 +683,22 @@ auto AutoRTFMLookupInstrumentedFunctorInvoker(const FunctorType& Functor) -> voi
 // Tells if we are currently running in a transaction. This will return true in an open nest
 // (see `Open`). This function is handled specially in the compiler and will be constant folded
 // as true in closed code, or preserved as a function call in open code.
-static UE_AUTORTFM_FORCEINLINE bool IsTransactional() { return autortfm_is_transactional(); }
+UE_AUTORTFM_CRITICAL_INLINE bool IsTransactional() { return autortfm_is_transactional(); }
 
 // Tells if we are currently running in the closed nest of a transaction. By default,
 // transactional code is in a closed nest; the only way to be in an open nest is to request it
 // via `Open`. This function is handled specially in the compiler and will be constant folded
 // as true in closed code, and false in open code.
-static UE_AUTORTFM_FORCEINLINE bool IsClosed() { return autortfm_is_closed(); }
+UE_AUTORTFM_CRITICAL_INLINE bool IsClosed() { return autortfm_is_closed(); }
 
 // Tells us if we are currently committing or aborting a transaction. This will return true
 // in an on-abort or on-commit.
-static UE_AUTORTFM_FORCEINLINE bool IsCommittingOrAborting() { return autortfm_is_committing_or_aborting(); }
+UE_AUTORTFM_CRITICAL_INLINE bool IsCommittingOrAborting() { return autortfm_is_committing_or_aborting(); }
 
 // Returns true if the passed-in pointer is on the stack of the currently-executing transaction.
 // This is occasionally necessary when writing OnAbort handlers for objects on the stack, since
 // we don't want to scribble on stack memory that might have been reused.
-static UE_AUTORTFM_FORCEINLINE bool IsOnCurrentTransactionStack(void* Ptr)
+UE_AUTORTFM_CRITICAL_INLINE bool IsOnCurrentTransactionStack(void* Ptr)
 {
 	return autortfm_is_on_current_transaction_stack(Ptr);
 }
@@ -704,7 +712,7 @@ static UE_AUTORTFM_FORCEINLINE bool IsOnCurrentTransactionStack(void* Ptr)
 //
 // If AutoRTFM is disabled, the code will be ran non-transactionally.
 template<typename FunctorType>
-static UE_AUTORTFM_FORCEINLINE ETransactionResult Transact(const FunctorType& Functor)
+UE_AUTORTFM_CRITICAL_INLINE ETransactionResult Transact(const FunctorType& Functor)
 {
 	ETransactionResult Result =
 		static_cast<ETransactionResult>(
@@ -721,7 +729,7 @@ static UE_AUTORTFM_FORCEINLINE ETransactionResult Transact(const FunctorType& Fu
 // be compiled with the AutoRTFM instrumentation of functions if the only
 // thing that's being invoked is a function in the open.
 template<typename FunctorType>
-static UE_AUTORTFM_FORCEINLINE ETransactionResult TransactThenOpen(const FunctorType& Functor)
+UE_AUTORTFM_CRITICAL_INLINE ETransactionResult TransactThenOpen(const FunctorType& Functor)
 {
 	ETransactionResult Result =
 		static_cast<ETransactionResult>(
@@ -737,7 +745,7 @@ static UE_AUTORTFM_FORCEINLINE ETransactionResult TransactThenOpen(const Functor
 // execution if the result is anything other than autortfm_committed.
 // Useful for testing.
 template<typename FunctorType>
-static UE_AUTORTFM_FORCEINLINE void Commit(const FunctorType& Functor)
+UE_AUTORTFM_CRITICAL_INLINE void Commit(const FunctorType& Functor)
 {
     autortfm_commit(
 		&AutoRTFMFunctorInvoker<FunctorType>,
@@ -747,7 +755,7 @@ static UE_AUTORTFM_FORCEINLINE void Commit(const FunctorType& Functor)
 
 // Ends a transaction while in the closed, discarding all effects. 
 // Sends control to the end of the transaction immediately.
-static UE_AUTORTFM_FORCEINLINE void AbortTransaction()
+UE_AUTORTFM_CRITICAL_INLINE void AbortTransaction()
 {
 	autortfm_abort_transaction();
 }
@@ -759,19 +767,19 @@ static UE_AUTORTFM_FORCEINLINE void AbortTransaction()
 // Transact.
 #if UE_AUTORTFM
 template <typename FunctorType>
-static UE_AUTORTFM_FORCEINLINE void CascadingAbortTransaction(FunctorType&& RunAfterAbort)
+UE_AUTORTFM_CRITICAL_INLINE void CascadingAbortTransaction(FunctorType&& RunAfterAbort)
 {
 	ForTheRuntime::CascadingAbortTransactionInternal(std::forward<FunctorType>(RunAfterAbort));
 }
 
-static UE_AUTORTFM_FORCEINLINE void CascadingAbortTransaction()
+UE_AUTORTFM_CRITICAL_INLINE void CascadingAbortTransaction()
 {
 	ForTheRuntime::CascadingAbortTransactionInternal({});
 }
 #else
 template <typename FunctorType>
-static UE_AUTORTFM_FORCEINLINE void CascadingAbortTransaction(FunctorType&&) {}
-static UE_AUTORTFM_FORCEINLINE void CascadingAbortTransaction()              {}
+UE_AUTORTFM_CRITICAL_INLINE void CascadingAbortTransaction(FunctorType&&) {}
+UE_AUTORTFM_CRITICAL_INLINE void CascadingAbortTransaction()              {}
 #endif
 
 namespace Detail
@@ -885,7 +893,7 @@ template
 	typename FunctorType = void,
 	typename ReturnType = decltype(std::declval<FunctorType>()())
 >
-static UE_AUTORTFM_FORCEINLINE ReturnType Open(const FunctorType& Functor)
+UE_AUTORTFM_CRITICAL_INLINE ReturnType Open(const FunctorType& Functor)
 {
 	static_assert(IsSafeToReturnFromOpen<ReturnType>,
 		"function return type is not safe to return from Open()");
@@ -959,7 +967,7 @@ static UE_AUTORTFM_FORCEINLINE ReturnType Open(const FunctorType& Functor)
 // (whether we are in open or closed code).
 //
 // Will crash if called outside of a transaction nest.
-template<typename FunctorType> [[nodiscard]] static UE_AUTORTFM_FORCEINLINE EContextStatus Close(const FunctorType& Functor)
+template<typename FunctorType> [[nodiscard]] UE_AUTORTFM_CRITICAL_INLINE EContextStatus Close(const FunctorType& Functor)
 {
     return static_cast<EContextStatus>(
 		autortfm_close(
@@ -974,7 +982,7 @@ template<typename FunctorType> [[nodiscard]] static UE_AUTORTFM_FORCEINLINE ECon
 // extreme caution. If this is called outside of a transaction then the
 // `RunAfterAbortBeforeRetryWork` callback is never executed.
 #if UE_AUTORTFM
-template<typename FunctorType> static UE_AUTORTFM_FORCEINLINE void CascadingRetryTransaction(FunctorType&& RunAfterAbortBeforeRetryWork)
+template<typename FunctorType> UE_AUTORTFM_CRITICAL_INLINE void CascadingRetryTransaction(FunctorType&& RunAfterAbortBeforeRetryWork)
 {
 	if (autortfm_is_closed())
 	{
@@ -982,7 +990,7 @@ template<typename FunctorType> static UE_AUTORTFM_FORCEINLINE void CascadingRetr
 	}
 }
 #else
-template<typename FunctorType> static UE_AUTORTFM_FORCEINLINE void CascadingRetryTransaction(FunctorType&&)
+template<typename FunctorType> UE_AUTORTFM_CRITICAL_INLINE void CascadingRetryTransaction(FunctorType&&)
 {
 }
 #endif
@@ -992,7 +1000,7 @@ template<typename FunctorType> static UE_AUTORTFM_FORCEINLINE void CascadingRetr
 // this just adds the work to the work deferred until the outer nest's commit.
 // If this is called outside a transaction or from an open nest then the work
 // happens immediately.
-template<typename FunctorType> static UE_AUTORTFM_FORCEINLINE void OnCommit(FunctorType&& Work)
+template<typename FunctorType> UE_AUTORTFM_CRITICAL_INLINE void OnCommit(FunctorType&& Work)
 {
 	if (autortfm_is_closed())
 	{
@@ -1008,13 +1016,13 @@ template<typename FunctorType> static UE_AUTORTFM_FORCEINLINE void OnCommit(Func
 // this just adds the work to the work deferred until the outer nest's commit.
 // If this is called outside a transaction or from an open nest then the work
 // happens immediately.
-template<typename FunctorType> static UE_AUTORTFM_FORCEINLINE void OnCommit(FunctorType&& Work) { Work(); }
+template<typename FunctorType> UE_AUTORTFM_CRITICAL_INLINE void OnCommit(FunctorType&& Work) { Work(); }
 #endif
 
 #if UE_AUTORTFM
 // Have some work happen when this transaction aborts. If this is called
 // outside a transaction or from an open nest then the work is ignored.
-template<typename FunctorType> static UE_AUTORTFM_FORCEINLINE void OnAbort(FunctorType&& Work)
+template<typename FunctorType> UE_AUTORTFM_CRITICAL_INLINE void OnAbort(FunctorType&& Work)
 {
 	if (autortfm_is_closed())
 	{
@@ -1024,7 +1032,7 @@ template<typename FunctorType> static UE_AUTORTFM_FORCEINLINE void OnAbort(Funct
 #else
 // Have some work happen when this transaction aborts. If this is called
 // outside a transaction or from an open nest then the work is ignored.
-template<typename FunctorType> static UE_AUTORTFM_FORCEINLINE void OnAbort(FunctorType&&) {}
+template<typename FunctorType> UE_AUTORTFM_CRITICAL_INLINE void OnAbort(FunctorType&&) {}
 #endif
 
 #if UE_AUTORTFM
@@ -1032,7 +1040,7 @@ template<typename FunctorType> static UE_AUTORTFM_FORCEINLINE void OnAbort(Funct
 // the handler can be unregistered (see `PopOnCommitHandler`). This is useful
 // for scoped mutations that need an abort handler present unless execution
 // reaches the end of the relevant scope.
-template<typename FunctorType> static UE_AUTORTFM_FORCEINLINE void PushOnCommitHandler(const void* Key, FunctorType&& Work)
+template<typename FunctorType> UE_AUTORTFM_CRITICAL_INLINE void PushOnCommitHandler(const void* Key, FunctorType&& Work)
 {
 	if (autortfm_is_closed())
 	{
@@ -1044,12 +1052,12 @@ template<typename FunctorType> static UE_AUTORTFM_FORCEINLINE void PushOnCommitH
 // the handler can be unregistered (see `PopOnCommitHandler`). This is useful
 // for scoped mutations that need an abort handler present unless execution
 // reaches the end of the relevant scope.
-template<typename FunctorType> static UE_AUTORTFM_FORCEINLINE void PushOnCommitHandler(const void*, FunctorType&&) {}
+template<typename FunctorType> UE_AUTORTFM_CRITICAL_INLINE void PushOnCommitHandler(const void*, FunctorType&&) {}
 #endif
 
 #if UE_AUTORTFM
 // Unregister the most recently pushed handler (via `PushOnCommitHandler`) for the given key.
-static UE_AUTORTFM_FORCEINLINE void PopOnCommitHandler(const void* Key)
+UE_AUTORTFM_CRITICAL_INLINE void PopOnCommitHandler(const void* Key)
 {
 	if (autortfm_is_closed())
 	{
@@ -1058,12 +1066,12 @@ static UE_AUTORTFM_FORCEINLINE void PopOnCommitHandler(const void* Key)
 }
 #else
 // Unregister the most recently pushed handler (via `PushOnCommitHandler`) for the given key.
-static UE_AUTORTFM_FORCEINLINE void PopOnCommitHandler(const void*) {}
+UE_AUTORTFM_CRITICAL_INLINE void PopOnCommitHandler(const void*) {}
 #endif
 
 #if UE_AUTORTFM
 // Unregister all pushed handlers (via `PushOnCommitHandler`) for the given key.
-static UE_AUTORTFM_FORCEINLINE void PopAllOnCommitHandlers(const void* Key)
+UE_AUTORTFM_CRITICAL_INLINE void PopAllOnCommitHandlers(const void* Key)
 {
 	if (autortfm_is_closed())
 	{
@@ -1072,7 +1080,7 @@ static UE_AUTORTFM_FORCEINLINE void PopAllOnCommitHandlers(const void* Key)
 }
 #else
 // Unregister all pushed handlers (via `PushOnCommitHandler`) for the given key.
-static UE_AUTORTFM_FORCEINLINE void PopAllOnCommitHandlers(const void*) {}
+UE_AUTORTFM_CRITICAL_INLINE void PopAllOnCommitHandlers(const void*) {}
 #endif
 
 #if UE_AUTORTFM
@@ -1080,7 +1088,7 @@ static UE_AUTORTFM_FORCEINLINE void PopAllOnCommitHandlers(const void*) {}
 // the handler can be unregistered (see `PopOnAbortHandler`). This is useful
 // for scoped mutations that need an abort handler present unless execution
 // reaches the end of the relevant scope.
-template<typename FunctorType> static UE_AUTORTFM_FORCEINLINE void PushOnAbortHandler(const void* Key, FunctorType&& Work)
+template<typename FunctorType> UE_AUTORTFM_CRITICAL_INLINE void PushOnAbortHandler(const void* Key, FunctorType&& Work)
 {
 	if (autortfm_is_closed())
 	{
@@ -1092,12 +1100,12 @@ template<typename FunctorType> static UE_AUTORTFM_FORCEINLINE void PushOnAbortHa
 // the handler can be unregistered (see `PopOnAbortHandler`). This is useful
 // for scoped mutations that need an abort handler present unless execution
 // reaches the end of the relevant scope.
-template<typename FunctorType> static UE_AUTORTFM_FORCEINLINE void PushOnAbortHandler(const void* Key, FunctorType&&) {}
+template<typename FunctorType> UE_AUTORTFM_CRITICAL_INLINE void PushOnAbortHandler(const void* Key, FunctorType&&) {}
 #endif
 
 #if UE_AUTORTFM
 // Unregister the most recently pushed handler (via `PushOnAbortHandler`) for the given key.
-static UE_AUTORTFM_FORCEINLINE void PopOnAbortHandler(const void* Key)
+UE_AUTORTFM_CRITICAL_INLINE void PopOnAbortHandler(const void* Key)
 {
 	if (autortfm_is_closed())
 	{
@@ -1106,7 +1114,7 @@ static UE_AUTORTFM_FORCEINLINE void PopOnAbortHandler(const void* Key)
 }
 #else
 // Unregister the most recently pushed handler (via `PushOnAbortHandler`) for the given key.
-static UE_AUTORTFM_FORCEINLINE void PopOnAbortHandler(const void* Key)
+UE_AUTORTFM_CRITICAL_INLINE void PopOnAbortHandler(const void* Key)
 {
 	UE_AUTORTFM_UNUSED(Key);
 }
@@ -1114,7 +1122,7 @@ static UE_AUTORTFM_FORCEINLINE void PopOnAbortHandler(const void* Key)
 
 #if UE_AUTORTFM
 // Unregister all pushed handlers (via `PushOnAbortHandler`) for the given key.
-static UE_AUTORTFM_FORCEINLINE void PopAllOnAbortHandlers(const void* Key)
+UE_AUTORTFM_CRITICAL_INLINE void PopAllOnAbortHandlers(const void* Key)
 {
 	if (autortfm_is_closed())
 	{
@@ -1123,7 +1131,7 @@ static UE_AUTORTFM_FORCEINLINE void PopAllOnAbortHandlers(const void* Key)
 }
 #else
 // Unregister all pushed handlers (via `PushOnAbortHandler`) for the given key.
-static UE_AUTORTFM_FORCEINLINE void PopAllOnAbortHandlers(const void* Key)
+UE_AUTORTFM_CRITICAL_INLINE void PopAllOnAbortHandlers(const void* Key)
 {
 	UE_AUTORTFM_UNUSED(Key);
 }
@@ -1135,13 +1143,13 @@ static UE_AUTORTFM_FORCEINLINE void PopAllOnAbortHandlers(const void* Key)
 // returned by this function. It's guaranteed to be equal to the pointer you
 // passed, but it's blessed specially from the compiler's perspective, leading
 // to some nice optimizations. This does nothing when called from open code.
-static UE_AUTORTFM_FORCEINLINE void* DidAllocate(void* Ptr, size_t Size)
+UE_AUTORTFM_CRITICAL_INLINE void* DidAllocate(void* Ptr, size_t Size)
 {
     return autortfm_did_allocate(Ptr, Size);
 }
 
 // Inform the runtime that we have free'd a given memory location.
-static UE_AUTORTFM_FORCEINLINE void DidFree(void* Ptr)
+UE_AUTORTFM_CRITICAL_INLINE void DidFree(void* Ptr)
 {
     autortfm_did_free(Ptr);
 }
@@ -1151,7 +1159,7 @@ static UE_AUTORTFM_FORCEINLINE void DidFree(void* Ptr)
 // a possible abort. Normally, tracking memory overwrites should be automatically
 // handled by AutoRTFM, but manual overwrite tracking may be required for third-party
 // libraries or outside compilers (such as ISPC).
-static UE_AUTORTFM_FORCEINLINE void RecordOpenWrite(void* Ptr, size_t Size)
+UE_AUTORTFM_CRITICAL_INLINE void RecordOpenWrite(void* Ptr, size_t Size)
 {
 	autortfm_record_open_write(Ptr, Size);
 }
@@ -1161,32 +1169,32 @@ static UE_AUTORTFM_FORCEINLINE void RecordOpenWrite(void* Ptr, size_t Size)
 // a possible abort. Normally, tracking memory overwrites should be automatically
 // handled by AutoRTFM, but manual overwrite tracking may be required for third-party
 // libraries or outside compilers (such as ISPC).
-template<typename Type> static UE_AUTORTFM_FORCEINLINE void RecordOpenWrite(Type* Ptr)
+template<typename Type> UE_AUTORTFM_CRITICAL_INLINE void RecordOpenWrite(Type* Ptr)
 {
 	autortfm_record_open_write(Ptr, sizeof(Type));
 }
 
 // Same as RecordOpenWrite() but marks the write as ignorable by the memory validator.
-static UE_AUTORTFM_FORCEINLINE void RecordOpenWriteNoMemoryValidation(void* Ptr, size_t Size)
+UE_AUTORTFM_CRITICAL_INLINE void RecordOpenWriteNoMemoryValidation(void* Ptr, size_t Size)
 {
 	autortfm_record_open_write_no_memory_validation(Ptr, Size);
 }
 
 // Same as RecordOpenWrite() but marks the write as ignorable by the memory validator.
-template<typename Type> static UE_AUTORTFM_FORCEINLINE void RecordOpenWriteNoMemoryValidation(Type* Ptr)
+template<typename Type> UE_AUTORTFM_CRITICAL_INLINE void RecordOpenWriteNoMemoryValidation(Type* Ptr)
 {
 	autortfm_record_open_write_no_memory_validation(Ptr, sizeof(Type));
 }
 
 // Report that a unreachable codepath is being hit. Used to manually ban certain codepaths
 // from being transactionally safe.
-static UE_AUTORTFM_FORCEINLINE void Unreachable(const char* Message = nullptr)
+UE_AUTORTFM_CRITICAL_INLINE void Unreachable(const char* Message = nullptr)
 {
 	autortfm_unreachable(Message);
 }
 
 // If we are running within a transaction, call `AutoRTFM::Unreachable`.
-static UE_AUTORTFM_FORCEINLINE void UnreachableIfTransactional(const char* Message = nullptr)
+UE_AUTORTFM_CRITICAL_INLINE void UnreachableIfTransactional(const char* Message = nullptr)
 {
 	if (AutoRTFM::IsTransactional())
 	{
@@ -1274,7 +1282,7 @@ namespace ForTheRuntime
 	UE_AUTORTFM_API bool IsAutoRTFMRuntimeEnabledInternal();
 
 	// Query whether the AutoRTFM runtime is enabled.
-	static UE_AUTORTFM_FORCEINLINE bool IsAutoRTFMRuntimeEnabled()
+	UE_AUTORTFM_CRITICAL_INLINE bool IsAutoRTFMRuntimeEnabled()
 	{
 		// If we are already in the closed nest of a transaction, we must have our runtime enabled!
 		if (AutoRTFM::IsClosed())
@@ -1364,7 +1372,7 @@ namespace ForTheRuntime
 	// Manually create a new transaction from open code and push it as a transaction nest.
 	// Can only be called within an already active parent transaction (EG. this cannot start
 	// a transaction nest itself).
-	static UE_AUTORTFM_FORCEINLINE bool StartTransaction()
+	UE_AUTORTFM_CRITICAL_INLINE bool StartTransaction()
 	{
 		return autortfm_start_transaction();
 	}
@@ -1372,23 +1380,23 @@ namespace ForTheRuntime
 	// Manually commit the top transaction nest, popping it from the execution scope.
 	// Can only be called within an already active parent transaction (EG. this cannot end
 	// a transaction nest itself).
-	static UE_AUTORTFM_FORCEINLINE ETransactionResult CommitTransaction()
+	UE_AUTORTFM_CRITICAL_INLINE ETransactionResult CommitTransaction()
 	{
 		return static_cast<ETransactionResult>(autortfm_commit_transaction());
 	}
 
 	// Manually clear the status of a user abort from the top transaction in a nest.
-	static UE_AUTORTFM_FORCEINLINE void ClearTransactionStatus()
+	UE_AUTORTFM_CRITICAL_INLINE void ClearTransactionStatus()
 	{
 		autortfm_clear_transaction_status();
 	}
 
-	static UE_AUTORTFM_FORCEINLINE EContextStatus GetContextStatus()
+	UE_AUTORTFM_CRITICAL_INLINE EContextStatus GetContextStatus()
 	{
 		return static_cast<EContextStatus>(autortfm_get_context_status());
 	}
 
-	static UE_AUTORTFM_FORCEINLINE EContextStatus ResetContextStatus()
+	UE_AUTORTFM_CRITICAL_INLINE EContextStatus ResetContextStatus()
 	{
 		return static_cast<EContextStatus>(autortfm_get_context_status());
 	}
@@ -1396,7 +1404,7 @@ namespace ForTheRuntime
 	// Ends a transaction while in the open, discarding all effects.
 	// Control flow is unaffected while in the open. 
 	// Sends control to the end of the transaction after leaving the open.
-	static UE_AUTORTFM_FORCEINLINE ETransactionResult RollbackTransaction()
+	UE_AUTORTFM_CRITICAL_INLINE ETransactionResult RollbackTransaction()
 	{
 		return static_cast<ETransactionResult>(autortfm_rollback_transaction());
 	}
@@ -1405,7 +1413,7 @@ namespace ForTheRuntime
 	// meaning an abort of a nested transaction will cause all transactions in the
 	// nest to abort. Control flow is unaffected while in the open.
 	// Sends control to the end of the transaction after leaving the open.
-	static UE_AUTORTFM_FORCEINLINE ETransactionResult CascadingRollbackTransaction()
+	UE_AUTORTFM_CRITICAL_INLINE ETransactionResult CascadingRollbackTransaction()
 	{
 		return static_cast<ETransactionResult>(autortfm_cascading_rollback_transaction());
 	}
@@ -1478,10 +1486,10 @@ namespace ForTheRuntime
 	};
 
 	// Reserved for future.
-	static UE_AUTORTFM_FORCEINLINE void RecordOpenRead(void const*, size_t) {}
+	UE_AUTORTFM_CRITICAL_INLINE void RecordOpenRead(void const*, size_t) {}
 
 	// Reserved for future.
-	template<typename Type> static UE_AUTORTFM_FORCEINLINE void RecordOpenRead(Type*) {}
+	template<typename Type> UE_AUTORTFM_CRITICAL_INLINE void RecordOpenRead(Type*) {}
 
 } // namespace ForTheRuntime
 
