Fix TDeque's copy constructor and assignment by setting Tail correctly.

Tail was being incorrectly set to Count, which in the case of it being equal to Capacity, would then result in Tail=Capacity without wrapping back to 0.

#rnx

diff --git a/Engine/Source/Runtime/Core/Public/Containers/Deque.h b/Engine/Source/Runtime/Core/Public/Containers/Deque.h
--- a/Engine/Source/Runtime/Core/Public/Containers/Deque.h
+++ b/Engine/Source/Runtime/Core/Public/Containers/Deque.h
@@ -528,29 +528,9 @@ private:
 			ConstructItems<ElementType>(GetData(), Other.GetData() + Other.Head, HeadToEndOffset);
 			ConstructItems<ElementType>(GetData() + HeadToEndOffset, Other.GetData(), Other.Tail);
 		}
+		Count = Other.Count;
 		Head = 0;
-		Tail = Count = Other.Count;
-	}
-
-	/*
-	 * Moves the parameter container elements into this one.
-	 * The moved range is linearized.
-	 */
-	void MoveElements(TDeque& Other)
-	{
-		if (Other.Head < Other.Tail)
-		{
-			RelocateConstructItems<ElementType>(GetData(), Other.GetData(), Other.Count);
-		}
-		else
-		{
-			const SizeType HeadToEndOffset = Other.Capacity - Other.Head;
-			RelocateConstructItems<ElementType>(GetData(), Other.GetData() + Other.Head, HeadToEndOffset);
-			RelocateConstructItems<ElementType>(GetData() + HeadToEndOffset, Other.GetData(), Other.Tail);
-		}
-		Head = 0;
-		Tail = Count = Other.Count;
-		Other.Head = Other.Tail = Other.Count = 0;
+		Tail = UE::Deque::Private::WrapAround(Count, Capacity); 
 	}
 
 	/**
diff --git a/Engine/Source/Runtime/Core/Tests/Containers/DequeTest.cpp b/Engine/Source/Runtime/Core/Tests/Containers/DequeTest.cpp
--- a/Engine/Source/Runtime/Core/Tests/Containers/DequeTest.cpp
+++ b/Engine/Source/Runtime/Core/Tests/Containers/DequeTest.cpp
@@ -748,6 +748,34 @@ TEST_CASE_NAMED(FDequeCopyVariableSizeWithHeadTailWrapTest, "System::Core::Conta
 	CHECK(int32Token::EvenConstructionDestructionCalls());
 }
 
+TEST_CASE_NAMED(FDequeCopyFullQueueAndPopFirstEmplaceLastTest, "System::Core::Containers::TDeque::Copy with source deque at capacity and PopFirst/EmplaceLast on copied deque", "[SmokeFilter][Core][Containers][Deque]")
+{
+	int32Token::Reset();
+	{
+		TDeque<int32Token> Deque;
+		Deque.Reserve(Test::DefaultCapacity);
+		CHECK(Deque.Max() == Test::DefaultCapacity);
+		for (int32 Value = 0; Value < Test::DefaultCapacity; ++Value)
+		{
+			Deque.EmplaceLast(Value);
+		}
+		CHECK(Deque.Num() == Test::DefaultCapacity);
+
+		TDeque<int32Token> DestQueue(Deque);
+		CHECK(Deque == DestQueue);
+		CHECK(DestQueue.Max() == Test::DefaultCapacity);
+
+		DestQueue.PopFirst();
+		DestQueue.EmplaceLast(Test::DefaultCapacity);
+		int32 CheckValue = 1;
+		for (const auto& Value : DestQueue)
+		{
+			CHECK(Value == CheckValue++);
+		}
+	}
+	CHECK(int32Token::EvenConstructionDestructionCalls());
+}
+
 TEST_CASE_NAMED(FDequeCopyWithInlineAllocator, "System::Core::Containers::TDeque Copy with inline allocator", "[SmokeFilter][Core][Containers][Deque]")
 {
 	SECTION("Copy empty TDeque")
