Made TNotNull<T>::operator() fail under SFINAE if T is not callable.  This allows TNotNull<NotCallable*> and TFunction<...> to be overloaded.

diff --git a/Engine/Source/Runtime/Core/Public/Misc/NotNull.h b/Engine/Source/Runtime/Core/Public/Misc/NotNull.h
--- a/Engine/Source/Runtime/Core/Public/Misc/NotNull.h
+++ b/Engine/Source/Runtime/Core/Public/Misc/NotNull.h
@@ -95,6 +95,10 @@
 	template <typename T>
 	struct TNotNull
 	{
+	private:
+		T Val;
+
+	public:
 		template <typename>
 		friend struct TNotNull;
 
@@ -185,15 +189,12 @@
 
 		// Allow function call syntax, if it's supported
 		template <typename... ArgTypes>
-		decltype(auto) operator()(ArgTypes&&... Args) const
+		auto operator()(ArgTypes&&... Args) const -> decltype(this->Val((ArgTypes&&)Args...))
 		{
 			UE_ASSUME(Val);
 			return Val((ArgTypes&&)Args...);
 		}
 
-	private:
-		T Val;
-
 		// Allow hashing, if it's supported
 		friend uint32 GetTypeHash(const TNotNull& NotNull)
 		{
