If we fail to iterate on a directory using JournalIterateDirectory, fall back and/or disable the directory cache
As well as send some telemetry on the directory, error code and file attributes to give an idea what the issue is

diff --git a/Engine/Plugins/Experimental/EditorTelemetry/Source/Private/EditorTelemetry.cpp b/Engine/Plugins/Experimental/EditorTelemetry/Source/Private/EditorTelemetry.cpp
--- a/Engine/Plugins/Experimental/EditorTelemetry/Source/Private/EditorTelemetry.cpp
+++ b/Engine/Plugins/Experimental/EditorTelemetry/Source/Private/EditorTelemetry.cpp
@@ -1036,6 +1036,14 @@ void FEditorTelemetry::StartSession()
 				{ TEXT("InitialSearchCompleted"), Data.bInitialSearchCompleted },
 			});
 		});
+		Router.OnTelemetry<FFileJournalTelemetry>([this](const FFileJournalTelemetry& Data){
+			FStudioTelemetry::Get().RecordEvent(TEXT("Editor.AssetRegistry.FileJournalError"),
+			{
+				{ TEXT("SchemaVersion"), 1 },
+				{ TEXT("Directory"), Data.Directory },
+				{ TEXT("ErrorString"), Data.ErrorString },
+			});
+		});
 	}
 
 	{
diff --git a/Engine/Source/Runtime/AssetRegistry/Private/AssetDataGatherer.cpp b/Engine/Source/Runtime/AssetRegistry/Private/AssetDataGatherer.cpp
--- a/Engine/Source/Runtime/AssetRegistry/Private/AssetDataGatherer.cpp
+++ b/Engine/Source/Runtime/AssetRegistry/Private/AssetDataGatherer.cpp
@@ -3,6 +3,7 @@
 #include "AssetDataGatherer.h"
 #include "AssetDataGathererPrivate.h"
 #include "AssetRegistry.h"
+#include "AssetRegistry/AssetRegistryTelemetry.h"
 
 #include "Algo/AnyOf.h"
 #include "Algo/Find.h"
@@ -36,6 +37,7 @@
 #include "Serialization/Archive.h"
 #include "String/LexFromString.h"
 #include "Tasks/Task.h"
+#include "TelemetryRouter.h"
 #include "Templates/UniquePtr.h"
 #include "Templates/UnrealTemplate.h"
 #include "Trace/Trace.h"
@@ -2134,18 +2136,49 @@ void FAssetDataDiscovery::TickInternal(bool bTickAll)
 				return true;
 			};
 
+			bool bIteratedDirectory = false;
 			if (Cache.IsWriteEnabled() != EFeatureEnabled::Never)
 			{
-				FPlatformFileManager::Get().GetPlatformFile().FileJournalIterateDirectory(Data.DirLocalAbsPath.ToString(),
+				// If we fail to iterate this directory, fall back to the old way. Meaning we will fail to cache this but still pick up on the assets in the directories
+				if (FPlatformFileManager::Get().GetPlatformFile().FileJournalIterateDirectory(Data.DirLocalAbsPath.ToString(),
 					[this, &Data, &CacheDataToAdd, &ProcessIterData]
 					(const TCHAR* IterFilename, const FFileJournalData& IterData)
 					{
 						return ProcessIterData(IterFilename, IterData.bIsDirectory, IterData.ModificationTime, IterData.JournalHandle, IterData.bIsReparsePoint);
-					});
+					}))
+				{
+					bIteratedDirectory = true;
+					Cache.QueueAdd(FString(Data.DirLocalAbsPath), MoveTemp(CacheDataToAdd));
+				}
+				else
+				{
+					// Only run this once to capture more information on why we fail sometimes here
+					static std::atomic<bool> bRunOnce = false;
+					if (bRunOnce)
+					{
+						// If we failed once, run once more time but grab an error so we can send some telemetry on the issue
+						FString OutError;
+						if (!FPlatformFileManager::Get().GetPlatformFile().FileJournalIterateDirectory(Data.DirLocalAbsPath.ToString(),
+							[this, &Data, &CacheDataToAdd, &ProcessIterData]
+							(const TCHAR* IterFilename, const FFileJournalData& IterData)
+							{
+								return true; // if it doesnt return true it acts as a failure for this function
+							}, &OutError))
+						{
+							bRunOnce = true;
+
+							UE_LOG(LogAssetRegistry, Warning, TEXT("Failed to FileJournalIterateDirectory, failing to cache this directory due to:\n  %s"), *OutError);
 
-				Cache.QueueAdd(FString(Data.DirLocalAbsPath), MoveTemp(CacheDataToAdd));
+							UE::Telemetry::AssetRegistry::FFileJournalTelemetry Telemetry;
+							Telemetry.Directory = Data.DirLocalAbsPath.ToString();
+							Telemetry.ErrorString = OutError;
+							FTelemetryRouter::Get().ProvideTelemetry(Telemetry);
+						}
+					}
+				}
 			}
-			else
+
+			if (!bIteratedDirectory)
 			{
 				IFileManager::Get().IterateDirectoryStat(Data.DirLocalAbsPath.ToString(),
 					[this, &Data, &ProcessIterData]
diff --git a/Engine/Source/Runtime/AssetRegistry/Private/AssetDataGathererDiscoveryCache.cpp b/Engine/Source/Runtime/AssetRegistry/Private/AssetDataGathererDiscoveryCache.cpp
--- a/Engine/Source/Runtime/AssetRegistry/Private/AssetDataGathererDiscoveryCache.cpp
+++ b/Engine/Source/Runtime/AssetRegistry/Private/AssetDataGathererDiscoveryCache.cpp
@@ -4,6 +4,7 @@
 
 #include "AssetDataGathererPrivate.h"
 #include "AssetRegistryPrivate.h"
+#include "AssetRegistry/AssetRegistryTelemetry.h"
 #include "Compression/CompressedBuffer.h"
 #include "HAL/FileManager.h"
 #include "HAL/PlatformFileManager.h"
@@ -15,6 +16,7 @@
 #include "Serialization/Archive.h"
 #include "Serialization/LargeMemoryWriter.h"
 #include "Serialization/MemoryReader.h"
+#include "TelemetryRouter.h"
 #include "Templates/UniquePtr.h"
 
 FArchive& operator<<(FArchive& Ar, FFileJournalFileHandle& H)
@@ -63,6 +65,22 @@ void FAssetDataDiscoveryCache::LoadAndUpdateCache()
 		TestLatestJournalEntry, &TestError);
 	bool bPlatformSupported = TestResult == EFileJournalResult::Success;
 
+	if (bPlatformSupported)
+	{
+		// Do a quick test if can grab data from the journal file which does a similar call in IterateDirectory
+		// If it fails we are very likely going to fail to iterate using journal files, so lets disable the cache
+		FFileJournalData TestData = PlatformFile.FileJournalGetFileData(*FPaths::ConvertRelativePathToFull(ProjectDir), &TestError);
+		if (!TestData.bIsValid)
+		{
+			bPlatformSupported = false;
+
+			UE::Telemetry::AssetRegistry::FFileJournalTelemetry Telemetry;
+			Telemetry.Directory = FPaths::ConvertRelativePathToFull(ProjectDir);
+			Telemetry.ErrorString = TestError;
+			FTelemetryRouter::Get().ProvideTelemetry(Telemetry);
+		}
+	}
+
 	bool bInvalidateEnabled = GPreloadSettings.IsDiscoveryCacheInvalidateEnabled();
 	bool bReadEnabled = GPreloadSettings.IsDiscoveryCacheReadEnabled() && (bPlatformSupported || !bInvalidateEnabled);
 	// Precalculate IfPlatformSupported -> Never if we already know the project doesn't support it
diff --git a/Engine/Source/Runtime/AssetRegistry/Public/AssetRegistry/AssetRegistryTelemetry.h b/Engine/Source/Runtime/AssetRegistry/Public/AssetRegistry/AssetRegistryTelemetry.h
--- a/Engine/Source/Runtime/AssetRegistry/Public/AssetRegistry/AssetRegistryTelemetry.h
+++ b/Engine/Source/Runtime/AssetRegistry/Public/AssetRegistry/AssetRegistryTelemetry.h
@@ -86,4 +86,16 @@ namespace UE::Telemetry::AssetRegistry
 		// Whether an additional main background async gather was started at this time
 		bool bAdditionalMountSearchInProgress = false;
     };
-}
\ No newline at end of file
+
+	// Information about a failure in our FileJournal area
+	struct FFileJournalTelemetry
+	{
+		static constexpr inline FGuid TelemetryID = FGuid(0x761ebe19, 0x2f0c427e, 0xaec0f06d, 0xe116ff82);
+
+		// Directory that failed for our FileJournal calls
+		FString Directory;
+
+		// The error string returned from our lower level FileJournal calls
+		FString ErrorString;
+	};
+}
diff --git a/Engine/Source/Runtime/Core/Private/GenericPlatform/GenericPlatformFile.cpp b/Engine/Source/Runtime/Core/Private/GenericPlatform/GenericPlatformFile.cpp
--- a/Engine/Source/Runtime/Core/Private/GenericPlatform/GenericPlatformFile.cpp
+++ b/Engine/Source/Runtime/Core/Private/GenericPlatform/GenericPlatformFile.cpp
@@ -1092,7 +1092,7 @@ EFileJournalResult IPlatformFile::FileJournalGetLatestEntry(const TCHAR* VolumeO
 	return EFileJournalResult::InvalidPlatform;
 }
 
-bool IPlatformFile::FileJournalIterateDirectory(const TCHAR* Directory, FDirectoryJournalVisitorFunc Visitor)
+bool IPlatformFile::FileJournalIterateDirectory(const TCHAR* Directory, FDirectoryJournalVisitorFunc Visitor, FString* OutError)
 {
 	return this->IterateDirectoryStat(Directory,
 		[&Visitor](const TCHAR* InPackageFilename, const FFileStatData& StatData)
@@ -1101,7 +1101,7 @@ bool IPlatformFile::FileJournalIterateDirectory(const TCHAR* Directory, FDirecto
 		});
 }
 
-FFileJournalData IPlatformFile::FileJournalGetFileData(const TCHAR* FilenameOrDirectory)
+FFileJournalData IPlatformFile::FileJournalGetFileData(const TCHAR* FilenameOrDirectory, FString* OutError)
 {
 	return UE::PlatformFileJournal::Private::ToJournalData(this->GetStatData(FilenameOrDirectory));
 }
diff --git a/Engine/Source/Runtime/Core/Private/Windows/WindowsPlatformFileJournal.cpp b/Engine/Source/Runtime/Core/Private/Windows/WindowsPlatformFileJournal.cpp
--- a/Engine/Source/Runtime/Core/Private/Windows/WindowsPlatformFileJournal.cpp
+++ b/Engine/Source/Runtime/Core/Private/Windows/WindowsPlatformFileJournal.cpp
@@ -191,14 +191,14 @@ EFileJournalResult FWindowsPlatformFile::FileJournalGetLatestEntry(const TCHAR*
 }
 
 bool FWindowsPlatformFile::FileJournalIterateDirectory(const TCHAR* Directory,
-	FDirectoryJournalVisitorFunc Visitor)
+	FDirectoryJournalVisitorFunc Visitor, FString* OutError)
 {
 	using namespace UE::WindowsPlatformFileJournal::Private;
 	using namespace UE::WindowsPlatformFile::Private;
 
 	if (!FileJournalIsAvailableSystemWide())
 	{
-		return IPlatformFile::FileJournalIterateDirectory(Directory, Visitor);
+		return IPlatformFile::FileJournalIterateDirectory(Directory, Visitor, OutError);
 	}
 
 	FNormalizedDirectory Normalized(Directory);
@@ -211,6 +211,13 @@ bool FWindowsPlatformFile::FileJournalIterateDirectory(const TCHAR* Directory,
 		nullptr);
 	if (Handle == INVALID_HANDLE_VALUE)
 	{
+		if (OutError)
+		{
+			uint32 LastError = GetLastError();
+			*OutError = FString::Printf(TEXT("FileJournal: CreateFile('%s') failed; FileJournal feature is unavailable. LastError=%u (0x%08x) (%s)."),
+				*Normalized, LastError, LastError, *JournalGetLastErrorString(LastError));
+		}
+
 		return false;
 	}
 	ON_SCOPE_EXIT{ ::CloseHandle(Handle); };
@@ -223,6 +230,13 @@ bool FWindowsPlatformFile::FileJournalIterateDirectory(const TCHAR* Directory,
 	BOOL Result = GetFileInformationByHandleEx(Handle, ::FileIdExtdDirectoryRestartInfo, Buffer, BufferSize);
 	if (!Result)
 	{
+		if (OutError)
+		{
+			uint32 LastError = FPlatformMisc::GetLastError();
+			*OutError = FString::Printf(TEXT("Unexpected failure of GetFileInformationByHandleEx('%s') with error code %u (0x%08x) (%s)."),
+				*Normalized, LastError, LastError, *JournalGetLastErrorString(LastError));
+		}
+
 		// TODO: Need to handle filenames with size > buffersize by checking GetLastError and passing in a dynamically sized buffer.
 		return false;
 	}
@@ -279,7 +293,7 @@ bool FWindowsPlatformFile::FileJournalIterateDirectory(const TCHAR* Directory,
 	}
 }
 
-FFileJournalData FWindowsPlatformFile::FileJournalGetFileData(const TCHAR* FilenameOrDirectory)
+FFileJournalData FWindowsPlatformFile::FileJournalGetFileData(const TCHAR* FilenameOrDirectory, FString* OutError)
 {
 	using namespace UE::WindowsPlatformFileJournal::Private;
 	using namespace UE::WindowsPlatformFile::Private;
@@ -308,6 +322,12 @@ FFileJournalData FWindowsPlatformFile::FileJournalGetFileData(const TCHAR* Filen
 		nullptr);
 	if (Handle == INVALID_HANDLE_VALUE)
 	{
+		if (OutError)
+		{
+			uint32 LastError = GetLastError();
+			*OutError = FString::Printf(TEXT("FileJournal: CreateFile('%s') failed; FileJournal feature is unavailable. LastError=%u (0x%08x) (%s)."),
+				*Normalized, LastError, LastError, *JournalGetLastErrorString(LastError));
+		}
 		return OutResult;
 	}
 	ON_SCOPE_EXIT{ ::CloseHandle(Handle); };
@@ -333,6 +353,13 @@ FFileJournalData FWindowsPlatformFile::FileJournalGetFileData(const TCHAR* Filen
 			UE_LOG(LogWindows, Error, TEXT("Unexpected failure of GetFileInformationByHandleEx('%s') with error code %u (0x%08x) (%s)."),
 				*Normalized, LastError, LastError, *JournalGetLastErrorString(LastError));
 		}
+
+		if (OutError)
+		{
+			*OutError = FString::Printf(TEXT("Unexpected failure of GetFileInformationByHandleEx('%s') with error code %u (0x%08x) (%s)."),
+				*Normalized, LastError, LastError, *JournalGetLastErrorString(LastError));
+		}
+
 		return OutResult;
 	}
 
@@ -909,14 +936,14 @@ EFileJournalResult  FWindowsPlatformFile::FileJournalGetLatestEntry(const TCHAR*
 }
 
 bool FWindowsPlatformFile::FileJournalIterateDirectory(const TCHAR* Directory,
-	FDirectoryJournalVisitorFunc Visitor)
+	FDirectoryJournalVisitorFunc Visitor, FString* OutError)
 {
-	return IPlatformFile::FileJournalIterateDirectory(Directory, Visitor);
+	return IPlatformFile::FileJournalIterateDirectory(Directory, Visitor, OutError);
 }
 
-FFileJournalData FWindowsPlatformFile::FileJournalGetFileData(const TCHAR* FilenameOrDirectory)
+FFileJournalData FWindowsPlatformFile::FileJournalGetFileData(const TCHAR* FilenameOrDirectory, FString* OutError)
 {
-	return IPlatformFile::FileJournalGetFileData(FilenameOrDirectory);
+	return IPlatformFile::FileJournalGetFileData(FilenameOrDirectory, OutError);
 }
 
 EFileJournalResult FWindowsPlatformFile::FileJournalReadModified(const TCHAR* VolumeName,
@@ -937,4 +964,4 @@ EFileJournalResult FWindowsPlatformFile::FileJournalReadModified(const TCHAR* Vo
 		KnownDirectories, OutModifiedDirectories, OutNextJournalEntry, OutError);
 }
 
-#endif // else !UE_WINDOWS_PLATFORM_FILEJOURNAL_ENABLED
\ No newline at end of file
+#endif // else !UE_WINDOWS_PLATFORM_FILEJOURNAL_ENABLED
diff --git a/Engine/Source/Runtime/Core/Private/Windows/WindowsPlatformFilePrivate.h b/Engine/Source/Runtime/Core/Private/Windows/WindowsPlatformFilePrivate.h
--- a/Engine/Source/Runtime/Core/Private/Windows/WindowsPlatformFilePrivate.h
+++ b/Engine/Source/Runtime/Core/Private/Windows/WindowsPlatformFilePrivate.h
@@ -84,8 +84,8 @@ public:
 		FString* OutError = nullptr) override;
 	virtual EFileJournalResult FileJournalGetLatestEntry(const TCHAR* VolumeName, FFileJournalId& OutJournalId, 
 		FFileJournalEntryHandle& OutEntryHandle, FString* OutError = nullptr) override;
-	virtual bool FileJournalIterateDirectory(const TCHAR* Directory, FDirectoryJournalVisitorFunc Visitor) override;
-	virtual FFileJournalData FileJournalGetFileData(const TCHAR* FilenameOrDirectory) override;
+	virtual bool FileJournalIterateDirectory(const TCHAR* Directory, FDirectoryJournalVisitorFunc Visitor, FString* OutError = nullptr) override;
+	virtual FFileJournalData FileJournalGetFileData(const TCHAR* FilenameOrDirectory, FString* OutError = nullptr) override;
 	virtual EFileJournalResult FileJournalReadModified(const TCHAR* VolumeName,
 		const FFileJournalId& JournalIdOfStartingEntry, const FFileJournalEntryHandle& StartingJournalEntry,
 		TMap<FFileJournalFileHandle, FString>& KnownDirectories, TSet<FString>& OutModifiedDirectories,
diff --git a/Engine/Source/Runtime/Core/Public/GenericPlatform/GenericPlatformFile.h b/Engine/Source/Runtime/Core/Public/GenericPlatform/GenericPlatformFile.h
--- a/Engine/Source/Runtime/Core/Public/GenericPlatform/GenericPlatformFile.h
+++ b/Engine/Source/Runtime/Core/Public/GenericPlatform/GenericPlatformFile.h
@@ -670,14 +670,14 @@ public:
 	 * 
 	 * @return	false if the directory did not exist or if the visitor returned false.
 	 */
-	CORE_API virtual bool FileJournalIterateDirectory(const TCHAR* Directory, FDirectoryJournalVisitorFunc Visitor);
+	CORE_API virtual bool FileJournalIterateDirectory(const TCHAR* Directory, FDirectoryJournalVisitorFunc Visitor, FString* OutError = nullptr);
 
 	/**
 	 * Return the data for the given path as with GetStatData, but report a FFileJournalData instead, which
 	 * notably includes the FFileJournalFileHandle for the file/directory.
 	 * Check the FFileJournalData::bIsValid member before using the returned data 
 	 */
-	CORE_API virtual FFileJournalData FileJournalGetFileData(const TCHAR* FilenameOrDirectory);
+	CORE_API virtual FFileJournalData FileJournalGetFileData(const TCHAR* FilenameOrDirectory, FString* OutError = nullptr);
 
 	/**
 	 * Query the FileJournal to find a list of all directories on the given volume with files that have been added,
