Changed update priority parameter to signed integer since priority should be signed.

diff --git a/Engine/Source/Runtime/Core/Private/IO/IoDispatcher.cpp b/Engine/Source/Runtime/Core/Private/IO/IoDispatcher.cpp
--- a/Engine/Source/Runtime/Core/Private/IO/IoDispatcher.cpp
+++ b/Engine/Source/Runtime/Core/Private/IO/IoDispatcher.cpp
@@ -1316,7 +1316,7 @@ FIoRequest::Cancel()
 }
 
 void
-FIoRequest::UpdatePriority(uint32 NewPriority)
+FIoRequest::UpdatePriority(int32 NewPriority)
 {
 	if (!GIoDispatcher || !Impl || Impl->Priority == NewPriority)
 	{
diff --git a/Engine/Source/Runtime/Core/Public/IO/IoDispatcher.h b/Engine/Source/Runtime/Core/Public/IO/IoDispatcher.h
--- a/Engine/Source/Runtime/Core/Public/IO/IoDispatcher.h
+++ b/Engine/Source/Runtime/Core/Public/IO/IoDispatcher.h
@@ -247,7 +247,7 @@ public:
 	CORE_API const FIoBuffer*				GetResult() const;
 	CORE_API const FIoBuffer&				GetResultOrDie() const;
 	CORE_API void							Cancel();
-	CORE_API void							UpdatePriority(uint32 NewPriority);
+	CORE_API void							UpdatePriority(int32 NewPriority);
 	CORE_API void							Release();
 
 private:
