Remove deprecated package store interface.

diff --git a/Engine/Source/Runtime/Core/Public/IO/PackageStore.h b/Engine/Source/Runtime/Core/Public/IO/PackageStore.h
--- a/Engine/Source/Runtime/Core/Public/IO/PackageStore.h
+++ b/Engine/Source/Runtime/Core/Public/IO/PackageStore.h
@@ -121,53 +121,6 @@ struct FPackageStoreEntryResource
 	CORE_API static FPackageStoreEntryResource CreateEmptyPackage(FName PackageName, bool bHasCookError);
 };
 
-/**
- * Stores information about available packages that can be loaded.
- */
-class UE_DEPRECATED(5.1, "Use IPackageStoreBackend instead") IPackageStore
-{
-public:
-	/* Destructor. */
-	virtual ~IPackageStore() { }
-
-	virtual void Initialize() = 0;
-
-	/** Lock the package store for reading */
-	virtual void Lock() = 0;
-
-	/** Unlock the package store */
-	virtual void Unlock() = 0;
-
-	/* Returns whether the package exists. */
-	virtual bool DoesPackageExist(FPackageId PackageId) = 0;
-
-	/* Returns the package store entry data with export info and imported packages for the specified package ID. */
-	virtual EPackageStoreEntryStatus GetPackageStoreEntry(FPackageId PackageId, FName PackageName,
-		FPackageStoreEntry& OutPackageStoreEntry) = 0;
-
-	/* Returns the redirected package ID and source package name for the specified package ID if it's being redirected. */
-	virtual bool GetPackageRedirectInfo(FPackageId PackageId, FName& OutSourcePackageName, FPackageId& OutRedirectedToPackageId) = 0;
-
-	/* Event broadcasted when pending entries are completed and added to the package store */
-	DECLARE_EVENT(IPackageStore, FEntriesAddedEvent);
-	virtual FEntriesAddedEvent& OnPendingEntriesAdded() = 0;
-};
-
-class UE_DEPRECATED(5.1, "Use IPackageStoreBackend instead") FPackageStoreBase
-PRAGMA_DISABLE_DEPRECATION_WARNINGS
-	: public IPackageStore
-PRAGMA_ENABLE_DEPRECATION_WARNINGS
-{
-public:
-	virtual FEntriesAddedEvent& OnPendingEntriesAdded() override
-	{
-		return PendingEntriesAdded;
-	}
-
-protected:
-	FEntriesAddedEvent PendingEntriesAdded;
-};
-
 class FPackageStoreBackendContext
 {
 public:
