Octree2: Add a return value of FOctreeElementId2 to FindFirstElementWithBoundsTest

diff --git a/Engine/Source/Runtime/Core/Public/Math/GenericOctree.h b/Engine/Source/Runtime/Core/Public/Math/GenericOctree.h
--- a/Engine/Source/Runtime/Core/Public/Math/GenericOctree.h
+++ b/Engine/Source/Runtime/Core/Public/Math/GenericOctree.h
@@ -660,20 +660,23 @@ private:
 	}
 
 	template<typename IterateFunc>
-	void FindFirstElementWithBoundsTestInternal(FNodeIndex CurrentNodeIndex, const FOctreeNodeContext& NodeContext, const FBoxCenterAndExtent& BoxBounds, const IterateFunc& Func, bool& ContinueTraversal) const
+	FOctreeElementId2 FindFirstElementWithBoundsTestInternal(FNodeIndex CurrentNodeIndex, const FOctreeNodeContext& NodeContext, const FBoxCenterAndExtent& BoxBounds, const IterateFunc& Func) const
 	{
-		if (ContinueTraversal && TreeNodes[CurrentNodeIndex].InclusiveNumElements > 0)
+		if (TreeNodes[CurrentNodeIndex].InclusiveNumElements > 0)
 		{
-			for (int Index = 0; ContinueTraversal && Index < TreeElements[CurrentNodeIndex].Num(); Index++)
+			for (int Index = 0; Index < TreeElements[CurrentNodeIndex].Num(); Index++)
 			{
 				typename TCallTraits<ElementType>::ConstReference Element = TreeElements[CurrentNodeIndex][Index];
 				if (Intersect(OctreeSemantics::GetBoundingBox(Element), BoxBounds))
 				{
-					ContinueTraversal &= Func(Element);
+					if (!Func(Element))
+					{
+						return FOctreeElementId2(CurrentNodeIndex, Index);
+					}
 				}
 			}
 
-			if (ContinueTraversal && !TreeNodes[CurrentNodeIndex].IsLeaf())
+			if (!TreeNodes[CurrentNodeIndex].IsLeaf())
 			{
 				const FOctreeChildNodeSubset IntersectingChildSubset = NodeContext.GetIntersectingChildren(BoxBounds);
 				FNodeIndex ChildStartIndex = TreeNodes[CurrentNodeIndex].ChildNodes;
@@ -681,11 +684,18 @@ private:
 				{
 					if (IntersectingChildSubset.Contains(FOctreeChildNodeRef(i)))
 					{
-						FindFirstElementWithBoundsTestInternal(ChildStartIndex + i, NodeContext.GetChildContext(FOctreeChildNodeRef(i)), BoxBounds, Func, ContinueTraversal);
+						FOctreeElementId2 FoundIndex = FindFirstElementWithBoundsTestInternal(ChildStartIndex + i, NodeContext.GetChildContext(FOctreeChildNodeRef(i)), BoxBounds, Func);
+						if (FoundIndex.IsValidId())
+						{
+							return FoundIndex;
+						}
 					}
 				}
 			}
 		}
+
+		// No matching element was found
+		return FOctreeElementId2();
 	}
 
 	template<typename IterateFunc>
@@ -789,12 +799,12 @@ public:
 	 * this function will traverse the Octree using a fast box-box intersection and aborts traversal as soon as the Element function returns false.
 	 * @param BoxBounds - the bounds to test if a node is traversed or skipped.
 	 * @param Func - Function to call with each Element for nodes that passed bounds test.
+	 * @return The ID of the found element. It's only valid until the next time the tree changes.
 	 */
 	template<typename IterateBoundsFunc>
-	inline void FindFirstElementWithBoundsTest(const FBoxCenterAndExtent& BoxBounds, const IterateBoundsFunc& Func) const
+	inline FOctreeElementId2 FindFirstElementWithBoundsTest(const FBoxCenterAndExtent& BoxBounds, const IterateBoundsFunc& Func) const
 	{
-		bool ContinueTraversal = true;
-		FindFirstElementWithBoundsTestInternal(0, RootNodeContext, BoxBounds, Func, ContinueTraversal);
+		return FindFirstElementWithBoundsTestInternal(0, RootNodeContext, BoxBounds, Func);
 	}
 
 	/**
