- When loading thumbnails with invalid dimensions, skip loading image data. In cases where there is no image the image data should be empty, but in cases where there is image data but the dimensions are 0 by 0 we treat this as corrupt thumbnail data that is unsafe to read and should be skipped. When saving, thumbnails with 0 in their dimensions will regenerate the thumbnail, so there is no need to force the image to be dirty upon load and we can let the correction happen the next time the asset needs to be saved.
- Since we early out when we see corrupt data, the knock on side-effect is that any counting archive will report only the valid bytes in the thumbnail since when counting we will also skip the corrupt data.

#rnx

diff --git a/Engine/Source/Runtime/Core/Private/Misc/ObjectThumbnail.cpp b/Engine/Source/Runtime/Core/Private/Misc/ObjectThumbnail.cpp
--- a/Engine/Source/Runtime/Core/Private/Misc/ObjectThumbnail.cpp
+++ b/Engine/Source/Runtime/Core/Private/Misc/ObjectThumbnail.cpp
@@ -74,22 +74,27 @@ void FObjectThumbnail::Serialize(FStructuredArchive::FSlot Slot)
 		
 		Record << SA_VALUE(TEXT("ImageHeight"), ImageHeightToSerialize);
 	}
-		
-	Record << SA_VALUE(TEXT("CompressedImageData"), CompressedImageData);
-
-	if (Slot.GetUnderlyingArchive().IsCountingMemory())
-	{
-		Record << SA_VALUE(TEXT("ImageData"), ImageData) << SA_VALUE(TEXT("bIsDirty"), bIsDirty);
-	}
-
+	
 	if (Slot.GetArchiveState().IsLoading())
 	{
 		bLoadedFromDisk = true;
-		if ((ImageWidth>0) && (ImageHeight>0))
+		bCreatedAfterCustomThumbForSharedTypesEnabled = true;
+
+		if ((ImageWidth == 0) || (ImageHeight == 0))
 		{
-			bCreatedAfterCustomThumbForSharedTypesEnabled = true;
+			// The thumbnail data doesn't exist and might be corrupt. Attempting to read the
+			// image data may lead to crashes. Mark the thumbnail dirtsy so we may regenerate it
+			bIsDirty = true; 
+			return;
 		}
 	}
+
+	Record << SA_VALUE(TEXT("CompressedImageData"), CompressedImageData);
+
+	if (Slot.GetUnderlyingArchive().IsCountingMemory())
+	{
+		Record << SA_VALUE(TEXT("ImageData"), ImageData) << SA_VALUE(TEXT("bIsDirty"), bIsDirty);
+	}
 }
 
 
diff --git a/Engine/Source/Runtime/CoreUObject/Private/UObject/LinkerLoad.cpp b/Engine/Source/Runtime/CoreUObject/Private/UObject/LinkerLoad.cpp
--- a/Engine/Source/Runtime/CoreUObject/Private/UObject/LinkerLoad.cpp
+++ b/Engine/Source/Runtime/CoreUObject/Private/UObject/LinkerLoad.cpp
@@ -2814,6 +2814,13 @@ FLinkerLoad::ELinkerStatus FLinkerLoad::SerializeThumbnails( bool bForceEnableIn
 				FObjectThumbnail LoadedThumbnail;
 				LoadedThumbnail.Serialize(DataStream.EnterElement());
 
+				if (LoadedThumbnail.AccessImageData().IsEmpty())
+				{
+					// If we failed to load the thumbnail, stop loading as it might be unsafe
+					// to continue reading the stream.
+					break;
+				}
+
 				// Store the data!
 				ThumbnailMap.Add(CurThumbnailInfo.ObjectFullName, LoadedThumbnail);
 			}
