Override DateTime digits format from the ICU to follow UE's standard digits

diff --git a/Engine/Source/Runtime/Core/Private/Internationalization/ICUCulture.cpp b/Engine/Source/Runtime/Core/Private/Internationalization/ICUCulture.cpp
--- a/Engine/Source/Runtime/Core/Private/Internationalization/ICUCulture.cpp
+++ b/Engine/Source/Runtime/Core/Private/Internationalization/ICUCulture.cpp
@@ -95,7 +95,7 @@ namespace
 		return Ptr.ToSharedRef();
 	}
 
-	TSharedRef<const icu::DateFormat, ESPMode::ThreadSafe> CreateDateFormat(const icu::Locale& ICULocale)
+	TSharedRef<const icu::DateFormat, ESPMode::ThreadSafe> CreateDateFormat(const icu::Locale& ICULocale, const icu::DecimalFormat& ICUDecimalFormat)
 	{
 		TSharedPtr<icu::DateFormat, ESPMode::ThreadSafe> Ptr = MakeShareable(icu::DateFormat::createDateInstance(icu::DateFormat::kDefault, ICULocale));
 		if (!ensureAlwaysMsgf(Ptr, TEXT("Creating a date format object failed using locale %s. Perhaps this locale has no data."), StringCast<TCHAR>(ICULocale.getName()).Get()))
@@ -104,10 +104,11 @@ namespace
 			check(Ptr);
 		}
 		Ptr->adoptTimeZone(icu::TimeZone::createDefault());
+		Ptr->setNumberFormat(ICUDecimalFormat);
 		return Ptr.ToSharedRef();
 	}
 
-	TSharedRef<const icu::DateFormat, ESPMode::ThreadSafe> CreateTimeFormat(const icu::Locale& ICULocale)
+	TSharedRef<const icu::DateFormat, ESPMode::ThreadSafe> CreateTimeFormat(const icu::Locale& ICULocale, const icu::DecimalFormat& ICUDecimalFormat)
 	{
 		TSharedPtr<icu::DateFormat, ESPMode::ThreadSafe> Ptr = MakeShareable(icu::DateFormat::createTimeInstance(icu::DateFormat::kDefault, ICULocale));
 		if (!ensureAlwaysMsgf(Ptr, TEXT("Creating a time format object failed using locale %s. Perhaps this locale has no data."), StringCast<TCHAR>(ICULocale.getName()).Get()))
@@ -116,10 +117,11 @@ namespace
 			check(Ptr);
 		}
 		Ptr->adoptTimeZone(icu::TimeZone::createDefault());
+		Ptr->setNumberFormat(ICUDecimalFormat);
 		return Ptr.ToSharedRef();
 	}
 
-	TSharedRef<const icu::DateFormat, ESPMode::ThreadSafe> CreateDateTimeFormat(const icu::Locale& ICULocale)
+	TSharedRef<const icu::DateFormat, ESPMode::ThreadSafe> CreateDateTimeFormat(const icu::Locale& ICULocale, const icu::DecimalFormat& ICUDecimalFormat)
 	{
 		TSharedPtr<icu::DateFormat, ESPMode::ThreadSafe> Ptr = MakeShareable(icu::DateFormat::createDateTimeInstance(icu::DateFormat::kDefault, icu::DateFormat::kDefault, ICULocale));
 		if (!ensureAlwaysMsgf(Ptr, TEXT("Creating a date-time format object failed using locale %s. Perhaps this locale has no data."), StringCast<TCHAR>(ICULocale.getName()).Get()))
@@ -128,6 +130,7 @@ namespace
 			check(Ptr);
 		}
 		Ptr->adoptTimeZone(icu::TimeZone::createDefault());
+		Ptr->setNumberFormat(ICUDecimalFormat);
 		return Ptr.ToSharedRef();
 	}
 }
@@ -413,11 +416,58 @@ TSharedRef<const icu::Collator, ESPMode::ThreadSafe> FICUCultureImplementation::
 	}
 }
 
+TSharedRef<const icu::DecimalFormat, ESPMode::ThreadSafe> FICUCultureImplementation::GetDateTimeDecimalFormat()
+{
+	if (!ICUDateTimeDecimalFormat)
+	{
+		// Create a culture decimal formatter
+		TSharedPtr<icu::DecimalFormat> DecimalFormatterForCulture;
+		{
+			UErrorCode ICUStatus = U_ZERO_ERROR;
+			DecimalFormatterForCulture = MakeShareable(static_cast<icu::DecimalFormat*>(icu::NumberFormat::createInstance(ICULocale, ICUStatus)));
+			if (!ensureAlwaysMsgf(DecimalFormatterForCulture, TEXT("Creating a decimal format object failed using locale %s. Perhaps this locale has no data."), StringCast<TCHAR>(ICULocale.getName()).Get()))
+			{
+				DecimalFormatterForCulture = MakeShareable(static_cast<icu::DecimalFormat*>(icu::NumberFormat::createInstance(GetInvariantLocale(), ICUStatus)));
+				check(DecimalFormatterForCulture);
+			}
+		}
+
+		// Sync the ICU decimal symbols with the ones we want to use for UE
+		{
+			const FDecimalNumberFormattingRules& NewUEDecimalNumberFormattingRules = GetDecimalNumberFormattingRules();
+
+			icu::DecimalFormatSymbols ICUDecimalSymbols = *DecimalFormatterForCulture->getDecimalFormatSymbols();
+
+			ICUDecimalSymbols.setSymbol(icu::DecimalFormatSymbols::kZeroDigitSymbol, ICUUtilities::ConvertString(FStringView(&NewUEDecimalNumberFormattingRules.DigitCharacters[0], 1)), false);
+			ICUDecimalSymbols.setSymbol(icu::DecimalFormatSymbols::kOneDigitSymbol, ICUUtilities::ConvertString(FStringView(&NewUEDecimalNumberFormattingRules.DigitCharacters[1], 1)), false);
+			ICUDecimalSymbols.setSymbol(icu::DecimalFormatSymbols::kTwoDigitSymbol, ICUUtilities::ConvertString(FStringView(&NewUEDecimalNumberFormattingRules.DigitCharacters[2], 1)), false);
+			ICUDecimalSymbols.setSymbol(icu::DecimalFormatSymbols::kThreeDigitSymbol, ICUUtilities::ConvertString(FStringView(&NewUEDecimalNumberFormattingRules.DigitCharacters[3], 1)), false);
+			ICUDecimalSymbols.setSymbol(icu::DecimalFormatSymbols::kFourDigitSymbol, ICUUtilities::ConvertString(FStringView(&NewUEDecimalNumberFormattingRules.DigitCharacters[4], 1)), false);
+			ICUDecimalSymbols.setSymbol(icu::DecimalFormatSymbols::kFiveDigitSymbol, ICUUtilities::ConvertString(FStringView(&NewUEDecimalNumberFormattingRules.DigitCharacters[5], 1)), false);
+			ICUDecimalSymbols.setSymbol(icu::DecimalFormatSymbols::kSixDigitSymbol, ICUUtilities::ConvertString(FStringView(&NewUEDecimalNumberFormattingRules.DigitCharacters[6], 1)), false);
+			ICUDecimalSymbols.setSymbol(icu::DecimalFormatSymbols::kSevenDigitSymbol, ICUUtilities::ConvertString(FStringView(&NewUEDecimalNumberFormattingRules.DigitCharacters[7], 1)), false);
+			ICUDecimalSymbols.setSymbol(icu::DecimalFormatSymbols::kEightDigitSymbol, ICUUtilities::ConvertString(FStringView(&NewUEDecimalNumberFormattingRules.DigitCharacters[8], 1)), false);
+			ICUDecimalSymbols.setSymbol(icu::DecimalFormatSymbols::kNineDigitSymbol, ICUUtilities::ConvertString(FStringView(&NewUEDecimalNumberFormattingRules.DigitCharacters[9], 1)), false);
+
+			ICUDecimalSymbols.setSymbol(icu::DecimalFormatSymbols::kGroupingSeparatorSymbol, ICUUtilities::ConvertString(FStringView(&NewUEDecimalNumberFormattingRules.GroupingSeparatorCharacter, 1)), false);
+			ICUDecimalSymbols.setSymbol(icu::DecimalFormatSymbols::kDecimalSeparatorSymbol, ICUUtilities::ConvertString(FStringView(&NewUEDecimalNumberFormattingRules.DecimalSeparatorCharacter, 1)), false);
+
+			DecimalFormatterForCulture->setDecimalFormatSymbols(ICUDecimalSymbols);
+
+			DecimalFormatterForCulture->setGroupingUsed(false);
+		}
+
+		ICUDateTimeDecimalFormat = DecimalFormatterForCulture;
+	}
+
+	return ICUDateTimeDecimalFormat.ToSharedRef();
+}
+
 TSharedRef<const icu::DateFormat, ESPMode::ThreadSafe> FICUCultureImplementation::GetDateFormatter(const EDateTimeStyle::Type DateStyle, const FString& TimeZone)
 {
 	if (!ICUDateFormat.IsValid())
 	{
-		ICUDateFormat = CreateDateFormat( ICULocale );
+		ICUDateFormat = CreateDateFormat( ICULocale, *GetDateTimeDecimalFormat() );
 	}
 
 	const FString SanitizedTimezoneCode = ICUUtilities::SanitizeTimezoneCode(TimeZone);
@@ -469,7 +519,7 @@ TSharedRef<const icu::DateFormat, ESPMode::ThreadSafe> FICUCultureImplementation
 {
 	if (!ICUTimeFormat.IsValid())
 	{
-		ICUTimeFormat = CreateTimeFormat( ICULocale );
+		ICUTimeFormat = CreateTimeFormat( ICULocale, *GetDateTimeDecimalFormat() );
 	}
 
 	const FString SanitizedTimezoneCode = ICUUtilities::SanitizeTimezoneCode(TimeZone);
@@ -513,6 +563,7 @@ TSharedRef<const icu::DateFormat, ESPMode::ThreadSafe> FICUCultureImplementation
 			check(Formatter);
 		}
 		Formatter->adoptTimeZone(bIsDefaultTimeZone ? icu::TimeZone::createDefault() : icu::TimeZone::createTimeZone(InputTimeZoneID));
+		Formatter->setNumberFormat(*GetDateTimeDecimalFormat());
 		return Formatter.ToSharedRef();
 	}
 }
@@ -521,7 +572,7 @@ TSharedRef<const icu::DateFormat, ESPMode::ThreadSafe> FICUCultureImplementation
 {
 	if (!ICUDateTimeFormat.IsValid())
 	{
-		ICUDateTimeFormat = CreateDateTimeFormat( ICULocale );
+		ICUDateTimeFormat = CreateDateTimeFormat( ICULocale, *GetDateTimeDecimalFormat() );
 	}
 
 	const FString SanitizedTimezoneCode = ICUUtilities::SanitizeTimezoneCode(TimeZone);
@@ -566,6 +617,7 @@ TSharedRef<const icu::DateFormat, ESPMode::ThreadSafe> FICUCultureImplementation
 			check(Formatter);
 		}
 		Formatter->adoptTimeZone(bIsDefaultTimeZone ? icu::TimeZone::createDefault() : icu::TimeZone::createTimeZone(InputTimeZoneID));
+		Formatter->setNumberFormat(*GetDateTimeDecimalFormat());
 		return Formatter.ToSharedRef();
 	}
 }
@@ -634,6 +686,8 @@ TSharedRef<const icu::DateFormat, ESPMode::ThreadSafe> FICUCultureImplementation
 			const FString SanitizedTimezoneCode = ICUUtilities::SanitizeTimezoneCode(TimeZone);
 			Formatter->adoptTimeZone(SanitizedTimezoneCode.IsEmpty() ? icu::TimeZone::createDefault() : icu::TimeZone::createTimeZone(ICUUtilities::ConvertString(SanitizedTimezoneCode, false)));
 
+			Formatter->setNumberFormat(*GetDateTimeDecimalFormat());
+
 			return Formatter.ToSharedRef();
 		}
 	}
diff --git a/Engine/Source/Runtime/Core/Private/Internationalization/ICUCulture.h b/Engine/Source/Runtime/Core/Private/Internationalization/ICUCulture.h
--- a/Engine/Source/Runtime/Core/Private/Internationalization/ICUCulture.h
+++ b/Engine/Source/Runtime/Core/Private/Internationalization/ICUCulture.h
@@ -200,6 +200,7 @@ public:
 
 	TSharedRef<const icu::BreakIterator> GetBreakIterator(const EBreakIteratorType Type);
 	TSharedRef<const icu::Collator, ESPMode::ThreadSafe> GetCollator(const ETextComparisonLevel::Type ComparisonLevel);
+	TSharedRef<const icu::DecimalFormat, ESPMode::ThreadSafe> GetDateTimeDecimalFormat();
 	TSharedRef<const icu::DateFormat, ESPMode::ThreadSafe> GetDateFormatter(const EDateTimeStyle::Type DateStyle, const FString& TimeZone);
 	TSharedRef<const icu::DateFormat, ESPMode::ThreadSafe> GetTimeFormatter(const EDateTimeStyle::Type TimeStyle, const FString& TimeZone);
 	TSharedRef<const icu::DateFormat, ESPMode::ThreadSafe> GetDateTimeFormatter(const EDateTimeStyle::Type DateStyle, const EDateTimeStyle::Type TimeStyle, const FString& TimeZone);
@@ -214,6 +215,7 @@ private:
 	TSharedPtr<const icu::BreakIterator> ICUTitleBreakIterator;
 	TSharedPtr<const icu::Collator, ESPMode::ThreadSafe> ICUCollator;
 
+	TSharedPtr<const icu::DecimalFormat, ESPMode::ThreadSafe> ICUDateTimeDecimalFormat;
 	TSharedPtr<const icu::DateFormat, ESPMode::ThreadSafe> ICUDateFormat;
 	TSharedPtr<const icu::DateFormat, ESPMode::ThreadSafe> ICUTimeFormat;
 	TSharedPtr<const icu::DateFormat, ESPMode::ThreadSafe> ICUDateTimeFormat;
diff --git a/Engine/Source/Runtime/Core/Private/Internationalization/ICUUtilities.cpp b/Engine/Source/Runtime/Core/Private/Internationalization/ICUUtilities.cpp
--- a/Engine/Source/Runtime/Core/Private/Internationalization/ICUUtilities.cpp
+++ b/Engine/Source/Runtime/Core/Private/Internationalization/ICUUtilities.cpp
@@ -126,9 +126,9 @@ namespace ICUUtilities
 		}
 	}
 
-	void FStringConverter::ConvertString(const FString& Source, icu::UnicodeString& Destination, const bool ShouldNullTerminate)
+	void FStringConverter::ConvertString(FStringView Source, icu::UnicodeString& Destination, const bool ShouldNullTerminate)
 	{
-		ConvertString(*Source, 0, Source.Len(), Destination, ShouldNullTerminate);
+		ConvertString(Source.GetData(), 0, Source.Len(), Destination, ShouldNullTerminate);
 	}
 
 	void FStringConverter::ConvertString(const TCHAR* Source, const int32 SourceStartIndex, const int32 SourceLen, icu::UnicodeString& Destination, const bool ShouldNullTerminate)
@@ -136,7 +136,7 @@ namespace ICUUtilities
 		Impl.ConvertString(Source, SourceStartIndex, SourceLen, Destination, ShouldNullTerminate);
 	}
 
-	icu::UnicodeString FStringConverter::ConvertString(const FString& Source, const bool ShouldNullTerminate)
+	icu::UnicodeString FStringConverter::ConvertString(FStringView Source, const bool ShouldNullTerminate)
 	{
 		icu::UnicodeString Destination;
 		ConvertString(Source, Destination, ShouldNullTerminate);
@@ -174,7 +174,7 @@ namespace ICUUtilities
 		return Destination;
 	}
 
-	void ConvertString(const FString& Source, icu::UnicodeString& Destination, const bool ShouldNullTerminate)
+	void ConvertString(FStringView Source, icu::UnicodeString& Destination, const bool ShouldNullTerminate)
 	{
 		if (Source.Len() > 0)
 		{
@@ -200,7 +200,7 @@ namespace ICUUtilities
 		}
 	}
 
-	icu::UnicodeString ConvertString(const FString& Source, const bool ShouldNullTerminate)
+	icu::UnicodeString ConvertString(FStringView Source, const bool ShouldNullTerminate)
 	{
 		icu::UnicodeString Destination;
 		ConvertString(Source, Destination, ShouldNullTerminate);
@@ -318,9 +318,9 @@ namespace ICUUtilities
 		return Len;
 	}
 
-	int32 GetUnicodeStringLength(const FString& Source)
+	int32 GetUnicodeStringLength(FStringView Source)
 	{
-		return GetUnicodeStringLength(*Source, 0, Source.Len());
+		return GetUnicodeStringLength(Source.GetData(), 0, Source.Len());
 	}
 
 	int32 GetUnicodeStringLength(const TCHAR* Source, const int32 InSourceStartIndex, const int32 InSourceLength)
diff --git a/Engine/Source/Runtime/Core/Private/Internationalization/ICUUtilities.h b/Engine/Source/Runtime/Core/Private/Internationalization/ICUUtilities.h
--- a/Engine/Source/Runtime/Core/Private/Internationalization/ICUUtilities.h
+++ b/Engine/Source/Runtime/Core/Private/Internationalization/ICUUtilities.h
@@ -3,6 +3,7 @@
 #pragma once
 
 #include "CoreTypes.h"
+#include "Containers/StringView.h"
 #include "Containers/UnrealString.h"
 #include "Misc/Timespan.h"
 #include "Internationalization/InternationalizationUtilities.h"
@@ -66,9 +67,9 @@ namespace ICUUtilities
 		FStringConverter& operator=(const FStringConverter&) = delete;
 
 		/** Convert FString -> icu::UnicodeString */
-		void ConvertString(const FString& Source, icu::UnicodeString& Destination, const bool ShouldNullTerminate = true);
+		void ConvertString(FStringView Source, icu::UnicodeString& Destination, const bool ShouldNullTerminate = true);
 		void ConvertString(const TCHAR* Source, const int32 SourceStartIndex, const int32 SourceLen, icu::UnicodeString& Destination, const bool ShouldNullTerminate = true);
-		icu::UnicodeString ConvertString(const FString& Source, const bool ShouldNullTerminate = true);
+		icu::UnicodeString ConvertString(FStringView Source, const bool ShouldNullTerminate = true);
 		icu::UnicodeString ConvertString(const TCHAR* Source, const int32 SourceStartIndex, const int32 SourceLen, const bool ShouldNullTerminate = true);
 
 		/** Convert icu::UnicodeString -> FString */
@@ -82,9 +83,9 @@ namespace ICUUtilities
 	};
 
 	/** Convert FString -> icu::UnicodeString */
-	void ConvertString(const FString& Source, icu::UnicodeString& Destination, const bool ShouldNullTerminate = true);
+	void ConvertString(FStringView Source, icu::UnicodeString& Destination, const bool ShouldNullTerminate = true);
 	void ConvertString(const TCHAR* Source, const int32 SourceStartIndex, const int32 SourceLen, icu::UnicodeString& Destination, const bool ShouldNullTerminate = true);
-	icu::UnicodeString ConvertString(const FString& Source, const bool ShouldNullTerminate = true);
+	icu::UnicodeString ConvertString(FStringView Source, const bool ShouldNullTerminate = true);
 	icu::UnicodeString ConvertString(const TCHAR* Source, const int32 SourceStartIndex, const int32 SourceLen, const bool ShouldNullTerminate = true);
 
 	/** Convert icu::UnicodeString -> FString */
@@ -98,7 +99,7 @@ namespace ICUUtilities
 	int32 GetNativeStringLength(const icu::UnicodeString& Source, const int32 InSourceStartIndex, const int32 InSourceLength);
 
 	/** Given an FString, count how many characters it would be if converted to an icu::UnicodeString (as FString may not always be UTF-16) */
-	int32 GetUnicodeStringLength(const FString& Source);
+	int32 GetUnicodeStringLength(FStringView Source);
 	int32 GetUnicodeStringLength(const TCHAR* Source, const int32 InSourceStartIndex, const int32 InSourceLength);
 }
 #endif
