Correct allocator name comparison after "binned2" became "Binned2" in 5.5. MallocBinned3 also supports fetching allocation sizes.

#rnx

diff --git a/Engine/Source/Runtime/Core/Private/GenericPlatform/GenericPlatformMallocCrash.cpp b/Engine/Source/Runtime/Core/Private/GenericPlatform/GenericPlatformMallocCrash.cpp
--- a/Engine/Source/Runtime/Core/Private/GenericPlatform/GenericPlatformMallocCrash.cpp
+++ b/Engine/Source/Runtime/Core/Private/GenericPlatform/GenericPlatformMallocCrash.cpp
@@ -360,9 +360,10 @@ void* FGenericPlatformMallocCrash::Realloc( void* Ptr, SIZE_T NewSize, uint32 Al
 			if( bPreviousMalloc )
 			{
 				// We can safely get allocation size only from a few mallocs, this may change in future.
-				if( FCStringWide::Strcmp( PreviousMalloc->GetDescriptiveName(), TEXT("binned") ) == 0 ||
-					FCStringWide::Strcmp( PreviousMalloc->GetDescriptiveName(), TEXT("binned2")) == 0 ||
-					FCStringWide::Strcmp( PreviousMalloc->GetDescriptiveName(), TEXT("jemalloc") ) == 0 )
+				if( FCStringWide::Stricmp( PreviousMalloc->GetDescriptiveName(), TEXT("binned") ) == 0 ||
+					FCStringWide::Stricmp( PreviousMalloc->GetDescriptiveName(), TEXT("binned2")) == 0 ||
+					FCStringWide::Stricmp( PreviousMalloc->GetDescriptiveName(), TEXT("binned3")) == 0 ||
+					FCStringWide::Stricmp( PreviousMalloc->GetDescriptiveName(), TEXT("jemalloc") ) == 0 )
 				{
 					// Realloc from the previous allocator.
 					if (!PreviousMalloc->GetAllocationSize(Ptr, PtrSize) || PtrSize == 0)
