Launch UTS with job breakout only when allowed.

In certain scenarios the runtime is running a job that does not allow breaking out of the job. In these cases create process calls with CREATE_BREAKAWAY_FROM_JOB will fail. Detect these cases and only set CREATE_BREAKAWAY_FROM_JOB  when allowed.

diff --git a/Engine/Source/Runtime/Core/Private/ProfilingDebugging/TraceAuxiliary.cpp b/Engine/Source/Runtime/Core/Private/ProfilingDebugging/TraceAuxiliary.cpp
--- a/Engine/Source/Runtime/Core/Private/ProfilingDebugging/TraceAuxiliary.cpp
+++ b/Engine/Source/Runtime/Core/Private/ProfilingDebugging/TraceAuxiliary.cpp
@@ -2242,7 +2242,15 @@ bool LaunchTraceServerCommand(ELaunchTraceServerCommand Command, bool bAddSponso
 		CreateProcArgs << TEXT(" --sponsor ") << FPlatformProcess::GetCurrentProcessId();
 	}
 
-	uint32 CreateProcFlags = CREATE_BREAKAWAY_FROM_JOB;
+	uint32 CreateProcFlags = 0;
+	JOBOBJECT_BASIC_LIMIT_INFORMATION JobLimits;
+	if (QueryInformationJobObject(NULL, JobObjectBasicLimitInformation, &JobLimits, sizeof(JOBOBJECT_BASIC_LIMIT_INFORMATION), NULL))
+	{
+		if (JobLimits.LimitFlags & JOB_OBJECT_LIMIT_BREAKAWAY_OK)
+		{
+			CreateProcFlags |= CREATE_BREAKAWAY_FROM_JOB;
+		}
+	}
 	if (FParse::Param(FCommandLine::Get(), TEXT("traceshowstore")))
 	{
 		CreateProcFlags |= CREATE_NEW_CONSOLE;
