Add a "stat GPU" macro command to toggle on all GPU related stats
- The new GPU profiler maintains multiple, separate stat groups, one per GPU, per queue, e.g. "GPU0_Graphics0", "GPU0_Compute0", "GPU1_Graphics0" etc.
- This command provides a convenient way to enable all the GPU stat groups at once, which will be a familiar workflow for anyone used to the behavior of the old profiler.

Added a fixed sort order when rendering stats
- The order of individual stat groups was indeterminate, because the list of enabled groups is a TMap, which doesn't guarantee iteration order when elements are added / removed.
- A new "SortOrder" field has been added which we use to sort the enabled groups, such that more recently enabled groups appear higher on screen. This helps keep the new stats visible even if the total number of stats causes us to render beyond the bottom of the screen. It also enforces the ordering of the "stat GPU" macro command, which otherwise presented stat groups in a random order each time it was used.

diff --git a/Engine/Source/Runtime/Core/Private/Stats/StatsCommand.cpp b/Engine/Source/Runtime/Core/Private/Stats/StatsCommand.cpp
--- a/Engine/Source/Runtime/Core/Private/Stats/StatsCommand.cpp
+++ b/Engine/Source/Runtime/Core/Private/Stats/StatsCommand.cpp
@@ -811,6 +811,7 @@ struct FInternalGroup
 		, GroupDescription( InGroupDescription )
 		, DisplayMode( InDisplayMode )
 		, bCompareByName( InCompareByName )
+		, SortOrder( NextCounter++ )
 	{
 		// To avoid copy.
 		Exchange( EnabledItems, InEnabledItems );
@@ -849,8 +850,14 @@ struct FInternalGroup
 
 	/** Sorting mode for this group. */
 	bool bCompareByName;
+
+	/** Counter used to render stat groups in the order they were enabled. */
+	const uint64 SortOrder;
+	static uint64 NextCounter;
 };
 
+uint64 FInternalGroup::NextCounter = 0;
+
 /** Stats for the particular frame. */
 struct FHudFrame
 {
@@ -1462,11 +1469,25 @@ struct FHUDGroupManager
 
 			FComplexStatUtils::DiviveStatArray( AggregatedNonStackStatsHistory, NumFrames, EComplexStatField::IncSum, EComplexStatField::IncAve );
 	
+			// Sort groups into a reliable order, based on the order they were activated.
+			// More recently activated stats appear higher on screen than older stats.
+			TArray<const FInternalGroup*> SortedInternalGroups;
+			SortedInternalGroups.Reserve(EnabledGroups.Num());
+			for (auto const& [Name, InternalGroup] : EnabledGroups)
+			{
+				SortedInternalGroups.Add(&InternalGroup);
+			}
+
+			SortedInternalGroups.Sort([](const FInternalGroup& LHS, const FInternalGroup& RHS)
+			{
+				return LHS.SortOrder >= RHS.SortOrder;
+			});
+
 			// Iterate through all enabled groups.
-			for( auto GroupIt = EnabledGroups.CreateIterator(); GroupIt; ++GroupIt )
+			for (const FInternalGroup* InternalGroupPtr : SortedInternalGroups)
 			{
-				const FName& GroupName = GroupIt.Key();
-				FInternalGroup& InternalGroup = GroupIt.Value();
+				const FInternalGroup& InternalGroup = *InternalGroupPtr;
+				const FName& GroupName = InternalGroup.GroupName;
 
 				// Create a new hud group.
 				ToGame->ActiveStatGroups.Add(new FActiveStatGroupInfo());
@@ -2136,6 +2157,64 @@ static void StatCmd(FString InCmd, bool bStatCommand, FOutputDevice* Ar /*= null
 			FStatSlowParams Params(Cmd);
 			FHUDGroupManager::Get(Stats).HandleCommand(Params, true);
 		}
+		// Skip this macro command if a real stat group called "GPU" exists. Required for legacy GPU profiler support until we permanently enable RHI_NEW_GPU_PROFILER.
+		else if (!Stats.Groups.Contains(TEXT("STATGROUP_GPU")) && FParse::Command(&Cmd, TEXT("gpu")))
+		{
+			// Macro command to toggle all the stat groups that contain "STATCAT_GPU" stats.
+			if (FCoreDelegates::StatCheckEnabled.IsBound())
+			{
+				// Gather the list of relevant stat groups
+				TSet<FName> GpuStatGroups;
+				for (auto const& [GroupName, StatName] : Stats.Groups)
+				{
+					FName Category = Stats.ShortNameToLongName.FindChecked(StatName).NameAndInfo.GetGroupCategory();
+					if (Category == FStatNameAndInfo::GpuStatCategory)
+					{
+						GpuStatGroups.Add(GroupName);
+					}
+				}
+
+				if (GpuStatGroups.Num() > 0)
+				{
+					// Check if any of these groups are already enabled.
+					bool bAllEnabled = true;
+					for (FName Group : GpuStatGroups)
+					{
+						FString GroupName = Group.ToString();
+						GroupName.RemoveFromStart(TEXT("STATGROUP_"));
+
+						bool bEnabled = false, bOtherEnabled = false;
+						FCoreDelegates::StatCheckEnabled.Broadcast(*GroupName, bEnabled, bOtherEnabled);
+
+						if (!bEnabled)
+						{
+							bAllEnabled = false;
+							break;
+						}
+					}
+
+					bool bShouldBeEnabled = !bAllEnabled;
+
+					// Toggle any stat groups which don't match the desired state.
+					for (FName Group : GpuStatGroups)
+					{
+						FString GroupName = Group.ToString();
+						GroupName.RemoveFromStart(TEXT("STATGROUP_"));
+
+						bool bEnabled = false, bOtherEnabled = false;
+						FCoreDelegates::StatCheckEnabled.Broadcast(*GroupName, bEnabled, bOtherEnabled);
+
+						if (bEnabled != bShouldBeEnabled)
+						{
+							FStatParams Params(Cmd);
+							Params.Group.Set(*GroupName);
+
+							FHUDGroupManager::Get(Stats).HandleCommand(Params, false);
+						}
+					}
+				}
+			}
+		}
 		else
 #endif
 		{
@@ -2306,6 +2385,9 @@ namespace UE::Stats
 				else if (FParse::Command(&TempCmd, TEXT("slow")))
 				{
 				}
+				else if (FParse::Command(&TempCmd, TEXT("gpu")))
+				{
+				}
 				else if (FParse::Command(&TempCmd, TEXT("namedmarker")))
 				{
 				}
