Fix TestTrueExpr() when passing an expression containing a comma.

Previously, code like `TestTrueExpr(Foo(x,y) == z)` would generate a compile error.

diff --git a/Engine/Source/Runtime/Core/Public/Misc/AutomationTest.h b/Engine/Source/Runtime/Core/Public/Misc/AutomationTest.h
--- a/Engine/Source/Runtime/Core/Public/Misc/AutomationTest.h
+++ b/Engine/Source/Runtime/Core/Public/Misc/AutomationTest.h
@@ -2523,7 +2523,7 @@ public:
 	}
 
 	/** Macro version of above, uses the passed in expression as the description as well */
-	#define TestTrueExpr(Expression) TestTrue(TEXT(#Expression), Expression)
+	#define TestTrueExpr(...) TestTrue(TEXT(#__VA_ARGS__), __VA_ARGS__)
 
 	/**
 	 * Logs an error if the given object returns false when calling its IsValid member.
