Units: Add "Nits" to the list of light units

This gives the ability to have lights that output more energy as their size increases. This change only applies to local lights (point, spot, rect) as the units for directional lights is not currently exposed.

diff --git a/Engine/Programs/UnrealHeaderTool/Config/DefaultEngine.ini b/Engine/Programs/UnrealHeaderTool/Config/DefaultEngine.ini
--- a/Engine/Programs/UnrealHeaderTool/Config/DefaultEngine.ini
+++ b/Engine/Programs/UnrealHeaderTool/Config/DefaultEngine.ini
@@ -150,6 +150,8 @@ gc.MaxObjectsInProgram=262144
 +Units="cd/m2"
 +Units="CandelaPerMeter2"
 +Units="EV"
++Units="Nits"
++Units="nt"
 +Units="Nanoseconds"
 +Units="ns"
 +Units="Microseconds"
diff --git a/Engine/Source/Editor/DetailCustomizations/Private/LightComponentDetails.cpp b/Engine/Source/Editor/DetailCustomizations/Private/LightComponentDetails.cpp
--- a/Engine/Source/Editor/DetailCustomizations/Private/LightComponentDetails.cpp
+++ b/Engine/Source/Editor/DetailCustomizations/Private/LightComponentDetails.cpp
@@ -72,6 +72,11 @@ void FLightComponentDetails::AddLocalLightIntensityWithUnit(IDetailLayoutBuilder
 		LightIntensityProperty->SetInstanceMetaData("Units", TEXT("ev"));
 		LightIntensityProperty->SetToolTipText(LOCTEXT("LightIntensityInEVToolTipText", "Luminous intensity in EV100"));
 	}
+	else if (Component->IntensityUnits == ELightUnits::Nits)
+	{
+		LightIntensityProperty->SetInstanceMetaData("Units", TEXT("nt"));
+		LightIntensityProperty->SetToolTipText(LOCTEXT("LightIntensityInNitsToolTipText", "Luminance in nits"));
+	}
 
 	// Make these come first
 	IDetailCategoryBuilder& LightCategory = DetailBuilder.EditCategory( "Light", FText::GetEmpty(), ECategoryPriority::TypeSpecific );
diff --git a/Engine/Source/Runtime/Core/Private/Math/UnitConversion.cpp b/Engine/Source/Runtime/Core/Private/Math/UnitConversion.cpp
--- a/Engine/Source/Runtime/Core/Private/Math/UnitConversion.cpp
+++ b/Engine/Source/Runtime/Core/Private/Math/UnitConversion.cpp
@@ -100,6 +100,7 @@ constexpr FParseCandidate ParseCandidates[] = {
 	{ TEXT("Candela"),				EUnit::Candela },				{ TEXT("cd"),		EUnit::Candela },
 	{ TEXT("Lux"),					EUnit::Lux },					{ TEXT("lx"),		EUnit::Lux },
 	{ TEXT("CandelaPerMeterSquared"), EUnit::CandelaPerMeter2 },	{ TEXT("cd/m2"),	EUnit::CandelaPerMeter2 },		{ TEXT("CandelaPerMeter2"),		EUnit::CandelaPerMeter2 },
+	{ TEXT("Nits"),                 EUnit::Nit },                   { TEXT("nt"),       EUnit::Nit },
 	{ TEXT("EV"),					EUnit::ExposureValue },			{ TEXT("EV"),		EUnit::ExposureValue },
 
 	{ TEXT("Nanoseconds"),			EUnit::Nanoseconds },			{ TEXT("ns"),		EUnit::Nanoseconds },
@@ -155,7 +156,7 @@ constexpr const TCHAR* const DisplayStrings[] = {
 
 	TEXT("B"), TEXT("KiB"), TEXT("MiB"), TEXT("GiB"), TEXT("TiB"),
 
-	TEXT("lm"), TEXT("cd"), TEXT("lux"), TEXT("cd/m2"), TEXT("EV"),
+	TEXT("lm"), TEXT("cd"), TEXT("lux"), TEXT("cd/m2"), TEXT("nt"), TEXT("EV"),
 
 	TEXT("ns"), TEXT("\u00B5s"), TEXT("ms"), TEXT("s"), TEXT("min"), TEXT("hr"), TEXT("dy"), TEXT("mth"), TEXT("yr"),
 
@@ -240,6 +241,7 @@ constexpr const TCHAR* const SupportedUnitsStrings[] = {
 	TEXT("Candela"),			
 	TEXT("Lux"),				
 	TEXT("CandelaPerMeterSquared"), 
+	TEXT("Nits"), 
 	TEXT("EV"),				
 
 	TEXT("Nanoseconds"),			
@@ -297,7 +299,7 @@ constexpr EUnitType UnitTypes[] = {
 
 	EUnitType::DataSize,	EUnitType::DataSize,	EUnitType::DataSize,	EUnitType::DataSize,	EUnitType::DataSize,
 
-	EUnitType::LuminousFlux, EUnitType::LuminousIntensity, EUnitType::Illuminance, EUnitType::Luminance, EUnitType::ExposureValue,
+	EUnitType::LuminousFlux, EUnitType::LuminousIntensity, EUnitType::Illuminance, EUnitType::Luminance, EUnitType::Luminance, EUnitType::ExposureValue,
 
 	EUnitType::Time,		EUnitType::Time,		EUnitType::Time,		EUnitType::Time,		EUnitType::Time,		EUnitType::Time,		EUnitType::Time,		EUnitType::Time,		EUnitType::Time,
 
diff --git a/Engine/Source/Runtime/Core/Public/Math/UnitConversion.h b/Engine/Source/Runtime/Core/Public/Math/UnitConversion.h
--- a/Engine/Source/Runtime/Core/Public/Math/UnitConversion.h
+++ b/Engine/Source/Runtime/Core/Public/Math/UnitConversion.h
@@ -69,7 +69,7 @@ enum class EUnit : uint8
 	Lux,
 	
 	/** Luminance unit. */
-	CandelaPerMeter2,
+	CandelaPerMeter2, Nit,
 	
 	/** Exposure value unit. */
 	ExposureValue,
diff --git a/Engine/Source/Runtime/Engine/Classes/Engine/Scene.h b/Engine/Source/Runtime/Engine/Classes/Engine/Scene.h
--- a/Engine/Source/Runtime/Engine/Classes/Engine/Scene.h
+++ b/Engine/Source/Runtime/Engine/Classes/Engine/Scene.h
@@ -77,6 +77,7 @@ enum class ELightUnits : uint8
 	Candelas,
 	Lumens,
 	EV,
+	Nits
 };
 
 UENUM()
diff --git a/Engine/Source/Runtime/Engine/Private/Components/LocalLightComponent.cpp b/Engine/Source/Runtime/Engine/Private/Components/LocalLightComponent.cpp
--- a/Engine/Source/Runtime/Engine/Private/Components/LocalLightComponent.cpp
+++ b/Engine/Source/Runtime/Engine/Private/Components/LocalLightComponent.cpp
@@ -183,7 +183,11 @@ float ULocalLightComponent::GetUnitsConversionFactor(ELightUnits SrcUnits, ELigh
 	{
 		float CnvFactor = 1.f;
 		
-		if (SrcUnits == ELightUnits::Candelas)
+		if (SrcUnits == ELightUnits::Nits)
+		{
+			// no scale factor (it depends on the light's area)
+		}
+		else if (SrcUnits == ELightUnits::Candelas)
 		{
 			CnvFactor = 100.f * 100.f;
 		}
@@ -200,7 +204,11 @@ float ULocalLightComponent::GetUnitsConversionFactor(ELightUnits SrcUnits, ELigh
 			CnvFactor = 16.f;
 		}
 
-		if (TargetUnits == ELightUnits::Candelas)
+		if (TargetUnits == ELightUnits::Nits)
+		{
+			// no scale factor (it depends on the light's area)
+		}
+		else if (TargetUnits == ELightUnits::Candelas)
 		{
 			CnvFactor *= 1.f / 100.f / 100.f;
 		}
diff --git a/Engine/Source/Runtime/Engine/Private/Components/PointLightComponent.cpp b/Engine/Source/Runtime/Engine/Private/Components/PointLightComponent.cpp
--- a/Engine/Source/Runtime/Engine/Private/Components/PointLightComponent.cpp
+++ b/Engine/Source/Runtime/Engine/Private/Components/PointLightComponent.cpp
@@ -206,6 +206,12 @@ float UPointLightComponent::ComputeLightBrightness() const
 		{
 			LightBrightness *= (100.f * 100.f); // Conversion from cm2 to m2
 		}
+		else if (IntensityUnits == ELightUnits::Nits)
+		{
+			// Capsule area (sphere+cylinder)
+			float AreaInCm2 = (4 * UE_PI) * SourceRadius * (SourceRadius + 0.5f * SourceLength);
+			LightBrightness *= AreaInCm2;
+		}
 		else if (IntensityUnits == ELightUnits::Lumens)
 		{
 			LightBrightness *= (100.f * 100.f / 4 / UE_PI); // Conversion from cm2 to m2 and 4PI from the sphere area in the 1/r2 attenuation
@@ -231,6 +237,12 @@ void UPointLightComponent::SetLightBrightness(float InBrightness)
 		{
 			Super::SetLightBrightness(InBrightness / (100.f * 100.f)); // Conversion from cm2 to m2
 		}
+		else if (IntensityUnits == ELightUnits::Nits)
+		{
+			// Capsule area (sphere+cylinder)
+			float AreaInCm2 = (4 * UE_PI) * SourceRadius * (SourceRadius + 0.5f * SourceLength);
+			Super::SetLightBrightness(InBrightness / AreaInCm2);
+		}
 		else if (IntensityUnits == ELightUnits::Lumens)
 		{
 			Super::SetLightBrightness(InBrightness / (100.f * 100.f / 4 / UE_PI)); // Conversion from cm2 to m2 and 4PI from the sphere area in the 1/r2 attenuation
diff --git a/Engine/Source/Runtime/Engine/Private/Components/RectLightComponent.cpp b/Engine/Source/Runtime/Engine/Private/Components/RectLightComponent.cpp
--- a/Engine/Source/Runtime/Engine/Private/Components/RectLightComponent.cpp
+++ b/Engine/Source/Runtime/Engine/Private/Components/RectLightComponent.cpp
@@ -173,7 +173,12 @@ float URectLightComponent::ComputeLightBrightness() const
 
 	if (IntensityUnits == ELightUnits::Candelas)
 	{
-		LightBrightness *= (100.f * 100.f); // Conversion from cm2 to m2
+		LightBrightness *= (100.f * 100.f); // Conversion from m2 to cm2
+	}
+	else if (IntensityUnits == ELightUnits::Nits)
+	{
+		const float AreaInCm2 = SourceWidth * SourceHeight;
+		LightBrightness *= AreaInCm2;
 	}
 	else if (IntensityUnits == ELightUnits::Lumens)
 	{
@@ -198,6 +203,11 @@ void URectLightComponent::SetLightBrightness(float InBrightness)
 	{
 		Super::SetLightBrightness(InBrightness / (100.f * 100.f)); // Conversion from cm2 to m2
 	}
+	else if (IntensityUnits == ELightUnits::Nits)
+	{
+		const float AreaInCm2 = SourceWidth * SourceHeight;
+		Super::SetLightBrightness(InBrightness / AreaInCm2);
+	}
 	else if (IntensityUnits == ELightUnits::Lumens)
 	{
 		Super::SetLightBrightness(InBrightness / (100.f * 100.f / UE_PI)); // Conversion from cm2 to m2 and PI from the cosine distribution
diff --git a/Engine/Source/Runtime/Engine/Private/Components/SpotLightComponent.cpp b/Engine/Source/Runtime/Engine/Private/Components/SpotLightComponent.cpp
--- a/Engine/Source/Runtime/Engine/Private/Components/SpotLightComponent.cpp
+++ b/Engine/Source/Runtime/Engine/Private/Components/SpotLightComponent.cpp
@@ -244,6 +244,12 @@ float USpotLightComponent::ComputeLightBrightness() const
 		{
 			LightBrightness *= (100.f * 100.f); // Conversion from cm2 to m2
 		}
+		else if (IntensityUnits == ELightUnits::Nits)
+		{
+			// Capsule area (sphere+cylinder)
+			float AreaInCm2 = (4 * UE_PI) * SourceRadius * (SourceRadius + 0.5f * SourceLength);
+			LightBrightness *= AreaInCm2;
+		}
 		else if (IntensityUnits == ELightUnits::Lumens)
 		{
 			LightBrightness *= (100.f * 100.f / 2.f / UE_PI / (1.f - GetCosHalfConeAngle())); // Conversion from cm2 to m2 and cone remapping.
@@ -269,6 +275,12 @@ void USpotLightComponent::SetLightBrightness(float InBrightness)
 		{
 			ULightComponent::SetLightBrightness(InBrightness / (100.f * 100.f)); // Conversion from cm2 to m2
 		}
+		else if (IntensityUnits == ELightUnits::Nits)
+		{
+			// Capsule area (sphere+cylinder)
+			float AreaInCm2 = (4 * UE_PI) * SourceRadius * (SourceRadius + 0.5f * SourceLength);
+			ULightComponent::SetLightBrightness(InBrightness / AreaInCm2);
+		}
 		else if (IntensityUnits == ELightUnits::Lumens)
 		{
 			ULightComponent::SetLightBrightness(InBrightness / (100.f * 100.f / 2.f / UE_PI / (1.f - GetCosHalfConeAngle()))); // Conversion from cm2 to m2 and cone remapping
diff --git a/Engine/Source/Runtime/Engine/Private/PointLightSceneProxy.h b/Engine/Source/Runtime/Engine/Private/PointLightSceneProxy.h
--- a/Engine/Source/Runtime/Engine/Private/PointLightSceneProxy.h
+++ b/Engine/Source/Runtime/Engine/Private/PointLightSceneProxy.h
@@ -38,7 +38,6 @@ public:
 	,	SourceLength(Component->SourceLength)
 	,	bInverseSquared(Component->bUseInverseSquaredFalloff)
 	{
-		UpdateRadius(Component->AttenuationRadius);
 	}
 
 	virtual float GetSourceRadius() const override
