Fix nw_path_t being released unexpectedly. Add nw_retain/nw_release

#rnx

diff --git a/Engine/Source/Runtime/Core/Private/IOS/IOSPlatformMisc.cpp b/Engine/Source/Runtime/Core/Private/IOS/IOSPlatformMisc.cpp
--- a/Engine/Source/Runtime/Core/Private/IOS/IOSPlatformMisc.cpp
+++ b/Engine/Source/Runtime/Core/Private/IOS/IOSPlatformMisc.cpp
@@ -183,8 +183,15 @@ void FIOSPlatformMisc::PlatformInit()
 		UE::FIOSPlatformMisc::PathMonitor, 
 		^(nw_path_t InPath)
 		{ 
+			UE_LOG(LogInit, Log, TEXT("Network path changed"));
+			
 			// cache the path we get
+			if (UE::FIOSPlatformMisc::PathMonitor != nullptr)
+			{
+				nw_release(UE::FIOSPlatformMisc::CachedNetworkPath);	
+			}
 			UE::FIOSPlatformMisc::CachedNetworkPath = InPath;
+			nw_retain(UE::FIOSPlatformMisc::CachedNetworkPath);
 		}
 	);
 	
@@ -1611,11 +1618,21 @@ void FIOSPlatformMisc::GetValidTargetPlatforms(TArray<FString>& TargetPlatformNa
 static inline ENetworkConnectionType CheckNetworkConnectionType()
 {
 	nw_path_t MyPath = UE::FIOSPlatformMisc::CachedNetworkPath;
-	nw_path_status_t PathStatus = nw_path_get_status(MyPath);
+	if (MyPath == nullptr)
+	{
+		UE_LOG(LogInit, Warning, TEXT("Network path currently NOT set"));
+		return ENetworkConnectionType::Unknown;
+	}
 	
-	if ((PathStatus != nw_path_status_satisfied))
+	nw_path_status_t PathStatus = nw_path_get_status(MyPath);
+	if (PathStatus != nw_path_status_satisfied)
 	{
-		// Apparently there's no valid network paths, so lets see if it's in "airplane" mode.
+		if (PathStatus == nw_path_status_invalid)
+		{
+			UE_LOG(LogInit, Warning, TEXT("Network path is invalid"));
+			return ENetworkConnectionType::None;	
+		}
+		
 		// Basically, if there's _no_ network interfaces available, assume it's in airplane mode.
 		// Note: if you have wifi enabled - but airplane mode also enabled, that will NOT be detected
 		//       as airplane mode as you still have a valid network path.
diff --git a/Engine/Source/Runtime/InstallBundleManager/Private/InstallBundleUtils.cpp b/Engine/Source/Runtime/InstallBundleManager/Private/InstallBundleUtils.cpp
--- a/Engine/Source/Runtime/InstallBundleManager/Private/InstallBundleUtils.cpp
+++ b/Engine/Source/Runtime/InstallBundleManager/Private/InstallBundleUtils.cpp
@@ -31,7 +31,8 @@ namespace InstallBundleUtil
 	bool HasInternetConnection(ENetworkConnectionType ConnectionType)
 	{
 		return ConnectionType != ENetworkConnectionType::AirplaneMode
-			&& ConnectionType != ENetworkConnectionType::None;
+			&& ConnectionType != ENetworkConnectionType::None
+			&& ConnectionType != ENetworkConnectionType::Unknown;
 	}
 
 	bool SplitHostUrl(const FStringView Url, FStringView& OutHost, FStringView& OutRemainder)
