Extended platform I/O dispatcher request to inherit from InheritedContextBase to capture LLM stats correctly.

#rnx

diff --git a/Engine/Source/Runtime/Core/Internal/IO/IoDispatcherFilesystemStats.cpp b/Engine/Source/Runtime/Core/Internal/IO/IoDispatcherFilesystemStats.cpp
--- a/Engine/Source/Runtime/Core/Internal/IO/IoDispatcherFilesystemStats.cpp
+++ b/Engine/Source/Runtime/Core/Internal/IO/IoDispatcherFilesystemStats.cpp
@@ -96,8 +96,8 @@ bool FIoDispatcherFilesystemStats::CsvTick(float DeltaTime)
 {
 	using namespace IoDispatcherFilesystemStats;
 #if CSV_PROFILER_STATS
-	CSV_CUSTOM_STAT_DEFINED(QueuedFilesystemReadMB, BytesToApproxMB(QueuedFilesystemReadBytes), ECsvCustomStatOp::Set);
-	CSV_CUSTOM_STAT_DEFINED(QueuedFilesystemReads, (int32)QueuedFilesystemReads, ECsvCustomStatOp::Set);
+	CSV_CUSTOM_STAT_DEFINED(QueuedFilesystemReadMB, BytesToApproxMB(QueuedFilesystemReadBytes.load(std::memory_order_relaxed)), ECsvCustomStatOp::Set);
+	CSV_CUSTOM_STAT_DEFINED(QueuedFilesystemReads, (int32)QueuedFilesystemReads.load(std::memory_order_relaxed), ECsvCustomStatOp::Set);
 
 	CSV_CUSTOM_STAT_DEFINED(QueuedUncompressBlocks, int32(QueuedUncompressBlocks), ECsvCustomStatOp::Set);
 	CSV_CUSTOM_STAT_DEFINED(QueuedUncompressInMB, BytesToApproxMB(QueuedUncompressBytesIn), ECsvCustomStatOp::Set);
@@ -110,8 +110,8 @@ bool FIoDispatcherFilesystemStats::CsvTick(float DeltaTime)
 void FIoDispatcherFilesystemStats::OnReadRequestsQueued(uint64 ByteCount, uint64 ReadCount)
 {
 #if CSV_PROFILER_STATS
-	QueuedFilesystemReadBytes += ByteCount;
-	QueuedFilesystemReads += ReadCount;
+	QueuedFilesystemReadBytes.fetch_add(ByteCount, std::memory_order_relaxed);
+	QueuedFilesystemReads.fetch_add(ReadCount, std::memory_order_relaxed);
 #endif
 #if COUNTERSTRACE_ENABLED
 	QueuedReadRequestsSizeCounter.Add(ByteCount);
diff --git a/Engine/Source/Runtime/Core/Internal/IO/IoDispatcherFilesystemStats.h b/Engine/Source/Runtime/Core/Internal/IO/IoDispatcherFilesystemStats.h
--- a/Engine/Source/Runtime/Core/Internal/IO/IoDispatcherFilesystemStats.h
+++ b/Engine/Source/Runtime/Core/Internal/IO/IoDispatcherFilesystemStats.h
@@ -45,33 +45,33 @@ private:
 	void			OnHandleChangeSeek();
 
 #if COUNTERSTRACE_ENABLED
-	FCountersTrace::FCounterInt QueuedReadRequestsSizeCounter;
-	FCountersTrace::FCounterInt CompletedReadRequestsSizeCounter;
-	FCountersTrace::FCounterInt QueuedCompressedSizeCounter;
-	FCountersTrace::FCounterInt QueuedUncompressedSizeCounter;
-	FCountersTrace::FCounterInt CompletedCompressedSizeCounter;
-	FCountersTrace::FCounterInt CompletedUncompressedSizeCounter;
-	FCountersTrace::FCounterInt FileSystemSeeksTotalDistanceCounter;
-	FCountersTrace::FCounterInt FileSystemSeeksForwardCountCounter;
-	FCountersTrace::FCounterInt FileSystemSeeksBackwardCountCounter;
-	FCountersTrace::FCounterInt FileSystemSeeksChangeHandleCountCounter;
-	FCountersTrace::FCounterInt FileSystemCompletedRequestsSizeCounter;
-	FCountersTrace::FCounterInt BlockCacheStoredSizeCounter;
-	FCountersTrace::FCounterInt BlockCacheHitSizeCounter;
-	FCountersTrace::FCounterInt BlockCacheMissedSizeCounter;
-	FCountersTrace::FCounterInt ScatteredSizeCounter;
-	FCountersTrace::FCounterInt TocMemoryCounter;
+	FCountersTrace::FCounterAtomicInt	QueuedReadRequestsSizeCounter;
+	FCountersTrace::FCounterInt			CompletedReadRequestsSizeCounter;
+	FCountersTrace::FCounterInt 		QueuedCompressedSizeCounter;
+	FCountersTrace::FCounterInt 		QueuedUncompressedSizeCounter;
+	FCountersTrace::FCounterInt 		CompletedCompressedSizeCounter;
+	FCountersTrace::FCounterInt 		CompletedUncompressedSizeCounter;
+	FCountersTrace::FCounterInt 		FileSystemSeeksTotalDistanceCounter;
+	FCountersTrace::FCounterInt 		FileSystemSeeksForwardCountCounter;
+	FCountersTrace::FCounterInt 		FileSystemSeeksBackwardCountCounter;
+	FCountersTrace::FCounterInt 		FileSystemSeeksChangeHandleCountCounter;
+	FCountersTrace::FCounterInt 		FileSystemCompletedRequestsSizeCounter;
+	FCountersTrace::FCounterInt 		BlockCacheStoredSizeCounter;
+	FCountersTrace::FCounterInt 		BlockCacheHitSizeCounter;
+	FCountersTrace::FCounterInt 		BlockCacheMissedSizeCounter;
+	FCountersTrace::FCounterInt 		ScatteredSizeCounter;
+	FCountersTrace::FCounterInt 		TocMemoryCounter;
 	FCountersTrace::TCounter<std::atomic<int64>, TraceCounterType_Int> AvailableBuffersCounter;
 #endif
 #if CSV_PROFILER_STATS
-	uint64						QueuedFilesystemReadBytes = 0;
-	uint64 						QueuedFilesystemReads = 0;
-	uint64 						QueuedUncompressBytesIn = 0;
-	uint64 						QueuedUncompressBytesOut = 0;
-	uint64 						QueuedUncompressBlocks = 0;
+	std::atomic_uint64_t				QueuedFilesystemReadBytes = 0;
+	std::atomic_uint64_t 				QueuedFilesystemReads = 0;
+	uint64 								QueuedUncompressBytesIn = 0;
+	uint64 								QueuedUncompressBytesOut = 0;
+	uint64 								QueuedUncompressBlocks = 0;
 #endif
-	FTSTicker::FDelegateHandle	TickerHandle;
-	FFileReadInfo				LastFileReadInfo;
+	FTSTicker::FDelegateHandle			TickerHandle;
+	FFileReadInfo						LastFileReadInfo;
 };
 
 #else
diff --git a/Engine/Source/Runtime/Core/Internal/IO/PlatformIoDispatcherBase.cpp b/Engine/Source/Runtime/Core/Internal/IO/PlatformIoDispatcherBase.cpp
--- a/Engine/Source/Runtime/Core/Internal/IO/PlatformIoDispatcherBase.cpp
+++ b/Engine/Source/Runtime/Core/Internal/IO/PlatformIoDispatcherBase.cpp
@@ -1117,6 +1117,7 @@ void FPlatformIoDispatcherBase::ProcessCompletedFileBlock(FIoFileBlockRequest& F
 			}
 		}
 		EncodedBlockRequests = MoveTemp(FileBlockRequest.EncodedBlockRequests);
+		Stats.OnFileBlockCompleted(FileBlockRequest);
 		RequestMgr.Destroy(FileBlockRequest);
 	}
 
@@ -1380,6 +1381,7 @@ void FPlatformIoDispatcherBase::ScatterDecodedBlock(FIoEncodedBlockRequest& Enco
 				check(Request.IsScatterGather());
 				if (Request.Dst.GetSize() == 0)
 				{
+					UE::FInheritedContextScope InheritedContextScope = Request.RestoreInheritedContext();
 					Request.Dst = FIoBuffer(ScatterTarget.Request->DstSize);
 				}
 			}
diff --git a/Engine/Source/Runtime/Core/Internal/IO/PlatformIoDispatcherBase.h b/Engine/Source/Runtime/Core/Internal/IO/PlatformIoDispatcherBase.h
--- a/Engine/Source/Runtime/Core/Internal/IO/PlatformIoDispatcherBase.h
+++ b/Engine/Source/Runtime/Core/Internal/IO/PlatformIoDispatcherBase.h
@@ -2,6 +2,7 @@
 
 #pragma once
 
+#include "Async/InheritedContext.h"
 #include "Async/Mutex.h"
 #include "HAL/PlatformTime.h"
 #include "HAL/Runnable.h"
@@ -134,6 +135,7 @@ using FIoFileBlockLinkAllocator = TSingleThreadedSlabAllocator<FIoFileBlockLink>
 ////////////////////////////////////////////////////////////////////////////////
 struct FIoPlatformReadRequest
 	: public TIntrusiveListElement<FIoPlatformReadRequest>
+	, public FInheritedContextBase
 {
 	enum class EType : uint8
 	{
diff --git a/Engine/Source/Runtime/Core/Internal/IO/PlatformIoDispatcherStats.cpp b/Engine/Source/Runtime/Core/Internal/IO/PlatformIoDispatcherStats.cpp
--- a/Engine/Source/Runtime/Core/Internal/IO/PlatformIoDispatcherStats.cpp
+++ b/Engine/Source/Runtime/Core/Internal/IO/PlatformIoDispatcherStats.cpp
@@ -24,6 +24,10 @@ void FPlatformIoDispatcherStats::OnFileBlockReadStarted(FIoFileBlockRequest& Fil
 void FPlatformIoDispatcherStats::OnFileBlockReadCompleted(FIoFileBlockRequest& FileBlockRequest)
 {
 	Stats.OnFilesystemReadCompleted(FileBlockRequest.FileHandle.Value(), FileBlockRequest.FileOffset, FileBlockRequest.Size);
+}
+
+void FPlatformIoDispatcherStats::OnFileBlockCompleted(FIoFileBlockRequest& FileBlockRequest)
+{
 	Stats.OnReadRequestsCompleted(FileBlockRequest.Size, 1);
 }
 
diff --git a/Engine/Source/Runtime/Core/Internal/IO/PlatformIoDispatcherStats.h b/Engine/Source/Runtime/Core/Internal/IO/PlatformIoDispatcherStats.h
--- a/Engine/Source/Runtime/Core/Internal/IO/PlatformIoDispatcherStats.h
+++ b/Engine/Source/Runtime/Core/Internal/IO/PlatformIoDispatcherStats.h
@@ -21,6 +21,7 @@ public:
 	void OnFileBlockRequestDequeued(FIoFileBlockRequest& FileBlockRequest);
 	void OnFileBlockReadStarted(FIoFileBlockRequest& FileBlockRequest);
 	void OnFileBlockReadCompleted(FIoFileBlockRequest& FileBlockRequest);
+	void OnFileBlockCompleted(FIoFileBlockRequest& FileBlockRequest);
 	void OnFileBlockMemoryAllocated(uint32 Size);
 	void OnFileBlockMemoryFreed(uint32 Size);
 	void OnFileBlockCacheHit(uint64 Size);
@@ -43,6 +44,7 @@ public:
 	void OnFileBlockRequestDequeued(FIoFileBlockRequest& FileBlockRequest) { };
 	void OnFileBlockReadStarted(FIoFileBlockRequest& FileBlockRequest) { };
 	void OnFileBlockReadCompleted(FIoFileBlockRequest& FileBlockRequest) { };
+	void OnFileBlockCompleted(FIoFileBlockRequest& FileBlockRequest) { }
 	void OnFileBlockMemoryAllocated(uint32 Size) { };
 	void OnFileBlockMemoryFreed(uint32 Size) { };
 	void OnFileBlockCacheHit(uint64 Size) { }
