[ModuleManager]

Allow the GameThread to re-use already loaded modules instead of re-running expensive file existence checks attempting to load load something we already loaded.

Reduces UContentBrowserDataSubsystem::BroadCastItemDataUpdate (19.3s) to UContentBrowserDataSubsystem::BroadCastItemDataUpdate (3.8s)

#rnx

diff --git a/Engine/Source/Runtime/Core/Private/Modules/ModuleManager.cpp b/Engine/Source/Runtime/Core/Private/Modules/ModuleManager.cpp
--- a/Engine/Source/Runtime/Core/Private/Modules/ModuleManager.cpp
+++ b/Engine/Source/Runtime/Core/Private/Modules/ModuleManager.cpp
@@ -714,13 +714,14 @@ IModuleInterface* FModuleManager::GetOrLoadModule(const FName InModuleName, EMod
 	LLM_SCOPE_BYNAME(TEXT("Modules"));
 	// We allow an already loaded module to be returned in other threads to simplify
 	// parallel processing scenarios but they must have been loaded from the main thread beforehand.
-	if(!IsInGameThread())
+	IModuleInterface* Module = GetModule(InModuleName);
+	if (Module)
 	{
-		IModuleInterface* Module = GetModule(InModuleName);
-		if (!Module)
-		{
-			OutFailureReason = EModuleLoadResult::NotLoadedByGameThread;
-		}
+		return Module;
+	}
+	else if (!IsInGameThread())
+	{
+		OutFailureReason = EModuleLoadResult::NotLoadedByGameThread;
 		return Module;
 	}
 
