Reference correct CPU Frequency Percentage stat. Drive by display fixes to displayed stats as float where applicable.

diff --git a/Engine/Source/Runtime/Core/Private/Android/AndroidStats.cpp b/Engine/Source/Runtime/Core/Private/Android/AndroidStats.cpp
--- a/Engine/Source/Runtime/Core/Private/Android/AndroidStats.cpp
+++ b/Engine/Source/Runtime/Core/Private/Android/AndroidStats.cpp
@@ -68,7 +68,7 @@ CSV_DEFINE_STAT(AndroidCPU, CPUTemp);
 DECLARE_FLOAT_COUNTER_STAT(TEXT("CPU Temperature"), STAT_CPUTemp, STATGROUP_AndroidCPU);
 
 CSV_DEFINE_STAT(AndroidCPU, ThermalStatus);
-DECLARE_FLOAT_COUNTER_STAT(TEXT("Thermal Status"), STAT_ThermalStatus, STATGROUP_AndroidCPU);
+DECLARE_DWORD_COUNTER_STAT(TEXT("Thermal Status"), STAT_ThermalStatus, STATGROUP_AndroidCPU);
 
 CSV_DEFINE_STAT(AndroidCPU, ThermalStress);
 DECLARE_FLOAT_COUNTER_STAT(TEXT("Thermal Stress"), STAT_ThermalStress, STATGROUP_AndroidCPU);
@@ -408,13 +408,13 @@ void FAndroidStats::UpdateAndroidStats()
 			SET_DWORD_STAT_FName(AndroidFrequencyGroupNumCoresStats[FrequencyGroupIndex], FrequencyGroup.CoreCount);
 			//SET_DWORD_STAT_FName(AndroidFrequencyGroupMinFreqStats[FrequencyGroupIndex], FrequencyGroup.MinFrequency);
 			SET_DWORD_STAT_FName(AndroidFrequencyGroupCurrentFreqStats[FrequencyGroupIndex], CurrentFrequencies[FrequencyGroupIndex]);
-			SET_DWORD_STAT_FName(AndroidFrequencyGroupCurrentFreqPercentageStats[FrequencyGroupIndex], CurrentFrequencies[FrequencyGroupIndex]);
-			SET_DWORD_STAT_FName(AndroidFrequencyGroupMaxCoresUtilizationStats[FrequencyGroupIndex], MaxSingleCoreUtilization[FrequencyGroupIndex]);
+			SET_FLOAT_STAT_FName(AndroidFrequencyGroupCurrentFreqPercentageStats[FrequencyGroupIndex], CurrentFrequenciesPercentage[FrequencyGroupIndex]);
+			SET_FLOAT_STAT_FName(AndroidFrequencyGroupMaxCoresUtilizationStats[FrequencyGroupIndex], MaxSingleCoreUtilization[FrequencyGroupIndex]);
 		}
 
 		static const FName CPUStatName = GET_STATFNAME(STAT_CPUTemp);
 		static const FName ThermalStatus = GET_STATFNAME(STAT_ThermalStatus);
-		SET_DWORD_STAT_FName(CPUStatName, CPUTemp);
+		SET_FLOAT_STAT_FName(CPUStatName, CPUTemp);
 		SET_DWORD_STAT_FName(ThermalStatus, GThermalStatus);
 #endif
 
