Fix for deadlock in FConfigCacheIni::PreInitializePlatformPlugins. The deadlock would happen when waiting on a config future inside a ParallelFor thread callback.

Also removed call to PreInitializePlatformPlugins in TargetPlatformManagerModule as it wasn't really making a performance difference in ContentWorker cooks.

#rnx

diff --git a/Engine/Source/Developer/TargetPlatform/Private/TargetPlatformManagerModule.cpp b/Engine/Source/Developer/TargetPlatform/Private/TargetPlatformManagerModule.cpp
--- a/Engine/Source/Developer/TargetPlatform/Private/TargetPlatformManagerModule.cpp
+++ b/Engine/Source/Developer/TargetPlatform/Private/TargetPlatformManagerModule.cpp
@@ -1049,8 +1049,6 @@ protected:
 		}
 #endif
 
-		FConfigCacheIni::PreInitializePlatformPlugins();
-
 		// find a set of valid target platform names (the platform DataDrivenPlatformInfo.ini file was found indicates support for the platform 
 		// exists on disk, so the TP is expected to work)
 		FScopedSlowTask SlowTask((float)FDataDrivenPlatformInfoRegistry::GetAllPlatformInfos().Num());
diff --git a/Engine/Source/Runtime/Core/Private/Misc/ConfigCacheIni.cpp b/Engine/Source/Runtime/Core/Private/Misc/ConfigCacheIni.cpp
--- a/Engine/Source/Runtime/Core/Private/Misc/ConfigCacheIni.cpp
+++ b/Engine/Source/Runtime/Core/Private/Misc/ConfigCacheIni.cpp
@@ -6636,6 +6636,19 @@ void FConfigCacheIni::PreInitializePlatformPlugins()
 	TArray<FDataDrivenPlatformInfo> PlatformInfos;
 	AllPlatformInfos.GetKeys(PlatformNames);
 
+	for (FName PlatformName : PlatformNames)
+	{
+		if (PlatformName != NAME_None)
+		{
+			// GetPlatformConfigFutures will be called inside FConfigCacheIni::ForPlatform. 
+			// To avoid deadlocks wait for all futures here before entering the ParallelFor.
+			if (auto* PlatformCache = GetPlatformConfigFutures().Find(PlatformName))
+			{
+				PlatformCache->Get();
+			}
+		}
+	}
+
 	ParallelFor(PlatformNames.Num(), [PlatformNames, AllPlatformInfos](int32 Index)
 	{
 		FName PlatfName = PlatformNames[Index];
