[AsymmetricFence]
- Introduce asymmetric fence that we can use to optimize a critical path while using heavier fence on the slow path

diff --git a/Engine/Source/Runtime/Core/Private/Windows/WindowsPlatformMisc.cpp b/Engine/Source/Runtime/Core/Private/Windows/WindowsPlatformMisc.cpp
--- a/Engine/Source/Runtime/Core/Private/Windows/WindowsPlatformMisc.cpp
+++ b/Engine/Source/Runtime/Core/Private/Windows/WindowsPlatformMisc.cpp
@@ -1263,6 +1263,15 @@ EProcessDiagnosticFlags FWindowsPlatformMisc::GetProcessDiagnostics()
 }
 #endif //!UE_BUILD_SHIPPING
 
+void FWindowsPlatformMisc::AsymmetricThreadFenceHeavy()
+{
+	TRACE_CPUPROFILER_EVENT_SCOPE(FPlatformMisc::AsymmetricThreadFenceHeavy);
+
+	// The function generates an interprocessor interrupt (IPI) to all processors that are part of the current process affinity.
+	// It guarantees the visibility of write operations performed on one processor to the other processors.
+	FlushProcessWriteBuffers();
+}
+
 bool FWindowsPlatformMisc::IsRemoteSession()
 {
 	return ::GetSystemMetrics(SM_REMOTESESSION) != 0;
diff --git a/Engine/Source/Runtime/Core/Public/HAL/Platform.h b/Engine/Source/Runtime/Core/Public/HAL/Platform.h
--- a/Engine/Source/Runtime/Core/Public/HAL/Platform.h
+++ b/Engine/Source/Runtime/Core/Public/HAL/Platform.h
@@ -519,6 +519,10 @@
 	#define PLATFORM_IS_ANSI_MALLOC_THREADSAFE					0
 #endif
 
+#ifndef PLATFORM_SUPPORTS_ASYMMETRIC_FENCES
+	#define PLATFORM_SUPPORTS_ASYMMETRIC_FENCES					0
+#endif
+
 #ifndef PLATFORM_SUPPORTS_OPUS_CODEC
 	#define PLATFORM_SUPPORTS_OPUS_CODEC						1
 #endif
diff --git a/Engine/Source/Runtime/Core/Public/Windows/WindowsPlatform.h b/Engine/Source/Runtime/Core/Public/Windows/WindowsPlatform.h
--- a/Engine/Source/Runtime/Core/Public/Windows/WindowsPlatform.h
+++ b/Engine/Source/Runtime/Core/Public/Windows/WindowsPlatform.h
@@ -93,6 +93,8 @@ typedef FWindowsPlatformTypes FPlatformTypes;
 #define PLATFORM_HAS_BSD_SOCKET_FEATURE_WINSOCKETS			1
 #define PLATFORM_USES_MICROSOFT_LIBC_FUNCTIONS				1
 #define PLATFORM_IS_ANSI_MALLOC_THREADSAFE					1
+#define PLATFORM_SUPPORTS_ASYMMETRIC_FENCES					1
+
 #if PLATFORM_CPU_ARM_FAMILY
 #define PLATFORM_SUPPORTS_TBB								0
 #else
diff --git a/Engine/Source/Runtime/Core/Public/Windows/WindowsPlatformMisc.h b/Engine/Source/Runtime/Core/Public/Windows/WindowsPlatformMisc.h
--- a/Engine/Source/Runtime/Core/Public/Windows/WindowsPlatformMisc.h
+++ b/Engine/Source/Runtime/Core/Public/Windows/WindowsPlatformMisc.h
@@ -128,6 +128,13 @@ struct FWindowsPlatformMisc
 #endif
 	}
 
+	static CORE_API void AsymmetricThreadFenceHeavy();
+
+	FORCEINLINE static void AsymmetricThreadFenceLight()
+	{
+		_ReadWriteBarrier();
+	}
+
 	static CORE_API bool IsRemoteSession();
 
 	static CORE_API void SetUTF8Output();
diff --git a/Engine/Source/Runtime/Core/Tests/Tasks/TasksTest.cpp b/Engine/Source/Runtime/Core/Tests/Tasks/TasksTest.cpp
--- a/Engine/Source/Runtime/Core/Tests/Tasks/TasksTest.cpp
+++ b/Engine/Source/Runtime/Core/Tests/Tasks/TasksTest.cpp
@@ -1760,6 +1760,44 @@ namespace UE { namespace TasksTests
 		// Restore original worker count
 		LowLevelTasks::FScheduler::Get().RestartWorkers();
 	}
+
+#if PLATFORM_SUPPORTS_ASYMMETRIC_FENCES
+	TEST_CASE_NAMED(FAsymmetricThreadFence, "System::Core::Async::AsymmetricThreadFence", "[.][ApplicationContextMask][EngineFilter]")
+	{
+		for (int32 Index = 0; Index < 10000; ++Index)
+		{
+			std::atomic<int> A = 0;
+			std::atomic<int> B = 0;
+			std::atomic<int> R0 = 0;
+			std::atomic<int> R1 = 0;
+
+			// If you use AsymmetricThreadFenceLight in both threads, you'll see you can end up with a race
+			// because nothing publishes the A value before reading the B value. Using an heavy fence
+			// on one side ensures that the test can't fail even when one side is only using relaxed operations
+			// along with a simple compiler barrier.
+			UE::Tasks::FTask T1 =
+				UE::Tasks::Launch(
+					TEXT("T1"),
+					[&A, &B, &R0]()
+					{
+						A.store(1, std::memory_order_relaxed);
+						FPlatformMisc::AsymmetricThreadFenceHeavy();
+						R0.store(B.load(std::memory_order_relaxed), std::memory_order_relaxed);
+					}
+				);
+
+			B.store(1, std::memory_order_relaxed);
+			FPlatformMisc::AsymmetricThreadFenceLight();
+			R1.store(A.load(std::memory_order_relaxed), std::memory_order_relaxed);
+
+			T1.Wait();
+
+			// Both value should never be 0 but it can happen when no synchronization is used due to the presence of a store buffer.
+			verify(R1 == 1 || R0 == 1);
+		}
+	}
+#endif
+
 }}
 
 #endif // WITH_TESTS
