[LLM]
- Fix thread-safety issue on init

diff --git a/Engine/Source/Runtime/Core/Private/HAL/LowLevelMemTracker.cpp b/Engine/Source/Runtime/Core/Private/HAL/LowLevelMemTracker.cpp
--- a/Engine/Source/Runtime/Core/Private/HAL/LowLevelMemTracker.cpp
+++ b/Engine/Source/Runtime/Core/Private/HAL/LowLevelMemTracker.cpp
@@ -1838,7 +1838,17 @@ void FLowLevelMemTracker::BootstrapInitialise()
 	{
 		return;
 	}
-	bInitialisedTracking = true;
+	FScopeLock BootstrapScopeLock(&BootstrapLock);
+	if (bInitialisedTracking)
+	{
+		return;
+	}
+	// just make sure this part isn't reentrant.
+	if (bBeganInitialisedTracking != false)
+	{
+		UE_LOG(LogInit, Fatal, TEXT("FLowLevelMemTracker::BootstrapInitialise is not reentrant"));
+	}
+	bBeganInitialisedTracking = true;
 
 	for (int32 TrackerIndex = 0; TrackerIndex < static_cast<int32>(ELLMTracker::Max); ++TrackerIndex)
 	{
@@ -1860,6 +1870,10 @@ void FLowLevelMemTracker::BootstrapInitialise()
 
 	// calculate program size early on... the platform can call SetProgramSize later if it sees fit
 	InitialiseProgramSize();
+
+	FPlatformMisc::MemoryBarrier();
+
+	bInitialisedTracking = true;
 }
 
 void FLowLevelMemTracker::Clear()
diff --git a/Engine/Source/Runtime/Core/Public/HAL/LowLevelMemTracker.h b/Engine/Source/Runtime/Core/Public/HAL/LowLevelMemTracker.h
--- a/Engine/Source/Runtime/Core/Public/HAL/LowLevelMemTracker.h
+++ b/Engine/Source/Runtime/Core/Public/HAL/LowLevelMemTracker.h
@@ -857,6 +857,7 @@ private:
 
 	mutable FRWLock TagDataLock;
 	FCriticalSection UpdateLock;
+	FCriticalSection BootstrapLock;
 
 	uint64 ProgramSize;
 	int64 MemoryUsageCurrentOverhead;
@@ -870,6 +871,7 @@ private:
 	bool bTraceWriterEnabled;
 	bool bInitialisedTracking;
 	bool bIsBootstrapping;
+	bool bBeganInitialisedTracking = false;
 	bool bFullyInitialised;
 	std::atomic<bool> bConfigurationComplete;
 	bool bTagAdded;
