Changed the return types of TLruCache::Find variants to not be const. Users can still bind the return to a const type if they need.

diff --git a/Engine/Source/Runtime/Core/Public/Containers/LruCache.h b/Engine/Source/Runtime/Core/Public/Containers/LruCache.h
--- a/Engine/Source/Runtime/Core/Public/Containers/LruCache.h
+++ b/Engine/Source/Runtime/Core/Public/Containers/LruCache.h
@@ -385,6 +385,21 @@ public:
 		return (*EntryPtr)->Value;
 	}
 
+	/**
+	 * Find the value of the entry with the specified key.
+	 *
+	 * @param Key The key of the entry to get.
+	 * @return Reference to the value, or triggers an assertion if the key does not exist.
+	 */
+	[[nodiscard]] FORCEINLINE ValueType& FindChecked(const KeyType& Key)
+	{
+		FCacheEntry*const * EntryPtr = LookupSet.Find(Key);
+
+		check(EntryPtr);
+
+		return (*EntryPtr)->Value;
+	}
+
 	/**
 	 * Find the value of the entry with the specified key.
 	 *
@@ -410,7 +425,7 @@ public:
 	 * @return Pointer to the value, or nullptr if not found.
 	 * @see Add, Contains, Empty, Find, GetKeys, Remove
 	 */
-	[[nodiscard]] const ValueType* FindAndTouch(const KeyType& Key)
+	[[nodiscard]] ValueType* FindAndTouch(const KeyType& Key)
 	{
 		FCacheEntry** EntryPtr = LookupSet.Find(Key);
 
@@ -430,7 +445,7 @@ public:
 	 * @param Key The key of the entry to get.
 	 * @return Pointer to the value, or triggers an assertion if the key does not exist.
 	 */
-	[[nodiscard]] const ValueType& FindAndTouchChecked(const KeyType& Key)
+	[[nodiscard]] ValueType& FindAndTouchChecked(const KeyType& Key)
 	{
 		FCacheEntry** EntryPtr = LookupSet.Find(Key);
 
