Logging: Truncate pre-formatted messages in StaticDebugV to avoid omitting them entirely when the length of the message is known

diff --git a/Engine/Source/Runtime/Core/Private/Logging/LogMacros.cpp b/Engine/Source/Runtime/Core/Private/Logging/LogMacros.cpp
--- a/Engine/Source/Runtime/Core/Private/Logging/LogMacros.cpp
+++ b/Engine/Source/Runtime/Core/Private/Logging/LogMacros.cpp
@@ -12,7 +12,7 @@
 #include "ProfilingDebugging/CsvProfiler.h"
 #include "AutoRTFM/AutoRTFM.h"
 
-void StaticFailDebugV(const TCHAR* Error, const ANSICHAR* Expression, const ANSICHAR* File, int32 Line, bool bIsEnsure, void* ProgramCounter, const TCHAR* DescriptionFormat, va_list DescriptionArgs);
+void StaticFailDebugV(const TCHAR* Error, const ANSICHAR* File, int32 Line, void* ProgramCounter, const TCHAR* DescriptionFormat, va_list DescriptionArgs);
 
 CSV_DEFINE_CATEGORY(FMsgLogf, true);
 
@@ -47,7 +47,7 @@ void FMsg::LogfImplV(const ANSICHAR* File, int32 Line, const FLogCategoryName& C
 	{
 		va_list ArgsCopy;
 		va_copy(ArgsCopy, Args);
-		StaticFailDebugV(TEXT("Fatal error:"), "", File, Line, /*bIsEnsure*/ false, ReturnAddress, Fmt, ArgsCopy);
+		StaticFailDebugV(TEXT("Fatal error:"), File, Line, ReturnAddress, Fmt, ArgsCopy);
 		va_end(ArgsCopy);
 
 		FDebug::AssertFailedV("", File, Line, Fmt, Args);
@@ -108,7 +108,7 @@ void FMsg::LogV(const ANSICHAR* File, int32 Line, const FLogCategoryName& Catego
 	}
 	else
 	{
-		StaticFailDebugV(TEXT("Fatal error:"), "", File, Line, /*bIsEnsure*/ false, PLATFORM_RETURN_ADDRESS(), Fmt, Args);
+		StaticFailDebugV(TEXT("Fatal error:"), File, Line, PLATFORM_RETURN_ADDRESS(), Fmt, Args);
 	}
 #endif
 }
diff --git a/Engine/Source/Runtime/Core/Private/Logging/StructuredLog.cpp b/Engine/Source/Runtime/Core/Private/Logging/StructuredLog.cpp
--- a/Engine/Source/Runtime/Core/Private/Logging/StructuredLog.cpp
+++ b/Engine/Source/Runtime/Core/Private/Logging/StructuredLog.cpp
@@ -28,8 +28,8 @@
 #include "Templates/Function.h"
 #include <cstdarg>
 
-void VARARGS StaticFailDebug (const TCHAR* Error, const ANSICHAR* Expression, const ANSICHAR* File, int32 Line, bool bIsEnsure, void* ProgramCounter, const TCHAR* DescriptionFormat, ...);
-void         StaticFailDebugV(const TCHAR* Error, const ANSICHAR* Expression, const ANSICHAR* File, int32 Line, bool bIsEnsure, void* ProgramCounter, const TCHAR* DescriptionFormat, va_list DescriptionArgs);
+void StaticFailDebug (const TCHAR* Error, const ANSICHAR* File, int32 Line, void* ProgramCounter, const TCHAR* Message);
+void StaticFailDebugV(const TCHAR* Error, const ANSICHAR* File, int32 Line, void* ProgramCounter, const TCHAR* DescriptionFormat, va_list DescriptionArgs);
 
 namespace UE::Logging::Private
 {
@@ -1029,7 +1029,7 @@ UE_AUTORTFM_ALWAYS_OPEN
 	TStringBuilder<512> Message;
 	Record.FormatMessageTo(Message);
 
-	StaticFailDebug(TEXT("Fatal error:"), "", Log.File, Log.Line, /*bIsEnsure*/ false, PLATFORM_RETURN_ADDRESS(), TEXT("%s"), *Message);
+	StaticFailDebug(TEXT("Fatal error:"), Log.File, Log.Line, PLATFORM_RETURN_ADDRESS(), *Message);
 
 	UE_DEBUG_BREAK_AND_PROMPT_FOR_REMOTE();
 	FDebug::ProcessFatalError(PLATFORM_RETURN_ADDRESS());
@@ -1051,7 +1051,7 @@ UE_AUTORTFM_ALWAYS_OPEN
 	TStringBuilder<512> Message;
 	Record.FormatMessageTo(Message);
 
-	StaticFailDebug(TEXT("Fatal error:"), "", Log.File, Log.Line, /*bIsEnsure*/ false, PLATFORM_RETURN_ADDRESS(), TEXT("%s"), *Message);
+	StaticFailDebug(TEXT("Fatal error:"), Log.File, Log.Line, PLATFORM_RETURN_ADDRESS(), *Message);
 
 	UE_DEBUG_BREAK_AND_PROMPT_FOR_REMOTE();
 	FDebug::ProcessFatalError(PLATFORM_RETURN_ADDRESS());
@@ -1243,7 +1243,7 @@ UE_AUTORTFM_ALWAYS_OPEN
 static void BasicFatalLogV(const FLogCategoryBase& Category, const FStaticBasicLogRecord* Log, void* ProgramCounter, va_list Args)
 {
 #if !NO_LOGGING
-	StaticFailDebugV(TEXT("Fatal error:"), "", Log->File, Log->Line, /*bIsEnsure*/ false, ProgramCounter, Log->Format, Args);
+	StaticFailDebugV(TEXT("Fatal error:"), Log->File, Log->Line, ProgramCounter, Log->Format, Args);
 
 	UE_DEBUG_BREAK_AND_PROMPT_FOR_REMOTE();
 	FDebug::ProcessFatalError(ProgramCounter);
diff --git a/Engine/Source/Runtime/Core/Private/Misc/AssertionMacros.cpp b/Engine/Source/Runtime/Core/Private/Misc/AssertionMacros.cpp
--- a/Engine/Source/Runtime/Core/Private/Misc/AssertionMacros.cpp
+++ b/Engine/Source/Runtime/Core/Private/Misc/AssertionMacros.cpp
@@ -180,24 +180,43 @@ public:
 
 	inline ~FErrorHistWriter()
 	{
-		Terminate();
+		TerminateWith(TEXT(""));
 		FailDebugMutex.Unlock();
 	}
 
-	inline void Terminate()
+	inline void TerminateWith(const TCHAR* Text)
 	{
+		const int32 TextLen = FCString::Strlen(Text);
+		if (Index + TextLen >= Capacity)
+		{
+			Index = Capacity - TextLen - 1;
+		}
+		Append(Text);
 		GErrorHist[Index] = TCHAR(0);
 	}
 
+	void Append(const TCHAR* Text)
+	{
+		const int32 TextLen = FCString::Strlen(Text);
+		const int32 CopyLen = FPlatformMath::Min(TextLen + 1, Capacity - Index);
+		if (CopyLen > 1)
+		{
+			FCString::Strncpy(GErrorHist + Index, Text, CopyLen);
+			Index += CopyLen - 1;
+		}
+	}
+
 	template <typename CharType>
 	void Append(const CharType* Text)
 	{
-		const int32 TextLen = TCString<CharType>::Strlen(Text);
-		const int32 RequiredLen = FPlatformString::ConvertedLength<TCHAR>(Text, TextLen);
-		if (Index + RequiredLen < Capacity)
+		if (const int32 TextLen = TCString<CharType>::Strlen(Text))
 		{
-			FPlatformString::Convert(GErrorHist + Index, RequiredLen, Text, TextLen);
-			Index += RequiredLen;
+			const int32 RequiredLen = FPlatformString::ConvertedLength<TCHAR>(Text, TextLen);
+			if (Index + RequiredLen < Capacity)
+			{
+				FPlatformString::Convert(GErrorHist + Index, RequiredLen, Text, TextLen);
+				Index += RequiredLen;
+			}
 		}
 	}
 
@@ -231,6 +250,7 @@ FORCENOINLINE void StaticFailDebugV(
 	int32 Line,
 	bool bIsEnsure,
 	void* ProgramCounter,
+	const TCHAR* Message,
 	const TCHAR* DescriptionFormat,
 	va_list DescriptionArgs)
 {
@@ -252,20 +272,26 @@ FORCENOINLINE void StaticFailDebugV(
 	Writer.Append(Error);
 	Writer.Append(Expression);
 	Writer.Appendf(TEXT(FILE_LINE_DESC_ANSI "\n"), File, Line);
-	Writer.AppendV(DescriptionFormat, DescriptionArgs);
-	Writer.Append(TEXT("\n"));
-	Writer.Terminate();
+	if (Message)
+	{
+		Writer.Append(Message);
+		Writer.TerminateWith(TEXT("\n"));
+	}
+	if (DescriptionFormat)
+	{
+		Writer.AppendV(DescriptionFormat, DescriptionArgs);
+		Writer.TerminateWith(TEXT("\n"));
+	}
 
 	FPlatformMisc::LowLevelOutputDebugString(GErrorHist);
 
 	if (*StackTrace)
 	{
 		Writer.Append(StackTrace);
-		Writer.Append(TEXT("\n"));
+		Writer.TerminateWith(TEXT("\n"));
 	}
 
-	Writer.Append(TEXT("\r\n\r\n"));
-	Writer.Terminate();
+	Writer.TerminateWith(TEXT("\r\n\r\n"));
 
 	if (GError)
 	{
@@ -273,13 +299,15 @@ FORCENOINLINE void StaticFailDebugV(
 	}
 }
 
-FORCENOINLINE void VARARGS StaticFailDebug(const TCHAR* Error, const ANSICHAR* Expression, const ANSICHAR* File, int32 Line,
-	bool bIsEnsure, void* ProgramCounter, const TCHAR* DescriptionFormat, ...)
+FORCENOINLINE void VARARGS StaticFailDebug(const TCHAR* Error, const ANSICHAR* File, int32 Line, void* ProgramCounter, const TCHAR* Message)
+{
+	va_list DescriptionArgs{};
+	StaticFailDebugV(Error, /*Expression*/ "", File, Line, /*bIsEnsure*/ false, ProgramCounter, Message, /*DescriptionFormat*/ nullptr, DescriptionArgs);
+}
+
+FORCENOINLINE void VARARGS StaticFailDebugV(const TCHAR* Error, const ANSICHAR* File, int32 Line, void* ProgramCounter, const TCHAR* DescriptionFormat, va_list DescriptionArgs)
 {
-	va_list DescriptionArgs;
-	va_start(DescriptionArgs, DescriptionFormat);
-	StaticFailDebugV(Error, Expression, File, Line, bIsEnsure, ProgramCounter, DescriptionFormat, DescriptionArgs);
-	va_end(DescriptionArgs);
+	StaticFailDebugV(Error, /*Expression*/ "", File, Line, /*bIsEnsure*/ false, ProgramCounter, /*Message*/ nullptr, DescriptionFormat, DescriptionArgs);
 }
 
 /// track thread asserts
@@ -386,7 +414,7 @@ void FDebug::LogAssertFailedMessageImplV(const ANSICHAR* Expr, const ANSICHAR* F
 	// Ignore this assert if we're already forcibly shutting down because of a critical error.
 	if( !GIsCriticalError )
 	{
-		StaticFailDebugV(TEXT("Assertion failed: "), Expr, File, Line, /*bIsEnsure*/ false, ProgramCounter, Fmt, Args);
+		StaticFailDebugV(TEXT("Assertion failed: "), Expr, File, Line, /*bIsEnsure*/ false, ProgramCounter, /*Message*/ nullptr, Fmt, Args);
 	}
 }
 
@@ -446,7 +474,8 @@ FORCENOINLINE void FDebug::EnsureFailed(const ANSICHAR* Expr, const ANSICHAR* Fi
 	// Should we spin here?
 	FPlatformAtomics::InterlockedIncrement(&ActiveEnsureCount);
 
-	StaticFailDebug(TEXT("Ensure condition failed: "), Expr, File, Line, /*bIsEnsure*/ true, ProgramCounter, TEXT("%s"), Msg);
+	va_list DescriptionArgs{};
+	StaticFailDebugV(TEXT("Ensure condition failed: "), Expr, File, Line, /*bIsEnsure*/ true, ProgramCounter, Msg, nullptr, DescriptionArgs);
 
 	// Is there a debugger attached?  If not we'll submit an error report.
 	if (FPlatformMisc::IsDebuggerPresent() && !GAlwaysReportCrash)
@@ -758,7 +787,7 @@ void UE_COLD UE_DEBUG_SECTION VARARGS LowLevelFatalErrorHandler(const ANSICHAR*
 {
 	va_list Args;
 	va_start(Args, Format);
-	StaticFailDebugV(TEXT("LowLevelFatalError"), "", File, Line, /*bIsEnsure*/ false, PLATFORM_RETURN_ADDRESS(), Format, Args);
+	StaticFailDebugV(TEXT("LowLevelFatalError"), /*Expression*/ "", File, Line, /*bIsEnsure*/ false, PLATFORM_RETURN_ADDRESS(), /*Message*/ nullptr, Format, Args);
 	va_end(Args);
 
 	UE_DEBUG_BREAK_AND_PROMPT_FOR_REMOTE();
