[AutoRTFM] Remove the AUTORTFM_EXPECT that checks transactions are only retried at most once.

Transactions are being retried multiple times by the async loader, likely because of sequential async-loads.

diff --git a/Engine/Source/Runtime/Core/Private/AutoRTFM/Context.cpp b/Engine/Source/Runtime/Core/Private/AutoRTFM/Context.cpp
--- a/Engine/Source/Runtime/Core/Private/AutoRTFM/Context.cpp
+++ b/Engine/Source/Runtime/Core/Private/AutoRTFM/Context.cpp
@@ -168,16 +168,7 @@ ETransactionResult FContext::RollbackTransaction(EContextStatus NewStatus)
 		   NewStatus == EContextStatus::AbortedByCascadingRetry);
 
 	Status = NewStatus;
-		
-	if (Status == EContextStatus::AbortedByCascadingRetry)
-	{
-		if (ForTheRuntime::GetInternalAbortAction() != ForTheRuntime::EAutoRTFMInternalAbortActionState::Abort)
-		{
-			AUTORTFM_EXPECT(0 == NumAbortAndRetries);
-			NumAbortAndRetries += 1;
-		}
-	}
-	
+
 	ETransactionResult Result = ETransactionResult::AbortedByRequest;
 
 	if (GetNumDeferredTransactions())
@@ -709,7 +700,6 @@ void FContext::Reset()
 	Status = EContextStatus::Idle;
 	StackLocalInitializerDepth = 0;
 	TaskPool.Reset();
-	NumAbortAndRetries = 0;
 }
 
 void FContext::Throw()
diff --git a/Engine/Source/Runtime/Core/Private/AutoRTFM/Context.h b/Engine/Source/Runtime/Core/Private/AutoRTFM/Context.h
--- a/Engine/Source/Runtime/Core/Private/AutoRTFM/Context.h
+++ b/Engine/Source/Runtime/Core/Private/AutoRTFM/Context.h
@@ -176,7 +176,6 @@ private:
 	FThreadID CurrentThreadId;
 	uint32_t StackLocalInitializerDepth = 0;
 	TTask<void()> PostAbortCallback;
-	uint32_t NumAbortAndRetries = 0;
 	FTaskPool TaskPool;
 	FTransactionPool TransactionPool;
 	FOpenHashThrottler OpenHashThrottler{
