Fixed .natvis of FNames under Clang.

diff --git a/Engine/Source/Runtime/Core/Private/Misc/CoreGlobals.cpp b/Engine/Source/Runtime/Core/Private/Misc/CoreGlobals.cpp
--- a/Engine/Source/Runtime/Core/Private/Misc/CoreGlobals.cpp
+++ b/Engine/Source/Runtime/Core/Private/Misc/CoreGlobals.cpp
@@ -258,6 +258,11 @@ float				GNearClippingPlane_RenderThread		= 10.0f;				/* Near clipping plane (Re
 bool					GExitPurge						= false;
 
 FChunkedFixedUObjectArray* GCoreObjectArrayForDebugVisualizers = nullptr;
+uint32 GCoreNameEntryStrideForDebugVisualizers = 0;
+uint32 GCoreNameOffsetBitsForDebugVisualizers = 0;
+uint32 GCoreNameOffsetMaskForDebugVisualizers = 0;
+uint32 GCoreNameUnusedMaskForDebugVisualizers = 0;
+uint32 GCoreNameMaxLengthForDebugVisualizers = 0;
 
 namespace UE::CoreUObject::Private
 {
@@ -269,6 +274,11 @@ UE::CoreUObject::Private::FObjectHandlePackageDebugData* GCoreObjectHandlePackag
 #if PLATFORM_UNIX
 uint8** CORE_API GNameBlocksDebug = FNameDebugVisualizer::GetBlocks();
 FChunkedFixedUObjectArray*& CORE_API GObjectArrayForDebugVisualizers = GCoreObjectArrayForDebugVisualizers;
+uint32& CORE_API GNameEntryStrideForDebugVisualizers = GCoreNameEntryStrideForDebugVisualizers;
+uint32& CORE_API GNameOffsetBitsForDebugVisualizers = GCoreNameOffsetBitsForDebugVisualizers;
+uint32& CORE_API GNameOffsetMaskForDebugVisualizers = GCoreNameOffsetMaskForDebugVisualizers;
+uint32& CORE_API GNameUnusedMaskForDebugVisualizers = GCoreNameUnusedMaskForDebugVisualizers;
+uint32& CORE_API GNameMaxLengthForDebugVisualizers = GCoreNameMaxLengthForDebugVisualizers;
 UE::CoreUObject::Private::FStoredObjectPathDebug*& GComplexObjectPathDebug = GCoreComplexObjectPathDebug;
 UE::CoreUObject::Private::FObjectHandlePackageDebugData*& CORE_API GObjectHandlePackageDebug = GCoreObjectHandlePackageDebug;
 #endif
diff --git a/Engine/Source/Runtime/Core/Private/UObject/UnrealNames.cpp b/Engine/Source/Runtime/Core/Private/UObject/UnrealNames.cpp
--- a/Engine/Source/Runtime/Core/Private/UObject/UnrealNames.cpp
+++ b/Engine/Source/Runtime/Core/Private/UObject/UnrealNames.cpp
@@ -2206,12 +2206,22 @@ static FNamePool& GetNamePool()
 	{
 		Singleton = new (NamePoolData) FNamePool;
 		bNamePoolInitialized = true;
+		InitDebugVisualizer();
 		LLM(FLowLevelMemTracker::Get().FinishInitialise());
 	};
 	
 	return *Singleton;
 }
 
+void InitDebugVisualizer()
+{
+	GCoreNameEntryStrideForDebugVisualizers = FNameDebugVisualizer::EntryStride;
+	GCoreNameOffsetBitsForDebugVisualizers = FNameDebugVisualizer::OffsetBits;
+	GCoreNameOffsetMaskForDebugVisualizers = FNameDebugVisualizer::OffsetMask;
+	GCoreNameUnusedMaskForDebugVisualizers = FNameDebugVisualizer::UnusedMask;
+	GCoreNameMaxLengthForDebugVisualizers = FNameDebugVisualizer::MaxLength;
+}
+
 // Only call from functions guaranteed to run after FName lazy initialization
 static FNamePool& GetNamePoolPostInit()
 {
diff --git a/Engine/Source/Runtime/Core/Public/CoreGlobals.h b/Engine/Source/Runtime/Core/Public/CoreGlobals.h
--- a/Engine/Source/Runtime/Core/Public/CoreGlobals.h
+++ b/Engine/Source/Runtime/Core/Public/CoreGlobals.h
@@ -808,6 +808,11 @@ void CORE_API SetEmitDrawEvents(bool EmitDrawEvents);
 class FChunkedFixedUObjectArray;
 
 extern CORE_API FChunkedFixedUObjectArray* GCoreObjectArrayForDebugVisualizers;
+extern CORE_API uint32 GCoreNameEntryStrideForDebugVisualizers;
+extern CORE_API uint32 GCoreNameOffsetBitsForDebugVisualizers;
+extern CORE_API uint32 GCoreNameOffsetMaskForDebugVisualizers;
+extern CORE_API uint32 GCoreNameUnusedMaskForDebugVisualizers;
+extern CORE_API uint32 GCoreNameMaxLengthForDebugVisualizers;
 
 namespace UE::CoreUObject::Private
 {
diff --git a/Engine/Source/Runtime/Core/Public/Modules/Boilerplate/ModuleBoilerplate.h b/Engine/Source/Runtime/Core/Public/Modules/Boilerplate/ModuleBoilerplate.h
--- a/Engine/Source/Runtime/Core/Public/Modules/Boilerplate/ModuleBoilerplate.h
+++ b/Engine/Source/Runtime/Core/Public/Modules/Boilerplate/ModuleBoilerplate.h
@@ -10,8 +10,6 @@
 //   FVisualizerDebuggingState::Ptrs or FVisualizerDebuggingState::GuidString
 #include "Modules/VisualizerDebuggingState.h"
 
-class FChunkedFixedUObjectArray;
-
 // Boilerplate that is included once for each module, even in monolithic builds
 #if !defined(PER_MODULE_BOILERPLATE_ANYLINK)
 #define PER_MODULE_BOILERPLATE_ANYLINK(ModuleImplClass, ModuleName)
@@ -96,6 +94,11 @@ class FChunkedFixedUObjectArray;
 	#define UE_VISUALIZERS_HELPERS \
 		uint8** GNameBlocksDebug = FNameDebugVisualizer::GetBlocks(); \
 		FChunkedFixedUObjectArray*& GObjectArrayForDebugVisualizers = GCoreObjectArrayForDebugVisualizers; \
+		uint32& GNameEntryStrideForDebugVisualizers = GCoreNameEntryStrideForDebugVisualizers; \
+		uint32& GNameOffsetBitsForDebugVisualizers = GCoreNameOffsetBitsForDebugVisualizers; \
+		uint32& GNameOffsetMaskForDebugVisualizers = GCoreNameOffsetMaskForDebugVisualizers; \
+		uint32& GNameUnusedMaskForDebugVisualizers = GCoreNameUnusedMaskForDebugVisualizers; \
+		uint32& GNameMaxLengthForDebugVisualizers = GCoreNameMaxLengthForDebugVisualizers; \
 		UE::CoreUObject::Private::FStoredObjectPathDebug*& GComplexObjectPathDebug = GCoreComplexObjectPathDebug; \
 		UE::CoreUObject::Private::FObjectHandlePackageDebugData*& GObjectHandlePackageDebug = GCoreObjectHandlePackageDebug; \
 		UE::Core::FVisualizerDebuggingState*& GDebuggingState = GCoreDebuggingState;
diff --git a/Engine/Source/Runtime/Core/Public/UObject/NameTypes.h b/Engine/Source/Runtime/Core/Public/UObject/NameTypes.h
--- a/Engine/Source/Runtime/Core/Public/UObject/NameTypes.h
+++ b/Engine/Source/Runtime/Core/Public/UObject/NameTypes.h
@@ -1638,10 +1638,14 @@ struct FNameLexicalLess
 	}
 };
 
+CORE_API void InitDebugVisualizer();
+
 struct FNameDebugVisualizer
 {
 	CORE_API static uint8** GetBlocks();
 private:
+	friend CORE_API void InitDebugVisualizer();
+
 	static constexpr uint32 EntryStride = alignof(FNameEntry);
 	static constexpr uint32 OffsetBits = 16;
 	static constexpr uint32 BlockBits = 13;
