Fixed TOptional::operator== where it would compare objects with an intrusive unset state.  This would cause incorrect results in cases like TOptional<FString>() == TOptional<FString>(TEXT("")).

diff --git a/Engine/Source/Runtime/Core/Public/Misc/Optional.h b/Engine/Source/Runtime/Core/Public/Misc/Optional.h
--- a/Engine/Source/Runtime/Core/Public/Misc/Optional.h
+++ b/Engine/Source/Runtime/Core/Public/Misc/Optional.h
@@ -289,18 +289,16 @@ public:
 
 	[[nodiscard]] friend bool operator==(const TOptional& Lhs, const TOptional& Rhs)
 	{
-		if constexpr (!bUsingIntrusiveUnsetState)
+		bool bIsLhsSet = Lhs.IsSet();
+		bool bIsRhsSet = Rhs.IsSet();
+
+		if (bIsLhsSet != bIsRhsSet)
 		{
-			bool bIsLhsSet = Lhs.Value.bIsSet;
-			bool bIsRhsSet = Rhs.Value.bIsSet;
-			if (bIsLhsSet != bIsRhsSet)
-			{
-				return false;
-			}
-			if (!bIsLhsSet) // both unset
-			{
-				return true;
-			}
+			return false;
+		}
+		if (!bIsLhsSet) // both unset
+		{
+			return true;
 		}
 
 		return (*(const OptionalType*)&Lhs.Value) == (*(const OptionalType*)&Rhs.Value);
