[RaceDetector]
- Fix remaining cases of trying to access a context after a thread has released it

diff --git a/Engine/Source/Runtime/Core/Internal/Sanitizer/RaceDetectorTypes.h b/Engine/Source/Runtime/Core/Internal/Sanitizer/RaceDetectorTypes.h
--- a/Engine/Source/Runtime/Core/Internal/Sanitizer/RaceDetectorTypes.h
+++ b/Engine/Source/Runtime/Core/Internal/Sanitizer/RaceDetectorTypes.h
@@ -279,6 +279,11 @@ namespace UE::Sanitizer::RaceDetector {
 
 		// Lazily created on first instrumented access
 		TUniquePtr<FAccessHistory> AccessHistory;
+
+		INSTRUMENTATION_FUNCTION_ATTRIBUTES static bool IsValid(FContext* Context)
+		{
+			return reinterpret_cast<int64>(Context) > 0;
+		}
 	};
 
 	CORE_API extern FContext* GetThreadContext();
@@ -289,7 +294,7 @@ namespace UE::Sanitizer::RaceDetector {
 		INSTRUMENTATION_FUNCTION_ATTRIBUTES FInstrumentationScope()
 		{
 			FContext* Context = GetThreadContext();
-			if (Context)
+			if (FContext::IsValid(Context))
 			{
 				Context->WinInstrumentationDepth++;
 				bNeedDecrement = true;
diff --git a/Engine/Source/Runtime/Core/Private/Sanitizer/RaceDetector.cpp b/Engine/Source/Runtime/Core/Private/Sanitizer/RaceDetector.cpp
--- a/Engine/Source/Runtime/Core/Private/Sanitizer/RaceDetector.cpp
+++ b/Engine/Source/Runtime/Core/Private/Sanitizer/RaceDetector.cpp
@@ -1304,10 +1304,11 @@ namespace UE::Sanitizer::RaceDetector {
 
 	INSTRUMENTATION_FUNCTION_ATTRIBUTES void ReleaseCurrentContext()
 	{
-		if (GetThreadContext())
+		if (FContext* Context = GetThreadContext(); FContext::IsValid(Context))
 		{
-			ReleaseContext(GetThreadContext());
-			Platform::SetTlsValue(ContextTLSIndex, 0);
+			// Mark ContextTLS as invalid to prevent any further usage / detection for this thread
+			Platform::SetTlsValue(ContextTLSIndex, (void*)-1);
+			ReleaseContext(Context);
 		}
 	}
 
@@ -1425,7 +1426,7 @@ using namespace UE::Sanitizer::RaceDetector;
 // since we unhook the instrumentation instead.
 #define BEGIN_HANDLE_INSTRUMENTATION(FuncDoIfNotInstrumenting) \
 	FContext* ContextPtr = (FContext*)GetTlsValueFast(ContextTLSIndex); \
-	if (ContextPtr == nullptr || (ContextPtr->WinInstrumentationDepth | ContextPtr->InstrumentationDepth) != 0) \
+	if (!FContext::IsValid(ContextPtr) || (ContextPtr->WinInstrumentationDepth | ContextPtr->InstrumentationDepth) != 0) \
 	{ \
 		FuncDoIfNotInstrumenting; \
 	} \
diff --git a/Engine/Source/Runtime/Core/Private/Sanitizer/RaceDetectorWindows.cpp b/Engine/Source/Runtime/Core/Private/Sanitizer/RaceDetectorWindows.cpp
--- a/Engine/Source/Runtime/Core/Private/Sanitizer/RaceDetectorWindows.cpp
+++ b/Engine/Source/Runtime/Core/Private/Sanitizer/RaceDetectorWindows.cpp
@@ -35,7 +35,7 @@ namespace UE::Sanitizer::RaceDetector {
 #define START_WINAPI_INSTRUMENTATION(DoIfNotInstrumenting) \
 	DWORD OriginalError = GetLastError(); \
 	FContext* ContextPtr = GetThreadContext(); \
-	if (ContextPtr == nullptr || !ShouldInstrument(*ContextPtr)) \
+	if (!FContext::IsValid(ContextPtr) || !ShouldInstrument(*ContextPtr)) \
 	{ \
 		SetLastError(OriginalError); \
 		DoIfNotInstrumenting; \
@@ -87,9 +87,9 @@ struct FLastErrorPreservationScope
 INSTRUMENTATION_FUNCTION_ATTRIBUTES void InstrumentEndThread()
 {
 	HANDLE Handle = nullptr;
-	if (GetThreadContext())
+	if (FContext* ContextPtr = GetThreadContext(); FContext::IsValid(ContextPtr))
 	{
-		FInstrumentedStartThreadArgs* ThreadArgs = (FInstrumentedStartThreadArgs*)GetThreadContext()->ThreadArgs;
+		FInstrumentedStartThreadArgs* ThreadArgs = (FInstrumentedStartThreadArgs*)ContextPtr->ThreadArgs;
 		if (ThreadArgs)
 		{
 			// This can only be nullptr if this thread was so quick to execute
@@ -100,7 +100,7 @@ INSTRUMENTATION_FUNCTION_ATTRIBUTES void InstrumentEndThread()
 				FPlatformProcess::YieldThread();
 			}
 			delete ThreadArgs;
-			GetThreadContext()->ThreadArgs = nullptr;
+			ContextPtr->ThreadArgs = nullptr;
 		}
 	}
 
