Wrapped ClangPlatformMath in a __clang__ define to prevent VS from compiling it in non unity incremental builds that fail on clang intrinsics

diff --git a/Engine/Source/Runtime/Core/Public/Clang/ClangPlatformMath.h b/Engine/Source/Runtime/Core/Public/Clang/ClangPlatformMath.h
--- a/Engine/Source/Runtime/Core/Public/Clang/ClangPlatformMath.h
+++ b/Engine/Source/Runtime/Core/Public/Clang/ClangPlatformMath.h
@@ -5,6 +5,9 @@
 ==============================================================================================*/
 
 #pragma once
+
+#ifdef __clang__
+
 #include "GenericPlatform/GenericPlatformMath.h"
 
 // HEADER_UNIT_UNSUPPORTED - Clang not supporting header units
@@ -162,3 +165,5 @@ struct FClangPlatformMath : public FGenericPlatformMath
 		return !__builtin_mul_overflow(A, B, &OutResult);
 	}
 };
+
+#endif //~__clang__
\ No newline at end of file
