[Mac] Add process filename to error output when posix_spawn() fails.

#rnx

diff --git a/Engine/Source/Runtime/Core/Private/Mac/MacPlatformProcess.cpp b/Engine/Source/Runtime/Core/Private/Mac/MacPlatformProcess.cpp
--- a/Engine/Source/Runtime/Core/Private/Mac/MacPlatformProcess.cpp
+++ b/Engine/Source/Runtime/Core/Private/Mac/MacPlatformProcess.cpp
@@ -543,7 +543,7 @@ FProcHandle FMacPlatformProcess::CreateProcInternal(const TCHAR* URL, const TCHA
 
 	if (PosixSpawnErrNo != 0)
 	{
-		UE_LOG(LogHAL, Fatal, TEXT("FMacPlatformProcess::CreateProc: posix_spawn() failed (%d, %s)"), PosixSpawnErrNo, UTF8_TO_TCHAR(strerror(PosixSpawnErrNo)));
+		UE_LOG(LogHAL, Fatal, TEXT("FMacPlatformProcess::CreateProc: posix_spawn(\"%s\") failed (%d, %s)"), *ProcessPath, PosixSpawnErrNo, UTF8_TO_TCHAR(strerror(PosixSpawnErrNo)));
 		return FProcHandle();	// produce knowingly invalid handle if for some reason Fatal log (above) returns
 	}
 
