Fixed header unit compile/linker errors

diff --git a/Engine/Source/Runtime/Core/Public/Containers/HashTable.h b/Engine/Source/Runtime/Core/Public/Containers/HashTable.h
--- a/Engine/Source/Runtime/Core/Public/Containers/HashTable.h
+++ b/Engine/Source/Runtime/Core/Public/Containers/HashTable.h
@@ -20,7 +20,7 @@ class FMemoryUnfreezeContent;
 class FPointerTableBase;
 class FSHA1;
 
-static FORCEINLINE uint32 MurmurFinalize32(uint32 Hash)
+FORCEINLINE uint32 MurmurFinalize32(uint32 Hash)
 {
 	Hash ^= Hash >> 16;
 	Hash *= 0x85ebca6b;
@@ -30,7 +30,7 @@ static FORCEINLINE uint32 MurmurFinalize32(uint32 Hash)
 	return Hash;
 }
 
-static FORCEINLINE uint64 MurmurFinalize64(uint64 Hash)
+FORCEINLINE uint64 MurmurFinalize64(uint64 Hash)
 {
 	Hash ^= Hash >> 33;
 	Hash *= 0xff51afd7ed558ccdull;
@@ -40,7 +40,7 @@ static FORCEINLINE uint64 MurmurFinalize64(uint64 Hash)
 	return Hash;
 }
 
-static FORCEINLINE uint32 Murmur32( std::initializer_list< uint32 > InitList )
+FORCEINLINE uint32 Murmur32( std::initializer_list< uint32 > InitList )
 {
 	uint32 Hash = 0;
 	for( auto Element : InitList )
@@ -57,7 +57,7 @@ static FORCEINLINE uint32 Murmur32( std::initializer_list< uint32 > InitList )
 	return MurmurFinalize32( Hash );
 }
 
-static FORCEINLINE uint64 Murmur64( std::initializer_list< uint64 > InitList )
+FORCEINLINE uint64 Murmur64( std::initializer_list< uint64 > InitList )
 {
 	uint64 Hash = 0;
 	for( auto Element : InitList )
diff --git a/Engine/Source/Runtime/Core/Public/Math/Box.h b/Engine/Source/Runtime/Core/Public/Math/Box.h
--- a/Engine/Source/Runtime/Core/Public/Math/Box.h
+++ b/Engine/Source/Runtime/Core/Public/Math/Box.h
@@ -690,7 +690,17 @@ public:
 		return true;
 	}
 
-	bool SerializeFromMismatchedTag(FName StructTag, FArchive& Ar);
+	bool SerializeFromMismatchedTag(FName StructTag, FArchive& Ar)
+	{
+		if constexpr (std::is_same_v<T, float>)
+		{
+			return UE_SERIALIZE_VARIANT_FROM_MISMATCHED_TAG(Ar, Box, Box3f, Box3d);
+		}
+		else if constexpr (std::is_same_v<T, double>)
+		{
+			return UE_SERIALIZE_VARIANT_FROM_MISMATCHED_TAG(Ar, Box, Box3d, Box3f);
+		}
+	}
 };
 
 /* TBox<T> inline functions
@@ -987,18 +997,6 @@ template<> struct TIsUECoreVariant<FBox3f> { enum { Value = true }; };
 template<> struct TIsPODType<FBox3d> { enum { Value = true }; };
 template<> struct TIsUECoreVariant<FBox3d> { enum { Value = true }; };
 
-template<>
-inline bool FBox3f::SerializeFromMismatchedTag(FName StructTag, FArchive& Ar)
-{
-	return UE_SERIALIZE_VARIANT_FROM_MISMATCHED_TAG(Ar, Box, Box3f, Box3d);
-}
-
-template<>
-inline bool FBox3d::SerializeFromMismatchedTag(FName StructTag, FArchive& Ar)
-{
-	return UE_SERIALIZE_VARIANT_FROM_MISMATCHED_TAG(Ar, Box, Box3d, Box3f);
-}
-
 /* FMath inline functions
  *****************************************************************************/
 
diff --git a/Engine/Source/Runtime/Core/Public/Math/Box2D.h b/Engine/Source/Runtime/Core/Public/Math/Box2D.h
--- a/Engine/Source/Runtime/Core/Public/Math/Box2D.h
+++ b/Engine/Source/Runtime/Core/Public/Math/Box2D.h
@@ -447,7 +447,17 @@ public:
 		return true;
 	}
 
-	bool SerializeFromMismatchedTag(FName StructTag, FArchive& Ar);
+	bool SerializeFromMismatchedTag(FName StructTag, FArchive& Ar)
+	{
+		if constexpr (std::is_same_v<T, float>)
+		{
+			return UE_SERIALIZE_VARIANT_FROM_MISMATCHED_TAG(Ar, Box2D, Box2f, Box2d);
+		}
+		else if constexpr (std::is_same_v<T, double>)
+		{
+			return UE_SERIALIZE_VARIANT_FROM_MISMATCHED_TAG(Ar, Box2D, Box2d, Box2f);
+		}
+	}
 
 	// Conversion from other type. 
 	template<typename FArg UE_REQUIRES(!std::is_same_v<T, FArg>)>
@@ -613,16 +623,3 @@ template<> struct TIsUECoreVariant<FBox2f> { enum { Value = true }; };
 //template<> struct TCanBulkSerialize<FBox2d> { enum { Value = false }; };	// LWC_TODO: This can be done (via versioning) once LWC is fixed to on.
 template<> struct TIsPODType<FBox2d> { enum { Value = true }; };
 template<> struct TIsUECoreVariant<FBox2d> { enum { Value = true }; };
-
-
-template<>
-inline bool FBox2f::SerializeFromMismatchedTag(FName StructTag, FArchive& Ar)
-{
-	return UE_SERIALIZE_VARIANT_FROM_MISMATCHED_TAG(Ar, Box2D, Box2f, Box2d);
-}
-
-template<>
-inline bool FBox2d::SerializeFromMismatchedTag(FName StructTag, FArchive& Ar)
-{
-	return UE_SERIALIZE_VARIANT_FROM_MISMATCHED_TAG(Ar, Box2D, Box2d, Box2f);
-}
diff --git a/Engine/Source/Runtime/Core/Public/Math/BoxSphereBounds.h b/Engine/Source/Runtime/Core/Public/Math/BoxSphereBounds.h
--- a/Engine/Source/Runtime/Core/Public/Math/BoxSphereBounds.h
+++ b/Engine/Source/Runtime/Core/Public/Math/BoxSphereBounds.h
@@ -134,7 +134,19 @@ public:
 public:
 
 	bool Serialize(FArchive &Ar);
-	bool SerializeFromMismatchedTag(FName StructTag, FArchive &Ar);
+	bool SerializeFromMismatchedTag(FName StructTag, FArchive &Ar)
+	{
+		if constexpr (std::is_same_v<T, float>)
+		{
+			// Falls back to UseSerializeItem to convert per property.
+			return false;
+		}
+		else if constexpr (std::is_same_v<T, double>)
+		{
+			// Falls back to UseSerializeItem to convert per property.
+			return false;
+		}
+	}
 
 	/**
 	 * Constructs a bounding volume containing both this and B.
@@ -553,17 +565,3 @@ template <> struct TIsPODType<FBoxSphereBounds3d> { enum { Value = true }; };
 template <> struct TIsUECoreVariant<FBoxSphereBounds3f> { enum { Value = true }; };
 template <> struct TIsUECoreVariant<FBoxSphereBounds3d> { enum { Value = true }; };
 template <> struct TIsUECoreVariant<FCompactBoxSphereBounds3d> { enum { Value = true }; };
-
-template<>
-inline bool FBoxSphereBounds3f::SerializeFromMismatchedTag(FName StructTag, FArchive& Ar)
-{
-	// Falls back to UseSerializeItem to convert per property.
-	return false;
-}
-
-template<>
-inline bool FBoxSphereBounds3d::SerializeFromMismatchedTag(FName StructTag, FArchive& Ar)
-{
-	// Falls back to UseSerializeItem to convert per property.
-	return false;
-}
diff --git a/Engine/Source/Runtime/Core/Public/Math/IntPoint.h b/Engine/Source/Runtime/Core/Public/Math/IntPoint.h
--- a/Engine/Source/Runtime/Core/Public/Math/IntPoint.h
+++ b/Engine/Source/Runtime/Core/Public/Math/IntPoint.h
@@ -500,7 +500,25 @@ PRAGMA_ENABLE_DEPRECATION_WARNINGS
 		return true;
 	}
 
-	bool SerializeFromMismatchedTag(FName StructTag, FArchive& Ar);
+	bool SerializeFromMismatchedTag(FName StructTag, FArchive& Ar)
+	{
+		if constexpr (std::is_same_v<IntType, int32>)
+		{
+			return UE_SERIALIZE_VARIANT_FROM_MISMATCHED_TAG(Ar, IntPoint, Int32Point, Int64Point);
+		}
+		else if constexpr (std::is_same_v<IntType, int64>)
+		{
+			return UE_SERIALIZE_VARIANT_FROM_MISMATCHED_TAG(Ar, IntPoint, Int64Point, Int32Point);
+		}
+		else if constexpr (std::is_same_v<IntType, uint32>)
+		{
+			return UE_SERIALIZE_VARIANT_FROM_MISMATCHED_TAG(Ar, UintPoint, Uint32Point, Uint64Point);
+		}
+		else if constexpr (std::is_same_v<IntType, uint64>)
+		{
+			return UE_SERIALIZE_VARIANT_FROM_MISMATCHED_TAG(Ar, UintPoint, Uint64Point, Uint32Point);
+		}
+	}
 };
 
 template <typename IntType>
@@ -576,28 +594,3 @@ template <> struct TIsPODType<FUint64Point> { enum { Value = true }; };
 
 template<> struct TIsUECoreVariant<FInt64Point> { enum { Value = true }; };
 template<> struct TIsUECoreVariant<FUint64Point> { enum { Value = true }; };
-
-
-template<>
-inline bool FInt32Point::SerializeFromMismatchedTag(FName StructTag, FArchive& Ar)
-{
-	return UE_SERIALIZE_VARIANT_FROM_MISMATCHED_TAG(Ar, IntPoint, Int32Point, Int64Point);
-}
-
-template<>
-inline bool FInt64Point::SerializeFromMismatchedTag(FName StructTag, FArchive& Ar)
-{
-	return UE_SERIALIZE_VARIANT_FROM_MISMATCHED_TAG(Ar, IntPoint, Int64Point, Int32Point);
-}
-
-template<>
-inline bool FUint32Point::SerializeFromMismatchedTag(FName StructTag, FArchive& Ar)
-{
-	return UE_SERIALIZE_VARIANT_FROM_MISMATCHED_TAG(Ar, UintPoint, Uint32Point, Uint64Point);
-}
-
-template<>
-inline bool FUint64Point::SerializeFromMismatchedTag(FName StructTag, FArchive& Ar)
-{
-	return UE_SERIALIZE_VARIANT_FROM_MISMATCHED_TAG(Ar, UintPoint, Uint64Point, Uint32Point);
-}
diff --git a/Engine/Source/Runtime/Core/Public/Math/IntVector.h b/Engine/Source/Runtime/Core/Public/Math/IntVector.h
--- a/Engine/Source/Runtime/Core/Public/Math/IntVector.h
+++ b/Engine/Source/Runtime/Core/Public/Math/IntVector.h
@@ -631,7 +631,25 @@ PRAGMA_ENABLE_DEPRECATION_WARNINGS
 		return true;
 	}
 
-	bool SerializeFromMismatchedTag(FName StructTag, FArchive& Ar);
+	bool SerializeFromMismatchedTag(FName StructTag, FArchive& Ar)
+	{
+		if constexpr (std::is_same_v<IntType, int32>)
+		{
+			return UE_SERIALIZE_VARIANT_FROM_MISMATCHED_TAG(Ar, IntVector, Int32Vector, Int64Vector);
+		}
+		else if constexpr (std::is_same_v<IntType, int64>)
+		{
+			return UE_SERIALIZE_VARIANT_FROM_MISMATCHED_TAG(Ar, IntVector, Int64Vector, Int32Vector);
+		}
+		else if constexpr (std::is_same_v<IntType, uint32>)
+		{
+			return UE_SERIALIZE_VARIANT_FROM_MISMATCHED_TAG(Ar, UintVector, Uint32Vector, Uint64Vector);
+		}
+		else if constexpr (std::is_same_v<IntType, uint64>)
+		{
+			return UE_SERIALIZE_VARIANT_FROM_MISMATCHED_TAG(Ar, UintVector, Uint64Vector, Uint32Vector);
+		}
+	}
 };
 
 template <typename IntType>
@@ -1154,7 +1172,25 @@ PRAGMA_ENABLE_DEPRECATION_WARNINGS
 		return true;
 	}
 
-	bool SerializeFromMismatchedTag(FName StructTag, FArchive& Ar);
+	bool SerializeFromMismatchedTag(FName StructTag, FArchive& Ar)
+	{
+		if constexpr (std::is_same_v<IntType, int32>)
+		{
+			return UE_SERIALIZE_VARIANT_FROM_MISMATCHED_TAG(Ar, IntVector2, Int32Vector2, Int64Vector2);
+		}
+		else if constexpr (std::is_same_v<IntType, int64>)
+		{
+			return UE_SERIALIZE_VARIANT_FROM_MISMATCHED_TAG(Ar, IntVector2, Int64Vector2, Int32Vector2);
+		}
+		else if constexpr (std::is_same_v<IntType, uint32>)
+		{
+			return UE_SERIALIZE_VARIANT_FROM_MISMATCHED_TAG(Ar, UintVector2, Uint32Vector2, Uint64Vector2);
+		}
+		else if constexpr (std::is_same_v<IntType, uint64>)
+		{
+			return UE_SERIALIZE_VARIANT_FROM_MISMATCHED_TAG(Ar, UintVector2, Uint64Vector2, Uint32Vector2);
+		}
+	}
 };
 
 template <typename IntType>
@@ -1763,7 +1799,25 @@ PRAGMA_ENABLE_DEPRECATION_WARNINGS
 		return true;
 	}
 
-	bool SerializeFromMismatchedTag(FName StructTag, FArchive& Ar);
+	bool SerializeFromMismatchedTag(FName StructTag, FArchive& Ar)
+	{
+		if constexpr (std::is_same_v<IntType, int32>)
+		{
+			return UE_SERIALIZE_VARIANT_FROM_MISMATCHED_TAG(Ar, IntVector4, Int32Vector4, Int64Vector4);
+		}
+		else if constexpr (std::is_same_v<IntType, int64>)
+		{
+			return UE_SERIALIZE_VARIANT_FROM_MISMATCHED_TAG(Ar, IntVector4, Int64Vector4, Int32Vector4);
+		}
+		else if constexpr (std::is_same_v<IntType, uint32>)
+		{
+			return UE_SERIALIZE_VARIANT_FROM_MISMATCHED_TAG(Ar, UintVector4, Uint32Vector4, Uint64Vector4);
+		}
+		else if constexpr (std::is_same_v<IntType, uint64>)
+		{
+			return UE_SERIALIZE_VARIANT_FROM_MISMATCHED_TAG(Ar, UintVector4, Uint64Vector4, Uint32Vector4);
+		}
+	}
 };
 
 template <typename IntType>
@@ -1840,76 +1894,3 @@ template <> struct TIsUECoreVariant<FInt64Vector3> { enum { Value = true }; };
 template <> struct TIsUECoreVariant<FUint64Vector3> { enum { Value = true }; };
 template <> struct TIsUECoreVariant<FInt64Vector4> { enum { Value = true }; };
 template <> struct TIsUECoreVariant<FUint64Vector4> { enum { Value = true }; };
-
-
-template<>
-inline bool FInt32Vector2::SerializeFromMismatchedTag(FName StructTag, FArchive& Ar)
-{
-	return UE_SERIALIZE_VARIANT_FROM_MISMATCHED_TAG(Ar, IntVector2, Int32Vector2, Int64Vector2);
-}
-
-template<>
-inline bool FInt64Vector2::SerializeFromMismatchedTag(FName StructTag, FArchive& Ar)
-{
-	return UE_SERIALIZE_VARIANT_FROM_MISMATCHED_TAG(Ar, IntVector2, Int64Vector2, Int32Vector2);
-}
-
-template<>
-inline bool FUint32Vector2::SerializeFromMismatchedTag(FName StructTag, FArchive& Ar)
-{
-	return UE_SERIALIZE_VARIANT_FROM_MISMATCHED_TAG(Ar, UintVector2, Uint32Vector2, Uint64Vector2);
-}
-
-template<>
-inline bool FUint64Vector2::SerializeFromMismatchedTag(FName StructTag, FArchive& Ar)
-{
-	return UE_SERIALIZE_VARIANT_FROM_MISMATCHED_TAG(Ar, UintVector2, Uint64Vector2, Uint32Vector2);
-}
-
-template<>
-inline bool FInt32Vector3::SerializeFromMismatchedTag(FName StructTag, FArchive& Ar)
-{
-	return UE_SERIALIZE_VARIANT_FROM_MISMATCHED_TAG(Ar, IntVector, Int32Vector, Int64Vector);
-}
-
-template<>
-inline bool FInt64Vector3::SerializeFromMismatchedTag(FName StructTag, FArchive& Ar)
-{
-	return UE_SERIALIZE_VARIANT_FROM_MISMATCHED_TAG(Ar, IntVector, Int64Vector, Int32Vector);
-}
-
-template<>
-inline bool FUint32Vector3::SerializeFromMismatchedTag(FName StructTag, FArchive& Ar)
-{
-	return UE_SERIALIZE_VARIANT_FROM_MISMATCHED_TAG(Ar, UintVector, Uint32Vector, Uint64Vector);
-}
-
-template<>
-inline bool FUint64Vector3::SerializeFromMismatchedTag(FName StructTag, FArchive& Ar)
-{
-	return UE_SERIALIZE_VARIANT_FROM_MISMATCHED_TAG(Ar, UintVector, Uint64Vector, Uint32Vector);
-}
-
-template<>
-inline bool FInt32Vector4::SerializeFromMismatchedTag(FName StructTag, FArchive& Ar)
-{
-	return UE_SERIALIZE_VARIANT_FROM_MISMATCHED_TAG(Ar, IntVector4, Int32Vector4, Int64Vector4);
-}
-
-template<>
-inline bool FInt64Vector4::SerializeFromMismatchedTag(FName StructTag, FArchive& Ar)
-{
-	return UE_SERIALIZE_VARIANT_FROM_MISMATCHED_TAG(Ar, IntVector4, Int64Vector4, Int32Vector4);
-}
-
-template<>
-inline bool FUint32Vector4::SerializeFromMismatchedTag(FName StructTag, FArchive& Ar)
-{
-	return UE_SERIALIZE_VARIANT_FROM_MISMATCHED_TAG(Ar, UintVector4, Uint32Vector4, Uint64Vector4);
-}
-
-template<>
-inline bool FUint64Vector4::SerializeFromMismatchedTag(FName StructTag, FArchive& Ar)
-{
-	return UE_SERIALIZE_VARIANT_FROM_MISMATCHED_TAG(Ar, UintVector4, Uint64Vector4, Uint32Vector4);
-}
diff --git a/Engine/Source/Runtime/Core/Public/Math/Plane.h b/Engine/Source/Runtime/Core/Public/Math/Plane.h
--- a/Engine/Source/Runtime/Core/Public/Math/Plane.h
+++ b/Engine/Source/Runtime/Core/Public/Math/Plane.h
@@ -309,7 +309,17 @@ public:
 		//return false;
 	}
 
-	bool SerializeFromMismatchedTag(FName StructTag, FArchive& Ar);
+	bool SerializeFromMismatchedTag(FName StructTag, FArchive& Ar)
+	{
+		if constexpr (std::is_same_v<T, float>)
+		{
+			return UE_SERIALIZE_VARIANT_FROM_MISMATCHED_TAG(Ar, Plane, Plane4f, Plane4d);
+		}
+		else if constexpr (std::is_same_v<T, double>)
+		{
+			return UE_SERIALIZE_VARIANT_FROM_MISMATCHED_TAG(Ar, Plane, Plane4d, Plane4f);
+		}
+	}
 
 	/**
 	 * Serializes the vector compressed for e.g. network transmission.
@@ -710,19 +720,6 @@ template<> struct TIsPODType<FPlane4d> { enum { Value = true }; };
 template<> struct TIsUECoreVariant<FPlane4d> { enum { Value = true }; };
 
 
-template<>
-inline bool FPlane4f::SerializeFromMismatchedTag(FName StructTag, FArchive& Ar)
-{
-	return UE_SERIALIZE_VARIANT_FROM_MISMATCHED_TAG(Ar, Plane, Plane4f, Plane4d);
-}
-
-template<>
-inline bool FPlane4d::SerializeFromMismatchedTag(FName StructTag, FArchive& Ar)
-{
-	return UE_SERIALIZE_VARIANT_FROM_MISMATCHED_TAG(Ar, Plane, Plane4d, Plane4f);
-}
-
-
 #ifdef _MSC_VER
 #pragma warning (pop)
 #endif
diff --git a/Engine/Source/Runtime/Core/Public/Math/Ray.h b/Engine/Source/Runtime/Core/Public/Math/Ray.h
--- a/Engine/Source/Runtime/Core/Public/Math/Ray.h
+++ b/Engine/Source/Runtime/Core/Public/Math/Ray.h
@@ -201,7 +201,17 @@ public:
 		return true;
 	}
 
-	bool SerializeFromMismatchedTag(FName StructTag, FArchive& Ar);
+	bool SerializeFromMismatchedTag(FName StructTag, FArchive& Ar)
+	{
+		if constexpr (std::is_same_v<T, float>)
+		{
+			return UE_SERIALIZE_VARIANT_FROM_MISMATCHED_TAG(Ar, Ray, Ray3f, Ray3d);
+		}
+		else if constexpr (std::is_same_v<T, double>)
+		{
+			return UE_SERIALIZE_VARIANT_FROM_MISMATCHED_TAG(Ar, Ray, Ray3d, Ray3f);
+		}
+	}
 
 	// Conversion to other type.
 	template<typename FArg UE_REQUIRES(!std::is_same_v<T, FArg>)>
@@ -222,15 +232,3 @@ template<> struct TIsUECoreVariant<FRay3f> { enum { Value = true }; };
 template<> struct TIsUECoreVariant<FRay3d> { enum { Value = true }; };
 template<> struct TCanBulkSerialize<FRay3f> { enum { Value = true }; };
 template<> struct TCanBulkSerialize<FRay3d> { enum { Value = true }; };
-
-template<>
-inline bool FRay3f::SerializeFromMismatchedTag(FName StructTag, FArchive& Ar)
-{
-	return UE_SERIALIZE_VARIANT_FROM_MISMATCHED_TAG(Ar, Ray, Ray3f, Ray3d);
-}
-
-template<>
-inline bool FRay3d::SerializeFromMismatchedTag(FName StructTag, FArchive& Ar)
-{
-	return UE_SERIALIZE_VARIANT_FROM_MISMATCHED_TAG(Ar, Ray, Ray3d, Ray3f);
-}
diff --git a/Engine/Source/Runtime/Core/Public/Math/Sphere.h b/Engine/Source/Runtime/Core/Public/Math/Sphere.h
--- a/Engine/Source/Runtime/Core/Public/Math/Sphere.h
+++ b/Engine/Source/Runtime/Core/Public/Math/Sphere.h
@@ -231,7 +231,17 @@ public:
 		return true;
 	}
 
-	bool SerializeFromMismatchedTag(FName StructTag, FArchive& Ar);
+	bool SerializeFromMismatchedTag(FName StructTag, FArchive& Ar)
+	{
+		if constexpr (std::is_same_v<T, float>)
+		{
+			return UE_SERIALIZE_VARIANT_FROM_MISMATCHED_TAG(Ar, Sphere, Sphere3f, Sphere3d);
+		}
+		else if constexpr (std::is_same_v<T, double>)
+		{
+			return UE_SERIALIZE_VARIANT_FROM_MISMATCHED_TAG(Ar, Sphere, Sphere3d, Sphere3f);
+		}
+	}
 
 };
 
@@ -369,19 +379,6 @@ template<> struct TCanBulkSerialize<FSphere3d> { enum { Value = true }; };
 template<> struct TIsPODType<FSphere3d> { enum { Value = true }; };
 template<> struct TIsUECoreVariant<FSphere3d> { enum { Value = true }; };
 
-template<>
-inline bool FSphere3f::SerializeFromMismatchedTag(FName StructTag, FArchive& Ar)
-{
-	return UE_SERIALIZE_VARIANT_FROM_MISMATCHED_TAG(Ar, Sphere, Sphere3f, Sphere3d);
-}
-
-template<>
-inline bool FSphere3d::SerializeFromMismatchedTag(FName StructTag, FArchive& Ar)
-{
-	return UE_SERIALIZE_VARIANT_FROM_MISMATCHED_TAG(Ar, Sphere, Sphere3d, Sphere3f);
-}
-
-
 /* FMath inline functions
  *****************************************************************************/
 
diff --git a/Engine/Source/Runtime/Core/Public/Math/Vector2D.h b/Engine/Source/Runtime/Core/Public/Math/Vector2D.h
--- a/Engine/Source/Runtime/Core/Public/Math/Vector2D.h
+++ b/Engine/Source/Runtime/Core/Public/Math/Vector2D.h
@@ -688,7 +688,17 @@ public:
 		return true;
 	}
 
-	bool SerializeFromMismatchedTag(FName StructTag, FArchive& Ar);
+	bool SerializeFromMismatchedTag(FName StructTag, FArchive& Ar)
+	{
+		if constexpr (std::is_same_v<T, float>)
+		{
+			return UE_SERIALIZE_VARIANT_FROM_MISMATCHED_TAG(Ar, Vector2D, Vector2f, Vector2d);
+		}
+		else
+		{
+			return UE_SERIALIZE_VARIANT_FROM_MISMATCHED_TAG(Ar, Vector2D, Vector2d, Vector2f);
+		}
+	}
 
 #if ENABLE_NAN_DIAGNOSTIC
 	FORCEINLINE void DiagnosticCheckNaN()
@@ -1333,20 +1343,6 @@ template <> struct TIsPODType<FVector2d> { enum { Value = true }; };
 template <> struct TIsUECoreVariant<FVector2d> { enum { Value = true }; };
 template<> struct TCanBulkSerialize<FVector2d> { enum { Value = true }; };
 
-template<>
-inline bool FVector2f::SerializeFromMismatchedTag(FName StructTag, FArchive& Ar)
-{
-	return UE_SERIALIZE_VARIANT_FROM_MISMATCHED_TAG(Ar, Vector2D, Vector2f, Vector2d);
-}
-
-template<>
-inline bool FVector2d::SerializeFromMismatchedTag(FName StructTag, FArchive& Ar)
-{
-	return UE_SERIALIZE_VARIANT_FROM_MISMATCHED_TAG(Ar, Vector2D, Vector2d, Vector2f);
-}
-
-
-
 
 #ifdef _MSC_VER
 #pragma warning (pop)
diff --git a/Engine/Source/Runtime/Core/Public/PixelFormat.h b/Engine/Source/Runtime/Core/Public/PixelFormat.h
--- a/Engine/Source/Runtime/Core/Public/PixelFormat.h
+++ b/Engine/Source/Runtime/Core/Public/PixelFormat.h
@@ -309,7 +309,7 @@ FORCEINLINE bool IsASTCBlockCompressedTextureFormat(EPixelFormat PixelFormat)
 	PRAGMA_RESTORE_SWITCH_UNHANDLED_ENUM_CASE_WARNINGS;
 }
 
-static inline bool IsBlockCompressedFormat(EPixelFormat Format)
+inline bool IsBlockCompressedFormat(EPixelFormat Format)
 {
 	PRAGMA_DISABLE_SWITCH_UNHANDLED_ENUM_CASE_WARNINGS;
 	switch (Format)
@@ -369,7 +369,7 @@ FORCEINLINE bool IsInteger(EPixelFormat PixelFormat)
 }
 
 // IsFloatFormat does not include all floating point formats; see also IsHDR
-static bool IsFloatFormat(EPixelFormat Format)
+inline bool IsFloatFormat(EPixelFormat Format)
 {
 	PRAGMA_DISABLE_SWITCH_UNHANDLED_ENUM_CASE_WARNINGS;
 	switch (Format)
diff --git a/Engine/Source/Runtime/CoreUObject/Private/UObject/PropertyHelper.h b/Engine/Source/Runtime/CoreUObject/Private/UObject/PropertyHelper.h
--- a/Engine/Source/Runtime/CoreUObject/Private/UObject/PropertyHelper.h
+++ b/Engine/Source/Runtime/CoreUObject/Private/UObject/PropertyHelper.h
@@ -99,7 +99,7 @@ T* FindObjectByTypePath(FPropertyTypeName TypePath)
 }
 
 #if WITH_EDITORONLY_DATA
-inline static const FName NAME_OriginalType(ANSITEXTVIEW("OriginalType"));
+inline const FName NAME_OriginalType(ANSITEXTVIEW("OriginalType"));
 const FString* FindOriginalTypeName(const UField* Field);
 const FString* FindOriginalTypeName(const FProperty* Property);
 FPropertyTypeName FindOriginalType(const UField* Field);
