Disabled some code when compiling with header units

diff --git a/Engine/Source/Runtime/Core/Private/Containers/ContainerAllocationPolicies.cpp b/Engine/Source/Runtime/Core/Private/Containers/ContainerAllocationPolicies.cpp
--- a/Engine/Source/Runtime/Core/Private/Containers/ContainerAllocationPolicies.cpp
+++ b/Engine/Source/Runtime/Core/Private/Containers/ContainerAllocationPolicies.cpp
@@ -15,7 +15,7 @@ FORCENOINLINE void UE::Core::Private::OnInvalidSizedHeapAllocatorNum(int32 Index
 	for (;;);
 }
 
-#if !IS_MERGEDMODULES
+#if !IS_MERGEDMODULES && !defined(UE_HEADER_UNITS)
 template void TSizedHeapAllocator<32, FMemory>::ForAnyElementType::ResizeAllocation(SizeType CurrentNum, SizeType NewMax, SIZE_T NumBytesPerElement);
 template void TSizedHeapAllocator<32, FMemory>::ForAnyElementType::ResizeAllocation(SizeType CurrentNum, SizeType NewMax, SIZE_T NumBytesPerElement, uint32 AlignmentOfElement);
 template void TSizedHeapAllocator<64, FMemory>::ForAnyElementType::ResizeAllocation(SizeType CurrentNum, SizeType NewMax, SIZE_T NumBytesPerElement);
diff --git a/Engine/Source/Runtime/Core/Public/Containers/ContainerAllocationPolicies.h b/Engine/Source/Runtime/Core/Public/Containers/ContainerAllocationPolicies.h
--- a/Engine/Source/Runtime/Core/Public/Containers/ContainerAllocationPolicies.h
+++ b/Engine/Source/Runtime/Core/Public/Containers/ContainerAllocationPolicies.h
@@ -824,7 +824,7 @@ public:
 };
 
 // Define the ResizeAllocation functions with the regular allocator as exported to avoid bloat
-#if !IS_MERGEDMODULES
+#if !IS_MERGEDMODULES && !defined(UE_HEADER_UNITS)
 extern template CORE_API FORCENOINLINE void TSizedHeapAllocator<32, FMemory>::ForAnyElementType::ResizeAllocation(SizeType CurrentNum, SizeType NewMax, SIZE_T NumBytesPerElement);
 extern template CORE_API FORCENOINLINE void TSizedHeapAllocator<32, FMemory>::ForAnyElementType::ResizeAllocation(SizeType CurrentNum, SizeType NewMax, SIZE_T NumBytesPerElement, uint32 AlignmentOfElement);
 extern template CORE_API FORCENOINLINE void TSizedHeapAllocator<64, FMemory>::ForAnyElementType::ResizeAllocation(SizeType CurrentNum, SizeType NewMax, SIZE_T NumBytesPerElement);
