Temporarily disable Nanite skel mesh fallback generation until recompute normal/tangents + some stretching verts are resolved. All other optimizations are still in effect.

diff --git a/Engine/Source/Developer/MeshBuilder/Private/SkeletalMeshBuilder.cpp b/Engine/Source/Developer/MeshBuilder/Private/SkeletalMeshBuilder.cpp
--- a/Engine/Source/Developer/MeshBuilder/Private/SkeletalMeshBuilder.cpp
+++ b/Engine/Source/Developer/MeshBuilder/Private/SkeletalMeshBuilder.cpp
@@ -606,7 +606,9 @@ bool FSkeletalMeshBuilder::Build(FSkeletalMeshRenderData& OutRenderData, const F
 	const bool bHasClothing = SkeletalMeshBuildParameters.SkeletalMesh->GetMeshClothingAssets().Num() > 0;
 	const bool bHasMorphTargets = MeshAttributes.GetMorphTargetNames().Num() > 0;
 	const bool bBuildNanite = SkeletalMeshBuildParameters.SkeletalMesh->IsNaniteEnabled() && SkeletalMeshBuildParameters.LODIndex == 0;
-	const bool bBuildNaniteFallback = bBuildNanite && !bHasClothing && !bHasMorphTargets;
+
+	// TODO: Some issues to work out with missing triangles, and corrupt TSB if recompute normals/tangents is enabled.
+	const bool bBuildNaniteFallback = false;// bBuildNanite && !bHasClothing && !bHasMorphTargets;
 
 	const float TaskTotal = 
 		5.01f /*FinalizeContext*/ + 
diff --git a/Engine/Source/Runtime/Core/Private/UObject/UE5MainStreamObjectVersion.cpp b/Engine/Source/Runtime/Core/Private/UObject/UE5MainStreamObjectVersion.cpp
--- a/Engine/Source/Runtime/Core/Private/UObject/UE5MainStreamObjectVersion.cpp
+++ b/Engine/Source/Runtime/Core/Private/UObject/UE5MainStreamObjectVersion.cpp
@@ -15,7 +15,7 @@ TMap<FGuid, FGuid> FUE5MainStreamObjectVersion::GetSystemGuids()
 	SystemGuids.Add(DevGuids.NANITE_DERIVEDDATA_VER, FGuid("49D2EF76967340EFB22F2F398E4DDDDD"));
 	SystemGuids.Add(DevGuids.NIAGARASHADERMAP_DERIVEDDATA_VER, FGuid("8A37C45D24F2423CBE5F8F371DE33575"));
 	SystemGuids.Add(DevGuids.Niagara_LatestScriptCompileVersion, FGuid("9C8812406D824DE29BBC6F303964BB7E"));
-	SystemGuids.Add(DevGuids.SkeletalMeshDerivedDataVersion, FGuid("B28EFD64577949B2083DCECC6BA3B9AC"));
+	SystemGuids.Add(DevGuids.SkeletalMeshDerivedDataVersion, FGuid("CFE0AF3B4F48465E8DE6A6BD81881F18"));
 	SystemGuids.Add(DevGuids.STATICMESH_DERIVEDDATA_VER, FGuid("98F7D79A5811013825E75BBCC41ED3E9"));
 	SystemGuids.Add(DevGuids.POSESEARCHDB_DERIVEDDATA_VER, FGuid("4E595C2AC5E947D6BA9ABC874353E5BC"));
 	SystemGuids.Add(DevGuids.GROOM_BINDING_DERIVED_DATA_VERSION, FGuid("156678D4F9084D7CAE98FADC6AB93573"));
