[Insights] Fix deprecation version for Bookmarks public API

#rnx

diff --git a/Engine/Source/Developer/TraceServices/Public/TraceServices/Model/Bookmarks.h b/Engine/Source/Developer/TraceServices/Public/TraceServices/Model/Bookmarks.h
--- a/Engine/Source/Developer/TraceServices/Public/TraceServices/Model/Bookmarks.h
+++ b/Engine/Source/Developer/TraceServices/Public/TraceServices/Model/Bookmarks.h
@@ -64,9 +64,9 @@ public:
 	*/
 	virtual void AppendBookmark(uint64 BookmarkPoint, double Time, uint32 CallstackId, const TCHAR* Text) = 0;
 
-	UE_DEPRECATED(5.7, "Use AppendBookmark that has CallstackId parameter")
+	UE_DEPRECATED(5.6, "Use AppendBookmark that has CallstackId parameter")
 	void AppendBookmark(uint64 BookmarkPoint, double Time, const uint8* FormatArgs) { AppendBookmark(BookmarkPoint, Time, 0, FormatArgs); }
-	UE_DEPRECATED(5.7, "Use AppendBookmark that has CallstackId parameter")
+	UE_DEPRECATED(5.6, "Use AppendBookmark that has CallstackId parameter")
 	void AppendBookmark(uint64 BookmarkPoint, double Time, const TCHAR* Text) { AppendBookmark(BookmarkPoint, Time, 0, Text); }
 };
 
diff --git a/Engine/Source/Runtime/Core/Public/ProfilingDebugging/MiscTrace.h b/Engine/Source/Runtime/Core/Public/ProfilingDebugging/MiscTrace.h
--- a/Engine/Source/Runtime/Core/Public/ProfilingDebugging/MiscTrace.h
+++ b/Engine/Source/Runtime/Core/Public/ProfilingDebugging/MiscTrace.h
@@ -91,14 +91,14 @@ struct FMiscTrace
 	CORE_API static bool ShouldTraceRegion();
 
 	template <typename... Types>
-	UE_DEPRECATED(5.7, "Use OutputBookmark that has CallstackId parameter")
+	UE_DEPRECATED(5.6, "Use OutputBookmark that has CallstackId parameter")
 	static void OutputBookmark(const void* BookmarkPoint, Types... FormatArgs)
 	{
 		OutputBookmark(CallstackTrace_GetCurrentId(), BookmarkPoint, Forward<Types>(FormatArgs)...);
 	}
 
 	template <typename... Types>
-	UE_DEPRECATED(5.7, "Use OutputBookmarkCycles that has CallstackId parameter")
+	UE_DEPRECATED(5.6, "Use OutputBookmarkCycles that has CallstackId parameter")
 	static void OutputBookmarkCycles(uint64 Cycles, const void* BookmarkPoint, Types... FormatArgs)
 	{
 		OutputBookmarkCycles(CallstackTrace_GetCurrentId(), Cycles, BookmarkPoint, Forward<Types>(FormatArgs)...);
