Fixed static analysis build warning related to SerializeFromMismatchedTag not handling non-implemented branch

diff --git a/Engine/Source/Runtime/Core/Public/Math/Box.h b/Engine/Source/Runtime/Core/Public/Math/Box.h
--- a/Engine/Source/Runtime/Core/Public/Math/Box.h
+++ b/Engine/Source/Runtime/Core/Public/Math/Box.h
@@ -700,6 +700,10 @@ public:
 		{
 			return UE_SERIALIZE_VARIANT_FROM_MISMATCHED_TAG(Ar, Box, Box3d, Box3f);
 		}
+		else
+		{
+			static_assert(sizeof(T) == 0, "Unimplemented");
+		}
 	}
 };
 
diff --git a/Engine/Source/Runtime/Core/Public/Math/Box2D.h b/Engine/Source/Runtime/Core/Public/Math/Box2D.h
--- a/Engine/Source/Runtime/Core/Public/Math/Box2D.h
+++ b/Engine/Source/Runtime/Core/Public/Math/Box2D.h
@@ -457,6 +457,10 @@ public:
 		{
 			return UE_SERIALIZE_VARIANT_FROM_MISMATCHED_TAG(Ar, Box2D, Box2d, Box2f);
 		}
+		else
+		{
+			static_assert(sizeof(T) == 0, "Unimplemented");
+		}
 	}
 
 	// Conversion from other type. 
diff --git a/Engine/Source/Runtime/Core/Public/Math/BoxSphereBounds.h b/Engine/Source/Runtime/Core/Public/Math/BoxSphereBounds.h
--- a/Engine/Source/Runtime/Core/Public/Math/BoxSphereBounds.h
+++ b/Engine/Source/Runtime/Core/Public/Math/BoxSphereBounds.h
@@ -146,6 +146,10 @@ public:
 			// Falls back to UseSerializeItem to convert per property.
 			return false;
 		}
+		else
+		{
+			static_assert(sizeof(T) == 0, "Unimplemented");
+		}
 	}
 
 	/**
diff --git a/Engine/Source/Runtime/Core/Public/Math/IntPoint.h b/Engine/Source/Runtime/Core/Public/Math/IntPoint.h
--- a/Engine/Source/Runtime/Core/Public/Math/IntPoint.h
+++ b/Engine/Source/Runtime/Core/Public/Math/IntPoint.h
@@ -518,6 +518,10 @@ PRAGMA_ENABLE_DEPRECATION_WARNINGS
 		{
 			return UE_SERIALIZE_VARIANT_FROM_MISMATCHED_TAG(Ar, UintPoint, Uint64Point, Uint32Point);
 		}
+		else
+		{
+			static_assert(sizeof(IntType) == 0, "Unimplemented");
+		}
 	}
 };
 
diff --git a/Engine/Source/Runtime/Core/Public/Math/IntVector.h b/Engine/Source/Runtime/Core/Public/Math/IntVector.h
--- a/Engine/Source/Runtime/Core/Public/Math/IntVector.h
+++ b/Engine/Source/Runtime/Core/Public/Math/IntVector.h
@@ -649,6 +649,10 @@ PRAGMA_ENABLE_DEPRECATION_WARNINGS
 		{
 			return UE_SERIALIZE_VARIANT_FROM_MISMATCHED_TAG(Ar, UintVector, Uint64Vector, Uint32Vector);
 		}
+		else
+		{
+			static_assert(sizeof(IntType) == 0, "Unimplemented");
+		}
 	}
 };
 
@@ -1190,6 +1194,10 @@ PRAGMA_ENABLE_DEPRECATION_WARNINGS
 		{
 			return UE_SERIALIZE_VARIANT_FROM_MISMATCHED_TAG(Ar, UintVector2, Uint64Vector2, Uint32Vector2);
 		}
+		else
+		{
+			static_assert(sizeof(IntType) == 0, "Unimplemented");
+		}
 	}
 };
 
@@ -1817,6 +1825,10 @@ PRAGMA_ENABLE_DEPRECATION_WARNINGS
 		{
 			return UE_SERIALIZE_VARIANT_FROM_MISMATCHED_TAG(Ar, UintVector4, Uint64Vector4, Uint32Vector4);
 		}
+		else
+		{
+			static_assert(sizeof(IntType) == 0, "Unimplemented");
+		}
 	}
 };
 
diff --git a/Engine/Source/Runtime/Core/Public/Math/Plane.h b/Engine/Source/Runtime/Core/Public/Math/Plane.h
--- a/Engine/Source/Runtime/Core/Public/Math/Plane.h
+++ b/Engine/Source/Runtime/Core/Public/Math/Plane.h
@@ -319,6 +319,10 @@ public:
 		{
 			return UE_SERIALIZE_VARIANT_FROM_MISMATCHED_TAG(Ar, Plane, Plane4d, Plane4f);
 		}
+		else
+		{
+			static_assert(sizeof(T) == 0, "Unimplemented");
+		}
 	}
 
 	/**
diff --git a/Engine/Source/Runtime/Core/Public/Math/Ray.h b/Engine/Source/Runtime/Core/Public/Math/Ray.h
--- a/Engine/Source/Runtime/Core/Public/Math/Ray.h
+++ b/Engine/Source/Runtime/Core/Public/Math/Ray.h
@@ -211,6 +211,10 @@ public:
 		{
 			return UE_SERIALIZE_VARIANT_FROM_MISMATCHED_TAG(Ar, Ray, Ray3d, Ray3f);
 		}
+		else
+		{
+			static_assert(sizeof(T) == 0, "Unimplemented");
+		}
 	}
 
 	// Conversion to other type.
diff --git a/Engine/Source/Runtime/Core/Public/Math/Sphere.h b/Engine/Source/Runtime/Core/Public/Math/Sphere.h
--- a/Engine/Source/Runtime/Core/Public/Math/Sphere.h
+++ b/Engine/Source/Runtime/Core/Public/Math/Sphere.h
@@ -241,6 +241,10 @@ public:
 		{
 			return UE_SERIALIZE_VARIANT_FROM_MISMATCHED_TAG(Ar, Sphere, Sphere3d, Sphere3f);
 		}
+		else
+		{
+			static_assert(sizeof(T) == 0, "Unimplemented");
+		}
 	}
 
 };
