Appease PVS
#rnx

diff --git a/Engine/Source/Runtime/Core/Public/Math/Box.h b/Engine/Source/Runtime/Core/Public/Math/Box.h
--- a/Engine/Source/Runtime/Core/Public/Math/Box.h
+++ b/Engine/Source/Runtime/Core/Public/Math/Box.h
@@ -703,6 +703,7 @@ public:
 		else
 		{
 			static_assert(sizeof(T) == 0, "Unimplemented");
+			return false;
 		}
 	}
 };
diff --git a/Engine/Source/Runtime/Core/Public/Math/Box2D.h b/Engine/Source/Runtime/Core/Public/Math/Box2D.h
--- a/Engine/Source/Runtime/Core/Public/Math/Box2D.h
+++ b/Engine/Source/Runtime/Core/Public/Math/Box2D.h
@@ -460,6 +460,7 @@ public:
 		else
 		{
 			static_assert(sizeof(T) == 0, "Unimplemented");
+			return false;
 		}
 	}
 
diff --git a/Engine/Source/Runtime/Core/Public/Math/BoxSphereBounds.h b/Engine/Source/Runtime/Core/Public/Math/BoxSphereBounds.h
--- a/Engine/Source/Runtime/Core/Public/Math/BoxSphereBounds.h
+++ b/Engine/Source/Runtime/Core/Public/Math/BoxSphereBounds.h
@@ -149,6 +149,7 @@ public:
 		else
 		{
 			static_assert(sizeof(T) == 0, "Unimplemented");
+			return false;
 		}
 	}
 
diff --git a/Engine/Source/Runtime/Core/Public/Math/IntPoint.h b/Engine/Source/Runtime/Core/Public/Math/IntPoint.h
--- a/Engine/Source/Runtime/Core/Public/Math/IntPoint.h
+++ b/Engine/Source/Runtime/Core/Public/Math/IntPoint.h
@@ -521,6 +521,7 @@ PRAGMA_ENABLE_DEPRECATION_WARNINGS
 		else
 		{
 			static_assert(sizeof(IntType) == 0, "Unimplemented");
+			return false;
 		}
 	}
 };
diff --git a/Engine/Source/Runtime/Core/Public/Math/IntVector.h b/Engine/Source/Runtime/Core/Public/Math/IntVector.h
--- a/Engine/Source/Runtime/Core/Public/Math/IntVector.h
+++ b/Engine/Source/Runtime/Core/Public/Math/IntVector.h
@@ -652,6 +652,7 @@ PRAGMA_ENABLE_DEPRECATION_WARNINGS
 		else
 		{
 			static_assert(sizeof(IntType) == 0, "Unimplemented");
+			return false;
 		}
 	}
 };
diff --git a/Engine/Source/Runtime/Core/Public/Math/Plane.h b/Engine/Source/Runtime/Core/Public/Math/Plane.h
--- a/Engine/Source/Runtime/Core/Public/Math/Plane.h
+++ b/Engine/Source/Runtime/Core/Public/Math/Plane.h
@@ -322,6 +322,7 @@ public:
 		else
 		{
 			static_assert(sizeof(T) == 0, "Unimplemented");
+			return false;
 		}
 	}
 
diff --git a/Engine/Source/Runtime/Core/Public/Math/Ray.h b/Engine/Source/Runtime/Core/Public/Math/Ray.h
--- a/Engine/Source/Runtime/Core/Public/Math/Ray.h
+++ b/Engine/Source/Runtime/Core/Public/Math/Ray.h
@@ -214,6 +214,7 @@ public:
 		else
 		{
 			static_assert(sizeof(T) == 0, "Unimplemented");
+			return false;
 		}
 	}
 
diff --git a/Engine/Source/Runtime/Core/Public/Math/Sphere.h b/Engine/Source/Runtime/Core/Public/Math/Sphere.h
--- a/Engine/Source/Runtime/Core/Public/Math/Sphere.h
+++ b/Engine/Source/Runtime/Core/Public/Math/Sphere.h
@@ -244,6 +244,7 @@ public:
 		else
 		{
 			static_assert(sizeof(T) == 0, "Unimplemented");
+			return false;
 		}
 	}
 
