Additional PVS static_assert appeasement
#rnx

diff --git a/Engine/Source/Runtime/Core/Public/Math/IntVector.h b/Engine/Source/Runtime/Core/Public/Math/IntVector.h
--- a/Engine/Source/Runtime/Core/Public/Math/IntVector.h
+++ b/Engine/Source/Runtime/Core/Public/Math/IntVector.h
@@ -1198,6 +1198,7 @@ PRAGMA_ENABLE_DEPRECATION_WARNINGS
 		else
 		{
 			static_assert(sizeof(IntType) == 0, "Unimplemented");
+			return false;
 		}
 	}
 };
@@ -1829,6 +1830,7 @@ PRAGMA_ENABLE_DEPRECATION_WARNINGS
 		else
 		{
 			static_assert(sizeof(IntType) == 0, "Unimplemented");
+			return false;
 		}
 	}
 };
