Cleanup a restricted wording from a commentary

diff --git a/Engine/Source/Runtime/Core/Public/Logging/LogMacros.h b/Engine/Source/Runtime/Core/Public/Logging/LogMacros.h
--- a/Engine/Source/Runtime/Core/Public/Logging/LogMacros.h
+++ b/Engine/Source/Runtime/Core/Public/Logging/LogMacros.h
@@ -25,7 +25,6 @@
 ----------------------------------------------------------------------------*/
 
 // Define NO_LOGGING to strip out all writing to log files, OutputDebugString(), etc.
-// This is needed for consoles that require no logging (Xbox, Xenon)
 
 /**
  * FMsg 
