Add [[nodiscard]] attribute to Vector's GetClampedToMaxSize.

#rnx

diff --git a/Engine/Source/Runtime/Core/Public/Math/Vector.h b/Engine/Source/Runtime/Core/Public/Math/Vector.h
--- a/Engine/Source/Runtime/Core/Public/Math/Vector.h
+++ b/Engine/Source/Runtime/Core/Public/Math/Vector.h
@@ -716,7 +716,7 @@ public:
     [[nodiscard]] TVector<T> GetClampedToSize2D(T Min, T Max) const;
 
     /** Create a copy of this vector, with its maximum magnitude clamped to MaxSize. */
-    TVector<T> GetClampedToMaxSize(T MaxSize) const;
+	[[nodiscard]] TVector<T> GetClampedToMaxSize(T MaxSize) const;
 
     /** Create a copy of this vector, with the maximum 2D magnitude clamped to MaxSize. Z is unchanged. */
     [[nodiscard]] TVector<T> GetClampedToMaxSize2D(T MaxSize) const;
