Unreal.natvis: New fix for FName visualization under Clang for Windows without breaking other platforms (tm).
Fixing up visualization of UObjects will be handled in a separate change.

#rnx

diff --git a/Engine/Source/Runtime/Core/Private/Misc/CoreGlobals.cpp b/Engine/Source/Runtime/Core/Private/Misc/CoreGlobals.cpp
--- a/Engine/Source/Runtime/Core/Private/Misc/CoreGlobals.cpp
+++ b/Engine/Source/Runtime/Core/Private/Misc/CoreGlobals.cpp
@@ -267,7 +267,7 @@ namespace UE::CoreUObject::Private
 UE::CoreUObject::Private::FStoredObjectPathDebug* GCoreComplexObjectPathDebug = nullptr;
 UE::CoreUObject::Private::FObjectHandlePackageDebugData* GCoreObjectHandlePackageDebug = nullptr;
 #if PLATFORM_UNIX
-uint8** CORE_API GNameBlocksDebug = FNameDebugVisualizer::GetBlocks();
+uint8** CORE_API GNameBlocksDebug = FNameDebugVisualizer(FClangKeepDebugInfo{}).GetBlocks();
 FChunkedFixedUObjectArray*& CORE_API GObjectArrayForDebugVisualizers = GCoreObjectArrayForDebugVisualizers;
 UE::CoreUObject::Private::FStoredObjectPathDebug*& GComplexObjectPathDebug = GCoreComplexObjectPathDebug;
 UE::CoreUObject::Private::FObjectHandlePackageDebugData*& CORE_API GObjectHandlePackageDebug = GCoreObjectHandlePackageDebug;
diff --git a/Engine/Source/Runtime/Core/Private/UObject/UnrealNames.cpp b/Engine/Source/Runtime/Core/Private/UObject/UnrealNames.cpp
--- a/Engine/Source/Runtime/Core/Private/UObject/UnrealNames.cpp
+++ b/Engine/Source/Runtime/Core/Private/UObject/UnrealNames.cpp
@@ -2341,7 +2341,6 @@ static uint16 GetRawNonCasePreservingHash(const TCharType* Source)
 -----------------------------------------------------------------------------*/
 
 // Make Clang keep FNameEntry type debug info needed for debug visualizers
-struct FClangKeepDebugInfo {};
 FNameEntry::FNameEntry(FClangKeepDebugInfo) {}
 
 void FNameEntry::CopyUnterminatedName(ANSICHAR* Out, uint32 Len) const
@@ -5705,6 +5704,9 @@ void CallNameCreationHook()
 
 #endif
 
+// Make Clang keep FNameDebugVisualizer type debug info needed for debug visualizers
+FNameDebugVisualizer::FNameDebugVisualizer(FClangKeepDebugInfo) {}
+
 uint8** FNameDebugVisualizer::GetBlocks()
 {
 	static_assert(EntryStride == FNameEntryAllocator::Stride,	"Natvis constants out of sync with actual constants");
diff --git a/Engine/Source/Runtime/Core/Public/Modules/Boilerplate/ModuleBoilerplate.h b/Engine/Source/Runtime/Core/Public/Modules/Boilerplate/ModuleBoilerplate.h
--- a/Engine/Source/Runtime/Core/Public/Modules/Boilerplate/ModuleBoilerplate.h
+++ b/Engine/Source/Runtime/Core/Public/Modules/Boilerplate/ModuleBoilerplate.h
@@ -10,8 +10,6 @@
 //   FVisualizerDebuggingState::Ptrs or FVisualizerDebuggingState::GuidString
 #include "Modules/VisualizerDebuggingState.h"
 
-class FChunkedFixedUObjectArray;
-
 // Boilerplate that is included once for each module, even in monolithic builds
 #if !defined(PER_MODULE_BOILERPLATE_ANYLINK)
 #define PER_MODULE_BOILERPLATE_ANYLINK(ModuleImplClass, ModuleName)
@@ -94,7 +92,7 @@ class FChunkedFixedUObjectArray;
 	#define UE_VISUALIZERS_HELPERS
 #else
 	#define UE_VISUALIZERS_HELPERS \
-		uint8** GNameBlocksDebug = FNameDebugVisualizer::GetBlocks(); \
+		uint8** GNameBlocksDebug = FNameDebugVisualizer(FClangKeepDebugInfo{}).GetBlocks(); \
 		FChunkedFixedUObjectArray*& GObjectArrayForDebugVisualizers = GCoreObjectArrayForDebugVisualizers; \
 		UE::CoreUObject::Private::FStoredObjectPathDebug*& GComplexObjectPathDebug = GCoreComplexObjectPathDebug; \
 		UE::CoreUObject::Private::FObjectHandlePackageDebugData*& GObjectHandlePackageDebug = GCoreObjectHandlePackageDebug; \
diff --git a/Engine/Source/Runtime/Core/Public/UObject/NameTypes.h b/Engine/Source/Runtime/Core/Public/UObject/NameTypes.h
--- a/Engine/Source/Runtime/Core/Public/UObject/NameTypes.h
+++ b/Engine/Source/Runtime/Core/Public/UObject/NameTypes.h
@@ -48,6 +48,9 @@ struct FScriptName;
 struct FNumberedEntry;
 class FName;
 
+// Generic constructor/initializer struct to make Clang keep type debug info needed for debug visualizers
+struct FClangKeepDebugInfo {};
+
 /** Opaque id to a deduplicated name */
 struct FNameEntryId
 {
@@ -250,7 +253,7 @@ private:
 	};
 
 
-	FNameEntry(struct FClangKeepDebugInfo);
+	FNameEntry(FClangKeepDebugInfo);
 	FNameEntry(const FNameEntry&) = delete;
 	FNameEntry(FNameEntry&&) = delete;
 	FNameEntry& operator=(const FNameEntry&) = delete;
@@ -1640,7 +1643,8 @@ struct FNameLexicalLess
 
 struct FNameDebugVisualizer
 {
-	CORE_API static uint8** GetBlocks();
+	CORE_API FNameDebugVisualizer(FClangKeepDebugInfo);
+	CORE_API uint8** GetBlocks();
 private:
 	static constexpr uint32 EntryStride = alignof(FNameEntry);
 	static constexpr uint32 OffsetBits = 16;
@@ -1853,4 +1857,4 @@ template <> struct TIsContiguousContainer<FNameBuilder> { static constexpr inlin
 
 /** Update the Hash with the FName's text and number */
 class FBlake3;
-CORE_API void AppendHash(FBlake3& Builder, FName In);
\ No newline at end of file
+CORE_API void AppendHash(FBlake3& Builder, FName In);
