[AutoRTFM] Update OutputDeviceRedirector read lock for transactional safety.

Previously, we would crash if `UE_LOG` was called inside a transaction after a second `FOutputDevice` was added to `GLog`.

diff --git a/Engine/Source/Programs/AutoRTFMTests/Private/AutoRTFMTests.cpp b/Engine/Source/Programs/AutoRTFMTests/Private/AutoRTFMTests.cpp
--- a/Engine/Source/Programs/AutoRTFMTests/Private/AutoRTFMTests.cpp
+++ b/Engine/Source/Programs/AutoRTFMTests/Private/AutoRTFMTests.cpp
@@ -83,6 +83,7 @@ int RunTests(int ArgC, const char* ArgV[])
 
 	const TCHAR* CommandLine = TEXT("-Multiprocess -LogCmds=\"LogCsvProfiler off, LogStreaming off, LogUObjectGlobals off, LogPackageName off, LogAutoRTFM warning\" -AsyncLoadingThread");
 	GEngineLoop.PreInit(CommandLine); // Note: Initializes AutoRTFM
+	GLog->SetCurrentThreadAsPrimaryThread();
 
 	if (bRetryNestedToo)
 	{
diff --git a/Engine/Source/Programs/AutoRTFMTests/Private/UECoreTests.cpp b/Engine/Source/Programs/AutoRTFMTests/Private/UECoreTests.cpp
--- a/Engine/Source/Programs/AutoRTFMTests/Private/UECoreTests.cpp
+++ b/Engine/Source/Programs/AutoRTFMTests/Private/UECoreTests.cpp
@@ -2455,6 +2455,46 @@ TEST_CASE("UECore.UE_LOGFMT")
 	}
 }
 
+TEST_CASE("UECore.FOutputDeviceRedirector")
+{
+	SECTION("Commit")
+	{
+		FStringOutputDevice StringLog;
+		GLog->AddOutputDevice(&StringLog);
+
+		AutoRTFM::Testing::Commit([&]
+		{
+			// This test will actually be run twice, because we test with SetRetryTransaction enabled.
+			// Logging always runs in the open because UE::Private::Log::BasicLog is marked ALWAYS_OPEN,
+			// so the log won't be undone when the transaction is rolled back before being retried.
+			// We handle this by making sure that the string log has "Commit" appended to it, rather
+			// than verifying that it contains "Commit" exactly.
+			FString PreviousStringLog = StringLog;
+			UE_LOG(LogAutoRTFM_UECoreTests, Display, TEXT("Commit"));
+			REQUIRE(StringLog == PreviousStringLog + FString("Commit"));
+		});
+
+		GLog->RemoveOutputDevice(&StringLog);
+	}
+
+	SECTION("Abort")
+	{
+		FStringOutputDevice StringLog;
+		GLog->AddOutputDevice(&StringLog);
+
+		AutoRTFM::Testing::Abort([&]
+		{
+			FString PreviousStringLog = StringLog;
+			UE_LOG(LogAutoRTFM_UECoreTests, Display, TEXT("Abort"));
+			REQUIRE(StringLog == PreviousStringLog + FString("Abort"));
+
+			AutoRTFM::AbortTransaction();
+		});
+
+		GLog->RemoveOutputDevice(&StringLog);
+	}
+}
+
 TEST_CASE("UECore.AsyncLoading")
 {
 	SECTION("LoadPackageAsync")
diff --git a/Engine/Source/Runtime/Core/Private/Misc/OutputDeviceRedirector.cpp b/Engine/Source/Runtime/Core/Private/Misc/OutputDeviceRedirector.cpp
--- a/Engine/Source/Runtime/Core/Private/Misc/OutputDeviceRedirector.cpp
+++ b/Engine/Source/Runtime/Core/Private/Misc/OutputDeviceRedirector.cpp
@@ -3,6 +3,7 @@
 #include "Misc/OutputDeviceRedirector.h"
 
 #include "Async/EventCount.h"
+#include "AutoRTFM/AutoRTFM.h"
 #include "Containers/BitArray.h"
 #include "Containers/ConsumeAllMpmcQueue.h"
 #include "Experimental/ConcurrentLinearAllocator.h"
@@ -14,6 +15,7 @@
 #include "Misc/App.h"
 #include "Misc/ScopeRWLock.h"
 #include "Misc/StringBuilder.h"
+#include "Misc/TransactionallySafeRWLock.h"
 #include "Misc/TVariant.h"
 #include "ProfilingDebugging/CpuProfilerTrace.h"
 #include <atomic>
@@ -229,6 +231,28 @@ class FOutputDevicesReadScopeLock
 public:
 	FORCEINLINE explicit FOutputDevicesReadScopeLock(FOutputDeviceRedirectorState& InState)
 		: State(InState)
+	{
+		UE_AUTORTFM_OPEN
+		{
+			AcquireLock();
+		};
+		AutoRTFM::PushOnAbortHandler(this, [this]
+		{
+			ReleaseLock();
+		});
+	}
+
+	FORCEINLINE ~FOutputDevicesReadScopeLock()
+	{
+		UE_AUTORTFM_OPEN
+		{
+			ReleaseLock();
+		};
+		AutoRTFM::PopOnAbortHandler(this);
+	}
+
+private:
+	void AcquireLock()
 	{
 		// Read locks add/sub by 2 to keep the LSB free for write locks to use.
 		if (State.OutputDevicesLockState.fetch_add(2, std::memory_order_acquire) & 1)
@@ -237,24 +261,23 @@ public:
 		}
 	}
 
+	void ReleaseLock()
+	{
+		State.OutputDevicesLockState.fetch_sub(2, std::memory_order_release);
+	}
+
 	FORCENOINLINE void WaitForWriteLock()
 	{
 		TRACE_CPUPROFILER_EVENT_SCOPE(FOutputDevicesReadScopeLock);
 		// A write lock has set the LSB. Cancel this read lock and wait for the write.
 		State.OutputDevicesLockState.fetch_sub(2, std::memory_order_relaxed);
 		// This read lock will wait until the write lock exits.
-		FReadScopeLock ScopeLock(State.OutputDevicesLock);
+		UE::TReadScopeLock ScopeLock(State.OutputDevicesLock);
 		// This is relaxed because locking OutputDevicesLock has acquire semantics.
 		uint32 LockState = State.OutputDevicesLockState.fetch_add(2, std::memory_order_relaxed);
 		check((LockState & 1) == 0);
 	}
 
-	FORCEINLINE ~FOutputDevicesReadScopeLock()
-	{
-		State.OutputDevicesLockState.fetch_sub(2, std::memory_order_release);
-	}
-
-private:
 	FOutputDeviceRedirectorState& State;
 };
 
@@ -271,6 +294,9 @@ public:
 	FORCEINLINE explicit FOutputDevicesWriteScopeLock(FOutputDeviceRedirectorState& InState)
 		: State(InState)
 	{
+		// FOutputDevicesWriteScopeLock is not transactionally safe.
+		check(!AutoRTFM::IsClosed());
+
 		TRACE_CPUPROFILER_EVENT_SCOPE(FOutputDevicesWriteScopeLock);
 		// Take the lock before modifying the state, to avoid contention on the LSB.
 		State.OutputDevicesLock.WriteLock();
