Add GetClampedToSize and GetClampedToMaxSize to Vector2D.

#rnx

diff --git a/Engine/Source/Runtime/Core/Public/Math/Vector2D.h b/Engine/Source/Runtime/Core/Public/Math/Vector2D.h
--- a/Engine/Source/Runtime/Core/Public/Math/Vector2D.h
+++ b/Engine/Source/Runtime/Core/Public/Math/Vector2D.h
@@ -645,6 +645,18 @@ public:
 	*/
 	[[nodiscard]] TVector2<T> ClampAxes(T MinAxisVal, T MaxAxisVal) const;
 
+	/**
+	* Create a copy of this vector, with its magnitude clamped between Min and Max.
+	* @return New vector with clamped magnitude.
+	*/
+	[[nodiscard]] TVector2<T> GetClampedToSize(T Min, T Max) const;
+
+	/**
+	* Create a copy of this vector, with its maximum magnitude clamped to MaxSize.
+	* @return New vector with clamped magnitude.
+	*/
+	[[nodiscard]] TVector2<T> GetClampedToMaxSize(T MaxSize) const;
+
 	/**
 	* Get a copy of the vector as sign only.
 	* Each component is set to +1 or -1, with the sign of zero treated as +1.
@@ -1298,6 +1310,37 @@ FORCEINLINE TVector2<T> TVector2<T>::ClampAxes(T MinAxisVal, T MaxAxisVal) const
 	return TVector2<T>(FMath::Clamp(X, MinAxisVal, MaxAxisVal), FMath::Clamp(Y, MinAxisVal, MaxAxisVal));
 }
 
+template<typename T>
+FORCEINLINE TVector2<T> TVector2<T>::GetClampedToSize(T Min, T Max) const
+{
+	T VecSize = Size();
+	const TVector2<T> VecDir = (VecSize > UE_SMALL_NUMBER) ? (*this / VecSize) : ZeroVector;
+
+	VecSize = FMath::Clamp(VecSize, Min, Max);
+
+	return VecSize * VecDir;
+}
+
+template<typename T>
+FORCEINLINE TVector2<T> TVector2<T>::GetClampedToMaxSize(T MaxSize) const
+{
+	if (MaxSize < UE_KINDA_SMALL_NUMBER)
+	{
+		return ZeroVector;
+	}
+
+	const T VSq = SizeSquared();
+	if (VSq > FMath::Square(MaxSize))
+	{
+		const T Scale = MaxSize * FMath::InvSqrt(VSq);
+		return TVector2<T>(X * Scale, Y * Scale);
+	}
+	else
+	{
+		return *this;
+	}
+}
+
 template<typename T>
 FORCEINLINE TVector2<T> TVector2<T>::GetSignVector() const
 {
