[Truncation Warnings] Add template versions for some of the 2d vector math to support slate 2d -> 2f updates

diff --git a/Engine/Source/Runtime/Core/Private/Math/UnrealMath.cpp b/Engine/Source/Runtime/Core/Private/Math/UnrealMath.cpp
--- a/Engine/Source/Runtime/Core/Private/Math/UnrealMath.cpp
+++ b/Engine/Source/Runtime/Core/Private/Math/UnrealMath.cpp
@@ -12,6 +12,7 @@
 #include "Math/Matrix.h"
 #include "Math/Quat.h"
 #include "Math/Vector.h"
+#include "Math/Vector2D.h"
 #include "Math/Vector4.h"
 
 DEFINE_LOG_CATEGORY(LogUnrealMath);
@@ -1715,20 +1716,23 @@ UE::Math::TVector<T> FMath::ClosestPointOnSegment(const UE::Math::TVector<T>&Poi
 	return StartPoint + Segment * (Dot1 / Dot2);
 }
 
-FVector2D FMath::ClosestPointOnSegment2D(const FVector2D &Point, const FVector2D &StartPoint, const FVector2D &EndPoint)
+template <typename T>
+UE::Math::TVector2<T> FMath::ClosestPointOnSegment2D(const UE::Math::TVector2<T>& Point, const UE::Math::TVector2<T>& StartPoint, const UE::Math::TVector2<T>& EndPoint)
 {
-	const FVector2D Segment = EndPoint - StartPoint;
-	const FVector2D VectToPoint = Point - StartPoint;
+	using VectorType = UE::Math::TVector2<T>;
+
+	const VectorType Segment = EndPoint - StartPoint;
+	const VectorType VectToPoint = Point - StartPoint;
 
 	// See if closest point is before StartPoint
-	const FVector2D::FReal Dot1 = VectToPoint | Segment;
+	const T Dot1 = VectToPoint | Segment;
 	if (Dot1 <= 0)
 	{
 		return StartPoint;
 	}
 
 	// See if closest point is beyond EndPoint
-	const FVector2D::FReal Dot2 = Segment | Segment;
+	const T Dot2 = Segment | Segment;
 	if (Dot2 <= Dot1)
 	{
 		return EndPoint;
@@ -2622,20 +2626,21 @@ CORE_API FVector2D FMath::Vector2DInterpConstantTo( const FVector2D& Current, co
 	return Target;
 }
 
-CORE_API FVector2D FMath::Vector2DInterpTo( const FVector2D& Current, const FVector2D& Target, float DeltaTime, float InterpSpeed )
+template <typename T>
+CORE_API UE::Math::TVector2<T> FMath::Vector2DInterpTo( const UE::Math::TVector2<T>& Current, const UE::Math::TVector2<T>& Target, float DeltaTime, float InterpSpeed )
 {
 	if( InterpSpeed <= 0.f )
 	{
 		return Target;
 	}
 
-	const FVector2D Dist = Target - Current;
-	if( Dist.SizeSquared() < UE_KINDA_SMALL_NUMBER )
+	const UE::Math::TVector2<T> Dist = Target - Current;
+	if( Dist.SizeSquared() < static_cast<UE::Math::TVector2<T>::FReal>(UE_KINDA_SMALL_NUMBER) )
 	{
 		return Target;
 	}
 
-	const FVector2D DeltaMove = Dist * FMath::Clamp<float>(DeltaTime * InterpSpeed, 0.f, 1.f);
+	const UE::Math::TVector2<T> DeltaMove = Dist * static_cast<UE::Math::TVector2<T>::FReal>(FMath::Clamp(DeltaTime * InterpSpeed, 0.f, 1.f));
 	return Current + DeltaMove;
 }
 
@@ -3701,6 +3706,8 @@ template struct UE::Math::TQuat<float>;
 template struct UE::Math::TQuat<double>;
 template struct UE::Math::TVector<float>;
 template struct UE::Math::TVector<double>;
+template struct UE::Math::TVector2<float>;
+template struct UE::Math::TVector2<double>;
 template struct UE::Math::TVector4<float>;
 template struct UE::Math::TVector4<double>;
 template struct UE::Math::TMatrix<float>;
@@ -3711,3 +3718,7 @@ template CORE_API UE::Math::TVector<float> FMath::ClosestPointOnLine(const UE::M
 template CORE_API UE::Math::TVector<double> FMath::ClosestPointOnLine(const UE::Math::TVector<double>& LineStart, const UE::Math::TVector<double>& LineEnd, const UE::Math::TVector<double>& Point);
 template CORE_API UE::Math::TVector<float> FMath::ClosestPointOnSegment(const UE::Math::TVector<float>& Point, const UE::Math::TVector<float>& StartPoint, const UE::Math::TVector<float>& EndPoint);
 template CORE_API UE::Math::TVector<double> FMath::ClosestPointOnSegment(const UE::Math::TVector<double>& Point, const UE::Math::TVector<double>& StartPoint, const UE::Math::TVector<double>& EndPoint);
+template CORE_API UE::Math::TVector2<float> FMath::ClosestPointOnSegment2D(const UE::Math::TVector2<float>& Point, const UE::Math::TVector2<float>& StartPoint, const UE::Math::TVector2<float>& EndPoint);
+template CORE_API UE::Math::TVector2<double> FMath::ClosestPointOnSegment2D(const UE::Math::TVector2<double>& Point, const UE::Math::TVector2<double>& StartPoint, const UE::Math::TVector2<double>& EndPoint);
+template CORE_API UE::Math::TVector2<float> FMath::Vector2DInterpTo(const UE::Math::TVector2<float>& Current, const UE::Math::TVector2<float>& Target, float DeltaTime, float InterpSpeed);
+template CORE_API UE::Math::TVector2<double> FMath::Vector2DInterpTo(const UE::Math::TVector2<double>& Current, const UE::Math::TVector2<double>& Target, float DeltaTime, float InterpSpeed);
diff --git a/Engine/Source/Runtime/Core/Public/Math/UnrealMathUtility.h b/Engine/Source/Runtime/Core/Public/Math/UnrealMathUtility.h
--- a/Engine/Source/Runtime/Core/Public/Math/UnrealMathUtility.h
+++ b/Engine/Source/Runtime/Core/Public/Math/UnrealMathUtility.h
@@ -1548,7 +1548,8 @@ public:
 	[[nodiscard]] static CORE_API FVector2D Vector2DInterpConstantTo( const FVector2D& Current, const FVector2D& Target, float DeltaTime, float InterpSpeed );
 
 	/** Interpolate vector2D from Current to Target. Scaled by distance to Target, so it has a strong start speed and ease out. */
-	[[nodiscard]] static CORE_API FVector2D Vector2DInterpTo( const FVector2D& Current, const FVector2D& Target, float DeltaTime, float InterpSpeed );
+	template <typename T>
+	[[nodiscard]] static CORE_API UE::Math::TVector2<T> Vector2DInterpTo( const UE::Math::TVector2<T>& Current, const UE::Math::TVector2<T>& Target, float DeltaTime, float InterpSpeed );
 
 	/** Interpolate rotator from Current to Target with constant step */
 	[[nodiscard]] static CORE_API FRotator RInterpConstantTo( const FRotator& Current, const FRotator& Target, float DeltaTime, float InterpSpeed);
@@ -2071,7 +2072,8 @@ public:
 	*
 	* @return	point on the segment defined by (StartPoint, EndPoint) that is closest to Point.
 	*/
-	[[nodiscard]] static CORE_API FVector2D ClosestPointOnSegment2D(const FVector2D &Point, const FVector2D &StartPoint, const FVector2D &EndPoint);
+	template <typename T>
+	[[nodiscard]] static CORE_API UE::Math::TVector2<T> ClosestPointOnSegment2D(const UE::Math::TVector2<T>& Point, const UE::Math::TVector2<T>& StartPoint, const UE::Math::TVector2<T>& EndPoint);
 
 	/**
 	 * Returns distance from a point to the closest point on a segment.
diff --git a/Engine/Source/Runtime/Core/Public/Math/Vector2D.h b/Engine/Source/Runtime/Core/Public/Math/Vector2D.h
--- a/Engine/Source/Runtime/Core/Public/Math/Vector2D.h
+++ b/Engine/Source/Runtime/Core/Public/Math/Vector2D.h
@@ -615,8 +615,14 @@ public:
 	* @param OutDir Reference passed in to store unit direction vector.
 	* @param OutLength Reference passed in to store length of the vector.
 	*/
-	void ToDirectionAndLength(TVector2<T> &OutDir, double &OutLength) const;
-	void ToDirectionAndLength(TVector2<T> &OutDir, float &OutLength) const;
+	void ToDirectionAndLength(TVector2<T>& OutDir, T& OutLength) const;
+	template<typename FArg UE_REQUIRES(!std::is_same_v<T, FArg>)>
+	void ToDirectionAndLength(TVector2<T>& OutDir, FArg& OutLength) const
+	{
+		FReal TempOut;
+		ToDirectionAndLength(OutDir, TempOut);
+		OutLength = static_cast<FArg>(TempOut);
+	}
 
 	/**
 	* Checks whether all components of the vector are exactly zero.
@@ -1211,28 +1217,13 @@ FORCEINLINE bool TVector2<T>::Normalize(T Tolerance)
 }
 
 template<typename T>
-FORCEINLINE void TVector2<T>::ToDirectionAndLength(TVector2<T> &OutDir, double &OutLength) const
-{
-	OutLength = Size();
-	if (OutLength > UE_SMALL_NUMBER)
-	{
-		T OneOverLength = 1.0f / OutLength;
-		OutDir = TVector2<T>(X*OneOverLength, Y*OneOverLength);
-	}
-	else
-	{
-		OutDir = TVector2<T>::ZeroVector;
-	}
-}
-
-template<typename T>
-FORCEINLINE void TVector2<T>::ToDirectionAndLength(TVector2<T> &OutDir, float &OutLength) const
+FORCEINLINE void TVector2<T>::ToDirectionAndLength(TVector2<T> &OutDir, T &OutLength) const
 {
 	OutLength = Size();
 	if (OutLength > UE_SMALL_NUMBER)
 	{
-		float OneOverLength = 1.0f / OutLength;
-		OutDir = TVector2<T>(X*OneOverLength, Y*OneOverLength);
+		T OneOverLength = static_cast<T>(1.0f) / OutLength;
+		OutDir = TVector2<T>(static_cast<T>(X*OneOverLength), static_cast<T>(Y*OneOverLength));
 	}
 	else
 	{
