Android CIS fix

#rnx

diff --git a/Engine/Source/Runtime/Core/Private/Math/UnrealMath.cpp b/Engine/Source/Runtime/Core/Private/Math/UnrealMath.cpp
--- a/Engine/Source/Runtime/Core/Private/Math/UnrealMath.cpp
+++ b/Engine/Source/Runtime/Core/Private/Math/UnrealMath.cpp
@@ -2635,12 +2635,12 @@ CORE_API UE::Math::TVector2<T> FMath::Vector2DInterpTo( const UE::Math::TVector2
 	}
 
 	const UE::Math::TVector2<T> Dist = Target - Current;
-	if( Dist.SizeSquared() < static_cast<UE::Math::TVector2<T>::FReal>(UE_KINDA_SMALL_NUMBER) )
+	if( Dist.SizeSquared() < static_cast<typename UE::Math::TVector2<T>::FReal>(UE_KINDA_SMALL_NUMBER) )
 	{
 		return Target;
 	}
 
-	const UE::Math::TVector2<T> DeltaMove = Dist * static_cast<UE::Math::TVector2<T>::FReal>(FMath::Clamp(DeltaTime * InterpSpeed, 0.f, 1.f));
+	const UE::Math::TVector2<T> DeltaMove = Dist * static_cast<typename UE::Math::TVector2<T>::FReal>(FMath::Clamp(DeltaTime * InterpSpeed, 0.f, 1.f));
 	return Current + DeltaMove;
 }
 
