Added LLM_IS_ENABLED macro

diff --git a/Engine/Source/Runtime/Core/Public/HAL/LowLevelMemTracker.h b/Engine/Source/Runtime/Core/Public/HAL/LowLevelMemTracker.h
--- a/Engine/Source/Runtime/Core/Public/HAL/LowLevelMemTracker.h
+++ b/Engine/Source/Runtime/Core/Public/HAL/LowLevelMemTracker.h
@@ -321,6 +321,7 @@ extern CORE_API const FName LLMGetUntaggedTagName(ELLMTagSet TagSet);
  */
 #define LLM(x) x
 #define LLM_IF_ENABLED(x) if (FLowLevelMemTracker::IsEnabled()) { x; }
+#define LLM_IS_ENABLED() FLowLevelMemTracker::IsEnabled()
 #define SCOPE_NAME PREPROCESSOR_JOIN(LLMScope,__LINE__)
 
 ///////////////////////////////////////////////////////////////////////////////////////
@@ -1060,6 +1061,7 @@ inline bool FLowLevelMemTracker::IsEnabled()
 
 #define LLM(...)
 #define LLM_IF_ENABLED(...)
+#define LLM_IS_ENABLED(...)
 #define LLM_SCOPE(...)
 #define LLM_SCOPE_DYNAMIC(...)
 #define LLM_TAGSET_SCOPE(...)
