Fixed LLM_IS_ENABLED() macro when ENABLE_LOW_LEVEL_MEM_TRACKER is set to 0

diff --git a/Engine/Source/Runtime/Core/Public/HAL/LowLevelMemTracker.h b/Engine/Source/Runtime/Core/Public/HAL/LowLevelMemTracker.h
--- a/Engine/Source/Runtime/Core/Public/HAL/LowLevelMemTracker.h
+++ b/Engine/Source/Runtime/Core/Public/HAL/LowLevelMemTracker.h
@@ -1061,7 +1061,7 @@ inline bool FLowLevelMemTracker::IsEnabled()
 
 #define LLM(...)
 #define LLM_IF_ENABLED(...)
-#define LLM_IS_ENABLED(...)
+#define LLM_IS_ENABLED() false
 #define LLM_SCOPE(...)
 #define LLM_SCOPE_DYNAMIC(...)
 #define LLM_TAGSET_SCOPE(...)
