Add UE::TScopeTryLock.

This is more flexible than FScopeTryLock (paralleling UE::TScopeLock). Rather than forcing the caller to use an FCriticalSection, this opens the door to try-locking a UE::FMutex or FTransactionallySafeCriticalSection.

diff --git a/Engine/Source/Runtime/Core/Public/Misc/ScopeTryLock.h b/Engine/Source/Runtime/Core/Public/Misc/ScopeTryLock.h
--- a/Engine/Source/Runtime/Core/Public/Misc/ScopeTryLock.h
+++ b/Engine/Source/Runtime/Core/Public/Misc/ScopeTryLock.h
@@ -63,3 +63,60 @@ private:
 
 	FCriticalSection *CriticalSection;
 };
+
+namespace UE
+{
+	/**
+	 * RAII-style scope locking of a synchronization primitive using TryLock.
+	 * `MutexType` must be a type like `UE::FMutex` or `FCriticalSection`
+	 * which implements `TryLock` and `Unlock` methods.
+	 * 
+	 * <code>
+	 *	{
+	 *		// Try to acquire a lock on Mutex for the current scope.
+	 *		UE::TScopeTryLock ScopeTryLock(Mutex);
+	 *		
+	 *		// Check that the lock was acquired.
+	 *		if (ScopeTryLock.IsLocked())
+	 *		{
+	 *			// If the lock was acquired, we can safely access resources protected
+	 *			// by the mutex.
+	 *		}
+	 *		
+	 *		// When ScopeTryLock goes out of scope, the mutex will be released if it was
+	 *		// ever acquired.
+	 *	}
+	 * </code>
+	 */
+	template <typename MutexType>
+	class TScopeTryLock
+	{
+	public:
+		TScopeTryLock(MutexType& InMutex)
+		{
+			if (InMutex.TryLock())
+			{
+				HeldMutex = &InMutex;
+			}
+		}
+
+		~TScopeTryLock()
+		{
+			if (HeldMutex)
+			{
+				HeldMutex->Unlock();
+			}
+		}
+
+		FORCEINLINE bool IsLocked()
+		{
+			return HeldMutex != nullptr;
+		}
+
+	private:
+		TScopeTryLock(const TScopeTryLock&) = delete;
+		TScopeTryLock& operator=(TScopeTryLock&) = delete;
+
+		MutexType* HeldMutex = nullptr;
+	};
+}
\ No newline at end of file
