[SHA1] Add template function to simplify updating hash of POD types.

#rnx

diff --git a/Engine/Source/Runtime/Core/Public/Misc/SecureHash.h b/Engine/Source/Runtime/Core/Public/Misc/SecureHash.h
--- a/Engine/Source/Runtime/Core/Public/Misc/SecureHash.h
+++ b/Engine/Source/Runtime/Core/Public/Misc/SecureHash.h
@@ -323,6 +323,13 @@ public:
 	// Update the hash value with string
 	CORE_API void UpdateWithString(const TCHAR *data, uint32 len);
 
+	// Update with POD data.
+	template <typename Type	UE_REQUIRES(TIsPODType<Type>::Value)>
+	void Update(const Type& InData)
+	{
+		Update(reinterpret_cast<const uint8*>(&InData), sizeof(Type));
+	}
+
 	// Finalize hash and report
 	CORE_API void Final();
 
