FrameRate: Added divide by 0 protection to FrameRate.

Problem: divide by 0.
Testing: Played the game. When the crash was still occuring the ensure bypassed it. Now that the data has been fixed everything is running like is should.

diff --git a/Engine/Source/Runtime/Core/Public/Misc/FrameRate.h b/Engine/Source/Runtime/Core/Public/Misc/FrameRate.h
--- a/Engine/Source/Runtime/Core/Public/Misc/FrameRate.h
+++ b/Engine/Source/Runtime/Core/Public/Misc/FrameRate.h
@@ -271,6 +271,9 @@ inline FFrameTime ConvertFrameTime(FFrameTime SourceTime, FFrameRate SourceRate,
 	//We want NewTime =SourceTime * (DestinationRate/SourceRate);
 	//And want to limit conversions and keep int precision as much as possible
 
+	// If SourceRate.Numerator is 0 then NewDenominator will be 0 and will cause a divide by 0 error.
+	check(SourceRate.Numerator != 0);
+
 	//@todo: These integers should not need the volatile keyword here, but adding it works around
 	//       a compiler bug that results in an uninitialized vector register being used
 	volatile int64 NewNumerator = static_cast<int64>(DestinationRate.Numerator) * SourceRate.Denominator;
