ColorManagement: Making Bradford the explicit default chromatic adaptation method.

diff --git a/Engine/Plugins/Media/ImgMedia/Source/ImgMedia/Public/ImgMediaSourceColorSettings.h b/Engine/Plugins/Media/ImgMedia/Source/ImgMedia/Public/ImgMediaSourceColorSettings.h
--- a/Engine/Plugins/Media/ImgMedia/Source/ImgMedia/Public/ImgMediaSourceColorSettings.h
+++ b/Engine/Plugins/Media/ImgMedia/Source/ImgMedia/Public/ImgMediaSourceColorSettings.h
@@ -37,7 +37,7 @@ struct FMediaSourceColorSettings
 		, GreenChromaticityCoordinate(FVector2D::ZeroVector)
 		, BlueChromaticityCoordinate(FVector2D::ZeroVector)
 		, WhiteChromaticityCoordinate(FVector2D::ZeroVector)
-		, ChromaticAdaptationMethod(ETextureChromaticAdaptationMethod::TCAM_Bradford)
+		, ChromaticAdaptationMethod(static_cast<ETextureChromaticAdaptationMethod>(UE::Color::DEFAULT_CHROMATIC_ADAPTATION_METHOD))
 	{}
 
 	/** Source encoding of the media. */
diff --git a/Engine/Source/Developer/TextureCompressor/Public/TextureCompressorModule.h b/Engine/Source/Developer/TextureCompressor/Public/TextureCompressorModule.h
--- a/Engine/Source/Developer/TextureCompressor/Public/TextureCompressorModule.h
+++ b/Engine/Source/Developer/TextureCompressor/Public/TextureCompressorModule.h
@@ -313,7 +313,7 @@ struct FTextureBuildSettings
 		, GreenChromaticityCoordinate(0, 0)
 		, BlueChromaticityCoordinate(0, 0)
 		, WhiteChromaticityCoordinate(0, 0)
-		, ChromaticAdaptationMethod(1 /*UE::Color::EChromaticAdaptationMethod::Bradford*/)
+		, ChromaticAdaptationMethod(1 /*UE::Color::DEFAULT_CHROMATIC_ADAPTATION_METHOD*/)
 		, bUseLegacyGamma(false)
 		, bPreserveBorder(false)
 		, bForceNoAlphaChannel(false)
diff --git a/Engine/Source/Runtime/Core/Public/ColorManagement/ColorManagementDefines.h b/Engine/Source/Runtime/Core/Public/ColorManagement/ColorManagementDefines.h
--- a/Engine/Source/Runtime/Core/Public/ColorManagement/ColorManagementDefines.h
+++ b/Engine/Source/Runtime/Core/Public/ColorManagement/ColorManagementDefines.h
@@ -76,6 +76,9 @@ enum class EChromaticAdaptationMethod : uint8
 	Max,
 };
 
+/** Default method used across the engine for chromatic adaptation. */
+constexpr EChromaticAdaptationMethod DEFAULT_CHROMATIC_ADAPTATION_METHOD = EChromaticAdaptationMethod::Bradford;
+
 /** List of standard white points. */
 enum class EWhitePoint : uint8
 {
diff --git a/Engine/Source/Runtime/Core/Public/ColorManagement/ColorSpace.h b/Engine/Source/Runtime/Core/Public/ColorManagement/ColorSpace.h
--- a/Engine/Source/Runtime/Core/Public/ColorManagement/ColorSpace.h
+++ b/Engine/Source/Runtime/Core/Public/ColorManagement/ColorSpace.h
@@ -338,7 +338,7 @@ struct FColorSpaceTransform : FMatrix44d
 	* @param Target Target color space.
 	* @param Method Chromatic adapation method.
 	*/
-	CORE_API explicit FColorSpaceTransform(const FColorSpace& Src, const FColorSpace& Dst, EChromaticAdaptationMethod Method = EChromaticAdaptationMethod::Bradford);
+	CORE_API explicit FColorSpaceTransform(const FColorSpace& Src, const FColorSpace& Dst, EChromaticAdaptationMethod Method = DEFAULT_CHROMATIC_ADAPTATION_METHOD);
 
 	/**
 	* Constructor: create a color space transformation matrix from a raw matrix.
@@ -361,7 +361,7 @@ struct FColorSpaceTransform : FMatrix44d
 	* @param TargetXYZ Target color in XYZ space.
 	* @param Method Adaptation method (None, Bradford, CAT02).
 	*/
-	CORE_API static FMatrix44d CalcChromaticAdaptionMatrix(FVector3d SourceXYZ, FVector3d TargetXYZ, EChromaticAdaptationMethod Method = EChromaticAdaptationMethod::Bradford);
+	CORE_API static FMatrix44d CalcChromaticAdaptionMatrix(FVector3d SourceXYZ, FVector3d TargetXYZ, EChromaticAdaptationMethod Method = DEFAULT_CHROMATIC_ADAPTATION_METHOD);
 
 	/**
 	* Convenience function to get a (statically cached) conversion from sRGB/Rec709 to the working color space.
diff --git a/Engine/Source/Runtime/Engine/Classes/Engine/Texture.h b/Engine/Source/Runtime/Engine/Classes/Engine/Texture.h
--- a/Engine/Source/Runtime/Engine/Classes/Engine/Texture.h
+++ b/Engine/Source/Runtime/Engine/Classes/Engine/Texture.h
@@ -1172,15 +1172,8 @@ struct FTextureSourceColorSettings
 {
 	GENERATED_USTRUCT_BODY()
 
-	FTextureSourceColorSettings()
-		: EncodingOverride(ETextureSourceEncoding::TSE_None)
-		, ColorSpace(ETextureColorSpace::TCS_None)
-		, RedChromaticityCoordinate(FVector2D::ZeroVector)
-		, GreenChromaticityCoordinate(FVector2D::ZeroVector)
-		, BlueChromaticityCoordinate(FVector2D::ZeroVector)
-		, WhiteChromaticityCoordinate(FVector2D::ZeroVector)
-		, ChromaticAdaptationMethod(ETextureChromaticAdaptationMethod::TCAM_Bradford)
-	{}
+	/** Default constructor. */
+	ENGINE_API FTextureSourceColorSettings();
 
 	/** Source encoding of the texture, exposing more options than just sRGB. */
 	UPROPERTY(EditAnywhere, BlueprintReadWrite, Category = ColorManagement)
diff --git a/Engine/Source/Runtime/Engine/Private/Texture.cpp b/Engine/Source/Runtime/Engine/Private/Texture.cpp
--- a/Engine/Source/Runtime/Engine/Private/Texture.cpp
+++ b/Engine/Source/Runtime/Engine/Private/Texture.cpp
@@ -258,6 +258,16 @@ FTextureResource* UTexture::GetResource()
 	return nullptr;
 }
 
+FTextureSourceColorSettings::FTextureSourceColorSettings()
+	: EncodingOverride(ETextureSourceEncoding::TSE_None)
+	, ColorSpace(ETextureColorSpace::TCS_None)
+	, RedChromaticityCoordinate(FVector2D::ZeroVector)
+	, GreenChromaticityCoordinate(FVector2D::ZeroVector)
+	, BlueChromaticityCoordinate(FVector2D::ZeroVector)
+	, WhiteChromaticityCoordinate(FVector2D::ZeroVector)
+	, ChromaticAdaptationMethod(static_cast<ETextureChromaticAdaptationMethod>(UE::Color::DEFAULT_CHROMATIC_ADAPTATION_METHOD))
+{}
+
 UTexture::UTexture(FVTableHelper& Helper)
 	: Super(Helper)
 	, TextureReference(*new FTextureReference())
diff --git a/Engine/Source/Runtime/Media/Public/IMediaTextureSample.h b/Engine/Source/Runtime/Media/Public/IMediaTextureSample.h
--- a/Engine/Source/Runtime/Media/Public/IMediaTextureSample.h
+++ b/Engine/Source/Runtime/Media/Public/IMediaTextureSample.h
@@ -500,7 +500,7 @@ public:
 	*/
 	virtual UE::Color::EChromaticAdaptationMethod GetChromaticAdapationMethod() const
 	{
-		return UE::Color::EChromaticAdaptationMethod::Bradford;
+		return UE::Color::DEFAULT_CHROMATIC_ADAPTATION_METHOD;
 	}
 
 	/**
diff --git a/Engine/Source/Runtime/MediaAssets/Private/Misc/MediaTextureResource.cpp b/Engine/Source/Runtime/MediaAssets/Private/Misc/MediaTextureResource.cpp
--- a/Engine/Source/Runtime/MediaAssets/Private/Misc/MediaTextureResource.cpp
+++ b/Engine/Source/Runtime/MediaAssets/Private/Misc/MediaTextureResource.cpp
@@ -1061,7 +1061,7 @@ void FMediaTextureResource::GetColorSpaceConversionMatrixForSample(const TShared
 	}
 	else
 	{
-		// Apply the color space transformation from source to destination (including the Bradford chromatic adaptation).
+		// Apply the color space transformation from source to destination.
 		ColorSpaceMtx = UE::Color::Transpose<float>(UE::Color::FColorSpaceTransform(Sample->GetSourceColorSpace(), DestinationCS, Sample->GetChromaticAdapationMethod()));
 	}
 	
