Use page size for memory maped file alignment on Android to support 16kb page size

diff --git a/Engine/Source/Runtime/Core/Public/Android/AndroidPlatformProperties.h b/Engine/Source/Runtime/Core/Public/Android/AndroidPlatformProperties.h
--- a/Engine/Source/Runtime/Core/Public/Android/AndroidPlatformProperties.h
+++ b/Engine/Source/Runtime/Core/Public/Android/AndroidPlatformProperties.h
@@ -131,7 +131,7 @@ struct FAndroidPlatformProperties
 
 	static FORCEINLINE int64 GetMemoryMappingAlignment()
 	{
-		return 4096;
+		return FPlatformMemory::GetConstants().PageSize;
 	}
 
 	static FORCEINLINE bool SupportsVirtualTextureStreaming()
