Clang compile fixes when using the __cpuid intrinsic

#rnx

- When compiling with -stresstestunity & clang WindowsPlatformMisc.cpp will have access to <cpuid.h> via another cpp file in the module. The macro for __cpuid found in this header will override the intrinsic found in <intrin.h> which is what WindowsPlatformMisc.cpp is expecting to use. Since these two versions have different parameters the code will no longer compile.
- SecureHash - Remove the PLATFORM_WINDOWS check for the clang path, this was added back in CL 25865891 quite possibly to fix the sort of issue seen in WindowsPlatformMisc, which should now be fixed.

diff --git a/Engine/Source/Runtime/Core/Private/Misc/SecureHash.cpp b/Engine/Source/Runtime/Core/Private/Misc/SecureHash.cpp
--- a/Engine/Source/Runtime/Core/Private/Misc/SecureHash.cpp
+++ b/Engine/Source/Runtime/Core/Private/Misc/SecureHash.cpp
@@ -12,7 +12,7 @@
 #if (defined(_M_AMD64) || defined(__x86_64__)) && !defined(_M_ARM64EC)
 
 #include <immintrin.h>
-#if PLATFORM_COMPILER_CLANG && !PLATFORM_WINDOWS
+#if PLATFORM_COMPILER_CLANG
 #include <cpuid.h>
 #include <shaintrin.h>
 #endif
@@ -726,7 +726,7 @@ static inline bool DetectShaInstructions()
 {
 	int info1[4];
 	int info7[4];
-#if PLATFORM_COMPILER_CLANG && !PLATFORM_WINDOWS
+#if PLATFORM_COMPILER_CLANG
 	__cpuid(1, info1[0], info1[1], info1[2], info1[3]);
 	__cpuid_count(7, 0, info7[0], info7[1], info7[2], info7[3]);
 #else
diff --git a/Engine/Source/Runtime/Core/Private/Windows/WindowsPlatformMisc.cpp b/Engine/Source/Runtime/Core/Private/Windows/WindowsPlatformMisc.cpp
--- a/Engine/Source/Runtime/Core/Private/Windows/WindowsPlatformMisc.cpp
+++ b/Engine/Source/Runtime/Core/Private/Windows/WindowsPlatformMisc.cpp
@@ -75,6 +75,12 @@ THIRD_PARTY_INCLUDES_END
 #include <fcntl.h>
 #include <io.h>
 
+#if PLATFORM_COMPILER_CLANG
+	#include <cpuid.h>
+#else
+	#include <intrin.h>
+#endif // PLATFORM_COMPILER_CLANG
+
 #include "Windows/AllowWindowsPlatformTypes.h"
 
 #include "FramePro/FrameProProfiler.h"
@@ -103,6 +109,12 @@ static TAutoConsoleVariable<int32> CVarDriverDetectionMethod(
 	     "  5: Use Windows SetupAPI functions"),
 	ECVF_RenderThreadSafe);
 
+#if PLATFORM_COMPILER_CLANG
+	#define UE_CPUID(Info, FunctionID) __cpuid(FunctionID, Info[0], Info[1], Info[2], Info[3]);
+#else
+	#define UE_CPUID(Info, FunctionID) __cpuid(Info, FunctionID);
+#endif //PLATFORM_COMPILER_CLANG
+
 int32 GetOSVersionsHelper( TCHAR* OutOSVersionLabel, int32 OSVersionLabelLength, TCHAR* OutOSSubVersionLabel, int32 OSSubVersionLabelLength )
 {
 	int32 ErrorCode = (int32)FWindowsOSVersionHelper::SUCCEEDED;
@@ -3009,7 +3021,7 @@ private:
 #else
 		bool bSupportsTpause = false;
 		int CPUInfo[4];
-		__cpuid(CPUInfo, 0);
+		UE_CPUID(CPUInfo, 0);
 
 		if (CPUInfo[0] >= 7)
 		{
@@ -3056,7 +3068,7 @@ private:
 		} VendorResult;
 
 		int Args[4];
-		__cpuid(Args, 0);
+		UE_CPUID(Args, 0);
 
 		VendorResult.Dw.dw0 = Args[1];
 		VendorResult.Dw.dw1 = Args[3];
@@ -3078,7 +3090,7 @@ private:
 		int32 CPUInfo[4] = {-1};
 		const SIZE_T CPUInfoSize = sizeof( CPUInfo );
 
-		__cpuid( CPUInfo, 0x80000000 );
+		UE_CPUID(CPUInfo, 0x80000000);
 		const uint32 MaxExtIDs = CPUInfo[0];
 
 		if( MaxExtIDs >= 0x80000004 )
@@ -3087,7 +3099,7 @@ private:
 			const uint32 NumBrandStrings = 3;
 			for( uint32 Index = 0; Index < NumBrandStrings; Index++ )
 			{
-				__cpuid( CPUInfo, FirstBrandString + Index );
+				UE_CPUID(CPUInfo, FirstBrandString + Index);
 				FPlatformMemory::Memcpy( BrandString + CPUInfoSize * Index, CPUInfo, CPUInfoSize );
 			}
 		}
@@ -3102,7 +3114,7 @@ private:
 	 */
 	static void QueryCPUInfo(int Args[4])
 	{
-		__cpuid(Args, 1);
+		UE_CPUID(Args, 1);
 	}
 
 	/**
@@ -3115,7 +3127,7 @@ private:
 		int32 Result = 1;
 
 		int Args[4];
-		__cpuid(Args, 0x80000006);
+		UE_CPUID(Args, 0x80000006);
 
 		Result = Args[2] & 0xFF;
 		check(Result && !(Result & (Result - 1))); // assumed to be a power of two
@@ -3209,11 +3221,11 @@ uint32 FWindowsPlatformMisc::GetFeatureBits_X86()
 	uint32 MaxLeaf;
 
 	// Basic CPUID information
-	CpuId(CpuInfo, 0);
+	UE_CPUID(CpuInfo, 0);
 	MaxLeaf = CpuInfo[0];
 
 	// Basic feature flags
-	CpuId(CpuInfo, 1);
+	UE_CPUID(CpuInfo, 1);
 
 	FeatureBits |= (CpuInfo[3] & (1u << 26)) ? ECPUFeatureBits_X86::SSE2 : 0;
 	FeatureBits |= (CpuInfo[2] & (1u << 9)) ? ECPUFeatureBits_X86::SSSE3 : 0;
@@ -4350,3 +4362,5 @@ const FPlatformDriveStats* FWindowsPlatformMisc::GetDriveStats(WIDECHAR DriveLet
 	}
 	return nullptr;
 }
+
+#undef UE_CPUID
