Undo AndroidPlatformProperties use of platform memory constant

diff --git a/Engine/Source/Runtime/Core/Public/Android/AndroidPlatformProperties.h b/Engine/Source/Runtime/Core/Public/Android/AndroidPlatformProperties.h
--- a/Engine/Source/Runtime/Core/Public/Android/AndroidPlatformProperties.h
+++ b/Engine/Source/Runtime/Core/Public/Android/AndroidPlatformProperties.h
@@ -131,7 +131,7 @@ struct FAndroidPlatformProperties
 
 	static FORCEINLINE int64 GetMemoryMappingAlignment()
 	{
-		return FPlatformMemory::GetConstants().PageSize;
+		return 4096;
 	}
 
 	static FORCEINLINE bool SupportsVirtualTextureStreaming()
