Fixed the problem by avoiding getting the entry length if the path is a directory.

TMap will call FindChecked for operator[] which will check if the file path passed in is a path already recorded.
When calling FAndroidPlatformFile::IterateDirectoryCommon() with a resource path "\Engine\", the function will remove the last char '\' of the resource path before passing the path into a lambda ResourceVisitor. Because the resource path recorded in the ZipResource's entry map is with the last char '\', it will trigger the check and crash the game.

diff --git a/Engine/Source/Runtime/Core/Private/Android/AndroidPlatformFile.cpp b/Engine/Source/Runtime/Core/Private/Android/AndroidPlatformFile.cpp
--- a/Engine/Source/Runtime/Core/Private/Android/AndroidPlatformFile.cpp
+++ b/Engine/Source/Runtime/Core/Private/Android/AndroidPlatformFile.cpp
@@ -2208,7 +2208,7 @@ public:
 					FDateTime::MinValue(),						// CreationTime
 					FDateTime::MinValue(),						// AccessTime
 					FDateTime::MinValue(),						// ModificationTime
-					ZipResource.GetEntryLength(InResourceName),	// FileSize
+					IsDir ? -1 : ZipResource.GetEntryLength(InResourceName),	// FileSize
 					IsDir,										// bIsDirectory
 					true										// bIsReadOnly
 					)
