Remove check for hours in FTimecode::IsValid().  This was causing test failures in debug because of the checkSlow().

diff --git a/Engine/Source/Runtime/Core/Public/Misc/Timecode.h b/Engine/Source/Runtime/Core/Public/Misc/Timecode.h
--- a/Engine/Source/Runtime/Core/Public/Misc/Timecode.h
+++ b/Engine/Source/Runtime/Core/Public/Misc/Timecode.h
@@ -407,7 +407,6 @@ public:
 	/**
 	 * Will return true if the timecode represents a valid timecode value where
 	 *
-	 * Hours is +/- [0,23]
 	 * Minutes is +/- [0, 59]
 	 * Seconds is +/- [0, 59]
 	 * Frames [INT_MIN, INT_MAX]
@@ -420,7 +419,7 @@ public:
 		{
 		    return FMath::Abs(Val) >= MinVal && FMath::Abs(Val) <= MaxVal;
 		};
-		return InAbsRange(Hours,0,40) && InAbsRange(Minutes, 0, 59) && InAbsRange(Seconds, 0, 59) && Subframe >= 0;
+		return InAbsRange(Minutes, 0, 59) && InAbsRange(Seconds, 0, 59) && Subframe >= 0;
 	}
 public:
 
