[Scheduler]
- Fix thread-safety issue when waking up worker threads from external threads while the game thread is restarting the workers

diff --git a/Engine/Source/Runtime/Core/Private/Async/Fundamental/ReserveScheduler.cpp b/Engine/Source/Runtime/Core/Private/Async/Fundamental/ReserveScheduler.cpp
--- a/Engine/Source/Runtime/Core/Private/Async/Fundamental/ReserveScheduler.cpp
+++ b/Engine/Source/Runtime/Core/Private/Async/Fundamental/ReserveScheduler.cpp
@@ -28,8 +28,9 @@ TUniquePtr<FThread> FReserveScheduler::CreateWorker(FThread::EForkable IsForkabl
 		*FString::Printf(TEXT("Reserve Worker #%d"), WorkerId),
 		[this, ReserveEvent]
 		{
-			FSchedulerTls::ActiveScheduler = this;
-			FSchedulerTls::LocalQueue = nullptr;
+			FSchedulerTls::FTlsValues& TlsValues = FSchedulerTls::GetTlsValuesRef();
+			TlsValues.ActiveScheduler = this;
+			TlsValues.LocalQueue = nullptr;
 
 			while (true)
 			{
@@ -37,7 +38,7 @@ TUniquePtr<FThread> FReserveScheduler::CreateWorker(FThread::EForkable IsForkabl
 				ReserveEvent->SleepEvent->Wait();
 				{
 					TRACE_CPUPROFILER_EVENT_SCOPE(FReserveScheduler::BusyWaitUntil);
-					FSchedulerTls::WorkerType = ReserveEvent->bPermitBackgroundWork ? FSchedulerTls::EWorkerType::Background : FSchedulerTls::EWorkerType::Foreground;
+					TlsValues.WorkerType = ReserveEvent->bPermitBackgroundWork ? FSchedulerTls::EWorkerType::Background : FSchedulerTls::EWorkerType::Foreground;
 
 					if (ActiveWorkers.load(std::memory_order_relaxed) == 0)
 					{
@@ -47,8 +48,8 @@ TUniquePtr<FThread> FReserveScheduler::CreateWorker(FThread::EForkable IsForkabl
 					BusyWaitUntil(MoveTemp(ReserveEvent->CompletedDelegate), ReserveEvent->bPermitBackgroundWork);
 				}
 			}
-			FSchedulerTls::WorkerType = EWorkerType::None;
-			FSchedulerTls::ActiveScheduler = nullptr;
+			TlsValues.WorkerType = EWorkerType::None;
+			TlsValues.ActiveScheduler = nullptr;
 		}, 0, Priority, FThreadAffinity{ FPlatformAffinity::GetTaskGraphThreadMask(), 0 }, IsForkable
 	);
 }
@@ -59,7 +60,7 @@ bool FReserveScheduler::DoReserveWorkUntil(FConditional&& Condition)
 	{
 		WorkerEvent->CompletedDelegate = MoveTemp(Condition);
 		// become a background worker if the reserve worker is replacing a blocked background worker
-		WorkerEvent->bPermitBackgroundWork = FSchedulerTls::IsBackgroundWorker();
+		WorkerEvent->bPermitBackgroundWork = FSchedulerTls::GetTlsValuesRef().IsBackgroundWorker();
 		WorkerEvent->SleepEvent->Trigger();
 		return true;
 	}
diff --git a/Engine/Source/Runtime/Core/Private/Async/Fundamental/Scheduler.cpp b/Engine/Source/Runtime/Core/Private/Async/Fundamental/Scheduler.cpp
--- a/Engine/Source/Runtime/Core/Private/Async/Fundamental/Scheduler.cpp
+++ b/Engine/Source/Runtime/Core/Private/Async/Fundamental/Scheduler.cpp
@@ -21,11 +21,10 @@ namespace LowLevelTasks
 {
 	DEFINE_LOG_CATEGORY(LowLevelTasks);
 
-	thread_local FSchedulerTls::FLocalQueueType* FSchedulerTls::LocalQueue = nullptr;
+	UE::FMutex FSchedulerTls::Mutex;
+	TArray<FSchedulerTls::FTlsValues*>  FSchedulerTls::ThreadTlsValues;
+	thread_local FSchedulerTls::FTlsValuesHolder FSchedulerTls::TlsValuesHolder;
 	thread_local FTask* FTask::ActiveTask = nullptr;
-	thread_local FSchedulerTls* FSchedulerTls::ActiveScheduler = nullptr;
-	thread_local FSchedulerTls::EWorkerType FSchedulerTls::WorkerType = FSchedulerTls::EWorkerType::None;
-	thread_local bool FSchedulerTls::Impl::bIsStandbyWorker = false;
 	thread_local bool Private::FOversubscriptionTls::bIsOversubscriptionAllowed = false;
 
 	FORCENOINLINE bool& Private::FOversubscriptionTls::GetIsOversubscriptionAllowedRef()
@@ -33,6 +32,11 @@ namespace LowLevelTasks
 		return bIsOversubscriptionAllowed;
 	}
 
+	FORCENOINLINE FSchedulerTls::FTlsValues& FSchedulerTls::GetTlsValuesRef()
+	{
+		return *TlsValuesHolder.TlsValues;
+	}
+
 	FScheduler FScheduler::Singleton;
 
 	TUniquePtr<FThread> FScheduler::CreateWorker(uint32 WorkerId, const TCHAR* Name, bool bPermitBackgroundWork, FThread::EForkable IsForkable, Private::FWaitEvent* ExternalWorkerEvent, FSchedulerTls::FLocalQueueType* ExternalWorkerLocalQueue, EThreadPriority Priority, uint64 InAffinity)
@@ -138,7 +142,7 @@ namespace LowLevelTasks
 			{
 				GameThreadLocalQueue = MakeUnique<FSchedulerTls::FLocalQueueType>(QueueRegistry, Private::ELocalQueueType::EForeground);
 			}
-			FSchedulerTls::LocalQueue = GameThreadLocalQueue.Get();
+			FSchedulerTls::GetTlsValuesRef().LocalQueue = GameThreadLocalQueue.Get();
 
 			WorkerEvents.SetNum(MaxWorkers);
 			WorkerLocalQueues.Reserve(MaxWorkers);
@@ -317,7 +321,7 @@ namespace LowLevelTasks
 			WaitingQueue[1].FinishShutdown();
 
 			GameThreadLocalQueue.Reset();
-			FSchedulerTls::LocalQueue = nullptr;
+			FSchedulerTls::GetTlsValuesRef().LocalQueue = nullptr;
 
 			NextWorkerId = 0;
 			WorkerThreads.Reset();
@@ -342,10 +346,38 @@ namespace LowLevelTasks
 		}
 	}
 
+	bool FSchedulerTls::HasPendingWakeUp() const
+	{
+		UE::TUniqueLock ScopedLock(Mutex);
+
+#if PLATFORM_SUPPORTS_ASYMMETRIC_FENCES
+		// Heavy barrier since bPendingWakeUp is only written to with a relaxed write, we need all cores to flush their store buffer to memory
+		FPlatformMisc::AsymmetricThreadFenceHeavy();
+		constexpr std::memory_order MemoryOrder = std::memory_order_relaxed;
+#else
+		constexpr std::memory_order MemoryOrder = std::memory_order_acquire;
+#endif
+
+		for (const FTlsValues* ThreadTls : ThreadTlsValues)
+		{
+			if (ThreadTls->ActiveScheduler != this &&
+				ThreadTls->bPendingWakeUp.load(MemoryOrder))
+			{
+				return true;
+			}
+		}
+
+		return false;
+	}
+
 	void FScheduler::RestartWorkers(uint32 NumForegroundWorkers, uint32 NumBackgroundWorkers, FThread::EForkable IsForkable, EThreadPriority InWorkerPriority, EThreadPriority InBackgroundPriority, uint64 InWorkerAffinity, uint64 InBackgroundAffinity)
 	{
 		FScopeLock Lock(&WorkerThreadsCS);
 		TemporaryShutdown.store(true, std::memory_order_release);
+		while (HasPendingWakeUp())
+		{
+			FPlatformProcess::Yield();
+		}
 		const bool bDrainGlobalQueue = false;
 		StopWorkers(bDrainGlobalQueue);
 		StartWorkers(NumForegroundWorkers, NumBackgroundWorkers, IsForkable, InWorkerPriority, InBackgroundPriority, InWorkerAffinity, InBackgroundAffinity);
@@ -356,10 +388,12 @@ namespace LowLevelTasks
 	{
 		if (ActiveWorkers.load(std::memory_order_relaxed) || TemporaryShutdown.load(std::memory_order_acquire))
 		{
+			FTlsValues& LocalTlsValues = GetTlsValuesRef();
+
 			const bool bIsBackgroundTask = Task.IsBackgroundTask();
-			const bool bIsBackgroundWorker = FSchedulerTls::IsBackgroundWorker();
-			const bool bIsStandbyWorker = FSchedulerTls::IsStandbyWorker();
-			FSchedulerTls::FLocalQueueType* const CachedLocalQueue = FSchedulerTls::LocalQueue;
+			const bool bIsBackgroundWorker = LocalTlsValues.IsBackgroundWorker();
+			const bool bIsStandbyWorker = LocalTlsValues.IsStandbyWorker();
+			FSchedulerTls::FLocalQueueType* const CachedLocalQueue = LocalTlsValues.LocalQueue;
 
 			// Standby workers always enqueue to the global queue and perform wakeup
 			// as they can go to sleep whenever the oversubscription period is done
@@ -390,10 +424,49 @@ namespace LowLevelTasks
 
 			if (bWakeUpWorker)
 			{
+#if PLATFORM_SUPPORTS_ASYMMETRIC_FENCES
+				constexpr std::memory_order MemoryOrder = std::memory_order_relaxed;
+#else
+				constexpr std::memory_order MemoryOrder = std::memory_order_seq_cst;
+#endif
+				// We don't need to pay this cost for worker threads because we already manage their shutdown gracefully.
+				const bool bExternalThread = LocalTlsValues.ActiveScheduler != this || LocalTlsValues.WorkerType == EWorkerType::None;
+				if (bExternalThread)
+				{
+#if PLATFORM_SUPPORTS_ASYMMETRIC_FENCES
+					FPlatformMisc::AsymmetricThreadFenceLight();
+#endif
+
+					LocalTlsValues.bPendingWakeUp.store(true, MemoryOrder);
+
+#if PLATFORM_SUPPORTS_ASYMMETRIC_FENCES
+					FPlatformMisc::AsymmetricThreadFenceLight();
+#endif
+
+					if (TemporaryShutdown.load(std::memory_order_acquire))
+					{
+						LocalTlsValues.bPendingWakeUp.store(false, MemoryOrder);
+						return;
+					}
+				}
+
 				if (!WakeUpWorker(bIsBackgroundTask) && !bIsBackgroundTask)
 				{
 					WakeUpWorker(true);
 				}
+
+				if (bExternalThread)
+				{
+#if PLATFORM_SUPPORTS_ASYMMETRIC_FENCES
+					FPlatformMisc::AsymmetricThreadFenceLight();
+#endif
+
+					LocalTlsValues.bPendingWakeUp.store(false, MemoryOrder);
+
+#if PLATFORM_SUPPORTS_ASYMMETRIC_FENCES
+					FPlatformMisc::AsymmetricThreadFenceLight();
+#endif
+				}
 			}
 		}
 		else
@@ -411,7 +484,7 @@ namespace LowLevelTasks
 
 	void FScheduler::IncrementOversubscription()
 	{
-		FSchedulerTls::EWorkerType LocalWorkerType = WorkerType;
+		FSchedulerTls::EWorkerType LocalWorkerType = FSchedulerTls::GetTlsValuesRef().WorkerType;
 
 		if (LocalWorkerType != EWorkerType::None)
 		{
@@ -426,7 +499,7 @@ namespace LowLevelTasks
 
 	void FScheduler::DecrementOversubscription()
 	{
-		FSchedulerTls::EWorkerType LocalWorkerType = WorkerType;
+		FSchedulerTls::EWorkerType LocalWorkerType = FSchedulerTls::GetTlsValuesRef().WorkerType;
 
 		if (LocalWorkerType != EWorkerType::None)
 		{
@@ -444,7 +517,8 @@ namespace LowLevelTasks
 
 	bool FSchedulerTls::IsWorkerThread() const
 	{
-		return WorkerType != FSchedulerTls::EWorkerType::None && ActiveScheduler == this;
+		FTlsValues& LocalTlsValues = FSchedulerTls::GetTlsValuesRef();
+		return LocalTlsValues.WorkerType != FSchedulerTls::EWorkerType::None && LocalTlsValues.ActiveScheduler == this;
 	}
 
 	bool FSchedulerTls::IsBusyWaiting()
@@ -471,7 +545,7 @@ namespace LowLevelTasks
 					// a potential wakeup we might have consumed during the cancellation.
 					if (WaitingQueue[bPermitBackgroundWork].CancelWait(WaitEvent))
 					{
-						if (!WakeUpWorker(bPermitBackgroundWork) && !FSchedulerTls::IsBackgroundWorker())
+						if (!WakeUpWorker(bPermitBackgroundWork) && !GetTlsValuesRef().IsBackgroundWorker())
 						{
 							WakeUpWorker(!bPermitBackgroundWork);
 						}
@@ -577,17 +651,19 @@ namespace LowLevelTasks
 
 	void FScheduler::WorkerMain(Private::FWaitEvent* WorkerEvent, FSchedulerTls::FLocalQueueType* WorkerLocalQueue, uint32 WaitCycles, bool bPermitBackgroundWork)
 	{
-		checkSlow(FSchedulerTls::LocalQueue == nullptr);
+		FTlsValues& LocalTlsValues = GetTlsValuesRef();
+
+		checkSlow(LocalTlsValues.LocalQueue == nullptr);
 		checkSlow(WorkerLocalQueue != nullptr);
 		checkSlow(WorkerEvent != nullptr);
 
 		FTaskTagScope WorkerScope(ETaskTag::EWorkerThread);
-		FSchedulerTls::ActiveScheduler = this;
+		LocalTlsValues.ActiveScheduler = this;
 
 		FMemory::SetupTLSCachesOnCurrentThread();
-		FSchedulerTls::WorkerType = bPermitBackgroundWork ? FSchedulerTls::EWorkerType::Background : FSchedulerTls::EWorkerType::Foreground;
-		FSchedulerTls::SetStandbyWorker(WorkerEvent->bIsStandby);
-		FSchedulerTls::LocalQueue = WorkerLocalQueue;
+		LocalTlsValues.WorkerType = bPermitBackgroundWork ? FSchedulerTls::EWorkerType::Background : FSchedulerTls::EWorkerType::Foreground;
+		LocalTlsValues.SetStandbyWorker(WorkerEvent->bIsStandby);
+		LocalTlsValues.LocalQueue = WorkerLocalQueue;
 
 		{
 			Private::FOversubscriptionAllowedScope _(true);
@@ -602,10 +678,10 @@ namespace LowLevelTasks
 			}
 		}
 
-		FSchedulerTls::LocalQueue = nullptr;
-		FSchedulerTls::ActiveScheduler = nullptr;
-		FSchedulerTls::SetStandbyWorker(false);
-		FSchedulerTls::WorkerType = FSchedulerTls::EWorkerType::None;
+		LocalTlsValues.LocalQueue = nullptr;
+		LocalTlsValues.ActiveScheduler = nullptr;
+		LocalTlsValues.SetStandbyWorker(false);
+		LocalTlsValues.WorkerType = FSchedulerTls::EWorkerType::None;
 		FMemory::ClearAndDisableTLSCachesOnCurrentThread();
 	}
 
diff --git a/Engine/Source/Runtime/Core/Public/Async/Fundamental/Scheduler.h b/Engine/Source/Runtime/Core/Public/Async/Fundamental/Scheduler.h
--- a/Engine/Source/Runtime/Core/Public/Async/Fundamental/Scheduler.h
+++ b/Engine/Source/Runtime/Core/Public/Async/Fundamental/Scheduler.h
@@ -5,11 +5,14 @@
 #include "Async/Fundamental/TaskShared.h"
 #include "Async/Fundamental/TaskDelegate.h"
 #include "Async/Fundamental/WaitingQueue.h"
+#include "Async/Mutex.h"
+#include "Async/UniqueLock.h"
 #include "Containers/Array.h"
 #include "Containers/ArrayView.h"
 #include "Containers/ContainerAllocationPolicies.h"
 #include "HAL/CriticalSection.h"
 #include "HAL/Event.h"
+#include "HAL/LowLevelMemTracker.h"
 #include "HAL/PlatformAffinity.h"
 #include "HAL/PlatformProcess.h"
 #include "HAL/Thread.h"
@@ -121,18 +124,60 @@ namespace LowLevelTasks
 			Foreground,
 		};
 
-		static thread_local FSchedulerTls* ActiveScheduler;
-		static thread_local FLocalQueueType* LocalQueue;
-		static thread_local EWorkerType WorkerType;
-	
-	private:
-		// Hide the member to avoid clashing with local variables
-		// already named bIsStandbyWorker.
-		struct Impl
+		struct FTlsValues;
+		static UE::FMutex Mutex;
+		static TArray<FTlsValues*> ThreadTlsValues;
+
+		struct alignas(PLATFORM_CACHE_LINE_SIZE) FTlsValues
 		{
-			static thread_local bool bIsStandbyWorker;
+			FSchedulerTls* ActiveScheduler = nullptr;
+			FLocalQueueType* LocalQueue = nullptr;
+			EWorkerType WorkerType = EWorkerType::None;
+			std::atomic<bool> bPendingWakeUp = false;
+			bool        bIsStandbyWorker = false;
+
+			inline bool IsBackgroundWorker()
+			{
+				return WorkerType == EWorkerType::Background;
+			}
+
+			inline bool IsStandbyWorker()
+			{
+				return bIsStandbyWorker;
+			}
+
+			inline void SetStandbyWorker(bool bInIsStandbyWorker)
+			{
+				bIsStandbyWorker = bInIsStandbyWorker;
+			}
+		};
+
+		struct FTlsValuesHolder
+		{
+			FTlsValuesHolder()
+			{
+				LLM_SCOPE_BYNAME(TEXT("Scheduler/TlsValues"));
+				TlsValues = new FTlsValues();
+
+				UE::TUniqueLock ScopedLock(Mutex);
+				ThreadTlsValues.Add(TlsValues);
+			}
+			~FTlsValuesHolder()
+			{
+				{
+					UE::TUniqueLock ScopedLock(Mutex);
+					ThreadTlsValues.Remove(TlsValues);
+				}
+				delete TlsValues;
+				TlsValues = nullptr;
+			}
+			
+			FTlsValues* TlsValues = nullptr;
 		};
 
+	private:
+		static thread_local FTlsValuesHolder TlsValuesHolder;
+
 	public:
 		CORE_API bool IsWorkerThread() const;
 
@@ -140,20 +185,9 @@ namespace LowLevelTasks
 		CORE_API static bool IsBusyWaiting();
 
 	protected:
-		inline static bool IsBackgroundWorker()
-		{
-			return WorkerType == EWorkerType::Background;
-		}
+		CORE_API bool HasPendingWakeUp() const;
 
-		inline static bool IsStandbyWorker()
-		{
-			return Impl::bIsStandbyWorker;
-		}
-
-		inline static void SetStandbyWorker(bool bInIsStandbyWorker)
-		{
-			Impl::bIsStandbyWorker = bInIsStandbyWorker;
-		}
+		static FTlsValues& GetTlsValuesRef();
 	};
 
 	class FScheduler final : public FSchedulerTls
diff --git a/Engine/Source/Runtime/Core/Tests/Tasks/TasksTest.cpp b/Engine/Source/Runtime/Core/Tests/Tasks/TasksTest.cpp
--- a/Engine/Source/Runtime/Core/Tests/Tasks/TasksTest.cpp
+++ b/Engine/Source/Runtime/Core/Tests/Tasks/TasksTest.cpp
@@ -1738,9 +1738,9 @@ namespace UE { namespace TasksTests
 	}
 
 	// StressTest RestartWorkers while having task performing oversubscription
-	TEST_CASE_NAMED(FTaskGraphOversubscribeDuringRestartWorkers, "System::Core::Async::Tasks::OversubscribeDuringRestartWorkers", "[.][ApplicationContextMask][EngineFilter]")
+	TEST_CASE_NAMED(FTaskGraphOversubscribeDuringRestartWorkers, "System::Core::Async::Tasks::RestartWorkersAndOversubscription", "[.][ApplicationContextMask][EngineFilter]")
 	{
-		for (int32 Index = 0; Index < 10; ++Index)
+		for (int32 Index = 0; Index < 10000; ++Index)
 		{
 			FTask Task = Launch(TEXT("Task"),
 				[&]()
@@ -1761,6 +1761,37 @@ namespace UE { namespace TasksTests
 		LowLevelTasks::FScheduler::Get().RestartWorkers();
 	}
 
+	// StressTest RestartWorkers while having a thread trying to launch tasks on workers
+	TEST_CASE_NAMED(FTaskGraphLaunchFromThreadDuringRestartWorkers, "System::Core::Async::Tasks::RestartWorkersAndExternalThreads", "[.][ApplicationContextMask][EngineFilter]")
+	{
+		std::atomic<bool> bExit = false;
+		UE::FManualResetEvent Done;
+		Async(EAsyncExecution::Thread,
+			[&Done, &bExit]()
+			{
+				while (!bExit)
+				{
+					Launch(TEXT("Task"), [&]() {}).Wait();
+				}
+
+				Done.Notify();
+			}
+		);
+
+		for (int32 Index = 0; Index < 1000; ++Index)
+		{
+			// For the repro to work we need to launch tasks to try to start new workers while
+			// we're shutting down and restarting workers.
+			LowLevelTasks::FScheduler::Get().RestartWorkers(0, 4);
+		}
+
+		bExit = true;
+		Done.Wait();
+
+		// Restore original worker count
+		LowLevelTasks::FScheduler::Get().RestartWorkers();
+	}
+
 #if PLATFORM_SUPPORTS_ASYMMETRIC_FENCES
 	TEST_CASE_NAMED(FAsymmetricThreadFence, "System::Core::Async::AsymmetricThreadFence", "[.][ApplicationContextMask][EngineFilter]")
 	{
