Serialization: Fixed a structured archive assert in debug configurations

Asserting that array index matches array count is incorrect when saving a static array with delta serialization because values identical to the default are skipped and the tracked index in the structured archive is not incremented.

diff --git a/Engine/Source/Runtime/Core/Private/Serialization/StructuredArchive.cpp b/Engine/Source/Runtime/Core/Private/Serialization/StructuredArchive.cpp
--- a/Engine/Source/Runtime/Core/Private/Serialization/StructuredArchive.cpp
+++ b/Engine/Source/Runtime/Core/Private/Serialization/StructuredArchive.cpp
@@ -159,7 +159,7 @@ void FStructuredArchive::SetScope(UE::StructuredArchive::Private::FSlotPosition
 				break;
 			case UE::StructuredArchive::Private::EElementType::Array:
 #if DO_STRUCTURED_ARCHIVE_CONTAINER_CHECKS
-				checkf(GetUnderlyingArchive().IsLoading() || CurrentContainer.Top()->Index == CurrentContainer.Top()->Count, TEXT("Incorrect number of elements serialized in array"));
+				checkf(GetUnderlyingArchive().IsLoading() || CurrentContainer.Top()->Index <= CurrentContainer.Top()->Count, TEXT("Incorrect number of elements serialized in array"));
 #endif
 				Formatter.LeaveArray();
 #if DO_STRUCTURED_ARCHIVE_CONTAINER_CHECKS
