Fixed a problem which lead to incomplete LLM CSV logs on some platforms, in such cases now we flush also on suspend and terminate

#rnx
#tests Tested with -llm -llmcsv on affected platforms

diff --git a/Engine/Source/Runtime/Core/Private/HAL/LowLevelMemTracker.cpp b/Engine/Source/Runtime/Core/Private/HAL/LowLevelMemTracker.cpp
--- a/Engine/Source/Runtime/Core/Private/HAL/LowLevelMemTracker.cpp
+++ b/Engine/Source/Runtime/Core/Private/HAL/LowLevelMemTracker.cpp
@@ -1291,7 +1291,7 @@ public:
 
 	void OnPreFork();
 
-	void FlushOnCrash();
+	void Flush(bool IsOnCrash);
 
 private:
 	void Write(FStringView Text);
@@ -1310,7 +1310,7 @@ private:
 	double LastWriteTime;
 	int32 WriteCount;
 	ELLMTracker Tracker;
-	bool bRegisteredCrashDelegate;
+	bool bRegisteredFlushDelegate;
 };
 
 /** Outputs the LLM tags and sizes to TraceLog events. */
@@ -5991,7 +5991,7 @@ FLLMCsvWriter::FLLMCsvWriter()
 	: Archive(nullptr)
 	, LastWriteTime(FPlatformTime::Seconds())
 	, WriteCount(0)
-	, bRegisteredCrashDelegate(false)
+	, bRegisteredFlushDelegate(false)
 {
 }
 
@@ -6016,11 +6016,15 @@ void FLLMCsvWriter::OnPreFork()
 	}
 }
 
-void FLLMCsvWriter::FlushOnCrash()
+void FLLMCsvWriter::Flush(bool IsOnCrash)
 {
 	if (Archive)
 	{
-		FPlatformMisc::LowLevelOutputDebugStringf(TEXT("Flushing LLM CSV on crash\n"));
+		if (IsOnCrash)
+		{
+			FPlatformMisc::LowLevelOutputDebugStringf(TEXT("Flushing LLM CSV on crash\n"));
+		}
+
 		Archive->Flush();
 	}
 }
@@ -6232,13 +6236,15 @@ void FLLMCsvWriter::AddRow(FLowLevelMemTracker& LLMRef, const FTrackerTagSizeMap
 	{
 		Archive->Flush();
 	}
-	else if (!bRegisteredCrashDelegate)
+	else if (!bRegisteredFlushDelegate)
 	{
-		// If we're not flushing every row, lazily register crash delegates to ensure we flush on a crash
+		// If we're not flushing every row, lazily register flush delegates to ensure we flush on a crash
 		// Note: we intentionally leak this since we can't clean it up safely 
-		FCoreDelegates::OnHandleSystemError.AddRaw(this, &FLLMCsvWriter::FlushOnCrash);
-		FCoreDelegates::GetOutOfMemoryDelegate().AddRaw(this, &FLLMCsvWriter::FlushOnCrash);
-		bRegisteredCrashDelegate = true;
+		FCoreDelegates::OnHandleSystemError.AddRaw(this, &FLLMCsvWriter::Flush, true);
+		FCoreDelegates::GetOutOfMemoryDelegate().AddRaw(this, &FLLMCsvWriter::Flush, true);
+		FCoreDelegates::ApplicationWillEnterBackgroundDelegate.AddRaw(this, &FLLMCsvWriter::Flush, false); // When you close the app on some platforms it is suspended rather than terminated
+		FCoreDelegates::GetApplicationWillTerminateDelegate().AddRaw(this, &FLLMCsvWriter::Flush, false);
+		bRegisteredFlushDelegate = true;
 	}
 
 }
