Catch missing specifiers in FormatStringSan.

diff --git a/Engine/Plugins/Experimental/RemoteSession/Source/RemoteSession/Private/RemoteSessionRole.cpp b/Engine/Plugins/Experimental/RemoteSession/Source/RemoteSession/Private/RemoteSessionRole.cpp
--- a/Engine/Plugins/Experimental/RemoteSession/Source/RemoteSession/Private/RemoteSessionRole.cpp
+++ b/Engine/Plugins/Experimental/RemoteSession/Source/RemoteSession/Private/RemoteSessionRole.cpp
@@ -631,7 +631,7 @@ void FRemoteSessionRole::OnReceiveChangeChannel(IBackChannelPacket& Message)
 					TSharedPtr<IRemoteSessionChannel> Channel = GetChannel(*ChannelName);
 					check(Channel.IsValid());
 					// Since we received a request the channel should exist on the other end and we can send to it now.
-					UE_LOG(LogRemoteSession, Log, TEXT("OnReceiveChangeChannel: Broadcasting channel creation with mode %s"), *ChannelName, *ChannelMode);
+					UE_LOG(LogRemoteSession, Log, TEXT("OnReceiveChangeChannel: Broadcasting channel %s creation with mode %s"), *ChannelName, *ChannelMode);
 					ChannelChangeDelegate.Broadcast(this, Channel, ERemoteSessionChannelChange::Created);
 				}
 			});	
diff --git a/Engine/Plugins/FX/Niagara/Source/NiagaraEditor/Private/NiagaraGraph.cpp b/Engine/Plugins/FX/Niagara/Source/NiagaraEditor/Private/NiagaraGraph.cpp
--- a/Engine/Plugins/FX/Niagara/Source/NiagaraEditor/Private/NiagaraGraph.cpp
+++ b/Engine/Plugins/FX/Niagara/Source/NiagaraEditor/Private/NiagaraGraph.cpp
@@ -1849,7 +1849,7 @@ UNiagaraGraph* UNiagaraGraph::CreateCompilationCopy(const TArray<ENiagaraScriptU
 					UEdGraphNode* NodeB = Pin->LinkedTo[j]->GetOwningNode();
 					if (Node && NodeB && Node->GetGraph() != NodeB->GetGraph())
 					{
-						UE_LOG(LogNiagaraEditor, Error, TEXT(""), *NodeB->GetPathName());
+						UE_LOG(LogNiagaraEditor, Error, TEXT("%s"), *NodeB->GetPathName());
 						const FText WarningTemplate = LOCTEXT("DuplicationErrorToast", "Duplication Error! {0} was touched when cloning for compilation. Please report this issue to the Niagara team! ");
 						FText MismatchWarning = FText::Format(
 							WarningTemplate,
diff --git a/Engine/Plugins/PCG/Source/PCG/Private/Utils/PCGGraphExecutionLogging.cpp b/Engine/Plugins/PCG/Source/PCG/Private/Utils/PCGGraphExecutionLogging.cpp
--- a/Engine/Plugins/PCG/Source/PCG/Private/Utils/PCGGraphExecutionLogging.cpp
+++ b/Engine/Plugins/PCG/Source/PCG/Private/Utils/PCGGraphExecutionLogging.cpp
@@ -455,7 +455,7 @@ namespace PCGGraphExecutionLogging
 
 		UE_LOG(LogPCG, Log, TEXT("[GRIDLINKING] [%s] RETRIEVE: WOKEN BY Component=%s"),
 			*PCGLog::GetComponentOwnerName(InContext->SourceComponent.Get(), /*bUseLabel=*/true),
-			*PCGLog::GetComponentOwnerName(InWokenBy), /*bUseLabel=*/true);
+			*PCGLog::GetComponentOwnerName(InWokenBy, /*bUseLabel=*/true));
 #endif // !(UE_BUILD_SHIPPING || UE_BUILD_TEST) || USE_LOGGING_IN_SHIPPING
 	}
 
diff --git a/Engine/Plugins/Runtime/DataRegistry/Source/DataRegistry/Private/DataRegistrySource_DataTable.cpp b/Engine/Plugins/Runtime/DataRegistry/Source/DataRegistry/Private/DataRegistrySource_DataTable.cpp
--- a/Engine/Plugins/Runtime/DataRegistry/Source/DataRegistry/Private/DataRegistrySource_DataTable.cpp
+++ b/Engine/Plugins/Runtime/DataRegistry/Source/DataRegistry/Private/DataRegistrySource_DataTable.cpp
@@ -78,7 +78,7 @@ void UDataRegistrySource_DataTable::SetCachedTable(bool bForceLoad /*= false*/)
 		}
 		else if (!RowStruct->IsChildOf(ItemStruct))
 		{
-			UE_LOG(LogDataRegistry, Error, TEXT("Cannot initialize DataRegistry source %s, Table %s type does not match %s"), *GetPathName(), *FoundTable->GetPathName(), *RowStruct->GetName(), *ItemStruct->GetName());
+			UE_LOG(LogDataRegistry, Error, TEXT("Cannot initialize DataRegistry source %s, Table %s row type %s does not match item type %s"), *GetPathName(), *FoundTable->GetPathName(), *RowStruct->GetName(), *ItemStruct->GetName());
 			bInvalidSourceTable = true;
 		}
 		else
diff --git a/Engine/Plugins/Runtime/GameFeatures/Source/GameFeaturesEditor/Private/GameFeatureActionConvertContentBundleWorldPartitionBuilder.cpp b/Engine/Plugins/Runtime/GameFeatures/Source/GameFeaturesEditor/Private/GameFeatureActionConvertContentBundleWorldPartitionBuilder.cpp
--- a/Engine/Plugins/Runtime/GameFeatures/Source/GameFeaturesEditor/Private/GameFeatureActionConvertContentBundleWorldPartitionBuilder.cpp
+++ b/Engine/Plugins/Runtime/GameFeatures/Source/GameFeaturesEditor/Private/GameFeatureActionConvertContentBundleWorldPartitionBuilder.cpp
@@ -260,7 +260,7 @@ bool UGameFeatureActionConvertContentBundleWorldPartitionBuilder::RunInternal(UW
 			FText FailureReason;
 			if (!FExternalDataLayerHelper::MoveActorsToExternalDataLayer({ Actor }, ExternalDataLayerInstance, &FailureReason))
 			{
-				UE_LOG(LogConvertContentBundleBuilder, Error, TEXT("Can't create package for actor %s. %s"), *ExternalDataLayerAsset->GetName(), *Actor->GetActorNameOrLabel(), *FailureReason.ToString());
+				UE_LOG(LogConvertContentBundleBuilder, Error, TEXT("Can't create package for actor %s. %s"), *Actor->GetActorNameOrLabel(), *FailureReason.ToString());
 				bContentBundleActorConversionSuccess = false;
 				break;
 			}
diff --git a/Engine/Source/Editor/BlueprintGraph/Private/EdGraphSchema_K2.cpp b/Engine/Source/Editor/BlueprintGraph/Private/EdGraphSchema_K2.cpp
--- a/Engine/Source/Editor/BlueprintGraph/Private/EdGraphSchema_K2.cpp
+++ b/Engine/Source/Editor/BlueprintGraph/Private/EdGraphSchema_K2.cpp
@@ -6262,7 +6262,7 @@ struct FBackwardCompatibilityConversionHelper
 					}
 					else
 					{
-						UE_LOG(LogBlueprint, Warning, TEXT("BackwardCompatibilityNodeConversion Error 'missing old pin' in blueprint: %s"),
+						UE_LOG(LogBlueprint, Warning, TEXT("BackwardCompatibilityNodeConversion Error 'missing old pin' in blueprint: %s, pin: %s"),
 							Blueprint ? *Blueprint->GetName() : TEXT("Unknown"),
 							Pin ? *Pin->PinName.ToString() : TEXT("Unknown"));
 					}
@@ -6273,7 +6273,7 @@ struct FBackwardCompatibilityConversionHelper
 			{
 				if (!OldPins.Contains(Pin))
 				{
-					UE_LOG(LogBlueprint, Warning, TEXT("BackwardCompatibilityNodeConversion Error 'missing new pin' in blueprint: %s"),
+					UE_LOG(LogBlueprint, Warning, TEXT("BackwardCompatibilityNodeConversion Error 'missing new pin' in blueprint: %s, pin: %s"),
 						Blueprint ? *Blueprint->GetName() : TEXT("Unknown"),
 						Pin ? *Pin->PinName.ToString() : TEXT("Unknown"));
 				}
diff --git a/Engine/Source/Editor/BlueprintGraph/Private/K2Node_Composite.cpp b/Engine/Source/Editor/BlueprintGraph/Private/K2Node_Composite.cpp
--- a/Engine/Source/Editor/BlueprintGraph/Private/K2Node_Composite.cpp
+++ b/Engine/Source/Editor/BlueprintGraph/Private/K2Node_Composite.cpp
@@ -145,7 +145,7 @@ void UK2Node_Composite::PostPasteNode()
 				}
 				else
 				{
-					ensureMsgf(false, TEXT("%s"), *LOCTEXT("UnexpectedTunnelNode", "Unexpected tunnel node '%s' in cloned graph '%s' (both I/O or neither)").ToString(), *Tunnel->GetName(), *GetName());
+					ensureMsgf(false, TEXT("Unexpected tunnel node '%s' in cloned graph '%s' (both I/O or neither)"), *Tunnel->GetName(), *GetName());
 				}
 			}
 		}
@@ -198,7 +198,7 @@ void UK2Node_Composite::FixupInputAndOutputSink()
 				}
 				else
 				{
-					ensureMsgf(false, TEXT("%s"), *LOCTEXT("UnexpectedTunnelNode", "Unexpected tunnel node '%s' in cloned graph '%s' (both I/O or neither)").ToString(), *Tunnel->GetName(), *GetName());
+					ensureMsgf(false, TEXT("Unexpected tunnel node '%s' in cloned graph '%s' (both I/O or neither)"), *Tunnel->GetName(), *GetName());
 				}
 			}
 		}
diff --git a/Engine/Source/Editor/UnrealEd/Private/Subsystems/EditorAssetSubsystem.cpp b/Engine/Source/Editor/UnrealEd/Private/Subsystems/EditorAssetSubsystem.cpp
--- a/Engine/Source/Editor/UnrealEd/Private/Subsystems/EditorAssetSubsystem.cpp
+++ b/Engine/Source/Editor/UnrealEd/Private/Subsystems/EditorAssetSubsystem.cpp
@@ -1467,7 +1467,7 @@ bool UEditorAssetSubsystem::MakeDirectory(const FString& DirectoryPath)
 	FString DirectoryPackageName = EditorScriptingHelpers::ConvertAnyPathToLongPackagePath(DirectoryPath, FailureReason);
 	if (DirectoryPackageName.IsEmpty())
 	{
-		UE_LOG(LogEditorAssetSubsystem, Error, TEXT("MakeDirectory failed: Could not convert the path. %s"), *DirectoryPath, *FailureReason);
+		UE_LOG(LogEditorAssetSubsystem, Error, TEXT("MakeDirectory failed: Could not convert the path '%s'. %s"), *DirectoryPath, *FailureReason);
 		return false;
 	}
 
diff --git a/Engine/Source/Runtime/Core/Public/String/FormatStringSan.h b/Engine/Source/Runtime/Core/Public/String/FormatStringSan.h
--- a/Engine/Source/Runtime/Core/Public/String/FormatStringSan.h
+++ b/Engine/Source/Runtime/Core/Public/String/FormatStringSan.h
@@ -133,7 +133,7 @@ namespace UE::Core::Private
 
 				if (Fmt[0] == (CharType)'\0' && !bInsideFormatSpec)
 				{
-					return {EFormatStringSanStatus::Ok, 0};
+					return {EFormatStringSanStatus::NotEnoughSpecifiers, CurArgPos};
 				}
 
 				if (!bInsideFormatSpec)
diff --git a/Engine/Source/Runtime/Core/Public/String/FormatStringSanErrors.inl b/Engine/Source/Runtime/Core/Public/String/FormatStringSanErrors.inl
--- a/Engine/Source/Runtime/Core/Public/String/FormatStringSanErrors.inl
+++ b/Engine/Source/Runtime/Core/Public/String/FormatStringSanErrors.inl
@@ -37,6 +37,7 @@ X(PNeedsPointerArg, "'%p' expects a pointer.")
 X(FNeedsFloatOrDoubleArg, "'%f' expects `float` or `double`.")
 X(InvalidFormatSpec, "unsupported '%' format specifier. (to print a percent sign, write '%%'.)")
 X(NotEnoughArguments, "not enough arguments provided to format string.")
+X(NotEnoughSpecifiers, "missing format specifier (Is there a missing '%' somewhere?)")
 X(DNeedsIntegerArg, "'%d' expects integral arg (eg. `char`, `int`, `long`, etc.)")
 X(ZNeedsIntegerArg, "'%z' expects integral arg (eg. `char`, `int`, `long`, etc.)")
 X(ZNeedsIntegerSpec, "'%z' must have a suffix (eg. '%zd', '%zu', etc.)")
diff --git a/Engine/Source/Runtime/Core/Tests/String/FormatStringSanTest.cpp b/Engine/Source/Runtime/Core/Tests/String/FormatStringSanTest.cpp
--- a/Engine/Source/Runtime/Core/Tests/String/FormatStringSanTest.cpp
+++ b/Engine/Source/Runtime/Core/Tests/String/FormatStringSanTest.cpp
@@ -260,7 +260,6 @@ TEST_CASE_NAMED(FFormatStringValidatorTest, "System::Core::String::FormatStringS
 		STATIC_CHECK(UE_CHECK_FORMAT_STRING_ERR(E::PNeedsPointerArg, TEXT("Hello %p"), 42));
 		STATIC_CHECK(UE_CHECK_FORMAT_STRING_ERR(E::FNeedsFloatOrDoubleArg, TEXT("Hello %f"), 42));
 		STATIC_CHECK(UE_CHECK_FORMAT_STRING_ERR(E::InvalidFormatSpec, TEXT("Hello %k"), 42));
-		STATIC_CHECK(UE_CHECK_FORMAT_STRING_ERR(E::NotEnoughArguments, TEXT("Hello %s")));
 		STATIC_CHECK(UE_CHECK_FORMAT_STRING_ERR(E::DNeedsIntegerArg, TEXT("Hello %d"), 42.0));
 		STATIC_CHECK(UE_CHECK_FORMAT_STRING_ERR(E::ZNeedsIntegerSpec, TEXT("Hello %z test"), 42));
 		STATIC_CHECK(UE_CHECK_FORMAT_STRING_ERR(E::ZNeedsIntegerArg, TEXT("Hello %zu"), "hi"));
@@ -336,6 +335,35 @@ TEST_CASE_NAMED(FFormatStringValidatorTest, "System::Core::String::FormatStringS
 			STATIC_CHECK(UE_CHECK_FORMAT_STRING_ERR(E::Ok, TEXT("%d"), EnumAsByteParam));
 		}
 	}
+
+	SECTION("NumberOfArguments")
+	{
+		STATIC_CHECK(UE_CHECK_FORMAT_STRING_ERR(E::Ok,                  TEXT("")));
+		STATIC_CHECK(UE_CHECK_FORMAT_STRING_ERR(E::NotEnoughSpecifiers, TEXT(""), 1));
+		STATIC_CHECK(UE_CHECK_FORMAT_STRING_ERR(E::NotEnoughSpecifiers, TEXT(""), 1, 2));
+		STATIC_CHECK(UE_CHECK_FORMAT_STRING_ERR(E::NotEnoughSpecifiers, TEXT(""), 1, 2, 3));
+
+		STATIC_CHECK(UE_CHECK_FORMAT_STRING_ERR(E::NotEnoughArguments,  TEXT("%d")));
+		STATIC_CHECK(UE_CHECK_FORMAT_STRING_ERR(E::Ok,                  TEXT("%d"), 1));
+		STATIC_CHECK(UE_CHECK_FORMAT_STRING_ERR(E::NotEnoughSpecifiers, TEXT("%d"), 1, 2));
+		STATIC_CHECK(UE_CHECK_FORMAT_STRING_ERR(E::NotEnoughSpecifiers, TEXT("%d"), 1, 2, 3));
+		STATIC_CHECK(UE_CHECK_FORMAT_STRING_ERR(E::NotEnoughSpecifiers, TEXT("%d"), 1, 2, 3, 4));
+
+		STATIC_CHECK(UE_CHECK_FORMAT_STRING_ERR(E::NotEnoughArguments,  TEXT("%d %d")));
+		STATIC_CHECK(UE_CHECK_FORMAT_STRING_ERR(E::NotEnoughArguments,  TEXT("%d %d"), 1));
+		STATIC_CHECK(UE_CHECK_FORMAT_STRING_ERR(E::Ok,                  TEXT("%d %d"), 1, 2));
+		STATIC_CHECK(UE_CHECK_FORMAT_STRING_ERR(E::NotEnoughSpecifiers, TEXT("%d %d"), 1, 2, 3));
+		STATIC_CHECK(UE_CHECK_FORMAT_STRING_ERR(E::NotEnoughSpecifiers, TEXT("%d %d"), 1, 2, 3, 4));
+		STATIC_CHECK(UE_CHECK_FORMAT_STRING_ERR(E::NotEnoughSpecifiers, TEXT("%d %d"), 1, 2, 3, 4, 5));
+
+		STATIC_CHECK(UE_CHECK_FORMAT_STRING_ERR(E::NotEnoughArguments,  TEXT("%d %d %d")));
+		STATIC_CHECK(UE_CHECK_FORMAT_STRING_ERR(E::NotEnoughArguments,  TEXT("%d %d %d"), 1));
+		STATIC_CHECK(UE_CHECK_FORMAT_STRING_ERR(E::NotEnoughArguments,  TEXT("%d %d %d"), 1, 2));
+		STATIC_CHECK(UE_CHECK_FORMAT_STRING_ERR(E::Ok,                  TEXT("%d %d %d"), 1, 2, 3));
+		STATIC_CHECK(UE_CHECK_FORMAT_STRING_ERR(E::NotEnoughSpecifiers, TEXT("%d %d %d"), 1, 2, 3, 4));
+		STATIC_CHECK(UE_CHECK_FORMAT_STRING_ERR(E::NotEnoughSpecifiers, TEXT("%d %d %d"), 1, 2, 3, 4, 5));
+		STATIC_CHECK(UE_CHECK_FORMAT_STRING_ERR(E::NotEnoughSpecifiers, TEXT("%d %d %d"), 1, 2, 3, 4, 5, 6));
+	}
 }
 
 TEST_CASE_NAMED(FFormatStringConstStringValidationTest, "System::Core::String::FormatStringSan::ConstString", "[Core][String][FormatStringSan]")
diff --git a/Engine/Source/Runtime/D3D12RHI/Private/D3D12CommandContext.cpp b/Engine/Source/Runtime/D3D12RHI/Private/D3D12CommandContext.cpp
--- a/Engine/Source/Runtime/D3D12RHI/Private/D3D12CommandContext.cpp
+++ b/Engine/Source/Runtime/D3D12RHI/Private/D3D12CommandContext.cpp
@@ -1048,7 +1048,7 @@ void FD3D12CommandContextRedirector::RHITransferResources(TConstArrayView<FTrans
 				ensureMsgf(
 					Param.Min.X >= 0 && Param.Min.Y >= 0 && Param.Min.Z >= 0 &&
 					Param.Max.X >= 0 && Param.Max.Y >= 0 && Param.Max.Z >= 0,
-					TEXT("Invalid rect for texture transfer: %i, %i, %i, %i"), Param.Min.X, Param.Min.Y, Param.Min.Z, Param.Max.X, Param.Max.Y, Param.Max.Z);
+					TEXT("Invalid rect for texture transfer: %i, %i, %i, %i, %i, %i"), Param.Min.X, Param.Min.Y, Param.Min.Z, Param.Max.X, Param.Max.Y, Param.Max.Z);
 
 				D3D12_BOX Box = { (UINT)Param.Min.X, (UINT)Param.Min.Y, (UINT)Param.Min.Z, (UINT)Param.Max.X, (UINT)Param.Max.Y, (UINT)Param.Max.Z };
 
@@ -1202,7 +1202,7 @@ void FD3D12CommandContextRedirector::RHICrossGPUTransfer(TConstArrayView<FTransf
 				ensureMsgf(
 					Param.Min.X >= 0 && Param.Min.Y >= 0 && Param.Min.Z >= 0 &&
 					Param.Max.X >= 0 && Param.Max.Y >= 0 && Param.Max.Z >= 0,
-					TEXT("Invalid rect for texture transfer: %i, %i, %i, %i"), Param.Min.X, Param.Min.Y, Param.Min.Z, Param.Max.X, Param.Max.Y, Param.Max.Z);
+					TEXT("Invalid rect for texture transfer: %i, %i, %i, %i, %i, %i"), Param.Min.X, Param.Min.Y, Param.Min.Z, Param.Max.X, Param.Max.Y, Param.Max.Z);
 
 				D3D12_BOX Box = { (UINT)Param.Min.X, (UINT)Param.Min.Y, (UINT)Param.Min.Z, (UINT)Param.Max.X, (UINT)Param.Max.Y, (UINT)Param.Max.Z };
 
diff --git a/Engine/Source/Runtime/Engine/Private/Particles/ParticlePerfStats.cpp b/Engine/Source/Runtime/Engine/Private/Particles/ParticlePerfStats.cpp
--- a/Engine/Source/Runtime/Engine/Private/Particles/ParticlePerfStats.cpp
+++ b/Engine/Source/Runtime/Engine/Private/Particles/ParticlePerfStats.cpp
@@ -1440,7 +1440,7 @@ int32 FParticlePerfStatsListener_DebugRender::RenderStats(class UWorld* World, c
 		// Average Per Frame
 		tempString.Reset();
 		tempString.Appendf(
-			TEXT("%4u | %4u | %4u"),
+			TEXT("%4u | %4u | %4u | %4u"),
 			uint32(FPlatformTime::ToMilliseconds64(GTStats.AccumulatedStats.TickGameThreadCycles + GTStats.AccumulatedStats.FinalizeCycles) * 1000.0 / double(GTStats.NumFrames)),
 			uint32(FPlatformTime::ToMilliseconds64(GTStats.AccumulatedStats.TickConcurrentCycles + GTStats.AccumulatedStats.EndOfFrameCycles) * 1000.0 / double(GTStats.NumFrames)),
 			uint32(FPlatformTime::ToMilliseconds64(RTStats.AccumulatedStats.RenderUpdateCycles + RTStats.AccumulatedStats.GetDynamicMeshElementsCycles) * 1000.0 / double(RTStats.NumFrames)),
@@ -1451,7 +1451,7 @@ int32 FParticlePerfStatsListener_DebugRender::RenderStats(class UWorld* World, c
 		// Average Per Instances
 		tempString.Reset();
 		tempString.Appendf(
-			TEXT("%4u | %4u | %4u"),
+			TEXT("%4u | %4u | %4u | %4u"),
 			uint32(FPlatformTime::ToMilliseconds64(GTStats.AccumulatedStats.TickGameThreadCycles + GTStats.AccumulatedStats.FinalizeCycles) * 1000.0 / double(GTStats.AccumulatedStats.NumInstances)),
 			uint32(FPlatformTime::ToMilliseconds64(GTStats.AccumulatedStats.TickConcurrentCycles + GTStats.AccumulatedStats.EndOfFrameCycles) * 1000.0 / double(GTStats.AccumulatedStats.NumInstances)),
 			uint32(FPlatformTime::ToMilliseconds64(RTStats.AccumulatedStats.RenderUpdateCycles + RTStats.AccumulatedStats.GetDynamicMeshElementsCycles) * 1000.0 / double(RTStats.AccumulatedStats.NumInstances)),
diff --git a/Engine/Source/Runtime/Experimental/Chaos/Public/Chaos/ErrorReporter.h b/Engine/Source/Runtime/Experimental/Chaos/Public/Chaos/ErrorReporter.h
--- a/Engine/Source/Runtime/Experimental/Chaos/Public/Chaos/ErrorReporter.h
+++ b/Engine/Source/Runtime/Experimental/Chaos/Public/Chaos/ErrorReporter.h
@@ -33,7 +33,7 @@ public:
 		}
 		else
 		{
-			UE_LOG(LogChaos, Log, TEXT("ErrorReporter: %s"), *Prefix, ErrorMsg);
+			UE_LOG(LogChaos, Log, TEXT("ErrorReporter: %s"), ErrorMsg);
 		}
 	}
 
@@ -45,7 +45,7 @@ public:
 		}
 		else
 		{
-			UE_LOG(LogChaos, Warning, TEXT("ErrorReporter: %s"), *Prefix, ErrorMsg);
+			UE_LOG(LogChaos, Warning, TEXT("ErrorReporter: %s"), ErrorMsg);
 		}
 	}
 
diff --git a/Engine/Source/Runtime/Online/HTTPServer/Private/HttpListener.cpp b/Engine/Source/Runtime/Online/HTTPServer/Private/HttpListener.cpp
--- a/Engine/Source/Runtime/Online/HTTPServer/Private/HttpListener.cpp
+++ b/Engine/Source/Runtime/Online/HTTPServer/Private/HttpListener.cpp
@@ -86,7 +86,7 @@ bool FHttpListener::StartListening()
 	{
 		UE_LOG(LogHttpListener, Error,
 			TEXT("HttpListener detected invalid port %u"),
-			ListenPort, *Config.BindAddress, ListenPort);
+			ListenPort);
 		return false;
 	}
 	BindAddress->SetPort(ListenPort);
diff --git a/Engine/Source/Runtime/SlateRHIRenderer/Private/SlateUTextureResource.cpp b/Engine/Source/Runtime/SlateRHIRenderer/Private/SlateUTextureResource.cpp
--- a/Engine/Source/Runtime/SlateRHIRenderer/Private/SlateUTextureResource.cpp
+++ b/Engine/Source/Runtime/SlateRHIRenderer/Private/SlateUTextureResource.cpp
@@ -77,7 +77,7 @@ void FSlateBaseUTextureResource::CheckForStaleResources() const
 		{
 			if (GSlateCheckUObjectRenderResourcesShouldLogFatal)
 			{
-				UE_LOG(LogSlate, Fatal, TEXT("%s"), TEXT("Texture %s has become invalid. This means the resource was garbage collected while slate was using it"), *DebugName.ToString());
+				UE_LOG(LogSlate, Fatal, TEXT("Texture %s has become invalid. This means the resource was garbage collected while slate was using it"), *DebugName.ToString());
 			}
 			else
 			{
