PlatformProcess: Log with lowered verbosity the failure to open a non-existent semaphore.

The caller may do this to determine if the named semaphore exists, and if it doesn't then ERROR_FILE_NOT_FOUND is normal, and should not issue a warning which may fail CIS tests. Instead it now logs as "Verbose", while other error types are still logged with Warning verbosity.

While the caller could use LOG_SCOPE_VERBOSITY_OVERRIDE to avoid this, that is only safe to call in the game thread, but the semaphore may be opened in other threads.

#rnx

diff --git a/Engine/Source/Runtime/Core/Private/Windows/WindowsPlatformProcess.cpp b/Engine/Source/Runtime/Core/Private/Windows/WindowsPlatformProcess.cpp
--- a/Engine/Source/Runtime/Core/Private/Windows/WindowsPlatformProcess.cpp
+++ b/Engine/Source/Runtime/Core/Private/Windows/WindowsPlatformProcess.cpp
@@ -2005,11 +2005,22 @@ FWindowsPlatformProcess::FSemaphore* FWindowsPlatformProcess::NewInterprocessSyn
 		Semaphore = OpenSemaphore(AccessRights, false, Name);
 		if (NULL == Semaphore)
 		{
-			DWORD ErrNo = GetLastError();
-			UE_LOG(LogHAL, Warning, TEXT("OpenSemaphore(AccessRights=0x%08x, bInherit=false, Name='%s') failed with LastError = %d"),
-				AccessRights,
-				Name,
-				ErrNo);
+			const DWORD ErrNo = GetLastError();
+
+			// ERROR_FILE_NOT_FOUND is expected when checking if the semaphore exists but it doesn't,
+			// so we log the failure to open using "Verbose". Other errors will be logged using "Warning" verbosity.
+
+			if (ErrNo == ERROR_FILE_NOT_FOUND)
+			{
+				UE_LOG(LogHAL, Verbose, TEXT("OpenSemaphore(AccessRights=0x%08x, bInherit=false, Name='%s') returned ERROR_FILE_NOT_FOUND (LastError = %d)"),
+					AccessRights, Name, ErrNo);
+			}
+			else
+			{
+				UE_LOG(LogHAL, Warning, TEXT("OpenSemaphore(AccessRights=0x%08x, bInherit=false, Name='%s') failed with LastError = %d"),
+					AccessRights, Name, ErrNo);
+			}
+
 			return NULL;
 		}
 	}
