[Backout] - CL40822132
#rnx
Original CL Desc
-----------------------------------------------------------------
[Scheduler]
- Fix thread-safety issue when waking up worker threads from external threads while the game thread is restarting the workers

diff --git a/Engine/Source/Runtime/Core/Private/Async/Fundamental/ReserveScheduler.cpp b/Engine/Source/Runtime/Core/Private/Async/Fundamental/ReserveScheduler.cpp
--- a/Engine/Source/Runtime/Core/Private/Async/Fundamental/ReserveScheduler.cpp
+++ b/Engine/Source/Runtime/Core/Private/Async/Fundamental/ReserveScheduler.cpp
@@ -28,9 +28,8 @@ TUniquePtr<FThread> FReserveScheduler::CreateWorker(FThread::EForkable IsForkabl
 		*FString::Printf(TEXT("Reserve Worker #%d"), WorkerId),
 		[this, ReserveEvent]
 		{
-			FSchedulerTls::FTlsValues& TlsValues = FSchedulerTls::GetTlsValuesRef();
-			TlsValues.ActiveScheduler = this;
-			TlsValues.LocalQueue = nullptr;
+			FSchedulerTls::ActiveScheduler = this;
+			FSchedulerTls::LocalQueue = nullptr;
 
 			while (true)
 			{
@@ -38,7 +37,7 @@ TUniquePtr<FThread> FReserveScheduler::CreateWorker(FThread::EForkable IsForkabl
 				ReserveEvent->SleepEvent->Wait();
 				{
 					TRACE_CPUPROFILER_EVENT_SCOPE(FReserveScheduler::BusyWaitUntil);
-					TlsValues.WorkerType = ReserveEvent->bPermitBackgroundWork ? FSchedulerTls::EWorkerType::Background : FSchedulerTls::EWorkerType::Foreground;
+					FSchedulerTls::WorkerType = ReserveEvent->bPermitBackgroundWork ? FSchedulerTls::EWorkerType::Background : FSchedulerTls::EWorkerType::Foreground;
 
 					if (ActiveWorkers.load(std::memory_order_relaxed) == 0)
 					{
@@ -48,8 +47,8 @@ TUniquePtr<FThread> FReserveScheduler::CreateWorker(FThread::EForkable IsForkabl
 					BusyWaitUntil(MoveTemp(ReserveEvent->CompletedDelegate), ReserveEvent->bPermitBackgroundWork);
 				}
 			}
-			TlsValues.WorkerType = EWorkerType::None;
-			TlsValues.ActiveScheduler = nullptr;
+			FSchedulerTls::WorkerType = EWorkerType::None;
+			FSchedulerTls::ActiveScheduler = nullptr;
 		}, 0, Priority, FThreadAffinity{ FPlatformAffinity::GetTaskGraphThreadMask(), 0 }, IsForkable
 	);
 }
@@ -60,7 +59,7 @@ bool FReserveScheduler::DoReserveWorkUntil(FConditional&& Condition)
 	{
 		WorkerEvent->CompletedDelegate = MoveTemp(Condition);
 		// become a background worker if the reserve worker is replacing a blocked background worker
-		WorkerEvent->bPermitBackgroundWork = FSchedulerTls::GetTlsValuesRef().IsBackgroundWorker();
+		WorkerEvent->bPermitBackgroundWork = FSchedulerTls::IsBackgroundWorker();
 		WorkerEvent->SleepEvent->Trigger();
 		return true;
 	}
diff --git a/Engine/Source/Runtime/Core/Private/Async/Fundamental/Scheduler.cpp b/Engine/Source/Runtime/Core/Private/Async/Fundamental/Scheduler.cpp
--- a/Engine/Source/Runtime/Core/Private/Async/Fundamental/Scheduler.cpp
+++ b/Engine/Source/Runtime/Core/Private/Async/Fundamental/Scheduler.cpp
@@ -21,10 +21,11 @@ namespace LowLevelTasks
 {
 	DEFINE_LOG_CATEGORY(LowLevelTasks);
 
-	UE::FMutex FSchedulerTls::Mutex;
-	TArray<FSchedulerTls::FTlsValues*>  FSchedulerTls::ThreadTlsValues;
-	thread_local FSchedulerTls::FTlsValuesHolder FSchedulerTls::TlsValuesHolder;
+	thread_local FSchedulerTls::FLocalQueueType* FSchedulerTls::LocalQueue = nullptr;
 	thread_local FTask* FTask::ActiveTask = nullptr;
+	thread_local FSchedulerTls* FSchedulerTls::ActiveScheduler = nullptr;
+	thread_local FSchedulerTls::EWorkerType FSchedulerTls::WorkerType = FSchedulerTls::EWorkerType::None;
+	thread_local bool FSchedulerTls::Impl::bIsStandbyWorker = false;
 	thread_local bool Private::FOversubscriptionTls::bIsOversubscriptionAllowed = false;
 
 	FORCENOINLINE bool& Private::FOversubscriptionTls::GetIsOversubscriptionAllowedRef()
@@ -32,11 +33,6 @@ namespace LowLevelTasks
 		return bIsOversubscriptionAllowed;
 	}
 
-	FORCENOINLINE FSchedulerTls::FTlsValues& FSchedulerTls::GetTlsValuesRef()
-	{
-		return *TlsValuesHolder.TlsValues;
-	}
-
 	FScheduler FScheduler::Singleton;
 
 	TUniquePtr<FThread> FScheduler::CreateWorker(uint32 WorkerId, const TCHAR* Name, bool bPermitBackgroundWork, FThread::EForkable IsForkable, Private::FWaitEvent* ExternalWorkerEvent, FSchedulerTls::FLocalQueueType* ExternalWorkerLocalQueue, EThreadPriority Priority, uint64 InAffinity)
@@ -142,7 +138,7 @@ namespace LowLevelTasks
 			{
 				GameThreadLocalQueue = MakeUnique<FSchedulerTls::FLocalQueueType>(QueueRegistry, Private::ELocalQueueType::EForeground);
 			}
-			FSchedulerTls::GetTlsValuesRef().LocalQueue = GameThreadLocalQueue.Get();
+			FSchedulerTls::LocalQueue = GameThreadLocalQueue.Get();
 
 			WorkerEvents.SetNum(MaxWorkers);
 			WorkerLocalQueues.Reserve(MaxWorkers);
@@ -321,7 +317,7 @@ namespace LowLevelTasks
 			WaitingQueue[1].FinishShutdown();
 
 			GameThreadLocalQueue.Reset();
-			FSchedulerTls::GetTlsValuesRef().LocalQueue = nullptr;
+			FSchedulerTls::LocalQueue = nullptr;
 
 			NextWorkerId = 0;
 			WorkerThreads.Reset();
@@ -346,38 +342,10 @@ namespace LowLevelTasks
 		}
 	}
 
-	bool FSchedulerTls::HasPendingWakeUp() const
-	{
-		UE::TUniqueLock ScopedLock(Mutex);
-
-#if PLATFORM_SUPPORTS_ASYMMETRIC_FENCES
-		// Heavy barrier since bPendingWakeUp is only written to with a relaxed write, we need all cores to flush their store buffer to memory
-		FPlatformMisc::AsymmetricThreadFenceHeavy();
-		constexpr std::memory_order MemoryOrder = std::memory_order_relaxed;
-#else
-		constexpr std::memory_order MemoryOrder = std::memory_order_acquire;
-#endif
-
-		for (const FTlsValues* ThreadTls : ThreadTlsValues)
-		{
-			if (ThreadTls->ActiveScheduler != this &&
-				ThreadTls->bPendingWakeUp.load(MemoryOrder))
-			{
-				return true;
-			}
-		}
-
-		return false;
-	}
-
 	void FScheduler::RestartWorkers(uint32 NumForegroundWorkers, uint32 NumBackgroundWorkers, FThread::EForkable IsForkable, EThreadPriority InWorkerPriority, EThreadPriority InBackgroundPriority, uint64 InWorkerAffinity, uint64 InBackgroundAffinity)
 	{
 		FScopeLock Lock(&WorkerThreadsCS);
 		TemporaryShutdown.store(true, std::memory_order_release);
-		while (HasPendingWakeUp())
-		{
-			FPlatformProcess::Yield();
-		}
 		const bool bDrainGlobalQueue = false;
 		StopWorkers(bDrainGlobalQueue);
 		StartWorkers(NumForegroundWorkers, NumBackgroundWorkers, IsForkable, InWorkerPriority, InBackgroundPriority, InWorkerAffinity, InBackgroundAffinity);
@@ -388,12 +356,10 @@ namespace LowLevelTasks
 	{
 		if (ActiveWorkers.load(std::memory_order_relaxed) || TemporaryShutdown.load(std::memory_order_acquire))
 		{
-			FTlsValues& LocalTlsValues = GetTlsValuesRef();
-
 			const bool bIsBackgroundTask = Task.IsBackgroundTask();
-			const bool bIsBackgroundWorker = LocalTlsValues.IsBackgroundWorker();
-			const bool bIsStandbyWorker = LocalTlsValues.IsStandbyWorker();
-			FSchedulerTls::FLocalQueueType* const CachedLocalQueue = LocalTlsValues.LocalQueue;
+			const bool bIsBackgroundWorker = FSchedulerTls::IsBackgroundWorker();
+			const bool bIsStandbyWorker = FSchedulerTls::IsStandbyWorker();
+			FSchedulerTls::FLocalQueueType* const CachedLocalQueue = FSchedulerTls::LocalQueue;
 
 			// Standby workers always enqueue to the global queue and perform wakeup
 			// as they can go to sleep whenever the oversubscription period is done
@@ -424,49 +390,10 @@ namespace LowLevelTasks
 
 			if (bWakeUpWorker)
 			{
-#if PLATFORM_SUPPORTS_ASYMMETRIC_FENCES
-				constexpr std::memory_order MemoryOrder = std::memory_order_relaxed;
-#else
-				constexpr std::memory_order MemoryOrder = std::memory_order_seq_cst;
-#endif
-				// We don't need to pay this cost for worker threads because we already manage their shutdown gracefully.
-				const bool bExternalThread = LocalTlsValues.ActiveScheduler != this || LocalTlsValues.WorkerType == EWorkerType::None;
-				if (bExternalThread)
-				{
-#if PLATFORM_SUPPORTS_ASYMMETRIC_FENCES
-					FPlatformMisc::AsymmetricThreadFenceLight();
-#endif
-
-					LocalTlsValues.bPendingWakeUp.store(true, MemoryOrder);
-
-#if PLATFORM_SUPPORTS_ASYMMETRIC_FENCES
-					FPlatformMisc::AsymmetricThreadFenceLight();
-#endif
-
-					if (TemporaryShutdown.load(std::memory_order_acquire))
-					{
-						LocalTlsValues.bPendingWakeUp.store(false, MemoryOrder);
-						return;
-					}
-				}
-
 				if (!WakeUpWorker(bIsBackgroundTask) && !bIsBackgroundTask)
 				{
 					WakeUpWorker(true);
 				}
-
-				if (bExternalThread)
-				{
-#if PLATFORM_SUPPORTS_ASYMMETRIC_FENCES
-					FPlatformMisc::AsymmetricThreadFenceLight();
-#endif
-
-					LocalTlsValues.bPendingWakeUp.store(false, MemoryOrder);
-
-#if PLATFORM_SUPPORTS_ASYMMETRIC_FENCES
-					FPlatformMisc::AsymmetricThreadFenceLight();
-#endif
-				}
 			}
 		}
 		else
@@ -484,7 +411,7 @@ namespace LowLevelTasks
 
 	void FScheduler::IncrementOversubscription()
 	{
-		FSchedulerTls::EWorkerType LocalWorkerType = FSchedulerTls::GetTlsValuesRef().WorkerType;
+		FSchedulerTls::EWorkerType LocalWorkerType = WorkerType;
 
 		if (LocalWorkerType != EWorkerType::None)
 		{
@@ -499,7 +426,7 @@ namespace LowLevelTasks
 
 	void FScheduler::DecrementOversubscription()
 	{
-		FSchedulerTls::EWorkerType LocalWorkerType = FSchedulerTls::GetTlsValuesRef().WorkerType;
+		FSchedulerTls::EWorkerType LocalWorkerType = WorkerType;
 
 		if (LocalWorkerType != EWorkerType::None)
 		{
@@ -517,8 +444,7 @@ namespace LowLevelTasks
 
 	bool FSchedulerTls::IsWorkerThread() const
 	{
-		FTlsValues& LocalTlsValues = FSchedulerTls::GetTlsValuesRef();
-		return LocalTlsValues.WorkerType != FSchedulerTls::EWorkerType::None && LocalTlsValues.ActiveScheduler == this;
+		return WorkerType != FSchedulerTls::EWorkerType::None && ActiveScheduler == this;
 	}
 
 	bool FSchedulerTls::IsBusyWaiting()
@@ -545,7 +471,7 @@ namespace LowLevelTasks
 					// a potential wakeup we might have consumed during the cancellation.
 					if (WaitingQueue[bPermitBackgroundWork].CancelWait(WaitEvent))
 					{
-						if (!WakeUpWorker(bPermitBackgroundWork) && !GetTlsValuesRef().IsBackgroundWorker())
+						if (!WakeUpWorker(bPermitBackgroundWork) && !FSchedulerTls::IsBackgroundWorker())
 						{
 							WakeUpWorker(!bPermitBackgroundWork);
 						}
@@ -651,19 +577,17 @@ namespace LowLevelTasks
 
 	void FScheduler::WorkerMain(Private::FWaitEvent* WorkerEvent, FSchedulerTls::FLocalQueueType* WorkerLocalQueue, uint32 WaitCycles, bool bPermitBackgroundWork)
 	{
-		FTlsValues& LocalTlsValues = GetTlsValuesRef();
-
-		checkSlow(LocalTlsValues.LocalQueue == nullptr);
+		checkSlow(FSchedulerTls::LocalQueue == nullptr);
 		checkSlow(WorkerLocalQueue != nullptr);
 		checkSlow(WorkerEvent != nullptr);
 
 		FTaskTagScope WorkerScope(ETaskTag::EWorkerThread);
-		LocalTlsValues.ActiveScheduler = this;
+		FSchedulerTls::ActiveScheduler = this;
 
 		FMemory::SetupTLSCachesOnCurrentThread();
-		LocalTlsValues.WorkerType = bPermitBackgroundWork ? FSchedulerTls::EWorkerType::Background : FSchedulerTls::EWorkerType::Foreground;
-		LocalTlsValues.SetStandbyWorker(WorkerEvent->bIsStandby);
-		LocalTlsValues.LocalQueue = WorkerLocalQueue;
+		FSchedulerTls::WorkerType = bPermitBackgroundWork ? FSchedulerTls::EWorkerType::Background : FSchedulerTls::EWorkerType::Foreground;
+		FSchedulerTls::SetStandbyWorker(WorkerEvent->bIsStandby);
+		FSchedulerTls::LocalQueue = WorkerLocalQueue;
 
 		{
 			Private::FOversubscriptionAllowedScope _(true);
@@ -678,10 +602,10 @@ namespace LowLevelTasks
 			}
 		}
 
-		LocalTlsValues.LocalQueue = nullptr;
-		LocalTlsValues.ActiveScheduler = nullptr;
-		LocalTlsValues.SetStandbyWorker(false);
-		LocalTlsValues.WorkerType = FSchedulerTls::EWorkerType::None;
+		FSchedulerTls::LocalQueue = nullptr;
+		FSchedulerTls::ActiveScheduler = nullptr;
+		FSchedulerTls::SetStandbyWorker(false);
+		FSchedulerTls::WorkerType = FSchedulerTls::EWorkerType::None;
 		FMemory::ClearAndDisableTLSCachesOnCurrentThread();
 	}
 
diff --git a/Engine/Source/Runtime/Core/Public/Async/Fundamental/Scheduler.h b/Engine/Source/Runtime/Core/Public/Async/Fundamental/Scheduler.h
--- a/Engine/Source/Runtime/Core/Public/Async/Fundamental/Scheduler.h
+++ b/Engine/Source/Runtime/Core/Public/Async/Fundamental/Scheduler.h
@@ -5,14 +5,11 @@
 #include "Async/Fundamental/TaskShared.h"
 #include "Async/Fundamental/TaskDelegate.h"
 #include "Async/Fundamental/WaitingQueue.h"
-#include "Async/Mutex.h"
-#include "Async/UniqueLock.h"
 #include "Containers/Array.h"
 #include "Containers/ArrayView.h"
 #include "Containers/ContainerAllocationPolicies.h"
 #include "HAL/CriticalSection.h"
 #include "HAL/Event.h"
-#include "HAL/LowLevelMemTracker.h"
 #include "HAL/PlatformAffinity.h"
 #include "HAL/PlatformProcess.h"
 #include "HAL/Thread.h"
@@ -124,60 +121,18 @@ namespace LowLevelTasks
 			Foreground,
 		};
 
-		struct FTlsValues;
-		static UE::FMutex Mutex;
-		static TArray<FTlsValues*> ThreadTlsValues;
-
-		struct alignas(PLATFORM_CACHE_LINE_SIZE) FTlsValues
-		{
-			FSchedulerTls* ActiveScheduler = nullptr;
-			FLocalQueueType* LocalQueue = nullptr;
-			EWorkerType WorkerType = EWorkerType::None;
-			std::atomic<bool> bPendingWakeUp = false;
-			bool        bIsStandbyWorker = false;
-
-			inline bool IsBackgroundWorker()
-			{
-				return WorkerType == EWorkerType::Background;
-			}
-
-			inline bool IsStandbyWorker()
-			{
-				return bIsStandbyWorker;
-			}
-
-			inline void SetStandbyWorker(bool bInIsStandbyWorker)
-			{
-				bIsStandbyWorker = bInIsStandbyWorker;
-			}
-		};
-
-		struct FTlsValuesHolder
+		static thread_local FSchedulerTls* ActiveScheduler;
+		static thread_local FLocalQueueType* LocalQueue;
+		static thread_local EWorkerType WorkerType;
+	
+	private:
+		// Hide the member to avoid clashing with local variables
+		// already named bIsStandbyWorker.
+		struct Impl
 		{
-			FTlsValuesHolder()
-			{
-				LLM_SCOPE_BYNAME(TEXT("Scheduler/TlsValues"));
-				TlsValues = new FTlsValues();
-
-				UE::TUniqueLock ScopedLock(Mutex);
-				ThreadTlsValues.Add(TlsValues);
-			}
-			~FTlsValuesHolder()
-			{
-				{
-					UE::TUniqueLock ScopedLock(Mutex);
-					ThreadTlsValues.Remove(TlsValues);
-				}
-				delete TlsValues;
-				TlsValues = nullptr;
-			}
-			
-			FTlsValues* TlsValues = nullptr;
+			static thread_local bool bIsStandbyWorker;
 		};
 
-	private:
-		static thread_local FTlsValuesHolder TlsValuesHolder;
-
 	public:
 		CORE_API bool IsWorkerThread() const;
 
@@ -185,9 +140,20 @@ namespace LowLevelTasks
 		CORE_API static bool IsBusyWaiting();
 
 	protected:
-		CORE_API bool HasPendingWakeUp() const;
+		inline static bool IsBackgroundWorker()
+		{
+			return WorkerType == EWorkerType::Background;
+		}
 
-		static FTlsValues& GetTlsValuesRef();
+		inline static bool IsStandbyWorker()
+		{
+			return Impl::bIsStandbyWorker;
+		}
+
+		inline static void SetStandbyWorker(bool bInIsStandbyWorker)
+		{
+			Impl::bIsStandbyWorker = bInIsStandbyWorker;
+		}
 	};
 
 	class FScheduler final : public FSchedulerTls
diff --git a/Engine/Source/Runtime/Core/Tests/Tasks/TasksTest.cpp b/Engine/Source/Runtime/Core/Tests/Tasks/TasksTest.cpp
--- a/Engine/Source/Runtime/Core/Tests/Tasks/TasksTest.cpp
+++ b/Engine/Source/Runtime/Core/Tests/Tasks/TasksTest.cpp
@@ -1738,9 +1738,9 @@ namespace UE { namespace TasksTests
 	}
 
 	// StressTest RestartWorkers while having task performing oversubscription
-	TEST_CASE_NAMED(FTaskGraphOversubscribeDuringRestartWorkers, "System::Core::Async::Tasks::RestartWorkersAndOversubscription", "[.][ApplicationContextMask][EngineFilter]")
+	TEST_CASE_NAMED(FTaskGraphOversubscribeDuringRestartWorkers, "System::Core::Async::Tasks::OversubscribeDuringRestartWorkers", "[.][ApplicationContextMask][EngineFilter]")
 	{
-		for (int32 Index = 0; Index < 10000; ++Index)
+		for (int32 Index = 0; Index < 10; ++Index)
 		{
 			FTask Task = Launch(TEXT("Task"),
 				[&]()
@@ -1761,37 +1761,6 @@ namespace UE { namespace TasksTests
 		LowLevelTasks::FScheduler::Get().RestartWorkers();
 	}
 
-	// StressTest RestartWorkers while having a thread trying to launch tasks on workers
-	TEST_CASE_NAMED(FTaskGraphLaunchFromThreadDuringRestartWorkers, "System::Core::Async::Tasks::RestartWorkersAndExternalThreads", "[.][ApplicationContextMask][EngineFilter]")
-	{
-		std::atomic<bool> bExit = false;
-		UE::FManualResetEvent Done;
-		Async(EAsyncExecution::Thread,
-			[&Done, &bExit]()
-			{
-				while (!bExit)
-				{
-					Launch(TEXT("Task"), [&]() {}).Wait();
-				}
-
-				Done.Notify();
-			}
-		);
-
-		for (int32 Index = 0; Index < 1000; ++Index)
-		{
-			// For the repro to work we need to launch tasks to try to start new workers while
-			// we're shutting down and restarting workers.
-			LowLevelTasks::FScheduler::Get().RestartWorkers(0, 4);
-		}
-
-		bExit = true;
-		Done.Wait();
-
-		// Restore original worker count
-		LowLevelTasks::FScheduler::Get().RestartWorkers();
-	}
-
 #if PLATFORM_SUPPORTS_ASYMMETRIC_FENCES
 	TEST_CASE_NAMED(FAsymmetricThreadFence, "System::Core::Async::AsymmetricThreadFence", "[.][ApplicationContextMask][EngineFilter]")
 	{
