AddStatMessage was inadvertently made private.
#rnx

diff --git a/Engine/Source/Runtime/Core/Public/Stats/StatsSystemTypes.h b/Engine/Source/Runtime/Core/Public/Stats/StatsSystemTypes.h
--- a/Engine/Source/Runtime/Core/Public/Stats/StatsSystemTypes.h
+++ b/Engine/Source/Runtime/Core/Public/Stats/StatsSystemTypes.h
@@ -1196,6 +1196,8 @@ class FThreadStats : FNoncopyable
 	/** Tracks current stack depth for cycle counters. **/
 	bool bSawExplicitFlush = false;
 
+	void SendMessage_Async(FStatPacket* ToSend);
+
 protected:
 	/** Gathers information about the current thread and sets up the TLS value. **/
 	CORE_API FThreadStats(EThreadType::Type InThreadType = EThreadType::Invalid);
@@ -1245,9 +1247,6 @@ public:
 	/** Checks the command line whether we want to enable collecting startup stats. */
 	static CORE_API void CheckForCollectingStartupStats();
 
-private:
-	void SendMessage_Async(FStatPacket* ToSend);
-
 	FORCEINLINE_STATS void AddStatMessage(const FStatMessage& StatMessage)
 	{
 		LLM_SCOPE(ELLMTag::Stats);
