Fix FBase64::Encode/Decode being lossy when encoding/decoding FString because intermediate conversion to ansi

diff --git a/Engine/Source/Runtime/Core/Private/Misc/Base64.cpp b/Engine/Source/Runtime/Core/Private/Misc/Base64.cpp
--- a/Engine/Source/Runtime/Core/Private/Misc/Base64.cpp
+++ b/Engine/Source/Runtime/Core/Private/Misc/Base64.cpp
@@ -52,7 +52,8 @@ namespace
 
 FString FBase64::Encode(const FString& Source, EBase64Mode Mode)
 {
-	return Encode((uint8*)TCHAR_TO_ANSI(*Source), Source.Len(), Mode);
+	auto Conv = StrCast<UTF8CHAR>(*Source, Source.Len());
+	return Encode((uint8*)Conv.Get(), Conv.Length(), Mode);
 }
 
 FString FBase64::Encode(const TArray<uint8>& Source, EBase64Mode Mode)
@@ -166,13 +167,13 @@ bool FBase64::Decode(const FString& Source, FString& OutDest, EBase64Mode Mode)
 {
 	uint32 ExpectedLength = GetDecodedDataSize(Source);
 
-	TArray<ANSICHAR> TempDest;
-	TempDest.AddZeroed(ExpectedLength + 1);
+	TArray<UTF8CHAR> TempDest;
+	TempDest.AddUninitialized(ExpectedLength);
 	if (!Decode(*Source, Source.Len(), (uint8*)TempDest.GetData(), Mode))
 	{
 		return false;
 	}
-	OutDest = ANSI_TO_TCHAR(TempDest.GetData());
+	OutDest = FString::ConstructFromPtrSize(TempDest.GetData(), ExpectedLength);
 	return true;
 }
 
