Csv Profiler: Add validation for stat names.
The cvar csv.StatNameValidation determines what happens when invalid names are encountered:
- If 0, does nothing
- If 1, emits a warning if a stat name contains invalid characters
- If 2, warns and sanitizes stat names with invalid characters (stripping them) (DEFAULT)
- If 3, warns and filters out stats with invalid characters
[RN] Csv Profiler: Add validation for stat names. The cvar csv.StatNameValidation determines what happens when invalid names are encountered
#tests Tested on client and via automation

diff --git a/Engine/Source/Runtime/Core/Private/ProfilingDebugging/CsvProfiler.cpp b/Engine/Source/Runtime/Core/Private/ProfilingDebugging/CsvProfiler.cpp
--- a/Engine/Source/Runtime/Core/Private/ProfilingDebugging/CsvProfiler.cpp
+++ b/Engine/Source/Runtime/Core/Private/ProfilingDebugging/CsvProfiler.cpp
@@ -233,6 +233,15 @@ TAutoConsoleVariable<bool> CVarNumberedFNamesStatsAreFatal(
 	ECVF_Default
 );
 
+TAutoConsoleVariable<int32> CVarCsvStatNameValidation(
+	TEXT("csv.StatNameValidation"),
+	2,
+	TEXT("If 0, does nothing\r\n")
+	TEXT("If 1, warns if there are invalid characters\r\n")
+	TEXT("If 2, warns and sanitizes stat names with invalid characters\r\n")
+	TEXT("If 3, warns and removes stats with invalid characters"),
+	ECVF_Default
+);
 
 static bool GCsvUseProcessingThread = true;
 static int32 GCsvRepeatCount = 0;
@@ -1905,6 +1914,53 @@ struct FCsvProcessedEvent
 	uint32 CategoryIndex;
 };
 
+
+class FCsvStatNameValidator
+{
+public:
+	FCsvStatNameValidator()
+	{
+		ValidCharacters = TEXT("0123456789abcdefghijklmnopqrstuvwxyzABCDEFGHIJKLMNOPQRSTUVWXYZ /_-[]()#.:");
+		for (TCHAR Character : ValidCharacters)
+		{
+			ValidCharactersSet.Add(Character);
+		}
+	}
+	bool IsNameValid(const FString& StatName) const
+	{
+		for (TCHAR Character : StatName)
+		{
+			if (!ValidCharactersSet.Contains(Character))
+			{
+				return false;
+			}
+		}
+		return true;
+	}
+
+	FString SanitizeName(const FString& StatName) const
+	{
+		FString NameOut;
+		for (TCHAR Character : StatName)
+		{
+			if (ValidCharactersSet.Contains(Character))
+			{
+				NameOut += Character;
+			}
+		}
+		return NameOut;
+	}
+
+	const FString& GetValidCharacters() const
+	{
+		return ValidCharacters;
+	}
+private:
+	TSet<TCHAR> ValidCharactersSet;
+	FString ValidCharacters;
+};
+static FCsvStatNameValidator* GCsvStatNameValidator = nullptr;
+
 typedef int32 FCsvStatID;
 
 struct FCsvStatSeries
@@ -2222,6 +2278,20 @@ FCsvStatSeries::FCsvStatSeries(EType InSeriesType, const FCsvStatID& InStatID, F
 	// If we have a linked stat series, don't write directly to the CSV
 	bool bVisible = LinkedAggregateStatSeries == nullptr || (CSV_DEBUG_EMIT_SEPARATE_THREAD_STATS_WHEN_TASK_AGGREGATION_ENABLED==1);
 
+	// Name validation
+	if (CVarCsvStatNameValidation.GetValueOnAnyThread() > 0 && GCsvStatNameValidator != nullptr && !GCsvStatNameValidator->IsNameValid(Name))
+	{
+		UE_LOG(LogCsvProfiler, Warning, TEXT("Stat name '%s' contains invalid characters. Valid characters are: '%s'"), *Name, *GCsvStatNameValidator->GetValidCharacters());
+		if (CVarCsvStatNameValidation.GetValueOnAnyThread() == 2)
+		{
+			Name = GCsvStatNameValidator->SanitizeName(Name);
+		}
+		else if (CVarCsvStatNameValidation.GetValueOnAnyThread() == 3)
+		{
+			bVisible = false;
+		}
+	}
+
 	// Internal CsvBench stats should be hidden
 	if (StatCategoryIndex == CSV_CATEGORY_INDEX(CsvBench))
 	{
@@ -3399,6 +3469,8 @@ FCsvProfiler::FCsvProfiler()
 	OSMinor.TrimStartAndEndInline();
 	FString OSString = FString::Printf(TEXT("%s %s"), *OSMajor, *OSMinor);
 
+	GCsvStatNameValidator = new FCsvStatNameValidator();
+
 	SetMetadataInternal(TEXT("Platform"), *PlatformStr);
 	SetMetadataInternal(TEXT("Config"), *BuildConfigurationStr);
 	SetMetadataInternal(TEXT("BuildVersion"), *BuildVersionString);
