SourceLocationUtils: Add adapter-based API for source location formatting.

- Introduce adapter pattern for source location formatting
- Add ToString functions on adapters. These will replace existing To[Full|FileAndLine]String functions, which will be deprecated in a future change.
- Add support for serialization to structured logs.
- Add support for appending to string builders.
- Support both wide and utf8 string formats.

#tests All new source location unit tests pass. Compiled and run tests with UE_INCLUDE_SOURCE_LOCATION set to both 0 and 1.

diff --git a/Engine/Source/Runtime/Core/Private/Misc/SourceLocationUtils.cpp b/Engine/Source/Runtime/Core/Private/Misc/SourceLocationUtils.cpp
new file mode 100644
--- /dev/null
+++ b/Engine/Source/Runtime/Core/Private/Misc/SourceLocationUtils.cpp
@@ -0,0 +1,145 @@
+// Copyright Epic Games, Inc. All Rights Reserved.
+
+#include "Misc/SourceLocationUtils.h"
+#include "Containers/Utf8String.h"
+#include "Serialization/CompactBinaryWriter.h"
+
+namespace UE::SourceLocation::Private
+{
+
+	/**
+	 * Helper to append a formatted full source location to either a string or string builder.
+	 *
+	 * @param Dest      Destination string or string builder
+	 * @param Location  Source location to format
+	 */
+	template <typename FmtCharType, typename DestBufferType>
+	static void AppendFull(DestBufferType& Dest, const FSourceLocation& Location)
+	{
+	#if UE_INCLUDE_SOURCE_LOCATION
+		Dest.Appendf(CHARTEXT(FmtCharType, "%hs(%d:%d) %hs"), Location.GetFileName(), Location.GetLine(), Location.GetColumn(), Location.GetFunctionName());
+	#endif // UE_INCLUDE_SOURCE_LOCATION
+	}
+
+	/**
+	 * Helper to append a formatted file and line source location to either a string or string builder.
+	 *
+	 * @param Dest      Destination string or string builder
+	 * @param Location  Source location to format
+	 */
+	template <typename FmtCharType, typename DestBufferType>
+	static void AppendFileAndLine(DestBufferType& Dest, const FSourceLocation& Location)
+	{
+	#if UE_INCLUDE_SOURCE_LOCATION
+		Dest.Appendf(CHARTEXT(FmtCharType, "%hs(%d)"), Location.GetFileName(), Location.GetLine());
+	#endif // UE_INCLUDE_SOURCE_LOCATION
+	}
+
+	void SerializeForLogFull(FCbWriter& Writer, const FSourceLocation& Location)
+	{
+		Writer.BeginObject();
+	#if UE_INCLUDE_SOURCE_LOCATION
+		Writer.AddString(ANSITEXTVIEW("$type"), ANSITEXTVIEW("SourceLocationFull"));
+		Writer.AddString(ANSITEXTVIEW("$text"), WriteToUtf8String<512>(Full(Location)));
+		Writer.AddString(ANSITEXTVIEW("File"), Location.GetFileName());
+		Writer.AddInteger(ANSITEXTVIEW("Line"), Location.GetLine());
+		Writer.AddInteger(ANSITEXTVIEW("Column"), Location.GetColumn());
+		Writer.AddString(ANSITEXTVIEW("Function"), Location.GetFunctionName());
+	#endif // UE_INCLUDE_SOURCE_LOCATION
+		Writer.EndObject();
+	}
+
+	void SerializeForLogFileAndLine(FCbWriter& Writer, const FSourceLocation& Location)
+	{
+		Writer.BeginObject();
+	#if UE_INCLUDE_SOURCE_LOCATION
+		Writer.AddString(ANSITEXTVIEW("$type"), ANSITEXTVIEW("SourceLocationFileAndLine"));
+		Writer.AddString(ANSITEXTVIEW("$text"), WriteToUtf8String<300>(FileAndLine(Location)));
+		Writer.AddString(ANSITEXTVIEW("File"), Location.GetFileName());
+		Writer.AddInteger(ANSITEXTVIEW("Line"), Location.GetLine());
+	#endif // UE_INCLUDE_SOURCE_LOCATION
+		Writer.EndObject();
+	}
+
+} // namespace UE::SourceLocation::Private
+
+namespace UE::SourceLocation
+{
+
+	FString FFullAdapter::ToString() const
+	{
+		FString Result;
+		ToString(Result);
+		return Result;
+	}
+
+	FUtf8String FFullAdapter::ToUtf8String() const
+	{
+		FUtf8String Result;
+		ToString(Result);
+		return Result;
+	}
+
+	void FFullAdapter::ToString(FWideString& Out) const
+	{
+		Out.Reset();
+		Private::AppendFull<WIDECHAR>(Out, Location);
+	}
+
+	void FFullAdapter::ToString(FUtf8String& Out) const
+	{
+		Out.Reset();
+		// Using ANSICHAR format string - FUtf8String::Appendf's use of TCheckedFormatString
+		// fails if UE_UTF8TEXT_IS_CONSTEXPR == 0 and UTF8TEXT is used.
+		Private::AppendFull<ANSICHAR>(Out, Location);
+	}
+
+	void FFullAdapter::AppendString(FWideStringBuilderBase& Out) const
+	{
+		Private::AppendFull<WIDECHAR>(Out, Location);
+	}
+
+	void FFullAdapter::AppendString(FUtf8StringBuilderBase& Out) const
+	{
+		Private::AppendFull<UTF8CHAR>(Out, Location);
+	}
+
+	FString FFileAndLineAdapter::ToString() const
+	{
+		FString Result;
+		ToString(Result);
+		return Result;
+	}
+
+	FUtf8String FFileAndLineAdapter::ToUtf8String() const
+	{
+		FUtf8String Result;
+		ToString(Result);
+		return Result;
+	}
+
+	void FFileAndLineAdapter::ToString(FWideString& Out) const
+	{
+		Out.Reset();
+		Private::AppendFileAndLine<WIDECHAR>(Out, Location);
+	}
+
+	void FFileAndLineAdapter::ToString(FUtf8String& Out) const
+	{
+		Out.Reset();
+		// Using ANSICHAR format string - FUtf8String::Appendf's use of TCheckedFormatString
+		// fails if UE_UTF8TEXT_IS_CONSTEXPR == 0 and UTF8TEXT is used.
+		Private::AppendFileAndLine<ANSICHAR>(Out, Location);
+	}
+
+	void FFileAndLineAdapter::AppendString(FWideStringBuilderBase& Out) const
+	{
+		Private::AppendFileAndLine<WIDECHAR>(Out, Location);
+	}
+
+	void FFileAndLineAdapter::AppendString(FUtf8StringBuilderBase& Out) const
+	{
+		Private::AppendFileAndLine<UTF8CHAR>(Out, Location);
+	}
+
+} // namespace UE::SourceLocation
diff --git a/Engine/Source/Runtime/Core/Public/Misc/SourceLocationUtils.h b/Engine/Source/Runtime/Core/Public/Misc/SourceLocationUtils.h
--- a/Engine/Source/Runtime/Core/Public/Misc/SourceLocationUtils.h
+++ b/Engine/Source/Runtime/Core/Public/Misc/SourceLocationUtils.h
@@ -3,40 +3,184 @@
 #pragma once
 
 #include "CoreTypes.h"
-
+#include "Containers/StringFwd.h"
 #include "Containers/UnrealString.h"
 #include "Misc/SourceLocation.h"
 #include "Misc/StringBuilder.h"
 
+class FCbWriter;
+
+namespace UE::SourceLocation::Private
+{
+
+	CORE_API void SerializeForLogFull(FCbWriter& Writer, const FSourceLocation& Location);
+	CORE_API void SerializeForLogFileAndLine(FCbWriter& Writer, const FSourceLocation& Location);
+
+}  // namespace UE::SourceLocation::Private
+
 namespace UE::SourceLocation
 {
+
+	/**
+	 * Adapter for formatting a source location with full information.
+	 */
+	struct FFullAdapter
+	{
+		const FSourceLocation& Location;
+
+		/**
+		 * Converts the source location to a readable format with full information.
+		 *
+		 * @return String representation of the source location.
+		 */
+		[[nodiscard]] CORE_API FString ToString() const;
+		[[nodiscard]] CORE_API FUtf8String ToUtf8String() const;
+
+		/**
+		 * Converts the source location to a readable format with full information, in place.
+		 *
+		 * @param Out String to fill with the string representation of the source location.
+		 */
+		CORE_API void ToString(FWideString& Out) const;
+		CORE_API void ToString(FUtf8String& Out) const;
+
+		/**
+		 * Converts the source location to a readable format with full information, appending to an existing string.
+		 *
+		 * @param Out StringBuilder to append with the string representation of the source location.
+		 */
+		CORE_API void AppendString(FWideStringBuilderBase& Out) const;
+		CORE_API void AppendString(FUtf8StringBuilderBase& Out) const;
+
+		friend void SerializeForLog(FCbWriter& Writer, const FFullAdapter& Adapter)
+		{
+			Private::SerializeForLogFull(Writer, Adapter.Location);
+		}
+
+		friend FWideStringBuilderBase& operator<<(FWideStringBuilderBase& Builder, const FFullAdapter& Adapter)
+		{
+			Adapter.AppendString(Builder);
+			return Builder;
+		}
+		friend FUtf8StringBuilderBase& operator<<(FUtf8StringBuilderBase& Builder, const FFullAdapter& Adapter)
+		{
+			Adapter.AppendString(Builder);
+			return Builder;
+		}
+	};
+
+	/**
+	 * Adapter for formatting a source location as file and line only.
+	 */
+	struct FFileAndLineAdapter
+	{
+		const FSourceLocation& Location;
+
+		/**
+		 * Converts the source location to a readable format with file and line information.
+		 *
+		 * @return String representation of the source location.
+		 */
+		[[nodiscard]] CORE_API FString ToString() const;
+		[[nodiscard]] CORE_API FUtf8String ToUtf8String() const;
+
+		/**
+		 * Converts the source location to a readable format with file and line information, in place.
+		 *
+		 * @param Out String to fill with the string representation of the source location.
+		 */
+		CORE_API void ToString(FWideString& Out) const;
+		CORE_API void ToString(FUtf8String& Out) const;
+
+		/**
+		 * Converts the source location to a readable format with file and line information, appending to an existing string.
+		 *
+		 * @param Out StringBuilder to append with the string representation of the source location.
+		 */
+		CORE_API void AppendString(FWideStringBuilderBase& Out) const;
+		CORE_API void AppendString(FUtf8StringBuilderBase& Out) const;
+
+		friend void SerializeForLog(FCbWriter& Writer, const FFileAndLineAdapter& Adapter)
+		{
+			Private::SerializeForLogFileAndLine(Writer, Adapter.Location);
+		}
+
+		friend FWideStringBuilderBase& operator<<(FWideStringBuilderBase& Builder, const FFileAndLineAdapter& Adapter)
+		{
+			Adapter.AppendString(Builder);
+			return Builder;
+		}
+		friend FUtf8StringBuilderBase& operator<<(FUtf8StringBuilderBase& Builder, const FFileAndLineAdapter& Adapter)
+		{
+			Adapter.AppendString(Builder);
+			return Builder;
+		}
+	};
+
+	/**
+	 * Returns an adapter that formats source location with full information (file name, line, column, and function name).
+	 *
+	 * @code
+	 *   // Example with string:
+	 *   FString LocationStr = UE::SourceLocation::Full(SourceLoc).ToString();
+	 *
+	 *   // Example with string builder:
+	 *   Builder << "Error at " << UE::SourceLocation::Full(SourceLoc);
+	 *
+	 *   // Example with structured log:
+	 *   UE_LOGFMT(LogCore, Warning, "Error at {Location}", UE::SourceLocation::Full(SourceLoc));
+	 * @endcode
+	 *
+	 * @param Location The source location to adapt.
+	 * @return Adapter for full source location formatting.
+	 */
+	inline FFullAdapter Full(const FSourceLocation& Location)
+	{
+		return { Location };
+	}
+
 	/**
-	 * Returns an owning string that contains function name as well as file, line and column
+	 * Returns an adapter that formats source location with file and line only.
+	 *
+	 * @code
+	 *   // Example with string:
+	 *   FString LocationStr = UE::SourceLocation::FileAndLine(SourceLoc).ToString();
+	 *
+	 *   // Example with string builder:
+	 *   Builder << "Error at " << UE::SourceLocation::FileAndLine(SourceLoc);
+	 *
+	 *   // Example with structured log:
+	 *   UE_LOGFMT(LogCore, Warning, "Error at {Location}", UE::SourceLocation::FileAndLine(SourceLoc));
+	 * @endcode
+	 *
+	 * @param Location The source location to adapt.
+	 * @return Adapter for file and line formatting.
 	 */
+	inline FFileAndLineAdapter FileAndLine(const FSourceLocation& Location)
+	{
+		return { Location };
+	}
+
+	/**
+	 * Returns an owning string with full source location information (file name, line, column, and function name).
+	 * @param Location - The source location to convert to string.
+	 * @return Formatted string with full source location information.
+	 */
+	// UE_DEPRECATED(5.6, "Use UE::SourceLocation::Full(Location).ToString() instead.")
 	static inline FString ToFullString(const FSourceLocation& Location)
 	{
-#if UE_INCLUDE_SOURCE_LOCATION
-		return FString(WriteToString<512>(
-			Location.GetFileName(), "(",
-			Location.GetLine(), ":",
-			Location.GetColumn(), ") ",
-			Location.GetFunctionName()).ToView());
-#else
-		return FString();
-#endif
+		return Full(Location).ToString();
 	}
 
 	/**
-	 * Returns an owning string that contains Source filename and line.
-	 * Equvivalent to UE_SOURCE_LOCATION
+	 * Returns an owning string with source filename and line. Equivalent to UE_SOURCE_LOCATION.
+	 * @param Location - The source location to convert to string.
+	 * @return Formatted string with file and line information.
 	 */
+	// UE_DEPRECATED(5.6, "Use UE::SourceLocation::FileAndLine(Location).ToString() instead.")
 	static inline FString ToFileAndLineString(const FSourceLocation& Location)
 	{
-#if UE_INCLUDE_SOURCE_LOCATION
-		return FString(WriteToString<300>(
-			Location.GetFileName(), "(", Location.GetLine(), ")").ToView());
-#else
-		return FString();
-#endif
+		return FileAndLine(Location).ToString();
 	}
-}
\ No newline at end of file
+
+}  // namespace UE::SourceLocation
diff --git a/Engine/Source/Runtime/Core/Tests/Misc/SourceLocationTest.cpp b/Engine/Source/Runtime/Core/Tests/Misc/SourceLocationTest.cpp
new file mode 100644
--- /dev/null
+++ b/Engine/Source/Runtime/Core/Tests/Misc/SourceLocationTest.cpp
@@ -0,0 +1,323 @@
+// Copyright Epic Games, Inc. All Rights Reserved.
+
+#if WITH_TESTS
+
+#include "Logging/StructuredLog.h"
+#include "Containers/Utf8String.h"
+#include "Misc/SourceLocationUtils.h"
+#include "Misc/StringBuilder.h"
+#include "Serialization/CompactBinaryWriter.h"
+#include "Tests/TestHarnessAdapter.h"
+
+namespace UE::SourceLocation::Private
+{
+
+UE::FSourceLocation GetTestSourceLocation()
+{
+	return UE::FSourceLocation::Current();
+}
+
+/*
+ * Helper function to find and remove a single occurrence of the parameter search text
+ * from the parameter string builder.
+ * @return True if an occurrence was found and removed.
+ */
+template <typename CharType>
+bool RemoveSingle(TStringBuilderBase<CharType>& Builder, TStringView<CharType> SearchText)
+{
+	if (SearchText.IsEmpty())
+	{
+		return false;
+	}
+	const int32 FindPos = Builder.ToView().Find(SearchText);
+	if (FindPos != INDEX_NONE)
+	{
+		Builder.RemoveAt(FindPos, SearchText.Len());
+		return true;
+	}
+	return false;
+}
+
+/*
+ * Helper function to verify full formatting.
+ * @note This validation assumes that the formatted elements appear in a specific order,
+ * since it relies on substring removal rather than whole-word matching.
+ * e.g. if column "15" is searched in the string "151:15" (line:column),
+ * the "15" would match part of "151", causing the subsequent search for "151" to fail.
+ */
+template <typename CharType>
+bool CheckFullFormatting(TStringView<CharType> Text, UE::FSourceLocation Location)
+{
+	TStringBuilderWithBuffer<CharType, 512> Builder;
+	Builder << Text;
+
+	// Verify all required elements are present by removing them one by one.
+	if (!RemoveSingle<CharType>(Builder, StringCast<CharType>(Location.GetFileName())))
+	{
+		return false;
+	}
+	if (!RemoveSingle<CharType>(Builder, TStringBuilderWithBuffer<CharType, 16>(InPlace, Location.GetLine())))
+	{
+		return false;
+	}
+	if (!RemoveSingle<CharType>(Builder, TStringBuilderWithBuffer<CharType, 16>(InPlace, Location.GetColumn())))
+	{
+		return false;
+	}
+	if (!RemoveSingle<CharType>(Builder, StringCast<CharType>(Location.GetFunctionName())))
+	{
+		return false;
+	}
+	return true;
+}
+/*
+ * Helper function to verify file and line formatting.
+ * @note This validation assumes that the formatted elements appear in a specific order,
+ * since it relies on substring removal rather than whole-word matching.
+ * e.g. if column "15" is searched in the string "151:15" (line:column),
+ * the "15" would match part of "151", causing the subsequent search for "151" to fail.
+ */
+template <typename CharType>
+bool CheckFileAndLineFormatting(TStringView<CharType> Text, UE::FSourceLocation Location)
+{
+	TStringBuilderWithBuffer<CharType, 512> Builder;
+	Builder << Text;
+
+	// Verify all required elements are present by removing them one by one.
+	if (!RemoveSingle<CharType>(Builder, StringCast<CharType>(Location.GetFileName())))
+	{
+		return false;
+	}
+	if (!RemoveSingle<CharType>(Builder, TStringBuilderWithBuffer<CharType, 16>(InPlace, Location.GetLine())))
+	{
+		return false;
+	}
+	if (RemoveSingle<CharType>(Builder, TStringBuilderWithBuffer<CharType, 16>(InPlace, Location.GetColumn())))
+	{
+		return false;
+	}
+	if (RemoveSingle<CharType>(Builder, StringCast<CharType>(Location.GetFunctionName())))
+	{
+		return false;
+	}
+	return true;
+}
+
+bool ValidateBinaryForFull(FCbField Field, UE::FSourceLocation Location)
+{
+	// Must have a value.
+	if (!Field.HasValue())
+	{
+		return false;
+	}
+
+	FCbObject Object = Field.AsObject();
+	// Check for required fields in full formatting.
+	if (!Object.Find("$type").HasValue())
+	{
+		return false;
+	}
+	if (!Object.Find("File").HasValue())
+	{
+		return false;
+	}
+	if (!Object.Find("Line").HasValue())
+	{
+		return false;
+	}
+	if (!Object.Find("Column").HasValue())
+	{
+		return false;
+	}
+	if (!Object.Find("Function").HasValue())
+	{
+		return false;
+	}
+
+	// Validate the text field.
+	FCbFieldView TextField = Object.Find("$text");
+	if (!TextField.HasValue() || !TextField.IsString())
+	{
+		return false;
+	}
+	return CheckFullFormatting(TextField.AsString(), Location);
+}
+
+bool ValidateBinaryForFileAndLine(FCbField Field, UE::FSourceLocation Location)
+{
+	// Must have a value.
+	if (!Field.HasValue())
+	{
+		return false;
+	}
+
+	FCbObject Object = Field.AsObject();
+	// Check for required fields in file and line formatting.
+	if (!Object.Find("$type").HasValue())
+	{
+		return false;
+	}
+	if (!Object.Find("File").HasValue())
+	{
+		return false;
+	}
+	if (!Object.Find("Line").HasValue())
+	{
+		return false;
+	}
+	// The following fields must NOT be present
+	if (Object.Find("Column").HasValue())
+	{
+		return false;
+	}
+	if (Object.Find("Function").HasValue())
+	{
+		return false;
+	}
+
+	// Validate the text field.
+	FCbFieldView TextField = Object.Find("$text");
+	if (!TextField.HasValue() || !TextField.IsString())
+	{
+		return false;
+	}
+	return CheckFileAndLineFormatting(TextField.AsString(), Location);
+}
+
+}  // namespace UE::SourceLocation::Private
+
+namespace UE::SourceLocation
+{
+
+TEST_CASE("System::Core::Misc::SourceLocation", "[Core][SourceLocation][ApplicationContextMask][SmokeFilter]")
+{
+	// Get a real source location from a function call - FSourceLocation doesn't allow assigning made-up values for
+	// testing.
+	UE::FSourceLocation TestLocation = Private::GetTestSourceLocation();
+
+	SECTION("UE::SourceLocation::Full (ToString)")
+	{
+		const FString String = UE::SourceLocation::Full(TestLocation).ToString();
+
+#if UE_INCLUDE_SOURCE_LOCATION
+		CHECK(Private::CheckFullFormatting(MakeStringView(String), TestLocation));
+#else
+		CHECK(String.IsEmpty());
+#endif
+	}
+
+	SECTION("UE::SourceLocation::Full (ToUtf8String)")
+	{
+		const FUtf8String String = UE::SourceLocation::Full(TestLocation).ToUtf8String();
+
+#if UE_INCLUDE_SOURCE_LOCATION
+		CHECK(Private::CheckFullFormatting(MakeStringView(String), TestLocation));
+#else
+		CHECK(String.IsEmpty());
+#endif
+	}
+
+	SECTION("UE::SourceLocation::Full (Wide String Builder)")
+	{
+		TWideStringBuilder<256> Builder;
+		Builder << UE::SourceLocation::Full(TestLocation);
+
+#if UE_INCLUDE_SOURCE_LOCATION
+		CHECK(Private::CheckFullFormatting(Builder.ToView(), TestLocation));
+#else
+		CHECK(Builder.Len() == 0);
+#endif
+	}
+
+	SECTION("UE::SourceLocation::Full (Utf8 String Builder)")
+	{
+		TUtf8StringBuilder<256> Builder;
+		Builder << UE::SourceLocation::Full(TestLocation);
+
+#if UE_INCLUDE_SOURCE_LOCATION
+		CHECK(Private::CheckFullFormatting(Builder.ToView(), TestLocation));
+#else
+		CHECK(Builder.Len() == 0);
+#endif
+	}
+
+	SECTION("UE::SourceLocation::Full (Compact Binary Writer)")
+	{
+		FCbWriter Writer;
+		SerializeForLog(Writer, UE::SourceLocation::Full(TestLocation));
+		FCbField Field = Writer.Save();
+
+#if UE_INCLUDE_SOURCE_LOCATION
+		CHECK(Private::ValidateBinaryForFull(Field, TestLocation));
+#else
+		REQUIRE(Field.IsObject());
+		// Use FCbObjectView::operator bool() to check if the object has any fields.
+		CHECK(!Field.AsObjectView());
+#endif
+	}
+
+	SECTION("UE::SourceLocation::FileAndLine (ToString)")
+	{
+		const FString String = UE::SourceLocation::FileAndLine(TestLocation).ToString();
+
+#if UE_INCLUDE_SOURCE_LOCATION
+		CHECK(Private::CheckFileAndLineFormatting(MakeStringView(String), TestLocation));
+#else
+		CHECK(String.IsEmpty());
+#endif
+	}
+
+	SECTION("UE::SourceLocation::FileAndLine (ToUtf8String)")
+	{
+		const FUtf8String String = UE::SourceLocation::FileAndLine(TestLocation).ToUtf8String();
+
+#if UE_INCLUDE_SOURCE_LOCATION
+		CHECK(Private::CheckFileAndLineFormatting(MakeStringView(String), TestLocation));
+#else
+		CHECK(String.IsEmpty());
+#endif
+	}
+
+	SECTION("UE::SourceLocation::FileAndLine (Wide String Builder)")
+	{
+		TWideStringBuilder<256> Builder;
+		Builder << UE::SourceLocation::FileAndLine(TestLocation);
+
+#if UE_INCLUDE_SOURCE_LOCATION
+		CHECK(Private::CheckFileAndLineFormatting(Builder.ToView(), TestLocation));
+#else
+		CHECK(Builder.Len() == 0);
+#endif
+	}
+
+	SECTION("UE::SourceLocation::FileAndLine (Utf8 String Builder)")
+	{
+		TUtf8StringBuilder<256> Builder;
+		Builder << UE::SourceLocation::FileAndLine(TestLocation);
+
+#if UE_INCLUDE_SOURCE_LOCATION
+		CHECK(Private::CheckFileAndLineFormatting(Builder.ToView(), TestLocation));
+#else
+		CHECK(Builder.Len() == 0);
+#endif
+	}
+
+	SECTION("UE::SourceLocation::FileAndLine (Compact Binary Writer)")
+	{
+		FCbWriter Writer;
+		SerializeForLog(Writer, UE::SourceLocation::FileAndLine(TestLocation));
+		FCbField Field = Writer.Save();
+
+#if UE_INCLUDE_SOURCE_LOCATION
+		CHECK(Private::ValidateBinaryForFileAndLine(Field, TestLocation));
+#else
+		REQUIRE(Field.IsObject());
+		// Use FCbObjectView::operator bool() to check if the object has any fields.
+		CHECK(!Field.AsObjectView());
+#endif
+	}
+}
+
+}  // namespace UE::SourceLocation
+
+#endif	// WITH_TESTS
