[AutoRTFM] Update StringTableRegistry to use transactionally-safe locking.

This code is reachable from Verse.

diff --git a/Engine/Source/Runtime/Core/Private/Internationalization/LegacyCulture.cpp b/Engine/Source/Runtime/Core/Private/Internationalization/LegacyCulture.cpp
--- a/Engine/Source/Runtime/Core/Private/Internationalization/LegacyCulture.cpp
+++ b/Engine/Source/Runtime/Core/Private/Internationalization/LegacyCulture.cpp
@@ -154,7 +154,7 @@ const FDecimalNumberFormattingRules& FLegacyCultureImplementation::GetCurrencyFo
 	}
 	else
 	{
-		FScopeLock MapLock(&UEAlternateCurrencyFormattingRulesCS);
+		UE::TScopeLock MapLock(UEAlternateCurrencyFormattingRulesCS);
 
 		auto FoundUEAlternateCurrencyFormattingRules = UEAlternateCurrencyFormattingRules.FindRef(InCurrencyCode);
 		if (FoundUEAlternateCurrencyFormattingRules.IsValid())
@@ -170,7 +170,7 @@ const FDecimalNumberFormattingRules& FLegacyCultureImplementation::GetCurrencyFo
 	NewUECurrencyFormattingRules.PositiveSuffixString.ReplaceInline(TEXT("$"), *InCurrencyCode, ESearchCase::CaseSensitive);
 
 	{
-		FScopeLock MapLock(&UEAlternateCurrencyFormattingRulesCS);
+		UE::TScopeLock MapLock(UEAlternateCurrencyFormattingRulesCS);
 
 		// Find again in case another thread beat us to it
 		auto FoundUEAlternateCurrencyFormattingRules = UEAlternateCurrencyFormattingRules.FindRef(InCurrencyCode);
diff --git a/Engine/Source/Runtime/Core/Private/Internationalization/LegacyCulture.h b/Engine/Source/Runtime/Core/Private/Internationalization/LegacyCulture.h
--- a/Engine/Source/Runtime/Core/Private/Internationalization/LegacyCulture.h
+++ b/Engine/Source/Runtime/Core/Private/Internationalization/LegacyCulture.h
@@ -8,6 +8,7 @@
 #include "Internationalization/Culture.h"
 #include "Internationalization/CultureImplementation.h"
 #include "Internationalization/FastDecimalFormat.h"
+#include "Misc/TransactionallySafeCriticalSection.h"
 
 #if !UE_ENABLE_ICU
 class FInternationalization;
@@ -101,6 +102,6 @@ private:
 
 	// Rules for formatting alternate currencies in this culture
 	TMap<FString, TSharedPtr<const FDecimalNumberFormattingRules>> UEAlternateCurrencyFormattingRules;
-	FCriticalSection UEAlternateCurrencyFormattingRulesCS;
+	FTransactionallySafeCriticalSection UEAlternateCurrencyFormattingRulesCS;
 };
 #endif
diff --git a/Engine/Source/Runtime/Core/Private/Internationalization/StringTableCore.cpp b/Engine/Source/Runtime/Core/Private/Internationalization/StringTableCore.cpp
--- a/Engine/Source/Runtime/Core/Private/Internationalization/StringTableCore.cpp
+++ b/Engine/Source/Runtime/Core/Private/Internationalization/StringTableCore.cpp
@@ -138,7 +138,7 @@ FString FStringTable::GetNamespace() const
 
 void FStringTable::SetNamespace(const FTextKey& InNamespace)
 {
-	FScopeLock KeyMappingLock(&KeyMappingCS);
+	UE::TScopeLock KeyMappingLock(KeyMappingCS);
 
 	if (TableNamespace != InNamespace)
 	{
@@ -157,7 +157,7 @@ void FStringTable::SetNamespace(const FTextKey& InNamespace)
 
 bool FStringTable::GetSourceString(const FTextKey& InKey, FString& OutSourceString) const
 {
-	FScopeLock KeyMappingLock(&KeyMappingCS);
+	UE::TScopeLock KeyMappingLock(KeyMappingCS);
 
 	FStringTableEntryPtr TableEntry = KeysToEntries.FindRef(InKey);
 	if (TableEntry.IsValid())
@@ -172,7 +172,7 @@ void FStringTable::SetSourceString(const FTextKey& InKey, const FString& InSourc
 {
 	checkf(!InKey.IsEmpty(), TEXT("String table key cannot be empty!"));
 
-	FScopeLock KeyMappingLock(&KeyMappingCS);
+	UE::TScopeLock KeyMappingLock(KeyMappingCS);
 	
 	FStringTableEntryPtr TableEntry = KeysToEntries.FindRef(InKey);
 	if (TableEntry.IsValid())
@@ -186,7 +186,7 @@ void FStringTable::SetSourceString(const FTextKey& InKey, const FString& InSourc
 
 void FStringTable::RemoveSourceString(const FTextKey& InKey)
 {
-	FScopeLock KeyMappingLock(&KeyMappingCS);
+	UE::TScopeLock KeyMappingLock(KeyMappingCS);
 
 	FStringTableEntryPtr TableEntry = KeysToEntries.FindRef(InKey);
 	if (TableEntry.IsValid())
@@ -207,7 +207,7 @@ void FStringTable::EnumerateSourceStrings(const TFunctionRef<bool(const FString&
 
 void FStringTable::EnumerateKeysAndSourceStrings(const TFunctionRef<bool(const FTextKey&, const FString&)>& InEnumerator) const
 {
-	FScopeLock KeyMappingLock(&KeyMappingCS);
+	UE::TScopeLock KeyMappingLock(KeyMappingCS);
 
 	for (const auto& KeyToEntryPair : KeysToEntries)
 	{
@@ -220,7 +220,7 @@ void FStringTable::EnumerateKeysAndSourceStrings(const TFunctionRef<bool(const F
 
 void FStringTable::ClearSourceStrings(const int32 InSlack)
 {
-	FScopeLock KeyMappingLock(&KeyMappingCS);
+	UE::TScopeLock KeyMappingLock(KeyMappingCS);
 
 	for (const auto& KeyToEntryPair : KeysToEntries)
 	{
@@ -234,7 +234,7 @@ void FStringTable::ClearSourceStrings(const int32 InSlack)
 
 FStringTableEntryConstPtr FStringTable::FindEntry(const FTextKey& InKey) const
 {
-	FScopeLock KeyMappingLock(&KeyMappingCS);
+	UE::TScopeLock KeyMappingLock(KeyMappingCS);
 	return KeysToEntries.FindRef(InKey);
 }
 
@@ -253,7 +253,7 @@ bool FStringTable::FindKey(const FStringTableEntryConstRef& InEntry, FTextKey& O
 {
 	if (InEntry->IsOwnedBy(*this))
 	{
-		FScopeLock KeyMappingLock(&KeyMappingCS);
+		UE::TScopeLock KeyMappingLock(KeyMappingCS);
 
 		for (const auto& KeyToEntryPair : KeysToEntries)
 		{
@@ -269,7 +269,7 @@ bool FStringTable::FindKey(const FStringTableEntryConstRef& InEntry, FTextKey& O
 
 FString FStringTable::GetMetaData(const FTextKey& InKey, const FName InMetaDataId) const
 {
-	FScopeLock MetaDataLock(&KeysToMetaDataCS);
+	UE::TScopeLock MetaDataLock(KeysToMetaDataCS);
 
 	const FMetaDataMap* MetaDataMap = KeysToMetaData.Find(InKey);
 	if (MetaDataMap)
@@ -281,7 +281,7 @@ FString FStringTable::GetMetaData(const FTextKey& InKey, const FName InMetaDataI
 
 void FStringTable::SetMetaData(const FTextKey& InKey, const FName InMetaDataId, const FString& InMetaDataValue)
 {
-	FScopeLock MetaDataLock(&KeysToMetaDataCS);
+	UE::TScopeLock MetaDataLock(KeysToMetaDataCS);
 
 	FMetaDataMap& MetaDataMap = KeysToMetaData.FindOrAdd(InKey);
 	MetaDataMap.Add(InMetaDataId, InMetaDataValue);
@@ -289,7 +289,7 @@ void FStringTable::SetMetaData(const FTextKey& InKey, const FName InMetaDataId,
 
 void FStringTable::RemoveMetaData(const FTextKey& InKey, const FName InMetaDataId)
 {
-	FScopeLock MetaDataLock(&KeysToMetaDataCS);
+	UE::TScopeLock MetaDataLock(KeysToMetaDataCS);
 
 	FMetaDataMap* MetaDataMap = KeysToMetaData.Find(InKey);
 	if (MetaDataMap)
@@ -304,7 +304,7 @@ void FStringTable::RemoveMetaData(const FTextKey& InKey, const FName InMetaDataI
 
 void FStringTable::EnumerateMetaData(const FTextKey& InKey, const TFunctionRef<bool(FName, const FString&)>& InEnumerator) const
 {
-	FScopeLock MetaDataLock(&KeysToMetaDataCS);
+	UE::TScopeLock MetaDataLock(KeysToMetaDataCS);
 
 	const FMetaDataMap* MetaDataMap = KeysToMetaData.Find(InKey);
 	if (MetaDataMap)
@@ -321,20 +321,20 @@ void FStringTable::EnumerateMetaData(const FTextKey& InKey, const TFunctionRef<b
 
 void FStringTable::ClearMetaData(const FTextKey& InKey)
 {
-	FScopeLock MetaDataLock(&KeysToMetaDataCS);
+	UE::TScopeLock MetaDataLock(KeysToMetaDataCS);
 	KeysToMetaData.Remove(InKey);
 }
 
 void FStringTable::ClearMetaData(const int32 InSlack)
 {
-	FScopeLock MetaDataLock(&KeysToMetaDataCS);
+	UE::TScopeLock MetaDataLock(KeysToMetaDataCS);
 	KeysToMetaData.Empty(InSlack);
 }
 
 void FStringTable::Serialize(FArchive& Ar)
 {
-	FScopeLock KeyMappingLock(&KeyMappingCS);
-	FScopeLock MetaDataLock(&KeysToMetaDataCS);
+	UE::TScopeLock KeyMappingLock(KeyMappingCS);
+	UE::TScopeLock MetaDataLock(KeysToMetaDataCS);
 
 	TableNamespace.SerializeAsString(Ar);
 
@@ -413,8 +413,8 @@ bool FStringTable::ExportStrings(const FString& InFilename) const
 	FString ExportedStrings;
 
 	{
-		FScopeLock KeyMappingLock(&KeyMappingCS);
-		FScopeLock MetaDataLock(&KeysToMetaDataCS);
+		UE::TScopeLock KeyMappingLock(KeyMappingCS);
+		UE::TScopeLock MetaDataLock(KeysToMetaDataCS);
 
 		// Collect meta-data column names
 		TSet<FName> MetaDataColumnNames;
@@ -541,8 +541,8 @@ bool FStringTable::ImportStrings(const FString& InFilename)
 
 	// Import rows
 	{
-		FScopeLock KeyMappingLock(&KeyMappingCS);
-		FScopeLock MetaDataLock(&KeysToMetaDataCS);
+		UE::TScopeLock KeyMappingLock(KeyMappingCS);
+		UE::TScopeLock MetaDataLock(KeysToMetaDataCS);
 
 		ClearSourceStrings(Rows.Num() - 1);
 		for (int32 RowIdx = 1; RowIdx < Rows.Num(); ++RowIdx)
diff --git a/Engine/Source/Runtime/Core/Private/Internationalization/StringTableRegistry.cpp b/Engine/Source/Runtime/Core/Private/Internationalization/StringTableRegistry.cpp
--- a/Engine/Source/Runtime/Core/Private/Internationalization/StringTableRegistry.cpp
+++ b/Engine/Source/Runtime/Core/Private/Internationalization/StringTableRegistry.cpp
@@ -62,7 +62,7 @@ FStringTableRegistry::~FStringTableRegistry()
 
 void FStringTableRegistry::RegisterStringTable(const FName InTableId, FStringTableRef InTable)
 {
-	FScopeLock RegisteredStringTablesLock(&RegisteredStringTablesCS);
+	UE::TScopeLock RegisteredStringTablesLock(RegisteredStringTablesCS);
 
 	checkf(!InTableId.IsNone(), TEXT("String table ID cannot be 'None'!"));
 #if DO_CHECK
@@ -78,7 +78,7 @@ void FStringTableRegistry::RegisterStringTable(const FName InTableId, FStringTab
 void FStringTableRegistry::UnregisterStringTable(const FName InTableId)
 {
 	{
-		FScopeLock RegisteredStringTablesLock(&RegisteredStringTablesCS);
+		UE::TScopeLock RegisteredStringTablesLock(RegisteredStringTablesCS);
 		RegisteredStringTables.Remove(InTableId);
 	}
 
@@ -98,13 +98,13 @@ void FStringTableRegistry::UnregisterStringTable(const FName InTableId)
 
 FStringTablePtr FStringTableRegistry::FindMutableStringTable(const FName InTableId) const
 {
-	FScopeLock RegisteredStringTablesLock(&RegisteredStringTablesCS);
+	UE::TScopeLock RegisteredStringTablesLock(RegisteredStringTablesCS);
 	return RegisteredStringTables.FindRef(InTableId);
 }
 
 FStringTableConstPtr FStringTableRegistry::FindStringTable(const FName InTableId) const
 {
-	FScopeLock RegisteredStringTablesLock(&RegisteredStringTablesCS);
+	UE::TScopeLock RegisteredStringTablesLock(RegisteredStringTablesCS);
 	return RegisteredStringTables.FindRef(InTableId);
 }
 
@@ -123,7 +123,7 @@ UStringTable* FStringTableRegistry::FindStringTableAsset(const FName InTableId)
 
 void FStringTableRegistry::EnumerateStringTables(const TFunctionRef<bool(const FName&, const FStringTableConstRef&)>& InEnumerator) const
 {
-	FScopeLock RegisteredStringTablesLock(&RegisteredStringTablesCS);
+	UE::TScopeLock RegisteredStringTablesLock(RegisteredStringTablesCS);
 
 	for (const auto& RegisteredStringTablePair : RegisteredStringTables)
 	{
@@ -141,7 +141,7 @@ bool FStringTableRegistry::FindTableIdAndKey(const FText& InText, FName& OutTabl
 
 void FStringTableRegistry::LogMissingStringTableEntry(const FName InTableId, const FTextKey& InKey)
 {
-	FScopeLock LoggedMissingEntriesLock(&LoggedMissingEntriesCS);
+	UE::TScopeLock LoggedMissingEntriesLock(LoggedMissingEntriesCS);
 
 	if (const FLocKeySet* LoggedMissingKeys = LoggedMissingEntries.Find(InTableId))
 	{
diff --git a/Engine/Source/Runtime/Core/Public/Internationalization/StringTableCore.h b/Engine/Source/Runtime/Core/Public/Internationalization/StringTableCore.h
--- a/Engine/Source/Runtime/Core/Public/Internationalization/StringTableCore.h
+++ b/Engine/Source/Runtime/Core/Public/Internationalization/StringTableCore.h
@@ -5,7 +5,6 @@
 #include "Containers/SparseArray.h"
 #include "Containers/UnrealString.h"
 #include "CoreTypes.h"
-#include "HAL/CriticalSection.h"
 #include "Internationalization/Internationalization.h"
 #include "Internationalization/LocKeyFuncs.h"
 #include "Internationalization/LocalizedTextSourceTypes.h"
@@ -14,6 +13,7 @@
 #include "Internationalization/TextKey.h"
 #include "Logging/LogMacros.h"
 #include "Misc/AssertionMacros.h"
+#include "Misc/TransactionallySafeCriticalSection.h"
 #include "Serialization/Archive.h"
 #include "Serialization/StructuredArchive.h"
 #include "Templates/Function.h"
@@ -183,14 +183,14 @@ private:
 	TMap<FTextKey, FStringTableEntryPtr> KeysToEntries;
 
 	/** Critical section preventing concurrent modification of KeysToEntries */
-	mutable FCriticalSection KeyMappingCS;
+	mutable FTransactionallySafeCriticalSection KeyMappingCS;
 
 	/** Mapping between the text key and its meta-data map */
 	typedef TMap<FName, FString> FMetaDataMap;
 	TMap<FTextKey, FMetaDataMap> KeysToMetaData;
 
 	/** Critical section preventing concurrent modification of KeysToMetaData */
-	mutable FCriticalSection KeysToMetaDataCS;
+	mutable FTransactionallySafeCriticalSection KeysToMetaDataCS;
 };
 
 /** Interface to allow Core code to access String Table assets from the Engine */
diff --git a/Engine/Source/Runtime/Core/Public/Internationalization/StringTableRegistry.h b/Engine/Source/Runtime/Core/Public/Internationalization/StringTableRegistry.h
--- a/Engine/Source/Runtime/Core/Public/Internationalization/StringTableRegistry.h
+++ b/Engine/Source/Runtime/Core/Public/Internationalization/StringTableRegistry.h
@@ -14,6 +14,7 @@
 #include "Internationalization/Text.h"
 #include "Internationalization/TextKey.h"
 #include "Internationalization/TextLocalizationManager.h"
+#include "Misc/TransactionallySafeCriticalSection.h"
 #include "UObject/NameTypes.h"
 
 class UStringTable;
@@ -75,14 +76,14 @@ private:
 	TMap<FName, FStringTablePtr> RegisteredStringTables;
 
 	/** Critical section preventing concurrent modification of RegisteredStringTables */
-	mutable FCriticalSection RegisteredStringTablesCS;
+	mutable FTransactionallySafeCriticalSection RegisteredStringTablesCS;
 
 	/** Mapping from a table ID to a set of keys that we've logged warnings for */
 	typedef TSet<FTextKey> FLocKeySet;
 	TMap<FName, FLocKeySet> LoggedMissingEntries;
 
 	/** Critical section preventing concurrent modification of LoggedMissingEntries */
-	mutable FCriticalSection LoggedMissingEntriesCS;
+	mutable FTransactionallySafeCriticalSection LoggedMissingEntriesCS;
 
 #if WITH_EDITOR
 	/** Callback handler for a directory change notification */
diff --git a/Engine/Source/Runtime/Core/Tests/Internationalization/TextTest.cpp b/Engine/Source/Runtime/Core/Tests/Internationalization/TextTest.cpp
--- a/Engine/Source/Runtime/Core/Tests/Internationalization/TextTest.cpp
+++ b/Engine/Source/Runtime/Core/Tests/Internationalization/TextTest.cpp
@@ -1162,7 +1162,7 @@ TEST_CASE_NAMED(FTextNumericParsingTest, "System::Core::Misc::TextNumericParsing
 	Instance.TextNumericParsingTest();
 }
 
-TEST_CASE_NAMED(FTextStringificationTest, "System::Core::Misc::TextStringification", "[.][EditorContext][ClientContext][EngineFilter]")
+static void RunTestCasesForSystem_Core_Misc_TextStringification()
 {
 	FInternationalization& I18N = FInternationalization::Get();
 
@@ -1210,6 +1210,14 @@ TEST_CASE_NAMED(FTextStringificationTest, "System::Core::Misc::TextStringificati
 #define TEST(Text, Str) DoSingleTest(Text, TEXT(Str), TEXT(#Text), true)
 #define TEST_EX(Text, Str, ImportCpp) DoSingleTest(Text, TEXT(Str), TEXT(#Text), ImportCpp)
 
+#if UE_ENABLE_ICU
+	#define TEST_ICU(Text, Str)               TEST(Text, Str)
+	#define TEST_EX_ICU(Text, Str, ImportCpp) TEST_EX(Text, Str, ImportCpp)
+#else
+	#define TEST_ICU(Text, Str)              
+	#define TEST_EX_ICU(Text, Str, ImportCpp)
+#endif
+
 	// Add the test string table, but only if it isn't already!
 	if (!FStringTableRegistry::Get().FindStringTable("Core.Tests.TextFormatTest"))
 	{
@@ -1234,27 +1242,27 @@ TEST_CASE_NAMED(FTextStringificationTest, "System::Core::Misc::TextStringificati
 	TEST(LOCGEN_NUMBER_GROUPED(12.5f, ""), "LOCGEN_NUMBER_GROUPED(12.500000f, \"\")");
 	TEST(LOCGEN_NUMBER_UNGROUPED(12.5f, ""), "LOCGEN_NUMBER_UNGROUPED(12.500000f, \"\")");
 	TEST(LOCGEN_NUMBER_CUSTOM(+10, SetAlwaysSign(true).SetRoundingMode(ERoundingMode::ToZero).SetMinimumFractionalDigits(2), ""), "LOCGEN_NUMBER_CUSTOM(10, SetAlwaysSign(true).SetRoundingMode(ERoundingMode::ToZero).SetMinimumFractionalDigits(2), \"\")");
-	TEST(LOCGEN_NUMBER(-10, "en"), "LOCGEN_NUMBER(-10, \"en\")");
+	TEST_ICU(LOCGEN_NUMBER(-10, "en"), "LOCGEN_NUMBER(-10, \"en\")");
 
 	TEST(LOCGEN_PERCENT(0.1f, ""), "LOCGEN_PERCENT(0.100000f, \"\")");
 	TEST(LOCGEN_PERCENT_GROUPED(0.1f, ""), "LOCGEN_PERCENT_GROUPED(0.100000f, \"\")");
 	TEST(LOCGEN_PERCENT_UNGROUPED(0.1f, ""), "LOCGEN_PERCENT_UNGROUPED(0.100000f, \"\")");
 	TEST(LOCGEN_PERCENT_CUSTOM(0.1f, SetAlwaysSign(true).SetRoundingMode(ERoundingMode::ToZero).SetMinimumFractionalDigits(2), ""), "LOCGEN_PERCENT_CUSTOM(0.100000f, SetAlwaysSign(true).SetRoundingMode(ERoundingMode::ToZero).SetMinimumFractionalDigits(2), \"\")");
-	TEST(LOCGEN_PERCENT(0.1, "en"), "LOCGEN_PERCENT(0.100000, \"en\")");
+	TEST_ICU(LOCGEN_PERCENT(0.1, "en"), "LOCGEN_PERCENT(0.100000, \"en\")");
 
 	TEST(LOCGEN_CURRENCY(125, "USD", ""), "LOCGEN_CURRENCY(125, \"USD\", \"\")");
-	TEST_EX(FText::AsCurrency(1.25f, TEXT("USD"), nullptr, FInternationalization::Get().GetCulture(TEXT("en"))), "LOCGEN_CURRENCY(125, \"USD\", \"en\")", false);
+	TEST_EX_ICU(FText::AsCurrency(1.25f, TEXT("USD"), nullptr, FInternationalization::Get().GetCulture(TEXT("en"))), "LOCGEN_CURRENCY(125, \"USD\", \"en\")", false);
 
-	TEST(LOCGEN_DATE_UTC(1526342400, EDateTimeStyle::Short, "", "en-GB"), "LOCGEN_DATE_UTC(1526342400, EDateTimeStyle::Short, \"\", \"en-GB\")");
+	TEST_ICU(LOCGEN_DATE_UTC(1526342400, EDateTimeStyle::Short, "", "en-GB"), "LOCGEN_DATE_UTC(1526342400, EDateTimeStyle::Short, \"\", \"en-GB\")");
 	TEST(LOCGEN_DATE_LOCAL(1526342400, EDateTimeStyle::Medium, ""), "LOCGEN_DATE_LOCAL(1526342400, EDateTimeStyle::Medium, \"\")");
 
-	TEST(LOCGEN_TIME_UTC(1526342400, EDateTimeStyle::Long, "", "en-GB"), "LOCGEN_TIME_UTC(1526342400, EDateTimeStyle::Long, \"\", \"en-GB\")");
+	TEST_ICU(LOCGEN_TIME_UTC(1526342400, EDateTimeStyle::Long, "", "en-GB"), "LOCGEN_TIME_UTC(1526342400, EDateTimeStyle::Long, \"\", \"en-GB\")");
 	TEST(LOCGEN_TIME_LOCAL(1526342400, EDateTimeStyle::Full, ""), "LOCGEN_TIME_LOCAL(1526342400, EDateTimeStyle::Full, \"\")");
 
-	TEST(LOCGEN_DATETIME_UTC(1526342400, EDateTimeStyle::Short, EDateTimeStyle::Medium, "", "en-GB"), "LOCGEN_DATETIME_UTC(1526342400, EDateTimeStyle::Short, EDateTimeStyle::Medium, \"\", \"en-GB\")");
+	TEST_ICU(LOCGEN_DATETIME_UTC(1526342400, EDateTimeStyle::Short, EDateTimeStyle::Medium, "", "en-GB"), "LOCGEN_DATETIME_UTC(1526342400, EDateTimeStyle::Short, EDateTimeStyle::Medium, \"\", \"en-GB\")");
 	TEST(LOCGEN_DATETIME_LOCAL(1526342400, EDateTimeStyle::Long, EDateTimeStyle::Full, ""), "LOCGEN_DATETIME_LOCAL(1526342400, EDateTimeStyle::Long, EDateTimeStyle::Full, \"\")");
 
-	TEST(LOCGEN_DATETIME_CUSTOM_UTC(1526342400, "%A, %B %e, %Y", "", "en-GB"), "LOCGEN_DATETIME_CUSTOM_UTC(1526342400, \"%A, %B %e, %Y\", \"\", \"en-GB\")");
+	TEST_ICU(LOCGEN_DATETIME_CUSTOM_UTC(1526342400, "%A, %B %e, %Y", "", "en-GB"), "LOCGEN_DATETIME_CUSTOM_UTC(1526342400, \"%A, %B %e, %Y\", \"\", \"en-GB\")");
 	TEST(LOCGEN_DATETIME_CUSTOM_LOCAL(1526342400, "%A, %B %e, %Y", ""), "LOCGEN_DATETIME_CUSTOM_LOCAL(1526342400, \"%A, %B %e, %Y\", \"\")");
 
 	TEST(LOCGEN_TOUPPER(LOCTEXT("TextStringificationTest_Lorem", "Lorem")), "LOCGEN_TOUPPER(NSLOCTEXT(\"Core.Tests.TextFormatTest\", \"TextStringificationTest_Lorem\", \"Lorem\"))");
@@ -1265,11 +1273,18 @@ TEST_CASE_NAMED(FTextStringificationTest, "System::Core::Misc::TextStringificati
 
 #undef TEST
 #undef TEST_EX
+#undef TEST_ICU
+#undef TEST_EX_ICU
 
 	// Restore original culture
 	I18N.RestoreCultureState(OriginalCultureState);
 }
 
+TEST_CASE_NAMED(FTextStringificationTest, "System::Core::Misc::TextStringification", "[.][EditorContext][ClientContext][EngineFilter]")
+{
+	RunTestCasesForSystem_Core_Misc_TextStringification();
+}
+
 static void RunTestCasesForSystem_Core_Misc_TextFormatArgModifiers()
 {
 	auto EnsureValidResult = [&](const FString& InResult, const FString& InExpected, const FString& InName, const FString& InDescription)
@@ -1574,8 +1589,6 @@ bool FAutoRTFMTextTests::RunTest(const FString& Parameters)
 	}
 
 	// Verify that FText works properly when used in an AutoRTFM transaction, both committing and aborting.
-	// System::Core::Misc::TextStringification is not included, because this test uses the StringTableRegistry.
-	// (The StringTableRegistry is not yet AutoRTFM-safe.)
 	AutoRTFM::Commit([]
 	{
 		FTextNumericParsingTestClass Instance;
@@ -1600,6 +1613,17 @@ bool FAutoRTFMTextTests::RunTest(const FString& Parameters)
 		AutoRTFM::AbortTransaction();
 	});
 
+	AutoRTFM::Commit([]
+	{
+		RunTestCasesForSystem_Core_Misc_TextStringification();
+	});
+
+	AutoRTFM::Transact([]
+	{
+		RunTestCasesForSystem_Core_Misc_TextStringification();
+		AutoRTFM::AbortTransaction();
+	});
+
 	return true;
 }
 
