Fix hitch stat scopes not emitting stat names.

#rnx
#tests verified hitch scopes work locally. Verified traces emit events correctly still via replay run. Tested no measurable additional perf overhead.

diff --git a/Engine/Source/Runtime/Core/Public/Stats/LightweightStats.h b/Engine/Source/Runtime/Core/Public/Stats/LightweightStats.h
--- a/Engine/Source/Runtime/Core/Public/Stats/LightweightStats.h
+++ b/Engine/Source/Runtime/Core/Public/Stats/LightweightStats.h
@@ -53,16 +53,59 @@ struct TStatId
 
 namespace UE::Stats::Private
 {
-	struct FCheckedStat : TStatId
+	// A TStatId that we've already checked can be emitted (via GetStatIfEnabled).
+	struct FCheckedStat
 	{
-		using TStatId::TStatId;
+		// Validity is separate from the stat name being not-null
+		// as even if the stat shouldn't be emitted as a named event, we may need the stat name for other
+		// things such as the FHitchTrackingStatScope.
+		// However, it's worth noting the StatString can be null for cases where we do not want to emit a hitch scope
+		// which is the case for conditional scopes when the condition is false.
+		static constexpr FCheckedStat MakeValid(const PROFILER_CHAR* InStatString)
+		{
+			return FCheckedStat{ InStatString, true };
+		}
+
+		static constexpr FCheckedStat MakeInvalid(const PROFILER_CHAR* InStatString)
+		{
+			return FCheckedStat{ InStatString, false };
+		}
+
+		// Allow implicit conversions to TStatId as FCheckedStat can be assigned to TStatId via GetStatIfEnabled.
+		constexpr operator TStatId() const
+		{
+			// Systems that use TStatId directly test validity through the stat string.
+			return bIsValid ? TStatId(StatString) : TStatId{};
+		}
+
+		// Returns the stat name, even if this stat isn't valid.
+		constexpr const PROFILER_CHAR* GetName() const
+		{
+			return StatString;
+		}
+
+		// Returns true if this stat should be emitted.
+		constexpr bool IsValidStat() const
+		{
+			return bIsValid;
+		}
+
+	private:
+		constexpr explicit FCheckedStat(const PROFILER_CHAR* InStatString, bool bValidStat)
+			: StatString(InStatString)
+			, bIsValid(bValidStat)
+		{
+		}
+
+		const PROFILER_CHAR* StatString;
+		bool bIsValid;
 	};
 
-	template<typename DerivedScopeCounterType, typename StatType = TStatId>
+	template<typename DerivedScopeCounterType, typename StatIdType = TStatId>
 	class TScopeCycleCounterBase
 	{
 	public:
-		FORCEINLINE explicit TScopeCycleCounterBase(StatType InStatId)
+		FORCEINLINE constexpr explicit TScopeCycleCounterBase(const StatIdType& InStatId)
 #if USE_LIGHTWEIGHT_STATS_FOR_HITCH_DETECTION && USE_HITCH_DETECTION
 			: StatScope(InStatId.GetName())
 #endif
@@ -105,14 +148,13 @@ namespace UE::Stats::Private
 	/**
 	 * Scope counter for stat names with static storage.
 	 */
-	template<typename StatIdType>
-	class FScopeCycleCounterStatic : TScopeCycleCounterBase<FScopeCycleCounterStatic<StatIdType>, StatIdType>
+	class FScopeCycleCounterStatic : TScopeCycleCounterBase<FScopeCycleCounterStatic, FCheckedStat>
 	{
-		using Base = TScopeCycleCounterBase<FScopeCycleCounterStatic<StatIdType>, StatIdType>;
+		using Base = TScopeCycleCounterBase<FScopeCycleCounterStatic, FCheckedStat>;
 	public:
 		using Base::CanEmitStat;
 
-		FORCEINLINE constexpr explicit FScopeCycleCounterStatic(StatIdType InStatId)
+		FORCEINLINE constexpr explicit FScopeCycleCounterStatic(const FCheckedStat& InStatId)
 			: Base(InStatId)
 		{
 		}
@@ -301,9 +343,9 @@ namespace UE::Stats::Private
 		if (GCycleStatsShouldEmitNamedEvents > 0 &&
 			IsStatEnabled<StatDataType, GroupDataType>(StatNameHash, GroupNameHash))
 		{
-			return FCheckedStat{ StatName };
+			return FCheckedStat::MakeValid(StatName);
 		}
-		return FCheckedStat{};
+		return FCheckedStat::MakeInvalid(StatName);
 	}
 
 	/* Returns a valid TStatId if the condition is true and we can emit it. */
@@ -314,7 +356,8 @@ namespace UE::Stats::Private
 		{
 			return GetStatIfEnabled<StatDataType, GroupDataType>(StatName, StatNameHash, GroupNameHash);
 		}
-		return FCheckedStat{};
+		// Preserve existing behavior of not emitting hitch scopes if the condition is false.
+		return FCheckedStat::MakeInvalid(bCondition ? StatName : nullptr);
 	}
 } // namespace UE::Stats::Private
 
