Add crash callstacks to crash report

#rnx

diff --git a/Engine/Source/Runtime/Core/Private/Apple/ApplePlatformCrashContext.cpp b/Engine/Source/Runtime/Core/Private/Apple/ApplePlatformCrashContext.cpp
--- a/Engine/Source/Runtime/Core/Private/Apple/ApplePlatformCrashContext.cpp
+++ b/Engine/Source/Runtime/Core/Private/Apple/ApplePlatformCrashContext.cpp
@@ -89,8 +89,11 @@ int32 FApplePlatformCrashContext::ReportCrash() const
 		const_cast<FApplePlatformCrashContext*>(this)->CapturePortableCallStack(ErrorFrame, Context);
 		// Walk the stack and dump it to the allocated memory (ignore first 2 callstack lines as those are in stack walking code)
 		FPlatformStackWalk::StackWalkAndDump( StackBuffer, UE_ARRAY_COUNT(MinidumpCallstackInfo) - 1, ErrorFrame, Context );
+		// We will put this in xml, need to escape it
+		FString EscapedStack;
+		AppendEscapedXMLString(EscapedStack, FString(MinidumpCallstackInfo));
+		FCString::Strncpy(GErrorHist, *EscapedStack, UE_ARRAY_COUNT(GErrorHist));
 		
-		FUTF8ToTCHAR_Convert::Convert(GErrorHist, UE_ARRAY_COUNT(GErrorHist) - 1, MinidumpCallstackInfo, FCStringAnsi::Strlen(MinidumpCallstackInfo));
 		CreateExceptionInfoString(Signal, Info);
 	}
 	
@@ -238,3 +241,10 @@ TCHAR* FApplePlatformCrashContext::ItoTCHAR(uint64 Val, uint64 Base, uint32 Len)
 	
 	return &InternalBuffer[i+1];
 }
+
+
+const TCHAR* FApplePlatformCrashContext::GetCallstackProperty() const
+{
+	// FApplePlatformCrashContext::ReportCrash() already parsed callstack into GErrorHist at this point
+	return GErrorHist;
+}
diff --git a/Engine/Source/Runtime/Core/Public/Apple/ApplePlatformCrashContext.h b/Engine/Source/Runtime/Core/Public/Apple/ApplePlatformCrashContext.h
--- a/Engine/Source/Runtime/Core/Public/Apple/ApplePlatformCrashContext.h
+++ b/Engine/Source/Runtime/Core/Public/Apple/ApplePlatformCrashContext.h
@@ -66,4 +66,10 @@ struct CORE_API FApplePlatformCrashContext : public FGenericCrashContext
 	
 	/** Async-safe ItoT */
 	static TCHAR* ItoTCHAR(uint64 Val, uint64 Base, uint32 Len = 0);
+	
+
+protected:
+
+	/** Allow platform implementations to provide a callstack property. Primarily used when non-native code triggers a crash. */
+	CORE_API virtual const TCHAR* GetCallstackProperty() const;
 };
