Change ProcessUntilTasksComplete to periodically call an update function that can set the current operation while waiting for tasks on other threads.
Add a new tick.IdleTaskWorkMS cvar to control this, if > 0 the game thread will spend that many ms trying to process other work (like worker thread tasks) when the game thread is idle.
The default value of 0 will tell the game thread to wait as soon as possible, while -1 makes it use WaitUntilTasksComplete like in 5.5.

diff --git a/Engine/Source/Runtime/Core/Private/Async/TaskGraph.cpp b/Engine/Source/Runtime/Core/Private/Async/TaskGraph.cpp
--- a/Engine/Source/Runtime/Core/Private/Async/TaskGraph.cpp
+++ b/Engine/Source/Runtime/Core/Private/Async/TaskGraph.cpp
@@ -1459,27 +1459,30 @@ private:
 			return;
 		}
 
-		ENamedThreads::Type CurrentThread = CurrentThreadIfKnown;
+		ENamedThreads::Type NamedThreadWithFlags = CurrentThreadIfKnown;
 		if (ENamedThreads::GetThreadIndex(CurrentThreadIfKnown) == ENamedThreads::AnyThread)
 		{
 			bool bIsHiPri = !!ENamedThreads::GetTaskPriority(CurrentThreadIfKnown);
 			int32 Priority = ENamedThreads::GetThreadPriorityIndex(CurrentThreadIfKnown);
 			check(!ENamedThreads::GetQueueIndex(CurrentThreadIfKnown));
 			CurrentThreadIfKnown = ENamedThreads::GetThreadIndex(GetCurrentThread());
-			CurrentThread = ENamedThreads::SetPriorities(CurrentThreadIfKnown, Priority, bIsHiPri);
+			NamedThreadWithFlags = ENamedThreads::SetPriorities(CurrentThreadIfKnown, Priority, bIsHiPri);
 		}
 		else
 		{
 			CurrentThreadIfKnown = ENamedThreads::GetThreadIndex(CurrentThreadIfKnown);
 			check(CurrentThreadIfKnown == ENamedThreads::GetThreadIndex(GetCurrentThread()));
-			// we don't modify CurrentThread here because it might be a local queue
 		}
 
-		if (CurrentThreadIfKnown != ENamedThreads::AnyThread && CurrentThreadIfKnown < NumNamedThreads && !IsThreadProcessingTasks(CurrentThread))
+		WaitOnNamedThreadForTasks(Tasks, CurrentThreadIfKnown, NamedThreadWithFlags);
+	}
+
+	void WaitOnNamedThreadForTasks(const FGraphEventArray& Tasks, ENamedThreads::Type BaseNamedThread, ENamedThreads::Type NamedThreadWithFlags)
+	{
+		if (BaseNamedThread != ENamedThreads::AnyThread && BaseNamedThread < NumNamedThreads && !IsThreadProcessingTasks(NamedThreadWithFlags))
 		{
-			// we wouldn't get past the loop above if all tasks were complete so process named tasks while we wait
-			TGraphTask<FReturnGraphTask>::CreateTask(&Tasks, CurrentThread).ConstructAndDispatchWhenReady(CurrentThread);
-			ProcessThreadUntilRequestReturn(CurrentThread);
+			TGraphTask<FReturnGraphTask>::CreateTask(&Tasks, NamedThreadWithFlags).ConstructAndDispatchWhenReady(NamedThreadWithFlags);
+			ProcessThreadUntilRequestReturn(NamedThreadWithFlags);
 			return;
 		}
 
@@ -1504,51 +1507,123 @@ private:
 
 		// We will just stall this thread on an event while we wait
 		FScopedEvent Event;
-		TriggerEventWhenTasksComplete(Event.Get(), Tasks, CurrentThreadIfKnown);
+		TriggerEventWhenTasksComplete(Event.Get(), Tasks, BaseNamedThread);
 	}
 
-	void ProcessUntilTasksComplete(const FGraphEventArray& Tasks, ENamedThreads::Type CurrentThreadIfKnown, const TFunction<void()>& ExecuteOnIdle) final override
+	bool ProcessUntilTasksComplete(const FGraphEventArray& Tasks, ENamedThreads::Type CurrentThreadIfKnown, const FProcessTasksUpdateCallback& IdleWorkUpdate) final override
 	{
-		ENamedThreads::Type CurrentThread = CurrentThreadIfKnown;
+		TRACE_CPUPROFILER_EVENT_SCOPE(ProcessUntilTasksComplete);
+
+		ENamedThreads::Type NamedThreadWithFlags = CurrentThreadIfKnown;
 		if (ENamedThreads::GetThreadIndex(CurrentThreadIfKnown) == ENamedThreads::AnyThread)
 		{
 			bool bIsHiPri = !!ENamedThreads::GetTaskPriority(CurrentThreadIfKnown);
 			int32 Priority = ENamedThreads::GetThreadPriorityIndex(CurrentThreadIfKnown);
 			check(!ENamedThreads::GetQueueIndex(CurrentThreadIfKnown));
 			CurrentThreadIfKnown = ENamedThreads::GetThreadIndex(GetCurrentThread());
-			CurrentThread = ENamedThreads::SetPriorities(CurrentThreadIfKnown, Priority, bIsHiPri);
+			NamedThreadWithFlags = ENamedThreads::SetPriorities(CurrentThreadIfKnown, Priority, bIsHiPri);
+		}
+		else
+		{
+			CurrentThreadIfKnown = ENamedThreads::GetThreadIndex(CurrentThreadIfKnown);
+			check(CurrentThreadIfKnown == ENamedThreads::GetThreadIndex(GetCurrentThread()));
 		}
 
-		// First process until this named thread (could be a local queue) is idle, which should complete most of the tasks
-		ProcessThreadUntilIdle(CurrentThread);
+		// Copy into faster array to avoid ref count changes, these will all be in Tasks array
+		TArray<FGraphEvent*> RemainingTasks;
+		RemainingTasks.SetNumUninitialized(Tasks.Num());
+		for (int32 TaskIndex = 0; TaskIndex < Tasks.Num(); TaskIndex++)
+		{
+			RemainingTasks[TaskIndex] = Tasks[TaskIndex].GetReference();
+			if (!ensure(RemainingTasks[TaskIndex] != nullptr))
+			{
+				// Fail if one of the events is not initialized
+				UE_LOG(LogTaskGraph, Error, TEXT("ProcessUntilTasksComplete was passed an invalid event!"));
+				return false;
+			}
+		}
 
-		// See which tasks are still active
-		FGraphEventArray RemainingTasks;
-		for (const FGraphEventRef& Task : Tasks)
+		// Always start by processing named thread tasks once
+		EProcessTasksOperation CurrentOperation = EProcessTasksOperation::ProcessNamedThreadTasks;
+		
+		const UE::FTimeout NeverTimeout = UE::FTimeout::Never();		
+		bool bHasOtherTasks = true;
+
+		while (CurrentOperation != EProcessTasksOperation::StopProcessing)
 		{
-			if (Task.IsValid())
+			if (CurrentOperation == EProcessTasksOperation::ProcessNamedThreadTasks)
 			{
-				if (!Task->IsCompleted())
+				// Process until this named thread (could be a local queue) is idle, which should complete most of the tasks
+				ProcessThreadUntilIdle(NamedThreadWithFlags);
+			}
+			else if (CurrentOperation == EProcessTasksOperation::ProcessAllOtherTasks || CurrentOperation == EProcessTasksOperation::ProcessOneOtherTask)
+			{
+				bHasOtherTasks = false;
+				for (int32 TaskIndex = 0; TaskIndex < RemainingTasks.Num(); TaskIndex++)
 				{
-					RemainingTasks.Add(Task);
+					FGraphEvent* Task = RemainingTasks[TaskIndex];
+					// Don't process any named thread tasks here
+					if (!Task->IsNamedThreadTask())
+					{
+						bHasOtherTasks = true;
+						if (Task->TryRetractAndExecute(NeverTimeout) && CurrentOperation == EProcessTasksOperation::ProcessOneOtherTask)
+						{
+							break; // Completed a task, check logic again
+						}
+					}
 				}
 			}
-		}
+			else if (CurrentOperation == EProcessTasksOperation::WaitUntilComplete)
+			{
+				FGraphEventArray TasksToWaitOn;
+				for (FGraphEvent* TaskEvent : RemainingTasks)
+				{
+					TasksToWaitOn.Add(TaskEvent->GetCompletionEvent());
+				}
 
-		if (RemainingTasks.Num() == 0)
-		{
-			// They're all done so no need to wait
-			return;
-		}
+				TaskTrace::FWaitingScope WaitingScope(GetTraceIds(TasksToWaitOn));
+				WaitOnNamedThreadForTasks(TasksToWaitOn, CurrentThreadIfKnown, NamedThreadWithFlags);
 
-		// Call the idle function, which can be used to try and resolve deadlocks
-		if (ExecuteOnIdle)
-		{
-			ExecuteOnIdle();
+				// All tasks must have succeeded
+				return true;
+			}
+
+			// Update remaining tasks list
+			for (int32 TaskIndex = RemainingTasks.Num() - 1; TaskIndex >= 0; TaskIndex--)
+			{
+				if (RemainingTasks[TaskIndex]->IsCompleted())
+				{
+					RemainingTasks.RemoveAtSwap(TaskIndex);
+				}
+			}
+
+			if (RemainingTasks.Num() == 0)
+			{
+				// All complete
+				return true;
+			}
+
+			// Run callback function if it exists, otherwise default to ProcessAllOtherTasks
+			if (IdleWorkUpdate)
+			{
+				CurrentOperation = IdleWorkUpdate(RemainingTasks.Num());
+				check(CurrentOperation >= EProcessTasksOperation::ProcessAllOtherTasks && CurrentOperation <= EProcessTasksOperation::StopProcessing);
+			}
+			else
+			{
+				CurrentOperation = EProcessTasksOperation::ProcessAllOtherTasks;
+			}
+
+			if ((CurrentOperation == EProcessTasksOperation::ProcessAllOtherTasks || CurrentOperation == EProcessTasksOperation::ProcessOneOtherTask)
+				&& !bHasOtherTasks)
+			{
+				// If we have no more tasks to try and retract, just wait which will handle processing named thread tasks as well
+				CurrentOperation = EProcessTasksOperation::WaitUntilComplete;
+			}
 		}
 
-		// Now wait on this named thread until all the tasks are done
-		WaitUntilTasksComplete(RemainingTasks, CurrentThread);
+		// Stopped due to callback, some tasks may not be complete
+		return false;
 	}
 
 	void TriggerEventWhenTasksComplete(FEvent* InEvent, const FGraphEventArray& Tasks, ENamedThreads::Type CurrentThreadIfKnown = ENamedThreads::AnyThread, ENamedThreads::Type TriggerThread = ENamedThreads::AnyHiPriThreadHiPriTask) final override
diff --git a/Engine/Source/Runtime/Core/Public/Async/TaskGraphInterfaces.h b/Engine/Source/Runtime/Core/Public/Async/TaskGraphInterfaces.h
--- a/Engine/Source/Runtime/Core/Public/Async/TaskGraphInterfaces.h
+++ b/Engine/Source/Runtime/Core/Public/Async/TaskGraphInterfaces.h
@@ -377,14 +377,38 @@ public:
 	**/
 	virtual void WaitUntilTasksComplete(const FGraphEventArray& Tasks, ENamedThreads::Type CurrentThreadIfKnown = ENamedThreads::AnyThread)=0;
 	
+
+	/** Used to define what ProcessUntilTasksComplete should do next */
+	enum class EProcessTasksOperation : int8
+	{
+		/** Default behavior with no update callback, try to process all other thread thread tasks in the queue before checking named thread tasks */
+		ProcessAllOtherTasks,
+
+		/** Try to process other thread tasks, but immediately process named thread tasks after completing one */
+		ProcessOneOtherTask,
+
+		/** Immediately try to process named thread tasks and call update again */
+		ProcessNamedThreadTasks,
+
+		/** Stop any idle processing and wait on this thread until all tasks are complete, update is not called again */
+		WaitUntilComplete,
+
+		/** Stop processing entirely and return, like if ProcessingTimeout has expired */
+		StopProcessing,
+	};
+
+	/** Function periodically called during task processing, with parameter set to number of tasks remaining */
+	typedef TFunction<EProcessTasksOperation(int32 /*NumTasksRemaining*/)> FProcessTasksUpdateCallback;
+
 	/**
-	 *	Requests that a named thread actively attempt to process a list of tasks, and calls a function when idle.
+	 *	Requests that a named thread actively attempt to process a list of tasks and periodically call an update function when idle.
 	 *	This function assumes that many of the tasks can be completed on this thread.
 	 *	@param	Tasks - tasks to wait for
 	 *	@param	CurrentThread - This thread, must be a named thread
-	 *	@param	ExecuteOnIdle - If set, call this function the first time this thread has run out of work to do, before continuing to wait
+	 *	@param	IdleWorkUpdate - If set, call this function after trying to process named thread work if any tasks are stalled
+	 *	@return true if all tasks were completed, return false if StopProcessing was returned or it failed to process all tasks
 	**/
-	virtual void ProcessUntilTasksComplete(const FGraphEventArray& Tasks, ENamedThreads::Type CurrentThreadIfKnown, const TFunction<void()>& ExecuteOnIdle = {}) = 0;
+	virtual bool ProcessUntilTasksComplete(const FGraphEventArray& Tasks, ENamedThreads::Type CurrentThreadIfKnown, const FProcessTasksUpdateCallback& IdleWorkUpdate = {}) = 0;
 
 	/** 
 	 *	When a set of tasks complete, fire a scoped event
diff --git a/Engine/Source/Runtime/Engine/Private/TickTaskManager.cpp b/Engine/Source/Runtime/Engine/Private/TickTaskManager.cpp
--- a/Engine/Source/Runtime/Engine/Private/TickTaskManager.cpp
+++ b/Engine/Source/Runtime/Engine/Private/TickTaskManager.cpp
@@ -89,6 +89,15 @@ static FAutoConsoleVariableRef CVarLightweightTimeguardThresholdMS(
 	TEXT("Threshold in milliseconds for the tick timeguard"),
 	ECVF_Default);
 
+static float GIdleTaskWorkMS = 0.0f;
+static FAutoConsoleVariableRef CVarIdleTaskWorkMS(
+	TEXT("tick.IdleTaskWorkMS"),
+	GIdleTaskWorkMS,
+	TEXT("If > 0, when the game thread is idle waiting for other threads to complete tasks it will try to spend this much time in milliseconds processing other tasks.\n")
+	TEXT("If this is 0, it will keep processing game thread ticks until blocked and then wait.\n")
+	TEXT("If < 0 this will use the UE 5.5 behavior to skip idle task processing and stall the game thread entirely."),
+	ECVF_Default);
+
 FAutoConsoleTaskPriority CPrio_DispatchTaskPriority(
 	TEXT("TaskGraph.TaskPriorities.TickDispatchTaskPriority"),
 	TEXT("Task and thread priority for tick tasks dispatch."),
@@ -910,14 +919,60 @@ public:
 				{
 					TRACE_CPUPROFILER_EVENT_SCOPE(TickCompletionEvents);
 
-					// Process tasks, and on first stall make sure all manual dispatch events have been dispatched
-					// TODO change this to a deadlock detector somehow? In theory this could happen too early if there are very few game thread ticks
-					// But right now it just runs directly inside the function as soon as the game thread runs out of work to do
-					FTaskGraphInterface::Get().ProcessUntilTasksComplete(TickCompletionEvents[Block], ENamedThreads::GameThread,
-						[this, WorldTickGroup]()
+					if (GIdleTaskWorkMS >= 0.0)
+					{
+						bool bCanDoIdleWork = (GIdleTaskWorkMS > 0.0);
+						uint64 EndIdleWork = 0;
+						int32 PreviousTasksRemaining = TickCompletionEvents[Block].Num();
+						auto IdleWorkUpdate = [this, &bCanDoIdleWork, &EndIdleWork, &PreviousTasksRemaining, WorldTickGroup](int32 TasksRemaining) -> FTaskGraphInterface::EProcessTasksOperation
 						{
-							this->VerifyManualDispatch(WorldTickGroup);
-						});
+							if (bCanDoIdleWork)
+							{
+								// Compute cycles for ending idle work if necessary, this is called after the first pass of processing
+								if (EndIdleWork == 0)
+								{
+									const double TaskWorkSeconds = (GIdleTaskWorkMS / 1000.0);
+									EndIdleWork = FPlatformTime::Cycles64() + FMath::TruncToInt64(TaskWorkSeconds / FPlatformTime::GetSecondsPerCycle64());
+								}
+								else if (FPlatformTime::Cycles64() > EndIdleWork)
+								{
+									bCanDoIdleWork = false;
+								}
+							}
+
+							if (TasksRemaining == PreviousTasksRemaining)
+							{
+								// Nothing changed so this could be deadlocked, make sure our manual dispatches have happened
+								this->VerifyManualDispatch(WorldTickGroup);
+
+								if (!bCanDoIdleWork)
+								{
+									// Out of idle work budget or never had any
+									return FTaskGraphInterface::EProcessTasksOperation::WaitUntilComplete;
+								}
+							}
+							else
+							{
+								PreviousTasksRemaining = TasksRemaining;
+
+								// This could be changed to execute other game thread work in the future
+
+								if (!bCanDoIdleWork)
+								{
+									return FTaskGraphInterface::EProcessTasksOperation::ProcessNamedThreadTasks;
+								}
+							}
+
+							return FTaskGraphInterface::EProcessTasksOperation::ProcessOneOtherTask;
+						};
+
+						FTaskGraphInterface::Get().ProcessUntilTasksComplete(TickCompletionEvents[Block], ENamedThreads::GameThread, IdleWorkUpdate);
+					}
+					else
+					{
+						// Old behavior of waiting for all of them
+						FTaskGraphInterface::Get().WaitUntilTasksComplete(TickCompletionEvents[Block], ENamedThreads::GameThread);
+					}
 
 					if (bSingleThreadMode || Block == TG_NewlySpawned || CVarAllowAsyncTickCleanup.GetValueOnGameThread() == 0 || TickCompletionEvents[Block].Num() < 50)
 					{
