[Nanite] Fix incorrect merge of skeletal mesh builder which created an issue with material slots for Nanite skeletal assemblies. Also, add GetSystemGuids to FUE5SpecialProjectStreamObjectVersion and register it to bump the skeletal mesh DDC version.

diff --git a/Engine/Source/Runtime/Core/Private/UObject/DevObjectVersion.cpp b/Engine/Source/Runtime/Core/Private/UObject/DevObjectVersion.cpp
--- a/Engine/Source/Runtime/Core/Private/UObject/DevObjectVersion.cpp
+++ b/Engine/Source/Runtime/Core/Private/UObject/DevObjectVersion.cpp
@@ -344,6 +344,7 @@ FDevSystemGuidRegistration GRegisterUE5ReleaseBranchSystemGuids(FUE5ReleaseStrea
 const FGuid FUE5SpecialProjectStreamObjectVersion::GUID(0x59DA5D52, 0x12324948, 0xB8785978, 0x70B8E98B);
 // Register UE5 private frosty stream custom version with Core
 FDevVersionRegistration GRegisterUE5SpecialProjectStreamObjectVersion(FUE5SpecialProjectStreamObjectVersion::GUID, FUE5SpecialProjectStreamObjectVersion::LatestVersion, TEXT("UE5-SpecialProject"));
+FDevSystemGuidRegistration GRegisterUE5SpecialProjectStreamSystemGuids(FUE5SpecialProjectStreamObjectVersion::GetSystemGuids());
 
 // Unique MediaFramework Object version id
 const FGuid FMediaFrameworkObjectVersion::GUID(0x6f0ed827, 0xa6094895, 0x9c91998d, 0x90180ea4);
diff --git a/Engine/Source/Runtime/Core/Private/UObject/UE5SpecialProjectStreamObjectVersion.cpp b/Engine/Source/Runtime/Core/Private/UObject/UE5SpecialProjectStreamObjectVersion.cpp
new file mode 100644
--- /dev/null
+++ b/Engine/Source/Runtime/Core/Private/UObject/UE5SpecialProjectStreamObjectVersion.cpp
@@ -0,0 +1,16 @@
+// Copyright Epic Games, Inc. All Rights Reserved.
+#include "UObject/UE5SpecialProjectStreamObjectVersion.h"
+#include "UObject/DevObjectVersion.h"
+#include "Containers/Map.h"
+#include "Misc/Guid.h"
+
+
+TMap<FGuid, FGuid> FUE5SpecialProjectStreamObjectVersion::GetSystemGuids()
+{
+	TMap<FGuid, FGuid> SystemGuids;
+	const FDevSystemGuids& DevGuids = FDevSystemGuids::Get();
+
+	SystemGuids.Add(DevGuids.SkeletalMeshDerivedDataVersion, FGuid("5D7013CA-0F57-48ED-A81F-CBE3038DDBBA"));
+
+	return SystemGuids;
+}
diff --git a/Engine/Source/Runtime/Core/Public/UObject/UE5SpecialProjectStreamObjectVersion.h b/Engine/Source/Runtime/Core/Public/UObject/UE5SpecialProjectStreamObjectVersion.h
--- a/Engine/Source/Runtime/Core/Public/UObject/UE5SpecialProjectStreamObjectVersion.h
+++ b/Engine/Source/Runtime/Core/Public/UObject/UE5SpecialProjectStreamObjectVersion.h
@@ -20,6 +20,8 @@ struct FUE5SpecialProjectStreamObjectVersion
 	// The GUID for this custom version number
 	CORE_API const static FGuid GUID;
 
+	static CORE_API TMap<FGuid, FGuid> GetSystemGuids();
+
 	FUE5SpecialProjectStreamObjectVersion() = delete;
 };
 
