Removed lambdas from the implementation of ensure() - this saves a lot of PDB space.

diff --git a/Engine/Source/Runtime/Core/Public/Clang/ClangPlatform.h b/Engine/Source/Runtime/Core/Public/Clang/ClangPlatform.h
--- a/Engine/Source/Runtime/Core/Public/Clang/ClangPlatform.h
+++ b/Engine/Source/Runtime/Core/Public/Clang/ClangPlatform.h
@@ -27,6 +27,8 @@
 
 #define UE_LIFETIMEBOUND [[clang::lifetimebound]]
 
+#define UE_NODEBUG [[gnu::nodebug]]
+
 // Clang supports the following GNU attributes to let us describe allocation functions.
 // `gnu::malloc` maps to the `noalias` function return on the caller in the LLVM IR.
 // `gnu::alloc_size` maps to the `allocsize` function attribute on the caller in the LLVM IR.
diff --git a/Engine/Source/Runtime/Core/Public/HAL/Platform.h b/Engine/Source/Runtime/Core/Public/HAL/Platform.h
--- a/Engine/Source/Runtime/Core/Public/HAL/Platform.h
+++ b/Engine/Source/Runtime/Core/Public/HAL/Platform.h
@@ -770,6 +770,11 @@
 	#define UE_LIFETIMEBOUND
 #endif
 
+/* Tag a function or variable to omit debug info for this symbol */
+#ifndef UE_NODEBUG
+	#define UE_NODEBUG
+#endif
+
 /* Annotate functions that allocate new memory to ensure the compiler can optimize them accordingly.
    The arguments to this macro specify the 1-based arguments of the annotated function that specify the size of the allocation. */
 #ifndef UE_ALLOCATION_FUNCTION
diff --git a/Engine/Source/Runtime/Core/Public/Misc/AssertionMacros.h b/Engine/Source/Runtime/Core/Public/Misc/AssertionMacros.h
--- a/Engine/Source/Runtime/Core/Public/Misc/AssertionMacros.h
+++ b/Engine/Source/Runtime/Core/Public/Misc/AssertionMacros.h
@@ -397,19 +397,35 @@ RetType UE_COLD UE_DEBUG_SECTION DispatchCheckVerify (InnerType&& Inner, ArgType
 
 	} // UE::Assert::Private
 
-	struct UE_DEPRECATED(5.4, "Do not use directly. This internal type is being removed.") FValidateArgsInternal
+	template <uint64 Uid>
+	UE_NODEBUG std::atomic<bool> bGEnsureHasExecuted = false;
+
+	UE_NODEBUG constexpr uint32 FileHashForEnsure(const char* Filename)
 	{
-		template <typename... Types>
-		FValidateArgsInternal(Types... Args)
+		// Based on djb2
+		uint32 Result = 5381;
+		for (;;)
 		{
-			static_assert((TIsValidVariadicFunctionArg<Types>::Value && ...), "Invalid argument(s) passed to ensureMsgf");
+			char Ch = *Filename++;
+			if (!Ch)
+			{
+				return Result;
+			}
+			Result = ((Result << 5) + Result) + Ch;
 		}
-	};
+	}
+
+	UE_NODEBUG FORCEINLINE bool BreakAndReturnFalse()
+	{
+		// This is tagged with UE_NODEBUG and FORCEINLINE to ensure that the break happens at the usage of ensure(), not inside this function.
+		PLATFORM_BREAK();
+		return false;
+	}
 
 	#define UE_ENSURE_IMPL(Always, InExpression) \
 		(LIKELY(!!(InExpression)) \
-			|| (::UE::Assert::Private::ExecCheckImplInternal([]() UE_COLD UE_DEBUG_SECTION -> std::atomic<bool>& { static std::atomic<bool> ENSURE_bExecuted = false; return ENSURE_bExecuted; } (), Always, __FILE__, __LINE__, #InExpression) \
-			&& [] () { PLATFORM_BREAK(); return false; } ()))
+			|| (::UE::Assert::Private::ExecCheckImplInternal(::bGEnsureHasExecuted<((uint64)FileHashForEnsure(__FILE__) << 32) | (uint64)__LINE__>, Always, __FILE__, __LINE__, #InExpression) \
+			&& BreakAndReturnFalse()))
 
 	#define UE_ENSURE_IMPL2(Capture, Always, InExpression, InFormat, ...) \
 		(LIKELY(!!(InExpression)) || ([Capture] () UE_COLD UE_DEBUG_SECTION \
