Remove duplicate implementations of FScopeLock, FScopeUnlock and FScopeTryLock.

Previously, the UE::TScope and FScope classes contained identical copy-paste logic. (The FScope versions were hard-coded to use FCriticalSection.) Now, the FScope classes are lightweight wrappers around the matching UE::TScope class, with zero lines of code. The constructors could not be removed entirely because the FScope classes took pointers and the UE::TScope methods take references, but the `check` could be replaced with `TNotNull` and everything else could be handled via inheritance or `using`.

diff --git a/Engine/Source/Runtime/Core/Public/Misc/ScopeLock.h b/Engine/Source/Runtime/Core/Public/Misc/ScopeLock.h
--- a/Engine/Source/Runtime/Core/Public/Misc/ScopeLock.h
+++ b/Engine/Source/Runtime/Core/Public/Misc/ScopeLock.h
@@ -4,142 +4,13 @@
 
 #include "CoreTypes.h"
 #include "Misc/AssertionMacros.h"
+#include "Misc/NotNull.h"
 #include "HAL/CriticalSection.h"
 
-/**
- * Implements a scope lock.
- *
- * This is a utility class that handles scope level locking. It's very useful
- * to keep from causing deadlocks due to exceptions being caught and knowing
- * about the number of locks a given thread has on a resource. Example:
- *
- * <code>
- *	{
- *		// Synchronize thread access to the following data
- *		FScopeLock ScopeLock(SynchObject);
- *		// Access data that is shared among multiple threads
- *		...
- *		// When ScopeLock goes out of scope, other threads can access data
- *	}
- * </code>
- */
-class FScopeLock
-{
-public:
-
-	/**
-	 * Constructor that performs a lock on the synchronization object
-	 *
-	 * @param InSynchObject The synchronization object to manage
-	 */
-	UE_NODISCARD_CTOR FScopeLock( FCriticalSection* InSynchObject )
-		: SynchObject(InSynchObject)
-	{
-		check(SynchObject);
-		SynchObject->Lock();
-	}
-
-	/** Destructor that performs a release on the synchronization object. */
-	~FScopeLock()
-	{
-		Unlock();
-	}
-
-	void Unlock()
-	{
-		if(SynchObject)
-		{
-			SynchObject->Unlock();
-			SynchObject = nullptr;
-		}
-	}
-
-private:
-
-	/** Default constructor (hidden on purpose). */
-	FScopeLock();
-
-	/** Copy constructor( hidden on purpose). */
-	FScopeLock(const FScopeLock& InScopeLock);
-
-	/** Assignment operator (hidden on purpose). */
-	FScopeLock& operator=( FScopeLock& InScopeLock )
-	{
-		return *this;
-	}
-
-private:
-
-	// Holds the synchronization object to aggregate and scope manage.
-	FCriticalSection* SynchObject;
-};
-
-/**
- * Implements a scope unlock.
- *
- * This is a utility class that handles scope level unlocking. It's very useful
- * to allow access to a protected object when you are sure it can happen.
- * Example:
- *
- * <code>
- *	{
- *		// Access data that is shared among multiple threads
- *		FScopeUnlock ScopeUnlock(SynchObject);
- *		...
- *		// When ScopeUnlock goes out of scope, other threads can no longer access data
- *	}
- * </code>
- */
-class FScopeUnlock
-{
-public:
-
-	/**
-	 * Constructor that performs a unlock on the synchronization object
-	 *
-	 * @param InSynchObject The synchronization object to manage, can be null.
-	 */
-	UE_NODISCARD_CTOR FScopeUnlock(FCriticalSection* InSynchObject)
-		: SynchObject(InSynchObject)
-	{
-		if (InSynchObject)
-		{
-			InSynchObject->Unlock();
-		}
-	}
-
-	/** Destructor that performs a lock on the synchronization object. */
-	~FScopeUnlock()
-	{
-		if (SynchObject)
-		{
-			SynchObject->Lock();
-		}
-	}
-private:
-
-	/** Default constructor (hidden on purpose). */
-	FScopeUnlock();
-
-	/** Copy constructor( hidden on purpose). */
-	FScopeUnlock(const FScopeUnlock& InScopeLock);
-
-	/** Assignment operator (hidden on purpose). */
-	FScopeUnlock& operator=(FScopeUnlock& InScopeLock)
-	{
-		return *this;
-	}
-
-private:
-
-	// Holds the synchronization object to aggregate and scope manage.
-	FCriticalSection* SynchObject;
-};
-
 namespace UE
 {
-	// RAII-style scope locking of a synchronisation primitive
-	// `MutexType` is required to implement `Lock` and `Unlock` methods
+	// RAII-style scope locking of a synchronization primitive.
+	// `MutexType` is required to implement `Lock` and `Unlock` methods.
 	// Example:
 	//	{
 	//		TScopeLock<FCriticalSection> ScopeLock(CriticalSection);
@@ -176,8 +47,8 @@ namespace UE
 		MutexType* Mutex;
 	};
 
-	// RAII-style scope locking of a synchronisation primitive. Same
-	// as TScopeLock except taking the lock is conditional
+	// RAII-style scope locking of a synchronization primitive. Same
+	// as TScopeLock except taking the lock is conditional.
 	template<typename MutexType>
 	class TConditionalScopeLock
 	{
@@ -212,8 +83,8 @@ namespace UE
 		MutexType* Mutex;
 	};
 
-	// RAII-style scope unlocking of a synchronisation primitive
-	// `MutexType` is required to implement `Lock` and `Unlock` methods
+	// RAII-style scope unlocking of a synchronization primitive.
+	// `MutexType` is required to implement `Lock` and `Unlock` methods.
 	// Example:
 	//	{
 	//		TScopeLock<FCriticalSection> ScopeLock(CriticalSection);
@@ -244,4 +115,65 @@ namespace UE
 	private:
 		MutexType* Mutex;
 	};
-}
\ No newline at end of file
+}
+
+/**
+ * Implements a scope lock.
+ *
+ * This is a utility class that handles scope level locking. It's very useful
+ * to keep from causing deadlocks due to exceptions being caught and knowing
+ * about the number of locks a given thread has on a resource. Example:
+ *
+ * <code>
+ *	{
+ *		// Synchronize thread access to the following data
+ *		FScopeLock ScopeLock(SynchObject);
+ *		// Access data that is shared among multiple threads
+ *		...
+ *		// When ScopeLock goes out of scope, other threads can access data
+ *	}
+ * </code>
+ */
+class FScopeLock : private UE::TScopeLock<FCriticalSection>
+{
+public:
+	/**
+	 * Constructor that performs a lock on the synchronization object
+	 *
+	 * @param InSynchObject The synchronization object to manage
+	 */
+	UE_NODISCARD_CTOR FScopeLock(TNotNull<FCriticalSection*> InSynchObject) : UE::TScopeLock<FCriticalSection>(*InSynchObject)
+	{
+	}
+
+	using UE::TScopeLock<FCriticalSection>::Unlock;
+};
+
+/**
+ * Implements a scope unlock.
+ *
+ * This is a utility class that handles scope level unlocking. It's very useful
+ * to allow access to a protected object when you are sure it can happen.
+ * Example:
+ *
+ * <code>
+ *	{
+ *		// Access data that is shared among multiple threads
+ *		FScopeUnlock ScopeUnlock(SynchObject);
+ *		...
+ *		// When ScopeUnlock goes out of scope, other threads can no longer access data
+ *	}
+ * </code>
+ */
+class FScopeUnlock : private UE::TScopeUnlock<FCriticalSection>
+{
+public:
+	/**
+	 * Constructor that performs a unlock on the synchronization object
+	 *
+	 * @param InSynchObject The synchronization object to manage, can be null.
+	 */
+	UE_NODISCARD_CTOR FScopeUnlock(TNotNull<FCriticalSection*> InSynchObject) : UE::TScopeUnlock<FCriticalSection>(*InSynchObject)
+	{
+	}
+};
diff --git a/Engine/Source/Runtime/Core/Public/Misc/ScopeTryLock.h b/Engine/Source/Runtime/Core/Public/Misc/ScopeTryLock.h
--- a/Engine/Source/Runtime/Core/Public/Misc/ScopeTryLock.h
+++ b/Engine/Source/Runtime/Core/Public/Misc/ScopeTryLock.h
@@ -4,66 +4,9 @@
 
 #include "CoreTypes.h"
 #include "Misc/AssertionMacros.h"
+#include "Misc/NotNull.h"
 #include "HAL/CriticalSection.h"
 
-/**
- * Implements a scope lock using TryLock.
- *
- * This is a utility class that handles scope level locking using TryLock.
- * Scope locking helps to avoid programming errors by which a lock is acquired
- * and never released.
- *
- * <code>
- *	{
- *		// Try to acquire a lock on CriticalSection for the current scope.
- *		FScopeTryLock ScopeTryLock(CriticalSection);
- *		// Check that the lock was acquired.
- *		if (ScopeTryLock.IsLocked())
- *		{
- *			// If the lock was acquired, we can safely access resources protected
- *			// by the lock.
- *		}
- *		// When ScopeTryLock goes out of scope, the critical section will be
- *		// released if it was ever acquired.
- *	}
- * </code>
- */
-class FScopeTryLock
-{
-public:
-	// Constructor that tries to lock the critical section. Note that you must
-	// call IsLocked to test that the lock was acquired.
-	FScopeTryLock(FCriticalSection *InCriticalSection)
-		: CriticalSection(InCriticalSection)
-	{
-		check(CriticalSection);
-		if (!CriticalSection->TryLock())
-		{
-			CriticalSection = nullptr;
-		}
-	}
-
-	// Destructor that will release the lock if it was ever acquired.
-	~FScopeTryLock()
-	{
-		if (CriticalSection)
-		{
-			CriticalSection->Unlock();
-		}
-	}
-
-	FORCEINLINE bool IsLocked() const
-	{
-		return nullptr != CriticalSection;
-	}
-
-private:
-	FScopeTryLock(FScopeTryLock const &InScopeLock) = delete;
-	FScopeTryLock &operator=(FScopeTryLock &Rhs) = delete;
-
-	FCriticalSection *CriticalSection;
-};
-
 namespace UE
 {
 	/**
@@ -119,4 +62,38 @@ namespace UE
 
 		MutexType* HeldMutex = nullptr;
 	};
-}
\ No newline at end of file
+}
+
+/**
+ * Implements a scope lock using TryLock.
+ *
+ * This is a utility class that handles scope level locking using TryLock.
+ * Scope locking helps to avoid programming errors by which a lock is acquired
+ * and never released.
+ *
+ * <code>
+ *	{
+ *		// Try to acquire a lock on CriticalSection for the current scope.
+ *		FScopeTryLock ScopeTryLock(CriticalSection);
+ *		// Check that the lock was acquired.
+ *		if (ScopeTryLock.IsLocked())
+ *		{
+ *			// If the lock was acquired, we can safely access resources protected
+ *			// by the lock.
+ *		}
+ *		// When ScopeTryLock goes out of scope, the critical section will be
+ *		// released if it was ever acquired.
+ *	}
+ * </code>
+ */
+class FScopeTryLock : private UE::TScopeTryLock<FCriticalSection>
+{
+public:
+	// Constructor that tries to lock the critical section.
+	FScopeTryLock(TNotNull<FCriticalSection*> InCriticalSection) : UE::TScopeTryLock<FCriticalSection>(*InCriticalSection)
+	{
+	}
+
+	// You must call IsLocked to test that the lock was acquired.
+	using UE::TScopeTryLock<FCriticalSection>::IsLocked;
+};
\ No newline at end of file
