Unreal.natvis: Prevent compiler optimizations and linker dead-stripping from removing global debugging state variables.
Follow up change to cl 40725151.

#rnx

diff --git a/Engine/Source/Runtime/Core/Public/Modules/Boilerplate/ModuleBoilerplate.h b/Engine/Source/Runtime/Core/Public/Modules/Boilerplate/ModuleBoilerplate.h
--- a/Engine/Source/Runtime/Core/Public/Modules/Boilerplate/ModuleBoilerplate.h
+++ b/Engine/Source/Runtime/Core/Public/Modules/Boilerplate/ModuleBoilerplate.h
@@ -77,6 +77,19 @@
 
 class FChunkedFixedUObjectArray;
 
+// Visualizer macro to prevent dead-stripping of global debug state variables.
+// This has not (yet) been an issue under MSVC, which is missing support for used and retain.
+#if defined(__has_cpp_attribute)
+	#if __has_cpp_attribute(used) && __has_cpp_attribute(retain)
+		#define UE_VISUALIZER_USE_AND_RETAIN [[used]] [[retain]]			// __GNUC__
+	#elif __has_cpp_attribute(gnu::used) && __has_cpp_attribute(gnu::retain)
+		#define UE_VISUALIZER_USE_AND_RETAIN [[gnu::used]] [[gnu::retain]]	// __clang__
+	#endif
+#endif
+#ifndef UE_VISUALIZER_USE_AND_RETAIN
+#define UE_VISUALIZER_USE_AND_RETAIN
+#endif
+
 #ifdef DISABLE_UE4_VISUALIZER_HELPERS
 	UE_DEPRECATED_HEADER(5.5, "DISABLE_UE4_VISUALIZER_HELPERS has been disabled, please define UE_ENABLE_VISUALIZER_HELPERS=0 instead.")
 	#define UE_ENABLE_VISUALIZER_HELPERS 0
@@ -92,11 +105,11 @@ class FChunkedFixedUObjectArray;
 	#define UE_VISUALIZERS_HELPERS
 #else
 	#define UE_VISUALIZERS_HELPERS \
-		uint8** GNameBlocksDebug = FNameDebugVisualizer(FClangKeepDebugInfo{}).GetBlocks(); \
-		FChunkedFixedUObjectArray*& GObjectArrayForDebugVisualizers = GCoreObjectArrayForDebugVisualizers; \
-		UE::CoreUObject::Private::FStoredObjectPathDebug*& GComplexObjectPathDebug = GCoreComplexObjectPathDebug; \
-		UE::CoreUObject::Private::FObjectHandlePackageDebugData*& GObjectHandlePackageDebug = GCoreObjectHandlePackageDebug; \
-		UE::Core::FVisualizerDebuggingState*& GDebuggingState = GCoreDebuggingState;
+		UE_VISUALIZER_USE_AND_RETAIN uint8** GNameBlocksDebug = FNameDebugVisualizer(FClangKeepDebugInfo{}).GetBlocks(); \
+		UE_VISUALIZER_USE_AND_RETAIN FChunkedFixedUObjectArray*& GObjectArrayForDebugVisualizers = GCoreObjectArrayForDebugVisualizers; \
+		UE_VISUALIZER_USE_AND_RETAIN UE::CoreUObject::Private::FStoredObjectPathDebug*& GComplexObjectPathDebug = GCoreComplexObjectPathDebug; \
+		UE_VISUALIZER_USE_AND_RETAIN UE::CoreUObject::Private::FObjectHandlePackageDebugData*& GObjectHandlePackageDebug = GCoreObjectHandlePackageDebug; \
+		UE_VISUALIZER_USE_AND_RETAIN UE::Core::FVisualizerDebuggingState*& GDebuggingState = GCoreDebuggingState;
 #endif
 
 // in DLL builds, these are done per-module, otherwise we just need one in the application
diff --git a/Engine/Source/Runtime/CoreUObject/Private/UObject/UObjectArray.cpp b/Engine/Source/Runtime/CoreUObject/Private/UObject/UObjectArray.cpp
--- a/Engine/Source/Runtime/CoreUObject/Private/UObject/UObjectArray.cpp
+++ b/Engine/Source/Runtime/CoreUObject/Private/UObject/UObjectArray.cpp
@@ -108,10 +108,7 @@ FUObjectArray::FUObjectArray()
 , OpenForDisregardForGC(true)
 , PrimarySerialNumber(START_SERIAL_NUMBER)
 {
-	// Set the global state using the (per-module) pointer reference rather than the GCoreObjectArrayForDebugVisualizers pointer directly
-	// in order to make sure the reference does not get stripped in monolithic (non-modular) builds on any platform.
-	extern FChunkedFixedUObjectArray*& GObjectArrayForDebugVisualizers;
-	GObjectArrayForDebugVisualizers = &GUObjectArray.ObjObjects;
+	GCoreObjectArrayForDebugVisualizers = &GUObjectArray.ObjObjects;
 }
 
 void FUObjectArray::AllocateObjectPool(int32 InMaxUObjects, int32 InMaxObjectsNotConsideredByGC, bool bPreAllocateObjectArray)
