Add a note about one unexpected cause for a Stat error
#rnx

diff --git a/Engine/Source/Runtime/Core/Private/Stats/StatsData.cpp b/Engine/Source/Runtime/Core/Private/Stats/StatsData.cpp
--- a/Engine/Source/Runtime/Core/Private/Stats/StatsData.cpp
+++ b/Engine/Source/Runtime/Core/Private/Stats/StatsData.cpp
@@ -888,7 +888,12 @@ void FStatsThreadState::ProcessNonFrameStats(FStatMessagesArray& Data, TSet<FNam
 				FStatMessage* Result = NotClearedEveryFrame.Find(Item.NameAndInfo.GetRawName());
 				if (!Result)
 				{
-					UE_LOG(LogStats, Error, TEXT( "Stat %s was cleared every frame, but we don't have metadata for it. Data loss." ), *Item.NameAndInfo.GetRawName().ToString() );
+					UE_LOG(
+						LogStats, Error,
+						TEXT("Stat %s was cleared every frame, but we don't have metadata for it. Data loss.")
+						TEXT(" (N.B. One cause for this can be duplicately named stats)"),
+						*Item.NameAndInfo.GetRawName().ToString()
+					);
 				}
 				else
 				{
