Improve safety of scope locks.

All constructors are now marked as `[[nodiscard]] explicit`, so that the lock is never immediately dropped and is always explicitly spelled out at the construction site.

diff --git a/Engine/Source/Runtime/Core/Public/Delegates/DelegateAccessHandler.h b/Engine/Source/Runtime/Core/Public/Delegates/DelegateAccessHandler.h
--- a/Engine/Source/Runtime/Core/Public/Delegates/DelegateAccessHandler.h
+++ b/Engine/Source/Runtime/Core/Public/Delegates/DelegateAccessHandler.h
@@ -78,12 +78,12 @@ protected:
 
 	[[nodiscard]] FReadAccessScope GetReadAccessScope() const
 	{
-		return FReadAccessScope{Mutex};
+		return FReadAccessScope{UE::TScopeLock{Mutex}};
 	}
 
 	[[nodiscard]] FWriteAccessScope GetWriteAccessScope()
 	{
-		return FWriteAccessScope{Mutex};
+		return FWriteAccessScope{UE::TScopeLock{Mutex}};
 	}
 
 private:
diff --git a/Engine/Source/Runtime/Core/Public/Misc/ScopeLock.h b/Engine/Source/Runtime/Core/Public/Misc/ScopeLock.h
--- a/Engine/Source/Runtime/Core/Public/Misc/ScopeLock.h
+++ b/Engine/Source/Runtime/Core/Public/Misc/ScopeLock.h
@@ -22,8 +22,7 @@ namespace UE
 	public:
 		UE_NONCOPYABLE(TScopeLock);
 
-		UE_NODISCARD_CTOR TScopeLock(MutexType& InMutex)
-			: Mutex(&InMutex)
+		[[nodiscard]] explicit TScopeLock(MutexType& InMutex) : Mutex(&InMutex)
 		{
 			check(Mutex);
 			Mutex->Lock();
@@ -55,7 +54,7 @@ namespace UE
 	public:
 		UE_NONCOPYABLE(TConditionalScopeLock);
 
-		UE_NODISCARD_CTOR TConditionalScopeLock(MutexType& InMutex, bool bShouldLock)
+		[[nodiscard]] explicit TConditionalScopeLock(MutexType& InMutex, bool bShouldLock)
 			: Mutex(bShouldLock ? &InMutex : nullptr)
 		{
 			if (bShouldLock)
@@ -100,8 +99,7 @@ namespace UE
 	public:
 		UE_NONCOPYABLE(TScopeUnlock);
 
-		UE_NODISCARD_CTOR TScopeUnlock(MutexType& InMutex)
-			: Mutex(&InMutex)
+		[[nodiscard]] explicit TScopeUnlock(MutexType& InMutex) : Mutex(&InMutex)
 		{
 			check(Mutex);
 			Mutex->Unlock();
@@ -127,7 +125,7 @@ namespace UE
  * <code>
  *	{
  *		// Synchronize thread access to the following data
- *		FScopeLock ScopeLock(SynchObject);
+ *		FScopeLock ScopeLock(SyncObject);
  *		// Access data that is shared among multiple threads
  *		...
  *		// When ScopeLock goes out of scope, other threads can access data
@@ -140,9 +138,9 @@ public:
 	/**
 	 * Constructor that performs a lock on the synchronization object
 	 *
-	 * @param InSynchObject The synchronization object to manage
+	 * @param InSyncObject The synchronization object to manage
 	 */
-	UE_NODISCARD_CTOR FScopeLock(TNotNull<FCriticalSection*> InSynchObject) : UE::TScopeLock<FCriticalSection>(*InSynchObject)
+	[[nodiscard]] explicit FScopeLock(TNotNull<FCriticalSection*> InSyncObject) : UE::TScopeLock<FCriticalSection>(*InSyncObject)
 	{
 	}
 
@@ -159,7 +157,7 @@ public:
  * <code>
  *	{
  *		// Access data that is shared among multiple threads
- *		FScopeUnlock ScopeUnlock(SynchObject);
+ *		FScopeUnlock ScopeUnlock(SyncObject);
  *		...
  *		// When ScopeUnlock goes out of scope, other threads can no longer access data
  *	}
@@ -171,9 +169,9 @@ public:
 	/**
 	 * Constructor that performs a unlock on the synchronization object
 	 *
-	 * @param InSynchObject The synchronization object to manage, can be null.
+	 * @param InSyncObject The synchronization object to manage, can be null.
 	 */
-	UE_NODISCARD_CTOR FScopeUnlock(TNotNull<FCriticalSection*> InSynchObject) : UE::TScopeUnlock<FCriticalSection>(*InSynchObject)
+	[[nodiscard]] explicit FScopeUnlock(TNotNull<FCriticalSection*> InSyncObject) : UE::TScopeUnlock<FCriticalSection>(*InSyncObject)
 	{
 	}
 };
diff --git a/Engine/Source/Runtime/Core/Public/Misc/ScopeTryLock.h b/Engine/Source/Runtime/Core/Public/Misc/ScopeTryLock.h
--- a/Engine/Source/Runtime/Core/Public/Misc/ScopeTryLock.h
+++ b/Engine/Source/Runtime/Core/Public/Misc/ScopeTryLock.h
@@ -35,7 +35,7 @@ namespace UE
 	class TScopeTryLock
 	{
 	public:
-		TScopeTryLock(MutexType& InMutex)
+		[[nodiscard]] explicit TScopeTryLock(MutexType& InMutex)
 		{
 			if (InMutex.TryLock())
 			{
@@ -90,7 +90,7 @@ class FScopeTryLock : private UE::TScopeTryLock<FCriticalSection>
 {
 public:
 	// Constructor that tries to lock the critical section.
-	FScopeTryLock(TNotNull<FCriticalSection*> InCriticalSection) : UE::TScopeTryLock<FCriticalSection>(*InCriticalSection)
+	[[nodiscard]] explicit FScopeTryLock(TNotNull<FCriticalSection*> InCriticalSection) : UE::TScopeTryLock<FCriticalSection>(*InCriticalSection)
 	{
 	}
 
