[FApp]

- Changes private static variables relating to Time to store and load via relaxed atomics to fix TSAN errors for these globals.
- Removed the private global variable "bUseDropFrameFormatWhenSupported" as it was never referenced or initialized.

#rnx

diff --git a/Engine/Source/Runtime/Core/Private/Misc/App.cpp b/Engine/Source/Runtime/Core/Private/Misc/App.cpp
--- a/Engine/Source/Runtime/Core/Private/Misc/App.cpp
+++ b/Engine/Source/Runtime/Core/Private/Misc/App.cpp
@@ -37,13 +37,13 @@ bool FApp::Standalone = true;
 bool FApp::bIsBenchmarking = false;
 bool FApp::bUseFixedSeed = false;
 bool FApp::bUseFixedTimeStep = false;
-double FApp::FixedDeltaTime = 1 / 30.0;
-double FApp::CurrentTime = 0.0;
-double FApp::LastTime = 0.0;
-double FApp::DeltaTime = 1 / 30.0;
-double FApp::IdleTime = 0.0;
-double FApp::IdleTimeOvershoot = 0.0;
-double FApp::GameTime = 0.0;
+std::atomic<double> FApp::FixedDeltaTime = 1 / 30.0;
+std::atomic<double> FApp::CurrentTime = 0.0;
+std::atomic<double> FApp::LastTime = 0.0;
+std::atomic<double> FApp::DeltaTime = 1 / 30.0;
+std::atomic<double> FApp::IdleTime = 0.0;
+std::atomic<double> FApp::IdleTimeOvershoot = 0.0;
+std::atomic<double> FApp::GameTime = 0.0;
 TOptional<FQualifiedFrameTime> FApp::CurrentFrameTime;
 float FApp::VolumeMultiplier = 1.0f;
 float FApp::UnfocusedVolumeMultiplier = 0.0f;
diff --git a/Engine/Source/Runtime/Core/Public/Misc/App.h b/Engine/Source/Runtime/Core/Public/Misc/App.h
--- a/Engine/Source/Runtime/Core/Public/Misc/App.h
+++ b/Engine/Source/Runtime/Core/Public/Misc/App.h
@@ -23,6 +23,9 @@
 #include "Templates/UnrealTemplate.h"
 #include "UObject/NameTypes.h"
 
+#include <atomic>
+#include <optional>
+
 class FCbObjectId;
 
 /**
@@ -526,7 +529,7 @@ public:
 	 */
 	FORCEINLINE static double GetFixedDeltaTime()
 	{
-		return FixedDeltaTime;
+		return FixedDeltaTime.load(std::memory_order_relaxed);
 	}
 
 	/**
@@ -536,7 +539,7 @@ public:
 	 */
 	static void SetFixedDeltaTime(double Seconds)
 	{
-		FixedDeltaTime = Seconds;
+		FixedDeltaTime.store(Seconds, std::memory_order_relaxed);
 	}
 
 	/**
@@ -570,7 +573,7 @@ public:
 	 */
 	FORCEINLINE static double GetCurrentTime()
 	{
-		return CurrentTime;
+		return CurrentTime.load(std::memory_order_relaxed);
 	}
 
 	/**
@@ -580,7 +583,7 @@ public:
 	 */
 	static void SetCurrentTime(double Seconds)
 	{
-		CurrentTime = Seconds;
+		CurrentTime.store(Seconds, std::memory_order_relaxed);
 	}
 
 	/**
@@ -590,13 +593,14 @@ public:
 	 */
 	FORCEINLINE static double GetLastTime()
 	{
-		return LastTime;
+		return LastTime.load(std::memory_order_relaxed);
 	}
 
 	/** Updates Last time to CurrentTime. */
 	static void UpdateLastTime()
 	{
-		LastTime = CurrentTime;
+		// Not an atomic operation, but preferred to a compare and swap
+		LastTime.store(CurrentTime.load(std::memory_order_relaxed), std::memory_order_relaxed);
 	}
 
 	/**
@@ -606,7 +610,7 @@ public:
 	 */
 	FORCEINLINE static double GetDeltaTime()
 	{
-		return DeltaTime;
+		return DeltaTime.load(std::memory_order_relaxed);
 	}
 
 	/**
@@ -616,7 +620,7 @@ public:
 	 */
 	static void SetDeltaTime(double Seconds)
 	{
-		DeltaTime = Seconds;
+		DeltaTime.store(Seconds, std::memory_order_relaxed);
 	}
 
 	/**
@@ -626,7 +630,7 @@ public:
 	 */
 	FORCEINLINE static double GetIdleTime()
 	{
-		return IdleTime;
+		return IdleTime.load(std::memory_order_relaxed);
 	}
 
 	/**
@@ -636,7 +640,7 @@ public:
 	 */
 	static void SetIdleTime(double Seconds)
 	{
-		IdleTime = Seconds;
+		IdleTime.store(Seconds, std::memory_order_relaxed);
 	}
 
 	/**
@@ -646,7 +650,7 @@ public:
 	 */
 	FORCEINLINE static double GetGameTime()
 	{
-		return GameTime;
+		return GameTime.load(std::memory_order_relaxed);
 	}
 
 	/**
@@ -656,7 +660,7 @@ public:
 	 */
 	static void SetGameTime(double Seconds)
 	{
-		GameTime = Seconds;
+		GameTime.store(Seconds, std::memory_order_relaxed);
 	}
 
 	/**
@@ -666,7 +670,7 @@ public:
 	 */
 	FORCEINLINE static double GetIdleTimeOvershoot()
 	{
-		return IdleTimeOvershoot;
+		return IdleTimeOvershoot.load(std::memory_order_relaxed);
 	}
 
 	/**
@@ -676,7 +680,7 @@ public:
 	 */
 	static void SetIdleTimeOvershoot(double Seconds)
 	{
-		IdleTimeOvershoot = Seconds;
+		IdleTimeOvershoot.store(Seconds, std::memory_order_relaxed);
 	}
 
 	/**
@@ -849,32 +853,29 @@ private:
 	static CORE_API bool bUseFixedTimeStep;
 
 	/** Holds time step if a fixed delta time is wanted. */
-	static CORE_API double FixedDeltaTime;
+	static CORE_API std::atomic<double> FixedDeltaTime;
 
 	/** Holds current time. */
-	static CORE_API double CurrentTime;
+	static CORE_API std::atomic<double> CurrentTime;
 
 	/** Holds previous value of CurrentTime. */
-	static CORE_API double LastTime;
+	static CORE_API std::atomic<double> LastTime;
 
 	/** Holds current delta time in seconds. */
-	static CORE_API double DeltaTime;
+	static CORE_API std::atomic<double> DeltaTime;
 
 	/** Holds time we spent sleeping in UpdateTimeAndHandleMaxTickRate() if our frame time was smaller than one allowed by target FPS. */
-	static CORE_API double IdleTime;
+	static CORE_API std::atomic<double> IdleTime;
 
 	/** Holds the amount of IdleTime that was LONGER than we tried to sleep. The OS can't sleep the exact amount of time, so this measures that overshoot. */
-	static CORE_API double IdleTimeOvershoot;
+	static CORE_API std::atomic<double> IdleTimeOvershoot;
 
 	/** Holds overall game time. */
-	static CORE_API double GameTime;
+	static CORE_API std::atomic<double> GameTime;
 
 	/** Holds the current frame time and framerate. */
 	static CORE_API TOptional<FQualifiedFrameTime> CurrentFrameTime;
 
-	/** Holds if we should generate a drop frame timecode when the frame rate does support it. */
-	static CORE_API bool bUseDropFrameFormatWhenSupported;
-
 	/** Use to affect the app volume when it loses focus */
 	static CORE_API float VolumeMultiplier;
 
