Added UE_NODEBUG across TArray functions which are simple getters, forwarding functions, or metaprogramming implementation details.

diff --git a/Engine/Source/Runtime/Core/Public/Containers/Array.h b/Engine/Source/Runtime/Core/Public/Containers/Array.h
--- a/Engine/Source/Runtime/Core/Public/Containers/Array.h
+++ b/Engine/Source/Runtime/Core/Public/Containers/Array.h
@@ -42,7 +42,7 @@
 template <typename T>
 struct TCanBulkSerialize
 {
-	enum { Value = std::is_arithmetic_v<T> };
+	enum { Value UE_NODEBUG = std::is_arithmetic_v<T> };
 };
 
 // Forward declarations
@@ -62,19 +62,19 @@ template <typename ContainerType, typename ElementType, typename SizeType>
 class TIndexedContainerIterator
 {
 public:
-	[[nodiscard]] TIndexedContainerIterator(ContainerType& InContainer, SizeType StartIndex = 0)
+	UE_NODEBUG [[nodiscard]] TIndexedContainerIterator(ContainerType& InContainer, SizeType StartIndex = 0)
 		: Container(InContainer)
 		, Index    (StartIndex)
 	{
 	}
 
 	/** Advances iterator to the next element in the container. */
-	TIndexedContainerIterator& operator++()
+	UE_NODEBUG TIndexedContainerIterator& operator++()
 	{
 		++Index;
 		return *this;
 	}
-	TIndexedContainerIterator operator++(int)
+	UE_NODEBUG TIndexedContainerIterator operator++(int)
 	{
 		TIndexedContainerIterator Tmp(*this);
 		++Index;
@@ -82,12 +82,12 @@ public:
 	}
 
 	/** Moves iterator to the previous element in the container. */
-	TIndexedContainerIterator& operator--()
+	UE_NODEBUG TIndexedContainerIterator& operator--()
 	{
 		--Index;
 		return *this;
 	}
-	TIndexedContainerIterator operator--(int)
+	UE_NODEBUG TIndexedContainerIterator operator--(int)
 	{
 		TIndexedContainerIterator Tmp(*this);
 		--Index;
@@ -95,65 +95,65 @@ public:
 	}
 
 	/** iterator arithmetic support */
-	TIndexedContainerIterator& operator+=(SizeType Offset)
+	UE_NODEBUG TIndexedContainerIterator& operator+=(SizeType Offset)
 	{
 		Index += Offset;
 		return *this;
 	}
 
-	[[nodiscard]] TIndexedContainerIterator operator+(SizeType Offset) const
+	UE_NODEBUG [[nodiscard]] TIndexedContainerIterator operator+(SizeType Offset) const
 	{
 		TIndexedContainerIterator Tmp(*this);
 		return Tmp += Offset;
 	}
 
-	TIndexedContainerIterator& operator-=(SizeType Offset)
+	UE_NODEBUG TIndexedContainerIterator& operator-=(SizeType Offset)
 	{
 		return *this += -Offset;
 	}
 
-	[[nodiscard]] TIndexedContainerIterator operator-(SizeType Offset) const
+	UE_NODEBUG [[nodiscard]] TIndexedContainerIterator operator-(SizeType Offset) const
 	{
 		TIndexedContainerIterator Tmp(*this);
 		return Tmp -= Offset;
 	}
 
-	[[nodiscard]] FORCEINLINE ElementType& operator* () const
+	UE_NODEBUG [[nodiscard]] FORCEINLINE ElementType& operator* () const
 	{
 		return Container[ Index ];
 	}
 
-	[[nodiscard]] FORCEINLINE ElementType* operator->() const
+	UE_NODEBUG [[nodiscard]] FORCEINLINE ElementType* operator->() const
 	{
 		return &Container[ Index ];
 	}
 
 	/** conversion to "bool" returning true if the iterator has not reached the last element. */
-	[[nodiscard]] FORCEINLINE explicit operator bool() const
+	UE_NODEBUG [[nodiscard]] FORCEINLINE explicit operator bool() const
 	{
 		return Container.IsValidIndex(Index);
 	}
 
 	/** Returns an index to the current element. */
-	[[nodiscard]] SizeType GetIndex() const
+	UE_NODEBUG [[nodiscard]] SizeType GetIndex() const
 	{
 		return Index;
 	}
 
 	/** Resets the iterator to the first element. */
-	void Reset()
+	UE_NODEBUG void Reset()
 	{
 		Index = 0;
 	}
 
 	/** Sets the iterator to one past the last element. */
-	void SetToEnd()
+	UE_NODEBUG void SetToEnd()
 	{
 		Index = Container.Num();
 	}
-	
+
 	/** Removes current element in array. This invalidates the current iterator value and it must be incremented */
-	void RemoveCurrent()
+	UE_NODEBUG void RemoveCurrent()
 	{
 		Container.RemoveAt(Index);
 		Index--;
@@ -164,17 +164,24 @@ public:
 	 * This invalidates the current iterator value and it must be incremented.
 	 * Note this modifies the order of the remaining elements in the array.
 	 */
-	void RemoveCurrentSwap()
+	UE_NODEBUG void RemoveCurrentSwap()
 	{
 		Container.RemoveAtSwap(Index);
 		Index--;
 	}
 
-	[[nodiscard]] FORCEINLINE bool operator==(const TIndexedContainerIterator& Rhs) const { return &Container == &Rhs.Container && Index == Rhs.Index; }
-	[[nodiscard]] FORCEINLINE bool operator!=(const TIndexedContainerIterator& Rhs) const { return &Container != &Rhs.Container || Index != Rhs.Index; }
+	UE_NODEBUG [[nodiscard]] FORCEINLINE bool operator==(const TIndexedContainerIterator& Rhs) const
+	{
+		return &Container == &Rhs.Container && Index == Rhs.Index;
+	}
+#if !PLATFORM_COMPILER_HAS_GENERATED_COMPARISON_OPERATORS
+	UE_NODEBUG [[nodiscard]] FORCEINLINE bool operator!=(const TIndexedContainerIterator& Rhs) const
+	{
+		return &Container != &Rhs.Container || Index != Rhs.Index;
+	}
+#endif
 
 private:
-
 	ContainerType& Container;
 	SizeType      Index;
 };
@@ -182,7 +189,7 @@ private:
 
 /** operator + */
 template <typename ContainerType, typename ElementType, typename SizeType>
-[[nodiscard]] FORCEINLINE TIndexedContainerIterator<ContainerType, ElementType, SizeType> operator+(SizeType Offset, TIndexedContainerIterator<ContainerType, ElementType, SizeType> RHS)
+UE_NODEBUG [[nodiscard]] FORCEINLINE TIndexedContainerIterator<ContainerType, ElementType, SizeType> operator+(SizeType Offset, TIndexedContainerIterator<ContainerType, ElementType, SizeType> RHS)
 {
 	return RHS + Offset;
 }
@@ -201,14 +208,14 @@ template <typename ContainerType, typename ElementType, typename SizeType>
 		//
 		// We do add an operator-- to help FString implementation
 
-		[[nodiscard]] explicit TCheckedPointerIterator(const SizeType& InNum, ElementType* InPtr)
+		UE_NODEBUG [[nodiscard]] explicit TCheckedPointerIterator(const SizeType& InNum, ElementType* InPtr)
 			: Ptr       (InPtr)
 			, CurrentNum(InNum)
 			, InitialNum(InNum)
 		{
 		}
 
-		[[nodiscard]] FORCEINLINE ElementType& operator*() const
+		UE_NODEBUG [[nodiscard]] FORCEINLINE ElementType& operator*() const
 		{
 			if constexpr (bReverse)
 			{
@@ -220,7 +227,7 @@ template <typename ContainerType, typename ElementType, typename SizeType>
 			}
 		}
 
-		FORCEINLINE TCheckedPointerIterator& operator++()
+		UE_NODEBUG FORCEINLINE TCheckedPointerIterator& operator++()
 		{
 			if constexpr (bReverse)
 			{
@@ -233,7 +240,7 @@ template <typename ContainerType, typename ElementType, typename SizeType>
 			return *this;
 		}
 
-		FORCEINLINE TCheckedPointerIterator& operator--()
+		UE_NODEBUG FORCEINLINE TCheckedPointerIterator& operator--()
 		{
 			if constexpr (bReverse)
 			{
@@ -246,7 +253,7 @@ template <typename ContainerType, typename ElementType, typename SizeType>
 			return *this;
 		}
 
-		[[nodiscard]] FORCEINLINE bool operator!=(const TCheckedPointerIterator& Rhs) const
+		UE_NODEBUG [[nodiscard]] FORCEINLINE bool operator!=(const TCheckedPointerIterator& Rhs) const
 		{
 			// We only need to do the check in this operator, because no other operator will be
 			// called until after this one returns.
@@ -268,23 +275,23 @@ template <typename ContainerType, typename ElementType, typename SizeType>
 template <typename ElementType, typename IteratorType>
 struct TDereferencingIterator
 {
-	[[nodiscard]] explicit TDereferencingIterator(IteratorType InIter)
+	UE_NODEBUG [[nodiscard]] explicit TDereferencingIterator(IteratorType InIter)
 		: Iter(InIter)
 	{
 	}
 
-	[[nodiscard]] FORCEINLINE ElementType& operator*() const
+	UE_NODEBUG [[nodiscard]] FORCEINLINE ElementType& operator*() const
 	{
 		return *(ElementType*)*Iter;
 	}
 
-	FORCEINLINE TDereferencingIterator& operator++()
+	UE_NODEBUG FORCEINLINE TDereferencingIterator& operator++()
 	{
 		++Iter;
 		return *this;
 	}
 
-	[[nodiscard]] FORCEINLINE bool operator!=(const TDereferencingIterator& Rhs) const
+	UE_NODEBUG [[nodiscard]] FORCEINLINE bool operator!=(const TDereferencingIterator& Rhs) const
 	{
 		return Iter != Rhs.Iter;
 	}
@@ -304,7 +311,7 @@ namespace UE::Core::Private
 	}
 
 	template <typename FromArrayType, typename ToArrayType>
-	[[nodiscard]] constexpr bool CanMoveTArrayPointersBetweenArrayTypes()
+	UE_NODEBUG [[nodiscard]] constexpr bool CanMoveTArrayPointersBetweenArrayTypes()
 	{
 		using FromAllocatorType          = typename FromArrayType::AllocatorType;
 		using ToAllocatorType            = typename ToArrayType::AllocatorType;
@@ -339,15 +346,14 @@ namespace UE::Core::Private
 	//     TArray<FThing> Arr; // this will cause errors without this workaround
 	// };
 	template <typename DestType, typename SourceType>
-	constexpr bool TArrayElementsAreCompatible_V = std::disjunction_v<std::is_same<DestType, std::decay_t<SourceType>>, std::is_constructible<DestType, SourceType>>;
+	UE_NODEBUG constexpr bool TArrayElementsAreCompatible_V = std::disjunction_v<std::is_same<DestType, std::decay_t<SourceType>>, std::is_constructible<DestType, SourceType>>;
 
 	template <typename ElementType, typename AllocatorType>
-	static char (&ResolveIsTArrayPtr(const volatile TArray<ElementType, AllocatorType>*))[2];
-
-	static char(&ResolveIsTArrayPtr(...))[1];
+	UE_NODEBUG static char (&ResolveIsTArrayPtr(const volatile TArray<ElementType, AllocatorType>*))[2];
+	UE_NODEBUG static char(&ResolveIsTArrayPtr(...))[1];
 
 	template <typename T>
-	constexpr bool TIsTArrayOrDerivedFromTArray_V = sizeof(ResolveIsTArrayPtr((T*)nullptr)) == 2;
+	UE_NODEBUG constexpr bool TIsTArrayOrDerivedFromTArray_V = sizeof(ResolveIsTArrayPtr((T*)nullptr)) == 2;
 
 	[[noreturn]] CORE_API void OnInvalidArrayNum(unsigned long long NewNum);
 
@@ -363,7 +369,7 @@ namespace UE::Core::Private
 	//
 	// When C++20 is guaranteed, concept checks can be used instead.
 	template <typename AllocatorType>
-	[[nodiscard]] constexpr uint32 GetAllocatorFlags()
+	UE_NODEBUG [[nodiscard]] constexpr uint32 GetAllocatorFlags()
 	{
 		uint32 Result = 0;
 		if constexpr (TAllocatorTraits<AllocatorType>::SupportsElementAlignment)
@@ -896,16 +902,16 @@ public:
 	///////////////////////////////////////////////
 	// Start - intrusive TOptional<TArray> state //
 	///////////////////////////////////////////////
-	constexpr static bool bHasIntrusiveUnsetOptionalState = true;
+	UE_NODEBUG constexpr static bool bHasIntrusiveUnsetOptionalState = true;
 	using IntrusiveUnsetOptionalStateType = TArray;
 
-	[[nodiscard]] explicit TArray(FIntrusiveUnsetOptionalState Tag)
+	UE_NODEBUG [[nodiscard]] explicit TArray(FIntrusiveUnsetOptionalState Tag)
 		: ArrayNum(0)
 		, ArrayMax(-1)
 	{
 		// Use ArrayMax == -1 as our intrusive state so that the destructor still works without change, as it doesn't use ArrayMax.
 	}
-	[[nodiscard]] bool operator==(FIntrusiveUnsetOptionalState Tag) const
+	UE_NODEBUG [[nodiscard]] bool operator==(FIntrusiveUnsetOptionalState Tag) const
 	{
 		return ArrayMax == -1;
 	}
@@ -918,7 +924,7 @@ public:
 	 *
 	 * @returns Pointer to first array entry or nullptr if ArrayMax == 0.
 	 */
-	[[nodiscard]] FORCEINLINE ElementType* GetData() UE_LIFETIMEBOUND
+	UE_NODEBUG [[nodiscard]] FORCEINLINE ElementType* GetData() UE_LIFETIMEBOUND
 	{
 		return (ElementType*)AllocatorInstance.GetAllocation();
 	}
@@ -928,7 +934,7 @@ public:
 	 *
 	 * @returns Pointer to first array entry or nullptr if ArrayMax == 0.
 	 */
-	[[nodiscard]] FORCEINLINE const ElementType* GetData() const UE_LIFETIMEBOUND
+	UE_NODEBUG [[nodiscard]] FORCEINLINE const ElementType* GetData() const UE_LIFETIMEBOUND
 	{
 		return (const ElementType*)AllocatorInstance.GetAllocation();
 	}
@@ -938,7 +944,7 @@ public:
 	 *
 	 * @returns Size in bytes of array type.
 	 */
-	[[nodiscard]] FORCEINLINE static constexpr uint32 GetTypeSize()
+	UE_NODEBUG [[nodiscard]] FORCEINLINE static constexpr uint32 GetTypeSize()
 	{
 		return sizeof(ElementType);
 	}
@@ -950,7 +956,7 @@ public:
 	 *
 	 * @returns Number of bytes allocated by this container.
 	 */
-	[[nodiscard]] FORCEINLINE SIZE_T GetAllocatedSize(void) const
+	UE_NODEBUG [[nodiscard]] FORCEINLINE SIZE_T GetAllocatedSize(void) const
 	{
 		return AllocatorInstance.GetAllocatedSize(ArrayMax, sizeof(ElementType));
 	}
@@ -960,7 +966,7 @@ public:
 	 *
 	 * @see Num, Shrink
 	 */
-	[[nodiscard]] FORCEINLINE SizeType GetSlack() const
+	UE_NODEBUG [[nodiscard]] FORCEINLINE SizeType GetSlack() const
 	{
 		return ArrayMax - ArrayNum;
 	}
@@ -969,7 +975,7 @@ public:
 	 * Checks array invariants: if array size is greater than or equal to zero and less
 	 * than or equal to the maximum.
 	 */
-	FORCEINLINE void CheckInvariants() const
+	UE_NODEBUG FORCEINLINE void CheckInvariants() const
 	{
 		checkSlow((ArrayNum >= 0) & (ArrayMax >= ArrayNum)); // & for one branch
 	}
@@ -1013,7 +1019,7 @@ public:
 	 * @param Index Index to test.
 	 * @returns True if index is valid. False otherwise.
 	 */
-	[[nodiscard]] FORCEINLINE bool IsValidIndex(SizeType Index) const
+	UE_NODEBUG [[nodiscard]] FORCEINLINE bool IsValidIndex(SizeType Index) const
 	{
 		return Index >= 0 && Index < ArrayNum;
 	}
@@ -1035,13 +1041,13 @@ public:
 	 * @returns Number of elements in array.
 	 * @see GetSlack
 	 */
-	[[nodiscard]] FORCEINLINE SizeType Num() const
+	UE_NODEBUG [[nodiscard]] FORCEINLINE SizeType Num() const
 	{
 		return ArrayNum;
 	}
 
 	/** @returns Number of bytes used, excluding slack */
-	[[nodiscard]] FORCEINLINE SIZE_T NumBytes() const
+	UE_NODEBUG [[nodiscard]] FORCEINLINE SIZE_T NumBytes() const
 	{
 		return static_cast<SIZE_T>(ArrayNum) * sizeof(ElementType);
 	}
@@ -1052,7 +1058,7 @@ public:
 	 * @returns Maximum number of elements in array.
 	 * @see GetSlack
 	 */
-	[[nodiscard]] FORCEINLINE SizeType Max() const
+	UE_NODEBUG [[nodiscard]] FORCEINLINE SizeType Max() const
 	{
 		return ArrayMax;
 	}
@@ -1062,7 +1068,7 @@ public:
 	 *
 	 * @returns Reference to indexed element.
 	 */
-	[[nodiscard]] FORCEINLINE ElementType& operator[](SizeType Index) UE_LIFETIMEBOUND
+	UE_NODEBUG [[nodiscard]] FORCEINLINE ElementType& operator[](SizeType Index) UE_LIFETIMEBOUND
 	{
 		RangeCheck(Index);
 		return GetData()[Index];
@@ -1075,7 +1081,7 @@ public:
 	 *
 	 * @returns Reference to indexed element.
 	 */
-	[[nodiscard]] FORCEINLINE const ElementType& operator[](SizeType Index) const UE_LIFETIMEBOUND
+	UE_NODEBUG [[nodiscard]] FORCEINLINE const ElementType& operator[](SizeType Index) const UE_LIFETIMEBOUND
 	{
 		RangeCheck(Index);
 		return GetData()[Index];
@@ -1105,7 +1111,7 @@ public:
 		return Result;
 	}
 	UE_ALLOWSHRINKING_BOOL_DEPRECATED("Pop")
-	FORCEINLINE ElementType Pop(bool bAllowShrinking)
+	UE_NODEBUG FORCEINLINE ElementType Pop(bool bAllowShrinking)
 	{
 		return Pop(bAllowShrinking ? EAllowShrinking::Yes : EAllowShrinking::No);
 	}
@@ -1115,7 +1121,7 @@ public:
 	 *
 	 * @param Item Item to push.
 	 */
-	FORCEINLINE void Push(ElementType&& Item)
+	UE_NODEBUG FORCEINLINE void Push(ElementType&& Item)
 	{
 		Add(MoveTempIfPossible(Item));
 	}
@@ -1128,7 +1134,7 @@ public:
 	 * @param Item Item to push.
 	 * @see Pop, Top
 	 */
-	FORCEINLINE void Push(const ElementType& Item)
+	UE_NODEBUG FORCEINLINE void Push(const ElementType& Item)
 	{
 		Add(Item);
 	}
@@ -1139,7 +1145,7 @@ public:
 	 * @returns Reference to the top element.
 	 * @see Pop, Push
 	 */
-	[[nodiscard]] FORCEINLINE ElementType& Top() UE_LIFETIMEBOUND
+	UE_NODEBUG [[nodiscard]] FORCEINLINE ElementType& Top() UE_LIFETIMEBOUND
 	{
 		return Last();
 	}
@@ -1152,7 +1158,7 @@ public:
 	 * @returns Reference to the top element.
 	 * @see Pop, Push
 	 */
-	[[nodiscard]] FORCEINLINE const ElementType& Top() const UE_LIFETIMEBOUND
+	UE_NODEBUG [[nodiscard]] FORCEINLINE const ElementType& Top() const UE_LIFETIMEBOUND
 	{
 		return Last();
 	}
@@ -1163,7 +1169,7 @@ public:
 	 * @param IndexFromTheEnd (Optional) Index from the end of array (default = 0).
 	 * @returns Reference to n-th last element from the array.
 	 */
-	[[nodiscard]] FORCEINLINE ElementType& Last(SizeType IndexFromTheEnd = 0) UE_LIFETIMEBOUND
+	UE_NODEBUG [[nodiscard]] FORCEINLINE ElementType& Last(SizeType IndexFromTheEnd = 0) UE_LIFETIMEBOUND
 	{
 		RangeCheck(ArrayNum - IndexFromTheEnd - 1);
 		return GetData()[ArrayNum - IndexFromTheEnd - 1];
@@ -1177,7 +1183,7 @@ public:
 	 * @param IndexFromTheEnd (Optional) Index from the end of array (default = 0).
 	 * @returns Reference to n-th last element from the array.
 	 */
-	[[nodiscard]] FORCEINLINE const ElementType& Last(SizeType IndexFromTheEnd = 0) const UE_LIFETIMEBOUND
+	UE_NODEBUG [[nodiscard]] FORCEINLINE const ElementType& Last(SizeType IndexFromTheEnd = 0) const UE_LIFETIMEBOUND
 	{
 		RangeCheck(ArrayNum - IndexFromTheEnd - 1);
 		return GetData()[ArrayNum - IndexFromTheEnd - 1];
@@ -1212,7 +1218,7 @@ public:
 	 * @returns True if found. False otherwise.
 	 * @see FindLast, FindLastByPredicate
 	 */
-	FORCEINLINE bool Find(const ElementType& Item, SizeType& Index) const
+	UE_NODEBUG FORCEINLINE bool Find(const ElementType& Item, SizeType& Index) const
 	{
 		Index = this->Find(Item);
 		return Index != INDEX_NONE;
@@ -1246,7 +1252,7 @@ public:
 	 * @returns True if found. False otherwise.
 	 * @see Find, FindLastByPredicate
 	 */
-	FORCEINLINE bool FindLast(const ElementType& Item, SizeType& Index) const
+	UE_NODEBUG FORCEINLINE bool FindLast(const ElementType& Item, SizeType& Index) const
 	{
 		Index = this->FindLast(Item);
 		return Index != INDEX_NONE;
@@ -1300,7 +1306,7 @@ public:
 	 * @returns Index of the found element. INDEX_NONE otherwise.
 	 */
 	template <typename Predicate>
-	[[nodiscard]] FORCEINLINE SizeType FindLastByPredicate(Predicate Pred) const
+	UE_NODEBUG [[nodiscard]] FORCEINLINE SizeType FindLastByPredicate(Predicate Pred) const
 	{
 		return FindLastByPredicate(Pred, ArrayNum);
 	}
@@ -1355,7 +1361,7 @@ public:
 	 * @see Find
 	 */
 	template <typename KeyType>
-	[[nodiscard]] FORCEINLINE const ElementType* FindByKey(const KeyType& Key) const
+	UE_NODEBUG [[nodiscard]] FORCEINLINE const ElementType* FindByKey(const KeyType& Key) const
 	{
 		return const_cast<TArray*>(this)->FindByKey(Key);
 	}
@@ -1390,7 +1396,7 @@ public:
 	 * @see FilterByPredicate, ContainsByPredicate
 	 */
 	template <typename Predicate>
-	[[nodiscard]] FORCEINLINE const ElementType* FindByPredicate(Predicate Pred) const
+	UE_NODEBUG [[nodiscard]] FORCEINLINE const ElementType* FindByPredicate(Predicate Pred) const
 	{
 		return const_cast<TArray*>(this)->FindByPredicate(Pred);
 	}
@@ -1464,7 +1470,7 @@ public:
 	 * @see Contains, Find
 	 */
 	template <typename Predicate>
-	[[nodiscard]] FORCEINLINE bool ContainsByPredicate(Predicate Pred) const
+	UE_NODEBUG [[nodiscard]] FORCEINLINE bool ContainsByPredicate(Predicate Pred) const
 	{
 		return FindByPredicate(Pred) != nullptr;
 	}
@@ -1475,7 +1481,7 @@ public:
 	 * @param OtherArray Array to compare.
 	 * @returns True if this array is the same as OtherArray. False otherwise.
 	 */
-	[[nodiscard]] bool operator==(const TArray& OtherArray) const
+	UE_NODEBUG [[nodiscard]] bool operator==(const TArray& OtherArray) const
 	{
 		SizeType Count = Num();
 
@@ -1489,14 +1495,12 @@ public:
 	 * @returns True if this array is NOT the same as OtherArray. False otherwise.
 	 */
 #if !PLATFORM_COMPILER_HAS_GENERATED_COMPARISON_OPERATORS
-	[[nodiscard]] FORCEINLINE bool operator!=(const TArray& OtherArray) const
+	UE_NODEBUG [[nodiscard]] FORCEINLINE bool operator!=(const TArray& OtherArray) const
 	{
 		return !(*this == OtherArray);
 	}
 #endif
 
-
-
 	/**
 	 * Bulk serialize array as a single memory blob when loading. Uses regular serialization code for saving
 	 * and doesn't serialize at all otherwise (e.g. transient, garbage collection, ...).
@@ -1577,7 +1581,7 @@ public:
 	 *
 	 * @param Ar Archive to count for.
 	 */
-	void CountBytes(FArchive& Ar) const
+	UE_NODEBUG void CountBytes(FArchive& Ar) const
 	{
 		Ar.CountBytes(ArrayNum*sizeof(ElementType), ArrayMax*sizeof(ElementType));
 	}
@@ -1592,7 +1596,7 @@ public:
 	 * @param Count Number of elements to add.
 	 * @returns Number of elements in array before addition.
 	 */
-	SizeType AddUninitialized()
+	FORCEINLINE SizeType AddUninitialized()
 	{
 		return UE::Core::Private::ReallocGrow<UE::Core::Private::GetAllocatorFlags<AllocatorType>()>(
 			sizeof(ElementType),
@@ -1648,7 +1652,7 @@ private:
 		typename OtherSizeType
 		UE_REQUIRES(sizeof(OtherSizeType) > sizeof(SizeType))
 	>
-	void InsertUninitializedImpl(SizeType Index, OtherSizeType Count)
+	UE_NODEBUG void InsertUninitializedImpl(SizeType Index, OtherSizeType Count)
 	{
 		checkf((OtherSizeType)(SizeType)Count == Count, TEXT("Invalid number of elements to add to this array type: %lld"), (long long)(SizeType)Count);
 		return InsertUninitializedImpl(Index, (SizeType)Count);
@@ -1667,11 +1671,11 @@ public:
 	 * @param Count Number of elements to add.
 	 * @see Insert, InsertZeroed, InsertDefaulted
 	 */
-	FORCEINLINE void InsertUninitialized(SizeType Index)
+	UE_NODEBUG FORCEINLINE void InsertUninitialized(SizeType Index)
 	{
 		InsertUninitializedImpl(Index);
 	}
-	FORCEINLINE void InsertUninitialized(SizeType Index, SizeType Count)
+	UE_NODEBUG FORCEINLINE void InsertUninitialized(SizeType Index, SizeType Count)
 	{
 		InsertUninitializedImpl(Index, Count);
 	}
@@ -2126,7 +2130,7 @@ public:
 	}
 	template <typename CountType>
 	UE_ALLOWSHRINKING_BOOL_DEPRECATED("RemoveAtSwap")
-	FORCEINLINE void RemoveAtSwap(SizeType Index, CountType Count, bool bAllowShrinking)
+	UE_NODEBUG FORCEINLINE void RemoveAtSwap(SizeType Index, CountType Count, bool bAllowShrinking)
 	{
 		RemoveAtSwap(Index, Count, bAllowShrinking ? EAllowShrinking::Yes : EAllowShrinking::No);
 	}
@@ -2217,7 +2221,7 @@ public:
 		}
 	}
 	UE_ALLOWSHRINKING_BOOL_DEPRECATED("SetNum")
-	FORCEINLINE void SetNum(SizeType NewNum, bool bAllowShrinking)
+	UE_NODEBUG FORCEINLINE void SetNum(SizeType NewNum, bool bAllowShrinking)
 	{
 		SetNum(NewNum, bAllowShrinking ? EAllowShrinking::Yes : EAllowShrinking::No);
 	}
@@ -2246,7 +2250,7 @@ public:
 		}
 	}
 	UE_ALLOWSHRINKING_BOOL_DEPRECATED("SetNumZeroed")
-	FORCEINLINE void SetNumZeroed(SizeType NewNum, bool bAllowShrinking)
+	UE_NODEBUG FORCEINLINE void SetNumZeroed(SizeType NewNum, bool bAllowShrinking)
 	{
 		SetNumZeroed(NewNum, bAllowShrinking ? EAllowShrinking::Yes : EAllowShrinking::No);
 	}
@@ -2274,7 +2278,7 @@ public:
 		}
 	}
 	UE_ALLOWSHRINKING_BOOL_DEPRECATED("SetNumUninitialized")
-	FORCEINLINE void SetNumUninitialized(SizeType NewNum, bool bAllowShrinking)
+	UE_NODEBUG FORCEINLINE void SetNumUninitialized(SizeType NewNum, bool bAllowShrinking)
 	{
 		SetNumUninitialized(NewNum, bAllowShrinking ? EAllowShrinking::Yes : EAllowShrinking::No);
 	}
@@ -2413,7 +2417,7 @@ public:
 	 *
 	 * @param Other The array to append.
 	 */
-	TArray& operator+=(TArray&& Other)
+	UE_NODEBUG TArray& operator+=(TArray&& Other)
 	{
 		Append(MoveTemp(Other));
 		return *this;
@@ -2425,7 +2429,7 @@ public:
 	 *
 	 * @param Other The array to append.
 	 */
-	TArray& operator+=(const TArray& Other)
+	UE_NODEBUG TArray& operator+=(const TArray& Other)
 	{
 		Append(Other);
 		return *this;
@@ -2436,7 +2440,7 @@ public:
 	 *
 	 * @param InitList The initializer list to append.
 	 */
-	TArray& operator+=(std::initializer_list<ElementType> InitList)
+	UE_NODEBUG TArray& operator+=(std::initializer_list<ElementType> InitList)
 	{
 		Append(InitList);
 		return *this;
@@ -2567,7 +2571,7 @@ public:
 	 * @return Index to the new item
 	 * @see AddDefaulted, AddUnique, AddZeroed, Append, Insert
 	 */
-	FORCEINLINE SizeType Add(ElementType&& Item)
+	UE_NODEBUG FORCEINLINE SizeType Add(ElementType&& Item)
 	{
 		CheckAddress(&Item);
 		return Emplace(MoveTempIfPossible(Item));
@@ -2580,7 +2584,7 @@ public:
 	 * @return Index to the new item
 	 * @see AddDefaulted, AddUnique, AddZeroed, Append, Insert
 	 */
-	FORCEINLINE SizeType Add(const ElementType& Item)
+	UE_NODEBUG FORCEINLINE SizeType Add(const ElementType& Item)
 	{
 		CheckAddress(&Item);
 		return Emplace(Item);
@@ -2595,7 +2599,7 @@ public:
 	 * @return A reference to the newly-inserted element.
 	 * @see AddDefaulted_GetRef, AddUnique_GetRef, AddZeroed_GetRef, Insert_GetRef
 	 */
-	[[nodiscard]] FORCEINLINE ElementType& Add_GetRef(ElementType&& Item) UE_LIFETIMEBOUND
+	UE_NODEBUG [[nodiscard]] FORCEINLINE ElementType& Add_GetRef(ElementType&& Item) UE_LIFETIMEBOUND
 	{
 		CheckAddress(&Item);
 		return Emplace_GetRef(MoveTempIfPossible(Item));
@@ -2608,7 +2612,7 @@ public:
 	 * @return A reference to the newly-inserted element.
 	 * @see AddDefaulted_GetRef, AddUnique_GetRef, AddZeroed_GetRef, Insert_GetRef
 	 */
-	[[nodiscard]] FORCEINLINE ElementType& Add_GetRef(const ElementType& Item) UE_LIFETIMEBOUND
+	UE_NODEBUG [[nodiscard]] FORCEINLINE ElementType& Add_GetRef(const ElementType& Item) UE_LIFETIMEBOUND
 	{
 		CheckAddress(&Item);
 		return Emplace_GetRef(Item);
@@ -3065,7 +3069,7 @@ public:
 	}
 	template <class PREDICATE_CLASS>
 	UE_ALLOWSHRINKING_BOOL_DEPRECATED("RemoveAllSwap")
-	FORCEINLINE SizeType RemoveAllSwap(const PREDICATE_CLASS& Predicate, bool bAllowShrinking)
+	UE_NODEBUG FORCEINLINE SizeType RemoveAllSwap(const PREDICATE_CLASS& Predicate, bool bAllowShrinking)
 	{
 		return RemoveAllSwap(Predicate, bAllowShrinking ? EAllowShrinking::Yes : EAllowShrinking::No);
 	}
@@ -3094,7 +3098,7 @@ public:
 		return 1;
 	}
 	UE_ALLOWSHRINKING_BOOL_DEPRECATED("RemoveSingleSwap")
-	FORCEINLINE SizeType RemoveSingleSwap(const ElementType& Item, bool bAllowShrinking)
+	UE_NODEBUG FORCEINLINE SizeType RemoveSingleSwap(const ElementType& Item, bool bAllowShrinking)
 	{
 		return RemoveSingleSwap(Item, bAllowShrinking ? EAllowShrinking::Yes : EAllowShrinking::No);
 	}
@@ -3141,7 +3145,7 @@ public:
 		return OriginalNum - ArrayNum;
 	}
 	UE_ALLOWSHRINKING_BOOL_DEPRECATED("RemoveSwap")
-	FORCEINLINE SizeType RemoveSwap(const ElementType& Item, bool bAllowShrinking)
+	UE_NODEBUG FORCEINLINE SizeType RemoveSwap(const ElementType& Item, bool bAllowShrinking)
 	{
 		return RemoveSwap(Item, bAllowShrinking ? EAllowShrinking::Yes : EAllowShrinking::No);
 	}
@@ -3223,7 +3227,7 @@ public:
 	 *
 	 * @returns The iterator.
 	 */
-	[[nodiscard]] TIterator CreateIterator()
+	UE_NODEBUG [[nodiscard]] TIterator CreateIterator()
 	{
 		return TIterator(*this);
 	}
@@ -3233,7 +3237,7 @@ public:
 	 *
 	 * @returns The const iterator.
 	 */
-	[[nodiscard]] TConstIterator CreateConstIterator() const
+	UE_NODEBUG [[nodiscard]] TConstIterator CreateConstIterator() const
 	{
 		return TConstIterator(*this);
 	}
@@ -3257,23 +3261,23 @@ public:
 	 * STL-like iterators to enable range-based for loop support.
 	 */
 	#if TARRAY_RANGED_FOR_CHECKS
-		[[nodiscard]] FORCEINLINE RangedForIteratorType             begin ()       { return RangedForIteratorType            (ArrayNum, GetData()); }
-		[[nodiscard]] FORCEINLINE RangedForConstIteratorType        begin () const { return RangedForConstIteratorType       (ArrayNum, GetData()); }
-		[[nodiscard]] FORCEINLINE RangedForIteratorType             end   ()       { return RangedForIteratorType            (ArrayNum, GetData() + Num()); }
-		[[nodiscard]] FORCEINLINE RangedForConstIteratorType        end   () const { return RangedForConstIteratorType       (ArrayNum, GetData() + Num()); }
-		[[nodiscard]] FORCEINLINE RangedForReverseIteratorType      rbegin()       { return RangedForReverseIteratorType     (ArrayNum, GetData() + Num()); }
-		[[nodiscard]] FORCEINLINE RangedForConstReverseIteratorType rbegin() const { return RangedForConstReverseIteratorType(ArrayNum, GetData() + Num()); }
-		[[nodiscard]] FORCEINLINE RangedForReverseIteratorType      rend  ()       { return RangedForReverseIteratorType     (ArrayNum, GetData()); }
-		[[nodiscard]] FORCEINLINE RangedForConstReverseIteratorType rend  () const { return RangedForConstReverseIteratorType(ArrayNum, GetData()); }
+		[[nodiscard]] UE_NODEBUG FORCEINLINE RangedForIteratorType             begin ()       { return RangedForIteratorType            (ArrayNum, GetData()); }
+		[[nodiscard]] UE_NODEBUG FORCEINLINE RangedForConstIteratorType        begin () const { return RangedForConstIteratorType       (ArrayNum, GetData()); }
+		[[nodiscard]] UE_NODEBUG FORCEINLINE RangedForIteratorType             end   ()       { return RangedForIteratorType            (ArrayNum, GetData() + Num()); }
+		[[nodiscard]] UE_NODEBUG FORCEINLINE RangedForConstIteratorType        end   () const { return RangedForConstIteratorType       (ArrayNum, GetData() + Num()); }
+		[[nodiscard]] UE_NODEBUG FORCEINLINE RangedForReverseIteratorType      rbegin()       { return RangedForReverseIteratorType     (ArrayNum, GetData() + Num()); }
+		[[nodiscard]] UE_NODEBUG FORCEINLINE RangedForConstReverseIteratorType rbegin() const { return RangedForConstReverseIteratorType(ArrayNum, GetData() + Num()); }
+		[[nodiscard]] UE_NODEBUG FORCEINLINE RangedForReverseIteratorType      rend  ()       { return RangedForReverseIteratorType     (ArrayNum, GetData()); }
+		[[nodiscard]] UE_NODEBUG FORCEINLINE RangedForConstReverseIteratorType rend  () const { return RangedForConstReverseIteratorType(ArrayNum, GetData()); }
 	#else
-		[[nodiscard]] FORCEINLINE RangedForIteratorType             begin ()       { return                                   GetData(); }
-		[[nodiscard]] FORCEINLINE RangedForConstIteratorType        begin () const { return                                   GetData(); }
-		[[nodiscard]] FORCEINLINE RangedForIteratorType             end   ()       { return                                   GetData() + Num(); }
-		[[nodiscard]] FORCEINLINE RangedForConstIteratorType        end   () const { return                                   GetData() + Num(); }
-		[[nodiscard]] FORCEINLINE RangedForReverseIteratorType      rbegin()       { return RangedForReverseIteratorType     (GetData() + Num()); }
-		[[nodiscard]] FORCEINLINE RangedForConstReverseIteratorType rbegin() const { return RangedForConstReverseIteratorType(GetData() + Num()); }
-		[[nodiscard]] FORCEINLINE RangedForReverseIteratorType      rend  ()       { return RangedForReverseIteratorType     (GetData()); }
-		[[nodiscard]] FORCEINLINE RangedForConstReverseIteratorType rend  () const { return RangedForConstReverseIteratorType(GetData()); }
+		[[nodiscard]] UE_NODEBUG FORCEINLINE RangedForIteratorType             begin ()       { return                                   GetData(); }
+		[[nodiscard]] UE_NODEBUG FORCEINLINE RangedForConstIteratorType        begin () const { return                                   GetData(); }
+		[[nodiscard]] UE_NODEBUG FORCEINLINE RangedForIteratorType             end   ()       { return                                   GetData() + Num(); }
+		[[nodiscard]] UE_NODEBUG FORCEINLINE RangedForConstIteratorType        end   () const { return                                   GetData() + Num(); }
+		[[nodiscard]] UE_NODEBUG FORCEINLINE RangedForReverseIteratorType      rbegin()       { return RangedForReverseIteratorType     (GetData() + Num()); }
+		[[nodiscard]] UE_NODEBUG FORCEINLINE RangedForConstReverseIteratorType rbegin() const { return RangedForConstReverseIteratorType(GetData() + Num()); }
+		[[nodiscard]] UE_NODEBUG FORCEINLINE RangedForReverseIteratorType      rend  ()       { return RangedForReverseIteratorType     (GetData()); }
+		[[nodiscard]] UE_NODEBUG FORCEINLINE RangedForConstReverseIteratorType rend  () const { return RangedForConstReverseIteratorType(GetData()); }
 	#endif
 
 public:
@@ -3286,7 +3290,7 @@ public:
 	 *        If this is not desirable, please use Algo::Sort(MyArray) directly instead.
 	 *        The auto-dereferencing behavior does not occur with smart pointers.
 	 */
-	void Sort()
+	UE_NODEBUG void Sort()
 	{
 		Algo::Sort(*this, TDereferenceWrapper<ElementType, TLess<>>(TLess<>()));
 	}
@@ -3302,7 +3306,7 @@ public:
 	 *        The auto-dereferencing behavior does not occur with smart pointers.
 	 */
 	template <class PREDICATE_CLASS>
-	void Sort(const PREDICATE_CLASS& Predicate)
+	UE_NODEBUG void Sort(const PREDICATE_CLASS& Predicate)
 	{
 		TDereferenceWrapper<ElementType, PREDICATE_CLASS> PredicateWrapper(Predicate);
 		Algo::Sort(*this, PredicateWrapper);
@@ -3318,7 +3322,7 @@ public:
 	 *        If this is not desirable, please use Algo::StableSort(MyArray) directly instead.
 	 *        The auto-dereferencing behavior does not occur with smart pointers.
 	 */
-	void StableSort()
+	UE_NODEBUG void StableSort()
 	{
 		Algo::StableSort(*this, TDereferenceWrapper<ElementType, TLess<>>(TLess<>()));
 	}
@@ -3336,7 +3340,7 @@ public:
 	 *        The auto-dereferencing behavior does not occur with smart pointers.
 	 */
 	template <class PREDICATE_CLASS>
-	void StableSort(const PREDICATE_CLASS& Predicate)
+	UE_NODEBUG void StableSort(const PREDICATE_CLASS& Predicate)
 	{
 		TDereferenceWrapper<ElementType, PREDICATE_CLASS> PredicateWrapper(Predicate);
 		Algo::StableSort(*this, PredicateWrapper);
@@ -3508,7 +3512,7 @@ public:
 	 *        The auto-dereferencing behavior does not occur with smart pointers.
 	 */
 	template <class PREDICATE_CLASS>
-	FORCEINLINE void Heapify(const PREDICATE_CLASS& Predicate)
+	UE_NODEBUG FORCEINLINE void Heapify(const PREDICATE_CLASS& Predicate)
 	{
 		TDereferenceWrapper<ElementType, PREDICATE_CLASS> PredicateWrapper(Predicate);
 		Algo::Heapify(*this, PredicateWrapper);
@@ -3522,7 +3526,7 @@ public:
 	 *        Therefore, your array will be heapified by the values being pointed to, rather than the pointers' values.
 	 *        The auto-dereferencing behavior does not occur with smart pointers.
 	 */
-	void Heapify()
+	UE_NODEBUG void Heapify()
 	{
 		Heapify(TLess<ElementType>());
 	}
@@ -3582,7 +3586,7 @@ public:
 	 *        Therefore, your array will be heapified by the values being pointed to, rather than the pointers' values.
 	 *        The auto-dereferencing behavior does not occur with smart pointers.
 	 */
-	SizeType HeapPush(ElementType&& InItem)
+	UE_NODEBUG SizeType HeapPush(ElementType&& InItem)
 	{
 		return HeapPush(MoveTempIfPossible(InItem), TLess<ElementType>());
 	}
@@ -3598,7 +3602,7 @@ public:
 	 *        Therefore, your array will be heapified by the values being pointed to, rather than the pointers' values.
 	 *        The auto-dereferencing behavior does not occur with smart pointers.
 	 */
-	SizeType HeapPush(const ElementType& InItem)
+	UE_NODEBUG SizeType HeapPush(const ElementType& InItem)
 	{
 		return HeapPush(InItem, TLess<ElementType>());
 	}
@@ -3625,7 +3629,7 @@ public:
 	}
 	template <class PREDICATE_CLASS>
 	UE_ALLOWSHRINKING_BOOL_DEPRECATED("HeapPop")
-	FORCEINLINE void HeapPop(ElementType& OutItem, const PREDICATE_CLASS& Predicate, bool bAllowShrinking)
+	UE_NODEBUG FORCEINLINE void HeapPop(ElementType& OutItem, const PREDICATE_CLASS& Predicate, bool bAllowShrinking)
 	{
 		HeapPop(OutItem, Predicate, bAllowShrinking ? EAllowShrinking::Yes : EAllowShrinking::No);
 	}
@@ -3641,12 +3645,12 @@ public:
 	 *        Therefore, your array will be heapified by the values being pointed to, rather than the pointers' values.
 	 *        The auto-dereferencing behavior does not occur with smart pointers.
 	 */
-	void HeapPop(ElementType& OutItem, EAllowShrinking AllowShrinking = EAllowShrinking::Default)
+	UE_NODEBUG void HeapPop(ElementType& OutItem, EAllowShrinking AllowShrinking = EAllowShrinking::Default)
 	{
 		HeapPop(OutItem, TLess<ElementType>(), AllowShrinking);
 	}
 	UE_ALLOWSHRINKING_BOOL_DEPRECATED("HeapPop")
-	FORCEINLINE void HeapPop(ElementType& OutItem, bool bAllowShrinking)
+	UE_NODEBUG FORCEINLINE void HeapPop(ElementType& OutItem, bool bAllowShrinking)
 	{
 		HeapPop(OutItem, bAllowShrinking ? EAllowShrinking::Yes : EAllowShrinking::No);
 	}
@@ -3657,7 +3661,7 @@ public:
 	 * @param Predicate Predicate class instance.
 	 */
 	template <class PREDICATE_CLASS>
-	void VerifyHeap(const PREDICATE_CLASS& Predicate)
+	UE_NODEBUG void VerifyHeap(const PREDICATE_CLASS& Predicate)
 	{
 		check(Algo::IsHeap(*this, Predicate));
 	}
@@ -3681,7 +3685,7 @@ public:
 	}
 	template <class PREDICATE_CLASS>
 	UE_ALLOWSHRINKING_BOOL_DEPRECATED("HeapPopDiscard")
-	FORCEINLINE void HeapPopDiscard(const PREDICATE_CLASS& Predicate, bool bAllowShrinking)
+	UE_NODEBUG FORCEINLINE void HeapPopDiscard(const PREDICATE_CLASS& Predicate, bool bAllowShrinking)
 	{
 		HeapPopDiscard(Predicate, bAllowShrinking ? EAllowShrinking::Yes : EAllowShrinking::No);
 	}
@@ -3696,12 +3700,12 @@ public:
 	 *        Therefore, your array will be heapified by the values being pointed to, rather than the pointers' values.
 	 *        The auto-dereferencing behavior does not occur with smart pointers.
 	 */
-	void HeapPopDiscard(EAllowShrinking AllowShrinking = EAllowShrinking::Default)
+	UE_NODEBUG void HeapPopDiscard(EAllowShrinking AllowShrinking = EAllowShrinking::Default)
 	{
 		HeapPopDiscard(TLess<ElementType>(), AllowShrinking);
 	}
 	UE_ALLOWSHRINKING_BOOL_DEPRECATED("HeapPopDiscard")
-	FORCEINLINE void HeapPopDiscard(bool bAllowShrinking)
+	UE_NODEBUG FORCEINLINE void HeapPopDiscard(bool bAllowShrinking)
 	{
 		HeapPopDiscard(TLess<ElementType>(), bAllowShrinking ? EAllowShrinking::Yes : EAllowShrinking::No);
 	}
@@ -3713,7 +3717,7 @@ public:
 	 *
 	 * @returns The reference to the top element from the heap.
 	 */
-	[[nodiscard]] const ElementType& HeapTop() const UE_LIFETIMEBOUND
+	UE_NODEBUG [[nodiscard]] const ElementType& HeapTop() const UE_LIFETIMEBOUND
 	{
 		return (*this)[0];
 	}
@@ -3723,7 +3727,7 @@ public:
 	 *
 	 * @returns The reference to the top element from the heap.
 	 */
-	[[nodiscard]] ElementType& HeapTop() UE_LIFETIMEBOUND
+	UE_NODEBUG [[nodiscard]] ElementType& HeapTop() UE_LIFETIMEBOUND
 	{
 		return (*this)[0];
 	}
@@ -3751,7 +3755,7 @@ public:
 	}
 	template <class PREDICATE_CLASS>
 	UE_ALLOWSHRINKING_BOOL_DEPRECATED("HeapRemoveAt")
-	void HeapRemoveAt(SizeType Index, const PREDICATE_CLASS& Predicate, bool bAllowShrinking)
+	UE_NODEBUG void HeapRemoveAt(SizeType Index, const PREDICATE_CLASS& Predicate, bool bAllowShrinking)
 	{
 		HeapRemoveAt(Index, Predicate, bAllowShrinking ? EAllowShrinking::Yes : EAllowShrinking::No);
 	}
@@ -3767,12 +3771,12 @@ public:
 	 *        Therefore, your array will be heapified by the values being pointed to, rather than the pointers' values.
 	 *        The auto-dereferencing behavior does not occur with smart pointers.
 	 */
-	void HeapRemoveAt(SizeType Index, EAllowShrinking AllowShrinking = EAllowShrinking::Default)
+	UE_NODEBUG void HeapRemoveAt(SizeType Index, EAllowShrinking AllowShrinking = EAllowShrinking::Default)
 	{
 		HeapRemoveAt(Index, TLess< ElementType >(), AllowShrinking);
 	}
 	UE_ALLOWSHRINKING_BOOL_DEPRECATED("HeapRemoveAt")
-	FORCEINLINE void HeapRemoveAt(SizeType Index, bool bAllowShrinking)
+	UE_NODEBUG FORCEINLINE void HeapRemoveAt(SizeType Index, bool bAllowShrinking)
 	{
 		HeapRemoveAt(Index, bAllowShrinking ? EAllowShrinking::Yes : EAllowShrinking::No);
 	}
@@ -3787,7 +3791,7 @@ public:
 	 *        The auto-dereferencing behavior does not occur with smart pointers.
 	 */
 	template <class PREDICATE_CLASS>
-	void HeapSort(const PREDICATE_CLASS& Predicate)
+	UE_NODEBUG void HeapSort(const PREDICATE_CLASS& Predicate)
 	{
 		TDereferenceWrapper<ElementType, PREDICATE_CLASS> PredicateWrapper(Predicate);
 		Algo::HeapSort(*this, PredicateWrapper);
@@ -3801,16 +3805,16 @@ public:
 	 *        Therefore, your array will be heapified by the values being pointed to, rather than the pointers' values.
 	 *        The auto-dereferencing behavior does not occur with smart pointers.
 	 */
-	void HeapSort()
+	UE_NODEBUG void HeapSort()
 	{
 		HeapSort(TLess<ElementType>());
 	}
 
-	[[nodiscard]] const ElementAllocatorType& GetAllocatorInstance() const
+	UE_NODEBUG [[nodiscard]] const ElementAllocatorType& GetAllocatorInstance() const
 	{
 		return AllocatorInstance;
 	}
-	[[nodiscard]] ElementAllocatorType& GetAllocatorInstance()
+	UE_NODEBUG [[nodiscard]] ElementAllocatorType& GetAllocatorInstance()
 	{
 		return AllocatorInstance;
 	}
@@ -3822,33 +3826,33 @@ public:
 namespace Freeze
 {
 	template<typename T, typename AllocatorType>
-	void IntrinsicWriteMemoryImage(FMemoryImageWriter& Writer, const TArray<T, AllocatorType>& Object, const FTypeLayoutDesc&)
+	UE_NODEBUG void IntrinsicWriteMemoryImage(FMemoryImageWriter& Writer, const TArray<T, AllocatorType>& Object, const FTypeLayoutDesc&)
 	{
 		Object.WriteMemoryImage(Writer);
 	}
 
 	template<typename T, typename AllocatorType>
-	[[nodiscard]] uint32 IntrinsicUnfrozenCopy(const FMemoryUnfreezeContent& Context, const TArray<T, AllocatorType>& Object, void* OutDst)
+	UE_NODEBUG [[nodiscard]] uint32 IntrinsicUnfrozenCopy(const FMemoryUnfreezeContent& Context, const TArray<T, AllocatorType>& Object, void* OutDst)
 	{
 		Object.CopyUnfrozen(Context, OutDst);
 		return sizeof(Object);
 	}
 
 	template<typename T, typename AllocatorType>
-	[[nodiscard]] uint32 IntrinsicAppendHash(const TArray<T, AllocatorType>* DummyObject, const FTypeLayoutDesc& TypeDesc, const FPlatformTypeLayoutParameters& LayoutParams, FSHA1& Hasher)
+	UE_NODEBUG [[nodiscard]] uint32 IntrinsicAppendHash(const TArray<T, AllocatorType>* DummyObject, const FTypeLayoutDesc& TypeDesc, const FPlatformTypeLayoutParameters& LayoutParams, FSHA1& Hasher)
 	{
 		return AppendHashForNameAndSize(TypeDesc.Name, sizeof(TArray<T, AllocatorType>), Hasher);
 	}
 
 	template<typename T, typename AllocatorType>
-	[[nodiscard]] uint32 IntrinsicGetTargetAlignment(const TArray<T, AllocatorType>* DummyObject, const FTypeLayoutDesc& TypeDesc, const FPlatformTypeLayoutParameters& LayoutParams)
+	UE_NODEBUG [[nodiscard]] uint32 IntrinsicGetTargetAlignment(const TArray<T, AllocatorType>* DummyObject, const FTypeLayoutDesc& TypeDesc, const FPlatformTypeLayoutParameters& LayoutParams)
 	{
 		// Assume alignment of array is drive by pointer
 		return FMath::Min(8u, LayoutParams.MaxFieldAlignment);
 	}
 
 	template<typename T, typename AllocatorType>
-	void IntrinsicToString(const TArray<T, AllocatorType>& Object, const FTypeLayoutDesc& TypeDesc, const FPlatformTypeLayoutParameters& LayoutParams, FMemoryToStringContext& OutContext)
+	UE_NODEBUG void IntrinsicToString(const TArray<T, AllocatorType>& Object, const FTypeLayoutDesc& TypeDesc, const FPlatformTypeLayoutParameters& LayoutParams, FMemoryToStringContext& OutContext)
 	{
 		Object.ToString(LayoutParams, OutContext);
 	}
@@ -3857,44 +3861,46 @@ namespace Freeze
 DECLARE_TEMPLATE_INTRINSIC_TYPE_LAYOUT((template <typename T, typename AllocatorType>), (TArray<T, AllocatorType>));
 
 template <typename InElementType, typename AllocatorType>
-struct TIsZeroConstructType<TArray<InElementType, AllocatorType>>
+struct UE_NODEBUG TIsZeroConstructType<TArray<InElementType, AllocatorType>>
 {
-	enum { Value = TAllocatorTraits<AllocatorType>::IsZeroConstruct };
+	enum { Value UE_NODEBUG = TAllocatorTraits<AllocatorType>::IsZeroConstruct };
 };
 
 template <typename T, typename AllocatorType>
-struct TIsContiguousContainer<TArray<T, AllocatorType>>
+struct UE_NODEBUG TIsContiguousContainer<TArray<T, AllocatorType>>
 {
-	enum { Value = true };
+	enum { Value UE_NODEBUG = true };
 };
 
 /**
  * Trait which determines whether or not a type is a TArray.
  */
-template <typename T> constexpr bool TIsTArray_V = false;
+template <typename T> UE_NODEBUG constexpr bool TIsTArray_V = false;
 
-template <typename InElementType, typename InAllocatorType> constexpr bool TIsTArray_V<               TArray<InElementType, InAllocatorType>> = true;
-template <typename InElementType, typename InAllocatorType> constexpr bool TIsTArray_V<const          TArray<InElementType, InAllocatorType>> = true;
-template <typename InElementType, typename InAllocatorType> constexpr bool TIsTArray_V<      volatile TArray<InElementType, InAllocatorType>> = true;
-template <typename InElementType, typename InAllocatorType> constexpr bool TIsTArray_V<const volatile TArray<InElementType, InAllocatorType>> = true;
+template <typename InElementType, typename InAllocatorType> UE_NODEBUG constexpr bool TIsTArray_V<               TArray<InElementType, InAllocatorType>> = true;
+template <typename InElementType, typename InAllocatorType> UE_NODEBUG constexpr bool TIsTArray_V<const          TArray<InElementType, InAllocatorType>> = true;
+template <typename InElementType, typename InAllocatorType> UE_NODEBUG constexpr bool TIsTArray_V<      volatile TArray<InElementType, InAllocatorType>> = true;
+template <typename InElementType, typename InAllocatorType> UE_NODEBUG constexpr bool TIsTArray_V<const volatile TArray<InElementType, InAllocatorType>> = true;
 
 template <typename T>
 struct TIsTArray
 {
-	enum { Value = TIsTArray_V<T> };
+	enum { Value UE_NODEBUG = TIsTArray_V<T> };
 };
 
 
 //
 // Array operator news.
 //
-template <typename T,typename AllocatorType> void* operator new(size_t Size, TArray<T, AllocatorType>& Array)
+template <typename T,typename AllocatorType>
+UE_NODEBUG void* operator new(size_t Size, TArray<T, AllocatorType>& Array)
 {
 	check(Size == sizeof(T));
 	const auto Index = Array.AddUninitialized();
 	return &Array[Index];
 }
-template <typename T,typename AllocatorType> void* operator new( size_t Size, TArray<T,AllocatorType>& Array, typename TArray<T, AllocatorType>::SizeType Index )
+template <typename T,typename AllocatorType>
+UE_NODEBUG void* operator new(size_t Size, TArray<T,AllocatorType>& Array, typename TArray<T, AllocatorType>::SizeType Index)
 {
 	check(Size == sizeof(T));
 	Array.InsertUninitialized(Index);
@@ -4007,7 +4013,7 @@ struct TArrayPrivateFriend
 
 
 template<typename ElementType, typename AllocatorType>
-FArchive& operator<<(FArchive& Ar, TArray<ElementType, AllocatorType>& A)
+UE_NODEBUG FArchive& operator<<(FArchive& Ar, TArray<ElementType, AllocatorType>& A)
 {
 	return TArrayPrivateFriend::Serialize(Ar, A);
 }
