[Backout] - CL40908676 - CIS Compile Error
#rnx
Original CL Desc
-----------------------------------------------------------------
Added UE_NODEBUG across TArray functions which are simple getters, forwarding functions, or metaprogramming implementation details.

diff --git a/Engine/Source/Runtime/Core/Public/Containers/Array.h b/Engine/Source/Runtime/Core/Public/Containers/Array.h
--- a/Engine/Source/Runtime/Core/Public/Containers/Array.h
+++ b/Engine/Source/Runtime/Core/Public/Containers/Array.h
@@ -42,7 +42,7 @@
 template <typename T>
 struct TCanBulkSerialize
 {
-	enum { Value UE_NODEBUG = std::is_arithmetic_v<T> };
+	enum { Value = std::is_arithmetic_v<T> };
 };
 
 // Forward declarations
@@ -62,19 +62,19 @@ template <typename ContainerType, typename ElementType, typename SizeType>
 class TIndexedContainerIterator
 {
 public:
-	UE_NODEBUG [[nodiscard]] TIndexedContainerIterator(ContainerType& InContainer, SizeType StartIndex = 0)
+	[[nodiscard]] TIndexedContainerIterator(ContainerType& InContainer, SizeType StartIndex = 0)
 		: Container(InContainer)
 		, Index    (StartIndex)
 	{
 	}
 
 	/** Advances iterator to the next element in the container. */
-	UE_NODEBUG TIndexedContainerIterator& operator++()
+	TIndexedContainerIterator& operator++()
 	{
 		++Index;
 		return *this;
 	}
-	UE_NODEBUG TIndexedContainerIterator operator++(int)
+	TIndexedContainerIterator operator++(int)
 	{
 		TIndexedContainerIterator Tmp(*this);
 		++Index;
@@ -82,12 +82,12 @@ public:
 	}
 
 	/** Moves iterator to the previous element in the container. */
-	UE_NODEBUG TIndexedContainerIterator& operator--()
+	TIndexedContainerIterator& operator--()
 	{
 		--Index;
 		return *this;
 	}
-	UE_NODEBUG TIndexedContainerIterator operator--(int)
+	TIndexedContainerIterator operator--(int)
 	{
 		TIndexedContainerIterator Tmp(*this);
 		--Index;
@@ -95,65 +95,65 @@ public:
 	}
 
 	/** iterator arithmetic support */
-	UE_NODEBUG TIndexedContainerIterator& operator+=(SizeType Offset)
+	TIndexedContainerIterator& operator+=(SizeType Offset)
 	{
 		Index += Offset;
 		return *this;
 	}
 
-	UE_NODEBUG [[nodiscard]] TIndexedContainerIterator operator+(SizeType Offset) const
+	[[nodiscard]] TIndexedContainerIterator operator+(SizeType Offset) const
 	{
 		TIndexedContainerIterator Tmp(*this);
 		return Tmp += Offset;
 	}
 
-	UE_NODEBUG TIndexedContainerIterator& operator-=(SizeType Offset)
+	TIndexedContainerIterator& operator-=(SizeType Offset)
 	{
 		return *this += -Offset;
 	}
 
-	UE_NODEBUG [[nodiscard]] TIndexedContainerIterator operator-(SizeType Offset) const
+	[[nodiscard]] TIndexedContainerIterator operator-(SizeType Offset) const
 	{
 		TIndexedContainerIterator Tmp(*this);
 		return Tmp -= Offset;
 	}
 
-	UE_NODEBUG [[nodiscard]] FORCEINLINE ElementType& operator* () const
+	[[nodiscard]] FORCEINLINE ElementType& operator* () const
 	{
 		return Container[ Index ];
 	}
 
-	UE_NODEBUG [[nodiscard]] FORCEINLINE ElementType* operator->() const
+	[[nodiscard]] FORCEINLINE ElementType* operator->() const
 	{
 		return &Container[ Index ];
 	}
 
 	/** conversion to "bool" returning true if the iterator has not reached the last element. */
-	UE_NODEBUG [[nodiscard]] FORCEINLINE explicit operator bool() const
+	[[nodiscard]] FORCEINLINE explicit operator bool() const
 	{
 		return Container.IsValidIndex(Index);
 	}
 
 	/** Returns an index to the current element. */
-	UE_NODEBUG [[nodiscard]] SizeType GetIndex() const
+	[[nodiscard]] SizeType GetIndex() const
 	{
 		return Index;
 	}
 
 	/** Resets the iterator to the first element. */
-	UE_NODEBUG void Reset()
+	void Reset()
 	{
 		Index = 0;
 	}
 
 	/** Sets the iterator to one past the last element. */
-	UE_NODEBUG void SetToEnd()
+	void SetToEnd()
 	{
 		Index = Container.Num();
 	}
-
+	
 	/** Removes current element in array. This invalidates the current iterator value and it must be incremented */
-	UE_NODEBUG void RemoveCurrent()
+	void RemoveCurrent()
 	{
 		Container.RemoveAt(Index);
 		Index--;
@@ -164,24 +164,17 @@ public:
 	 * This invalidates the current iterator value and it must be incremented.
 	 * Note this modifies the order of the remaining elements in the array.
 	 */
-	UE_NODEBUG void RemoveCurrentSwap()
+	void RemoveCurrentSwap()
 	{
 		Container.RemoveAtSwap(Index);
 		Index--;
 	}
 
-	UE_NODEBUG [[nodiscard]] FORCEINLINE bool operator==(const TIndexedContainerIterator& Rhs) const
-	{
-		return &Container == &Rhs.Container && Index == Rhs.Index;
-	}
-#if !PLATFORM_COMPILER_HAS_GENERATED_COMPARISON_OPERATORS
-	UE_NODEBUG [[nodiscard]] FORCEINLINE bool operator!=(const TIndexedContainerIterator& Rhs) const
-	{
-		return &Container != &Rhs.Container || Index != Rhs.Index;
-	}
-#endif
+	[[nodiscard]] FORCEINLINE bool operator==(const TIndexedContainerIterator& Rhs) const { return &Container == &Rhs.Container && Index == Rhs.Index; }
+	[[nodiscard]] FORCEINLINE bool operator!=(const TIndexedContainerIterator& Rhs) const { return &Container != &Rhs.Container || Index != Rhs.Index; }
 
 private:
+
 	ContainerType& Container;
 	SizeType      Index;
 };
@@ -189,7 +182,7 @@ private:
 
 /** operator + */
 template <typename ContainerType, typename ElementType, typename SizeType>
-UE_NODEBUG [[nodiscard]] FORCEINLINE TIndexedContainerIterator<ContainerType, ElementType, SizeType> operator+(SizeType Offset, TIndexedContainerIterator<ContainerType, ElementType, SizeType> RHS)
+[[nodiscard]] FORCEINLINE TIndexedContainerIterator<ContainerType, ElementType, SizeType> operator+(SizeType Offset, TIndexedContainerIterator<ContainerType, ElementType, SizeType> RHS)
 {
 	return RHS + Offset;
 }
@@ -208,14 +201,14 @@ UE_NODEBUG [[nodiscard]] FORCEINLINE TIndexedContainerIterator<ContainerType, El
 		//
 		// We do add an operator-- to help FString implementation
 
-		UE_NODEBUG [[nodiscard]] explicit TCheckedPointerIterator(const SizeType& InNum, ElementType* InPtr)
+		[[nodiscard]] explicit TCheckedPointerIterator(const SizeType& InNum, ElementType* InPtr)
 			: Ptr       (InPtr)
 			, CurrentNum(InNum)
 			, InitialNum(InNum)
 		{
 		}
 
-		UE_NODEBUG [[nodiscard]] FORCEINLINE ElementType& operator*() const
+		[[nodiscard]] FORCEINLINE ElementType& operator*() const
 		{
 			if constexpr (bReverse)
 			{
@@ -227,7 +220,7 @@ UE_NODEBUG [[nodiscard]] FORCEINLINE TIndexedContainerIterator<ContainerType, El
 			}
 		}
 
-		UE_NODEBUG FORCEINLINE TCheckedPointerIterator& operator++()
+		FORCEINLINE TCheckedPointerIterator& operator++()
 		{
 			if constexpr (bReverse)
 			{
@@ -240,7 +233,7 @@ UE_NODEBUG [[nodiscard]] FORCEINLINE TIndexedContainerIterator<ContainerType, El
 			return *this;
 		}
 
-		UE_NODEBUG FORCEINLINE TCheckedPointerIterator& operator--()
+		FORCEINLINE TCheckedPointerIterator& operator--()
 		{
 			if constexpr (bReverse)
 			{
@@ -253,7 +246,7 @@ UE_NODEBUG [[nodiscard]] FORCEINLINE TIndexedContainerIterator<ContainerType, El
 			return *this;
 		}
 
-		UE_NODEBUG [[nodiscard]] FORCEINLINE bool operator!=(const TCheckedPointerIterator& Rhs) const
+		[[nodiscard]] FORCEINLINE bool operator!=(const TCheckedPointerIterator& Rhs) const
 		{
 			// We only need to do the check in this operator, because no other operator will be
 			// called until after this one returns.
@@ -275,23 +268,23 @@ UE_NODEBUG [[nodiscard]] FORCEINLINE TIndexedContainerIterator<ContainerType, El
 template <typename ElementType, typename IteratorType>
 struct TDereferencingIterator
 {
-	UE_NODEBUG [[nodiscard]] explicit TDereferencingIterator(IteratorType InIter)
+	[[nodiscard]] explicit TDereferencingIterator(IteratorType InIter)
 		: Iter(InIter)
 	{
 	}
 
-	UE_NODEBUG [[nodiscard]] FORCEINLINE ElementType& operator*() const
+	[[nodiscard]] FORCEINLINE ElementType& operator*() const
 	{
 		return *(ElementType*)*Iter;
 	}
 
-	UE_NODEBUG FORCEINLINE TDereferencingIterator& operator++()
+	FORCEINLINE TDereferencingIterator& operator++()
 	{
 		++Iter;
 		return *this;
 	}
 
-	UE_NODEBUG [[nodiscard]] FORCEINLINE bool operator!=(const TDereferencingIterator& Rhs) const
+	[[nodiscard]] FORCEINLINE bool operator!=(const TDereferencingIterator& Rhs) const
 	{
 		return Iter != Rhs.Iter;
 	}
@@ -311,7 +304,7 @@ namespace UE::Core::Private
 	}
 
 	template <typename FromArrayType, typename ToArrayType>
-	UE_NODEBUG [[nodiscard]] constexpr bool CanMoveTArrayPointersBetweenArrayTypes()
+	[[nodiscard]] constexpr bool CanMoveTArrayPointersBetweenArrayTypes()
 	{
 		using FromAllocatorType          = typename FromArrayType::AllocatorType;
 		using ToAllocatorType            = typename ToArrayType::AllocatorType;
@@ -346,14 +339,15 @@ namespace UE::Core::Private
 	//     TArray<FThing> Arr; // this will cause errors without this workaround
 	// };
 	template <typename DestType, typename SourceType>
-	UE_NODEBUG constexpr bool TArrayElementsAreCompatible_V = std::disjunction_v<std::is_same<DestType, std::decay_t<SourceType>>, std::is_constructible<DestType, SourceType>>;
+	constexpr bool TArrayElementsAreCompatible_V = std::disjunction_v<std::is_same<DestType, std::decay_t<SourceType>>, std::is_constructible<DestType, SourceType>>;
 
 	template <typename ElementType, typename AllocatorType>
-	UE_NODEBUG static char (&ResolveIsTArrayPtr(const volatile TArray<ElementType, AllocatorType>*))[2];
-	UE_NODEBUG static char(&ResolveIsTArrayPtr(...))[1];
+	static char (&ResolveIsTArrayPtr(const volatile TArray<ElementType, AllocatorType>*))[2];
+
+	static char(&ResolveIsTArrayPtr(...))[1];
 
 	template <typename T>
-	UE_NODEBUG constexpr bool TIsTArrayOrDerivedFromTArray_V = sizeof(ResolveIsTArrayPtr((T*)nullptr)) == 2;
+	constexpr bool TIsTArrayOrDerivedFromTArray_V = sizeof(ResolveIsTArrayPtr((T*)nullptr)) == 2;
 
 	[[noreturn]] CORE_API void OnInvalidArrayNum(unsigned long long NewNum);
 
@@ -369,7 +363,7 @@ namespace UE::Core::Private
 	//
 	// When C++20 is guaranteed, concept checks can be used instead.
 	template <typename AllocatorType>
-	UE_NODEBUG [[nodiscard]] constexpr uint32 GetAllocatorFlags()
+	[[nodiscard]] constexpr uint32 GetAllocatorFlags()
 	{
 		uint32 Result = 0;
 		if constexpr (TAllocatorTraits<AllocatorType>::SupportsElementAlignment)
@@ -902,16 +896,16 @@ public:
 	///////////////////////////////////////////////
 	// Start - intrusive TOptional<TArray> state //
 	///////////////////////////////////////////////
-	UE_NODEBUG constexpr static bool bHasIntrusiveUnsetOptionalState = true;
+	constexpr static bool bHasIntrusiveUnsetOptionalState = true;
 	using IntrusiveUnsetOptionalStateType = TArray;
 
-	UE_NODEBUG [[nodiscard]] explicit TArray(FIntrusiveUnsetOptionalState Tag)
+	[[nodiscard]] explicit TArray(FIntrusiveUnsetOptionalState Tag)
 		: ArrayNum(0)
 		, ArrayMax(-1)
 	{
 		// Use ArrayMax == -1 as our intrusive state so that the destructor still works without change, as it doesn't use ArrayMax.
 	}
-	UE_NODEBUG [[nodiscard]] bool operator==(FIntrusiveUnsetOptionalState Tag) const
+	[[nodiscard]] bool operator==(FIntrusiveUnsetOptionalState Tag) const
 	{
 		return ArrayMax == -1;
 	}
@@ -924,7 +918,7 @@ public:
 	 *
 	 * @returns Pointer to first array entry or nullptr if ArrayMax == 0.
 	 */
-	UE_NODEBUG [[nodiscard]] FORCEINLINE ElementType* GetData() UE_LIFETIMEBOUND
+	[[nodiscard]] FORCEINLINE ElementType* GetData() UE_LIFETIMEBOUND
 	{
 		return (ElementType*)AllocatorInstance.GetAllocation();
 	}
@@ -934,7 +928,7 @@ public:
 	 *
 	 * @returns Pointer to first array entry or nullptr if ArrayMax == 0.
 	 */
-	UE_NODEBUG [[nodiscard]] FORCEINLINE const ElementType* GetData() const UE_LIFETIMEBOUND
+	[[nodiscard]] FORCEINLINE const ElementType* GetData() const UE_LIFETIMEBOUND
 	{
 		return (const ElementType*)AllocatorInstance.GetAllocation();
 	}
@@ -944,7 +938,7 @@ public:
 	 *
 	 * @returns Size in bytes of array type.
 	 */
-	UE_NODEBUG [[nodiscard]] FORCEINLINE static constexpr uint32 GetTypeSize()
+	[[nodiscard]] FORCEINLINE static constexpr uint32 GetTypeSize()
 	{
 		return sizeof(ElementType);
 	}
@@ -956,7 +950,7 @@ public:
 	 *
 	 * @returns Number of bytes allocated by this container.
 	 */
-	UE_NODEBUG [[nodiscard]] FORCEINLINE SIZE_T GetAllocatedSize(void) const
+	[[nodiscard]] FORCEINLINE SIZE_T GetAllocatedSize(void) const
 	{
 		return AllocatorInstance.GetAllocatedSize(ArrayMax, sizeof(ElementType));
 	}
@@ -966,7 +960,7 @@ public:
 	 *
 	 * @see Num, Shrink
 	 */
-	UE_NODEBUG [[nodiscard]] FORCEINLINE SizeType GetSlack() const
+	[[nodiscard]] FORCEINLINE SizeType GetSlack() const
 	{
 		return ArrayMax - ArrayNum;
 	}
@@ -975,7 +969,7 @@ public:
 	 * Checks array invariants: if array size is greater than or equal to zero and less
 	 * than or equal to the maximum.
 	 */
-	UE_NODEBUG FORCEINLINE void CheckInvariants() const
+	FORCEINLINE void CheckInvariants() const
 	{
 		checkSlow((ArrayNum >= 0) & (ArrayMax >= ArrayNum)); // & for one branch
 	}
@@ -1019,7 +1013,7 @@ public:
 	 * @param Index Index to test.
 	 * @returns True if index is valid. False otherwise.
 	 */
-	UE_NODEBUG [[nodiscard]] FORCEINLINE bool IsValidIndex(SizeType Index) const
+	[[nodiscard]] FORCEINLINE bool IsValidIndex(SizeType Index) const
 	{
 		return Index >= 0 && Index < ArrayNum;
 	}
@@ -1041,13 +1035,13 @@ public:
 	 * @returns Number of elements in array.
 	 * @see GetSlack
 	 */
-	UE_NODEBUG [[nodiscard]] FORCEINLINE SizeType Num() const
+	[[nodiscard]] FORCEINLINE SizeType Num() const
 	{
 		return ArrayNum;
 	}
 
 	/** @returns Number of bytes used, excluding slack */
-	UE_NODEBUG [[nodiscard]] FORCEINLINE SIZE_T NumBytes() const
+	[[nodiscard]] FORCEINLINE SIZE_T NumBytes() const
 	{
 		return static_cast<SIZE_T>(ArrayNum) * sizeof(ElementType);
 	}
@@ -1058,7 +1052,7 @@ public:
 	 * @returns Maximum number of elements in array.
 	 * @see GetSlack
 	 */
-	UE_NODEBUG [[nodiscard]] FORCEINLINE SizeType Max() const
+	[[nodiscard]] FORCEINLINE SizeType Max() const
 	{
 		return ArrayMax;
 	}
@@ -1068,7 +1062,7 @@ public:
 	 *
 	 * @returns Reference to indexed element.
 	 */
-	UE_NODEBUG [[nodiscard]] FORCEINLINE ElementType& operator[](SizeType Index) UE_LIFETIMEBOUND
+	[[nodiscard]] FORCEINLINE ElementType& operator[](SizeType Index) UE_LIFETIMEBOUND
 	{
 		RangeCheck(Index);
 		return GetData()[Index];
@@ -1081,7 +1075,7 @@ public:
 	 *
 	 * @returns Reference to indexed element.
 	 */
-	UE_NODEBUG [[nodiscard]] FORCEINLINE const ElementType& operator[](SizeType Index) const UE_LIFETIMEBOUND
+	[[nodiscard]] FORCEINLINE const ElementType& operator[](SizeType Index) const UE_LIFETIMEBOUND
 	{
 		RangeCheck(Index);
 		return GetData()[Index];
@@ -1111,7 +1105,7 @@ public:
 		return Result;
 	}
 	UE_ALLOWSHRINKING_BOOL_DEPRECATED("Pop")
-	UE_NODEBUG FORCEINLINE ElementType Pop(bool bAllowShrinking)
+	FORCEINLINE ElementType Pop(bool bAllowShrinking)
 	{
 		return Pop(bAllowShrinking ? EAllowShrinking::Yes : EAllowShrinking::No);
 	}
@@ -1121,7 +1115,7 @@ public:
 	 *
 	 * @param Item Item to push.
 	 */
-	UE_NODEBUG FORCEINLINE void Push(ElementType&& Item)
+	FORCEINLINE void Push(ElementType&& Item)
 	{
 		Add(MoveTempIfPossible(Item));
 	}
@@ -1134,7 +1128,7 @@ public:
 	 * @param Item Item to push.
 	 * @see Pop, Top
 	 */
-	UE_NODEBUG FORCEINLINE void Push(const ElementType& Item)
+	FORCEINLINE void Push(const ElementType& Item)
 	{
 		Add(Item);
 	}
@@ -1145,7 +1139,7 @@ public:
 	 * @returns Reference to the top element.
 	 * @see Pop, Push
 	 */
-	UE_NODEBUG [[nodiscard]] FORCEINLINE ElementType& Top() UE_LIFETIMEBOUND
+	[[nodiscard]] FORCEINLINE ElementType& Top() UE_LIFETIMEBOUND
 	{
 		return Last();
 	}
@@ -1158,7 +1152,7 @@ public:
 	 * @returns Reference to the top element.
 	 * @see Pop, Push
 	 */
-	UE_NODEBUG [[nodiscard]] FORCEINLINE const ElementType& Top() const UE_LIFETIMEBOUND
+	[[nodiscard]] FORCEINLINE const ElementType& Top() const UE_LIFETIMEBOUND
 	{
 		return Last();
 	}
@@ -1169,7 +1163,7 @@ public:
 	 * @param IndexFromTheEnd (Optional) Index from the end of array (default = 0).
 	 * @returns Reference to n-th last element from the array.
 	 */
-	UE_NODEBUG [[nodiscard]] FORCEINLINE ElementType& Last(SizeType IndexFromTheEnd = 0) UE_LIFETIMEBOUND
+	[[nodiscard]] FORCEINLINE ElementType& Last(SizeType IndexFromTheEnd = 0) UE_LIFETIMEBOUND
 	{
 		RangeCheck(ArrayNum - IndexFromTheEnd - 1);
 		return GetData()[ArrayNum - IndexFromTheEnd - 1];
@@ -1183,7 +1177,7 @@ public:
 	 * @param IndexFromTheEnd (Optional) Index from the end of array (default = 0).
 	 * @returns Reference to n-th last element from the array.
 	 */
-	UE_NODEBUG [[nodiscard]] FORCEINLINE const ElementType& Last(SizeType IndexFromTheEnd = 0) const UE_LIFETIMEBOUND
+	[[nodiscard]] FORCEINLINE const ElementType& Last(SizeType IndexFromTheEnd = 0) const UE_LIFETIMEBOUND
 	{
 		RangeCheck(ArrayNum - IndexFromTheEnd - 1);
 		return GetData()[ArrayNum - IndexFromTheEnd - 1];
@@ -1218,7 +1212,7 @@ public:
 	 * @returns True if found. False otherwise.
 	 * @see FindLast, FindLastByPredicate
 	 */
-	UE_NODEBUG FORCEINLINE bool Find(const ElementType& Item, SizeType& Index) const
+	FORCEINLINE bool Find(const ElementType& Item, SizeType& Index) const
 	{
 		Index = this->Find(Item);
 		return Index != INDEX_NONE;
@@ -1252,7 +1246,7 @@ public:
 	 * @returns True if found. False otherwise.
 	 * @see Find, FindLastByPredicate
 	 */
-	UE_NODEBUG FORCEINLINE bool FindLast(const ElementType& Item, SizeType& Index) const
+	FORCEINLINE bool FindLast(const ElementType& Item, SizeType& Index) const
 	{
 		Index = this->FindLast(Item);
 		return Index != INDEX_NONE;
@@ -1306,7 +1300,7 @@ public:
 	 * @returns Index of the found element. INDEX_NONE otherwise.
 	 */
 	template <typename Predicate>
-	UE_NODEBUG [[nodiscard]] FORCEINLINE SizeType FindLastByPredicate(Predicate Pred) const
+	[[nodiscard]] FORCEINLINE SizeType FindLastByPredicate(Predicate Pred) const
 	{
 		return FindLastByPredicate(Pred, ArrayNum);
 	}
@@ -1361,7 +1355,7 @@ public:
 	 * @see Find
 	 */
 	template <typename KeyType>
-	UE_NODEBUG [[nodiscard]] FORCEINLINE const ElementType* FindByKey(const KeyType& Key) const
+	[[nodiscard]] FORCEINLINE const ElementType* FindByKey(const KeyType& Key) const
 	{
 		return const_cast<TArray*>(this)->FindByKey(Key);
 	}
@@ -1396,7 +1390,7 @@ public:
 	 * @see FilterByPredicate, ContainsByPredicate
 	 */
 	template <typename Predicate>
-	UE_NODEBUG [[nodiscard]] FORCEINLINE const ElementType* FindByPredicate(Predicate Pred) const
+	[[nodiscard]] FORCEINLINE const ElementType* FindByPredicate(Predicate Pred) const
 	{
 		return const_cast<TArray*>(this)->FindByPredicate(Pred);
 	}
@@ -1470,7 +1464,7 @@ public:
 	 * @see Contains, Find
 	 */
 	template <typename Predicate>
-	UE_NODEBUG [[nodiscard]] FORCEINLINE bool ContainsByPredicate(Predicate Pred) const
+	[[nodiscard]] FORCEINLINE bool ContainsByPredicate(Predicate Pred) const
 	{
 		return FindByPredicate(Pred) != nullptr;
 	}
@@ -1481,7 +1475,7 @@ public:
 	 * @param OtherArray Array to compare.
 	 * @returns True if this array is the same as OtherArray. False otherwise.
 	 */
-	UE_NODEBUG [[nodiscard]] bool operator==(const TArray& OtherArray) const
+	[[nodiscard]] bool operator==(const TArray& OtherArray) const
 	{
 		SizeType Count = Num();
 
@@ -1495,12 +1489,14 @@ public:
 	 * @returns True if this array is NOT the same as OtherArray. False otherwise.
 	 */
 #if !PLATFORM_COMPILER_HAS_GENERATED_COMPARISON_OPERATORS
-	UE_NODEBUG [[nodiscard]] FORCEINLINE bool operator!=(const TArray& OtherArray) const
+	[[nodiscard]] FORCEINLINE bool operator!=(const TArray& OtherArray) const
 	{
 		return !(*this == OtherArray);
 	}
 #endif
 
+
+
 	/**
 	 * Bulk serialize array as a single memory blob when loading. Uses regular serialization code for saving
 	 * and doesn't serialize at all otherwise (e.g. transient, garbage collection, ...).
@@ -1581,7 +1577,7 @@ public:
 	 *
 	 * @param Ar Archive to count for.
 	 */
-	UE_NODEBUG void CountBytes(FArchive& Ar) const
+	void CountBytes(FArchive& Ar) const
 	{
 		Ar.CountBytes(ArrayNum*sizeof(ElementType), ArrayMax*sizeof(ElementType));
 	}
@@ -1596,7 +1592,7 @@ public:
 	 * @param Count Number of elements to add.
 	 * @returns Number of elements in array before addition.
 	 */
-	FORCEINLINE SizeType AddUninitialized()
+	SizeType AddUninitialized()
 	{
 		return UE::Core::Private::ReallocGrow<UE::Core::Private::GetAllocatorFlags<AllocatorType>()>(
 			sizeof(ElementType),
@@ -1652,7 +1648,7 @@ private:
 		typename OtherSizeType
 		UE_REQUIRES(sizeof(OtherSizeType) > sizeof(SizeType))
 	>
-	UE_NODEBUG void InsertUninitializedImpl(SizeType Index, OtherSizeType Count)
+	void InsertUninitializedImpl(SizeType Index, OtherSizeType Count)
 	{
 		checkf((OtherSizeType)(SizeType)Count == Count, TEXT("Invalid number of elements to add to this array type: %lld"), (long long)(SizeType)Count);
 		return InsertUninitializedImpl(Index, (SizeType)Count);
@@ -1671,11 +1667,11 @@ public:
 	 * @param Count Number of elements to add.
 	 * @see Insert, InsertZeroed, InsertDefaulted
 	 */
-	UE_NODEBUG FORCEINLINE void InsertUninitialized(SizeType Index)
+	FORCEINLINE void InsertUninitialized(SizeType Index)
 	{
 		InsertUninitializedImpl(Index);
 	}
-	UE_NODEBUG FORCEINLINE void InsertUninitialized(SizeType Index, SizeType Count)
+	FORCEINLINE void InsertUninitialized(SizeType Index, SizeType Count)
 	{
 		InsertUninitializedImpl(Index, Count);
 	}
@@ -2130,7 +2126,7 @@ public:
 	}
 	template <typename CountType>
 	UE_ALLOWSHRINKING_BOOL_DEPRECATED("RemoveAtSwap")
-	UE_NODEBUG FORCEINLINE void RemoveAtSwap(SizeType Index, CountType Count, bool bAllowShrinking)
+	FORCEINLINE void RemoveAtSwap(SizeType Index, CountType Count, bool bAllowShrinking)
 	{
 		RemoveAtSwap(Index, Count, bAllowShrinking ? EAllowShrinking::Yes : EAllowShrinking::No);
 	}
@@ -2221,7 +2217,7 @@ public:
 		}
 	}
 	UE_ALLOWSHRINKING_BOOL_DEPRECATED("SetNum")
-	UE_NODEBUG FORCEINLINE void SetNum(SizeType NewNum, bool bAllowShrinking)
+	FORCEINLINE void SetNum(SizeType NewNum, bool bAllowShrinking)
 	{
 		SetNum(NewNum, bAllowShrinking ? EAllowShrinking::Yes : EAllowShrinking::No);
 	}
@@ -2250,7 +2246,7 @@ public:
 		}
 	}
 	UE_ALLOWSHRINKING_BOOL_DEPRECATED("SetNumZeroed")
-	UE_NODEBUG FORCEINLINE void SetNumZeroed(SizeType NewNum, bool bAllowShrinking)
+	FORCEINLINE void SetNumZeroed(SizeType NewNum, bool bAllowShrinking)
 	{
 		SetNumZeroed(NewNum, bAllowShrinking ? EAllowShrinking::Yes : EAllowShrinking::No);
 	}
@@ -2278,7 +2274,7 @@ public:
 		}
 	}
 	UE_ALLOWSHRINKING_BOOL_DEPRECATED("SetNumUninitialized")
-	UE_NODEBUG FORCEINLINE void SetNumUninitialized(SizeType NewNum, bool bAllowShrinking)
+	FORCEINLINE void SetNumUninitialized(SizeType NewNum, bool bAllowShrinking)
 	{
 		SetNumUninitialized(NewNum, bAllowShrinking ? EAllowShrinking::Yes : EAllowShrinking::No);
 	}
@@ -2417,7 +2413,7 @@ public:
 	 *
 	 * @param Other The array to append.
 	 */
-	UE_NODEBUG TArray& operator+=(TArray&& Other)
+	TArray& operator+=(TArray&& Other)
 	{
 		Append(MoveTemp(Other));
 		return *this;
@@ -2429,7 +2425,7 @@ public:
 	 *
 	 * @param Other The array to append.
 	 */
-	UE_NODEBUG TArray& operator+=(const TArray& Other)
+	TArray& operator+=(const TArray& Other)
 	{
 		Append(Other);
 		return *this;
@@ -2440,7 +2436,7 @@ public:
 	 *
 	 * @param InitList The initializer list to append.
 	 */
-	UE_NODEBUG TArray& operator+=(std::initializer_list<ElementType> InitList)
+	TArray& operator+=(std::initializer_list<ElementType> InitList)
 	{
 		Append(InitList);
 		return *this;
@@ -2571,7 +2567,7 @@ public:
 	 * @return Index to the new item
 	 * @see AddDefaulted, AddUnique, AddZeroed, Append, Insert
 	 */
-	UE_NODEBUG FORCEINLINE SizeType Add(ElementType&& Item)
+	FORCEINLINE SizeType Add(ElementType&& Item)
 	{
 		CheckAddress(&Item);
 		return Emplace(MoveTempIfPossible(Item));
@@ -2584,7 +2580,7 @@ public:
 	 * @return Index to the new item
 	 * @see AddDefaulted, AddUnique, AddZeroed, Append, Insert
 	 */
-	UE_NODEBUG FORCEINLINE SizeType Add(const ElementType& Item)
+	FORCEINLINE SizeType Add(const ElementType& Item)
 	{
 		CheckAddress(&Item);
 		return Emplace(Item);
@@ -2599,7 +2595,7 @@ public:
 	 * @return A reference to the newly-inserted element.
 	 * @see AddDefaulted_GetRef, AddUnique_GetRef, AddZeroed_GetRef, Insert_GetRef
 	 */
-	UE_NODEBUG [[nodiscard]] FORCEINLINE ElementType& Add_GetRef(ElementType&& Item) UE_LIFETIMEBOUND
+	[[nodiscard]] FORCEINLINE ElementType& Add_GetRef(ElementType&& Item) UE_LIFETIMEBOUND
 	{
 		CheckAddress(&Item);
 		return Emplace_GetRef(MoveTempIfPossible(Item));
@@ -2612,7 +2608,7 @@ public:
 	 * @return A reference to the newly-inserted element.
 	 * @see AddDefaulted_GetRef, AddUnique_GetRef, AddZeroed_GetRef, Insert_GetRef
 	 */
-	UE_NODEBUG [[nodiscard]] FORCEINLINE ElementType& Add_GetRef(const ElementType& Item) UE_LIFETIMEBOUND
+	[[nodiscard]] FORCEINLINE ElementType& Add_GetRef(const ElementType& Item) UE_LIFETIMEBOUND
 	{
 		CheckAddress(&Item);
 		return Emplace_GetRef(Item);
@@ -3069,7 +3065,7 @@ public:
 	}
 	template <class PREDICATE_CLASS>
 	UE_ALLOWSHRINKING_BOOL_DEPRECATED("RemoveAllSwap")
-	UE_NODEBUG FORCEINLINE SizeType RemoveAllSwap(const PREDICATE_CLASS& Predicate, bool bAllowShrinking)
+	FORCEINLINE SizeType RemoveAllSwap(const PREDICATE_CLASS& Predicate, bool bAllowShrinking)
 	{
 		return RemoveAllSwap(Predicate, bAllowShrinking ? EAllowShrinking::Yes : EAllowShrinking::No);
 	}
@@ -3098,7 +3094,7 @@ public:
 		return 1;
 	}
 	UE_ALLOWSHRINKING_BOOL_DEPRECATED("RemoveSingleSwap")
-	UE_NODEBUG FORCEINLINE SizeType RemoveSingleSwap(const ElementType& Item, bool bAllowShrinking)
+	FORCEINLINE SizeType RemoveSingleSwap(const ElementType& Item, bool bAllowShrinking)
 	{
 		return RemoveSingleSwap(Item, bAllowShrinking ? EAllowShrinking::Yes : EAllowShrinking::No);
 	}
@@ -3145,7 +3141,7 @@ public:
 		return OriginalNum - ArrayNum;
 	}
 	UE_ALLOWSHRINKING_BOOL_DEPRECATED("RemoveSwap")
-	UE_NODEBUG FORCEINLINE SizeType RemoveSwap(const ElementType& Item, bool bAllowShrinking)
+	FORCEINLINE SizeType RemoveSwap(const ElementType& Item, bool bAllowShrinking)
 	{
 		return RemoveSwap(Item, bAllowShrinking ? EAllowShrinking::Yes : EAllowShrinking::No);
 	}
@@ -3227,7 +3223,7 @@ public:
 	 *
 	 * @returns The iterator.
 	 */
-	UE_NODEBUG [[nodiscard]] TIterator CreateIterator()
+	[[nodiscard]] TIterator CreateIterator()
 	{
 		return TIterator(*this);
 	}
@@ -3237,7 +3233,7 @@ public:
 	 *
 	 * @returns The const iterator.
 	 */
-	UE_NODEBUG [[nodiscard]] TConstIterator CreateConstIterator() const
+	[[nodiscard]] TConstIterator CreateConstIterator() const
 	{
 		return TConstIterator(*this);
 	}
@@ -3261,23 +3257,23 @@ public:
 	 * STL-like iterators to enable range-based for loop support.
 	 */
 	#if TARRAY_RANGED_FOR_CHECKS
-		[[nodiscard]] UE_NODEBUG FORCEINLINE RangedForIteratorType             begin ()       { return RangedForIteratorType            (ArrayNum, GetData()); }
-		[[nodiscard]] UE_NODEBUG FORCEINLINE RangedForConstIteratorType        begin () const { return RangedForConstIteratorType       (ArrayNum, GetData()); }
-		[[nodiscard]] UE_NODEBUG FORCEINLINE RangedForIteratorType             end   ()       { return RangedForIteratorType            (ArrayNum, GetData() + Num()); }
-		[[nodiscard]] UE_NODEBUG FORCEINLINE RangedForConstIteratorType        end   () const { return RangedForConstIteratorType       (ArrayNum, GetData() + Num()); }
-		[[nodiscard]] UE_NODEBUG FORCEINLINE RangedForReverseIteratorType      rbegin()       { return RangedForReverseIteratorType     (ArrayNum, GetData() + Num()); }
-		[[nodiscard]] UE_NODEBUG FORCEINLINE RangedForConstReverseIteratorType rbegin() const { return RangedForConstReverseIteratorType(ArrayNum, GetData() + Num()); }
-		[[nodiscard]] UE_NODEBUG FORCEINLINE RangedForReverseIteratorType      rend  ()       { return RangedForReverseIteratorType     (ArrayNum, GetData()); }
-		[[nodiscard]] UE_NODEBUG FORCEINLINE RangedForConstReverseIteratorType rend  () const { return RangedForConstReverseIteratorType(ArrayNum, GetData()); }
+		[[nodiscard]] FORCEINLINE RangedForIteratorType             begin ()       { return RangedForIteratorType            (ArrayNum, GetData()); }
+		[[nodiscard]] FORCEINLINE RangedForConstIteratorType        begin () const { return RangedForConstIteratorType       (ArrayNum, GetData()); }
+		[[nodiscard]] FORCEINLINE RangedForIteratorType             end   ()       { return RangedForIteratorType            (ArrayNum, GetData() + Num()); }
+		[[nodiscard]] FORCEINLINE RangedForConstIteratorType        end   () const { return RangedForConstIteratorType       (ArrayNum, GetData() + Num()); }
+		[[nodiscard]] FORCEINLINE RangedForReverseIteratorType      rbegin()       { return RangedForReverseIteratorType     (ArrayNum, GetData() + Num()); }
+		[[nodiscard]] FORCEINLINE RangedForConstReverseIteratorType rbegin() const { return RangedForConstReverseIteratorType(ArrayNum, GetData() + Num()); }
+		[[nodiscard]] FORCEINLINE RangedForReverseIteratorType      rend  ()       { return RangedForReverseIteratorType     (ArrayNum, GetData()); }
+		[[nodiscard]] FORCEINLINE RangedForConstReverseIteratorType rend  () const { return RangedForConstReverseIteratorType(ArrayNum, GetData()); }
 	#else
-		[[nodiscard]] UE_NODEBUG FORCEINLINE RangedForIteratorType             begin ()       { return                                   GetData(); }
-		[[nodiscard]] UE_NODEBUG FORCEINLINE RangedForConstIteratorType        begin () const { return                                   GetData(); }
-		[[nodiscard]] UE_NODEBUG FORCEINLINE RangedForIteratorType             end   ()       { return                                   GetData() + Num(); }
-		[[nodiscard]] UE_NODEBUG FORCEINLINE RangedForConstIteratorType        end   () const { return                                   GetData() + Num(); }
-		[[nodiscard]] UE_NODEBUG FORCEINLINE RangedForReverseIteratorType      rbegin()       { return RangedForReverseIteratorType     (GetData() + Num()); }
-		[[nodiscard]] UE_NODEBUG FORCEINLINE RangedForConstReverseIteratorType rbegin() const { return RangedForConstReverseIteratorType(GetData() + Num()); }
-		[[nodiscard]] UE_NODEBUG FORCEINLINE RangedForReverseIteratorType      rend  ()       { return RangedForReverseIteratorType     (GetData()); }
-		[[nodiscard]] UE_NODEBUG FORCEINLINE RangedForConstReverseIteratorType rend  () const { return RangedForConstReverseIteratorType(GetData()); }
+		[[nodiscard]] FORCEINLINE RangedForIteratorType             begin ()       { return                                   GetData(); }
+		[[nodiscard]] FORCEINLINE RangedForConstIteratorType        begin () const { return                                   GetData(); }
+		[[nodiscard]] FORCEINLINE RangedForIteratorType             end   ()       { return                                   GetData() + Num(); }
+		[[nodiscard]] FORCEINLINE RangedForConstIteratorType        end   () const { return                                   GetData() + Num(); }
+		[[nodiscard]] FORCEINLINE RangedForReverseIteratorType      rbegin()       { return RangedForReverseIteratorType     (GetData() + Num()); }
+		[[nodiscard]] FORCEINLINE RangedForConstReverseIteratorType rbegin() const { return RangedForConstReverseIteratorType(GetData() + Num()); }
+		[[nodiscard]] FORCEINLINE RangedForReverseIteratorType      rend  ()       { return RangedForReverseIteratorType     (GetData()); }
+		[[nodiscard]] FORCEINLINE RangedForConstReverseIteratorType rend  () const { return RangedForConstReverseIteratorType(GetData()); }
 	#endif
 
 public:
@@ -3290,7 +3286,7 @@ public:
 	 *        If this is not desirable, please use Algo::Sort(MyArray) directly instead.
 	 *        The auto-dereferencing behavior does not occur with smart pointers.
 	 */
-	UE_NODEBUG void Sort()
+	void Sort()
 	{
 		Algo::Sort(*this, TDereferenceWrapper<ElementType, TLess<>>(TLess<>()));
 	}
@@ -3306,7 +3302,7 @@ public:
 	 *        The auto-dereferencing behavior does not occur with smart pointers.
 	 */
 	template <class PREDICATE_CLASS>
-	UE_NODEBUG void Sort(const PREDICATE_CLASS& Predicate)
+	void Sort(const PREDICATE_CLASS& Predicate)
 	{
 		TDereferenceWrapper<ElementType, PREDICATE_CLASS> PredicateWrapper(Predicate);
 		Algo::Sort(*this, PredicateWrapper);
@@ -3322,7 +3318,7 @@ public:
 	 *        If this is not desirable, please use Algo::StableSort(MyArray) directly instead.
 	 *        The auto-dereferencing behavior does not occur with smart pointers.
 	 */
-	UE_NODEBUG void StableSort()
+	void StableSort()
 	{
 		Algo::StableSort(*this, TDereferenceWrapper<ElementType, TLess<>>(TLess<>()));
 	}
@@ -3340,7 +3336,7 @@ public:
 	 *        The auto-dereferencing behavior does not occur with smart pointers.
 	 */
 	template <class PREDICATE_CLASS>
-	UE_NODEBUG void StableSort(const PREDICATE_CLASS& Predicate)
+	void StableSort(const PREDICATE_CLASS& Predicate)
 	{
 		TDereferenceWrapper<ElementType, PREDICATE_CLASS> PredicateWrapper(Predicate);
 		Algo::StableSort(*this, PredicateWrapper);
@@ -3512,7 +3508,7 @@ public:
 	 *        The auto-dereferencing behavior does not occur with smart pointers.
 	 */
 	template <class PREDICATE_CLASS>
-	UE_NODEBUG FORCEINLINE void Heapify(const PREDICATE_CLASS& Predicate)
+	FORCEINLINE void Heapify(const PREDICATE_CLASS& Predicate)
 	{
 		TDereferenceWrapper<ElementType, PREDICATE_CLASS> PredicateWrapper(Predicate);
 		Algo::Heapify(*this, PredicateWrapper);
@@ -3526,7 +3522,7 @@ public:
 	 *        Therefore, your array will be heapified by the values being pointed to, rather than the pointers' values.
 	 *        The auto-dereferencing behavior does not occur with smart pointers.
 	 */
-	UE_NODEBUG void Heapify()
+	void Heapify()
 	{
 		Heapify(TLess<ElementType>());
 	}
@@ -3586,7 +3582,7 @@ public:
 	 *        Therefore, your array will be heapified by the values being pointed to, rather than the pointers' values.
 	 *        The auto-dereferencing behavior does not occur with smart pointers.
 	 */
-	UE_NODEBUG SizeType HeapPush(ElementType&& InItem)
+	SizeType HeapPush(ElementType&& InItem)
 	{
 		return HeapPush(MoveTempIfPossible(InItem), TLess<ElementType>());
 	}
@@ -3602,7 +3598,7 @@ public:
 	 *        Therefore, your array will be heapified by the values being pointed to, rather than the pointers' values.
 	 *        The auto-dereferencing behavior does not occur with smart pointers.
 	 */
-	UE_NODEBUG SizeType HeapPush(const ElementType& InItem)
+	SizeType HeapPush(const ElementType& InItem)
 	{
 		return HeapPush(InItem, TLess<ElementType>());
 	}
@@ -3629,7 +3625,7 @@ public:
 	}
 	template <class PREDICATE_CLASS>
 	UE_ALLOWSHRINKING_BOOL_DEPRECATED("HeapPop")
-	UE_NODEBUG FORCEINLINE void HeapPop(ElementType& OutItem, const PREDICATE_CLASS& Predicate, bool bAllowShrinking)
+	FORCEINLINE void HeapPop(ElementType& OutItem, const PREDICATE_CLASS& Predicate, bool bAllowShrinking)
 	{
 		HeapPop(OutItem, Predicate, bAllowShrinking ? EAllowShrinking::Yes : EAllowShrinking::No);
 	}
@@ -3645,12 +3641,12 @@ public:
 	 *        Therefore, your array will be heapified by the values being pointed to, rather than the pointers' values.
 	 *        The auto-dereferencing behavior does not occur with smart pointers.
 	 */
-	UE_NODEBUG void HeapPop(ElementType& OutItem, EAllowShrinking AllowShrinking = EAllowShrinking::Default)
+	void HeapPop(ElementType& OutItem, EAllowShrinking AllowShrinking = EAllowShrinking::Default)
 	{
 		HeapPop(OutItem, TLess<ElementType>(), AllowShrinking);
 	}
 	UE_ALLOWSHRINKING_BOOL_DEPRECATED("HeapPop")
-	UE_NODEBUG FORCEINLINE void HeapPop(ElementType& OutItem, bool bAllowShrinking)
+	FORCEINLINE void HeapPop(ElementType& OutItem, bool bAllowShrinking)
 	{
 		HeapPop(OutItem, bAllowShrinking ? EAllowShrinking::Yes : EAllowShrinking::No);
 	}
@@ -3661,7 +3657,7 @@ public:
 	 * @param Predicate Predicate class instance.
 	 */
 	template <class PREDICATE_CLASS>
-	UE_NODEBUG void VerifyHeap(const PREDICATE_CLASS& Predicate)
+	void VerifyHeap(const PREDICATE_CLASS& Predicate)
 	{
 		check(Algo::IsHeap(*this, Predicate));
 	}
@@ -3685,7 +3681,7 @@ public:
 	}
 	template <class PREDICATE_CLASS>
 	UE_ALLOWSHRINKING_BOOL_DEPRECATED("HeapPopDiscard")
-	UE_NODEBUG FORCEINLINE void HeapPopDiscard(const PREDICATE_CLASS& Predicate, bool bAllowShrinking)
+	FORCEINLINE void HeapPopDiscard(const PREDICATE_CLASS& Predicate, bool bAllowShrinking)
 	{
 		HeapPopDiscard(Predicate, bAllowShrinking ? EAllowShrinking::Yes : EAllowShrinking::No);
 	}
@@ -3700,12 +3696,12 @@ public:
 	 *        Therefore, your array will be heapified by the values being pointed to, rather than the pointers' values.
 	 *        The auto-dereferencing behavior does not occur with smart pointers.
 	 */
-	UE_NODEBUG void HeapPopDiscard(EAllowShrinking AllowShrinking = EAllowShrinking::Default)
+	void HeapPopDiscard(EAllowShrinking AllowShrinking = EAllowShrinking::Default)
 	{
 		HeapPopDiscard(TLess<ElementType>(), AllowShrinking);
 	}
 	UE_ALLOWSHRINKING_BOOL_DEPRECATED("HeapPopDiscard")
-	UE_NODEBUG FORCEINLINE void HeapPopDiscard(bool bAllowShrinking)
+	FORCEINLINE void HeapPopDiscard(bool bAllowShrinking)
 	{
 		HeapPopDiscard(TLess<ElementType>(), bAllowShrinking ? EAllowShrinking::Yes : EAllowShrinking::No);
 	}
@@ -3717,7 +3713,7 @@ public:
 	 *
 	 * @returns The reference to the top element from the heap.
 	 */
-	UE_NODEBUG [[nodiscard]] const ElementType& HeapTop() const UE_LIFETIMEBOUND
+	[[nodiscard]] const ElementType& HeapTop() const UE_LIFETIMEBOUND
 	{
 		return (*this)[0];
 	}
@@ -3727,7 +3723,7 @@ public:
 	 *
 	 * @returns The reference to the top element from the heap.
 	 */
-	UE_NODEBUG [[nodiscard]] ElementType& HeapTop() UE_LIFETIMEBOUND
+	[[nodiscard]] ElementType& HeapTop() UE_LIFETIMEBOUND
 	{
 		return (*this)[0];
 	}
@@ -3755,7 +3751,7 @@ public:
 	}
 	template <class PREDICATE_CLASS>
 	UE_ALLOWSHRINKING_BOOL_DEPRECATED("HeapRemoveAt")
-	UE_NODEBUG void HeapRemoveAt(SizeType Index, const PREDICATE_CLASS& Predicate, bool bAllowShrinking)
+	void HeapRemoveAt(SizeType Index, const PREDICATE_CLASS& Predicate, bool bAllowShrinking)
 	{
 		HeapRemoveAt(Index, Predicate, bAllowShrinking ? EAllowShrinking::Yes : EAllowShrinking::No);
 	}
@@ -3771,12 +3767,12 @@ public:
 	 *        Therefore, your array will be heapified by the values being pointed to, rather than the pointers' values.
 	 *        The auto-dereferencing behavior does not occur with smart pointers.
 	 */
-	UE_NODEBUG void HeapRemoveAt(SizeType Index, EAllowShrinking AllowShrinking = EAllowShrinking::Default)
+	void HeapRemoveAt(SizeType Index, EAllowShrinking AllowShrinking = EAllowShrinking::Default)
 	{
 		HeapRemoveAt(Index, TLess< ElementType >(), AllowShrinking);
 	}
 	UE_ALLOWSHRINKING_BOOL_DEPRECATED("HeapRemoveAt")
-	UE_NODEBUG FORCEINLINE void HeapRemoveAt(SizeType Index, bool bAllowShrinking)
+	FORCEINLINE void HeapRemoveAt(SizeType Index, bool bAllowShrinking)
 	{
 		HeapRemoveAt(Index, bAllowShrinking ? EAllowShrinking::Yes : EAllowShrinking::No);
 	}
@@ -3791,7 +3787,7 @@ public:
 	 *        The auto-dereferencing behavior does not occur with smart pointers.
 	 */
 	template <class PREDICATE_CLASS>
-	UE_NODEBUG void HeapSort(const PREDICATE_CLASS& Predicate)
+	void HeapSort(const PREDICATE_CLASS& Predicate)
 	{
 		TDereferenceWrapper<ElementType, PREDICATE_CLASS> PredicateWrapper(Predicate);
 		Algo::HeapSort(*this, PredicateWrapper);
@@ -3805,16 +3801,16 @@ public:
 	 *        Therefore, your array will be heapified by the values being pointed to, rather than the pointers' values.
 	 *        The auto-dereferencing behavior does not occur with smart pointers.
 	 */
-	UE_NODEBUG void HeapSort()
+	void HeapSort()
 	{
 		HeapSort(TLess<ElementType>());
 	}
 
-	UE_NODEBUG [[nodiscard]] const ElementAllocatorType& GetAllocatorInstance() const
+	[[nodiscard]] const ElementAllocatorType& GetAllocatorInstance() const
 	{
 		return AllocatorInstance;
 	}
-	UE_NODEBUG [[nodiscard]] ElementAllocatorType& GetAllocatorInstance()
+	[[nodiscard]] ElementAllocatorType& GetAllocatorInstance()
 	{
 		return AllocatorInstance;
 	}
@@ -3826,33 +3822,33 @@ public:
 namespace Freeze
 {
 	template<typename T, typename AllocatorType>
-	UE_NODEBUG void IntrinsicWriteMemoryImage(FMemoryImageWriter& Writer, const TArray<T, AllocatorType>& Object, const FTypeLayoutDesc&)
+	void IntrinsicWriteMemoryImage(FMemoryImageWriter& Writer, const TArray<T, AllocatorType>& Object, const FTypeLayoutDesc&)
 	{
 		Object.WriteMemoryImage(Writer);
 	}
 
 	template<typename T, typename AllocatorType>
-	UE_NODEBUG [[nodiscard]] uint32 IntrinsicUnfrozenCopy(const FMemoryUnfreezeContent& Context, const TArray<T, AllocatorType>& Object, void* OutDst)
+	[[nodiscard]] uint32 IntrinsicUnfrozenCopy(const FMemoryUnfreezeContent& Context, const TArray<T, AllocatorType>& Object, void* OutDst)
 	{
 		Object.CopyUnfrozen(Context, OutDst);
 		return sizeof(Object);
 	}
 
 	template<typename T, typename AllocatorType>
-	UE_NODEBUG [[nodiscard]] uint32 IntrinsicAppendHash(const TArray<T, AllocatorType>* DummyObject, const FTypeLayoutDesc& TypeDesc, const FPlatformTypeLayoutParameters& LayoutParams, FSHA1& Hasher)
+	[[nodiscard]] uint32 IntrinsicAppendHash(const TArray<T, AllocatorType>* DummyObject, const FTypeLayoutDesc& TypeDesc, const FPlatformTypeLayoutParameters& LayoutParams, FSHA1& Hasher)
 	{
 		return AppendHashForNameAndSize(TypeDesc.Name, sizeof(TArray<T, AllocatorType>), Hasher);
 	}
 
 	template<typename T, typename AllocatorType>
-	UE_NODEBUG [[nodiscard]] uint32 IntrinsicGetTargetAlignment(const TArray<T, AllocatorType>* DummyObject, const FTypeLayoutDesc& TypeDesc, const FPlatformTypeLayoutParameters& LayoutParams)
+	[[nodiscard]] uint32 IntrinsicGetTargetAlignment(const TArray<T, AllocatorType>* DummyObject, const FTypeLayoutDesc& TypeDesc, const FPlatformTypeLayoutParameters& LayoutParams)
 	{
 		// Assume alignment of array is drive by pointer
 		return FMath::Min(8u, LayoutParams.MaxFieldAlignment);
 	}
 
 	template<typename T, typename AllocatorType>
-	UE_NODEBUG void IntrinsicToString(const TArray<T, AllocatorType>& Object, const FTypeLayoutDesc& TypeDesc, const FPlatformTypeLayoutParameters& LayoutParams, FMemoryToStringContext& OutContext)
+	void IntrinsicToString(const TArray<T, AllocatorType>& Object, const FTypeLayoutDesc& TypeDesc, const FPlatformTypeLayoutParameters& LayoutParams, FMemoryToStringContext& OutContext)
 	{
 		Object.ToString(LayoutParams, OutContext);
 	}
@@ -3861,46 +3857,44 @@ namespace Freeze
 DECLARE_TEMPLATE_INTRINSIC_TYPE_LAYOUT((template <typename T, typename AllocatorType>), (TArray<T, AllocatorType>));
 
 template <typename InElementType, typename AllocatorType>
-struct UE_NODEBUG TIsZeroConstructType<TArray<InElementType, AllocatorType>>
+struct TIsZeroConstructType<TArray<InElementType, AllocatorType>>
 {
-	enum { Value UE_NODEBUG = TAllocatorTraits<AllocatorType>::IsZeroConstruct };
+	enum { Value = TAllocatorTraits<AllocatorType>::IsZeroConstruct };
 };
 
 template <typename T, typename AllocatorType>
-struct UE_NODEBUG TIsContiguousContainer<TArray<T, AllocatorType>>
+struct TIsContiguousContainer<TArray<T, AllocatorType>>
 {
-	enum { Value UE_NODEBUG = true };
+	enum { Value = true };
 };
 
 /**
  * Trait which determines whether or not a type is a TArray.
  */
-template <typename T> UE_NODEBUG constexpr bool TIsTArray_V = false;
+template <typename T> constexpr bool TIsTArray_V = false;
 
-template <typename InElementType, typename InAllocatorType> UE_NODEBUG constexpr bool TIsTArray_V<               TArray<InElementType, InAllocatorType>> = true;
-template <typename InElementType, typename InAllocatorType> UE_NODEBUG constexpr bool TIsTArray_V<const          TArray<InElementType, InAllocatorType>> = true;
-template <typename InElementType, typename InAllocatorType> UE_NODEBUG constexpr bool TIsTArray_V<      volatile TArray<InElementType, InAllocatorType>> = true;
-template <typename InElementType, typename InAllocatorType> UE_NODEBUG constexpr bool TIsTArray_V<const volatile TArray<InElementType, InAllocatorType>> = true;
+template <typename InElementType, typename InAllocatorType> constexpr bool TIsTArray_V<               TArray<InElementType, InAllocatorType>> = true;
+template <typename InElementType, typename InAllocatorType> constexpr bool TIsTArray_V<const          TArray<InElementType, InAllocatorType>> = true;
+template <typename InElementType, typename InAllocatorType> constexpr bool TIsTArray_V<      volatile TArray<InElementType, InAllocatorType>> = true;
+template <typename InElementType, typename InAllocatorType> constexpr bool TIsTArray_V<const volatile TArray<InElementType, InAllocatorType>> = true;
 
 template <typename T>
 struct TIsTArray
 {
-	enum { Value UE_NODEBUG = TIsTArray_V<T> };
+	enum { Value = TIsTArray_V<T> };
 };
 
 
 //
 // Array operator news.
 //
-template <typename T,typename AllocatorType>
-UE_NODEBUG void* operator new(size_t Size, TArray<T, AllocatorType>& Array)
+template <typename T,typename AllocatorType> void* operator new(size_t Size, TArray<T, AllocatorType>& Array)
 {
 	check(Size == sizeof(T));
 	const auto Index = Array.AddUninitialized();
 	return &Array[Index];
 }
-template <typename T,typename AllocatorType>
-UE_NODEBUG void* operator new(size_t Size, TArray<T,AllocatorType>& Array, typename TArray<T, AllocatorType>::SizeType Index)
+template <typename T,typename AllocatorType> void* operator new( size_t Size, TArray<T,AllocatorType>& Array, typename TArray<T, AllocatorType>::SizeType Index )
 {
 	check(Size == sizeof(T));
 	Array.InsertUninitialized(Index);
@@ -4013,7 +4007,7 @@ struct TArrayPrivateFriend
 
 
 template<typename ElementType, typename AllocatorType>
-UE_NODEBUG FArchive& operator<<(FArchive& Ar, TArray<ElementType, AllocatorType>& A)
+FArchive& operator<<(FArchive& Ar, TArray<ElementType, AllocatorType>& A)
 {
 	return TArrayPrivateFriend::Serialize(Ar, A);
 }
