Add SecondsToCycles64 to PlatformTime and clean up the header to properly explain Cycles64.
This will be used to replace slow calls to Seconds() where it would be better to precompute the cycles.
Convert ThreadHeartBeat to use the new function as an example.

diff --git a/Engine/Source/Runtime/Core/Private/GenericPlatform/GenericPlatformTime.cpp b/Engine/Source/Runtime/Core/Private/GenericPlatform/GenericPlatformTime.cpp
--- a/Engine/Source/Runtime/Core/Private/GenericPlatform/GenericPlatformTime.cpp
+++ b/Engine/Source/Runtime/Core/Private/GenericPlatform/GenericPlatformTime.cpp
@@ -15,12 +15,6 @@ double FGenericPlatformTime::SecondsPerCycle = 0.0;
 double FGenericPlatformTime::SecondsPerCycle64 = 0.0;
 double FGenericPlatformTime::LastIntervalCPUTimeInSeconds = 0.0;
 
-double FGenericPlatformTime::GetSecondsPerCycle64()
-{
-	check(SecondsPerCycle64 != 0.0);
-	return SecondsPerCycle64;
-}
-
 #if PLATFORM_HAS_BSD_TIME 
 
 #include <sched.h>
diff --git a/Engine/Source/Runtime/Core/Private/HAL/ThreadHeartBeat.cpp b/Engine/Source/Runtime/Core/Private/HAL/ThreadHeartBeat.cpp
--- a/Engine/Source/Runtime/Core/Private/HAL/ThreadHeartBeat.cpp
+++ b/Engine/Source/Runtime/Core/Private/HAL/ThreadHeartBeat.cpp
@@ -35,7 +35,7 @@ const double HangDetectorClock_MaxTimeStep_MS = 2000.0;
 const double HitchDetectorClock_MaxTimeStep_MS = 50.0;
 
 FThreadHeartBeatClock::FThreadHeartBeatClock(double InMaxTimeStep)
-	: MaxTimeStepCycles((uint64)(InMaxTimeStep / FPlatformTime::GetSecondsPerCycle64()))
+	: MaxTimeStepCycles(FPlatformTime::SecondsToCycles64(InMaxTimeStep))
 {
 	CurrentCycles = FPlatformTime::Cycles64();
 	LastRealTickCycles = CurrentCycles;
@@ -55,12 +55,12 @@ double FThreadHeartBeatClock::Seconds()
 {
 	uint64 Offset = FPlatformTime::Cycles64() - LastRealTickCycles;
 	uint64 ClampedOffset = FMath::Min(Offset, MaxTimeStepCycles);
-	uint64 CyclesPerSecond = (uint64)(1.0 / FPlatformTime::GetSecondsPerCycle64());
+	uint64 CyclesPerSecond = FPlatformTime::SecondsToCycles64(1.0);
 	uint64 Cycles = CurrentCycles + ClampedOffset;
 	uint64 Seconds = Cycles / CyclesPerSecond;
 	uint64 RemainderCycles = Cycles % CyclesPerSecond;
 
-	return (double)Seconds + (double)RemainderCycles * FPlatformTime::GetSecondsPerCycle64();
+	return (double)Seconds + FPlatformTime::ToSeconds64(RemainderCycles);
 }
 
 FThreadHeartBeat::FThreadHeartBeat()
diff --git a/Engine/Source/Runtime/Core/Public/GenericPlatform/GenericPlatformTime.h b/Engine/Source/Runtime/Core/Public/GenericPlatform/GenericPlatformTime.h
--- a/Engine/Source/Runtime/Core/Public/GenericPlatform/GenericPlatformTime.h
+++ b/Engine/Source/Runtime/Core/Public/GenericPlatform/GenericPlatformTime.h
@@ -166,12 +166,19 @@ struct FGenericPlatformTime
 	}
 
 	/**
-	* @return seconds per cycle.
-	*/
+	 * Each platform implements these two functions, which return the current time in platform-specific cpu cycles.
+	 * Cycles64 should be used for most purposes as it may be higher resolution and will not roll over during execution:
+	 * 
+	 * static uint32 Cycles();
+	 * static uint64 Cycles64();
+	 */
+
+	/** Returns seconds per cycle, to pair with Cycles(). */
 	static double GetSecondsPerCycle()
 	{
 		return SecondsPerCycle;
 	}
+
 	/** Converts cycles to milliseconds. */
 	static float ToMilliseconds( const uint32 Cycles )
 	{
@@ -183,22 +190,32 @@ struct FGenericPlatformTime
 	{
 		return (float)double( SecondsPerCycle * Cycles );
 	}
-	/**
-	 * @return seconds per cycle.
-	 */
-	static CORE_API double GetSecondsPerCycle64();
-	/** Converts cycles to milliseconds. */
+
+
+	/** Returns seconds per cycle, to pair with Cycles64(). */
+	static double GetSecondsPerCycle64()
+	{
+		return SecondsPerCycle64;
+	}
+
+	/** Converts 64 bit cycles to milliseconds. */
 	static double ToMilliseconds64(const uint64 Cycles)
 	{
 		return ToSeconds64(Cycles) * 1000.0;
 	}
 
-	/** Converts cycles to seconds. */
+	/** Converts 64 bit cycles to seconds. */
 	static double ToSeconds64(const uint64 Cycles)
 	{
 		return GetSecondsPerCycle64() * double(Cycles);
 	}
 
+	/** Convert seconds to cycles, can be added to Cycles64 to set a high resolution timeout */
+	static uint64 SecondsToCycles64(double Seconds)
+	{
+		return static_cast<uint64>(Seconds / GetSecondsPerCycle64());
+	}
+
 protected:
 
 	static CORE_API double SecondsPerCycle;
