Fix build error see in some versions of compilers
#rnx

diff --git a/Engine/Source/Runtime/Core/Public/Stats/LightweightStats.h b/Engine/Source/Runtime/Core/Public/Stats/LightweightStats.h
--- a/Engine/Source/Runtime/Core/Public/Stats/LightweightStats.h
+++ b/Engine/Source/Runtime/Core/Public/Stats/LightweightStats.h
@@ -357,7 +357,7 @@ namespace UE::Stats::Private
 			return GetStatIfEnabled<StatDataType, GroupDataType>(StatName, StatNameHash, GroupNameHash);
 		}
 		// Preserve existing behavior of not emitting hitch scopes if the condition is false.
-		return FCheckedStat::MakeInvalid(bCondition ? StatName : nullptr);
+		return FCheckedStat::MakeInvalid(nullptr);
 	}
 } // namespace UE::Stats::Private
 
