[AutoRTFM] Skip the AutoRTFM extern-API glue from crash stack traces

The top of all crash stack traces include AutoRTFM <-> UE binding functions which are not interesting, and may be preventing our crash reporting tooling from properly clustering bugs. Pass down the program counter to the UE fatal / ensure APIs so that we skip pass the boilerplate.

diff --git a/Engine/Source/Runtime/Core/Private/AutoRTFM/API.cpp b/Engine/Source/Runtime/Core/Private/AutoRTFM/API.cpp
--- a/Engine/Source/Runtime/Core/Private/AutoRTFM/API.cpp
+++ b/Engine/Source/Runtime/Core/Private/AutoRTFM/API.cpp
@@ -482,7 +482,7 @@ AUTORTFM_CORE_API autortfm_status RTFM_autortfm_close(void (*UninstrumentedWork)
 	}
 	else
 	{
-		ReportError("Could not find function %p '%s' for autortfm_close()",
+		AUTORTFM_REPORT_ERROR("Could not find function %p '%s' for autortfm_close()",
 			reinterpret_cast<void*>(UninstrumentedWork), GetFunctionDescription(UninstrumentedWork).c_str());
 	}
 
diff --git a/Engine/Source/Runtime/Core/Private/AutoRTFM/AutoRTFMUE.cpp b/Engine/Source/Runtime/Core/Private/AutoRTFM/AutoRTFMUE.cpp
--- a/Engine/Source/Runtime/Core/Private/AutoRTFM/AutoRTFMUE.cpp
+++ b/Engine/Source/Runtime/Core/Private/AutoRTFM/AutoRTFMUE.cpp
@@ -280,7 +280,7 @@ TStringConversion<TStringConvert<UTF8CHAR, TCHAR>, 128> FmtToTChar(const char* F
 }
 
 UE_AUTORTFM_ALWAYS_OPEN_NO_MEMORY_VALIDATION
-void AutoRTFMLog(const char* File, int Line, autortfm_log_severity Severity, const char* Format, va_list Args)
+void AutoRTFMLog(const char* File, int Line, void* ProgramCounter, autortfm_log_severity Severity, const char* Format, va_list Args)
 {
 	if (!IsSeverityActive(Severity))
 	{
@@ -306,7 +306,7 @@ void AutoRTFMLog(const char* File, int Line, autortfm_log_severity Severity, con
 			break;
 		case autortfm_log_fatal:
 			FDebug::DumpStackTraceToLog(TEXT("AutoRTFM backtrace"), ELogVerbosity::Error);
-			UE::Logging::Private::BasicFatalLog(LogAutoRTFM, &Record, FmtToTChar(Format, Args).Get());
+			UE::Logging::Private::BasicFatalLogWithProgramCounter(LogAutoRTFM, &Record, ProgramCounter, FmtToTChar(Format, Args).Get());
 			break;
 	}
 
@@ -314,7 +314,7 @@ void AutoRTFMLog(const char* File, int Line, autortfm_log_severity Severity, con
 }
 
 UE_AUTORTFM_ALWAYS_OPEN_NO_MEMORY_VALIDATION
-void AutoRTFMLogWithCallstack(autortfm_log_severity Severity, const char* Format, va_list Args)
+void AutoRTFMLogWithCallstack(void* ProgramCounter, autortfm_log_severity Severity, const char* Format, va_list Args)
 {
 	if (!IsSeverityActive(Severity))
 	{
@@ -341,7 +341,7 @@ void AutoRTFMLogWithCallstack(autortfm_log_severity Severity, const char* Format
 			break;
 		case autortfm_log_fatal:
 			FDebug::DumpStackTraceToLog(TEXT("AutoRTFM backtrace"), ELogVerbosity::Error);
-			UE::Logging::Private::BasicFatalLog(LogAutoRTFM, &Record, FmtToTChar(Format, Args).Get());
+			UE::Logging::Private::BasicFatalLogWithProgramCounter(LogAutoRTFM, &Record, ProgramCounter, FmtToTChar(Format, Args).Get());
 			break;
 	}
 
@@ -349,7 +349,7 @@ void AutoRTFMLogWithCallstack(autortfm_log_severity Severity, const char* Format
 }
 
 UE_AUTORTFM_ALWAYS_OPEN_NO_MEMORY_VALIDATION
-void AutoRTFMEnsureFailure(const char* File, int Line, const char* Condition, void* ProgramCounter, const char* Format, va_list Args)
+void AutoRTFMEnsureFailure(const char* File, int Line, void* ProgramCounter, const char* Condition, const char* Format, va_list Args)
 {
 #if DO_ENSURE
 	if (Format == nullptr)
diff --git a/Engine/Source/Runtime/Core/Private/AutoRTFM/CompilerInterface.cpp b/Engine/Source/Runtime/Core/Private/AutoRTFM/CompilerInterface.cpp
--- a/Engine/Source/Runtime/Core/Private/AutoRTFM/CompilerInterface.cpp
+++ b/Engine/Source/Runtime/Core/Private/AutoRTFM/CompilerInterface.cpp
@@ -139,18 +139,18 @@ extern "C" UE_AUTORTFM_API void autortfm_memset(void* Dst, int Value, size_t Siz
 
 extern "C" UE_AUTORTFM_API void autortfm_unreachable(const char* Message)
 {
-	ReportError("AutoRTFM Unreachable: %s", Message);
+	ReportError(__FILE__, __LINE__, __builtin_return_address(0), "AutoRTFM Unreachable: %s", Message);
 }
 
 extern "C" UE_AUTORTFM_API void autortfm_llvm_fail(const char* Message)
 {
 	if (Message)
 	{
-		ReportError("AutoRTFM LLVM Failure: %s", Message);
+		ReportError(__FILE__, __LINE__, __builtin_return_address(0), "AutoRTFM LLVM Failure: %s", Message);
 	}
 	else
 	{
-		ReportError("AutoRTFM LLVM Failure");
+		ReportError(__FILE__, __LINE__, __builtin_return_address(0), "AutoRTFM LLVM Failure");
 	}
 }
 
diff --git a/Engine/Source/Runtime/Core/Private/AutoRTFM/ExternAPI.cpp b/Engine/Source/Runtime/Core/Private/AutoRTFM/ExternAPI.cpp
--- a/Engine/Source/Runtime/Core/Private/AutoRTFM/ExternAPI.cpp
+++ b/Engine/Source/Runtime/Core/Private/AutoRTFM/ExternAPI.cpp
@@ -20,7 +20,7 @@ namespace
 // which have logic that can (potentially) log. Logging is unlikely to be called
 // before initialization, but a default printf implementation is preferable to
 // crashing due to a nullptr.
-void DefaultLog(const char* File, int Line, autortfm_log_severity Severity, const char* Format, va_list Args)
+void DefaultLog(const char* File, int Line, [[maybe_unused]] void* ProgramCounter, autortfm_log_severity Severity, const char* Format, va_list Args)
 {
 	static constexpr size_t InlineBufferLength = 256;
 	char InlineBuffer[InlineBufferLength];
diff --git a/Engine/Source/Runtime/Core/Private/AutoRTFM/FunctionMap.cpp b/Engine/Source/Runtime/Core/Private/AutoRTFM/FunctionMap.cpp
--- a/Engine/Source/Runtime/Core/Private/AutoRTFM/FunctionMap.cpp
+++ b/Engine/Source/Runtime/Core/Private/AutoRTFM/FunctionMap.cpp
@@ -145,11 +145,11 @@ void* FunctionMapReportError(void* OldFunction, const char* Where)
     std::string FunctionDescription = GetFunctionDescription(OldFunction);
 	if (Where)
 	{
-		ReportError("Could not find function %p '%s' where '%s'.", OldFunction, FunctionDescription.c_str(), Where);
+		ReportError(__FILE__, __LINE__, __builtin_return_address(0), "Could not find function %p '%s' where '%s'.", OldFunction, FunctionDescription.c_str(), Where);
 	}
 	else
 	{
-		ReportError("Could not find function %p '%s'.", OldFunction, FunctionDescription.c_str());
+		ReportError(__FILE__, __LINE__, __builtin_return_address(0), "Could not find function %p '%s'.", OldFunction, FunctionDescription.c_str());
 	}
 	return nullptr;
 }
diff --git a/Engine/Source/Runtime/Core/Private/AutoRTFM/Utils.cpp b/Engine/Source/Runtime/Core/Private/AutoRTFM/Utils.cpp
--- a/Engine/Source/Runtime/Core/Private/AutoRTFM/Utils.cpp
+++ b/Engine/Source/Runtime/Core/Private/AutoRTFM/Utils.cpp
@@ -49,11 +49,12 @@ std::string GetFunctionDescription(void* FunctionPtr)
 #endif // AUTORTFM_PLATFORM_WINDOWS -> so !AUTORTFM_PLATFORM_WINDOWS
 }
 
-UE_AUTORTFM_API UE_AUTORTFM_FORCENOINLINE void ReportError(const char* Format, ...)
+UE_AUTORTFM_API UE_AUTORTFM_FORCENOINLINE void ReportError(const char* File, int Line, void* ProgramCounter, const char* Format, ...)
 {
-	FContext* const Context = FContext::Get();
-
-	AUTORTFM_FATAL_IF(!Context->IsTransactional(), "AutoRTFM::ReportError() was called from outside a transaction");
+	if (ProgramCounter == nullptr)
+	{
+		ProgramCounter = __builtin_return_address(0);
+	}
 
 	const ForTheRuntime::EAutoRTFMInternalAbortActionState InternalAbortAction = ForTheRuntime::GetInternalAbortAction();
 
@@ -63,14 +64,18 @@ UE_AUTORTFM_API UE_AUTORTFM_FORCENOINLINE void ReportError(const char* Format, .
 		{
 			va_list Args;
 			va_start(Args, Format);
-			AUTORTFM_FATAL_V(Format, Args);
+			::AutoRTFM::LogV(File, Line, ProgramCounter, autortfm_log_fatal, Format, Args);
 			va_end(Args);
 		}
-		else
+		else if (ForTheRuntime::GetEnsureOnInternalAbort())
 		{
 			va_list Args;
 			va_start(Args, Format);
-			AUTORTFM_ENSURE_MSG_V(!ForTheRuntime::GetEnsureOnInternalAbort(), Format, Args);
+			static bool bCalled = [&]
+			{
+				::AutoRTFM::EnsureFailureV(File, Line, ProgramCounter, "!GetEnsureOnInternalAbort()", Format, Args);
+				return true;
+			}();
 			va_end(Args);
 		}
 	}
@@ -78,22 +83,29 @@ UE_AUTORTFM_API UE_AUTORTFM_FORCENOINLINE void ReportError(const char* Format, .
 	{
 		if (InternalAbortAction == ForTheRuntime::EAutoRTFMInternalAbortActionState::Crash)
 		{
-			AUTORTFM_FATAL("Transaction failing because of internal issue");
+			::AutoRTFM::Log(File, Line, ProgramCounter, autortfm_log_fatal, "Transaction failing because of internal issue");
 		}
-		else
+		else if (ForTheRuntime::GetEnsureOnInternalAbort())
 		{
-			AUTORTFM_ENSURE_MSG(!ForTheRuntime::GetEnsureOnInternalAbort(), "Transaction failing because of internal issue");
+			static bool bCalled = [&]
+			{
+				::AutoRTFM::EnsureFailure(File, Line, ProgramCounter, "!GetEnsureOnInternalAbort()", "Transaction failing because of internal issue");
+				return true;
+			}();
 		}
 	}
 
-	if (InternalAbortAction == ForTheRuntime::EAutoRTFMInternalAbortActionState::Abort)
+	if (FContext* const Context = FContext::Get())
 	{
-		Context->AbortByLanguageAndThrow();
-	}
-	else
-	{
-		TTask<void()> Task = [] { AUTORTFM_ENSURE(ForTheRuntime::SetAutoRTFMRuntime(ForTheRuntime::AutoRTFM_ForcedDisabled)); };
-		Context->AbortTransactionWithPostAbortCallback(EContextStatus::AbortedByCascadingRetry, std::move(Task));
+		if (InternalAbortAction == ForTheRuntime::EAutoRTFMInternalAbortActionState::Abort)
+		{
+			Context->AbortByLanguageAndThrow();
+		}
+		else
+		{
+			TTask<void()> Task = [] { AUTORTFM_ENSURE(ForTheRuntime::SetAutoRTFMRuntime(ForTheRuntime::AutoRTFM_ForcedDisabled)); };
+			Context->AbortTransactionWithPostAbortCallback(EContextStatus::AbortedByCascadingRetry, std::move(Task));
+		}
 	}
 }
 
diff --git a/Engine/Source/Runtime/Core/Private/AutoRTFM/Utils.h b/Engine/Source/Runtime/Core/Private/AutoRTFM/Utils.h
--- a/Engine/Source/Runtime/Core/Private/AutoRTFM/Utils.h
+++ b/Engine/Source/Runtime/Core/Private/AutoRTFM/Utils.h
@@ -19,8 +19,12 @@
 namespace AutoRTFM
 {
 
-__attribute__((__format__(__printf__, 1, 2)))
-UE_AUTORTFM_API void ReportError(const char* Format, ...);
+// Reports an internal AutoRTFM issue. The behavior of this function will vary
+// based on ForTheRuntime::GetInternalAbortAction() and 
+// ForTheRuntime::GetEnsureOnInternalAbort().
+// If ProgramCounter is null, then the return address of the call will be used.
+__attribute__((__format__(__printf__, 4, 5)))
+UE_AUTORTFM_API void ReportError(const char* File, int Line, void* ProgramCounter, const char* Format, ...);
 
 [[noreturn]] inline void InternalUnreachable(); //-V1082 Silence PVS [[noreturn]] false positive warning
 
@@ -55,25 +59,35 @@ UE_AUTORTFM_FORCENOINLINE inline void DoExpect(void(*Logger)())
 }
 
 // Logs a message using a printf-style format string and va_list arguments.
-inline void LogV(const char* File, int Line, autortfm_log_severity Severity, const char* Format, va_list Args)
+inline void LogV(const char* File, int Line, void* ProgramCounter, autortfm_log_severity Severity, const char* Format, va_list Args)
 {
-	GExternAPI.Log(File, Line, Severity, Format, Args);
+	if (ProgramCounter == nullptr)
+	{
+		ProgramCounter = __builtin_return_address(0);
+	}
+
+	GExternAPI.Log(File, Line, ProgramCounter, Severity, Format, Args);
 }
 
 // Logs a message using a printf-style format string and variadic arguments.
-__attribute__((__format__(__printf__, 4, 5)))
-inline void Log(const char* File, int Line, autortfm_log_severity Severity, const char* Format, ...)
+__attribute__((__format__(__printf__, 5, 6)))
+inline void Log(const char* File, int Line, void* ProgramCounter, autortfm_log_severity Severity, const char* Format, ...)
 {
+	if (ProgramCounter == nullptr)
+	{
+		ProgramCounter = __builtin_return_address(0);
+	}
+
 	va_list Args;
 	va_start(Args, Format);
-	GExternAPI.Log(File, Line, Severity, Format, Args);
+	GExternAPI.Log(File, Line, ProgramCounter, Severity, Format, Args);
 	va_end(Args);
 }
 
 // Logs a message with a callstack using a printf-style format string and va_list arguments.
 inline void LogWithCallstackV(autortfm_log_severity Severity, const char* Format, va_list Args)
 {
-	GExternAPI.LogWithCallstack(Severity, Format, Args);
+	GExternAPI.LogWithCallstack(__builtin_return_address(0), Severity, Format, Args);
 }
 
 // Logs a message with a callstack using a printf-style format string and variadic arguments.
@@ -82,23 +96,35 @@ inline void LogWithCallstack(autortfm_log_severity Severity, const char* Format,
 {
 	va_list Args;
 	va_start(Args, Format);
-	GExternAPI.LogWithCallstack(Severity, Format, Args);
+	GExternAPI.LogWithCallstack(__builtin_return_address(0), Severity, Format, Args);
 	va_end(Args);
 }
 
 // Reports an ensure failure using a printf-style format string and va_list arguments.
-inline void EnsureFailureV(const char* File, int Line, const char* Condition, const char* Format, va_list Args)
+// If ProgramCounter is null, then the return address of the call will be used.
+inline void EnsureFailureV(const char* File, int Line, void* ProgramCounter, const char* Condition, const char* Format, va_list Args)
 {
-	GExternAPI.EnsureFailure(File, Line, Condition, __builtin_return_address(0), Format, Args);
+	if (ProgramCounter == nullptr)
+	{
+		ProgramCounter = __builtin_return_address(0);
+	}
+
+	GExternAPI.EnsureFailure(File, Line, ProgramCounter, Condition, Format, Args);
 }
 
 // Reports an ensure failure using a printf-style format string and variadic arguments.
-__attribute__((__format__(__printf__, 4, 5)))
-inline void EnsureFailure(const char* File, int Line, const char* Condition, const char* Format, ...)
+// If ProgramCounter is null, then the return address of the call will be used.
+__attribute__((__format__(__printf__, 5, 6)))
+inline void EnsureFailure(const char* File, int Line, void* ProgramCounter, const char* Condition, const char* Format, ...)
 {
+	if (ProgramCounter == nullptr)
+	{
+		ProgramCounter = __builtin_return_address(0);
+	}
+
 	va_list Args;
 	va_start(Args, Format);
-	GExternAPI.EnsureFailure(File, Line, Condition, __builtin_return_address(0), Format, Args);
+	GExternAPI.EnsureFailure(File, Line, ProgramCounter, Condition, Format, Args);
 	va_end(Args);
 }
 
@@ -152,17 +178,19 @@ inline constexpr T Lerp(T Start, T End, T Fraction)
 
 #define AUTORTFM_REQUIRE_SEMICOLON static_assert(true)
 
-#define AUTORTFM_VERBOSE(Format, ...) ::AutoRTFM::Log(__FILE__, __LINE__, autortfm_log_verbose, Format, ##__VA_ARGS__)
-#define AUTORTFM_LOG(Format, ...) ::AutoRTFM::Log(__FILE__, __LINE__, autortfm_log_info, Format, ##__VA_ARGS__)
-#define AUTORTFM_WARN(Format, ...) ::AutoRTFM::Log(__FILE__, __LINE__, autortfm_log_warn, Format, ##__VA_ARGS__)
-#define AUTORTFM_ERROR(Format, ...) ::AutoRTFM::Log(__FILE__, __LINE__, autortfm_log_error, Format, ##__VA_ARGS__)
-#define AUTORTFM_FATAL(Format, ...) ::AutoRTFM::Log(__FILE__, __LINE__, autortfm_log_fatal, Format, ##__VA_ARGS__)
+#define AUTORTFM_REPORT_ERROR(Format, ...) ::AutoRTFM::ReportError(__FILE__, __LINE__, /* ProgramCounter */ nullptr, Format, ##__VA_ARGS__)
+
+#define AUTORTFM_VERBOSE(Format, ...) ::AutoRTFM::Log(__FILE__, __LINE__, /* ProgramCounter */ nullptr, autortfm_log_verbose, Format, ##__VA_ARGS__)
+#define AUTORTFM_LOG(Format, ...) ::AutoRTFM::Log(__FILE__, __LINE__, /* ProgramCounter */ nullptr, autortfm_log_info, Format, ##__VA_ARGS__)
+#define AUTORTFM_WARN(Format, ...) ::AutoRTFM::Log(__FILE__, __LINE__, /* ProgramCounter */ nullptr, autortfm_log_warn, Format, ##__VA_ARGS__)
+#define AUTORTFM_ERROR(Format, ...) ::AutoRTFM::Log(__FILE__, __LINE__, /* ProgramCounter */ nullptr, autortfm_log_error, Format, ##__VA_ARGS__)
+#define AUTORTFM_FATAL(Format, ...) ::AutoRTFM::Log(__FILE__, __LINE__, /* ProgramCounter */ nullptr, autortfm_log_fatal, Format, ##__VA_ARGS__)
 
-#define AUTORTFM_VERBOSE_V(Format, Args) ::AutoRTFM::LogV(__FILE__, __LINE__, autortfm_log_verbose, Format, Args)
-#define AUTORTFM_LOG_V(Format, Args) ::AutoRTFM::LogV(__FILE__, __LINE__, autortfm_log_info, Format, Args)
-#define AUTORTFM_WARN_V(Format, Args) ::AutoRTFM::LogV(__FILE__, __LINE__, autortfm_log_warn, Format, Args)
-#define AUTORTFM_ERROR_V(Format, Args) ::AutoRTFM::LogV(__FILE__, __LINE__, autortfm_log_error, Format, Args)
-#define AUTORTFM_FATAL_V(Format, Args) ::AutoRTFM::LogV(__FILE__, __LINE__, autortfm_log_fatal, Format, Args)
+#define AUTORTFM_VERBOSE_V(Format, Args) ::AutoRTFM::LogV(__FILE__, __LINE__, /* ProgramCounter */ nullptr, autortfm_log_verbose, Format, Args)
+#define AUTORTFM_LOG_V(Format, Args) ::AutoRTFM::LogV(__FILE__, __LINE__, /* ProgramCounter */ nullptr, autortfm_log_info, Format, Args)
+#define AUTORTFM_WARN_V(Format, Args) ::AutoRTFM::LogV(__FILE__, __LINE__, /* ProgramCounter */ nullptr, autortfm_log_warn, Format, Args)
+#define AUTORTFM_ERROR_V(Format, Args) ::AutoRTFM::LogV(__FILE__, __LINE__, /* ProgramCounter */ nullptr, autortfm_log_error, Format, Args)
+#define AUTORTFM_FATAL_V(Format, Args) ::AutoRTFM::LogV(__FILE__, __LINE__, /* ProgramCounter */ nullptr, autortfm_log_fatal, Format, Args)
 
 #define AUTORTFM_VERBOSE_IF(Condition, Format, ...) do { if (Condition) { AUTORTFM_VERBOSE(Format, ##__VA_ARGS__); } } while(false)
 #define AUTORTFM_LOG_IF(Condition, Format, ...) do { if (Condition) { AUTORTFM_LOG(Format, ##__VA_ARGS__); } } while(false)
@@ -184,7 +212,7 @@ inline constexpr T Lerp(T Start, T End, T Fraction)
 #define AUTORTFM_ENSURE(Condition) do { \
 	if (AUTORTFM_UNLIKELY(!(Condition))) { \
 		static bool bCalled = [] { \
-			::AutoRTFM::EnsureFailure(__FILE__, __LINE__, #Condition, nullptr); \
+			::AutoRTFM::EnsureFailure(__FILE__, __LINE__, /* ProgramCounter */ nullptr, #Condition, nullptr); \
 			return true; \
 		}(); \
 	} } while(false)
@@ -192,7 +220,7 @@ inline constexpr T Lerp(T Start, T End, T Fraction)
 #define AUTORTFM_ENSURE_MSG(Condition, Format, ...) do { \
 	if (AUTORTFM_UNLIKELY(!(Condition))) { \
 		static bool bCalled = [&] { \
-			::AutoRTFM::EnsureFailure(__FILE__, __LINE__, #Condition, Format, ##__VA_ARGS__); \
+			::AutoRTFM::EnsureFailure(__FILE__, __LINE__, /* ProgramCounter */ nullptr, #Condition, Format, ##__VA_ARGS__); \
 			return true; \
 		}(); \
 	} } while(false)
@@ -200,7 +228,7 @@ inline constexpr T Lerp(T Start, T End, T Fraction)
 #define AUTORTFM_ENSURE_MSG_V(Condition, Format, Args) do { \
 	if (AUTORTFM_UNLIKELY(!(Condition))) { \
 		static bool bCalled = [&] { \
-			::AutoRTFM::EnsureFailureV(__FILE__, __LINE__, #Condition, Format, Args); \
+			::AutoRTFM::EnsureFailureV(__FILE__, __LINE__, /* ProgramCounter */ nullptr, #Condition, Format, Args); \
 			return true; \
 		}(); \
 	} } while(false)
diff --git a/Engine/Source/Runtime/Core/Private/Logging/StructuredLog.cpp b/Engine/Source/Runtime/Core/Private/Logging/StructuredLog.cpp
--- a/Engine/Source/Runtime/Core/Private/Logging/StructuredLog.cpp
+++ b/Engine/Source/Runtime/Core/Private/Logging/StructuredLog.cpp
@@ -1259,6 +1259,17 @@ void BasicFatalLog(const FLogCategoryBase& Category, const FStaticBasicLogRecord
 	va_end(Args);
 #endif
 }
+
+void BasicFatalLogWithProgramCounter(const FLogCategoryBase& Category, const FStaticBasicLogRecord* Log, void* ProgramCounter, ...)
+{
+#if !NO_LOGGING
+	va_list Args;
+	va_start(Args, ProgramCounter);
+	BasicFatalLogV(Category, Log, ProgramCounter, Args);
+	va_end(Args);
+#endif
+}
+
 } // UE::Logging::Private
 
 ///////////////////////////////////////////////////////////////////////////////////////////////////////////////////////
diff --git a/Engine/Source/Runtime/Core/Public/AutoRTFM/AutoRTFM.h b/Engine/Source/Runtime/Core/Public/AutoRTFM/AutoRTFM.h
--- a/Engine/Source/Runtime/Core/Public/AutoRTFM/AutoRTFM.h
+++ b/Engine/Source/Runtime/Core/Public/AutoRTFM/AutoRTFM.h
@@ -216,17 +216,17 @@ typedef struct
 	// Function used to log messages using a printf-style format string and va_list arguments.
 	// Strings use UTF-8 encoding.
 	// Must not be null.
-	void (*Log)(const char* File, int Line, autortfm_log_severity Severity, const char* Format, va_list Args);
+	void (*Log)(const char* File, int Line, void* ProgramCounter, autortfm_log_severity Severity, const char* Format, va_list Args);
 
 	// Function used to log messages with a callstack using a printf-style format string and va_list arguments.
 	// Strings use UTF-8 encoding.
 	// Must not be null.
-	void (*LogWithCallstack)(autortfm_log_severity Severity, const char* Format, va_list Args);
+	void (*LogWithCallstack)(void* ProgramCounter, autortfm_log_severity Severity, const char* Format, va_list Args);
 	
 	// Function used to report an ensure failure using a printf-style format string and va_list arguments.
 	// Strings use UTF-8 encoding.
 	// Must not be null.
-	void (*EnsureFailure)(const char* File, int Line, const char* Condition, void* ProgramCounter, const char* Format, va_list Args);
+	void (*EnsureFailure)(const char* File, int Line, void* ProgramCounter, const char* Condition, const char* Format, va_list Args);
 
 	// Function used to query whether a log severity is active.
 	// Must not be null.
diff --git a/Engine/Source/Runtime/Core/Public/Logging/LogMacros.h b/Engine/Source/Runtime/Core/Public/Logging/LogMacros.h
--- a/Engine/Source/Runtime/Core/Public/Logging/LogMacros.h
+++ b/Engine/Source/Runtime/Core/Public/Logging/LogMacros.h
@@ -140,6 +140,7 @@ struct FStaticBasicLogRecord
 
 CORE_API void BasicLog(const FLogCategoryBase& Category, const FStaticBasicLogRecord* Log, ...);
 CORE_API void BasicFatalLog(const FLogCategoryBase& Category, const FStaticBasicLogRecord* Log, ...);
+CORE_API void BasicFatalLogWithProgramCounter(const FLogCategoryBase& Category, const FStaticBasicLogRecord* Log, void* ProgramCounter, ...);
 
 } // UE::Logging::Private
 
