- Optimize TLruCache::RemoveByPredicate to iterate the LookupSet and RemoveCurrent, avoids some unnecessary hashes/lookups
- Also fixes it to not call Remove during a range-for iteration.

diff --git a/Engine/Source/Runtime/Core/Public/Containers/LruCache.h b/Engine/Source/Runtime/Core/Public/Containers/LruCache.h
--- a/Engine/Source/Runtime/Core/Public/Containers/LruCache.h
+++ b/Engine/Source/Runtime/Core/Public/Containers/LruCache.h
@@ -572,11 +572,25 @@ public:
 	{
 		int32 NumRemoved = 0;
 
-		for (const FCacheEntry* Entry : LookupSet)
+		for (auto It = LookupSet.CreateIterator(); It; ++It)
 		{
+			FCacheEntry* Entry = *It;
 			if (Pred(Entry->Key, Entry->Value))
 			{
-				Remove(const_cast<FCacheEntry*>(Entry));
+				if (Entry == LeastRecent)
+				{
+					LeastRecent = Entry->MoreRecent;
+				}
+
+				if (Entry == MostRecent)
+				{
+					MostRecent = Entry->LessRecent;
+				}
+
+				Entry->Unlink();
+				delete Entry;
+
+				It.RemoveCurrent();
 				++NumRemoved;
 			}
 		}
