Added UnifiedError and integrated with FStreamableManager

diff --git a/Engine/Source/Runtime/Core/Private/Experimental/IO/IoStatusError.cpp b/Engine/Source/Runtime/Core/Private/Experimental/IO/IoStatusError.cpp
new file mode 100644
--- /dev/null
+++ b/Engine/Source/Runtime/Core/Private/Experimental/IO/IoStatusError.cpp
@@ -0,0 +1,59 @@
+#include "Experimental/IO/IOStatusError.h"
+#include "IO/IoStatus.h"
+
+#include "Internationalization/Internationalization.h"
+#include "Internationalization/Text.h"
+
+
+namespace UE::UnifiedError::IoStore
+{
+	class FIoStoreErrorDetails : public UE::UnifiedError::FRefCountedErrorDetails
+	{
+	private:
+		FIoStatus::FErrorMessage ErrorMessage;
+		mutable FUtf8String CachedErrorName;
+		static const FText GenericErrorFormatString;
+	public:
+		FIoStoreErrorDetails(const FIoStatus::FErrorMessage& InErrorMessage)
+		{
+			FCString::Strncpy(ErrorMessage, InErrorMessage, FIoStatus::MaxErrorMessageLength);
+		}
+
+		virtual const FUtf8String& GetErrorName(const FError& Error) const final
+		{
+			if (CachedErrorName.IsEmpty() && (ErrorMessage[0] != TEXT('\0')))
+			{
+				CachedErrorName = TCHAR_TO_UTF8(GetIoErrorText(EIoErrorCode(Error.GetErrorCode())));
+			}
+			return CachedErrorName;
+		}
+		virtual const FText& GetErrorFormatString(const FError& Error) const final
+		{
+			return GenericErrorFormatString;
+		}
+
+		virtual void GetErrorProperties(const FError& Error, IErrorPropertyExtractor& OutProperties) const override
+		{
+			OutProperties.AddProperty(UTF8TEXTVIEW("IoStoreErrorMessage"), FStringView(ErrorMessage));
+		}
+
+	};
+
+	const FText FIoStoreErrorDetails::GenericErrorFormatString = NSLOCTEXT("IoStore", "GenericErrorMessage", "{IoStoreErrorMessage}");
+
+	// allow access to the private FError implementation
+	class FIoStoreError : public FError
+	{
+	public:
+
+		FIoStoreError(const EIoErrorCode InErrorCode, const IErrorDetails* ErrorDetails) : FError(UE::UnifiedError::IoStore::StaticModuleId, (int32)(InErrorCode), ErrorDetails)
+		{
+		}
+	};
+
+	CORE_API UE::UnifiedError::FError ConvertError(const FIoStatus& Status)
+	{
+		FIoStoreErrorDetails* ErrorDetails = new FIoStoreErrorDetails(Status.GetErrorMessage());
+		return FIoStoreError(Status.GetErrorCode(), ErrorDetails);
+	}
+}
\ No newline at end of file
diff --git a/Engine/Source/Runtime/Core/Private/Experimental/UnifiedError/UnifiedError.cpp b/Engine/Source/Runtime/Core/Private/Experimental/UnifiedError/UnifiedError.cpp
new file mode 100644
--- /dev/null
+++ b/Engine/Source/Runtime/Core/Private/Experimental/UnifiedError/UnifiedError.cpp
@@ -0,0 +1,381 @@
+#include "Experimental/UnifiedError/UnifiedError.h"
+#include "Internationalization/Text.h"
+#include "Internationalization/Internationalization.h"
+#include "Logging/StructuredLog.h"
+
+
+namespace UE::UnifiedError
+{
+
+class FTextFormatArgsPropertyExtractor : public IErrorPropertyExtractor
+{
+private:
+	FFormatNamedArguments& Arguments;
+public:
+	FTextFormatArgsPropertyExtractor(FFormatNamedArguments& InArguments) : Arguments(InArguments) 
+	{
+	}
+
+	virtual void AddProperty(const FUtf8StringView& PropertyName, const FStringView& PropertyValue) override
+	{
+		Arguments.Add(UTF8_TO_TCHAR(PropertyName.GetData()), FFormatArgumentValue(FText::FromString(PropertyValue.GetData())));
+	}
+
+	virtual void AddProperty(const FUtf8StringView& PropertyName, const FUtf8StringView& PropertyValue)  override
+	{
+		Arguments.Add(UTF8_TO_TCHAR(PropertyName.GetData()), FFormatArgumentValue(FText::FromStringView(UTF8_TO_TCHAR(PropertyValue.GetData()))));
+	}
+
+	virtual void AddProperty(const FUtf8StringView& PropertyName, const int64 PropertyValue) override
+	{
+		Arguments.Add(UTF8_TO_TCHAR(PropertyName.GetData()), FFormatArgumentValue(PropertyValue));
+	}
+
+	virtual void AddProperty(const FUtf8StringView& PropertyName, const int32 PropertyValue) override
+	{
+		Arguments.Add(UTF8_TO_TCHAR(PropertyName.GetData()), FFormatArgumentValue(PropertyValue));
+	}
+
+	virtual void AddProperty(const FUtf8StringView& PropertyName, const float PropertyValue) override
+	{
+		Arguments.Add(UTF8_TO_TCHAR(PropertyName.GetData()), FFormatArgumentValue(PropertyValue));
+	}
+
+	virtual void AddProperty(const FUtf8StringView& PropertyName, const double PropertyValue) override
+	{
+		Arguments.Add(UTF8_TO_TCHAR(PropertyName.GetData()), FFormatArgumentValue(PropertyValue));
+	}
+
+	virtual void AddProperty(const FUtf8StringView& PropertyName, const FText& PropertyValue) override
+	{
+		Arguments.Add(UTF8_TO_TCHAR(PropertyName.GetData()), FFormatArgumentValue(PropertyValue));
+	}
+};
+
+
+
+// FError functions
+
+void FError::PushErrorDetails(TRefCountPtr<IErrorDetails> InErrorDetails)
+{
+	/*if (ErrorDetails)
+	{
+		InErrorDetails->SetInnerErrorDetails(ErrorDetails);
+	}
+	InErrorDetails->AddRef();
+	ErrorDetails->ReleaseRef();*/
+	if (!InErrorDetails.IsValid())
+	{
+		return;
+	}
+	InErrorDetails->SetInnerErrorDetails(ErrorDetails);
+	ErrorDetails = InErrorDetails;
+}
+
+TRefCountPtr<const IErrorDetails> FError::GetInnerMostErrorDetails() const
+{
+	const IErrorDetails* Result = ErrorDetails;
+	while (true)
+	{
+		if (Result->GetInnerErrorDetails())
+		{
+			Result = Result->GetInnerErrorDetails();
+		}
+		else
+		{
+			return Result;
+		}
+	}
+}
+
+FText FError::GetErrorMessage() const
+{
+	FFormatNamedArguments Args;
+	FTextFormatArgsPropertyExtractor PropertyExtractor(Args);
+	ErrorDetails->GetErrorProperties(*this, PropertyExtractor);
+	return FText::Format(ErrorDetails->GetErrorFormatString(*this), Args);
+}
+
+int32 FError::GetErrorCode() const
+{
+	return ErrorCode;
+}
+
+int32 FError::GetModuleId() const
+{
+	return ModuleId;
+}
+
+class FNullErrorPropertyExtractor : public IErrorPropertyExtractor
+{
+public:
+	virtual void AddProperty(const FUtf8StringView& PropertyName, const FStringView& PropertyValue) override { } 
+	virtual void AddProperty(const FUtf8StringView& PropertyName, const FUtf8StringView& PropertyValue) override { }
+	virtual void AddProperty(const FUtf8StringView& PropertyName, const FText& PropertyValue) override { }
+	virtual void AddProperty(const FUtf8StringView& PropertyName, const int64 PropertyValue) override { }
+	virtual void AddProperty(const FUtf8StringView& PropertyName, const int32 PropertyValue) override { }
+	virtual void AddProperty(const FUtf8StringView& PropertyName, const float PropertyValue) override { }
+	virtual void AddProperty(const FUtf8StringView& PropertyName, const double PropertyValue) override { }
+};
+
+template<typename SearchType, typename ResultType>
+class FPropertySearchVisitor : public FNullErrorPropertyExtractor
+{
+public:
+	FPropertySearchVisitor(const FUtf8StringView& InSearchName, ResultType* InResult) : SearchName(InSearchName), Result(InResult)
+	{
+	}
+
+	virtual void AddProperty(const FUtf8StringView& PropertyName, SearchType PropertyValue) override
+	{
+		if ((PropertyName == SearchName) && (Result != nullptr))
+		{
+			*Result = PropertyValue;
+			bFound = true;
+		}
+	}
+
+	bool WasFound() const { return bFound; }
+private:
+
+	const FUtf8StringView& SearchName;
+	ResultType* Result;
+	mutable bool bFound = false;
+};
+
+bool FError::GetDetailByKey(const FUtf8StringView& KeyName, FString& Result) const
+{
+	FPropertySearchVisitor<const FStringView&, FString> Visitor(KeyName, &Result);
+	ErrorDetails->GetErrorProperties(*this, Visitor);
+	return Visitor.WasFound();
+}
+
+bool FError::GetDetailByKey(const FUtf8StringView& KeyName, FUtf8String& Result) const
+{
+	FPropertySearchVisitor<const FUtf8StringView&, FUtf8String> Visitor(KeyName, &Result);
+	ErrorDetails->GetErrorProperties(*this, Visitor);
+	return Visitor.WasFound();
+}
+
+bool FError::GetDetailByKey(const FUtf8StringView& KeyName, FText& Result) const
+{
+	FPropertySearchVisitor<const FText&, FText> Visitor(KeyName, &Result);
+	ErrorDetails->GetErrorProperties(*this, Visitor);
+	return Visitor.WasFound();
+}
+
+bool FError::GetDetailByKey(const FUtf8StringView& KeyName, int64& Result) const
+{
+	FPropertySearchVisitor<const int64, int64> Visitor(KeyName, &Result);
+	ErrorDetails->GetErrorProperties(*this, Visitor);
+	return Visitor.WasFound();
+}
+
+bool FError::GetDetailByKey(const FUtf8StringView& KeyName, int32& Result) const
+{
+	FPropertySearchVisitor<const int32, int32> Visitor(KeyName, &Result);
+	ErrorDetails->GetErrorProperties(*this, Visitor);
+	return Visitor.WasFound();
+}
+
+bool FError::GetDetailByKey(const FUtf8StringView& KeyName, double& Result) const
+{
+	FPropertySearchVisitor<const double, double> Visitor(KeyName, &Result);
+	ErrorDetails->GetErrorProperties(*this, Visitor);
+	return Visitor.WasFound();
+}
+
+bool FError::GetDetailByKey(const FUtf8StringView& KeyName, float& Result) const
+{
+	FPropertySearchVisitor<const float, float> Visitor(KeyName, &Result);
+	ErrorDetails->GetErrorProperties(*this, Visitor);
+	return Visitor.WasFound();
+}
+
+void FError::GetErrorProperties(IErrorPropertyExtractor& Visitor) const
+{
+	ErrorDetails->GetErrorProperties(*this, Visitor);
+}
+
+TRefCountPtr<const FStaticErrorDetails> FError::GetStaticErrorDetails() const
+{
+	TRefCountPtr<const IErrorDetails> InnerMostError = GetInnerMostErrorDetails();
+	check(InnerMostError.IsValid() && (InnerMostError->GetErrorDetailsTypeId() == FStaticErrorDetails::StaticGetErrorDetailsTypeId())); // innermost error should always be something static
+	TRefCountPtr<const FStaticErrorDetails> StaticErrorDetails((const FStaticErrorDetails*)InnerMostError.GetReference());
+	return StaticErrorDetails;
+}
+
+const FAnsiStringView& FError::GetErrorCodeString() const
+{
+	TRefCountPtr<const FStaticErrorDetails> StaticErrorDetails = GetStaticErrorDetails();
+	return StaticErrorDetails->GetErrorCodeString();
+}
+
+const FAnsiStringView& FError::GetModuleIdString() const
+{
+	TRefCountPtr<const FStaticErrorDetails> StaticErrorDetails = GetStaticErrorDetails();
+	return StaticErrorDetails->GetModuleIdString();
+}
+
+// Structured logging FError integration
+class FCbWriterErrorPropertyExtractor : public IErrorPropertyExtractor
+{
+private:
+	FCbWriter& Writer;
+public:
+	FCbWriterErrorPropertyExtractor(FCbWriter& InWriter) : Writer(InWriter) {}
+
+
+
+public:
+	virtual void AddProperty(const FUtf8StringView& PropertyName, const FStringView& PropertyValue) override 
+	{ 
+		Writer.AddString(PropertyName, PropertyValue);
+	}
+	virtual void AddProperty(const FUtf8StringView& PropertyName, const FUtf8StringView& PropertyValue) override
+	{
+		Writer.AddString(PropertyName, PropertyValue);
+	}
+	virtual void AddProperty(const FUtf8StringView& PropertyName, const FText& PropertyValue) override
+	{
+		Writer.AddString(PropertyName, PropertyValue.ToString());
+	}
+	virtual void AddProperty(const FUtf8StringView& PropertyName, const int64 PropertyValue) override
+	{
+		Writer.AddInteger(PropertyName, PropertyValue);
+	}
+	virtual void AddProperty(const FUtf8StringView& PropertyName, const int32 PropertyValue) override 
+	{
+		Writer.AddInteger(PropertyName, PropertyValue);
+	}
+	virtual void AddProperty(const FUtf8StringView& PropertyName, const float PropertyValue) override
+	{
+		Writer.AddFloat(PropertyName, PropertyValue);
+	}
+	virtual void AddProperty(const FUtf8StringView& PropertyName, const double PropertyValue) override 
+	{ 
+		Writer.AddFloat(PropertyName, PropertyValue);
+	}
+};
+
+void SerializeForLog(FCbWriter& Writer, const FError& Error)
+{
+	Writer.BeginObject();
+	Writer.AddString(ANSITEXTVIEW("$type"), ANSITEXTVIEW("FError"));
+
+
+	FCbWriterErrorPropertyExtractor ExtractToCbWriter(Writer);
+	Error.GetErrorProperties(ExtractToCbWriter);
+
+
+	// get the error message this error wants to produce
+	FText ErrorMessageText = Error.GetErrorMessage();
+	Writer.AddString(ANSITEXTVIEW("ErrorMessage"), ErrorMessageText.ToString());
+
+	// build up a formatted error string including some additional information about the error
+	TRefCountPtr<const FStaticErrorDetails> StaticErrorDetails = Error.GetStaticErrorDetails();
+	FFormatNamedArguments ErrorMessageArgs;
+	FTextFormatArgsPropertyExtractor ExtractToFTextArgs(ErrorMessageArgs);
+	StaticErrorDetails->GetErrorProperties(Error, ExtractToFTextArgs);;
+	ErrorMessageArgs.Add(TEXT("ErrorMessage"), ErrorMessageText);
+
+	
+	Writer.AddString(ANSITEXTVIEW("ErrorFormatString"), StaticErrorDetails->GetErrorFormatString(Error).ToString());
+	Writer.AddString("$text", FText::Format(NSLOCTEXT("UnifiedError", "StructuredLogErrorMessage", "{ModuleIdString}.{ErrorCodeString}({ModuleId}.{ErrorCode}): {ErrorMessage}"), ErrorMessageArgs).ToString());
+
+	Writer.EndObject();
+}
+
+// FRefCountedErrorDetails fucntions
+FRefCountedErrorDetails::~FRefCountedErrorDetails()
+{
+}
+
+// FDynamicErrorDetails functions
+
+FDynamicErrorDetails::FDynamicErrorDetails(TRefCountPtr<const IErrorDetails> InInnerErrorDetails)
+{
+	InnerErrorDetails = InInnerErrorDetails;
+}
+FDynamicErrorDetails::~FDynamicErrorDetails() 
+{
+}
+
+const FText& FDynamicErrorDetails::GetErrorFormatString(const FError & Error) const
+{
+	check(InnerErrorDetails);
+	return InnerErrorDetails->GetErrorFormatString(Error);
+}
+/*
+
+const FUtf8String& FDynamicErrorDetails::GetErrorName(const FError& Error) const
+{
+	check(InnerErrorDetails);
+	return InnerErrorDetails->GetErrorName(Error);
+}
+*/
+
+
+void FDynamicErrorDetails::GetErrorProperties(const FError& Error, IErrorPropertyExtractor& OutProperties) const
+{
+	InnerErrorDetails->GetErrorProperties(Error, OutProperties);
+}
+
+void FDynamicErrorDetails::SetInnerErrorDetails(TRefCountPtr<const IErrorDetails> InInnerErrorDetails)
+{
+	InnerErrorDetails = InInnerErrorDetails;
+}
+
+// FStaticErrorDetails functions
+
+FStaticErrorDetails::FStaticErrorDetails(const FAnsiStringView& InErrorName, const FAnsiStringView& InModuleName, const FText& InErrorFormatString) : ErrorName(InErrorName), ModuleName(InModuleName), ErrorFormatString(InErrorFormatString)
+{
+}
+
+
+/*
+FStaticErrorDetails::FStaticErrorDetails(const FUtf8StringView& InErrorName, const FText& InErrorFormatString)
+{
+	ErrorName = InErrorName.GetData();
+	ErrorFormatString = InErrorFormatString;
+}*/
+
+void FStaticErrorDetails::GetErrorProperties(const FError& Error, IErrorPropertyExtractor& OutProperties) const
+{
+	OutProperties.AddProperty(UTF8TEXTVIEW("ErrorCodeString"), ErrorName);
+	OutProperties.AddProperty(UTF8TEXTVIEW("ModuleIdString"), ModuleName);
+	OutProperties.AddProperty(UTF8TEXTVIEW("ErrorCode"), Error.GetErrorCode());
+	OutProperties.AddProperty(UTF8TEXTVIEW("ModuleId"), Error.GetModuleId());
+}
+
+const FAnsiStringView& FStaticErrorDetails::GetErrorCodeString() const
+{
+	return ErrorName;
+}
+
+const FAnsiStringView& FStaticErrorDetails::GetModuleIdString() const
+{
+	return ModuleName;
+}
+
+const FText& FStaticErrorDetails::GetErrorFormatString(const FError& Error) const
+{
+	return ErrorFormatString;
+}
+
+// FErrorDetailsRegistry
+uint32 FErrorDetailsRegistry::RegisterDetails(const FAnsiStringView& ErrorDetailsName, TFunction<IErrorDetails*()> CreationFunction)
+{
+	uint32 DetailsId = FCrc::StrCrc32<FAnsiStringView::ElementType>(ErrorDetailsName.GetData());
+	CreateFunctions.Add(DetailsId, CreationFunction);
+	return DetailsId;
+}
+
+
+} // namespace UE::Error
+
+
+FString LexToString(const UE::UnifiedError::FError& Error)
+{
+	return Error.GetErrorMessage().ToString();
+}
\ No newline at end of file
diff --git a/Engine/Source/Runtime/Core/Private/IO/IoStatus.cpp b/Engine/Source/Runtime/Core/Private/IO/IoStatus.cpp
--- a/Engine/Source/Runtime/Core/Private/IO/IoStatus.cpp
+++ b/Engine/Source/Runtime/Core/Private/IO/IoStatus.cpp
@@ -6,6 +6,8 @@
 #include "HAL/UnrealMemory.h"
 #include "Logging/LogCategory.h"
 #include "Logging/LogMacros.h"
+#include "Experimental/UnifiedError/UnifiedError.h"
+#include "Experimental/IO/IoStatusError.h"
 #include "Math/UnrealMathUtility.h"
 
 DEFINE_LOG_CATEGORY_STATIC(LogIoStatus, Log, All);
@@ -110,3 +112,4 @@ FIoStatusBuilder operator<<(const FIoStatus& Status, FStringView String)
 { 
 	return FIoStatusBuilder(Status, String);
 }
+
diff --git a/Engine/Source/Runtime/Core/Public/Experimental/IO/IOStatusError.h b/Engine/Source/Runtime/Core/Public/Experimental/IO/IOStatusError.h
new file mode 100644
--- /dev/null
+++ b/Engine/Source/Runtime/Core/Public/Experimental/IO/IOStatusError.h
@@ -0,0 +1,15 @@
+// Copyright Epic Games, Inc. All Rights Reserved.
+#pragma once
+
+#include "Experimental/UnifiedError/UnifiedError.h"
+
+
+class FIoStatus;
+
+DECLARE_ERROR_MODULE(IoStore, 0x15);
+
+
+namespace UE::UnifiedError::IoStore
+{
+	CORE_API UE::UnifiedError::FError ConvertError(const FIoStatus& Status);
+}
\ No newline at end of file
diff --git a/Engine/Source/Runtime/Core/Public/Experimental/UnifiedError/UnifiedError.h b/Engine/Source/Runtime/Core/Public/Experimental/UnifiedError/UnifiedError.h
new file mode 100644
--- /dev/null
+++ b/Engine/Source/Runtime/Core/Public/Experimental/UnifiedError/UnifiedError.h
@@ -0,0 +1,501 @@
+// Copyright Epic Games, Inc. All Rights Reserved.
+#pragma once
+
+#include "CoreTypes.h"
+#include "Containers/Utf8String.h"
+#include "Containers/StringView.h"
+#include "Internationalization/Internationalization.h"
+#include "Internationalization/Text.h"
+#include "Templates/RefCounting.h"
+
+#define DECLARE_ERROR_DETAILS_INTERNAL(DetailsNamespace, TypeName)\
+	static uint32 StaticGetErrorDetailsTypeId() \
+	{ \
+		return StaticDetailsTypeId;\
+	}\
+	virtual uint32 GetErrorDetailsTypeId() const \
+	{ \
+		return StaticGetErrorDetailsTypeId(); \
+	}
+
+
+
+#define DECLARE_FERROR_DETAILS_ABSTRACT(DetailsNamespace, TypeName)\
+	public:\
+	inline static const uint32 StaticDetailsTypeId = UE::UnifiedError::FErrorDetailsRegistry::Get().RegisterDetails(ANSITEXTVIEW(#DetailsNamespace "::" #TypeName), nullptr);\
+	DECLARE_ERROR_DETAILS_INTERNAL(DetailsNamespace, TypeName);
+	
+
+#define DECLARE_FERROR_DETAILS(DetailsNamespace, TypeName)\
+	public:\
+	friend IErrorDetails* Create();\
+	inline static IErrorDetails* Create()\
+	{\
+		return new TypeName();\
+	}\
+	inline static const uint32 StaticDetailsTypeId = UE::UnifiedError::FErrorDetailsRegistry::Get().RegisterDetails(ANSITEXTVIEW(#DetailsNamespace "::" #TypeName), TFunction<IErrorDetails* ()>([]() -> IErrorDetails* { return DetailsNamespace::TypeName::Create(); }));\
+	DECLARE_ERROR_DETAILS_INTERNAL(DetailsNamespace, TypeName);
+
+
+
+namespace UE::UnifiedError
+{
+	
+	namespace DetailsTypes
+	{
+		const int32 IERROR_DETAILS_TYPE = 1;
+		const int32 STATIC_ERROR_DETAILS_TYPE = 2;
+		const int32 DYNAMIC_ERROR_DETAILS_TYPE = 3;
+		const int32 FIRST_CUSTOM_DETAILS_TYPE = 4;
+	}
+
+	/// <summary>
+	/// IErrorPropertyExtractor is an interface used to visit properties exposed by IErrorDetails implementations
+	/// See also IErrorDetails::GetErrorDetails
+	/// </summary>
+	class IErrorPropertyExtractor
+	{
+	public:
+		virtual void AddProperty(const FUtf8StringView& PropertyName, const FStringView& PropertyValue) = 0;
+		virtual void AddProperty(const FUtf8StringView& PropertyName, const FUtf8StringView& PropertyValue) = 0;
+		virtual void AddProperty(const FUtf8StringView& PropertyName, const FText& PropertyValue) = 0;
+		virtual void AddProperty(const FUtf8StringView& PropertyName, const int64 PropertyValue) = 0;
+		virtual void AddProperty(const FUtf8StringView& PropertyName, const int32 PropertyValue) = 0;
+		virtual void AddProperty(const FUtf8StringView& PropertyName, const float PropertyValue) = 0;
+		virtual void AddProperty(const FUtf8StringView& PropertyName, const double PropertyValue) = 0;
+	};
+
+	class FError;
+
+	class IErrorDetails : public IRefCountedObject
+	{
+	public:
+		CORE_API virtual ~IErrorDetails() = default;
+
+		/// <summary>
+		/// GetErrorFormatString; specifies the default error format string to be used when generating FError::GetErrorMessage.
+		///	 The format string can specify any property exposed by any encapsulated IErrorDetails::GetErrorProperties.
+		///  Example: GetErrorProperties adds Name:"ModuleId" Value:10. GetErrorFormatString returns "Module id was {ModuleId}".  Result "Module id was 10".
+		/// </summary>
+		/// <param name="Error"></param>
+		/// <returns></returns>
+		CORE_API virtual const FText& GetErrorFormatString(const FError& Error) const = 0;
+
+		/// <summary>
+		/// GetErrorProperties; used to expose error properties in name:value format, Error properties can be used for string formating functions, searching, exposing context to higher level stacks
+		///	 See also: FError::GetErrorMessage, IErrorDetails::GetErrorFormatString, FError::GetDetailByKey
+		/// </summary>
+		/// <param name="Error"></param>
+		/// <param name="OutProperties"></param>
+		CORE_API virtual void GetErrorProperties(const FError& Error, IErrorPropertyExtractor& OutProperties) const = 0;
+
+		/// <summary>
+		/// GetInnerErrorDetails; Exposes inner error details to FError, if this ErrorDetails allows inner details
+		/// </summary>
+		/// <returns></returns>
+		CORE_API virtual TRefCountPtr<const IErrorDetails> GetInnerErrorDetails() const { return nullptr; }
+
+		/// <summary>
+		/// SetInnerErrorDetails; Exposes inner error details to FError, if this ErrorDetails allows inner details
+		/// </summary>
+		/// <returns></returns>
+		CORE_API virtual void SetInnerErrorDetails(TRefCountPtr<const IErrorDetails> ErrorDetails) { checkf(false, TEXT("SetInnerErrorDetails not implemented!")); }
+
+		/// <summary>
+		/// GetErrorDetialsTypeId; Simple type information for error details, generated using hash of details name
+		///  See also: #define FERROR_DETAILS
+		/// </summary>
+		/// <returns></returns>
+		virtual uint32 GetErrorDetailsTypeId() const = 0;
+
+	};
+
+
+	class FErrorDetailsRegistry
+	{
+	private:
+		TMap<uint32, TFunction<IErrorDetails* ()>> CreateFunctions;
+
+		FErrorDetailsRegistry() { }
+	public:
+		static FErrorDetailsRegistry& Get()
+		{
+			static FErrorDetailsRegistry Registry;
+			return Registry;
+		}
+
+		CORE_API uint32 RegisterDetails(const FAnsiStringView& ErrorDetailsName, TFunction<IErrorDetails*()> CreationFunction);
+	};
+
+
+	/// <summary>
+	/// FRefCountedErrorDetails; base implementation of refcounting for IErrorDetails, this is used for heap allocated IErrorDetails implementations
+	/// </summary>
+	class FRefCountedErrorDetails : public IErrorDetails, public FRefCountBase
+	{
+		DECLARE_FERROR_DETAILS_ABSTRACT(UnifiedError, FRefCountedErrorDetails);
+
+
+		CORE_API virtual ~FRefCountedErrorDetails();
+
+		virtual FReturnedRefCountValue AddRef() const final { return FRefCountBase::AddRef(); }
+		virtual uint32 Release() const final { return FRefCountBase::Release(); }
+		virtual uint32 GetRefCount() const final { return FRefCountBase::GetRefCount(); }
+
+		/*using FRefCountBase::AddRef;
+		using FRefCountBase::Release;
+		using FRefCountBase::GetRefCount;*/
+	};
+
+
+	/// <summary>
+	/// FDynamicErrorDetails; base implementation of inner error details, for use by derived classes to reduce unnessisary reimplementation 
+	/// </summary>
+	class FDynamicErrorDetails : public FRefCountedErrorDetails
+	{
+		DECLARE_FERROR_DETAILS(UnifiedError, FDynamicErrorDetails);
+	private:
+		TRefCountPtr<const IErrorDetails> InnerErrorDetails;
+	public:
+		CORE_API FDynamicErrorDetails(TRefCountPtr<const IErrorDetails> InInnerErrorDetails = nullptr);
+		CORE_API virtual ~FDynamicErrorDetails();
+
+
+		// IErrorDetails functions
+		CORE_API virtual void SetInnerErrorDetails(TRefCountPtr<const IErrorDetails> InInnerErrorDetails);
+		CORE_API virtual TRefCountPtr<const IErrorDetails> GetInnerErrorDetails() const { return InnerErrorDetails; }
+
+		/// <summary>
+		/// GetErrorFormatString; Pass through to the InnerErrorDetails.
+		/// </summary>
+		/// <param name="Error"></param>
+		/// <returns></returns>
+		CORE_API virtual const FText& GetErrorFormatString(const FError& Error) const override;
+
+		/// <summary>
+		/// GetErrorProperties; by default pass through to the InnerErrorDetails.
+		///  It's expected FDynamicErrorDetails will be derived from and implement GetErrorProperties. 
+		/// </summary>
+		/// <param name="Error"></param>
+		/// <param name="OutProperties"></param>
+		CORE_API virtual void GetErrorProperties(const FError& Error, IErrorPropertyExtractor& OutProperties) const override;
+	};
+
+
+	/// <summary>
+	/// FDefaultErrorDetails; serializable FStaticErrorDetails standin. 
+	///  This class is primarily used as a stand in for FStaticErrorDetails in cases of serialization over network or to disk.
+	/// </summary>
+	class FDefaultErrorDetails
+	{
+	private:
+		FAnsiStringView ErrorName;
+		FAnsiStringView ModuleName;
+		FText ErrorFormatString;
+	public:
+		FDefaultErrorDetails() = default;
+
+	};
+
+	/// <summary>
+	/// FStaticErrorDetails; static error details and members are statically allocated
+	///  Every error which uses DEFINE_ERROR will have FStaticErrorDetails generated for it
+	///  Can not rely on it to be available for every error as some Error conversion functions will not use pregenerated errors or error codesF
+	///  Use FError::GetErrorDetails to discover FStaticErrorDetails
+	/// </summary>
+	class FStaticErrorDetails : public IErrorDetails
+	{
+		DECLARE_FERROR_DETAILS_ABSTRACT(UnifiedError,FStaticErrorDetails);
+	private:
+		const FAnsiStringView& ErrorName;
+		const FAnsiStringView& ModuleName;
+		FText ErrorFormatString;
+	public:
+		// TODO: convert AnsiStringView to UTF8StringView when c++ 20 is supported
+		CORE_API FStaticErrorDetails(const FAnsiStringView& InErrorName, const FAnsiStringView& InModuleName, const FText& InErrorFormatString);
+
+		CORE_API virtual ~FStaticErrorDetails() = default;
+
+		/// <summary>
+		/// GetErrorCodeString; Accessor for ErrorName.  
+		///  Can be called directly on FStaticErrorDetails object.  
+		///  See also: FError::GetErrorDetails
+		/// </summary>
+		/// <returns></returns>
+		CORE_API const FAnsiStringView& GetErrorCodeString() const;
+
+		/// <summary>
+		/// GetModuleIdString; accessor for ModuleName.
+		///  Can be called directly on FStaticErrorDetails object
+		///  See also: FError::GetErrorDetails
+		/// </summary>
+		CORE_API const FAnsiStringView& GetModuleIdString() const;
+
+		// IErrorDetails implementation
+
+		/// <summary>
+		/// GetErrorProperties;  Adds error properties exposed by this error details object to the OutProperties value
+		///  FStaticErrorDetails includes error details all errors have ErrorCode, ModuleId, ErrorCodeString, ModuleIdString
+		///  These can be used in any error format string see also GetErrorFormatString
+		/// </summary>
+		/// <param name="Error"></param>
+		/// <param name="OutProperties"></param>
+		CORE_API virtual void GetErrorProperties(const FError& Error, IErrorPropertyExtractor& OutProperties) const;
+
+		/// <summary>
+		/// GetErrorFormatString; return the localized format text generated in DECLARE_ERROR macro.
+		/// </summary>
+		/// <param name="Error"></param>
+		/// <returns></returns>
+		CORE_API virtual const FText& GetErrorFormatString(const FError& Error) const final;
+
+
+
+
+		// IRefCountedObject implementation
+		//  FStaticErrorDetails is statically allocated, make sure it is never released.
+		virtual FReturnedRefCountValue AddRef() const final { return FReturnedRefCountValue(10); }
+		virtual uint32 Release() const final { return 10; }
+		virtual uint32 GetRefCount() const final { return 10; }
+	};
+
+	class FError
+	{
+	private: 
+		int32 ErrorCode;
+		int32 ModuleId;
+		TRefCountPtr<const IErrorDetails> ErrorDetails;
+	public:
+		CORE_API FError(int32 InModuleId, int32 InErrorCode, const IErrorDetails* InErrorDetails)
+		{
+			ErrorCode = InErrorCode;
+			ModuleId = InModuleId;
+			ErrorDetails = InErrorDetails;
+		}
+	public:
+		CORE_API FError(FError&& InError)
+		{
+			ErrorCode = InError.ErrorCode;
+			InError.ErrorCode = 0;
+			ModuleId = InError.ModuleId;
+			InError.ModuleId = 0;
+			if (InError.ErrorDetails)
+			{
+				ErrorDetails = InError.ErrorDetails;
+				InError.ErrorDetails = nullptr;
+			}
+		}
+
+		CORE_API FError(const FError& InError)
+		{
+			ErrorCode = InError.ErrorCode;
+			ModuleId = InError.ModuleId;
+			ErrorDetails = InError.ErrorDetails;
+		}
+
+		CORE_API bool operator==(const FError& OtherError)
+		{
+			if ((ModuleId == OtherError.ModuleId) && 
+				(ErrorCode == OtherError.ErrorCode))
+			{
+				return true;
+			}
+			return false;
+		}
+
+		~FError()
+		{
+			Invalidate();
+		}
+
+		CORE_API void GetErrorProperties(IErrorPropertyExtractor& Visitor) const;
+
+		CORE_API int32 GetErrorCode() const;
+		CORE_API int32 GetModuleId() const;
+
+		CORE_API const FAnsiStringView& GetErrorCodeString() const;
+		CORE_API const FAnsiStringView& GetModuleIdString() const;
+
+		CORE_API void PushErrorDetails(TRefCountPtr<IErrorDetails> InErrorDetails);
+
+		CORE_API FText GetErrorMessage() const;
+
+		template<typename DetailType>
+		TRefCountPtr<const DetailType> GetErrorDetails() const
+		{
+			TRefCountPtr<const IErrorDetails> CurrentIt = ErrorDetails;
+			while (CurrentIt != nullptr)
+			{
+				if (CurrentIt->GetErrorDetailsTypeId() == DetailType::StaticGetErrorDetailsTypeId())
+				{
+					return TRefCountPtr<const DetailType>((const DetailType*)CurrentIt.GetReference());
+				}
+				CurrentIt = CurrentIt->GetInnerErrorDetails();
+			}
+			return nullptr;
+		}
+
+		CORE_API TRefCountPtr<const FStaticErrorDetails> GetStaticErrorDetails() const;
+
+		CORE_API bool GetDetailByKey(const FUtf8StringView& KeyName, FString& Result) const;
+		CORE_API bool GetDetailByKey(const FUtf8StringView& KeyName, FUtf8String& Result) const;
+		CORE_API bool GetDetailByKey(const FUtf8StringView& KeyName, FText& Result) const;
+		CORE_API bool GetDetailByKey(const FUtf8StringView& KeyName, int64& Result) const;
+		CORE_API bool GetDetailByKey(const FUtf8StringView& KeyName, int32& Result) const;
+		CORE_API bool GetDetailByKey(const FUtf8StringView& KeyName, double& Result) const;
+		CORE_API bool GetDetailByKey(const FUtf8StringView& KeyName, float& Result) const;
+
+		FORCEINLINE bool IsValid() const
+		{
+			return (ErrorCode != 0) && (ModuleId != 0);
+		}
+
+		FORCEINLINE void Invalidate()
+		{
+			ErrorCode = 0;
+			ModuleId = 0;
+			ErrorDetails = nullptr;/*
+			if (ErrorDetails)
+			{
+				ErrorDetails->ReleaseRef();
+				ErrorDetails = nullptr;
+			}*/
+		}
+	private:
+
+		TRefCountPtr<const IErrorDetails> GetInnerMostErrorDetails() const;
+	};
+
+	class FErrorRegistry
+	{
+	public:
+		static FErrorRegistry& Get()
+		{
+			static FErrorRegistry StaticRegistry;
+			return StaticRegistry;
+		}
+
+
+	private:
+
+		FErrorRegistry() {}
+
+	public:
+		uint32 RegisterModule(const FStringView& ModuleName, uint32 ModuleId)
+		{
+			checkf(ModuleNameMap.Contains(ModuleId) == false, TEXT("Module %s and %s are trying to register under module id %d"), ModuleName.GetData(), *ModuleNameMap.FindRef(ModuleId), ModuleId);
+			ModuleNameMap.Add(ModuleId, ModuleName.GetData());
+			return ModuleId;
+		}
+
+		void RegisterErrorCode(const FStringView& ErrorName, uint32 ModuleId, uint32 ErrorCode)
+		{
+			int64 CombinedErrorCode = ((int64)ModuleId) << 32 | ErrorCode;
+			checkf(ErrorCodeNameMap.Contains(CombinedErrorCode) == false, TEXT("Error %s and %s are trying to register under same error code %d"), ErrorName.GetData(), *ErrorCodeNameMap.FindRef(CombinedErrorCode), CombinedErrorCode);
+			ErrorCodeNameMap.Add(CombinedErrorCode, ErrorName.GetData());
+		}
+
+	private:
+		TMap<uint32, FString> ModuleNameMap;
+		TMap<uint64, FString> ErrorCodeNameMap;
+	};
+
+
+
+
+	CORE_API void SerializeForLog(FCbWriter& Writer, const FError& OnlineError);
+
+} // namespace UnifiedError
+
+
+CORE_API FString LexToString(const UE::UnifiedError::FError& Error);
+/*
+
+CORE_API constexpr uint32 GetTypeIdHash(const TCHAR* TypeName)
+{
+	return FCrc::StrCrc32(GetData(S));
+}
+*/
+
+
+
+
+#define DECLARE_ERROR_MODULE(ModuleName, ModuleId) \
+	namespace UE::UnifiedError { namespace ModuleName {\
+		static constexpr uint32 StaticModuleId = ModuleId;\
+		static constexpr FAnsiStringView StaticModuleName = ANSITEXTVIEW(#ModuleName);\
+		FORCEINLINE int32 GetNextErrorCode()  \
+		{\
+			static int32 ErrorCodeCounter = 0;\
+			++ErrorCodeCounter;\
+			return ErrorCodeCounter;\
+		}\
+	}}
+
+
+#define DECLARE_ERROR_INTERNAL(ErrorName, ErrorCode, ModuleName, FormatString) \
+	namespace UE::UnifiedError { namespace ModuleName { namespace ErrorName {\
+			static constexpr FAnsiStringView StaticErrorName = ANSITEXTVIEW(#ErrorName);\
+			FORCEINLINE const FStaticErrorDetails& GetStaticErrorDetails() \
+			{\
+				static FStaticErrorDetails StaticErrorDetails(StaticErrorName, UE::UnifiedError::ModuleName::StaticModuleName, FormatString);\
+				return StaticErrorDetails;\
+			}\
+			FORCEINLINE int32 GetErrorCodeId() \
+			{ \
+				static int32 ErrorCodeId = ErrorCode < 0 ? UE::UnifiedError::ModuleName::GetNextErrorCode() : ErrorCode; \
+				return ErrorCodeId; \
+			} \
+			static const int32 ErrorCodeId = GetErrorCodeId(); \
+	}}}
+
+
+#define DECLARE_ERROR(ErrorName, ErrorCode, ModuleName, FormatString) \
+	DECLARE_ERROR_INTERNAL(ErrorName, ErrorCode, ModuleName, FormatString) \
+	namespace UE::UnifiedError { namespace ModuleName { namespace ErrorName {\
+			FORCEINLINE FError MakeError() { return FError(UE::UnifiedError::ModuleName::StaticModuleId, UE::UnifiedError::ModuleName::ErrorName::ErrorCodeId, &GetStaticErrorDetails()); }\
+			FORCEINLINE const FError& GetStaticError() \
+			{\
+				static FError StaticError = MakeError();\
+				return StaticError;\
+			}\
+	}}}
+
+#define DEFINE_ERROR(ErrorName, ModuleName)
+
+
+#define DECLARE_ERROR_ONEPARAM(ErrorName, ErrorCode, ModuleName, FormatString, ParamOneType, ParamOneName, ParamOneDefault) \
+	DECLARE_ERROR_INTERNAL(ErrorName, ErrorCode, ModuleName, FormatString) \
+	namespace UE::UnifiedError { namespace ModuleName { namespace ErrorName { \
+			class FDetails : public FDynamicErrorDetails \
+			{ \
+				DECLARE_FERROR_DETAILS(ModuleName::ErrorName, FDetails)\
+			private:\
+				ParamOneType ParamOneName; \
+				FDetails() : ParamOneName(ParamOneDefault) {}\
+			public: \
+				FDetails(const ParamOneType& In##ParamOneName, const IErrorDetails* InErrorDetails) : FDynamicErrorDetails(InErrorDetails) \
+				{ \
+					ParamOneName = In##ParamOneName; \
+				} \
+				virtual void GetErrorProperties(const FError& Error, IErrorPropertyExtractor& OutProperties) const override \
+				{ \
+					OutProperties.AddProperty(UTF8TEXTVIEW(#ParamOneName), ParamOneName); \
+					FDynamicErrorDetails::GetErrorProperties(Error, OutProperties); \
+				} \
+				const ParamOneType& Get##ParamOneName() const \
+				{ \
+					return ParamOneName; \
+				} \
+			}; \
+			FORCEINLINE FError MakeError(ParamOneType ParamOneName = ParamOneDefault) \
+			{ \
+				FDetails* NewDetails = new FDetails(ParamOneName, &GetStaticErrorDetails()); \
+				return FError(UE::UnifiedError::ModuleName::StaticModuleId, GetErrorCodeId(), NewDetails); \
+			} \
+			FORCEINLINE const FError& GetStaticError() \
+			{\
+				static FError StaticError = MakeError();\
+				return StaticError;\
+			}\
+	}}}
diff --git a/Engine/Source/Runtime/Core/Public/IO/IoStatus.h b/Engine/Source/Runtime/Core/Public/IO/IoStatus.h
--- a/Engine/Source/Runtime/Core/Public/IO/IoStatus.h
+++ b/Engine/Source/Runtime/Core/Public/IO/IoStatus.h
@@ -4,9 +4,11 @@
 
 #include "Containers/StringView.h"
 #include "Containers/UnrealString.h"
+#include "Experimental/UnifiedError/UnifiedError.h"
 
 template <typename CharType> class TStringBuilderBase;
 
+
 ///////////////////////////////////////////////////////////////////////////////
 
 /**
@@ -70,10 +72,11 @@ public:
 	CORE_API static const FIoStatus Unknown;
 	CORE_API static const FIoStatus Invalid;
 
-private:
 	static constexpr int32 MaxErrorMessageLength = 128;
 	using FErrorMessage = TCHAR[MaxErrorMessageLength];
 
+	CORE_API const FErrorMessage& GetErrorMessage() const { return ErrorMessage; }
+private:
 	EIoErrorCode	ErrorCode = EIoErrorCode::Ok;
 	FErrorMessage	ErrorMessage;
 
diff --git a/Engine/Source/Runtime/Engine/Classes/Engine/Experimental/StreamableManagerError.h b/Engine/Source/Runtime/Engine/Classes/Engine/Experimental/StreamableManagerError.h
new file mode 100644
--- /dev/null
+++ b/Engine/Source/Runtime/Engine/Classes/Engine/Experimental/StreamableManagerError.h
@@ -0,0 +1,55 @@
+// Copyright Epic Games, Inc. All Rights Reserved.
+#pragma once
+
+#include "Experimental/UnifiedError/UnifiedError.h"
+#include "UObject/UObjectGlobals.h"
+
+
+DECLARE_ERROR_MODULE(StreamableManager, 0x10);
+
+
+DECLARE_ERROR_ONEPARAM(PackageLoadFailed, 1, StreamableManager, NSLOCTEXT("StreamableManager", "PackageLoadError", "Failed to load package {PackageName}"), FString, PackageName, TEXT("Unknown"));
+
+DECLARE_ERROR(PackageLoadCancelled, 2, StreamableManager, NSLOCTEXT("StreamableManager", "PackageLoadCancelled", "Async load canceled"));
+DECLARE_ERROR(DownloadError, 3, StreamableManager, NSLOCTEXT("StreamableManager", "DownloadError", "Failed to download"));
+DECLARE_ERROR_ONEPARAM(PackageNameInvalid, 4, StreamableManager, NSLOCTEXT("StreamableManager", "PackageNameInvalid", "Found invalid package name {InvalidPackageName}"), FString, InvalidPackageName, TEXT("Unknown"));
+
+
+DECLARE_ERROR_ONEPARAM(MissingDownloadParams, 5, StreamableManager, NSLOCTEXT("StreamableManager", "MissingDownloadParams", "Missing download params for primary asset {PrimaryAssetName}"), FString, PrimaryAssetName, TEXT("Unknown"));
+DECLARE_ERROR(IoStoreNotFound, 6, StreamableManager, NSLOCTEXT("StreamableManager", "IoStoreNotFound", "IoStore did not load correctly."));
+DECLARE_ERROR_ONEPARAM(SyncLoadIncomplete, 7, StreamableManager, NSLOCTEXT("StreamableManager", "SyncLoadIncomplete", "Sync load did not complete correctly for {DebugName}."), FString, DebugName, TEXT("Unknown"));
+
+DECLARE_ERROR(AsyncLoadFailed, 8, StreamableManager, NSLOCTEXT("StreamableManager", "AsyncLoadFailed", "Async load failed."));
+DECLARE_ERROR(AsyncLoadCancelled, 9, StreamableManager, NSLOCTEXT("StreamableManager", "AsyncLoadCancelled", "Async load cancelled."));
+DECLARE_ERROR_ONEPARAM(AsyncLoadUnknownError, 10, StreamableManager, NSLOCTEXT("StreamableManager", "AsyncLoadUnknownError", "Unknown async loading error {AsyncLoadingErrorId}."), int32, AsyncLoadingErrorId, -1);
+
+
+DECLARE_ERROR(UnknownError, 11, StreamableManager, NSLOCTEXT("StreamableManager", "UnknownError", "Unknown error occured while streaming asset"));
+
+
+namespace UE::UnifiedError::StreamableManager
+{
+	class FStreamableManagerErrorDetails : public FDynamicErrorDetails
+	{
+		DECLARE_FERROR_DETAILS(StreamableManager, FStreamableManagerErrorDetails);
+	private:
+		FString SoftObjectPath;
+		FStreamableManagerErrorDetails() = default;
+	public:
+
+		FStreamableManagerErrorDetails(const FString& InSoftObjectPath) : SoftObjectPath(InSoftObjectPath)
+		{}
+
+		virtual void GetErrorProperties(const FError& Error, IErrorPropertyExtractor& OutProperties) const override
+		{
+			OutProperties.AddProperty(UTF8TEXTVIEW("SoftObjectPath"), SoftObjectPath);
+			FDynamicErrorDetails::GetErrorProperties(Error, OutProperties);
+		}
+
+		const FString& GetSoftObjectPath() const { return SoftObjectPath; }
+	};
+
+	UE::UnifiedError::FError GetStreamableError(EAsyncLoadingResult::Type Result);
+}
+
+//DECLARE_DYNAMIC_ERROR_CONTEXT(StreamableManagerContext, )
diff --git a/Engine/Source/Runtime/Engine/Classes/Engine/StreamableManager.h b/Engine/Source/Runtime/Engine/Classes/Engine/StreamableManager.h
--- a/Engine/Source/Runtime/Engine/Classes/Engine/StreamableManager.h
+++ b/Engine/Source/Runtime/Engine/Classes/Engine/StreamableManager.h
@@ -15,6 +15,7 @@
 #include "UObject/GCObject.h"
 #include "UObject/ObjectMacros.h"
 #include "UObject/SoftObjectPtr.h"
+#include "Experimental/StreamableManagerError.h"
 
 enum class EStreamableProgressType : uint8
 {
@@ -26,31 +27,6 @@ enum class EStreamableProgressType : uint8
 	DownloadAbsolute,
 };
 
-enum class EStreamableErrorType : uint8
-{
-	Unknown = 0,
-	/** Low-level package loading error */
-	PackageLoadError,
-	/** Low-level package loading was canceled. This is different than canceling a FStreamableHandle which just detaches from the low level request. */
-	PackageLoadCanceled,
-	/** Error downloading assets */
-	DownloadError,
-};
-
-ENGINE_API const TCHAR* LexToString(EStreamableErrorType ErrorType);
-
-struct FStreamableError
-{
-	FStreamableError(EStreamableErrorType InErrorType, FString InErrorMessage)
-		: ErrorMessage(MoveTemp(InErrorMessage))
-		, ErrorType(InErrorType)
-	{}
-
-	FString ErrorMessage;
-	EStreamableErrorType ErrorType = EStreamableErrorType::Unknown;
-};
-
-ENGINE_API FString LexToString(const FStreamableError Error);
 
 /** Defines FStreamableDelegate delegate interface */
 using FStreamableDelegate = TDelegate<void()>;
@@ -260,7 +236,7 @@ struct FStreamableHandle : public TSharedFromThis<FStreamableHandle>
 	}
 
 	/** Get error details */
-	const TOptional<FStreamableError>& GetError() const
+	const TOptional<UE::UnifiedError::FError>& GetError() const
 	{
 		return Error;
 	}
@@ -593,8 +569,8 @@ private:
 	void UnbindDelegates();
 
 	/** Record any errors that occured with this handle */
-	void SetError(FStreamableError&& Error);
-	void SetErrorOnParents(const FStreamableError& Error);
+	void SetError(UE::UnifiedError::FError&& Error);
+	void SetErrorOnParents(const UE::UnifiedError::FError& Error);
 
 	/** Delegate to call when streaming is completed */
 	FStreamableDelegateWithHandle CompleteDelegate;
@@ -612,7 +588,7 @@ private:
 	UE::FSourceLocation Location;
 
 	/** Error details in case of an error */
-	TOptional<FStreamableError> Error;
+	TOptional<UE::UnifiedError::FError> Error;
 
 	/** Handle that holds packages in the download cache */
 	FStreamableDownloadCachePin DownloadCachePin;
@@ -860,7 +836,7 @@ struct FStreamableManager : public FGCObject
 	 * @param Location				[optional] Is not intended for direct use, the parameter catches call site source location. That location is used for debug tools output.
 	 */
 	ENGINE_API TSharedPtr<FStreamableHandle> CreateErrorHandle(
-		FStreamableError&& Error,
+		UE::UnifiedError::FError&& Error,
 		FStreamableAsyncLoadParams&& Params,
 		FString DebugName = FStreamableHandle::HandleDebugName_CombinedHandle,
 		UE::FSourceLocation Location = UE::FSourceLocation::Current());
diff --git a/Engine/Source/Runtime/Engine/Private/AssetManager.cpp b/Engine/Source/Runtime/Engine/Private/AssetManager.cpp
--- a/Engine/Source/Runtime/Engine/Private/AssetManager.cpp
+++ b/Engine/Source/Runtime/Engine/Private/AssetManager.cpp
@@ -20,6 +20,7 @@
 #include "Interfaces/IPluginManager.h"
 #include "Internationalization/PackageLocalizationManager.h"
 #include "IPlatformFilePak.h"
+#include "Logging/StructuredLog.h"
 #include "Misc/CoreDelegates.h"
 #include "Misc/FileHelper.h"
 #include "Misc/DelayedAutoRegister.h"
@@ -1952,7 +1953,7 @@ TSharedPtr<FStreamableHandle> UAssetManager::ChangeBundleStateForPrimaryAssets(
 				FStreamableAsyncLoadParams StreamableLoadParams;
 				StreamableLoadParams.TargetsToStream = MoveTemp(PathsToLoad);
 				NewHandle = StreamableManager.CreateErrorHandle(
-					FStreamableError(EStreamableErrorType::DownloadError, FString::Printf(TEXT("Expected DownloadParams for %s"), *PrimaryAssetId.ToString())),
+					UE::UnifiedError::StreamableManager::MissingDownloadParams::MakeError(*PrimaryAssetId.ToString()),
 					MoveTemp(StreamableLoadParams), FString(DebugName), CopyTemp(Location));
 			}
 
@@ -2180,7 +2181,17 @@ void UAssetManager::OnAssetStateChangeCompleted(const TSharedPtr<FStreamableHand
 			// If an error occured, don't update the current state with a bad state
 			if (BoundHandle->HasError())
 			{
-				UE_LOG(LogAssetManager, Warning, TEXT("OnAssetStateChangeCompleted: Received error (%s) loading asset! (%s)"), *LexToString(*BoundHandle->GetError()), *PrimaryAssetId.ToString());
+				TRefCountPtr<const UE::UnifiedError::StreamableManager::FStreamableManagerErrorDetails> ErrorDetails = BoundHandle->GetError()->GetErrorDetails<UE::UnifiedError::StreamableManager::FStreamableManagerErrorDetails>();
+
+				if (ErrorDetails.IsValid())
+				{
+					UE_LOG(LogAssetManager, Warning, TEXT("OnAssetStateChangeCompleted: Received error (%s) loading asset! Primary asset id: (%s) erroring asset:%s"), *LexToString(*BoundHandle->GetError()), *PrimaryAssetId.ToString(), *ErrorDetails->GetSoftObjectPath());
+				}
+				else
+				{
+					UE_LOG(LogAssetManager, Warning, TEXT("OnAssetStateChangeCompleted: Received error (%s) loading asset! Primary asset id:(%s)"), *LexToString(*BoundHandle->GetError()), *PrimaryAssetId.ToString());
+				}
+				UE_LOGFMT(LogAssetManager, Warning, "OnAssetStateChangeCompleted: Received error ({Error}) loading asset! Primary asset id:({PrimaryAssetId})", ("Error", BoundHandle->GetError()), ("PrimaryAssetId", *PrimaryAssetId.ToString()));
 			}
 			else
 			{
diff --git a/Engine/Source/Runtime/Engine/Private/Experimental/StreamableManagerError.cpp b/Engine/Source/Runtime/Engine/Private/Experimental/StreamableManagerError.cpp
new file mode 100644
--- /dev/null
+++ b/Engine/Source/Runtime/Engine/Private/Experimental/StreamableManagerError.cpp
@@ -0,0 +1,24 @@
+#include "Engine/Experimental/StreamableManagerError.h"
+#include "UObject/UObjectGlobals.h"
+
+
+
+namespace UE::UnifiedError::StreamableManager
+{
+
+UE::UnifiedError::FError GetStreamableError(EAsyncLoadingResult::Type Result)
+{
+	switch (Result)
+	{
+	case EAsyncLoadingResult::Failed:
+		// We could supply an error string if async loading bubbled one up...
+		// Possibly GetExplanationForUnavailablePackage?
+		return AsyncLoadFailed::MakeError();
+	case EAsyncLoadingResult::Canceled:
+		return AsyncLoadCancelled::MakeError();
+	}
+
+	return AsyncLoadUnknownError::MakeError((int32)Result);
+}
+
+} // namespace UE::UnifiedError::StreamableManager
\ No newline at end of file
diff --git a/Engine/Source/Runtime/Engine/Private/StreamableManager.cpp b/Engine/Source/Runtime/Engine/Private/StreamableManager.cpp
--- a/Engine/Source/Runtime/Engine/Private/StreamableManager.cpp
+++ b/Engine/Source/Runtime/Engine/Private/StreamableManager.cpp
@@ -1,6 +1,8 @@
 // Copyright Epic Games, Inc. All Rights Reserved.
 
 #include "Engine/StreamableManager.h"
+#include "Experimental/UnifiedError/UnifiedError.h"
+#include "Experimental/IO/IoStatusError.h"
 #include "UObject/ICookInfo.h"
 #include "UObject/ObjectRedirector.h"
 #include "UObject/Package.h"
@@ -20,6 +22,9 @@
 
 DEFINE_LOG_CATEGORY_STATIC(LogStreamableManager, Log, All);
 
+
+// DEFINE_ERROR_DETAILS(StreamableManager, FStreamableManagerErrorDetails);
+
 // Default to 1 frame, this will cause the delegates to go off on the next tick to avoid recursion issues. Set higher to fake disk lag
 static int32 GStreamableDelegateDelayFrames = 1;
 static FAutoConsoleVariableRef CVarStreamableDelegateDelayFrames(
@@ -46,26 +51,6 @@ static FAutoConsoleVariableRef CVarStreamableStripDebugNameInShipping(
 	ECVF_Default
 );
 
-const TCHAR* LexToString(EStreamableErrorType ErrorType)
-{
-	switch (ErrorType)
-	{
-	default:
-	case EStreamableErrorType::Unknown:
-		return TEXT("Unknown");
-	case EStreamableErrorType::PackageLoadError:
-		return TEXT("PackageLoadError");
-	case EStreamableErrorType::PackageLoadCanceled:
-		return TEXT("PackageLoadCanceled");
-	case EStreamableErrorType::DownloadError:
-		return TEXT("DownloadError");
-	}
-}
-
-FString LexToString(const FStreamableError Error)
-{
-	return FString::Printf(TEXT("%s - %s"), LexToString(Error.ErrorType), *Error.ErrorMessage);
-}
 
 FStreamableDownloadCachePin::FStreamableDownloadCachePin() = default;
 FStreamableDownloadCachePin::~FStreamableDownloadCachePin() = default;
@@ -422,7 +407,7 @@ namespace UE::StreamableManager::Private
 			}
 		}
 
-		void CompleteLoad(const TSharedRef<FStreamableHandle>& OwningHandle, FStreamableError&& Error)
+		void CompleteLoad(const TSharedRef<FStreamableHandle>& OwningHandle, UE::UnifiedError::FError&& Error)
 		{
 			OwningHandle->SetError(MoveTemp(Error));
 			CompleteLoad(OwningHandle);
@@ -439,7 +424,7 @@ namespace UE::StreamableManager::Private
 			UE::IoStore::IOnDemandIoStore* IoStore = UE::IoStore::TryGetOnDemandIoStore();
 			if (!IoStore)
 			{
-				CompleteLoad(OwningHandle, FStreamableError{ EStreamableErrorType::DownloadError, TEXT("IoStore not found") });
+				CompleteLoad(OwningHandle, UE::UnifiedError::StreamableManager::IoStoreNotFound::MakeError());
 				OwningHandle->DownloadContext = nullptr; // Destroy self
 				return;
 			}
@@ -496,7 +481,7 @@ namespace UE::StreamableManager::Private
 				}
 				else
 				{
-					CompleteLoad(OwningHandle, FStreamableError{ EStreamableErrorType::DownloadError, OnDemandInstallResult.Status.ToString() });
+					CompleteLoad(OwningHandle, UE::UnifiedError::IoStore::ConvertError(OnDemandInstallResult.Status));
 				}
 
 				return;
@@ -525,21 +510,6 @@ namespace UE::StreamableManager::Private
 			}
 		}
 	};
-
-	static FStreamableError GetStreamableError(EAsyncLoadingResult::Type Result)
-	{
-		switch (Result)
-		{
-		case EAsyncLoadingResult::Failed:
-			// We could supply an error string if async loading bubbled one up...
-			// Possibly GetExplanationForUnavailablePackage?
-			return FStreamableError{ EStreamableErrorType::PackageLoadError, TEXT("Async load failed") };
-		case EAsyncLoadingResult::Canceled:
-			return FStreamableError{ EStreamableErrorType::PackageLoadCanceled, TEXT("Async load canceled") };
-		}
-
-		return FStreamableError{ EStreamableErrorType::Unknown, FString() };
-	}
 }
 
 FStreamableHandle::FStreamableHandle(FPrivateToken, UE::FSourceLocation&& InLocation)
@@ -1083,6 +1053,10 @@ FStreamableHandle::~FStreamableHandle()
 void FStreamableHandle::CompleteLoad()
 {
 	// Only complete if it's still active
+	if (Error.IsSet())
+	{
+		UE_LOG(LogStreamableManager, Verbose, TEXT("FStreamableHandle::CompleteLoad failed load %s, with error %s!"), *DebugName, *Error->GetErrorMessage().ToString());
+	}
 	if (IsActive())
 	{
 		bLoadCompleted = true;
@@ -1217,16 +1191,17 @@ void FStreamableHandle::UnbindDelegates()
 	CompleteDelegate.Unbind();
 }
 
-void FStreamableHandle::SetError(FStreamableError&& InError)
+void FStreamableHandle::SetError(UE::UnifiedError::FError&& InError)
 {
 	if (!Error)
 	{
 		SetErrorOnParents(InError);
 		Error = MoveTemp(InError);
+		
 	}
 }
 
-void FStreamableHandle::SetErrorOnParents(const FStreamableError& InError)
+void FStreamableHandle::SetErrorOnParents(const UE::UnifiedError::FError& InError)
 {
 	// Update any meta handles that are still active
 	for (TWeakPtr<FStreamableHandle>& WeakHandle : ParentHandles)
@@ -1249,7 +1224,12 @@ void FStreamableHandle::AsyncLoadCallbackWrapper(const FName& PackageName, UPack
 	{
 		if (Result != EAsyncLoadingResult::Succeeded)
 		{
-			SetError(UE::StreamableManager::Private::GetStreamableError(Result));
+			UE::UnifiedError::FError NewError = UE::UnifiedError::StreamableManager::GetStreamableError(Result);
+
+			TRefCountPtr<UE::UnifiedError::IErrorDetails> NewErrorDetails( new UE::UnifiedError::StreamableManager::FStreamableManagerErrorDetails(TargetName.ToString()));
+
+			NewError.PushErrorDetails(NewErrorDetails);
+			SetError(MoveTemp(NewError));
 		}
 
 		OwningManager->AsyncLoadCallback(MoveTemp(TargetName), Package, Result);
@@ -1707,7 +1687,7 @@ FStreamable* FStreamableManager::StreamInternal(const FSoftObjectPath& InTargetN
 				Existing->bLoadFailed = true;
 				UE_LOG(LogStreamableManager, Log, TEXT("Failed attempt to load %s"), *TargetName.ToString());
 
-				Handle->SetError(FStreamableError{ EStreamableErrorType::PackageLoadError, TEXT("Static load failed") });
+				Handle->SetError(UE::UnifiedError::StreamableManager::PackageLoadFailed::MakeError(*TargetName.ToString()));
 			}
 			Existing->bAsyncLoadRequestOutstanding = false;
 		}
@@ -1728,7 +1708,7 @@ FStreamable* FStreamableManager::StreamInternal(const FSoftObjectPath& InTargetN
 				Existing->bLoadFailed = true;
 				Existing->bAsyncLoadRequestOutstanding = false;
 
-				Handle->SetError(FStreamableError{ EStreamableErrorType::PackageLoadError, TEXT("Unkown package")});
+				Handle->SetError(UE::UnifiedError::StreamableManager::PackageNameInvalid::MakeError(Package));
 			}
 			else
 			{
@@ -1930,7 +1910,7 @@ TSharedPtr<FStreamableHandle> FStreamableManager::RequestSyncLoadInternal(TArray
 
 		if (bBadPackageState || bLoadNotComplete)
 		{
-			Request->SetError(FStreamableError{ EStreamableErrorType::PackageLoadError, TEXT("Sync load incomplete") });
+			Request->SetError(UE::UnifiedError::StreamableManager::SyncLoadIncomplete::MakeError(Request->DebugName));
 		}
 	}
 
@@ -2139,14 +2119,15 @@ void FStreamableManager::AsyncLoadCallback(FSoftObjectPath TargetName, UPackage*
 		{
 			for (const TSharedRef<FStreamableHandle>& Handle : Existing->LoadingHandles)
 			{
-				Handle->SetError(UE::StreamableManager::Private::GetStreamableError(Result));
+				// @Per @Justin In my testing this has already been set by the AsyncLoadCallbackWrapper.  does this need to happen twice?
+				Handle->SetError(UE::UnifiedError::StreamableManager::GetStreamableError(Result));
 			}
 		}
 		else if (!Existing->Target)
 		{
 			for (const TSharedRef<FStreamableHandle>& Handle : Existing->LoadingHandles)
 			{
-				Handle->SetError(FStreamableError{ EStreamableErrorType::PackageLoadError, TEXT("Failed to find target asset in memory") });
+				Handle->SetError(UE::UnifiedError::StreamableManager::PackageLoadFailed::MakeError(TargetName.ToString()));
 			}
 		}
 
@@ -2157,6 +2138,7 @@ void FStreamableManager::AsyncLoadCallback(FSoftObjectPath TargetName, UPackage*
 	}
 	else
 	{
+		// should this be an error?  Is it an error if no handles are waiting for it? 
 		UE_LOG(LogStreamableManager, Error, TEXT("Can't find streamable for %s in AsyncLoadCallback!"), *TargetName.ToString());
 	}
 }
@@ -2430,7 +2412,7 @@ TSharedPtr<FStreamableHandle> FStreamableManager::CreateCombinedHandle(
 	return NewRequest;
 }
 
-TSharedPtr<FStreamableHandle> FStreamableManager::CreateErrorHandle(FStreamableError&& Error, FStreamableAsyncLoadParams&& Params, FString DebugName, UE::FSourceLocation Location)
+TSharedPtr<FStreamableHandle> FStreamableManager::CreateErrorHandle(UE::UnifiedError::FError&& Error, FStreamableAsyncLoadParams&& Params, FString DebugName, UE::FSourceLocation Location)
 {
 	TSharedRef<FStreamableHandle> NewRequest = MakeShared<FStreamableHandle>(FStreamableHandle::FPrivateToken(), MoveTemp(Location));
 	NewRequest->OwningManager = this;
