Added file with correct case from cl 40976944

diff --git a/Engine/Source/Runtime/Core/Public/Experimental/IO/IoStatusError.h b/Engine/Source/Runtime/Core/Public/Experimental/IO/IoStatusError.h
new file mode 100644
--- /dev/null
+++ b/Engine/Source/Runtime/Core/Public/Experimental/IO/IoStatusError.h
@@ -0,0 +1,15 @@
+// Copyright Epic Games, Inc. All Rights Reserved.
+#pragma once
+
+#include "Experimental/UnifiedError/UnifiedError.h"
+
+
+class FIoStatus;
+
+DECLARE_ERROR_MODULE(IoStore, 0x15);
+
+
+namespace UE::UnifiedError::IoStore
+{
+	CORE_API UE::UnifiedError::FError ConvertError(const FIoStatus& Status);
+}
\ No newline at end of file
