Added copyright to some files which were missing it.  Fix for cis warnings. Fix case on include.

diff --git a/Engine/Source/Runtime/Core/Private/Experimental/IO/IoStatusError.cpp b/Engine/Source/Runtime/Core/Private/Experimental/IO/IoStatusError.cpp
--- a/Engine/Source/Runtime/Core/Private/Experimental/IO/IoStatusError.cpp
+++ b/Engine/Source/Runtime/Core/Private/Experimental/IO/IoStatusError.cpp
@@ -1,6 +1,7 @@
-#include "Experimental/IO/IOStatusError.h"
-#include "IO/IoStatus.h"
+// Copyright Epic Games, Inc. All Rights Reserved.
 
+#include "Experimental/IO/IoStatusError.h"
+#include "IO/IoStatus.h"
 #include "Internationalization/Internationalization.h"
 #include "Internationalization/Text.h"
 
diff --git a/Engine/Source/Runtime/Core/Private/Experimental/UnifiedError/UnifiedError.cpp b/Engine/Source/Runtime/Core/Private/Experimental/UnifiedError/UnifiedError.cpp
--- a/Engine/Source/Runtime/Core/Private/Experimental/UnifiedError/UnifiedError.cpp
+++ b/Engine/Source/Runtime/Core/Private/Experimental/UnifiedError/UnifiedError.cpp
@@ -1,3 +1,5 @@
+// Copyright Epic Games, Inc. All Rights Reserved.
+
 #include "Experimental/UnifiedError/UnifiedError.h"
 #include "Internationalization/Text.h"
 #include "Internationalization/Internationalization.h"
diff --git a/Engine/Source/Runtime/Engine/Private/Experimental/StreamableManagerError.cpp b/Engine/Source/Runtime/Engine/Private/Experimental/StreamableManagerError.cpp
--- a/Engine/Source/Runtime/Engine/Private/Experimental/StreamableManagerError.cpp
+++ b/Engine/Source/Runtime/Engine/Private/Experimental/StreamableManagerError.cpp
@@ -1,3 +1,4 @@
+// Copyright Epic Games, Inc. All Rights Reserved.
 #include "Engine/Experimental/StreamableManagerError.h"
 #include "UObject/UObjectGlobals.h"
 
