[Source Location] Deprecate UE::SourceLocation::To[Full|FileAndLine]String functions in favor of the corresponding adapter-based functions.

diff --git a/Engine/Source/Runtime/Core/Public/Misc/SourceLocation.h b/Engine/Source/Runtime/Core/Public/Misc/SourceLocation.h
--- a/Engine/Source/Runtime/Core/Public/Misc/SourceLocation.h
+++ b/Engine/Source/Runtime/Core/Public/Misc/SourceLocation.h
@@ -101,7 +101,7 @@ namespace UE
 		 * #include "SourceLocationUtils.h"
 		 * 
 		 * void MyFunction(int Param1, int Param2, UE::FSourceLocation Location = UE::FSourceLocation::Current()) {
-		 *    UE_LOG(TEXT("My caller is %s"), *UE::SourceLocation::ToFullString(Location));
+		 *    UE_LOG(TEXT("My caller is %s"), *UE::SourceLocation::Full(Location).ToString());
 		 * }
 		 */
 		static UE_CONSTEVAL FSourceLocation Current(FSourceLocationImpl Impl = FSourceLocationImpl::current()) noexcept
diff --git a/Engine/Source/Runtime/Core/Public/Misc/SourceLocationUtils.h b/Engine/Source/Runtime/Core/Public/Misc/SourceLocationUtils.h
--- a/Engine/Source/Runtime/Core/Public/Misc/SourceLocationUtils.h
+++ b/Engine/Source/Runtime/Core/Public/Misc/SourceLocationUtils.h
@@ -166,7 +166,7 @@ namespace UE::SourceLocation
 	 * @param Location - The source location to convert to string.
 	 * @return Formatted string with full source location information.
 	 */
-	// UE_DEPRECATED(5.6, "Use UE::SourceLocation::Full(Location).ToString() instead.")
+	UE_DEPRECATED(5.6, "Use UE::SourceLocation::Full(Location).ToString() instead.")
 	static inline FString ToFullString(const FSourceLocation& Location)
 	{
 		return Full(Location).ToString();
@@ -177,7 +177,7 @@ namespace UE::SourceLocation
 	 * @param Location - The source location to convert to string.
 	 * @return Formatted string with file and line information.
 	 */
-	// UE_DEPRECATED(5.6, "Use UE::SourceLocation::FileAndLine(Location).ToString() instead.")
+	UE_DEPRECATED(5.6, "Use UE::SourceLocation::FileAndLine(Location).ToString() instead.")
 	static inline FString ToFileAndLineString(const FSourceLocation& Location)
 	{
 		return FileAndLine(Location).ToString();
