1. Bypassed the NaN thermal headroom return value from getThermalHeadroom.
2. Provided an enum value EForecastPeriod::INSTANT to call getThermalHeadroom with 0 forecast period.

diff --git a/Engine/Source/Runtime/Core/Private/Android/AndroidPlatformThermal.cpp b/Engine/Source/Runtime/Core/Private/Android/AndroidPlatformThermal.cpp
--- a/Engine/Source/Runtime/Core/Private/Android/AndroidPlatformThermal.cpp
+++ b/Engine/Source/Runtime/Core/Private/Android/AndroidPlatformThermal.cpp
@@ -33,14 +33,20 @@ namespace UE
 	{
 		class FAndroidThermalManagerImpl
 		{
+			// *** From Google: 
+			// * We recommend calling this every 10 seconds because Android devices do not much frequently update thermal data.
+			// * In the case of the Pixel, it updates every 7 seconds, and the Galaxy is 10 secs.
+			// * It's ok if you call it with multiple forecasts within 10 secs, but please never call it no more than twice per second.
+			// ***
+			// We set the period to 1 second to track the thermal status as tightly as we can.
 			static constexpr float CallRateLimitSeconds = 1.0f;
 			static constexpr int NumPeriods = FAndroidPlatformThermal::EForecastPeriod::NUM_FORECASTPERIODS;
-			static constexpr int PeriodEnumToS[NumPeriods] = { 1,5,10 };
-			static_assert(NumPeriods == 3); // update PeriodEnumToS if this changes.
+			static constexpr int PeriodEnumToS[NumPeriods] = { 0, 1, 5, 10 };
+			static_assert(NumPeriods == 4); // update PeriodEnumToS if this changes.
 			double LastCallTime = 0;
 			FCriticalSection ThermalCS;
 
-			FAndroidPlatformThermal::EForecastPeriod PeriodToRead = FAndroidPlatformThermal::EForecastPeriod::ONE_SEC;
+			FAndroidPlatformThermal::EForecastPeriod PeriodToRead = FAndroidPlatformThermal::EForecastPeriod::INSTANT;
 
 			FJavaClassObject ActivityClass;			// note: ActivityClass must always be accessed from the same thread that created it.
 			FJavaClassMethod ThermalHeadroomMethod;
@@ -69,9 +75,12 @@ namespace UE
 				{
 					if (ThermalCS.TryLock())
 					{
-						/// thermal headroom docs state we should not call this API more than once per second.
 						float RecentThermalReading = ActivityClass.CallMethod<float>(ThermalHeadroomMethod, PeriodEnumToS[PeriodToRead]);
-						ThermalInfo.HeadroomForecasts[PeriodToRead] = RecentThermalReading;
+						// Bypass NaN values returned from getThermalHeadroom() because the device won't update thermal data frequently.
+						if (!isnan(RecentThermalReading))
+						{
+							ThermalInfo.HeadroomForecasts[PeriodToRead] = RecentThermalReading;
+						}
 
 						LastCallTime = CurrentTime;
 						PeriodToRead = (FAndroidPlatformThermal::EForecastPeriod)((int)(PeriodToRead+1) % NumPeriods);
diff --git a/Engine/Source/Runtime/Core/Public/Android/AndroidPlatformThermal.h b/Engine/Source/Runtime/Core/Public/Android/AndroidPlatformThermal.h
--- a/Engine/Source/Runtime/Core/Public/Android/AndroidPlatformThermal.h
+++ b/Engine/Source/Runtime/Core/Public/Android/AndroidPlatformThermal.h
@@ -4,11 +4,11 @@
 
 struct FAndroidPlatformThermal
 {
-	enum EForecastPeriod { ONE_SEC, FIVE_SEC, TEN_SEC, NUM_FORECASTPERIODS };
+	enum EForecastPeriod { INSTANT, ONE_SEC, FIVE_SEC, TEN_SEC, NUM_FORECASTPERIODS };
 	// returns the 'thermal stress' value for the forecasted period.
 	// This value is a hint for the extent of thermal throttling the device will experience during the time scale provided.
 	// The implementation is provided by PowerManager.getThermalHeadroom(). 
 	// Values of 1.0 or above indicate the current workload can not be sustained and thermal throttling will occur.
-	static float GetThermalStress(EForecastPeriod ForecastPeriod);
+	static CORE_API float GetThermalStress(EForecastPeriod ForecastPeriod);
 };
 
