Update TArray single addition functions to avoid a function call when reallocation is not needed.

#rnx

diff --git a/Engine/Source/Runtime/Core/Public/Containers/Array.h b/Engine/Source/Runtime/Core/Public/Containers/Array.h
--- a/Engine/Source/Runtime/Core/Public/Containers/Array.h
+++ b/Engine/Source/Runtime/Core/Public/Containers/Array.h
@@ -383,6 +383,78 @@ namespace UE::Core::Private
 		return Result;
 	}
 
+	
+	// Called only when we KNOW we are going to do a realloc increasing by 1.
+	// In this case, we know that max == num and can simplify things in a very 
+	// hot location in the code.
+	// This returns the old ArrayMax in order to save a register clobber/reload.
+	template <uint32 Flags, typename AllocatorInstanceType>
+	FORCEINLINE TAllocatorSizeType_T<AllocatorInstanceType> ReallocGrow1_DoAlloc_Impl(
+		uint32                                       ElementSize,
+		uint32                                       ElementAlignment,
+		AllocatorInstanceType&                       AllocatorInstance,
+		TAllocatorSizeType_T<AllocatorInstanceType>& ArrayMax
+	)
+	{		
+		using SizeType  = TAllocatorSizeType_T<AllocatorInstanceType>;
+		using USizeType = std::make_unsigned_t<SizeType>;
+
+		const USizeType UOldMax = (USizeType)ArrayMax;
+		const USizeType UNewNum = UOldMax + 1U;
+		const SizeType  OldMax  = (SizeType)UOldMax;
+		const SizeType  NewNum  = (SizeType)UNewNum;
+		
+		// This should only happen when we've underflowed or overflowed SizeType
+		if (NewNum < OldMax)
+		{
+			OnInvalidArrayNum((unsigned long long)UNewNum);
+		}
+
+		SizeType NewMax;
+		if constexpr (!!(Flags & 1)) // TAllocatorTraits<AllocatorType>::SupportsElementAlignment
+		{
+			NewMax = AllocatorInstance.CalculateSlackGrow(NewNum, OldMax, ElementSize, ElementAlignment);
+			AllocatorInstance.ResizeAllocation(UOldMax, NewMax, ElementSize, ElementAlignment);
+		}
+		else
+		{
+			NewMax = AllocatorInstance.CalculateSlackGrow(NewNum, OldMax, ElementSize);
+			AllocatorInstance.ResizeAllocation(UOldMax, NewMax, ElementSize);
+		}
+		ArrayMax = NewMax;
+		#if UE_ENABLE_ARRAY_SLACK_TRACKING
+		if constexpr (!!(Flags & 2)) // TAllocatorTraits<AllocatorType>::SupportsSlackTracking
+		{
+			AllocatorInstance.SlackTrackerLogNum(NewNum);
+		}
+		#endif
+
+		return OldMax;
+	}
+
+	// Version for small sizes/alignments. This allows the parameter setup to be a single instruction
+	// note the uint16 limitation allows for a single instruction setup on arm.
+	template <uint32 Flags, typename AllocatorInstanceType>
+	FORCENOINLINE TAllocatorSizeType_T<AllocatorInstanceType> ReallocGrow1_DoAlloc_Tiny(
+		uint16                                       ElementSizeAndAlignment,
+		AllocatorInstanceType&                       AllocatorInstance,
+		TAllocatorSizeType_T<AllocatorInstanceType>& ArrayMax
+	)
+	{
+		return ReallocGrow1_DoAlloc_Impl<Flags, AllocatorInstanceType>(ElementSizeAndAlignment & 0xff, ElementSizeAndAlignment >> 8, AllocatorInstance, ArrayMax);		
+	}
+
+	template <uint32 Flags, typename AllocatorInstanceType>
+	FORCENOINLINE TAllocatorSizeType_T<AllocatorInstanceType> ReallocGrow1_DoAlloc(
+		uint32                                       ElementSize,
+		uint32                                       ElementAlignment,
+		AllocatorInstanceType&                       AllocatorInstance,
+		TAllocatorSizeType_T<AllocatorInstanceType>& ArrayMax
+	)
+	{
+		return ReallocGrow1_DoAlloc_Impl<Flags, AllocatorInstanceType>(ElementSize, ElementAlignment, AllocatorInstance, ArrayMax);		
+	}
+
 	// This should be used for repeated growing operations when reallocations are to be amortized over multiple inserts.
 	template <uint32 Flags, typename AllocatorInstanceType>
 	FORCENOINLINE TAllocatorSizeType_T<AllocatorInstanceType> ReallocGrow(
@@ -1598,14 +1670,48 @@ public:
 	 */
 	FORCEINLINE SizeType AddUninitialized()
 	{
-		return UE::Core::Private::ReallocGrow<UE::Core::Private::GetAllocatorFlags<AllocatorType>()>(
-			sizeof(ElementType),
-			alignof(ElementType),
-			1,
-			AllocatorInstance,
-			ArrayNum,
-			ArrayMax
-		);
+		// Begin sensitive code! Single element additions/insertions get inlined _everywhere_ and thus
+		// punch above their weight with respect to exe size, so we pay a lot of attention to every instruction
+		// here until we get in to the NOINLINE functions that actually handle the growth.
+
+		// Modulo some register allocation and pointer offsets, inlined sites on x64 for the "Tiny" path should look roughly like:
+		// 
+		// mov eax, dword ptr [array + 8h] // load arraynum
+		// cmp eax, dword ptr [array + 0Ch] // load arraymax
+		// jne [to ArrayNum++]
+		// mov ecx, immediate representing size/alignment
+		// lea r8, [array + 0Ch]] // array max
+		// lea rdx, [array] // allocator instance
+		// call
+		// lea ecx, [rax+1] // ArrayNum++
+		// mov dword ptr [array + 8h], ecx // Save ArrayNum
+		//
+
+		// Single cmp, which we can assume because we are adding a single element.
+		if (ArrayNum == ArrayMax)
+		{
+			// Both branches here write the return in to ArrayNum. This is because the function call clobbers the registers
+			// and if we assign as part of the return into something we need, the compiler doesn't have to reload the data 
+			// into the clobbered register.
+
+			// When we can pack size and alignment into a single 16 bit load, we save a parameter setup instruction
+			// for the function call. The 16 bit part is load-bearing for arm codegen.
+			if constexpr (sizeof(ElementType) <= 255 && alignof(ElementType) <= 255)
+			{
+				// Note that the realloc functions are templated ONLY on allocator instance so they are not duplicated
+				// in the code for every type!
+				ArrayNum = UE::Core::Private::ReallocGrow1_DoAlloc_Tiny<UE::Core::Private::GetAllocatorFlags<AllocatorType>()>(sizeof(ElementType) | (alignof(ElementType) << 8), AllocatorInstance, ArrayMax);
+			}
+			else
+			{
+				ArrayNum = UE::Core::Private::ReallocGrow1_DoAlloc<UE::Core::Private::GetAllocatorFlags<AllocatorType>()>(sizeof(ElementType), alignof(ElementType), AllocatorInstance, ArrayMax);
+			}
+		}
+		// End sensitive code!
+
+		SizeType OldArrayNum = ArrayNum;
+		ArrayNum++;
+		return OldArrayNum;	
 	}
 	FORCEINLINE SizeType AddUninitialized(SizeType Count)
 	{
@@ -1623,14 +1729,21 @@ public:
 private:
 	void InsertUninitializedImpl(SizeType Index)
 	{
-		SizeType OldNum = UE::Core::Private::ReallocGrow<UE::Core::Private::GetAllocatorFlags<AllocatorType>()>(
-			sizeof(ElementType),
-			alignof(ElementType),
-			1,
-			AllocatorInstance,
-			ArrayNum,
-			ArrayMax
-		);
+		// Begin sensitive code! See comments in AddUninitialized() before touching!
+		if (ArrayNum == ArrayMax)
+		{
+			if constexpr (sizeof(ElementType) <= 255 && alignof(ElementType) <= 255)
+			{
+				ArrayNum = UE::Core::Private::ReallocGrow1_DoAlloc_Tiny<UE::Core::Private::GetAllocatorFlags<AllocatorType>()>(sizeof(ElementType) | (alignof(ElementType) << 8), AllocatorInstance, ArrayMax);
+			}
+			else
+			{
+				ArrayNum = UE::Core::Private::ReallocGrow1_DoAlloc<UE::Core::Private::GetAllocatorFlags<AllocatorType>()>(sizeof(ElementType), alignof(ElementType), AllocatorInstance, ArrayMax);
+			}
+		}
+		// End sensitive code!
+		SizeType OldNum = ArrayNum;
+		ArrayNum++;		
 		ElementType* Data = GetData() + Index;
 		RelocateConstructItems<ElementType>((void*)(Data + 1), Data, OldNum - Index);
 	}
@@ -2477,17 +2590,25 @@ public:
 		//     // Won't compile if the caller doesn't have access to FMyType::FPrivateToken
 		//     Arr.Emplace(FMyType::FPrivateToken{}, 5, 3.14f, TEXT("Banana"));
 		//
-		SizeType Result = UE::Core::Private::ReallocGrow<UE::Core::Private::GetAllocatorFlags<AllocatorType>()>(
-			sizeof(ElementType),
-			alignof(ElementType),
-			1,
-			AllocatorInstance,
-			ArrayNum,
-			ArrayMax
-		);
-		void* Ptr = (char*)AllocatorInstance.GetAllocation() + sizeof(ElementType) * Result;
+
+		// Begin sensitive code! See comments in AddUninitialized() before touching!
+		if (ArrayNum == ArrayMax)
+		{
+			if constexpr (sizeof(ElementType) <= 255 && alignof(ElementType) <= 255)
+			{
+				ArrayNum = UE::Core::Private::ReallocGrow1_DoAlloc_Tiny<UE::Core::Private::GetAllocatorFlags<AllocatorType>()>(sizeof(ElementType) | (alignof(ElementType) << 8), AllocatorInstance, ArrayMax);
+			}
+			else
+			{
+				ArrayNum = UE::Core::Private::ReallocGrow1_DoAlloc<UE::Core::Private::GetAllocatorFlags<AllocatorType>()>(sizeof(ElementType), alignof(ElementType), AllocatorInstance, ArrayMax);
+			}
+		}
+		// End sensitive code!
+		SizeType OldArrayNum = ArrayNum;
+		ArrayNum++;		
+		void* Ptr = (char*)AllocatorInstance.GetAllocation() + sizeof(ElementType) * OldArrayNum;
 		(void)new (Ptr) ElementType(Forward<ArgsType>(Args)...);
-		return ArrayNum - 1;
+		return OldArrayNum;
 	}
 
 	/**
@@ -2521,15 +2642,24 @@ public:
 		//     // Won't compile if the caller doesn't have access to FMyType::FPrivateToken
 		//     Arr.Emplace(FMyType::FPrivateToken{}, 5, 3.14f, TEXT("Banana"));
 		//
-		SizeType Result = UE::Core::Private::ReallocGrow<UE::Core::Private::GetAllocatorFlags<AllocatorType>()>(
-			sizeof(ElementType),
-			alignof(ElementType),
-			1,
-			AllocatorInstance,
-			ArrayNum,
-			ArrayMax
-		);
-		void* Ptr = (char*)AllocatorInstance.GetAllocation() + sizeof(ElementType) * Result;
+		// Begin sensitive code! See comments in AddUninitialized() before touching!
+		if (ArrayNum == ArrayMax)
+		{
+			if constexpr (sizeof(ElementType) <= 255 && alignof(ElementType) <= 255)
+			{
+				// Looks weird but saves the register reload
+				ArrayNum = UE::Core::Private::ReallocGrow1_DoAlloc_Tiny<UE::Core::Private::GetAllocatorFlags<AllocatorType>()>(sizeof(ElementType) | (alignof(ElementType) << 8), AllocatorInstance, ArrayMax);
+			}
+			else
+			{
+				// Looks weird but saves the register reload
+				ArrayNum = UE::Core::Private::ReallocGrow1_DoAlloc<UE::Core::Private::GetAllocatorFlags<AllocatorType>()>(sizeof(ElementType), alignof(ElementType), AllocatorInstance, ArrayMax);
+			}
+		}
+		// End sensitive code!
+		SizeType OldArrayNum = ArrayNum;
+		ArrayNum++;
+		void* Ptr = (char*)AllocatorInstance.GetAllocation() + sizeof(ElementType) * OldArrayNum;
 		return *new (Ptr) ElementType(Forward<ArgsType>(Args)...);
 	}
 
