PVS is technically correct but not usefully correct here.

#rnx

diff --git a/Engine/Source/Runtime/Core/Public/Containers/Array.h b/Engine/Source/Runtime/Core/Public/Containers/Array.h
--- a/Engine/Source/Runtime/Core/Public/Containers/Array.h
+++ b/Engine/Source/Runtime/Core/Public/Containers/Array.h
@@ -1696,7 +1696,9 @@ public:
 
 			// When we can pack size and alignment into a single 16 bit load, we save a parameter setup instruction
 			// for the function call. The 16 bit part is load-bearing for arm codegen.
-			if constexpr (sizeof(ElementType) <= 255 && alignof(ElementType) <= 255)
+			// PVS believes this is overwrought because technically sizeof() <= 255 implies alignof <= 255. I think this is far more
+			// clear and costs nothing.
+			if constexpr (sizeof(ElementType) <= 255 && alignof(ElementType) <= 255) // -V590 
 			{
 				// Note that the realloc functions are templated ONLY on allocator instance so they are not duplicated
 				// in the code for every type!
@@ -1732,7 +1734,7 @@ private:
 		// Begin sensitive code! See comments in AddUninitialized() before touching!
 		if (ArrayNum == ArrayMax)
 		{
-			if constexpr (sizeof(ElementType) <= 255 && alignof(ElementType) <= 255)
+			if constexpr (sizeof(ElementType) <= 255 && alignof(ElementType) <= 255) // -V590 
 			{
 				ArrayNum = UE::Core::Private::ReallocGrow1_DoAlloc_Tiny<UE::Core::Private::GetAllocatorFlags<AllocatorType>()>(sizeof(ElementType) | (alignof(ElementType) << 8), AllocatorInstance, ArrayMax);
 			}
@@ -2594,7 +2596,7 @@ public:
 		// Begin sensitive code! See comments in AddUninitialized() before touching!
 		if (ArrayNum == ArrayMax)
 		{
-			if constexpr (sizeof(ElementType) <= 255 && alignof(ElementType) <= 255)
+			if constexpr (sizeof(ElementType) <= 255 && alignof(ElementType) <= 255) // -V590 
 			{
 				ArrayNum = UE::Core::Private::ReallocGrow1_DoAlloc_Tiny<UE::Core::Private::GetAllocatorFlags<AllocatorType>()>(sizeof(ElementType) | (alignof(ElementType) << 8), AllocatorInstance, ArrayMax);
 			}
@@ -2645,7 +2647,7 @@ public:
 		// Begin sensitive code! See comments in AddUninitialized() before touching!
 		if (ArrayNum == ArrayMax)
 		{
-			if constexpr (sizeof(ElementType) <= 255 && alignof(ElementType) <= 255)
+			if constexpr (sizeof(ElementType) <= 255 && alignof(ElementType) <= 255) // -V590 
 			{
 				// Looks weird but saves the register reload
 				ArrayNum = UE::Core::Private::ReallocGrow1_DoAlloc_Tiny<UE::Core::Private::GetAllocatorFlags<AllocatorType>()>(sizeof(ElementType) | (alignof(ElementType) << 8), AllocatorInstance, ArrayMax);
