Allow null mutexes to be passed into FScopeUnlock/TScopeUnlock.

Previously, FScopeUnlock allowed null mutexes to be passed in, but TScopeUnlock did not. A recent CL caused FScopeUnlock to share TScopeUnlock's implementation and thus broke null mutex support. We now allow null mutexes in both.

There are no other existing callers of TScopeUnlock so it does not matter that its signature has changed.

diff --git a/Engine/Source/Runtime/Core/Public/Misc/ScopeLock.h b/Engine/Source/Runtime/Core/Public/Misc/ScopeLock.h
--- a/Engine/Source/Runtime/Core/Public/Misc/ScopeLock.h
+++ b/Engine/Source/Runtime/Core/Public/Misc/ScopeLock.h
@@ -89,7 +89,7 @@ namespace UE
 	//		TScopeLock<FCriticalSection> ScopeLock(CriticalSection);
 	//		for (FElementType& Element : ThreadUnsafeContainer)
 	//		{
-	//			TScopeUnlock<FCriticalSection> ScopeUnlock(CriticalSection);
+	//			TScopeUnlock<FCriticalSection> ScopeUnlock(&CriticalSection);
 	//			Process(Element);
 	//		}
 	//	}
@@ -99,15 +99,20 @@ namespace UE
 	public:
 		UE_NONCOPYABLE(TScopeUnlock);
 
-		[[nodiscard]] explicit TScopeUnlock(MutexType& InMutex) : Mutex(&InMutex)
+		[[nodiscard]] explicit TScopeUnlock(MutexType* InMutex) : Mutex(InMutex)
 		{
-			check(Mutex);
-			Mutex->Unlock();
+			if (Mutex)
+			{
+				Mutex->Unlock();
+			}
 		}
 
 		~TScopeUnlock()
 		{
-			Mutex->Lock();
+			if (Mutex)
+			{
+				Mutex->Lock();
+			}
 		}
 
 	private:
@@ -171,7 +176,7 @@ public:
 	 *
 	 * @param InSyncObject The synchronization object to manage, can be null.
 	 */
-	[[nodiscard]] explicit FScopeUnlock(TNotNull<FCriticalSection*> InSyncObject) : UE::TScopeUnlock<FCriticalSection>(*InSyncObject)
+	[[nodiscard]] explicit FScopeUnlock(FCriticalSection* InSyncObject) : UE::TScopeUnlock<FCriticalSection>(InSyncObject)
 	{
 	}
 };
