Enable platform I/O dispatcher from -ExecCmds

#rnx

diff --git a/Engine/Source/Runtime/Core/Internal/IO/IoDispatcherConfig.cpp b/Engine/Source/Runtime/Core/Internal/IO/IoDispatcherConfig.cpp
--- a/Engine/Source/Runtime/Core/Internal/IO/IoDispatcherConfig.cpp
+++ b/Engine/Source/Runtime/Core/Internal/IO/IoDispatcherConfig.cpp
@@ -120,6 +120,15 @@ bool IsPlatformIoDispatcherEnabled()
 		FIoDispatcherConfig(const TCHAR* CmdLine)
 		{
 			bPlatformIoDispatcherEnabled = FParse::Param(CmdLine, TEXT("PlatformIo"));
+			if (!bPlatformIoDispatcherEnabled)
+			{
+				FString Line;
+				const bool bShouldStopOnSeparator = false;
+				while (FParse::Value(CmdLine, TEXT("ExecCmds="), Line, bShouldStopOnSeparator, &CmdLine) && !bPlatformIoDispatcherEnabled)
+				{
+					bPlatformIoDispatcherEnabled = Line.Contains(TEXT("s.PlatformIo 1"), ESearchCase::IgnoreCase);
+				}
+			}
 		}
 	} IoDispatcherConfig(FCommandLine::Get());
 
