Remove bogus char logging, which hasn't worked for 11 years.

diff --git a/Engine/Source/Runtime/Core/Private/GenericPlatform/GenericPlatformString.cpp b/Engine/Source/Runtime/Core/Private/GenericPlatform/GenericPlatformString.cpp
--- a/Engine/Source/Runtime/Core/Private/GenericPlatform/GenericPlatformString.cpp
+++ b/Engine/Source/Runtime/Core/Private/GenericPlatform/GenericPlatformString.cpp
@@ -22,16 +22,6 @@ void* FGenericPlatformString::Memcpy(void* Dest, const void* Src, SIZE_T Count)
 	return FMemory::Memcpy(Dest, Src, Count);
 }
 
-namespace
-{
-	void TrimStringAndLogBogusCharsError(FString& SrcStr, const TCHAR* SourceCharName, const TCHAR* DestCharName)
-	{
-		SrcStr.TrimStartInline();
-		// @todo: Put this back in once GLog becomes a #define, or is replaced with GLog::GetLog()
-		//UE_LOG(LogGenericPlatformString, Warning, TEXT("Bad chars found when trying to convert \"%s\" from %s to %s"), *SrcStr, SourceCharName, DestCharName);
-	}
-}
-
 namespace UE::Core::Private
 {
 	/**
@@ -698,61 +688,6 @@ namespace UE::Core::Private
 		}
 		return Dest;
 	}
-
-	template <typename DestEncoding, typename SourceEncoding, typename SourceEndType>
-	void LogBogusCharsImpl(const SourceEncoding* Src, SourceEndType SourceEnd)
-	{
-		static_assert(TIsFixedWidthCharEncoding_V<SourceEncoding>, "Currently unimplemented for non-fixed-width source conversions");
-
-		FString SrcStr;
-		bool    bFoundBogusChars = false;
-		for (; !IsRangeEmpty(Src, SourceEnd); PopFront(Src, SourceEnd))
-		{
-			SourceEncoding SrcCh = *Src;
-			if (!FGenericPlatformString::CanConvertCodepoint<DestEncoding>(SrcCh))
-			{
-				SrcStr += FString::Printf(TEXT("[0x%X]"), (int32)SrcCh);
-				bFoundBogusChars = true;
-			}
-			else if (FGenericPlatformString::CanConvertCodepoint<TCHAR>(SrcCh))
-			{
-				if (TChar<SourceEncoding>::IsLinebreak(SrcCh))
-				{
-					if (bFoundBogusChars)
-					{
-						TrimStringAndLogBogusCharsError(SrcStr, FGenericPlatformString::GetEncodingTypeName<SourceEncoding>(), FGenericPlatformString::GetEncodingTypeName<DestEncoding>());
-						bFoundBogusChars = false;
-					}
-					SrcStr.Empty();
-				}
-				else
-				{
-					SrcStr.AppendChar((TCHAR)SrcCh);
-				}
-			}
-			else
-			{
-				SrcStr.AppendChar((TCHAR)'?');
-			}
-		}
-
-		if (bFoundBogusChars)
-		{
-			TrimStringAndLogBogusCharsError(SrcStr, FGenericPlatformString::GetEncodingTypeName<SourceEncoding>(), FGenericPlatformString::GetEncodingTypeName<DestEncoding>());
-		}
-	}
-}
-
-template <typename DestEncoding, typename SourceEncoding>
-void FGenericPlatformString::LogBogusChars(const SourceEncoding* Src)
-{
-	UE::Core::Private::LogBogusCharsImpl<DestEncoding>(Src, UE::Core::Private::FNullTerminal{});
-}
-
-template <typename DestEncoding, typename SourceEncoding>
-void FGenericPlatformString::LogBogusChars(const SourceEncoding* Src, int32 SrcSize)
-{
-	UE::Core::Private::LogBogusCharsImpl<DestEncoding>(Src, SrcSize);
 }
 
 namespace GenericPlatformStringPrivate
@@ -790,30 +725,3 @@ int32 FGenericPlatformString::Strncmp(const UTF8CHAR* Str1, const WIDECHAR* Str2
 int32 FGenericPlatformString::Strncmp(const ANSICHAR* Str1, const UTF8CHAR* Str2, SIZE_T Count) { return GenericPlatformStringPrivate::StrncmpImpl(Str1, Str2, Count); }
 int32 FGenericPlatformString::Strncmp(const WIDECHAR* Str1, const UTF8CHAR* Str2, SIZE_T Count) { return GenericPlatformStringPrivate::StrncmpImpl(Str1, Str2, Count); }
 int32 FGenericPlatformString::Strncmp(const UTF8CHAR* Str1, const UTF8CHAR* Str2, SIZE_T Count) { return GenericPlatformStringPrivate::StrncmpImpl(Str1, Str2, Count); }
-
-#if !UE_BUILD_DOCS
-template CORE_API void FGenericPlatformString::LogBogusChars<ANSICHAR, WIDECHAR>(const WIDECHAR* Src);
-template CORE_API void FGenericPlatformString::LogBogusChars<ANSICHAR, WIDECHAR>(const WIDECHAR* Src, int32 SrcSize);
-template CORE_API void FGenericPlatformString::LogBogusChars<ANSICHAR, UCS2CHAR>(const UCS2CHAR* Src);
-template CORE_API void FGenericPlatformString::LogBogusChars<ANSICHAR, UCS2CHAR>(const UCS2CHAR* Src, int32 SrcSize);
-template CORE_API void FGenericPlatformString::LogBogusChars<WIDECHAR, ANSICHAR>(const ANSICHAR* Src);
-template CORE_API void FGenericPlatformString::LogBogusChars<WIDECHAR, ANSICHAR>(const ANSICHAR* Src, int32 SrcSize);
-template CORE_API void FGenericPlatformString::LogBogusChars<WIDECHAR, UCS2CHAR>(const UCS2CHAR* Src);
-template CORE_API void FGenericPlatformString::LogBogusChars<WIDECHAR, UCS2CHAR>(const UCS2CHAR* Src, int32 SrcSize);
-template CORE_API void FGenericPlatformString::LogBogusChars<UCS2CHAR, ANSICHAR>(const ANSICHAR* Src);
-template CORE_API void FGenericPlatformString::LogBogusChars<UCS2CHAR, ANSICHAR>(const ANSICHAR* Src, int32 SrcSize);
-template CORE_API void FGenericPlatformString::LogBogusChars<UCS2CHAR, WIDECHAR>(const WIDECHAR* Src);
-template CORE_API void FGenericPlatformString::LogBogusChars<UCS2CHAR, WIDECHAR>(const WIDECHAR* Src, int32 SrcSize);
-template CORE_API void FGenericPlatformString::LogBogusChars<UTF8CHAR, ANSICHAR>(const ANSICHAR* Src);
-template CORE_API void FGenericPlatformString::LogBogusChars<UTF8CHAR, ANSICHAR>(const ANSICHAR* Src, int32 SrcSize);
-template CORE_API void FGenericPlatformString::LogBogusChars<UTF8CHAR, WIDECHAR>(const WIDECHAR* Src);
-template CORE_API void FGenericPlatformString::LogBogusChars<UTF8CHAR, WIDECHAR>(const WIDECHAR* Src, int32 SrcSize);
-template CORE_API void FGenericPlatformString::LogBogusChars<UTF8CHAR, UCS2CHAR>(const UCS2CHAR* Src);
-template CORE_API void FGenericPlatformString::LogBogusChars<UTF8CHAR, UCS2CHAR>(const UCS2CHAR* Src, int32 SrcSize);
-#if PLATFORM_TCHAR_IS_CHAR16
-template CORE_API void FGenericPlatformString::LogBogusChars<wchar_t, char16_t>(const char16_t* Src);
-template CORE_API void FGenericPlatformString::LogBogusChars<wchar_t, char16_t>(const char16_t* Src, int32 SrcSize);
-template CORE_API void FGenericPlatformString::LogBogusChars<char16_t, wchar_t>(const wchar_t* Src);
-template CORE_API void FGenericPlatformString::LogBogusChars<char16_t, wchar_t>(const wchar_t* Src, int32 SrcSize);
-#endif
-#endif
diff --git a/Engine/Source/Runtime/Core/Public/GenericPlatform/GenericPlatformString.h b/Engine/Source/Runtime/Core/Public/GenericPlatform/GenericPlatformString.h
--- a/Engine/Source/Runtime/Core/Public/GenericPlatform/GenericPlatformString.h
+++ b/Engine/Source/Runtime/Core/Public/GenericPlatform/GenericPlatformString.h
@@ -202,8 +202,6 @@ struct FGenericPlatformString : public FGenericPlatformStricmp
 
 					--DestSizeCopy;
 				}
-
-				LogBogusChars<DestEncoding>(Src);
 			}
 
 			return Dest;
@@ -273,8 +271,6 @@ struct FGenericPlatformString : public FGenericPlatformStricmp
 						Dest[I] = UNICODE_BOGUS_CHAR_CODEPOINT;
 					}
 				}
-
-				LogBogusChars<DestEncoding>(Src, Size);
 			}
 
 			return DestSize < SrcSize ? nullptr : Dest + Size;
@@ -353,23 +349,4 @@ private:
 	 * @return      Dest
 	 */
 	static CORE_API void* Memcpy(void* Dest, const void* Src, SIZE_T Count);
-
-
-	/**
-	 * Logs a message about bogus characters which were detected during string conversion.
-	 *
-	 * @param Src     Pointer to the null-terminated string being converted.
-	 */
-	template <typename DestEncoding, typename SourceEncoding>
-	CORE_API static void LogBogusChars(const SourceEncoding* Src);
-
-
-	/**
-	 * Logs a message about bogus characters which were detected during string conversion.
-	 *
-	 * @param Src     Pointer to the possibly-not-null-terminated string being converted.
-	 * @param SrcSize Number of characters in the Src string.
-	 */
-	template <typename DestEncoding, typename SourceEncoding>
-	CORE_API static void LogBogusChars(const SourceEncoding* Src, int32 SrcSize);
 };
