Transform the Right into a Left

#okforversepublic

diff --git a/Engine/Source/Runtime/Core/Public/UObject/FortniteMainBranchObjectVersions.inl b/Engine/Source/Runtime/Core/Public/UObject/FortniteMainBranchObjectVersions.inl
--- a/Engine/Source/Runtime/Core/Public/UObject/FortniteMainBranchObjectVersions.inl
+++ b/Engine/Source/Runtime/Core/Public/UObject/FortniteMainBranchObjectVersions.inl
@@ -607,5 +607,8 @@ FORTNITE_MAIN_VERSION(ConvertGlintDensity, 197)
 // Introduced skinweight validation to avoid render crashes and disappearing simulation meshes
 FORTNITE_MAIN_VERSION(ClothAssetSkinweightsValidation, 198)
 
+// Switching verse from right handed to left handed
+FORTNITE_MAIN_VERSION(VerseRightToLeftHandedness, 199)
+
 // -----<new versions can be added above this line>-------------------------------------------------
 #undef FORTNITE_MAIN_VERSION
\ No newline at end of file
