Change GetMemoryMappingAlignment to 16kb to support 16kb page size on Android

diff --git a/Engine/Source/Runtime/Core/Public/Android/AndroidPlatformProperties.h b/Engine/Source/Runtime/Core/Public/Android/AndroidPlatformProperties.h
--- a/Engine/Source/Runtime/Core/Public/Android/AndroidPlatformProperties.h
+++ b/Engine/Source/Runtime/Core/Public/Android/AndroidPlatformProperties.h
@@ -131,7 +131,9 @@ struct FAndroidPlatformProperties
 
 	static FORCEINLINE int64 GetMemoryMappingAlignment()
 	{
-		return 4096;
+		// Cook for largest page size available.
+		// Data cooked for 16kb page size will work for 4kb page size firmware.
+		return 16384;
 	}
 
 	static FORCEINLINE bool SupportsVirtualTextureStreaming()
