Fix performance regression in editor builds where excess file checks are done for merged modules support

We don't actually need this when merged modules are not available.

#tests Confirmed fix with a partner
#rnx

diff --git a/Engine/Source/Runtime/Core/Private/Modules/ModuleManager.cpp b/Engine/Source/Runtime/Core/Private/Modules/ModuleManager.cpp
--- a/Engine/Source/Runtime/Core/Private/Modules/ModuleManager.cpp
+++ b/Engine/Source/Runtime/Core/Private/Modules/ModuleManager.cpp
@@ -778,7 +778,7 @@ IModuleInterface* FModuleManager::LoadModuleWithFailureReason(const FName InModu
 
 	const auto EnsureModuleFilename = [this, InModuleName, &FoundModulePtr]()
 	{
-#if !IS_MONOLITHIC
+#if UE_MERGED_MODULES
 		ModuleInfoRef ModuleInfo = FoundModulePtr.ToSharedRef();
 		if (ModuleInfo->Filename.IsEmpty() || !FPlatformProcess::ModuleExists(ModuleInfo->Filename))
 		{
@@ -789,7 +789,7 @@ IModuleInterface* FModuleManager::LoadModuleWithFailureReason(const FName InModu
 				ModuleInfo->Filename = MoveTemp(TMap<FName, FString>::TIterator(ModulePathMap).Value());
 			}
 		}
-#endif // !IS_MONOLITHIC
+#endif // UE_MERGED_MODULES
 	};
 
 	if (FoundModulePtr.IsValid())
