Allow TStrings to accept UTF32CHAR* strings.

diff --git a/Engine/Source/Runtime/Core/Private/Containers/String.cpp.inl b/Engine/Source/Runtime/Core/Private/Containers/String.cpp.inl
--- a/Engine/Source/Runtime/Core/Private/Containers/String.cpp.inl
+++ b/Engine/Source/Runtime/Core/Private/Containers/String.cpp.inl
@@ -230,6 +230,7 @@ UE_STRING_CLASS::UE_STRING_CLASS(const ANSICHAR* Str)								{ if (Str) { UE::Co
 UE_STRING_CLASS::UE_STRING_CLASS(const WIDECHAR* Str)								{ if (Str) { UE::Core::Private::ConstructFromCString(Data, Str, 0); } }
 UE_STRING_CLASS::UE_STRING_CLASS(const UTF8CHAR* Str)								{ if (Str) { UE::Core::Private::ConstructFromCString(Data, Str, 0); } }
 UE_STRING_CLASS::UE_STRING_CLASS(const UCS2CHAR* Str)								{ if (Str) { UE::Core::Private::ConstructFromCString(Data, Str, 0); } }
+UE_STRING_CLASS::UE_STRING_CLASS(const UTF32CHAR* Str)								{ if (Str) { UE::Core::Private::ConstructFromCString(Data, Str, 0); } }
 
 // Deprecated
 UE_STRING_CLASS::UE_STRING_CLASS(int32 Len, const ANSICHAR* Str)					{ if (Str && *Str) { UE::Core::Private::ConstructWithLength(Data, Str, Len, 0); } }
@@ -246,14 +247,17 @@ UE_STRING_CLASS UE_STRING_CLASS::ConstructFromPtrSize(const ANSICHAR* Str, int32
 UE_STRING_CLASS UE_STRING_CLASS::ConstructFromPtrSize(const WIDECHAR* Str, int32 Len)								{ UE_STRING_CLASS Result; UE::Core::Private::ConstructWithLength(Result.Data, Str, Len, 0); return Result; }
 UE_STRING_CLASS UE_STRING_CLASS::ConstructFromPtrSize(const UTF8CHAR* Str, int32 Len)								{ UE_STRING_CLASS Result; UE::Core::Private::ConstructWithLength(Result.Data, Str, Len, 0); return Result; }
 UE_STRING_CLASS UE_STRING_CLASS::ConstructFromPtrSize(const UCS2CHAR* Str, int32 Len)								{ UE_STRING_CLASS Result; UE::Core::Private::ConstructWithLength(Result.Data, Str, Len, 0); return Result; }
+UE_STRING_CLASS UE_STRING_CLASS::ConstructFromPtrSize(const UTF32CHAR* Str, int32 Len)								{ UE_STRING_CLASS Result; UE::Core::Private::ConstructWithLength(Result.Data, Str, Len, 0); return Result; }
 UE_STRING_CLASS UE_STRING_CLASS::ConstructWithSlack(const ANSICHAR* Str, int32 ExtraSlack)							{ UE_STRING_CLASS Result; UE::Core::Private::ConstructFromCString(Result.Data, Str, ExtraSlack); return Result; }
 UE_STRING_CLASS UE_STRING_CLASS::ConstructWithSlack(const WIDECHAR* Str, int32 ExtraSlack)							{ UE_STRING_CLASS Result; UE::Core::Private::ConstructFromCString(Result.Data, Str, ExtraSlack); return Result; }
 UE_STRING_CLASS UE_STRING_CLASS::ConstructWithSlack(const UTF8CHAR* Str, int32 ExtraSlack)							{ UE_STRING_CLASS Result; UE::Core::Private::ConstructFromCString(Result.Data, Str, ExtraSlack); return Result; }
 UE_STRING_CLASS UE_STRING_CLASS::ConstructWithSlack(const UCS2CHAR* Str, int32 ExtraSlack)							{ UE_STRING_CLASS Result; UE::Core::Private::ConstructFromCString(Result.Data, Str, ExtraSlack); return Result; }
+UE_STRING_CLASS UE_STRING_CLASS::ConstructWithSlack(const UTF32CHAR* Str, int32 ExtraSlack)							{ UE_STRING_CLASS Result; UE::Core::Private::ConstructFromCString(Result.Data, Str, ExtraSlack); return Result; }
 UE_STRING_CLASS UE_STRING_CLASS::ConstructFromPtrSizeWithSlack(const ANSICHAR* Str, int32 Len, int32 ExtraSlack)	{ UE_STRING_CLASS Result; UE::Core::Private::ConstructWithLength(Result.Data, Str, Len, ExtraSlack); return Result; }
 UE_STRING_CLASS UE_STRING_CLASS::ConstructFromPtrSizeWithSlack(const WIDECHAR* Str, int32 Len, int32 ExtraSlack)	{ UE_STRING_CLASS Result; UE::Core::Private::ConstructWithLength(Result.Data, Str, Len, ExtraSlack); return Result; }
 UE_STRING_CLASS UE_STRING_CLASS::ConstructFromPtrSizeWithSlack(const UTF8CHAR* Str, int32 Len, int32 ExtraSlack)	{ UE_STRING_CLASS Result; UE::Core::Private::ConstructWithLength(Result.Data, Str, Len, ExtraSlack); return Result; }
 UE_STRING_CLASS UE_STRING_CLASS::ConstructFromPtrSizeWithSlack(const UCS2CHAR* Str, int32 Len, int32 ExtraSlack)	{ UE_STRING_CLASS Result; UE::Core::Private::ConstructWithLength(Result.Data, Str, Len, ExtraSlack); return Result; }
+UE_STRING_CLASS UE_STRING_CLASS::ConstructFromPtrSizeWithSlack(const UTF32CHAR* Str, int32 Len, int32 ExtraSlack)	{ UE_STRING_CLASS Result; UE::Core::Private::ConstructWithLength(Result.Data, Str, Len, ExtraSlack); return Result; }
 
 UE_STRING_CLASS& UE_STRING_CLASS::operator=(const ElementType* Other)
 {
@@ -406,6 +410,12 @@ void UE_STRING_CLASS::AppendChars(const UTF8CHAR* Str, int32 Count)
 	UE::Core::Private::AppendCharacters(Data, Str, Count);
 }
 
+void UE_STRING_CLASS::AppendChars(const UTF32CHAR* Str, int32 Count)
+{
+	CheckInvariants();
+	UE::Core::Private::AppendCharacters(Data, Str, Count);
+}
+
 void UE_STRING_CLASS::TrimToNullTerminator()
 {
 	if (Data.Num())
diff --git a/Engine/Source/Runtime/Core/Private/Misc/CString.cpp b/Engine/Source/Runtime/Core/Private/Misc/CString.cpp
--- a/Engine/Source/Runtime/Core/Private/Misc/CString.cpp
+++ b/Engine/Source/Runtime/Core/Private/Misc/CString.cpp
@@ -183,3 +183,14 @@ bool FToBoolHelper::FromCStringUtf8( const UTF8CHAR* String )
 	return FToBoolHelper::FromCStringWide( StringCast<WIDECHAR>(String).Get() );
 #endif
 }
+
+int32 UE::Core::Private::Strlen32(const UTF32CHAR* String)
+{
+	int32 Result = 0;
+	while (*String)
+	{
+		++Result;
+		++String;
+	}
+	return Result;
+}
diff --git a/Engine/Source/Runtime/Core/Public/Containers/UnrealString.h.inl b/Engine/Source/Runtime/Core/Public/Containers/UnrealString.h.inl
--- a/Engine/Source/Runtime/Core/Public/Containers/UnrealString.h.inl
+++ b/Engine/Source/Runtime/Core/Public/Containers/UnrealString.h.inl
@@ -102,6 +102,7 @@ public:
 	[[nodiscard]] CORE_API UE_STRING_CLASS(const WIDECHAR* Str);
 	[[nodiscard]] CORE_API UE_STRING_CLASS(const UTF8CHAR* Str);
 	[[nodiscard]] CORE_API UE_STRING_CLASS(const UCS2CHAR* Str);
+	[[nodiscard]] CORE_API UE_STRING_CLASS(const UTF32CHAR* Str);
 
 	/** Construct from null-terminated C substring or nullptr */
 	UE_STRING_DEPRECATED(5.4, "This constructor has been deprecated - please use " PREPROCESSOR_TO_STRING(UE_STRING_CLASS) "::ConstructFromPtrSize(Ptr, Size) instead.") [[nodiscard]] CORE_API UE_STRING_CLASS(int32 Len, const ANSICHAR* Str);
@@ -120,18 +121,21 @@ public:
 	[[nodiscard]] static CORE_API UE_STRING_CLASS ConstructWithSlack(const WIDECHAR* Str, int32 ExtraSlack);
 	[[nodiscard]] static CORE_API UE_STRING_CLASS ConstructWithSlack(const UTF8CHAR* Str, int32 ExtraSlack);
 	[[nodiscard]] static CORE_API UE_STRING_CLASS ConstructWithSlack(const UCS2CHAR* Str, int32 ExtraSlack);
+	[[nodiscard]] static CORE_API UE_STRING_CLASS ConstructWithSlack(const UTF32CHAR* Str, int32 ExtraSlack);
 
 	/** Construct from a buffer.  If the buffer contains zeros, these will be present in the constructed string. */
 	[[nodiscard]] static CORE_API UE_STRING_CLASS ConstructFromPtrSize(const ANSICHAR* Str, int32 Size);
 	[[nodiscard]] static CORE_API UE_STRING_CLASS ConstructFromPtrSize(const WIDECHAR* Str, int32 Size);
 	[[nodiscard]] static CORE_API UE_STRING_CLASS ConstructFromPtrSize(const UTF8CHAR* Str, int32 Size);
 	[[nodiscard]] static CORE_API UE_STRING_CLASS ConstructFromPtrSize(const UCS2CHAR* Str, int32 Size);
+	[[nodiscard]] static CORE_API UE_STRING_CLASS ConstructFromPtrSize(const UTF32CHAR* Str, int32 Size);
 
 	/** Construct from a buffer with extra slack on top of original string length.  If the buffer contains zeros, these will be present in the constructed string. */
 	[[nodiscard]] static CORE_API UE_STRING_CLASS ConstructFromPtrSizeWithSlack(const ANSICHAR* Str, int32 Size, int32 ExtraSlack);
 	[[nodiscard]] static CORE_API UE_STRING_CLASS ConstructFromPtrSizeWithSlack(const WIDECHAR* Str, int32 Size, int32 ExtraSlack);
 	[[nodiscard]] static CORE_API UE_STRING_CLASS ConstructFromPtrSizeWithSlack(const UTF8CHAR* Str, int32 Size, int32 ExtraSlack);
 	[[nodiscard]] static CORE_API UE_STRING_CLASS ConstructFromPtrSizeWithSlack(const UCS2CHAR* Str, int32 Size, int32 ExtraSlack);
+	[[nodiscard]] static CORE_API UE_STRING_CLASS ConstructFromPtrSizeWithSlack(const UTF32CHAR* Str, int32 Size, int32 ExtraSlack);
 
 	/** Construct from contiguous range of characters such as a string view or string builder */
 	template <
@@ -402,6 +406,7 @@ public:
 	CORE_API void AppendChars(const WIDECHAR* Str, int32 Count);
 	CORE_API void AppendChars(const UCS2CHAR* Str, int32 Count);
 	CORE_API void AppendChars(const UTF8CHAR* Str, int32 Count);
+	CORE_API void AppendChars(const UTF32CHAR* Str, int32 Count);
 
 	/** Append a string and return a reference to this */
 	template<class CharType>
diff --git a/Engine/Source/Runtime/Core/Public/Misc/CString.h b/Engine/Source/Runtime/Core/Public/Misc/CString.h
--- a/Engine/Source/Runtime/Core/Public/Misc/CString.h
+++ b/Engine/Source/Runtime/Core/Public/Misc/CString.h
@@ -980,10 +980,22 @@ int32 TCString<T>::Strnicmp( const CharType* String1, const CharType* String2, S
 	return FPlatformString::Strnicmp(String1, String2, Count);
 }
 
-template <typename T> FORCEINLINE
-int32 TCString<T>::Strlen( const CharType* String ) 
+namespace UE::Core::Private
+{
+	CORE_API int32 Strlen32(const UTF32CHAR* String);
+}
+
+template <typename T>
+UE_NODEBUG FORCEINLINE int32 TCString<T>::Strlen( const CharType* String ) 
 {
-	return FPlatformString::Strlen(String);
+	if constexpr (std::is_same_v<T, UTF32CHAR>)
+	{
+		return UE::Core::Private::Strlen32(String);
+	}
+	else
+	{
+		return FPlatformString::Strlen(String);
+	}
 }
 
 template <typename T> FORCEINLINE
