Fix invalidated GameUserSettings not getting deleted
- GGameUserSettingsIni is set to "GameUserSettings", not a path, so the delete was silently failing. Instead, use `FConfigBranch::IniPath` so we know we're deleting the correct file if it exists.
Use a consistent log category.

diff --git a/Engine/Source/Runtime/Core/Private/Misc/ConfigCacheIni.cpp b/Engine/Source/Runtime/Core/Private/Misc/ConfigCacheIni.cpp
--- a/Engine/Source/Runtime/Core/Private/Misc/ConfigCacheIni.cpp
+++ b/Engine/Source/Runtime/Core/Private/Misc/ConfigCacheIni.cpp
@@ -668,6 +668,17 @@ static bool SaveConfigFileWrapper(const TCHAR* IniFile, const FString& Contents)
 	return SavedCount > 0 || bLocalWriteSucceeded;
 }
 
+static bool DeleteConfigFileWrapper(const TCHAR* IniFile)
+{
+	TRACE_CPUPROFILER_EVENT_SCOPE(DeleteConfigFileWrapper);
+
+	bool bDeleted = false;
+	FCoreDelegates::TSPreDeleteConfigFileDelegate().Broadcast(IniFile, bDeleted);
+
+	bDeleted |= IFileManager::Get().Delete(IniFile);
+	return bDeleted;
+}
+
 static FConfigCommandStream CalculateDiff(const FConfigFile& First, const FConfigFile& Second, const FString& SingleSection=FString(), const FString& SingleProperty=FString())
 {
 	FConfigCommandStream Diff;
@@ -901,8 +912,6 @@ static bool SaveBranch(FConfigBranch& Branch)
 	if (bBuiltString && Output.Len() > 0)
 	{
 		Output = FString(TEXT(";METADATA=(Diff=true, UseCommands=true)" LINE_TERMINATOR_ANSI)) + Output;
-		
-	
 		return SaveConfigFileWrapper(*Branch.IniPath, Output);
 	}
 
@@ -2231,7 +2240,7 @@ bool FConfigFile::WriteInternal(const FString& Filename, bool bDoRemoteWrite, TM
 	// don't write out non-default configs that are only whitespace
 	if (!bIsADefaultIniWrite && Text.TrimStart().Len() == 0)
 	{
-		IFileManager::Get().Delete(*Filename);
+		DeleteConfigFileWrapper(*Filename);
 		return true;
 	}
 	
@@ -3543,6 +3552,12 @@ bool FConfigBranch::RemoveSection(const TCHAR* Section)
 	return NumRemoved > 0;
 }
 
+bool FConfigBranch::Delete()
+{
+	bool bDeleted = DeleteConfigFileWrapper(*IniPath);
+	return bDeleted;
+}
+
 void FConfigBranch::Shrink()
 {
 	RunOnEachFile([](FConfigFile& File, const FString& Name)
@@ -3812,12 +3827,13 @@ FConfigFile* FConfigCacheIni::FindConfigFile( const FString& Filename )
 		if (Result->bIsSafeUnloaded)
 		{
 			Result->SafeReload();
-		};
+		}
+
 		// track that this branch is being used, so re-set the time
 		if (Result && Result->InactiveTimer >= 0 && GTimeToUnloadConfig > 0)
-	{
+		{
 			Result->InactiveTimer = FPlatformTime::Seconds();
-			UE_LOG(LogConfig, VeryVerbose, TEXT("REsetting InactiveTimer for %s"), *Result->IniName.ToString());
+			UE_LOG(LogConfig, VeryVerbose, TEXT("Resetting InactiveTimer for %s"), *Result->IniName.ToString());
 		}
 		return &Result->InMemoryFile;
 	}
@@ -4463,7 +4479,7 @@ bool FConfigCacheIni::EmptySection( const TCHAR* Section, const FString& Filenam
 				}
 				else
 				{
-					IFileManager::Get().Delete(*Filename);	
+					DeleteConfigFileWrapper(*Filename);
 				}
 			}
 			return true;
diff --git a/Engine/Source/Runtime/Core/Private/Misc/CoreDelegates.cpp b/Engine/Source/Runtime/Core/Private/Misc/CoreDelegates.cpp
--- a/Engine/Source/Runtime/Core/Private/Misc/CoreDelegates.cpp
+++ b/Engine/Source/Runtime/Core/Private/Misc/CoreDelegates.cpp
@@ -198,6 +198,12 @@ TTSMulticastDelegate<void(const TCHAR*, const FString&, int32&)>& FCoreDelegates
 	return Singleton;
 }
 
+TTSMulticastDelegate<void(const TCHAR*, bool&)>& FCoreDelegates::TSPreDeleteConfigFileDelegate()
+{
+	static TTSMulticastDelegate<void(const TCHAR*, bool&)> Singleton;
+	return Singleton;
+}
+
 TTSMulticastDelegate<void(const FConfigFile*)>& FCoreDelegates::TSOnFConfigCreated()
 {
 	static TTSMulticastDelegate<void(const FConfigFile*)> Singleton;
diff --git a/Engine/Source/Runtime/Core/Public/Misc/ConfigCacheIni.h b/Engine/Source/Runtime/Core/Public/Misc/ConfigCacheIni.h
--- a/Engine/Source/Runtime/Core/Public/Misc/ConfigCacheIni.h
+++ b/Engine/Source/Runtime/Core/Public/Misc/ConfigCacheIni.h
@@ -111,7 +111,7 @@ enum class EKnownIniFile : uint8
 	// make an enum entry for each known ini above
 	ENUMERATE_KNOWN_INI_FILES(KNOWN_INI_ENUM)
 	
-	// conventient counter for the above list
+	// convenient counter for the above list
 	NumKnownFiles,
 };
 
@@ -1063,6 +1063,8 @@ enum class DynamicLayerPriority : uint16
 	Hotfix = 80,
 };
 
+// A set of config files with the same "base name" like Engine, Input, or GameUserSettings. Contains the hierarchy of
+// all possible files that will be searched for and the files that actually exist and can be loaded.
 class FConfigBranch
 {
 public:
@@ -1076,7 +1078,7 @@ public:
 	// base name of the branch, like "Engine"
 	FName IniName;
 	
-	// "final" path for the branch like "Saved/Config/WIndows/Engine.ini"
+	// "final" path for the branch like "Saved/Config/Windows/Engine.ini"
 	FString IniPath;
 	
 	FName Platform;
@@ -1097,7 +1099,6 @@ public:
 	TMap<FString, FConfigCommandStream> StaticLayers;
 	using DynamicLayerList = TDoubleLinkedList<FConfigCommandStream*>;
 	DynamicLayerList DynamicLayers;
-	//	TMap<FString, FConfigCommandStream> DynamicLayers;
 	FConfigCommandStream SavedLayer;
 	
 	// cache the static layers so when remaking dynamic layers after removing a dynamic layer it's faster
@@ -1123,7 +1124,7 @@ public:
 	CORE_API bool AddDynamicLayerToHierarchy(const FString& Filename, FConfigModificationTracker* ModificationTracker=nullptr, TSet<FString>* GlobalConfigFileCache=nullptr, TSet<FString>* PluginConfigFileCache=nullptr);
 	CORE_API bool AddDynamicLayersToHierarchy(const TArray<FDynamicLayerInfo>& Layers, FConfigModificationTracker* ModificationTracker=nullptr, TSet<FString>* GlobalConfigFileCache=nullptr, TSet<FString>* PluginConfigFileCache=nullptr, bool bForceFullDynamicLayerUpdate=false);
 
-		UE_DEPRECATED(5.6, "Use AddDynamicLayersToHierarchy that takes a FDynamicLayerInfo list")
+	UE_DEPRECATED(5.6, "Use AddDynamicLayersToHierarchy that takes a FDynamicLayerInfo list")
 	CORE_API bool AddDynamicLayersToHierarchy(const TArray<FString>& Filenames, FName Tag=NAME_None, DynamicLayerPriority Priority=DynamicLayerPriority::Unknown, FConfigModificationTracker* ModificationTracker =nullptr);
 	
 	// Add a preloaded string as a dynamic layer (useful for hotfixing)
@@ -1143,13 +1144,19 @@ public:
 	CORE_API void SafeReload();
 
 	/**
-	  * Removes the section completely from all layers of this branch 
-	  * This is destructive! It will not reload on demand (like SafeUnload does). Use this only when you know you will _never_ need
-	  * the values in the section again!
-	  * @return true if anything was removed
+	 * Removes the section completely from all layers of this branch 
+	 * This is destructive! It will not reload on demand (like SafeUnload does). Use this only when you know you will _never_ need
+	 * the values in the section again!
+	 * @return true if anything was removed
 	 */
 	CORE_API bool RemoveSection(const TCHAR* Section);
 
+	/**
+	 * Deletes the final saved branch file, such as "Saved/Config/Windows/Engine.ini".
+	 * @return Whether the file was successfully deleted.
+	 */
+	CORE_API bool Delete();
+
 	CORE_API void Flush();
 
 	CORE_API void Shrink();
@@ -1291,14 +1298,13 @@ public:
 
 	CORE_API FConfigFile& Add(const FString& Filename, const FConfigFile& File);
 
+	/** Finds an FConfigBranch, using base name or a filename (can pass in NAME_None or empty string as needed) */
+	CORE_API FConfigBranch* FindBranch(FName BaseIniName, const FString& Filename);
 
-    /** Finds an FConfigBranch, using base name or a filename (can pass in NAME_None or empty string as needed) */
-    CORE_API FConfigBranch* FindBranch(FName BaseIniName, const FString& Filename);
-
-	/** Same as FConfigBRanch, but this will not reload a formerly SafeUnload'd branch */
+	/** Same as FindBranch, but this will not reload a formerly SafeUnload'd branch */
 	CORE_API FConfigBranch* FindBranchWithNoReload(FName BaseIniName, const FString& Filename);
 
-    /** Create a new branch for FIlename, and return it */
+	/** Create a new branch for FIlename, and return it */
 	CORE_API FConfigBranch& AddNewBranch(const FString& Filename);
 	
 	CORE_API int32 Remove(const FString& Filename);
@@ -1317,13 +1323,15 @@ public:
 	CORE_API bool GetText( const TCHAR* Section, const TCHAR* Key, FText& Value, const FString& Filename );
 	CORE_API bool GetSection( const TCHAR* Section, TArray<FString>& Result, const FString& Filename );
 	CORE_API bool DoesSectionExist(const TCHAR* Section, const FString& Filename);
+
 	/**
 	 * @param Force Whether to create the Section on Filename if it did not exist previously.
 	 * @param Const If Const (and not Force), then it will not modify File->Dirty. If not Const (or Force is true), then File->Dirty will be set to true.
 	 */
-    UE_DEPRECATED(5.4, "Use GetSection instead, and/or use the new AddToSection, etc APIs to modify sections without retrieving the section. See top of ConfigCacheIni.h for more info.")
-    CORE_API FConfigSection* GetSectionPrivate( const TCHAR* Section, const bool Force, const bool Const, const FString& Filename );
-    CORE_API const FConfigSection* GetSection( const TCHAR* Section, const bool Force, const FString& Filename );
+	UE_DEPRECATED(5.4, "Use GetSection instead, and/or use the new AddToSection, etc APIs to modify sections without retrieving the section. See top of ConfigCacheIni.h for more info.")
+	CORE_API FConfigSection* GetSectionPrivate( const TCHAR* Section, const bool Force, const bool Const, const FString& Filename );
+
+	CORE_API const FConfigSection* GetSection( const TCHAR* Section, const bool Force, const FString& Filename );
 	CORE_API void SetString( const TCHAR* Section, const TCHAR* Key, const TCHAR* Value, const FString& Filename );
 	CORE_API void SetText( const TCHAR* Section, const TCHAR* Key, const FText& Value, const FString& Filename );
 	CORE_API bool RemoveKey( const TCHAR* Section, const TCHAR* Key, const FString& Filename );
diff --git a/Engine/Source/Runtime/Core/Public/Misc/CoreDelegates.h b/Engine/Source/Runtime/Core/Public/Misc/CoreDelegates.h
--- a/Engine/Source/Runtime/Core/Public/Misc/CoreDelegates.h
+++ b/Engine/Source/Runtime/Core/Public/Misc/CoreDelegates.h
@@ -304,6 +304,7 @@ public:
 	static CORE_API TTSMulticastDelegate<void(const TCHAR* IniFilename, int32& ResponderCount)>& TSCountPreLoadConfigFileRespondersDelegate();
 	static CORE_API TTSMulticastDelegate<void(const TCHAR* IniFilename, FString& LoadedContents)>& TSPreLoadConfigFileDelegate();
 	static CORE_API TTSMulticastDelegate<void(const TCHAR* IniFilename, const FString& ContentsToSave, int32& SavedCount)>& TSPreSaveConfigFileDelegate();
+	static CORE_API TTSMulticastDelegate<void(const TCHAR* IniFilename, bool& bOutDeleted)>& TSPreDeleteConfigFileDelegate();
 	static CORE_API TTSMulticastDelegate<void(const FConfigFile*)>& TSOnFConfigCreated();
 	static CORE_API TTSMulticastDelegate<void(const FConfigFile*)>& TSOnFConfigDeleted();
 	static CORE_API TTSMulticastDelegate<void(const TCHAR* IniFilename, const TCHAR* SectionName, const TCHAR* Key)>& TSOnConfigValueRead();
diff --git a/Engine/Source/Runtime/Engine/Private/GameUserSettings.cpp b/Engine/Source/Runtime/Engine/Private/GameUserSettings.cpp
--- a/Engine/Source/Runtime/Engine/Private/GameUserSettings.cpp
+++ b/Engine/Source/Runtime/Engine/Private/GameUserSettings.cpp
@@ -26,6 +26,8 @@
 
 #include UE_INLINE_GENERATED_CPP_BY_NAME(GameUserSettings)
 
+DEFINE_LOG_CATEGORY_STATIC(LogGameUserSettings, Log, All);
+
 extern EWindowMode::Type GetWindowModeType(EWindowMode::Type WindowMode);
 
 static TAutoConsoleVariable<int32> CVar_GameUserSettings_WindowPositionsMaxSize(
@@ -392,7 +394,9 @@ void UGameUserSettings::ValidateSettings()
 				SetDynamicResolutionEnabled(GEngine->GetDynamicResolutionUserSetting());
 			}
 
-			IFileManager::Get().Delete(*GGameUserSettingsIni);
+			FConfigBranch* Branch = GConfig->FindBranch(FName("GameUserSettings"), FString());
+			Branch->Delete();
+
 			LoadSettings(true);
 		}
 	}
@@ -534,11 +538,11 @@ void UGameUserSettings::LoadSettings(bool bForceReload/*=false*/)
         if (OnUpdateGameUserSettingsFileFromCloud.IsBound())
         {
             FString IniFileLocation = FPaths::GeneratedConfigDir() + UGameplayStatics::GetPlatformName() + "/" +  GGameUserSettingsIni + ".ini";
-            UE_LOG(LogTemp, Verbose, TEXT("%s"), *IniFileLocation);
+            UE_LOG(LogGameUserSettings, Verbose, TEXT("%s"), *IniFileLocation);
 
             if (!OnUpdateGameUserSettingsFileFromCloud.Execute(FString(*IniFileLocation)))
             {
-                UE_LOG(LogTemp, Warning, TEXT("Failed to read the ini file from the Cloud interface %s"), *IniFileLocation);
+                UE_LOG(LogGameUserSettings, Warning, TEXT("Failed to read the ini file from the Cloud interface %s"), *IniFileLocation);
             }
         }
         
@@ -606,14 +610,14 @@ void UGameUserSettings::SaveSettings()
     if (OnUpdateCloudDataFromGameUserSettings.IsBound())
     {
         FString IniFileLocation = FPaths::GeneratedConfigDir() + UGameplayStatics::GetPlatformName() + "/" +  GGameUserSettingsIni + ".ini";
-        UE_LOG(LogTemp, Verbose, TEXT("%s"), *IniFileLocation);
+        UE_LOG(LogGameUserSettings, Verbose, TEXT("%s"), *IniFileLocation);
 
         bool bDidSucceed = false;
         bDidSucceed = OnUpdateCloudDataFromGameUserSettings.Execute(FString(*IniFileLocation));
         
         if (!bDidSucceed)
         {
-            UE_LOG(LogTemp, Warning, TEXT("Failed to load the ini file from the Cloud interface %s"), *IniFileLocation);
+            UE_LOG(LogGameUserSettings, Warning, TEXT("Failed to load the ini file from the Cloud interface %s"), *IniFileLocation);
         }
     }
 }
@@ -646,12 +650,12 @@ void UGameUserSettings::LoadConfigIni(bool bForceReload/*=false*/)
 					IFileManager::Get().GetTimeStampPair(*MediumIntegrityPath, *LowIntegrityPath, MediumIntegrityTimestamp, LowIntegrityTimestamp);
 					if (MediumIntegrityTimestamp > LowIntegrityTimestamp)
 					{
-						UE_LOG(LogTemp, Log, TEXT("Migrating config file to low integrity. %s to %s"), *MediumIntegrityPath, *LowIntegrityPath);
+						UE_LOG(LogGameUserSettings, Log, TEXT("Migrating config file to low integrity. %s to %s"), *MediumIntegrityPath, *LowIntegrityPath);
 						IFileManager::Get().Copy(*LowIntegrityPath, *MediumIntegrityPath, true);
 					}
 					else
 					{
-						UE_LOG(LogTemp, Log, TEXT("Skipping config file migration to low integrity. %s is newer than %s"), *LowIntegrityPath, *MediumIntegrityPath);
+						UE_LOG(LogGameUserSettings, Log, TEXT("Skipping config file migration to low integrity. %s is newer than %s"), *LowIntegrityPath, *MediumIntegrityPath);
 					}
 				}
 			}
@@ -805,7 +809,7 @@ void UGameUserSettings::SetWindowPosition(int32 WinX, int32 WinY, int32 SaveSlot
 	}
 	else
 	{
-		UE_LOG(LogEngine, Warning, TEXT("UGameUserSettings::SetWindowPosition attempted to save window position in invalid slot %d. Maximum slot index is %d."),
+		UE_LOG(LogGameUserSettings, Warning, TEXT("UGameUserSettings::SetWindowPosition attempted to save window position in invalid slot %d. Maximum slot index is %d."),
 			SaveSlotIndex, CVar_GameUserSettings_WindowPositionsMaxSize.GetValueOnAnyThread() - 1);
 	}
 }
@@ -1053,7 +1057,7 @@ void UGameUserSettings::EnableHDRDisplayOutputInternal(bool bEnable, int32 Displ
 	{
 		if (bEnable && !(GRHISupportsHDROutput && IsHDRAllowed()))
 		{
-			UE_LOG(LogConsoleResponse, Display, TEXT("Tried to enable HDR display output but unsupported or disallowed, forcing off."));
+			UE_LOG(LogGameUserSettings, Display, TEXT("Tried to enable HDR display output but unsupported or disallowed, forcing off."));
 			bEnable = false;
 		}
 
@@ -1090,4 +1094,3 @@ bool UGameUserSettings::IsHDREnabled() const
 {
 	return bUseHDRDisplayOutput;
 }
-
