Removed UE_NODEBUG from non-functions and switched the order of its use when paired with UE_INTRINSIC_CAST, as it doesn't compile under MSVC when UE_NODEBUG is defined to be a __declspec.

#rnx

diff --git a/Engine/Source/Runtime/Core/Public/Containers/Array.h b/Engine/Source/Runtime/Core/Public/Containers/Array.h
--- a/Engine/Source/Runtime/Core/Public/Containers/Array.h
+++ b/Engine/Source/Runtime/Core/Public/Containers/Array.h
@@ -346,14 +346,14 @@ namespace UE::Core::Private
 	//     TArray<FThing> Arr; // this will cause errors without this workaround
 	// };
 	template <typename DestType, typename SourceType>
-	UE_NODEBUG constexpr bool TArrayElementsAreCompatible_V = std::disjunction_v<std::is_same<DestType, std::decay_t<SourceType>>, std::is_constructible<DestType, SourceType>>;
+	constexpr bool TArrayElementsAreCompatible_V = std::disjunction_v<std::is_same<DestType, std::decay_t<SourceType>>, std::is_constructible<DestType, SourceType>>;
 
 	template <typename ElementType, typename AllocatorType>
 	UE_NODEBUG static char (&ResolveIsTArrayPtr(const volatile TArray<ElementType, AllocatorType>*))[2];
 	UE_NODEBUG static char(&ResolveIsTArrayPtr(...))[1];
 
 	template <typename T>
-	UE_NODEBUG constexpr bool TIsTArrayOrDerivedFromTArray_V = sizeof(ResolveIsTArrayPtr((T*)nullptr)) == 2;
+	constexpr bool TIsTArrayOrDerivedFromTArray_V = sizeof(ResolveIsTArrayPtr((T*)nullptr)) == 2;
 
 	[[noreturn]] CORE_API void OnInvalidArrayNum(unsigned long long NewNum);
 
@@ -974,7 +974,7 @@ public:
 	///////////////////////////////////////////////
 	// Start - intrusive TOptional<TArray> state //
 	///////////////////////////////////////////////
-	UE_NODEBUG constexpr static bool bHasIntrusiveUnsetOptionalState = true;
+	constexpr static bool bHasIntrusiveUnsetOptionalState = true;
 	using IntrusiveUnsetOptionalStateType = TArray;
 
 	UE_NODEBUG [[nodiscard]] explicit TArray(FIntrusiveUnsetOptionalState Tag)
@@ -4007,12 +4007,12 @@ struct TIsContiguousContainer<TArray<T, AllocatorType>>
 /**
  * Trait which determines whether or not a type is a TArray.
  */
-template <typename T> UE_NODEBUG constexpr bool TIsTArray_V = false;
+template <typename T> constexpr bool TIsTArray_V = false;
 
-template <typename InElementType, typename InAllocatorType> UE_NODEBUG constexpr bool TIsTArray_V<               TArray<InElementType, InAllocatorType>> = true;
-template <typename InElementType, typename InAllocatorType> UE_NODEBUG constexpr bool TIsTArray_V<const          TArray<InElementType, InAllocatorType>> = true;
-template <typename InElementType, typename InAllocatorType> UE_NODEBUG constexpr bool TIsTArray_V<      volatile TArray<InElementType, InAllocatorType>> = true;
-template <typename InElementType, typename InAllocatorType> UE_NODEBUG constexpr bool TIsTArray_V<const volatile TArray<InElementType, InAllocatorType>> = true;
+template <typename InElementType, typename InAllocatorType> constexpr bool TIsTArray_V<               TArray<InElementType, InAllocatorType>> = true;
+template <typename InElementType, typename InAllocatorType> constexpr bool TIsTArray_V<const          TArray<InElementType, InAllocatorType>> = true;
+template <typename InElementType, typename InAllocatorType> constexpr bool TIsTArray_V<      volatile TArray<InElementType, InAllocatorType>> = true;
+template <typename InElementType, typename InAllocatorType> constexpr bool TIsTArray_V<const volatile TArray<InElementType, InAllocatorType>> = true;
 
 template <typename T>
 struct TIsTArray
diff --git a/Engine/Source/Runtime/Core/Public/Misc/AssertionMacros.h b/Engine/Source/Runtime/Core/Public/Misc/AssertionMacros.h
--- a/Engine/Source/Runtime/Core/Public/Misc/AssertionMacros.h
+++ b/Engine/Source/Runtime/Core/Public/Misc/AssertionMacros.h
@@ -398,7 +398,7 @@ RetType UE_COLD UE_DEBUG_SECTION DispatchCheckVerify (InnerType&& Inner, ArgType
 	} // UE::Assert::Private
 
 	template <uint64 Uid>
-	UE_NODEBUG std::atomic<bool> bGEnsureHasExecuted = false;
+	std::atomic<bool> bGEnsureHasExecuted = false;
 
 	UE_NODEBUG constexpr uint32 FileHashForEnsure(const char* Filename)
 	{
diff --git a/Engine/Source/Runtime/Core/Public/Templates/UnrealTemplate.h b/Engine/Source/Runtime/Core/Public/Templates/UnrealTemplate.h
--- a/Engine/Source/Runtime/Core/Public/Templates/UnrealTemplate.h
+++ b/Engine/Source/Runtime/Core/Public/Templates/UnrealTemplate.h
@@ -546,7 +546,7 @@ template <typename T> struct TRemovePointer<T*> { typedef T Type; };
  * const object, because we would prefer to be informed when MoveTemp will have no effect.
  */
 template <typename T>
-UE_NODEBUG UE_INTRINSIC_CAST FORCEINLINE constexpr std::remove_reference_t<T>&& MoveTemp(T&& Obj) noexcept
+UE_INTRINSIC_CAST UE_NODEBUG FORCEINLINE constexpr std::remove_reference_t<T>&& MoveTemp(T&& Obj) noexcept
 {
 	using CastType = std::remove_reference_t<T>;
 
@@ -564,7 +564,7 @@ UE_NODEBUG UE_INTRINSIC_CAST FORCEINLINE constexpr std::remove_reference_t<T>&&
  * where you can but not stop compilation.
  */
 template <typename T>
-UE_NODEBUG UE_INTRINSIC_CAST FORCEINLINE constexpr std::remove_reference_t<T>&& MoveTempIfPossible(T&& Obj) noexcept
+UE_INTRINSIC_CAST UE_NODEBUG FORCEINLINE constexpr std::remove_reference_t<T>&& MoveTempIfPossible(T&& Obj) noexcept
 {
 	using CastType = std::remove_reference_t<T>;
 	return (CastType&&)Obj;
@@ -610,13 +610,13 @@ FORCEINLINE constexpr std::decay_t<T> CopyTempIfNecessary(T&& Val)
  * This is UE's equivalent of std::forward.
  */
 template <typename T>
-UE_NODEBUG UE_INTRINSIC_CAST FORCEINLINE constexpr T&& Forward(std::remove_reference_t<T>& Obj) noexcept
+UE_INTRINSIC_CAST UE_NODEBUG FORCEINLINE constexpr T&& Forward(std::remove_reference_t<T>& Obj) noexcept
 {
 	return (T&&)Obj;
 }
 
 template <typename T>
-UE_NODEBUG UE_INTRINSIC_CAST FORCEINLINE constexpr T&& Forward(std::remove_reference_t<T>&& Obj) noexcept
+UE_INTRINSIC_CAST UE_NODEBUG FORCEINLINE constexpr T&& Forward(std::remove_reference_t<T>&& Obj) noexcept
 {
 	return (T&&)Obj;
 }
@@ -783,7 +783,7 @@ template <
 	typename Base
 	UE_REQUIRES(std::is_convertible_v<std::remove_reference_t<T>*, const volatile Base*>)
 >
-UE_NODEBUG UE_INTRINSIC_CAST FORCEINLINE decltype(auto) ForwardAsBase(std::remove_reference_t<T>& Obj)
+UE_INTRINSIC_CAST UE_NODEBUG FORCEINLINE decltype(auto) ForwardAsBase(std::remove_reference_t<T>& Obj)
 {
 	return (TCopyQualifiersAndRefsFromTo_T<T&&, Base>)Obj;
 }
@@ -793,7 +793,7 @@ template <
 	typename Base
 	UE_REQUIRES(std::is_convertible_v<std::remove_reference_t<T>*, const volatile Base*>)
 >
-UE_NODEBUG UE_INTRINSIC_CAST FORCEINLINE decltype(auto) ForwardAsBase(std::remove_reference_t<T>&& Obj)
+UE_INTRINSIC_CAST UE_NODEBUG FORCEINLINE decltype(auto) ForwardAsBase(std::remove_reference_t<T>&& Obj)
 {
 	return (TCopyQualifiersAndRefsFromTo_T<T&&, Base>)Obj;
 }
