Added missing include for FString.

#rnx

diff --git a/Engine/Source/Runtime/Core/Tests/Templates/SortingTests.cpp b/Engine/Source/Runtime/Core/Tests/Templates/SortingTests.cpp
--- a/Engine/Source/Runtime/Core/Tests/Templates/SortingTests.cpp
+++ b/Engine/Source/Runtime/Core/Tests/Templates/SortingTests.cpp
@@ -3,6 +3,7 @@
 #if WITH_TESTS
 
 #include "Containers/Array.h"
+#include "Containers/UnrealString.h"
 #include "Templates/Sorting.h"
 #include "Tests/TestHarnessAdapter.h"
 
