Removed some elaborated type specifiers which were confusing IntelliSense.

#rnx

diff --git a/Engine/Source/Runtime/Core/Public/Serialization/Archive.h b/Engine/Source/Runtime/Core/Public/Serialization/Archive.h
--- a/Engine/Source/Runtime/Core/Public/Serialization/Archive.h
+++ b/Engine/Source/Runtime/Core/Public/Serialization/Archive.h
@@ -36,6 +36,8 @@ class UObject;
 namespace Verse { struct VCell; }
 namespace UE::Serialization{ class FEditorBulkData; }
 struct FArchiveSerializedPropertyChain;
+struct FBulkDataSerializationParams;
+struct FCompressedChunk;
 template <typename TEnum> class TEnumAsByte;
 template <typename FuncType> class TFunction;
 template <typename T> struct TObjectPtr;
@@ -1732,7 +1734,7 @@ public:
 	 *
 	 * @return true if the bulk data was serialized, false to fallback to default serialization (inline)
 	 */
-	virtual bool SerializeBulkData(FBulkData& BulkData, const struct FBulkDataSerializationParams& Params)
+	virtual bool SerializeBulkData(FBulkData& BulkData, const FBulkDataSerializationParams& Params)
 	{ 
 		return false;
 	}
@@ -1785,7 +1787,7 @@ public:
 	 *
 	 * @return true if archive supports translating offsets & uncompressing on read, false otherwise
 	 */
-	virtual bool SetCompressionMap(TArray<struct FCompressedChunk>* CompressedChunks, ECompressionFlags CompressionFlags)
+	virtual bool SetCompressionMap(TArray<FCompressedChunk>* CompressedChunks, ECompressionFlags CompressionFlags)
 	{
 		return false;
 	}
diff --git a/Engine/Source/Runtime/Core/Public/Serialization/ArchiveProxy.h b/Engine/Source/Runtime/Core/Public/Serialization/ArchiveProxy.h
--- a/Engine/Source/Runtime/Core/Public/Serialization/ArchiveProxy.h
+++ b/Engine/Source/Runtime/Core/Public/Serialization/ArchiveProxy.h
@@ -6,6 +6,10 @@
 #include "Internationalization/TextNamespaceFwd.h"
 #include "Serialization/Archive.h"
 
+class FBulkData;
+struct FBulkDataSerializationParams;
+struct FCompressedChunk;
+
 /**
  * Base class for archive proxies.
  *
@@ -154,7 +158,7 @@ public:
 		InnerArchive.DetachBulkData(BulkData, bEnsureBulkDataIsLoaded);
 	}
 
-	virtual bool SerializeBulkData(class FBulkData& BulkData, const struct FBulkDataSerializationParams& Params) override
+	virtual bool SerializeBulkData(FBulkData& BulkData, const FBulkDataSerializationParams& Params) override
 	{
 		return InnerArchive.SerializeBulkData(BulkData, Params);
 	}
diff --git a/Engine/Source/Runtime/Engine/Classes/Engine/Texture.h b/Engine/Source/Runtime/Engine/Classes/Engine/Texture.h
--- a/Engine/Source/Runtime/Engine/Classes/Engine/Texture.h
+++ b/Engine/Source/Runtime/Engine/Classes/Engine/Texture.h
@@ -37,6 +37,7 @@ class ITargetPlatform;
 class ITargetPlatformSettings;
 class UAssetUserData;
 struct FPropertyChangedEvent;
+struct FTextureBuildSettings;
 
 #if WITH_EDITORONLY_DATA
 namespace UE::DerivedData { struct FValueId; }
@@ -1088,8 +1089,8 @@ public:
 	// in the ddc, it is used, otherwise it is built.
 	void Cache(
 		class UTexture& InTexture,
-		const struct FTextureBuildSettings* InSettingsPerLayerFetchFirst,
-		const struct FTextureBuildSettings* InSettingsPerLayerFetchOrBuild,
+		const FTextureBuildSettings* InSettingsPerLayerFetchFirst,
+		const FTextureBuildSettings* InSettingsPerLayerFetchOrBuild,
 		const FTexturePlatformData::FTextureEncodeResultMetadata* OutResultMetadataPerLayerFetchFirst,
 		const FTexturePlatformData::FTextureEncodeResultMetadata* OutResultMetadataPerLayerFetchOrBuild,
 		uint32 InFlags,
