[TSAN on Windows]
- Add support for virtual table pointer update during destruction to remove a few false positives
- Add unittest for both harmful and benign vptr race
- Make ToggleFilterDetailedLogOnAddress global instead of per thread as it makes more sense
- Minor build scripts improvements
- Improve application shutdown robustness

diff --git a/Engine/Programs/UnrealInstrumentation/Compiler/CMakeLists.txt b/Engine/Programs/UnrealInstrumentation/Compiler/CMakeLists.txt
--- a/Engine/Programs/UnrealInstrumentation/Compiler/CMakeLists.txt
+++ b/Engine/Programs/UnrealInstrumentation/Compiler/CMakeLists.txt
@@ -11,6 +11,14 @@ if(DEFINED CMAKE_CONFIGURATION_TYPES)
   message(FATAL_ERROR "Cannot use multi-configuration generator '${CMAKE_GENERATOR}' with the LLVM builds (recommend you use 'Ninja')!")
 endif()
 
+set(BUILD_TYPE_TO_USE "RELEASE")
+
+if(NOT "${CMAKE_BUILD_TYPE}" STREQUAL "${BUILD_TYPE_TO_USE}")
+  message(WARNING "CMAKE_BUILD_TYPE was '${CMAKE_BUILD_TYPE}'. Overriding to '${BUILD_TYPE_TO_USE}'!")
+
+  set(CMAKE_BUILD_TYPE "${BUILD_TYPE_TO_USE}" CACHE STRING "" FORCE)
+endif()
+
 cmake_minimum_required(VERSION 3.19)
 project(BuildLLVMs)
 
@@ -46,11 +54,17 @@ include(FetchContent)
 # directory layout on disk.
 set_property(DIRECTORY PROPERTY EP_BASE "${CMAKE_CURRENT_BINARY_DIR}")
 
-set(CAREFULLY_REDIST_DIR "$ENV{UE_SDKS_ROOT}/HostWin64")
+set(UE_SDKS_ROOT $ENV{UE_SDKS_ROOT})
+
+if("${UE_SDKS_ROOT}" STREQUAL "")
+  message(FATAL_ERROR "The environment variable 'UE_SDKS_ROOT' was not set!")
+endif()
+
+set(CAREFULLY_REDIST_DIR "${UE_SDKS_ROOT}/HostWin64")
 set(CAREFULLY_REDIST_WIN64_DIR "${CAREFULLY_REDIST_DIR}/Win64")
 file(TO_CMAKE_PATH "${CAREFULLY_REDIST_WIN64_DIR}" CAREFULLY_REDIST_WIN64_DIR)
 
-set(BOOTSTRAP_DIR $ENV{UE_SDKS_ROOT}/HostWin64/Win64/LLVM/18.1.8)
+set(BOOTSTRAP_DIR ${UE_SDKS_ROOT}/HostWin64/Win64/LLVM/18.1.8)
 file(TO_CMAKE_PATH "${BOOTSTRAP_DIR}" BOOTSTRAP_DIR)
 
 if(WIN32)
@@ -132,7 +146,7 @@ function(build)
     SOURCE_SUBDIR "llvm"
     CMAKE_ARGS
       -DUNREALINSTRUMENTATION_DIR=${UNREALINSTRUMENTATION_DIR}
-      -DCMAKE_BUILD_TYPE=RELEASE
+      -DCMAKE_BUILD_TYPE:STRING=${CMAKE_BUILD_TYPE}
       -DCMAKE_INSTALL_PREFIX=${CMAKE_CURRENT_BINARY_DIR}/Install/llvm-${NAME}
       -DLLVM_ENABLE_PROJECTS=clang$<SEMICOLON>lld
       -DLLVM_TARGETS_TO_BUILD=X86
@@ -145,7 +159,7 @@ function(build)
       -DLLVM_ENABLE_PDB=OFF
 
       # Always statically link the CRT on Windows.
-      -DLLVM_USE_CRT_RELEASE=MT
+      -DLLVM_USE_CRT_${CMAKE_BUILD_TYPE}=MT
       -DCMAKE_MSVC_RUNTIME_LIBRARY=MultiThreaded
 
       # Force C++17
@@ -241,7 +255,7 @@ set(DEPLOY_VERSION "18")
 
 option(DO_LTO "Should we do LTO or not?" OFF)
 
-build(NAME "llvm18" GIT_VERSION "${LLVM18_TAG}")
+build(NAME "llvm18" GIT_VERSION "${LLVM18_TAG}" PATCH_PROJECT_TO_USE "llvm${DEPLOY_VERSION}")
 
 # Deploy only this version to the P4 location!
 deploy(NAME "llvm${DEPLOY_VERSION}" VERSION "${DEPLOY_VERSION}")
diff --git a/Engine/Programs/UnrealInstrumentation/Compiler/CustomMemoryInstrumentation.cpp b/Engine/Programs/UnrealInstrumentation/Compiler/CustomMemoryInstrumentation.cpp
--- a/Engine/Programs/UnrealInstrumentation/Compiler/CustomMemoryInstrumentation.cpp
+++ b/Engine/Programs/UnrealInstrumentation/Compiler/CustomMemoryInstrumentation.cpp
@@ -400,6 +400,15 @@ void CustomMemoryInstrumentationPass::cacheInstrumentationFunctions(Module &M) {
   InstrumentFunctions.push_back(&FuncEntryInstrumentFunction);
   InstrumentFunctions.push_back(&FuncExitInstrumentFunction);
 
+  // Virtual Ptr Load and Store
+  StoreVPtrInstrumentFunction =
+      M.getOrInsertFunction("__Instrument_VPtr_Store", Attrs, Builder.getVoidTy(), 
+		  Builder.getPtrTy(), Builder.getPtrTy());
+
+  LoadVPtrInstrumentFunction =
+      M.getOrInsertFunction("__Instrument_VPtr_Load", Attrs, Builder.getVoidTy(),
+                            Builder.getPtrTy());
+
   // Non-atomic loads/stores.
   StoreInstrumentFunction =
       M.getOrInsertFunction("__Instrument_Store", Attrs, Builder.getVoidTy(),
@@ -775,11 +784,22 @@ bool CustomMemoryInstrumentationPass::instrumentStore(StoreInst *Inst,
   InstrumentationIRBuilder Builder(Inst);
 
   Value *Addr = Inst->getPointerOperand();
-  // swifterror memory addresses are mem2reg promoted by instruction selection.
-  // As such they cannot have regular uses like an instrumentation function and
-  // it makes no sense to track them as memory.
-  if (Addr->isSwiftError())
-    return false;
+
+  // Special case for virtual table pointer updates.
+  if (MDNode *Metadata = Inst->getMetadata(LLVMContext::MD_tbaa)) {
+    if (Metadata->isTBAAVtableAccess()) {
+      Value *ValueOperand = Inst->getValueOperand();
+      if (isa<VectorType>(ValueOperand->getType()))
+        ValueOperand = Builder.CreateExtractElement(
+            ValueOperand, ConstantInt::get(Builder.getInt32Ty(), 0));
+
+      if (ValueOperand->getType()->isIntegerTy())
+        ValueOperand = Builder.CreateIntToPtr(ValueOperand, Builder.getPtrTy());
+
+	  Builder.CreateCall(StoreVPtrInstrumentFunction, {Addr, ValueOperand});
+      return true;
+    }
+  }
 
   Value *Ptr =
       Builder.CreateCast(Instruction::CastOps::PtrToInt,
@@ -808,6 +828,14 @@ bool CustomMemoryInstrumentationPass::instrumentStore(StoreInst *Inst,
 bool CustomMemoryInstrumentationPass::instrumentLoad(LoadInst *Inst,
                                                      bool SkipNonAtomics) {
   InstrumentationIRBuilder Builder(Inst);
+  
+  // Special case for virtual table pointer reads.
+  if (MDNode *Metadata = Inst->getMetadata(LLVMContext::MD_tbaa)) {
+    if (Metadata->isTBAAVtableAccess()) {
+      Builder.CreateCall(LoadVPtrInstrumentFunction, Inst->getPointerOperand());
+      return true;
+    }
+  }
 
   Value *Ptr =
       Builder.CreateCast(Instruction::CastOps::PtrToInt,
diff --git a/Engine/Programs/UnrealInstrumentation/Compiler/CustomMemoryInstrumentation.h b/Engine/Programs/UnrealInstrumentation/Compiler/CustomMemoryInstrumentation.h
--- a/Engine/Programs/UnrealInstrumentation/Compiler/CustomMemoryInstrumentation.h
+++ b/Engine/Programs/UnrealInstrumentation/Compiler/CustomMemoryInstrumentation.h
@@ -207,6 +207,8 @@ private:
 
   FunctionCallee StoreInstrumentFunction;
   FunctionCallee LoadInstrumentFunction;
+  FunctionCallee StoreVPtrInstrumentFunction;
+  FunctionCallee LoadVPtrInstrumentFunction;
   FunctionCallee StoreRangeInstrumentFunction;
   FunctionCallee LoadRangeInstrumentFunction;
   
diff --git a/Engine/Source/Runtime/Core/Internal/Sanitizer/RaceDetectorTypes.h b/Engine/Source/Runtime/Core/Internal/Sanitizer/RaceDetectorTypes.h
--- a/Engine/Source/Runtime/Core/Internal/Sanitizer/RaceDetectorTypes.h
+++ b/Engine/Source/Runtime/Core/Internal/Sanitizer/RaceDetectorTypes.h
@@ -101,6 +101,7 @@ namespace UE::Sanitizer::RaceDetector {
 		ACCESS_TYPE_READ = 0b1,
 		ACCESS_TYPE_WRITE = 0b10,
 		ACCESS_TYPE_ATOMIC = 0b100,
+		ACCESS_TYPE_VPTR = 0b1000,
 
 		ACCESS_TYPE_ATOMIC_READ = ACCESS_TYPE_ATOMIC | ACCESS_TYPE_READ,
 		ACCESS_TYPE_ATOMIC_WRITE = ACCESS_TYPE_ATOMIC | ACCESS_TYPE_WRITE,
@@ -109,7 +110,7 @@ namespace UE::Sanitizer::RaceDetector {
 	ENUM_CLASS_FLAGS(EMemoryAccessType);
 	INSTRUMENTATION_FUNCTION_ATTRIBUTES inline const TCHAR* AccessTypeToString(EMemoryAccessType AccessType)
 	{
-		switch (AccessType)
+		switch (AccessType & ACCESS_TYPE_ATOMIC_READ_WRITE)
 		{
 		case ACCESS_TYPE_READ:
 			return TEXT("Read");
@@ -179,8 +180,9 @@ namespace UE::Sanitizer::RaceDetector {
 						uint8 bIsRead : 1;
 						uint8 bIsWrite : 1;
 						uint8 bIsAtomic : 1;
+						uint8 bIsVPtr : 1;
 					};
-					EMemoryAccessType AccessType : 3;
+					EMemoryAccessType AccessType : 4;
 				};
 				uint8 bHasSyncObjectAttached : 1;
 			};
@@ -247,8 +249,6 @@ namespace UE::Sanitizer::RaceDetector {
 		uint32 ThreadCreationDepth = 0; 
 		// When we want detailed logging for diagnostic purpose
 		uint32 DetailedLogDepth = 0;
-		// When we want detailed log only on a single address
-		void* DetailedLogAddress = nullptr; 
 		// Clock used while waiting to get a context id assigned.
 		FClock StandbyClock;
 		// Each thread holds a bank of clocks to synchronize with every other context.
@@ -291,16 +291,7 @@ namespace UE::Sanitizer::RaceDetector {
 	struct FInstrumentationScope
 	{
 		bool bNeedDecrement = false;
-		INSTRUMENTATION_FUNCTION_ATTRIBUTES FInstrumentationScope()
-		{
-			FContext* Context = GetThreadContext();
-			if (FContext::IsValid(Context))
-			{
-				Context->WinInstrumentationDepth++;
-				bNeedDecrement = true;
-			}
-		}
-
+		INSTRUMENTATION_FUNCTION_ATTRIBUTES FInstrumentationScope();
 		INSTRUMENTATION_FUNCTION_ATTRIBUTES ~FInstrumentationScope()
 		{
 			if (bNeedDecrement)
diff --git a/Engine/Source/Runtime/Core/Private/Sanitizer/RaceDetector.cpp b/Engine/Source/Runtime/Core/Private/Sanitizer/RaceDetector.cpp
--- a/Engine/Source/Runtime/Core/Private/Sanitizer/RaceDetector.cpp
+++ b/Engine/Source/Runtime/Core/Private/Sanitizer/RaceDetector.cpp
@@ -131,6 +131,7 @@ namespace UE::Sanitizer::RaceDetector {
 	volatile bool bIsResettingShadow = false;
 	volatile bool bIsDebuggerPresent = false;
 	volatile bool bDetailedLogGlobal = false;
+	volatile void* DetailedLogAddress = nullptr;
 
 	// Can't be an atomic as it might cause infinite recursion in debug when atomics
 	// are unoptimized and we end up with a function call to an instrumented function
@@ -547,14 +548,15 @@ namespace UE::Sanitizer::RaceDetector {
 			Report.Appendf(
 				TEXT("=========================================") LINE_TERMINATOR
 				TEXT("WARNING: RaceDetector: data race detected") LINE_TERMINATOR
-				TEXT("%s of size %d at %p %smade at clock %lu by thread %s (%d) which is now at clock %lu:") LINE_TERMINATOR 
+				TEXT("%s%s of size %d at %p %smade at clock %lu by thread %s (%d) which is now at clock %lu:") LINE_TERMINATOR 
 				TEXT("%s") LINE_TERMINATOR
-				TEXT("Previous %s of size %d at %p at clock %lu by thread %s (%d %s) which is now at clock %lu:") LINE_TERMINATOR
+				TEXT("Previous %s%s of size %d at %p at clock %lu by thread %s (%d %s) which is now at clock %lu:") LINE_TERMINATOR
 				TEXT("%s") LINE_TERMINATOR
 				TEXT("Last known sync clock between both threads is %lu:") LINE_TERMINATOR
 				TEXT("%s") LINE_TERMINATOR
 				TEXT("=========================================") LINE_TERMINATOR,
 				AccessTypeToString(CurrentAccess.AccessType),
+				(CurrentAccess.AccessType & EMemoryAccessType::ACCESS_TYPE_VPTR) ? TEXT(" (vptr)") : TEXT(""),
 				CurrentAccess.GetSize(),
 				(void*)((UPTRINT)AlignedPointer + CurrentAccess.GetOffset()),
 				bIsRaceOnStack ? TEXT("(Stack) ") : TEXT(""),
@@ -564,6 +566,7 @@ namespace UE::Sanitizer::RaceDetector {
 				Context.CurrentClock(),
 				*FirstFullLocation.ToString(Alignment),
 				AccessTypeToString(RacingAccess.AccessType),
+				(RacingAccess.AccessType& EMemoryAccessType::ACCESS_TYPE_VPTR) ? TEXT(" (vptr)") : TEXT(""),
 				RacingAccess.GetSize(),
 				(void*)((UPTRINT)AlignedPointer + RacingAccess.GetOffset()),
 				RacingAccess.Clock,
@@ -635,8 +638,8 @@ namespace UE::Sanitizer::RaceDetector {
 		// exactly what's going on inside.
 		if (UNLIKELY(Context.DetailedLogDepth || bDetailedLogGlobal))
 		{
-			if (Context.DetailedLogAddress == nullptr ||
-				Context.DetailedLogAddress == Ptr)
+			if (DetailedLogAddress == nullptr ||
+				DetailedLogAddress == Ptr)
 			{
 				FPlatformMisc::LowLevelOutputDebugStringf(
 					TEXT("[%d] %s / %s of size %d offset %d (%s) at %p (ctx:%d clk:%u)")
@@ -919,7 +922,7 @@ namespace UE::Sanitizer::RaceDetector {
 			}
 		}
 	}
-	
+
 	INSTRUMENTATION_FUNCTION_ATTRIBUTES void UnhookInstrumentation();
 	INSTRUMENTATION_FUNCTION_ATTRIBUTES void HookInstrumentation();
 
@@ -1227,7 +1230,23 @@ namespace UE::Sanitizer::RaceDetector {
 
 		// FMalloc interface end
 	};
-	
+
+	INSTRUMENTATION_FUNCTION_ATTRIBUTES FInstrumentationScope::FInstrumentationScope()
+	{
+		// This is important since we're never unregistering the malloc instrumentation so we need to avoid
+		// handling context once the runtime is shut down otherwise we can end up with use-after-free on the
+		// TLS during application exit.
+		if (bRuntimeInitialized)
+		{
+			FContext* Context = GetThreadContext();
+			if (FContext::IsValid(Context))
+			{
+				Context->WinInstrumentationDepth++;
+				bNeedDecrement = true;
+			}
+		}
+	}
+
 	INSTRUMENTATION_FUNCTION_ATTRIBUTES void InitMemoryAllocator()
 	{
 		while (true)
@@ -1308,6 +1327,11 @@ namespace UE::Sanitizer::RaceDetector {
 		{
 			// Mark ContextTLS as invalid to prevent any further usage / detection for this thread
 			Platform::SetTlsValue(ContextTLSIndex, (void*)-1);
+
+			// We don't care about the depth we're currently in when releasing the current context 
+			// since this only happens during shutdown and thread cleanup.
+			Context->InstrumentationDepth = 0;
+
 			ReleaseContext(Context);
 		}
 	}
@@ -1412,8 +1436,7 @@ namespace UE::Sanitizer::RaceDetector {
 
 	INSTRUMENTATION_FUNCTION_ATTRIBUTES void ToggleFilterDetailedLogOnAddress(void* Address)
 	{
-		FContext& Context = EnsureCurrentContext();
-		Context.DetailedLogAddress = Address;
+		DetailedLogAddress = Address;
 	}
 
 
@@ -1506,6 +1529,34 @@ extern "C" {
 		FINISH_HANDLE_INSTRUMENTATION();
 	}
 
+	INSTRUMENTATION_CORE_FUNCTION_PREFIX void __RaceDetector__Instrument_VPtr_Store(void** Address, void* Value)
+	{
+		BEGIN_HANDLE_INSTRUMENTATION(return);
+
+		// For virtual table pointers, there is no race if the store is simply rewriting the same value as this will generally happen
+		// when entering a destructor of a base class that wasn't subclassed.
+		// See FRaceDetectorVirtualPointerBenignTest in RaceDetectorTests.cpp for an example.
+		if (*Address != Value)
+		{
+			AddCallstackFrame(Context, _ReturnAddress());
+			InstrumentAccessPreamble(Context, (UPTRINT)Address, sizeof(void*), EMemoryAccessType::ACCESS_TYPE_WRITE | EMemoryAccessType::ACCESS_TYPE_VPTR, FAtomicMemoryOrder::MEMORY_ORDER_RELAXED, _ReturnAddress(), TEXT("VPtr Store"));
+			RemoveCallstackFrame(Context);
+		}
+
+		FINISH_HANDLE_INSTRUMENTATION();
+	}
+
+	INSTRUMENTATION_CORE_FUNCTION_PREFIX void __RaceDetector__Instrument_VPtr_Load(void** Address)
+	{
+		BEGIN_HANDLE_INSTRUMENTATION(return);
+
+		AddCallstackFrame(Context, _ReturnAddress());
+		InstrumentAccessPreamble(Context, (UPTRINT)Address, sizeof(void*), EMemoryAccessType::ACCESS_TYPE_READ | EMemoryAccessType::ACCESS_TYPE_VPTR, FAtomicMemoryOrder::MEMORY_ORDER_RELAXED, _ReturnAddress(), TEXT("VPtr Load"));
+		RemoveCallstackFrame(Context);
+
+		FINISH_HANDLE_INSTRUMENTATION();
+	}
+
 	INSTRUMENTATION_CORE_FUNCTION_PREFIX void __RaceDetector__Instrument_StoreRange(uint64 Address, uint32 Size)
 	{
 		if (Size == 0)
@@ -1772,6 +1823,8 @@ extern "C" {
 	void __Instrument_LoadRange(uint64 Address, uint32 Size);
 	void __Instrument_Store(uint64 Address, uint32 Size);
 	void __Instrument_Load(uint64 Address, uint32 Size);
+	void __Instrument_VPtr_Store(void** Address, void* Value);
+	void __Instrument_VPtr_Load(void** Address);
 }
 
 namespace UE::Sanitizer::RaceDetector {
@@ -1801,6 +1854,8 @@ INSTRUMENTATION_FUNCTION_ATTRIBUTES void PopulateHotPatchFunctions()
 	HotpatchInstrumentationFunctionsFast.Emplace(__Thunk__Instrument_LoadRange, __RaceDetector__Instrument_LoadRange);
 	HotpatchInstrumentationFunctionsFast.Emplace(__Thunk__Instrument_Store, __RaceDetector__Instrument_Store);
 	HotpatchInstrumentationFunctionsFast.Emplace(__Thunk__Instrument_Load, __RaceDetector__Instrument_Load);
+	HotpatchInstrumentationFunctionsFast.Emplace(__Thunk__Instrument_VPtr_Store, __RaceDetector__Instrument_VPtr_Store);
+	HotpatchInstrumentationFunctionsFast.Emplace(__Thunk__Instrument_VPtr_Load, __RaceDetector__Instrument_VPtr_Load);
 
 	// In non-monolithic, we also hook the core functions directly instead of via the thunks to save on another set of JMPs
 #if !IS_MONOLITHIC
@@ -1813,6 +1868,8 @@ INSTRUMENTATION_FUNCTION_ATTRIBUTES void PopulateHotPatchFunctions()
 	HotpatchInstrumentationFunctionsFast.Emplace(__Instrument_LoadRange, __RaceDetector__Instrument_LoadRange);
 	HotpatchInstrumentationFunctionsFast.Emplace(__Instrument_Store, __RaceDetector__Instrument_Store);
 	HotpatchInstrumentationFunctionsFast.Emplace(__Instrument_Load, __RaceDetector__Instrument_Load);
+	HotpatchInstrumentationFunctionsFast.Emplace(__Instrument_VPtr_Store, __RaceDetector__Instrument_VPtr_Store);
+	HotpatchInstrumentationFunctionsFast.Emplace(__Instrument_VPtr_Load, __RaceDetector__Instrument_VPtr_Load);
 #endif
 
 #define HOTPATCH_FUNC(Type) \
diff --git a/Engine/Source/Runtime/Core/Private/Sanitizer/RaceDetectorWindows.cpp b/Engine/Source/Runtime/Core/Private/Sanitizer/RaceDetectorWindows.cpp
--- a/Engine/Source/Runtime/Core/Private/Sanitizer/RaceDetectorWindows.cpp
+++ b/Engine/Source/Runtime/Core/Private/Sanitizer/RaceDetectorWindows.cpp
@@ -30,6 +30,7 @@ namespace UE::Sanitizer::RaceDetector {
 	extern void SanitizerThreadRun(volatile bool& bContinue);
 	extern volatile bool bIsResettingShadow;
 	extern void WriteToLog(FString&& Message);
+	extern bool bRuntimeInitialized;
 }
 
 #define START_WINAPI_INSTRUMENTATION(DoIfNotInstrumenting) \
@@ -86,6 +87,14 @@ struct FLastErrorPreservationScope
 // Must be called from thread about to end.
 INSTRUMENTATION_FUNCTION_ATTRIBUTES void InstrumentEndThread()
 {
+	// Important to guard against trying to access TLS if the runtime has shut down
+	// otherwise we might end up with use-after-free during some thread cleanup 
+	// at application exit.
+	if (!bRuntimeInitialized)
+	{
+		return;
+	}
+
 	HANDLE Handle = nullptr;
 	if (FContext* ContextPtr = GetThreadContext(); FContext::IsValid(ContextPtr))
 	{
diff --git a/Engine/Source/Runtime/Core/Tests/Sanitizer/RaceDetectorTests.cpp b/Engine/Source/Runtime/Core/Tests/Sanitizer/RaceDetectorTests.cpp
--- a/Engine/Source/Runtime/Core/Tests/Sanitizer/RaceDetectorTests.cpp
+++ b/Engine/Source/Runtime/Core/Tests/Sanitizer/RaceDetectorTests.cpp
@@ -151,7 +151,7 @@ namespace UE { namespace TasksTests
 			ToggleRaceDetection(false);
 			FThreadRaceCollector Collector;
 
-			for (uint32 i = 0; i < 10000; ++i)
+			for (uint32 i = 0; i < 100; ++i)
 			{
 				Collector.Reset();
 				ToggleRaceDetection(true);
@@ -450,6 +450,84 @@ namespace UE { namespace TasksTests
 		}
 	}
 
+	TEST_CASE_NAMED(FRaceDetectorVirtualPointerHarmfulTest, "System::Core::Sanitizer::RaceDetector::VirtualPointerHarmful", "[ApplicationContextMask][EngineFilter]")
+	{
+		ToggleRaceDetection(false);
+
+		FThreadRaceCollector Collector;
+
+		ToggleRaceDetection(true);
+
+		class FBase
+		{
+		public:
+			virtual void Function() {}
+			void Done() { Event.Notify(); }
+			virtual ~FBase() { Event.Wait(); }
+		private:
+			UE::FManualResetEvent Event;
+		};
+
+		class FDerived : public FBase
+		{
+		public:
+			virtual void Function() {}
+			virtual ~FDerived() {}
+		};
+
+		FBase* Base = new FDerived; 
+		TFuture<void> FutureA = AsyncThread([Base] { 
+			Base->Function(); 
+			Base->Done(); 
+		});
+
+		// This is a race because the function called could be the one of FDerived or FBase
+		// depending if the call is made before or after we enter the destructor and the vptr
+		// is rewritten to point to the base functions.
+		delete Base;
+		FutureA.Get();
+
+		ToggleRaceDetection(false);
+
+		// For this test, assume the vptr is stored as the first member of the instance.
+		CHECK(Collector.Contains(Base));
+	}
+
+	TEST_CASE_NAMED(FRaceDetectorVirtualPointerBenignTest, "System::Core::Sanitizer::RaceDetector::VirtualPointerBenign", "[ApplicationContextMask][EngineFilter]")
+	{
+		ToggleRaceDetection(false);
+
+		FThreadRaceCollector Collector;
+
+		ToggleRaceDetection(true);
+
+		class FBase
+		{
+		public:
+			virtual void Function() {}
+			void Done() { Event.Notify(); }
+			virtual ~FBase() { Event.Wait(); }
+		private:
+			UE::FManualResetEvent Event;
+		};
+
+		FBase* Base = new FBase;
+		TFuture<void> FutureA = AsyncThread([Base] {
+			Base->Function();
+			Base->Done();
+		});
+
+		// This race is considered benign since the vptr can only point on the base class
+		// so the racedetector won't report it.
+		delete Base;
+		FutureA.Get();
+
+		ToggleRaceDetection(false);
+
+		// For this test, assume the vptr is stored as the first member of the instance.
+		CHECK(!Collector.Contains(Base));
+	}
+
 	struct FTestStruct
 	{
 		int32 Index;
diff --git a/Engine/Source/Runtime/Instrumentation/Internal/Instrumentation/EntryPoints.h b/Engine/Source/Runtime/Instrumentation/Internal/Instrumentation/EntryPoints.h
--- a/Engine/Source/Runtime/Instrumentation/Internal/Instrumentation/EntryPoints.h
+++ b/Engine/Source/Runtime/Instrumentation/Internal/Instrumentation/EntryPoints.h
@@ -14,6 +14,8 @@ extern "C"
 	INSTRUMENTATION_API void __Thunk__Instrument_FuncExit();
 	INSTRUMENTATION_API void __Thunk__Instrument_Store(uint64 Address, uint32 Size);
 	INSTRUMENTATION_API void __Thunk__Instrument_Load(uint64 Address, uint32 Size);
+	INSTRUMENTATION_API void __Thunk__Instrument_VPtr_Store(void** Address, void* Value);
+	INSTRUMENTATION_API void __Thunk__Instrument_VPtr_Load(void** Address);
 	INSTRUMENTATION_API void __Thunk__Instrument_StoreRange(uint64 Address, uint32 Size);
 	INSTRUMENTATION_API void __Thunk__Instrument_LoadRange(uint64 Address, uint32 Size);
 
diff --git a/Engine/Source/Runtime/Instrumentation/Internal/Instrumentation/PerModuleInline.inl b/Engine/Source/Runtime/Instrumentation/Internal/Instrumentation/PerModuleInline.inl
--- a/Engine/Source/Runtime/Instrumentation/Internal/Instrumentation/PerModuleInline.inl
+++ b/Engine/Source/Runtime/Instrumentation/Internal/Instrumentation/PerModuleInline.inl
@@ -38,6 +38,16 @@ extern "C" {
 		__Thunk__Instrument_FuncExit();
 	}
 
+	INSTRUMENTATION_PERMODULE_FUNCTION_ATTRIBUTES void __Instrument_VPtr_Store(void** Address, void* Value)
+	{
+		__Thunk__Instrument_VPtr_Store(Address, Value);
+	}
+
+	INSTRUMENTATION_PERMODULE_FUNCTION_ATTRIBUTES void __Instrument_VPtr_Load(void** Address)
+	{
+		__Thunk__Instrument_VPtr_Load(Address);
+	}
+
 	INSTRUMENTATION_PERMODULE_FUNCTION_ATTRIBUTES void __Instrument_Store(uint64 Address, uint32 Size)
 	{
 		__Thunk__Instrument_Store(Address, Size);
diff --git a/Engine/Source/Runtime/Instrumentation/Private/Instrumentation/EntryPoints.cpp b/Engine/Source/Runtime/Instrumentation/Private/Instrumentation/EntryPoints.cpp
--- a/Engine/Source/Runtime/Instrumentation/Private/Instrumentation/EntryPoints.cpp
+++ b/Engine/Source/Runtime/Instrumentation/Private/Instrumentation/EntryPoints.cpp
@@ -37,6 +37,14 @@ INSTRUMENTATION_THUNK_ATTRIBUTES void __Thunk__Instrument_Load(uint64 Address, u
 {
 }
 
+INSTRUMENTATION_THUNK_ATTRIBUTES void __Thunk__Instrument_VPtr_Store(void** Address, void* Value)
+{
+}
+
+INSTRUMENTATION_THUNK_ATTRIBUTES void __Thunk__Instrument_VPtr_Load(void** Address)
+{
+}
+
 INSTRUMENTATION_THUNK_ATTRIBUTES void __Thunk__Instrument_StoreRange(uint64 Address, uint32 Size)
 {
 }
