Fix the config system being unable to set cvars that end with the following pattern "_[numbers]"

- This issue was affecting cvars with multiple underscores in the name and ending with _[numbers], so something like "Our_CVar_001".
- The problem occured because in ConfigUtilities we were using FName::GetPlainNameString which with the given name format assumed that the trailing "_001" was the internal index number of the FName and so would be discarded. So instead of applying a change from config file to cvar "Our_CVar_001" it would try to apply it to "Our_CVar".
- Replaced use of FName::GetPlainNameString with FName::ToString
-- As a minor optimization I opted to use the overload that writes the string to a TStringBuilder to reduce the number of temp allocations made.

diff --git a/Engine/Source/Runtime/Core/Private/Misc/ConfigUtilities.cpp b/Engine/Source/Runtime/Core/Private/Misc/ConfigUtilities.cpp
--- a/Engine/Source/Runtime/Core/Private/Misc/ConfigUtilities.cpp
+++ b/Engine/Source/Runtime/Core/Private/Misc/ConfigUtilities.cpp
@@ -392,11 +392,12 @@ void ApplyCVarSettingsFromIni(const TCHAR* InSectionName, const TCHAR* InIniFile
 
 	UE_LOG(LogConfig,Log,TEXT("Applying CVar settings from Section [%s] File [%s]"),InSectionName,InIniFilename);
 
-	if(const FConfigSection* Section = GConfig->GetSection(InSectionName, false, InIniFilename))
+	if (const FConfigSection* Section = GConfig->GetSection(InSectionName, false, InIniFilename))
 	{
-		for(FConfigSectionMap::TConstIterator It(*Section); It; ++It)
+		for (FConfigSectionMap::TConstIterator It(*Section); It; ++It)
 		{
-			const FString& KeyString = It.Key().GetPlainNameString(); 
+			TStringBuilder<128> KeyString;
+			It.Key().ToString(KeyString);
 			const FString& ValueString = It.Value().GetValue();
 
 			OnSetCVarFromIniEntry(InIniFilename, *KeyString, *ValueString, SetBy, bAllowCheating, false, Tag);
@@ -410,7 +411,7 @@ void ForEachCVarInSectionFromIni(const TCHAR* InSectionName, const TCHAR* InIniF
 	{
 		for (FConfigSectionMap::TConstIterator It(*Section); It; ++It)
 		{
-			const FString& KeyString = It.Key().GetPlainNameString();
+			const FString& KeyString = It.Key().ToString();
 			const FString& ValueString = ConvertValueFromHumanFriendlyValue(*It.Value().GetValue());
 
 			IConsoleVariable* CVar = IConsoleManager::Get().FindConsoleVariable(*KeyString);
@@ -475,10 +476,10 @@ public:
 			{
 				for (FConfigSectionMap::TConstIterator It(*Section); It; ++It)
 				{
-					const FString& KeyString = It.Key().GetPlainNameString();
+					TStringBuilder<128> KeyString;
+					It.Key().ToString(KeyString);
 					const FString& ValueString = It.Value().GetValue();
 
-
 					IConsoleVariable* CVar = IConsoleManager::Get().FindConsoleVariable(*KeyString);
 
 					if ( !CVar )
