[TSAN on Windows]
- Fix potentially undetected races because we weren't prioritizing the latest access to be stored in the shadow memory
- Split tests so they can be run separately
- Running FRaceDetectorStdThreadSynchEventTest repeatedly would sometimes fail because of an undetected race on the y value.

diff --git a/Engine/Source/Runtime/Core/Private/Sanitizer/RaceDetector.cpp b/Engine/Source/Runtime/Core/Private/Sanitizer/RaceDetector.cpp
--- a/Engine/Source/Runtime/Core/Private/Sanitizer/RaceDetector.cpp
+++ b/Engine/Source/Runtime/Core/Private/Sanitizer/RaceDetector.cpp
@@ -730,11 +730,12 @@ namespace UE::Sanitizer::RaceDetector {
 			// If we already have a slot, upgrade it if possible to avoid spilling to too many slots.
 			if (PreviousAccess.ContextId == CurrentAccess.ContextId)
 			{
-				// Ultimately, what we want in the slot is what is most susceptible to cause a race.
+				// Most recent clock first, then what we want in the slot is what is most susceptible to cause a race.
 				// We try to keep non-atomic first, then write first.. then finally reads.
 				if (bSaveNeeded &&
-					(PreviousAccess.bIsAtomic > CurrentAccess.bIsAtomic ||
-						CurrentAccess.bIsWrite > PreviousAccess.bIsWrite))
+					(CurrentAccess.Clock > PreviousAccess.Clock ||
+					 PreviousAccess.bIsAtomic > CurrentAccess.bIsAtomic ||
+					 CurrentAccess.bIsWrite > PreviousAccess.bIsWrite))
 				{
 					RawValue = CurrentAccess.RawValue;
 					bSaveNeeded = false;
diff --git a/Engine/Source/Runtime/Core/Tests/Sanitizer/RaceDetectorTests.cpp b/Engine/Source/Runtime/Core/Tests/Sanitizer/RaceDetectorTests.cpp
--- a/Engine/Source/Runtime/Core/Tests/Sanitizer/RaceDetectorTests.cpp
+++ b/Engine/Source/Runtime/Core/Tests/Sanitizer/RaceDetectorTests.cpp
@@ -137,321 +137,338 @@ namespace UE { namespace TasksTests
 	void Test();
 	void TestLockFree(int32 OuterIters = 3);
 
-	TEST_CASE_NAMED(FRaceDetectorThreadTest, "System::Core::Sanitizer::RaceDetector::Threads", "[ApplicationContextMask][EngineFilter]")
+	TEST_CASE_NAMED(FRaceDetectorTasksWithPrereqTest, "System::Core::Sanitizer::RaceDetector::TasksWithPrereq", "[ApplicationContextMask][EngineFilter]")
 	{
 		CHECK(UE::Sanitizer::RaceDetector::Initialize());
 
-		SECTION("Manual test")
-		{
-			//Test();
-		}
+		ToggleRaceDetection(false);
+		FThreadRaceCollector Collector;
 
-		SECTION("Task with prerequisite")
+		for (uint32 i = 0; i < 100; ++i)
 		{
-			ToggleRaceDetection(false);
-			FThreadRaceCollector Collector;
-
-			for (uint32 i = 0; i < 100; ++i)
-			{
-				Collector.Reset();
-				ToggleRaceDetection(true);
+			Collector.Reset();
+			ToggleRaceDetection(true);
 
-				int x = 0, y = 0;
+			int x = 0, y = 0;
 
-				FGraphEventRef LegacyTask = FFunctionGraphTask::CreateAndDispatchWhenReady([&]
-				{
-					x = 1;
-					y = 1;
-				}, TStatId(), nullptr, ENamedThreads::AnyHiPriThreadHiPriTask);
+			FGraphEventRef LegacyTask = FFunctionGraphTask::CreateAndDispatchWhenReady([&]
+			{
+				x = 1;
+				y = 1;
+			}, TStatId(), nullptr, ENamedThreads::AnyHiPriThreadHiPriTask);
 
+			{
+				FGraphEventArray Prereqs;
+				Prereqs.Add(LegacyTask);
+				LegacyTask = FFunctionGraphTask::CreateAndDispatchWhenReady([&]
 				{
-					FGraphEventArray Prereqs;
-					Prereqs.Add(LegacyTask);
-					LegacyTask = FFunctionGraphTask::CreateAndDispatchWhenReady([&]
-					{
-						x = 2;
-						y = 2;
-					}, TStatId(), &Prereqs, ENamedThreads::AnyHiPriThreadHiPriTask);
-				}
-
-				LegacyTask->Wait();
+					x = 2;
+					y = 2;
+				}, TStatId(), &Prereqs, ENamedThreads::AnyHiPriThreadHiPriTask);
+			}
 
-				ToggleRaceDetection(false);
+			LegacyTask->Wait();
 
-				CHECK(!Collector.Contains(&x));
-				CHECK(!Collector.Contains(&y));
-			}
-			
 			ToggleRaceDetection(false);
+
+			CHECK(!Collector.Contains(&x));
+			CHECK(!Collector.Contains(&y));
 		}
+	}
 
-		SECTION("Standard threads (race)")
-		{
-			ToggleRaceDetection(false);
+	TEST_CASE_NAMED(FRaceDetectorStdThreadRacesTest, "System::Core::Sanitizer::RaceDetector::StdThreadRaces", "[ApplicationContextMask][EngineFilter]")
+	{
+		CHECK(UE::Sanitizer::RaceDetector::Initialize());
 
-			FThreadRaceCollector Collector;
+		ToggleRaceDetection(false);
 
-			ToggleRaceDetection(true);
+		FThreadRaceCollector Collector;
 
-			int x = 0, y = 0, z = 0, w = 0;
-			std::atomic<int> sync = 0;
-			std::thread t1{ [&]() {
-				x = 1;
-				y = 1;
-				z = 1;
-			} };
+		ToggleRaceDetection(true);
 
-			w = 2;
+		int x = 0, y = 0, z = 0, w = 0;
+		std::atomic<int> sync = 0;
+		std::thread t1{ [&]() {
+			x = 1;
+			y = 1;
+			z = 1;
+		} };
 
-			std::thread t2{ [&]() {
-				w = 3;
-				z = 3;
-			} };
+		w = 2;
 
-			x = 2;
+		std::thread t2{ [&]() {
+			w = 3;
+			z = 3;
+		} };
 
-			t1.join();
-			t2.join();
+		x = 2;
 
-			y = 2;
+		t1.join();
+		t2.join();
 
-			ToggleRaceDetection(false);
+		y = 2;
 
-			CHECK(Collector.Contains(&x));
-			CHECK(!Collector.Contains(&y));
-			CHECK(Collector.Contains(&z));
-			CHECK(!Collector.Contains(&w));	
-		}
+		ToggleRaceDetection(false);
 
-		SECTION("Standard threads (atomic synchronization)")
-		{
-			ToggleRaceDetection(false);
-			FThreadRaceCollector Collector;
+		CHECK(Collector.Contains(&x));
+		CHECK(!Collector.Contains(&y));
+		CHECK(Collector.Contains(&z));
+		CHECK(!Collector.Contains(&w));
+	}
 
-			ToggleRaceDetection(true);
+	TEST_CASE_NAMED(FRaceDetectorStdThreadAtomicSyncTest, "System::Core::Sanitizer::RaceDetector::StdThreadAtomicSync", "[ApplicationContextMask][EngineFilter]")
+	{
+		CHECK(UE::Sanitizer::RaceDetector::Initialize());
 
-			int x = 0, y = 0, result = 0;
-			std::atomic<int> sync = 0;
-			std::thread t1{ [&]() {
-				x = 10;
+		ToggleRaceDetection(false);
+		FThreadRaceCollector Collector;
 
-				int expected = 0;
-				if (!sync.compare_exchange_strong(expected, 1))
-				{
-					result = y;
-				}
-			}};
+		ToggleRaceDetection(true);
 
-			y = 20;
+		int x = 0, y = 0, result = 0;
+		std::atomic<int> sync = 0;
+		std::thread t1{ [&]() {
+			x = 10;
 
 			int expected = 0;
 			if (!sync.compare_exchange_strong(expected, 1))
 			{
-				result = x;
+				result = y;
 			}
+		} };
 
-			t1.join();
+		y = 20;
 
-			ToggleRaceDetection(false);
-
-			CHECK(!Collector.Contains(&x));
-			CHECK(!Collector.Contains(&y));
-			CHECK(!Collector.Contains(&result));
+		int expected = 0;
+		if (!sync.compare_exchange_strong(expected, 1))
+		{
+			result = x;
 		}
 
-		SECTION("Standard threads (event synchronization)")
-		{
-			ToggleRaceDetection(false);
+		t1.join();
 
-			FThreadRaceCollector Collector;
+		ToggleRaceDetection(false);
 
-			ToggleRaceDetection(true);
+		CHECK(!Collector.Contains(&x));
+		CHECK(!Collector.Contains(&y));
+		CHECK(!Collector.Contains(&result));
+	}
 
-			FEvent* Event = FPlatformProcess::GetSynchEventFromPool(false);
-			Event->Reset();
+	TEST_CASE_NAMED(FRaceDetectorStdThreadSynchEventTest, "System::Core::Sanitizer::RaceDetector::StdThreadSynchEvent", "[ApplicationContextMask][EngineFilter]")
+	{
+		CHECK(UE::Sanitizer::RaceDetector::Initialize());
 
-			int x = 0, y= 0, result = 0;
-			std::thread t1{ [&]() {
-				x = 10;
-				Event->Trigger();
-				y = 20;
-			} };
+		ToggleRaceDetection(false);
 
-			Event->Wait();
-			result = x;
-			result += y;
-			
-			t1.join();
+		FThreadRaceCollector Collector;
 
-			ToggleRaceDetection(false);
+		ToggleRaceDetection(true);
 
-			CHECK(!Collector.Contains(&x));
-			CHECK(Collector.Contains(&y));
-			CHECK(result); // avoid the result and its operations to be optimized out.
-		}
+		FEvent* Event = FPlatformProcess::GetSynchEventFromPool(false);
+		Event->Reset();
 
-		SECTION("Standard threads (manual reset event synchronization)")
-		{
-			ToggleRaceDetection(false);
-			FThreadRaceCollector Collector;
+		int x = 0, y = 0, result = 0;
 
-			ToggleRaceDetection(true);
+		ToggleFilterDetailedLogOnAddress(&y);
+		ToggleGlobalDetailedLog(true);
 
-			UE::FPlatformManualResetEvent Event;
-			Event.Reset();
+		std::thread t1{ [&]() {
+			x = 10;
+			Event->Trigger();
+			y = 20;
+		} };
 
-			int x = 0, y = 0, result = 0;
-			std::thread t1{ [&]() {
-				x = 10;
-				Event.Notify();
-				y = 20;
-			} };
+		Event->Wait();
+		result = x;
+		result += y;
 
-			Event.Wait();
-			result = x;
-			result += y;
+		t1.join();
 
-			t1.join();
+		ToggleFilterDetailedLogOnAddress(nullptr);
+		ToggleGlobalDetailedLog(false);
 
-			ToggleRaceDetection(false);
+		ToggleRaceDetection(false);
 
-			CHECK(!Collector.Contains(&x));
-			CHECK(Collector.Contains(&y));
-			CHECK(result); // avoid the result and its operations to be optimized out.
-		}
-		
-		SECTION("Modern tasks")
-		{
-			ToggleRaceDetection(false);
+		CHECK(!Collector.Contains(&x));
+		CHECK(Collector.Contains(&y));
+		CHECK(result); // avoid the result and its operations to be optimized out.
+	}
 
-			FThreadRaceCollector Collector;
+	TEST_CASE_NAMED(FRaceDetectorStdThreadManualResetEventTest, "System::Core::Sanitizer::RaceDetector::StdThreadManualResetEvent", "[ApplicationContextMask][EngineFilter]")
+	{
+		CHECK(UE::Sanitizer::RaceDetector::Initialize());
 
-			ToggleRaceDetection(true);
+		ToggleRaceDetection(false);
+		FThreadRaceCollector Collector;
 
-			int x = 0, y = 0;
-			FTask TaskA = UE::Tasks::Launch(TEXT("A"), [&]() {
-				x = 3;
-				y = 3;
-			});
-			x = 5;
-			TestWait(TaskA);
+		ToggleRaceDetection(true);
 
-			y = 5;
+		UE::FPlatformManualResetEvent Event;
+		Event.Reset();
 
-			ToggleRaceDetection(false);
-			CHECK(Collector.Contains(&x));
-			CHECK(!Collector.Contains(&y));
-		}
-		
-		SECTION("Legacy tasks")
-		{
-			ToggleRaceDetection(false);
+		int x = 0, y = 0, result = 0;
+		std::thread t1{ [&]() {
+			x = 10;
+			Event.Notify();
+			y = 20;
+		} };
 
-			FThreadRaceCollector Collector;
+		Event.Wait();
+		result = x;
+		result += y;
 
-			ToggleRaceDetection(true);
+		t1.join();
 
-			int x = 0, y = 0;
-			FGraphEventRef LegacyTask = FFunctionGraphTask::CreateAndDispatchWhenReady([&]
-			{
-				x = 3;
-				y = 3;
-			}, TStatId(), nullptr, ENamedThreads::AnyHiPriThreadHiPriTask);
+		ToggleRaceDetection(false);
 
-			x = 5;
+		CHECK(!Collector.Contains(&x));
+		CHECK(Collector.Contains(&y));
+		CHECK(result); // avoid the result and its operations to be optimized out.
+	}
 
-			TestWait(LegacyTask);
+	TEST_CASE_NAMED(FRaceDetectorModernTasksTest, "System::Core::Sanitizer::RaceDetector::ModernTasks", "[ApplicationContextMask][EngineFilter]")
+	{
+		CHECK(UE::Sanitizer::RaceDetector::Initialize());
 
-			y = 5;
+		ToggleRaceDetection(false);
 
-			ToggleRaceDetection(false);
+		FThreadRaceCollector Collector;
 
-			CHECK(Collector.Contains(&x));
-			CHECK(!Collector.Contains(&y));
-		}
+		ToggleRaceDetection(true);
+
+		int x = 0, y = 0;
+		FTask TaskA = UE::Tasks::Launch(TEXT("A"), [&]() {
+			x = 3;
+			y = 3;
+		});
+		x = 5;
+		TestWait(TaskA);
+
+		y = 5;
+
+		ToggleRaceDetection(false);
+		CHECK(Collector.Contains(&x));
+		CHECK(!Collector.Contains(&y));
+	}
+
+	TEST_CASE_NAMED(FRaceDetectorLegacyTasksTest, "System::Core::Sanitizer::RaceDetector::LegacyTasks", "[ApplicationContextMask][EngineFilter]")
+	{
+		CHECK(UE::Sanitizer::RaceDetector::Initialize());
 
-		SECTION("Two async threads (race)")
+		ToggleRaceDetection(false);
+
+		FThreadRaceCollector Collector;
+
+		ToggleRaceDetection(true);
+
+		int x = 0, y = 0;
+		FGraphEventRef LegacyTask = FFunctionGraphTask::CreateAndDispatchWhenReady([&]
 		{
-			ToggleRaceDetection(false);
+			x = 3;
+			y = 3;
+		}, TStatId(), nullptr, ENamedThreads::AnyHiPriThreadHiPriTask);
 
-			FThreadRaceCollector Collector;
+		x = 5;
 
-			ToggleRaceDetection(true);
+		TestWait(LegacyTask);
 
-			int x = 0;
-			TFuture<void> FutureA = AsyncThread([&x]() {
-				x = 1;
-			});
+		y = 5;
 
-			x = 2;
+		ToggleRaceDetection(false);
 
-			FutureA.Get();
+		CHECK(Collector.Contains(&x));
+		CHECK(!Collector.Contains(&y));
+	}
 
-			ToggleRaceDetection(false);
+	TEST_CASE_NAMED(FRaceDetectorTwoAsyncThreadRaceTest, "System::Core::Sanitizer::RaceDetector::TwoAsyncThreadRace", "[ApplicationContextMask][EngineFilter]")
+	{
+		CHECK(UE::Sanitizer::RaceDetector::Initialize());
 
-			CHECK(Collector.Contains(&x));
-		}
+		ToggleRaceDetection(false);
 
-		SECTION("Async threads (atomic synchronization)")
-		{
-			ToggleRaceDetection(false);
+		FThreadRaceCollector Collector;
 
-			FThreadRaceCollector Collector;
+		ToggleRaceDetection(true);
 
-			int x = 0;
-			std::atomic<bool> sync = false;
-			TFuture<void> FutureA = AsyncThread([&x, &sync]() {
-				ToggleRaceDetection(true);
-				x = 1;
-				sync.store(true);
-				ToggleRaceDetection(false);
-			});
+		int x = 0;
+		TFuture<void> FutureA = AsyncThread([&x]() {
+			x = 1;
+		});
 
-			FutureA.Get();
+		x = 2;
 
-			TFuture<void> FutureB = AsyncThread([&x, &sync]() {
-				ToggleRaceDetection(true);
-				if (sync.load())
-				{
-					x = 2;
-				}
-				ToggleRaceDetection(false);
-			});
+		FutureA.Get();
 
-			
-			FutureB.Get();
+		ToggleRaceDetection(false);
 
-			CHECK(!Collector.Contains(&x));
-		}
+		CHECK(Collector.Contains(&x));
+	}
 
-		SECTION("Async threads (atomic access)")
-		{
-			ToggleRaceDetection(false);
+	TEST_CASE_NAMED(FRaceDetectorAsyncThreadAtomicSyncTest, "System::Core::Sanitizer::RaceDetector::AsyncThreadAtomicSync", "[ApplicationContextMask][EngineFilter]")
+	{
+		CHECK(UE::Sanitizer::RaceDetector::Initialize());
 
-			FThreadRaceCollector Collector;
+		ToggleRaceDetection(false);
+
+		FThreadRaceCollector Collector;
 
+		int x = 0;
+		std::atomic<bool> sync = false;
+		TFuture<void> FutureA = AsyncThread([&x, &sync]() {
 			ToggleRaceDetection(true);
+			x = 1;
+			sync.store(true);
+			ToggleRaceDetection(false);
+		});
 
-			std::atomic<uint32> x = 0;
-			TFuture<void> FutureA = AsyncThread([&x]() {
-				x = 1;
-			});
+		FutureA.Get();
 
-			TFuture<void> FutureB = AsyncThread([&x]() {
+		TFuture<void> FutureB = AsyncThread([&x, &sync]() {
+			ToggleRaceDetection(true);
+			if (sync.load())
+			{
 				x = 2;
-			});
+			}
+			ToggleRaceDetection(false);
+		});
 
-			FutureA.Get();
-			FutureB.Get();
 
-			ToggleRaceDetection(false);
+		FutureB.Get();
 
-			CHECK(!Collector.Contains(&x));
-		}
+		CHECK(!Collector.Contains(&x));
+	}
+
+	TEST_CASE_NAMED(FRaceDetectorAsyncThreadAtomicAccessTest, "System::Core::Sanitizer::RaceDetector::AsyncThreadAtomicAccess", "[ApplicationContextMask][EngineFilter]")
+	{
+		CHECK(UE::Sanitizer::RaceDetector::Initialize());
+
+		ToggleRaceDetection(false);
+
+		FThreadRaceCollector Collector;
+
+		ToggleRaceDetection(true);
+
+		std::atomic<uint32> x = 0;
+		TFuture<void> FutureA = AsyncThread([&x]() {
+			x = 1;
+		});
+
+		TFuture<void> FutureB = AsyncThread([&x]() {
+			x = 2;
+		});
+
+		FutureA.Get();
+		FutureB.Get();
+
+		ToggleRaceDetection(false);
+
+		CHECK(!Collector.Contains(&x));
 	}
 
 	TEST_CASE_NAMED(FRaceDetectorVirtualPointerHarmfulTest, "System::Core::Sanitizer::RaceDetector::VirtualPointerHarmful", "[ApplicationContextMask][EngineFilter]")
 	{
+		CHECK(UE::Sanitizer::RaceDetector::Initialize());
+
 		ToggleRaceDetection(false);
 
 		FThreadRaceCollector Collector;
@@ -495,6 +512,8 @@ namespace UE { namespace TasksTests
 
 	TEST_CASE_NAMED(FRaceDetectorVirtualPointerBenignTest, "System::Core::Sanitizer::RaceDetector::VirtualPointerBenign", "[ApplicationContextMask][EngineFilter]")
 	{
+		CHECK(UE::Sanitizer::RaceDetector::Initialize());
+
 		ToggleRaceDetection(false);
 
 		FThreadRaceCollector Collector;
