Change FTimeout to use FPlatformTime::Cycles64 internally instead of FTimespan, this can be significantly faster depending on the platform.
Add support for double seconds API functions, which are more efficient and easier to use with the new internals.
Add functions for modifying an existing timeout and efficiently processing 0 duration timeouts.

diff --git a/Engine/Source/Runtime/Core/Public/Misc/Timeout.h b/Engine/Source/Runtime/Core/Public/Misc/Timeout.h
--- a/Engine/Source/Runtime/Core/Public/Misc/Timeout.h
+++ b/Engine/Source/Runtime/Core/Public/Misc/Timeout.h
@@ -9,19 +9,17 @@
 
 namespace UE
 {
-	// utility class to handle timeouts
-	// usage:
-	// ------------------------------------
-	// FTimeout Timeout(FTimespan::FromMilliseconds(2));
-	// while (!Timeout.IsExpired()) { ... }
-	// ------------------------------------
+	/** 
+	 * Utility class to create a timeout that will expire at a point in the future.
+	 * Example usage:
+	 *
+	 *	FTimeout TimeoutFromTimespan(FTimespan::FromMilliseconds(2));
+	 *	FTimeout TimeoutFromSeconds(0.002);
+	 *	while (!TimeoutFromSeconds.IsExpired()) { ... }
+	 */
 	class FTimeout
 	{
 	public:
-		explicit FTimeout(FTimespan Value)
-			: Timeout(Value)
-		{
-		}
 
 		UE_DEPRECATED(5.5, "Use IsExpired() instead.")
 		explicit operator bool() const
@@ -29,36 +27,188 @@ namespace UE
 			return IsExpired();
 		}
 
+		/** Return true if elapsed time is greater than the initially requested timeout */
 		bool IsExpired() const
 		{
-			return GetRemainingTime() <= FTimespan::Zero();
+			// First two cases can skip the slow current time check
+			if (WillNeverExpire())
+			{
+				return false;
+			}
+			else if (IsAlwaysExpired())
+			{
+				return true;
+			}
+			else
+			{
+				return FPlatformTime::Cycles64() > (StartCycles + TimeoutCycles);
+			}
 		}
 
-		FTimespan GetElapsedTime() const
+		/** Create a timeout that will never return true for IsExpired */
+		static FTimeout Never()
 		{
-			return FTimespan::FromSeconds(FPlatformTime::Seconds()) - Start;
+			return FTimeout(FPlatformTime::Cycles64(), NeverExpireCycles);
 		}
 
-		FTimespan GetRemainingTime() const
+		/** Returns true if this was created from Never and does not need to be repeatedly checked */
+		FORCEINLINE bool WillNeverExpire() const
 		{
-			return Timeout == FTimespan::MaxValue() ? FTimespan::MaxValue() : Timeout - GetElapsedTime();
+			return TimeoutCycles == NeverExpireCycles;
 		}
 
-		static FTimeout Never()
+		/** Create a timeout that will always return true for IsExpired */
+		static FTimeout AlwaysExpired()
+		{
+			return FTimeout(FPlatformTime::Cycles64(), 0);
+		}
+
+		/** Returns true if this was created from AlwaysExpired and does not need to be repeatedly checked */
+		FORCEINLINE bool IsAlwaysExpired() const
+		{
+			return TimeoutCycles == 0;
+		}
+
+		/** Set this timeout to explicitly expired without recalculating start time */
+		void SetToExpired()
+		{
+			TimeoutCycles = 0;
+		}
+
+
+		// Preferred API for creating and querying using double seconds
+
+		/** Construct a timeout that starts right now and will end after the passed in time in seconds */
+		explicit FTimeout(double TimeoutSeconds)
+			: StartCycles(FPlatformTime::Cycles64())
+		{
+			SetTimeoutSeconds(TimeoutSeconds);
+		}
+
+		/** Construct a timeout that started at the same time as BaseTimeout, but with a new duration */
+		explicit FTimeout(const FTimeout& BaseTimeout, double TimeoutSeconds)
+			: StartCycles(BaseTimeout.StartCycles)
+		{
+			SetTimeoutSeconds(TimeoutSeconds);
+		}
+
+		/** Returns time since the timeout was created, in seconds */
+		double GetElapsedSeconds() const
+		{
+			// StartCycles can never be greater than current time as there is no way to construct a timeout starting in the future
+			return FPlatformTime::ToSeconds64(FPlatformTime::Cycles64() - StartCycles);
+		}
+
+		/** Returns time left until the timeout expires (which can be negative) in seconds */
+		double GetRemainingSeconds() const
+		{
+			if (WillNeverExpire())
+			{
+				return NeverExpireSeconds;
+			}
+
+			// We convert to double separately to avoid underflow on the cycles
+			// This could also be done with some branches or treating cycles as signed int64
+			return GetTimeoutSeconds() - GetElapsedSeconds();
+		}
+
+		/** Returns duration of timeout in seconds */
+		double GetTimeoutSeconds() const
+		{
+			return FPlatformTime::ToSeconds64(TimeoutCycles);
+		}
+
+		/** Sets the timeout to new value in seconds */
+		void SetTimeoutSeconds(double TimeoutSeconds)
 		{
-			return FTimeout{ FTimespan::MaxValue() };
+			if (TimeoutSeconds <= 0.0)
+			{
+				SetToExpired();
+			}
+			else
+			{
+				TimeoutCycles = FPlatformTime::SecondsToCycles64(TimeoutSeconds);
+			}
 		}
 
+		/** Safely modify the remaining time by adding the delta time in seconds to the timeout */
+		void ModifyTimeoutSeconds(double DeltaTimeoutSeconds)
+		{
+			if (IsAlwaysExpired() || WillNeverExpire())
+			{
+				return;
+			}
+
+			if (DeltaTimeoutSeconds >= 0.0)
+			{
+				TimeoutCycles += FPlatformTime::SecondsToCycles64(DeltaTimeoutSeconds);
+			}
+			else
+			{
+				uint64 RemovedCycles = FPlatformTime::SecondsToCycles64(-DeltaTimeoutSeconds);
+				if (RemovedCycles >= TimeoutCycles)
+				{
+					SetToExpired();
+				}
+				else
+				{
+					TimeoutCycles -= RemovedCycles;
+				}
+			}
+		}
+
+
+		// Older API for creating and querying using FTimespan
+
+		/** Construct a timeout that starts right now and will end after the passed in timespan */
+		explicit FTimeout(FTimespan TimeoutValue)
+			: StartCycles(FPlatformTime::Cycles64())
+		{
+			if (TimeoutValue == FTimespan::MaxValue())
+			{
+				TimeoutCycles = NeverExpireCycles;
+			}
+			else
+			{
+				SetTimeoutSeconds(TimeoutValue.GetTotalSeconds());
+			}
+		}
+
+		/** Returns time since the timeout was created, as a timespan */
+		FTimespan GetElapsedTime() const
+		{
+			return FTimespan::FromSeconds(GetElapsedSeconds());
+		}
+
+		/** Returns time left until the timeout expires (which can be negative) as a timespan */
+		FTimespan GetRemainingTime() const
+		{
+			if (WillNeverExpire())
+			{
+				return FTimespan::MaxValue();
+			}
+
+			return FTimespan::FromSeconds(GetRemainingSeconds());
+		}
+
+		/** Returns duration of timeout as a timespan */
 		FTimespan GetTimeoutValue() const
 		{
-			return Timeout;
+			if (WillNeverExpire())
+			{
+				return FTimespan::MaxValue();
+			}
+			
+			return FTimespan::FromSeconds(GetTimeoutSeconds());
 		}
 
-		// intended for use in waiting functions, e.g. `FEvent::Wait()`
-		// returns the whole number (rounded up) of remaining milliseconds, clamped into [0, MAX_uint32] range
+		/**
+		 * Intended for use in waiting functions, e.g. `FEvent::Wait()`
+		 * returns the whole number (rounded up) of remaining milliseconds, clamped into [0, MAX_uint32] range
+		 */
 		uint32 GetRemainingRoundedUpMilliseconds() const
 		{
-			if (*this == Never())
+			if (WillNeverExpire())
 			{
 				return MAX_uint32;
 			}
@@ -71,7 +221,10 @@ namespace UE
 
 		friend bool operator==(FTimeout Left, FTimeout Right)
 		{
-			return Left.Timeout == Right.Timeout && (Left.Timeout == FTimespan::MaxValue() || Left.Start == Right.Start);
+			// Timeout cycles need to match which handles differentiating between always and never
+			// For normal timeouts, also check the start cycles
+			return Left.TimeoutCycles == Right.TimeoutCycles 
+				&& (Left.WillNeverExpire() || Left.IsAlwaysExpired() || Left.StartCycles == Right.StartCycles);
 		}
 
 		friend bool operator!=(FTimeout Left, FTimeout Right)
@@ -80,7 +233,19 @@ namespace UE
 		}
 
 	private:
-		FTimespan Start{ FTimespan::FromSeconds(FPlatformTime::Seconds()) };
-		FTimespan Timeout;
+		FTimeout(uint64 StartValue, uint64 TimeoutValue)
+			: StartCycles(StartValue)
+			, TimeoutCycles(TimeoutValue)
+		{
+		}
+
+		static constexpr uint64 NeverExpireCycles = TNumericLimits<uint64>::Max();
+		static constexpr double NeverExpireSeconds = TNumericLimits<double>::Max();
+
+		// Value of FPlatformTime::Cycles64 at timeout creation, cannot be directly converted to seconds
+		uint64 StartCycles = 0;
+
+		// Length of timeout, can be converted to seconds as it is relative to StartCycles
+		uint64 TimeoutCycles = 0;
 	};
 }
