Minimize macro usage in FText.

Previously, our null-guard accessor was implemented as a macro instead of as a method. This led to macro invocations spread across the file. This has been replaced with a normal method; the compiler should make reasonable inlining choices.

Proper typechecking also exposed places--serialization, rebuilding--where the macro's behavior would have been undesirable. In particular, we don't ever want to mutate the static EmptyText object. The const-correctness of an actual method correctly detected this mistake.

diff --git a/Engine/Source/Runtime/Core/Private/Internationalization/Text.cpp b/Engine/Source/Runtime/Core/Private/Internationalization/Text.cpp
--- a/Engine/Source/Runtime/Core/Private/Internationalization/Text.cpp
+++ b/Engine/Source/Runtime/Core/Private/Internationalization/Text.cpp
@@ -38,17 +38,6 @@ DEFINE_LOG_CATEGORY(LogText);
 
 #define LOCTEXT_NAMESPACE "Core.Text"
 
-// LiveCoding can result in FText with null TextData, so this guards against them crashing while still emitting an ensure
-#ifndef UE_ENABLE_SAFE_GET_TEXT_DATA
-	#define UE_ENABLE_SAFE_GET_TEXT_DATA (WITH_EDITOR && WITH_LIVE_CODING)
-#endif
-
-#if UE_ENABLE_SAFE_GET_TEXT_DATA
-	#define UE_SAFE_GET_TEXT_DATA(Text) (ensure((Text).TextData) ? (Text).TextData : FText::GetEmpty().TextData)
-#else
-	#define UE_SAFE_GET_TEXT_DATA(Text) ((Text).TextData)
-#endif
-
 bool FTextInspector::ShouldGatherForLocalization(const FText& Text)
 {
 	return Text.ShouldGatherForLocalization();
@@ -76,18 +65,18 @@ TOptional<FString> FTextInspector::GetKey(const FText& Text)
 
 FTextId FTextInspector::GetTextId(const FText& Text)
 {
-	return UE_SAFE_GET_TEXT_DATA(Text)->GetTextHistory().GetTextId();
+	return Text.GetTextData().GetTextHistory().GetTextId();
 }
 
 const FString* FTextInspector::GetSourceString(const FText& Text)
 {
-	return &UE_SAFE_GET_TEXT_DATA(Text)->GetSourceString();
+	return &Text.GetTextData().GetSourceString();
 }
 
 const FString& FTextInspector::GetDisplayString(const FText& Text)
 {
 	Text.Rebuild();
-	return UE_SAFE_GET_TEXT_DATA(Text)->GetDisplayString();
+	return Text.GetTextData().GetDisplayString();
 }
 
 bool FTextInspector::GetTableIdAndKey(const FText& Text, FName& OutTableId, FString& OutKey)
@@ -130,7 +119,7 @@ bool FTextInspector::GetHistoricNumericData(const FText& Text, FHistoricTextNume
 
 const void* FTextInspector::GetSharedDataId(const FText& Text)
 {
-	return UE_SAFE_GET_TEXT_DATA(Text).GetReference();
+	return &Text.GetTextData();
 }
 
 // These default values have been duplicated to the KismetTextLibrary functions for Blueprints. Please replicate any changes there!
@@ -317,6 +306,16 @@ void FText::AutoRTFMAssignFromOpenToClosed(FText& Closed, const FText& Open)
 	ensure(AutoRTFM::EContextStatus::OnTrack == Status);
 }
 
+const ITextData& FText::GetTextData() const
+{
+// LiveCoding can result in FText with null TextData, so this guards against them crashing while still emitting an ensure
+#if WITH_EDITOR && WITH_LIVE_CODING
+	return ensure(TextData) ? *TextData : *FText::GetEmpty().TextData;
+#else
+	return *TextData;
+#endif
+}
+
 bool FText::IsEmpty() const
 {
 	return ToString().IsEmpty();
@@ -920,7 +919,7 @@ void FText::SerializeText(FStructuredArchive::FSlot Slot, FText& Value)
 
 			if (bSerializeHistory)
 			{
-				int8 HistoryType = (int8)UE_SAFE_GET_TEXT_DATA(Value)->GetTextHistory().GetType();
+				int8 HistoryType = (int8)Value.GetTextData().GetTextHistory().GetType();
 				Record << SA_VALUE(TEXT("HistoryType"), HistoryType);
 			}
 			else
@@ -932,7 +931,7 @@ void FText::SerializeText(FStructuredArchive::FSlot Slot, FText& Value)
 				Record << SA_VALUE(TEXT("bHasCultureInvariantString"), bHasCultureInvariantString);
 				if (bHasCultureInvariantString)
 				{
-					FString CultureInvariantString = UE_SAFE_GET_TEXT_DATA(Value)->GetSourceString();
+					FString CultureInvariantString = Value.GetTextData().GetSourceString();
 					Record << SA_VALUE(TEXT("CultureInvariantString"), CultureInvariantString);
 				}
 			}
@@ -1032,7 +1031,8 @@ void FText::SerializeText(FStructuredArchive::FSlot Slot, FText& Value)
 
 		if (bSerializeHistory)
 		{
-			UE_SAFE_GET_TEXT_DATA(Value)->GetMutableTextHistory().Serialize(Record);
+			check(Value.TextData);
+			Value.TextData->GetMutableTextHistory().Serialize(Record);
 		}
 	}
 
@@ -1056,7 +1056,7 @@ void FText::SerializeText(FStructuredArchive::FSlot Slot, FText& Value)
 #if WITH_EDITORONLY_DATA
 FText FText::ChangeKey( const FTextKey& Namespace, const FTextKey& Key, const FText& Text )
 {
-	return FText(FString(*UE_SAFE_GET_TEXT_DATA(Text)->GetSourceString()), Namespace, Key, 0);
+	return FText(FString(*Text.GetTextData().GetSourceString()), Namespace, Key, 0);
 }
 #endif
 
@@ -1250,12 +1250,12 @@ bool FText::FindTextInLiveTable_Advanced(const FTextKey& Namespace, const FTextK
 const FString& FText::ToString() const
 {
 	Rebuild();
-	return UE_SAFE_GET_TEXT_DATA(*this)->GetDisplayString();
+	return GetTextData().GetDisplayString();
 }
 
 FString FText::BuildSourceString() const
 {
-	return UE_SAFE_GET_TEXT_DATA(*this)->GetTextHistory().BuildInvariantDisplayString();
+	return GetTextData().GetTextHistory().BuildInvariantDisplayString();
 }
 
 bool FText::IsNumeric() const
@@ -1285,7 +1285,10 @@ bool FText::EqualToCaseIgnored(const FText& Other) const
 
 void FText::Rebuild() const
 {
-	UE_SAFE_GET_TEXT_DATA(*this)->GetMutableTextHistory().UpdateDisplayStringIfOutOfDate();
+	if (TextData)
+	{
+		TextData->GetMutableTextHistory().UpdateDisplayStringIfOutOfDate();
+	}
 }
 
 bool FText::IsTransient() const
@@ -1305,7 +1308,7 @@ bool FText::IsInitializedFromString() const
 
 bool FText::IsFromStringTable() const
 {
-	return UE_SAFE_GET_TEXT_DATA(*this)->GetTextHistory().GetType() == ETextHistoryType::StringTableEntry;
+	return GetTextData().GetTextHistory().GetType() == ETextHistoryType::StringTableEntry;
 }
 
 bool FText::ShouldGatherForLocalization() const
@@ -1315,7 +1318,7 @@ bool FText::ShouldGatherForLocalization() const
 		return false;
 	}
 
-	const FString& SourceString = UE_SAFE_GET_TEXT_DATA(*this)->GetSourceString();
+	const FString& SourceString = GetTextData().GetSourceString();
 
 	auto IsAllWhitespace = [](const FString& String) -> bool
 	{
@@ -1339,12 +1342,12 @@ void FText::PreloadStringTable( FName InTableId, const FTextKey& InKey )
 
 void FText::GetHistoricFormatData(TArray<FHistoricTextFormatData>& OutHistoricFormatData) const
 {
-	UE_SAFE_GET_TEXT_DATA(*this)->GetTextHistory().GetHistoricFormatData(*this, OutHistoricFormatData);
+	GetTextData().GetTextHistory().GetHistoricFormatData(*this, OutHistoricFormatData);
 }
 
 bool FText::GetHistoricNumericData(FHistoricTextNumericData& OutHistoricNumericData) const
 {
-	return UE_SAFE_GET_TEXT_DATA(*this)->GetTextHistory().GetHistoricNumericData(*this, OutHistoricNumericData);
+	return GetTextData().GetTextHistory().GetHistoricNumericData(*this, OutHistoricNumericData);
 }
 
 bool FText::IdenticalTo( const FText& Other, const ETextIdenticalModeFlags CompareModeFlags ) const
@@ -1362,8 +1365,8 @@ bool FText::IdenticalTo( const FText& Other, const ETextIdenticalModeFlags Compa
 	// This is fast as it skips a lexical compare, however it can also return false for two instances that have identical strings, but in different pointers.
 	// For instance, this method will return false for two FText objects created from FText::FromString("Wooble") as they each have unique (or null), non-shared instances.
 	{
-		FTextConstDisplayStringPtr DisplayStringPtr = UE_SAFE_GET_TEXT_DATA(*this)->GetLocalizedString();
-		FTextConstDisplayStringPtr OtherDisplayStringPtr = UE_SAFE_GET_TEXT_DATA(Other)->GetLocalizedString();
+		FTextConstDisplayStringPtr DisplayStringPtr = GetTextData().GetLocalizedString();
+		FTextConstDisplayStringPtr OtherDisplayStringPtr = Other.GetTextData().GetLocalizedString();
 		if (DisplayStringPtr && OtherDisplayStringPtr && DisplayStringPtr == OtherDisplayStringPtr)
 		{
 			return true;
@@ -1372,8 +1375,8 @@ bool FText::IdenticalTo( const FText& Other, const ETextIdenticalModeFlags Compa
 
 	if (EnumHasAnyFlags(CompareModeFlags, ETextIdenticalModeFlags::DeepCompare))
 	{
-		const FTextHistory& ThisTextHistory = UE_SAFE_GET_TEXT_DATA(*this)->GetTextHistory();
-		const FTextHistory& OtherTextHistory = UE_SAFE_GET_TEXT_DATA(Other)->GetTextHistory();
+		const FTextHistory& ThisTextHistory = GetTextData().GetTextHistory();
+		const FTextHistory& OtherTextHistory = Other.GetTextData().GetTextHistory();
 		if (ThisTextHistory.GetType() == OtherTextHistory.GetType() && ThisTextHistory.IdenticalTo(OtherTextHistory, CompareModeFlags))
 		{
 			return true;
@@ -1750,7 +1753,7 @@ FTextSnapshot::FTextSnapshot()
 
 FTextSnapshot::FTextSnapshot(const FText& InText)
 	: TextDataPtr(InText.TextData)
-	, LocalizedStringPtr(UE_SAFE_GET_TEXT_DATA(InText)->GetLocalizedString())
+	, LocalizedStringPtr(InText.GetTextData().GetLocalizedString())
 	, GlobalHistoryRevision(GetGlobalHistoryRevisionForText(InText))
 	, LocalHistoryRevision(GetLocalHistoryRevisionForText(InText))
 	, Flags(InText.Flags)
@@ -1765,7 +1768,7 @@ bool FTextSnapshot::IdenticalTo(const FText& InText) const
 
 	return TextDataPtr
 		&& TextDataPtr == InText.TextData
-		&& LocalizedStringPtr == UE_SAFE_GET_TEXT_DATA(InText)->GetLocalizedString()
+		&& LocalizedStringPtr == InText.GetTextData().GetLocalizedString()
 		&& GlobalHistoryRevision == GetGlobalHistoryRevisionForText(InText)
 		&& LocalHistoryRevision == GetLocalHistoryRevisionForText(InText)
 		&& Flags == InText.Flags;
@@ -1782,17 +1785,17 @@ bool FTextSnapshot::IsDisplayStringEqualTo(const FText& InText) const
 	return TextDataPtr
 		&& GlobalHistoryRevision == GetGlobalHistoryRevisionForText(InText)
 		&& LocalHistoryRevision == GetLocalHistoryRevisionForText(InText)
-		&& TextDataPtr->GetDisplayString().Equals(UE_SAFE_GET_TEXT_DATA(InText)->GetDisplayString(), ESearchCase::CaseSensitive);
+		&& TextDataPtr->GetDisplayString().Equals(InText.GetTextData().GetDisplayString(), ESearchCase::CaseSensitive);
 }
 
 uint16 FTextSnapshot::GetGlobalHistoryRevisionForText(const FText& InText)
 {
-	return (InText.IsEmpty() || InText.IsCultureInvariant()) ? 0 : UE_SAFE_GET_TEXT_DATA(InText)->GetGlobalHistoryRevision();
+	return (InText.IsEmpty() || InText.IsCultureInvariant()) ? 0 : InText.GetTextData().GetGlobalHistoryRevision();
 }
 
 uint16 FTextSnapshot::GetLocalHistoryRevisionForText(const FText& InText)
 {
-	return (InText.IsEmpty() || InText.IsCultureInvariant()) ? 0 : UE_SAFE_GET_TEXT_DATA(InText)->GetLocalHistoryRevision();
+	return (InText.IsEmpty() || InText.IsCultureInvariant()) ? 0 : InText.GetTextData().GetLocalHistoryRevision();
 }
 
 bool TextBiDi::IsControlCharacter(const TCHAR InChar)
@@ -1972,7 +1975,7 @@ void FTextStringHelper::WriteToBuffer(FString& Buffer, const FText& Value, const
 #undef LOC_DEFINE_REGION
 	}
 	// Is this text that should be written via its text history?
-	else if (UE_SAFE_GET_TEXT_DATA(Value)->GetTextHistory().WriteToBuffer(Buffer, bStripPackageNamespace))
+	else if (Value.GetTextData().GetTextHistory().WriteToBuffer(Buffer, bStripPackageNamespace))
 	{
 	}
 	// This isn't special text, so write as a raw string (potentially quoted)
@@ -2215,6 +2218,4 @@ const TCHAR* LexToString(ERoundingMode InValue)
 	}
 }
 
-#undef UE_SAFE_GET_TEXT_DATA
-
 #undef LOCTEXT_NAMESPACE
diff --git a/Engine/Source/Runtime/Core/Public/Internationalization/Text.h b/Engine/Source/Runtime/Core/Public/Internationalization/Text.h
--- a/Engine/Source/Runtime/Core/Public/Internationalization/Text.h
+++ b/Engine/Source/Runtime/Core/Public/Internationalization/Text.h
@@ -891,6 +891,9 @@ private:
 	template < typename T >
 	static TSharedRef<ITextGenerator> CreateTextGenerator(FStructuredArchive::FRecord Record);
 
+	/** This accessor for TextData guards against the case where TextData might be null and returns the ITextData from EmptyText instead. */
+	const ITextData& GetTextData() const;
+
 private:
 	/** The internal shared data for this FText */
 	TRefCountPtr<ITextData> TextData;
