Trivial perf improvement to the TaskGraph modifications I made last week

diff --git a/Engine/Source/Runtime/Core/Private/Async/TaskGraph.cpp b/Engine/Source/Runtime/Core/Private/Async/TaskGraph.cpp
--- a/Engine/Source/Runtime/Core/Private/Async/TaskGraph.cpp
+++ b/Engine/Source/Runtime/Core/Private/Async/TaskGraph.cpp
@@ -1593,7 +1593,7 @@ private:
 			{
 				if (RemainingTasks[TaskIndex]->IsCompleted())
 				{
-					RemainingTasks.RemoveAtSwap(TaskIndex);
+					RemainingTasks.RemoveAtSwap(TaskIndex, EAllowShrinking::No);
 				}
 			}
 
