- FIxed an assert in FMappedName to check correct range

diff --git a/Engine/Source/Runtime/Core/Public/Serialization/MappedName.h b/Engine/Source/Runtime/Core/Public/Serialization/MappedName.h
--- a/Engine/Source/Runtime/Core/Public/Serialization/MappedName.h
+++ b/Engine/Source/Runtime/Core/Public/Serialization/MappedName.h
@@ -38,7 +38,7 @@ public:
 
 	static inline FMappedName Create(const uint32 InIndex, const uint32 InNumber, EType InType)
 	{
-		check(InIndex <= MAX_int32);
+		check(InIndex < (1 << TypeShift));
 		return FMappedName((uint32(InType) << TypeShift) | InIndex, InNumber);
 	}
 
