[AutoRTFM] Address CL comments from review.

These changes shouldn't have any functional impact, just naming/formatting tweaks.

diff --git a/Engine/Source/Runtime/Core/Public/Serialization/TransactionallySafeArchiveWriter.h b/Engine/Source/Runtime/Core/Public/Serialization/TransactionallySafeArchiveWriter.h
--- a/Engine/Source/Runtime/Core/Public/Serialization/TransactionallySafeArchiveWriter.h
+++ b/Engine/Source/Runtime/Core/Public/Serialization/TransactionallySafeArchiveWriter.h
@@ -39,7 +39,7 @@ public:
 		{
 			// Our wrapper is going away, so move everything into a lambda-capture for safekeeping.
 			AutoRTFM::PopOnCommitHandler(this);
-			AutoRTFM::PushOnCommitHandler(this, [Archive = InnerArchive.Release(), WriteBuffer = MoveTemp(DeferredWrites), bFlush = bRequestFlush]() mutable
+			AutoRTFM::PushOnCommitHandler(this, [Archive = InnerArchive.Release(), WriteBuffer = MoveTemp(DeferredWrites), bFlush = bFlushRequested]() mutable
 			{
 				DoDeferredWrites(*Archive, WriteBuffer, bFlush);
 				delete Archive;
@@ -79,7 +79,7 @@ public:
 		if (MaybeRegisterCommitHandler())
 		{
 			// Defer the flush until commit time.
-			bRequestFlush = true;
+			bFlushRequested = true;
 		}
 		else
 		{
@@ -132,25 +132,25 @@ private:
 
 		AutoRTFM::PushOnCommitHandler(this, [this]
 		{
-			DoDeferredWrites(*InnerArchive, DeferredWrites, bRequestFlush);
+			DoDeferredWrites(*InnerArchive, DeferredWrites, bFlushRequested);
 			bRegisteredCommitHandler = false;
 		});
 	}
 
-	static void DoDeferredWrites(FArchive& Archive, DeferredWriteBuffer& WriteBuffer, bool& bRequestFlush)
+	static void DoDeferredWrites(FArchive& Archive, DeferredWriteBuffer& WriteBuffer, bool& bFlushRequested)
 	{
 		Archive.Serialize(WriteBuffer.GetData(), WriteBuffer.NumBytes());
 		WriteBuffer.Reset();
 
-		if (bRequestFlush)
+		if (bFlushRequested)
 		{
 			Archive.Flush();
-			bRequestFlush = false;
+			bFlushRequested = false;
 		}
 	}
 
 	TUniquePtr<FArchive> InnerArchive;
 	DeferredWriteBuffer DeferredWrites;
 	bool bRegisteredCommitHandler = false;
-	bool bRequestFlush = false;
+	bool bFlushRequested = false;
 };
\ No newline at end of file
