Deprecated FPlatformType::CHAR32 and added FPlatformType::UTF32CHAR instead, which is a char type rather than an integer type.

diff --git a/Engine/Plugins/Experimental/SerializationUtils/Source/JsonSerialization/Private/Formatters/JsonArchiveInputFormatterEx.cpp b/Engine/Plugins/Experimental/SerializationUtils/Source/JsonSerialization/Private/Formatters/JsonArchiveInputFormatterEx.cpp
--- a/Engine/Plugins/Experimental/SerializationUtils/Source/JsonSerialization/Private/Formatters/JsonArchiveInputFormatterEx.cpp
+++ b/Engine/Plugins/Experimental/SerializationUtils/Source/JsonSerialization/Private/Formatters/JsonArchiveInputFormatterEx.cpp
@@ -365,6 +365,11 @@ void FJsonArchiveInputFormatterEx::Serialize(bool& Value)
 	Value = ValueStack.Top()->AsBool();
 }
 
+void FJsonArchiveInputFormatterEx::Serialize(UTF32CHAR& Value)
+{
+	verify(ValueStack.Top()->TryGetNumber(*(uint32*)&Value));
+}
+
 void FJsonArchiveInputFormatterEx::Serialize(FString& Value)
 {
 	// If the string we serialized was empty, this value will be a null object rather than a string, so we have to
diff --git a/Engine/Plugins/Experimental/SerializationUtils/Source/JsonSerialization/Private/Formatters/JsonArchiveOutputFormatterEx.cpp b/Engine/Plugins/Experimental/SerializationUtils/Source/JsonSerialization/Private/Formatters/JsonArchiveOutputFormatterEx.cpp
--- a/Engine/Plugins/Experimental/SerializationUtils/Source/JsonSerialization/Private/Formatters/JsonArchiveOutputFormatterEx.cpp
+++ b/Engine/Plugins/Experimental/SerializationUtils/Source/JsonSerialization/Private/Formatters/JsonArchiveOutputFormatterEx.cpp
@@ -330,6 +330,11 @@ void FJsonArchiveOutputFormatterEx::Serialize(bool& Value)
 	WriteValue(LexToString(Value));
 }
 
+void FJsonArchiveOutputFormatterEx::Serialize(UTF32CHAR& Value)
+{
+	WriteValue(LexToString(*(uint32*)&Value));
+}
+
 void FJsonArchiveOutputFormatterEx::Serialize(FString& Value)
 {
 	// Insert a "String:" prefix to prevent incorrect interpretation as another explicit type
diff --git a/Engine/Plugins/Experimental/SerializationUtils/Source/JsonSerialization/Public/Formatters/JsonArchiveInputFormatterEx.h b/Engine/Plugins/Experimental/SerializationUtils/Source/JsonSerialization/Public/Formatters/JsonArchiveInputFormatterEx.h
--- a/Engine/Plugins/Experimental/SerializationUtils/Source/JsonSerialization/Public/Formatters/JsonArchiveInputFormatterEx.h
+++ b/Engine/Plugins/Experimental/SerializationUtils/Source/JsonSerialization/Public/Formatters/JsonArchiveInputFormatterEx.h
@@ -79,6 +79,7 @@ public:
 	virtual void Serialize(float& Value) override;
 	virtual void Serialize(double& Value) override;
 	virtual void Serialize(bool& Value) override;
+	virtual void Serialize(UTF32CHAR& Value) override;
 	virtual void Serialize(FString& Value) override;
 	virtual void Serialize(FName& Value) override;
 	virtual void Serialize(UObject*& Value) override;
diff --git a/Engine/Plugins/Experimental/SerializationUtils/Source/JsonSerialization/Public/Formatters/JsonArchiveOutputFormatterEx.h b/Engine/Plugins/Experimental/SerializationUtils/Source/JsonSerialization/Public/Formatters/JsonArchiveOutputFormatterEx.h
--- a/Engine/Plugins/Experimental/SerializationUtils/Source/JsonSerialization/Public/Formatters/JsonArchiveOutputFormatterEx.h
+++ b/Engine/Plugins/Experimental/SerializationUtils/Source/JsonSerialization/Public/Formatters/JsonArchiveOutputFormatterEx.h
@@ -78,6 +78,7 @@ public:
 	virtual void Serialize(float& Value) override;
 	virtual void Serialize(double& Value) override;
 	virtual void Serialize(bool& Value) override;
+	virtual void Serialize(UTF32CHAR& Value) override;
 	virtual void Serialize(FString& Value) override;
 	virtual void Serialize(FName& Value) override;
 	virtual void Serialize(UObject*& Value) override;
diff --git a/Engine/Plugins/Experimental/SerializationUtils/Source/XmlSerialization/Private/Formatters/XmlArchiveInputFormatter.cpp b/Engine/Plugins/Experimental/SerializationUtils/Source/XmlSerialization/Private/Formatters/XmlArchiveInputFormatter.cpp
--- a/Engine/Plugins/Experimental/SerializationUtils/Source/XmlSerialization/Private/Formatters/XmlArchiveInputFormatter.cpp
+++ b/Engine/Plugins/Experimental/SerializationUtils/Source/XmlSerialization/Private/Formatters/XmlArchiveInputFormatter.cpp
@@ -746,6 +746,15 @@ void FXmlArchiveInputFormatter::Serialize(bool& Value)
 	}
 }
 
+void FXmlArchiveInputFormatter::Serialize(UTF32CHAR& Value)
+{
+	const pugi::xml_attribute Attrib = Impl->GetValueAttrib();
+	if (!Attrib.empty())
+	{
+		Value = Attrib.as_uint();
+	}
+}
+
 void FXmlArchiveInputFormatter::Serialize(FString& Value)
 {
 	const pugi::xml_attribute Attrib = Impl->GetValueAttrib();
diff --git a/Engine/Plugins/Experimental/SerializationUtils/Source/XmlSerialization/Private/Formatters/XmlArchiveOutputFormatter.cpp b/Engine/Plugins/Experimental/SerializationUtils/Source/XmlSerialization/Private/Formatters/XmlArchiveOutputFormatter.cpp
--- a/Engine/Plugins/Experimental/SerializationUtils/Source/XmlSerialization/Private/Formatters/XmlArchiveOutputFormatter.cpp
+++ b/Engine/Plugins/Experimental/SerializationUtils/Source/XmlSerialization/Private/Formatters/XmlArchiveOutputFormatter.cpp
@@ -403,6 +403,11 @@ void FXmlArchiveOutputFormatter::Serialize(bool& Value)
 	Impl->SetValue(Value);
 }
 
+void FXmlArchiveOutputFormatter::Serialize(UTF32CHAR& Value)
+{
+	Impl->SetValue(*(uint32*)&Value);
+}
+
 void FXmlArchiveOutputFormatter::Serialize(FString& Value)
 {
 	Impl->SetValue(*Value);
diff --git a/Engine/Plugins/Experimental/SerializationUtils/Source/XmlSerialization/Public/Formatters/XmlArchiveInputFormatter.h b/Engine/Plugins/Experimental/SerializationUtils/Source/XmlSerialization/Public/Formatters/XmlArchiveInputFormatter.h
--- a/Engine/Plugins/Experimental/SerializationUtils/Source/XmlSerialization/Public/Formatters/XmlArchiveInputFormatter.h
+++ b/Engine/Plugins/Experimental/SerializationUtils/Source/XmlSerialization/Public/Formatters/XmlArchiveInputFormatter.h
@@ -70,6 +70,7 @@ public:
 	virtual void Serialize(float& Value) override;
 	virtual void Serialize(double& Value) override;
 	virtual void Serialize(bool& Value) override;
+	virtual void Serialize(UTF32CHAR& Value) override;
 	virtual void Serialize(FString& Value) override;
 	virtual void Serialize(FName& Value) override;
 	virtual void Serialize(UObject*& Value) override;
diff --git a/Engine/Plugins/Experimental/SerializationUtils/Source/XmlSerialization/Public/Formatters/XmlArchiveOutputFormatter.h b/Engine/Plugins/Experimental/SerializationUtils/Source/XmlSerialization/Public/Formatters/XmlArchiveOutputFormatter.h
--- a/Engine/Plugins/Experimental/SerializationUtils/Source/XmlSerialization/Public/Formatters/XmlArchiveOutputFormatter.h
+++ b/Engine/Plugins/Experimental/SerializationUtils/Source/XmlSerialization/Public/Formatters/XmlArchiveOutputFormatter.h
@@ -76,6 +76,7 @@ public:
 	virtual void Serialize(float& Value) override;
 	virtual void Serialize(double& Value) override;
 	virtual void Serialize(bool& Value) override;
+	virtual void Serialize(UTF32CHAR& Value) override;
 	virtual void Serialize(FString& Value) override;
 	virtual void Serialize(FName& Value) override;
 	virtual void Serialize(UObject*& Value) override;
diff --git a/Engine/Source/Programs/UnrealPackageTool/UnrealPackageTool.cpp b/Engine/Source/Programs/UnrealPackageTool/UnrealPackageTool.cpp
--- a/Engine/Source/Programs/UnrealPackageTool/UnrealPackageTool.cpp
+++ b/Engine/Source/Programs/UnrealPackageTool/UnrealPackageTool.cpp
@@ -607,6 +607,10 @@ public:
 	{
 		GetValueBuffer() << (Value ? UTF8TEXTVIEW("true") : UTF8TEXTVIEW("false"));
 	}
+	virtual void Serialize(UTF32CHAR& Value) override
+	{
+		GetValueBuffer() << Value;
+	}
 	virtual void Serialize(FString& Value) override
 	{
 		if (LooksLikeBinary(FStringView(Value)))
diff --git a/Engine/Source/Runtime/Core/Private/Serialization/StructuredArchiveSlots.cpp b/Engine/Source/Runtime/Core/Private/Serialization/StructuredArchiveSlots.cpp
--- a/Engine/Source/Runtime/Core/Private/Serialization/StructuredArchiveSlots.cpp
+++ b/Engine/Source/Runtime/Core/Private/Serialization/StructuredArchiveSlots.cpp
@@ -239,6 +239,13 @@ void FStructuredArchiveSlot::operator<< (bool& Value)
 	StructuredArchive.LeaveSlot();
 }
 
+void FStructuredArchiveSlot::operator<< (UTF32CHAR& Value)
+{
+	StructuredArchive.EnterSlot(*this);
+	StructuredArchive.Formatter.Serialize(Value);
+	StructuredArchive.LeaveSlot();
+}
+
 void FStructuredArchiveSlot::operator<< (FString& Value)
 {
 	StructuredArchive.EnterSlot(*this);
diff --git a/Engine/Source/Runtime/Core/Public/GenericPlatform/GenericPlatform.h b/Engine/Source/Runtime/Core/Public/GenericPlatform/GenericPlatform.h
--- a/Engine/Source/Runtime/Core/Public/GenericPlatform/GenericPlatform.h
+++ b/Engine/Source/Runtime/Core/Public/GenericPlatform/GenericPlatform.h
@@ -74,8 +74,12 @@ struct FGenericPlatformTypes
 	typedef uint16				CHAR16;		
 
 	// A 32-bit character type. In-memory only. 32-bit representation. Should really be char32_t but making this the generic option is easier for compilers which don't fully support C++11 yet (i.e. MSVC).
+	[[deprecated("CHAR32 has been deprecated - please use UTF32CHAR instead.")]] // not using UE_DEPRECATED(5.6, "...") for dependency reasons
 	typedef uint32				CHAR32;
 
+	// A 32-bit character type. In-memory only. 32-bit representation.
+	using UTF32CHAR = char32_t;
+
 	// A switchable character. In-memory only. Either ANSICHAR or WIDECHAR, depending on a licensee's requirements.
 	typedef WIDECHAR			TCHAR;
 
diff --git a/Engine/Source/Runtime/Core/Public/HAL/Platform.h b/Engine/Source/Runtime/Core/Public/HAL/Platform.h
--- a/Engine/Source/Runtime/Core/Public/HAL/Platform.h
+++ b/Engine/Source/Runtime/Core/Public/HAL/Platform.h
@@ -1080,7 +1080,7 @@ typedef FPlatformTypes::CHAR16		UCS2CHAR;
 /// A 16-bit character containing a UTF16 (Unicode, 16-bit, variable-width) code unit.
 typedef FPlatformTypes::CHAR16		UTF16CHAR;
 /// A 32-bit character containing a UTF32 (Unicode, 32-bit, fixed-width) code unit.
-typedef FPlatformTypes::CHAR32		UTF32CHAR;
+typedef FPlatformTypes::UTF32CHAR	UTF32CHAR;
 
 /// An unsigned integer the same size as a pointer
 typedef FPlatformTypes::UPTRINT UPTRINT;
diff --git a/Engine/Source/Runtime/Core/Public/Misc/StringBuilder.h b/Engine/Source/Runtime/Core/Public/Misc/StringBuilder.h
--- a/Engine/Source/Runtime/Core/Public/Misc/StringBuilder.h
+++ b/Engine/Source/Runtime/Core/Public/Misc/StringBuilder.h
@@ -520,6 +520,8 @@ inline FWideStringBuilderBase&		operator<<(FWideStringBuilderBase& Builder, WIDE
 inline FUtf8StringBuilderBase&		operator<<(FUtf8StringBuilderBase& Builder, ANSICHAR Char)							{ return Builder.AppendChar(UTF8CHAR(Char)); }
 inline FUtf8StringBuilderBase&		operator<<(FUtf8StringBuilderBase& Builder, UTF8CHAR Char)							{ return Builder.AppendChar(Char); }
 inline FUtf8StringBuilderBase&		operator<<(FUtf8StringBuilderBase& Builder, WIDECHAR Char) = delete;
+inline FWideStringBuilderBase&		operator<<(FWideStringBuilderBase& Builder, UTF32CHAR Char)							{ return Builder.AppendChar(Char); }
+inline FUtf8StringBuilderBase&		operator<<(FUtf8StringBuilderBase& Builder, UTF32CHAR Char)							{ return Builder.AppendChar(Char); }
 
 // Prefer using << instead of += as operator+= is only intended for mechanical FString -> FStringView replacement.
 inline FStringBuilderBase&			operator+=(FStringBuilderBase& Builder, ANSICHAR Char)								{ return Builder.AppendChar(Char); }
diff --git a/Engine/Source/Runtime/Core/Public/Serialization/Archive.h b/Engine/Source/Runtime/Core/Public/Serialization/Archive.h
--- a/Engine/Source/Runtime/Core/Public/Serialization/Archive.h
+++ b/Engine/Source/Runtime/Core/Public/Serialization/Archive.h
@@ -1385,6 +1385,23 @@ public:
 		return Ar;
 	}
 
+	/**
+	 * Serializes a UTF32CHAR value from or into an archive.
+	 *
+	 * @param Ar The archive to serialize from or to.
+	 * @param Value The value to serialize.
+	 */
+	FORCEINLINE friend FArchive& operator<<(FArchive& Ar, UTF32CHAR& Value)
+	{
+#if DEVIRTUALIZE_FLinkerLoad_Serialize
+		if (!Ar.FastPathLoad<sizeof(Value)>(&Value))
+#endif
+		{
+			Ar.ByteOrderSerialize(&Value, sizeof(Value));
+		}
+		return Ar;
+	}
+
 	/**
 	 * Serializes an unsigned 8-bit integer value from or into an archive.
 	 *
diff --git a/Engine/Source/Runtime/Core/Public/Serialization/Formatters/BinaryArchiveFormatter.h b/Engine/Source/Runtime/Core/Public/Serialization/Formatters/BinaryArchiveFormatter.h
--- a/Engine/Source/Runtime/Core/Public/Serialization/Formatters/BinaryArchiveFormatter.h
+++ b/Engine/Source/Runtime/Core/Public/Serialization/Formatters/BinaryArchiveFormatter.h
@@ -67,6 +67,7 @@ public:
 	CORE_API virtual void Serialize(float& Value) override;
 	CORE_API virtual void Serialize(double& Value) override;
 	CORE_API virtual void Serialize(bool& Value) override;
+	CORE_API virtual void Serialize(UTF32CHAR& Value) override;
 	CORE_API virtual void Serialize(FString& Value) override;
 	CORE_API virtual void Serialize(FName& Value) override;
 	CORE_API virtual void Serialize(UObject*& Value) override;
@@ -261,6 +262,11 @@ inline void FBinaryArchiveFormatter::Serialize(bool& Value)
 	Inner << Value;
 }
 
+inline void FBinaryArchiveFormatter::Serialize(UTF32CHAR& Value)
+{
+	Inner << Value;
+}
+
 inline void FBinaryArchiveFormatter::Serialize(FString& Value)
 {
 	Inner << Value;
diff --git a/Engine/Source/Runtime/Core/Public/Serialization/StructuredArchiveFormatter.h b/Engine/Source/Runtime/Core/Public/Serialization/StructuredArchiveFormatter.h
--- a/Engine/Source/Runtime/Core/Public/Serialization/StructuredArchiveFormatter.h
+++ b/Engine/Source/Runtime/Core/Public/Serialization/StructuredArchiveFormatter.h
@@ -106,6 +106,7 @@ public:
 	virtual void Serialize(float& Value) = 0;
 	virtual void Serialize(double& Value) = 0;
 	virtual void Serialize(bool& Value) = 0;
+	virtual void Serialize(UTF32CHAR& Value) = 0;
 	virtual void Serialize(FString& Value) = 0;
 	virtual void Serialize(FName& Value) = 0;
 	virtual void Serialize(UObject*& Value) = 0;
diff --git a/Engine/Source/Runtime/Core/Public/Serialization/StructuredArchiveSlots.h b/Engine/Source/Runtime/Core/Public/Serialization/StructuredArchiveSlots.h
--- a/Engine/Source/Runtime/Core/Public/Serialization/StructuredArchiveSlots.h
+++ b/Engine/Source/Runtime/Core/Public/Serialization/StructuredArchiveSlots.h
@@ -64,6 +64,7 @@ public:
 	CORE_API void operator << (float& Value);
 	CORE_API void operator << (double& Value);
 	CORE_API void operator << (bool& Value);
+	CORE_API void operator << (UTF32CHAR& Value);
 	CORE_API void operator << (FString& Value);
 	CORE_API void operator << (FName& Value);
 	CORE_API void operator << (UObject*& Value);
@@ -297,6 +298,11 @@ FORCEINLINE_DEBUGGABLE void operator<<(FStructuredArchiveSlot Slot, TArray<uint8
 		UE::StructuredArchive::Private::GetFormatterImpl(StructuredArchive).Serialize(Value);
 	}
 
+	FORCEINLINE void FStructuredArchiveSlot::operator<< (UTF32CHAR& Value)
+	{
+		UE::StructuredArchive::Private::GetFormatterImpl(StructuredArchive).Serialize(Value);
+	}
+
 	FORCEINLINE void FStructuredArchiveSlot::operator<< (uint64& Value)
 	{
 		UE::StructuredArchive::Private::GetFormatterImpl(StructuredArchive).Serialize(Value);
diff --git a/Engine/Source/Runtime/CoreUObject/Private/Serialization/Formatters/JsonArchiveInputFormatter.cpp b/Engine/Source/Runtime/CoreUObject/Private/Serialization/Formatters/JsonArchiveInputFormatter.cpp
--- a/Engine/Source/Runtime/CoreUObject/Private/Serialization/Formatters/JsonArchiveInputFormatter.cpp
+++ b/Engine/Source/Runtime/CoreUObject/Private/Serialization/Formatters/JsonArchiveInputFormatter.cpp
@@ -418,6 +418,11 @@ void FJsonArchiveInputFormatter::Serialize(bool& Value)
 	Value = ValueStack.Top()->AsBool();
 }
 
+void FJsonArchiveInputFormatter::Serialize(UTF32CHAR& Value)
+{
+	verify(ValueStack.Top()->TryGetNumber(*(uint32*)&Value));
+}
+
 void FJsonArchiveInputFormatter::Serialize(FString& Value)
 {
 	// If the string we serialized was empty, this value will be a null object rather than a string, so we have to
diff --git a/Engine/Source/Runtime/CoreUObject/Private/Serialization/Formatters/JsonArchiveOutputFormatter.cpp b/Engine/Source/Runtime/CoreUObject/Private/Serialization/Formatters/JsonArchiveOutputFormatter.cpp
--- a/Engine/Source/Runtime/CoreUObject/Private/Serialization/Formatters/JsonArchiveOutputFormatter.cpp
+++ b/Engine/Source/Runtime/CoreUObject/Private/Serialization/Formatters/JsonArchiveOutputFormatter.cpp
@@ -302,6 +302,11 @@ void FJsonArchiveOutputFormatter::Serialize(bool& Value)
 	WriteValue(LexToString(Value));
 }
 
+void FJsonArchiveOutputFormatter::Serialize(UTF32CHAR& Value)
+{
+	WriteValue(LexToString(*(uint32*)&Value));
+}
+
 void FJsonArchiveOutputFormatter::Serialize(FString& Value)
 {
 	// Insert a "String:" prefix to prevent incorrect interpretation as another explicit type
diff --git a/Engine/Source/Runtime/CoreUObject/Public/Serialization/Formatters/JsonArchiveInputFormatter.h b/Engine/Source/Runtime/CoreUObject/Public/Serialization/Formatters/JsonArchiveInputFormatter.h
--- a/Engine/Source/Runtime/CoreUObject/Public/Serialization/Formatters/JsonArchiveInputFormatter.h
+++ b/Engine/Source/Runtime/CoreUObject/Public/Serialization/Formatters/JsonArchiveInputFormatter.h
@@ -85,6 +85,7 @@ public:
 	COREUOBJECT_API virtual void Serialize(float& Value) override;
 	COREUOBJECT_API virtual void Serialize(double& Value) override;
 	COREUOBJECT_API virtual void Serialize(bool& Value) override;
+	COREUOBJECT_API virtual void Serialize(UTF32CHAR& Value) override;
 	COREUOBJECT_API virtual void Serialize(FString& Value) override;
 	COREUOBJECT_API virtual void Serialize(FName& Value) override;
 	COREUOBJECT_API virtual void Serialize(UObject*& Value) override;
diff --git a/Engine/Source/Runtime/CoreUObject/Public/Serialization/Formatters/JsonArchiveOutputFormatter.h b/Engine/Source/Runtime/CoreUObject/Public/Serialization/Formatters/JsonArchiveOutputFormatter.h
--- a/Engine/Source/Runtime/CoreUObject/Public/Serialization/Formatters/JsonArchiveOutputFormatter.h
+++ b/Engine/Source/Runtime/CoreUObject/Public/Serialization/Formatters/JsonArchiveOutputFormatter.h
@@ -81,6 +81,7 @@ public:
 	COREUOBJECT_API virtual void Serialize(float& Value) override;
 	COREUOBJECT_API virtual void Serialize(double& Value) override;
 	COREUOBJECT_API virtual void Serialize(bool& Value) override;
+	COREUOBJECT_API virtual void Serialize(UTF32CHAR& Value) override;
 	COREUOBJECT_API virtual void Serialize(FString& Value) override;
 	COREUOBJECT_API virtual void Serialize(FName& Value) override;
 	COREUOBJECT_API virtual void Serialize(UObject*& Value) override;
diff --git a/Engine/Source/Runtime/Experimental/JsonObjectGraph/Private/JsonStringifyStructuredArchive.cpp b/Engine/Source/Runtime/Experimental/JsonObjectGraph/Private/JsonStringifyStructuredArchive.cpp
--- a/Engine/Source/Runtime/Experimental/JsonObjectGraph/Private/JsonStringifyStructuredArchive.cpp
+++ b/Engine/Source/Runtime/Experimental/JsonObjectGraph/Private/JsonStringifyStructuredArchive.cpp
@@ -452,6 +452,11 @@ void FJsonStringifyStructuredArchive::Serialize(bool& Value)
 	WriteValue(LexToString(Value));
 }
 
+void FJsonStringifyStructuredArchive::Serialize(UTF32CHAR& Value)
+{
+	WriteValue(LexToString(*(uint32*)&Value));
+}
+
 void FJsonStringifyStructuredArchive::Serialize(FString& Value)
 {
 	// Insert a "String:" prefix to prevent incorrect interpretation as another explicit type
diff --git a/Engine/Source/Runtime/Experimental/JsonObjectGraph/Private/JsonStringifyStructuredArchive.h b/Engine/Source/Runtime/Experimental/JsonObjectGraph/Private/JsonStringifyStructuredArchive.h
--- a/Engine/Source/Runtime/Experimental/JsonObjectGraph/Private/JsonStringifyStructuredArchive.h
+++ b/Engine/Source/Runtime/Experimental/JsonObjectGraph/Private/JsonStringifyStructuredArchive.h
@@ -84,6 +84,7 @@ private:
 	virtual void Serialize(float& Value) override;
 	virtual void Serialize(double& Value) override;
 	virtual void Serialize(bool& Value) override;
+	virtual void Serialize(UTF32CHAR& Value) override;
 	virtual void Serialize(FString& Value) override;
 	virtual void Serialize(FName& Value) override;
 	virtual void Serialize(UObject*& Value) override;
