Logging: Simplified and removed duplication in FOutputDeviceStdOut

#rnx

diff --git a/Engine/Source/Runtime/Core/Private/Misc/OutputDeviceStdOut.cpp b/Engine/Source/Runtime/Core/Private/Misc/OutputDeviceStdOut.cpp
--- a/Engine/Source/Runtime/Core/Private/Misc/OutputDeviceStdOut.cpp
+++ b/Engine/Source/Runtime/Core/Private/Misc/OutputDeviceStdOut.cpp
@@ -17,10 +17,164 @@
 #include "Windows/WindowsHWrapper.h"
 #endif
 
+// Several functions below are FORCENOINLINE to reduce total required stack space by limiting
+// the scope of string builders and compact binary writers.
+
+namespace UE::Logging::Private
+{
+
+#if PLATFORM_WINDOWS || PLATFORM_TCHAR_IS_UTF8CHAR || PLATFORM_TCHAR_IS_CHAR16
+using StdOutCharType = UTF8CHAR;
+#else
+using StdOutCharType = WIDECHAR;
+#endif
+
+static void WriteLineToStdOut(FUtf8StringBuilderBase& Line)
+{
+	printf("%s", (const char*)*Line);
+	fflush(stdout);
+}
+
+#if PLATFORM_WINDOWS || !(PLATFORM_TCHAR_IS_UTF8CHAR || PLATFORM_TCHAR_IS_CHAR16)
+static void WriteLineToStdOut(FWideStringBuilderBase& Line)
+{
+#if PLATFORM_USE_LS_SPEC_FOR_WIDECHAR
+	// printf prints wchar_t strings just fine with %ls, while mixing printf()/wprintf() is not recommended (see https://stackoverflow.com/questions/8681623/printf-and-wprintf-in-single-c-code)
+	printf("%ls", *Line);
+#else
+	wprintf(TEXT("%s"), *Line);
+#endif
+	fflush(stdout);
+}
+#endif
+
+#if PLATFORM_WINDOWS
+static void WriteLineToConsole(FWideStringBuilderBase& Line)
+{
+	WriteConsoleW(GetStdHandle(STD_OUTPUT_HANDLE), *Line, Line.Len(), nullptr, nullptr);
+}
+#endif
+
+FORCEINLINE static void AddFormatPrefix(FUtf8StringBuilderBase& Format, bool bShowCategory, bool bShowVerbosity)
+{
+	if (bShowCategory)
+	{
+		Format << ANSITEXTVIEW("{_channel}: ");
+	}
+	if (bShowVerbosity)
+	{
+		Format << ANSITEXTVIEW("{_severity}: ");
+	}
+}
+
+FORCEINLINE static void AddMessagePrefix(FUtf8StringBuilderBase& Message, const FName& Category, ELogVerbosity::Type Verbosity, bool bShowCategory, bool bShowVerbosity)
+{
+	if (bShowCategory)
+	{
+		Message << Category << ANSITEXTVIEW(": ");
+	}
+	if (bShowVerbosity)
+	{
+		Message << ToString(Verbosity) << ANSITEXTVIEW(": ");
+	}
+}
+
+FORCEINLINE static void AddMessageFields(FCbWriter& Writer, const FName& Category, ELogVerbosity::Type Verbosity, bool bShowCategory, bool bShowVerbosity)
+{
+	if (bShowCategory)
+	{
+		Writer.BeginObject(ANSITEXTVIEW("_channel"));
+		Writer.AddString(ANSITEXTVIEW("$type"), ANSITEXTVIEW("Channel"));
+		Writer.AddString(ANSITEXTVIEW("$text"), WriteToUtf8String<64>(Category));
+		Writer.EndObject();
+	}
+	if (bShowVerbosity)
+	{
+		Writer.BeginObject(ANSITEXTVIEW("_severity"));
+		Writer.AddString(ANSITEXTVIEW("$type"), ANSITEXTVIEW("Severity"));
+		Writer.AddString(ANSITEXTVIEW("$text"), ToString(Verbosity));
+		Writer.EndObject();
+	}
+}
+
+FORCENOINLINE static void AddMessage(FCbWriter& Writer, const TCHAR* V, ELogVerbosity::Type Verbosity, const FName& Category, bool bShowCategory, bool bShowVerbosity)
+{
+	TUtf8StringBuilder<512> Message;
+	AddMessagePrefix(Message, Category, Verbosity, bShowCategory, bShowVerbosity);
+	Message << V;
+	Writer.AddString(ANSITEXTVIEW("message"), Message);
+}
+
+FORCENOINLINE static void AddFormat(FCbWriter& Writer, const TCHAR* Message, bool bShowCategory, bool bShowVerbosity)
+{
+	TUtf8StringBuilder<512> Format;
+	AddFormatPrefix(Format, bShowCategory, bShowVerbosity);
+
+	// Escape {} in Message
+	for (constexpr FAsciiSet Brackets("{}");;)
+	{
+		const TCHAR* End = FAsciiSet::FindFirstOrEnd(Message, Brackets);
+		Format.Append(Message, UE_PTRDIFF_TO_INT32(End - Message));
+		if (!*End)
+		{
+			break;
+		}
+		Format.AppendChar(*End);
+		Format.AppendChar(*End);
+		Message = End + 1;
+	}
+
+	Writer.AddString(ANSITEXTVIEW("format"), Format);
+}
+
+FORCENOINLINE static void AddMessage(FCbWriter& Writer, const UE::FLogRecord& Record, bool bShowCategory, bool bShowVerbosity)
+{
+	TUtf8StringBuilder<512> Message;
+	AddMessagePrefix(Message, Record.GetCategory(), Record.GetVerbosity(), bShowCategory, bShowVerbosity);
+	Record.FormatMessageTo(Message);
+	Writer.AddString(ANSITEXTVIEW("message"), Message);
+}
+
+static FAnsiStringView GetLevel(ELogVerbosity::Type Verbosity)
+{
+	switch (Verbosity & ELogVerbosity::VerbosityMask)
+	{
+	case ELogVerbosity::Fatal:
+		return ANSITEXTVIEW("Critical");
+	case ELogVerbosity::Error:
+		return ANSITEXTVIEW("Error");
+	case ELogVerbosity::Warning:
+		return ANSITEXTVIEW("Warning");
+	case ELogVerbosity::Display:
+	case ELogVerbosity::Log:
+	default:
+		return ANSITEXTVIEW("Information");
+	case ELogVerbosity::Verbose:
+	case ELogVerbosity::VeryVerbose:
+		return ANSITEXTVIEW("Debug");
+	}
+}
+
+} // UE::Logging::Private
+
+#if PLATFORM_WINDOWS
+static bool IsStdOutAttachedToConsole()
+{
+	if (HANDLE Handle = GetStdHandle(STD_OUTPUT_HANDLE); Handle != INVALID_HANDLE_VALUE)
+	{
+		if (DWORD FileType = GetFileType(Handle); FileType == FILE_TYPE_CHAR)
+		{
+			return true;
+		}
+	}
+	return false;
+}
+#endif
+
 FOutputDeviceStdOutput::FOutputDeviceStdOutput()
 {
 #if PLATFORM_WINDOWS
-	bIsConsoleOutput = IsStdoutAttachedToConsole() && !FParse::Param(FCommandLine::Get(), TEXT("GenericConsoleOutput"));
+	bIsConsoleOutput = IsStdOutAttachedToConsole() && !FParse::Param(FCommandLine::Get(), TEXT("GenericConsoleOutput"));
 #endif
 
 	bIsJsonOutput = FParse::Param(FCommandLine::Get(), TEXT("JsonStdOut"));
@@ -64,23 +218,28 @@ void FOutputDeviceStdOutput::SerializeRecord(const UE::FLogRecord& Record)
 	}
 }
 
-// Several functions below are FORCENOINLINE to reduce total required stack space by limiting
-// the scope of string builders and compact binary writers.
-
 FORCENOINLINE void FOutputDeviceStdOutput::SerializeAsText(const TCHAR* V, ELogVerbosity::Type Verbosity, const FName& Category, double Time)
 {
+	using namespace UE::Logging::Private;
+
+	auto FormatLine = [V, Verbosity, &Category, Time]<typename CharType>(TStringBuilderBase<CharType>& Line) -> TStringBuilderBase<CharType>&
+	{
+		FOutputDeviceHelper::AppendFormatLogLine(Line, Verbosity, Category, V, GPrintLogTimes, Time);
+		Line.AppendChar('\n');
+		return Line;
+	};
+
 #if PLATFORM_WINDOWS
 	if (bIsConsoleOutput)
 	{
-		WriteLine<WIDECHAR>(V, Verbosity, Category, Time);
+		TStringBuilderWithBuffer<WIDECHAR, 512> Line;
+		WriteLineToConsole(FormatLine(Line));
 		return;
 	}
 #endif
-#if PLATFORM_WINDOWS || PLATFORM_TCHAR_IS_UTF8CHAR || PLATFORM_TCHAR_IS_CHAR16
-	WriteLine<UTF8CHAR>(V, Verbosity, Category, Time);
-#else
-	WriteLine<WIDECHAR>(V, Verbosity, Category, Time);
-#endif
+
+	TStringBuilderWithBuffer<StdOutCharType, 512> Line;
+	WriteLineToStdOut(FormatLine(Line));
 }
 
 FORCENOINLINE void FOutputDeviceStdOutput::SerializeRecordAsText(const UE::FLogRecord& Record)
@@ -90,17 +249,10 @@ FORCENOINLINE void FOutputDeviceStdOutput::SerializeRecordAsText(const UE::FLogR
 	Serialize(*V, Record.GetVerbosity(), Record.GetCategory(), -1.0);
 }
 
-template <typename CharType>
-FORCENOINLINE void FOutputDeviceStdOutput::WriteLine(const TCHAR* V, ELogVerbosity::Type Verbosity, const FName& Category, double Time)
-{
-	TStringBuilderWithBuffer<CharType, 512> Line;
-	FOutputDeviceHelper::AppendFormatLogLine(Line, Verbosity, Category, V, GPrintLogTimes, Time);
-	Line.AppendChar('\n');
-	WriteLine(Line);
-}
-
 FORCENOINLINE void FOutputDeviceStdOutput::SerializeAsJson(const TCHAR* V, ELogVerbosity::Type Verbosity, const FName& Category, double Time)
 {
+	using namespace UE::Logging::Private;
+
 	const bool bShowCategory = GPrintLogCategory && !Category.IsNone();
 	const bool bShowVerbosity = GPrintLogVerbosity && (Verbosity & ELogVerbosity::VerbosityMask) != ELogVerbosity::Log;
 
@@ -111,76 +263,22 @@ FORCENOINLINE void FOutputDeviceStdOutput::SerializeAsJson(const TCHAR* V, ELogV
 	AddMessage(Writer, V, Verbosity, Category, bShowCategory, bShowVerbosity);
 	if (bShowCategory || bShowVerbosity)
 	{
-		AddFormat(Writer, V, bShowCategory, bShowVerbosity);
-
 		Writer.BeginObject(ANSITEXTVIEW("properties"));
-		if (bShowCategory)
-		{
-			Writer.BeginObject(ANSITEXTVIEW("_channel"));
-			Writer.AddString(ANSITEXTVIEW("$type"), ANSITEXTVIEW("Channel"));
-			Writer.AddString(ANSITEXTVIEW("$text"), WriteToUtf8String<64>(Category));
-			Writer.EndObject();
-		}
-		if (bShowVerbosity)
-		{
-			Writer.BeginObject(ANSITEXTVIEW("_severity"));
-			Writer.AddString(ANSITEXTVIEW("$type"), ANSITEXTVIEW("Severity"));
-			Writer.AddString(ANSITEXTVIEW("$text"), ToString(Verbosity));
-			Writer.EndObject();
-		}
+		AddMessageFields(Writer, Category, Verbosity, bShowCategory, bShowVerbosity);
 		Writer.EndObject();
+
+		AddFormat(Writer, V, bShowCategory, bShowVerbosity);
 	}
 	Writer.EndObject();
 
 	WriteAsJson(Writer);
 }
 
-FORCENOINLINE void FOutputDeviceStdOutput::AddMessage(FCbWriter& Writer, const TCHAR* V, ELogVerbosity::Type Verbosity, const FName& Category, bool bShowCategory, bool bShowVerbosity)
-{
-	TUtf8StringBuilder<512> Message;
-	if (bShowCategory)
-	{
-		Message << Category << ANSITEXTVIEW(": ");
-	}
-	if (bShowVerbosity)
-	{
-		Message << ToString(Verbosity) << ANSITEXTVIEW(": ");
-	}
-	Message << V;
-	Writer.AddString(ANSITEXTVIEW("message"), Message);
-}
-
-FORCENOINLINE void FOutputDeviceStdOutput::AddFormat(FCbWriter& Writer, const TCHAR* Message, bool bShowCategory, bool bShowVerbosity)
-{
-	TUtf8StringBuilder<512> Format;
-	if (bShowCategory)
-	{
-		Format << ANSITEXTVIEW("{_channel}: ");
-	}
-	if (bShowVerbosity)
-	{
-		Format << ANSITEXTVIEW("{_severity}: ");
-	}
-
-	// Escape {} in Message
-	for (constexpr FAsciiSet Brackets("{}");;)
-	{
-		const TCHAR* End = FAsciiSet::FindFirstOrEnd(Message, Brackets);
-		Format.Append(Message, UE_PTRDIFF_TO_INT32(End - Message));
-		if (!*End)
-		{
-			break;
-		}
-		Format.AppendChar(*End);
-		Format.AppendChar(*End);
-		Message = End + 1;
-	}
-
-	Writer.AddString(ANSITEXTVIEW("format"), Format);
-}
-
 FORCENOINLINE void FOutputDeviceStdOutput::SerializeRecordAsJson(const UE::FLogRecord& Record)
 {
+	using namespace UE;
+	using namespace UE::Logging::Private;
+
 	const bool bShowCategory = GPrintLogCategory && !Record.GetCategory().IsNone();
 	const bool bShowVerbosity = GPrintLogVerbosity && (Record.GetVerbosity() & ELogVerbosity::VerbosityMask) != ELogVerbosity::Log;
 
@@ -191,23 +289,13 @@ FORCENOINLINE void FOutputDeviceStdOutput::SerializeRecordAsJson(const UE::FLogR
 	AddMessage(Writer, Record, bShowCategory, bShowVerbosity);
 	if (bShowCategory || bShowVerbosity || Record.GetFields())
 	{
-		AddFormat(Writer, Record, bShowCategory, bShowVerbosity);
+		TUtf8StringBuilder<512> Format;
+		AddFormatPrefix(Format, bShowCategory, bShowVerbosity);
+		Format.Append(Record.GetFormat());
+		Writer.AddString(ANSITEXTVIEW("format"), Format);
 
 		Writer.BeginObject(ANSITEXTVIEW("properties"));
-		if (bShowCategory)
-		{
-			Writer.BeginObject(ANSITEXTVIEW("_channel"));
-			Writer.AddString(ANSITEXTVIEW("$type"), ANSITEXTVIEW("Channel"));
-			Writer.AddString(ANSITEXTVIEW("$text"), WriteToUtf8String<64>(Record.GetCategory()));
-			Writer.EndObject();
-		}
-		if (bShowVerbosity)
-		{
-			Writer.BeginObject(ANSITEXTVIEW("_severity"));
-			Writer.AddString(ANSITEXTVIEW("$type"), ANSITEXTVIEW("Severity"));
-			Writer.AddString(ANSITEXTVIEW("$text"), ToString(Record.GetVerbosity()));
-			Writer.EndObject();
-		}
+		AddMessageFields(Writer, Record.GetCategory(), Record.GetVerbosity(), bShowCategory, bShowVerbosity);
 		if (const TCHAR* TextNamespace = Record.GetTextNamespace())
 		{
 			Writer.AddString(ANSITEXTVIEW("_ns"), TextNamespace);
@@ -227,125 +315,30 @@ FORCENOINLINE void FOutputDeviceStdOutput::SerializeRecordAsJson(const UE::FLogR
 	WriteAsJson(Writer);
 }
 
-FORCENOINLINE void FOutputDeviceStdOutput::AddMessage(FCbWriter& Writer, const UE::FLogRecord& Record, bool bShowCategory, bool bShowVerbosity)
+FORCENOINLINE void FOutputDeviceStdOutput::WriteAsJson(const FCbWriter& Writer)
 {
-	TUtf8StringBuilder<512> Message;
-	if (bShowCategory)
-	{
-		Message << Record.GetCategory() << ANSITEXTVIEW(": ");
-	}
-	if (bShowVerbosity)
-	{
-		Message << ToString(Record.GetVerbosity()) << ANSITEXTVIEW(": ");
-	}
-	Record.FormatMessageTo(Message);
-	Writer.AddString(ANSITEXTVIEW("message"), Message);
-}
+	using namespace UE::Logging::Private;
 
-FORCENOINLINE void FOutputDeviceStdOutput::AddFormat(FCbWriter& Writer, const UE::FLogRecord& Record, bool bShowCategory, bool bShowVerbosity)
-{
-	TUtf8StringBuilder<512> Format;
-	if (bShowCategory)
-	{
-		Format << ANSITEXTVIEW("{_channel}: ");
-	}
-	if (bShowVerbosity)
-	{
-		Format << ANSITEXTVIEW("{_severity}: ");
-	}
-	Format.Append(Record.GetFormat());
-	Writer.AddString(ANSITEXTVIEW("format"), Format);
-}
-
-void FOutputDeviceStdOutput::WriteAsJson(const FCbWriter& Writer)
-{
 	TArray<uint8, TInlineAllocator64<512>> Buffer;
 	Buffer.AddUninitialized((int64)Writer.GetSaveSize());
 	FCbFieldView Object = Writer.Save(MakeMemoryView(Buffer));
 
-#if PLATFORM_WINDOWS
-	if (bIsConsoleOutput)
+	auto FormatLine = []<typename CharType>(const FCbFieldView& Field, TStringBuilderBase<CharType>& Line) -> TStringBuilderBase<CharType>&
 	{
-		WriteLine<WIDECHAR>(Object);
-		return;
-	}
-#endif
-#if PLATFORM_WINDOWS || PLATFORM_TCHAR_IS_UTF8CHAR || PLATFORM_TCHAR_IS_CHAR16
-	WriteLine<UTF8CHAR>(Object);
-#else
-	WriteLine<WIDECHAR>(Object);
-#endif
-}
-
-template <typename CharType>
-FORCENOINLINE void FOutputDeviceStdOutput::WriteLine(const FCbFieldView& Field)
-{
-	TStringBuilderWithBuffer<CharType, 512> Line;
-	CompactBinaryToCompactJson(Field, Line);
-	Line.AppendChar('\n');
-	WriteLine(Line);
-}
+		CompactBinaryToCompactJson(Field, Line);
+		Line.AppendChar('\n');
+		return Line;
+	};
 
-void FOutputDeviceStdOutput::WriteLine(FUtf8StringBuilderBase& Line) const
-{
-	printf("%s", (const char*)*Line);
-	fflush(stdout);
-}
-
-#if PLATFORM_WINDOWS || !(PLATFORM_TCHAR_IS_UTF8CHAR || PLATFORM_TCHAR_IS_CHAR16)
-void FOutputDeviceStdOutput::WriteLine(FWideStringBuilderBase& Line) const
-{
 #if PLATFORM_WINDOWS
 	if (bIsConsoleOutput)
 	{
-		WriteConsoleW(GetStdHandle(STD_OUTPUT_HANDLE), *Line, Line.Len(), nullptr, nullptr);
+		TStringBuilderWithBuffer<WIDECHAR, 512> Line;
+		WriteLineToConsole(FormatLine(Object, Line));
 		return;
 	}
 #endif
 
-#if PLATFORM_USE_LS_SPEC_FOR_WIDECHAR
-	// printf prints wchar_t strings just fine with %ls, while mixing printf()/wprintf() is not recommended (see https://stackoverflow.com/questions/8681623/printf-and-wprintf-in-single-c-code)
-	printf("%ls", *Line);
-#else
-	wprintf(TEXT("%s"), *Line);
-#endif
-	fflush(stdout);
-}
-#endif
-
-FAnsiStringView FOutputDeviceStdOutput::GetLevel(ELogVerbosity::Type Verbosity)
-{
-	switch (Verbosity & ELogVerbosity::VerbosityMask)
-	{
-	case ELogVerbosity::Fatal:
-		return ANSITEXTVIEW("Critical");
-	case ELogVerbosity::Error:
-		return ANSITEXTVIEW("Error");
-	case ELogVerbosity::Warning:
-		return ANSITEXTVIEW("Warning");
-	case ELogVerbosity::Display:
-	case ELogVerbosity::Log:
-	default:
-		return ANSITEXTVIEW("Information");
-	case ELogVerbosity::Verbose:
-	case ELogVerbosity::VeryVerbose:
-		return ANSITEXTVIEW("Debug");
-	}
+	TStringBuilderWithBuffer<StdOutCharType, 512> Line;
+	WriteLineToStdOut(FormatLine(Object, Line));
 }
-
-#if PLATFORM_WINDOWS
-bool FOutputDeviceStdOutput::IsStdoutAttachedToConsole()
-{
-	HANDLE StdoutHandle = GetStdHandle(STD_OUTPUT_HANDLE);
-	if (StdoutHandle != INVALID_HANDLE_VALUE)
-	{
-		DWORD FileType = GetFileType(StdoutHandle);
-		if (FileType == FILE_TYPE_CHAR)
-		{
-			return true;
-		}
-	}
-
-	return false;
-}
-#endif
diff --git a/Engine/Source/Runtime/Core/Public/Misc/OutputDeviceStdOut.h b/Engine/Source/Runtime/Core/Public/Misc/OutputDeviceStdOut.h
--- a/Engine/Source/Runtime/Core/Public/Misc/OutputDeviceStdOut.h
+++ b/Engine/Source/Runtime/Core/Public/Misc/OutputDeviceStdOut.h
@@ -2,10 +2,8 @@
 
 #pragma once
 
-#include "Containers/StringFwd.h"
 #include "Misc/OutputDevice.h"
 
-class FCbFieldView;
 class FCbWriter;
 
 class FOutputDeviceStdOutput final : public FOutputDevice
@@ -28,37 +26,13 @@ private:
 	void SerializeAsText(const TCHAR* V, ELogVerbosity::Type Verbosity, const FName& Category, double Time);
 	void SerializeRecordAsText(const UE::FLogRecord& Record);
 
-	template <typename CharType>
-	void WriteLine(const TCHAR* V, ELogVerbosity::Type Verbosity, const FName& Category, double Time);
-
 	void SerializeAsJson(const TCHAR* V, ELogVerbosity::Type Verbosity, const FName& Category, double Time);
-
-	static void AddMessage(FCbWriter& Writer, const TCHAR* V, ELogVerbosity::Type Verbosity, const FName& Category, bool bShowCategory, bool bShowVerbosity);
-	static void AddFormat(FCbWriter& Writer, const TCHAR* Message, bool bShowCategory, bool bShowVerbosity);
-
 	void SerializeRecordAsJson(const UE::FLogRecord& Record);
 
-	static void AddMessage(FCbWriter& Writer, const UE::FLogRecord& Record, bool bShowCategory, bool bShowVerbosity);
-	static void AddFormat(FCbWriter& Writer, const UE::FLogRecord& Record, bool bShowCategory, bool bShowVerbosity);
-
 	void WriteAsJson(const FCbWriter& Writer);
 
-	template <typename CharType>
-	void WriteLine(const FCbFieldView& Field);
-
-	void WriteLine(FUtf8StringBuilderBase& Line) const;
-#if PLATFORM_WINDOWS || !(PLATFORM_TCHAR_IS_UTF8CHAR || PLATFORM_TCHAR_IS_CHAR16)
-	void WriteLine(FWideStringBuilderBase& Line) const;
-#endif
-
-	static FAnsiStringView GetLevel(ELogVerbosity::Type Verbosity);
-
 private:
 	ELogVerbosity::Type AllowedLogVerbosity = ELogVerbosity::Display;
 	bool bIsConsoleOutput = false;
 	bool bIsJsonOutput = false;
-
-#if PLATFORM_WINDOWS
-	static bool IsStdoutAttachedToConsole();
-#endif
 };
