[WindowsPlatformStackWalk]
- Remove global lock that is not needed for CaptureBacktrace and caused deadlocks when FMallocLeakDetection was used
- Remove global initialization that does nothing useful

diff --git a/Engine/Source/Runtime/Core/Private/Windows/WindowsPlatformStackWalk.cpp b/Engine/Source/Runtime/Core/Private/Windows/WindowsPlatformStackWalk.cpp
--- a/Engine/Source/Runtime/Core/Private/Windows/WindowsPlatformStackWalk.cpp
+++ b/Engine/Source/Runtime/Core/Private/Windows/WindowsPlatformStackWalk.cpp
@@ -118,21 +118,8 @@ NTSYSAPI uint16 NTAPI RtlCaptureStackBackTrace(
 	__out_opt PDWORD BackTraceHash
 	);
 
-/** Maximum callstack depth that is supported by the current OS. */
-static ULONG GMaxCallstackDepth = 62;
-
-/** Whether DetermineMaxCallstackDepth() has been called or not. */
-static bool GMaxCallstackDepthInitialized = false;
-
 /** Maximum callstack depth we support, no matter what OS we're running on. */
-#define MAX_CALLSTACK_DEPTH 128
-
-/** Checks the current OS version and sets up the GMaxCallstackDepth variable. */
-void DetermineMaxCallstackDepth()
-{
-	GMaxCallstackDepth = MAX_CALLSTACK_DEPTH;
-	GMaxCallstackDepthInitialized = true;
-}
+static constexpr ULONG GMaxCallstackDepth = 128;
 
 #endif
 
@@ -385,8 +372,6 @@ UE_AUTORTFM_ALWAYS_OPEN uint32 FWindowsPlatformStackWalk::CaptureThreadStackBack
  */
 UE_AUTORTFM_ALWAYS_OPEN uint32 FWindowsPlatformStackWalk::CaptureStackBackTrace( uint64* BackTrace, uint32 MaxDepth, void* Context )
 {
-	UE::TUniqueLock GlobalLock(GStackWalkingLock);
-
 	// Make sure we have place to store the information before we go through the process of raising
 	// an exception and handling it.
 	if (BackTrace == NULL || MaxDepth == 0)
@@ -403,11 +388,7 @@ UE_AUTORTFM_ALWAYS_OPEN uint32 FWindowsPlatformStackWalk::CaptureStackBackTrace(
 	else
 	{
 #if USE_FAST_STACKTRACE
-		if (!GMaxCallstackDepthInitialized)
-		{
-			DetermineMaxCallstackDepth();
-		}
-		PVOID WinBackTrace[MAX_CALLSTACK_DEPTH];
+		PVOID WinBackTrace[GMaxCallstackDepth];
 		ULONG RequestedDepth = FMath::Min<ULONG>(GMaxCallstackDepth, MaxDepth);
 		uint16 NumFrames = RtlCaptureStackBackTrace(0, RequestedDepth, WinBackTrace, NULL);
 		if (NumFrames == 0 && RequestedDepth > 0)
