Fixed TNotNull to not give a C6504 static analysis warning about when wrapping a non-pointer.

diff --git a/Engine/Source/Runtime/Core/Public/Misc/NotNull.h b/Engine/Source/Runtime/Core/Public/Misc/NotNull.h
--- a/Engine/Source/Runtime/Core/Public/Misc/NotNull.h
+++ b/Engine/Source/Runtime/Core/Public/Misc/NotNull.h
@@ -90,12 +90,65 @@
 
 		template <typename... ArgTypes>
 		inline constexpr bool TIsTNotNullParam_V = TIsTNotNullParam<ArgTypes...>::Value;
+
+		template <typename T>
+		struct TNotNullBase
+		{
+			using Sub = TNotNull<T>;
+
+			// Allow implicit conversion to the wrapped value, for compatibility with non-TNotNull APIs
+			operator T() const&
+			{
+				T Result = ((Sub*)this)->Val;
+				UE_ASSUME(Result);
+				return Result;
+			}
+			operator T() &&
+			{
+				T Result = (T&&)((Sub*)this)->Val;
+				UE_ASSUME(Result);
+				return Result;
+			}
+
+			// Allow dereferencing, if it's supported
+			const T& operator->() const
+			{
+				return ((Sub*)this)->Val;
+			}
+		};
+
+		template <typename T>
+		struct TNotNullBase<T*>
+		{
+			using Sub = TNotNull<T*>;
+
+			// Allow implicit conversion to the wrapped value, for compatibility with non-TNotNull APIs
+		UE_NOTNULL_FUNCTION_NON_NULL_RETURN_START
+			operator T*() const
+		UE_NOTNULL_FUNCTION_NON_NULL_RETURN_END
+			{
+				T* Result = ((Sub*)this)->Val;
+				UE_ASSUME(Result);
+				return Result;
+			}
+
+			// Allow dereferencing, if it's supported
+		UE_NOTNULL_FUNCTION_NON_NULL_RETURN_START
+			T* operator->() const
+		UE_NOTNULL_FUNCTION_NON_NULL_RETURN_END
+			{
+				return ((Sub*)this)->Val;
+			}
+		};
 	}
 
 	template <typename T>
-	struct TNotNull
+	struct TNotNull : UE::Core::Private::TNotNullBase<T>
 	{
 	private:
+		using Super = UE::Core::Private::TNotNullBase<T>;
+		friend Super;
+
 		T Val;
 
 	public:
@@ -157,31 +210,7 @@
 		// Disallow testing for nullness - implicitly or explicitly
 		operator bool() const = delete;
 
-		// Allow implicit conversion to the wrapped value, for compatibility with non-TNotNull APIs
-	UE_NOTNULL_FUNCTION_NON_NULL_RETURN_START
-		operator T() const&
-	UE_NOTNULL_FUNCTION_NON_NULL_RETURN_END
-		{
-			T Result = Val;
-			UE_ASSUME(Result);
-			return Result;
-		}
-	UE_NOTNULL_FUNCTION_NON_NULL_RETURN_START
-		operator T() &&
-	UE_NOTNULL_FUNCTION_NON_NULL_RETURN_END
-		{
-			T Result = (T&&)Val;
-			UE_ASSUME(Result);
-			return Result;
-		}
-
 		// Allow dereferencing, if it's supported
-	UE_NOTNULL_FUNCTION_NON_NULL_RETURN_START
-		const T& operator->() const
-	UE_NOTNULL_FUNCTION_NON_NULL_RETURN_END
-		{
-			return Val;
-		}
 		decltype(auto) operator*() const
 		{
 			return *Val;
