Callstack trace module filtering

diff --git a/Engine/Source/Runtime/Core/Private/ProfilingDebugging/Microsoft/WindowsCallstackTrace.cpp b/Engine/Source/Runtime/Core/Private/ProfilingDebugging/Microsoft/WindowsCallstackTrace.cpp
--- a/Engine/Source/Runtime/Core/Private/ProfilingDebugging/Microsoft/WindowsCallstackTrace.cpp
+++ b/Engine/Source/Runtime/Core/Private/ProfilingDebugging/Microsoft/WindowsCallstackTrace.cpp
@@ -271,11 +271,10 @@ void FBacktracer::AddModule(UPTRINT ModuleBase, const TCHAR* Name)
 	const auto* NtHeader = (IMAGE_NT_HEADERS*)(ModuleBase + DosHeader->e_lfanew);
 	const IMAGE_FILE_HEADER* FileHeader = &(NtHeader->FileHeader);
 
-	if(!GFullBacktraces)
+	if (!GFullBacktraces)
 	{
-		int32 NameLen = FCString::Strlen(Name);
-		if (!(NameLen > 4 && FCString::Strcmp(Name + NameLen - 4, TEXT(".exe")) == 0) && // we want always load an .exe file
-			(FCString::Strfind(Name, TEXT("Binaries")) == nullptr || FCString::Strfind(Name, TEXT("ThirdParty")) != nullptr))
+		FStringView NameView(Name);
+		if (!NameView.EndsWith(TEXTVIEW(".exe")) && CallstackTrace_FilterModule(NameView))
 		{
 			return;
 		}
@@ -881,6 +880,14 @@ void CallstackTrace_InitializeInternal()
 	GModulesAreInitialized = true;
 }
 
+////////////////////////////////////////////////////////////////////////////////
+#if PLATFORM_WINDOWS
+bool CallstackTrace_FilterModule(FStringView ModuleName)
+{
+	return !ModuleName.Contains(TEXTVIEW("Binaries")) && ModuleName.Contains(TEXTVIEW("ThirdParty"));
+}
+#endif
+
 ////////////////////////////////////////////////////////////////////////////////
 uint32 CallstackTrace_GetCurrentId()
 {
diff --git a/Engine/Source/Runtime/Core/Public/ProfilingDebugging/CallstackTrace.h b/Engine/Source/Runtime/Core/Public/ProfilingDebugging/CallstackTrace.h
--- a/Engine/Source/Runtime/Core/Public/ProfilingDebugging/CallstackTrace.h
+++ b/Engine/Source/Runtime/Core/Public/ProfilingDebugging/CallstackTrace.h
@@ -5,6 +5,7 @@
 #include "Math/NumericLimits.h"
 #include "Misc/Build.h"
 #include "Trace/Config.h"
+#include "Containers/StringFwd.h"
 
 ////////////////////////////////////////////////////////////////////////////////
 #if !defined(UE_CALLSTACK_TRACE_ENABLED)
@@ -37,6 +38,15 @@ void CallstackTrace_Create(class FMalloc* Malloc);
  */
 void CallstackTrace_Initialize();
 
+/**
+ * Filter modules that should be considered for callstack tracing. Only used on
+ * platforms that supports custom stack walking. See platform implementations for
+ * details.
+ * @param ModuleName Name of the module to be filtered
+ * @return True if module should ignored for stack walking, false otherwise
+ */
+bool CallstackTrace_FilterModule(FStringView ModuleName);
+
 /**
  * Capture the current callstack, and trace the definition if it has not already been encountered. The returned value
  * can be used in trace events and be resolved in analysis.
