Logging: Updated structured logging to support nested $text fields

Updated stdout to use a more basic format because the JSON written to stdout needs to be in the format supported by UAT and Horde. Nested $text fields are collapsed into the root of each field.

diff --git a/Engine/Source/Runtime/Core/Private/Logging/StructuredLog.cpp b/Engine/Source/Runtime/Core/Private/Logging/StructuredLog.cpp
--- a/Engine/Source/Runtime/Core/Private/Logging/StructuredLog.cpp
+++ b/Engine/Source/Runtime/Core/Private/Logging/StructuredLog.cpp
@@ -31,6 +31,12 @@
 void StaticFailDebug (const TCHAR* Error, const ANSICHAR* File, int32 Line, void* ProgramCounter, const TCHAR* Message);
 void StaticFailDebugV(const TCHAR* Error, const ANSICHAR* File, int32 Line, void* ProgramCounter, const TCHAR* DescriptionFormat, va_list DescriptionArgs);
 
+namespace UE::Serialization::Private
+{
+template <typename CharType>
+void AppendQuotedJsonString(TStringBuilderBase<CharType>& Builder, FUtf8StringView Value);
+} // UE::Serialization::Private
+
 namespace UE::Logging::Private
 {
 
@@ -112,11 +118,42 @@ static void LogFieldValue(TStringBuilderBase<CharType>& Out, const FCbFieldView&
 		if (FCbFieldView TextField = Accessor.AsObjectView().FindViewIgnoreCase(ANSITEXTVIEW("$text")))
 		{
 			Out.Append(TextField.AsString());
-			break;
 		}
-		[[fallthrough]];
+		else
+		{
+			Out.AppendChar('{');
+			bool bNeedsComma = false;
+			for (FCbFieldView It : Field)
+			{
+				if (bNeedsComma)
+				{
+					Out.AppendChar(',').AppendChar(' ');
+				}
+				bNeedsComma = true;
+				Serialization::Private::AppendQuotedJsonString(Out, It.GetName());
+				Out.AppendChar(':').AppendChar(' ');
+				LogFieldValue(Out, It);
+			}
+			Out.AppendChar('}');
+		}
+		break;
 	case ECbFieldType::Array:
 	case ECbFieldType::UniformArray:
+		{
+			Out.AppendChar('[');
+			bool bNeedsComma = false;
+			for (FCbFieldView It : Field)
+			{
+				if (bNeedsComma)
+				{
+					Out.AppendChar(',').AppendChar(' ');
+				}
+				bNeedsComma = true;
+				LogFieldValue(Out, It);
+			}
+			Out.AppendChar(']');
+		}
+		break;
 	case ECbFieldType::Binary:
 		CompactBinaryToCompactJson(Field.RemoveName(), Out);
 		break;
@@ -581,6 +618,44 @@ FText FormatLogToText(const FLogTemplate* Template, const FCbFieldViewIterator&
 	return Template->FormatToText(Fields);
 }
 
+void ConvertToCommonLog(FUtf8StringBuilderBase& OutFormat, FCbWriter& OutFields, const FLogTemplate* Template, const FCbFieldIterator& Fields)
+{
+	using namespace Logging::Private;
+
+	for (FCbField& Field : Fields)
+	{
+		OutFields.SetName(Field.GetName());
+		if (FCbArray Array = Field.AsArray(); !Field.HasError())
+		{
+			OutFields.BeginObject();
+			OutFields.AddArray(ANSITEXTVIEW("$value"), Array);
+			TUtf8StringBuilder<256> Text;
+			LogFieldValue(Text, Field);
+			OutFields.AddString(ANSITEXTVIEW("$text"), Text);
+			OutFields.EndObject();
+		}
+		else if (FCbObject Object = Field.AsObject(); !Field.HasError() && !Object.FindView(ANSITEXTVIEW("$text")))
+		{
+			OutFields.BeginObject();
+			for (FCbField& Child : Object)
+			{
+				OutFields.AddField(Child.GetName(), Child);
+			}
+			TUtf8StringBuilder<256> Text;
+			LogFieldValue(Text, Field);
+			OutFields.AddString(ANSITEXTVIEW("$text"), Text);
+			OutFields.EndObject();
+		}
+		else
+		{
+			OutFields.AddField(Field);
+		}
+	}
+
+	// TODO: Process localized format strings to remove argument modifiers and convert escaped braces.
+	OutFormat.Append(Template->GetFormat());
+}
+
 ///////////////////////////////////////////////////////////////////////////////////////////////////////////////////////
 
 FLogTime FLogTime::Now()
diff --git a/Engine/Source/Runtime/Core/Private/Misc/OutputDeviceStdOut.cpp b/Engine/Source/Runtime/Core/Private/Misc/OutputDeviceStdOut.cpp
--- a/Engine/Source/Runtime/Core/Private/Misc/OutputDeviceStdOut.cpp
+++ b/Engine/Source/Runtime/Core/Private/Misc/OutputDeviceStdOut.cpp
@@ -4,6 +4,7 @@
 
 #include "CoreGlobals.h"
 #include "Logging/StructuredLog.h"
+#include "Logging/StructuredLogFormat.h"
 #include "Misc/AsciiSet.h"
 #include "Misc/CommandLine.h"
 #include "Misc/DateTime.h"
@@ -291,8 +292,6 @@ FORCENOINLINE void FOutputDeviceStdOutput::SerializeRecordAsJson(const UE::FLogR
 	{
 		TUtf8StringBuilder<512> Format;
 		AddFormatPrefix(Format, bShowCategory, bShowVerbosity);
-		Format.Append(Record.GetFormat());
-		Writer.AddString(ANSITEXTVIEW("format"), Format);
 
 		Writer.BeginObject(ANSITEXTVIEW("properties"));
 		AddMessageFields(Writer, Record.GetCategory(), Record.GetVerbosity(), bShowCategory, bShowVerbosity);
@@ -304,11 +303,10 @@ FORCENOINLINE void FOutputDeviceStdOutput::SerializeRecordAsJson(const UE::FLogR
 		{
 			Writer.AddString(ANSITEXTVIEW("_key"), TextKey);
 		}
-		for (const FCbField& Field : Record.GetFields())
-		{
-			Writer.AddField(Field.GetName(), Field);
-		}
+		ConvertToCommonLog(Format, Writer, Record.GetTemplate(), Record.GetFields().CreateIterator());
 		Writer.EndObject();
+
+		Writer.AddString(ANSITEXTVIEW("format"), Format);
 	}
 	Writer.EndObject();
 
diff --git a/Engine/Source/Runtime/Core/Private/Serialization/CompactBinarySerialization.cpp b/Engine/Source/Runtime/Core/Private/Serialization/CompactBinarySerialization.cpp
--- a/Engine/Source/Runtime/Core/Private/Serialization/CompactBinarySerialization.cpp
+++ b/Engine/Source/Runtime/Core/Private/Serialization/CompactBinarySerialization.cpp
@@ -35,9 +35,118 @@
 namespace UE::Serialization::Private
 {
 
+template <typename CharType>
+struct TCbJsonWriterConstants;
+
+template <>
+struct TCbJsonWriterConstants<ANSICHAR>
+{
+	static inline const FAnsiStringView LineTerminator = LINE_TERMINATOR_ANSI;
+	static inline const FAnsiStringView Space = ANSITEXTVIEW(" ");
+	static inline const FAnsiStringView Tab = ANSITEXTVIEW("\t");
+
+	static inline const FAnsiStringView Null = ANSITEXTVIEW("null");
+	static inline const FAnsiStringView True = ANSITEXTVIEW("true");
+	static inline const FAnsiStringView False = ANSITEXTVIEW("false");
+
+	static inline const FAnsiStringView EscapeBackslash = ANSITEXTVIEW("\\\\");
+	static inline const FAnsiStringView EscapeDoubleQuote = ANSITEXTVIEW("\\\"");
+	static inline const FAnsiStringView EscapeBell = ANSITEXTVIEW("\\b");
+	static inline const FAnsiStringView EscapeFormFeed = ANSITEXTVIEW("\\f");
+	static inline const FAnsiStringView EscapeLineFeed = ANSITEXTVIEW("\\n");
+	static inline const FAnsiStringView EscapeCarriageReturn = ANSITEXTVIEW("\\r");
+	static inline const FAnsiStringView EscapeTab = ANSITEXTVIEW("\\t");
+
+	static inline const FAnsiStringView FieldId = ANSITEXTVIEW("\"Id\":");
+	static inline const FAnsiStringView FieldName = ANSITEXTVIEW("\"Name\":");
+	static inline const FAnsiStringView FieldData = ANSITEXTVIEW("\"Data\":");
+
+	static inline decltype(auto) FormatEscapeUnicode4 = "\\u%04x";
+	static inline decltype(auto) FormatFloat32 = "%.9g";
+	static inline decltype(auto) FormatFloat64 = "%.17g";
+};
+
+template <>
+struct TCbJsonWriterConstants<WIDECHAR>
+{
+	static inline const FWideStringView LineTerminator = LINE_TERMINATOR;
+	static inline const FWideStringView Space = WIDETEXTVIEW(" ");
+	static inline const FWideStringView Tab = WIDETEXTVIEW("\t");
+
+	static inline const FWideStringView Null = WIDETEXTVIEW("null");
+	static inline const FWideStringView True = WIDETEXTVIEW("true");
+	static inline const FWideStringView False = WIDETEXTVIEW("false");
+
+	static inline const FWideStringView EscapeBackslash = WIDETEXTVIEW("\\\\");
+	static inline const FWideStringView EscapeDoubleQuote = WIDETEXTVIEW("\\\"");
+	static inline const FWideStringView EscapeBell = WIDETEXTVIEW("\\b");
+	static inline const FWideStringView EscapeFormFeed = WIDETEXTVIEW("\\f");
+	static inline const FWideStringView EscapeLineFeed = WIDETEXTVIEW("\\n");
+	static inline const FWideStringView EscapeCarriageReturn = WIDETEXTVIEW("\\r");
+	static inline const FWideStringView EscapeTab = WIDETEXTVIEW("\\t");
+
+	static inline const FWideStringView FieldId = WIDETEXTVIEW("\"Id\":");
+	static inline const FWideStringView FieldName = WIDETEXTVIEW("\"Name\":");
+	static inline const FWideStringView FieldData = WIDETEXTVIEW("\"Data\":");
+
+	static inline decltype(auto) FormatEscapeUnicode4 = WIDETEXT("\\u%04x");
+	static inline decltype(auto) FormatFloat32 = WIDETEXT("%.9g");
+	static inline decltype(auto) FormatFloat64 = WIDETEXT("%.17g");
+};
+
+template <>
+struct TCbJsonWriterConstants<UTF8CHAR> : TCbJsonWriterConstants<ANSICHAR>
+{
+};
+
+template <typename CharType>
+void AppendQuotedJsonString(TStringBuilderBase<CharType>& Builder, FUtf8StringView Value)
+{
+	using FConstants = TCbJsonWriterConstants<CharType>;
+	const FAsciiSet EscapeSet("\\\"\b\f\n\r\t"
+		"\x00\x01\x02\x03\x04\x05\x06\x07\x08\x09\x0a\x0b\x0c\x0d\x0e\x0f"
+		"\x10\x11\x12\x13\x14\x15\x16\x17\x18\x19\x1a\x1b\x1c\x1d\x1e\x1f");
+	Builder << '\"';
+	while (!Value.IsEmpty())
+	{
+		FUtf8StringView Verbatim = FAsciiSet::FindPrefixWithout(Value, EscapeSet);
+		Builder << Verbatim;
+		Value.RightChopInline(Verbatim.Len());
+		FUtf8StringView Escape = FAsciiSet::FindPrefixWith(Value, EscapeSet);
+		for (CharType Char : Escape)
+		{
+			switch (Char)
+			{
+			case '\\': Builder << FConstants::EscapeBackslash; break;
+			case '\"': Builder << FConstants::EscapeDoubleQuote; break;
+			case '\b': Builder << FConstants::EscapeBell; break;
+			case '\f': Builder << FConstants::EscapeFormFeed; break;
+			case '\n': Builder << FConstants::EscapeLineFeed; break;
+			case '\r': Builder << FConstants::EscapeCarriageReturn; break;
+			case '\t': Builder << FConstants::EscapeTab; break;
+			default:
+				Builder.Appendf(FConstants::FormatEscapeUnicode4, uint32(Char));
+				break;
+			}
+		}
+		Value.RightChopInline(Escape.Len());
+	}
+	Builder << '\"';
+}
+
+template void AppendQuotedJsonString<UTF8CHAR>(FUtf8StringBuilderBase& Builder, FUtf8StringView Value);
+template void AppendQuotedJsonString<WIDECHAR>(FWideStringBuilderBase& Builder, FUtf8StringView Value);
+
+} // UE::Serialization::Private
+
+///////////////////////////////////////////////////////////////////////////////////////////////////////////////////////
+
+namespace UE::Serialization::Private
+{
+
 void LogFieldTooLargeForArrayWarning(uint64 FieldLength)
 {
-	UE_LOG(LogSerialization, Warning, TEXT("Attempting to load a FCbFieldView that has too many entries (%" UINT64_FMT ") to fit in a TArray "), FieldLength);
+	UE_LOG(LogSerialization, Warning, TEXT("Attempting to load a FCbFieldView that has too many entries (%" UINT64_FMT ") to fit in a TArray."), FieldLength);
 }
 
 } // namespace UE::Serialization::Private
@@ -366,68 +475,9 @@
 ///////////////////////////////////////////////////////////////////////////////////////////////////////////////////////
 
 template <typename CharType>
-struct TCbJsonWriterConstants;
-
-template <>
-struct TCbJsonWriterConstants<ANSICHAR>
-{
-	static inline const FAnsiStringView LineTerminator = LINE_TERMINATOR_ANSI;
-	static inline const FAnsiStringView Space = ANSITEXTVIEW(" ");
-	static inline const FAnsiStringView Tab = ANSITEXTVIEW("\t");
-
-	static inline const FAnsiStringView Null = ANSITEXTVIEW("null");
-	static inline const FAnsiStringView True = ANSITEXTVIEW("true");
-	static inline const FAnsiStringView False = ANSITEXTVIEW("false");
-
-	static inline const FAnsiStringView EscapeBackslash = ANSITEXTVIEW("\\\\");
-	static inline const FAnsiStringView EscapeDoubleQuote = ANSITEXTVIEW("\\\"");
-	static inline const FAnsiStringView EscapeBell = ANSITEXTVIEW("\\b");
-	static inline const FAnsiStringView EscapeFormFeed = ANSITEXTVIEW("\\f");
-	static inline const FAnsiStringView EscapeLineFeed = ANSITEXTVIEW("\\n");
-	static inline const FAnsiStringView EscapeCarriageReturn = ANSITEXTVIEW("\\r");
-	static inline const FAnsiStringView EscapeTab = ANSITEXTVIEW("\\t");
-
-	static inline const FAnsiStringView FieldId = ANSITEXTVIEW("\"Id\":");
-	static inline const FAnsiStringView FieldName = ANSITEXTVIEW("\"Name\":");
-	static inline const FAnsiStringView FieldData = ANSITEXTVIEW("\"Data\":");
-
-	static inline decltype(auto) FormatEscapeUnicode4 = "\\u%04x";
-	static inline decltype(auto) FormatFloat32 = "%.9g";
-	static inline decltype(auto) FormatFloat64 = "%.17g";
-};
-
-template <>
-struct TCbJsonWriterConstants<WIDECHAR>
-{
-	static inline const FWideStringView LineTerminator = LINE_TERMINATOR;
-	static inline const FWideStringView Space = WIDETEXTVIEW(" ");
-	static inline const FWideStringView Tab = WIDETEXTVIEW("\t");
-
-	static inline const FWideStringView Null = WIDETEXTVIEW("null");
-	static inline const FWideStringView True = WIDETEXTVIEW("true");
-	static inline const FWideStringView False = WIDETEXTVIEW("false");
-
-	static inline const FWideStringView EscapeBackslash = WIDETEXTVIEW("\\\\");
-	static inline const FWideStringView EscapeDoubleQuote = WIDETEXTVIEW("\\\"");
-	static inline const FWideStringView EscapeBell = WIDETEXTVIEW("\\b");
-	static inline const FWideStringView EscapeFormFeed = WIDETEXTVIEW("\\f");
-	static inline const FWideStringView EscapeLineFeed = WIDETEXTVIEW("\\n");
-	static inline const FWideStringView EscapeCarriageReturn = WIDETEXTVIEW("\\r");
-	static inline const FWideStringView EscapeTab = WIDETEXTVIEW("\\t");
-
-	static inline const FWideStringView FieldId = WIDETEXTVIEW("\"Id\":");
-	static inline const FWideStringView FieldName = WIDETEXTVIEW("\"Name\":");
-	static inline const FWideStringView FieldData = WIDETEXTVIEW("\"Data\":");
-
-	static inline decltype(auto) FormatEscapeUnicode4 = WIDETEXT("\\u%04x");
-	static inline decltype(auto) FormatFloat32 = WIDETEXT("%.9g");
-	static inline decltype(auto) FormatFloat64 = WIDETEXT("%.17g");
-};
-
-template <typename CharType, typename ConstantCharType = CharType>
 class TCbJsonWriter
 {
-	using FConstants = TCbJsonWriterConstants<ConstantCharType>;
+	using FConstants = UE::Serialization::Private::TCbJsonWriterConstants<CharType>;
 
 public:
 	explicit TCbJsonWriter(
@@ -593,35 +643,7 @@
 
 	void AppendQuotedString(FUtf8StringView Value)
 	{
-		const FAsciiSet EscapeSet("\\\"\b\f\n\r\t"
-			"\x00\x01\x02\x03\x04\x05\x06\x07\x08\x09\x0a\x0b\x0c\x0d\x0e\x0f"
-			"\x10\x11\x12\x13\x14\x15\x16\x17\x18\x19\x1a\x1b\x1c\x1d\x1e\x1f");
-		Builder << '\"';
-		while (!Value.IsEmpty())
-		{
-			FUtf8StringView Verbatim = FAsciiSet::FindPrefixWithout(Value, EscapeSet);
-			Builder << Verbatim;
-			Value.RightChopInline(Verbatim.Len());
-			FUtf8StringView Escape = FAsciiSet::FindPrefixWith(Value, EscapeSet);
-			for (CharType Char : Escape)
-			{
-				switch (Char)
-				{
-				case '\\': Builder << FConstants::EscapeBackslash; break;
-				case '\"': Builder << FConstants::EscapeDoubleQuote; break;
-				case '\b': Builder << FConstants::EscapeBell; break;
-				case '\f': Builder << FConstants::EscapeFormFeed; break;
-				case '\n': Builder << FConstants::EscapeLineFeed; break;
-				case '\r': Builder << FConstants::EscapeCarriageReturn; break;
-				case '\t': Builder << FConstants::EscapeTab; break;
-				default:
-					Builder.Appendf(FConstants::FormatEscapeUnicode4, uint32(Char));
-					break;
-				}
-			}
-			Value.RightChopInline(Escape.Len());
-		}
-		Builder << '\"';
+		UE::Serialization::Private::AppendQuotedJsonString(Builder, Value);
 	}
 
 	void AppendBase64String(FMemoryView Value)
@@ -632,7 +654,7 @@
 		const uint32 EncodedSize = FBase64::GetEncodedDataSize(uint32(Value.GetSize()));
 		const int32 EncodedIndex = Builder.AddUninitialized(int32(EncodedSize));
 		FBase64::Encode(static_cast<const uint8*>(Value.GetData()), uint32(Value.GetSize()),
-			reinterpret_cast<ConstantCharType*>(Builder.GetData() + EncodedIndex));
+			reinterpret_cast<CharType*>(Builder.GetData() + EncodedIndex));
 		Builder << '"';
 	}
 
@@ -647,7 +669,7 @@
 
 void CompactBinaryToJson(const FCbFieldView& Field, FUtf8StringBuilderBase& Builder)
 {
-	TCbJsonWriter<UTF8CHAR, ANSICHAR> Writer(Builder);
+	TCbJsonWriter<UTF8CHAR> Writer(Builder);
 	Writer.WriteField(Field);
 }
 
@@ -679,7 +701,7 @@
 
 void CompactBinaryToCompactJson(const FCbFieldView& Field, FUtf8StringBuilderBase& Builder)
 {
-	TCbJsonWriter<UTF8CHAR, ANSICHAR> Writer(Builder, ANSITEXTVIEW(""), ANSITEXTVIEW(""), ANSITEXTVIEW(""));
+	TCbJsonWriter<UTF8CHAR> Writer(Builder, ANSITEXTVIEW(""), ANSITEXTVIEW(""), ANSITEXTVIEW(""));
 	Writer.WriteField(Field);
 }
 
diff --git a/Engine/Source/Runtime/Core/Public/Logging/StructuredLogFormat.h b/Engine/Source/Runtime/Core/Public/Logging/StructuredLogFormat.h
--- a/Engine/Source/Runtime/Core/Public/Logging/StructuredLogFormat.h
+++ b/Engine/Source/Runtime/Core/Public/Logging/StructuredLogFormat.h
@@ -6,7 +6,9 @@
 
 #define UE_API CORE_API
 
+class FCbFieldIterator;
 class FCbFieldViewIterator;
+class FCbWriter;
 class FText;
 
 namespace UE { class FLogTemplate; }
@@ -22,6 +24,23 @@ UE_API void FormatLogTo(FUtf8StringBuilderBase& Out, const FLogTemplate* Templat
 UE_API void FormatLogTo(FWideStringBuilderBase& Out, const FLogTemplate* Template, const FCbFieldViewIterator& Fields);
 UE_API FText FormatLogToText(const FLogTemplate* Template, const FCbFieldViewIterator& Fields);
 
+/**
+ * Converts the template and fields into a common format string and compatible fields.
+ *
+ * The common format string uses '{{' as a literal '{' and uses as a literal '}}'.
+ * The common format string contains no format specifiers or format argument modifiers.
+ * The common format string uses field names matching [A-Za-z0-9_]+.
+ *
+ * When a compatible field is an object, it will have a string field named $text containing the formatted value.
+ * A compatible field will never be an array.
+ *
+ * @param OutFormat   Appended with a transformed version of the format from Template.
+ * @param OutFields   Must be within an object scope. Appended with a compatible version of Fields.
+ * @param Template   Template from which to generate the common format string and its compatible fields.
+ * @param Fields   Fields matching Template, to be transformed into fields compatible with the common format string.
+ */
+UE_API void ConvertToCommonLog(FUtf8StringBuilderBase& OutFormat, FCbWriter& OutFields, const FLogTemplate* Template, const FCbFieldIterator& Fields);
+
 } // UE
 
 #undef UE_API
