PR #13071: Fix TSphere::ToString not compiling when used with FSphere3f or FSphere3d

diff --git a/Engine/Source/Runtime/Core/Public/Math/Sphere.h b/Engine/Source/Runtime/Core/Public/Math/Sphere.h
--- a/Engine/Source/Runtime/Core/Public/Math/Sphere.h
+++ b/Engine/Source/Runtime/Core/Public/Math/Sphere.h
@@ -221,7 +221,7 @@ public:
 	*/
 	[[nodiscard]] FString ToString() const
 	{
-		return FString::Printf(TEXT("Center=(%s), Radius=(%s)"), *Center.ToString(), *W.ToString());
+		return FString::Printf(TEXT("Center=(%s), Radius=(%s)"), *Center.ToString(), *LexToSanitizedString(W));
 	}
 
 	// Note: TSphere is usually written via binary serialization. This function exists for SerializeFromMismatchedTag conversion usage. 
