[RaceDetector]
- Fix a bug where we didn't detect race conditions on memory past the shadow block because the shadow size was not aligned and messed up the pointer arithmetic in GetShadowMemoryAddress.
- The result was that we could end up storing memory accesses in the wrong shadow slot, and then fail to resolve the access from the history due to the pointer being wrong.

diff --git a/Engine/Source/Runtime/Core/Private/Sanitizer/RaceDetectorWindows.cpp b/Engine/Source/Runtime/Core/Private/Sanitizer/RaceDetectorWindows.cpp
--- a/Engine/Source/Runtime/Core/Private/Sanitizer/RaceDetectorWindows.cpp
+++ b/Engine/Source/Runtime/Core/Private/Sanitizer/RaceDetectorWindows.cpp
@@ -898,7 +898,7 @@ namespace UE::Sanitizer::RaceDetector::Platform {
 		UPTRINT TotalAddressSpace = (UPTRINT)si.lpMaximumApplicationAddress - (UPTRINT)si.lpMinimumApplicationAddress;
 
 		// Reserve 7/8 of the space for the shadow so that we leave the application with 1/8 of the total address space
-		ShadowSize = 7 * (TotalAddressSpace >> 3);
+		ShadowSize = Align(7 * (TotalAddressSpace >> 3), PageSize);
 
 		// Because page faults are super slow when attached with a debugger
 		// we're going to use a bitmap view where each 4KB page of the shadow bitmap
