Core: Add GetTypeHash for FBox so that it can be used in sets and maps.

diff --git a/Engine/Source/Runtime/Core/Public/Math/Box.h b/Engine/Source/Runtime/Core/Public/Math/Box.h
--- a/Engine/Source/Runtime/Core/Public/Math/Box.h
+++ b/Engine/Source/Runtime/Core/Public/Math/Box.h
@@ -989,7 +989,12 @@ TBox<T> TBox<T>::Overlap(const TBox<T>& Other) const
 	return TBox<T>(MinVector, MaxVector);
 }
 
-
+template<typename T>
+uint32 GetTypeHash(const TBox<T>& Box)
+{
+	return !Box.IsValid ? 0 // All invalid boxes are the same
+		: HashCombineFast(GetTypeHash(Box.Min), GetTypeHash(Box.Max));
+}
 } // namespace UE::Math
 
 UE_DECLARE_LWC_TYPE(Box, 3);
