Add methods to declare inheritable command line arguments that are automatically passed down when the engine spawns an engine subprocess.

diff --git a/Engine/Plugins/MovieScene/MovieRenderPipeline/Source/MovieRenderPipelineEditor/Private/MoviePipelineNewProcessExecutor.cpp b/Engine/Plugins/MovieScene/MovieRenderPipeline/Source/MovieRenderPipelineEditor/Private/MoviePipelineNewProcessExecutor.cpp
--- a/Engine/Plugins/MovieScene/MovieRenderPipeline/Source/MovieRenderPipelineEditor/Private/MoviePipelineNewProcessExecutor.cpp
+++ b/Engine/Plugins/MovieScene/MovieRenderPipeline/Source/MovieRenderPipelineEditor/Private/MoviePipelineNewProcessExecutor.cpp
@@ -18,6 +18,7 @@
 #include "MoviePipelineEditorBlueprintLibrary.h"
 #include "MoviePipelineGameMode.h"
 #include "MoviePipelineUtils.h"
+#include "Misc/StringBuilder.h"
 
 #include UE_INLINE_GENERATED_CPP_BY_NAME(MoviePipelineNewProcessExecutor)
 
@@ -202,7 +203,9 @@ void UMoviePipelineNewProcessExecutor::Execute_Implementation(UMoviePipelineQueu
 	NamedArguments.Add(TEXT("GameNameOrProjectFile"), GameNameOrProjectFile);
 	NamedArguments.Add(TEXT("PlayWorld"), TEXT("MoviePipelineEntryMap")); // Boot up on an empty map, the executor will immediately transition to the correct one.
 	NamedArguments.Add(TEXT("UnrealURL"), UnrealURLParams); // Pass the command line arguments for this job
-	NamedArguments.Add(TEXT("SubprocessCommandLine"), FCommandLine::GetSubprocessCommandline());
+	TStringBuilder<64> SubprocessCommandLine;
+	FCommandLine::BuildSubprocessCommandLine(ECommandLineArgumentFlags::GameContexts, SubprocessCommandLine);
+	NamedArguments.Add(TEXT("SubprocessCommandLine"), *SubprocessCommandLine);
 	NamedArguments.Add(TEXT("CommandLineParams"), CommandLineArgs);
 	NamedArguments.Add(TEXT("MoviePipelineArgs"), MoviePipelineArgs);
 
diff --git a/Engine/Source/Developer/DerivedDataCache/Private/FileSystemCacheStore.cpp b/Engine/Source/Developer/DerivedDataCache/Private/FileSystemCacheStore.cpp
--- a/Engine/Source/Developer/DerivedDataCache/Private/FileSystemCacheStore.cpp
+++ b/Engine/Source/Developer/DerivedDataCache/Private/FileSystemCacheStore.cpp
@@ -3023,6 +3023,10 @@ bool FFileSystemCacheStore::RunInitialSpeedTest(const FFileSystemCacheStoreParam
 
 void FFileSystemCacheStoreParams::Parse(const TCHAR* Name, const TCHAR* Config)
 {
+	auto RegisterInheritedCommandlineArg = [](const FStringView ArgName)
+	{
+		FCommandLine::RegisterArgument(ArgName, ECommandLineArgumentFlags::EditorContext | ECommandLineArgumentFlags::CommandletContext | ECommandLineArgumentFlags::Inherit);
+	};
 	CacheName = Name;
 
 	// Default remote behavior based on historical cache names.
@@ -3051,6 +3055,7 @@ void FFileSystemCacheStoreParams::Parse(const TCHAR* Name, const TCHAR* Config)
 	if (FParse::Value(Config, TEXT("CommandLineOverride="), Key) &&
 		FParse::Value(FCommandLine::Get(), *(Key + TEXT("=")), CachePath))
 	{
+		RegisterInheritedCommandlineArg(Key);
 		UE_LOG(LogDerivedDataCache, Log, TEXT("%s: Found command line override %s=%s"), Name, *Key, *CachePath);
 	}
 
diff --git a/Engine/Source/Developer/DerivedDataCache/Private/HttpCacheStore.cpp b/Engine/Source/Developer/DerivedDataCache/Private/HttpCacheStore.cpp
--- a/Engine/Source/Developer/DerivedDataCache/Private/HttpCacheStore.cpp
+++ b/Engine/Source/Developer/DerivedDataCache/Private/HttpCacheStore.cpp
@@ -3172,6 +3172,10 @@ void FHttpCacheStore::GetChunks(
 
 void FHttpCacheStoreParams::Parse(const TCHAR* NodeName, const TCHAR* Config)
 {
+	auto RegisterInheritedCommandlineArg = [](const FStringView ArgName)
+	{
+		FCommandLine::RegisterArgument(ArgName, ECommandLineArgumentFlags::EditorContext | ECommandLineArgumentFlags::CommandletContext | ECommandLineArgumentFlags::Inherit);
+	};
 	Name = NodeName;
 
 	FString ServerId;
@@ -3212,6 +3216,7 @@ void FHttpCacheStoreParams::Parse(const TCHAR* NodeName, const TCHAR* Config)
 	{
 		if (FParse::Value(FCommandLine::Get(), *(OverrideName + TEXT("=")), Host))
 		{
+			RegisterInheritedCommandlineArg(OverrideName);
 			UE_LOG(LogDerivedDataCache, Log, TEXT("%s: Found command line override for Host %s=%s"), NodeName, *OverrideName, *Host);
 		}
 	}
@@ -3238,6 +3243,7 @@ void FHttpCacheStoreParams::Parse(const TCHAR* NodeName, const TCHAR* Config)
 	{
 		if (FParse::Value(FCommandLine::Get(), *(OverrideName + TEXT("=")), HttpVersion))
 		{
+			RegisterInheritedCommandlineArg(OverrideName);
 			UE_LOG(LogDerivedDataCache, Log, TEXT("%s: Found command line override for HttpVersion %s=%s"), NodeName, *OverrideName, *HttpVersion);
 		}
 	}
@@ -3257,6 +3263,7 @@ void FHttpCacheStoreParams::Parse(const TCHAR* NodeName, const TCHAR* Config)
 	{
 		if (FParse::Value(FCommandLine::Get(), *(OverrideName + TEXT("=")), UnixSocketPath))
 		{
+			RegisterInheritedCommandlineArg(OverrideName);
 			UE_LOG(LogDerivedDataCache, Log, TEXT("%s: Found command line override for UnixSocketPath %s=%s"), NodeName, *OverrideName, *UnixSocketPath);
 		}
 	}
@@ -3277,6 +3284,7 @@ void FHttpCacheStoreParams::Parse(const TCHAR* NodeName, const TCHAR* Config)
 	{
 		if (FParse::Value(FCommandLine::Get(), *(OverrideName + TEXT("=")), OAuthProvider))
 		{
+			RegisterInheritedCommandlineArg(OverrideName);
 			UE_LOG(LogDerivedDataCache, Log, TEXT("%s: Found command line override for OAuthProvider %s=%s"), NodeName, *OverrideName, *OAuthProvider);
 		}
 	}
@@ -3297,6 +3305,7 @@ void FHttpCacheStoreParams::Parse(const TCHAR* NodeName, const TCHAR* Config)
 	{
 		if (FParse::Value(FCommandLine::Get(), *(OverrideName + TEXT("=")), OAuthSecret))
 		{
+			RegisterInheritedCommandlineArg(OverrideName);
 			UE_LOG(LogDerivedDataCache, Log, TEXT("%s: Found command line override for OAuthSecret %s={SECRET}"), NodeName, *OverrideName);
 		}
 	}
@@ -3330,6 +3339,7 @@ void FHttpCacheStoreParams::Parse(const TCHAR* NodeName, const TCHAR* Config)
 	{
 		if (FParse::Value(FCommandLine::Get(), *(OverrideName + TEXT("=")), OAuthProviderIdentifier))
 		{
+			RegisterInheritedCommandlineArg(OverrideName);
 			UE_LOG(LogDerivedDataCache, Log, TEXT("%s: Found command line override for OAuthProviderIdentifier %s=%s"), NodeName, *OverrideName, *OAuthProviderIdentifier);
 		}
 	}
diff --git a/Engine/Source/Developer/DerivedDataCache/Private/ZenCacheStore.cpp b/Engine/Source/Developer/DerivedDataCache/Private/ZenCacheStore.cpp
--- a/Engine/Source/Developer/DerivedDataCache/Private/ZenCacheStore.cpp
+++ b/Engine/Source/Developer/DerivedDataCache/Private/ZenCacheStore.cpp
@@ -2031,6 +2031,10 @@ void FZenCacheStore::GetChunks(
 
 void FZenCacheStoreParams::Parse(const TCHAR* NodeName, const TCHAR* Config)
 {
+	auto RegisterInheritedCommandlineArg = [](const FStringView ArgName)
+	{
+		FCommandLine::RegisterArgument(ArgName, ECommandLineArgumentFlags::EditorContext | ECommandLineArgumentFlags::CommandletContext | ECommandLineArgumentFlags::Inherit);
+	};
 	Name = NodeName;
 
 	if (FString ServerId; FParse::Value(Config, TEXT("ServerID="), ServerId))
@@ -2064,6 +2068,7 @@ void FZenCacheStoreParams::Parse(const TCHAR* NodeName, const TCHAR* Config)
 	{
 		if (FParse::Value(FCommandLine::Get(), *(OverrideName + TEXT("=")), Host))
 		{
+			RegisterInheritedCommandlineArg(OverrideName);
 			UE_LOG(LogDerivedDataCache, Log, TEXT("%s: Found command line override for Host %s=%s"), NodeName, *OverrideName, *Host);
 		}
 	}
diff --git a/Engine/Source/Developer/LauncherServices/Private/Launcher/LauncherWorker.cpp b/Engine/Source/Developer/LauncherServices/Private/Launcher/LauncherWorker.cpp
--- a/Engine/Source/Developer/LauncherServices/Private/Launcher/LauncherWorker.cpp
+++ b/Engine/Source/Developer/LauncherServices/Private/Launcher/LauncherWorker.cpp
@@ -595,11 +595,17 @@ FString FLauncherWorker::CreateUATCommand( const ILauncherProfileRef& InProfile,
 	}
 
 	FString AdditionalCommandLine;
-	FString AddCmdLine = FString::Printf(TEXT("%s%s %s %s"),
+	TStringBuilder<64> InheritableGameOptions;
+	ECommandLineArgumentFlags CommandLineContextFlags = ECommandLineArgumentFlags::ClientContext | ECommandLineArgumentFlags::ServerContext;
+	FCommandLine::BuildSubprocessCommandLine(CommandLineContextFlags, InheritableGameOptions);
+	FString EscapedInheritableGameOptions(InheritableGameOptions);
+	EscapedInheritableGameOptions.ReplaceInline(TEXT("\""), TEXT("\\\""));
+	FString AddCmdLine = FString::Printf(TEXT("%s%s %s %s%s"),
 		*SessionCommands,
 		*RoleCommands,
 		*LocalizationCommands,
-		*InProfile->GetAdditionalCommandLineParameters());
+		*InProfile->GetAdditionalCommandLineParameters(),
+		*EscapedInheritableGameOptions);
 	if (!AddCmdLine.TrimStartAndEnd().IsEmpty()) // don't bother with -addcmdline if it's empty
 	{
 		AdditionalCommandLine = FString::Printf(TEXT(" -addcmdline=\"%s\""), *AddCmdLine);
@@ -675,14 +681,27 @@ FString FLauncherWorker::CreateUATCommand( const ILauncherProfileRef& InProfile,
 				UATCommand += TEXT(" -encryptinifiles");
 			}
 
+			TStringBuilder<64> InheritableCookOptions;
+			FCommandLine::BuildSubprocessCommandLine(ECommandLineArgumentFlags::CommandletContext, InheritableCookOptions);
 			FString AdditionalOptions = InProfile->GetCookOptions();
-			if (!AdditionalOptions.IsEmpty())
+			if (!AdditionalOptions.IsEmpty() || (InheritableCookOptions.Len() > 0))
 			{
 				UATCommand += TEXT(" -additionalcookeroptions=\"");
 
 				// Escape any quotes in the argument list
 				UATCommand += AdditionalOptions.Replace(TEXT("\""), TEXT("\\\""));
 
+				if (InheritableCookOptions.Len() > 0)
+				{
+					if (!AdditionalOptions.IsEmpty())
+					{
+						UATCommand += TEXT(" ");
+					}
+					FString EscapedInheritableCookOptions(InheritableCookOptions);
+					EscapedInheritableCookOptions.ReplaceInline(TEXT("\""), TEXT("\\\""));
+					UATCommand += EscapedInheritableCookOptions;
+				}
+
 				// If the additional options ends with a slash, make sure we don't escape the quote
 				if (UATCommand.EndsWith("\\"))
 				{
@@ -821,6 +840,37 @@ FString FLauncherWorker::CreateUATCommand( const ILauncherProfileRef& InProfile,
 				UATCommand += TEXT(" -snapshot");
 			}
 
+			TStringBuilder<64> InheritableCookOptions;
+			FCommandLine::BuildSubprocessCommandLine(ECommandLineArgumentFlags::CommandletContext, InheritableCookOptions);
+
+			FString AdditionalOptions = InProfile->GetCookOptions();
+			if (!AdditionalOptions.IsEmpty() || (InheritableCookOptions.Len() > 0))
+			{
+				UATCommand += TEXT(" -additionalcookeroptions=\"");
+
+				// Escape any quotes in the argument list
+				UATCommand += AdditionalOptions.Replace(TEXT("\""), TEXT("\\\""));
+
+				if (InheritableCookOptions.Len() > 0)
+				{
+					if (!AdditionalOptions.IsEmpty())
+					{
+						UATCommand += TEXT(" ");
+					}
+					FString EscapedInheritableCookOptions(InheritableCookOptions);
+					EscapedInheritableCookOptions.ReplaceInline(TEXT("\""), TEXT("\\\""));
+					UATCommand += EscapedInheritableCookOptions;
+				}
+
+				// If the additional options ends with a slash, make sure we don't escape the quote
+				if (UATCommand.EndsWith("\\"))
+				{
+					UATCommand += TEXT("\\");
+				}
+
+				UATCommand += TEXT("\"");
+			}
+
 			if (FDerivedDataCacheInterface* DDC = GetDerivedDataCache())
 			{
 				UATCommand += FString::Printf(TEXT(" -ddc=%s"), DDC->GetGraphName());
diff --git a/Engine/Source/Editor/MovieSceneCaptureDialog/Private/MovieSceneCaptureDialogModule.cpp b/Engine/Source/Editor/MovieSceneCaptureDialog/Private/MovieSceneCaptureDialogModule.cpp
--- a/Engine/Source/Editor/MovieSceneCaptureDialog/Private/MovieSceneCaptureDialogModule.cpp
+++ b/Engine/Source/Editor/MovieSceneCaptureDialog/Private/MovieSceneCaptureDialogModule.cpp
@@ -26,6 +26,7 @@
 #include "HAL/PlatformFileManager.h"
 #include "HAL/FileManager.h"
 #include "Misc/CommandLine.h"
+#include "Misc/StringBuilder.h"
 #include "Misc/FileHelper.h"
 #include "Misc/App.h"
 #include "Serialization/ObjectWriter.h"
@@ -733,14 +734,16 @@ void FNewProcessCapture::Start()
 	// Ensure game session is correctly set up 
 	EditorCommandLine += FString::Printf(TEXT(" -messaging -SessionName=\"%s\""), MovieCaptureSessionName);
 
+	TStringBuilder<64> SubprocessCommandLine;
+	FCommandLine::BuildSubprocessCommandLine(ECommandLineArgumentFlags::GameContexts, SubprocessCommandLine);
 	FString Params;
 	if (FPaths::IsProjectFilePathSet())
 	{
-		Params = FString::Printf(TEXT("\"%s\" %s %s"), *FPaths::GetProjectFilePath(), *EditorCommandLine, *FCommandLine::GetSubprocessCommandline());
+		Params = FString::Printf(TEXT("\"%s\" %s %s"), *FPaths::GetProjectFilePath(), *EditorCommandLine, *SubprocessCommandLine);
 	}
 	else
 	{
-		Params = FString::Printf(TEXT("%s %s %s"), FApp::GetProjectName(), *EditorCommandLine, *FCommandLine::GetSubprocessCommandline());
+		Params = FString::Printf(TEXT("%s %s %s"), FApp::GetProjectName(), *EditorCommandLine, *SubprocessCommandLine);
 	}
 
 	FString GamePath = FPlatformProcess::GenerateApplicationPath(FApp::GetName(), FApp::GetBuildConfiguration());
diff --git a/Engine/Source/Editor/UnrealEd/Private/PlayLevelNewProcess.cpp b/Engine/Source/Editor/UnrealEd/Private/PlayLevelNewProcess.cpp
--- a/Engine/Source/Editor/UnrealEd/Private/PlayLevelNewProcess.cpp
+++ b/Engine/Source/Editor/UnrealEd/Private/PlayLevelNewProcess.cpp
@@ -9,6 +9,7 @@
 #include "Settings/LevelEditorPlaySettings.h"
 #include "DataDrivenShaderPlatformInfo.h"
 #include "Net/OnlineEngineInterface.h"
+#include "Misc/StringBuilder.h"
 
 void UEditorEngine::StartPlayInNewProcessSession(FRequestPlaySessionParams& InRequestParams)
 {
@@ -347,7 +348,14 @@ void UEditorEngine::LaunchNewProcess(const FRequestPlaySessionParams& InParams,
 		// Otherwise hosts just connect and accept whatever the server's settings are.
 		NamedArguments.Add(TEXT("PlayWorldURL"), UnrealURLParams);
 	}
-	NamedArguments.Add(TEXT("SubprocessCommandLine"), FCommandLine::GetSubprocessCommandline());
+	TStringBuilder<64> SubprocessCommandLine;
+	ECommandLineArgumentFlags CommandLineArgFlags = ECommandLineArgumentFlags::ServerContext;
+	if (!bIsDedicatedServer)
+	{
+		CommandLineArgFlags = ECommandLineArgumentFlags::ClientContext;
+	}
+	FCommandLine::BuildSubprocessCommandLine(CommandLineArgFlags, SubprocessCommandLine);
+	NamedArguments.Add(TEXT("SubprocessCommandLine"), *SubprocessCommandLine);
 	NamedArguments.Add(TEXT("CommandLineParams"), CommandLine);
 
 	FString FinalCommandLine = FString::Format(
diff --git a/Engine/Source/Editor/UnrealEd/Private/UnrealEdEngine.cpp b/Engine/Source/Editor/UnrealEd/Private/UnrealEdEngine.cpp
--- a/Engine/Source/Editor/UnrealEd/Private/UnrealEdEngine.cpp
+++ b/Engine/Source/Editor/UnrealEd/Private/UnrealEdEngine.cpp
@@ -266,7 +266,7 @@ void UUnrealEdEngine::Init(IEngineLoop* InEngineLoop)
 	if (FParse::Param(FCommandLine::Get(), TEXT("nomcp")))
 	{
 		// If our editor has nomcp, pass it through to any subprocesses.
-		FCommandLine::AddToSubprocessCommandline(TEXT(" -nomcp"));
+		FCommandLine::AddToSubprocessCommandLine(TEXT(" -nomcp"), ECommandLineArgumentFlags::AllContexts);
 	}
 
 	bPivotMovedIndependently = false;
diff --git a/Engine/Source/Runtime/Core/Private/Misc/CommandLine.cpp b/Engine/Source/Runtime/Core/Private/Misc/CommandLine.cpp
--- a/Engine/Source/Runtime/Core/Private/Misc/CommandLine.cpp
+++ b/Engine/Source/Runtime/Core/Private/Misc/CommandLine.cpp
@@ -18,11 +18,24 @@ TCHAR FCommandLine::CmdLine[FCommandLine::MaxCommandLineSize] = {};
 TCHAR FCommandLine::OriginalCmdLine[FCommandLine::MaxCommandLineSize] = {};
 TCHAR FCommandLine::LoggingCmdLine[FCommandLine::MaxCommandLineSize] = {};
 TCHAR FCommandLine::LoggingOriginalCmdLine[FCommandLine::MaxCommandLineSize] = {};
+TMap<FString, FCommandLine::FRegisteredArgData> FCommandLine::RegisteredArgs;
 
-FString& FCommandLine::GetSubprocessCommandLine_Internal()
+class FInitializedCommandlinesArray : public TArray<FString>
 {
-	static FString SubprocessCommandLine = TEXT(" -Multiprocess");
-	return SubprocessCommandLine;
+public:
+	FInitializedCommandlinesArray()
+	{
+		SetNum((int)ECommandLineArgumentFlags::AllContexts);
+		this->operator[]((int)ECommandLineArgumentFlags::AllContexts - 1) = TEXT(" -Multiprocess");
+	}
+};
+
+FString& FCommandLine::GetSubprocessCommandLine_Internal(ECommandLineArgumentFlags ContextFlags)
+{
+	static FInitializedCommandlinesArray SubprocessCommandLines;
+	int Index = (int)(ContextFlags & ECommandLineArgumentFlags::AllContexts) - 1;
+	check(Index >= 0);
+	return SubprocessCommandLines[Index];
 }
 
 bool FCommandLine::IsInitialized()
@@ -337,9 +350,21 @@ void FCommandLine::BuildCommandLineAllowList(TCHAR* CommandLine, uint32 ArrayCou
 }
 #endif
 
-void FCommandLine::AddToSubprocessCommandline( const TCHAR* Param )
+void FCommandLine::RegisterArgument(FStringView Name, ECommandLineArgumentFlags Flags, FStringView Description)
+{
+	if ((Flags & ECommandLineArgumentFlags::AllContexts) == ECommandLineArgumentFlags::None)
+	{
+		return;
+	}
+
+	FRegisteredArgData& RegisteredArg = RegisteredArgs.FindOrAdd(FString(Name));
+	RegisteredArg.Flags |= Flags;
+	RegisteredArg.Description = Description;
+}
+
+void FCommandLine::AddToSubprocessCommandLine_Internal(const TCHAR* Param, ECommandLineArgumentFlags ApplicationContextFlags)
 {
-	FString& SubprocessCommandLine = GetSubprocessCommandLine_Internal();
+	FString& SubprocessCommandLine = GetSubprocessCommandLine_Internal(ApplicationContextFlags);
 
 	check( Param != NULL );
 	if (Param[0] != ' ')
@@ -349,9 +374,75 @@ void FCommandLine::AddToSubprocessCommandline( const TCHAR* Param )
 	SubprocessCommandLine += Param;
 }
 
+void FCommandLine::AddToSubprocessCommandline( const TCHAR* Param )
+{
+	AddToSubprocessCommandLine_Internal(Param, ECommandLineArgumentFlags::AllContexts);
+}
+
+void FCommandLine::AddToSubprocessCommandLine( const TCHAR* Param, ECommandLineArgumentFlags ApplicationContextFlags )
+{
+	if ((ApplicationContextFlags & ECommandLineArgumentFlags::AllContexts) == ECommandLineArgumentFlags::None)
+	{
+		return;
+	}
+
+	AddToSubprocessCommandLine_Internal(Param, ApplicationContextFlags);
+}
+
 const FString& FCommandLine::GetSubprocessCommandline()
 {
-	return GetSubprocessCommandLine_Internal();
+	// Get exclusively the arguments that are applied to an AllContexts run
+	return GetSubprocessCommandLine_Internal(ECommandLineArgumentFlags::AllContexts);
+}
+
+void FCommandLine::BuildSubprocessCommandLine(ECommandLineArgumentFlags ApplicationContextFlags, FStringBuilderBase& OutCommandline)
+{
+	if ((ApplicationContextFlags & ECommandLineArgumentFlags::AllContexts) == ECommandLineArgumentFlags::None)
+	{
+		return;
+	}
+
+	// Append all explicitly supplied subprocess arguments
+	for (int ContextPermutation = 1; ContextPermutation < (int)ECommandLineArgumentFlags::AllContexts; ++ContextPermutation)
+	{
+		if (ContextPermutation & (int)ApplicationContextFlags)
+		{
+			OutCommandline.Append(GetSubprocessCommandLine_Internal((ECommandLineArgumentFlags)ContextPermutation));
+		}
+	}
+	OutCommandline.Append(GetSubprocessCommandLine_Internal(ECommandLineArgumentFlags::AllContexts));
+
+	// Append all inherited arguments
+	const TCHAR* Stream = Get();
+	FString NextToken;
+	while (FParse::Token(Stream, NextToken, false))
+	{
+		if (NextToken.StartsWith(TEXT("-")))
+		{
+			FStringView NextTokenView(NextToken);
+			NextTokenView.RightChopInline(1); //Chop the leading dash from the view
+			for (const TPair<FString, FRegisteredArgData>& RegisteredArg : RegisteredArgs)
+			{
+				if ((RegisteredArg.Value.Flags & ECommandLineArgumentFlags::Inherit) != ECommandLineArgumentFlags::Inherit)
+				{
+					continue;
+				}
+
+				if (((RegisteredArg.Value.Flags & ApplicationContextFlags) != ECommandLineArgumentFlags::None) && NextTokenView.StartsWith(RegisteredArg.Key))
+				{
+					if ((NextTokenView.Len() == RegisteredArg.Key.Len()) || (NextTokenView[RegisteredArg.Key.Len()] == TCHAR('=')))
+					{
+						if (OutCommandline.Len() != 0)
+						{
+							OutCommandline.AppendChar(TCHAR(' '));
+						}
+						OutCommandline.Append(NextToken);
+						break;
+					}
+				}
+			}
+		}
+	}
 }
 
 /** 
diff --git a/Engine/Source/Runtime/Core/Public/Misc/CommandLine.h b/Engine/Source/Runtime/Core/Public/Misc/CommandLine.h
--- a/Engine/Source/Runtime/Core/Public/Misc/CommandLine.h
+++ b/Engine/Source/Runtime/Core/Public/Misc/CommandLine.h
@@ -4,6 +4,7 @@
 
 #include "CoreTypes.h"
 #include "Containers/Array.h"
+#include "Containers/StringView.h"
 #include "Containers/UnrealString.h"
 
 /*-----------------------------------------------------------------------------
@@ -20,6 +21,27 @@
 #endif
 
 
+enum class ECommandLineArgumentFlags
+{
+	None = 0x0000,
+
+	//~ Application context for a commandline argument
+	EditorContext = 0x0001,
+	ClientContext = 0x0002,
+	ServerContext = 0x0004,
+	CommandletContext = 0x0008,
+	ProgramContext = 0x0010,
+
+	GameContexts = 0x0006,
+	AllContexts = 0x001F,
+
+	// Indicates this argument should be automatically inherited by subprocesses.
+	// To be combined with one or more application context to inherit to.
+	Inherit = 0x0020,
+};
+
+ENUM_CLASS_FLAGS(ECommandLineArgumentFlags)
+
 struct FCommandLine
 {
 	/** maximum size of the command line */
@@ -71,19 +93,46 @@ struct FCommandLine
 	 */
 	CORE_API static void Append(const TCHAR* AppendString);
 
+	/**
+	* Registers a command line argument and associates it with a set of flags and optionally a description.
+	*
+	* @param Name 				Argument name on the commandline that should be inherited.
+	* @param Flags 				Flags to be associated with the argument.
+	*/
+	CORE_API static void RegisterArgument(FStringView Name, ECommandLineArgumentFlags Flags, FStringView Description = FStringView());
+
 	/**
 	 * Adds a new parameter to subprocess command line. If Param
 	 * does not start with a space, one is added.
 	 *
 	 * @param Param Command line param string.
 	 */
+	UE_DEPRECATED(5.6, "Use AddToSubprocessCommandLine version with capital L and flags argument.")
 	CORE_API static void AddToSubprocessCommandline( const TCHAR* Param );
 
+	/**
+	 * Adds a new parameter to subprocess command line. If Param
+	 * does not start with a space, one is added.
+	 *
+	 * @param Param Command line param string.
+	 * @param ApplicationContextFlags The one or more application context flags to which this parameter applies.
+	 */
+	CORE_API static void AddToSubprocessCommandLine( const TCHAR* Param, ECommandLineArgumentFlags ApplicationContextFlags );
+
 	/** 
-	 * Returns the subprocess command line string 
+	 * Returns the subprocess command line string (without inherited or context specific arguments)
 	 */
+	UE_DEPRECATED(5.6, "Use BuildSubprocessCommandLine instead.")
 	CORE_API static const FString& GetSubprocessCommandline();
 
+	/**
+	* Builds a commandline of inheritable and subprocess arguments for a specified application context.
+	*
+	* @param Flags 				The flags indicating the subprocess context we want a commandline for.
+	* @param Flags 				The flags indicating the subprocess context we want a commandline for.
+	*/
+	CORE_API static void BuildSubprocessCommandLine(ECommandLineArgumentFlags ApplicationContextFlags, FStringBuilderBase& OutCommandline);
+
 	/** 
 	* Removes the executable name from the passed CmdLine, denoted by parentheses.
 	* Returns the CmdLine string without the executable name.
@@ -154,9 +203,18 @@ private:
 	static TCHAR LoggingCmdLine[MaxCommandLineSize];
 	/** character buffer containing the original command line filtered for logging purposes */
 	static TCHAR LoggingOriginalCmdLine[MaxCommandLineSize];
-	/** subprocess command line */
-	static FString& GetSubprocessCommandLine_Internal();
+	/** subprocess command lines */
+	static FString& GetSubprocessCommandLine_Internal(ECommandLineArgumentFlags ContextFlags);
+	/** add subprocess command lines */
+	static void AddToSubprocessCommandLine_Internal(const TCHAR* Param, ECommandLineArgumentFlags ApplicationContextFlags);
 	/** What version of the command line this is, incremented on any modification to the command line */
 	static uint32 CmdLineVersion;
+
+	struct FRegisteredArgData
+	{
+		ECommandLineArgumentFlags Flags = ECommandLineArgumentFlags::None;
+		FString Description;
+	};
+	static TMap<FString, FRegisteredArgData> RegisteredArgs;
 };
 
diff --git a/Engine/Source/Runtime/Engine/Private/ShaderCompiler/ShaderCompiler.cpp b/Engine/Source/Runtime/Engine/Private/ShaderCompiler/ShaderCompiler.cpp
--- a/Engine/Source/Runtime/Engine/Private/ShaderCompiler/ShaderCompiler.cpp
+++ b/Engine/Source/Runtime/Engine/Private/ShaderCompiler/ShaderCompiler.cpp
@@ -1593,7 +1593,9 @@ FProcHandle FShaderCompilingManager::LaunchWorker(const FString& WorkingDirector
 			WorkerParameters += FString(TEXT(" -core "));
 		}
 	}
-	WorkerParameters += FCommandLine::GetSubprocessCommandline();
+	TStringBuilder<64> SubprocessCommandLine;
+	FCommandLine::BuildSubprocessCommandLine(ECommandLineArgumentFlags::ProgramContext, SubprocessCommandLine);
+	WorkerParameters += SubprocessCommandLine;
 
 #if USE_SHADER_COMPILER_WORKER_TRACE
 	// When doing utrace functionality we can't run with -nothreading, since it won't create the utrace thread to send events.
diff --git a/Engine/Source/Runtime/Engine/Private/ShaderCompiler/ShaderCompilerDistributed.cpp b/Engine/Source/Runtime/Engine/Private/ShaderCompiler/ShaderCompilerDistributed.cpp
--- a/Engine/Source/Runtime/Engine/Private/ShaderCompiler/ShaderCompilerDistributed.cpp
+++ b/Engine/Source/Runtime/Engine/Private/ShaderCompiler/ShaderCompilerDistributed.cpp
@@ -4,6 +4,7 @@
 #include "Async/Future.h"
 #include "HAL/FileManager.h"
 #include "Misc/CommandLine.h"
+#include "Misc/StringBuilder.h"
 #include "ShaderCompiler.h"
 #include "ShaderCompilerJobTypes.h"
 
@@ -196,7 +197,9 @@ void FShaderCompileDistributedThreadRunnable_Interface::DispatchShaderCompileJob
 	TaskCommandData.DispatcherPID = Manager->ProcessId;
 	TaskCommandData.InputFileName = InputFilePath;
 	TaskCommandData.OutputFileName = OutputFilePath;
-	TaskCommandData.ExtraCommandArgs = FString::Printf(TEXT("%s%s"), *FCommandLine::GetSubprocessCommandline(), GIsBuildMachine ? TEXT(" -buildmachine") : TEXT(""));
+	TStringBuilder<64> SubprocessCommandLine;
+	FCommandLine::BuildSubprocessCommandLine(ECommandLineArgumentFlags::ProgramContext, SubprocessCommandLine);
+	TaskCommandData.ExtraCommandArgs = FString::Printf(TEXT("%s%s"), *SubprocessCommandLine, GIsBuildMachine ? TEXT(" -buildmachine") : TEXT(""));
 	TaskCommandData.Dependencies = GetDependencyFilesForJobs(JobsToSerialize);
 
 	// Register any debug info paths that may be written to as additional output folders
diff --git a/Engine/Source/Runtime/Launch/Private/LaunchEngineLoop.cpp b/Engine/Source/Runtime/Launch/Private/LaunchEngineLoop.cpp
--- a/Engine/Source/Runtime/Launch/Private/LaunchEngineLoop.cpp
+++ b/Engine/Source/Runtime/Launch/Private/LaunchEngineLoop.cpp
@@ -6142,7 +6142,7 @@ bool FEngineLoop::AppInit( )
 	{
 		GIsBuildMachine = true;
 		// propagate to subprocesses, especially because some - like ShaderCompileWorker - use DDC, for which this switch matters
-		FCommandLine::AddToSubprocessCommandline(TEXT(" -buildmachine"));
+		FCommandLine::AddToSubprocessCommandLine(TEXT(" -buildmachine"), ECommandLineArgumentFlags::AllContexts);
 	}
 #endif // !UE_BUILD_SHIPPING
 
diff --git a/Engine/Source/Runtime/NetworkFile/Private/NetworkPlatformFile.cpp b/Engine/Source/Runtime/NetworkFile/Private/NetworkPlatformFile.cpp
--- a/Engine/Source/Runtime/NetworkFile/Private/NetworkPlatformFile.cpp
+++ b/Engine/Source/Runtime/NetworkFile/Private/NetworkPlatformFile.cpp
@@ -136,7 +136,7 @@ bool FNetworkPlatformFile::InitializeInternal(IPlatformFile* Inner, const TCHAR*
 	// lets see we can test whether the server is up. 
 	if (Out.Num())
 	{
-		FCommandLine::AddToSubprocessCommandline( *FString::Printf( TEXT("-FileHostIP=%s"), HostIP ) );
+		FCommandLine::AddToSubprocessCommandLine( *FString::Printf( TEXT("-FileHostIP=%s"), HostIP ), ECommandLineArgumentFlags::AllContexts );
 		bIsUsable = true; 
 	}
 	return bIsUsable;
diff --git a/Engine/Source/Runtime/SandboxFile/Private/IPlatformFileSandboxWrapper.cpp b/Engine/Source/Runtime/SandboxFile/Private/IPlatformFileSandboxWrapper.cpp
--- a/Engine/Source/Runtime/SandboxFile/Private/IPlatformFileSandboxWrapper.cpp
+++ b/Engine/Source/Runtime/SandboxFile/Private/IPlatformFileSandboxWrapper.cpp
@@ -180,11 +180,11 @@ bool FSandboxPlatformFile::Initialize(IPlatformFile* Inner, const TCHAR* CmdLine
 		{
 			if (InjectedTargetDirectory.Len())
 			{
-				FCommandLine::AddToSubprocessCommandline(*FString::Printf(TEXT("-InjectedSandbox=%s;%s"), *InjectedSourceDirectory, *InjectedTargetDirectory));
+				FCommandLine::AddToSubprocessCommandLine(*FString::Printf(TEXT("-InjectedSandbox=%s;%s"), *InjectedSourceDirectory, *InjectedTargetDirectory), ECommandLineArgumentFlags::AllContexts);
 			}
 			else
 			{
-				FCommandLine::AddToSubprocessCommandline(*FString::Printf(TEXT("-Sandbox=%s"), *SandboxDirectory));
+				FCommandLine::AddToSubprocessCommandLine(*FString::Printf(TEXT("-Sandbox=%s"), *SandboxDirectory), ECommandLineArgumentFlags::AllContexts);
 			}
 		}
 	}
diff --git a/Engine/Source/Runtime/StreamingFile/Private/StreamingNetworkPlatformFile.cpp b/Engine/Source/Runtime/StreamingFile/Private/StreamingNetworkPlatformFile.cpp
--- a/Engine/Source/Runtime/StreamingFile/Private/StreamingNetworkPlatformFile.cpp
+++ b/Engine/Source/Runtime/StreamingFile/Private/StreamingNetworkPlatformFile.cpp
@@ -380,7 +380,7 @@ bool FStreamingNetworkPlatformFile::InitializeInternal(IPlatformFile* Inner, con
 			UE_LOG(LogStreamingPlatformFile, Fatal, TEXT("Could not open test file %s."), *TestSyncFile);
 		}
 
-		FCommandLine::AddToSubprocessCommandline( *FString::Printf( TEXT("-StreamingHostIP=%s"), HostIP ) );
+		FCommandLine::AddToSubprocessCommandLine( *FString::Printf( TEXT("-StreamingHostIP=%s"), HostIP ), ECommandLineArgumentFlags::AllContexts );
 
 		return true; 
 	}
