Add logic to track number of times the application has been suspended and activated

#rnx

diff --git a/Engine/Source/Runtime/Core/Private/ProfilingDebugging/BootProfiling.cpp b/Engine/Source/Runtime/Core/Private/ProfilingDebugging/BootProfiling.cpp
--- a/Engine/Source/Runtime/Core/Private/ProfilingDebugging/BootProfiling.cpp
+++ b/Engine/Source/Runtime/Core/Private/ProfilingDebugging/BootProfiling.cpp
@@ -2,10 +2,16 @@
 
 #include "ProfilingDebugging/BootProfiling.h"
 #include "CoreGlobals.h"
+#include "Misc/CoreDelegates.h"
+#include "HAL/PlatformTime.h"
 
-double GEnginePreInitPreStartupScreenEndTime;
-double GEnginePreInitPostStartupScreenEndTime;
-double GEngineInitEndTime;
+static double GEnginePreInitPreStartupScreenEndTime;
+static double GEnginePreInitPostStartupScreenEndTime;
+static double GEngineInitEndTime;
+static bool GEngineInitializing = true;
+static bool GCurrentlyInBackground = true;
+static int GAppSuspendedCounter[2] = { 0,0 };
+static int GAppActivatedCounter[2] = { 0,0 };
 
 double FBootProfiling::GetBootDuration()
 {
@@ -26,3 +32,59 @@ double FBootProfiling::GetEngineInitDuration()
 {
 	return GEngineInitEndTime - GEnginePreInitPostStartupScreenEndTime;
 }
+
+void FBootProfiling::OnPreInitPreStartupScreenComplete()
+{
+	GEnginePreInitPreStartupScreenEndTime = FPlatformTime::Seconds();
+}
+
+void FBootProfiling::OnPreInitPostStartupScreenComplete()
+{
+	GEnginePreInitPostStartupScreenEndTime = FPlatformTime::Seconds();
+}
+
+void FBootProfiling::OnInitComplete()
+{
+	GEngineInitEndTime = FPlatformTime::Seconds();
+	GEngineInitializing = false;
+}
+
+void FBootProfiling::InitCounters()
+{
+	FCoreDelegates::ApplicationWillEnterBackgroundDelegate.AddStatic([] 
+		{ 
+			++GAppSuspendedCounter[!GEngineInitializing]; 
+			GCurrentlyInBackground = true;
+		});
+	FCoreDelegates::ApplicationHasEnteredForegroundDelegate.AddStatic([] 
+		{ 
+			if (!GCurrentlyInBackground)
+			{
+				// never saw a WillEnterBackground event, which happens on e.g. PS4, so we have to bump GAppSuspendedCounter
+				++GAppSuspendedCounter[!GEngineInitializing];
+			}
+			++GAppActivatedCounter[!GEngineInitializing]; 
+			GCurrentlyInBackground = false;
+		});
+}
+
+static int GetPhaseCounter(FBootProfiling::ECounterPhase phase, const int counters[2])
+{
+	switch (phase)
+	{
+	case FBootProfiling::ECounterPhase::DuringEngineInit: return counters[0];
+	case FBootProfiling::ECounterPhase::PostEngineInit: return counters[1];
+	case FBootProfiling::ECounterPhase::AllPhases: return counters[0] + counters[1];
+	default: return -1;
+	}
+}
+
+int FBootProfiling::GetAppSuspendedCounter(ECounterPhase phase)
+{
+	return GetPhaseCounter(phase, GAppSuspendedCounter);
+}
+
+int FBootProfiling::GetAppActivatedCounter(ECounterPhase phase)
+{
+	return GetPhaseCounter(phase, GAppActivatedCounter);
+}
diff --git a/Engine/Source/Runtime/Core/Public/ProfilingDebugging/BootProfiling.h b/Engine/Source/Runtime/Core/Public/ProfilingDebugging/BootProfiling.h
--- a/Engine/Source/Runtime/Core/Public/ProfilingDebugging/BootProfiling.h
+++ b/Engine/Source/Runtime/Core/Public/ProfilingDebugging/BootProfiling.h
@@ -6,14 +6,33 @@
 
 // TODO: Consider moving boot timing stuff from CoreGlobals.h here in UE5 - involves adding a new include to 50+ files
 
-extern CORE_API double GEnginePreInitPreStartupScreenEndTime;
-extern CORE_API double GEnginePreInitPostStartupScreenEndTime;
-extern CORE_API double GEngineInitEndTime;
-
 struct FBootProfiling
 {
 	static CORE_API double GetBootDuration();
 	static CORE_API double GetPreInitPreStartupScreenDuration();
 	static CORE_API double GetPreInitPostStartupScreenDuration();
 	static CORE_API double GetEngineInitDuration();
-};
\ No newline at end of file
+
+	enum class ECounterPhase
+	{
+		DuringEngineInit,
+		PostEngineInit,
+		AllPhases,	// combine all phases above
+	};
+
+	// These are to be called at the appropriate time during EngineInit
+	static CORE_API void InitCounters();
+	static CORE_API void OnPreInitPreStartupScreenComplete();
+	static CORE_API void OnPreInitPostStartupScreenComplete();
+	static CORE_API void OnInitComplete();
+
+	/* Returns the amount of times the app has been suspended (mobile/switch) during the specified phase. 
+	 * NOTE: returns -1 when an incorrect phase is passed.
+	 */
+	static CORE_API int GetAppSuspendedCounter(ECounterPhase phase);
+
+	/* Returns the amount of times the app has been activated (mobile/switch) during the specified phase. 
+	 * NOTE: returns -1 when an incorrect phase is passed.
+	 */
+	static CORE_API int GetAppActivatedCounter(ECounterPhase phase);
+};
diff --git a/Engine/Source/Runtime/Launch/Private/LaunchEngineLoop.cpp b/Engine/Source/Runtime/Launch/Private/LaunchEngineLoop.cpp
--- a/Engine/Source/Runtime/Launch/Private/LaunchEngineLoop.cpp
+++ b/Engine/Source/Runtime/Launch/Private/LaunchEngineLoop.cpp
@@ -1578,7 +1578,7 @@ DECLARE_CYCLE_STAT(TEXT("FEngineLoop::PreInitPostStartupScreen.AfterStats"), STA
 
 int32 FEngineLoop::PreInitPreStartupScreen(const TCHAR* CmdLine)
 {
-	ON_SCOPE_EXIT { GEnginePreInitPreStartupScreenEndTime = FPlatformTime::Seconds(); };
+	ON_SCOPE_EXIT { FBootProfiling::OnPreInitPreStartupScreenComplete(); };
 	FDelayedAutoRegisterHelper::RunAndClearDelayedAutoRegisterDelegates(EDelayedRegisterRunPhase::StartOfEnginePreInit);
 	SCOPED_BOOT_TIMING("FEngineLoop::PreInitPreStartupScreen");
 
@@ -3223,7 +3223,7 @@ void ConditionallyEnsureOnCommandletErrors(int32 InNumErrors)
 
 int32 FEngineLoop::PreInitPostStartupScreen(const TCHAR* CmdLine)
 {
-	ON_SCOPE_EXIT{ GEnginePreInitPostStartupScreenEndTime = FPlatformTime::Seconds(); };
+	ON_SCOPE_EXIT{ FBootProfiling::OnPreInitPostStartupScreenComplete(); };
 	SCOPED_BOOT_TIMING("FEngineLoop::PreInitPostStartupScreen");
 	LLM_SCOPE(ELLMTag::EnginePreInitMemory);
 
@@ -4127,6 +4127,8 @@ int32 FEngineLoop::PreInit(const TCHAR* CmdLine)
 	ArraySlackTrackInit();
 #endif
 
+	FBootProfiling::InitCounters();
+
 	const int32 rv1 = PreInitPreStartupScreen(CmdLine);
 	if (rv1 != 0)
 	{
@@ -4144,7 +4146,6 @@ int32 FEngineLoop::PreInit(const TCHAR* CmdLine)
 	return 0;
 }
 
-
 bool FEngineLoop::LoadCoreModules()
 {
 	// Always attempt to load CoreUObject. It requires additional pre-init which is called from its module's StartupModule method.
@@ -4474,7 +4475,7 @@ int32 FEngineLoop::Init()
 #if TRACK_DISK_UTILIZATION
 		CSV_METADATA(TEXT("BootMBLoaded"), *FString::FromInt(GDiskUtilizationTracker.GetLongTermStats().TotalBytesRead / (1024 * 1024)));
 #endif
-		GEngineInitEndTime = FPlatformTime::Seconds(); 
+		FBootProfiling::OnInitComplete();
 	};
 	LLM_SCOPE(ELLMTag::EngineInitMemory);
 	SCOPED_BOOT_TIMING("FEngineLoop::Init");
