Fixed FError == operator conflicts.

diff --git a/Engine/Source/Runtime/Core/Public/Experimental/UnifiedError/UnifiedError.h b/Engine/Source/Runtime/Core/Public/Experimental/UnifiedError/UnifiedError.h
--- a/Engine/Source/Runtime/Core/Public/Experimental/UnifiedError/UnifiedError.h
+++ b/Engine/Source/Runtime/Core/Public/Experimental/UnifiedError/UnifiedError.h
@@ -293,16 +293,6 @@ namespace UE::UnifiedError
 			ErrorDetails = InError.ErrorDetails;
 		}
 
-		CORE_API bool operator==(const FError& OtherError)
-		{
-			if ((ModuleId == OtherError.ModuleId) && 
-				(ErrorCode == OtherError.ErrorCode))
-			{
-				return true;
-			}
-			return false;
-		}
-
 		~FError()
 		{
 			Invalidate();
@@ -362,10 +352,22 @@ namespace UE::UnifiedError
 			}*/
 		}
 	private:
+		friend CORE_API bool operator==(const FError& Error, const FError& OtherError);
 
 		TRefCountPtr<const IErrorDetails> GetInnerMostErrorDetails() const;
 	};
 
+
+	CORE_API inline bool operator==(const FError& Error, const FError& OtherError)
+	{
+		if ((Error.ModuleId == OtherError.ModuleId) &&
+			(Error.ErrorCode == OtherError.ErrorCode))
+		{
+			return true;
+		}
+		return false;
+	}
+
 	class FErrorRegistry
 	{
 	public:
