FileManager: Add feebdack in error cases: we were previously passing in 0 for GetLastError rather than passing in the value for GetLastError, and so were printing "Success" instead of the correct string for the error code.
#rnx

diff --git a/Engine/Source/Runtime/Core/Private/HAL/FileManagerGeneric.cpp b/Engine/Source/Runtime/Core/Private/HAL/FileManagerGeneric.cpp
--- a/Engine/Source/Runtime/Core/Private/HAL/FileManagerGeneric.cpp
+++ b/Engine/Source/Runtime/Core/Private/HAL/FileManagerGeneric.cpp
@@ -699,7 +699,8 @@ void FArchiveFileReaderGeneric::Seek( int64 InPos )
 	{
 		TCHAR ErrorBuffer[1024];
 		SetError();
-		UE_CLOG( !IsSilent(), LogFileManager, Error, TEXT("SetFilePointer on %s Failed %lld/%lld: %lld %s"), *Filename, InPos, Size, Pos, FPlatformMisc::GetSystemErrorMessage(ErrorBuffer, 1024, 0));
+		UE_CLOG( !IsSilent(), LogFileManager, Error, TEXT("SetFilePointer on %s Failed %lld/%lld: %lld %s"),
+			*Filename, InPos, Size, Pos, FPlatformMisc::GetSystemErrorMessage(ErrorBuffer, 1024, FPlatformMisc::GetLastError()));
 	}
 
 	Pos = InPos;
@@ -830,7 +831,8 @@ bool FArchiveFileReaderGeneric::InternalPrecache( int64 PrecacheOffset, int64 Pr
 	if (Count!=ReadCount)
 	{
 		TCHAR ErrorBuffer[1024];
-		UE_CLOG( !IsSilent(), LogFileManager, Warning, TEXT( "ReadFile failed: Count=%lld ReadCount=%lld Error=%s" ), Count, ReadCount, FPlatformMisc::GetSystemErrorMessage( ErrorBuffer, 1024, 0 ) );
+		UE_CLOG( !IsSilent(), LogFileManager, Warning, TEXT( "ReadFile failed: Count=%lld ReadCount=%lld Error=%s" ),
+			Count, ReadCount, FPlatformMisc::GetSystemErrorMessage( ErrorBuffer, 1024, FPlatformMisc::GetLastError()) );
 		BufferCount = WriteOffset + Count;
 		BufferArray.SetNumUninitialized(BufferCount);
 		// The read failed, but we do not SetError or return false just because a precache read fails.
@@ -879,7 +881,8 @@ void FArchiveFileReaderGeneric::Serialize( void* V, int64 Length )
 					TCHAR ErrorBuffer[1024];
 					SetError();
 					UE_CLOG( !IsSilent(), LogFileManager, Warning, TEXT( "ReadFile failed: Count=%lld Length=%lld Error=%s for file %s" ), 
-						Count, Length, FPlatformMisc::GetSystemErrorMessage( ErrorBuffer, 1024, 0 ), *Filename );
+						Count, Length, FPlatformMisc::GetSystemErrorMessage( ErrorBuffer, 1024, FPlatformMisc::GetLastError()),
+						*Filename );
 				}
 				Pos += Length;
 				return;
@@ -1046,7 +1049,8 @@ void FArchiveFileWriterGeneric::LogWriteError(const TCHAR* Message)
 	{
 		bLoggingError = true;
 		TCHAR ErrorBuffer[1024];
-		UE_LOG(LogFileManager, Error, TEXT("%s: %s (%s)"), Message, *Filename, FPlatformMisc::GetSystemErrorMessage(ErrorBuffer, 1024, 0));
+		UE_LOG(LogFileManager, Error, TEXT("%s: %s (%s)"), Message, *Filename,
+			FPlatformMisc::GetSystemErrorMessage(ErrorBuffer, 1024, FPlatformMisc::GetLastError()));
 		bLoggingError = false;
 	}
 }
diff --git a/Engine/Source/Runtime/Core/Private/Misc/MessageDialog.cpp b/Engine/Source/Runtime/Core/Private/Misc/MessageDialog.cpp
--- a/Engine/Source/Runtime/Core/Private/Misc/MessageDialog.cpp
+++ b/Engine/Source/Runtime/Core/Private/Misc/MessageDialog.cpp
@@ -66,12 +66,12 @@ void FMessageDialog::ShowLastError()
 	TCHAR ErrorBuffer[1024];
 	if (FApp::IsUnattended())
 	{
-		UE_LOG(LogOutputDevice, Fatal, TEXT("GetLastError : %d\n\n%s"), LastError, FPlatformMisc::GetSystemErrorMessage(ErrorBuffer, 1024, 0));
+		UE_LOG(LogOutputDevice, Fatal, TEXT("GetLastError : %d\n\n%s"), LastError, FPlatformMisc::GetSystemErrorMessage(ErrorBuffer, 1024, LastError));
 	}
 	else
 	{
 		TCHAR TempStr[MAX_SPRINTF] = {};
-		FCString::Sprintf(TempStr, TEXT("GetLastError : %d\n\n%s"), LastError, FPlatformMisc::GetSystemErrorMessage(ErrorBuffer, 1024, 0));
+		FCString::Sprintf(TempStr, TEXT("GetLastError : %d\n\n%s"), LastError, FPlatformMisc::GetSystemErrorMessage(ErrorBuffer, 1024, LastError));
 		FPlatformMisc::MessageBoxExt(EAppMsgType::Ok, TempStr, *NSLOCTEXT("MessageDialog", "DefaultSystemErrorTitle", "System Error").ToString());
 	}
 }
