Mute arm64ec FPCR warnings

#rnx

diff --git a/Engine/Source/Runtime/Core/Public/Math/UnrealMathNeon.h b/Engine/Source/Runtime/Core/Public/Math/UnrealMathNeon.h
--- a/Engine/Source/Runtime/Core/Public/Math/UnrealMathNeon.h
+++ b/Engine/Source/Runtime/Core/Public/Math/UnrealMathNeon.h
@@ -2240,6 +2240,20 @@ FORCEINLINE int32 VectorAnyGreaterThan(VectorRegister4Double Vec1, VectorRegiste
  */
 #define VectorResetFloatRegisters()
 
+
+/**
+ * Disable arm64ec FPCR warnings for the time being
+ * 
+ * https://learn.microsoft.com/en-us/cpp/build/arm64ec-windows-abi-conventions
+ * "...Avoid directly reading, writing, or computing mappings of FPCR and FPSR. These bits may be used in the future and are subject to change."
+ */
+#if PLATFORM_WINDOWS && defined(_M_ARM64EC)
+	#pragma warning(push)
+	#pragma warning(disable:5076) // warning C5076: read from FPCR
+	#pragma warning(disable:5077) // warning C5076: write to FPCR
+#endif
+
+
 /**
  * Returns the control register.
  *
@@ -2272,6 +2286,14 @@ FORCEINLINE void VectorSetControlRegister(uint32_t ControlStatus)
 #endif
 }
 
+/**
+ * Restore arm64ec FPCR warnings
+ */
+#if PLATFORM_WINDOWS && defined(_M_ARM64EC)
+	#pragma warning(pop)
+#endif
+
+
 /**
  * Control status bit to round all floating point math results towards zero.
  */
