Changed error code from NotFound to UnknownChunkID in IoStoreReader.

#rnx

diff --git a/Engine/Source/Runtime/Core/Private/IO/IoStore.cpp b/Engine/Source/Runtime/Core/Private/IO/IoStore.cpp
--- a/Engine/Source/Runtime/Core/Private/IO/IoStore.cpp
+++ b/Engine/Source/Runtime/Core/Private/IO/IoStore.cpp
@@ -492,7 +492,7 @@ public:
 		}
 		else
 		{
-			return FIoStatus(EIoErrorCode::NotFound, TEXT("Not found"));
+			return FIoStatus(EIoErrorCode::UnknownChunkID, FString::Printf(TEXT("Unknown chunk ID '%s'"), *LexToString(ChunkId)));
 		}
 	}
 
@@ -530,7 +530,7 @@ public:
 			// Currently there's no way to make a task with a valid result that just emplaces
 			// without running.
 			return UE::Tasks::Launch(TEXT("FIoStoreRead_Error"), 
-				[] { return TIoStatusOr<FIoBuffer>(FIoStatus(EIoErrorCode::NotFound, TEXT("Unknown chunk ID"))); },
+				[ChunkId] { return TIoStatusOr<FIoBuffer>(FIoStatus(EIoErrorCode::UnknownChunkID, FString::Printf(TEXT("Unknown chunk ID '%s'"), *LexToString(ChunkId)))); },
 				UE::Tasks::ETaskPriority::Normal,
 				UE::Tasks::EExtendedTaskPriority::Inline); // force execution on this thread
 		}
@@ -666,7 +666,7 @@ public:
 		const FIoOffsetAndLength* OffsetAndLength = TocReader.GetOffsetAndLength(ChunkId);
 		if (!OffsetAndLength)
 		{
-			return FIoStatus(EIoErrorCode::NotFound, TEXT("Unknown chunk ID"));
+			return FIoStatus(EIoErrorCode::UnknownChunkID, FString::Printf(TEXT("Unknown chunk ID '%s'"), *LexToString(ChunkId)));
 		}
 
 		uint64 RequestedOffset = Options.GetOffset();
@@ -810,7 +810,7 @@ public:
 		const FIoOffsetAndLength* OffsetAndLength = TocReader.GetOffsetAndLength(ChunkId);
 		if (!OffsetAndLength)
 		{
-			return FIoStatus(EIoErrorCode::NotFound, TEXT("Unknown chunk ID"));
+			return FIoStatus(EIoErrorCode::UnknownChunkID, FString::Printf(TEXT("Unknown chunk ID '%s'"), *LexToString(ChunkId)));
 		}
 
 		// Combine with offset/size requested by the reader.
