improve comment explaining how matrix-matrix multiplication applies

diff --git a/Engine/Source/Runtime/Core/Public/Math/Matrix.h b/Engine/Source/Runtime/Core/Public/Math/Matrix.h
--- a/Engine/Source/Runtime/Core/Public/Math/Matrix.h
+++ b/Engine/Source/Runtime/Core/Public/Math/Matrix.h
@@ -26,9 +26,11 @@ UE_DECLARE_LWC_TYPE(Quat, 4);
 
 /**
  * 4x4 matrix of floating point values.
- * Matrix-matrix multiplication happens with a pre-multiple of the transpose --
- * in other words, Res = Mat1.operator*(Mat2) means Res = Mat2^FArg * Mat1, as
- * opposed to Res = Mat1 * Mat2.
+ * 
+ * Note that, like with FTransform, Matrix-Matrix multiplication is applied such that
+ * C = A * B will yield a transform C that logically first applies A then B,
+ * so (A*B).TransformPosition(Pt) == B.TransformPosition(A.TransformPosition(Pt))
+ * 
  * Matrix elements are accessed with M[RowIndex][ColumnIndex].
  */
 namespace UE
