Unified Error Fixes
- Moved FTextFormatArgsPropertyExtractor to header
- Improve LexToString

diff --git a/Engine/Source/Runtime/Core/Private/Experimental/UnifiedError/UnifiedError.cpp b/Engine/Source/Runtime/Core/Private/Experimental/UnifiedError/UnifiedError.cpp
--- a/Engine/Source/Runtime/Core/Private/Experimental/UnifiedError/UnifiedError.cpp
+++ b/Engine/Source/Runtime/Core/Private/Experimental/UnifiedError/UnifiedError.cpp
@@ -9,51 +9,6 @@
 namespace UE::UnifiedError
 {
 
-class FTextFormatArgsPropertyExtractor : public IErrorPropertyExtractor
-{
-private:
-	FFormatNamedArguments& Arguments;
-public:
-	FTextFormatArgsPropertyExtractor(FFormatNamedArguments& InArguments) : Arguments(InArguments) 
-	{
-	}
-
-	virtual void AddProperty(const FUtf8StringView& PropertyName, const FStringView& PropertyValue) override
-	{
-		Arguments.Add(UTF8_TO_TCHAR(PropertyName.GetData()), FFormatArgumentValue(FText::FromString(PropertyValue.GetData())));
-	}
-
-	virtual void AddProperty(const FUtf8StringView& PropertyName, const FUtf8StringView& PropertyValue)  override
-	{
-		Arguments.Add(UTF8_TO_TCHAR(PropertyName.GetData()), FFormatArgumentValue(FText::FromStringView(UTF8_TO_TCHAR(PropertyValue.GetData()))));
-	}
-
-	virtual void AddProperty(const FUtf8StringView& PropertyName, const int64 PropertyValue) override
-	{
-		Arguments.Add(UTF8_TO_TCHAR(PropertyName.GetData()), FFormatArgumentValue(PropertyValue));
-	}
-
-	virtual void AddProperty(const FUtf8StringView& PropertyName, const int32 PropertyValue) override
-	{
-		Arguments.Add(UTF8_TO_TCHAR(PropertyName.GetData()), FFormatArgumentValue(PropertyValue));
-	}
-
-	virtual void AddProperty(const FUtf8StringView& PropertyName, const float PropertyValue) override
-	{
-		Arguments.Add(UTF8_TO_TCHAR(PropertyName.GetData()), FFormatArgumentValue(PropertyValue));
-	}
-
-	virtual void AddProperty(const FUtf8StringView& PropertyName, const double PropertyValue) override
-	{
-		Arguments.Add(UTF8_TO_TCHAR(PropertyName.GetData()), FFormatArgumentValue(PropertyValue));
-	}
-
-	virtual void AddProperty(const FUtf8StringView& PropertyName, const FText& PropertyValue) override
-	{
-		Arguments.Add(UTF8_TO_TCHAR(PropertyName.GetData()), FFormatArgumentValue(PropertyValue));
-	}
-};
-
 
 
 // FError functions
@@ -379,5 +334,8 @@ uint32 FErrorDetailsRegistry::RegisterDetails(const FAnsiStringView& ErrorDetail
 
 FString LexToString(const UE::UnifiedError::FError& Error)
 {
-	return Error.GetErrorMessage().ToString();
+	return FString::Printf(TEXT("%.*hs:%.*hs - %s"), 
+		Error.GetModuleIdString().Len(), Error.GetModuleIdString().GetData(),
+		Error.GetErrorCodeString().Len(), Error.GetErrorCodeString().GetData(),
+		*Error.GetErrorMessage().ToString());
 }
\ No newline at end of file
diff --git a/Engine/Source/Runtime/Core/Public/Experimental/UnifiedError/UnifiedError.h b/Engine/Source/Runtime/Core/Public/Experimental/UnifiedError/UnifiedError.h
--- a/Engine/Source/Runtime/Core/Public/Experimental/UnifiedError/UnifiedError.h
+++ b/Engine/Source/Runtime/Core/Public/Experimental/UnifiedError/UnifiedError.h
@@ -404,6 +404,53 @@ namespace UE::UnifiedError
 
 
 
+	class FTextFormatArgsPropertyExtractor : public IErrorPropertyExtractor
+	{
+	private:
+		FFormatNamedArguments& Arguments;
+	public:
+		inline  FTextFormatArgsPropertyExtractor(FFormatNamedArguments& InArguments) : Arguments(InArguments)
+		{
+		}
+
+		inline virtual void AddProperty(const FUtf8StringView& PropertyName, const FStringView& PropertyValue) override
+		{
+			Arguments.Add(UTF8_TO_TCHAR(PropertyName.GetData()), FFormatArgumentValue(FText::FromString(PropertyValue.GetData())));
+		}
+
+		inline virtual void AddProperty(const FUtf8StringView& PropertyName, const FUtf8StringView& PropertyValue)  override
+		{
+			Arguments.Add(UTF8_TO_TCHAR(PropertyName.GetData()), FFormatArgumentValue(FText::FromStringView(UTF8_TO_TCHAR(PropertyValue.GetData()))));
+		}
+
+		inline virtual void AddProperty(const FUtf8StringView& PropertyName, const int64 PropertyValue) override
+		{
+			Arguments.Add(UTF8_TO_TCHAR(PropertyName.GetData()), FFormatArgumentValue(PropertyValue));
+		}
+
+		inline virtual void AddProperty(const FUtf8StringView& PropertyName, const int32 PropertyValue) override
+		{
+			Arguments.Add(UTF8_TO_TCHAR(PropertyName.GetData()), FFormatArgumentValue(PropertyValue));
+		}
+
+		inline virtual void AddProperty(const FUtf8StringView& PropertyName, const float PropertyValue) override
+		{
+			Arguments.Add(UTF8_TO_TCHAR(PropertyName.GetData()), FFormatArgumentValue(PropertyValue));
+		}
+
+		inline virtual void AddProperty(const FUtf8StringView& PropertyName, const double PropertyValue) override
+		{
+			Arguments.Add(UTF8_TO_TCHAR(PropertyName.GetData()), FFormatArgumentValue(PropertyValue));
+		}
+
+		inline virtual void AddProperty(const FUtf8StringView& PropertyName, const FText& PropertyValue) override
+		{
+			Arguments.Add(UTF8_TO_TCHAR(PropertyName.GetData()), FFormatArgumentValue(PropertyValue));
+		}
+	};
+
+
+
 
 	CORE_API void SerializeForLog(FCbWriter& Writer, const FError& OnlineError);
 
