Update minimum compiler static_assert to match the minimum versions allowed in Windows_SDK.json

#rnx

diff --git a/Engine/Source/Runtime/Core/Public/Windows/WindowsPlatformCompilerSetup.h b/Engine/Source/Runtime/Core/Public/Windows/WindowsPlatformCompilerSetup.h
--- a/Engine/Source/Runtime/Core/Public/Windows/WindowsPlatformCompilerSetup.h
+++ b/Engine/Source/Runtime/Core/Public/Windows/WindowsPlatformCompilerSetup.h
@@ -9,16 +9,16 @@
 #endif
 
 //
-// Future-proofing the min version check so we keep bumping it whenever we upgrade.
+// If LLVM Clang is used, we require at minimum v18.1.3 to compile
 //
-#if defined(_MSC_VER) && _MSC_VER > 1949 
-	#pragma message("Detected compiler newer than Visual Studio 2022, please update min version checking in WindowsPlatformCompilerSetup.h")
+#if defined(__clang__)
+static_assert(__clang_major__ >= 18 && __clang_minor__ >= 1 && __clang_patchlevel__ >= 3, "LLVM Clang v18.1.3 or greater is required to build Unreal Engine.");
 #endif
 
 //
-// We require at least Visual Studio 2022 v17.4 to compile
+// We require at minimum Visual Studio 2022 v17.8 to compile
 //
-static_assert(_MSC_VER >= 1934, "Visual Studio 2022 v17.4 or greater is required to build Unreal Engine.");
+static_assert(_MSC_VER >= 1938, "Visual Studio 2022 v17.8 (MSVC 14.38) or greater is required to build Unreal Engine.");
 
 //
 // Manually enable all warnings as errors, except ones that are explicitly skipped.
