[RaceDetector]
- Fix stack overflow during instrumentation caused by some thread being created with very small stack size

diff --git a/Engine/Source/Runtime/Core/Private/Sanitizer/RaceDetectorWindows.cpp b/Engine/Source/Runtime/Core/Private/Sanitizer/RaceDetectorWindows.cpp
--- a/Engine/Source/Runtime/Core/Private/Sanitizer/RaceDetectorWindows.cpp
+++ b/Engine/Source/Runtime/Core/Private/Sanitizer/RaceDetectorWindows.cpp
@@ -31,6 +31,7 @@ namespace UE::Sanitizer::RaceDetector {
 	extern volatile bool bIsResettingShadow;
 	extern void WriteToLog(FString&& Message);
 	extern bool bRuntimeInitialized;
+	constexpr int32 MinStackSize = 256 * 1024;
 }
 
 #define START_WINAPI_INSTRUMENTATION(DoIfNotInstrumenting) \
@@ -164,6 +165,13 @@ HANDLE (WINAPI* TrueCreateThread)(LPSECURITY_ATTRIBUTES lpThreadAttributes, SIZE
 INSTRUMENTATION_FUNCTION_ATTRIBUTES HANDLE WINAPI DetouredCreateThread(LPSECURITY_ATTRIBUTES lpThreadAttributes, SIZE_T dwStackSize, LPTHREAD_START_ROUTINE lpStartAddress, LPVOID lpParameter, DWORD dwCreationFlags, LPDWORD lpThreadId)
 {
 	FContext& Context = EnsureCurrentContext();
+
+	// Make sure threads have enough stack for additional instrumentation depth
+	if (dwStackSize != 0 && dwStackSize < UE::Sanitizer::RaceDetector::MinStackSize)
+	{
+		dwStackSize = UE::Sanitizer::RaceDetector::MinStackSize;
+	}
+
 	// If ThreadCreationDepth is non-zero here, we're probably coming from a _beginthreadex call 
 	// which was already instrumented. 
 	// Pass along the arguments as-they-are as they're already the instrumented arguments.
@@ -203,6 +211,12 @@ INSTRUMENTATION_FUNCTION_ATTRIBUTES UPTRINT WINAPI Detoured_beginthreadex(void*
 		FPlatformMisc::LowLevelOutputDebugStringf(TEXT("[%d] Creating thread (current clock %d)\n"), Context.ThreadId, Context.CurrentClock());
 	}
 
+	// Make sure threads have enough stack for additional instrumentation depth
+	if (stack_size != 0 && stack_size < UE::Sanitizer::RaceDetector::MinStackSize)
+	{
+		stack_size = UE::Sanitizer::RaceDetector::MinStackSize;
+	}
+
 	FInstrumentedStartThreadArgs* Args = new FInstrumentedStartThreadArgs;
 	Args->RealThreadParameter = arglist;
 	Args->RealStartRoutine = (LPTHREAD_START_ROUTINE)start_address;
