Logging: Fixed log tracing when the trace channel is enabled after logging has occurred

diff --git a/Engine/Source/Runtime/Core/Private/Logging/StructuredLog.cpp b/Engine/Source/Runtime/Core/Private/Logging/StructuredLog.cpp
--- a/Engine/Source/Runtime/Core/Private/Logging/StructuredLog.cpp
+++ b/Engine/Source/Runtime/Core/Private/Logging/StructuredLog.cpp
@@ -857,6 +857,38 @@
 
 ///////////////////////////////////////////////////////////////////////////////////////////////////////////////////////
 
+#if LOGTRACE_ENABLED
+template <typename StaticLogRecordType>
+FORCENOINLINE static void EnsureLogMessageSpec(const FLogCategoryBase& Category, const StaticLogRecordType& Log)
+{
+	if (!Log.DynamicData.bInitializedTrace.load(std::memory_order_acquire))
+	{
+		FLogTrace::OutputLogMessageSpec(&Log, &Category, Log.Verbosity, Log.File, Log.Line, TEXT("%s"));
+		Log.DynamicData.bInitializedTrace.store(true, std::memory_order_release);
+	}
+}
+
+// Tracing the log happens in its own function because that allows stack space for the message to
+// be returned before calling into the output devices.
+FORCENOINLINE static void LogToTrace(const void* LogPoint, const FLogRecord& Record)
+{
+	TStringBuilder<1024> Message;
+	Record.FormatMessageTo(Message);
+	FLogTrace::OutputLogMessage(LogPoint, *Message);
+}
+
+// Tracing the log happens in its own function because that allows stack space for the message to
+// be returned before calling into the output devices.
+FORCENOINLINE static void BasicLogToTrace(const void* LogPoint, const TCHAR* Format, va_list Args)
+{
+	TStringBuilder<1024> Message;
+	Message.AppendV(Format, Args);
+	FLogTrace::OutputLogMessage(LogPoint, *Message);
+}
+#endif // LOGTRACE_ENABLED
+
+///////////////////////////////////////////////////////////////////////////////////////////////////////////////////////
+
 struct FStaticLogTemplateManager
 {
 	std::atomic<FLogTemplate*> Head = nullptr;
@@ -874,17 +906,6 @@
 
 static FStaticLogTemplateManager GStaticLogTemplateManager;
 
-#if LOGTRACE_ENABLED
-// Tracing the log happens in its own function because that allows stack space for the message to
-// be returned before calling into the output devices.
-FORCENOINLINE static void LogToTrace(const void* LogPoint, const FLogRecord& Record)
-{
-	TStringBuilder<1024> Message;
-	Record.FormatMessageTo(Message);
-	FLogTrace::OutputLogMessage(LogPoint, *Message);
-}
-#endif
-
 // Serializing log fields to compact binary happens in its own function because that allows stack
 // space for the writer to be returned before calling into the output devices.
 FORCENOINLINE static FCbObject SerializeLogFields(
@@ -977,13 +998,6 @@
 template <typename StaticLogRecordType>
 FORCENOINLINE static FLogTemplate& CreateLogTemplate(const FLogCategoryBase& Category, const StaticLogRecordType& Log, const FLogField* Fields, const int32 FieldCount)
 {
-#if LOGTRACE_ENABLED
-	if (UE_TRACE_CHANNELEXPR_IS_ENABLED(LogChannel))
-	{
-		FLogTrace::OutputLogMessageSpec(&Log, &Category, Log.Verbosity, Log.File, Log.Line, TEXT("%s"));
-	}
-#endif
-
 	for (FLogTemplate* Template = Log.DynamicData.Template.load(std::memory_order_acquire);;)
 	{
 		if (Template && Template->GetFormat() == Log.Format)
@@ -1023,6 +1037,13 @@
 template <typename StaticLogRecordType>
 inline static FLogRecord CreateLogRecord(const FLogCategoryBase& Category, const StaticLogRecordType& Log, const FLogField* Fields, const int32 FieldCount)
 {
+#if LOGTRACE_ENABLED
+	if (UE_TRACE_CHANNELEXPR_IS_ENABLED(LogChannel))
+	{
+		EnsureLogMessageSpec(Category, Log);
+	}
+#endif
+
 	FLogTemplate& Template = EnsureLogTemplate(Category, Log, Fields, FieldCount);
 
 	FLogRecord Record;
@@ -1191,36 +1212,6 @@
 
 static FStaticBasicLogTemplateManager GStaticBasicLogTemplateManager;
 
-FORCENOINLINE static void OutputBasicLogMessageSpec(const FLogCategoryBase& Category, const FStaticBasicLogRecord& Log)
-{
-#if LOGTRACE_ENABLED
-	if (UE_TRACE_CHANNELEXPR_IS_ENABLED(LogChannel))
-	{
-		FLogTrace::OutputLogMessageSpec(&Log, &Category, Log.Verbosity, Log.File, Log.Line, TEXT("%s"));
-	}
-#endif
-	Log.DynamicData.bInitialized.store(true, std::memory_order_release);
-}
-
-inline static void EnsureBasicLogMessageSpec(const FLogCategoryBase& Category, const FStaticBasicLogRecord& Log)
-{
-	if (!Log.DynamicData.bInitialized.load(std::memory_order_acquire))
-	{
-		OutputBasicLogMessageSpec(Category, Log);
-	}
-}
-
-// Tracing the log happens in its own function because that allows stack space for the message to
-// be returned before calling into the output devices.
-FORCENOINLINE static void BasicLogToTrace(const void* LogPoint, const TCHAR* Format, va_list Args)
-{
-#if LOGTRACE_ENABLED
-	TStringBuilder<1024> Message;
-	Message.AppendV(Format, Args);
-	FLogTrace::OutputLogMessage(LogPoint, *Message);
-#endif
-}
-
 // Serializing the log to compact binary happens in its own function because that allows stack
 // space for the writer to be returned before calling into the output devices.
 FORCENOINLINE static FCbObject SerializeBasicLogMessage(const FStaticBasicLogRecord& Log, va_list Args)
@@ -1228,13 +1219,6 @@
 	TStringBuilder<512> Message;
 	Message.AppendV(Log.Format, Args);
 
-#if LOGTRACE_ENABLED
-	if (UE_TRACE_CHANNELEXPR_IS_ENABLED(LogChannel))
-	{
-		FLogTrace::OutputLogMessage(&Log, *Message);
-	}
-#endif
-
 	TCbWriter<512> Writer;
 	Writer.BeginObject();
 
@@ -1273,7 +1257,12 @@
 static void BasicLogV(const FLogCategoryBase& Category, const FStaticBasicLogRecord* Log, va_list Args)
 {
 #if !NO_LOGGING
-	EnsureBasicLogMessageSpec(Category, *Log);
+#if LOGTRACE_ENABLED
+	if (UE_TRACE_CHANNELEXPR_IS_ENABLED(LogChannel))
+	{
+		EnsureLogMessageSpec(Category, *Log);
+	}
+#endif
 
 	if (GConvertBasicLogToLogRecord)
 	{
diff --git a/Engine/Source/Runtime/Core/Public/Logging/LogMacros.h b/Engine/Source/Runtime/Core/Public/Logging/LogMacros.h
--- a/Engine/Source/Runtime/Core/Public/Logging/LogMacros.h
+++ b/Engine/Source/Runtime/Core/Public/Logging/LogMacros.h
@@ -110,7 +110,7 @@ namespace UE::Logging::Private
 /** Data about a static basic log that is created on-demand. */
 struct FStaticBasicLogDynamicData
 {
-	std::atomic<bool> bInitialized = false;
+	std::atomic<bool> bInitializedTrace = false;
 };
 
 /** Data about a static basic log that is constant for every occurrence. */
diff --git a/Engine/Source/Runtime/Core/Public/Logging/StructuredLog.h b/Engine/Source/Runtime/Core/Public/Logging/StructuredLog.h
--- a/Engine/Source/Runtime/Core/Public/Logging/StructuredLog.h
+++ b/Engine/Source/Runtime/Core/Public/Logging/StructuredLog.h
@@ -318,7 +318,8 @@ namespace UE::Logging::Private
 /** Data about a static log that is created on-demand. */
 struct FStaticLogDynamicData
 {
-	std::atomic<FLogTemplate*> Template;
+	std::atomic<FLogTemplate*> Template = nullptr;
+	std::atomic<bool> bInitializedTrace = false;
 };
 
 /** Data about a static log that is constant for every occurrence. */
