TVector4: Add a IsNearlyZero to complement the equivalent in TVector

#rnx

diff --git a/Engine/Source/Runtime/Core/Public/Math/Vector4.h b/Engine/Source/Runtime/Core/Public/Math/Vector4.h
--- a/Engine/Source/Runtime/Core/Public/Math/Vector4.h
+++ b/Engine/Source/Runtime/Core/Public/Math/Vector4.h
@@ -526,9 +526,12 @@ public:
 	/** Utility to check if there are any non-finite values (NaN or Inf) in this vector. */
 	[[nodiscard]] bool ContainsNaN() const;
 
-	/** Utility to check if all of the components of this vector are nearly zero given the tolerance. */
+	/** Utility to check if the XYZ components of this vector are nearly zero given the tolerance. */
 	[[nodiscard]] bool IsNearlyZero3(T Tolerance = UE_KINDA_SMALL_NUMBER) const;
 
+	/** Utility to check if all of the components of this vector are nearly zero given the tolerance. */
+	[[nodiscard]] bool IsNearlyZero(T Tolerance = UE_KINDA_SMALL_NUMBER) const;
+
 	/** Reflect vector. */
 	[[nodiscard]] TVector4<T> Reflect3(const TVector4<T>& Normal) const;
 
@@ -927,6 +930,16 @@ FORCEINLINE bool TVector4<T>::IsNearlyZero3(T Tolerance) const
 		&&	FMath::Abs(Z)<=Tolerance;
 }
 
+template<typename T>
+FORCEINLINE bool TVector4<T>::IsNearlyZero(T Tolerance) const
+{
+	return
+			FMath::Abs(X)<=Tolerance
+		&&	FMath::Abs(Y)<=Tolerance
+		&&	FMath::Abs(Z)<=Tolerance
+		&&	FMath::Abs(W)<=Tolerance;
+}
+
 template<typename T>
 FORCEINLINE TVector4<T> TVector4<T>::operator*=(const TVector4<T>& V)
 {
