Fix TestEqual not supporting an Actual/Expected as FString

diff --git a/Engine/Source/Runtime/Core/Public/Misc/AutomationTest.h b/Engine/Source/Runtime/Core/Public/Misc/AutomationTest.h
--- a/Engine/Source/Runtime/Core/Public/Misc/AutomationTest.h
+++ b/Engine/Source/Runtime/Core/Public/Misc/AutomationTest.h
@@ -1960,11 +1960,19 @@ public:
 	CORE_API bool TestEqual(const TCHAR* What, const FLinearColor Actual, const FLinearColor Expected);
 	CORE_API bool TestEqual(const TCHAR* What, const TCHAR* Actual, const TCHAR* Expected);
 	CORE_API bool TestEqual(const TCHAR* What, FStringView Actual, FStringView Expected);
+	CORE_API bool TestEqual(const TCHAR* What, const FString& Actual, const FString& Expected)
+	{
+		return TestEqual(What, *Actual, *Expected);
+	}
 	CORE_API bool TestEqual(const TCHAR* What, FUtf8StringView Actual, FUtf8StringView Expected);
 	CORE_API bool TestEqual(const TCHAR* What, const FText Actual, const FText Expected);
 	CORE_API bool TestEqual(const TCHAR* What, const FName Actual, const FName Expected);
 	CORE_API bool TestNotEqual(const TCHAR* What, const TCHAR* Actual, const TCHAR* Expected);
 	CORE_API bool TestNotEqual(const TCHAR* What, FStringView Actual, FStringView Expected);
+	CORE_API bool TestNotEqual(const TCHAR* What, const FString& Actual, const FString& Expected)
+	{
+		return TestNotEqual(What, *Actual, *Expected);
+	}
 	CORE_API bool TestNotEqual(const TCHAR* What, FUtf8StringView Actual, FUtf8StringView Expected);
 	CORE_API bool TestNotEqual(const TCHAR* What, const FText Actual, const FText Expected);
 	CORE_API bool TestNotEqual(const TCHAR* What, const FName Actual, const FName Expected);
@@ -1984,9 +1992,17 @@ public:
 
 	CORE_API bool TestEqualSensitive(const TCHAR* What, const TCHAR* Actual, const TCHAR* Expected);
 	CORE_API bool TestEqualSensitive(const TCHAR* What, FStringView Actual, FStringView Expected);
+	CORE_API bool TestEqualSensitive(const TCHAR* What, const FString& Actual, const FString& Expected)
+	{
+		return TestEqualSensitive(What, *Actual, *Expected);
+	}
 	CORE_API bool TestEqualSensitive(const TCHAR* What, FUtf8StringView Actual, FUtf8StringView Expected);
 	CORE_API bool TestNotEqualSensitive(const TCHAR* What, const TCHAR* Actual, const TCHAR* Expected);
 	CORE_API bool TestNotEqualSensitive(const TCHAR* What, FStringView Actual, FStringView Expected);
+	CORE_API bool TestNotEqualSensitive(const TCHAR* What, const FString& Actual, const FString& Expected)
+	{
+		return TestNotEqualSensitive(What, *Actual, *Expected);
+	}
 	CORE_API bool TestNotEqualSensitive(const TCHAR* What, FUtf8StringView Actual, FUtf8StringView Expected);
 
 	bool TestEqual(const FString& What, const int32 Actual, const int32 Expected)
