AndroidPlatformString: Fix compile errors in the inactive !PLATFORM_TCHAR_IS_CHAR16 case.
#rnx

diff --git a/Engine/Source/Runtime/Core/Public/Android/AndroidPlatformString.h b/Engine/Source/Runtime/Core/Public/Android/AndroidPlatformString.h
--- a/Engine/Source/Runtime/Core/Public/Android/AndroidPlatformString.h
+++ b/Engine/Source/Runtime/Core/Public/Android/AndroidPlatformString.h
@@ -30,6 +30,8 @@ typedef FGenericWidePlatformString FAndroidPlatformString;
 // to ANSI.  The Android NDK appears to have wchar_t support, but many of the functions appear
 // to just be stubs to the non-wide versions.  This doesn't work for obvious reasons.
 
+#include <alloca.h>
+#include <cstdlib>
 #include "Misc/Char.h"
 #include "GenericPlatform/GenericPlatformMemory.h"
 #include "GenericPlatform/GenericPlatformString.h"
@@ -198,7 +200,7 @@ struct FAndroidPlatformString : public FGenericPlatformString
 		{
 			if (Src[Pos] <= 255)
 			{
-				Dest[Pos] = Src[Pos];
+				Dest[Pos] = static_cast<ANSICHAR>(Src[Pos]);
 			}
 			else
 			{
@@ -236,7 +238,7 @@ struct FAndroidPlatformString : public FGenericPlatformString
 
 		++Find;
 		size_t MemCmpLen = wcslen(Find);
-		WIDECHAR* FoundChar;
+		const WIDECHAR* FoundChar;
 		for (;;)
 		{
 			FoundChar = wcschr(String, FindChar);
@@ -556,7 +558,7 @@ struct FAndroidPlatformString : public FGenericPlatformString
 
 	static FORCEINLINE int32 Strnlen( const ANSICHAR* String, SIZE_T StringSize )
 	{
-		return strnlen_s( String, StringSize );
+		return strnlen( String, StringSize );
 	}
 
 	static FORCEINLINE const ANSICHAR* Strstr( const ANSICHAR* String, const ANSICHAR* Find)
