Fix clang version static_assert to only use major version

#rnx

diff --git a/Engine/Source/Runtime/Core/Public/Windows/WindowsPlatformCompilerSetup.h b/Engine/Source/Runtime/Core/Public/Windows/WindowsPlatformCompilerSetup.h
--- a/Engine/Source/Runtime/Core/Public/Windows/WindowsPlatformCompilerSetup.h
+++ b/Engine/Source/Runtime/Core/Public/Windows/WindowsPlatformCompilerSetup.h
@@ -9,10 +9,10 @@
 #endif
 
 //
-// If LLVM Clang is used, we require at minimum v18.1.3 to compile
+// If LLVM Clang is used, we require at minimum v18 to compile
 //
 #if defined(__clang__)
-static_assert(__clang_major__ >= 18 && __clang_minor__ >= 1 && __clang_patchlevel__ >= 3, "LLVM Clang v18.1.3 or greater is required to build Unreal Engine.");
+static_assert(__clang_major__ >= 18, "LLVM Clang v18 or greater is required to build Unreal Engine.");
 #endif
 
 //
