FileManager: More feedback in error cases: Add the integer value of LastError to the message, and capture LastError earlier after the low-level failure to prevent the possibility of it being reset before we read it.
#rnx

diff --git a/Engine/Source/Runtime/Core/Private/HAL/FileManagerGeneric.cpp b/Engine/Source/Runtime/Core/Private/HAL/FileManagerGeneric.cpp
--- a/Engine/Source/Runtime/Core/Private/HAL/FileManagerGeneric.cpp
+++ b/Engine/Source/Runtime/Core/Private/HAL/FileManagerGeneric.cpp
@@ -697,10 +697,11 @@ void FArchiveFileReaderGeneric::Seek( int64 InPos )
 
 	if (!SeekLowLevel(SeekPos))
 	{
+		uint32 LastError = FPlatformMisc::GetLastError();
 		TCHAR ErrorBuffer[1024];
 		SetError();
-		UE_CLOG( !IsSilent(), LogFileManager, Error, TEXT("SetFilePointer on %s Failed %lld/%lld: %lld %s"),
-			*Filename, InPos, Size, Pos, FPlatformMisc::GetSystemErrorMessage(ErrorBuffer, 1024, FPlatformMisc::GetLastError()));
+		UE_CLOG( !IsSilent(), LogFileManager, Error, TEXT("SetFilePointer on %s Failed %lld/%lld: %lld LastError==%u: %s"),
+			*Filename, InPos, Size, Pos, LastError, FPlatformMisc::GetSystemErrorMessage(ErrorBuffer, 1024, LastError));
 	}
 
 	Pos = InPos;
@@ -830,9 +831,10 @@ bool FArchiveFileReaderGeneric::InternalPrecache( int64 PrecacheOffset, int64 Pr
 
 	if (Count!=ReadCount)
 	{
+		uint32 LastError = FPlatformMisc::GetLastError();
 		TCHAR ErrorBuffer[1024];
-		UE_CLOG( !IsSilent(), LogFileManager, Warning, TEXT( "ReadFile failed: Count=%lld ReadCount=%lld Error=%s" ),
-			Count, ReadCount, FPlatformMisc::GetSystemErrorMessage( ErrorBuffer, 1024, FPlatformMisc::GetLastError()) );
+		UE_CLOG( !IsSilent(), LogFileManager, Warning, TEXT( "ReadFile failed: Count=%lld ReadCount=%lld LastError=%u: %s" ),
+			Count, ReadCount, LastError, FPlatformMisc::GetSystemErrorMessage( ErrorBuffer, 1024, LastError) );
 		BufferCount = WriteOffset + Count;
 		BufferArray.SetNumUninitialized(BufferCount);
 		// The read failed, but we do not SetError or return false just because a precache read fails.
@@ -878,10 +880,11 @@ void FArchiveFileReaderGeneric::Serialize( void* V, int64 Length )
 				}
 				if( Count!=Length )
 				{
+					uint32 LastError = FPlatformMisc::GetLastError();
 					TCHAR ErrorBuffer[1024];
 					SetError();
-					UE_CLOG( !IsSilent(), LogFileManager, Warning, TEXT( "ReadFile failed: Count=%lld Length=%lld Error=%s for file %s" ), 
-						Count, Length, FPlatformMisc::GetSystemErrorMessage( ErrorBuffer, 1024, FPlatformMisc::GetLastError()),
+					UE_CLOG( !IsSilent(), LogFileManager, Warning, TEXT( "ReadFile failed: Count=%lld Length=%lld Error=%u: %s for file %s" ), 
+						Count, Length, LastError, FPlatformMisc::GetSystemErrorMessage( ErrorBuffer, 1024, LastError),
 						*Filename );
 				}
 				Pos += Length;
@@ -972,8 +975,9 @@ void FArchiveFileWriterGeneric::Seek( int64 InPos )
 	FlushBuffer();
 	if( !SeekLowLevel( InPos ) )
 	{
+		uint32 LastError = FPlatformMisc::GetLastError();
 		SetError();
-		LogWriteError(TEXT("Error seeking file"));
+		LogWriteError(TEXT("Error seeking file"), LastError);
 	}
 	Pos = InPos;
 }
@@ -983,8 +987,9 @@ bool FArchiveFileWriterGeneric::Close()
 	FlushBuffer();
 	if( !CloseLowLevel() )
 	{
+		uint32 LastError = FPlatformMisc::GetLastError();
 		SetError();
-		LogWriteError(TEXT("Error closing file"));
+		LogWriteError(TEXT("Error closing file"), LastError);
 	}
 	return !IsError();
 }
@@ -997,8 +1002,9 @@ void FArchiveFileWriterGeneric::Serialize( void* V, int64 Length )
 		FlushBuffer();
 		if( !WriteLowLevel( (uint8*)V, Length ) )
 		{
+			uint32 LastError = FPlatformMisc::GetLastError();
 			SetError();
-			LogWriteError(TEXT("Error writing to file"));
+			LogWriteError(TEXT("Error writing to file"), LastError);
 		}
 	}
 	else
@@ -1034,8 +1040,9 @@ bool FArchiveFileWriterGeneric::FlushBuffer()
 		bDidWriteData = WriteLowLevel(BufferArray.GetData(), BufferNum);
 		if (!bDidWriteData)
 		{
+			uint32 LastError = FPlatformMisc::GetLastError();
 			SetError();
-			LogWriteError(TEXT("Error flushing file"));
+			LogWriteError(TEXT("Error flushing file"), LastError);
 		}
 		BufferArray.Reset();
 	}
@@ -1043,14 +1050,19 @@ bool FArchiveFileWriterGeneric::FlushBuffer()
 }
 
 void FArchiveFileWriterGeneric::LogWriteError(const TCHAR* Message)
+{
+	LogWriteError(Message, FPlatformMisc::GetLastError());
+}
+
+void FArchiveFileWriterGeneric::LogWriteError(const TCHAR* Message, uint32 LastError)
 {
 	// Prevent re-entry if logging causes another log error leading to a stack overflow
 	if (!bLoggingError && !IsSilent())
 	{
 		bLoggingError = true;
 		TCHAR ErrorBuffer[1024];
-		UE_LOG(LogFileManager, Error, TEXT("%s: %s (%s)"), Message, *Filename,
-			FPlatformMisc::GetSystemErrorMessage(ErrorBuffer, 1024, FPlatformMisc::GetLastError()));
+		UE_LOG(LogFileManager, Error, TEXT("%s: %s (%u: %s)"), Message, *Filename,
+			LastError, FPlatformMisc::GetSystemErrorMessage(ErrorBuffer, 1024, LastError));
 		bLoggingError = false;
 	}
 }
diff --git a/Engine/Source/Runtime/Core/Public/HAL/FileManagerGeneric.h b/Engine/Source/Runtime/Core/Public/HAL/FileManagerGeneric.h
--- a/Engine/Source/Runtime/Core/Public/HAL/FileManagerGeneric.h
+++ b/Engine/Source/Runtime/Core/Public/HAL/FileManagerGeneric.h
@@ -346,14 +346,16 @@ protected:
 	**/
 	CORE_API virtual bool WriteLowLevel(const uint8* Src, int64 CountToWrite);
 
-	/** 
+	UE_DEPRECATED(5.6, "Pass in LastError from FPlatformMisc::GetLastError")
+	CORE_API void LogWriteError(const TCHAR* Message);
+	/**
 	 * Logs I/O error
 	 * It is important to not call any platform API functions after the error occurred and before 
 	 * calling this functions as the system error code may get reset and will not be properly
 	 * logged in this message.
 	 * @param Message Brief description of what went wrong
 	 */
-	CORE_API void LogWriteError(const TCHAR* Message);
+	CORE_API void LogWriteError(const TCHAR* Message, uint32 LastError);
 
 	/** Returns true if the archive should suppress logging in case of error */
 	bool IsSilent() const
