Workaround for expf() optimization bug on MSVC Arm64 & Arm64ec

#rnx

diff --git a/Engine/Source/Runtime/Core/Public/GenericPlatform/GenericPlatformMath.h b/Engine/Source/Runtime/Core/Public/GenericPlatform/GenericPlatformMath.h
--- a/Engine/Source/Runtime/Core/Public/GenericPlatform/GenericPlatformMath.h
+++ b/Engine/Source/Runtime/Core/Public/GenericPlatform/GenericPlatformMath.h
@@ -480,7 +480,11 @@ struct FGenericPlatformMath
 	}
 
 	// Returns e^Value
+#if PLATFORM_WINDOWS && (defined(_M_ARM64) || defined(_M_ARM64EC))
+	[[nodiscard]] static FORCEINLINE float Exp( volatile float Value ) { return expf(Value); } // #jira UE-261901 : workaround optimization bug for expf()
+#else
 	[[nodiscard]] static FORCEINLINE float Exp( float Value ) { return expf(Value); }
+#endif
 	[[nodiscard]] static FORCEINLINE double Exp(double Value) { return exp(Value); }
 
 	// Returns 2^Value
