Fix possible LiveCoding crash

#rnx

diff --git a/Engine/Source/Runtime/Core/Private/Logging/LogSuppressionInterface.cpp b/Engine/Source/Runtime/Core/Private/Logging/LogSuppressionInterface.cpp
--- a/Engine/Source/Runtime/Core/Private/Logging/LogSuppressionInterface.cpp
+++ b/Engine/Source/Runtime/Core/Private/Logging/LogSuppressionInterface.cpp
@@ -16,6 +16,7 @@
 #include "Misc/CommandLine.h"
 #include "Misc/ConfigCacheIni.h"
 #include "Misc/OutputDeviceHelper.h"
+#include "Modules/ModuleManager.h"
 
 namespace
 {
@@ -448,7 +449,10 @@ class FLogSuppressionImplementation: public FLogSuppressionInterface, private FS
 		{
 			if (It.Value() == Destination)
 			{
-				UE_LOG(LogHAL, Fatal,TEXT("Log suppression category %s was somehow declared twice with the same data."), *NameFName.ToString());
+				if (!IsReloadActive())
+				{
+					UE_LOG(LogHAL, Fatal,TEXT("Log suppression category %s was somehow declared twice with the same data."), *NameFName.ToString());
+				}
 			}
 			// if it is registered, it better be the same
 			if (It.Value()->CompileTimeVerbosity != Destination->CompileTimeVerbosity)
