FTextKey: Move implementation of delegating FTextKey constructors to TextKey.cpp to prevent MSVC inlining the FStringView construction and spending time optimizing the UE::Core::Private::StringLength loop.

diff --git a/Engine/Source/Runtime/Core/Private/Internationalization/TextKey.cpp b/Engine/Source/Runtime/Core/Private/Internationalization/TextKey.cpp
--- a/Engine/Source/Runtime/Core/Private/Internationalization/TextKey.cpp
+++ b/Engine/Source/Runtime/Core/Private/Internationalization/TextKey.cpp
@@ -607,6 +607,16 @@ FTextKey::FTextKey(FStringView InStr)
 	}
 }
 
+FTextKey::FTextKey(const TCHAR* InStr)
+	: FTextKey(FStringView(InStr))
+{
+}
+
+FTextKey::FTextKey(const FString& InStr)
+	: FTextKey(FStringView(InStr))
+{
+}
+
 const TCHAR* FTextKey::GetChars() const
 {
 	return Index != INDEX_NONE
diff --git a/Engine/Source/Runtime/Core/Public/Internationalization/TextKey.h b/Engine/Source/Runtime/Core/Public/Internationalization/TextKey.h
--- a/Engine/Source/Runtime/Core/Public/Internationalization/TextKey.h
+++ b/Engine/Source/Runtime/Core/Public/Internationalization/TextKey.h
@@ -76,16 +76,8 @@ class FTextKey
 public:
 	FTextKey() = default;
 	CORE_API FTextKey(FStringView InStr);
-	
-	FTextKey(const TCHAR* InStr)
-		: FTextKey(FStringView(InStr))
-	{
-	}
-
-	FTextKey(const FString& InStr)
-		: FTextKey(FStringView(InStr))
-	{
-	}
+	CORE_API FTextKey(const TCHAR* InStr);
+	CORE_API FTextKey(const FString& InStr);
 
 	/** Get the underlying chars buffer this text key represents */
 	UE_DEPRECATED(5.5, "GetChars is deprecated as FTextKey may now store its internal data as UTF-8. Use ToString/AppendString instead.")
