[CString/strnstr]
- Reads and tests 64 bit on each iterations to reduce branches and improve performance.
- Worst case where the entire string match the first character in find 326.12 MByte/s -> 357.73 MByte/s
- Best case where none of the string match the first character in find 598.25 MByte/s -> 8580.73 MByte/s (14x)
- When TSAN is active, we see an even greater improvement of 56.95 MByte/s -> 2200MByte/s (40x)

diff --git a/Engine/Source/Runtime/Core/Public/Misc/CString.h b/Engine/Source/Runtime/Core/Public/Misc/CString.h
--- a/Engine/Source/Runtime/Core/Public/Misc/CString.h
+++ b/Engine/Source/Runtime/Core/Public/Misc/CString.h
@@ -892,13 +892,72 @@ const typename TCString<T>::CharType* TCString<T>::Strnstr(const CharType* Str,
 	}
 
 	// get first letter of the find string (to reduce the number of full strncmps)
-	CharType FindInitial = *Find;
+	const CharType FindInitial = *Find;
 	// Set FindSuffix,FindSuffixLength to the characters of Find after the first letter
-	int32 FindSuffixLength = FindLen - 1;
+	const int32 FindSuffixLength = FindLen - 1;
 	const CharType* FindSuffix = Find + 1;
 
 	// while the length of the remaining string is >= FindLen
 	const CharType* StrLastChance = Str + InStrLen - FindLen;
+	
+	if constexpr (sizeof(CharType) >= 1 && sizeof(CharType) <= 2)
+	{
+		constexpr uint64 SignBit    = sizeof(CharType) == 1 ? 0x80ull : 0x8000ull;
+		constexpr uint64 Mask       = sizeof(CharType) == 1 ? 0xFFull : 0xFFFFull;
+		constexpr uint64 Ones64     = sizeof(CharType) == 1 ? 0x0101010101010101ull : 0x0001000100010001ull;
+		constexpr uint64 Positive64 = sizeof(CharType) == 1 ? 0x7F7F7F7F7F7F7F7Full : 0x7FFF7FFF7FFF7FFFull;
+		constexpr uint64 SignBit64  = sizeof(CharType) == 1 ? 0x8080808080808080ull : 0x8000800080008000ull;
+		constexpr uint64 CharPer64  = sizeof(uint64) / sizeof(CharType);
+		constexpr uint64 BitPerChar = sizeof(CharType) == 1 ? 8 : 16;
+		constexpr uint64 BitShift   = sizeof(CharType) == 1 ? 3 : 4;
+
+		InStrLen -= FindSuffixLength;
+		if (InStrLen >= CharPer64)
+		{
+			// Process normally if we're unaligned
+			for (; (UPTRINT(Str) & (BitPerChar-1)) && InStrLen > 1; --InStrLen)
+			{
+				CharType StrChar = *Str++;
+
+				// if it matches the first letter of the find string, do a string compare for the length of the find string
+				if (StrChar == FindInitial && !TCString<CharType>::Strncmp(Str, FindSuffix, FindSuffixLength))
+				{
+					// if we found the string, then return a pointer to the beginning of it in the search string
+					return Str - 1;
+				}
+			}
+
+			// Broadcast the initial value in each slots.
+			const uint64 FindInitial64 = (uint64)FindInitial * Ones64;
+
+			for (; InStrLen >= CharPer64; Str += CharPer64, InStrLen -= CharPer64)
+			{
+				uint64 StrValue = *(const uint64*)Str;
+				// We want to test each slots perfect match with the initial character
+				// The trick is to avoid overflowing a slot into another one.
+				// So we XOR and NEGATE so we end up with 0xFF when the character match.
+				// Then we remove the sign bit with 0x7F and ADD one, which will end up with 0x80
+				// if all the bits are set. We also AND with 0x80 to test the sign bit separately.
+				// And then we AND those two things together that will end up with 0x80 for each
+				// slot where there is a match.
+				StrValue = ~(StrValue ^ FindInitial64);
+				StrValue = ((StrValue & Positive64) + Ones64) & (StrValue & SignBit64);
+
+				// If any slot has a match, begin analysis
+				if (StrValue)
+				{
+					for (const CharType* InnerStr = Str; StrValue; StrValue >>= BitPerChar, ++InnerStr)
+					{
+						if ((StrValue & Mask) == SignBit && !TCString<CharType>::Strncmp(InnerStr + 1, FindSuffix, FindSuffixLength))
+						{
+							return InnerStr;
+						}
+					}
+				}
+			}
+		}
+	}
+
 	while (Str <= StrLastChance)
 	{
 		CharType StrChar = *Str++;
