[TextLocalizationManager]
- Move instead of copy to avoid useless and costly atomics for additional refcounting

diff --git a/Engine/Source/Runtime/Core/Private/Internationalization/TextLocalizationManager.cpp b/Engine/Source/Runtime/Core/Private/Internationalization/TextLocalizationManager.cpp
--- a/Engine/Source/Runtime/Core/Private/Internationalization/TextLocalizationManager.cpp
+++ b/Engine/Source/Runtime/Core/Private/Internationalization/TextLocalizationManager.cpp
@@ -829,15 +829,17 @@ void FTextLocalizationManager::AddOrUpdateDisplayStringInLiveTable(const FString
 	}
 	else if (SourceStringPtr)
 	{
-		// Add new entry
-		FDisplayStringEntry NewLiveEntry(
-			FTextKey(),													/*LocResID*/
-			INDEX_NONE,													/*LocalizationTargetPathId*/
-			FTextLocalizationResource::HashString(*SourceStringPtr),	/*SourceStringHash*/
-			MakeTextDisplayString(CopyTemp(DisplayString))				/*String*/
-		);
-
-		DisplayStringLookupTable.Emplace(TextId, NewLiveEntry);
+		{
+			// Add new entry
+			FDisplayStringEntry NewLiveEntry(
+				FTextKey(),													/*LocResID*/
+				INDEX_NONE,													/*LocalizationTargetPathId*/
+				FTextLocalizationResource::HashString(*SourceStringPtr),	/*SourceStringHash*/
+				MakeTextDisplayString(CopyTemp(DisplayString))				/*String*/
+			);
+
+			DisplayStringLookupTable.Add(TextId, MoveTemp(NewLiveEntry));
+		}
 
 		UE_LOG(LogConsoleResponse, Display, TEXT("Added string for Namespace='%s', Key='%s' to DisplayString='%s'"), *Namespace, *Key, *DisplayString);
 	}
