[CString/strnstr]
- Fix for PVS-Studio

diff --git a/Engine/Source/Runtime/Core/Public/Misc/CString.h b/Engine/Source/Runtime/Core/Public/Misc/CString.h
--- a/Engine/Source/Runtime/Core/Public/Misc/CString.h
+++ b/Engine/Source/Runtime/Core/Public/Misc/CString.h
@@ -930,7 +930,7 @@ const typename TCString<T>::CharType* TCString<T>::Strnstr(const CharType* Str,
 			// Broadcast the initial value in each slots.
 			const uint64 FindInitial64 = (uint64)FindInitial * Ones64;
 
-			for (; InStrLen >= CharPer64; Str += CharPer64, InStrLen -= CharPer64)
+			for (; InStrLen >= (int32)CharPer64; Str += CharPer64, InStrLen -= (int32)CharPer64)
 			{
 				uint64 StrValue = *(const uint64*)Str;
 				// We want to test each slots perfect match with the initial character
