Fixed a deadlock in FSharedMutex and FSharedRecursiveMutex

FWakeState::bHasWaitingThreads indicates that there MAY be waiting threads, and is not a guarantee. It cannot be used as an assumption that there is a waiting thread that will take the lock and later unlock it to wake the other threads.

diff --git a/Engine/Source/Runtime/Core/Private/Async/ParkingLot.cpp b/Engine/Source/Runtime/Core/Private/Async/ParkingLot.cpp
--- a/Engine/Source/Runtime/Core/Private/Async/ParkingLot.cpp
+++ b/Engine/Source/Runtime/Core/Private/Async/ParkingLot.cpp
@@ -729,9 +729,9 @@ uint32 WakeMultiple(const void* const Address, const uint32 WakeCount)
 	return uint32(WakeThreads.Num());
 }
 
-void WakeAll(const void* Address)
+uint32 WakeAll(const void* Address)
 {
-	WakeMultiple(Address, MAX_uint32);
+	return WakeMultiple(Address, MAX_uint32);
 }
 
 } // UE::ParkingLot
diff --git a/Engine/Source/Runtime/Core/Private/Async/SharedMutex.cpp b/Engine/Source/Runtime/Core/Private/Async/SharedMutex.cpp
--- a/Engine/Source/Runtime/Core/Private/Async/SharedMutex.cpp
+++ b/Engine/Source/Runtime/Core/Private/Async/SharedMutex.cpp
@@ -39,7 +39,7 @@ void FSharedMutex::LockSlow()
 		}
 
 		// Spin up to the spin limit while there are no waiting threads.
-		if (LIKELY(!(CurrentState & HasWaitingLockFlag) && SpinCount < SpinLimit))
+		if (LIKELY(!(CurrentState & MayHaveWaitingLockFlag) && SpinCount < SpinLimit))
 		{
 			FPlatformProcess::YieldThread();
 			++SpinCount;
@@ -48,13 +48,13 @@ void FSharedMutex::LockSlow()
 		}
 
 		// Store that there are waiting threads. Restart if the state has changed since it was loaded.
-		if (LIKELY(!(CurrentState & HasWaitingLockFlag)))
+		if (LIKELY(!(CurrentState & MayHaveWaitingLockFlag)))
 		{
-			if (UNLIKELY(!State.compare_exchange_weak(CurrentState, CurrentState | HasWaitingLockFlag, std::memory_order_relaxed)))
+			if (UNLIKELY(!State.compare_exchange_weak(CurrentState, CurrentState | MayHaveWaitingLockFlag, std::memory_order_relaxed)))
 			{
 				continue;
 			}
-			CurrentState |= HasWaitingLockFlag;
+			CurrentState |= MayHaveWaitingLockFlag;
 		}
 
 		// Do not enter oversubscription during a wait on a mutex since the wait is generally too short
@@ -62,7 +62,10 @@ void FSharedMutex::LockSlow()
 		LowLevelTasks::Private::FOversubscriptionAllowedScope _(false);
 
 		// Wait if the state has not changed. Either way, loop back and try to acquire the lock after trying to wait.
-		ParkingLot::Wait(GetLockAddress(), [this, CurrentState] { return State.load(std::memory_order_relaxed) == CurrentState; }, []{});
+		ParkingLot::Wait(GetLockAddress(), [this, CurrentState]
+		{
+			return State.load(std::memory_order_relaxed) == CurrentState;
+		}, []{});
 		CurrentState = State.load(std::memory_order_relaxed);
 	}
 }
@@ -74,7 +77,7 @@ void FSharedMutex::LockSharedSlow()
 	for (uint32 CurrentState = State.load(std::memory_order_relaxed);;)
 	{
 		// Try to acquire the lock if it is unlocked and there are no waiting threads.
-		if (LIKELY(!(CurrentState & (IsLockedFlag | HasWaitingLockFlag))))
+		if (LIKELY(!(CurrentState & (IsLockedFlag | MayHaveWaitingLockFlag))))
 		{
 			if (LIKELY(State.compare_exchange_weak(CurrentState, CurrentState + (1 << SharedLockCountShift), std::memory_order_acquire, std::memory_order_relaxed)))
 			{
@@ -84,7 +87,7 @@ void FSharedMutex::LockSharedSlow()
 		}
 
 		// Spin up to the spin limit while there are no waiting threads.
-		if (LIKELY(!(CurrentState & HasWaitingLockFlag) && SpinCount < SpinLimit))
+		if (LIKELY(!(CurrentState & MayHaveWaitingLockFlag) && SpinCount < SpinLimit))
 		{
 			FPlatformProcess::YieldThread();
 			++SpinCount;
@@ -93,13 +96,13 @@ void FSharedMutex::LockSharedSlow()
 		}
 
 		// Store that there are waiting threads. Restart if the state has changed since it was loaded.
-		if (LIKELY(!(CurrentState & HasWaitingSharedLockFlag)))
+		if (LIKELY(!(CurrentState & MayHaveWaitingSharedLockFlag)))
 		{
-			if (UNLIKELY(!State.compare_exchange_weak(CurrentState, CurrentState | HasWaitingSharedLockFlag, std::memory_order_relaxed)))
+			if (UNLIKELY(!State.compare_exchange_weak(CurrentState, CurrentState | MayHaveWaitingSharedLockFlag, std::memory_order_relaxed)))
 			{
 				continue;
 			}
-			CurrentState |= HasWaitingSharedLockFlag;
+			CurrentState |= MayHaveWaitingSharedLockFlag;
 		}
 
 		// Do not enter oversubscription during a wait on a mutex since the wait is generally too short
@@ -107,7 +110,10 @@ void FSharedMutex::LockSharedSlow()
 		LowLevelTasks::Private::FOversubscriptionAllowedScope _(false);
 
 		// Wait if the state has not changed. Either way, loop back and try to acquire the lock after trying to wait.
-		ParkingLot::Wait(GetSharedLockAddress(), [this, CurrentState] { return State.load(std::memory_order_relaxed) == CurrentState; }, []{});
+		ParkingLot::Wait(GetSharedLockAddress(), [this, CurrentState]
+		{
+			return State.load(std::memory_order_relaxed) == CurrentState;
+		}, []{});
 		CurrentState = State.load(std::memory_order_relaxed);
 	}
 }
@@ -116,9 +122,11 @@ FORCENOINLINE void FSharedMutex::WakeWaitingThread()
 {
 	ParkingLot::WakeOne(GetLockAddress(), [this](ParkingLot::FWakeState WakeState) -> uint64
 	{
-		if (!WakeState.bHasWaitingThreads)
+		if (!WakeState.bDidWake)
 		{
-			State.fetch_and(~HasWaitingLockFlag, std::memory_order_relaxed);
+			// Keep the flag until no thread wakes, otherwise shared locks may win before
+			// an exclusive lock has a chance.
+			State.fetch_and(~MayHaveWaitingLockFlag, std::memory_order_relaxed);
 		}
 		return 0;
 	});
@@ -126,15 +134,17 @@ FORCENOINLINE void FSharedMutex::WakeWaitingThread()
 
 FORCENOINLINE void FSharedMutex::WakeWaitingThreads(uint32 LastState)
 {
-	if (LastState & HasWaitingLockFlag)
+	if (LastState & MayHaveWaitingLockFlag)
 	{
 		// Wake one thread that is waiting to acquire an exclusive lock.
 		bool bDidWake = false;
 		ParkingLot::WakeOne(GetLockAddress(), [this, &bDidWake](ParkingLot::FWakeState WakeState) -> uint64
 		{
-			if (!WakeState.bHasWaitingThreads)
+			if (!WakeState.bDidWake)
 			{
-				State.fetch_and(~HasWaitingLockFlag, std::memory_order_relaxed);
+				// Keep the flag until no thread wakes, otherwise shared locks may win before
+				// an exclusive lock has a chance.
+				State.fetch_and(~MayHaveWaitingLockFlag, std::memory_order_relaxed);
 			}
 			bDidWake = WakeState.bDidWake;
 			return 0;
@@ -145,11 +155,11 @@ FORCENOINLINE void FSharedMutex::WakeWaitingThreads(uint32 LastState)
 		}
 	}
 
-	if (LastState & HasWaitingSharedLockFlag)
+	if (LastState & MayHaveWaitingSharedLockFlag)
 	{
 		// Wake every thread that is waiting to acquire a shared lock.
 		// The awoken threads might race against other exclusive locks.
-		if (State.fetch_and(~HasWaitingSharedLockFlag, std::memory_order_relaxed) & HasWaitingSharedLockFlag)
+		if (State.fetch_and(~MayHaveWaitingSharedLockFlag, std::memory_order_relaxed) & MayHaveWaitingSharedLockFlag)
 		{
 			ParkingLot::WakeAll(GetSharedLockAddress());
 		}
diff --git a/Engine/Source/Runtime/Core/Private/Async/SharedRecursiveMutex.cpp b/Engine/Source/Runtime/Core/Private/Async/SharedRecursiveMutex.cpp
--- a/Engine/Source/Runtime/Core/Private/Async/SharedRecursiveMutex.cpp
+++ b/Engine/Source/Runtime/Core/Private/Async/SharedRecursiveMutex.cpp
@@ -153,7 +153,7 @@ FORCENOINLINE void FSharedRecursiveMutex::LockSlow(uint32 CurrentState, const ui
 		}
 
 		// Spin up to the spin limit while there are no waiting threads.
-		if (LIKELY(!(CurrentState & HasWaitingLockFlag) && SpinCount < SpinLimit))
+		if (LIKELY(!(CurrentState & MayHaveWaitingLockFlag) && SpinCount < SpinLimit))
 		{
 			FPlatformProcess::YieldThread();
 			++SpinCount;
@@ -162,13 +162,13 @@ FORCENOINLINE void FSharedRecursiveMutex::LockSlow(uint32 CurrentState, const ui
 		}
 
 		// Store that there are waiting threads. Restart if the state has changed since it was loaded.
-		if (LIKELY(!(CurrentState & HasWaitingLockFlag)))
+		if (LIKELY(!(CurrentState & MayHaveWaitingLockFlag)))
 		{
-			if (UNLIKELY(!State.compare_exchange_weak(CurrentState, CurrentState | HasWaitingLockFlag, std::memory_order_relaxed)))
+			if (UNLIKELY(!State.compare_exchange_weak(CurrentState, CurrentState | MayHaveWaitingLockFlag, std::memory_order_relaxed)))
 			{
 				continue;
 			}
-			CurrentState |= HasWaitingLockFlag;
+			CurrentState |= MayHaveWaitingLockFlag;
 		}
 
 		// Do not enter oversubscription during a wait on a mutex since the wait is generally too short
@@ -176,7 +176,10 @@ FORCENOINLINE void FSharedRecursiveMutex::LockSlow(uint32 CurrentState, const ui
 		LowLevelTasks::Private::FOversubscriptionAllowedScope _(false);
 
 		// Wait if the state has not changed. Either way, loop back and try to acquire the lock after trying to wait.
-		ParkingLot::Wait(GetLockAddress(), [this, CurrentState] { return State.load(std::memory_order_relaxed) == CurrentState; }, []{});
+		ParkingLot::Wait(GetLockAddress(), [this, CurrentState]
+		{
+			return State.load(std::memory_order_relaxed) == CurrentState;
+		}, []{});
 		CurrentState = State.load(std::memory_order_relaxed);
 	}
 }
@@ -196,9 +199,9 @@ void FSharedRecursiveMutex::Unlock()
 		uint32 LastState = State.fetch_sub(1 << LockCountShift, std::memory_order_release);
 
 		// Wake one exclusive waiter or every shared waiter if there are waiting threads.
-		if (UNLIKELY(LastState & (HasWaitingLockFlag | HasWaitingSharedLockFlag)))
+		if (UNLIKELY(LastState & (MayHaveWaitingLockFlag | MayHaveWaitingSharedLockFlag)))
 		{
-			WakeWaitingThreads(CurrentState);
+			WakeWaitingThreads(LastState);
 		}
 	}
 	else
@@ -212,7 +215,7 @@ bool FSharedRecursiveMutex::TryLockShared(Core::Private::FSharedRecursiveMutexLi
 {
 	uint32 CurrentState = State.load(std::memory_order_relaxed);
 	// Recursive shared locks are quick to acquire. Check for that case and add 1 to the count.
-	if ((CurrentState & SharedLockCountMask) && (!(CurrentState & HasWaitingLockFlag) || Core::Private::FSharedRecursiveMutexLink::Owns(this)))
+	if ((CurrentState & SharedLockCountMask) && (!(CurrentState & MayHaveWaitingLockFlag) || Core::Private::FSharedRecursiveMutexLink::Owns(this)))
 	{
 		[[maybe_unused]] uint32 LastState = State.fetch_add(1 << SharedLockCountShift, std::memory_order_relaxed);
 		checkSlow((LastState & SharedLockCountMask) != SharedLockCountMask);
@@ -220,7 +223,7 @@ bool FSharedRecursiveMutex::TryLockShared(Core::Private::FSharedRecursiveMutexLi
 		return true;
 	}
 	// Try to acquire a shared lock if there is no active or waiting exclusive lock.
-	else if (!(CurrentState & (LockCountMask | HasWaitingLockFlag)) &&
+	else if (!(CurrentState & (LockCountMask | MayHaveWaitingLockFlag)) &&
 		State.compare_exchange_strong(CurrentState, CurrentState + (1 << SharedLockCountShift), std::memory_order_acquire, std::memory_order_relaxed))
 	{
 		Link.Push(this);
@@ -233,7 +236,7 @@ void FSharedRecursiveMutex::LockShared(Core::Private::FSharedRecursiveMutexLink&
 {
 	uint32 CurrentState = State.load(std::memory_order_relaxed);
 	// Recursive shared locks are quick to acquire. Check for that case and add 1 to the count.
-	if ((CurrentState & SharedLockCountMask) && (!(CurrentState & HasWaitingLockFlag) || Core::Private::FSharedRecursiveMutexLink::Owns(this)))
+	if ((CurrentState & SharedLockCountMask) && (!(CurrentState & MayHaveWaitingLockFlag) || Core::Private::FSharedRecursiveMutexLink::Owns(this)))
 	{
 		[[maybe_unused]] uint32 LastState = State.fetch_add(1 << SharedLockCountShift, std::memory_order_relaxed);
 		checkSlow((LastState & SharedLockCountMask) != SharedLockCountMask);
@@ -241,7 +244,7 @@ void FSharedRecursiveMutex::LockShared(Core::Private::FSharedRecursiveMutexLink&
 		return;
 	}
 	// Try to acquire a shared lock if there is no active or waiting exclusive lock.
-	else if (!(CurrentState & (LockCountMask | HasWaitingLockFlag)) &&
+	else if (!(CurrentState & (LockCountMask | MayHaveWaitingLockFlag)) &&
 		State.compare_exchange_weak(CurrentState, CurrentState + (1 << SharedLockCountShift), std::memory_order_acquire, std::memory_order_relaxed))
 	{
 		Link.Push(this);
@@ -257,7 +260,7 @@ FORCENOINLINE void FSharedRecursiveMutex::LockSharedSlow(Core::Private::FSharedR
 	for (uint32 CurrentState = State.load(std::memory_order_relaxed);;)
 	{
 		// Try to acquire the lock if it is unlocked and there are no waiting threads.
-		if (LIKELY(!(CurrentState & (LockCountMask | HasWaitingLockFlag))))
+		if (LIKELY(!(CurrentState & (LockCountMask | MayHaveWaitingLockFlag))))
 		{
 			checkSlow((CurrentState & SharedLockCountMask) != SharedLockCountMask);
 			if (LIKELY(State.compare_exchange_weak(CurrentState, CurrentState + (1 << SharedLockCountShift), std::memory_order_acquire, std::memory_order_relaxed)))
@@ -269,7 +272,7 @@ FORCENOINLINE void FSharedRecursiveMutex::LockSharedSlow(Core::Private::FSharedR
 		}
 
 		// Spin up to the spin limit while there are no waiting threads.
-		if (LIKELY(!(CurrentState & HasWaitingLockFlag) && SpinCount < SpinLimit))
+		if (LIKELY(!(CurrentState & MayHaveWaitingLockFlag) && SpinCount < SpinLimit))
 		{
 			FPlatformProcess::YieldThread();
 			++SpinCount;
@@ -278,13 +281,13 @@ FORCENOINLINE void FSharedRecursiveMutex::LockSharedSlow(Core::Private::FSharedR
 		}
 
 		// Store that there are waiting threads. Restart if the state has changed since it was loaded.
-		if (LIKELY(!(CurrentState & HasWaitingSharedLockFlag)))
+		if (LIKELY(!(CurrentState & MayHaveWaitingSharedLockFlag)))
 		{
-			if (UNLIKELY(!State.compare_exchange_weak(CurrentState, CurrentState | HasWaitingSharedLockFlag, std::memory_order_relaxed)))
+			if (UNLIKELY(!State.compare_exchange_weak(CurrentState, CurrentState | MayHaveWaitingSharedLockFlag, std::memory_order_relaxed)))
 			{
 				continue;
 			}
-			CurrentState |= HasWaitingSharedLockFlag;
+			CurrentState |= MayHaveWaitingSharedLockFlag;
 		}
 
 		// Do not enter oversubscription during a wait on a mutex since the wait is generally too short
@@ -292,7 +295,10 @@ FORCENOINLINE void FSharedRecursiveMutex::LockSharedSlow(Core::Private::FSharedR
 		LowLevelTasks::Private::FOversubscriptionAllowedScope _(false);
 
 		// Wait if the state has not changed. Either way, loop back and try to acquire the lock after trying to wait.
-		ParkingLot::Wait(GetSharedLockAddress(), [this, CurrentState] { return State.load(std::memory_order_relaxed) == CurrentState; }, []{});
+		ParkingLot::Wait(GetSharedLockAddress(), [this, CurrentState]
+		{
+			return State.load(std::memory_order_relaxed) == CurrentState;
+		}, []{});
 		CurrentState = State.load(std::memory_order_relaxed);
 	}
 }
@@ -302,33 +308,37 @@ void FSharedRecursiveMutex::UnlockShared(Core::Private::FSharedRecursiveMutexLin
 	Link.Pop();
 	const uint32 LastState = State.fetch_sub(1 << SharedLockCountShift, std::memory_order_release);
 	checkSlow(LastState & SharedLockCountMask);
-	constexpr uint32 WakeState = HasWaitingLockFlag | (1 << SharedLockCountShift);
-	if (UNLIKELY((LastState & ~HasWaitingSharedLockFlag) == WakeState))
+	constexpr uint32 WakeState = MayHaveWaitingLockFlag | (1 << SharedLockCountShift);
+	if (UNLIKELY((LastState & ~MayHaveWaitingSharedLockFlag) == WakeState))
 	{
 		// The last shared lock was released and there is a waiting exclusive lock.
 		// Wake one thread that is waiting to acquire an exclusive lock.
 		ParkingLot::WakeOne(GetLockAddress(), [this](ParkingLot::FWakeState WakeState) -> uint64
 		{
-			if (!WakeState.bHasWaitingThreads)
+			if (!WakeState.bDidWake)
 			{
-				State.fetch_and(~HasWaitingLockFlag, std::memory_order_relaxed);
+				// Keep the flag until no thread wakes, otherwise shared locks may win before
+				// an exclusive lock has a chance.
+				State.fetch_and(~MayHaveWaitingLockFlag, std::memory_order_relaxed);
 			}
 			return 0;
 		});
 	}
 }
 
-FORCENOINLINE void FSharedRecursiveMutex::WakeWaitingThreads(uint32 CurrentState)
+FORCENOINLINE void FSharedRecursiveMutex::WakeWaitingThreads(uint32 LastState)
 {
-	if (CurrentState & HasWaitingLockFlag)
+	if (LastState & MayHaveWaitingLockFlag)
 	{
 		// Wake one thread that is waiting to acquire an exclusive lock.
 		bool bDidWake = false;
 		ParkingLot::WakeOne(GetLockAddress(), [this, &bDidWake](ParkingLot::FWakeState WakeState) -> uint64
 		{
-			if (!WakeState.bHasWaitingThreads)
+			if (!WakeState.bDidWake)
 			{
-				State.fetch_and(~HasWaitingLockFlag, std::memory_order_relaxed);
+				// Keep the flag until no thread wakes, otherwise shared locks may win before
+				// an exclusive lock has a chance.
+				State.fetch_and(~MayHaveWaitingLockFlag, std::memory_order_relaxed);
 			}
 			bDidWake = WakeState.bDidWake;
 			return 0;
@@ -339,11 +349,11 @@ FORCENOINLINE void FSharedRecursiveMutex::WakeWaitingThreads(uint32 CurrentState
 		}
 	}
 
-	if (CurrentState & HasWaitingSharedLockFlag)
+	if (LastState & MayHaveWaitingSharedLockFlag)
 	{
 		// Wake every thread that is waiting to acquire a shared lock.
 		// The awoken threads might race against other exclusive locks.
-		if (State.fetch_and(~HasWaitingSharedLockFlag, std::memory_order_relaxed) & HasWaitingSharedLockFlag)
+		if (State.fetch_and(~MayHaveWaitingSharedLockFlag, std::memory_order_relaxed) & MayHaveWaitingSharedLockFlag)
 		{
 			ParkingLot::WakeAll(GetSharedLockAddress());
 		}
diff --git a/Engine/Source/Runtime/Core/Public/Async/ParkingLot.h b/Engine/Source/Runtime/Core/Public/Async/ParkingLot.h
--- a/Engine/Source/Runtime/Core/Public/Async/ParkingLot.h
+++ b/Engine/Source/Runtime/Core/Public/Async/ParkingLot.h
@@ -28,7 +28,7 @@ struct FWakeState final
 {
 	/** Did a thread wake up? */
 	bool bDidWake = false;
-	/** Does the queue have another thread waiting? */
+	/** Does the queue MAYBE have another thread waiting on this address? */
 	bool bHasWaitingThreads = false;
 };
 
@@ -114,8 +114,9 @@ UE_API uint32 WakeMultiple(const void* Address, uint32 WakeCount);
  * Wake all threads from the queue of threads waiting on the address.
  *
  * @param Address   Address to use as the key for the queue. Must match the address used in Wait.
+ * @return The number of threads that this call woke up.
  */
-UE_API void WakeAll(const void* Address);
+UE_API uint32 WakeAll(const void* Address);
 
 } // UE::ParkingLot
 
diff --git a/Engine/Source/Runtime/Core/Public/Async/SharedMutex.h b/Engine/Source/Runtime/Core/Public/Async/SharedMutex.h
--- a/Engine/Source/Runtime/Core/Public/Async/SharedMutex.h
+++ b/Engine/Source/Runtime/Core/Public/Async/SharedMutex.h
@@ -49,7 +49,7 @@ public:
 		// Unlock immediately to allow other threads to acquire the lock while this thread looks for a thread to wake.
 		uint32 LastState = State.fetch_sub(IsLockedFlag, std::memory_order_release);
 		checkSlow(LastState & IsLockedFlag);
-		if (LIKELY(!(LastState & (HasWaitingLockFlag | HasWaitingSharedLockFlag))))
+		if (LIKELY(!(LastState & (MayHaveWaitingLockFlag | MayHaveWaitingSharedLockFlag))))
 		{
 			return;
 		}
@@ -59,7 +59,7 @@ public:
 	inline bool TryLockShared()
 	{
 		uint32 Expected = State.load(std::memory_order_relaxed);
-		return !(Expected & (IsLockedFlag | HasWaitingLockFlag)) &&
+		return !(Expected & (IsLockedFlag | MayHaveWaitingLockFlag)) &&
 			State.compare_exchange_strong(Expected, Expected + (1 << SharedLockCountShift),
 				std::memory_order_acquire, std::memory_order_relaxed);
 	}
@@ -67,7 +67,7 @@ public:
 	inline void LockShared()
 	{
 		uint32 Expected = State.load(std::memory_order_relaxed);
-		if (LIKELY(!(Expected & (IsLockedFlag | HasWaitingLockFlag)) &&
+		if (LIKELY(!(Expected & (IsLockedFlag | MayHaveWaitingLockFlag)) &&
 			State.compare_exchange_weak(Expected, Expected + (1 << SharedLockCountShift),
 				std::memory_order_acquire, std::memory_order_relaxed)))
 		{
@@ -81,8 +81,8 @@ public:
 		// Unlock immediately to allow other threads to acquire the lock while this thread looks for a thread to wake.
 		const uint32 LastState = State.fetch_sub(1 << SharedLockCountShift, std::memory_order_release);
 		checkSlow(LastState & SharedLockCountMask);
-		constexpr uint32 WakeState = HasWaitingLockFlag | (1 << SharedLockCountShift);
-		if (LIKELY((LastState & ~HasWaitingSharedLockFlag) != WakeState))
+		constexpr uint32 WakeState = MayHaveWaitingLockFlag | (1 << SharedLockCountShift);
+		if (LIKELY((LastState & ~MayHaveWaitingSharedLockFlag) != WakeState))
 		{
 			return;
 		}
@@ -99,8 +99,8 @@ private:
 	const void* GetSharedLockAddress() const;
 
 	static constexpr uint32 IsLockedFlag = 1 << 0;
-	static constexpr uint32 HasWaitingLockFlag = 1 << 1;
-	static constexpr uint32 HasWaitingSharedLockFlag = 1 << 2;
+	static constexpr uint32 MayHaveWaitingLockFlag = 1 << 1;
+	static constexpr uint32 MayHaveWaitingSharedLockFlag = 1 << 2;
 	static constexpr uint32 SharedLockCountShift = 3;
 	static constexpr uint32 SharedLockCountMask = 0xffff'fff8;
 
diff --git a/Engine/Source/Runtime/Core/Public/Async/SharedRecursiveMutex.h b/Engine/Source/Runtime/Core/Public/Async/SharedRecursiveMutex.h
--- a/Engine/Source/Runtime/Core/Public/Async/SharedRecursiveMutex.h
+++ b/Engine/Source/Runtime/Core/Public/Async/SharedRecursiveMutex.h
@@ -61,8 +61,8 @@ private:
 	const void* GetLockAddress() const;
 	const void* GetSharedLockAddress() const;
 
-	static constexpr uint32 HasWaitingLockFlag = 1 << 0;
-	static constexpr uint32 HasWaitingSharedLockFlag = 1 << 1;
+	static constexpr uint32 MayHaveWaitingLockFlag = 1 << 0;
+	static constexpr uint32 MayHaveWaitingSharedLockFlag = 1 << 1;
 	static constexpr uint32 LockCountShift = 2;
 	static constexpr uint32 LockCountMask = 0x0000'0ffc;
 	static constexpr uint32 SharedLockCountShift = 12;
