Renamed HasWaitingLockFlag in mutex types for clarity

#rnx

diff --git a/Engine/Source/Runtime/Core/Private/Async/ExternalMutex.cpp b/Engine/Source/Runtime/Core/Private/Async/ExternalMutex.cpp
--- a/Engine/Source/Runtime/Core/Private/Async/ExternalMutex.cpp
+++ b/Engine/Source/Runtime/Core/Private/Async/ExternalMutex.cpp
@@ -27,7 +27,7 @@ void FExternalMutex::LockSlow()
 		}
 
 		// Spin up to the spin limit while there are no waiting threads.
-		if (LIKELY(!(CurrentState & HasWaitingLockFlag) && SpinCount < SpinLimit))
+		if (LIKELY(!(CurrentState & MayHaveWaitingLockFlag) && SpinCount < SpinLimit))
 		{
 			FPlatformProcess::YieldThread();
 			++SpinCount;
@@ -36,13 +36,13 @@ void FExternalMutex::LockSlow()
 		}
 
 		// Store that there are waiting threads. Restart if the state has changed since it was loaded.
-		if (LIKELY(!(CurrentState & HasWaitingLockFlag)))
+		if (LIKELY(!(CurrentState & MayHaveWaitingLockFlag)))
 		{
-			if (UNLIKELY(!State.compare_exchange_weak(CurrentState, CurrentState | HasWaitingLockFlag, std::memory_order_relaxed)))
+			if (UNLIKELY(!State.compare_exchange_weak(CurrentState, CurrentState | MayHaveWaitingLockFlag, std::memory_order_relaxed)))
 			{
 				continue;
 			}
-			CurrentState |= HasWaitingLockFlag;
+			CurrentState |= MayHaveWaitingLockFlag;
 		}
 
 		// Do not enter oversubscription during a wait on a mutex since the wait is generally too short
@@ -52,7 +52,7 @@ void FExternalMutex::LockSlow()
 		// Wait if the state has not changed. Either way, loop back and try to acquire the lock after trying to wait.
 		ParkingLot::Wait(&State, [this]
 		{
-			return (State.load(std::memory_order_relaxed) & (IsLockedFlag | HasWaitingLockFlag)) == (IsLockedFlag | HasWaitingLockFlag);
+			return (State.load(std::memory_order_relaxed) & (IsLockedFlag | MayHaveWaitingLockFlag)) == (IsLockedFlag | MayHaveWaitingLockFlag);
 		}, []{});
 		CurrentState = State.load(std::memory_order_relaxed);
 	}
@@ -66,7 +66,7 @@ FORCENOINLINE void FExternalMutex::WakeWaitingThread()
 		checkSlow(CurrentState & IsLockedFlag);
 
 		// Spin on the fast path because there may be spurious failures.
-		while (LIKELY(!(CurrentState & HasWaitingLockFlag)))
+		while (LIKELY(!(CurrentState & MayHaveWaitingLockFlag)))
 		{
 			if (LIKELY(State.compare_exchange_weak(CurrentState, CurrentState & ~IsLockedFlag, std::memory_order_release, std::memory_order_relaxed)))
 			{
@@ -81,14 +81,14 @@ FORCENOINLINE void FExternalMutex::WakeWaitingThread()
 		{
 			if (!WakeState.bHasWaitingThreads)
 			{
-				State.fetch_and(~HasWaitingLockFlag, std::memory_order_relaxed);
+				State.fetch_and(~MayHaveWaitingLockFlag, std::memory_order_relaxed);
 			}
 		}
 		else
 		{
-			const uint8 ClearFlags = IsLockedFlag | (WakeState.bHasWaitingThreads ? 0 : HasWaitingLockFlag);
+			const uint8 ClearFlags = IsLockedFlag | (WakeState.bHasWaitingThreads ? 0 : MayHaveWaitingLockFlag);
 			const uint8 OldState = State.fetch_and(~ClearFlags, std::memory_order_release);
-			checkSlow((OldState & IsLockedFlag) && (OldState & HasWaitingLockFlag));
+			checkSlow((OldState & IsLockedFlag) && (OldState & MayHaveWaitingLockFlag));
 		}
 		return 0;
 	});
diff --git a/Engine/Source/Runtime/Core/Private/Async/Mutex.cpp b/Engine/Source/Runtime/Core/Private/Async/Mutex.cpp
--- a/Engine/Source/Runtime/Core/Private/Async/Mutex.cpp
+++ b/Engine/Source/Runtime/Core/Private/Async/Mutex.cpp
@@ -27,7 +27,7 @@ void FMutex::LockSlow()
 		}
 
 		// Spin up to the spin limit while there are no waiting threads.
-		if (LIKELY(!(CurrentState & HasWaitingLockFlag) && SpinCount < SpinLimit))
+		if (LIKELY(!(CurrentState & MayHaveWaitingLockFlag) && SpinCount < SpinLimit))
 		{
 			FPlatformProcess::YieldThread();
 			++SpinCount;
@@ -36,13 +36,13 @@ void FMutex::LockSlow()
 		}
 
 		// Store that there are waiting threads. Restart if the state has changed since it was loaded.
-		if (LIKELY(!(CurrentState & HasWaitingLockFlag)))
+		if (LIKELY(!(CurrentState & MayHaveWaitingLockFlag)))
 		{
-			if (UNLIKELY(!State.compare_exchange_weak(CurrentState, CurrentState | HasWaitingLockFlag, std::memory_order_relaxed)))
+			if (UNLIKELY(!State.compare_exchange_weak(CurrentState, CurrentState | MayHaveWaitingLockFlag, std::memory_order_relaxed)))
 			{
 				continue;
 			}
-			CurrentState |= HasWaitingLockFlag;
+			CurrentState |= MayHaveWaitingLockFlag;
 		}
 
 		// Do not enter oversubscription during a wait on a mutex since the wait is generally too short
@@ -63,7 +63,7 @@ FORCENOINLINE void FMutex::WakeWaitingThread()
 		checkSlow(CurrentState & IsLockedFlag);
 
 		// Spin on the fast path because there may be spurious failures.
-		while (LIKELY(!(CurrentState & HasWaitingLockFlag)))
+		while (LIKELY(!(CurrentState & MayHaveWaitingLockFlag)))
 		{
 			if (LIKELY(State.compare_exchange_weak(CurrentState, 0, std::memory_order_release, std::memory_order_relaxed)))
 			{
@@ -78,14 +78,14 @@ FORCENOINLINE void FMutex::WakeWaitingThread()
 		{
 			if (!WakeState.bHasWaitingThreads)
 			{
-				State.fetch_and(~HasWaitingLockFlag, std::memory_order_relaxed);
+				State.fetch_and(~MayHaveWaitingLockFlag, std::memory_order_relaxed);
 			}
 		}
 		else
 		{
-			const uint8 NewState = WakeState.bHasWaitingThreads ? HasWaitingLockFlag : 0;
+			const uint8 NewState = WakeState.bHasWaitingThreads ? MayHaveWaitingLockFlag : 0;
 			const uint8 OldState = State.exchange(NewState, std::memory_order_release);
-			checkSlow((OldState & IsLockedFlag) && (OldState & HasWaitingLockFlag));
+			checkSlow((OldState & IsLockedFlag) && (OldState & MayHaveWaitingLockFlag));
 		}
 		return 0;
 	});
diff --git a/Engine/Source/Runtime/Core/Private/Async/RecursiveMutex.cpp b/Engine/Source/Runtime/Core/Private/Async/RecursiveMutex.cpp
--- a/Engine/Source/Runtime/Core/Private/Async/RecursiveMutex.cpp
+++ b/Engine/Source/Runtime/Core/Private/Async/RecursiveMutex.cpp
@@ -86,7 +86,7 @@ FORCENOINLINE void FRecursiveMutex::LockSlow(uint32 CurrentState, const uint32 C
 		}
 
 		// Spin up to the spin limit while there are no waiting threads.
-		if (LIKELY(!(CurrentState & HasWaitingLockFlag) && SpinCount < SpinLimit))
+		if (LIKELY(!(CurrentState & MayHaveWaitingLockFlag) && SpinCount < SpinLimit))
 		{
 			FPlatformProcess::YieldThread();
 			++SpinCount;
@@ -95,13 +95,13 @@ FORCENOINLINE void FRecursiveMutex::LockSlow(uint32 CurrentState, const uint32 C
 		}
 
 		// Store that there are waiting threads. Restart if the state has changed since it was loaded.
-		if (LIKELY(!(CurrentState & HasWaitingLockFlag)))
+		if (LIKELY(!(CurrentState & MayHaveWaitingLockFlag)))
 		{
-			if (UNLIKELY(!State.compare_exchange_weak(CurrentState, CurrentState | HasWaitingLockFlag, std::memory_order_relaxed)))
+			if (UNLIKELY(!State.compare_exchange_weak(CurrentState, CurrentState | MayHaveWaitingLockFlag, std::memory_order_relaxed)))
 			{
 				continue;
 			}
-			CurrentState |= HasWaitingLockFlag;
+			CurrentState |= MayHaveWaitingLockFlag;
 		}
 
 		// Do not enter oversubscription during a wait on a mutex since the wait is generally too short
@@ -131,14 +131,14 @@ void FRecursiveMutex::Unlock()
 			const uint32 LastState = State.fetch_sub(1 << LockCountShift, std::memory_order_release);
 
 			// Wake one exclusive waiter if there are waiting threads.
-			if (UNLIKELY(LastState & HasWaitingLockFlag))
+			if (UNLIKELY(LastState & MayHaveWaitingLockFlag))
 			{
 				WakeWaitingThread();
 			}
 		}
 		else
 		{
-			if (LIKELY(!(CurrentState & HasWaitingLockFlag)))
+			if (LIKELY(!(CurrentState & MayHaveWaitingLockFlag)))
 			{
 				if (LIKELY(State.compare_exchange_strong(CurrentState, 0, std::memory_order_release, std::memory_order_relaxed)))
 				{
@@ -164,14 +164,14 @@ FORCENOINLINE void FRecursiveMutex::WakeWaitingThread()
 		{
 			if (!WakeState.bHasWaitingThreads)
 			{
-				State.fetch_and(~HasWaitingLockFlag, std::memory_order_relaxed);
+				State.fetch_and(~MayHaveWaitingLockFlag, std::memory_order_relaxed);
 			}
 		}
 		else
 		{
-			const uint32 NewState = WakeState.bHasWaitingThreads ? HasWaitingLockFlag : 0;
+			const uint32 NewState = WakeState.bHasWaitingThreads ? MayHaveWaitingLockFlag : 0;
 			const uint32 OldState = State.exchange(NewState, std::memory_order_release);
-			checkSlow((OldState & LockCountMask) && (OldState & HasWaitingLockFlag));
+			checkSlow((OldState & LockCountMask) && (OldState & MayHaveWaitingLockFlag));
 		}
 		return 0;
 	});
diff --git a/Engine/Source/Runtime/Core/Public/Async/ExternalMutex.h b/Engine/Source/Runtime/Core/Public/Async/ExternalMutex.h
--- a/Engine/Source/Runtime/Core/Public/Async/ExternalMutex.h
+++ b/Engine/Source/Runtime/Core/Public/Async/ExternalMutex.h
@@ -51,7 +51,7 @@ public:
 
 	inline void Lock()
 	{
-		uint8 Expected = State.load(std::memory_order_relaxed) & ~IsLockedFlag & ~HasWaitingLockFlag;
+		uint8 Expected = State.load(std::memory_order_relaxed) & ~IsLockedFlag & ~MayHaveWaitingLockFlag;
 		if (LIKELY(State.compare_exchange_weak(Expected, Expected | IsLockedFlag, std::memory_order_acquire, std::memory_order_relaxed)))
 		{
 			return;
@@ -65,7 +65,7 @@ public:
 		{
 			// Unlock immediately to allow other threads to acquire the lock while this thread looks for a thread to wake.
 			const uint8 LastState = State.fetch_sub(IsLockedFlag, std::memory_order_release);
-			if (LIKELY(!(LastState & HasWaitingLockFlag)))
+			if (LIKELY(!(LastState & MayHaveWaitingLockFlag)))
 			{
 				return;
 			}
@@ -73,7 +73,7 @@ public:
 		}
 		else
 		{
-			uint8 Expected = (State.load(std::memory_order_relaxed) | IsLockedFlag) & ~HasWaitingLockFlag;
+			uint8 Expected = (State.load(std::memory_order_relaxed) | IsLockedFlag) & ~MayHaveWaitingLockFlag;
 			if (LIKELY(State.compare_exchange_weak(Expected, Expected & ~IsLockedFlag, std::memory_order_release, std::memory_order_relaxed)))
 			{
 				return;
@@ -87,7 +87,7 @@ private:
 	UE_API void WakeWaitingThread();
 
 	static constexpr uint8 IsLockedFlag = 1 << 0;
-	static constexpr uint8 HasWaitingLockFlag = 1 << 1;
+	static constexpr uint8 MayHaveWaitingLockFlag = 1 << 1;
 
 	static constexpr bool bUnlockImmediately = false;
 
diff --git a/Engine/Source/Runtime/Core/Public/Async/Mutex.h b/Engine/Source/Runtime/Core/Public/Async/Mutex.h
--- a/Engine/Source/Runtime/Core/Public/Async/Mutex.h
+++ b/Engine/Source/Runtime/Core/Public/Async/Mutex.h
@@ -57,7 +57,7 @@ public:
 		{
 			// Unlock immediately to allow other threads to acquire the lock while this thread looks for a thread to wake.
 			const uint8 LastState = State.fetch_sub(IsLockedFlag, std::memory_order_release);
-			if (LIKELY(!(LastState & HasWaitingLockFlag)))
+			if (LIKELY(!(LastState & MayHaveWaitingLockFlag)))
 			{
 				return;
 			}
@@ -81,7 +81,7 @@ private:
 	static constexpr bool bUnlockImmediately = false;
 
 	static constexpr uint8 IsLockedFlag = 1 << 0;
-	static constexpr uint8 HasWaitingLockFlag = 1 << 1;
+	static constexpr uint8 MayHaveWaitingLockFlag = 1 << 1;
 
 	std::atomic<uint8> State = 0;
 };
diff --git a/Engine/Source/Runtime/Core/Public/Async/RecursiveMutex.h b/Engine/Source/Runtime/Core/Public/Async/RecursiveMutex.h
--- a/Engine/Source/Runtime/Core/Public/Async/RecursiveMutex.h
+++ b/Engine/Source/Runtime/Core/Public/Async/RecursiveMutex.h
@@ -31,7 +31,7 @@ private:
 	void LockSlow(uint32 CurrentState, uint32 CurrentThreadId);
 	void WakeWaitingThread();
 
-	static constexpr uint32 HasWaitingLockFlag = 1 << 0;
+	static constexpr uint32 MayHaveWaitingLockFlag = 1 << 0;
 	static constexpr uint32 LockCountShift = 1;
 	static constexpr uint32 LockCountMask = 0xffff'fffe;
 
