Logging: Improvements to the backlog

- Automatically serialize the backlog into added output devices.
- Fixed incorrectly supported copy/move of FOutputDeviceLine.
- Enabled the backlog by default for non-program targets and enabled it earlier during engine initialization otherwise, to capture more early logging in the backlog.
- Deprecated IsBacklogEnabled and SerializeBacklog and exposed them internally for crash reporting.

diff --git a/Engine/Plugins/AudioInsights/Source/AudioInsights/Private/Views/LogDashboardViewFactory.cpp b/Engine/Plugins/AudioInsights/Source/AudioInsights/Private/Views/LogDashboardViewFactory.cpp
--- a/Engine/Plugins/AudioInsights/Source/AudioInsights/Private/Views/LogDashboardViewFactory.cpp
+++ b/Engine/Plugins/AudioInsights/Source/AudioInsights/Private/Views/LogDashboardViewFactory.cpp
@@ -29,11 +29,6 @@ namespace UE::Audio::Insights
 		}
 	} // namespace LogDashboardViewFactoryPrivate
 
-	FLogDashboardViewFactory::FLogDashboardViewFactory()
-	{
-		CategoryCollector.RunAsync();
-	}
-
 	FName FLogDashboardViewFactory::GetName() const
 	{
 		return "Log";
@@ -92,17 +87,6 @@ namespace UE::Audio::Insights
 		return CollectedCategories;
 	}
 
-	void FLogDashboardViewFactory::FLogCategoryCollector::RunAsync()
-	{
-		FScopeLock CollectionLock(&CollectionCritSec);
-		CollectedCategories.Empty();
-
-		if (GLog)
-		{
-			GLog->SerializeBacklog(this);
-		}
-	}
-
 	void FLogDashboardViewFactory::FLogCategoryCollector::Serialize(const TCHAR* InMsg, ELogVerbosity::Type Verbosity, const FName& InCategory)
 	{
 		using namespace LogDashboardViewFactoryPrivate;
diff --git a/Engine/Plugins/AudioInsights/Source/AudioInsights/Public/Views/LogDashboardViewFactory.h b/Engine/Plugins/AudioInsights/Source/AudioInsights/Public/Views/LogDashboardViewFactory.h
--- a/Engine/Plugins/AudioInsights/Source/AudioInsights/Public/Views/LogDashboardViewFactory.h
+++ b/Engine/Plugins/AudioInsights/Source/AudioInsights/Public/Views/LogDashboardViewFactory.h
@@ -18,8 +18,6 @@ namespace UE::Audio::Insights
 	class AUDIOINSIGHTS_API FLogDashboardViewFactory : public IDashboardViewFactory
 	{
 	public:
-		FLogDashboardViewFactory();
-
 		virtual FName GetName() const override;
 		virtual FText GetDisplayName() const override;
 		virtual EDefaultDashboardTabStack GetDefaultTabStack() const override;
@@ -33,7 +31,6 @@ namespace UE::Audio::Insights
 			~FLogCategoryCollector();
 
 			FDefaultCategorySelectionMap GetCollectedCategories() const;
-			void RunAsync();
 
 			virtual bool IsMemoryOnly() const override;
 			virtual void Serialize(const TCHAR* InMsg, ELogVerbosity::Type Verbosity, const FName& InCategory) override;
diff --git a/Engine/Plugins/Experimental/Avalanche/Source/AvalancheMedia/Private/Playback/AvaPlaybackServer.cpp b/Engine/Plugins/Experimental/Avalanche/Source/AvalancheMedia/Private/Playback/AvaPlaybackServer.cpp
--- a/Engine/Plugins/Experimental/Avalanche/Source/AvalancheMedia/Private/Playback/AvaPlaybackServer.cpp
+++ b/Engine/Plugins/Experimental/Avalanche/Source/AvalancheMedia/Private/Playback/AvaPlaybackServer.cpp
@@ -1914,7 +1914,6 @@ FAvaPlaybackServer::FReplicationOutputDevice::FReplicationOutputDevice(FAvaPlayb
 		: Server(InServer)
 {
 	GLog->AddOutputDevice(this);
-	GLog->SerializeBacklog(this);
 }
 		
 FAvaPlaybackServer::FReplicationOutputDevice::~FReplicationOutputDevice()
diff --git a/Engine/Source/Developer/OutputLog/Private/OutputLogModule.cpp b/Engine/Source/Developer/OutputLog/Private/OutputLogModule.cpp
--- a/Engine/Source/Developer/OutputLog/Private/OutputLogModule.cpp
+++ b/Engine/Source/Developer/OutputLog/Private/OutputLogModule.cpp
@@ -62,7 +62,6 @@ public:
 	FOutputLogHistory()
 	{
 		GLog->AddOutputDevice(this);
-		GLog->SerializeBacklog(this);
 	}
 
 	~FOutputLogHistory()
diff --git a/Engine/Source/Runtime/Core/Private/GenericPlatform/GenericPlatformCrashContext.cpp b/Engine/Source/Runtime/Core/Private/GenericPlatform/GenericPlatformCrashContext.cpp
--- a/Engine/Source/Runtime/Core/Private/GenericPlatform/GenericPlatformCrashContext.cpp
+++ b/Engine/Source/Runtime/Core/Private/GenericPlatform/GenericPlatformCrashContext.cpp
@@ -1572,7 +1572,7 @@ FString FGenericCrashContext::DumpLog(const FString& CrashFolderAbsolute)
 	// If we have a memory only log, make sure it's dumped to file before we attach it to the report
 #if !NO_LOGGING
 	bool bMemoryOnly = FPlatformOutputDevices::GetLog()->IsMemoryOnly();
-	bool bBacklogEnabled = FOutputDeviceRedirector::Get()->IsBacklogEnabled();
+	bool bBacklogEnabled = FOutputDeviceRedirector::Get()->HasBacklogForCrash();
 
 	if (bMemoryOnly || bBacklogEnabled)
 	{
@@ -1586,7 +1586,7 @@ FString FGenericCrashContext::DumpLog(const FString& CrashFolderAbsolute)
 			else
 			{
 				FOutputDeviceArchiveWrapper Wrapper(LogFile.Get());
-				GLog->SerializeBacklog(&Wrapper);
+				GLog->SerializeBacklogForCrash(&Wrapper);
 			}
 
 			LogFile->Flush();
diff --git a/Engine/Source/Runtime/Core/Private/Misc/OutputDeviceRedirector.cpp b/Engine/Source/Runtime/Core/Private/Misc/OutputDeviceRedirector.cpp
--- a/Engine/Source/Runtime/Core/Private/Misc/OutputDeviceRedirector.cpp
+++ b/Engine/Source/Runtime/Core/Private/Misc/OutputDeviceRedirector.cpp
@@ -64,8 +64,17 @@ struct FOutputDeviceLine
 	const FName Category;
 	const ELogVerbosity::Type Verbosity;
 
-	FOutputDeviceLine(const FBufferedLine&) = delete;
-	FOutputDeviceLine& operator=(const FBufferedLine&) = delete;
+	FOutputDeviceLine(const FOutputDeviceLine&) = delete;
+	FOutputDeviceLine& operator=(const FOutputDeviceLine&) = delete;
+
+	FORCEINLINE FOutputDeviceLine(FOutputDeviceLine&& Other)
+		: Time(Other.Time)
+		, Data(Other.Data)
+		, Category(Other.Category)
+		, Verbosity(Other.Verbosity)
+	{
+		Other.Data = nullptr;
+	}
 
 	FORCEINLINE FOutputDeviceLine(const TCHAR* const InData, const FName& InCategory, const ELogVerbosity::Type InVerbosity, const double InTime)
 		: Time(InTime)
@@ -77,7 +86,10 @@ struct FOutputDeviceLine
 
 	FORCEINLINE ~FOutputDeviceLine()
 	{
-		FOutputDeviceLinearAllocator::Free(const_cast<TCHAR*>(Data));
+		if (Data)
+		{
+			FOutputDeviceLinearAllocator::Free(const_cast<TCHAR*>(Data));
+		}
 	}
 
 private:
@@ -102,6 +114,12 @@ struct FOutputDeviceItem
 		: Value(TInPlaceType<FLogRecord>(), Record)
 	{
 	}
+
+	FOutputDeviceItem(FOutputDeviceItem&&) = default;
+	FOutputDeviceItem& operator=(FOutputDeviceItem&&) = default;
+
+	FOutputDeviceItem(const FOutputDeviceItem&) = delete;
+	FOutputDeviceItem& operator=(const FOutputDeviceItem&) = delete;
 };
 
 static constexpr uint64 CalculateRedirectorCacheLinePadding(const uint64 Size)
@@ -154,7 +172,7 @@ struct FOutputDeviceRedirectorState
 	std::atomic<bool> bThreadStarted = false;
 
 	/** Whether the backlog is enabled. */
-	bool bEnableBacklog = false;
+	bool bEnableBacklog = !IS_PROGRAM; // Programs may not use FEngineLoop which disables the backlog but other targets do.
 
 	/** Whether the output device at the corresponding index can be used on the panic thread. */
 	TBitArray<TInlineAllocator<1>> BufferedOutputDevicesCanBeUsedOnPanicThread;
@@ -216,6 +234,7 @@ struct FOutputDeviceRedirectorState
 
 	FORCENOINLINE void AddToBacklog(const FLogRecord& Record);
 	FORCENOINLINE void AddToBacklog(const TCHAR* Data, ELogVerbosity::Type Verbosity, const FName& Category, double Time);
+	FORCENOINLINE void SerializeBacklog(TNotNull<FOutputDevice*> OutputDevice);
 };
 
 /**
@@ -539,6 +558,15 @@ void FOutputDeviceRedirectorState::AddToBacklog(
 	BacklogLines.Emplace(Data, Category, Verbosity, Time);
 }
 
+void FOutputDeviceRedirectorState::SerializeBacklog(TNotNull<FOutputDevice*> OutputDevice)
+{
+	FReadScopeLock ScopeLock(BacklogLock);
+	for (const FBufferedLine& BacklogLine : BacklogLines)
+	{
+		OutputDevice->Serialize(BacklogLine.Data.Get(), BacklogLine.Verbosity, BacklogLine.Category, BacklogLine.Time);
+	}
+}
+
 } // UE::Private
 
 FOutputDeviceRedirector::FOutputDeviceRedirector()
@@ -567,6 +595,10 @@ void FOutputDeviceRedirector::AddOutputDevice(FOutputDevice* OutputDevice)
 	if (OutputDevice)
 	{
 		State->AddOutputDevice(OutputDevice);
+		if (State->bEnableBacklog)
+		{
+			State->SerializeBacklog(OutputDevice);
+		}
 	}
 }
 
@@ -606,11 +638,7 @@ void FOutputDeviceRedirector::FlushThreadedLogs(EOutputDeviceRedirectorFlushOpti
 
 void FOutputDeviceRedirector::SerializeBacklog(FOutputDevice* OutputDevice)
 {
-	FReadScopeLock ScopeLock(State->BacklogLock);
-	for (const FBufferedLine& BacklogLine : State->BacklogLines)
-	{
-		OutputDevice->Serialize(BacklogLine.Data.Get(), BacklogLine.Verbosity, BacklogLine.Category, BacklogLine.Time);
-	}
+	// Deprecated. Do nothing because AddOutputDevice now serializes the backlog.
 }
 
 void FOutputDeviceRedirector::EnableBacklog(bool bEnable)
@@ -845,11 +873,21 @@ void FOutputDeviceRedirector::TearDown()
 }
 
 bool FOutputDeviceRedirector::IsBacklogEnabled() const
+{
+	return HasBacklogForCrash();
+}
+
+bool FOutputDeviceRedirector::HasBacklogForCrash() const
 {
 	FReadScopeLock Lock(State->BacklogLock);
 	return State->bEnableBacklog;
 }
 
+void FOutputDeviceRedirector::SerializeBacklogForCrash(TNotNull<FOutputDevice*> OutputDevice) const
+{
+	State->SerializeBacklog(OutputDevice);
+}
+
 CORE_API FOutputDeviceRedirector* GetGlobalLogSingleton()
 {
 	return GWasGLogDestroyed ? nullptr : FOutputDeviceRedirector::Get();
diff --git a/Engine/Source/Runtime/Core/Private/Unix/UnixPlatformCrashContext.cpp b/Engine/Source/Runtime/Core/Private/Unix/UnixPlatformCrashContext.cpp
--- a/Engine/Source/Runtime/Core/Private/Unix/UnixPlatformCrashContext.cpp
+++ b/Engine/Source/Runtime/Core/Private/Unix/UnixPlatformCrashContext.cpp
@@ -642,7 +642,7 @@ void FUnixCrashContext::GenerateCrashInfoAndLaunchReporter() const
 
 #if !NO_LOGGING
 			bool bMemoryOnly = FPlatformOutputDevices::GetLog()->IsMemoryOnly();
-			bool bBacklogEnabled = FOutputDeviceRedirector::Get()->IsBacklogEnabled();
+			bool bBacklogEnabled = FOutputDeviceRedirector::Get()->HasBacklogForCrash();
 
 			// The minimum free space on drive for saving a crash log
 			const uint64 MinDriveSpaceForCrashLog = 250 * 1024 * 1024;
@@ -675,7 +675,7 @@ void FUnixCrashContext::GenerateCrashInfoAndLaunchReporter() const
 						else
 						{
 							FOutputDeviceArchiveWrapper Wrapper(LogFile);
-							GLog->SerializeBacklog(&Wrapper);
+							GLog->SerializeBacklogForCrash(&Wrapper);
 						}
 
 						LogFile->Flush();
diff --git a/Engine/Source/Runtime/Core/Public/Misc/OutputDeviceRedirector.h b/Engine/Source/Runtime/Core/Public/Misc/OutputDeviceRedirector.h
--- a/Engine/Source/Runtime/Core/Public/Misc/OutputDeviceRedirector.h
+++ b/Engine/Source/Runtime/Core/Public/Misc/OutputDeviceRedirector.h
@@ -7,6 +7,7 @@
 #include "CoreTypes.h"
 #include "Logging/LogVerbosity.h"
 #include "Misc/EnumClassFlags.h"
+#include "Misc/NotNull.h"
 #include "Misc/OutputDevice.h"
 #include "Templates/PimplPtr.h"
 #include "Templates/UniquePtr.h"
@@ -89,6 +90,7 @@ public:
 	 * Serializes the current backlog to the specified output device.
 	 * @param OutputDevice   Output device that will receive the current backlog.
 	 */
+	UE_DEPRECATED(5.6, "The backlog is now automatically serialized within AddOutputDevice.")
 	CORE_API void SerializeBacklog(FOutputDevice* OutputDevice);
 
 	/**
@@ -164,11 +166,15 @@ public:
 	 */
 	CORE_API void TearDown() final;
 
-	/**
-	 * Determine if the backlog is enabled.
-	 */
+	/** Determine if the backlog is enabled. */
+	UE_DEPRECATED(5.6, "The backlog is now automatically serialized within AddOutputDevice.")
 	CORE_API bool IsBacklogEnabled() const;
 
+	/** Internal. Used for crash reporting. */
+	UE_INTERNAL bool HasBacklogForCrash() const;
+	/** Internal. Used for crash reporting. */
+	UE_INTERNAL void SerializeBacklogForCrash(TNotNull<FOutputDevice*> OutputDevice) const;
+
 private:
 	TPimplPtr<UE::Private::FOutputDeviceRedirectorState> State;
 };
diff --git a/Engine/Source/Runtime/Launch/Private/LaunchEngineLoop.cpp b/Engine/Source/Runtime/Launch/Private/LaunchEngineLoop.cpp
--- a/Engine/Source/Runtime/Launch/Private/LaunchEngineLoop.cpp
+++ b/Engine/Source/Runtime/Launch/Private/LaunchEngineLoop.cpp
@@ -1586,10 +1586,11 @@ int32 FEngineLoop::PreInitPreStartupScreen(const TCHAR* CmdLine)
 	// This lazy initialization can happen during initialization of a DLL, which Windows does on a
 	// worker thread, which makes that worker thread the primary thread. Make this the primary thread
 	// until initialization is far enough along to try to start a dedicated primary thread.
-	if (GLog)
-	{
-		GLog->SetCurrentThreadAsPrimaryThread();
-	}
+	GLog->SetCurrentThreadAsPrimaryThread();
+
+	// Most targets have the backlog enabled by default. Enable it as early as possible here for
+	// cases where it is disabled by default. It will be disabled when module loading is complete.
+	GLog->EnableBacklog(true);
 
 	// Command line option for enabling named events
 	if (FParse::Param(CmdLine, TEXT("statnamedevents")))
@@ -1734,10 +1735,6 @@ int32 FEngineLoop::PreInitPreStartupScreen(const TCHAR* CmdLine)
 	}
 #endif
 
-	// Always enable the backlog so we get all messages, we will disable and clear it in the game
-	// as soon as we determine whether GIsEditor == false
-	GLog->EnableBacklog(true);
-
 	// Initialize std out device as early as possible if requested in the command line
 #if PLATFORM_DESKTOP
 	// consoles don't typically have stdout, and FOutputDeviceDebug is responsible for echoing logs to the
