[RaceDetector]
- Add history replay support to improve performance (almost 2x) and allow full callstack recovery.
- Add cvar to increase the history size if needed
- Add history memory usage reporting
- Move the construction of the current stack inside ReportRace to reduce the stack space required by InstrumentAccess to remove a call to __chkstk and improve its performance
- Fixed memory leak for context when a thread exit
- Remove the EnsureContextRegistered function call which was not inlined even if specifically told to with FORCEINLINE
- Fix a FMemoryAccess constructor issue that caused lots of CPU stalls on Store-to-Load Forwarding
- FORCEINLINE the entire InstrumentAccess for another almost 2x perf boost in warm case
- This CL makes the race detector between 2x and 4x faster depending on use cases

diff --git a/Engine/Source/Runtime/Core/Internal/Sanitizer/RaceDetectorTypes.h b/Engine/Source/Runtime/Core/Internal/Sanitizer/RaceDetectorTypes.h
--- a/Engine/Source/Runtime/Core/Internal/Sanitizer/RaceDetectorTypes.h
+++ b/Engine/Source/Runtime/Core/Internal/Sanitizer/RaceDetectorTypes.h
@@ -17,6 +17,8 @@
 #include "Instrumentation/Types.h"
 #include "Instrumentation/Containers.h"
 
+extern int32 GRaceDetectorHistoryLength;
+
 namespace UE::Sanitizer::RaceDetector {
 
 	using namespace UE::Instrumentation;
@@ -92,7 +94,6 @@ namespace UE::Sanitizer::RaceDetector {
 		void* Locations[256] = {};
 	};
 
-
 	// ------------------------------------------------------------------------------
 	// Memory access.
 	// ------------------------------------------------------------------------------
@@ -107,7 +108,8 @@ namespace UE::Sanitizer::RaceDetector {
 		ACCESS_TYPE_ATOMIC_WRITE = ACCESS_TYPE_ATOMIC | ACCESS_TYPE_WRITE,
 		ACCESS_TYPE_ATOMIC_READ_WRITE = ACCESS_TYPE_ATOMIC_READ | ACCESS_TYPE_ATOMIC_WRITE
 	};
-	ENUM_CLASS_FLAGS(EMemoryAccessType);
+	INSTRUMENTATION_ENUM_CLASS_FLAGS(EMemoryAccessType);
+
 	INSTRUMENTATION_FUNCTION_ATTRIBUTES inline const TCHAR* AccessTypeToString(EMemoryAccessType AccessType)
 	{
 		switch (AccessType & ACCESS_TYPE_ATOMIC_READ_WRITE)
@@ -152,12 +154,13 @@ namespace UE::Sanitizer::RaceDetector {
 		}
 
 		INSTRUMENTATION_FUNCTION_ATTRIBUTES FORCEINLINE FMemoryAccess(FContextId InContextId, FClock InClock, uint8 InOffset, uint8 InSize, EMemoryAccessType InAccessType, bool bInHasSyncObjectAttached)
-			: Clock(InClock)
-			, ContextId(InContextId)
-			, Access((uint8)(((1ul<<InSize) - 1ul) << InOffset))
-			, AccessType(InAccessType)
-			, bHasSyncObjectAttached(bInHasSyncObjectAttached)
 		{
+			// Apparently this is much faster than using the bitfields. 30% faster in warm TSAN benchmark!!!
+			// The constructor was accessing the same value in memory/store buffer multiple time doing its bit tweaking and it caused tons of Core::X86::Pmc::Core::LsBadStatus2.
+			// That ended up stalling on S[0].RawValue == CurrentAccess.RawValue in InstrumentMemoryAccess when trying to extract the full uint64.
+			// https://blog.stuffedcow.net/2014/01/x86-memory-disambiguation/
+			const uint8 AccessValue = (uint8)(((1ull << InSize) - 1ull) << InOffset);
+			RawValue = (uint64)bInHasSyncObjectAttached << 56 | (uint64)InAccessType << 48 | (uint64)AccessValue << 40 | (uint64)InContextId << 32 | (uint64)InClock;
 		}
 		
 		INSTRUMENTATION_FUNCTION_ATTRIBUTES FORCEINLINE uint8 GetOffset() const { return (uint8)FMath::CountTrailingZeros(Access); }
@@ -166,57 +169,339 @@ namespace UE::Sanitizer::RaceDetector {
 
 		union
 		{
+			// If you touch this layout, make sure to update the constructor above
 			struct
 			{
-				FClock     Clock;
-				FContextId ContextId;
+				FClock     Clock;                    // 0
+				FContextId ContextId;                // 4
 				// Each bit represent a 1-byte slot used in our 8 byte shadow
 				// and can easily be tested for overlaps with other accesses.
-				uint8      Access; 
+				uint8      Access;                   // 5
 				union
 				{
 					struct
 					{
-						uint8 bIsRead : 1;
-						uint8 bIsWrite : 1;
+						uint8 bIsRead   : 1;
+						uint8 bIsWrite  : 1;
 						uint8 bIsAtomic : 1;
-						uint8 bIsVPtr : 1;
+						uint8 bIsVPtr   : 1;
 					};
-					EMemoryAccessType AccessType : 4;
+					EMemoryAccessType AccessType : 4; // 6  plenty of bits left here
 				};
-				uint8 bHasSyncObjectAttached : 1;
+				uint8 bHasSyncObjectAttached : 1;     // 7  plenty of bits left here
 			};
 
 			uint64 RawValue;
 		};
 	};
 
-	// Used to store all accesses made by a thread since we don't have 
-	// enough space in the shadow memory to store callstacks.
-	// Storage space can certainly be improved by storing
-	// function entry and exit instead of full callstack every time.
-	// The current setup amount to about 8MB of memory per thread
-	// for a maximum of 256 thread the shadow memory can account for.
-	// This is totally acceptable for race detection purpose for now
-	// and can be revisited later if needed.
+	enum class EHistoryEntryType : uint8
+	{
+                Invalid       = 0,     // Just make sure we do not mistake 0 memory for something valid
+		FunctionEntry = 0xAA,  // Any number will do but make them stand out in the trace
+		MemoryAccess  = 0xBB,
+		FunctionExit  = 0xCC
+	};
+
+	struct FHistoryEntryBase
+	{
+		EHistoryEntryType  Type;
+
+		INSTRUMENTATION_FUNCTION_ATTRIBUTES FHistoryEntryBase(EHistoryEntryType InType)
+			: Type(InType)
+		{
+		}
+	};
+
+	struct FHistoryEntryAccess : public FHistoryEntryBase
+	{
+		void*          Pointer;
+		FMemoryAccess  Access;
+
+		INSTRUMENTATION_FUNCTION_ATTRIBUTES FHistoryEntryAccess(void* InPointer, uint64 InRawAccess)
+			: FHistoryEntryBase(EHistoryEntryType::MemoryAccess)
+			, Pointer(InPointer)
+			, Access(InRawAccess)
+		{
+		}
+	};
+
+	struct FHistoryEntryFunctionEntry : public FHistoryEntryBase
+	{
+		void* ReturnAddress;
+
+		INSTRUMENTATION_FUNCTION_ATTRIBUTES FHistoryEntryFunctionEntry(void* InReturnAddress)
+			: FHistoryEntryBase(EHistoryEntryType::FunctionEntry)
+			, ReturnAddress(InReturnAddress)
+		{
+		}
+	};
+
+	struct FHistoryEntryFunctionExit : public FHistoryEntryBase
+	{
+		INSTRUMENTATION_FUNCTION_ATTRIBUTES FHistoryEntryFunctionExit()
+			: FHistoryEntryBase(EHistoryEntryType::FunctionExit)
+		{
+		}
+	};
+
+	struct FHistoryChunk
+	{
+		SAFE_OPERATOR_NEW_DELETE()
+
+		INSTRUMENTATION_FUNCTION_ATTRIBUTES FHistoryChunk();
+		INSTRUMENTATION_FUNCTION_ATTRIBUTES ~FHistoryChunk();
+		INSTRUMENTATION_FUNCTION_ATTRIBUTES void InitStack();
+
+		uint32 StartClock = 0;
+		uint32 EndClock = 0;
+		uint32 Offset = 0;
+		double LastUsed = FPlatformTime::Seconds();
+		uint8  Buffer[2*1024*1024];
+		FHistoryChunk* Prev = nullptr;
+		FHistoryChunk* Next = nullptr;
+	};
+
+	struct FClockRange
+	{
+		FClock First;
+		FClock Last;
+	};
+
 	struct FAccessHistory
 	{
 		SAFE_OPERATOR_NEW_DELETE()
 
-		struct FEntry
+		INSTRUMENTATION_FUNCTION_ATTRIBUTES FAccessHistory()
 		{
-			void*              Pointer;
-			FMemoryAccess      Access;
-			FCallstackLocation Location;
-			UPTRINT            Data;
-		};
+			Tail = Head = new FHistoryChunk();
+			Tail->InitStack();
+			NumChunks++;
+		}
+
+		INSTRUMENTATION_FUNCTION_ATTRIBUTES ~FAccessHistory()
+		{
+			UE::TWriteScopeLock Scope(Lock);
+
+			while (Head)
+			{
+				FHistoryChunk* ToDelete = Head;
+				Head = Head->Next;
+				delete ToDelete;
+			}
+
+			NumChunks = 0;
+			Head = Tail = nullptr;
+
+			while (Spare)
+			{
+				FHistoryChunk* Next = Spare->Next;
+				delete Spare;
+				Spare = Next;
+			}
+		}
+
+		// Number of chunks that have been recycled
+		uint64 RecycleCount = 0;
+
+		// Total number of chunks currently allocated
+		int32  NumChunks = 0;
+		int32  NumSpares = 0;
+
+		// Used to dump information in case we can't find the memory access in the history
+		double LastRecycle = 0.0;
+
+		// Just used on the slow path between recycling and scanning.
+		UE::FRWSpinLock Lock;
+
+		// Can be used by other threads doing race reporting
+		FHistoryChunk* Head = nullptr;
+
+		// Only used by the owner thread
+		FHistoryChunk* Tail = nullptr;
+
+		// Only used by the owner thread to store unused buffers
+		FHistoryChunk* Spare = nullptr;
+
+		INSTRUMENTATION_FUNCTION_ATTRIBUTES bool HasTooManyChunks() const
+		{
+			// We need to have at least 2 chunks so that we always have one filled with data
+			// while we start filling the new one.
+			return NumChunks > 2 && NumChunks > GRaceDetectorHistoryLength;
+		}
+
+		INSTRUMENTATION_FUNCTION_ATTRIBUTES void TrimChunks()
+		{
+			if (HasTooManyChunks())
+			{
+				UE::TWriteScopeLock Scope(Lock);
 
-		static const uint32 NumEntry = 128 * 1024;
-		uint32 Index;
-		uint64 ResetCount;
-		double LastReset;
-		uint32 LastResetClock;
-		FEntry Entries[NumEntry];
+				while (HasTooManyChunks())
+				{
+					FHistoryChunk* Recycle = Head;
+					Head = Head->Next;
+					Head->Prev = nullptr;
+
+					RecycleCount++;
+					LastRecycle = Recycle->LastUsed;
+
+					Recycle->Next = Spare;
+					Spare = Recycle;
+
+					Spare->StartClock = 0;
+					Spare->EndClock = 0;
+					Spare->Offset = 0;
+					Spare->LastUsed = 0.0;
+
+					NumChunks--;
+					NumSpares++;
+				}
+			}
+		}
+
+		INSTRUMENTATION_FUNCTION_ATTRIBUTES FORCENOINLINE void EnsureNewChunk()
+		{
+			FHistoryChunk* NewChunk = nullptr;
+			if (Spare)
+			{
+				NumSpares--;
+				NewChunk = Spare;
+				Spare = Spare->Next;
+				NewChunk->Next = nullptr;
+			}
+			else
+			{
+				NewChunk = new FHistoryChunk();
+			}
+
+			NewChunk->InitStack();
+
+			Tail->LastUsed = FPlatformTime::Seconds();
+			NewChunk->Prev = Tail;
+			Tail->Next = NewChunk;
+			Tail = NewChunk;
+			NumChunks++;
+
+			TrimChunks();
+		}
+
+		template <typename EntryType, typename... ArgsType>
+		INSTRUMENTATION_FUNCTION_ATTRIBUTES FORCEINLINE void EmplaceEntry(ArgsType&&... Args)
+		{
+			if (UNLIKELY(Tail->Offset + sizeof(EntryType) > UE_ARRAY_COUNT(Tail->Buffer)))
+			{
+				EnsureNewChunk();
+			}
+
+			new (Tail->Buffer + Tail->Offset) EntryType(Forward<ArgsType>(Args)...);
+
+			FPlatformMisc::AsymmetricThreadFenceLight();
+			Tail->Offset += sizeof(EntryType);
+		}
+
+		INSTRUMENTATION_FUNCTION_ATTRIBUTES FORCEINLINE void AddFunctionEntry(void* InReturnAddress)
+		{
+			EmplaceEntry<FHistoryEntryFunctionEntry>(InReturnAddress);
+		}
+
+		INSTRUMENTATION_FUNCTION_ATTRIBUTES FORCEINLINE void AddMemoryAccess(void* InPointer, const FMemoryAccess& InAccess)
+		{
+			EmplaceEntry<FHistoryEntryAccess>(InPointer, InAccess.RawValue);
+			Tail->EndClock = InAccess.Clock;
+		}
+
+		INSTRUMENTATION_FUNCTION_ATTRIBUTES FORCEINLINE void AddFunctionExit()
+		{
+			EmplaceEntry<FHistoryEntryFunctionExit>();
+		}
+
+		INSTRUMENTATION_FUNCTION_ATTRIBUTES bool ResolveAccess(const void* InAlignedPointer, FMemoryAccess InAccess, FCallstackLocation& OutLocation, FClockRange& OutClockRange)
+		{
+			UE::TReadScopeLock Scope(Lock);
+
+			// Make sure we sync with the light fence.
+			FPlatformMisc::AsymmetricThreadFenceHeavy();
+
+			OutClockRange.First = MAX_uint32;
+			OutClockRange.Last  = 0;
+
+			for (FHistoryChunk* Chunk = Head; Chunk; Chunk = Chunk->Next)
+			{
+				if (Chunk->StartClock < OutClockRange.First)
+				{
+					OutClockRange.First = Chunk->StartClock;
+				}
+
+				if (Chunk->EndClock > OutClockRange.Last)
+				{
+					OutClockRange.Last = Chunk->EndClock;
+				}
+
+				// Do not bother searching a chunk that is outside the range we're looking for.
+				if (InAccess.Clock >= Chunk->StartClock && InAccess.Clock <= Chunk->EndClock)
+				{
+					int32 StackDepth = 0;
+					void* Stack[1024];
+					Stack[StackDepth] = 0;
+
+					int32 Offset = 0;
+					while (Offset < Chunk->Offset)
+					{
+						FHistoryEntryBase* Entry = (FHistoryEntryBase*)(Chunk->Buffer + Offset);
+						switch (Entry->Type)
+						{
+							case EHistoryEntryType::FunctionEntry:
+								check(StackDepth < UE_ARRAY_COUNT(Stack));
+								Stack[StackDepth++] = ((FHistoryEntryFunctionEntry*)Entry)->ReturnAddress;
+								Offset += sizeof(FHistoryEntryFunctionEntry);
+							break;
+							case EHistoryEntryType::MemoryAccess:
+							{
+								FHistoryEntryAccess* MemoryAccessEntry = (FHistoryEntryAccess*)Entry;
+								if (MemoryAccessEntry->Pointer == InAlignedPointer &&
+									MemoryAccessEntry->Access.RawValue == InAccess.RawValue)
+								{
+									OutLocation = FCallstackLocation::FromCallstackFrames(Stack, StackDepth, StackDepth + 1);
+									return true;
+								}
+								Offset += sizeof(FHistoryEntryAccess);
+							}
+							break;
+							case EHistoryEntryType::FunctionExit:
+								Stack[StackDepth--] = 0;
+								check(StackDepth >= 0);
+								Offset += sizeof(FHistoryEntryBase);
+							break;
+							default:
+								check(false);
+								// This should never happen, but if it does it is most likely a race condition
+								// so just restart the tracing from the beginning as a last resort.
+								StackDepth = 0;
+								Stack[StackDepth] = 0;
+								Offset = 0;
+							break;
+						}
+					}
+				}
+			}
+
+			// We should never end up empty handed unless the clock is from
+			// a long time ago and outside the history range.
+			// We have no choice but to include the equality comparison here
+			// since we have no idea if the chunk that got recycled had some
+			// entry for this clock.
+			check(InAccess.Clock <= OutClockRange.First);
+			if (OutClockRange.First == MAX_uint32)
+			{
+				OutClockRange.First = OutClockRange.Last;
+			}
+			return false;
+		}
+
+		INSTRUMENTATION_FUNCTION_ATTRIBUTES int32 GetOffset() const
+		{
+			return Tail->Offset;
+		}
 	};
 
 	// ------------------------------------------------------------------------------
@@ -229,11 +514,20 @@ namespace UE::Sanitizer::RaceDetector {
 	struct FContext : public TRefCountingMixin<FContext> {
 		SAFE_OPERATOR_NEW_DELETE()
 
-			INSTRUMENTATION_FUNCTION_ATTRIBUTES FContext(uint32 InThreadId)
+		INSTRUMENTATION_FUNCTION_ATTRIBUTES FContext(uint32 InThreadId)
 			: ThreadId(InThreadId)
 		{
 		}
 
+		INSTRUMENTATION_FUNCTION_ATTRIBUTES ~FContext()
+		{
+			if (AccessHistory)
+			{
+				delete AccessHistory;
+				AccessHistory = nullptr;
+			}
+		}
+
 		uint32 GlobalEpoch = 0;
 		// When we activate tracing, we need to recapture the current stack
 		uint32 StackEpoch = 0;
@@ -254,7 +548,7 @@ namespace UE::Sanitizer::RaceDetector {
 		// Each thread holds a bank of clocks to synchronize with every other context.
 		FClockBank ClockBank;
 		// We need to keep the callstack for each thread 
-		static constexpr SIZE_T MaxCallstackSize = 512;
+		static constexpr SIZE_T MaxCallstackSize = 1024; // Yes, callstacks can get huge in UE!
 		void* CurrentCallstack[MaxCallstackSize];
 		uint16 CurrentCallstackSize = 0;
 		// We use this to pass thread arguments to functions that dont have parameters (i.e. ExitThread).
@@ -277,8 +571,8 @@ namespace UE::Sanitizer::RaceDetector {
 			}
 		}
 
-		// Lazily created on first instrumented access
-		TUniquePtr<FAccessHistory> AccessHistory;
+		// Avoid using UniquePtr because it's instrumented and each access has a cost.
+		FAccessHistory* AccessHistory = nullptr;
 
 		INSTRUMENTATION_FUNCTION_ATTRIBUTES static bool IsValid(FContext* Context)
 		{
diff --git a/Engine/Source/Runtime/Core/Internal/Sanitizer/Types.h b/Engine/Source/Runtime/Core/Internal/Sanitizer/Types.h
--- a/Engine/Source/Runtime/Core/Internal/Sanitizer/Types.h
+++ b/Engine/Source/Runtime/Core/Internal/Sanitizer/Types.h
@@ -14,10 +14,20 @@
 #include "HAL/PlatformStackWalk.h"
 #include "Hash/CityHash.h"
 
+#define INSTRUMENTATION_ENUM_CLASS_FLAGS(Enum) \
+	FORCEINLINE INSTRUMENTATION_FUNCTION_ATTRIBUTES           Enum& operator|=(Enum& Lhs, Enum Rhs) { return Lhs = (Enum)((__underlying_type(Enum))Lhs | (__underlying_type(Enum))Rhs); } \
+	FORCEINLINE INSTRUMENTATION_FUNCTION_ATTRIBUTES           Enum& operator&=(Enum& Lhs, Enum Rhs) { return Lhs = (Enum)((__underlying_type(Enum))Lhs & (__underlying_type(Enum))Rhs); } \
+	FORCEINLINE INSTRUMENTATION_FUNCTION_ATTRIBUTES           Enum& operator^=(Enum& Lhs, Enum Rhs) { return Lhs = (Enum)((__underlying_type(Enum))Lhs ^ (__underlying_type(Enum))Rhs); } \
+	FORCEINLINE INSTRUMENTATION_FUNCTION_ATTRIBUTES constexpr Enum  operator| (Enum  Lhs, Enum Rhs) { return (Enum)((__underlying_type(Enum))Lhs | (__underlying_type(Enum))Rhs); } \
+	FORCEINLINE INSTRUMENTATION_FUNCTION_ATTRIBUTES constexpr Enum  operator& (Enum  Lhs, Enum Rhs) { return (Enum)((__underlying_type(Enum))Lhs & (__underlying_type(Enum))Rhs); } \
+	FORCEINLINE INSTRUMENTATION_FUNCTION_ATTRIBUTES constexpr Enum  operator^ (Enum  Lhs, Enum Rhs) { return (Enum)((__underlying_type(Enum))Lhs ^ (__underlying_type(Enum))Rhs); } \
+	FORCEINLINE INSTRUMENTATION_FUNCTION_ATTRIBUTES constexpr bool  operator! (Enum  E)             { return !(__underlying_type(Enum))E; } \
+	FORCEINLINE INSTRUMENTATION_FUNCTION_ATTRIBUTES constexpr Enum  operator~ (Enum  E)             { return (Enum)~(__underlying_type(Enum))E; }
+
 namespace UE::Sanitizer {
 
 	struct FLocation {
-		FLocation(const FString& ModuleName, const FString& FunctionName, const FString& Filename, uint32 Line, UPTRINT ProgramCounter)
+		INSTRUMENTATION_FUNCTION_ATTRIBUTES FLocation(const FString& ModuleName, const FString& FunctionName, const FString& Filename, uint32 Line, UPTRINT ProgramCounter)
 			: ModuleName(ModuleName), FunctionName(FunctionName), Filename(Filename), Line(Line), ProgramCounter(ProgramCounter)
 		{
 		}
@@ -28,21 +38,31 @@ namespace UE::Sanitizer {
 		const uint32  Line;
 		const UPTRINT ProgramCounter;
 
-		int32 GetAlignment() const
+		INSTRUMENTATION_FUNCTION_ATTRIBUTES int32 GetAlignment() const
 		{
+			if (Filename.Len() == 0)
+			{
+				return ModuleName.Len() - 3 /* Remove 1 space and () for line number */ ;
+			}
 			return Filename.Len() + GetDigits(Line);
 		}
 
-		FString ToString(int32 Alignment = 0) const
+		INSTRUMENTATION_FUNCTION_ATTRIBUTES FString ToString(int32 Alignment = 0) const
 		{
 			Alignment = FMath::Max(0, Alignment - GetAlignment());
 
+			if (Filename.Len() == 0)
+			{
+				// Format in a way that visual studio can understand so we can click the file and goto source.
+				return FString::Printf(TEXT("%s:%s %s 0x%p"), *ModuleName, FCString::Spc(Alignment), FunctionName.Len() == 0 ? TEXT("[Unknown Function]") : *FunctionName, (void*)ProgramCounter);
+			}
+
 			// Format in a way that visual studio can understand so we can click the file and goto source.
-			return FString::Printf(TEXT("%s (%d):%s %s 0x%p"), *Filename, Line, FCString::Spc(Alignment), *FunctionName, (void*)ProgramCounter);
+			return FString::Printf(TEXT("%s (%d):%s %s 0x%p"), *Filename, Line, FCString::Spc(Alignment), FunctionName.Len() == 0 ? TEXT("[Unknown Function]") : *FunctionName, (void*)ProgramCounter);
 		}
 
 	private:
-		static int32 GetDigits(int32 Value)
+		INSTRUMENTATION_FUNCTION_ATTRIBUTES static int32 GetDigits(int32 Value)
 		{
 			int32 Digits = 0;
 			do
@@ -58,9 +78,9 @@ namespace UE::Sanitizer {
 	struct FFullLocation {
 		TArray<FLocation> Locations;
 
-		int32 Num() const { return Locations.Num(); }
+		INSTRUMENTATION_FUNCTION_ATTRIBUTES int32 Num() const { return Locations.Num(); }
 
-		int32 GetAlignment() const
+		INSTRUMENTATION_FUNCTION_ATTRIBUTES int32 GetAlignment() const
 		{
 			int32 MaxAlign = 0;
 			for (const FLocation& Location : Locations)
@@ -71,7 +91,7 @@ namespace UE::Sanitizer {
 			return MaxAlign;
 		}
 
-		FString ToString(int32 Alignment = 0) const
+		INSTRUMENTATION_FUNCTION_ATTRIBUTES FString ToString(int32 Alignment = 0) const
 		{
 			if (Alignment == 0)
 			{
@@ -98,7 +118,7 @@ namespace UE::Sanitizer {
 	// ------------------------------------------------------------------------------
 	class FCallstackLocation {
 	private:
-		static constexpr uint32 NumCallstackFrames = 5;
+		static constexpr uint32 NumCallstackFrames = 1024;
 
 		static constexpr UPTRINT NoCallstackMarker = 0;
 		static constexpr UPTRINT OverflowCallstackMarker = 0x1;
@@ -106,12 +126,12 @@ namespace UE::Sanitizer {
 		static constexpr const TCHAR* OverflowString = TEXT("[Uncaptured frames, callstack too deep]");
 
 	public:
-		FCallstackLocation()
+		INSTRUMENTATION_FUNCTION_ATTRIBUTES FCallstackLocation()
 		{
 			Callstack[0] = NoCallstackMarker;
 		}
 
-		static FCallstackLocation FromCallstackFrames(const void* InCallstack, uint32 NumFrames, uint32 MaxNumFrames)
+		INSTRUMENTATION_FUNCTION_ATTRIBUTES static FCallstackLocation FromCallstackFrames(const void* InCallstack, uint32 NumFrames, uint32 MaxNumFrames)
 		{
 			FCallstackLocation Location;
 
@@ -137,7 +157,7 @@ namespace UE::Sanitizer {
 			return Location;
 		}
 
-		FFullLocation GetFullLocation() const
+		INSTRUMENTATION_FUNCTION_ATTRIBUTES FFullLocation GetFullLocation() const
 		{
 			FFullLocation Loc;
 			if (Callstack[0] == NoCallstackMarker)
@@ -171,7 +191,7 @@ namespace UE::Sanitizer {
 					// By using Frame - 1, we make sure that we actually point inside the chunk of assembly that represents the function call and it is enough
 					// to make symbol decode correct since we only need to point inside a valid range of a symbol entry, not directly at the start of it.
 					FPlatformStackWalk::EnumerateSymbolInfosForProgramCounter(Frame - 1, bIncludeInlineFrames, 
-						[&Loc, Frame](FProgramCounterSymbolInfo& SymbolInfo)
+						[&Loc, Frame](FProgramCounterSymbolInfo& SymbolInfo) INSTRUMENTATION_FUNCTION_ATTRIBUTES
 						{
 							Loc.Locations.Emplace(FString(SymbolInfo.ModuleName), FString(SymbolInfo.FunctionName), FString(SymbolInfo.Filename), SymbolInfo.LineNumber, Frame);
 						}
@@ -182,7 +202,7 @@ namespace UE::Sanitizer {
 			return Loc;
 		}
 
-		uint32 GetHash() const
+		INSTRUMENTATION_FUNCTION_ATTRIBUTES uint32 GetHash() const
 		{
 			if (Callstack[0] == NoCallstackMarker)
 			{
@@ -192,7 +212,7 @@ namespace UE::Sanitizer {
 			return CityHash32((const char*)Callstack, sizeof(UPTRINT) * NumCallstackFrames);
 		}
 
-		UPTRINT GetLastFrame() const
+		INSTRUMENTATION_FUNCTION_ATTRIBUTES UPTRINT GetLastFrame() const
 		{
 			for (int32 Idx = NumCallstackFrames - 1; Idx >= 0; --Idx)
 			{
diff --git a/Engine/Source/Runtime/Core/Private/Sanitizer/RaceDetector.cpp b/Engine/Source/Runtime/Core/Private/Sanitizer/RaceDetector.cpp
--- a/Engine/Source/Runtime/Core/Private/Sanitizer/RaceDetector.cpp
+++ b/Engine/Source/Runtime/Core/Private/Sanitizer/RaceDetector.cpp
@@ -101,6 +101,19 @@ static FAutoConsoleVariableRef CVarRaceDetectorBreakOnRace(
 	TEXT("	4 - Break only when the second callstack is missing\n"),
 	ECVF_Default);
 
+static int32 GRaceDetectorHistoryLength = 4;
+static FAutoConsoleVariableRef CVarRaceDetectorHistoryLength(
+	TEXT("r.RaceDetector.HistoryLength"),
+	GRaceDetectorHistoryLength,
+	TEXT("Represents the number of history blocks each thread is keeping to resolve callstacks of race conditions\n")
+	TEXT("Can be increased to improve detection rate for races that are very far apart\n")
+	TEXT("Trying to set this below a minimum of 2 blocks won't have any effect\n"),
+	FConsoleVariableDelegate::CreateLambda([](IConsoleVariable*) INSTRUMENTATION_FUNCTION_ATTRIBUTES
+	{
+		GRaceDetectorHistoryLength = FMath::Max(2, GRaceDetectorHistoryLength);
+	}),
+	ECVF_Default);
+
 using namespace UE::Instrumentation;
 
 namespace UE::Sanitizer::RaceDetector {
@@ -114,6 +127,7 @@ namespace UE::Sanitizer::RaceDetector {
 	volatile uint32 GlobalEpoch = 0;
 	volatile int64 FSyncObject::SyncObjectCount = 0;
 	volatile FContextId CurrentContextId = 1;
+	volatile int64 HistoryChunkCount = 0;
 
 	UE::FSpinLock RaceReportsLock;
 	TArray<FString> RaceReports;
@@ -158,6 +172,11 @@ namespace UE::Sanitizer::RaceDetector {
 		return FSyncObject::GetSyncObjectCount() * Align(sizeof(FSyncObject), Platform::GetPageSize());
 	}
 
+	INSTRUMENTATION_FUNCTION_ATTRIBUTES uint64 GetHistoryChunkMemoryUsage()
+	{
+		return HistoryChunkCount * sizeof(FHistoryChunk);
+	}
+
 	// ------------------------------------------------------------------------------
 	// Callstack management.
 	// ------------------------------------------------------------------------------
@@ -170,6 +189,14 @@ namespace UE::Sanitizer::RaceDetector {
 
 	FORCEINLINE INSTRUMENTATION_FUNCTION_ATTRIBUTES void AddCallstackFrame(FContext& Context, void* ReturnAddress)
 	{
+		// History tracing needs to come first since initializing a new history chunk
+		// needs to copy the current stack and we don't want the current stack
+		// to contain the frame we're going to add via this function.
+		if (Context.AccessHistory)
+		{
+			Context.AccessHistory->AddFunctionEntry(ReturnAddress);
+		}
+
 		// Don't capture extremely deep callstacks.
 		if (Context.CurrentCallstackSize < FContext::MaxCallstackSize)
 		{
@@ -181,12 +208,20 @@ namespace UE::Sanitizer::RaceDetector {
 
 	FORCEINLINE INSTRUMENTATION_FUNCTION_ATTRIBUTES void RemoveCallstackFrame(FContext& Context)
 	{
+		// History tracing needs to come first since initializing a new history chunk
+		// needs to copy the current stack before applying the exit.
+		if (Context.AccessHistory)
+		{
+			Context.AccessHistory->AddFunctionExit();
+		}
+
 		Context.CurrentCallstackSize--;
 	}
 
 	INSTRUMENTATION_FUNCTION_ATTRIBUTES void RegisterContext(FContext& Context)
 	{
-		Context.AccessHistory = MakeUnique<FAccessHistory>();
+		check(Context.AccessHistory == nullptr);
+		Context.AccessHistory = new FAccessHistory();
 
 		UE::TScopeLock Scope(ContextMappingLock);
 		Context.GlobalEpoch = GlobalEpoch;
@@ -235,14 +270,6 @@ namespace UE::Sanitizer::RaceDetector {
 		check(Context.CurrentClock() > 0);
 	}
 
-	FORCEINLINE INSTRUMENTATION_FUNCTION_ATTRIBUTES void EnsureContextRegistered(FContext& Context)
-	{
-		if (Context.ContextId == 0)
-		{
-			RegisterContext(Context);
-		}
-	}
-
 	INSTRUMENTATION_FUNCTION_ATTRIBUTES void ReleaseContext(FContext* Context)
 	{
 		if (Context != nullptr)
@@ -437,8 +464,32 @@ namespace UE::Sanitizer::RaceDetector {
 		RaceReports.Emplace(MoveTemp(Msg));
 	}
 
+	INSTRUMENTATION_FUNCTION_ATTRIBUTES FHistoryChunk::FHistoryChunk()
+	{
+		FPlatformAtomics::InterlockedIncrement(&HistoryChunkCount);
+	}
+
+	INSTRUMENTATION_FUNCTION_ATTRIBUTES void FHistoryChunk::InitStack()
+	{
+		FContext* ContextPtr = (FContext*)GetTlsValueFast(ContextTLSIndex);
+		if (FContext::IsValid(ContextPtr))
+		{
+			StartClock = ContextPtr->CurrentClock();
+			for (uint16 Index = 0; Index < ContextPtr->CurrentCallstackSize; ++Index)
+			{
+				new (Buffer + Offset) FHistoryEntryFunctionEntry(ContextPtr->CurrentCallstack[Index]);
+				Offset += sizeof(FHistoryEntryFunctionEntry);
+			}
+		}
+	}
+
+	INSTRUMENTATION_FUNCTION_ATTRIBUTES FHistoryChunk::~FHistoryChunk()
+	{
+		FPlatformAtomics::InterlockedDecrement(&HistoryChunkCount);
+	}
+
 	TOptional<TRaceCallbackFn> RaceCallbackFn;
-	INSTRUMENTATION_FUNCTION_ATTRIBUTES void ReportRace(FContext& Context, void* Pointer, const FMemoryAccess& CurrentAccess, const FMemoryAccess& RacingAccess, const FCallstackLocation& FirstLocation)
+	INSTRUMENTATION_FUNCTION_ATTRIBUTES void ReportRace(FContext& Context, void* Pointer, const FMemoryAccess& CurrentAccess, const FMemoryAccess& RacingAccess)
 	{
 		TRefCountPtr<FContext> RacingContext = GetContextById(RacingAccess.ContextId);
 		// The thread might have been destroyed
@@ -447,26 +498,17 @@ namespace UE::Sanitizer::RaceDetector {
 			return;
 		}
 
+		FCallstackLocation FirstLocation = GetCurrentCallstackLocation();
+
 		const uint32 FirstThreadId = Context.ThreadId;
 		const uint32 SecondThreadId = RacingContext->ThreadId;
 
 		// Dumb scan of the history to find the matching race so we can get its location.
 		FCallstackLocation SecondLocation;
-		FAccessHistory& SecondAccessHistory = *RacingContext->AccessHistory;
-		bool bFoundLocation = false;
-
 		const void* AlignedPointer = GetAlignedPointer(Pointer);
-		for (FAccessHistory::FEntry& Entry : SecondAccessHistory.Entries)
-		{
-			if (Entry.Pointer == AlignedPointer &&
-				Entry.Access.RawValue == RacingAccess.RawValue
-				)
-			{
-				SecondLocation = Entry.Location;
-				bFoundLocation = true;
-				break;
-			}
-		}
+		FAccessHistory& SecondAccessHistory = *RacingContext->AccessHistory;
+		FClockRange HistoryRange;
+		bool bFoundLocation = SecondAccessHistory.ResolveAccess(AlignedPointer, RacingAccess, SecondLocation, HistoryRange);
 
 		// Resolving symbols can cause us to enter a wait that could try to start a new thread
 		// and wait for it to be started. We can't allow that since we could deadlock
@@ -523,19 +565,36 @@ namespace UE::Sanitizer::RaceDetector {
 			FFullLocation FirstFullLocation = FirstLocation.GetFullLocation();
 			FFullLocation SecondFullLocation = SecondLocation.GetFullLocation();
 
-			const FString& FirstThreadName = FThreadManager::Get().GetThreadName(FirstThreadId);
-			const FString& SecondThreadName = FThreadManager::Get().GetThreadName(SecondThreadId);
+			FString FirstThreadName;
+			FString SecondThreadName;
+			bool bNeedResolve = true;
+			FThreadManager::Get().ForEachThread(
+				[&](uint32, FRunnableThread*)
+				{
+					// GetThreadName is not thread-safe when a thread is exiting since it just sends us a reference to the string
+					// which can be deleted once the lock is not held. So resolve the names while inside ForEachThread
+					// since it maintains the lock while we copy the thread names. This is just a workaround until GetThreadName 
+					// can be fixed.
+					if (bNeedResolve)
+					{
+						FirstThreadName = FThreadManager::Get().GetThreadName(FirstThreadId);
+						SecondThreadName = FThreadManager::Get().GetThreadName(SecondThreadId);
+						bNeedResolve = false;
+					}
+				}
+			);
 
 			FString NotFoundMessage;
 			if (!bFoundLocation)
 			{
 				NotFoundMessage = FString::Printf(
 					TEXT("Location not found in access history.\n")
-					TEXT("Number of history rollovers: %llu\n")
-					TEXT("Last history rollover: %.02f seconds ago at clock %d"),
-					SecondAccessHistory.ResetCount,
-					FPlatformTime::Seconds() - SecondAccessHistory.LastReset,
-					SecondAccessHistory.LastResetClock
+					TEXT("Number of history blocks recycled: %llu\n")
+					TEXT("History range: clock %d to %d\n")
+					TEXT("Last recycle: %.02f seconds ago"),
+					SecondAccessHistory.RecycleCount,
+					HistoryRange.First, HistoryRange.Last,
+					FPlatformTime::Seconds() - SecondAccessHistory.LastRecycle
 				);
 			}
 			
@@ -614,26 +673,35 @@ namespace UE::Sanitizer::RaceDetector {
 		RaceCallbackFn.Reset();
 	}
 
-	INSTRUMENTATION_FUNCTION_ATTRIBUTES void InstrumentAccess(FContext& Context, void* Ptr, uint8 Size, EMemoryAccessType AccessType, FAtomicMemoryOrder Order, void* ReturnAddress, const TCHAR* OpName, bool bHasSyncObjectAttached)
+	INSTRUMENTATION_FUNCTION_ATTRIBUTES FORCEINLINE void InstrumentAccess(FContext& Context, void* Ptr, uint8 Size, EMemoryAccessType AccessType, FAtomicMemoryOrder Order, void* ReturnAddress, const TCHAR* OpName, bool bHasSyncObjectAttached, bool& bHasAddedStack)
 	{
 #if WITH_RACEDETECTOR_CHECK
 		// Validate size and alignment since it should have been done before entering this function.
 		check(Size <= 8 && GetShadowMemoryAddress((UPTRINT)Ptr) == GetShadowMemoryAddress((UPTRINT)Ptr + Size - 1));
 #endif
+		// GetShadowMemory will contain the cost of committing pages to shadow memory
+		// The cost of this is higher when run under the debugger because we're not
+		// using vectored exception handler in this case due to abysmal perf caused
+		// by the kernel sending exception events to the attached debugger.
+		FShadowMemory* Shadow = GetShadowMemory((UPTRINT)Ptr);
+
+		// Use short variable name here to improve readability further down.
+		FMemoryAccess* S = Shadow->Accesses;
+
+		// Try to hide cache miss on shadow memory behind some computational stuff.
+		FPlatformMisc::PrefetchBlock(S, sizeof(Shadow->Accesses));
+
 		// Grab a ContextId on first instrumented access to avoid wasting our precious 256 available context
 		// on thirdparty and OS threads that never enter into instrumented code.
-		EnsureContextRegistered(Context);
+		if (Context.ContextId == 0)
+		{
+			RegisterContext(Context);
+		}
 
 		const uint32 Epoch = GlobalEpoch;
 		const uint8 Offset = GetOffsetInBytes(Ptr);
 		const FMemoryAccess CurrentAccess(Context.ContextId, Context.CurrentClock(), Offset, Size, AccessType, bHasSyncObjectAttached);
 
-		// GetShadowMemory will contain the cost of committing pages to shadow memory
-		FShadowMemory* Shadow = GetShadowMemory((UPTRINT)Ptr);
-
-		// Use short variable name here to improve readability further down.
-		FMemoryAccess* S = Shadow->Accesses;
-
 		// This is very practical to scope small section of code that we want to understand
 		// exactly what's going on inside.
 		if (UNLIKELY(Context.DetailedLogDepth || bDetailedLogGlobal))
@@ -664,6 +732,8 @@ namespace UE::Sanitizer::RaceDetector {
 			}
 		}
 
+		// The cost of this is going to be higher when we're not running under the debugger
+		// because the vectored page fault handler is called as part of this access on new shadow memory.
 		if (S[0].RawValue == CurrentAccess.RawValue ||
 			S[1].RawValue == CurrentAccess.RawValue ||
 			S[2].RawValue == CurrentAccess.RawValue ||
@@ -673,23 +743,15 @@ namespace UE::Sanitizer::RaceDetector {
 			return;
 		}
 
+		if (!bHasAddedStack)
+		{
+			AddCallstackFrame(Context, ReturnAddress);
+			bHasAddedStack = true;
+		}
+
 		// Write to our history first so that any race detected from the shadow
 		// can be resolved.
-		FAccessHistory& AccessHistory = *Context.AccessHistory;
-		const uint32 HistoryIndex = AccessHistory.Index++;
-		if (UNLIKELY(AccessHistory.Index >= AccessHistory.NumEntry))
-		{
-			AccessHistory.ResetCount++;
-			AccessHistory.Index = 0;
-			AccessHistory.LastReset = FPlatformTime::Seconds();
-			AccessHistory.LastResetClock = CurrentAccess.Clock;
-		}
-		FAccessHistory::FEntry& HistoryEntry = AccessHistory.Entries[HistoryIndex];
-		HistoryEntry.Access.RawValue = CurrentAccess.RawValue;
-		// This is expensive. We need to optimize history so that we don't need to keep full callstack location
-		// and instead replay the event to discover the callstack along the way.
-		HistoryEntry.Location = GetCurrentCallstackLocation();
-		HistoryEntry.Pointer = GetAlignedPointer(Ptr);
+		Context.AccessHistory->AddMemoryAccess(GetAlignedPointer(Ptr), CurrentAccess);
 
 		bool bSaveNeeded = true;
 		FMemoryAccess PreviousAccess;
@@ -771,7 +833,7 @@ namespace UE::Sanitizer::RaceDetector {
 			// Just ignore any race we have found if anything seems incoherent.
 			if (bWonReportingRace && bRaceDetectorActive && Epoch == GlobalEpoch && !bIsResettingShadow)
 			{
-				ReportRace(Context, Ptr, CurrentAccess, PreviousAccess, GetCurrentCallstackLocation());
+				ReportRace(Context, Ptr, CurrentAccess, PreviousAccess);
 			}
 
 			// Just return now as we're in the race reporting case where we don't need to save our access anymore.
@@ -782,21 +844,21 @@ namespace UE::Sanitizer::RaceDetector {
 		{
 			// We haven't saved our access yet.
 			// Use history index as a 'random' position to avoid expensive computations.
-			Shadow->Accesses[HistoryIndex & 3].RawValue = CurrentAccess.RawValue;
+			Shadow->Accesses[Context.AccessHistory->GetOffset() & 3].RawValue = CurrentAccess.RawValue;
 		}
 	}
 
-	INSTRUMENTATION_FUNCTION_ATTRIBUTES FORCEINLINE void InstrumentAccess(FContext& Context, void* Ptr, uint8 Size, EMemoryAccessType AccessType, FAtomicMemoryOrder Order, void* ReturnAddress, const TCHAR* OpName)
+	INSTRUMENTATION_FUNCTION_ATTRIBUTES FORCEINLINE void InstrumentAccess(FContext& Context, void* Ptr, uint8 Size, EMemoryAccessType AccessType, FAtomicMemoryOrder Order, void* ReturnAddress, const TCHAR* OpName, bool& bHasAddedCallstackFrame)
 	{
 		const bool bHasSyncObjectAttached = !IsAtomicOrderRelaxed(Order);
-		InstrumentAccess(Context, Ptr, Size, AccessType, Order, ReturnAddress, OpName, bHasSyncObjectAttached);
+		InstrumentAccess(Context, Ptr, Size, AccessType, Order, ReturnAddress, OpName, bHasSyncObjectAttached, bHasAddedCallstackFrame);
 	}
 
 	// Handles alignment and sizes bigger than a shadow cell
-	INSTRUMENTATION_FUNCTION_ATTRIBUTES void InstrumentAccessPreamble(FContext& Context, UPTRINT Ptr, uint32 Size, EMemoryAccessType AccessType, FAtomicMemoryOrder Order, void* ReturnAddress, const TCHAR* OpName)
+	INSTRUMENTATION_FUNCTION_ATTRIBUTES FORCEINLINE void InstrumentAccessPreamble(FContext& Context, UPTRINT Ptr, uint32 Size, EMemoryAccessType AccessType, FAtomicMemoryOrder Order, void* ReturnAddress, const TCHAR* OpName, bool& bHasAddedCallstackFrame)
 	{
 		UPTRINT AdjustedSize = FMath::Min((UPTRINT)Size, GetAlignedAddress(Ptr) + 8 - Ptr);
-		InstrumentAccess(Context, (void*)Ptr, static_cast<uint8>(AdjustedSize), AccessType, Order, ReturnAddress, OpName);
+		InstrumentAccess(Context, (void*)Ptr, static_cast<uint8>(AdjustedSize), AccessType, Order, ReturnAddress, OpName, bHasAddedCallstackFrame);
 		Ptr += AdjustedSize;
 		Size -= AdjustedSize;
 
@@ -804,24 +866,24 @@ namespace UE::Sanitizer::RaceDetector {
 		{
 			while (Size >= 8)
 			{
-				InstrumentAccess(Context, (void*)Ptr, static_cast<uint8>(8), AccessType, Order, ReturnAddress, OpName);
+				InstrumentAccess(Context, (void*)Ptr, static_cast<uint8>(8), AccessType, Order, ReturnAddress, OpName, bHasAddedCallstackFrame);
 				Ptr += 8;
 				Size -= 8;
 			}
 
 			if (Size)
 			{
-				InstrumentAccess(Context, (void*)Ptr, static_cast<uint8>(Size), AccessType, Order, ReturnAddress, OpName);
+				InstrumentAccess(Context, (void*)Ptr, static_cast<uint8>(Size), AccessType, Order, ReturnAddress, OpName, bHasAddedCallstackFrame);
 			}
 		}
 	}
 
-	INSTRUMENTATION_FUNCTION_ATTRIBUTES void InstrumentAtomicAccess(FContext& Context, void* AtomicAddr, uint8 Size, EMemoryAccessType AccessType, FAtomicMemoryOrder Order, const TFunctionRef<void()>& AtomicOp, void* ReturnAddress, const TCHAR* OpName)
+	INSTRUMENTATION_FUNCTION_ATTRIBUTES void InstrumentAtomicAccess(FContext& Context, void* AtomicAddr, uint8 Size, EMemoryAccessType AccessType, FAtomicMemoryOrder Order, const TFunctionRef<void()>& AtomicOp, void* ReturnAddress, const TCHAR* OpName, bool& bHasAddedCallstackFrame)
 	{
 #if WITH_RACEDETECTOR_CHECK
 		check(FMath::Min((UPTRINT)Size, GetAlignedAddress((UPTRINT)AtomicAddr) + 8 - (UPTRINT)AtomicAddr) == (UPTRINT)Size);
 #endif
-		InstrumentAccess(Context, AtomicAddr, Size, AccessType, Order, ReturnAddress, OpName);
+		InstrumentAccess(Context, AtomicAddr, Size, AccessType, Order, ReturnAddress, OpName, bHasAddedCallstackFrame);
 		if (IsAtomicOrderRelaxed(Order))
 		{
 			AtomicOp();
@@ -850,7 +912,7 @@ namespace UE::Sanitizer::RaceDetector {
 		Context.IncrementClock();
 	}
 
-	INSTRUMENTATION_FUNCTION_ATTRIBUTES void InstrumentAtomicAccess(FContext& Context, void* AtomicAddr, uint8 Size, EMemoryAccessType AccessType, FAtomicMemoryOrder SuccessOrder, FAtomicMemoryOrder FailureOrder, const TFunctionRef<bool()>& AtomicOp, void* ReturnAddress, const TCHAR* OpName)
+	INSTRUMENTATION_FUNCTION_ATTRIBUTES void InstrumentAtomicAccess(FContext& Context, void* AtomicAddr, uint8 Size, EMemoryAccessType AccessType, FAtomicMemoryOrder SuccessOrder, FAtomicMemoryOrder FailureOrder, const TFunctionRef<bool()>& AtomicOp, void* ReturnAddress, const TCHAR* OpName, bool& bHasAddedCallstackFrame)
 	{
 #if WITH_RACEDETECTOR_CHECK
 		check(FMath::Min((UPTRINT)Size, GetAlignedAddress((UPTRINT)AtomicAddr) + 8 - (UPTRINT)AtomicAddr) == (UPTRINT)Size);
@@ -858,17 +920,17 @@ namespace UE::Sanitizer::RaceDetector {
 		if (IsAtomicOrderRelaxed(SuccessOrder) && IsAtomicOrderRelaxed(FailureOrder))
 		{
 			// Both orders are relaxed, forward either one of them.
-			InstrumentAccess(Context, AtomicAddr, Size, AccessType, SuccessOrder, ReturnAddress, OpName);
+			InstrumentAccess(Context, AtomicAddr, Size, AccessType, SuccessOrder, ReturnAddress, OpName, bHasAddedCallstackFrame);
 			
 			AtomicOp();
 		}
 		else
 		{
 			GetSyncObject(AtomicAddr)->SyncWithFailureSupport(Context, AtomicOp, AccessType, SuccessOrder, FailureOrder, ReturnAddress, AtomicAddr, OpName,
-				[&Context, AtomicAddr, Size, AccessType, ReturnAddress, &SuccessOrder, OpName](FAtomicMemoryOrder ActualOrder)
+				[&Context, AtomicAddr, Size, AccessType, ReturnAddress, &SuccessOrder, OpName, &bHasAddedCallstackFrame](FAtomicMemoryOrder ActualOrder) INSTRUMENTATION_FUNCTION_ATTRIBUTES
 				{
 					const bool bHasSyncObjectAttached = true;
-					InstrumentAccess(Context, AtomicAddr, Size, AccessType, ActualOrder, ReturnAddress, OpName, bHasSyncObjectAttached);
+					InstrumentAccess(Context, AtomicAddr, Size, AccessType, ActualOrder, ReturnAddress, OpName, bHasSyncObjectAttached, bHasAddedCallstackFrame);
 				}
 			);
 		}
@@ -1025,13 +1087,15 @@ namespace UE::Sanitizer::RaceDetector {
 
 				double ShadowMemoryUsage = double(Platform::GetShadowMemoryUsage()) / (1024 *1024 *1024);
 				double SyncObjectMemoryUsage = double(GetSyncObjectMemoryUsage()) / (1024 * 1024 * 1024);
+				double HistoryChunkMemoryUsage = double(GetHistoryChunkMemoryUsage()) / (1024 * 1024 * 1024);
 
 				if (++MemoryUsageStatIterations > 10)
 				{
 					FPlatformMisc::LowLevelOutputDebugStringf(
-						TEXT("Sanitizer Memory Usage (Shadow : %.04f GB, SyncObjects : %.04f GB)\n"),
+						TEXT("Sanitizer Memory Usage (Shadow : %.04f GB, SyncObjects : %.04f GB, History : %.04f GB)\n"),
 						ShadowMemoryUsage,
-						SyncObjectMemoryUsage
+						SyncObjectMemoryUsage,
+						HistoryChunkMemoryUsage
 					);
 					MemoryUsageStatIterations = 0;
 				}
@@ -1289,9 +1353,11 @@ namespace UE::Sanitizer::RaceDetector {
 
 		InitMemoryAllocator();
 		
-		HookAlwaysOnInstrumentation();
-
+		// Set as initialized now otherwise we can end up with messed up callstack state
+		// since we start using incrementation scope inside HookAlwaysOnInstrumentation.
 		bRuntimeInitialized = true;
+
+		HookAlwaysOnInstrumentation();
 		
 		const TCHAR* CommandLine = Platform::GetCommandLine();
 		if (FCString::Stristr(CommandLine, TEXT("-racedetector")))
@@ -1334,6 +1400,9 @@ namespace UE::Sanitizer::RaceDetector {
 			Context->InstrumentationDepth = 0;
 
 			ReleaseContext(Context);
+
+			// Release the refcount owned by the thread.
+			Context->Release();
 		}
 	}
 
@@ -1468,11 +1537,14 @@ extern "C" {
 	{
 		BEGIN_HANDLE_INSTRUMENTATION(return);
 
-		InstrumentAccess(Context, addr, sizeof(UPTRINT), ACCESS_TYPE_ATOMIC_READ, FAtomicMemoryOrder::MEMORY_ORDER_RELEASE, _ReturnAddress(), TEXT("AnnotateHappensBefore"));
+		AddCallstackFrame(Context, _ReturnAddress());
+		bool bHasAddedCallstackFrame = true;
+		InstrumentAccess(Context, addr, sizeof(UPTRINT), ACCESS_TYPE_ATOMIC_READ, FAtomicMemoryOrder::MEMORY_ORDER_RELEASE, _ReturnAddress(), TEXT("AnnotateHappensBefore"), bHasAddedCallstackFrame);
 
 		TRefCountPtr<FSyncObject> Atomic = GetSyncObject(addr);
 		Atomic->SyncRelease(Context, []() {}, _ReturnAddress(), addr, TEXT("AnnotateHappensBefore"));
 		Context.IncrementClock();
+		RemoveCallstackFrame(Context);
 
 		FINISH_HANDLE_INSTRUMENTATION();
 	}
@@ -1481,11 +1553,14 @@ extern "C" {
 	{
 		BEGIN_HANDLE_INSTRUMENTATION(return);
 
-		InstrumentAccess(Context, addr, sizeof(UPTRINT), ACCESS_TYPE_ATOMIC_READ, FAtomicMemoryOrder::MEMORY_ORDER_ACQUIRE, _ReturnAddress(), TEXT("AnnotateHappensBefore"));
+		AddCallstackFrame(Context, _ReturnAddress());
+		bool bHasAddedCallstackFrame = true;
+		InstrumentAccess(Context, addr, sizeof(UPTRINT), ACCESS_TYPE_ATOMIC_READ, FAtomicMemoryOrder::MEMORY_ORDER_ACQUIRE, _ReturnAddress(), TEXT("AnnotateHappensBefore"), bHasAddedCallstackFrame);
 
 		TRefCountPtr<FSyncObject> Atomic = GetSyncObject(addr);
 		Atomic->SyncAcquire(Context, []() {}, _ReturnAddress(), addr, TEXT("AnnotateHappensAfter"));
 		Context.IncrementClock();
+		RemoveCallstackFrame(Context);
 
 		FINISH_HANDLE_INSTRUMENTATION();
 	}
@@ -1512,9 +1587,12 @@ extern "C" {
 	{
 		BEGIN_HANDLE_INSTRUMENTATION(return);
 
-		AddCallstackFrame(Context, _ReturnAddress());
-		InstrumentAccessPreamble(Context, Address, Size, EMemoryAccessType::ACCESS_TYPE_WRITE, FAtomicMemoryOrder::MEMORY_ORDER_RELAXED, _ReturnAddress(), TEXT("Store"));
-		RemoveCallstackFrame(Context);
+		bool bHasAddedCallstackFrame = false;
+		InstrumentAccessPreamble(Context, Address, Size, EMemoryAccessType::ACCESS_TYPE_WRITE, FAtomicMemoryOrder::MEMORY_ORDER_RELAXED, _ReturnAddress(), TEXT("Store"), bHasAddedCallstackFrame);
+		if (bHasAddedCallstackFrame)
+		{
+			RemoveCallstackFrame(Context);
+		}
 
 		FINISH_HANDLE_INSTRUMENTATION();
 	}
@@ -1523,9 +1601,12 @@ extern "C" {
 	{
 		BEGIN_HANDLE_INSTRUMENTATION(return);
 
-		AddCallstackFrame(Context, _ReturnAddress());
-		InstrumentAccessPreamble(Context, Address, Size, EMemoryAccessType::ACCESS_TYPE_READ, FAtomicMemoryOrder::MEMORY_ORDER_RELAXED, _ReturnAddress(), TEXT("Load"));
-		RemoveCallstackFrame(Context);
+		bool bHasAddedCallstackFrame = false;
+		InstrumentAccessPreamble(Context, Address, Size, EMemoryAccessType::ACCESS_TYPE_READ, FAtomicMemoryOrder::MEMORY_ORDER_RELAXED, _ReturnAddress(), TEXT("Load"), bHasAddedCallstackFrame);
+		if (bHasAddedCallstackFrame)
+		{
+			RemoveCallstackFrame(Context);
+		}
 
 		FINISH_HANDLE_INSTRUMENTATION();
 	}
@@ -1539,9 +1620,12 @@ extern "C" {
 		// See FRaceDetectorVirtualPointerBenignTest in RaceDetectorTests.cpp for an example.
 		if (*Address != Value)
 		{
-			AddCallstackFrame(Context, _ReturnAddress());
-			InstrumentAccessPreamble(Context, (UPTRINT)Address, sizeof(void*), EMemoryAccessType::ACCESS_TYPE_WRITE | EMemoryAccessType::ACCESS_TYPE_VPTR, FAtomicMemoryOrder::MEMORY_ORDER_RELAXED, _ReturnAddress(), TEXT("VPtr Store"));
-			RemoveCallstackFrame(Context);
+			bool bHasAddedCallstackFrame = false;
+			InstrumentAccessPreamble(Context, (UPTRINT)Address, sizeof(void*), EMemoryAccessType::ACCESS_TYPE_WRITE | EMemoryAccessType::ACCESS_TYPE_VPTR, FAtomicMemoryOrder::MEMORY_ORDER_RELAXED, _ReturnAddress(), TEXT("VPtr Store"), bHasAddedCallstackFrame);
+			if (bHasAddedCallstackFrame)
+			{
+				RemoveCallstackFrame(Context);
+			}
 		}
 
 		FINISH_HANDLE_INSTRUMENTATION();
@@ -1551,9 +1635,12 @@ extern "C" {
 	{
 		BEGIN_HANDLE_INSTRUMENTATION(return);
 
-		AddCallstackFrame(Context, _ReturnAddress());
-		InstrumentAccessPreamble(Context, (UPTRINT)Address, sizeof(void*), EMemoryAccessType::ACCESS_TYPE_READ | EMemoryAccessType::ACCESS_TYPE_VPTR, FAtomicMemoryOrder::MEMORY_ORDER_RELAXED, _ReturnAddress(), TEXT("VPtr Load"));
-		RemoveCallstackFrame(Context);
+		bool bHasAddedCallstackFrame = false;
+		InstrumentAccessPreamble(Context, (UPTRINT)Address, sizeof(void*), EMemoryAccessType::ACCESS_TYPE_READ | EMemoryAccessType::ACCESS_TYPE_VPTR, FAtomicMemoryOrder::MEMORY_ORDER_RELAXED, _ReturnAddress(), TEXT("VPtr Load"), bHasAddedCallstackFrame);
+		if (bHasAddedCallstackFrame)
+		{
+			RemoveCallstackFrame(Context);
+		}
 
 		FINISH_HANDLE_INSTRUMENTATION();
 	}
@@ -1567,9 +1654,12 @@ extern "C" {
 
 		BEGIN_HANDLE_INSTRUMENTATION(return);
 		
-		AddCallstackFrame(Context, _ReturnAddress());
-		InstrumentAccessPreamble(Context, Address, Size, EMemoryAccessType::ACCESS_TYPE_WRITE, FAtomicMemoryOrder::MEMORY_ORDER_RELAXED, _ReturnAddress(), TEXT("StoreRange"));
-		RemoveCallstackFrame(Context);
+		bool bHasAddedCallstackFrame = false;
+		InstrumentAccessPreamble(Context, Address, Size, EMemoryAccessType::ACCESS_TYPE_WRITE, FAtomicMemoryOrder::MEMORY_ORDER_RELAXED, _ReturnAddress(), TEXT("StoreRange"), bHasAddedCallstackFrame);
+		if (bHasAddedCallstackFrame)
+		{
+			RemoveCallstackFrame(Context);
+		}
 
 		FINISH_HANDLE_INSTRUMENTATION();
 	}
@@ -1583,9 +1673,12 @@ extern "C" {
 
 		BEGIN_HANDLE_INSTRUMENTATION(return);
 
-		AddCallstackFrame(Context, _ReturnAddress());
-		InstrumentAccessPreamble(Context, Address, Size, EMemoryAccessType::ACCESS_TYPE_READ, FAtomicMemoryOrder::MEMORY_ORDER_RELAXED, _ReturnAddress(), TEXT("LoadRange"));
-		RemoveCallstackFrame(Context);
+		bool bHasAddedCallstackFrame = false;
+		InstrumentAccessPreamble(Context, Address, Size, EMemoryAccessType::ACCESS_TYPE_READ, FAtomicMemoryOrder::MEMORY_ORDER_RELAXED, _ReturnAddress(), TEXT("LoadRange"), bHasAddedCallstackFrame);
+		if (bHasAddedCallstackFrame)
+		{
+			RemoveCallstackFrame(Context);
+		}
 
 		FINISH_HANDLE_INSTRUMENTATION();
 	}
@@ -1595,14 +1688,17 @@ template <typename AtomicType>
 FORCEINLINE INSTRUMENTATION_FUNCTION_ATTRIBUTES AtomicType InstrumentAtomicLoad(AtomicType* Atomic, FAtomicMemoryOrder Order)
 {
 	AtomicType Ret;
-	auto AtomicOp = [Atomic, &Ret]() {
+	auto AtomicOp = [Atomic, &Ret]() INSTRUMENTATION_FUNCTION_ATTRIBUTES {
 		Ret = UE::Core::Private::Atomic::Load(Atomic);
 		};
 
 	BEGIN_HANDLE_INSTRUMENTATION(AtomicOp(); return Ret);
-	AddCallstackFrame(Context, _ReturnAddress());
-	InstrumentAtomicAccess(Context, Atomic, sizeof(AtomicType), EMemoryAccessType::ACCESS_TYPE_ATOMIC_READ, Order, AtomicOp, _ReturnAddress(), TEXT("AtomicLoad"));
-	RemoveCallstackFrame(Context);
+	bool bHasAddedCallstackFrame = false;
+	InstrumentAtomicAccess(Context, Atomic, sizeof(AtomicType), EMemoryAccessType::ACCESS_TYPE_ATOMIC_READ, Order, AtomicOp, _ReturnAddress(), TEXT("AtomicLoad"), bHasAddedCallstackFrame);
+	if (bHasAddedCallstackFrame)
+	{
+		RemoveCallstackFrame(Context);
+	}
 	FINISH_HANDLE_INSTRUMENTATION()
 	return Ret;
 }
@@ -1610,14 +1706,17 @@ FORCEINLINE INSTRUMENTATION_FUNCTION_ATTRIBUTES AtomicType InstrumentAtomicLoad(
 template <typename AtomicType>
 FORCEINLINE INSTRUMENTATION_FUNCTION_ATTRIBUTES void InstrumentAtomicStore(AtomicType* Atomic, AtomicType Val, FAtomicMemoryOrder Order)
 {
-	auto AtomicOp = [Atomic, Val]() {
+	auto AtomicOp = [Atomic, Val]() INSTRUMENTATION_FUNCTION_ATTRIBUTES {
 		UE::Core::Private::Atomic::Store(Atomic, Val);
 		};
 
 	BEGIN_HANDLE_INSTRUMENTATION(AtomicOp(); return);
-	AddCallstackFrame(Context, _ReturnAddress());
-	InstrumentAtomicAccess(Context, Atomic, sizeof(AtomicType), EMemoryAccessType::ACCESS_TYPE_ATOMIC_WRITE, Order, AtomicOp, _ReturnAddress(), TEXT("AtomicStore"));
-	RemoveCallstackFrame(Context);
+	bool bHasAddedCallstackFrame = false;
+	InstrumentAtomicAccess(Context, Atomic, sizeof(AtomicType), EMemoryAccessType::ACCESS_TYPE_ATOMIC_WRITE, Order, AtomicOp, _ReturnAddress(), TEXT("AtomicStore"), bHasAddedCallstackFrame);
+	if (bHasAddedCallstackFrame)
+	{
+		RemoveCallstackFrame(Context);
+	}
 	FINISH_HANDLE_INSTRUMENTATION();
 }
 
@@ -1625,14 +1724,17 @@ template <typename AtomicType>
 FORCEINLINE INSTRUMENTATION_FUNCTION_ATTRIBUTES AtomicType InstrumentAtomicExchange(AtomicType* Atomic, AtomicType Val, FAtomicMemoryOrder Order)
 {
 	AtomicType Ret;
-	auto AtomicOp = [Atomic, Val, &Ret]() {
+	auto AtomicOp = [Atomic, Val, &Ret]() INSTRUMENTATION_FUNCTION_ATTRIBUTES {
 		Ret = UE::Core::Private::Atomic::Exchange(Atomic, Val);
 		};
 
 	BEGIN_HANDLE_INSTRUMENTATION(AtomicOp(); return Ret);
-	AddCallstackFrame(Context, _ReturnAddress());
-	InstrumentAtomicAccess(Context, Atomic, sizeof(AtomicType), EMemoryAccessType::ACCESS_TYPE_ATOMIC_READ_WRITE, Order, AtomicOp, _ReturnAddress(), TEXT("AtomicExchange"));
-	RemoveCallstackFrame(Context);
+	bool bHasAddedCallstackFrame = false;
+	InstrumentAtomicAccess(Context, Atomic, sizeof(AtomicType), EMemoryAccessType::ACCESS_TYPE_ATOMIC_READ_WRITE, Order, AtomicOp, _ReturnAddress(), TEXT("AtomicExchange"), bHasAddedCallstackFrame);
+	if (bHasAddedCallstackFrame)
+	{
+		RemoveCallstackFrame(Context);
+	}
 	FINISH_HANDLE_INSTRUMENTATION();
 	return Ret;
 }
@@ -1641,7 +1743,7 @@ template <typename AtomicType>
 FORCEINLINE INSTRUMENTATION_FUNCTION_ATTRIBUTES AtomicType InstrumentAtomicCompareExchange(AtomicType* Atomic, AtomicType* Expected, AtomicType Val, FAtomicMemoryOrder SuccessOrder, FAtomicMemoryOrder FailureOrder)
 {
 	AtomicType Ret;
-	auto AtomicOp = [Atomic, Expected, Val, &Ret]() -> bool {
+	auto AtomicOp = [Atomic, Expected, Val, &Ret]() INSTRUMENTATION_FUNCTION_ATTRIBUTES -> bool {
 		Ret = UE::Core::Private::Atomic::CompareExchange(Atomic, *Expected, Val);
 		if (Ret != *Expected)
 		{
@@ -1652,9 +1754,12 @@ FORCEINLINE INSTRUMENTATION_FUNCTION_ATTRIBUTES AtomicType InstrumentAtomicCompa
 	};
 
 	BEGIN_HANDLE_INSTRUMENTATION(AtomicOp(); return Ret);
-	AddCallstackFrame(Context, _ReturnAddress());
-	InstrumentAtomicAccess(Context, Atomic, sizeof(AtomicType), EMemoryAccessType::ACCESS_TYPE_ATOMIC_READ_WRITE, SuccessOrder, FailureOrder, AtomicOp, _ReturnAddress(), TEXT("CompareExchange"));
-	RemoveCallstackFrame(Context);
+	bool bHasAddedCallstackFrame = false;
+	InstrumentAtomicAccess(Context, Atomic, sizeof(AtomicType), EMemoryAccessType::ACCESS_TYPE_ATOMIC_READ_WRITE, SuccessOrder, FailureOrder, AtomicOp, _ReturnAddress(), TEXT("CompareExchange"), bHasAddedCallstackFrame);
+	if (bHasAddedCallstackFrame)
+	{
+		RemoveCallstackFrame(Context);
+	}
 	FINISH_HANDLE_INSTRUMENTATION();
 	return Ret;
 }
@@ -1663,14 +1768,17 @@ template <typename AtomicType>
 FORCEINLINE INSTRUMENTATION_FUNCTION_ATTRIBUTES AtomicType InstrumentAtomicFetchAdd(AtomicType* Atomic, AtomicType Val, FAtomicMemoryOrder Order)
 {
 	AtomicType Ret;
-	auto AtomicOp = [Atomic, Val, &Ret]() {
+	auto AtomicOp = [Atomic, Val, &Ret]() INSTRUMENTATION_FUNCTION_ATTRIBUTES {
 		Ret = UE::Core::Private::Atomic::AddExchange(Atomic, Val);
 		};
 
 	BEGIN_HANDLE_INSTRUMENTATION(AtomicOp(); return Ret);
-	AddCallstackFrame(Context, _ReturnAddress());
-	InstrumentAtomicAccess(Context, Atomic, sizeof(AtomicType), EMemoryAccessType::ACCESS_TYPE_ATOMIC_READ_WRITE, Order, AtomicOp, _ReturnAddress(), TEXT("FetchAdd"));
-	RemoveCallstackFrame(Context);
+	bool bHasAddedCallstackFrame = false;
+	InstrumentAtomicAccess(Context, Atomic, sizeof(AtomicType), EMemoryAccessType::ACCESS_TYPE_ATOMIC_READ_WRITE, Order, AtomicOp, _ReturnAddress(), TEXT("FetchAdd"), bHasAddedCallstackFrame);
+	if (bHasAddedCallstackFrame)
+	{
+		RemoveCallstackFrame(Context);
+	}
 	FINISH_HANDLE_INSTRUMENTATION();
 	return Ret;
 }
@@ -1679,14 +1787,17 @@ template <typename AtomicType>
 FORCEINLINE INSTRUMENTATION_FUNCTION_ATTRIBUTES AtomicType InstrumentAtomicFetchSub(AtomicType* Atomic, AtomicType Val, FAtomicMemoryOrder Order)
 {
 	AtomicType Ret;
-	auto AtomicOp = [Atomic, Val, &Ret]() {
+	auto AtomicOp = [Atomic, Val, &Ret]() INSTRUMENTATION_FUNCTION_ATTRIBUTES {
 		Ret = UE::Core::Private::Atomic::SubExchange(Atomic, Val);
 		};
 
 	BEGIN_HANDLE_INSTRUMENTATION(AtomicOp(); return Ret);
-	AddCallstackFrame(Context, _ReturnAddress());
-	InstrumentAtomicAccess(Context, Atomic, sizeof(AtomicType), EMemoryAccessType::ACCESS_TYPE_ATOMIC_READ_WRITE, Order, AtomicOp, _ReturnAddress(), TEXT("FetchSub"));
-	RemoveCallstackFrame(Context);
+	bool bHasAddedCallstackFrame = false;
+	InstrumentAtomicAccess(Context, Atomic, sizeof(AtomicType), EMemoryAccessType::ACCESS_TYPE_ATOMIC_READ_WRITE, Order, AtomicOp, _ReturnAddress(), TEXT("FetchSub"), bHasAddedCallstackFrame);
+	if (bHasAddedCallstackFrame)
+	{
+		RemoveCallstackFrame(Context);
+	}
 	FINISH_HANDLE_INSTRUMENTATION();
 	return Ret;
 }
@@ -1695,14 +1806,17 @@ template <typename AtomicType>
 FORCEINLINE INSTRUMENTATION_FUNCTION_ATTRIBUTES AtomicType InstrumentAtomicFetchOr(AtomicType* Atomic, AtomicType Val, FAtomicMemoryOrder Order)
 {
 	AtomicType Ret;
-	auto AtomicOp = [Atomic, Val, &Ret]() {
+	auto AtomicOp = [Atomic, Val, &Ret]() INSTRUMENTATION_FUNCTION_ATTRIBUTES {
 		Ret = UE::Core::Private::Atomic::OrExchange(Atomic, Val);
 		};
 
 	BEGIN_HANDLE_INSTRUMENTATION(AtomicOp(); return Ret);
-	AddCallstackFrame(Context, _ReturnAddress());
-	InstrumentAtomicAccess(Context, Atomic, sizeof(AtomicType), EMemoryAccessType::ACCESS_TYPE_ATOMIC_READ_WRITE, Order, AtomicOp, _ReturnAddress(), TEXT("FetchOr"));
-	RemoveCallstackFrame(Context);
+	bool bHasAddedCallstackFrame = false;
+	InstrumentAtomicAccess(Context, Atomic, sizeof(AtomicType), EMemoryAccessType::ACCESS_TYPE_ATOMIC_READ_WRITE, Order, AtomicOp, _ReturnAddress(), TEXT("FetchOr"), bHasAddedCallstackFrame);
+	if (bHasAddedCallstackFrame)
+	{
+		RemoveCallstackFrame(Context);
+	}
 	FINISH_HANDLE_INSTRUMENTATION();
 	return Ret;
 }
@@ -1711,14 +1825,17 @@ template <typename AtomicType>
 FORCEINLINE INSTRUMENTATION_FUNCTION_ATTRIBUTES AtomicType InstrumentAtomicFetchXor(AtomicType* Atomic, AtomicType Val, FAtomicMemoryOrder Order)
 {
 	AtomicType Ret;
-	auto AtomicOp = [Atomic, Val, &Ret]() {
+	auto AtomicOp = [Atomic, Val, &Ret]() INSTRUMENTATION_FUNCTION_ATTRIBUTES {
 		Ret = UE::Core::Private::Atomic::XorExchange(Atomic, Val);
 	};
 	
 	BEGIN_HANDLE_INSTRUMENTATION(AtomicOp(); return Ret);
-	AddCallstackFrame(Context, _ReturnAddress());
-	InstrumentAtomicAccess(Context, Atomic, sizeof(AtomicType), EMemoryAccessType::ACCESS_TYPE_ATOMIC_READ_WRITE, Order, AtomicOp, _ReturnAddress(), TEXT("FetchXor"));
-	RemoveCallstackFrame(Context);
+	bool bHasAddedCallstackFrame = false;
+	InstrumentAtomicAccess(Context, Atomic, sizeof(AtomicType), EMemoryAccessType::ACCESS_TYPE_ATOMIC_READ_WRITE, Order, AtomicOp, _ReturnAddress(), TEXT("FetchXor"), bHasAddedCallstackFrame);
+	if (bHasAddedCallstackFrame)
+	{
+		RemoveCallstackFrame(Context);
+	}
 	FINISH_HANDLE_INSTRUMENTATION();
 	return Ret;
 }
@@ -1727,14 +1844,17 @@ template <typename AtomicType>
 FORCEINLINE INSTRUMENTATION_FUNCTION_ATTRIBUTES AtomicType InstrumentAtomicFetchAnd(AtomicType* Atomic, AtomicType Val, FAtomicMemoryOrder Order)
 {
 	AtomicType Ret;
-	auto AtomicOp = [Atomic, Val, &Ret]() {
+	auto AtomicOp = [Atomic, Val, &Ret]() INSTRUMENTATION_FUNCTION_ATTRIBUTES {
 		Ret = UE::Core::Private::Atomic::AndExchange(Atomic, Val);
 		};
 	
 	BEGIN_HANDLE_INSTRUMENTATION(AtomicOp(); return Ret);
-	AddCallstackFrame(Context, _ReturnAddress());
-	InstrumentAtomicAccess(Context, Atomic, sizeof(AtomicType), EMemoryAccessType::ACCESS_TYPE_ATOMIC_READ_WRITE, Order, AtomicOp, _ReturnAddress(), TEXT("FetchAnd"));
-	RemoveCallstackFrame(Context);
+	bool bHasAddedCallstackFrame = false;
+	InstrumentAtomicAccess(Context, Atomic, sizeof(AtomicType), EMemoryAccessType::ACCESS_TYPE_ATOMIC_READ_WRITE, Order, AtomicOp, _ReturnAddress(), TEXT("FetchAnd"), bHasAddedCallstackFrame);
+	if (bHasAddedCallstackFrame)
+	{
+		RemoveCallstackFrame(Context);
+	}
 	FINISH_HANDLE_INSTRUMENTATION();
 	return Ret;
 }
@@ -1911,6 +2031,8 @@ INSTRUMENTATION_FUNCTION_ATTRIBUTES void PopulateHotPatchFunctions()
 
 INSTRUMENTATION_FUNCTION_ATTRIBUTES void PrepareHotPatchFunctions()
 {
+	FInstrumentationScope Scope;
+
 	for (const auto [ThunkFn, DetouredFn] : AlwaysOnInstrumentationFunctions)
 	{
 		Platform::PrepareTrampoline(ThunkFn, DetouredFn, false /* bUseRETBypass */);
@@ -1929,6 +2051,8 @@ INSTRUMENTATION_FUNCTION_ATTRIBUTES void PrepareHotPatchFunctions()
 
 INSTRUMENTATION_FUNCTION_ATTRIBUTES void CleanupHotPatchFunctions()
 {
+	FInstrumentationScope Scope;
+
 	for (const auto [ThunkFn, DetouredFn] : AlwaysOnInstrumentationFunctions)
 	{
 		Platform::CleanupTrampoline(ThunkFn);
@@ -1947,6 +2071,8 @@ INSTRUMENTATION_FUNCTION_ATTRIBUTES void CleanupHotPatchFunctions()
 
 INSTRUMENTATION_FUNCTION_ATTRIBUTES void HookAlwaysOnInstrumentation()
 {
+	FInstrumentationScope Scope;
+
 	for (const auto [ThunkFn, DetouredFn] : AlwaysOnInstrumentationFunctions)
 	{
 		Platform::ActivateTrampoline(ThunkFn);
@@ -1957,6 +2083,8 @@ INSTRUMENTATION_FUNCTION_ATTRIBUTES void HookAlwaysOnInstrumentation()
 
 INSTRUMENTATION_FUNCTION_ATTRIBUTES void HookInstrumentation()
 {
+	FInstrumentationScope Scope;
+
 	for (const auto [ThunkFn, DetouredFn] : HotpatchInstrumentationFunctions)
 	{
 		Platform::ActivateTrampoline(ThunkFn);
@@ -1972,6 +2100,8 @@ INSTRUMENTATION_FUNCTION_ATTRIBUTES void HookInstrumentation()
 
 INSTRUMENTATION_FUNCTION_ATTRIBUTES void UnhookAlwaysOnInstrumentation()
 {
+	FInstrumentationScope Scope;
+
 	for (const auto [ThunkFn, DetouredFn] : AlwaysOnInstrumentationFunctions)
 	{
 		Platform::DeactivateTrampoline(ThunkFn, false /* bUseRETBypass */);
@@ -1982,6 +2112,8 @@ INSTRUMENTATION_FUNCTION_ATTRIBUTES void UnhookAlwaysOnInstrumentation()
 
 INSTRUMENTATION_FUNCTION_ATTRIBUTES void UnhookInstrumentation()
 {
+	FInstrumentationScope Scope;
+
 	for (const auto [ThunkFn, DetouredFn] : HotpatchInstrumentationFunctions)
 	{
 		Platform::DeactivateTrampoline(ThunkFn, false /* bUseRETBypass */);
