Added UE_USE_LITE_ENSURES for opting in to reduced PDB sizes.

diff --git a/Engine/Source/Runtime/Core/Public/Misc/AssertionMacros.h b/Engine/Source/Runtime/Core/Public/Misc/AssertionMacros.h
--- a/Engine/Source/Runtime/Core/Public/Misc/AssertionMacros.h
+++ b/Engine/Source/Runtime/Core/Public/Misc/AssertionMacros.h
@@ -33,6 +33,19 @@
 #endif // DO_CHECK || DO_GUARD_SLOW
 #endif
 
+// Lite ensures use an alternate implementation that avoid lambdas which are expensive in the amount of
+// debug information generated. However, this comes at the cost of ensures breaking inside an implementation
+// function rather than on the line that the ensure is on. On Clang, this limitation can be mitigated by
+// using UE_NODEBUG, so that the break happens where the code is inlined rather than where it is defined,
+// but no such option exists for MSVC, so we disable lite ensures by default on that compiler.
+#ifndef UE_USE_LITE_ENSURES
+	#ifdef __clang__
+		#define UE_USE_LITE_ENSURES 1
+	#else
+		#define UE_USE_LITE_ENSURES 0
+	#endif
+#endif
+
 namespace ELogVerbosity
 {
 	enum Type : uint8;
@@ -415,6 +428,7 @@ RetType UE_COLD UE_DEBUG_SECTION DispatchCheckVerify (InnerType&& Inner, ArgType
 		}
 	}
 
+#if UE_USE_LITE_ENSURES
 	UE_NODEBUG FORCEINLINE bool BreakAndReturnFalse()
 	{
 		// This is tagged with UE_NODEBUG and FORCEINLINE to ensure that the break happens at the usage of ensure(), not inside this function.
@@ -426,6 +440,12 @@ RetType UE_COLD UE_DEBUG_SECTION DispatchCheckVerify (InnerType&& Inner, ArgType
 		(LIKELY(!!(InExpression)) \
 			|| (::UE::Assert::Private::ExecCheckImplInternal(::bGEnsureHasExecuted<((uint64)FileHashForEnsure(__FILE__) << 32) | (uint64)__LINE__>, Always, __FILE__, __LINE__, #InExpression) \
 			&& BreakAndReturnFalse()))
+#else
+	#define UE_ENSURE_IMPL(Always, InExpression) \
+		(LIKELY(!!(InExpression)) \
+			|| (::UE::Assert::Private::ExecCheckImplInternal(::bGEnsureHasExecuted<((uint64)FileHashForEnsure(__FILE__) << 32) | (uint64)__LINE__>, Always, __FILE__, __LINE__, #InExpression) \
+			&& [](){ PLATFORM_BREAK(); return false; }()))
+#endif
 
 	#define UE_ENSURE_IMPL2(Capture, Always, InExpression, InFormat, ...) \
 		(LIKELY(!!(InExpression)) || ([Capture] () UE_COLD UE_DEBUG_SECTION \
