[RaceDetector]
- Implement clock bank shadow to get rid of hash table
- Implement instrumentation aware versions of our spinlock and scopes to fix reentrancy slowdowns
- Use templates instead of TFunctionRef to further improve performance
- 3x performance improvements on non-relaxed atomic operations

diff --git a/Engine/Source/Runtime/Core/Internal/Sanitizer/RaceDetectorInterface.h b/Engine/Source/Runtime/Core/Internal/Sanitizer/RaceDetectorInterface.h
--- a/Engine/Source/Runtime/Core/Internal/Sanitizer/RaceDetectorInterface.h
+++ b/Engine/Source/Runtime/Core/Internal/Sanitizer/RaceDetectorInterface.h
@@ -7,7 +7,6 @@
 #if USING_INSTRUMENTATION
 
 #include "CoreTypes.h"
-#include "Templates/RefCounting.h"
 #include "Sanitizer/RaceDetectorTypes.h"
 
 DECLARE_LOG_CATEGORY_EXTERN(LogRaceDetector, Log, All);
@@ -30,6 +29,8 @@ namespace UE::Sanitizer::RaceDetector {
 		// Returns the size of the entire shadow memory address space.
 		// Note: InitShadowMemory must have been called before, otherwise this is meaningless.
 		UPTRINT GetShadowMemorySize();
+		// Get the base of the shadow memory for clock banks.
+		UPTRINT GetShadowClockBase();
 		// Returns the number of bytes currently mapped in the shadow memory.
 		uint64 GetShadowMemoryUsage();
 		// Returns whether a particular range is already accessible in shadow memory.
@@ -87,7 +88,7 @@ namespace UE::Sanitizer::RaceDetector {
 	// Releases the current thread context.
 	void ReleaseCurrentContext();
 	// Returns a sync object for the given address, initialize one if there isn't one already.
-	TRefCountPtr<FSyncObject> GetSyncObject(void* SyncAddr);
+	FSyncObjectRef GetSyncObject(void* SyncAddr);
 }
 
 #endif // USING_INSTRUMENTATION
\ No newline at end of file
diff --git a/Engine/Source/Runtime/Core/Internal/Sanitizer/RaceDetectorTypes.h b/Engine/Source/Runtime/Core/Internal/Sanitizer/RaceDetectorTypes.h
--- a/Engine/Source/Runtime/Core/Internal/Sanitizer/RaceDetectorTypes.h
+++ b/Engine/Source/Runtime/Core/Internal/Sanitizer/RaceDetectorTypes.h
@@ -8,11 +8,6 @@
 
 #include "CoreTypes.h"
 
-#include "Misc/EnumClassFlags.h"
-#include "Misc/RWSpinLock.h"
-#include "Misc/ScopeRWLock.h"
-#include "Templates/RefCounting.h"
-
 #include "Sanitizer/Types.h"
 #include "Instrumentation/Types.h"
 #include "Instrumentation/Containers.h"
@@ -25,6 +20,153 @@ namespace UE::Sanitizer::RaceDetector {
 
 	extern volatile bool bDetailedLogGlobal;
 
+	// A read-write lock that doesn't put the thread into a WAIT state but instead repeatedly tries to acquire the lock.
+	// This version is customized to remove instrumentation and be as optimized as possible for instrumentation purpose.
+	// FPlatformAtomics are used to make sure all atomics are inlined.
+	// std::atomic often end up causing calls into non-inlined instrumented functions that causes costly reentrancy.
+	class FRWSpinLock
+	{
+	public:
+		UE_NONCOPYABLE(FRWSpinLock);
+
+		INSTRUMENTATION_FUNCTION_ATTRIBUTES FORCEINLINE FRWSpinLock() = default;
+
+		INSTRUMENTATION_FUNCTION_ATTRIBUTES FORCEINLINE bool TryWriteLock()
+		{
+			return FPlatformAtomics::InterlockedCompareExchange((volatile int32*)&Lock, UINT32_MAX, 0) == 0;
+		}
+
+		INSTRUMENTATION_FUNCTION_ATTRIBUTES FORCEINLINE void WriteLock()
+		{
+			while (!TryWriteLock())
+			{
+				// Reduce contention by doing a simple relaxed read to see if we have a chance of being able to lock.
+				while (Lock != 0)
+				{
+					FPlatformProcess::Yield();
+				}
+			}
+		}
+
+		INSTRUMENTATION_FUNCTION_ATTRIBUTES FORCEINLINE void WriteUnlock()
+		{
+			Lock = 0;
+		}
+
+		INSTRUMENTATION_FUNCTION_ATTRIBUTES FORCEINLINE bool TryReadLock()
+		{
+			uint32 LocalValue = Lock;
+			// Check to make sure we don't already have a write lock or that we've not reached the limit of reader locks.
+			if (LocalValue >= UINT32_MAX - 1)
+			{
+				return false;
+			}
+
+			return FPlatformAtomics::InterlockedCompareExchange((volatile int32*)&Lock, LocalValue + 1, LocalValue) == LocalValue;
+		}
+
+		INSTRUMENTATION_FUNCTION_ATTRIBUTES FORCEINLINE void ReadUnlock()
+		{
+			FPlatformAtomics::InterlockedDecrement((volatile int32*)&Lock);
+		}
+
+		INSTRUMENTATION_FUNCTION_ATTRIBUTES FORCEINLINE void ReadLock()
+		{
+			while (!TryReadLock())
+			{
+				FPlatformProcess::Yield();
+			}
+		}
+
+	private:
+		volatile uint32 Lock = 0;
+	};
+
+	template<typename MutexType>
+	class TReadScopeLock
+	{
+	public:
+		UE_NONCOPYABLE(TReadScopeLock);
+
+		INSTRUMENTATION_FUNCTION_ATTRIBUTES FORCEINLINE TReadScopeLock(MutexType& InMutex)
+			: Mutex(InMutex)
+		{
+			Mutex.ReadLock();
+		}
+
+		INSTRUMENTATION_FUNCTION_ATTRIBUTES FORCEINLINE ~TReadScopeLock()
+		{
+			Mutex.ReadUnlock();
+		}
+
+	private:
+		MutexType& Mutex;
+	};
+
+	template<typename MutexType>
+	class TWriteScopeLock
+	{
+	public:
+		UE_NONCOPYABLE(TWriteScopeLock);
+
+		INSTRUMENTATION_FUNCTION_ATTRIBUTES FORCEINLINE TWriteScopeLock(MutexType& InMutex)
+			: Mutex(InMutex)
+		{
+			Mutex.WriteLock();
+		}
+
+		INSTRUMENTATION_FUNCTION_ATTRIBUTES FORCEINLINE ~TWriteScopeLock()
+		{
+			Mutex.WriteUnlock();
+		}
+
+	private:
+		MutexType& Mutex;
+	};
+
+	enum FRWScopeLockType
+	{
+		SLT_ReadOnly = 0,
+		SLT_Write,
+	};
+
+	template<typename MutexType>
+	class TRWScopeLock
+	{
+	public:
+		UE_NONCOPYABLE(TRWScopeLock);
+
+		INSTRUMENTATION_FUNCTION_ATTRIBUTES FORCEINLINE TRWScopeLock(MutexType& InMutex, FRWScopeLockType InLockType)
+			: Mutex(InMutex)
+			, LockType(InLockType)
+		{
+			if (LockType == SLT_ReadOnly)
+			{
+				Mutex.ReadLock();
+			}
+			else
+			{
+				Mutex.WriteLock();
+			}
+		}
+
+		INSTRUMENTATION_FUNCTION_ATTRIBUTES FORCEINLINE ~TRWScopeLock()
+		{
+			if (LockType == SLT_ReadOnly)
+			{
+				Mutex.ReadUnlock();
+			}
+			else
+			{
+				Mutex.WriteUnlock();
+			}
+		}
+
+	private:
+		MutexType& Mutex;
+		FRWScopeLockType LockType;
+	};
+
 	// ------------------------------------------------------------------------------
 	// Clocks.
 	// ------------------------------------------------------------------------------
@@ -153,14 +295,14 @@ namespace UE::Sanitizer::RaceDetector {
 		{
 		}
 
-		INSTRUMENTATION_FUNCTION_ATTRIBUTES FORCEINLINE FMemoryAccess(FContextId InContextId, FClock InClock, uint8 InOffset, uint8 InSize, EMemoryAccessType InAccessType, bool bInHasSyncObjectAttached)
+		INSTRUMENTATION_FUNCTION_ATTRIBUTES FORCEINLINE FMemoryAccess(FContextId InContextId, FClock InClock, uint8 InOffset, uint8 InSize, EMemoryAccessType InAccessType)
 		{
 			// Apparently this is much faster than using the bitfields. 30% faster in warm TSAN benchmark!!!
 			// The constructor was accessing the same value in memory/store buffer multiple time doing its bit tweaking and it caused tons of Core::X86::Pmc::Core::LsBadStatus2.
 			// That ended up stalling on S[0].RawValue == CurrentAccess.RawValue in InstrumentMemoryAccess when trying to extract the full uint64.
 			// https://blog.stuffedcow.net/2014/01/x86-memory-disambiguation/
 			const uint8 AccessValue = (uint8)(((1ull << InSize) - 1ull) << InOffset);
-			RawValue = (uint64)bInHasSyncObjectAttached << 56 | (uint64)InAccessType << 48 | (uint64)AccessValue << 40 | (uint64)InContextId << 32 | (uint64)InClock;
+			RawValue = (uint64)InAccessType << 48 | (uint64)AccessValue << 40 | (uint64)InContextId << 32 | (uint64)InClock;
 		}
 		
 		INSTRUMENTATION_FUNCTION_ATTRIBUTES FORCEINLINE uint8 GetOffset() const { return (uint8)FMath::CountTrailingZeros(Access); }
@@ -188,7 +330,7 @@ namespace UE::Sanitizer::RaceDetector {
 					};
 					EMemoryAccessType AccessType : 4; // 6  plenty of bits left here
 				};
-				uint8 bHasSyncObjectAttached : 1;     // 7  plenty of bits left here
+				uint8 Reserved;                           // 7  plenty of bits left here
 			};
 
 			uint64 RawValue;
@@ -281,7 +423,7 @@ namespace UE::Sanitizer::RaceDetector {
 
 		INSTRUMENTATION_FUNCTION_ATTRIBUTES ~FAccessHistory()
 		{
-			UE::TWriteScopeLock Scope(Lock);
+			TWriteScopeLock Scope(Lock);
 
 			while (Head)
 			{
@@ -312,7 +454,7 @@ namespace UE::Sanitizer::RaceDetector {
 		double LastRecycle = 0.0;
 
 		// Just used on the slow path between recycling and scanning.
-		UE::FRWSpinLock Lock;
+		FRWSpinLock Lock;
 
 		// Can be used by other threads doing race reporting
 		FHistoryChunk* Head = nullptr;
@@ -334,7 +476,7 @@ namespace UE::Sanitizer::RaceDetector {
 		{
 			if (HasTooManyChunks())
 			{
-				UE::TWriteScopeLock Scope(Lock);
+				TWriteScopeLock Scope(Lock);
 
 				while (HasTooManyChunks())
 				{
@@ -417,7 +559,7 @@ namespace UE::Sanitizer::RaceDetector {
 
 		INSTRUMENTATION_FUNCTION_ATTRIBUTES bool ResolveAccess(const void* InAlignedPointer, FMemoryAccess InAccess, FCallstackLocation& OutLocation, FClockRange& OutClockRange)
 		{
-			UE::TReadScopeLock Scope(Lock);
+			TReadScopeLock Scope(Lock);
 
 			// Make sure we sync with the light fence.
 			FPlatformMisc::AsymmetricThreadFenceHeavy();
@@ -595,47 +737,55 @@ namespace UE::Sanitizer::RaceDetector {
 		}
 	};
 
-	class FSyncObject : public TRefCountingMixin<FSyncObject> {
+	struct FSyncObjectBank;
+	class FSyncObject
+	{
 	public:
-		SAFE_OPERATOR_NEW_DELETE()
-
-		INSTRUMENTATION_FUNCTION_ATTRIBUTES FSyncObject()
+		INSTRUMENTATION_FUNCTION_ATTRIBUTES void* operator new(SIZE_T Size)
 		{
-			FPlatformAtomics::InterlockedIncrement(&SyncObjectCount);
+			// Only count SyncObject that have been allocated separately to 
+			// avoid counting the one embedded in the SyncObjectBank.
+			FPlatformAtomics::InterlockedIncrement(&ObjectCount);
+			return FInstrumentationSafeWinAllocator::Alloc(Size);
 		}
-		
-		INSTRUMENTATION_FUNCTION_ATTRIBUTES ~FSyncObject()
+
+		INSTRUMENTATION_FUNCTION_ATTRIBUTES void operator delete(void* Ptr)
 		{
-			FPlatformAtomics::InterlockedDecrement(&SyncObjectCount);
+			FPlatformAtomics::InterlockedDecrement(&ObjectCount);
+			FInstrumentationSafeWinAllocator::Free(Ptr);
 		}
 
-		INSTRUMENTATION_FUNCTION_ATTRIBUTES void SyncAcquire(FContext& Context, const TFunctionRef<void()>& AtomicOp, void* ReturnAddress, void* SyncAddress, const TCHAR* OpName)
+		template <typename AtomicOpType>
+		INSTRUMENTATION_FUNCTION_ATTRIBUTES void SyncAcquire(FContext& Context, AtomicOpType&& AtomicOp, void* ReturnAddress, void* SyncAddress, const TCHAR* OpName)
 		{
-			UE::TWriteScopeLock Scope(Lock);
+			TWriteScopeLock Scope(Lock);
 			SyncAcquireAsSoleOwnerOrReadOwner(Context, ReturnAddress, SyncAddress, OpName);
 			AtomicOp();
 		}
 
-		INSTRUMENTATION_FUNCTION_ATTRIBUTES void SyncRelease(FContext& Context, const TFunctionRef<void()>& AtomicOp, void* ReturnAddress, void* SyncAddress, const TCHAR* OpName)
+		template <typename AtomicOpType>
+		INSTRUMENTATION_FUNCTION_ATTRIBUTES void SyncRelease(FContext& Context, AtomicOpType&& AtomicOp, void* ReturnAddress, void* SyncAddress, const TCHAR* OpName)
 		{
-			UE::TWriteScopeLock Scope(Lock);
+			TWriteScopeLock Scope(Lock);
 			SyncReleaseAsSoleOwner(Context, ReturnAddress, SyncAddress, OpName);
 			AtomicOp();
 		}
 
-		INSTRUMENTATION_FUNCTION_ATTRIBUTES void SyncAcquireRelease(FContext& Context, const TFunctionRef<void()>& AtomicOp, void* ReturnAddress, void* SyncAddress, const TCHAR* OpName)
+		template <typename AtomicOpType>
+		INSTRUMENTATION_FUNCTION_ATTRIBUTES void SyncAcquireRelease(FContext& Context, AtomicOpType&& AtomicOp, void* ReturnAddress, void* SyncAddress, const TCHAR* OpName)
 		{
 			if (UNLIKELY(Context.DetailedLogDepth || bDetailedLogGlobal))
 			{
 				FPlatformMisc::LowLevelOutputDebugStringf(TEXT("[%d] %s acq_rel of 0x%p from function at 0x%p\n"), Context.ThreadId, OpName, SyncAddress, ReturnAddress);
 			}
 
-			UE::TWriteScopeLock Scope(Lock);
+			TWriteScopeLock Scope(Lock);
 			Context.ClockBank.AcquireRelease(ClockBank, ReturnAddress);
 			AtomicOp();
 		}
 
-		INSTRUMENTATION_FUNCTION_ATTRIBUTES void SyncWithFailureSupport(FContext& Context, const TFunctionRef<bool()>& AtomicOp, EMemoryAccessType AccessType, FAtomicMemoryOrder SuccessOrder, FAtomicMemoryOrder FailureOrder, void* ReturnAddress, void* SyncAddress, const TCHAR* OpName, const TFunctionRef<void(FAtomicMemoryOrder ActualOrder)>& ActualAccessCallback)
+		template <typename AtomicOpType, typename ActualAccessCallbackType>
+		INSTRUMENTATION_FUNCTION_ATTRIBUTES void SyncWithFailureSupport(FContext& Context, AtomicOpType&& AtomicOp, EMemoryAccessType AccessType, FAtomicMemoryOrder SuccessOrder, FAtomicMemoryOrder FailureOrder, void* ReturnAddress, void* SyncAddress, const TCHAR* OpName, ActualAccessCallbackType&& ActualAccessCallback)
 		{
 			using namespace UE::Instrumentation;
 
@@ -647,7 +797,7 @@ namespace UE::Sanitizer::RaceDetector {
 			// #17 Requires: The failure argument shall not be memory_order_release nor memory_order_acq_rel.
 			// 
 			// Also for the time being MSVC atomic implementation always combines both order so we have no way to test our support of different SuccessOrder and FailureOrder until they fix their implementation.
-			UE::TRWScopeLock Scope(Lock, IsAtomicOrderRelease(SuccessOrder) ? SLT_Write : SLT_ReadOnly);
+			TRWScopeLock Scope(Lock, IsAtomicOrderRelease(SuccessOrder) ? SLT_Write : SLT_ReadOnly);
 			bool bSucceeded = AtomicOp();
 
 			FAtomicMemoryOrder Order = bSucceeded ? SuccessOrder : FailureOrder;
@@ -722,14 +872,15 @@ namespace UE::Sanitizer::RaceDetector {
 			Context.ClockBank.Acquire(ClockBank, ReturnAddress);
 		}
 
-		INSTRUMENTATION_FUNCTION_ATTRIBUTES static uint64 GetSyncObjectCount()
+		INSTRUMENTATION_FUNCTION_ATTRIBUTES static uint64 GetObjectCount()
 		{
-			return SyncObjectCount;
+			return ObjectCount;
 		}
+
 	private:
-		UE::FRWSpinLock Lock;
+		FRWSpinLock Lock;
 		FClockBank ClockBank;
-		static volatile int64 SyncObjectCount;
+		static volatile int64 ObjectCount;
 	};
 
 	// One SyncObjectBank per 64-bit aligned address
@@ -737,8 +888,116 @@ namespace UE::Sanitizer::RaceDetector {
 	{
 		SAFE_OPERATOR_NEW_DELETE()
 
-		// Contains 8 sync object for each byte of the 64-bit.
-		TRefCountPtr<FSyncObject> SyncObjects[8];
+		INSTRUMENTATION_FUNCTION_ATTRIBUTES FSyncObjectBank()
+		{
+			FPlatformAtomics::InterlockedIncrement(&ObjectCount);
+		}
+
+		INSTRUMENTATION_FUNCTION_ATTRIBUTES ~FSyncObjectBank()
+		{
+			for (int32 Index = 0; Index < 8; ++Index)
+			{
+				if (SyncObjects[Index])
+				{
+					delete SyncObjects[Index];
+				}
+			}
+
+			FPlatformAtomics::InterlockedDecrement(&ObjectCount);
+		}
+
+		INSTRUMENTATION_FUNCTION_ATTRIBUTES static uint64 GetObjectCount()
+		{
+			return ObjectCount;
+		}
+
+		INSTRUMENTATION_FUNCTION_ATTRIBUTES FSyncObject* GetSyncObject(int32 Index)
+		{
+			if (Index == 0)
+			{
+				return &EmbeddedObject;
+			}
+
+			Index--;
+			if (SyncObjects[Index] == nullptr)
+			{
+				FSyncObject* SyncObject = new FSyncObject();
+				if (FPlatformAtomics::InterlockedCompareExchangePointer((void**)&SyncObjects[Index], SyncObject, nullptr) != nullptr)
+				{
+					delete SyncObject;
+				}
+			}
+
+			return SyncObjects[Index];
+		}
+
+		// We maintain a linked list of clock banks for recycling purpose
+		FSyncObjectBank* Next = 0;
+
+		INSTRUMENTATION_FUNCTION_ATTRIBUTES void AddRef()
+		{
+			FPlatformAtomics::InterlockedIncrement(&RefCount);
+		}
+
+		INSTRUMENTATION_FUNCTION_ATTRIBUTES void Release()
+		{
+			if (FPlatformAtomics::InterlockedDecrement(&RefCount) == 0)
+			{
+				delete this;
+			}
+		}
+
+		static volatile int64 ObjectCount;
+	private:
+		volatile int32 RefCount = 0;
+
+		// Save space by allocating the first sync object as part of the bank itself.
+		// Most of the time the sync object will be at offset 0.
+		// The safe allocator uses virtual memory with 4KB pages so this first entry
+		// is completely free.
+		FSyncObject  EmbeddedObject;
+
+		// Contains optional sync object for each unaligned bytes of the 64-bit.
+		FSyncObject* SyncObjects[7] = { 0 };
+	};
+
+	struct FSyncObjectRef
+	{
+		INSTRUMENTATION_FUNCTION_ATTRIBUTES FORCEINLINE FSyncObjectRef(FSyncObjectBank* InBank, FSyncObject* InSyncObject)
+			: Bank(InBank)
+			, Object(InSyncObject)
+		{
+			AddRef();
+		}
+
+		INSTRUMENTATION_FUNCTION_ATTRIBUTES FORCEINLINE ~FSyncObjectRef()
+		{
+			Release();
+		}
+
+		// Not needed for now
+		FSyncObjectRef(const FSyncObjectRef&) = delete;
+		FSyncObjectRef(FSyncObjectRef&&) = delete;
+		FSyncObjectRef& operator=(const FSyncObjectRef&) = delete;
+		FSyncObjectRef& operator=(FSyncObjectRef&&) = delete;
+
+		INSTRUMENTATION_FUNCTION_ATTRIBUTES FORCEINLINE FSyncObject* operator->() const
+		{
+			return Object;
+		}
+	private:
+		INSTRUMENTATION_FUNCTION_ATTRIBUTES FORCEINLINE void AddRef()
+		{
+			Bank->AddRef();
+		}
+
+		INSTRUMENTATION_FUNCTION_ATTRIBUTES FORCEINLINE void Release()
+		{
+			Bank->Release();
+		}
+
+		FSyncObjectBank* Bank;
+		FSyncObject* Object;
 	};
 
 	// ------------------------------------------------------------------------------
@@ -749,6 +1008,11 @@ namespace UE::Sanitizer::RaceDetector {
 		FMemoryAccess Accesses[4];
 	};
 
+	struct FShadowClockBankSlot
+	{
+		FSyncObjectBank* SyncObjectBank;
+	};
+
 } // UE::Sanitizer::RaceDetector
 
 #endif // USING_INSTRUMENTATION
\ No newline at end of file
diff --git a/Engine/Source/Runtime/Core/Private/Sanitizer/RaceDetector.cpp b/Engine/Source/Runtime/Core/Private/Sanitizer/RaceDetector.cpp
--- a/Engine/Source/Runtime/Core/Private/Sanitizer/RaceDetector.cpp
+++ b/Engine/Source/Runtime/Core/Private/Sanitizer/RaceDetector.cpp
@@ -13,16 +13,13 @@
 
 #include "CoreTypes.h"
 #include "Trace/Trace.h"
-#include "Misc/SpinLock.h"
 #include "Misc/CString.h"
 #include "Misc/ScopeLock.h"
 #include "Misc/Paths.h"
 #include "Async/Mutex.h"
 #include "Misc/Parse.h"
 #include "Math/NumericLimits.h"
-#include "Algo/Reverse.h"
 #include "Misc/CommandLine.h"
-#include "Async/UniqueLock.h"
 #include "HAL/FileManager.h"
 #include "HAL/ThreadManager.h"
 #include "Tasks/Task.h"
@@ -122,24 +119,30 @@ namespace UE::Sanitizer::RaceDetector {
 
 	UPTRINT ShadowMemoryBase = 0;
 	UPTRINT ShadowMemorySize = 0;
+	UPTRINT ShadowClockBase = 0;
 	UPTRINT ShadowMemoryEnd = 0;
 	uint32  ContextTLSIndex  = UINT32_MAX;
 	volatile uint32 GlobalEpoch = 0;
-	volatile int64 FSyncObject::SyncObjectCount = 0;
+	volatile int64 FSyncObject::ObjectCount = 0;
+	volatile int64 FSyncObjectBank::ObjectCount = 0;
 	volatile FContextId CurrentContextId = 1;
 	volatile int64 HistoryChunkCount = 0;
 
-	UE::FSpinLock RaceReportsLock;
+	// Since the shadow space can be unmapped in a single shot loosing all our pointers,
+	// we need to keep a list of all the banks we've allocated so we can free them.
+	volatile FSyncObjectBank* SyncObjectBankHead = nullptr;
+
+	FRWSpinLock RaceReportsLock;
 	TArray<FString> RaceReports;
 
 	// Lock for both ContextMapping and FreeContexts
-	UE::FSpinLock ContextMappingLock; 
+	FRWSpinLock ContextMappingLock; 
 	TSafeMap<FContextId, TRefCountPtr<FContext>> ContextMapping;
 	TSafeArray<FContextId> FreeContexts;
 
 	// Used for filtering duplicates
 	TSafeSet<uint64> RaceHashes;
-	UE::FSpinLock RaceHashesLock;
+	FRWSpinLock RaceHashesLock;
 	UE::FTimeout RaceDetectorTimeout = UE::FTimeout::Never();
 	
 	volatile bool bIsResettingShadow = false;
@@ -150,7 +153,7 @@ namespace UE::Sanitizer::RaceDetector {
 	// Can't be an atomic as it might cause infinite recursion in debug when atomics
 	// are unoptimized and we end up with a function call to an instrumented function
 	// going directly back to the ShouldInstrument call. std::atomics that are tested
-	// after InstrumentationDepth should be OK.
+	// after InstrumentationDepth should work but are not optimal.
 	volatile bool bRaceDetectorActive = false;
 	INSTRUMENTATION_FUNCTION_ATTRIBUTES bool ShouldInstrument(FContext& Context)
 	{
@@ -169,7 +172,8 @@ namespace UE::Sanitizer::RaceDetector {
 
 	INSTRUMENTATION_FUNCTION_ATTRIBUTES uint64 GetSyncObjectMemoryUsage()
 	{
-		return FSyncObject::GetSyncObjectCount() * Align(sizeof(FSyncObject), Platform::GetPageSize());
+		return FSyncObject::GetObjectCount() * Align(sizeof(FSyncObject), Platform::GetPageSize()) +
+			FSyncObjectBank::GetObjectCount() * Align(sizeof(FSyncObjectBank), Platform::GetPageSize());
 	}
 
 	INSTRUMENTATION_FUNCTION_ATTRIBUTES uint64 GetHistoryChunkMemoryUsage()
@@ -223,7 +227,7 @@ namespace UE::Sanitizer::RaceDetector {
 		check(Context.AccessHistory == nullptr);
 		Context.AccessHistory = new FAccessHistory();
 
-		UE::TScopeLock Scope(ContextMappingLock);
+		TWriteScopeLock Scope(ContextMappingLock);
 		Context.GlobalEpoch = GlobalEpoch;
 		if (CurrentContextId < TNumericLimits<FContextId>::Max())
 		{
@@ -287,7 +291,7 @@ namespace UE::Sanitizer::RaceDetector {
 				// We don't destroy the context here as we want to leave a chance
 				// to find race in short lived threads or threads that race just
 				// before exiting. We'll recycle it when its ContextId gets reused.
-				UE::TScopeLock Scope(ContextMappingLock);
+				TWriteScopeLock Scope(ContextMappingLock);
 				FreeContexts.Push(Context->ContextId);
 				Context->ContextId = 0;
 			}
@@ -296,7 +300,7 @@ namespace UE::Sanitizer::RaceDetector {
 
 	INSTRUMENTATION_FUNCTION_ATTRIBUTES TRefCountPtr<FContext> GetContextById(FContextId ContextId)
 	{
-		UE::TScopeLock Scope(ContextMappingLock);
+		TWriteScopeLock Scope(ContextMappingLock);
 		return ContextMapping.FindRef(ContextId);
 	}
 
@@ -335,76 +339,85 @@ namespace UE::Sanitizer::RaceDetector {
 		}
 		return (FShadowMemory*)ShadowMemoryAddress;
 	}
-
-	// Reduce lock contention by striping the addresses into buckets.
-	constexpr uint32 STRIPING_BITS = 8;
-	constexpr uint32 STRIPING_SIZE = 1 << STRIPING_BITS;
-	constexpr uint32 STRIPING_MASK = STRIPING_SIZE - 1;
-
-	struct FSyncBanks
-	{
-		TSafeMap<UPTRINT, FSyncObjectBank> Map;
-		UE::FSpinLock Lock;
-	};
-	FSyncBanks StripedSyncBanks[STRIPING_SIZE];
 	
-	INSTRUMENTATION_FUNCTION_ATTRIBUTES TRefCountPtr<FSyncObject> GetSyncObject(void* SyncAddr)
+	FORCEINLINE INSTRUMENTATION_FUNCTION_ATTRIBUTES UPTRINT GetShadowClockBankMemoryAddress(UPTRINT Ptr)
 	{
-		UPTRINT Index    = UPTRINT(SyncAddr) & 7;
-		UPTRINT BankAddr = UPTRINT(SyncAddr) >> 3;
+		static_assert(sizeof(FShadowClockBankSlot) <= sizeof(UPTRINT));
 
-		FSyncBanks& SyncBanks = StripedSyncBanks[BankAddr & STRIPING_MASK];
-		
-		UE::TScopeLock Scope(SyncBanks.Lock);
-		FSyncObjectBank* Bank = SyncBanks.Map.Find(BankAddr);
-		if (!Bank)
-		{
-			Bank = &SyncBanks.Map.Emplace(BankAddr, FSyncObjectBank());
-		}
-		TRefCountPtr<FSyncObject>& SyncObject = Bank->SyncObjects[Index];
-		if (SyncObject == nullptr)
+		// Map the higher address space as a continuation of the lower one by just removing the shadow size from it.
+		if (Ptr >= ShadowMemoryEnd)
 		{
-			SyncObject = new FSyncObject();
+			Ptr -= ShadowMemorySize;
 		}
 
-		return SyncObject;
+		return ShadowClockBase + (Ptr >> 3) * sizeof(FShadowClockBankSlot);
 	}
 
-	INSTRUMENTATION_FUNCTION_ATTRIBUTES bool HasInstrumentedAtomics()
+	INSTRUMENTATION_FUNCTION_ATTRIBUTES FShadowClockBankSlot* GetShadowClockBank(UPTRINT Ptr)
 	{
-		FInstrumentationScope InstrumentationScope;
-		for (uint32 Index = 0; Index < STRIPING_SIZE; ++Index)
-		{
-			FSyncBanks& SyncBanks = StripedSyncBanks[Index];
+		UPTRINT ShadowMemoryAddress = GetShadowClockBankMemoryAddress(Ptr);
+		
+		// Always do mapping to avoid reentrancy in the vectored exception handler that
+		// uses SRW lock and end up calling back into this function for another sync object.
+		Platform::MapShadowMemory(ShadowMemoryAddress, sizeof(FShadowClockBankSlot));
 
-			TSafeMap<UPTRINT, FSyncObjectBank> OldMap;
+		return (FShadowClockBankSlot*)ShadowMemoryAddress;
+	}
+
+	INSTRUMENTATION_FUNCTION_ATTRIBUTES FSyncObjectRef GetSyncObject(void* SyncAddr)
+	{
+		FShadowClockBankSlot* ClockBankSlot = GetShadowClockBank((UPTRINT)SyncAddr);
+
+		if (ClockBankSlot->SyncObjectBank == nullptr)
+		{
+			FSyncObjectBank* SyncObjectBank = new FSyncObjectBank();
+			if (FPlatformAtomics::InterlockedCompareExchangePointer((void**)&ClockBankSlot->SyncObjectBank, SyncObjectBank, nullptr) != nullptr)
+			{
+				// Should not happen very often, and at most once per thread taking part of the race
+				delete SyncObjectBank;
+			}
+			else
 			{
-				UE::TScopeLock Scope(SyncBanks.Lock);
-				if (SyncBanks.Map.Num())
+				SyncObjectBank->AddRef();
+
+				// Publish the object in the linked-list
+				while (true)
 				{
-					return true;
+					FSyncObjectBank* LocalSyncBankHead = (FSyncObjectBank*)SyncObjectBankHead;
+					SyncObjectBank->Next = LocalSyncBankHead;
+
+					if (FPlatformAtomics::InterlockedCompareExchangePointer((void**)&SyncObjectBankHead, SyncObjectBank, LocalSyncBankHead) == LocalSyncBankHead)
+					{
+						break;
+					}
 				}
 			}
 		}
 
-		return false;
+		UPTRINT Index = UPTRINT(SyncAddr) & 7;
+		return FSyncObjectRef(ClockBankSlot->SyncObjectBank, ClockBankSlot->SyncObjectBank->GetSyncObject(Index));
+	}
+
+	INSTRUMENTATION_FUNCTION_ATTRIBUTES bool HasInstrumentedAtomics()
+	{
+		FInstrumentationScope InstrumentationScope;
+		// Very handy way of knowing if we still have any clock banks left.
+		// This will make sure all the refcounting works and that we can free everything properly.
+		return GetSyncObjectMemoryUsage() != 0;
 	}
 
 	INSTRUMENTATION_FUNCTION_ATTRIBUTES void ResetInstrumentedAtomics()
 	{
 		FInstrumentationScope InstrumentationScope;
-		for (uint32 Index = 0; Index < STRIPING_SIZE; ++Index)
+		FSyncObjectBank* ClockBank = (FSyncObjectBank*)FPlatformAtomics::InterlockedExchangePtr((void**)&SyncObjectBankHead, nullptr);
+		while (ClockBank)
 		{
-			FSyncBanks& SyncBanks = StripedSyncBanks[Index];
-
-			TSafeMap<UPTRINT, FSyncObjectBank> OldMap;
-			{
-				UE::TScopeLock Scope(SyncBanks.Lock);
-				OldMap = MoveTemp(SyncBanks.Map);
-			}
+			FSyncObjectBank* ToDelete = ClockBank;
+			ClockBank = ClockBank->Next;
+			ToDelete->Release();
 		}
 	}
-	
+
 	INSTRUMENTATION_FUNCTION_ATTRIBUTES void FreeMemoryRange(void* Ptr, uint64 Size)
 	{
 		if (!bRaceDetectorActive || bIsResettingShadow)
@@ -421,23 +434,6 @@ namespace UE::Sanitizer::RaceDetector {
 			{
 				FShadowMemory* Shadow = (FShadowMemory*)ShadowAddr;
 
-				// Best effort to clear sync objects if we know some were attached.
-				// If we miss any, they'll be cleaned up on the next shadow reset anyway.
-				const bool bHasSyncObject =
-					Shadow->Accesses[0].bHasSyncObjectAttached ||
-					Shadow->Accesses[1].bHasSyncObjectAttached ||
-					Shadow->Accesses[2].bHasSyncObjectAttached ||
-					Shadow->Accesses[3].bHasSyncObjectAttached;
-
-				if (bHasSyncObject)
-				{
-					FSyncBanks& SyncBanks = StripedSyncBanks[BankAddr & STRIPING_MASK];
-					{
-						UE::TScopeLock ScopeLock(SyncBanks.Lock);
-						SyncBanks.Map.Remove(BankAddr);
-					}
-				}
-
 				Shadow->Accesses[0].RawValue = 0;
 				Shadow->Accesses[1].RawValue = 0;
 				Shadow->Accesses[2].RawValue = 0;
@@ -460,7 +456,7 @@ namespace UE::Sanitizer::RaceDetector {
 
 	INSTRUMENTATION_FUNCTION_ATTRIBUTES void WriteToLog(FString&& Msg)
 	{
-		UE::TScopeLock Lock(RaceReportsLock);
+		TWriteScopeLock Lock(RaceReportsLock);
 		RaceReports.Emplace(MoveTemp(Msg));
 	}
 
@@ -536,7 +532,7 @@ namespace UE::Sanitizer::RaceDetector {
 				uint64 LocationHash = FirstLocationLastFrame ^ SecondLocationLastFrame;
 				bool bAlreadyFound = false;
 				{
-					UE::TUniqueLock Lock(RaceHashesLock);
+					TWriteScopeLock Lock(RaceHashesLock);
 					RaceHashes.FindOrAdd(LocationHash, &bAlreadyFound);
 				}
 				if (bAlreadyFound)
@@ -673,7 +669,7 @@ namespace UE::Sanitizer::RaceDetector {
 		RaceCallbackFn.Reset();
 	}
 
-	INSTRUMENTATION_FUNCTION_ATTRIBUTES FORCEINLINE void InstrumentAccess(FContext& Context, void* Ptr, uint8 Size, EMemoryAccessType AccessType, FAtomicMemoryOrder Order, void* ReturnAddress, const TCHAR* OpName, bool bHasSyncObjectAttached, bool& bHasAddedStack)
+	INSTRUMENTATION_FUNCTION_ATTRIBUTES FORCEINLINE void InstrumentAccess(FContext& Context, void* Ptr, uint8 Size, EMemoryAccessType AccessType, FAtomicMemoryOrder Order, void* ReturnAddress, const TCHAR* OpName, bool& bHasAddedStack)
 	{
 #if WITH_RACEDETECTOR_CHECK
 		// Validate size and alignment since it should have been done before entering this function.
@@ -700,7 +696,7 @@ namespace UE::Sanitizer::RaceDetector {
 
 		const uint32 Epoch = GlobalEpoch;
 		const uint8 Offset = GetOffsetInBytes(Ptr);
-		const FMemoryAccess CurrentAccess(Context.ContextId, Context.CurrentClock(), Offset, Size, AccessType, bHasSyncObjectAttached);
+		const FMemoryAccess CurrentAccess(Context.ContextId, Context.CurrentClock(), Offset, Size, AccessType);
 
 		// This is very practical to scope small section of code that we want to understand
 		// exactly what's going on inside.
@@ -848,12 +844,6 @@ namespace UE::Sanitizer::RaceDetector {
 		}
 	}
 
-	INSTRUMENTATION_FUNCTION_ATTRIBUTES FORCEINLINE void InstrumentAccess(FContext& Context, void* Ptr, uint8 Size, EMemoryAccessType AccessType, FAtomicMemoryOrder Order, void* ReturnAddress, const TCHAR* OpName, bool& bHasAddedCallstackFrame)
-	{
-		const bool bHasSyncObjectAttached = !IsAtomicOrderRelaxed(Order);
-		InstrumentAccess(Context, Ptr, Size, AccessType, Order, ReturnAddress, OpName, bHasSyncObjectAttached, bHasAddedCallstackFrame);
-	}
-
 	// Handles alignment and sizes bigger than a shadow cell
 	INSTRUMENTATION_FUNCTION_ATTRIBUTES FORCEINLINE void InstrumentAccessPreamble(FContext& Context, UPTRINT Ptr, uint32 Size, EMemoryAccessType AccessType, FAtomicMemoryOrder Order, void* ReturnAddress, const TCHAR* OpName, bool& bHasAddedCallstackFrame)
 	{
@@ -878,7 +868,8 @@ namespace UE::Sanitizer::RaceDetector {
 		}
 	}
 
-	INSTRUMENTATION_FUNCTION_ATTRIBUTES void InstrumentAtomicAccess(FContext& Context, void* AtomicAddr, uint8 Size, EMemoryAccessType AccessType, FAtomicMemoryOrder Order, const TFunctionRef<void()>& AtomicOp, void* ReturnAddress, const TCHAR* OpName, bool& bHasAddedCallstackFrame)
+	template <typename AtomicOpType>
+	INSTRUMENTATION_FUNCTION_ATTRIBUTES void InstrumentAtomicAccess(FContext& Context, void* AtomicAddr, uint8 Size, EMemoryAccessType AccessType, FAtomicMemoryOrder Order, AtomicOpType&& AtomicOp, void* ReturnAddress, const TCHAR* OpName, bool& bHasAddedCallstackFrame)
 	{
 #if WITH_RACEDETECTOR_CHECK
 		check(FMath::Min((UPTRINT)Size, GetAlignedAddress((UPTRINT)AtomicAddr) + 8 - (UPTRINT)AtomicAddr) == (UPTRINT)Size);
@@ -890,7 +881,7 @@ namespace UE::Sanitizer::RaceDetector {
 			return;
 		}
 
-		TRefCountPtr<FSyncObject> Atomic = GetSyncObject(AtomicAddr);
+		FSyncObjectRef Atomic = GetSyncObject(AtomicAddr);
 
 		if (AccessType == ACCESS_TYPE_ATOMIC_READ_WRITE && IsAtomicOrderAcquireRelease(Order))
 		{
@@ -912,7 +903,8 @@ namespace UE::Sanitizer::RaceDetector {
 		Context.IncrementClock();
 	}
 
-	INSTRUMENTATION_FUNCTION_ATTRIBUTES void InstrumentAtomicAccess(FContext& Context, void* AtomicAddr, uint8 Size, EMemoryAccessType AccessType, FAtomicMemoryOrder SuccessOrder, FAtomicMemoryOrder FailureOrder, const TFunctionRef<bool()>& AtomicOp, void* ReturnAddress, const TCHAR* OpName, bool& bHasAddedCallstackFrame)
+	template <typename AtomicOpType>
+	INSTRUMENTATION_FUNCTION_ATTRIBUTES void InstrumentAtomicAccess(FContext& Context, void* AtomicAddr, uint8 Size, EMemoryAccessType AccessType, FAtomicMemoryOrder SuccessOrder, FAtomicMemoryOrder FailureOrder, AtomicOpType&& AtomicOp, void* ReturnAddress, const TCHAR* OpName, bool& bHasAddedCallstackFrame)
 	{
 #if WITH_RACEDETECTOR_CHECK
 		check(FMath::Min((UPTRINT)Size, GetAlignedAddress((UPTRINT)AtomicAddr) + 8 - (UPTRINT)AtomicAddr) == (UPTRINT)Size);
@@ -929,8 +921,7 @@ namespace UE::Sanitizer::RaceDetector {
 			GetSyncObject(AtomicAddr)->SyncWithFailureSupport(Context, AtomicOp, AccessType, SuccessOrder, FailureOrder, ReturnAddress, AtomicAddr, OpName,
 				[&Context, AtomicAddr, Size, AccessType, ReturnAddress, &SuccessOrder, OpName, &bHasAddedCallstackFrame](FAtomicMemoryOrder ActualOrder) INSTRUMENTATION_FUNCTION_ATTRIBUTES
 				{
-					const bool bHasSyncObjectAttached = true;
-					InstrumentAccess(Context, AtomicAddr, Size, AccessType, ActualOrder, ReturnAddress, OpName, bHasSyncObjectAttached, bHasAddedCallstackFrame);
+					InstrumentAccess(Context, AtomicAddr, Size, AccessType, ActualOrder, ReturnAddress, OpName, bHasAddedCallstackFrame);
 				}
 			);
 		}
@@ -1037,7 +1028,7 @@ namespace UE::Sanitizer::RaceDetector {
 	{
 		TArray<FString> LocalReports;
 		{
-			UE::TScopeLock Lock(RaceReportsLock);
+			TWriteScopeLock Lock(RaceReportsLock);
 			LocalReports = MoveTemp(RaceReports);
 		}
 
@@ -1350,6 +1341,7 @@ namespace UE::Sanitizer::RaceDetector {
 		ShadowMemoryBase = Platform::GetShadowMemoryBase();
 		ShadowMemorySize = Platform::GetShadowMemorySize();
 		ShadowMemoryEnd = ShadowMemoryBase + ShadowMemorySize;
+		ShadowClockBase = Platform::GetShadowClockBase();
 
 		InitMemoryAllocator();
 		
@@ -1541,7 +1533,7 @@ extern "C" {
 		bool bHasAddedCallstackFrame = true;
 		InstrumentAccess(Context, addr, sizeof(UPTRINT), ACCESS_TYPE_ATOMIC_READ, FAtomicMemoryOrder::MEMORY_ORDER_RELEASE, _ReturnAddress(), TEXT("AnnotateHappensBefore"), bHasAddedCallstackFrame);
 
-		TRefCountPtr<FSyncObject> Atomic = GetSyncObject(addr);
+		FSyncObjectRef Atomic = GetSyncObject(addr);
 		Atomic->SyncRelease(Context, []() {}, _ReturnAddress(), addr, TEXT("AnnotateHappensBefore"));
 		Context.IncrementClock();
 		RemoveCallstackFrame(Context);
@@ -1557,7 +1549,7 @@ extern "C" {
 		bool bHasAddedCallstackFrame = true;
 		InstrumentAccess(Context, addr, sizeof(UPTRINT), ACCESS_TYPE_ATOMIC_READ, FAtomicMemoryOrder::MEMORY_ORDER_ACQUIRE, _ReturnAddress(), TEXT("AnnotateHappensBefore"), bHasAddedCallstackFrame);
 
-		TRefCountPtr<FSyncObject> Atomic = GetSyncObject(addr);
+		FSyncObjectRef Atomic = GetSyncObject(addr);
 		Atomic->SyncAcquire(Context, []() {}, _ReturnAddress(), addr, TEXT("AnnotateHappensAfter"));
 		Context.IncrementClock();
 		RemoveCallstackFrame(Context);
diff --git a/Engine/Source/Runtime/Core/Private/Sanitizer/RaceDetectorWindows.cpp b/Engine/Source/Runtime/Core/Private/Sanitizer/RaceDetectorWindows.cpp
--- a/Engine/Source/Runtime/Core/Private/Sanitizer/RaceDetectorWindows.cpp
+++ b/Engine/Source/Runtime/Core/Private/Sanitizer/RaceDetectorWindows.cpp
@@ -117,7 +117,7 @@ INSTRUMENTATION_FUNCTION_ATTRIBUTES void InstrumentEndThread()
 	START_WINAPI_INSTRUMENTATION(ReleaseCurrentContext(); return);
 	if (Handle)
 	{
-		TRefCountPtr<FSyncObject> Sync = GetSyncObject(Handle);
+		FSyncObjectRef Sync = GetSyncObject(Handle);
 		Sync->SyncReleaseAsSoleOwner(Context, _ReturnAddress(), Handle, TEXT("InstrumentEndThread"));
 		Context.IncrementClock();
 	}
@@ -259,7 +259,7 @@ INSTRUMENTATION_FUNCTION_ATTRIBUTES void AcquireWaitHandle(HANDLE Handle, void*
 {
 	FContext& Context = EnsureCurrentContext();
 
-	TRefCountPtr<FSyncObject> Sync = GetSyncObject(Handle);
+	FSyncObjectRef Sync = GetSyncObject(Handle);
 	Sync->SyncAcquire(Context, []() {}, ReturnAddress, Handle, OpName);
 	Context.IncrementClock();
 }
@@ -306,7 +306,7 @@ INSTRUMENTATION_FUNCTION_ATTRIBUTES BOOL WINAPI DetouredSetEvent(HANDLE hHandle)
 
 	// The scope is important to make sure the TRefCountPtr is destroyed before we FINISH_WINAPI_INSTRUMENTATION.
 	{
-		TRefCountPtr<FSyncObject> Sync = GetSyncObject(hHandle);
+		FSyncObjectRef Sync = GetSyncObject(hHandle);
 		Sync->SyncRelease(Context, []() {}, _ReturnAddress(), hHandle, TEXT("SetEvent"));
 		Context.IncrementClock();
 	}
@@ -342,7 +342,7 @@ INSTRUMENTATION_FUNCTION_ATTRIBUTES BOOL WINAPI DetouredWaitOnAddress(volatile V
 
 		if (bIsValueChanged)
 		{
-			TRefCountPtr<FSyncObject> Sync = GetSyncObject((void*)Address);
+			FSyncObjectRef Sync = GetSyncObject((void*)Address);
 			Sync->SyncAcquire(Context, []() {}, _ReturnAddress(), (void*)Address, TEXT("WaitOnAddress"));
 			Context.IncrementClock();
 		}
@@ -359,7 +359,7 @@ INSTRUMENTATION_FUNCTION_ATTRIBUTES void WINAPI DetouredWakeByAddressSingle(PVOI
 
 	{
 		FLastErrorPreservationScope LastErrorScope;
-		TRefCountPtr<FSyncObject> Sync = GetSyncObject(Address);
+		FSyncObjectRef Sync = GetSyncObject(Address);
 		Sync->SyncRelease(Context, []() {}, _ReturnAddress(), Address, TEXT("WakeByAddressSingle"));
 		Context.IncrementClock();
 	}
@@ -381,7 +381,7 @@ INSTRUMENTATION_FUNCTION_ATTRIBUTES BOOL WINAPI DetouredTryEnterCriticalSection(
 	// If we entered the critical section, we are now the sole owner.
 	if (Result)
 	{
-		TRefCountPtr<FSyncObject> Sync = GetSyncObject(lpCriticalSection);
+		FSyncObjectRef Sync = GetSyncObject(lpCriticalSection);
 		Sync->SyncAcquireAsSoleOwnerOrReadOwner(Context, _ReturnAddress(), lpCriticalSection, TEXT("TryEnterCriticalSection"));
 		Context.IncrementClock();
 	}
@@ -401,7 +401,7 @@ INSTRUMENTATION_FUNCTION_ATTRIBUTES void WINAPI DetouredEnterCriticalSection(CRI
 	// The scope is important to make sure the TRefCountPtr is destroyed before we FINISH_WINAPI_INSTRUMENTATION.
 	{
 		// We are now the sole owner.
-		TRefCountPtr<FSyncObject> Sync = GetSyncObject(lpCriticalSection);
+		FSyncObjectRef Sync = GetSyncObject(lpCriticalSection);
 		Sync->SyncAcquireAsSoleOwnerOrReadOwner(Context, _ReturnAddress(), lpCriticalSection, TEXT("EnterCriticalSection"));
 		Context.IncrementClock();
 	}
@@ -418,7 +418,7 @@ INSTRUMENTATION_FUNCTION_ATTRIBUTES void WINAPI DetouredLeaveCriticalSection(CRI
 		FLastErrorPreservationScope LastErrorScope;
 		// Assume this is called with the right semantics, i.e. that we are actual
 		// owners of the critical section.
-		TRefCountPtr<FSyncObject> Sync = GetSyncObject(lpCriticalSection);
+		FSyncObjectRef Sync = GetSyncObject(lpCriticalSection);
 		Sync->SyncReleaseAsSoleOwner(Context, _ReturnAddress(), lpCriticalSection, TEXT("LeaveCriticalSection"));
 		Context.IncrementClock();
 	}
@@ -438,7 +438,7 @@ INSTRUMENTATION_FUNCTION_ATTRIBUTES void WINAPI DetouredAcquireSRWLockShared(PSR
 
 	// The scope is important to make sure the TRefCountPtr is destroyed before we FINISH_WINAPI_INSTRUMENTATION.
 	{
-		TRefCountPtr<FSyncObject> Sync = GetSyncObject(SRWLock);
+		FSyncObjectRef Sync = GetSyncObject(SRWLock);
 		Sync->SyncAcquireAsSoleOwnerOrReadOwner(Context, _ReturnAddress(), SRWLock, TEXT("AcquireSRWLockShared"));
 		Context.IncrementClock();
 	}
@@ -456,7 +456,7 @@ INSTRUMENTATION_FUNCTION_ATTRIBUTES BOOLEAN WINAPI DetouredTryAcquireSRWLockShar
 	FLastErrorPreservationScope LastErrorScope;
 	if (Result)
 	{
-		TRefCountPtr<FSyncObject> Sync = GetSyncObject(SRWLock);
+		FSyncObjectRef Sync = GetSyncObject(SRWLock);
 		Sync->SyncAcquireAsSoleOwnerOrReadOwner(Context, _ReturnAddress(), SRWLock, TEXT("TryAcquireSRWLockShared"));
 		Context.IncrementClock();
 	}
@@ -472,7 +472,7 @@ INSTRUMENTATION_FUNCTION_ATTRIBUTES void WINAPI DetouredReleaseSRWLockShared(PSR
 
 	{
 		FLastErrorPreservationScope LastErrorScope;
-		TRefCountPtr<FSyncObject> Sync = GetSyncObject(SRWLock);
+		FSyncObjectRef Sync = GetSyncObject(SRWLock);
 		Sync->SyncRelease(Context, []() {}, _ReturnAddress(), SRWLock, TEXT("ReleaseSRWLockShared"));
 		Context.IncrementClock();
 	}
@@ -492,7 +492,7 @@ INSTRUMENTATION_FUNCTION_ATTRIBUTES void WINAPI DetouredAcquireSRWLockExclusive(
 
 	// The scope is important to make sure the TRefCountPtr is destroyed before we FINISH_WINAPI_INSTRUMENTATION.
 	{
-		TRefCountPtr<FSyncObject> Sync = GetSyncObject(SRWLock);
+		FSyncObjectRef Sync = GetSyncObject(SRWLock);
 		Sync->SyncAcquireAsSoleOwnerOrReadOwner(Context, _ReturnAddress(), SRWLock, TEXT("AcquireSRWLockExclusive"));
 		Context.IncrementClock();
 	}
@@ -510,7 +510,7 @@ INSTRUMENTATION_FUNCTION_ATTRIBUTES BOOLEAN WINAPI DetouredTryAcquireSRWLockExcl
 
 	if (Result)
 	{
-		TRefCountPtr<FSyncObject> Sync = GetSyncObject(SRWLock);
+		FSyncObjectRef Sync = GetSyncObject(SRWLock);
 		Sync->SyncAcquireAsSoleOwnerOrReadOwner(Context, _ReturnAddress(), SRWLock, TEXT("TryAcquireSRWLockExclusive"));
 		Context.IncrementClock();
 	}
@@ -526,7 +526,7 @@ INSTRUMENTATION_FUNCTION_ATTRIBUTES void WINAPI DetouredReleaseSRWLockExclusive(
 
 	{
 		FLastErrorPreservationScope LastErrorScope;
-		TRefCountPtr<FSyncObject> Sync = GetSyncObject(SRWLock);
+		FSyncObjectRef Sync = GetSyncObject(SRWLock);
 		Sync->SyncReleaseAsSoleOwner(Context, _ReturnAddress(), SRWLock, TEXT("ReleaseSRWLockExclusive"));
 		Context.IncrementClock();
 	}
@@ -672,6 +672,7 @@ namespace UE::Sanitizer::RaceDetector::Platform {
 
 	UPTRINT ShadowBase = 0;
 	UPTRINT ShadowSize = 0;
+	UPTRINT ShadowClockBase = 0;
 	UPTRINT ShadowEnd = 0;
 	UPTRINT ShadowBitmapBase = 0;
 	UPTRINT ShadowBitmapSize = 0;
@@ -907,24 +908,27 @@ namespace UE::Sanitizer::RaceDetector::Platform {
 		// -----------
 		// High Application Address Space
 		// -----------
-		// We need to split the total address space in a way that we always have
-		// 4 times the shadow address space as the total address space usable by the application.
 		UPTRINT TotalAddressSpace = (UPTRINT)si.lpMaximumApplicationAddress - (UPTRINT)si.lpMinimumApplicationAddress;
 
-		// Reserve 7/8 of the space for the shadow so that we leave the application with 1/8 of the total address space
-		ShadowSize = Align(7 * (TotalAddressSpace >> 3), PageSize);
+		// We need 64 bytes of shadow for each 8 bytes of real application data. 8:1
+		// We need a pointer to a clock bank (8 bytes) for each 8 bytes of application data 1:1
+		// Reserve 9/10 of the space for the shadow so that we leave the application with 1/10 of the total address space
+		ShadowSize = Align(9 * (TotalAddressSpace / 10), PageSize);
+
+		// The 8:1 space so that we can compute the clock bank base address
+		UPTRINT ShadowAccessSize = Align(8 * (TotalAddressSpace / 10), PageSize);
 
 		// Because page faults are super slow when attached with a debugger
 		// we're going to use a bitmap view where each 4KB page of the shadow bitmap
 		// is a bit. When that bit is 1, the page has already been committed so we
 		// don't need to call VirtualAlloc for that page.
-		// The required memory for this scheme is a maximum of 4GB support 128TB of address space.
+		// The required memory for this scheme is a maximum of 4GB to support 128TB of address space.
 		ShadowBitmapSize = Align((ShadowSize / PageSize) >> 3, PageSize);
 
 		// Because zeroing 4GB is very slow when unmapping shadow.
 		// We keep another 128KB of bits to know which page in the ShadowBitmap have been committed
 		// so we just have to zero those page instead of uncommitting the whole 4GB and recommitting it.
-		DirtyShadowBitmapSize = (ShadowBitmapSize / PageSize) >> 3;
+		DirtyShadowBitmapSize = Align((ShadowBitmapSize / PageSize) >> 3, PageSize);
 
 		// Let VirtualAlloc decide the best region to reserve. This works around a Windows 10 bug where specifying 
 		// a base address when reserving large regions can cause extreme system-wide performance degradation.
@@ -946,6 +950,7 @@ namespace UE::Sanitizer::RaceDetector::Platform {
 		// Put the real shadow memory after the bitmaps.
 		ShadowBase = DirtyShadowBitmapEnd;
 		ShadowEnd = ShadowBase + ShadowSize;
+		ShadowClockBase = ShadowBase + ShadowAccessSize;
 	}
 
 	INSTRUMENTATION_FUNCTION_ATTRIBUTES void DirtyShadowBitmapPage(UPTRINT ShadowBitmapAddress)
@@ -1144,6 +1149,11 @@ namespace UE::Sanitizer::RaceDetector::Platform {
 		return ShadowSize;
 	}
 
+	INSTRUMENTATION_FUNCTION_ATTRIBUTES UPTRINT GetShadowClockBase()
+	{
+		return ShadowClockBase;
+	}
+
 	INSTRUMENTATION_FUNCTION_ATTRIBUTES uint64 GetShadowMemoryUsage()
 	{
 		uint64 ShadowPageCount = 0;
