WindowsPlatformFile: More feedback in error cases: When we return from WriteFile because it succeeded but wrote less than expected, call SetLastError to report the error condition to FArchiveFileWriterGeneric
#rnx

diff --git a/Engine/Source/Runtime/Core/Private/Windows/WindowsPlatformFile.cpp b/Engine/Source/Runtime/Core/Private/Windows/WindowsPlatformFile.cpp
--- a/Engine/Source/Runtime/Core/Private/Windows/WindowsPlatformFile.cpp
+++ b/Engine/Source/Runtime/Core/Private/Windows/WindowsPlatformFile.cpp
@@ -866,6 +866,9 @@ public:
 			// Early out as a failure case if we didn't write all of the data we expected
 			if (BytesToWrite32 != NumWritten)
 			{
+				// Pick an error we can use to report to the caller why we failed. We arbitrarily use
+				// ERROR_CAN_NOT_COMPLETE, since it seems like one that WriteFile does not normally return.
+				SetLastError(ERROR_CAN_NOT_COMPLETE);
 				TRACE_PLATFORMFILE_END_WRITE(this, TotalNumWritten);
 				return false;
 			}
