Added FPlatformRecursiveMutex and FPlatformSharedMutex

Made FCriticalSection an alias for FPlatformRecursiveMutex.
Made FRWLock an alias for FPlatformSharedMutex.

The new names reflect that it is a platform-provided mutex as opposed to a mutex provided by UE.

diff --git a/Engine/Source/Runtime/Core/Private/GenericPlatform/GenericPlatformProcess.cpp b/Engine/Source/Runtime/Core/Private/GenericPlatform/GenericPlatformProcess.cpp
--- a/Engine/Source/Runtime/Core/Private/GenericPlatform/GenericPlatformProcess.cpp
+++ b/Engine/Source/Runtime/Core/Private/GenericPlatform/GenericPlatformProcess.cpp
@@ -4,7 +4,7 @@
 #include "Misc/Timespan.h"
 #include "HAL/PlatformProperties.h"
 #include "HAL/PlatformProcess.h"
-#include "GenericPlatform/GenericPlatformCriticalSection.h"
+#include "GenericPlatform/GenericPlatformMutex.h"
 #include "Logging/LogMacros.h"
 #include "CoreGlobals.h"
 #include "HAL/FileManager.h"
@@ -690,9 +690,9 @@ PRAGMA_DISABLE_DEPRECATION_WARNINGS
 PRAGMA_ENABLE_DEPRECATION_WARNINGS
 }
 
-FSystemWideCriticalSectionNotImplemented::FSystemWideCriticalSectionNotImplemented(const FString& Name, FTimespan Timeout)
+UE::FPlatformSystemWideMutexNotImplemented::FPlatformSystemWideMutexNotImplemented(const FString& Name, FTimespan Timeout)
 {
-	UE_LOG(LogHAL, Fatal, TEXT("FSystemWideCriticalSection not implemented on this platform"));
+	UE_LOG(LogHAL, Fatal, TEXT("FPlatformSystemWideMutex not implemented on this platform."));
 }
 
 void FGenericPlatformProcess::TearDown()
diff --git a/Engine/Source/Runtime/Core/Private/Mac/MacCriticalSection.cpp b/Engine/Source/Runtime/Core/Private/Mac/MacCriticalSection.cpp
deleted file mode 100644
diff --git a/Engine/Source/Runtime/Core/Private/Mac/MacPlatformMutex.cpp b/Engine/Source/Runtime/Core/Private/Mac/MacPlatformMutex.cpp
new file mode 100644
--- /dev/null
+++ b/Engine/Source/Runtime/Core/Private/Mac/MacPlatformMutex.cpp
@@ -0,0 +1,79 @@
+// Copyright Epic Games, Inc. All Rights Reserved.
+
+#include "Mac/MacPlatformMutex.h"
+
+#include "HAL/PlatformTime.h"
+#include "Misc/App.h"
+#include "Misc/DateTime.h"
+#include "Apple/PreAppleSystemHeaders.h"
+#include <mach-o/dyld.h>
+#include <libproc.h>
+#include "Apple/PostAppleSystemHeaders.h"
+
+namespace UE
+{
+
+FMacSystemWideMutex::FMacSystemWideMutex(const FString& InName, FTimespan InTimeout)
+{
+	check(InName.Len() > 0);
+	check(InTimeout >= FTimespan::Zero());
+	check(InTimeout.GetTotalSeconds() < (double)FLT_MAX);
+
+	const FString LockPath = FString(FPlatformProcess::ApplicationSettingsDir()) / InName;
+	FString NormalizedFilepath(LockPath);
+	NormalizedFilepath.ReplaceInline(TEXT("\\"), TEXT("/"));
+	FileHandle = -1;
+
+	double ExpireTimeSecs = FPlatformTime::Seconds() + InTimeout.GetTotalSeconds();
+	while (true)
+	{
+		if (FileHandle == -1)
+		{
+			// Try to open the file.
+			FileHandle = open(TCHAR_TO_UTF8(*NormalizedFilepath), O_CREAT | O_WRONLY | O_NONBLOCK, S_IRUSR | S_IWUSR | S_IRGRP | S_IWGRP | S_IROTH | S_IWOTH);
+		}
+
+		// If the file is open, try to lock it, but don't block. If the file is already locked by another process or another thread, blocking here may not honor InTimeout.
+		if (FileHandle != -1 && flock(FileHandle, LOCK_EX | LOCK_NB) == 0)
+		{
+			return; // Lock was successfully taken.
+		}
+
+		// If the lock isn't acquired and no time is left to retry, clean up and set the state as 'invalid'
+		if (InTimeout == FTimespan::Zero() || FPlatformTime::Seconds() > ExpireTimeSecs)
+		{
+			if (FileHandle != -1)
+			{
+				close(FileHandle);
+				FileHandle = -1;
+			}
+			return; // Lock wasn't acquired within the allowed time.
+		}
+
+		// Either the file did not open or the lock wasn't acquired, retry.
+		const float RetrySeconds = FMath::Min((float)InTimeout.GetTotalSeconds(), 0.25f);
+		FPlatformProcess::Sleep(RetrySeconds);
+	}
+}
+
+FMacSystemWideMutex::~FMacSystemWideMutex()
+{
+	Release();
+}
+
+bool FMacSystemWideMutex::IsValid() const
+{
+	return FileHandle != -1;
+}
+
+void FMacSystemWideMutex::Release()
+{
+	if (IsValid())
+	{
+		flock(FileHandle, LOCK_UN);
+		close(FileHandle);
+		FileHandle = -1;
+	}
+}
+
+} // UE
diff --git a/Engine/Source/Runtime/Core/Private/Unix/UnixCriticalSection.cpp b/Engine/Source/Runtime/Core/Private/Unix/UnixCriticalSection.cpp
deleted file mode 100644
diff --git a/Engine/Source/Runtime/Core/Private/Unix/UnixPlatformMutex.cpp b/Engine/Source/Runtime/Core/Private/Unix/UnixPlatformMutex.cpp
new file mode 100644
--- /dev/null
+++ b/Engine/Source/Runtime/Core/Private/Unix/UnixPlatformMutex.cpp
@@ -0,0 +1,83 @@
+// Copyright Epic Games, Inc. All Rights Reserved.
+
+#include "Unix/UnixPlatformMutex.h"
+
+#include "Containers/StringConv.h"
+#include "HAL/PlatformProcess.h"
+#include "HAL/PlatformTime.h"
+#include "Math/UnrealMathUtility.h"
+#include "Misc/AssertionMacros.h"
+
+#include <fcntl.h>
+#include <float.h>
+#include <sys/file.h>
+#include <unistd.h>
+
+namespace UE
+{
+
+FUnixSystemWideMutex::FUnixSystemWideMutex(const FString& InName, FTimespan InTimeout)
+{
+	check(InName.Len() > 0);
+	check(InTimeout >= FTimespan::Zero());
+	check(InTimeout.GetTotalSeconds() < (double)FLT_MAX);
+
+	const FString LockPath = FString(FPlatformProcess::ApplicationSettingsDir()) / InName;
+	FString NormalizedFilepath(LockPath);
+	NormalizedFilepath.ReplaceInline(TEXT("\\"), TEXT("/"));
+	FileHandle = -1;
+
+	double ExpireTimeSecs = FPlatformTime::Seconds() + InTimeout.GetTotalSeconds();
+	while (true)
+	{
+		if (FileHandle == -1)
+		{
+			// Try to open the file.
+			FileHandle = open(TCHAR_TO_UTF8(*NormalizedFilepath), O_CREAT | O_WRONLY | O_NONBLOCK, S_IRUSR | S_IWUSR | S_IRGRP | S_IWGRP | S_IROTH | S_IWOTH);
+		}
+
+		// If the file is open, try to lock it, but don't block. If the file is already locked by another process or another thread, blocking here may not honor InTimeout.
+		// NOTE: In old Linux kernels, flock() wasn't always atomic, but in recent ones, that was fixed and flock() is expected to be an atomic operation.
+		if (FileHandle != -1 && flock(FileHandle, LOCK_EX | LOCK_NB) == 0)
+		{
+			return; // Lock was successfully taken.
+		}
+
+		// If the lock isn't acquired and no time is left to retry, clean up and set the state as 'invalid'
+		if (InTimeout == FTimespan::Zero() || FPlatformTime::Seconds() > ExpireTimeSecs)
+		{
+			if (FileHandle != -1)
+			{
+				close(FileHandle);
+				FileHandle = -1;
+			}
+			return; // Lock wasn't acquired within the allowed time.
+		}
+
+		// Either the file did not open or the lock wasn't acquired, retry.
+		const float RetrySeconds = FMath::Min((float)InTimeout.GetTotalSeconds(), 0.25f);
+		FPlatformProcess::Sleep(RetrySeconds);
+	}
+}
+
+FUnixSystemWideMutex::~FUnixSystemWideMutex()
+{
+	Release();
+}
+
+bool FUnixSystemWideMutex::IsValid() const
+{
+	return FileHandle != -1;
+}
+
+void FUnixSystemWideMutex::Release()
+{
+	if (IsValid())
+	{
+		flock(FileHandle, LOCK_UN);
+		close(FileHandle);
+		FileHandle = -1;
+	}
+}
+
+} // UE
diff --git a/Engine/Source/Runtime/Core/Private/Windows/WindowsCriticalSection.cpp b/Engine/Source/Runtime/Core/Private/Windows/WindowsCriticalSection.cpp
deleted file mode 100644
diff --git a/Engine/Source/Runtime/Core/Private/Windows/WindowsPlatformMutex.cpp b/Engine/Source/Runtime/Core/Private/Windows/WindowsPlatformMutex.cpp
new file mode 100644
--- /dev/null
+++ b/Engine/Source/Runtime/Core/Private/Windows/WindowsPlatformMutex.cpp
@@ -0,0 +1,83 @@
+// Copyright Epic Games, Inc. All Rights Reserved.
+
+#include "Windows/WindowsPlatformMutex.h"
+
+#include "Containers/UnrealString.h"
+#include "Math/UnrealMathUtility.h"
+#include "Misc/AssertionMacros.h"
+#include "Windows/AllowWindowsPlatformTypes.h"
+
+namespace UE
+{
+
+FWindowsSystemWideMutex::FWindowsSystemWideMutex(const FString& InName, FTimespan InTimeout)
+{
+	// Check for valid name length not exceeding the Window mutex name length limit
+	check(InName.Len() > 0)
+	check(InName.Len() < MAX_PATH)
+	// Check for valid positive timeouts
+	check(InTimeout >= FTimespan::Zero())
+	check(InTimeout.GetTotalMilliseconds() < (double)0x7FFFFFFE)	// limit timespan to a number of millisec that will fit in a signed int32
+
+	// Disallow backslashes as they aren't allowed in Windows mutex names
+	FString NormalizedMutexName(InName);
+	NormalizedMutexName.ReplaceInline(TEXT("\\"), TEXT("/"));
+
+	TCHAR MutexName[MAX_PATH] = {};
+	FCString::Strcpy(MutexName, *NormalizedMutexName);
+
+	// Attempt to create and take ownership of a named mutex
+	Mutex = CreateMutex(NULL, true, MutexName);
+
+	if (Mutex != NULL && GetLastError() == ERROR_ALREADY_EXISTS)
+	{
+		// CreateMutex returned a valid handle but we didn't get ownership because another process/thread has already created it
+		bool bMutexOwned = false;
+
+		if (InTimeout != FTimespan::Zero())
+		{
+			// We have a handle already so try waiting for it to be released by the current owner
+			DWORD WaitResult = WaitForSingleObject(Mutex, FMath::TruncToInt((float)InTimeout.GetTotalMilliseconds()));
+
+			// WAIT_OBJECT_0 = we got ownership when the previous owner released it
+			// WAIT_ABANDONED = we got ownership when the previous owner exited WITHOUT releasing the mutex gracefully (we own it now but the state of any shared resource could be corrupted!)
+			if (WaitResult == WAIT_ABANDONED || WaitResult == WAIT_OBJECT_0)
+			{
+				bMutexOwned = true;
+			}
+		}
+
+		if (!bMutexOwned)
+		{
+			// We failed to gain ownership by waiting so close the handle to avoid leaking it.
+			CloseHandle(Mutex);
+			Mutex = NULL;
+		}
+	}
+}
+
+FWindowsSystemWideMutex::~FWindowsSystemWideMutex()
+{
+	Release();
+}
+
+bool FWindowsSystemWideMutex::IsValid() const
+{
+	return Mutex != NULL;
+}
+
+void FWindowsSystemWideMutex::Release()
+{
+	if (IsValid())
+	{
+		// Release ownership
+		ReleaseMutex(Mutex);
+		// Also release the handle so it isn't leaked
+		CloseHandle(Mutex);
+		Mutex = NULL;
+	}
+}
+
+} // UE
+
+#include "Windows/HideWindowsPlatformTypes.h"
diff --git a/Engine/Source/Runtime/Core/Public/Android/AndroidCriticalSection.h b/Engine/Source/Runtime/Core/Public/Android/AndroidCriticalSection.h
deleted file mode 100644
diff --git a/Engine/Source/Runtime/Core/Public/Android/AndroidPlatformMutex.h b/Engine/Source/Runtime/Core/Public/Android/AndroidPlatformMutex.h
new file mode 100644
--- /dev/null
+++ b/Engine/Source/Runtime/Core/Public/Android/AndroidPlatformMutex.h
@@ -0,0 +1,16 @@
+// Copyright Epic Games, Inc. All Rights Reserved.
+
+#pragma once
+
+#include "GenericPlatform/GenericPlatformMutex.h"
+#include "HAL/PThreadsRecursiveMutex.h"
+#include "HAL/PThreadsSharedMutex.h"
+
+namespace UE
+{
+
+using FPlatformRecursiveMutex = FPThreadsRecursiveMutex;
+using FPlatformSharedMutex = FPThreadsSharedMutex;
+using FPlatformSystemWideMutex = FPlatformSystemWideMutexNotImplemented;
+
+} // UE
diff --git a/Engine/Source/Runtime/Core/Public/Core.h b/Engine/Source/Runtime/Core/Public/Core.h
--- a/Engine/Source/Runtime/Core/Public/Core.h
+++ b/Engine/Source/Runtime/Core/Public/Core.h
@@ -46,7 +46,7 @@ MONOLITHIC_HEADER_BOILERPLATE()
 #include "HAL/PlatformTime.h"
 #include "GenericPlatform/GenericPlatformProcess.h"
 #include "HAL/PlatformProcess.h"
-#include "GenericPlatform/GenericPlatformCriticalSection.h"
+#include "GenericPlatform/GenericPlatformMutex.h"
 #include "HAL/CriticalSection.h"
 #include "GenericPlatform/GenericPlatformAtomics.h"
 #include "HAL/PlatformAtomics.h"
diff --git a/Engine/Source/Runtime/Core/Public/GenericPlatform/GenericPlatformCriticalSection.h b/Engine/Source/Runtime/Core/Public/GenericPlatform/GenericPlatformCriticalSection.h
deleted file mode 100644
diff --git a/Engine/Source/Runtime/Core/Public/GenericPlatform/GenericPlatformMutex.h b/Engine/Source/Runtime/Core/Public/GenericPlatform/GenericPlatformMutex.h
new file mode 100644
--- /dev/null
+++ b/Engine/Source/Runtime/Core/Public/GenericPlatform/GenericPlatformMutex.h
@@ -0,0 +1,89 @@
+// Copyright Epic Games, Inc. All Rights Reserved.
+
+#pragma once
+
+#include "CoreTypes.h"
+#include "Misc/Timespan.h"
+
+namespace UE
+{
+
+/**
+* TGenericPlatformSharedMutex - Read/Write Mutex
+*	- Provides non-recursive Read/Write (or shared-exclusive) access.
+*	- As a fallback default for non implemented platforms, using a single mutex to provide complete single mutual exclusion - no separate Read/Write access.
+*/
+template <class MutexType>
+class TGenericPlatformSharedMutex
+{
+public:
+	TGenericPlatformSharedMutex(const TGenericPlatformSharedMutex&) = delete;
+	TGenericPlatformSharedMutex& operator=(const TGenericPlatformSharedMutex&) = delete;
+
+	TGenericPlatformSharedMutex() = default;
+	~TGenericPlatformSharedMutex() = default;
+
+	FORCEINLINE bool TryLock()
+	{
+		return Mutex.TryLock();
+	}
+
+	FORCEINLINE void Lock()
+	{
+		Mutex.Lock();
+	}
+
+	FORCEINLINE void Unlock()
+	{
+		Mutex.Unlock();
+	}
+
+	FORCEINLINE bool TryLockShared()
+	{
+		return Mutex.TryLock();
+	}
+
+	FORCEINLINE void LockShared()
+	{
+		Mutex.Lock();
+	}
+
+	FORCEINLINE void UnlockShared()
+	{
+		Mutex.Unlock();
+	}
+
+private:
+	MutexType Mutex;
+};
+
+/** Platforms that don't need a working FPlatformSystemWideMutex can alias this one. */
+class FPlatformSystemWideMutexNotImplemented
+{
+public:
+	FPlatformSystemWideMutexNotImplemented(const FPlatformSystemWideMutexNotImplemented&) = delete;
+	FPlatformSystemWideMutexNotImplemented& operator=(const FPlatformSystemWideMutexNotImplemented&) = delete;
+
+	/** Construct a named, system-wide mutex and attempt to get access/ownership of it. */
+	CORE_API explicit FPlatformSystemWideMutexNotImplemented(const FString& Name, FTimespan Timeout = FTimespan::Zero());
+
+	/** Destructor releases system-wide mutex if it is currently owned. */
+	~FPlatformSystemWideMutexNotImplemented() = default;
+
+	/**
+	 * Does the calling thread have ownership of the system-wide mutex?
+	 *
+	 * @return True if obtained. WARNING: Returns true for an owned but previously abandoned locks so shared resources can be in undetermined states. You must handle shared data robustly.
+	 */
+	bool IsValid() const
+	{
+		return false;
+	}
+
+	/** Releases system-wide mutex if it is currently owned. */
+	void Release()
+	{
+	}
+};
+
+} // UE
diff --git a/Engine/Source/Runtime/Core/Public/HAL/CriticalSection.h b/Engine/Source/Runtime/Core/Public/HAL/CriticalSection.h
--- a/Engine/Source/Runtime/Core/Public/HAL/CriticalSection.h
+++ b/Engine/Source/Runtime/Core/Public/HAL/CriticalSection.h
@@ -2,6 +2,57 @@
 
 #pragma once
 
-#include "CoreTypes.h"
+#include "HAL/PlatformMutex.h"
 
-#include COMPILED_PLATFORM_HEADER(CriticalSection.h)// IWYU pragma: export
+namespace UE
+{
+
+/**
+ * Adapter for FPlatformSharedMutex to the API expected for FRWLock.
+ */
+class FPlatformRWLock : private FPlatformSharedMutex
+{
+public:
+	using FPlatformSharedMutex::FPlatformSharedMutex;
+
+	FORCEINLINE bool TryWriteLock()
+	{
+		return FPlatformSharedMutex::TryLock();
+	}
+
+	FORCEINLINE void WriteLock()
+	{
+		FPlatformSharedMutex::Lock();
+	}
+
+	FORCEINLINE void WriteUnlock()
+	{
+		FPlatformSharedMutex::Unlock();
+	}
+
+	FORCEINLINE bool TryReadLock()
+	{
+		return FPlatformSharedMutex::TryLockShared();
+	}
+
+	FORCEINLINE void ReadLock()
+	{
+		FPlatformSharedMutex::LockShared();
+	}
+
+	FORCEINLINE void ReadUnlock()
+	{
+		FPlatformSharedMutex::UnlockShared();
+	}
+};
+
+} // UE
+
+/** Alias for a mutex that supports recursive locking and may not be fair. */
+using FCriticalSection = UE::FPlatformRecursiveMutex;
+
+/** Alias for a shared mutex that does not support recursive locking and may not be fair. */
+using FRWLock = UE::FPlatformRWLock;
+
+/** Alias for a system-wide (cross-process) mutex that does not support recursive locking and may not be fair. */
+using FSystemWideCriticalSection = UE::FPlatformSystemWideMutex;
diff --git a/Engine/Source/Runtime/Core/Public/HAL/PThreadCriticalSection.h b/Engine/Source/Runtime/Core/Public/HAL/PThreadCriticalSection.h
deleted file mode 100644
diff --git a/Engine/Source/Runtime/Core/Public/HAL/PThreadRWLock.h b/Engine/Source/Runtime/Core/Public/HAL/PThreadRWLock.h
deleted file mode 100644
diff --git a/Engine/Source/Runtime/Core/Public/HAL/PThreadsRecursiveMutex.h b/Engine/Source/Runtime/Core/Public/HAL/PThreadsRecursiveMutex.h
new file mode 100644
--- /dev/null
+++ b/Engine/Source/Runtime/Core/Public/HAL/PThreadsRecursiveMutex.h
@@ -0,0 +1,58 @@
+// Copyright Epic Games, Inc. All Rights Reserved.
+
+#pragma once
+
+// HEADER_UNIT_UNSUPPORTED - Unsupported platform
+
+#include "CoreTypes.h"
+#include <pthread.h>
+#include <errno.h>
+
+namespace UE
+{
+
+/**
+ * A mutex that supports recursive locking.
+ *
+ * Prefer FRecursiveMutex.
+ */
+class FPThreadsRecursiveMutex
+{
+public:
+	FPThreadsRecursiveMutex(const FPThreadsRecursiveMutex&) = delete;
+	FPThreadsRecursiveMutex& operator=(const FPThreadsRecursiveMutex&) = delete;
+
+	FORCEINLINE FPThreadsRecursiveMutex()
+	{
+		pthread_mutexattr_t MutexAttributes;
+		pthread_mutexattr_init(&MutexAttributes);
+		pthread_mutexattr_settype(&MutexAttributes, PTHREAD_MUTEX_RECURSIVE);
+		pthread_mutex_init(&Mutex, &MutexAttributes);
+		pthread_mutexattr_destroy(&MutexAttributes);
+	}
+
+	FORCEINLINE ~FPThreadsRecursiveMutex()
+	{
+		pthread_mutex_destroy(&Mutex);
+	}
+
+	FORCEINLINE bool TryLock()
+	{
+		return pthread_mutex_trylock(&Mutex) == 0;
+	}
+
+	FORCEINLINE void Lock()
+	{
+		pthread_mutex_lock(&Mutex);
+	}
+
+	FORCEINLINE void Unlock()
+	{
+		pthread_mutex_unlock(&Mutex);
+	}
+
+private:
+	pthread_mutex_t Mutex;
+};
+
+} // UE
diff --git a/Engine/Source/Runtime/Core/Public/HAL/PThreadsSharedMutex.h b/Engine/Source/Runtime/Core/Public/HAL/PThreadsSharedMutex.h
new file mode 100644
--- /dev/null
+++ b/Engine/Source/Runtime/Core/Public/HAL/PThreadsSharedMutex.h
@@ -0,0 +1,78 @@
+// Copyright Epic Games, Inc. All Rights Reserved.
+
+#pragma once
+
+// HEADER_UNIT_UNSUPPORTED - Unsupported platform
+
+#include "CoreTypes.h"
+#include "Misc/AssertionMacros.h"
+#include <pthread.h>
+#include <errno.h>
+
+namespace UE
+{
+
+/**
+ * A shared (read/write) mutex that does not support recursive locking.
+ *
+ * Prefer FSharedMutex.
+ */
+class FPThreadsSharedMutex
+{
+public:
+	FPThreadsSharedMutex(const FPThreadsSharedMutex&) = delete;
+	FPThreadsSharedMutex& operator=(const FPThreadsSharedMutex&) = delete;
+
+	FPThreadsSharedMutex()
+	{
+		int Err = pthread_rwlock_init(&Mutex, nullptr);
+		checkf(Err == 0, TEXT("pthread_rwlock_init failed with error: %d"), Err);
+	}
+
+	~FPThreadsSharedMutex()
+	{
+		int Err = pthread_rwlock_destroy(&Mutex);
+		checkf(Err == 0, TEXT("pthread_rwlock_destroy failed with error: %d"), Err);
+	}
+
+	bool TryLock()
+	{
+		int Err = pthread_rwlock_trywrlock(&Mutex);
+		return Err == 0;
+	}
+
+	void Lock()
+	{
+		int Err = pthread_rwlock_wrlock(&Mutex);
+		checkf(Err == 0, TEXT("pthread_rwlock_wrlock failed with error: %d"), Err);
+	}
+
+	void Unlock()
+	{
+		int Err = pthread_rwlock_unlock(&Mutex);
+		checkf(Err == 0, TEXT("pthread_rwlock_unlock failed with error: %d"), Err);
+	}
+
+	bool TryLockShared()
+	{
+		int Err = pthread_rwlock_tryrdlock(&Mutex);
+		return Err == 0;
+	}
+
+	void LockShared()
+	{
+		int Err = pthread_rwlock_rdlock(&Mutex);
+		checkf(Err == 0, TEXT("pthread_rwlock_rdlock failed with error: %d"), Err);
+	}
+
+	void UnlockShared()
+	{
+		int Err = pthread_rwlock_unlock(&Mutex);
+		checkf(Err == 0, TEXT("pthread_rwlock_unlock failed with error: %d"), Err);
+	}
+
+private:
+	pthread_rwlock_t Mutex;
+};
+
+} // UE
diff --git a/Engine/Source/Runtime/Core/Public/HAL/PlatformMutex.h b/Engine/Source/Runtime/Core/Public/HAL/PlatformMutex.h
new file mode 100644
--- /dev/null
+++ b/Engine/Source/Runtime/Core/Public/HAL/PlatformMutex.h
@@ -0,0 +1,7 @@
+// Copyright Epic Games, Inc. All Rights Reserved.
+
+#pragma once
+
+#include "CoreTypes.h"
+
+#include COMPILED_PLATFORM_HEADER(PlatformMutex.h) // IWYU pragma: export
diff --git a/Engine/Source/Runtime/Core/Public/IOS/IOSCriticalSection.h b/Engine/Source/Runtime/Core/Public/IOS/IOSCriticalSection.h
deleted file mode 100644
diff --git a/Engine/Source/Runtime/Core/Public/IOS/IOSPlatformMutex.h b/Engine/Source/Runtime/Core/Public/IOS/IOSPlatformMutex.h
new file mode 100644
--- /dev/null
+++ b/Engine/Source/Runtime/Core/Public/IOS/IOSPlatformMutex.h
@@ -0,0 +1,16 @@
+// Copyright Epic Games, Inc. All Rights Reserved.
+
+#pragma once
+
+#include "GenericPlatform/GenericPlatformMutex.h"
+#include "HAL/PThreadsRecursiveMutex.h"
+#include "HAL/PThreadsSharedMutex.h"
+
+namespace UE
+{
+
+using FPlatformRecursiveMutex = FPThreadsRecursiveMutex;
+using FPlatformSharedMutex = FPThreadsSharedMutex;
+using FPlatformSystemWideMutex = FPlatformSystemWideMutexNotImplemented;
+
+} // UE
diff --git a/Engine/Source/Runtime/Core/Public/Linux/LinuxCriticalSection.h b/Engine/Source/Runtime/Core/Public/Linux/LinuxCriticalSection.h
deleted file mode 100644
diff --git a/Engine/Source/Runtime/Core/Public/Linux/LinuxPlatformMutex.h b/Engine/Source/Runtime/Core/Public/Linux/LinuxPlatformMutex.h
new file mode 100644
--- /dev/null
+++ b/Engine/Source/Runtime/Core/Public/Linux/LinuxPlatformMutex.h
@@ -0,0 +1,5 @@
+// Copyright Epic Games, Inc. All Rights Reserved.
+
+#pragma once
+
+#include "Unix/UnixPlatformMutex.h"
diff --git a/Engine/Source/Runtime/Core/Public/Mac/MacCriticalSection.h b/Engine/Source/Runtime/Core/Public/Mac/MacCriticalSection.h
deleted file mode 100644
diff --git a/Engine/Source/Runtime/Core/Public/Mac/MacPlatformMutex.h b/Engine/Source/Runtime/Core/Public/Mac/MacPlatformMutex.h
new file mode 100644
--- /dev/null
+++ b/Engine/Source/Runtime/Core/Public/Mac/MacPlatformMutex.h
@@ -0,0 +1,44 @@
+// Copyright Epic Games, Inc. All Rights Reserved.
+
+#pragma once
+
+#include "Containers/UnrealString.h"
+#include "HAL/PThreadsRecursiveMutex.h"
+#include "HAL/PThreadsSharedMutex.h"
+#include "Misc/Timespan.h"
+
+namespace UE
+{
+
+/** A system-wide mutex for Mac. Uses exclusive file locking. */
+class FMacSystemWideMutex
+{
+public:
+	FMacSystemWideMutex(const FMacSystemWideMutex&) = delete;
+	FMacSystemWideMutex& operator=(const FMacSystemWideMutex&) = delete;
+
+	/** Construct a named, system-wide mutex and attempt to get access/ownership of it. */
+	CORE_API explicit FMacSystemWideMutex(const FString& InName, FTimespan InTimeout = FTimespan::Zero());
+
+	/** Destructor releases system-wide mutex if it is currently owned. */
+	CORE_API ~FMacSystemWideMutex();
+
+	/**
+	 * Does the calling thread have ownership of the system-wide mutex?
+	 *
+	 * @return True if obtained. WARNING: Returns true for an owned but previously abandoned locks so shared resources can be in undetermined states. You must handle shared data robustly.
+	 */
+	CORE_API bool IsValid() const;
+
+	/** Releases system-wide mutex if it is currently owned. */
+	CORE_API void Release();
+
+private:
+	int32 FileHandle;
+};
+
+using FPlatformRecursiveMutex = FPThreadsRecursiveMutex;
+using FPlatformSharedMutex = FPThreadsSharedMutex;
+using FPlatformSystemWideMutex = FMacSystemWideMutex;
+
+} // UE
diff --git a/Engine/Source/Runtime/Core/Public/Unix/UnixCriticalSection.h b/Engine/Source/Runtime/Core/Public/Unix/UnixCriticalSection.h
deleted file mode 100644
diff --git a/Engine/Source/Runtime/Core/Public/Unix/UnixPlatformMutex.h b/Engine/Source/Runtime/Core/Public/Unix/UnixPlatformMutex.h
new file mode 100644
--- /dev/null
+++ b/Engine/Source/Runtime/Core/Public/Unix/UnixPlatformMutex.h
@@ -0,0 +1,46 @@
+// Copyright Epic Games, Inc. All Rights Reserved.
+
+#pragma once
+
+#include "CoreTypes.h"
+#include "HAL/PThreadsRecursiveMutex.h"
+#include "HAL/PThreadsSharedMutex.h"
+#include "Misc/Timespan.h"
+
+class FString;
+
+namespace UE
+{
+
+/** A system-wide mutex for Unix. Uses exclusive file locking. */
+class FUnixSystemWideMutex
+{
+public:
+	FUnixSystemWideMutex(const FUnixSystemWideMutex&) = delete;
+	FUnixSystemWideMutex& operator=(const FUnixSystemWideMutex&) = delete;
+
+	/** Construct a named, system-wide mutex and attempt to get access/ownership of it. */
+	CORE_API explicit FUnixSystemWideMutex(const FString& InName, FTimespan InTimeout = FTimespan::Zero());
+
+	/** Destructor releases system-wide mutex if it is currently owned. */
+	CORE_API ~FUnixSystemWideMutex();
+
+	/**
+	 * Does the calling thread have ownership of the system-wide mutex?
+	 *
+	 * @return True if obtained. WARNING: Returns true for an owned but previously abandoned locks so shared resources can be in undetermined states. You must handle shared data robustly.
+	 */
+	CORE_API bool IsValid() const;
+
+	/** Releases system-wide mutex if it is currently owned. */
+	CORE_API void Release();
+
+private:
+	int32 FileHandle;
+};
+
+using FPlatformRecursiveMutex = FPThreadsRecursiveMutex;
+using FPlatformSharedMutex = FPThreadsSharedMutex;
+using FPlatformSystemWideMutex = FUnixSystemWideMutex;
+
+} // UE
diff --git a/Engine/Source/Runtime/Core/Public/Windows/WindowsCriticalSection.h b/Engine/Source/Runtime/Core/Public/Windows/WindowsCriticalSection.h
deleted file mode 100644
diff --git a/Engine/Source/Runtime/Core/Public/Windows/WindowsPlatformMutex.h b/Engine/Source/Runtime/Core/Public/Windows/WindowsPlatformMutex.h
new file mode 100644
--- /dev/null
+++ b/Engine/Source/Runtime/Core/Public/Windows/WindowsPlatformMutex.h
@@ -0,0 +1,159 @@
+// Copyright Epic Games, Inc. All Rights Reserved.
+
+#pragma once
+
+#include "CoreTypes.h"
+#include "Misc/AssertionMacros.h"
+#include "Misc/Timespan.h"
+
+class FString;
+
+namespace UE
+{
+
+/**
+ * A mutex that supports recursive locking.
+ *
+ * Prefer FRecursiveMutex.
+ */
+class FWindowsRecursiveMutex
+{
+public:
+	FWindowsRecursiveMutex(const FWindowsRecursiveMutex&) = delete;
+	FWindowsRecursiveMutex& operator=(const FWindowsRecursiveMutex&) = delete;
+
+	FORCEINLINE FWindowsRecursiveMutex()
+	{
+		CA_SUPPRESS(28125);
+		Windows::InitializeCriticalSection(&CriticalSection);
+		Windows::SetCriticalSectionSpinCount(&CriticalSection, 4000);
+		//Windows::InitializeCriticalSectionAndSpinCount(&CriticalSection, 4000);
+	}
+
+	FORCEINLINE ~FWindowsRecursiveMutex()
+	{
+		Windows::DeleteCriticalSection(&CriticalSection);
+	}
+
+	FORCEINLINE bool TryLock()
+	{
+		return !!Windows::TryEnterCriticalSection(&CriticalSection);
+	}
+
+	FORCEINLINE void Lock()
+	{
+		Windows::EnterCriticalSection(&CriticalSection);
+	}
+
+	FORCEINLINE void Unlock()
+	{
+		Windows::LeaveCriticalSection(&CriticalSection);
+	}
+
+private:
+	Windows::CRITICAL_SECTION CriticalSection;
+};
+
+/**
+ * A shared (read/write) mutex that does not support recursive locking.
+ *
+ * Prefer FSharedMutex.
+ *
+ * SRWLOCK: https://learn.microsoft.com/en-us/windows/win32/sync/slim-reader-writer--srw--locks
+ */
+class FWindowsSharedMutex
+{
+public:
+	FWindowsSharedMutex(const FWindowsSharedMutex&) = delete;
+	FWindowsSharedMutex& operator=(const FWindowsSharedMutex&) = delete;
+
+	FORCEINLINE FWindowsSharedMutex()
+	{
+		Windows::InitializeSRWLock(&Mutex);
+	}
+
+	~FWindowsSharedMutex()
+	{
+		checkf(!IsLocked(), TEXT("Destroying a lock that is still held!"));
+	}
+
+	FORCEINLINE bool TryLock()
+	{
+		return !!Windows::TryAcquireSRWLockExclusive(&Mutex);
+	}
+
+	FORCEINLINE void Lock()
+	{
+		Windows::AcquireSRWLockExclusive(&Mutex);
+	}
+
+	FORCEINLINE void Unlock()
+	{
+		Windows::ReleaseSRWLockExclusive(&Mutex);
+	}
+
+	FORCEINLINE bool TryLockShared()
+	{
+		return !!Windows::TryAcquireSRWLockShared(&Mutex);
+	}
+
+	FORCEINLINE void LockShared()
+	{
+		Windows::AcquireSRWLockShared(&Mutex);
+	}
+
+	FORCEINLINE void UnlockShared()
+	{
+		Windows::ReleaseSRWLockShared(&Mutex);
+	}
+
+private:
+	bool IsLocked()
+	{
+		if (Windows::TryAcquireSRWLockExclusive(&Mutex))
+		{
+			Windows::ReleaseSRWLockExclusive(&Mutex);
+			return false;
+		}
+		return true;
+	}
+
+	Windows::SRWLOCK Mutex;
+};
+
+/** A system-wide mutex for Windows. Uses a named mutex. */
+class FWindowsSystemWideMutex
+{
+public:
+	FWindowsSystemWideMutex(const FWindowsSystemWideMutex&) = delete;
+	FWindowsSystemWideMutex& operator=(const FWindowsSystemWideMutex&) = delete;
+
+	/** Construct a named, system-wide mutex and attempt to get access/ownership of it. */
+	CORE_API explicit FWindowsSystemWideMutex(const FString& InName, FTimespan InTimeout = FTimespan::Zero());
+
+	/** Destructor releases system-wide mutex if it is currently owned. */
+	CORE_API ~FWindowsSystemWideMutex();
+
+	/**
+	 * Does the calling thread have ownership of the system-wide mutex?
+	 *
+	 * @return True if obtained. WARNING: Returns true for an owned but previously abandoned locks so shared resources can be in undetermined states. You must handle shared data robustly.
+	 */
+	CORE_API bool IsValid() const;
+
+	/** Releases system-wide mutex if it is currently owned. */
+	CORE_API void Release();
+
+private:
+	Windows::HANDLE Mutex;
+};
+
+using FPlatformRecursiveMutex = FWindowsRecursiveMutex;
+using FPlatformSharedMutex = FWindowsSharedMutex;
+using FPlatformSystemWideMutex = FWindowsSystemWideMutex;
+
+} // UE
+
+#if UE_ENABLE_INCLUDE_ORDER_DEPRECATED_IN_5_6
+#include "HAL/PlatformMemory.h"
+#endif
diff --git a/Engine/Source/Runtime/Experimental/Dataflow/Core/Public/Dataflow/DataflowInputOutput.h b/Engine/Source/Runtime/Experimental/Dataflow/Core/Public/Dataflow/DataflowInputOutput.h
--- a/Engine/Source/Runtime/Experimental/Dataflow/Core/Public/Dataflow/DataflowInputOutput.h
+++ b/Engine/Source/Runtime/Experimental/Dataflow/Core/Public/Dataflow/DataflowInputOutput.h
@@ -8,7 +8,7 @@
 #include "Dataflow/DataflowAnyType.h"
 #include "Templates/Function.h"
 #include "Async/Async.h"
-#include "GenericPlatform/GenericPlatformCriticalSection.h"
+#include "HAL/CriticalSection.h"
 
 #include "DataflowInputOutput.generated.h"
 
diff --git a/Engine/Source/Runtime/Experimental/Dataflow/Core/Public/Dataflow/DataflowNodeParameters.h b/Engine/Source/Runtime/Experimental/Dataflow/Core/Public/Dataflow/DataflowNodeParameters.h
--- a/Engine/Source/Runtime/Experimental/Dataflow/Core/Public/Dataflow/DataflowNodeParameters.h
+++ b/Engine/Source/Runtime/Experimental/Dataflow/Core/Public/Dataflow/DataflowNodeParameters.h
@@ -6,7 +6,7 @@
 #include "Misc/ScopeExit.h"
 #include "UObject/UnrealType.h"
 #include "Templates/UniquePtr.h"
-#include "GenericPlatform/GenericPlatformCriticalSection.h"
+#include "HAL/CriticalSection.h"
 #include "Dataflow/DataflowContextCache.h"
 #include "Serialization/Archive.h"
 
