Scheduler: Use the platform mutex to ensure the scheduler does not depend on ParkingLot

ParkingLot indirectly depends on the scheduler because FPlatformManualResetEvent calls into the scheduler to support oversubscription. The scheduler will deadlock if it calls back into ParkingLot while starting a thread for oversubscription due to a wait in ParkingLot.

#rnx

diff --git a/Engine/Source/Runtime/Core/Private/Async/Fundamental/ReserveScheduler.cpp b/Engine/Source/Runtime/Core/Private/Async/Fundamental/ReserveScheduler.cpp
--- a/Engine/Source/Runtime/Core/Private/Async/Fundamental/ReserveScheduler.cpp
+++ b/Engine/Source/Runtime/Core/Private/Async/Fundamental/ReserveScheduler.cpp
@@ -4,6 +4,7 @@
 
 #include "Async/Fundamental/LocalQueue.h"
 #include "Async/Fundamental/Scheduler.h"
+#include "Async/UniqueLock.h"
 #include "HAL/LowLevelMemTracker.h"
 #include "HAL/PlatformProcess.h"
 #include "Misc/Fork.h"
@@ -83,7 +84,7 @@ void FReserveScheduler::StartWorkers(uint32 NumWorkers, FThread::EForkable IsFor
 	{
 		LLM_SCOPE_BYNAME(TEXT("EngineMisc/WorkerThreads"));
 
-		FScopeLock Lock(&WorkerThreadsCS);
+		UE::TUniqueLock Lock(WorkerThreadsCS);
 		check(!WorkerThreads.Num());
 		check(NextWorkerId == 0);
 
@@ -103,7 +104,7 @@ void FReserveScheduler::StopWorkers()
 	uint32 OldActiveWorkers = ActiveWorkers.load(std::memory_order_relaxed);
 	if(OldActiveWorkers != 0 && ActiveWorkers.compare_exchange_strong(OldActiveWorkers, 0, std::memory_order_relaxed))
 	{
-		FScopeLock Lock(&WorkerThreadsCS);
+		UE::TUniqueLock Lock(WorkerThreadsCS);
 
 		for (FYieldedWork& Event : ReserveEvents)
 		{
diff --git a/Engine/Source/Runtime/Core/Private/Async/Fundamental/Scheduler.cpp b/Engine/Source/Runtime/Core/Private/Async/Fundamental/Scheduler.cpp
--- a/Engine/Source/Runtime/Core/Private/Async/Fundamental/Scheduler.cpp
+++ b/Engine/Source/Runtime/Core/Private/Async/Fundamental/Scheduler.cpp
@@ -3,6 +3,7 @@
 #include "Async/Fundamental/Scheduler.h"
 #include "Async/Fundamental/Task.h"
 #include "Async/TaskTrace.h"
+#include "Async/UniqueLock.h"
 #include "Logging/LogMacros.h"
 #include "ProfilingDebugging/CsvProfiler.h"
 #include "HAL/RunnableThread.h"
@@ -124,7 +125,7 @@ namespace LowLevelTasks
 		uint32 OldActiveWorkers = ActiveWorkers.load(std::memory_order_relaxed);
 		if (OldActiveWorkers == 0 && bSupportsMultithreading && ActiveWorkers.compare_exchange_strong(OldActiveWorkers, NumForegroundWorkers + NumBackgroundWorkers, std::memory_order_relaxed))
 		{
-			FScopeLock Lock(&WorkerThreadsCS);
+			UE::TUniqueLock Lock(WorkerThreadsCS);
 			check(!WorkerThreads.Num());
 			check(!WorkerLocalQueues.Num());
 			check(!WorkerEvents.Num());
@@ -151,7 +152,7 @@ namespace LowLevelTasks
 			auto CreateThread =
 				[this, IsForkable](Private::ELocalQueueType LocalQueueType, const TCHAR* ThreadGroup, const TCHAR* Prefix, std::atomic<int32>& CreationIndex, int32 NumWorkers, int32 NumMaxWorkers, EThreadPriority Priority, uint64 Affinity)
 				{
-					FScopeLock Lock(&WorkerCreationCS);
+					UE::TUniqueLock Lock(WorkerCreationCS);
 					LLM_SCOPE_BYNAME(TEXT("EngineMisc/WorkerThreads"));
 
 					// Thread creation can end up waiting, we don't want to recursively oversubscribe if that happens.
@@ -301,7 +302,7 @@ namespace LowLevelTasks
 		uint32 OldActiveWorkers = ActiveWorkers.load(std::memory_order_relaxed);
 		if (OldActiveWorkers != 0 && ActiveWorkers.compare_exchange_strong(OldActiveWorkers, 0, std::memory_order_relaxed))
 		{
-			FScopeLock Lock(&WorkerThreadsCS);
+			UE::TUniqueLock Lock(WorkerThreadsCS);
 
 			WaitingQueue[0].StartShutdown();
 			WaitingQueue[1].StartShutdown();
@@ -372,7 +373,7 @@ namespace LowLevelTasks
 
 	void FScheduler::RestartWorkers(uint32 NumForegroundWorkers, uint32 NumBackgroundWorkers, FThread::EForkable IsForkable, EThreadPriority InWorkerPriority, EThreadPriority InBackgroundPriority, uint64 InWorkerAffinity, uint64 InBackgroundAffinity)
 	{
-		FScopeLock Lock(&WorkerThreadsCS);
+		UE::TUniqueLock Lock(WorkerThreadsCS);
 		TemporaryShutdown.store(true, std::memory_order_release);
 		while (HasPendingWakeUp())
 		{
diff --git a/Engine/Source/Runtime/Core/Public/Async/Fundamental/ReserveScheduler.h b/Engine/Source/Runtime/Core/Public/Async/Fundamental/ReserveScheduler.h
--- a/Engine/Source/Runtime/Core/Public/Async/Fundamental/ReserveScheduler.h
+++ b/Engine/Source/Runtime/Core/Public/Async/Fundamental/ReserveScheduler.h
@@ -5,9 +5,9 @@
 #include "Containers/Array.h"
 #include "Containers/ContainerAllocationPolicies.h"
 #include "CoreTypes.h"
-#include "HAL/CriticalSection.h"
 #include "HAL/Event.h"
 #include "HAL/PlatformAffinity.h"
+#include "HAL/PlatformMutex.h"
 #include "HAL/Thread.h"
 #include "Scheduler.h"
 #include "Templates/UniquePtr.h"
@@ -51,7 +51,7 @@ namespace LowLevelTasks
 		TEventStack<FYieldedWork> 						EventStack = { ReserveEvents };
 	PRAGMA_ENABLE_DEPRECATION_WARNINGS
 		TAlignedArray<FYieldedWork>						ReserveEvents;
-		FCriticalSection 								WorkerThreadsCS;
+		UE::FPlatformRecursiveMutex						WorkerThreadsCS;
 		TArray<TUniquePtr<FThread>>						WorkerThreads;
 		std::atomic_uint								ActiveWorkers { 0 };
 		std::atomic_uint								NextWorkerId { 0 };
@@ -73,4 +73,8 @@ namespace LowLevelTasks
 		StopWorkers();
 	}
 	PRAGMA_ENABLE_DEPRECATION_WARNINGS
-}
\ No newline at end of file
+}
+
+#if UE_ENABLE_INCLUDE_ORDER_DEPRECATED_IN_5_6
+#include "HAL/CriticalSection.h"
+#endif
diff --git a/Engine/Source/Runtime/Core/Public/Async/Fundamental/Scheduler.h b/Engine/Source/Runtime/Core/Public/Async/Fundamental/Scheduler.h
--- a/Engine/Source/Runtime/Core/Public/Async/Fundamental/Scheduler.h
+++ b/Engine/Source/Runtime/Core/Public/Async/Fundamental/Scheduler.h
@@ -10,9 +10,9 @@
 #include "Containers/Array.h"
 #include "Containers/ArrayView.h"
 #include "Containers/ContainerAllocationPolicies.h"
-#include "HAL/CriticalSection.h"
 #include "HAL/Event.h"
 #include "HAL/PlatformAffinity.h"
+#include "HAL/PlatformMutex.h"
 #include "HAL/PlatformProcess.h"
 #include "HAL/Thread.h"
 #include "LocalQueue.h"
@@ -264,8 +264,8 @@ namespace LowLevelTasks
 	private:
 		Private::FWaitingQueue                         WaitingQueue[2] = { { WorkerEvents, OversubscriptionLimitReachedEvent }, { WorkerEvents, OversubscriptionLimitReachedEvent } };
 		FSchedulerTls::FQueueRegistry                  QueueRegistry;
-		FCriticalSection                               WorkerCreationCS;
-		FCriticalSection                               WorkerThreadsCS;
+		UE::FPlatformRecursiveMutex                    WorkerCreationCS;
+		UE::FPlatformRecursiveMutex                    WorkerThreadsCS;
 		TArray<std::atomic<FThread*>>                  WorkerThreads;
 		TAlignedArray<FSchedulerTls::FLocalQueueType>  WorkerLocalQueues;
 		TAlignedArray<Private::FWaitEvent>             WorkerEvents;
@@ -470,3 +470,6 @@ namespace LowLevelTasks
 	}
 }
 
+#if UE_ENABLE_INCLUDE_ORDER_DEPRECATED_IN_5_6
+#include "HAL/CriticalSection.h"
+#endif
