[RaceDetector]
- Fix potential race between GetSyncObject and ResetShadow by using hazard pointers (could only repro in stress tests)
- Fix potential race between multiple GetSyncObject on the same address due to a missing refcount before publishing
- Remove remaining std::atomic usage to cut down on reentrancy as much as possible

diff --git a/Engine/Source/Runtime/Core/Internal/Sanitizer/RaceDetectorInterface.h b/Engine/Source/Runtime/Core/Internal/Sanitizer/RaceDetectorInterface.h
--- a/Engine/Source/Runtime/Core/Internal/Sanitizer/RaceDetectorInterface.h
+++ b/Engine/Source/Runtime/Core/Internal/Sanitizer/RaceDetectorInterface.h
@@ -88,7 +88,7 @@ namespace UE::Sanitizer::RaceDetector {
 	// Releases the current thread context.
 	void ReleaseCurrentContext();
 	// Returns a sync object for the given address, initialize one if there isn't one already.
-	FSyncObjectRef GetSyncObject(void* SyncAddr);
+	FSyncObjectRef GetSyncObject(FContext& Context, void* SyncAddr);
 }
 
 #endif // USING_INSTRUMENTATION
\ No newline at end of file
diff --git a/Engine/Source/Runtime/Core/Internal/Sanitizer/RaceDetectorTypes.h b/Engine/Source/Runtime/Core/Internal/Sanitizer/RaceDetectorTypes.h
--- a/Engine/Source/Runtime/Core/Internal/Sanitizer/RaceDetectorTypes.h
+++ b/Engine/Source/Runtime/Core/Internal/Sanitizer/RaceDetectorTypes.h
@@ -646,6 +646,7 @@ namespace UE::Sanitizer::RaceDetector {
 		}
 	};
 
+	struct FSyncObjectBank;
 	// ------------------------------------------------------------------------------
 	// Race Detector Context.
 	// ------------------------------------------------------------------------------
@@ -689,6 +690,8 @@ namespace UE::Sanitizer::RaceDetector {
 		FClock StandbyClock;
 		// Each thread holds a bank of clocks to synchronize with every other context.
 		FClockBank ClockBank;
+		// Hazard pointer used between GetSyncObject and ResetShadow
+		FSyncObjectBank* BankHazard = nullptr;
 		// We need to keep the callstack for each thread 
 		static constexpr SIZE_T MaxCallstackSize = 1024; // Yes, callstacks can get huge in UE!
 		void* CurrentCallstack[MaxCallstackSize];
@@ -737,7 +740,6 @@ namespace UE::Sanitizer::RaceDetector {
 		}
 	};
 
-	struct FSyncObjectBank;
 	class FSyncObject
 	{
 	public:
@@ -934,22 +936,29 @@ namespace UE::Sanitizer::RaceDetector {
 		// We maintain a linked list of clock banks for recycling purpose
 		FSyncObjectBank* Next = 0;
 
-		INSTRUMENTATION_FUNCTION_ATTRIBUTES void AddRef()
+		INSTRUMENTATION_FUNCTION_ATTRIBUTES int32 AddRef()
 		{
-			FPlatformAtomics::InterlockedIncrement(&RefCount);
+			return FPlatformAtomics::InterlockedIncrement(&RefCount);
 		}
 
-		INSTRUMENTATION_FUNCTION_ATTRIBUTES void Release()
+		INSTRUMENTATION_FUNCTION_ATTRIBUTES int32 Release()
 		{
-			if (FPlatformAtomics::InterlockedDecrement(&RefCount) == 0)
+			int32 NewRefCount = FPlatformAtomics::InterlockedDecrement(&RefCount);
+			if (NewRefCount == 0)
 			{
 				delete this;
 			}
+			return NewRefCount;
+		}
+
+		INSTRUMENTATION_FUNCTION_ATTRIBUTES int32 GetRefCount() const
+		{
+			return RefCount;
 		}
 
 		static volatile int64 ObjectCount;
 	private:
-		volatile int32 RefCount = 0;
+		volatile int32 RefCount = 1;
 
 		// Save space by allocating the first sync object as part of the bank itself.
 		// Most of the time the sync object will be at offset 0.
@@ -977,10 +986,18 @@ namespace UE::Sanitizer::RaceDetector {
 
 		// Not needed for now
 		FSyncObjectRef(const FSyncObjectRef&) = delete;
-		FSyncObjectRef(FSyncObjectRef&&) = delete;
 		FSyncObjectRef& operator=(const FSyncObjectRef&) = delete;
 		FSyncObjectRef& operator=(FSyncObjectRef&&) = delete;
 
+		FSyncObjectRef(FSyncObjectRef&& Other)
+		{
+			Bank = Other.Bank;
+			Object = Other.Object;
+
+			Other.Bank = nullptr;
+			Other.Object = nullptr;
+		}
+
 		INSTRUMENTATION_FUNCTION_ATTRIBUTES FORCEINLINE FSyncObject* operator->() const
 		{
 			return Object;
@@ -993,7 +1010,10 @@ namespace UE::Sanitizer::RaceDetector {
 
 		INSTRUMENTATION_FUNCTION_ATTRIBUTES FORCEINLINE void Release()
 		{
-			Bank->Release();
+			if (Bank)
+			{
+				Bank->Release();
+			}
 		}
 
 		FSyncObjectBank* Bank;
diff --git a/Engine/Source/Runtime/Core/Private/Sanitizer/RaceDetector.cpp b/Engine/Source/Runtime/Core/Private/Sanitizer/RaceDetector.cpp
--- a/Engine/Source/Runtime/Core/Private/Sanitizer/RaceDetector.cpp
+++ b/Engine/Source/Runtime/Core/Private/Sanitizer/RaceDetector.cpp
@@ -37,6 +37,12 @@ UE_TRACE_CHANNEL_DEFINE(RaceDetectorChannel);
 #define WITH_RACEDETECTOR_CHECK 0
 #define WITH_RACEDETECTOR_DEBUG 0
 
+#if WITH_RACEDETECTOR_CHECK
+#define RACEDETECTOR_CHECK(cond) check(cond)
+#else
+#define RACEDETECTOR_CHECK(cond)
+#endif
+
 #if WITH_RACEDETECTOR_DEBUG
 UE_DISABLE_OPTIMIZATION
 #endif
@@ -324,7 +330,9 @@ namespace UE::Sanitizer::RaceDetector {
 			Ptr -= ShadowMemorySize;
 		}
 
-		return ShadowMemoryBase + (Ptr >> 3) * sizeof(FShadowMemory);
+		UPTRINT Result = ShadowMemoryBase + (Ptr >> 3) * sizeof(FShadowMemory);
+		RACEDETECTOR_CHECK(Result >= ShadowMemoryBase && Result < ShadowClockBase);
+		return Result;
 	}
 
 	INSTRUMENTATION_FUNCTION_ATTRIBUTES FShadowMemory* GetShadowMemory(UPTRINT Ptr)
@@ -350,7 +358,9 @@ namespace UE::Sanitizer::RaceDetector {
 			Ptr -= ShadowMemorySize;
 		}
 
-		return ShadowClockBase + (Ptr >> 3) * sizeof(FShadowClockBankSlot);
+		UPTRINT Result = ShadowClockBase + (Ptr >> 3) * sizeof(FShadowClockBankSlot);
+		RACEDETECTOR_CHECK(Result >= ShadowMemoryBase && Result < ShadowMemoryEnd);
+		return Result;
 	}
 
 	INSTRUMENTATION_FUNCTION_ATTRIBUTES FShadowClockBankSlot* GetShadowClockBank(UPTRINT Ptr)
@@ -364,57 +374,74 @@ namespace UE::Sanitizer::RaceDetector {
 		return (FShadowClockBankSlot*)ShadowMemoryAddress;
 	}
 
-	INSTRUMENTATION_FUNCTION_ATTRIBUTES FSyncObjectRef GetSyncObject(void* SyncAddr)
+	INSTRUMENTATION_FUNCTION_ATTRIBUTES void PushSyncObjectBank(FSyncObjectBank* SyncObjectBank)
 	{
-		FShadowClockBankSlot* ClockBankSlot = GetShadowClockBank((UPTRINT)SyncAddr);
-
-		if (ClockBankSlot->SyncObjectBank == nullptr)
+		while (true)
 		{
-			FSyncObjectBank* SyncObjectBank = new FSyncObjectBank();
-			if (FPlatformAtomics::InterlockedCompareExchangePointer((void**)&ClockBankSlot->SyncObjectBank, SyncObjectBank, nullptr) != nullptr)
+			FSyncObjectBank* LocalSyncBankHead = (FSyncObjectBank*)SyncObjectBankHead;
+			SyncObjectBank->Next = LocalSyncBankHead;
+
+			if (FPlatformAtomics::InterlockedCompareExchangePointer((void**)&SyncObjectBankHead, SyncObjectBank, LocalSyncBankHead) == LocalSyncBankHead)
 			{
-				// Should not happen very often, and at most once per thread taking part of the race
-				delete SyncObjectBank;
+				break;
 			}
-			else
+		}
+	}
+
+	INSTRUMENTATION_FUNCTION_ATTRIBUTES FSyncObjectRef GetSyncObject(FContext& Context, void* SyncAddr)
+	{
+		FShadowClockBankSlot* ClockBankSlot = GetShadowClockBank((UPTRINT)SyncAddr);
+		const UPTRINT Index = UPTRINT(SyncAddr) & 7;
+		RACEDETECTOR_CHECK(Context.BankHazard == nullptr);
+
+		// This can race with ResetShadow so use Hazard Pointer mechanism to avoid
+		// the object bank from being deleted while we're still adding a refcount to it.
+		while (true)
+		{
+			FSyncObjectBank* LocalBank = ClockBankSlot->SyncObjectBank;
+			if (LocalBank == nullptr)
 			{
-				SyncObjectBank->AddRef();
+				FSyncObjectBank* NewObjectBank = new FSyncObjectBank();
+
+				// When publishing in shadow memory, we need to make sure the refcount is not 0, since another thread could try to use
+				// the newly published bank from the shadow, perform an AddRef and then release it and we'd end up with
+				// an invalid NewObjectBank before even starting to use it.
+				RACEDETECTOR_CHECK(NewObjectBank->GetRefCount() == 1);
 
-				// Publish the object in the linked-list
-				while (true)
+				FSyncObjectBank* OldObjectBank = (FSyncObjectBank*)FPlatformAtomics::InterlockedCompareExchangePointer((void**)&ClockBankSlot->SyncObjectBank, NewObjectBank, nullptr);
+				if (OldObjectBank == nullptr)
 				{
-					FSyncObjectBank* LocalSyncBankHead = (FSyncObjectBank*)SyncObjectBankHead;
-					SyncObjectBank->Next = LocalSyncBankHead;
+					LocalBank = NewObjectBank;
 
-					if (FPlatformAtomics::InterlockedCompareExchangePointer((void**)&SyncObjectBankHead, SyncObjectBank, LocalSyncBankHead) == LocalSyncBankHead)
-					{
-						break;
-					}
+					// Take our refcount now so that we don't have to do any hazard pointer handling.
+					FSyncObjectRef Result = FSyncObjectRef(LocalBank, LocalBank->GetSyncObject(Index));
+
+					// Publish the object in the linked-list
+					// The creation ref-count is now owned by the linked-list
+					PushSyncObjectBank(LocalBank);
+
+					Context.BankHazard = nullptr;
+					return MoveTemp(Result);
 				}
-			}
-		}
 
-		UPTRINT Index = UPTRINT(SyncAddr) & 7;
-		return FSyncObjectRef(ClockBankSlot->SyncObjectBank, ClockBankSlot->SyncObjectBank->GetSyncObject(Index));
-	}
+				LocalBank = OldObjectBank;
 
-	INSTRUMENTATION_FUNCTION_ATTRIBUTES bool HasInstrumentedAtomics()
-	{
-		FInstrumentationScope InstrumentationScope;
-		// Very handy way of knowing if we still have any clock banks left.
-		// This will make sure all the refcounting works and that we can free everything properly.
-		return GetSyncObjectMemoryUsage() != 0;
-	}
+				// Should not happen very often, and at most once per thread taking part of the race
+				int32 Result = NewObjectBank->Release();
+				RACEDETECTOR_CHECK(Result == 0);
+			}
 
-	INSTRUMENTATION_FUNCTION_ATTRIBUTES void ResetInstrumentedAtomics()
-	{
-		FInstrumentationScope InstrumentationScope;
-		FSyncObjectBank* ClockBank = (FSyncObjectBank*)FPlatformAtomics::InterlockedExchangePtr((void**)&SyncObjectBankHead, nullptr);
-		while (ClockBank)
-		{
-			FSyncObjectBank* ToDelete = ClockBank;
-			ClockBank = ClockBank->Next;
-			ToDelete->Release();
+			Context.BankHazard = LocalBank;
+			FPlatformMisc::AsymmetricThreadFenceLight();
+
+			// Confirm that the entry is still in the shadow slot, otherwise we iterate again as our
+			// hazard pointer protection wouldn't be guaranteed.
+			if (LocalBank == ClockBankSlot->SyncObjectBank)
+			{
+				FSyncObjectRef Result = FSyncObjectRef(LocalBank, LocalBank->GetSyncObject(Index));
+				Context.BankHazard = nullptr;
+				return MoveTemp(Result);
+			}
 		}
 	}
 
@@ -881,7 +908,7 @@ namespace UE::Sanitizer::RaceDetector {
 			return;
 		}
 
-		FSyncObjectRef Atomic = GetSyncObject(AtomicAddr);
+		FSyncObjectRef Atomic = GetSyncObject(Context, AtomicAddr);
 
 		if (AccessType == ACCESS_TYPE_ATOMIC_READ_WRITE && IsAtomicOrderAcquireRelease(Order))
 		{
@@ -918,7 +945,7 @@ namespace UE::Sanitizer::RaceDetector {
 		}
 		else
 		{
-			GetSyncObject(AtomicAddr)->SyncWithFailureSupport(Context, AtomicOp, AccessType, SuccessOrder, FailureOrder, ReturnAddress, AtomicAddr, OpName,
+			GetSyncObject(Context, AtomicAddr)->SyncWithFailureSupport(Context, AtomicOp, AccessType, SuccessOrder, FailureOrder, ReturnAddress, AtomicAddr, OpName,
 				[&Context, AtomicAddr, Size, AccessType, ReturnAddress, &SuccessOrder, OpName, &bHasAddedCallstackFrame](FAtomicMemoryOrder ActualOrder) INSTRUMENTATION_FUNCTION_ATTRIBUTES
 				{
 					InstrumentAccess(Context, AtomicAddr, Size, AccessType, ActualOrder, ReturnAddress, OpName, bHasAddedCallstackFrame);
@@ -979,6 +1006,23 @@ namespace UE::Sanitizer::RaceDetector {
 
 	INSTRUMENTATION_FUNCTION_ATTRIBUTES void UnhookInstrumentation();
 	INSTRUMENTATION_FUNCTION_ATTRIBUTES void HookInstrumentation();
+		
+	INSTRUMENTATION_FUNCTION_ATTRIBUTES TSet<FSyncObjectBank*> GatherBankHazards()
+	{
+		// Make sure all threads have published their BankHazard
+		FPlatformMisc::AsymmetricThreadFenceHeavy();
+
+		TSet<FSyncObjectBank*> Hazards;
+		TReadScopeLock Scope(ContextMappingLock);
+		for (auto& Pair : ContextMapping)
+		{
+			if (FSyncObjectBank* BankHazard = Pair.Value->BankHazard)
+			{
+				Hazards.Add(BankHazard);
+			}
+		}
+		return Hazards;
+	}
 
 	INSTRUMENTATION_FUNCTION_ATTRIBUTES void ResetShadow()
 	{
@@ -995,7 +1039,7 @@ namespace UE::Sanitizer::RaceDetector {
 
 		FPlatformMisc::LowLevelOutputDebugStringf(TEXT("[%d] Reset Shadow Memory Started\n"), Platform::GetCurrentThreadId());
 
-		bIsResettingShadow = true;
+		FPlatformAtomics::InterlockedExchange((volatile int8*)&bIsResettingShadow, (int8)1);
 
 		if (bRaceDetectorActive)
 		{
@@ -1006,10 +1050,32 @@ namespace UE::Sanitizer::RaceDetector {
 		// both collections are empty and coherent to together to avoid false positives.
 		do
 		{
+			// This need to synchronize with GetSyncObject so null the list before resetting shadow memory
+			FSyncObjectBank* ClockBank = (FSyncObjectBank*)FPlatformAtomics::InterlockedExchangePtr((void**)&SyncObjectBankHead, nullptr);
+			
 			ResetShadowMemory();
-			ResetInstrumentedAtomics();
 
-		} while (Platform::HasShadowMemoryMapped() || HasInstrumentedAtomics());
+			TSet<FSyncObjectBank*> Hazards = GatherBankHazards();
+
+			// Garbage collect what we can and put back what we can't.
+			while (ClockBank)
+			{
+				FSyncObjectBank* ToDelete = ClockBank;
+				ClockBank = ClockBank->Next;
+
+				if (Hazards.Contains(ToDelete))
+				{
+					// Put it back in the list for next try.
+					PushSyncObjectBank(ToDelete);
+					FPlatformMisc::LowLevelOutputDebugStringf(TEXT("[%d] %p was in hazard list\n"), Platform::GetCurrentThreadId(), ToDelete);
+				}
+				else
+				{
+					ToDelete->Release();
+				}
+			}
+
+		} while (Platform::HasShadowMemoryMapped());
 
 		FPlatformMisc::LowLevelOutputDebugStringf(TEXT("[%d] Reset Shadow Memory Ended\n"), Platform::GetCurrentThreadId());
 		
@@ -1533,7 +1599,7 @@ extern "C" {
 		bool bHasAddedCallstackFrame = true;
 		InstrumentAccess(Context, addr, sizeof(UPTRINT), ACCESS_TYPE_ATOMIC_READ, FAtomicMemoryOrder::MEMORY_ORDER_RELEASE, _ReturnAddress(), TEXT("AnnotateHappensBefore"), bHasAddedCallstackFrame);
 
-		FSyncObjectRef Atomic = GetSyncObject(addr);
+		FSyncObjectRef Atomic = GetSyncObject(Context, addr);
 		Atomic->SyncRelease(Context, []() {}, _ReturnAddress(), addr, TEXT("AnnotateHappensBefore"));
 		Context.IncrementClock();
 		RemoveCallstackFrame(Context);
@@ -1549,7 +1615,7 @@ extern "C" {
 		bool bHasAddedCallstackFrame = true;
 		InstrumentAccess(Context, addr, sizeof(UPTRINT), ACCESS_TYPE_ATOMIC_READ, FAtomicMemoryOrder::MEMORY_ORDER_ACQUIRE, _ReturnAddress(), TEXT("AnnotateHappensBefore"), bHasAddedCallstackFrame);
 
-		FSyncObjectRef Atomic = GetSyncObject(addr);
+		FSyncObjectRef Atomic = GetSyncObject(Context, addr);
 		Atomic->SyncAcquire(Context, []() {}, _ReturnAddress(), addr, TEXT("AnnotateHappensAfter"));
 		Context.IncrementClock();
 		RemoveCallstackFrame(Context);
diff --git a/Engine/Source/Runtime/Core/Private/Sanitizer/RaceDetectorWindows.cpp b/Engine/Source/Runtime/Core/Private/Sanitizer/RaceDetectorWindows.cpp
--- a/Engine/Source/Runtime/Core/Private/Sanitizer/RaceDetectorWindows.cpp
+++ b/Engine/Source/Runtime/Core/Private/Sanitizer/RaceDetectorWindows.cpp
@@ -117,7 +117,7 @@ INSTRUMENTATION_FUNCTION_ATTRIBUTES void InstrumentEndThread()
 	START_WINAPI_INSTRUMENTATION(ReleaseCurrentContext(); return);
 	if (Handle)
 	{
-		FSyncObjectRef Sync = GetSyncObject(Handle);
+		FSyncObjectRef Sync = GetSyncObject(Context, Handle);
 		Sync->SyncReleaseAsSoleOwner(Context, _ReturnAddress(), Handle, TEXT("InstrumentEndThread"));
 		Context.IncrementClock();
 	}
@@ -259,7 +259,7 @@ INSTRUMENTATION_FUNCTION_ATTRIBUTES void AcquireWaitHandle(HANDLE Handle, void*
 {
 	FContext& Context = EnsureCurrentContext();
 
-	FSyncObjectRef Sync = GetSyncObject(Handle);
+	FSyncObjectRef Sync = GetSyncObject(Context, Handle);
 	Sync->SyncAcquire(Context, []() {}, ReturnAddress, Handle, OpName);
 	Context.IncrementClock();
 }
@@ -306,7 +306,7 @@ INSTRUMENTATION_FUNCTION_ATTRIBUTES BOOL WINAPI DetouredSetEvent(HANDLE hHandle)
 
 	// The scope is important to make sure the TRefCountPtr is destroyed before we FINISH_WINAPI_INSTRUMENTATION.
 	{
-		FSyncObjectRef Sync = GetSyncObject(hHandle);
+		FSyncObjectRef Sync = GetSyncObject(Context, hHandle);
 		Sync->SyncRelease(Context, []() {}, _ReturnAddress(), hHandle, TEXT("SetEvent"));
 		Context.IncrementClock();
 	}
@@ -342,7 +342,7 @@ INSTRUMENTATION_FUNCTION_ATTRIBUTES BOOL WINAPI DetouredWaitOnAddress(volatile V
 
 		if (bIsValueChanged)
 		{
-			FSyncObjectRef Sync = GetSyncObject((void*)Address);
+			FSyncObjectRef Sync = GetSyncObject(Context, (void*)Address);
 			Sync->SyncAcquire(Context, []() {}, _ReturnAddress(), (void*)Address, TEXT("WaitOnAddress"));
 			Context.IncrementClock();
 		}
@@ -359,7 +359,7 @@ INSTRUMENTATION_FUNCTION_ATTRIBUTES void WINAPI DetouredWakeByAddressSingle(PVOI
 
 	{
 		FLastErrorPreservationScope LastErrorScope;
-		FSyncObjectRef Sync = GetSyncObject(Address);
+		FSyncObjectRef Sync = GetSyncObject(Context, Address);
 		Sync->SyncRelease(Context, []() {}, _ReturnAddress(), Address, TEXT("WakeByAddressSingle"));
 		Context.IncrementClock();
 	}
@@ -381,7 +381,7 @@ INSTRUMENTATION_FUNCTION_ATTRIBUTES BOOL WINAPI DetouredTryEnterCriticalSection(
 	// If we entered the critical section, we are now the sole owner.
 	if (Result)
 	{
-		FSyncObjectRef Sync = GetSyncObject(lpCriticalSection);
+		FSyncObjectRef Sync = GetSyncObject(Context, lpCriticalSection);
 		Sync->SyncAcquireAsSoleOwnerOrReadOwner(Context, _ReturnAddress(), lpCriticalSection, TEXT("TryEnterCriticalSection"));
 		Context.IncrementClock();
 	}
@@ -401,7 +401,7 @@ INSTRUMENTATION_FUNCTION_ATTRIBUTES void WINAPI DetouredEnterCriticalSection(CRI
 	// The scope is important to make sure the TRefCountPtr is destroyed before we FINISH_WINAPI_INSTRUMENTATION.
 	{
 		// We are now the sole owner.
-		FSyncObjectRef Sync = GetSyncObject(lpCriticalSection);
+		FSyncObjectRef Sync = GetSyncObject(Context, lpCriticalSection);
 		Sync->SyncAcquireAsSoleOwnerOrReadOwner(Context, _ReturnAddress(), lpCriticalSection, TEXT("EnterCriticalSection"));
 		Context.IncrementClock();
 	}
@@ -418,7 +418,7 @@ INSTRUMENTATION_FUNCTION_ATTRIBUTES void WINAPI DetouredLeaveCriticalSection(CRI
 		FLastErrorPreservationScope LastErrorScope;
 		// Assume this is called with the right semantics, i.e. that we are actual
 		// owners of the critical section.
-		FSyncObjectRef Sync = GetSyncObject(lpCriticalSection);
+		FSyncObjectRef Sync = GetSyncObject(Context, lpCriticalSection);
 		Sync->SyncReleaseAsSoleOwner(Context, _ReturnAddress(), lpCriticalSection, TEXT("LeaveCriticalSection"));
 		Context.IncrementClock();
 	}
@@ -438,7 +438,7 @@ INSTRUMENTATION_FUNCTION_ATTRIBUTES void WINAPI DetouredAcquireSRWLockShared(PSR
 
 	// The scope is important to make sure the TRefCountPtr is destroyed before we FINISH_WINAPI_INSTRUMENTATION.
 	{
-		FSyncObjectRef Sync = GetSyncObject(SRWLock);
+		FSyncObjectRef Sync = GetSyncObject(Context, SRWLock);
 		Sync->SyncAcquireAsSoleOwnerOrReadOwner(Context, _ReturnAddress(), SRWLock, TEXT("AcquireSRWLockShared"));
 		Context.IncrementClock();
 	}
@@ -456,7 +456,7 @@ INSTRUMENTATION_FUNCTION_ATTRIBUTES BOOLEAN WINAPI DetouredTryAcquireSRWLockShar
 	FLastErrorPreservationScope LastErrorScope;
 	if (Result)
 	{
-		FSyncObjectRef Sync = GetSyncObject(SRWLock);
+		FSyncObjectRef Sync = GetSyncObject(Context, SRWLock);
 		Sync->SyncAcquireAsSoleOwnerOrReadOwner(Context, _ReturnAddress(), SRWLock, TEXT("TryAcquireSRWLockShared"));
 		Context.IncrementClock();
 	}
@@ -472,7 +472,7 @@ INSTRUMENTATION_FUNCTION_ATTRIBUTES void WINAPI DetouredReleaseSRWLockShared(PSR
 
 	{
 		FLastErrorPreservationScope LastErrorScope;
-		FSyncObjectRef Sync = GetSyncObject(SRWLock);
+		FSyncObjectRef Sync = GetSyncObject(Context, SRWLock);
 		Sync->SyncRelease(Context, []() {}, _ReturnAddress(), SRWLock, TEXT("ReleaseSRWLockShared"));
 		Context.IncrementClock();
 	}
@@ -492,7 +492,7 @@ INSTRUMENTATION_FUNCTION_ATTRIBUTES void WINAPI DetouredAcquireSRWLockExclusive(
 
 	// The scope is important to make sure the TRefCountPtr is destroyed before we FINISH_WINAPI_INSTRUMENTATION.
 	{
-		FSyncObjectRef Sync = GetSyncObject(SRWLock);
+		FSyncObjectRef Sync = GetSyncObject(Context, SRWLock);
 		Sync->SyncAcquireAsSoleOwnerOrReadOwner(Context, _ReturnAddress(), SRWLock, TEXT("AcquireSRWLockExclusive"));
 		Context.IncrementClock();
 	}
@@ -510,7 +510,7 @@ INSTRUMENTATION_FUNCTION_ATTRIBUTES BOOLEAN WINAPI DetouredTryAcquireSRWLockExcl
 
 	if (Result)
 	{
-		FSyncObjectRef Sync = GetSyncObject(SRWLock);
+		FSyncObjectRef Sync = GetSyncObject(Context, SRWLock);
 		Sync->SyncAcquireAsSoleOwnerOrReadOwner(Context, _ReturnAddress(), SRWLock, TEXT("TryAcquireSRWLockExclusive"));
 		Context.IncrementClock();
 	}
@@ -526,7 +526,7 @@ INSTRUMENTATION_FUNCTION_ATTRIBUTES void WINAPI DetouredReleaseSRWLockExclusive(
 
 	{
 		FLastErrorPreservationScope LastErrorScope;
-		FSyncObjectRef Sync = GetSyncObject(SRWLock);
+		FSyncObjectRef Sync = GetSyncObject(Context, SRWLock);
 		Sync->SyncReleaseAsSoleOwner(Context, _ReturnAddress(), SRWLock, TEXT("ReleaseSRWLockExclusive"));
 		Context.IncrementClock();
 	}
@@ -684,7 +684,7 @@ namespace UE::Sanitizer::RaceDetector::Platform {
 	UPTRINT PageSizeBitShift = 0;
 	UPTRINT PageSizeMask = 0;
 	HANDLE  SanitizerThreadHandle = 0;
-	std::atomic<bool> bHasShadowMemoryMapped = false;
+	volatile bool bHasShadowMemoryMapped = false;
 	volatile bool bSanitizerThreadContinue = true;
 
 	INSTRUMENTATION_FUNCTION_ATTRIBUTES DWORD WINAPI SanitizerThreadProc(LPVOID lpParameter)
@@ -695,7 +695,7 @@ namespace UE::Sanitizer::RaceDetector::Platform {
 	
 	INSTRUMENTATION_FUNCTION_ATTRIBUTES bool HasShadowMemoryMapped()
 	{
-		return bHasShadowMemoryMapped.load(std::memory_order_relaxed);
+		return bHasShadowMemoryMapped;
 	}
 
 	INSTRUMENTATION_FUNCTION_ATTRIBUTES void SleepMS(uint32 MilliSeconds)
@@ -962,8 +962,7 @@ namespace UE::Sanitizer::RaceDetector::Platform {
 		checkSlow((UPTRINT)DirtyBitmapPtr < DirtyShadowBitmapEnd && (UPTRINT)DirtyBitmapPtr >= DirtyShadowBitmapBase);
 		if ((*DirtyBitmapPtr & PageBit) == 0)
 		{
-			std::atomic<uint8>* AtomicDirtyBitmapPtr = (std::atomic<uint8>*)DirtyBitmapPtr;
-			AtomicDirtyBitmapPtr->fetch_or(PageBit, std::memory_order_relaxed);
+			FPlatformAtomics::InterlockedOr((volatile int8*)DirtyBitmapPtr, (int8)PageBit);
 		}
 	}
 
@@ -1058,7 +1057,7 @@ namespace UE::Sanitizer::RaceDetector::Platform {
 				uintptr_t Result = (uintptr_t)::VirtualAlloc((LPVOID)BaseAddress, PageSize, MEM_COMMIT, PAGE_READWRITE);
 				if (Result == BaseAddress)
 				{
-					bHasShadowMemoryMapped.store(true, std::memory_order_relaxed);
+					bHasShadowMemoryMapped = true;
 
 					// Mark the page as committed in the bitmap. We use the bitmap when the debugger is
 					// present and to track memory usage of the shadow.
@@ -1067,8 +1066,7 @@ namespace UE::Sanitizer::RaceDetector::Platform {
 					uint8* BitmapPtr = (uint8*)(ShadowBitmapBase + (ShadowPage >> 3));
 					uint8 PageBit = (uint8)(1 << (ShadowPage & 7));
 					
-					std::atomic<uint8>* AtomicBitmapPtr = (std::atomic<uint8>*)BitmapPtr;
-					AtomicBitmapPtr->fetch_or(PageBit, std::memory_order_relaxed);
+					FPlatformAtomics::InterlockedOr((volatile int8*)BitmapPtr, (int8)PageBit);
 
 					DirtyShadowBitmapPage((UPTRINT)BitmapPtr);
 
@@ -1108,7 +1106,7 @@ namespace UE::Sanitizer::RaceDetector::Platform {
 		{
 			VirtualAlloc((void*)AlignedStart, AlignedEnd - AlignedStart, MEM_COMMIT, PAGE_READWRITE);
 
-			bHasShadowMemoryMapped.store(true, std::memory_order_relaxed);
+			bHasShadowMemoryMapped = true;
 
 			// Mark all the committed page in the bitmap in a thread-safe manner.
 			// We could unroll this to make it faster but this is not on the critical path anyway.
@@ -1118,8 +1116,7 @@ namespace UE::Sanitizer::RaceDetector::Platform {
 				uint8 PageBit = (uint8)(1 << (ShadowPage & 7));
 				if ((*BitmapPtr & PageBit) == 0)
 				{
-					std::atomic<uint8>* AtomicBitmapPtr = (std::atomic<uint8>*)BitmapPtr;
-					AtomicBitmapPtr->fetch_or(PageBit, std::memory_order_relaxed);
+					FPlatformAtomics::InterlockedOr((volatile int8*)BitmapPtr, (int8)PageBit);
 
 					DirtyShadowBitmapPage((UPTRINT)BitmapPtr);
 				}
