Fix crash in FError when using non static error details object.

diff --git a/Engine/Source/Runtime/Core/Private/Experimental/IO/IoStatusError.cpp b/Engine/Source/Runtime/Core/Private/Experimental/IO/IoStatusError.cpp
--- a/Engine/Source/Runtime/Core/Private/Experimental/IO/IoStatusError.cpp
+++ b/Engine/Source/Runtime/Core/Private/Experimental/IO/IoStatusError.cpp
@@ -11,13 +11,13 @@ namespace UE::UnifiedError::IoStore
 	class FIoStoreErrorDetails : public UE::UnifiedError::FRefCountedErrorDetails
 	{
 	private:
-		FIoStatus::FErrorMessage ErrorMessage;
+		FString ErrorMessage;
 		mutable FUtf8String CachedErrorName;
 		static const FText GenericErrorFormatString;
 	public:
-		FIoStoreErrorDetails(const FIoStatus::FErrorMessage& InErrorMessage)
+		FIoStoreErrorDetails(const FStringView& InErrorMessage)
 		{
-			FCString::Strncpy(ErrorMessage, InErrorMessage, FIoStatus::MaxErrorMessageLength);
+			ErrorMessage = InErrorMessage;
 		}
 
 		virtual const FUtf8String& GetErrorName(const FError& Error) const final
@@ -28,33 +28,28 @@ namespace UE::UnifiedError::IoStore
 			}
 			return CachedErrorName;
 		}
-		virtual const FText& GetErrorFormatString(const FError& Error) const final
+		virtual const FText GetErrorFormatString(const FError& Error) const final
 		{
 			return GenericErrorFormatString;
 		}
 
 		virtual void GetErrorProperties(const FError& Error, IErrorPropertyExtractor& OutProperties) const override
 		{
-			OutProperties.AddProperty(UTF8TEXTVIEW("IoStoreErrorMessage"), FStringView(ErrorMessage));
+			OutProperties.AddProperty(UTF8TEXTVIEW("ErrorCodeString"), CachedErrorName);
+			OutProperties.AddProperty(UTF8TEXTVIEW("ModuleIdString"), UE::UnifiedError::IoStore::StaticModuleId);
+			OutProperties.AddProperty(UTF8TEXTVIEW("ErrorCode"), Error.GetErrorCode());
+			OutProperties.AddProperty(UTF8TEXTVIEW("ModuleId"), Error.GetModuleId());
+			OutProperties.AddProperty(UTF8TEXTVIEW("IoStoreErrorMessage"), ErrorMessage);
 		}
 
 	};
 
 	const FText FIoStoreErrorDetails::GenericErrorFormatString = NSLOCTEXT("IoStore", "GenericErrorMessage", "{IoStoreErrorMessage}");
 
-	// allow access to the private FError implementation
-	class FIoStoreError : public FError
-	{
-	public:
-
-		FIoStoreError(const EIoErrorCode InErrorCode, const IErrorDetails* ErrorDetails) : FError(UE::UnifiedError::IoStore::StaticModuleId, (int32)(InErrorCode), ErrorDetails)
-		{
-		}
-	};
 
 	CORE_API UE::UnifiedError::FError ConvertError(const FIoStatus& Status)
 	{
-		FIoStoreErrorDetails* ErrorDetails = new FIoStoreErrorDetails(Status.GetErrorMessage());
-		return FIoStoreError(Status.GetErrorCode(), ErrorDetails);
+		TRefCountPtr<FIoStoreErrorDetails> ErrorDetails = new FIoStoreErrorDetails(Status.GetErrorMessage());
+		return FError(UE::UnifiedError::IoStore::StaticModuleId, (int32)(Status.GetErrorCode()), ErrorDetails);
 	}
 }
\ No newline at end of file
diff --git a/Engine/Source/Runtime/Core/Private/Experimental/UnifiedError/UnifiedError.cpp b/Engine/Source/Runtime/Core/Private/Experimental/UnifiedError/UnifiedError.cpp
--- a/Engine/Source/Runtime/Core/Private/Experimental/UnifiedError/UnifiedError.cpp
+++ b/Engine/Source/Runtime/Core/Private/Experimental/UnifiedError/UnifiedError.cpp
@@ -4,6 +4,7 @@
 #include "Internationalization/Text.h"
 #include "Internationalization/Internationalization.h"
 #include "Logging/StructuredLog.h"
+#include "Logging/StructuredLogFormat.h"
 
 
 namespace UE::UnifiedError
@@ -45,6 +46,12 @@ TRefCountPtr<const IErrorDetails> FError::GetInnerMostErrorDetails() const
 	}
 }
 
+FText FError::GetFormatErrorText() const
+{
+	return ErrorDetails->GetErrorFormatString(*this);
+}
+
+
 FText FError::GetErrorMessage() const
 {
 	FFormatNamedArguments Args;
@@ -225,20 +232,16 @@ void SerializeForLog(FCbWriter& Writer, const FError& Error)
 	Error.GetErrorProperties(ExtractToCbWriter);
 
 
-	// get the error message this error wants to produce
-	FText ErrorMessageText = Error.GetErrorMessage();
-	Writer.AddString(ANSITEXTVIEW("ErrorMessage"), ErrorMessageText.ToString());
+	// build a format string
+	const FText StructuredLogFormatString = NSLOCTEXT("UnifiedError", "StructuredLogErrorMessage", "{ModuleIdString}.{ErrorCodeString}({ModuleId}.{ErrorCode}): ");
+	FText ErrorFormatText = FText::Join(FText::FromString(TEXT("")), StructuredLogFormatString, Error.GetFormatErrorText());
+	Writer.AddString(ANSITEXTVIEW("ErrorFormatString"), ErrorFormatText.ToString());
 
 	// build up a formatted error string including some additional information about the error
-	TRefCountPtr<const FStaticErrorDetails> StaticErrorDetails = Error.GetStaticErrorDetails();
 	FFormatNamedArguments ErrorMessageArgs;
 	FTextFormatArgsPropertyExtractor ExtractToFTextArgs(ErrorMessageArgs);
-	StaticErrorDetails->GetErrorProperties(Error, ExtractToFTextArgs);;
-	ErrorMessageArgs.Add(TEXT("ErrorMessage"), ErrorMessageText);
-
-	
-	Writer.AddString(ANSITEXTVIEW("ErrorFormatString"), StaticErrorDetails->GetErrorFormatString(Error).ToString());
-	Writer.AddString("$text", FText::Format(NSLOCTEXT("UnifiedError", "StructuredLogErrorMessage", "{ModuleIdString}.{ErrorCodeString}({ModuleId}.{ErrorCode}): {ErrorMessage}"), ErrorMessageArgs).ToString());
+	Error.GetErrorProperties(ExtractToFTextArgs);
+	Writer.AddString("$text", FText::Format(ErrorFormatText, MoveTemp(ErrorMessageArgs)).ToString());
 
 	Writer.EndObject();
 }
@@ -258,7 +261,7 @@ FDynamicErrorDetails::~FDynamicErrorDetails()
 {
 }
 
-const FText& FDynamicErrorDetails::GetErrorFormatString(const FError & Error) const
+const FText FDynamicErrorDetails::GetErrorFormatString(const FError & Error) const
 {
 	check(InnerErrorDetails);
 	return InnerErrorDetails->GetErrorFormatString(Error);
@@ -290,12 +293,6 @@ FStaticErrorDetails::FStaticErrorDetails(const FAnsiStringView& InErrorName, con
 }
 
 
-/*
-FStaticErrorDetails::FStaticErrorDetails(const FUtf8StringView& InErrorName, const FText& InErrorFormatString)
-{
-	ErrorName = InErrorName.GetData();
-	ErrorFormatString = InErrorFormatString;
-}*/
 
 void FStaticErrorDetails::GetErrorProperties(const FError& Error, IErrorPropertyExtractor& OutProperties) const
 {
@@ -315,7 +312,7 @@ const FAnsiStringView& FStaticErrorDetails::GetModuleIdString() const
 	return ModuleName;
 }
 
-const FText& FStaticErrorDetails::GetErrorFormatString(const FError& Error) const
+const FText FStaticErrorDetails::GetErrorFormatString(const FError& Error) const
 {
 	return ErrorFormatString;
 }
diff --git a/Engine/Source/Runtime/Core/Public/Experimental/UnifiedError/UnifiedError.h b/Engine/Source/Runtime/Core/Public/Experimental/UnifiedError/UnifiedError.h
--- a/Engine/Source/Runtime/Core/Public/Experimental/UnifiedError/UnifiedError.h
+++ b/Engine/Source/Runtime/Core/Public/Experimental/UnifiedError/UnifiedError.h
@@ -79,7 +79,7 @@ namespace UE::UnifiedError
 		/// </summary>
 		/// <param name="Error"></param>
 		/// <returns></returns>
-		CORE_API virtual const FText& GetErrorFormatString(const FError& Error) const = 0;
+		CORE_API virtual const FText GetErrorFormatString(const FError& Error) const = 0;
 
 		/// <summary>
 		/// GetErrorProperties; used to expose error properties in name:value format, Error properties can be used for string formating functions, searching, exposing context to higher level stacks
@@ -153,7 +153,7 @@ namespace UE::UnifiedError
 	/// </summary>
 	class FDynamicErrorDetails : public FRefCountedErrorDetails
 	{
-		DECLARE_FERROR_DETAILS(UnifiedError, FDynamicErrorDetails);
+		DECLARE_FERROR_DETAILS_ABSTRACT(UnifiedError, FDynamicErrorDetails);
 	private:
 		TRefCountPtr<const IErrorDetails> InnerErrorDetails;
 	public:
@@ -170,7 +170,7 @@ namespace UE::UnifiedError
 		/// </summary>
 		/// <param name="Error"></param>
 		/// <returns></returns>
-		CORE_API virtual const FText& GetErrorFormatString(const FError& Error) const override;
+		CORE_API virtual const FText GetErrorFormatString(const FError& Error) const override;
 
 		/// <summary>
 		/// GetErrorProperties; by default pass through to the InnerErrorDetails.
@@ -182,21 +182,6 @@ namespace UE::UnifiedError
 	};
 
 
-	/// <summary>
-	/// FDefaultErrorDetails; serializable FStaticErrorDetails standin. 
-	///  This class is primarily used as a stand in for FStaticErrorDetails in cases of serialization over network or to disk.
-	/// </summary>
-	class FDefaultErrorDetails
-	{
-	private:
-		FAnsiStringView ErrorName;
-		FAnsiStringView ModuleName;
-		FText ErrorFormatString;
-	public:
-		FDefaultErrorDetails() = default;
-
-	};
-
 	/// <summary>
 	/// FStaticErrorDetails; static error details and members are statically allocated
 	///  Every error which uses DEFINE_ERROR will have FStaticErrorDetails generated for it
@@ -247,9 +232,7 @@ namespace UE::UnifiedError
 		/// </summary>
 		/// <param name="Error"></param>
 		/// <returns></returns>
-		CORE_API virtual const FText& GetErrorFormatString(const FError& Error) const final;
-
-
+		CORE_API virtual const FText GetErrorFormatString(const FError& Error) const final;
 
 
 		// IRefCountedObject implementation
@@ -307,9 +290,6 @@ namespace UE::UnifiedError
 		CORE_API const FAnsiStringView& GetModuleIdString() const;
 
 		CORE_API void PushErrorDetails(TRefCountPtr<IErrorDetails> InErrorDetails);
-
-		CORE_API FText GetErrorMessage() const;
-
 		template<typename DetailType>
 		TRefCountPtr<const DetailType> GetErrorDetails() const
 		{
@@ -325,6 +305,10 @@ namespace UE::UnifiedError
 			return nullptr;
 		}
 
+		CORE_API FText GetErrorMessage() const;
+		CORE_API FText GetFormatErrorText() const;
+
+
 		CORE_API TRefCountPtr<const FStaticErrorDetails> GetStaticErrorDetails() const;
 
 		CORE_API bool GetDetailByKey(const FUtf8StringView& KeyName, FString& Result) const;
@@ -344,12 +328,7 @@ namespace UE::UnifiedError
 		{
 			ErrorCode = 0;
 			ModuleId = 0;
-			ErrorDetails = nullptr;/*
-			if (ErrorDetails)
-			{
-				ErrorDetails->ReleaseRef();
-				ErrorDetails = nullptr;
-			}*/
+			ErrorDetails = nullptr;
 		}
 	private:
 		friend CORE_API bool operator==(const FError& Error, const FError& OtherError);
@@ -383,27 +362,29 @@ namespace UE::UnifiedError
 		FErrorRegistry() {}
 
 	public:
-		uint32 RegisterModule(const FStringView& ModuleName, uint32 ModuleId)
+		inline uint32 RegisterModule(const FStringView& ModuleName)
 		{
+			// todo: need to replace this with a stable hashing function
+			uint32 ModuleId = GetTypeHash(ModuleName);
 			checkf(ModuleNameMap.Contains(ModuleId) == false, TEXT("Module %s and %s are trying to register under module id %d"), ModuleName.GetData(), *ModuleNameMap.FindRef(ModuleId), ModuleId);
 			ModuleNameMap.Add(ModuleId, ModuleName.GetData());
 			return ModuleId;
 		}
 
-		void RegisterErrorCode(const FStringView& ErrorName, uint32 ModuleId, uint32 ErrorCode)
+		inline int32 RegisterErrorCode(const FStringView& ErrorName, int32 ModuleId, int32 ErrorCode)
 		{
-			int64 CombinedErrorCode = ((int64)ModuleId) << 32 | ErrorCode;
-			checkf(ErrorCodeNameMap.Contains(CombinedErrorCode) == false, TEXT("Error %s and %s are trying to register under same error code %d"), ErrorName.GetData(), *ErrorCodeNameMap.FindRef(CombinedErrorCode), CombinedErrorCode);
-			ErrorCodeNameMap.Add(CombinedErrorCode, ErrorName.GetData());
+			TPair<int32, int32> CombinedErrorId(ModuleId, ErrorCode);
+			checkf(ErrorCodeNameMap.Contains(CombinedErrorId) == false, TEXT("Error %s and %s are trying to register under same error code moduleid:%d errorcode:%d"), ErrorName.GetData(), *ErrorCodeNameMap.FindRef(CombinedErrorId), ModuleId, ErrorCode);
+			ErrorCodeNameMap.Add(CombinedErrorId, ErrorName.GetData());
+			return ErrorCode;
 		}
 
 	private:
-		TMap<uint32, FString> ModuleNameMap;
-		TMap<uint64, FString> ErrorCodeNameMap;
+		TMap<int32, FString> ModuleNameMap;
+		TMap<TPair<int32, int32>, FString> ErrorCodeNameMap;
 	};
 
 
-
 	class FTextFormatArgsPropertyExtractor : public IErrorPropertyExtractor
 	{
 	private:
@@ -449,14 +430,12 @@ namespace UE::UnifiedError
 		}
 	};
 
-
-
-
 	CORE_API void SerializeForLog(FCbWriter& Writer, const FError& OnlineError);
 
 } // namespace UnifiedError
 
 
+
 CORE_API FString LexToString(const UE::UnifiedError::FError& Error);
 /*
 
@@ -469,40 +448,34 @@ CORE_API constexpr uint32 GetTypeIdHash(const TCHAR* TypeName)
 
 
 
+
 #define DECLARE_ERROR_MODULE(ModuleName, ModuleId) \
 	namespace UE::UnifiedError { namespace ModuleName {\
-		static constexpr uint32 StaticModuleId = ModuleId;\
+		inline const int32 StaticModuleId = UE::UnifiedError::FErrorRegistry::Get().RegisterModule(TEXTVIEW(#ModuleName));\
 		static constexpr FAnsiStringView StaticModuleName = ANSITEXTVIEW(#ModuleName);\
-		FORCEINLINE int32 GetNextErrorCode()  \
-		{\
-			static int32 ErrorCodeCounter = 0;\
-			++ErrorCodeCounter;\
-			return ErrorCodeCounter;\
-		}\
 	}}
 
 
 #define DECLARE_ERROR_INTERNAL(ErrorName, ErrorCode, ModuleName, FormatString) \
 	namespace UE::UnifiedError { namespace ModuleName { namespace ErrorName {\
 			static constexpr FAnsiStringView StaticErrorName = ANSITEXTVIEW(#ErrorName);\
-			FORCEINLINE const FStaticErrorDetails& GetStaticErrorDetails() \
+			FORCEINLINE TRefCountPtr<const FStaticErrorDetails> GetStaticErrorDetails() \
 			{\
 				static FStaticErrorDetails StaticErrorDetails(StaticErrorName, UE::UnifiedError::ModuleName::StaticModuleName, FormatString);\
-				return StaticErrorDetails;\
+				return TRefCountPtr<const FStaticErrorDetails>(&StaticErrorDetails);\
 			}\
+			inline const int32 ErrorCodeId = UE::UnifiedError::FErrorRegistry::Get().RegisterErrorCode(TEXTVIEW(#ErrorName), UE::UnifiedError::ModuleName::StaticModuleId, ErrorCode);\
 			FORCEINLINE int32 GetErrorCodeId() \
-			{ \
-				static int32 ErrorCodeId = ErrorCode < 0 ? UE::UnifiedError::ModuleName::GetNextErrorCode() : ErrorCode; \
-				return ErrorCodeId; \
-			} \
-			static const int32 ErrorCodeId = GetErrorCodeId(); \
+			{\
+				return ErrorCodeId;\
+			}\
 	}}}
 
 
 #define DECLARE_ERROR(ErrorName, ErrorCode, ModuleName, FormatString) \
 	DECLARE_ERROR_INTERNAL(ErrorName, ErrorCode, ModuleName, FormatString) \
 	namespace UE::UnifiedError { namespace ModuleName { namespace ErrorName {\
-			FORCEINLINE FError MakeError() { return FError(UE::UnifiedError::ModuleName::StaticModuleId, UE::UnifiedError::ModuleName::ErrorName::ErrorCodeId, &GetStaticErrorDetails()); }\
+			FORCEINLINE FError MakeError() { return FError(UE::UnifiedError::ModuleName::StaticModuleId, GetErrorCodeId(), TRefCountPtr<const IErrorDetails>(GetStaticErrorDetails().GetReference())); }\
 			FORCEINLINE const FError& GetStaticError() \
 			{\
 				static FError StaticError = MakeError();\
@@ -510,7 +483,7 @@ CORE_API constexpr uint32 GetTypeIdHash(const TCHAR* TypeName)
 			}\
 	}}}
 
-#define DEFINE_ERROR(ErrorName, ModuleName)
+
 
 
 #define DECLARE_ERROR_ONEPARAM(ErrorName, ErrorCode, ModuleName, FormatString, ParamOneType, ParamOneName, ParamOneDefault) \
@@ -523,7 +496,7 @@ CORE_API constexpr uint32 GetTypeIdHash(const TCHAR* TypeName)
 				ParamOneType ParamOneName; \
 				FDetails() : ParamOneName(ParamOneDefault) {}\
 			public: \
-				FDetails(const ParamOneType& In##ParamOneName, const IErrorDetails* InErrorDetails) : FDynamicErrorDetails(InErrorDetails) \
+				FDetails(const ParamOneType& In##ParamOneName, TRefCountPtr<const IErrorDetails> InErrorDetails) : FDynamicErrorDetails(InErrorDetails) \
 				{ \
 					ParamOneName = In##ParamOneName; \
 				} \
@@ -539,7 +512,7 @@ CORE_API constexpr uint32 GetTypeIdHash(const TCHAR* TypeName)
 			}; \
 			FORCEINLINE FError MakeError(ParamOneType ParamOneName = ParamOneDefault) \
 			{ \
-				FDetails* NewDetails = new FDetails(ParamOneName, &GetStaticErrorDetails()); \
+				FDetails* NewDetails = new FDetails(ParamOneName, TRefCountPtr<const IErrorDetails>(GetStaticErrorDetails().GetReference())); \
 				return FError(UE::UnifiedError::ModuleName::StaticModuleId, GetErrorCodeId(), NewDetails); \
 			} \
 			FORCEINLINE const FError& GetStaticError() \
diff --git a/Engine/Source/Runtime/Engine/Classes/Engine/Experimental/StreamableManagerError.h b/Engine/Source/Runtime/Engine/Classes/Engine/Experimental/StreamableManagerError.h
--- a/Engine/Source/Runtime/Engine/Classes/Engine/Experimental/StreamableManagerError.h
+++ b/Engine/Source/Runtime/Engine/Classes/Engine/Experimental/StreamableManagerError.h
@@ -54,9 +54,13 @@ namespace UE::UnifiedError::StreamableManager
 			FDynamicErrorDetails::GetErrorProperties(Error, OutProperties);
 		}
 
-		virtual const FText& GetErrorFormatString(const FError& Error) const override
+		virtual const FText GetErrorFormatString(const FError& Error) const override
 		{
-			return ErrorFormatString;
+			const FText InnerFormatString = GetInnerErrorDetails()->GetErrorFormatString(Error);
+
+			FText MyFormatString = NSLOCTEXT("StreamableManager", "DetailsFormatString", "{FStreamableManagerAdditionalContext}");
+
+			return FText::Join(NSLOCTEXT("StreamableManager", "DetailsFormatStringDelimiter", " "), InnerFormatString, MyFormatString);
 		}
 
 		const FString& GetSoftObjectPath() const { return SoftObjectPath; }
@@ -65,4 +69,3 @@ namespace UE::UnifiedError::StreamableManager
 	UE::UnifiedError::FError GetStreamableError(EAsyncLoadingResult::Type Result);
 }
 
-//DECLARE_DYNAMIC_ERROR_CONTEXT(StreamableManagerContext, )
diff --git a/Engine/Source/Runtime/Engine/Private/StreamableManager.cpp b/Engine/Source/Runtime/Engine/Private/StreamableManager.cpp
--- a/Engine/Source/Runtime/Engine/Private/StreamableManager.cpp
+++ b/Engine/Source/Runtime/Engine/Private/StreamableManager.cpp
@@ -2,6 +2,7 @@
 
 #include "Engine/StreamableManager.h"
 #include "Experimental/UnifiedError/UnifiedError.h"
+#include "Engine/Experimental/StreamableManagerError.h"
 #include "Experimental/IO/IoStatusError.h"
 #include "UObject/ICookInfo.h"
 #include "UObject/ObjectRedirector.h"
