Fixed format specifier warnings for "%lc" in FGenericWidePlatformString::GetVarArgs when building for Clang under Windows.

diff --git a/Engine/Source/Runtime/Core/Private/GenericPlatform/GenericWidePlatformString.cpp b/Engine/Source/Runtime/Core/Private/GenericPlatform/GenericWidePlatformString.cpp
--- a/Engine/Source/Runtime/Core/Private/GenericPlatform/GenericWidePlatformString.cpp
+++ b/Engine/Source/Runtime/Core/Private/GenericPlatform/GenericWidePlatformString.cpp
@@ -598,7 +598,7 @@ int32 FGenericWidePlatformString::GetVarArgs( WIDECHAR* Dest, SIZE_T DestSize, c
 			{
 				if (Src[1] == 0)
 				{
-					printf("Unknown percent [%lc] in FGenericWidePlatformString::GetVarArgs() [%s]\n.", Src[0], TCHAR_TO_ANSI(Fmt));
+					printf("Unknown percent [%lc] in FGenericWidePlatformString::GetVarArgs() [%s]\n.", (wint_t)Src[0], TCHAR_TO_ANSI(Fmt));
 					Src++;  // skip it, I guess.
 					break;
 				}
@@ -675,7 +675,7 @@ int32 FGenericWidePlatformString::GetVarArgs( WIDECHAR* Dest, SIZE_T DestSize, c
 
 				if (RemainingSize < 2)
 				{
-					printf("Unknown percent [%lc] in FGenericWidePlatformString::GetVarArgs() [%s]\n.", Src[0], TCHAR_TO_ANSI(Fmt));
+					printf("Unknown percent [%lc] in FGenericWidePlatformString::GetVarArgs() [%s]\n.", (wint_t)Src[0], TCHAR_TO_ANSI(Fmt));
 					Src++;  // skip it, I guess.
 					break;
 				}
@@ -734,7 +734,7 @@ int32 FGenericWidePlatformString::GetVarArgs( WIDECHAR* Dest, SIZE_T DestSize, c
 				// At this point we expect RemainingSize >= 3, but we know it's at least 2.
 				if (RemainingSize < 3)
 				{
-					printf("Unknown percent [%lc%lc] in FGenericWidePlatformString::GetVarArgs() [%s]\n.", Src[0], Src[1], TCHAR_TO_ANSI(Fmt));
+					printf("Unknown percent [%lc%lc] in FGenericWidePlatformString::GetVarArgs() [%s]\n.", (wint_t)Src[0], (wint_t)Src[1], TCHAR_TO_ANSI(Fmt));
 					Src++;  // skip it, I guess.
 					break;
 				}
@@ -744,7 +744,7 @@ int32 FGenericWidePlatformString::GetVarArgs( WIDECHAR* Dest, SIZE_T DestSize, c
 				{
 					if (Src[1] != 'h' || !CharIsIntegerFormatSpecifier(Src[2]))
 					{
-						printf("Unknown percent [%lc%lc%lc] in FGenericWidePlatformString::GetVarArgs() [%s]\n.", Src[0], Src[1], Src[2], TCHAR_TO_ANSI(Fmt));
+						printf("Unknown percent [%lc%lc%lc] in FGenericWidePlatformString::GetVarArgs() [%s]\n.", (wint_t)Src[0], (wint_t)Src[1], (wint_t)Src[2], TCHAR_TO_ANSI(Fmt));
 						Src++;  // skip it, I guess.
 						break;
 					}
@@ -774,7 +774,7 @@ int32 FGenericWidePlatformString::GetVarArgs( WIDECHAR* Dest, SIZE_T DestSize, c
 				{
 					if (Src[1] != 'l' || !CharIsIntegerFormatSpecifier(Src[2]))
 					{
-						printf("Unknown percent [%lc%lc%lc] in FGenericWidePlatformString::GetVarArgs() [%s]\n.", Src[0], Src[1], Src[2], TCHAR_TO_ANSI(Fmt));
+						printf("Unknown percent [%lc%lc%lc] in FGenericWidePlatformString::GetVarArgs() [%s]\n.", (wint_t)Src[0], (wint_t)Src[1], (wint_t)Src[2], TCHAR_TO_ANSI(Fmt));
 						Src++;  // skip it, I guess.
 						break;
 					}
@@ -783,20 +783,20 @@ int32 FGenericWidePlatformString::GetVarArgs( WIDECHAR* Dest, SIZE_T DestSize, c
 				{
 					if (RemainingSize < 4)
 					{
-						printf("Unknown percent [%lc%lc%lc] in FGenericWidePlatformString::GetVarArgs() [%s]\n.", Src[0], Src[1], Src[2], TCHAR_TO_ANSI(Fmt));
+						printf("Unknown percent [%lc%lc%lc] in FGenericWidePlatformString::GetVarArgs() [%s]\n.", (wint_t)Src[0], (wint_t)Src[1], (wint_t)Src[2], TCHAR_TO_ANSI(Fmt));
 						Src++;  // skip it, I guess.
 						break;
 					}
 					else if (Src[1] != '6' || Src[2] != '4' || !CharIsIntegerFormatSpecifier(Src[3]))
 					{
-						printf("Unknown percent [%lc%lc%lc%lc] in FGenericWidePlatformString::GetVarArgs() [%s]\n.", Src[0], Src[1], Src[2], Src[3], TCHAR_TO_ANSI(Fmt));
+						printf("Unknown percent [%lc%lc%lc%lc] in FGenericWidePlatformString::GetVarArgs() [%s]\n.", (wint_t)Src[0], (wint_t)Src[1], (wint_t)Src[2], (wint_t)Src[3], TCHAR_TO_ANSI(Fmt));
 						Src++;  // skip it, I guess.
 						break;
 					}
 				}
 				else
 				{
-					printf("Unknown percent [%lc%lc%lc] in FGenericWidePlatformString::GetVarArgs() [%s]\n.", Src[0], Src[1], Src[2], TCHAR_TO_ANSI(Fmt));
+					printf("Unknown percent [%lc%lc%lc] in FGenericWidePlatformString::GetVarArgs() [%s]\n.", (wint_t)Src[0], (wint_t)Src[1], (wint_t)Src[2], TCHAR_TO_ANSI(Fmt));
 					Src++;  // skip it, I guess.
 					break;
 				}
@@ -871,7 +871,7 @@ int32 FGenericWidePlatformString::GetVarArgs( WIDECHAR* Dest, SIZE_T DestSize, c
 				{
 					// We should print what we have written into AnsiNum but ensure we null terminate before printing
 					AnsiNum[UE_ARRAY_COUNT(AnsiNum) - 1] = '\0';
-					checkf(0, TEXT("Attempting to read past the size our buffer. Buffer Size: %d Size to read: %d. Current contents: '%s'\n"), UE_ARRAY_COUNT(AnsiNum), RetCnt, UTF8_TO_TCHAR(AnsiNum));
+					checkf(0, TEXT("Attempting to read past the size our buffer. Buffer Size: %" SIZE_T_FMT " Size to read: %d. Current contents: '%s'\n"), UE_ARRAY_COUNT(AnsiNum), RetCnt, UTF8_TO_TCHAR(AnsiNum));
 				}
 				if (RetCnt < 0 || !DestIter.Write(AnsiNum, RetCnt))
 				{
