Add final to platform mutex types, which impacts FCriticalSection and FRWLock

diff --git a/Engine/Source/Runtime/Core/Public/HAL/CriticalSection.h b/Engine/Source/Runtime/Core/Public/HAL/CriticalSection.h
--- a/Engine/Source/Runtime/Core/Public/HAL/CriticalSection.h
+++ b/Engine/Source/Runtime/Core/Public/HAL/CriticalSection.h
@@ -10,40 +10,41 @@ namespace UE
 /**
  * Adapter for FPlatformSharedMutex to the API expected for FRWLock.
  */
-class FPlatformRWLock : private FPlatformSharedMutex
+class FPlatformRWLock final
 {
 public:
-	using FPlatformSharedMutex::FPlatformSharedMutex;
-
 	FORCEINLINE bool TryWriteLock()
 	{
-		return FPlatformSharedMutex::TryLock();
+		return Mutex.TryLock();
 	}
 
 	FORCEINLINE void WriteLock()
 	{
-		FPlatformSharedMutex::Lock();
+		Mutex.Lock();
 	}
 
 	FORCEINLINE void WriteUnlock()
 	{
-		FPlatformSharedMutex::Unlock();
+		Mutex.Unlock();
 	}
 
 	FORCEINLINE bool TryReadLock()
 	{
-		return FPlatformSharedMutex::TryLockShared();
+		return Mutex.TryLockShared();
 	}
 
 	FORCEINLINE void ReadLock()
 	{
-		FPlatformSharedMutex::LockShared();
+		Mutex.LockShared();
 	}
 
 	FORCEINLINE void ReadUnlock()
 	{
-		FPlatformSharedMutex::UnlockShared();
+		Mutex.UnlockShared();
 	}
+
+private:
+	FPlatformSharedMutex Mutex;
 };
 
 } // UE
diff --git a/Engine/Source/Runtime/Core/Public/HAL/PThreadsRecursiveMutex.h b/Engine/Source/Runtime/Core/Public/HAL/PThreadsRecursiveMutex.h
--- a/Engine/Source/Runtime/Core/Public/HAL/PThreadsRecursiveMutex.h
+++ b/Engine/Source/Runtime/Core/Public/HAL/PThreadsRecursiveMutex.h
@@ -19,7 +19,7 @@ namespace UE
  *
  * Prefer FRecursiveMutex.
  */
-class FPThreadsRecursiveMutex
+class FPThreadsRecursiveMutex final
 {
 public:
 	FPThreadsRecursiveMutex(const FPThreadsRecursiveMutex&) = delete;
diff --git a/Engine/Source/Runtime/Core/Public/HAL/PThreadsSharedMutex.h b/Engine/Source/Runtime/Core/Public/HAL/PThreadsSharedMutex.h
--- a/Engine/Source/Runtime/Core/Public/HAL/PThreadsSharedMutex.h
+++ b/Engine/Source/Runtime/Core/Public/HAL/PThreadsSharedMutex.h
@@ -20,7 +20,7 @@ namespace UE
  *
  * Prefer FSharedMutex.
  */
-class FPThreadsSharedMutex
+class FPThreadsSharedMutex final
 {
 public:
 	FPThreadsSharedMutex(const FPThreadsSharedMutex&) = delete;
diff --git a/Engine/Source/Runtime/Core/Public/Mac/MacPlatformMutex.h b/Engine/Source/Runtime/Core/Public/Mac/MacPlatformMutex.h
--- a/Engine/Source/Runtime/Core/Public/Mac/MacPlatformMutex.h
+++ b/Engine/Source/Runtime/Core/Public/Mac/MacPlatformMutex.h
@@ -11,7 +11,7 @@ namespace UE
 {
 
 /** A system-wide mutex for Mac. Uses exclusive file locking. */
-class FMacSystemWideMutex
+class FMacSystemWideMutex final
 {
 public:
 	FMacSystemWideMutex(const FMacSystemWideMutex&) = delete;
diff --git a/Engine/Source/Runtime/Core/Public/Unix/UnixPlatformMutex.h b/Engine/Source/Runtime/Core/Public/Unix/UnixPlatformMutex.h
--- a/Engine/Source/Runtime/Core/Public/Unix/UnixPlatformMutex.h
+++ b/Engine/Source/Runtime/Core/Public/Unix/UnixPlatformMutex.h
@@ -13,7 +13,7 @@ namespace UE
 {
 
 /** A system-wide mutex for Unix. Uses exclusive file locking. */
-class FUnixSystemWideMutex
+class FUnixSystemWideMutex final
 {
 public:
 	FUnixSystemWideMutex(const FUnixSystemWideMutex&) = delete;
diff --git a/Engine/Source/Runtime/Core/Public/Windows/WindowsPlatformMutex.h b/Engine/Source/Runtime/Core/Public/Windows/WindowsPlatformMutex.h
--- a/Engine/Source/Runtime/Core/Public/Windows/WindowsPlatformMutex.h
+++ b/Engine/Source/Runtime/Core/Public/Windows/WindowsPlatformMutex.h
@@ -16,7 +16,7 @@ namespace UE
  *
  * Prefer FRecursiveMutex.
  */
-class FWindowsRecursiveMutex
+class FWindowsRecursiveMutex final
 {
 public:
 	FWindowsRecursiveMutex(const FWindowsRecursiveMutex&) = delete;
@@ -61,7 +61,7 @@ private:
  *
  * SRWLOCK: https://learn.microsoft.com/en-us/windows/win32/sync/slim-reader-writer--srw--locks
  */
-class FWindowsSharedMutex
+class FWindowsSharedMutex final
 {
 public:
 	FWindowsSharedMutex(const FWindowsSharedMutex&) = delete;
@@ -122,7 +122,7 @@ private:
 };
 
 /** A system-wide mutex for Windows. Uses a named mutex. */
-class FWindowsSystemWideMutex
+class FWindowsSystemWideMutex final
 {
 public:
 	FWindowsSystemWideMutex(const FWindowsSystemWideMutex&) = delete;
