Avoid calling engine exit if already exiting

diff --git a/Engine/Source/Runtime/Core/Private/Windows/WindowsPlatformCrashContext.cpp b/Engine/Source/Runtime/Core/Private/Windows/WindowsPlatformCrashContext.cpp
--- a/Engine/Source/Runtime/Core/Private/Windows/WindowsPlatformCrashContext.cpp
+++ b/Engine/Source/Runtime/Core/Private/Windows/WindowsPlatformCrashContext.cpp
@@ -1744,8 +1744,10 @@ LONG WINAPI EngineUnhandledExceptionFilter(LPEXCEPTION_POINTERS ExceptionInfo)
 #else
 
 	GIsCriticalError = true;
-	FPlatformMisc::RequestExit(true, TEXT("WindowsPlatformCrashContext.EngineUnhandledExceptionFilter"));
-
+	if (!IsEngineExitRequested())
+	{
+		FPlatformMisc::RequestExit(true, TEXT("WindowsPlatformCrashContext.EngineUnhandledExceptionFilter"));
+	}
 	return EXCEPTION_CONTINUE_SEARCH; // Not really important, RequestExit() terminates the process just above.
 #endif // WINDOWS_USE_WER
 }
