Added support for %hs not being supported by vsnprintf.

diff --git a/Engine/Source/Runtime/Core/Public/GenericPlatform/GenericWidePlatformString.h b/Engine/Source/Runtime/Core/Public/GenericPlatform/GenericWidePlatformString.h
--- a/Engine/Source/Runtime/Core/Public/GenericPlatform/GenericWidePlatformString.h
+++ b/Engine/Source/Runtime/Core/Public/GenericPlatform/GenericWidePlatformString.h
@@ -346,7 +346,36 @@ PRAGMA_ENABLE_DEPRECATION_WARNINGS
 
 	static int32 GetVarArgs( ANSICHAR* Dest, SIZE_T DestSize, const ANSICHAR*& Fmt, va_list ArgPtr )
 	{
+#if PLATFORM_USE_S_SPEC_FOR_NARROWCHAR_IN_VSPRINTF
+		// fix up the Fmt string, as fast as possible, without using an FAnsiString
+		ANSICHAR* NewFormat = (ANSICHAR*)alloca(Strlen(Fmt) + 1);
+
+		const ANSICHAR* From = Fmt;
+		ANSICHAR* To = NewFormat;
+		for (;;)
+		{
+			ANSICHAR FromCh = *From++;
+			*To++ = FromCh;
+			if (FromCh != '%')
+			{
+				if (!FromCh)
+				{
+					break;
+				}
+
+				continue;
+			}
+
+			if (From[0] == 'h' && From[1] == 's')
+			{
+				// Skip the h and copy the s next time to give %s
+				++From;
+			}
+		}
+		int32 Result = vsnprintf(Dest, DestSize, NewFormat, ArgPtr);
+#else
 		int32 Result = vsnprintf(Dest, DestSize, Fmt, ArgPtr);
+#endif
 		return (Result != -1 && Result < (int32)DestSize) ? Result : -1;
 	}
 
diff --git a/Engine/Source/Runtime/Core/Public/HAL/Platform.h b/Engine/Source/Runtime/Core/Public/HAL/Platform.h
--- a/Engine/Source/Runtime/Core/Public/HAL/Platform.h
+++ b/Engine/Source/Runtime/Core/Public/HAL/Platform.h
@@ -231,9 +231,14 @@
 #ifndef UE_VALIDATE_FORMAT_STRINGS
 	#define UE_VALIDATE_FORMAT_STRINGS 0
 #endif
+// Does vswprintf want %ls instead of %s when displaying wide strings?
 #ifndef PLATFORM_USE_LS_SPEC_FOR_WIDECHAR
 	#define PLATFORM_USE_LS_SPEC_FOR_WIDECHAR	1
 #endif
+// Does vsprintf want %s instead of %hs when displaying narrow strings?
+#ifndef PLATFORM_USE_S_SPEC_FOR_NARROWCHAR_IN_VSPRINTF
+	#define PLATFORM_USE_S_SPEC_FOR_NARROWCHAR_IN_VSPRINTF	0
+#endif
 #ifndef PLATFORM_USE_SYSTEM_VSWPRINTF
 	#define PLATFORM_USE_SYSTEM_VSWPRINTF		1
 #endif
