Fixed incorrect format string specifiers in Core.

diff --git a/Engine/Source/Runtime/Core/Private/Compression/OodleDataCompression.cpp b/Engine/Source/Runtime/Core/Private/Compression/OodleDataCompression.cpp
--- a/Engine/Source/Runtime/Core/Private/Compression/OodleDataCompression.cpp
+++ b/Engine/Source/Runtime/Core/Private/Compression/OodleDataCompression.cpp
@@ -369,7 +369,7 @@ struct OodleScratchBuffers
 					OodleScratches[i].OodleScratchMemory = FMemory::Malloc(OodleScratchMemorySize);
 					if (OodleScratches[i].OodleScratchMemory == nullptr)
 					{
-						UE_LOG(OodleDataCompression, Error, TEXT("OodleDecode: Failed to allocate scratch buffer %d bytes!"), OodleScratchMemorySize);
+						UE_LOG(OodleDataCompression, Error, TEXT("OodleDecode: Failed to allocate scratch buffer %" PTRINT_FMT " bytes!"), OodleScratchMemorySize);
 						return OODLELZ_FAILED;
 					}
 				}
@@ -472,7 +472,7 @@ struct OodleScratchBuffers
 					OodleScratches[i].OodleScratchMemory = FMemory::Malloc(OodleScratchMemorySize);
 					if (OodleScratches[i].OodleScratchMemory == nullptr)
 					{
-						UE_LOG(OodleDataCompression, Error, TEXT("OodleEncode: Failed to allocate scratch buffer %d bytes!"), OodleScratchMemorySize);
+						UE_LOG(OodleDataCompression, Error, TEXT("OodleEncode: Failed to allocate scratch buffer %" PTRINT_FMT " bytes!"), OodleScratchMemorySize);
 						return OODLELZ_FAILED;
 					}
 				}
diff --git a/Engine/Source/Runtime/Core/Private/Containers/String.cpp.inl b/Engine/Source/Runtime/Core/Private/Containers/String.cpp.inl
--- a/Engine/Source/Runtime/Core/Private/Containers/String.cpp.inl
+++ b/Engine/Source/Runtime/Core/Private/Containers/String.cpp.inl
@@ -1819,7 +1819,7 @@ FArchive& operator<<(FArchive& Ar, UE_STRING_CLASS& Str)
 		if ((MaxSerializeSize > 0) && (SaveNum > MaxSerializeSize))
 		{
 			Ar.SetCriticalError();
-			UE_LOG(LogCore, Error, TEXT("String is too large (Size: %i, Max: %i)"), SaveNum, MaxSerializeSize);
+			UE_LOG(LogCore, Error, TEXT("String is too large (Size: %i, Max: %" INT64_FMT ")"), SaveNum, MaxSerializeSize);
 			return Ar;
 		}
 
diff --git a/Engine/Source/Runtime/Core/Private/Containers/Ticker.cpp b/Engine/Source/Runtime/Core/Private/Containers/Ticker.cpp
--- a/Engine/Source/Runtime/Core/Private/Containers/Ticker.cpp
+++ b/Engine/Source/Runtime/Core/Private/Containers/Ticker.cpp
@@ -70,7 +70,7 @@ void FTSTicker::Tick(float DeltaTime)
 		while (TOptional<FElementPtr> AddedElement = AddedElements.Dequeue())
 		{
 			uint64 State = AddedElement.GetValue()->State.load(std::memory_order_relaxed);
-			checkf(GetThreadId(State) == 0, TEXT("Invalid state %u"), State);
+			checkf(GetThreadId(State) == 0, TEXT("Invalid state %" UINT64_FMT), State);
 			if (State == FElement::DefaultState)
 			{
 				Elements.Add(MoveTemp(AddedElement.GetValue()));
@@ -99,7 +99,7 @@ void FTSTicker::Tick(float DeltaTime)
 			FElementPtr& Element = Elements[ElementIdx];
 			// set the execution state
 			uint64 PrevState = Element->State.fetch_add(CurrentThreadId, std::memory_order_acquire); // "acquire" to prevent anything between this and `ClearExecutionFlag()` to be reordered outside of this scope, is coupled with "release" in `ClearExecutionFlag`
-			checkf(GetThreadId(PrevState) == 0, TEXT("Invalid state %u"), PrevState);
+			checkf(GetThreadId(PrevState) == 0, TEXT("Invalid state %" UINT64_FMT), PrevState);
 			auto ClearExecutionFlag = [CurrentThreadId](FElementPtr& Element)
 			{
 				return (uint8)(Element->State.fetch_sub(CurrentThreadId, std::memory_order_release) - CurrentThreadId);
@@ -123,7 +123,7 @@ void FTSTicker::Tick(float DeltaTime)
 				PrevState = ClearExecutionFlag(Element); // it can be removed during execution
 				if (PrevState != FElement::RemovedState)
 				{
-					checkf(PrevState == FElement::DefaultState, TEXT("Invalid state %u"), PrevState);
+					checkf(PrevState == FElement::DefaultState, TEXT("Invalid state %" UINT64_FMT), PrevState);
 					Element->FireTime = CurrentTime.load(std::memory_order_relaxed) + Element->DelayTime;
 					TickedElements.Add(MoveTemp(Element));
 				}
@@ -131,7 +131,7 @@ void FTSTicker::Tick(float DeltaTime)
 			else
 			{
 				PrevState = ClearExecutionFlag(Element);
-				checkf(PrevState == FElement::DefaultState || PrevState == FElement::RemovedState, TEXT("Invalid state %u"), PrevState);
+				checkf(PrevState == FElement::DefaultState || PrevState == FElement::RemovedState, TEXT("Invalid state %" UINT64_FMT), PrevState);
 			}
 		}
 
diff --git a/Engine/Source/Runtime/Core/Private/FileCache/FileCache.cpp b/Engine/Source/Runtime/Core/Private/FileCache/FileCache.cpp
--- a/Engine/Source/Runtime/Core/Private/FileCache/FileCache.cpp
+++ b/Engine/Source/Runtime/Core/Private/FileCache/FileCache.cpp
@@ -315,7 +315,7 @@ void FileCache_PostIoStoreCompressionBlockSize(uint32 InCompressionBlockSize, FS
 		// We can direct check since CacheSlotID is correct (cvars resolved).
 		if (InCompressionBlockSize != CacheSlotID::GetSize())
 		{
-			UE_LOG(LogStreamingFileCache, Warning, TEXT("IoStore container %s has a different block sizes than FileCache (%d vs %d)!"), *InContainerFilePath, InCompressionBlockSize, CacheSlotID::GetSize());
+			UE_LOG(LogStreamingFileCache, Warning, TEXT("IoStore container %s has a different block sizes than FileCache (%d vs %" INT64_FMT ")!"), *InContainerFilePath, InCompressionBlockSize, CacheSlotID::GetSize());
 			UE_LOG(LogStreamingFileCache, Warning, TEXT("	Check your IoStore compression block size (Project Settings -> 'Package Compression Commandline Options'"));
 			UE_LOG(LogStreamingFileCache, Warning, TEXT("	and your File Cache block size (fc.BlockSize cvar). They should match!"));
 		}
@@ -358,7 +358,7 @@ FFileCache::FFileCache(int32 NumSlots)
 	else if (IoStoreCompressionBlockSize && // If we load without any iostore containers (i.e. editor) we don't need to warn.
 		IoStoreCompressionBlockSize != CacheSlotID::GetSize())
 	{
-		UE_LOG(LogStreamingFileCache, Warning, TEXT("IoStore containers have a different block sizes than FileCache (%d vs %d)!"), IoStoreCompressionBlockSize, CacheSlotID::GetSize());
+		UE_LOG(LogStreamingFileCache, Warning, TEXT("IoStore containers have a different block sizes than FileCache (%d vs %" INT64_FMT ")!"), IoStoreCompressionBlockSize, CacheSlotID::GetSize());
 		UE_LOG(LogStreamingFileCache, Warning, TEXT("	Check your IoStore compression block size (Project Settings -> 'Package Compression Commandline Options'"));
 		UE_LOG(LogStreamingFileCache, Warning, TEXT("	and your File Cache block size (fc.BlockSize cvar). They should match!"));
 	}
@@ -744,7 +744,7 @@ IMemoryReadStreamRef FFileCacheHandle::ReadData(FGraphEventArray& OutCompletionE
 		// not enough free slots in the cache to service this request
 		CacheLock.Unlock();
 
-		UE_LOG(LogStreamingFileCache, Verbose, TEXT("ReadData(%d, %d) is skipping cache, cache is full"), Offset, BytesToRead);
+		UE_LOG(LogStreamingFileCache, Verbose, TEXT("ReadData(%" INT64_FMT ", %" INT64_FMT ") is skipping cache, cache is full"), Offset, BytesToRead);
 		return ReadDataUncached(OutCompletionEvents, Offset, BytesToRead, Priority);
 	}
 
diff --git a/Engine/Source/Runtime/Core/Private/GenericPlatform/GenericPlatformCrashContext.cpp b/Engine/Source/Runtime/Core/Private/GenericPlatform/GenericPlatformCrashContext.cpp
--- a/Engine/Source/Runtime/Core/Private/GenericPlatform/GenericPlatformCrashContext.cpp
+++ b/Engine/Source/Runtime/Core/Private/GenericPlatform/GenericPlatformCrashContext.cpp
@@ -370,7 +370,7 @@ void FGenericCrashContext::Initialize()
 	NCached::Set(NCached::Session.SymbolsLabel, *Symbols);
 	if (Symbols.Len() >= UE_ARRAY_COUNT(NCached::Session.SymbolsLabel))
 	{
-		UE_LOG(LogInit, Error, TEXT("Symbols label too long (%d) for field size(%d), truncated. This may cause problems with crash report symbolication."),
+		UE_LOG(LogInit, Error, TEXT("Symbols label too long (%d) for field size(%" SIZE_T_FMT "), truncated. This may cause problems with crash report symbolication."),
 			Symbols.Len(), UE_ARRAY_COUNT(NCached::Session.SymbolsLabel));
 	}
 
diff --git a/Engine/Source/Runtime/Core/Private/HAL/FileManagerGeneric.cpp b/Engine/Source/Runtime/Core/Private/HAL/FileManagerGeneric.cpp
--- a/Engine/Source/Runtime/Core/Private/HAL/FileManagerGeneric.cpp
+++ b/Engine/Source/Runtime/Core/Private/HAL/FileManagerGeneric.cpp
@@ -849,7 +849,7 @@ void FArchiveFileReaderGeneric::Serialize( void* V, int64 Length )
 	if (Pos + Length > Size)
 	{
 		SetError();
-		UE_CLOG( !IsSilent(), LogFileManager, Error, TEXT("Requested read of %d bytes when %d bytes remain (file=%s, size=%d)"), Length, Size-Pos, *Filename, Size);
+		UE_CLOG( !IsSilent(), LogFileManager, Error, TEXT("Requested read of %" INT64_FMT " bytes when %" INT64_FMT " bytes remain (file=%s, size=%" INT64_FMT ")"), Length, Size-Pos, *Filename, Size);
 		return;
 	}
 
diff --git a/Engine/Source/Runtime/Core/Private/HAL/LowLevelMemTracker.cpp b/Engine/Source/Runtime/Core/Private/HAL/LowLevelMemTracker.cpp
--- a/Engine/Source/Runtime/Core/Private/HAL/LowLevelMemTracker.cpp
+++ b/Engine/Source/Runtime/Core/Private/HAL/LowLevelMemTracker.cpp
@@ -6172,7 +6172,7 @@ void FLLMCsvWriter::WriteHeader(const FTagData* OverrideTrackedTotalTagData, con
 	if (ColumnTitleTotalSize >= CVarLLMHeaderMaxSize.GetValueOnAnyThread())
 	{
 		UE_LOG(LogHAL, Error,
-			TEXT("LLM column titles have overflowed, LLM CSM data will be corrupted. Increase CVarLLMHeaderMaxSize > %d"),
+			TEXT("LLM column titles have overflowed, LLM CSM data will be corrupted. Increase CVarLLMHeaderMaxSize > %" INT64_FMT),
 			ColumnTitleTotalSize);
 	}
 
diff --git a/Engine/Source/Runtime/Core/Private/HAL/ThreadingBase.cpp b/Engine/Source/Runtime/Core/Private/HAL/ThreadingBase.cpp
--- a/Engine/Source/Runtime/Core/Private/HAL/ThreadingBase.cpp
+++ b/Engine/Source/Runtime/Core/Private/HAL/ThreadingBase.cpp
@@ -286,7 +286,7 @@ PRAGMA_DISABLE_DEPRECATION_WARNINGS
 	{
 		const uint32 CurrentThreadId = FPlatformTLS::GetCurrentThreadId();
 		bool oldValue = (GRenderThreadId == 0) ? (CurrentThreadId == GGameThreadId) || FTaskTagScope::IsRunningDuringStaticInit() : (CurrentThreadId == GRenderThreadId);
-		ensureMsgf(oldValue == newValue, TEXT("oldValue(%i) newValue(%i) If this check fails make sure that there is a FTaskTagScope(ETaskTag::ERenderingThread) as deep as possible on the current callstack, you can see the current value in ActiveNamedThreads(%x), GRenderingThread(%x)"), oldValue, newValue, FTaskTagScope::GetCurrentTag(), GRenderingThread);
+		ensureMsgf(oldValue == newValue, TEXT("oldValue(%i) newValue(%i) If this check fails make sure that there is a FTaskTagScope(ETaskTag::ERenderingThread) as deep as possible on the current callstack, you can see the current value in ActiveNamedThreads(%x), GRenderingThread(%p)"), oldValue, newValue, FTaskTagScope::GetCurrentTag(), GRenderingThread);
 		newValue = oldValue;
 	}
 #endif
@@ -319,7 +319,7 @@ PRAGMA_DISABLE_DEPRECATION_WARNINGS
 	{
 		const uint32 CurrentThreadId = FPlatformTLS::GetCurrentThreadId();
 		bool oldValue = (GRenderThreadId == 0) ?  true : CurrentThreadId != GGameThreadId;
-		ensureMsgf(oldValue == newValue, TEXT("oldValue(%i) newValue(%i) If this check fails make sure that there is a FTaskTagScope(ETaskTag::EParallelRenderingThread) as deep as possible on the current callstack, you can see the current value in ActiveNamedThreads(%x), GRenderingThread(%x)"), oldValue, newValue, FTaskTagScope::GetCurrentTag(), GRenderingThread);
+		ensureMsgf(oldValue == newValue, TEXT("oldValue(%i) newValue(%i) If this check fails make sure that there is a FTaskTagScope(ETaskTag::EParallelRenderingThread) as deep as possible on the current callstack, you can see the current value in ActiveNamedThreads(%x), GRenderingThread(%p)"), oldValue, newValue, FTaskTagScope::GetCurrentTag(), GRenderingThread);
 		newValue = oldValue;
 	}
 #endif
diff --git a/Engine/Source/Runtime/Core/Private/HAL/UnrealMemory.cpp b/Engine/Source/Runtime/Core/Private/HAL/UnrealMemory.cpp
--- a/Engine/Source/Runtime/Core/Private/HAL/UnrealMemory.cpp
+++ b/Engine/Source/Runtime/Core/Private/HAL/UnrealMemory.cpp
@@ -146,7 +146,7 @@ public:
 						if (Pop[At] != PURGATORY_STOMP_CHECKS_CANARYBYTE)
 						{
 							FPlatformMisc::LowLevelOutputDebugStringf(TEXT("Freed memory at %p + %d == %x (should be %x)\r\n"), Pop, At, (int32)Pop[At], (int32)PURGATORY_STOMP_CHECKS_CANARYBYTE);
-							UE_LOG(LogMemory, Fatal, TEXT("Freed memory at %p + %d == %x (should be %x)"), Pop, At, (int32)Pop[At], (int32)PURGATORY_STOMP_CHECKS_CANARYBYTE);
+							UE_LOG(LogMemory, Fatal, TEXT("Freed memory at %p + %" SIZE_T_FMT " == %x (should be %x)"), Pop, At, (int32)Pop[At], (int32)PURGATORY_STOMP_CHECKS_CANARYBYTE);
 						}
 					}
 					UsedMalloc->Free(Pop);
diff --git a/Engine/Source/Runtime/Core/Private/IO/IoChunkEncoding.cpp b/Engine/Source/Runtime/Core/Private/IO/IoChunkEncoding.cpp
--- a/Engine/Source/Runtime/Core/Private/IO/IoChunkEncoding.cpp
+++ b/Engine/Source/Runtime/Core/Private/IO/IoChunkEncoding.cpp
@@ -320,7 +320,7 @@ bool FIoChunkEncoding::Decode(const FIoChunkDecodingParams& Params, FMemoryView
 			{
 				if (!FCompression::UncompressMemory(Params.CompressionFormat, OutRawData.GetData(), int32(RawBlockReadSize), BlockView.GetData(), CompressedBlockSize))
 				{
-					UE_LOG(LogIoDispatcher, Warning, TEXT("Failed to uncompress block '%d', format '%s', raw size '" UINT64_FMT "' compressed size '%u'"),
+					UE_LOG(LogIoDispatcher, Warning, TEXT("Failed to uncompress block '%d', format '%s', raw size '%" UINT64_FMT "' compressed size '%u'"),
 						BlockIndex, *Params.CompressionFormat.ToString(), RawBlockSize, AlignedBlockSize);
 					return false;
 				}
@@ -330,7 +330,7 @@ bool FIoChunkEncoding::Decode(const FIoChunkDecodingParams& Params, FMemoryView
 				FIoBuffer RawBlockTmp = FIoBuffer(RawBlockSize);
 				if (!FCompression::UncompressMemory(Params.CompressionFormat, RawBlockTmp.GetData(), int32(RawBlockSize), BlockView.GetData(), CompressedBlockSize))
 				{
-					UE_LOG(LogIoDispatcher, Warning, TEXT("Failed to uncompress block '%d', format '%s', raw size '" UINT64_FMT "' compressed size '%u'"),
+					UE_LOG(LogIoDispatcher, Warning, TEXT("Failed to uncompress block '%d', format '%s', raw size '%" UINT64_FMT "' compressed size '%u'"),
 						BlockIndex, *Params.CompressionFormat.ToString(), RawBlockSize, AlignedBlockSize);
 					return false;
 				}
diff --git a/Engine/Source/Runtime/Core/Private/IO/IoStore.cpp b/Engine/Source/Runtime/Core/Private/IO/IoStore.cpp
--- a/Engine/Source/Runtime/Core/Private/IO/IoStore.cpp
+++ b/Engine/Source/Runtime/Core/Private/IO/IoStore.cpp
@@ -878,7 +878,7 @@ public:
 
 			if (bReadSucceeded == false)
 			{
-				UE_LOG(LogIoStore, Error, TEXT("Read from container %s failed (partition %d, offset %lld, size %d)"), *ContainerPath, PartitionIndex, PartitionOffset, TotalAlignedSize);
+				UE_LOG(LogIoStore, Error, TEXT("Read from container %s failed (partition %d, offset %" INT64_FMT ", size %" UINT64_FMT ")"), *ContainerPath, PartitionIndex, PartitionOffset, TotalAlignedSize);
 				return FIoStoreCompressedReadResult();
 			}
 		}
diff --git a/Engine/Source/Runtime/Core/Private/Internationalization/TextLocalizationManager.cpp b/Engine/Source/Runtime/Core/Private/Internationalization/TextLocalizationManager.cpp
--- a/Engine/Source/Runtime/Core/Private/Internationalization/TextLocalizationManager.cpp
+++ b/Engine/Source/Runtime/Core/Private/Internationalization/TextLocalizationManager.cpp
@@ -717,11 +717,11 @@ void FTextLocalizationManager::DumpMemoryInfo() const
 {
 	{
 		FReadScopeLock ScopeLock(DisplayStringTableRW);
-		UE_LOG(LogTextLocalizationManager, Log, TEXT("DisplayStringLookupTable.GetAllocatedSize()=%d elems=%d"), DisplayStringLookupTable.GetAllocatedSize(), DisplayStringLookupTable.Num());
+		UE_LOG(LogTextLocalizationManager, Log, TEXT("DisplayStringLookupTable.GetAllocatedSize()=%" SIZE_T_FMT " elems=%d"), DisplayStringLookupTable.GetAllocatedSize(), DisplayStringLookupTable.Num());
 	}
 	{
 		FReadScopeLock ScopeLock(TextRevisionRW);
-		UE_LOG(LogTextLocalizationManager, Log, TEXT("LocalTextRevisions.GetAllocatedSize()=%d elems=%d"), LocalTextRevisions.GetAllocatedSize(), LocalTextRevisions.Num());
+		UE_LOG(LogTextLocalizationManager, Log, TEXT("LocalTextRevisions.GetAllocatedSize()=%" SIZE_T_FMT " elems=%d"), LocalTextRevisions.GetAllocatedSize(), LocalTextRevisions.Num());
 	}
 }
 
diff --git a/Engine/Source/Runtime/Core/Private/Memory/LinearAllocator.cpp b/Engine/Source/Runtime/Core/Private/Memory/LinearAllocator.cpp
--- a/Engine/Source/Runtime/Core/Private/Memory/LinearAllocator.cpp
+++ b/Engine/Source/Runtime/Core/Private/Memory/LinearAllocator.cpp
@@ -32,7 +32,7 @@ FLinearAllocator::FLinearAllocator(SIZE_T ReserveMemorySize)
 		VirtualMemory = VirtualMemory.AllocateVirtual(ReserveMemorySize);
 		if (!VirtualMemory.GetVirtualPointer())
 		{
-			UE_LOG(LogMemory, Warning, TEXT("LinearAllocator failed to reserve %u MB and will default to FMemory::Malloc instead"), ReserveMemorySize / 1024 / 1024);
+			UE_LOG(LogMemory, Warning, TEXT("LinearAllocator failed to reserve %" SIZE_T_FMT " MB and will default to FMemory::Malloc instead"), ReserveMemorySize / 1024 / 1024);
 			Reserved = 0;
 		}
 	}
@@ -80,7 +80,7 @@ void* FLinearAllocator::Allocate(SIZE_T Size, uint32 Alignment)
 	static bool bOnce = false;
 	if (!bOnce)
 	{
-		UE_LOG(LogMemory, Warning, TEXT("LinearAllocator exceeded %u MB it reserved. Please tune PersistentAllocatorReserveSizeMB setting in [MemoryPools] ini group. Falling back to FMemory::Malloc"), Reserved / 1024 / 1024);
+		UE_LOG(LogMemory, Warning, TEXT("LinearAllocator exceeded %" SIZE_T_FMT " MB it reserved. Please tune PersistentAllocatorReserveSizeMB setting in [MemoryPools] ini group. Falling back to FMemory::Malloc"), Reserved / 1024 / 1024);
 		bOnce = true;
 	}
 	return FMemory::Malloc(Size, Alignment);
diff --git a/Engine/Source/Runtime/Core/Private/Misc/Compression.cpp b/Engine/Source/Runtime/Core/Private/Misc/Compression.cpp
--- a/Engine/Source/Runtime/Core/Private/Misc/Compression.cpp
+++ b/Engine/Source/Runtime/Core/Private/Misc/Compression.cpp
@@ -276,7 +276,7 @@ static bool appUncompressMemoryGZIP(void* UncompressedBuffer, int64 Uncompressed
 	// Sanity check to make sure we uncompressed as much data as we expected to.
 	if (UncompressedSize != ZUncompressedSize)
 	{
-		UE_LOG(LogCompression, Warning, TEXT("appUncompressMemoryGZIP failed: Mismatched uncompressed size. Expected: %d, Got:%d. Result: %d"), UncompressedSize, ZUncompressedSize, Result);
+		UE_LOG(LogCompression, Warning, TEXT("appUncompressMemoryGZIP failed: Mismatched uncompressed size. Expected: %" INT64_FMT ", Got:%d. Result: %d"), UncompressedSize, ZUncompressedSize, Result);
 		bOperationSucceeded = false;
 	}
 	return bOperationSucceeded;
@@ -350,7 +350,7 @@ static bool appUncompressMemoryZLIB( void* UncompressedBuffer, int64 Uncompresse
 	// Sanity check to make sure we uncompressed as much data as we expected to.
 	if( UncompressedSize != ZUncompressedSize )
 	{
-		UE_LOG( LogCompression, Warning, TEXT("appUncompressMemoryZLIB failed: Mismatched uncompressed size. Expected: %d, Got:%d. Result: %d"), UncompressedSize, ZUncompressedSize, Result );
+		UE_LOG( LogCompression, Warning, TEXT("appUncompressMemoryZLIB failed: Mismatched uncompressed size. Expected: %" INT64_FMT ", Got:%d. Result: %d"), UncompressedSize, ZUncompressedSize, Result );
 		bOperationSucceeded = false;
 	}
 	return bOperationSucceeded;
@@ -952,7 +952,7 @@ bool FCompression::UncompressMemory(FName FormatName, void* UncompressedBuffer,
 			bUncompressSucceeded = true;
 		}
 		// Always log an error
-		UE_LOG(LogCompression, Error, TEXT("FCompression::UncompressMemory - Failed to uncompress memory (%d/%d) from address %p using format %s, this may indicate the asset is corrupt!"), CompressedSize, UncompressedSize, CompressedBuffer, *FormatName.ToString());
+		UE_LOG(LogCompression, Error, TEXT("FCompression::UncompressMemory - Failed to uncompress memory (%" INT64_FMT "/%" INT64_FMT ") from address %p using format %s, this may indicate the asset is corrupt!"), CompressedSize, UncompressedSize, CompressedBuffer, *FormatName.ToString());
 		// this extra logging is added to understand shader decompression errors, see UE-159777. However in unrelated
 		// corruption issues this gets hit a lot causing massive log sizes. Since for UE-159777 we crash afterwards, we
 		// can safely limit to one instance for the purpose of diagnosing this.
@@ -961,7 +961,7 @@ bool FCompression::UncompressMemory(FName FormatName, void* UncompressedBuffer,
 		const int32 MaxSizeToLogOut = 16384;
 		if (bAllowLog && FormatName == NAME_Oodle && CompressedSize <= MaxSizeToLogOut)
 		{
-			UE_LOG(LogCompression, Error, TEXT("FCompression::UncompressMemory - Logging compressed data (%d bytes) as a hex dump for investigation"), CompressedSize);
+			UE_LOG(LogCompression, Error, TEXT("FCompression::UncompressMemory - Logging compressed data (%" INT64_FMT " bytes) as a hex dump for investigation"), CompressedSize);
 			FCompressionUtil::LogHexDump(reinterpret_cast<const uint8*>(CompressedBuffer), CompressedSize, 0, CompressedSize);
 		}
 	}
diff --git a/Engine/Source/Runtime/Core/Private/Misc/PreloadableFile.cpp b/Engine/Source/Runtime/Core/Private/Misc/PreloadableFile.cpp
--- a/Engine/Source/Runtime/Core/Private/Misc/PreloadableFile.cpp
+++ b/Engine/Source/Runtime/Core/Private/Misc/PreloadableFile.cpp
@@ -498,7 +498,7 @@ void FPreloadableArchive::SerializeInternal(void* V, int64 Length)
 	if (Pos + Length > Size)
 	{
 		SetError();
-		UE_LOG(LogCore, Error, TEXT("Requested read of %d bytes when %d bytes remain (file=%s, size=%d)"), Length, Size - Pos, *GetArchiveName(), Size);
+		UE_LOG(LogCore, Error, TEXT("Requested read of %" INT64_FMT " bytes when %" INT64_FMT " bytes remain (file=%s, size=%" INT64_FMT ")"), Length, Size - Pos, *GetArchiveName(), Size);
 		return;
 	}
 
diff --git a/Engine/Source/Runtime/Core/Private/ProfilingDebugging/InstanceCounter.cpp b/Engine/Source/Runtime/Core/Private/ProfilingDebugging/InstanceCounter.cpp
--- a/Engine/Source/Runtime/Core/Private/ProfilingDebugging/InstanceCounter.cpp
+++ b/Engine/Source/Runtime/Core/Private/ProfilingDebugging/InstanceCounter.cpp
@@ -61,7 +61,7 @@ void FInstanceCountingObject::IncrementStats()
 
 	if (DoLog)
 	{
-		UE_LOG(LogInstanceCount, Log, TEXT("Constructed %s at 0x%08X, count=%d"), *Name.ToString(), this, Count);
+		UE_LOG(LogInstanceCount, Log, TEXT("Constructed %s at 0x%08" UPTRINT_X_FMT ", count=%d"), *Name.ToString(), this, Count);
 	}
 }
 
@@ -81,7 +81,7 @@ void FInstanceCountingObject::DecrementStats()
 
 	if (DoLog)
 	{
-		UE_LOG(LogInstanceCount, Log, TEXT("Destructed %s at 0x%08X, count=%d"), *Name.ToString(), this, Count);
+		UE_LOG(LogInstanceCount, Log, TEXT("Destructed %s at 0x%08" UPTRINT_X_FMT ", count=%d"), *Name.ToString(), this, Count);
 	}
 }
 
diff --git a/Engine/Source/Runtime/Core/Private/Serialization/BitReader.cpp b/Engine/Source/Runtime/Core/Private/Serialization/BitReader.cpp
--- a/Engine/Source/Runtime/Core/Private/Serialization/BitReader.cpp
+++ b/Engine/Source/Runtime/Core/Private/Serialization/BitReader.cpp
@@ -274,17 +274,17 @@ void FBitReader::SetOverflowed(int64 LengthBits)
 
 	if (GLogFatalOnOverflow)
 	{
-		UE_LOG(LogNetSerialization, Fatal, TEXT("FBitReader::SetOverflowed() called! (ReadLen: %i, Remaining: %i, Max: %i)"),
+		UE_LOG(LogNetSerialization, Fatal, TEXT("FBitReader::SetOverflowed() called! (ReadLen: %" INT64_X_FMT ", Remaining: %" INT64_X_FMT ", Max: %" INT64_X_FMT ")"),
 				LengthBits, (Num - Pos), Num);
 	}
 	else if (GEnsureOnOverflow)
 	{
-		ensureMsgf(false, TEXT("FBitReader::SetOverflowed() called! (ReadLen: %i, Remaining: %i, Max: %i)"),
+		ensureMsgf(false, TEXT("FBitReader::SetOverflowed() called! (ReadLen: %" INT64_X_FMT ", Remaining: %" INT64_X_FMT ", Max: %" INT64_X_FMT ")"),
 					LengthBits, (Num - Pos), Num);
 	}
 	else
 	{
-		UE_LOG(LogNetSerialization, Error, TEXT("FBitReader::SetOverflowed() called! (ReadLen: %i, Remaining: %i, Max: %i)"),
+		UE_LOG(LogNetSerialization, Error, TEXT("FBitReader::SetOverflowed() called! (ReadLen: %" INT64_X_FMT ", Remaining: %" INT64_X_FMT ", Max: %" INT64_X_FMT ")"),
 				LengthBits, (Num - Pos), Num);
 	}
 
diff --git a/Engine/Source/Runtime/Core/Private/Serialization/BitWriter.cpp b/Engine/Source/Runtime/Core/Private/Serialization/BitWriter.cpp
--- a/Engine/Source/Runtime/Core/Private/Serialization/BitWriter.cpp
+++ b/Engine/Source/Runtime/Core/Private/Serialization/BitWriter.cpp
@@ -214,7 +214,7 @@ void FBitWriter::SetOverflowed(int32 LengthBits)
 {
 	if (!bAllowOverflow)
 	{
-		UE_LOG(LogNetSerialization, Error, TEXT("FBitWriter overflowed! (WriteLen: %i, Remaining: %i, Max: %i)"),
+		UE_LOG(LogNetSerialization, Error, TEXT("FBitWriter overflowed! (WriteLen: %" INT64_X_FMT ", Remaining: %" INT64_X_FMT ", Max: %" INT64_X_FMT ")"),
 				LengthBits, (Max-Num), Max);
 	}
 
diff --git a/Engine/Source/Runtime/Core/Private/Stats/StatsSystemTypes.cpp b/Engine/Source/Runtime/Core/Private/Stats/StatsSystemTypes.cpp
--- a/Engine/Source/Runtime/Core/Private/Stats/StatsSystemTypes.cpp
+++ b/Engine/Source/Runtime/Core/Private/Stats/StatsSystemTypes.cpp
@@ -101,7 +101,7 @@ static void DebugLeakTest()
 
 		if (GFrameCounter == 300)
 		{
-			UE_LOG(LogTemp, Warning, TEXT("Stat_ReallocTest: %i / %i"), Array.GetAllocatedSize(), Initial);
+			UE_LOG(LogTemp, Warning, TEXT("Stat_ReallocTest: %" SIZE_T_FMT " / %i"), Array.GetAllocatedSize(), Initial);
 		}
 
 		// General memory leak.
diff --git a/Engine/Source/Runtime/Core/Private/Windows/WindowsPlatformMemory.cpp b/Engine/Source/Runtime/Core/Private/Windows/WindowsPlatformMemory.cpp
--- a/Engine/Source/Runtime/Core/Private/Windows/WindowsPlatformMemory.cpp
+++ b/Engine/Source/Runtime/Core/Private/Windows/WindowsPlatformMemory.cpp
@@ -715,7 +715,7 @@ FPlatformMemory::FSharedMemoryRegion* FWindowsPlatformMemory::MapNamedSharedMemo
 	if (Ptr == NULL)
 	{
 		DWORD ErrNo = GetLastError();
-		UE_LOG(LogHAL, Warning, TEXT("MapViewOfFile(mapping=0x%x, access=0x%x, OffsetHigh=0, OffsetLow=0, NumBytes=%u) failed with GetLastError() = %d"), 
+		UE_LOG(LogHAL, Warning, TEXT("MapViewOfFile(mapping=0x%" UPTRINT_x_FMT ", access=0x%x, OffsetHigh=0, OffsetLow=0, NumBytes=%" SIZE_T_FMT ") failed with GetLastError() = %d"), 
 			Mapping, OpenMappingAccess, Size,
 			ErrNo
 			);
@@ -751,7 +751,7 @@ bool FWindowsPlatformMemory::UnmapNamedSharedMemoryRegion(FSharedMemoryRegion *
 			bAllSucceeded = false;
 
 			int ErrNo = GetLastError();
-			UE_LOG(LogHAL, Warning, TEXT("CloseHandle(handle=0x%x) failed with GetLastError() = %d"), 
+			UE_LOG(LogHAL, Warning, TEXT("CloseHandle(handle=0x%" UPTRINT_x_FMT ") failed with GetLastError() = %d"), 
 				WindowsRegion->GetMapping(),
 				ErrNo
 				);
diff --git a/Engine/Source/Runtime/Core/Private/Windows/WindowsPlatformMisc.cpp b/Engine/Source/Runtime/Core/Private/Windows/WindowsPlatformMisc.cpp
--- a/Engine/Source/Runtime/Core/Private/Windows/WindowsPlatformMisc.cpp
+++ b/Engine/Source/Runtime/Core/Private/Windows/WindowsPlatformMisc.cpp
@@ -952,7 +952,7 @@ void FWindowsPlatformMisc::PlatformInit()
 
 	// Get CPU info.
 	const FPlatformMemoryConstants& MemoryConstants = FPlatformMemory::GetConstants();
-	UE_LOG(LogInit, Log, TEXT("CPU Page size=%i, Cores=%i"), MemoryConstants.PageSize, FPlatformMisc::NumberOfCores() );
+	UE_LOG(LogInit, Log, TEXT("CPU Page size=%" SIZE_T_FMT ", Cores=%i"), MemoryConstants.PageSize, FPlatformMisc::NumberOfCores() );
 
 	// Timer resolution.
 	UE_LOG(LogInit, Log, TEXT("High frequency timer resolution =%f MHz"), 0.000001 / FPlatformTime::GetSecondsPerCycle() );
diff --git a/Engine/Source/Runtime/Core/Private/Windows/WindowsPlatformProcess.cpp b/Engine/Source/Runtime/Core/Private/Windows/WindowsPlatformProcess.cpp
--- a/Engine/Source/Runtime/Core/Private/Windows/WindowsPlatformProcess.cpp
+++ b/Engine/Source/Runtime/Core/Private/Windows/WindowsPlatformProcess.cpp
@@ -1613,7 +1613,7 @@ bool FWindowsPlatformProcess::LaunchFileInDefaultExternalApplication( const TCHA
 	UE_LOG(LogWindows, Log,  TEXT("LaunchFileInExternalEditor %s %s"), FileName, Parms ? Parms : TEXT("") );
 	HINSTANCE Code = ::ShellExecuteW( NULL, VerbString, FileName, Parms ? Parms : TEXT(""), TEXT(""), SW_SHOWNORMAL );
 	
-	UE_LOG(LogWindows, Log,  TEXT("Launch application code for %s %s: %d"), FileName, Parms ? Parms : TEXT(""), (PTRINT)Code );
+	UE_LOG(LogWindows, Log,  TEXT("Launch application code for %s %s: %" PTRINT_FMT), FileName, Parms ? Parms : TEXT(""), (PTRINT)Code );
 
 	// Fallback to a true windows-defined default for the asset type
 	if ((PTRINT)Code == SE_ERR_NOASSOC || (PTRINT)Code == SE_ERR_ASSOCINCOMPLETE)
diff --git a/Engine/Source/Runtime/Core/Public/Async/Fundamental/Scheduler.h b/Engine/Source/Runtime/Core/Public/Async/Fundamental/Scheduler.h
--- a/Engine/Source/Runtime/Core/Public/Async/Fundamental/Scheduler.h
+++ b/Engine/Source/Runtime/Core/Public/Async/Fundamental/Scheduler.h
@@ -80,7 +80,7 @@ namespace LowLevelTasks
 				}
 #if DO_CHECK
 				const int64 NewRevision = int64(LocalTop.Revision) < LastRevision ? (int64(UINT32_MAX) + int64(LocalTop.Revision)) : int64(LocalTop.Revision);
-				ensureMsgf((NewRevision - LastRevision) < (1ll << 31), TEXT("Dangerously close to the wraparound: %d, %d"), LastRevision, NewRevision);
+				ensureMsgf((NewRevision - LastRevision) < (1ll << 31), TEXT("Dangerously close to the wraparound: %" INT64_FMT ", %" INT64_FMT), LastRevision, NewRevision);
 #endif
 			}
 		}
@@ -104,7 +104,7 @@ namespace LowLevelTasks
 				}
 #if DO_CHECK
 				const int64 NewRevision = int64(LocalTop.Revision) < LastRevision ? (int64(UINT32_MAX) + int64(LocalTop.Revision)) : int64(LocalTop.Revision);
-				ensureMsgf((NewRevision - LastRevision) < (1ll << 31), TEXT("Dangerously close to the wraparound: %d, %d"), LastRevision, NewRevision);
+				ensureMsgf((NewRevision - LastRevision) < (1ll << 31), TEXT("Dangerously close to the wraparound: %" INT64_FMT ", %" INT64_FMT), LastRevision, NewRevision);
 #endif
 			}
 		}
diff --git a/Engine/Source/Runtime/Core/Public/Containers/Array.h b/Engine/Source/Runtime/Core/Public/Containers/Array.h
--- a/Engine/Source/Runtime/Core/Public/Containers/Array.h
+++ b/Engine/Source/Runtime/Core/Public/Containers/Array.h
@@ -1961,7 +1961,7 @@ public:
 	 */
 	FORCEINLINE void CheckAddress(const ElementType* Addr) const
 	{
-		checkf(Addr < GetData() || Addr >= (GetData() + ArrayMax), TEXT("Attempting to use a container element (%p) which already comes from the container being modified (%p, ArrayMax: %lld, ArrayNum: %lld, SizeofElement: %d)!"), Addr, GetData(), (long long)ArrayMax, (long long)ArrayNum, sizeof(ElementType));
+		checkf(Addr < GetData() || Addr >= (GetData() + ArrayMax), TEXT("Attempting to use a container element (%p) which already comes from the container being modified (%p, ArrayMax: %lld, ArrayNum: %lld, SizeofElement: %zu)!"), Addr, GetData(), (long long)ArrayMax, (long long)ArrayNum, sizeof(ElementType));
 	}
 
 	/**
diff --git a/Engine/Source/Runtime/Core/Public/Containers/DynamicRHIResourceArray.h b/Engine/Source/Runtime/Core/Public/Containers/DynamicRHIResourceArray.h
--- a/Engine/Source/Runtime/Core/Public/Containers/DynamicRHIResourceArray.h
+++ b/Engine/Source/Runtime/Core/Public/Containers/DynamicRHIResourceArray.h
@@ -68,7 +68,7 @@ public:
 	{
 		if (this->Num() > UINT32_MAX / sizeof(ElementType))
 		{
-			UE_LOG(LogCore, Fatal, TEXT("Resource data size too large for uint32, will overflow. Calculate with larger data type or use fewer elements. sizeof(ElementType): %d"), sizeof(ElementType));
+			UE_LOG(LogCore, Fatal, TEXT("Resource data size too large for uint32, will overflow. Calculate with larger data type or use fewer elements. sizeof(ElementType): %" SIZE_T_FMT), sizeof(ElementType));
 			return 0;
 		}
 		return this->Num() * sizeof(ElementType);
diff --git a/Engine/Source/Runtime/Core/Public/Serialization/MemoryWriter.h b/Engine/Source/Runtime/Core/Public/Serialization/MemoryWriter.h
--- a/Engine/Source/Runtime/Core/Public/Serialization/MemoryWriter.h
+++ b/Engine/Source/Runtime/Core/Public/Serialization/MemoryWriter.h
@@ -46,7 +46,7 @@ public:
 				constexpr IndexSizeType MaxValue = TNumericLimits<IndexSizeType>::Max();
 				if (NewArrayCount >= MaxValue)
 				{
-					UE_LOG(LogSerialization, Fatal, TEXT("TMemoryWriter with IndexSize=%d does not support data larger than %d bytes. Archive name: %s."), IndexSize, MaxValue, *ArchiveName.ToString());
+					UE_LOG(LogSerialization, Fatal, TEXT("TMemoryWriter with IndexSize=%d does not support data larger than %" INT64_FMT " bytes. Archive name: %s."), IndexSize, (int64)MaxValue, *ArchiveName.ToString());
 				}
 			}
 
