MallocLeakDetection uses a 64 bit hash now as 32 bit hash had been known to collide.

#rnx

diff --git a/Engine/Source/Runtime/Core/Private/HAL/MallocLeakDetection.cpp b/Engine/Source/Runtime/Core/Private/HAL/MallocLeakDetection.cpp
--- a/Engine/Source/Runtime/Core/Private/HAL/MallocLeakDetection.cpp
+++ b/Engine/Source/Runtime/Core/Private/HAL/MallocLeakDetection.cpp
@@ -111,7 +111,7 @@ void FMallocLeakDetection::PopContext()
 void FMallocLeakDetection::AddCallstack(FCallstackTrack& Callstack)
 {
 	FScopeLock Lock(&AllocatedPointersCritical);
-	uint32 CallstackHash = Callstack.GetHash();
+	CallstackHashType CallstackHash = Callstack.GetHash();
 	FCallstackTrack& UniqueCallstack = UniqueCallstacks.FindOrAdd(CallstackHash);
 	//if we had a hash collision bail and lose the data rather than corrupting existing data.
 	if ((UniqueCallstack.Count > 0 || UniqueCallstack.NumCheckPoints > 0) && UniqueCallstack != Callstack)
@@ -136,7 +136,7 @@ void FMallocLeakDetection::AddCallstack(FCallstackTrack& Callstack)
 void FMallocLeakDetection::RemoveCallstack(FCallstackTrack& Callstack)
 {
 	FScopeLock Lock(&AllocatedPointersCritical);
-	uint32 CallstackHash = Callstack.GetHash();
+	CallstackHashType CallstackHash = Callstack.GetHash();
 	FCallstackTrack* UniqueCallstack = UniqueCallstacks.Find(CallstackHash);
 	if (UniqueCallstack)
 	{
@@ -164,14 +164,14 @@ void FMallocLeakDetection::SetAllocationCollection(bool bEnabled, int32 Size)
 	}
 }
 
-void FMallocLeakDetection::GetOpenCallstacks(TArray<uint32>& OutCallstacks, SIZE_T& OutTotalSize, const FMallocLeakReportOptions& Options /*= FMallocLeakReportOptions()*/)
+void FMallocLeakDetection::GetOpenCallstacks(TArray<CallstackHashType>& OutCallstacks, SIZE_T& OutTotalSize, const FMallocLeakReportOptions& Options /*= FMallocLeakReportOptions()*/)
 {
 	// Make sure we preallocate memory to avoid
 	OutCallstacks.Empty(UniqueCallstacks.Num() + 32);
 
 	OutTotalSize = 0;
 
-	TMap<uint32, float> HashesToAllocRate;
+	TMap<CallstackHashType, float> HashesToAllocRate;
 	HashesToAllocRate.Empty(UniqueCallstacks.Num() + 32);
 
 	{
@@ -224,7 +224,7 @@ void FMallocLeakDetection::GetOpenCallstacks(TArray<uint32>& OutCallstacks, SIZE
 		}
 
 		// now sort
-		OutCallstacks.Sort([this, &Options, &HashesToAllocRate](uint32 lhs, uint32 rhs) {
+		OutCallstacks.Sort([this, &Options, &HashesToAllocRate](CallstackHashType lhs, CallstackHashType rhs) {
 
 			if (Options.SortBy == FMallocLeakReportOptions::ESortOption::SortRate)
 			{
@@ -295,7 +295,7 @@ int32 FMallocLeakDetection::DumpOpenCallstacks(const TCHAR* FileName, const FMal
 
 	SIZE_T ReportedSize = 0;
 	SIZE_T TotalSize = 0;
-	TArray<uint32> SortedKeys;
+	TArray<CallstackHashType> SortedKeys;
 
 	GetOpenCallstacks(SortedKeys, ReportedSize, Options);
 
@@ -380,9 +380,9 @@ int32 FMallocLeakDetection::DumpOpenCallstacks(const TCHAR* FileName, const FMal
 
 		for (int32 i = 0; i < FCallstackTrack::Depth; ++i)
 		{
-			FPlatformStackWalk::ProgramCounterToHumanReadableString(i, Callstack.CallStack[i], CallstackString, MaxCallstackLineChars);
+			FPlatformStackWalk::ProgramCounterToHumanReadableString(i, Callstack.Callstack[i], CallstackString, MaxCallstackLineChars);
 
-			if (Callstack.CallStack[i] && FCStringAnsi::Strlen(CallstackString) > 0)
+			if (Callstack.Callstack[i] && FCStringAnsi::Strlen(CallstackString) > 0)
 			{
 				//convert ansi -> tchar without mallocs in case we are in OOM handler.
 				for (int32 CurrChar = 0; CurrChar < MaxCallstackLineChars; ++CurrChar)
@@ -469,7 +469,7 @@ void FMallocLeakDetection::Malloc(void* Ptr, SIZE_T Size)
 			{
 				bRecursive = true;
 				FCallstackTrack Callstack;
-				FPlatformStackWalk::CaptureStackBackTrace(Callstack.CallStack, FCallstackTrack::Depth);
+				FPlatformStackWalk::CaptureStackBackTrace(Callstack.Callstack, FCallstackTrack::Depth);
 				Callstack.FirstFrame = (uint32)GFrameCounter;
 				Callstack.LastFrame = (uint32)GFrameCounter;
 				Callstack.Size = Size;
@@ -527,7 +527,7 @@ void FMallocLeakDetection::Realloc(void* OldPtr, SIZE_T OldSize, void* NewPtr, S
 		if (OldPtr != NewPtr)
 		{
 			FCallstackTrack* Callstack = OpenPointers.Find(OldPtr);
-			uint32 OldHash = 0;
+			CallstackHashType OldHash = 0;
 			bool WasKnownDeleter = false;
 
 			if (Callstack)
diff --git a/Engine/Source/Runtime/Core/Public/HAL/MallocLeakDetection.h b/Engine/Source/Runtime/Core/Public/HAL/MallocLeakDetection.h
--- a/Engine/Source/Runtime/Core/Public/HAL/MallocLeakDetection.h
+++ b/Engine/Source/Runtime/Core/Public/HAL/MallocLeakDetection.h
@@ -14,6 +14,7 @@
 #include "HAL/MemoryBase.h"
 #include "HAL/ThreadSafeBool.h"
 #include "HAL/UnrealMemory.h"
+#include "Hash/CityHash.h"
 #include "Misc/Crc.h"
 #include "Misc/ScopeLock.h"
 
@@ -67,6 +68,8 @@ struct FMallocLeakReportOptions
  */
 class FMallocLeakDetection
 {
+	using CallstackHashType = uint64;
+
 	struct FCallstackTrack
 	{
 		FCallstackTrack()
@@ -74,12 +77,12 @@ class FMallocLeakDetection
 			FMemory::Memzero(this, sizeof(FCallstackTrack));
 		}
 		static constexpr int32 Depth = 32;
-		uint64 CallStack[Depth];
+		uint64 Callstack[Depth];
 		uint32 FirstFrame;
 		uint32 LastFrame;
 		uint64 Size;
+		CallstackHashType CachedHash;
 		uint32 Count;
-		uint32 CachedHash;
 
 		// least square line fit stuff
 		uint32 NumCheckPoints;
@@ -97,7 +100,7 @@ class FMallocLeakDetection
 			bool bEqual = true;
 			for (int32 i = 0; i < Depth; ++i)
 			{
-				if (CallStack[i] != Other.CallStack[i])
+				if (Callstack[i] != Other.Callstack[i])
 				{
 					bEqual = false;
 					break;
@@ -113,9 +116,10 @@ class FMallocLeakDetection
 
 		void GetLinearFit();
 		
-		uint32 GetHash() 
+		CallstackHashType GetHash() 
 		{
-			CachedHash = FCrc::MemCrc32(CallStack, sizeof(CallStack), 0);
+			const char* Data = reinterpret_cast<const char*>(Callstack);
+			CachedHash = CityHash64(Data, sizeof(Callstack));
 			return CachedHash;
 		}
 	};
@@ -135,13 +139,13 @@ private:
 	TMap<void*, FCallstackTrack> OpenPointers;
 
 	/** List of all unique callstacks with allocated memory */
-	TMap<uint32, FCallstackTrack> UniqueCallstacks;
+	TMap<CallstackHashType, FCallstackTrack> UniqueCallstacks;
 
 	/** Set of callstacks that are known to delete memory (not reset on ClearData()) */
-	TSet<uint32>	KnownDeleters;
+	TSet<CallstackHashType>	KnownDeleters;
 
 	/** Set of callstacks that are known to resize memory (not reset on ClearData()) */
-	TSet<uint32>	KnownTrimmers;
+	TSet<CallstackHashType>	KnownTrimmers;
 
 	/** Contexts that are associated with allocations */
 	TMap<void*, FString>		PointerContexts;
@@ -215,7 +219,7 @@ public:
 	CORE_API void PopContext();
 
 	/** Returns */
-	CORE_API void GetOpenCallstacks(TArray<uint32>& OutCallstacks, SIZE_T& OutTotalSize, const FMallocLeakReportOptions& Options = FMallocLeakReportOptions());
+	CORE_API void GetOpenCallstacks(TArray<CallstackHashType>& OutCallstacks, SIZE_T& OutTotalSize, const FMallocLeakReportOptions& Options = FMallocLeakReportOptions());
 };
 
 /**
