[Log] Warn when runtime verbosity level is capped by compiled-time verbosity level of log category.

#rnx

diff --git a/Engine/Source/Runtime/Core/Private/Logging/LogCategory.cpp b/Engine/Source/Runtime/Core/Private/Logging/LogCategory.cpp
--- a/Engine/Source/Runtime/Core/Private/Logging/LogCategory.cpp
+++ b/Engine/Source/Runtime/Core/Private/Logging/LogCategory.cpp
@@ -38,12 +38,21 @@ FLogCategoryBase::~FLogCategoryBase()
 	}
 }
 
-void FLogCategoryBase::SetVerbosity(ELogVerbosity::Type NewVerbosity)
+void FLogCategoryBase::SetVerbosity(ELogVerbosity::Type InNewVerbosity)
 {
 	const ELogVerbosity::Type OldVerbosity = Verbosity;
+	const ELogVerbosity::Type NewVerbosity = (ELogVerbosity::Type)(InNewVerbosity & ELogVerbosity::VerbosityMask);
 	// regularize the verbosity to be at most whatever we were compiled with
-	Verbosity = FMath::Min<ELogVerbosity::Type>(CompileTimeVerbosity, (ELogVerbosity::Type)(NewVerbosity & ELogVerbosity::VerbosityMask));
-	DebugBreakOnLog = !!(NewVerbosity & ELogVerbosity::BreakOnLog);
+	if (NewVerbosity <= CompileTimeVerbosity)
+	{
+		Verbosity = NewVerbosity;
+	}
+	else
+	{
+		Verbosity = CompileTimeVerbosity;
+		UE_LOG(LogCore, Warning, TEXT("Log category '%s' tried to set verbosity (%s) higher than its compile-time verbosity level (%s)"), *GetCategoryName().ToString(), ToString(NewVerbosity), ToString(CompileTimeVerbosity));
+	}
+	DebugBreakOnLog = !!(InNewVerbosity & ELogVerbosity::BreakOnLog);
 	checkSlow(!(Verbosity & ELogVerbosity::BreakOnLog)); // this bit is factored out of this variable, always
 
 	if (OldVerbosity != Verbosity)
