[Log/2] Warn when runtime verbosity level is capped by compiled-time verbosity level of log category. 2nd commit after previous backout (41602918).

This also updates the compile-time verbosity settings for all log categories that previously triggered a warning in CIS runs.

#rnx

diff --git a/Engine/Plugins/Experimental/ImagePlate/Source/ImagePlate/Private/ImagePlateFileSequence.cpp b/Engine/Plugins/Experimental/ImagePlate/Source/ImagePlate/Private/ImagePlateFileSequence.cpp
--- a/Engine/Plugins/Experimental/ImagePlate/Source/ImagePlate/Private/ImagePlateFileSequence.cpp
+++ b/Engine/Plugins/Experimental/ImagePlate/Source/ImagePlate/Private/ImagePlateFileSequence.cpp
@@ -16,7 +16,7 @@
 
 #include UE_INLINE_GENERATED_CPP_BY_NAME(ImagePlateFileSequence)
 
-DEFINE_LOG_CATEGORY_STATIC(LogImagePlateFileSequence, Log, Warning);
+DEFINE_LOG_CATEGORY_STATIC(LogImagePlateFileSequence, Log, Log);
 
 UImagePlateFileSequence::UImagePlateFileSequence(const FObjectInitializer& Init)
 	: Super(Init)
diff --git a/Engine/Plugins/Runtime/Harmonix/Source/HarmonixMetasound/Private/HarmonixMetasound/DataTypes/FusionSyncLink.h b/Engine/Plugins/Runtime/Harmonix/Source/HarmonixMetasound/Private/HarmonixMetasound/DataTypes/FusionSyncLink.h
--- a/Engine/Plugins/Runtime/Harmonix/Source/HarmonixMetasound/Private/HarmonixMetasound/DataTypes/FusionSyncLink.h
+++ b/Engine/Plugins/Runtime/Harmonix/Source/HarmonixMetasound/Private/HarmonixMetasound/DataTypes/FusionSyncLink.h
@@ -6,7 +6,7 @@
 #include "MetasoundDataReferenceMacro.h"
 #include "Tasks/Task.h"
 
-DECLARE_LOG_CATEGORY_EXTERN(LogFusionAsync, Log, Warning)
+DECLARE_LOG_CATEGORY_EXTERN(LogFusionAsync, Log, Log)
 
 namespace HarmonixMetasound
 {
diff --git a/Engine/Source/Editor/FoliageEdit/Private/FoliageEdMode.cpp b/Engine/Source/Editor/FoliageEdit/Private/FoliageEdMode.cpp
--- a/Engine/Source/Editor/FoliageEdit/Private/FoliageEdMode.cpp
+++ b/Engine/Source/Editor/FoliageEdit/Private/FoliageEdMode.cpp
@@ -124,7 +124,7 @@ class UClass;
 #define LOCTEXT_NAMESPACE "FoliageEdMode"
 #define FOLIAGE_SNAP_TRACE (10000.f)
 
-DEFINE_LOG_CATEGORY_STATIC(LogFoliage, Log, Warning);
+DEFINE_LOG_CATEGORY_STATIC(LogFoliage, Log, Log);
 
 DECLARE_CYCLE_STAT(TEXT("Calculate Potential Instance"), STAT_FoliageCalculatePotentialInstance, STATGROUP_Foliage);
 DECLARE_CYCLE_STAT(TEXT("Add Instance Imp"), STAT_FoliageAddInstanceImp, STATGROUP_Foliage);
diff --git a/Engine/Source/Runtime/Core/Private/Logging/LogCategory.cpp b/Engine/Source/Runtime/Core/Private/Logging/LogCategory.cpp
--- a/Engine/Source/Runtime/Core/Private/Logging/LogCategory.cpp
+++ b/Engine/Source/Runtime/Core/Private/Logging/LogCategory.cpp
@@ -38,12 +38,21 @@ FLogCategoryBase::~FLogCategoryBase()
 	}
 }
 
-void FLogCategoryBase::SetVerbosity(ELogVerbosity::Type NewVerbosity)
+void FLogCategoryBase::SetVerbosity(ELogVerbosity::Type InNewVerbosity)
 {
 	const ELogVerbosity::Type OldVerbosity = Verbosity;
+	const ELogVerbosity::Type NewVerbosity = (ELogVerbosity::Type)(InNewVerbosity & ELogVerbosity::VerbosityMask);
 	// regularize the verbosity to be at most whatever we were compiled with
-	Verbosity = FMath::Min<ELogVerbosity::Type>(CompileTimeVerbosity, (ELogVerbosity::Type)(NewVerbosity & ELogVerbosity::VerbosityMask));
-	DebugBreakOnLog = !!(NewVerbosity & ELogVerbosity::BreakOnLog);
+	if (NewVerbosity <= CompileTimeVerbosity)
+	{
+		Verbosity = NewVerbosity;
+	}
+	else
+	{
+		Verbosity = CompileTimeVerbosity;
+		UE_LOG(LogCore, Warning, TEXT("Log category '%s' tried to set verbosity (%s) higher than its compile-time verbosity level (%s)"), *GetCategoryName().ToString(), ToString(NewVerbosity), ToString(CompileTimeVerbosity));
+	}
+	DebugBreakOnLog = !!(InNewVerbosity & ELogVerbosity::BreakOnLog);
 	checkSlow(!(Verbosity & ELogVerbosity::BreakOnLog)); // this bit is factored out of this variable, always
 
 	if (OldVerbosity != Verbosity)
diff --git a/Engine/Source/Runtime/Engine/Public/GenericQuadTree.h b/Engine/Source/Runtime/Engine/Public/GenericQuadTree.h
--- a/Engine/Source/Runtime/Engine/Public/GenericQuadTree.h
+++ b/Engine/Source/Runtime/Engine/Public/GenericQuadTree.h
@@ -4,7 +4,7 @@
 
 #include "CoreMinimal.h"
 
-ENGINE_API DECLARE_LOG_CATEGORY_EXTERN(LogQuadTree, Log, Warning);
+ENGINE_API DECLARE_LOG_CATEGORY_EXTERN(LogQuadTree, Log, Log);
 
 template <typename ElementType, int32 NodeCapacity = 4>
 class TQuadTree
diff --git a/Engine/Source/Runtime/UniversalObjectLocator/Private/UniversalObjectLocatorFragment.cpp b/Engine/Source/Runtime/UniversalObjectLocator/Private/UniversalObjectLocatorFragment.cpp
--- a/Engine/Source/Runtime/UniversalObjectLocator/Private/UniversalObjectLocatorFragment.cpp
+++ b/Engine/Source/Runtime/UniversalObjectLocator/Private/UniversalObjectLocatorFragment.cpp
@@ -18,7 +18,7 @@
 
 #define LOCTEXT_NAMESPACE "UOL"
 
-DECLARE_LOG_CATEGORY_EXTERN(LogUniversalObjectLocator, Log, Warning);
+DECLARE_LOG_CATEGORY_EXTERN(LogUniversalObjectLocator, Log, Log);
 DEFINE_LOG_CATEGORY(LogUniversalObjectLocator);
 
 namespace UE::UniversalObjectLocator
