[AutoRTFM] Handle Verse ensures from the open.

Low-level error reporting doesn't need transactional rollback.

diff --git a/Engine/Source/Runtime/Core/Private/Misc/AssertionMacros.cpp b/Engine/Source/Runtime/Core/Private/Misc/AssertionMacros.cpp
--- a/Engine/Source/Runtime/Core/Private/Misc/AssertionMacros.cpp
+++ b/Engine/Source/Runtime/Core/Private/Misc/AssertionMacros.cpp
@@ -441,6 +441,7 @@ TFunction<bool(const FEnsureHandlerArgs& Args)> GetEnsureHandler()
  * @param	Msg		Informative error message text
  * @param	NumStackFramesToIgnore	Number of stack frames to ignore in the callstack
  */
+UE_AUTORTFM_ALWAYS_OPEN
 FORCENOINLINE void FDebug::EnsureFailed(const ANSICHAR* Expr, const ANSICHAR* File, int32 Line, void* ProgramCounter, const TCHAR* Msg)
 {
 	if (UNLIKELY(EnsureHandler && EnsureHandler({ Expr, Msg })))
