Added an option for platforms to define which platform variant the game is currently runnin on.

It is a preparation step for some more specific changes on one of our platforms. This change includes edits to some common files.

diff --git a/Engine/Source/Runtime/Core/Public/GenericPlatform/GenericPlatformProperties.h b/Engine/Source/Runtime/Core/Public/GenericPlatform/GenericPlatformProperties.h
--- a/Engine/Source/Runtime/Core/Public/GenericPlatform/GenericPlatformProperties.h
+++ b/Engine/Source/Runtime/Core/Public/GenericPlatform/GenericPlatformProperties.h
@@ -109,6 +109,19 @@ struct FGenericPlatformProperties
 	 */
 	static FORCEINLINE const char* PlatformName();
 
+	/**
+	  * Get the name of the hardware variant of the current platform.
+	  * 
+	  * Most platforms don't need to provide overrides for this member. This member is intended to be used by the few
+	  * which come in different hardware flavours or which may operate in different runtime modes.
+	  * 
+	  * @return Name of the platform variant.
+	  */
+	static FORCEINLINE const char* PlatformVariantName()
+	{
+		return "";
+	}
+
 	/**
 	 * Checks whether this platform requires cooked data.
 	 *
diff --git a/Engine/Source/Runtime/Engine/Private/UnrealEngine.cpp b/Engine/Source/Runtime/Engine/Private/UnrealEngine.cpp
--- a/Engine/Source/Runtime/Engine/Private/UnrealEngine.cpp
+++ b/Engine/Source/Runtime/Engine/Private/UnrealEngine.cpp
@@ -19649,8 +19649,18 @@ void InitThreadConfig()
 
 	if (GConfig)
 	{
+		FString PlatformVariant = FPlatformProperties::PlatformVariantName();
 		TArray<FString> ThreadConfigEntries;
-		GConfig->GetArray(TEXT("/Script/Engine.Engine"), TEXT("ThreadConfigs"), ThreadConfigEntries, GEngineIni);
+
+		if (PlatformVariant.IsEmpty())
+		{
+			GConfig->GetArray(TEXT("/Script/Engine.Engine"), TEXT("ThreadConfigs"), ThreadConfigEntries, GEngineIni);
+		}
+		else
+		{
+			GConfig->GetArray(TEXT("/Script/Engine.Engine"), *FString::Printf(TEXT("ThreadConfigs_%s"), *PlatformVariant), ThreadConfigEntries, GEngineIni);
+		}
+
 		if (!ThreadConfigEntries.IsEmpty())
 		{
 			SetupThreadConfig(ThreadConfigEntries);
