Moved inlined function to be dllexported

diff --git a/Engine/Source/Runtime/Core/Private/IO/IoDispatcher.cpp b/Engine/Source/Runtime/Core/Private/IO/IoDispatcher.cpp
--- a/Engine/Source/Runtime/Core/Private/IO/IoDispatcher.cpp
+++ b/Engine/Source/Runtime/Core/Private/IO/IoDispatcher.cpp
@@ -60,7 +60,7 @@ static const TCHAR* const GetIoErrorText_ErrorCodeTextArray[] =
 	TEXT("Not Installed")
 };
 
-CORE_API const TCHAR* const* GetIoErrorText_ErrorCodeText = GetIoErrorText_ErrorCodeTextArray;
+const TCHAR* const* GetIoErrorText_ErrorCodeText = GetIoErrorText_ErrorCodeTextArray;
 
 #if UE_IODISPATCHER_STATS_ENABLED
 class FIoRequestStats
diff --git a/Engine/Source/Runtime/Core/Private/IO/IoStatus.cpp b/Engine/Source/Runtime/Core/Private/IO/IoStatus.cpp
--- a/Engine/Source/Runtime/Core/Private/IO/IoStatus.cpp
+++ b/Engine/Source/Runtime/Core/Private/IO/IoStatus.cpp
@@ -18,6 +18,13 @@ const FIoStatus FIoStatus::Ok		{ EIoErrorCode::Ok,				TEXT("OK")				};
 const FIoStatus FIoStatus::Unknown	{ EIoErrorCode::Unknown,		TEXT("Unknown Status")	};
 const FIoStatus FIoStatus::Invalid	{ EIoErrorCode::InvalidCode,	TEXT("Invalid Code")	};
 
+extern const TCHAR* const* GetIoErrorText_ErrorCodeText;
+
+const TCHAR* GetIoErrorText(EIoErrorCode ErrorCode)
+{
+	return GetIoErrorText_ErrorCodeText[static_cast<uint32>(ErrorCode)];
+}
+
 //////////////////////////////////////////////////////////////////////////
 
 FIoStatus::FIoStatus()
diff --git a/Engine/Source/Runtime/Core/Public/IO/IoStatus.h b/Engine/Source/Runtime/Core/Public/IO/IoStatus.h
--- a/Engine/Source/Runtime/Core/Public/IO/IoStatus.h
+++ b/Engine/Source/Runtime/Core/Public/IO/IoStatus.h
@@ -37,12 +37,7 @@ enum class EIoErrorCode
 };
 
 /** Get I/O error code description. */
-inline const TCHAR* GetIoErrorText(EIoErrorCode ErrorCode)
-{
-	extern CORE_API const TCHAR* const* GetIoErrorText_ErrorCodeText;
-
-	return GetIoErrorText_ErrorCodeText[static_cast<uint32>(ErrorCode)];
-}
+CORE_API const TCHAR* GetIoErrorText(EIoErrorCode ErrorCode);
 
 ///////////////////////////////////////////////////////////////////////////////
 
