Changed ctor from multiple specializations to one with constexpr branches in it

diff --git a/Engine/Source/Runtime/Core/Public/Math/Vector.h b/Engine/Source/Runtime/Core/Public/Math/Vector.h
--- a/Engine/Source/Runtime/Core/Public/Math/Vector.h
+++ b/Engine/Source/Runtime/Core/Public/Math/Vector.h
@@ -2726,22 +2726,26 @@ inline FVector FMath::VRand()
 namespace UE {
 namespace Math {
 
-template <>
-template<typename T>
-FORCEINLINE TIntVector3<int32>::TIntVector3(TVector<T> InVector)
-	: X(FMath::TruncToInt32(InVector.X))
-	, Y(FMath::TruncToInt32(InVector.Y))
-	, Z(FMath::TruncToInt32(InVector.Z))
-{
-}
-
-template <>
-template<typename T>
-FORCEINLINE TIntVector3<uint32>::TIntVector3(TVector<T> InVector)
-	: X(IntCastChecked<uint32, int64>(FMath::TruncToInt64(InVector.X)))
-	, Y(IntCastChecked<uint32, int64>(FMath::TruncToInt64(InVector.Y)))
-	, Z(IntCastChecked<uint32, int64>(FMath::TruncToInt64(InVector.Z)))
+template<typename InIntType>
+template <typename FloatType>
+TIntVector3<InIntType>::TIntVector3(TVector<FloatType> InVector)
 {
+	if constexpr (std::is_same_v<IntType, int32>)
+	{
+		X = FMath::TruncToInt32(InVector.X);
+		Y = FMath::TruncToInt32(InVector.Y);
+		Z = FMath::TruncToInt32(InVector.Z);
+	}
+	else if constexpr (std::is_same_v<IntType, uint32>)
+	{
+		X = IntCastChecked<uint32, int64>(FMath::TruncToInt64(InVector.X));
+		Y = IntCastChecked<uint32, int64>(FMath::TruncToInt64(InVector.Y));
+		Z = IntCastChecked<uint32, int64>(FMath::TruncToInt64(InVector.Z));
+	}
+	else
+	{
+		static_assert(sizeof(IntType) == 0, "Unimplemented");
+	}
 }
 
 template<typename T>
