Workaround for c++ header unit linker error in TMemoryImage

diff --git a/Engine/Source/Runtime/Core/Public/Serialization/MemoryImage.h b/Engine/Source/Runtime/Core/Public/Serialization/MemoryImage.h
--- a/Engine/Source/Runtime/Core/Public/Serialization/MemoryImage.h
+++ b/Engine/Source/Runtime/Core/Public/Serialization/MemoryImage.h
@@ -385,7 +385,8 @@ public:
 	inline bool IsValid() const { return UnfrozenPtr != nullptr; }
 	inline bool IsNull() const { return UnfrozenPtr == nullptr; }
 
-	inline TMemoryImagePtr(T* InPtr = nullptr) : UnfrozenPtr(InPtr) { check(!Frozen.IsFrozen()); }
+	inline TMemoryImagePtr() : UnfrozenPtr(nullptr) {}
+	inline TMemoryImagePtr(T* InPtr) : UnfrozenPtr(InPtr) { check(!Frozen.IsFrozen()); }
 	inline TMemoryImagePtr(const TMemoryImagePtr<T>& InPtr) : UnfrozenPtr(InPtr.Get()) { check(!Frozen.IsFrozen()); }
 	inline TMemoryImagePtr& operator=(T* InPtr) { UnfrozenPtr = InPtr; check(!Frozen.IsFrozen()); return *this; }
 	inline TMemoryImagePtr& operator=(const TMemoryImagePtr<T>& InPtr) { UnfrozenPtr = InPtr.Get(); check(!Frozen.IsFrozen()); return *this; }
