Fixing a crash that occurs when primitive descs without FPrimitiveSceneDesc::PrimitiveUObject assigned are added to a scene while tracing -trace=metadata is enabled

diff --git a/Engine/Source/Runtime/Core/Public/HAL/LowLevelMemStats.h b/Engine/Source/Runtime/Core/Public/HAL/LowLevelMemStats.h
--- a/Engine/Source/Runtime/Core/Public/HAL/LowLevelMemStats.h
+++ b/Engine/Source/Runtime/Core/Public/HAL/LowLevelMemStats.h
@@ -67,7 +67,7 @@
 #define LLM_SCOPE_DYNAMIC_STAT_OBJECTPATH(Object, Set) LLM_SCOPE_DYNAMIC(FName(*Object->GetPathName()), \
 					ELLMTracker::Default, Set, FLLMDynamicTagConstructorStatString(Object->GetPathName()))
 #define LLM_SCOPE_DYNAMIC_STAT_OBJECTPATH_FNAME(ObjectPath, Set) LLM_SCOPE_DYNAMIC(ObjectPath, \
-					ELLMTracker::Default, Set, FLLMDynamicTagConstructorStatString(ObjectPath.ToString()))
+					ELLMTracker::Default, Set, FLLMDynamicTagConstructorStatString((ObjectPath).ToString()))
 
 // special stat pushing to update threads after each asset when tracking assets
 // Currently this is unused, but we may use it for optimizations later
diff --git a/Engine/Source/Runtime/Renderer/Private/RendererScene.cpp b/Engine/Source/Runtime/Renderer/Private/RendererScene.cpp
--- a/Engine/Source/Runtime/Renderer/Private/RendererScene.cpp
+++ b/Engine/Source/Runtime/Renderer/Private/RendererScene.cpp
@@ -2175,8 +2175,8 @@ void FScene::BatchAddPrimitivesInternal(TArrayView<T*> InPrimitives)
 		return;
 	}
 #endif
-	LLM_SCOPE_DYNAMIC_STAT_OBJECTPATH(InPrimitives[0]->GetOutermost(), ELLMTagSet::Assets);
-	UE_TRACE_METADATA_SCOPE_ASSET_FNAME(NAME_None, NAME_None, InPrimitives[0]->GetOutermost()->GetFName());
+	LLM_SCOPE_DYNAMIC_STAT_OBJECTPATH_FNAME(ToUObject(InPrimitives[0]) ? InPrimitives[0]->GetOutermost()->GetFName() : NAME_None, ELLMTagSet::Assets);
+	UE_TRACE_METADATA_SCOPE_ASSET_FNAME(NAME_None, NAME_None, ToUObject(InPrimitives[0]) ? InPrimitives[0]->GetOutermost()->GetFName() : NAME_None);
 
 	SCOPE_CYCLE_COUNTER(STAT_AddScenePrimitiveGT);
 
