[Core]
* Fixed so FMD5, FSHAHash, FTextSourceSiteContext and FGatherableTextData use real hidden friends (must be inlined)

diff --git a/Engine/Source/Runtime/Core/Private/Internationalization/GatherableTextData.cpp b/Engine/Source/Runtime/Core/Private/Internationalization/GatherableTextData.cpp
--- a/Engine/Source/Runtime/Core/Private/Internationalization/GatherableTextData.cpp
+++ b/Engine/Source/Runtime/Core/Private/Internationalization/GatherableTextData.cpp
@@ -10,46 +10,43 @@
 
 class FArchive;
 
-FArchive& operator<<(FArchive& Archive, FTextSourceSiteContext& This)
+void FTextSourceSiteContext::StreamArchive(FArchive& Archive)
 {
-	FStructuredArchiveFromArchive(Archive).GetSlot() << This;
-	return Archive;
+	FStructuredArchiveFromArchive(Archive).GetSlot() << *this;
 }
 
-void operator<<(FStructuredArchive::FSlot Slot, FTextSourceSiteContext& This)
+void FTextSourceSiteContext::StreamStructuredArchive(FStructuredArchive::FSlot Slot)
 {
 	FStructuredArchive::FRecord Record = Slot.EnterRecord();
-	Record << SA_VALUE(TEXT("KeyName"), This.KeyName);
-	Record << SA_VALUE(TEXT("SiteDescription"), This.SiteDescription);
-	Record << SA_VALUE(TEXT("IsEditorOnly"), This.IsEditorOnly);
-	Record << SA_VALUE(TEXT("IsOptional"), This.IsOptional);
-	Record << SA_VALUE(TEXT("InfoMetaData"), This.InfoMetaData);
-	Record << SA_VALUE(TEXT("KeyMetaData"), This.KeyMetaData);
+	Record << SA_VALUE(TEXT("KeyName"), KeyName);
+	Record << SA_VALUE(TEXT("SiteDescription"), SiteDescription);
+	Record << SA_VALUE(TEXT("IsEditorOnly"), IsEditorOnly);
+	Record << SA_VALUE(TEXT("IsOptional"), IsOptional);
+	Record << SA_VALUE(TEXT("InfoMetaData"), InfoMetaData);
+	Record << SA_VALUE(TEXT("KeyMetaData"), KeyMetaData);
 }
 
-FArchive& operator<<(FArchive& Archive, FTextSourceData& This)
+void FTextSourceData::StreamArchive(FArchive& Archive)
 {
-	FStructuredArchiveFromArchive(Archive).GetSlot() << This;
-	return Archive;
+	FStructuredArchiveFromArchive(Archive).GetSlot() << *this;
 }
 
-void operator<<(FStructuredArchive::FSlot Slot, FTextSourceData& This)
+void FTextSourceData::StreamStructuredArchive(FStructuredArchive::FSlot Slot)
 {
 	FStructuredArchive::FRecord Record = Slot.EnterRecord();
-	Record << SA_VALUE(TEXT("SourceString"), This.SourceString);
-	Record << SA_VALUE(TEXT("SourceStringMetaData"), This.SourceStringMetaData);
+	Record << SA_VALUE(TEXT("SourceString"), SourceString);
+	Record << SA_VALUE(TEXT("SourceStringMetaData"), SourceStringMetaData);
 }
 
-FArchive& operator<<(FArchive& Archive, FGatherableTextData& This)
+void FGatherableTextData::StreamArchive(FArchive& Archive)
 {
-	FStructuredArchiveFromArchive(Archive).GetSlot() << This;
-	return Archive;
+	FStructuredArchiveFromArchive(Archive).GetSlot() << *this;
 }
 
-void operator<<(FStructuredArchive::FSlot Slot, FGatherableTextData& This)
+void FGatherableTextData::StreamStructuredArchive(FStructuredArchive::FSlot Slot)
 {
 	FStructuredArchive::FRecord Record = Slot.EnterRecord();
-	Record << SA_VALUE(TEXT("NamespaceName"), This.NamespaceName);
-	Record << SA_VALUE(TEXT("SourceData"), This.SourceData);
-	Record << SA_VALUE(TEXT("SourceSiteContexts"), This.SourceSiteContexts);
+	Record << SA_VALUE(TEXT("NamespaceName"), NamespaceName);
+	Record << SA_VALUE(TEXT("SourceData"), SourceData);
+	Record << SA_VALUE(TEXT("SourceSiteContexts"), SourceSiteContexts);
 }
\ No newline at end of file
diff --git a/Engine/Source/Runtime/Core/Private/Misc/SecureHash.cpp b/Engine/Source/Runtime/Core/Private/Misc/SecureHash.cpp
--- a/Engine/Source/Runtime/Core/Private/Misc/SecureHash.cpp
+++ b/Engine/Source/Runtime/Core/Private/Misc/SecureHash.cpp
@@ -495,19 +495,19 @@ void FMD5::Decode( uint32* output, const uint8* input, int32 len )
 // }
 /// @endcond
 
-FString LexToString(const FMD5Hash& Hash)
+FString FMD5Hash::LexToString() const
 {
-	if (!Hash.bIsValid)
+	if (!bIsValid)
 	{
 		return FString();
 	}
 
 	return FString::Printf(TEXT("%02x%02x%02x%02x%02x%02x%02x%02x%02x%02x%02x%02x%02x%02x%02x%02x"),
-		Hash.Bytes[0], Hash.Bytes[1], Hash.Bytes[2], Hash.Bytes[3], Hash.Bytes[4], Hash.Bytes[5], Hash.Bytes[6], Hash.Bytes[7],
-		Hash.Bytes[8], Hash.Bytes[9], Hash.Bytes[10], Hash.Bytes[11], Hash.Bytes[12], Hash.Bytes[13], Hash.Bytes[14], Hash.Bytes[15]);
+		Bytes[0], Bytes[1], Bytes[2], Bytes[3], Bytes[4], Bytes[5], Bytes[6], Bytes[7],
+		Bytes[8], Bytes[9], Bytes[10], Bytes[11], Bytes[12], Bytes[13], Bytes[14], Bytes[15]);
 }
 
-void LexFromString(FMD5Hash& Hash, const TCHAR* Buffer)
+void FMD5Hash::LexFromString(const TCHAR* Buffer)
 {
 	auto HexCharacterToDecimalValue = [](const TCHAR InHexChar, uint8& OutDecValue) -> bool
 	{
@@ -532,7 +532,7 @@ void LexFromString(FMD5Hash& Hash, const TCHAR* Buffer)
 		return true;
 	};
 
-	uint8 Bytes[16];
+	uint8 NewBytes[16];
 	for (int32 ByteIndex = 0, BufferIndex = 0; ByteIndex < 16; ++ByteIndex)
 	{
 		const TCHAR FirstChar = Buffer[BufferIndex++];
@@ -553,11 +553,11 @@ void LexFromString(FMD5Hash& Hash, const TCHAR* Buffer)
 			return;
 		}
 
-		Bytes[ByteIndex] = (uint8)((FirstCharVal << 4) + SecondCharVal);
+		NewBytes[ByteIndex] = (uint8)((FirstCharVal << 4) + SecondCharVal);
 	}
 
-	FMemory::Memcpy(Hash.Bytes, Bytes, 16);
-	Hash.bIsValid = true;
+	FMemory::Memcpy(Bytes, NewBytes, 16);
+	bIsValid = true;
 }
 
 FMD5Hash FMD5Hash::HashFile(const TCHAR* InFilename, TArray<uint8>* Buffer)
@@ -608,17 +608,17 @@ FCbWriter& FMD5Hash::WriteCompactBinary(FCbWriter& Writer) const
 	return Writer;
 }
 
-bool LoadFromCompactBinary(FCbFieldView Field, FMD5Hash& OutHash)
+bool FMD5Hash::LoadFromCompactBinary(const FCbFieldView& Field)
 {
-	FMemoryView BinaryView = Field.AsBinaryView();
-	if (BinaryView.GetSize() != OutHash.GetSize())
+	FMemoryView BinaryView = const_cast<FCbFieldView&>(Field).AsBinaryView();
+	if (BinaryView.GetSize() != GetSize())
 	{
-		OutHash.bIsValid = false;
+		bIsValid = false;
 		return false;
 	}
 
-	FMemory::Memcpy(OutHash.Bytes, BinaryView.GetData(), OutHash.GetSize());
-	OutHash.bIsValid = true;
+	FMemory::Memcpy(Bytes, BinaryView.GetData(), GetSize());
+	bIsValid = true;
 	return true;
 }
 
@@ -658,40 +658,24 @@ TMap<FString, uint8*> FSHA1::ScriptSHAHashMap;
 
 #define SHA1_BLOCK_SIZE 64
 
-FArchive& operator<<( FArchive& Ar, FSHAHash& G )
+bool FSHAHash::LoadFromCompactBinary(const FCbFieldView& Field)
 {
-	Ar.Serialize(&G.Hash, sizeof(G.Hash));
-	return Ar;
-}
-
-FString LexToString(const FSHAHash& InHash)
-{
-	return InHash.ToString();
-}
-
-void LexFromString(FSHAHash& InHash, const TCHAR* InString)
-{
-	InHash.FromString(InString);
-}
-
-FCbWriter& operator<<(FCbWriter& Writer, const FSHAHash& Hash)
-{
-	Writer.AddBinary(FMemoryView(&Hash.Hash, sizeof(Hash.Hash)));
-	return Writer;
-}
-
-bool LoadFromCompactBinary(FCbFieldView Field, FSHAHash& OutHash)
-{
-	FMemoryView Bytes = Field.AsBinaryView();
-	if (Bytes.GetSize() != sizeof(OutHash.Hash))
+	FMemoryView Bytes = const_cast<FCbFieldView&>(Field).AsBinaryView();
+	if (Bytes.GetSize() != sizeof(Hash))
 	{
-		OutHash = FSHAHash();
+		*this = FSHAHash();
 		return false;
 	}
-	FMemory::Memcpy(OutHash.Hash, Bytes.GetData(), Bytes.GetSize());
+	FMemory::Memcpy(Hash, Bytes.GetData(), Bytes.GetSize());
 	return true;
 }
 
+FCbWriter& FSHAHash::WriteCompactBinary(FCbWriter& Writer) const
+{
+	Writer.AddBinary(FMemoryView(&Hash, sizeof(Hash)));
+	return Writer;
+}
+
 void Freeze::IntrinsicToString(const FSHAHash& Object, const FTypeLayoutDesc& TypeDesc, const FPlatformTypeLayoutParameters& LayoutParams, FMemoryToStringContext& OutContext)
 {
 	OutContext.String->Appendf(TEXT("%s\n"), *Object.ToString());
diff --git a/Engine/Source/Runtime/Core/Public/Internationalization/GatherableTextData.h b/Engine/Source/Runtime/Core/Public/Internationalization/GatherableTextData.h
--- a/Engine/Source/Runtime/Core/Public/Internationalization/GatherableTextData.h
+++ b/Engine/Source/Runtime/Core/Public/Internationalization/GatherableTextData.h
@@ -24,18 +24,24 @@ struct FTextSourceSiteContext
 	FLocMetadataObject InfoMetaData;
 	FLocMetadataObject KeyMetaData;
 
-	CORE_API friend FArchive& operator<<(FArchive& Archive, FTextSourceSiteContext& This);
-	CORE_API friend void operator<<(FStructuredArchive::FSlot Slot, FTextSourceSiteContext& This);
+	friend FArchive& operator<<(FArchive& Archive, FTextSourceSiteContext& This) { This.StreamArchive(Archive); return Archive; }
+	friend void operator<<(FStructuredArchive::FSlot Slot, FTextSourceSiteContext& This) { This.StreamStructuredArchive(Slot); }
+
+	CORE_API void StreamArchive(FArchive& Archive);
+	CORE_API void StreamStructuredArchive(FStructuredArchive::FSlot Slot);
 };
 
 struct FTextSourceData
 {
 	FString SourceString;
 	FLocMetadataObject SourceStringMetaData;
-};
 
-CORE_API FArchive& operator<<(FArchive& Archive, FTextSourceData& This);
-CORE_API void operator<<(FStructuredArchive::FSlot Slot, FTextSourceData& This);
+	friend FArchive& operator<<(FArchive& Archive, FTextSourceData& This) { This.StreamArchive(Archive); return Archive; }
+	friend void operator<<(FStructuredArchive::FSlot Slot, FTextSourceData& This) { This.StreamStructuredArchive(Slot); }
+
+	CORE_API void StreamArchive(FArchive& Archive);
+	CORE_API void StreamStructuredArchive(FStructuredArchive::FSlot Slot);
+};
 
 struct FGatherableTextData
 {
@@ -44,6 +50,9 @@ struct FGatherableTextData
 
 	TArray<FTextSourceSiteContext> SourceSiteContexts;
 
-	CORE_API friend FArchive& operator<<(FArchive& Archive, FGatherableTextData& This);
-	CORE_API friend void operator<<(FStructuredArchive::FSlot Slot, FGatherableTextData& This);
+	friend FArchive& operator<<(FArchive& Archive, FGatherableTextData& This) { This.StreamArchive(Archive); return Archive; }
+	friend void operator<<(FStructuredArchive::FSlot Slot, FGatherableTextData& This) { This.StreamStructuredArchive(Slot); }
+
+	CORE_API void StreamArchive(FArchive& Archive);
+	CORE_API void StreamStructuredArchive(FStructuredArchive::FSlot Slot);
 };
diff --git a/Engine/Source/Runtime/Core/Public/Misc/SecureHash.h b/Engine/Source/Runtime/Core/Public/Misc/SecureHash.h
--- a/Engine/Source/Runtime/Core/Public/Misc/SecureHash.h
+++ b/Engine/Source/Runtime/Core/Public/Misc/SecureHash.h
@@ -186,10 +186,14 @@ private:
 	{
 		return Hash.WriteCompactBinary(Writer);
 	}
+	friend bool LoadFromCompactBinary(const FCbFieldView& Field, FMD5Hash& OutHash) { return OutHash.LoadFromCompactBinary(Field); }
+	friend FString LexToString(const FMD5Hash& InHash) { return InHash.LexToString(); }
+	friend void LexFromString(FMD5Hash& Hash, const TCHAR* InString) { Hash.LexFromString(InString); }
+
 	CORE_API FCbWriter& WriteCompactBinary(FCbWriter& Writer) const;
-	friend CORE_API bool LoadFromCompactBinary(FCbFieldView Field, FMD5Hash& OutHash); // inlining not available because we don't want FCbFieldView defined
-	friend CORE_API FString LexToString(const FMD5Hash&);
-	friend CORE_API void LexFromString(FMD5Hash& Hash, const TCHAR*);
+	CORE_API bool LoadFromCompactBinary(const FCbFieldView& Field); // inlining not available because we don't want FCbFieldView defined
+	CORE_API FString LexToString() const;
+	CORE_API void LexFromString(const TCHAR*);
 };
 
 /** 
@@ -277,19 +281,26 @@ public:
 		return FMemory::Memcmp(&X.Hash, &Y.Hash, sizeof(X.Hash)) < 0;
 	}
 
-	friend CORE_API FArchive& operator<<( FArchive& Ar, FSHAHash& G );
+	friend FArchive& operator<<( FArchive& Ar, FSHAHash& G )
+	{
+		Ar.Serialize(&G.Hash, sizeof(G.Hash));
+		return Ar;
+	}
 	
 	friend uint32 GetTypeHash(const FSHAHash& InKey)
 	{
 		return *reinterpret_cast<const uint32*>(InKey.Hash);
 	}
 
-	friend CORE_API FString LexToString(const FSHAHash&);
-	friend CORE_API void LexFromString(FSHAHash& Hash, const TCHAR*);
-	friend inline FStringBuilderBase& operator<<(FStringBuilderBase& Builder, const FSHAHash& InHash) { UE::String::BytesToHex(InHash.Hash, Builder); return Builder; }
-	friend inline FAnsiStringBuilderBase& operator<<(FAnsiStringBuilderBase& Builder, const FSHAHash& InHash) { UE::String::BytesToHex(InHash.Hash, Builder); return Builder; }
-	CORE_API friend FCbWriter& operator<<(FCbWriter& Writer, const FSHAHash& Hash);
-	CORE_API friend bool LoadFromCompactBinary(FCbFieldView Field, FSHAHash& OutHash);
+	friend FString LexToString(const FSHAHash& InHash) { return InHash.ToString(); }
+	friend void LexFromString(FSHAHash& InHash, const TCHAR* InString) { InHash.FromString(InString); }
+	friend FStringBuilderBase& operator<<(FStringBuilderBase& Builder, const FSHAHash& InHash) { UE::String::BytesToHex(InHash.Hash, Builder); return Builder; }
+	friend FAnsiStringBuilderBase& operator<<(FAnsiStringBuilderBase& Builder, const FSHAHash& InHash) { UE::String::BytesToHex(InHash.Hash, Builder); return Builder; }
+	friend FCbWriter& operator<<(FCbWriter& Writer, const FSHAHash& InHash) { return InHash.WriteCompactBinary(Writer); }
+	friend bool LoadFromCompactBinary(const FCbFieldView& Field, FSHAHash& OutHash) { return OutHash.LoadFromCompactBinary(Field); }
+
+	CORE_API bool LoadFromCompactBinary(const FCbFieldView& Field);
+	CORE_API FCbWriter& WriteCompactBinary(FCbWriter& Writer) const;
 };
 
 namespace Freeze
