[Core]
* Disabled static assert in TSharedPointer when header units are enabled. Will require fix from microsoft to enable again

diff --git a/Engine/Source/Runtime/Core/Public/Templates/SharedPointer.h b/Engine/Source/Runtime/Core/Public/Templates/SharedPointer.h
--- a/Engine/Source/Runtime/Core/Public/Templates/SharedPointer.h
+++ b/Engine/Source/Runtime/Core/Public/Templates/SharedPointer.h
@@ -1778,10 +1778,12 @@ protected:
 	/** Hidden stub constructor */
 	TSharedFromThis()
 	{
+#ifndef UE_HEADER_UNITS
 		// Cause a compile error if ObjectType is a UObject.
 		int32 UObjectTestOverload(const volatile UObject*);
 		int16 UObjectTestOverload(...);
 		static_assert(sizeof(UObjectTestOverload((ObjectType*)nullptr)) == sizeof(int16), "TSharedFromThis is not supported on UObjects");
+#endif
 	}
 
 	/** Hidden stub copy constructor */
