Disable backtrace when running on arm64 via x64 emulation for now

#rnx

diff --git a/Engine/Source/Runtime/Core/Private/ProfilingDebugging/Microsoft/WindowsCallstackTrace.cpp b/Engine/Source/Runtime/Core/Private/ProfilingDebugging/Microsoft/WindowsCallstackTrace.cpp
--- a/Engine/Source/Runtime/Core/Private/ProfilingDebugging/Microsoft/WindowsCallstackTrace.cpp
+++ b/Engine/Source/Runtime/Core/Private/ProfilingDebugging/Microsoft/WindowsCallstackTrace.cpp
@@ -847,6 +847,14 @@ void Modules_Initialize();
 ////////////////////////////////////////////////////////////////////////////////
 void CallstackTrace_CreateInternal(FMalloc* Malloc)
 {
+	// backtrace doesn't work on arm64, even with x64 emulation
+#if PLATFORM_CPU_X86_FAMILY
+	if (FPlatformMisc::GetHostArchitecture() == TEXT("arm64"))
+	{
+		return;
+	}
+#endif
+
 	if (FBacktracer::Get() != nullptr)
 	{
 		return;
