[ASIS] Force use "New Window Behavior"  if asis plugin was enabled

#rnx

diff --git a/Engine/Source/Runtime/Core/Private/Android/AndroidPlatformMisc.cpp b/Engine/Source/Runtime/Core/Private/Android/AndroidPlatformMisc.cpp
--- a/Engine/Source/Runtime/Core/Private/Android/AndroidPlatformMisc.cpp
+++ b/Engine/Source/Runtime/Core/Private/Android/AndroidPlatformMisc.cpp
@@ -2452,6 +2452,9 @@ JNI_METHOD void Java_com_epicgames_unreal_GameActivity_nativeSetAndroidStartupSt
 
 bool FAndroidMisc::UseNewWindowBehavior()
 {
+#if USE_ANDROID_STANDALONE
+	return true;
+#endif
 	static const FString* EnableNewBackgroundBehavior = FAndroidMisc::GetConfigRulesVariable(TEXT("EnableNewBackgroundBehavior"));
 	return EnableNewBackgroundBehavior && EnableNewBackgroundBehavior->ToLower() == TEXT("true");
 }
