Optimize FMutex and FRecursiveMutex by unlocking immediately

This allows other threads to acquire the lock while the unlocking thread looks for a thread to wake.

diff --git a/Engine/Source/Runtime/Core/Private/Async/RecursiveMutex.cpp b/Engine/Source/Runtime/Core/Private/Async/RecursiveMutex.cpp
--- a/Engine/Source/Runtime/Core/Private/Async/RecursiveMutex.cpp
+++ b/Engine/Source/Runtime/Core/Private/Async/RecursiveMutex.cpp
@@ -10,7 +10,7 @@
 namespace UE
 {
 
-static constexpr bool bUnlockRecursiveMutexImmediately = false;
+static constexpr bool bUnlockRecursiveMutexImmediately = true;
 
 bool FRecursiveMutex::TryLock()
 {
diff --git a/Engine/Source/Runtime/Core/Public/Async/ExternalMutex.h b/Engine/Source/Runtime/Core/Public/Async/ExternalMutex.h
--- a/Engine/Source/Runtime/Core/Public/Async/ExternalMutex.h
+++ b/Engine/Source/Runtime/Core/Public/Async/ExternalMutex.h
@@ -89,7 +89,7 @@ private:
 	static constexpr uint8 IsLockedFlag = 1 << 0;
 	static constexpr uint8 MayHaveWaitingLockFlag = 1 << 1;
 
-	static constexpr bool bUnlockImmediately = false;
+	static constexpr bool bUnlockImmediately = true;
 
 	std::atomic<uint8>& State;
 };
diff --git a/Engine/Source/Runtime/Core/Public/Async/Mutex.h b/Engine/Source/Runtime/Core/Public/Async/Mutex.h
--- a/Engine/Source/Runtime/Core/Public/Async/Mutex.h
+++ b/Engine/Source/Runtime/Core/Public/Async/Mutex.h
@@ -78,7 +78,7 @@ private:
 	UE_API void LockSlow();
 	UE_API void WakeWaitingThread();
 
-	static constexpr bool bUnlockImmediately = false;
+	static constexpr bool bUnlockImmediately = true;
 
 	static constexpr uint8 IsLockedFlag = 1 << 0;
 	static constexpr uint8 MayHaveWaitingLockFlag = 1 << 1;
