Don't get stuck in a loop if encountering an unknown opcode

#rnx

diff --git a/Engine/Source/Runtime/Core/Private/ProfilingDebugging/Microsoft/WindowsCallstackTrace.cpp b/Engine/Source/Runtime/Core/Private/ProfilingDebugging/Microsoft/WindowsCallstackTrace.cpp
--- a/Engine/Source/Runtime/Core/Private/ProfilingDebugging/Microsoft/WindowsCallstackTrace.cpp
+++ b/Engine/Source/Runtime/Core/Private/ProfilingDebugging/Microsoft/WindowsCallstackTrace.cpp
@@ -351,7 +351,8 @@ void FBacktracer::AddModule(UPTRINT ModuleBase, const TCHAR* Name)
 
 		const auto* Code = (FUnwindCode*)(UnwindInfo + 1);
 		const auto* EndCode = Code + UnwindInfo->NumUnwindCodes;
-		while (Code < EndCode)
+		bool bEncounteredUnknownOpCode = false;
+		while (Code < EndCode && !bEncounteredUnknownOpCode)
 		{
 #if BACKTRACE_DBGLVL >= 2
 			if (Code->PrologOffset > PrologVerify)
@@ -425,6 +426,7 @@ void FBacktracer::AddModule(UPTRINT ModuleBase, const TCHAR* Name)
 #if BACKTRACE_DBGLVL >= 2
 				PLATFORM_BREAK();
 #endif
+				bEncounteredUnknownOpCode = true;
 				break;
 			}
 		}
