[Core]
* Changed operators in StringView to not be hidden friend unless they have to. This fixes compile errors with hidden friends complaining about using unions in inlined hidden friends

diff --git a/Engine/Source/Runtime/Core/Public/Containers/DirectoryTree.h b/Engine/Source/Runtime/Core/Public/Containers/DirectoryTree.h
--- a/Engine/Source/Runtime/Core/Public/Containers/DirectoryTree.h
+++ b/Engine/Source/Runtime/Core/Public/Containers/DirectoryTree.h
@@ -980,7 +980,7 @@ ValueType& TDirectoryTree<ValueType>::FTreeNode::FindOrAdd(FStringView InRelPath
 			FStringView NextExistingRemainingPath;
 			FPathViews::SplitFirstComponent(ExistingRemainingPath, NextExistingFirstComponent, NextExistingRemainingPath);
 			check(!NextExistingFirstComponent.IsEmpty());
-			if (NextFirstComponent == NextExistingFirstComponent)
+			if (NextFirstComponent.Equals(NextExistingFirstComponent, ESearchCase::IgnoreCase)) // This works around a compiler bug. Should be changed to "NextFirstComponent == NextExistingFirstComponent"
 			{
 				// Next component is also a match, go to the next loop iteration to handle the new remainingpaths
 				RemainingPath = NextRemainingPath;
