Small change to make parsing values and params consistent.  With this change, back-to-back lines like this are now consistent with one another (notice the hyphens):

bool HasBug = FParse::Value(CommandLine, TEXT("-bug="), Filename);
bool Submit = FParse::Param(CommandLine, TEXT("-submit"));

With the original code, however, the lines above would have to be written like this (notice the inconsistent hypen):

bool HasBug = FParse::Value(CommandLine, TEXT("-bug="), Filename);
bool Submit = FParse::Param(CommandLine, TEXT("submit"));

The hyphen is totally optional, and, if supplied, is simply ignored.  This way, no existing code needs to be touched.

diff --git a/Engine/Source/Runtime/Core/Private/Misc/Parse.cpp b/Engine/Source/Runtime/Core/Private/Misc/Parse.cpp
--- a/Engine/Source/Runtime/Core/Private/Misc/Parse.cpp
+++ b/Engine/Source/Runtime/Core/Private/Misc/Parse.cpp
@@ -336,6 +336,11 @@ bool FParse::Value(
 //
 bool FParse::Param( const TCHAR* Stream, const TCHAR* Param )
 {
+	if (*Param == '-' || *Param == '/')
+	{
+		++Param;
+	}
+
 	const TCHAR* Start = Stream;
 	if( *Stream )
 	{
