Removed some extern functions declared inside inlined functions. c++modules/header units do not like this. Made them static so they are not exposed in ifc

diff --git a/Engine/Source/Runtime/Core/Private/Misc/CoreGlobals.cpp b/Engine/Source/Runtime/Core/Private/Misc/CoreGlobals.cpp
--- a/Engine/Source/Runtime/Core/Private/Misc/CoreGlobals.cpp
+++ b/Engine/Source/Runtime/Core/Private/Misc/CoreGlobals.cpp
@@ -257,7 +257,7 @@ float				GNearClippingPlane_RenderThread		= 10.0f;				/* Near clipping plane (Re
 
 bool					GExitPurge						= false;
 
-FChunkedFixedUObjectArray* GCoreObjectArrayForDebugVisualizers = nullptr;
+CORE_API FChunkedFixedUObjectArray* GCoreObjectArrayForDebugVisualizers = nullptr;
 
 namespace UE::CoreUObject::Private
 {
@@ -1008,6 +1008,17 @@ bool IsUsingZenPakFileStreaming()
 
 }
 
+bool IsThisNotNull(const void* This, const ANSICHAR* Function)
+{
+	if (This)
+	{
+		return true;
+	}
+
+	UE_LOG(LogCore, Warning, TEXT("'%hs' was called with a null 'this' pointer - this is no longer supported."), Function);
+	return false;
+}
+
 namespace UE::Private
 {
 
diff --git a/Engine/Source/Runtime/Core/Public/CoreGlobals.h b/Engine/Source/Runtime/Core/Public/CoreGlobals.h
--- a/Engine/Source/Runtime/Core/Public/CoreGlobals.h
+++ b/Engine/Source/Runtime/Core/Public/CoreGlobals.h
@@ -834,13 +834,4 @@ namespace UE {
  * on undefined behavior.
  */
 UE_DEPRECATED(5.6, "IsThisNotNull has been deprecated.")
-FORCENOINLINE static bool IsThisNotNull(const void* This, const ANSICHAR* Function)
-{
-	if (This)
-	{
-		return true;
-	}
-
-	UE_LOG(LogCore, Warning, TEXT("'%hs' was called with a null 'this' pointer - this is no longer supported."), Function);
-	return false;
-}
+CORE_API FORCENOINLINE bool IsThisNotNull(const void* This, const ANSICHAR* Function);
diff --git a/Engine/Source/Runtime/Core/Public/Misc/CoreMisc.h b/Engine/Source/Runtime/Core/Public/Misc/CoreMisc.h
--- a/Engine/Source/Runtime/Core/Public/Misc/CoreMisc.h
+++ b/Engine/Source/Runtime/Core/Public/Misc/CoreMisc.h
@@ -138,6 +138,12 @@ CORE_API bool WillNeedAudioVisualData();
 	Runtime.
 -----------------------------------------------------------------------------*/
 
+#if UE_EDITOR
+CORE_API int32 StaticDedicatedServerCheck(); // Internal, use IsRunningDedicatedServer()
+CORE_API int32 StaticGameCheck(); // Internal, use IsRunningGame()
+CORE_API int32 StaticClientOnlyCheck(); // Internal, use IsRunningClientOnly()
+#endif
+
 /**
  * Check to see if this executable was launched as a dedicated server process and should not load client only data.
  * An editor build can be launched with -server to set this to true, but it will be false during single process PlayInEditor mode.
@@ -156,7 +162,6 @@ FORCEINLINE bool IsRunningDedicatedServer()
 	}
 
 #if UE_EDITOR
-	extern CORE_API int32 StaticDedicatedServerCheck();
 	return (StaticDedicatedServerCheck() == 1);
 #else
 	return false;
@@ -182,7 +187,6 @@ FORCEINLINE bool IsRunningGame()
 	}
 
 #if UE_EDITOR
-	extern CORE_API int32 StaticGameCheck();
 	return (StaticGameCheck() == 1);
 #else
 	return false;
@@ -203,7 +207,6 @@ FORCEINLINE bool IsRunningClientOnly()
 	}
 
 #if UE_EDITOR
-	extern CORE_API int32 StaticClientOnlyCheck();
 	return (StaticClientOnlyCheck() == 1);
 #else
 	return false;
diff --git a/Engine/Source/Runtime/Engine/Private/AudioDebug.cpp b/Engine/Source/Runtime/Engine/Private/AudioDebug.cpp
--- a/Engine/Source/Runtime/Engine/Private/AudioDebug.cpp
+++ b/Engine/Source/Runtime/Engine/Private/AudioDebug.cpp
@@ -29,7 +29,15 @@
 
 namespace Audio
 {
-	ENGINE_API FString ExtraTracingSoundWaveFilterCVar;
+	FString ExtraTracingSoundWaveFilterCVar;
+
+#if !UE_BUILD_SHIPPING
+	bool MatchesLogFilter(const TCHAR* InName)
+	{
+		return !ExtraTracingSoundWaveFilterCVar.IsEmpty()	// Trivially reject if we have no filter.
+		 && FCString::Stristr(InName, *ExtraTracingSoundWaveFilterCVar);
+	}
+#endif
 }
 
 static FAutoConsoleVariableRef CVarExtraTracingSoundWaveFilter(
diff --git a/Engine/Source/Runtime/Engine/Public/Audio/AudioDebug.h b/Engine/Source/Runtime/Engine/Public/Audio/AudioDebug.h
--- a/Engine/Source/Runtime/Engine/Public/Audio/AudioDebug.h
+++ b/Engine/Source/Runtime/Engine/Public/Audio/AudioDebug.h
@@ -9,7 +9,7 @@
 namespace Audio
 {
 	// enum shortener for logging.
-	static const TCHAR* ShortEnumString(const TCHAR* InString)
+	inline const TCHAR* ShortEnumString(const TCHAR* InString)
 	{
 		if (const TCHAR* DoubleColon = FCString::Strstr(InString, TEXT("::")))
 		{
@@ -18,16 +18,11 @@ namespace Audio
 		return InString;
 	}
 	
-	FORCEINLINE bool MatchesLogFilter(const TCHAR* InName)
-	{
 #if UE_BUILD_SHIPPING
-		return false;		// Off for shipping.
+	FORCEINLINE bool MatchesLogFilter(const TCHAR* InName) { return false; } // Off for shipping.
 #else
-		extern ENGINE_API FString ExtraTracingSoundWaveFilterCVar;
-		return !ExtraTracingSoundWaveFilterCVar.IsEmpty()	// Trivially reject if we have no filter.
-		 && FCString::Stristr(InName, *ExtraTracingSoundWaveFilterCVar);
+	ENGINE_API bool MatchesLogFilter(const TCHAR* InName);
 #endif // UE_BUILD_SHIPPING
-	}
 	
 	FORCEINLINE bool MatchesLogFilter(const FName& InName)
 	{
diff --git a/Engine/Source/Runtime/RHI/Private/GPUProfiler.cpp b/Engine/Source/Runtime/RHI/Private/GPUProfiler.cpp
--- a/Engine/Source/Runtime/RHI/Private/GPUProfiler.cpp
+++ b/Engine/Source/Runtime/RHI/Private/GPUProfiler.cpp
@@ -88,7 +88,7 @@ static TAutoConsoleVariable<int32> GProfileShowEventHistogram(
 	TEXT("Whether the event histogram should be shown."),
 	ECVF_Default);
 
-RHI_API TAutoConsoleVariable<int32> GProfileGPUTransitions(
+TAutoConsoleVariable<int32> GProfileGPUTransitions(
 	TEXT("r.ProfileGPU.ShowTransitions"),
 	0,
 	TEXT("Allows profileGPU to display resource transition events."),
diff --git a/Engine/Source/Runtime/RHI/Private/RHICommandList.cpp b/Engine/Source/Runtime/RHI/Private/RHICommandList.cpp
--- a/Engine/Source/Runtime/RHI/Private/RHICommandList.cpp
+++ b/Engine/Source/Runtime/RHI/Private/RHICommandList.cpp
@@ -98,6 +98,8 @@ FAutoConsoleTaskPriority CPrio_SceneRenderingTask(
 	ENamedThreads::HighTaskPriority
 );
 
+extern TAutoConsoleVariable<int32> GProfileGPUTransitions;
+
 DECLARE_CYCLE_STAT(TEXT("Parallel Translate"),                 STAT_ParallelTranslate,      STATGROUP_RHICMDLIST);
 DECLARE_CYCLE_STAT(TEXT("RHI Thread Parallel Translate Wait"), STAT_ParallelTranslateWait,  STATGROUP_RHICMDLIST);
 DECLARE_CYCLE_STAT(TEXT("Explicit wait for tasks"),            STAT_ExplicitWait,           STATGROUP_RHICMDLIST);
@@ -1934,6 +1936,13 @@ void FRHICommandListImmediate::EndFrame()
 	GDynamicRHI->RHIEndFrame_RenderThread(*this);
 }
 
+#if WITH_PROFILEGPU && (RHI_NEW_GPU_PROFILER == 0)
+int32 FRHIComputeCommandList::GetGProfileGPUTransitions()
+{
+	return GProfileGPUTransitions.GetValueOnAnyThread();
+}
+#endif
+
 void FRHIComputeCommandList::Transition(TArrayView<const FRHITransitionInfo> Infos, ERHITransitionCreateFlags CreateFlags)
 {
 	TransitionInternal(Infos, CreateFlags);
diff --git a/Engine/Source/Runtime/RHI/Public/RHICommandList.h b/Engine/Source/Runtime/RHI/Public/RHICommandList.h
--- a/Engine/Source/Runtime/RHI/Public/RHICommandList.h
+++ b/Engine/Source/Runtime/RHI/Public/RHICommandList.h
@@ -3062,11 +3062,14 @@ public:
 		ALLOC_COMMAND(FRHICommandClearUAVUint)(UnorderedAccessViewRHI, Values);
 	}
 
+#if WITH_PROFILEGPU && (RHI_NEW_GPU_PROFILER == 0)
+	RHI_API static int32 GetGProfileGPUTransitions();
+#endif
+
 	FORCEINLINE_DEBUGGABLE void BeginTransitions(TArrayView<const FRHITransition*> Transitions)
 	{
 #if WITH_PROFILEGPU && (RHI_NEW_GPU_PROFILER == 0)
-		extern RHI_API TAutoConsoleVariable<int32> GProfileGPUTransitions;
-		RHI_BREADCRUMB_EVENT_CONDITIONAL(*this, GProfileGPUTransitions.GetValueOnAnyThread() != 0, "RHIBeginTransitions");
+		RHI_BREADCRUMB_EVENT_CONDITIONAL(*this, GetGProfileGPUTransitions() != 0, "RHIBeginTransitions");
 #endif
 		if (Bypass())
 		{
@@ -3106,8 +3109,7 @@ public:
 	FORCEINLINE_DEBUGGABLE void EndTransitions(TArrayView<const FRHITransition*> Transitions)
 	{
 #if WITH_PROFILEGPU && (RHI_NEW_GPU_PROFILER == 0)
-		extern RHI_API TAutoConsoleVariable<int32> GProfileGPUTransitions;
-		RHI_BREADCRUMB_EVENT_CONDITIONAL(*this, GProfileGPUTransitions.GetValueOnAnyThread() != 0, "RHIEndTransitions");
+		RHI_BREADCRUMB_EVENT_CONDITIONAL(*this, GetGProfileGPUTransitions() != 0, "RHIEndTransitions");
 #endif
 		if (Bypass())
 		{
