Added support for TErrorDetails type.

diff --git a/Engine/Source/Runtime/Core/Private/Experimental/IO/IoStatusError.cpp b/Engine/Source/Runtime/Core/Private/Experimental/IO/IoStatusError.cpp
--- a/Engine/Source/Runtime/Core/Private/Experimental/IO/IoStatusError.cpp
+++ b/Engine/Source/Runtime/Core/Private/Experimental/IO/IoStatusError.cpp
@@ -42,6 +42,10 @@ namespace UE::UnifiedError::IoStore
 			OutProperties.AddProperty(UTF8TEXTVIEW("IoStoreErrorMessage"), ErrorMessage);
 		}
 
+		virtual const FAnsiStringView GetErrorDetailsTypeName() const
+		{
+			return ANSITEXTVIEW("FIoStoreErrorDetails");
+		}
 	};
 
 	const FText FIoStoreErrorDetails::GenericErrorFormatString = NSLOCTEXT("IoStore", "GenericErrorMessage", "{IoStoreErrorMessage}");
diff --git a/Engine/Source/Runtime/Core/Private/Experimental/UnifiedError/UnifiedError.cpp b/Engine/Source/Runtime/Core/Private/Experimental/UnifiedError/UnifiedError.cpp
--- a/Engine/Source/Runtime/Core/Private/Experimental/UnifiedError/UnifiedError.cpp
+++ b/Engine/Source/Runtime/Core/Private/Experimental/UnifiedError/UnifiedError.cpp
@@ -52,12 +52,47 @@ FText FError::GetFormatErrorText() const
 }
 
 
-FText FError::GetErrorMessage() const
+#define USE_STRUCTURED_LOG_FOR_FERRORMESSAGE 0
+
+FText FError::GetErrorMessage(bool bIncludeContext) const
 {
+#if USE_STRUCTURED_LOG_FOR_FERRORMESSAGE
+	FString FormatString = ErrorDetails->GetErrorFormatString(*this).ToString();
+
+	FLogTemplate* Template = UE::CreateLogTemplate(*FormatString);
+
+
+	TCbWriter<1024> Writer;
+	Writer.BeginObject();
+
+	SerializeDetailsForLog(Writer);
+
+	Writer.EndObject();
+
+	FWideStringBuilderBase OutputMessage;
+	FormatLogTo(OutputMessage,  Template, Writer.Save()); 
+
+	UE::DestroyLogTemplate(Template);
+	return FText::FromString(OutputMessage.ToString());
+#else
 	FFormatNamedArguments Args;
 	FTextFormatArgsPropertyExtractor PropertyExtractor(Args);
 	ErrorDetails->GetErrorProperties(*this, PropertyExtractor);
 	return FText::Format(ErrorDetails->GetErrorFormatString(*this), Args);
+#endif
+}
+
+void FError::SerializeDetailsForLog(FCbWriter& Writer) const
+{
+	const IErrorDetails* DetailsIt = ErrorDetails.GetReference();
+	while (DetailsIt != nullptr)
+	{
+		Writer.SetName(DetailsIt->GetErrorDetailsTypeName());
+		DetailsIt->SerializeForLog(Writer);
+
+
+		DetailsIt = DetailsIt->GetInnerErrorDetails().GetReference();
+	}
 }
 
 int32 FError::GetErrorCode() const
@@ -227,6 +262,14 @@ void SerializeForLog(FCbWriter& Writer, const FError& Error)
 	Writer.BeginObject();
 	Writer.AddString(ANSITEXTVIEW("$type"), ANSITEXTVIEW("FError"));
 
+#if USE_STRUCTURED_LOG_FOR_FERRORMESSAGE
+	
+
+	Error.SerializeDetailsForLog(Writer);
+
+	Writer.AddString(ANSITEXTVIEW("$text"), Error.GetFormatErrorText().ToString());
+
+#else
 
 	FCbWriterErrorPropertyExtractor ExtractToCbWriter(Writer);
 	Error.GetErrorProperties(ExtractToCbWriter);
@@ -244,6 +287,7 @@ void SerializeForLog(FCbWriter& Writer, const FError& Error)
 	Writer.AddString("$text", FText::Format(ErrorFormatText, MoveTemp(ErrorMessageArgs)).ToString());
 
 	Writer.EndObject();
+#endif
 }
 
 // FRefCountedErrorDetails fucntions
diff --git a/Engine/Source/Runtime/Core/Public/Experimental/UnifiedError/UnifiedError.h b/Engine/Source/Runtime/Core/Public/Experimental/UnifiedError/UnifiedError.h
--- a/Engine/Source/Runtime/Core/Public/Experimental/UnifiedError/UnifiedError.h
+++ b/Engine/Source/Runtime/Core/Public/Experimental/UnifiedError/UnifiedError.h
@@ -6,8 +6,50 @@
 #include "Containers/StringView.h"
 #include "Internationalization/Internationalization.h"
 #include "Internationalization/Text.h"
+#include "Logging/StructuredLog.h"
+#include "Serialization/CompactBinaryWriter.h"
 #include "Templates/RefCounting.h"
 
+
+#include <type_traits>
+
+template<typename>
+struct deferred_false : std::false_type {};
+
+
+template<typename T>
+class TErrorStructFeatures
+{
+public:
+	static const FAnsiStringView GetErrorDetailsTypeNameAsString()
+	{
+		static_assert(deferred_false<T>::value, "All types are required to define DECLARE_ERRORSTRUCT_FEATURES to be used as a error context type");
+	}
+
+	static const FAnsiStringView GetErrorContextTypeNameAsString()
+	{
+		static_assert(deferred_false<T>::value, "All types are required to define DECLARE_ERRORSTRUCT_FEATURES to be used as a error context type");
+	}
+};
+
+
+
+#define DECLARE_ERRORSTRUCT_FEATURES(DetailsTypeName)\
+	template<>\
+	class TErrorStructFeatures<UE::UnifiedError::DetailsTypeName>\
+	{\
+	public:\
+		static const FAnsiStringView GetErrorDetailsTypeNameAsString() \
+		{\
+			return ANSITEXTVIEW("TErrorDetails<"#DetailsTypeName">");\
+		}\
+		static const FAnsiStringView GetErrorContextTypeNameAsString() \
+		{\
+			return ANSITEXTVIEW(#DetailsTypeName);\
+		}\
+	};
+
+
 #define DECLARE_ERROR_DETAILS_INTERNAL(DetailsNamespace, TypeName)\
 	static uint32 StaticGetErrorDetailsTypeId() \
 	{ \
@@ -16,10 +58,13 @@
 	virtual uint32 GetErrorDetailsTypeId() const \
 	{ \
 		return StaticGetErrorDetailsTypeId(); \
+	}\
+	virtual const FAnsiStringView GetErrorDetailsTypeName() const override \
+	{\
+		return #DetailsNamespace"::"#TypeName;\
 	}
 
 
-
 #define DECLARE_FERROR_DETAILS_ABSTRACT(DetailsNamespace, TypeName)\
 	public:\
 	inline static const uint32 StaticDetailsTypeId = UE::UnifiedError::FErrorDetailsRegistry::Get().RegisterDetails(ANSITEXTVIEW(#DetailsNamespace "::" #TypeName), nullptr);\
@@ -108,6 +153,10 @@ namespace UE::UnifiedError
 		/// <returns></returns>
 		virtual uint32 GetErrorDetailsTypeId() const = 0;
 
+
+		CORE_API virtual const FAnsiStringView GetErrorDetailsTypeName() const = 0;
+		CORE_API virtual void SerializeForLog(FCbWriter& Writer) const {}
+
 	};
 
 
@@ -242,6 +291,79 @@ namespace UE::UnifiedError
 		virtual uint32 GetRefCount() const final { return 10; }
 	};
 
+
+	class FError;
+} // UE::UnifiedError
+
+template <typename ValueType>
+inline void GatherPropertiesForError(const UE::UnifiedError::FError& Writer, const ValueType& Value, UE::UnifiedError::IErrorPropertyExtractor& Extractor)
+{
+	static_assert(deferred_false<ValueType>::value, "undefined function, define your own function pls");
+}
+
+namespace UE::UnifiedError
+{
+
+	template<typename T>
+	class TErrorDetails : public FDynamicErrorDetails
+	{
+	public:
+		TErrorDetails() {}
+		TErrorDetails(T&& InErrorDetail)
+		{
+			ErrorDetail = MoveTemp(InErrorDetail);
+		}
+		TErrorDetails(T&& InErrorDetail, TRefCountPtr<const IErrorDetails> InInnerErrorDetails) : FDynamicErrorDetails(InInnerErrorDetails)
+		{
+			ErrorDetail = MoveTemp(InErrorDetail);
+		}
+
+		friend IErrorDetails* Create();
+		inline static IErrorDetails* Create()
+		{
+			return new TErrorDetails<T>();
+		}
+		inline static const uint32 StaticDetailsTypeId = UE::UnifiedError::FErrorDetailsRegistry::Get().RegisterDetails(TErrorStructFeatures<T>::GetNameAsString(), TFunction<IErrorDetails * ()>([]() -> IErrorDetails* { return Create(); }));
+		static uint32 StaticGetErrorDetailsTypeId()
+		{
+			return StaticDetailsTypeId;
+		}
+		virtual uint32 GetErrorDetailsTypeId() const
+		{
+			return StaticGetErrorDetailsTypeId();
+		}
+	private:
+		T ErrorDetail;
+
+		FAnsiStringView GetTypeNameAsString() const
+		{
+			return TErrorStructFeatures<T>::GetErrorDetailsTypeNameAsString();
+		}
+	public:
+		const T& GetErrorContext() const
+		{
+			return &ErrorDetail;
+		}
+
+		virtual void GetErrorProperties(const FError& Error, IErrorPropertyExtractor& OutProperties) const override
+		{
+			GatherPropertiesForError(Error, ErrorDetail, OutProperties);
+			
+			FDynamicErrorDetails::GetErrorProperties(Error, OutProperties);
+		}
+		virtual void SerializeForLog(FCbWriter& Writer) const override
+		{
+			::SerializeForLog(Writer, ErrorDetail);
+		}
+
+		const T& GetValue() const
+		{
+			return ErrorDetail;
+		}
+
+	};
+
+
 	class FError
 	{
 	private: 
@@ -305,9 +427,31 @@ namespace UE::UnifiedError
 			return nullptr;
 		}
 
-		CORE_API FText GetErrorMessage() const;
+		CORE_API FText GetErrorMessage(bool bAppendContext = false) const;
 		CORE_API FText GetFormatErrorText() const;
+		template<typename T>
+		void PushErrorContext(T&& InErrorStruct)
+		{
+			PushErrorDetails(TRefCountPtr<IErrorDetails>(new TErrorDetails<T>(MoveTemp(InErrorStruct))));
+		}
+
+		template<typename T>
+		const T* GetErrorContext() const
+		{
+			TRefCountPtr<const IErrorDetails> CurrentIt = ErrorDetails;
+			while (CurrentIt != nullptr)
+			{
+				if (CurrentIt->GetErrorDetailsTypeName() == TErrorStructFeatures<T>::GetErrorDetailsTypeNameAsString())
+				{
+					const TErrorDetails<const T>* Details = (const TErrorDetails<const T>*)(CurrentIt.GetReference());
+					return &Details->GetErrorContext();
+				}
+				CurrentIt = CurrentIt->GetInnerErrorDetails();
+			}
+			return nullptr;
+		}
 
+		CORE_API void SerializeDetailsForLog(FCbWriter& Writer) const;
 
 		CORE_API TRefCountPtr<const FStaticErrorDetails> GetStaticErrorDetails() const;
 
diff --git a/Engine/Source/Runtime/Engine/Classes/Engine/Experimental/StreamableManagerError.h b/Engine/Source/Runtime/Engine/Classes/Engine/Experimental/StreamableManagerError.h
--- a/Engine/Source/Runtime/Engine/Classes/Engine/Experimental/StreamableManagerError.h
+++ b/Engine/Source/Runtime/Engine/Classes/Engine/Experimental/StreamableManagerError.h
@@ -9,7 +9,6 @@ DECLARE_ERROR_MODULE(StreamableManager, 0x10);
 
 
 DECLARE_ERROR_ONEPARAM(PackageLoadFailed, 1, StreamableManager, NSLOCTEXT("StreamableManager", "PackageLoadError", "Failed to load package {PackageName}"), FString, PackageName, TEXT("Unknown"));
-
 DECLARE_ERROR(PackageLoadCancelled, 2, StreamableManager, NSLOCTEXT("StreamableManager", "PackageLoadCancelled", "Async load canceled"));
 DECLARE_ERROR(DownloadError, 3, StreamableManager, NSLOCTEXT("StreamableManager", "DownloadError", "Failed to download"));
 DECLARE_ERROR_ONEPARAM(PackageNameInvalid, 4, StreamableManager, NSLOCTEXT("StreamableManager", "PackageNameInvalid", "Found invalid package name {InvalidPackageName}"), FString, InvalidPackageName, TEXT("Unknown"));
@@ -64,8 +63,38 @@ namespace UE::UnifiedError::StreamableManager
 		}
 
 		const FString& GetSoftObjectPath() const { return SoftObjectPath; }
-	};
 
+		
+	};
 	UE::UnifiedError::FError GetStreamableError(EAsyncLoadingResult::Type Result);
+
+
+
+	class FStreamableManagerAdditionalContext
+	{
+	public:
+		/*friend void SerializeForLog(FCbWriter& Writer, const UE::UnifiedError::StreamableManager::FStreamableManagerAdditionalContext& Context);
+	private:*/
+		FString ExtraSoftObjectPath;
+	};
+
+
+	inline void GatherPropertiesForError(const UE::UnifiedError::FError& Error, const UE::UnifiedError::StreamableManager::FStreamableManagerAdditionalContext& Context, UE::UnifiedError::IErrorPropertyExtractor& PropertyExtractor)
+	{
+		PropertyExtractor.AddProperty(ANSITEXTVIEW("ExtraSoftObjectPath"), Context.ExtraSoftObjectPath);
+	}
+	
+}
+
+inline void SerializeForLog(FCbWriter& Writer, const UE::UnifiedError::StreamableManager::FStreamableManagerAdditionalContext& Context)
+{
+	Writer.BeginObject();
+	Writer.AddString(ANSITEXTVIEW("$type"), ANSITEXTVIEW("FStreamableManagerAdditionalContext"));
+	FStringFormatNamedArguments NamedArguments;
+	NamedArguments.Add(TTuple<FString, FStringFormatArg>(TEXT("SoftObjectPath"), Context.ExtraSoftObjectPath));
+	Writer.AddString(ANSITEXTVIEW("$text"), FString::Format(TEXT("(SoftObjectPath: {SoftObjectPath})"), NamedArguments));
+	Writer.AddString(ANSITEXTVIEW("SoftObjectPath"), Context.ExtraSoftObjectPath);
+	Writer.EndObject();
 }
+DECLARE_ERRORSTRUCT_FEATURES(StreamableManager::FStreamableManagerAdditionalContext);
 
diff --git a/Engine/Source/Runtime/Engine/Private/StreamableManager.cpp b/Engine/Source/Runtime/Engine/Private/StreamableManager.cpp
--- a/Engine/Source/Runtime/Engine/Private/StreamableManager.cpp
+++ b/Engine/Source/Runtime/Engine/Private/StreamableManager.cpp
@@ -1228,6 +1228,8 @@ void FStreamableHandle::AsyncLoadCallbackWrapper(const FName& PackageName, UPack
 			TRefCountPtr<UE::UnifiedError::IErrorDetails> ErrorDetails(new UE::UnifiedError::StreamableManager::FStreamableManagerErrorDetails(TargetName.ToString()));
 			MaybeError->PushErrorDetails(ErrorDetails);
 
+			/*UE::UnifiedError::StreamableManager::FStreamableManagerAdditionalContext MoreContext = {TargetName.ToString()};
+			MaybeError->PushErrorStruct(MoveTemp(MoreContext));*/
 			SetError(*MaybeError);
 		}
 
