[Core] Made FLinearColor::CopyWithNewOpacity constexpr

diff --git a/Engine/Source/Runtime/Core/Public/Math/Color.h b/Engine/Source/Runtime/Core/Public/Math/Color.h
--- a/Engine/Source/Runtime/Core/Public/Math/Color.h
+++ b/Engine/Source/Runtime/Core/Public/Math/Color.h
@@ -287,10 +287,10 @@ struct FLinearColor
 		return FMath::Abs(this->R - ColorB.R) < Tolerance && FMath::Abs(this->G - ColorB.G) < Tolerance && FMath::Abs(this->B - ColorB.B) < Tolerance && FMath::Abs(this->A - ColorB.A) < Tolerance;
 	}
 
-	[[nodiscard]] FLinearColor CopyWithNewOpacity(float NewOpacicty) const
+	[[nodiscard]] constexpr FLinearColor CopyWithNewOpacity(float NewOpacity) const
 	{
 		FLinearColor NewCopy = *this;
-		NewCopy.A = NewOpacicty;
+		NewCopy.A = NewOpacity;
 		return NewCopy;
 	}
 
