Fix issue with FError crash on non static errors.

diff --git a/Engine/Source/Runtime/Core/Private/Experimental/UnifiedError/UnifiedError.cpp b/Engine/Source/Runtime/Core/Private/Experimental/UnifiedError/UnifiedError.cpp
--- a/Engine/Source/Runtime/Core/Private/Experimental/UnifiedError/UnifiedError.cpp
+++ b/Engine/Source/Runtime/Core/Private/Experimental/UnifiedError/UnifiedError.cpp
@@ -196,24 +196,18 @@ void FError::GetErrorProperties(IErrorPropertyExtractor& Visitor) const
 	ErrorDetails->GetErrorProperties(*this, Visitor);
 }
 
-TRefCountPtr<const FStaticErrorDetails> FError::GetStaticErrorDetails() const
+const FUtf8String FError::GetErrorCodeString() const
 {
-	TRefCountPtr<const IErrorDetails> InnerMostError = GetInnerMostErrorDetails();
-	check(InnerMostError.IsValid() && (InnerMostError->GetErrorDetailsTypeId() == FStaticErrorDetails::StaticGetErrorDetailsTypeId())); // innermost error should always be something static
-	TRefCountPtr<const FStaticErrorDetails> StaticErrorDetails((const FStaticErrorDetails*)InnerMostError.GetReference());
-	return StaticErrorDetails;
+	FUtf8String Result;
+	GetDetailByKey(UTF8TEXTVIEW("ErrorCodeString"), Result);
+	return Result;
 }
 
-const FAnsiStringView& FError::GetErrorCodeString() const
+const FUtf8String FError::GetModuleIdString() const
 {
-	TRefCountPtr<const FStaticErrorDetails> StaticErrorDetails = GetStaticErrorDetails();
-	return StaticErrorDetails->GetErrorCodeString();
-}
-
-const FAnsiStringView& FError::GetModuleIdString() const
-{
-	TRefCountPtr<const FStaticErrorDetails> StaticErrorDetails = GetStaticErrorDetails();
-	return StaticErrorDetails->GetModuleIdString();
+	FUtf8String Result;
+	GetDetailByKey(UTF8TEXTVIEW("ModuleIdString"), Result);
+	return Result;
 }
 
 // Structured logging FError integration
@@ -375,8 +369,10 @@ uint32 FErrorDetailsRegistry::RegisterDetails(const FAnsiStringView& ErrorDetail
 
 FString LexToString(const UE::UnifiedError::FError& Error)
 {
-	return FString::Printf(TEXT("%.*hs:%.*hs - %s"), 
-		Error.GetModuleIdString().Len(), Error.GetModuleIdString().GetData(),
-		Error.GetErrorCodeString().Len(), Error.GetErrorCodeString().GetData(),
+	
+	// TODO: daniel, switch to using structured logging output in the future instead of this
+	return FString::Printf(TEXT("%s:%s - %s"), 
+		UTF8_TO_TCHAR(*Error.GetModuleIdString()),
+		UTF8_TO_TCHAR(*Error.GetErrorCodeString()),
 		*Error.GetErrorMessage().ToString());
 }
\ No newline at end of file
diff --git a/Engine/Source/Runtime/Core/Public/Experimental/UnifiedError/UnifiedError.h b/Engine/Source/Runtime/Core/Public/Experimental/UnifiedError/UnifiedError.h
--- a/Engine/Source/Runtime/Core/Public/Experimental/UnifiedError/UnifiedError.h
+++ b/Engine/Source/Runtime/Core/Public/Experimental/UnifiedError/UnifiedError.h
@@ -24,11 +24,13 @@ public:
 	static const FAnsiStringView GetErrorDetailsTypeNameAsString()
 	{
 		static_assert(deferred_false<T>::value, "All types are required to define DECLARE_ERRORSTRUCT_FEATURES to be used as a error context type");
+		return ANSITEXTVIEW("");
 	}
 
 	static const FAnsiStringView GetErrorContextTypeNameAsString()
 	{
 		static_assert(deferred_false<T>::value, "All types are required to define DECLARE_ERRORSTRUCT_FEATURES to be used as a error context type");
+		return ANSITEXTVIEW("");
 	}
 };
 
@@ -408,8 +410,8 @@ namespace UE::UnifiedError
 		CORE_API int32 GetErrorCode() const;
 		CORE_API int32 GetModuleId() const;
 
-		CORE_API const FAnsiStringView& GetErrorCodeString() const;
-		CORE_API const FAnsiStringView& GetModuleIdString() const;
+		CORE_API const FUtf8String GetErrorCodeString() const;
+		CORE_API const FUtf8String GetModuleIdString() const;
 
 		CORE_API void PushErrorDetails(TRefCountPtr<IErrorDetails> InErrorDetails);
 		template<typename DetailType>
@@ -453,8 +455,6 @@ namespace UE::UnifiedError
 
 		CORE_API void SerializeDetailsForLog(FCbWriter& Writer) const;
 
-		CORE_API TRefCountPtr<const FStaticErrorDetails> GetStaticErrorDetails() const;
-
 		CORE_API bool GetDetailByKey(const FUtf8StringView& KeyName, FString& Result) const;
 		CORE_API bool GetDetailByKey(const FUtf8StringView& KeyName, FUtf8String& Result) const;
 		CORE_API bool GetDetailByKey(const FUtf8StringView& KeyName, FText& Result) const;
