Logging: Added support for optional sub-object references in structured log formatting, with tests

diff --git a/Engine/Source/Runtime/Core/Private/Logging/StructuredLog.cpp b/Engine/Source/Runtime/Core/Private/Logging/StructuredLog.cpp
--- a/Engine/Source/Runtime/Core/Private/Logging/StructuredLog.cpp
+++ b/Engine/Source/Runtime/Core/Private/Logging/StructuredLog.cpp
@@ -4,6 +4,7 @@
 
 #include "Algo/Accumulate.h"
 #include "Algo/AllOf.h"
+#include "Algo/Find.h"
 #include "Algo/NoneOf.h"
 #include "AutoRTFM.h"
 #include "Containers/AnsiString.h"
@@ -25,6 +26,7 @@
 #include "Serialization/CompactBinarySerialization.h"
 #include "Serialization/CompactBinaryValue.h"
 #include "Serialization/VarInt.h"
+#include "String/Split.h"
 #include "Templates/Function.h"
 #include <cstdarg>
 
@@ -46,7 +48,9 @@ CORE_API bool GConvertBasicLogToLogRecord = false;
 // Experimental feature to prepend log context to the log message during formatting.
 CORE_API bool GPrependLogContextToLogMessage = false;
 
+static constexpr ANSICHAR GSubObjectFieldDelimiter = '/';
 static constexpr FAsciiSet GValidLogFieldName("0123456789ABCDEFGHIJKLMNOPQRSTUVWXYZabcdefghijklmnopqrstuvwxyz_");
+static constexpr FAsciiSet GValidLogSubObjectFieldName = GValidLogFieldName | FAsciiSet({GSubObjectFieldDelimiter, '\0'});
 
 static constexpr FAnsiStringView GLogContextsFieldName = ANSITEXTVIEW("$Contexts");
 
@@ -54,7 +58,7 @@ static constexpr FAnsiStringView GLogContextsFieldName = ANSITEXTVIEW("$Contexts
 
 struct FLogTemplateOp
 {
-	enum EOpCode : int32 { OpEnd, OpSkip, OpText, OpName, OpIndex, OpLocText, OpCount };
+	enum EOpCode : int32 { OpEnd, OpSkip, OpText, OpName, OpPath, OpIndex, OpLocText, OpCount };
 
 	static constexpr int32 ValueShift = 3;
 	static_assert(OpCount <= (1 << ValueShift));
@@ -257,20 +261,21 @@ class FLogTemplate
 public:
 	static FLogTemplate* CreateStatic(const Logging::Private::FStaticLogRecord& Log, const FLogField* Fields, int32 FieldCount)
 	{
-		return Create(Log.Format, Fields, FieldCount);
+		return Create(Log.Format, {}, Fields, FieldCount);
 	}
 
 	static FLogTemplate* CreateStatic(const Logging::Private::FStaticLocalizedLogRecord& Log, const FLogField* Fields, int32 FieldCount)
 	{
-		return CreateLocalized(Log.TextNamespace, Log.TextKey, Log.Format, Fields, FieldCount);
+		return CreateLocalized(Log.TextNamespace, Log.TextKey, Log.Format, {}, Fields, FieldCount);
 	}
 
-	static FLogTemplate* Create(const TCHAR* Format, const FLogField* Fields = nullptr, int32 FieldCount = 0);
+	static FLogTemplate* Create(const TCHAR* Format, const FLogTemplateOptions& Options, const FLogField* Fields = nullptr, int32 FieldCount = 0);
 
 	static FLogTemplate* CreateLocalized(
 		const TCHAR* TextNamespace,
 		const TCHAR* TextKey,
 		const TCHAR* Format,
+		const FLogTemplateOptions& Options,
 		const FLogField* Fields = nullptr,
 		int32 FieldCount = 0);
 
@@ -308,7 +313,7 @@ private:
 	FLogTemplate* Next = nullptr;
 };
 
-FLogTemplate* FLogTemplate::Create(const TCHAR* Format, const FLogField* Fields, const int32 FieldCount)
+FLogTemplate* FLogTemplate::Create(const TCHAR* Format, const FLogTemplateOptions& Options, const FLogField* Fields, const int32 FieldCount)
 {
 	using namespace Logging::Private;
 
@@ -371,7 +376,7 @@ FLogTemplate* FLogTemplate::Create(const TCHAR* Format, const FLogField* Fields,
 
 		// Field
 		const TCHAR* const FieldStart = TextStart;
-		const TCHAR* const FieldNameEnd = FAsciiSet::Skip(FieldStart + 1, GValidLogFieldName);
+		const TCHAR* const FieldNameEnd = FAsciiSet::Skip(FieldStart + 1, GValidLogSubObjectFieldName);
 		checkf(*FieldNameEnd, TEXT("Log format has an unterminated field reference. Use '{{' to escape '{' if needed. [[%s]]"), Format);
 		checkf(*FieldNameEnd == TEXT('}'), TEXT("Log format has invalid character '%c' in field name. "
 			"Use '{{' to escape '{' if needed. Names must match \"[A-Za-z0-9][A-Za-z0-9_]*\". [[%s]]"), *FieldNameEnd, Format);
@@ -380,6 +385,9 @@ FLogTemplate* FLogTemplate::Create(const TCHAR* Format, const FLogField* Fields,
 		const int32 FieldNameLen = FieldLen - 2;
 		checkf(FieldStart[1] != TEXT('_'), TEXT("Log format uses reserved field name '%.*s' with leading '_'. "
 			"Names must match \"[A-Za-z0-9][A-Za-z0-9_]*\". [[%s]]"), FieldNameLen, FieldStart + 1, Format);
+		const bool bHasSubObjectReference = !!Algo::Find(FStringView(FieldStart + 1, FieldNameLen), GSubObjectFieldDelimiter);
+		checkf(!bHasSubObjectReference || Options.bAllowSubObjectReferences, TEXT("Log format has a sub-object reference (%c) "
+			"in a context that does not allow them. [[%s]]"), GSubObjectFieldDelimiter, Format);
 
 		if (bFindFields && !bPositional)
 		{
@@ -399,7 +407,7 @@ FLogTemplate* FLogTemplate::Create(const TCHAR* Format, const FLogField* Fields,
 				FieldNameLen, FieldStart + 1, Format);
 		}
 
-		Ops.Add({FLogTemplateOp::OpName, FieldLen});
+		Ops.Add({bHasSubObjectReference ? FLogTemplateOp::OpPath : FLogTemplateOp::OpName, FieldLen});
 		++FormatFieldCount;
 
 		TextStart = FieldEnd;
@@ -418,10 +426,18 @@ FLogTemplate* FLogTemplate::Create(const TCHAR* Format, const FLogField* Fields,
 	return Template;
 }
 
-FLogTemplate* FLogTemplate::CreateLocalized(const TCHAR* TextNamespace, const TCHAR* TextKey, const TCHAR* Format, const FLogField* Fields, const int32 FieldCount)
+FLogTemplate* FLogTemplate::CreateLocalized(
+	const TCHAR* TextNamespace,
+	const TCHAR* TextKey,
+	const TCHAR* Format,
+	const FLogTemplateOptions& Options,
+	const FLogField* Fields,
+	const int32 FieldCount)
 {
 	using namespace Logging::Private;
 
+	checkf(!Options.bAllowSubObjectReferences, TEXT("Localized log templates do not support sub-object references yet."));
+
 	FTextFormat TextFormat(FText::AsLocalizable_Advanced(TextNamespace, TextKey, Format));
 
 	const bool bFindFields = !!Fields;
@@ -521,6 +537,21 @@ void FLogTemplate::FormatTo(TStringBuilderBase<CharType>& Out, const FCbFieldVie
 		return It;
 	};
 
+	auto FindFieldByPath = [&FindField, Format = StaticFormat](FAnsiStringView Path, int32 IndexHint = -1) mutable -> FCbFieldView
+	{
+		FAnsiStringView Name = Path;
+		bool bMore = String::SplitFirstChar(Path, GSubObjectFieldDelimiter, Name, Path);
+		FCbFieldView Field = FindField(Name, IndexHint);
+		while (bMore)
+		{
+			Name = Path;
+			bMore = String::SplitFirstChar(Path, GSubObjectFieldDelimiter, Name, Path);
+			Field = Field.AsObjectView().FindView(Name);
+			checkf(Field, TEXT("Log format requires field '%.*hs' which was not provided. [[%s]]"), Path.Len(), Path.GetData(), Format);
+		}
+		return Field;
+	};
+
 	int32 FieldIndexHint = -1;
 	const uint8* NextOp = GetOpData();
 	const TCHAR* NextFormat = StaticFormat;
@@ -540,9 +571,18 @@ void FLogTemplate::FormatTo(TStringBuilderBase<CharType>& Out, const FCbFieldVie
 			FieldIndexHint = Op.Value;
 			break;
 		case FLogTemplateOp::OpName:
-			const auto Name = StringCast<ANSICHAR>(NextFormat + 1, Op.Value - 2);
-			LogFieldValue(Out, FindField(MakeStringView(Name.Get(), Name.Length()), FieldIndexHint));
-			FieldIndexHint = -1;
+			{
+				const auto Name = StringCast<ANSICHAR>(NextFormat + 1, Op.Value - 2);
+				LogFieldValue(Out, FindField(MakeStringView(Name.Get(), Name.Length()), FieldIndexHint));
+				FieldIndexHint = -1;
+			}
+			break;
+		case FLogTemplateOp::OpPath:
+			{
+				const auto Path = StringCast<ANSICHAR>(NextFormat + 1, Op.Value - 2);
+				LogFieldValue(Out, FindFieldByPath(MakeStringView(Path.Get(), Path.Length()), FieldIndexHint));
+				FieldIndexHint = -1;
+			}
 			break;
 		}
 		NextFormat += Op.GetSkipSize();
@@ -588,14 +628,14 @@ FText FLogTemplate::FormatLocalizedToText(const FCbFieldViewIterator& Fields) co
 
 ///////////////////////////////////////////////////////////////////////////////////////////////////////////////////////
 
-FLogTemplate* CreateLogTemplate(const TCHAR* Format)
+FLogTemplate* CreateLogTemplate(const TCHAR* Format, const FLogTemplateOptions& Options)
 {
-	return FLogTemplate::Create(Format);
+	return FLogTemplate::Create(Format, Options);
 }
 
-FLogTemplate* CreateLogTemplate(const TCHAR* TextNamespace, const TCHAR* TextKey, const TCHAR* Format)
+FLogTemplate* CreateLogTemplate(const TCHAR* TextNamespace, const TCHAR* TextKey, const TCHAR* Format, const FLogTemplateOptions& Options)
 {
-	return FLogTemplate::CreateLocalized(TextNamespace, TextKey, Format);
+	return FLogTemplate::CreateLocalized(TextNamespace, TextKey, Format, Options);
 }
 
 void DestroyLogTemplate(FLogTemplate* Template)
@@ -655,7 +695,8 @@ FORCENOINLINE static void FormatDynamicRecordMessageTo(TStringBuilderBase<CharTy
 	checkf(!TextNamespace == !TextKey,
 		TEXT("Log record must have both or neither of the text namespace and text key. [[%s]]"), Format);
 
-	FLogTemplate* LocalTemplate = TextKey ? FLogTemplate::CreateLocalized(TextNamespace, TextKey, Format) : FLogTemplate::Create(Format);
+	const FLogTemplateOptions Options;
+	FLogTemplate* LocalTemplate = TextKey ? FLogTemplate::CreateLocalized(TextNamespace, TextKey, Format, Options) : FLogTemplate::Create(Format, Options);
 	LocalTemplate->FormatTo(Out, Record.GetFields().CreateViewIterator());
 	FLogTemplate::Destroy(LocalTemplate);
 }
@@ -1197,7 +1238,7 @@ public:
 private:
 	FORCENOINLINE FLogTemplate& CreateTemplate()
 	{
-		FLogTemplate* NewTemplate = FLogTemplate::Create(GStaticBasicLogFormat);
+		FLogTemplate* NewTemplate = FLogTemplate::Create(GStaticBasicLogFormat, {});
 		if (FLogTemplate* ExistingTemplate = nullptr;
 			UNLIKELY(!Template.compare_exchange_strong(ExistingTemplate, NewTemplate, std::memory_order_release, std::memory_order_acquire)))
 		{
diff --git a/Engine/Source/Runtime/Core/Public/Logging/StructuredLogFormat.h b/Engine/Source/Runtime/Core/Public/Logging/StructuredLogFormat.h
--- a/Engine/Source/Runtime/Core/Public/Logging/StructuredLogFormat.h
+++ b/Engine/Source/Runtime/Core/Public/Logging/StructuredLogFormat.h
@@ -14,8 +14,15 @@ namespace UE { class FLogTemplate; }
 namespace UE
 {
 
-UE_API FLogTemplate* CreateLogTemplate(const TCHAR* Format);
-UE_API FLogTemplate* CreateLogTemplate(const TCHAR* TextNamespace, const TCHAR* TextKey, const TCHAR* Format);
+/** Options to control how log templates are parsed and perform formatting. */
+struct FLogTemplateOptions
+{
+	/** If true, allow field references of the form A.B.C to access fields of objects within objects. */
+	bool bAllowSubObjectReferences = false;
+};
+
+UE_API FLogTemplate* CreateLogTemplate(const TCHAR* Format, const FLogTemplateOptions& Options = {});
+UE_API FLogTemplate* CreateLogTemplate(const TCHAR* TextNamespace, const TCHAR* TextKey, const TCHAR* Format, const FLogTemplateOptions& Options = {});
 UE_API void DestroyLogTemplate(FLogTemplate* Template);
 
 UE_API void FormatLogTo(FUtf8StringBuilderBase& Out, const FLogTemplate* Template, const FCbFieldViewIterator& Fields);
diff --git a/Engine/Source/Runtime/Core/Tests/Logging/StructuredLogFormatTest.cpp b/Engine/Source/Runtime/Core/Tests/Logging/StructuredLogFormatTest.cpp
new file mode 100644
--- /dev/null
+++ b/Engine/Source/Runtime/Core/Tests/Logging/StructuredLogFormatTest.cpp
@@ -0,0 +1,156 @@
+// Copyright Epic Games, Inc. All Rights Reserved.
+
+#include "Logging/StructuredLogFormat.h"
+
+#if WITH_TESTS
+
+#include "Containers/UnrealString.h"
+#include "Misc/DateTime.h"
+#include "Misc/Guid.h"
+#include "Misc/Timespan.h"
+#include "Misc/StringBuilder.h"
+#include "Serialization/CompactBinaryWriter.h"
+#include "Tests/TestHarnessAdapter.h"
+
+namespace UE
+{
+
+TEST_CASE_NAMED(FStructuredLogFormatTest, "Core::Logging::StructuredLogFormat", "[Core][SmokeFilter]")
+{
+	const uint8 Binary[16] = { 1, 2, 3, 4, 5, 6, 7, 8, 9, 10, 11, 12, 13, 14, 15, 16 };
+	const FIoHash ObjectAttachment = FIoHash::HashBuffer(MakeMemoryView(ANSITEXTVIEW("ObjectAttachment")));
+	const FIoHash BinaryAttachment = FIoHash::HashBuffer(MakeMemoryView(ANSITEXTVIEW("BinaryAttachment")));
+	const FIoHash Hash = FIoHash::HashBuffer(MakeMemoryView(ANSITEXTVIEW("Hash")));
+	const FGuid Uuid = FGuid::NewGuidFromHash(FBlake3::HashBuffer(MakeMemoryView(ANSITEXTVIEW("Guid"))));
+	const FDateTime DateTime(2025, 4, 10, 11, 15, 30, 123);
+	const FTimespan TimeSpan(10, 11, 15, 30, 123456789);
+	const FCbObjectId ObjectId = FromGuid(FGuid::NewGuidFromHash(FBlake3::HashBuffer(MakeMemoryView(ANSITEXTVIEW("ObjectId")))));
+	const uint8 CustomById[8] = { 17, 18, 19, 20, 21, 22, 23, 24 };
+	const uint8 CustomByName[4] = { 25, 26, 27, 28 };
+
+	const auto MakeFields = [&]
+	{
+		TCbWriter<1024> Fields;
+
+		Fields.BeginObject(ANSITEXTVIEW("Object"));
+		Fields.AddInteger(ANSITEXTVIEW("X"), 1);
+		Fields.AddInteger(ANSITEXTVIEW("Y"), 2);
+		Fields.AddInteger(ANSITEXTVIEW("Z"), 3);
+		Fields.EndObject();
+
+		Fields.BeginArray(ANSITEXTVIEW("Array"));
+		Fields.AddInteger(1);
+		Fields.AddInteger(2);
+		Fields.AddInteger(3);
+		Fields.EndArray();
+
+		Fields.AddNull(ANSITEXTVIEW("Null"));
+		Fields.AddBinary(ANSITEXTVIEW("Binary"), MakeMemoryView(Binary));
+		Fields.AddString(ANSITEXTVIEW("String"), ANSITEXTVIEW("\"Quote\" with 4 words."));
+		Fields.AddInteger(ANSITEXTVIEW("IntegerNegative"), -64);
+		Fields.AddInteger(ANSITEXTVIEW("IntegerPositive"), 63);
+		Fields.AddFloat(ANSITEXTVIEW("Float"), 128.25f);
+		Fields.AddFloat(ANSITEXTVIEW("Double"), 123.456);
+		Fields.AddBool(ANSITEXTVIEW("False"), false);
+		Fields.AddBool(ANSITEXTVIEW("True"), true);
+		Fields.AddObjectAttachment(ANSITEXTVIEW("ObjectAttachment"), ObjectAttachment);
+		Fields.AddBinaryAttachment(ANSITEXTVIEW("BinaryAttachment"), BinaryAttachment);
+		Fields.AddHash(ANSITEXTVIEW("Hash"), Hash);
+		Fields.AddUuid(ANSITEXTVIEW("Uuid"), Uuid);
+		Fields.AddDateTime(ANSITEXTVIEW("DateTime"), DateTime);
+		Fields.AddTimeSpan(ANSITEXTVIEW("TimeSpan"), TimeSpan);
+		Fields.AddObjectId(ANSITEXTVIEW("ObjectId"), ObjectId);
+		Fields.AddCustom(ANSITEXTVIEW("CustomById"), 128, MakeMemoryView(CustomById));
+		Fields.AddCustom(ANSITEXTVIEW("CustomByName"), ANSITEXTVIEW("Custom"), MakeMemoryView(CustomByName));
+
+		Fields.BeginObject(ANSITEXTVIEW("ObjectText"));
+		Fields.AddInteger(ANSITEXTVIEW("X"), 1);
+		Fields.AddInteger(ANSITEXTVIEW("Y"), 2);
+		Fields.AddInteger(ANSITEXTVIEW("Z"), 3);
+		Fields.AddString(ANSITEXTVIEW("$text"), ANSITEXTVIEW("X=1;Y=2;Z=3"));
+		Fields.EndObject();
+
+		Fields.BeginObject(ANSITEXTVIEW("ObjectWithNestedText"));
+		Fields.BeginObject(ANSITEXTVIEW("X"));
+		Fields.AddInteger(ANSITEXTVIEW("$value"), 1);
+		Fields.AddString(ANSITEXTVIEW("$text"), ANSITEXTVIEW("0001"));
+		Fields.EndObject();
+		Fields.BeginObject(ANSITEXTVIEW("Y"));
+		Fields.AddInteger(ANSITEXTVIEW("$value"), 2);
+		Fields.AddString(ANSITEXTVIEW("$text"), ANSITEXTVIEW("0002"));
+		Fields.EndObject();
+		Fields.BeginObject(ANSITEXTVIEW("Z"));
+		Fields.AddInteger(ANSITEXTVIEW("$value"), 3);
+		Fields.AddString(ANSITEXTVIEW("$text"), ANSITEXTVIEW("0003"));
+		Fields.EndObject();
+		Fields.EndObject();
+
+		Fields.BeginArray(ANSITEXTVIEW("ArrayWithNestedText"));
+		Fields.BeginObject();
+		Fields.AddInteger(ANSITEXTVIEW("$value"), 1);
+		Fields.AddString(ANSITEXTVIEW("$text"), ANSITEXTVIEW("0001"));
+		Fields.EndObject();
+		Fields.BeginObject();
+		Fields.AddInteger(ANSITEXTVIEW("$value"), 2);
+		Fields.AddString(ANSITEXTVIEW("$text"), ANSITEXTVIEW("0002"));
+		Fields.EndObject();
+		Fields.BeginObject();
+		Fields.AddInteger(ANSITEXTVIEW("$value"), 3);
+		Fields.AddString(ANSITEXTVIEW("$text"), ANSITEXTVIEW("0003"));
+		Fields.EndObject();
+		Fields.EndArray();
+
+		return Fields.Save();
+	};
+
+	const FCbFieldIterator Fields = MakeFields();
+
+	const auto Test = [&Fields](const TCHAR* Format, FUtf8StringView Expected, const FLogTemplateOptions& Options = {})
+	{
+		FLogTemplate* Template = CreateLogTemplate(Format, Options);
+		TUtf8StringBuilder<1024> Message;
+		FormatLogTo(Message, Template, Fields);
+		DestroyLogTemplate(Template);
+		// TODO: Requires a fix for CaptureExpressionsAndValues in LowLevelTestAdapter.h
+		//CAPTURE(Format, Expected, Message);
+		CHECK(Message.ToView().Equals(Expected));
+	};
+
+	// Test each of the types that can be represented by compact binary.
+	Test(TEXT("{Object}"), UTF8TEXTVIEW(R"({"X": 1, "Y": 2, "Z": 3})"));
+	Test(TEXT("{Array}"), UTF8TEXTVIEW(R"([1, 2, 3])"));
+	Test(TEXT("{Null}"), UTF8TEXTVIEW(R"(null)"));
+	Test(TEXT("{Binary}"), UTF8TEXTVIEW(R"("AQIDBAUGBwgJCgsMDQ4PEA==")"));
+	Test(TEXT("{String}"), UTF8TEXTVIEW(R"("Quote" with 4 words.)"));
+	Test(TEXT("{IntegerNegative}"), UTF8TEXTVIEW(R"(-64)"));
+	Test(TEXT("{IntegerPositive}"), UTF8TEXTVIEW(R"(63)"));
+	Test(TEXT("{Float}"), UTF8TEXTVIEW(R"(128.25)"));
+	Test(TEXT("{Double}"), UTF8TEXTVIEW(R"(123.456)"));
+	Test(TEXT("{False}"), UTF8TEXTVIEW(R"(false)"));
+	Test(TEXT("{True}"), UTF8TEXTVIEW(R"(true)"));
+	Test(TEXT("{ObjectAttachment}"), UTF8TEXTVIEW(R"(cb42395cfe025324d80c31c88746d3392f330e58)"));
+	Test(TEXT("{BinaryAttachment}"), UTF8TEXTVIEW(R"(0ba7b01905a760046bacb86f092f291924c1f24a)"));
+	Test(TEXT("{Hash}"), UTF8TEXTVIEW(R"(700b0783bebf169c6d473141e82dd88c67f31ce2)"));
+	Test(TEXT("{Uuid}"), UTF8TEXTVIEW(R"(b1718295-2b3d-3379-9948-ef46ee9de6c3)"));
+	Test(TEXT("{DateTime}"), UTF8TEXTVIEW(R"(2025-04-10T11:15:30.123Z)"));
+	Test(TEXT("{TimeSpan}"), UTF8TEXTVIEW(R"(+10.11:15:30.123456700)"));
+	Test(TEXT("{ObjectId}"), UTF8TEXTVIEW(R"(354fa969c46eb24d1ea03026)"));
+	Test(TEXT("{CustomById}"), UTF8TEXTVIEW(R"({"Id":128,"Data":"ERITFBUWFxg="})"));
+	Test(TEXT("{CustomByName}"), UTF8TEXTVIEW(R"({"Name":"Custom","Data":"GRobHA=="})"));
+
+	// Test an object with a $text field.
+	Test(TEXT("{ObjectText}"), UTF8TEXTVIEW(R"(X=1;Y=2;Z=3)"));
+
+	// Test an object and an array with nested $text fields.
+	Test(TEXT("{ObjectWithNestedText}"), UTF8TEXTVIEW(R"({"X": 0001, "Y": 0002, "Z": 0003})"));
+	Test(TEXT("{ArrayWithNestedText}"), UTF8TEXTVIEW(R"([0001, 0002, 0003])"));
+
+	// Test accessing sub-object fields.
+	Test(TEXT("{{{Object/X}, {Object/Y}, {Object/Z}}}"), UTF8TEXTVIEW("{1, 2, 3}"), {.bAllowSubObjectReferences = true});
+	Test(TEXT("{ObjectText/X}"), UTF8TEXTVIEW("1"), {.bAllowSubObjectReferences = true});
+	Test(TEXT("{ObjectWithNestedText/X}"), UTF8TEXTVIEW("0001"), {.bAllowSubObjectReferences = true});
+}
+
+} // UE
+
+#endif // WITH_TESTS
