[RaceDetector]
- Use mimalloc by default whenever instrumentation is active

diff --git a/Engine/Source/Runtime/Core/Private/Windows/WindowsPlatformMemory.cpp b/Engine/Source/Runtime/Core/Private/Windows/WindowsPlatformMemory.cpp
--- a/Engine/Source/Runtime/Core/Private/Windows/WindowsPlatformMemory.cpp
+++ b/Engine/Source/Runtime/Core/Private/Windows/WindowsPlatformMemory.cpp
@@ -156,7 +156,7 @@ FMalloc* FWindowsPlatformMemory::BaseAllocator()
 	// Mimalloc is now the default allocator for editor and programs because it has shown
 	// both great performance and as much as half the memory usage of TBB after
 	// heavy editor workloads. See CL 15887498 description for benchmarks.
-	else if ((WITH_EDITORONLY_DATA || IS_PROGRAM) && MIMALLOC_ENABLED) //-V517
+	else if ((WITH_EDITORONLY_DATA || IS_PROGRAM || USING_INSTRUMENTATION) && MIMALLOC_ENABLED) //-V517
 	{
 		AllocatorToUse = EMemoryAllocatorToUse::Mimalloc;
 	}
