[MallocStomp]
- Make stomp malloc thread-safe for improved scalability

diff --git a/Engine/Source/Runtime/Core/Private/HAL/MallocStomp.cpp b/Engine/Source/Runtime/Core/Private/HAL/MallocStomp.cpp
--- a/Engine/Source/Runtime/Core/Private/HAL/MallocStomp.cpp
+++ b/Engine/Source/Runtime/Core/Private/HAL/MallocStomp.cpp
@@ -20,6 +20,8 @@
 // Virtual memory limit for a process on Win10 is 128 TB, which means we can afford to keep virtual memory reserved for a very long time.
 // Running Infiltrator demo consumes ~700MB of virtual address space per second.
 #define MALLOC_STOMP_KEEP_VIRTUAL_MEMORY 1
+thread_local UPTRINT VirtualAddressCursor = 0;
+thread_local SIZE_T VirtualAddressMax = 0;
 #else
 #define MALLOC_STOMP_KEEP_VIRTUAL_MEMORY 0
 #endif
@@ -125,7 +127,7 @@ void* FMallocStomp::TryMalloc(SIZE_T Size, uint32 Alignment)
 		VirtualAddressMax = VirtualAddressCursor + ReserveSize;
 	}
 
-	// No atomics or locks required here, as Malloc is externally synchronized (as indicated by FMallocStomp::IsInternallyThreadSafe()).
+	// No atomics as this is a thread_local
 	VirtualAddressCursor += TotalAllocationSize;
 
 #else
diff --git a/Engine/Source/Runtime/Core/Public/HAL/MallocStomp.h b/Engine/Source/Runtime/Core/Public/HAL/MallocStomp.h
--- a/Engine/Source/Runtime/Core/Public/HAL/MallocStomp.h
+++ b/Engine/Source/Runtime/Core/Public/HAL/MallocStomp.h
@@ -34,8 +34,6 @@ class FMallocStomp final : public FMalloc
 	/** If it is set to true, instead of focusing on overruns the allocator will focus on underruns. */
 	const bool bUseUnderrunMode;
 
-	UPTRINT VirtualAddressCursor = 0;
-	SIZE_T VirtualAddressMax = 0;
 	static constexpr SIZE_T VirtualAddressBlockSize = 1 * 1024 * 1024 * 1024; // 1 GB blocks
 
 public:
@@ -122,8 +120,11 @@ public:
 
 	virtual bool IsInternallyThreadSafe() const override
 	{
-		// Stomp allocator is NOT thread-safe and must be externally-synchronized.
-		return false;
+		// Unix and Mac uses mmap which is thread-safe
+		// Windows uses virtual alloc which is thread-safe
+		// Other platforms using FPlatformMemory::BinnedAllocFromOS are all assumed to have thread-safe implementation already
+		// VirtualAddressCursor and VirtualAddressMax are thread-locals so each thread has its virtual reserved block
+		return true;
 	}
 };
 
