[AutoRTFM] Ensure that MallocBinned2/3 are never used to allocate memory in the closed.

This marks the low-level Malloc, Free and Realloc binned methods as disallowed when in AutoRTFM. Callers should use FMemory::Malloc and friends to allocate memory. Previously, if a caller were to access GMalloc->Malloc from the closed and then abort the transaction, this could potentially have corrupted the malloc heap. In practice, these APIs would generally take a lock and crash immediately if accessed from the closed anyway, but in some cases, a small allocation can be fulfilled directly from TLS pools without ever taking a lock.

An earlier version of this CL was more extreme and marked every single method in these classes as NOAUTORTFM. This however is discouraged, as NOAUTORTFM functions currently cannot be inlined. This CL should still protect us without impacting performance. These functions are virtual anyway, so inlining is very unlikely in practice (even though they were mysteriously marked FORCEINLINE before now).

diff --git a/Engine/Source/Runtime/Core/Public/HAL/MallocBinned2.h b/Engine/Source/Runtime/Core/Public/HAL/MallocBinned2.h
--- a/Engine/Source/Runtime/Core/Public/HAL/MallocBinned2.h
+++ b/Engine/Source/Runtime/Core/Public/HAL/MallocBinned2.h
@@ -3,6 +3,7 @@
 #pragma once
 
 #include "CoreTypes.h"
+#include "AutoRTFM.h"
 #include "HAL/Allocators/CachedOSPageAllocator.h"
 #include "HAL/Allocators/CachedOSVeryLargePageAllocator.h"
 #include "HAL/Allocators/PooledVirtualMemoryAllocator.h"
@@ -262,6 +263,7 @@ public:
 	// FMalloc interface.
 	CORE_API virtual bool IsInternallyThreadSafe() const override;
 
+	UE_AUTORTFM_NOAUTORTFM
 	FORCEINLINE virtual void* Malloc(SIZE_T Size, uint32 Alignment) override
 	{
 #if UE_MB2_ALLOCATOR_STATS_VALIDATION
@@ -289,6 +291,7 @@ public:
 #endif
 	}
 
+protected:
 	FORCEINLINE void* MallocInline(SIZE_T Size, uint32 Alignment)
 	{
 		// Only allocate from the small pools if the size is small enough and the alignment isn't crazy large.
@@ -334,6 +337,8 @@ public:
 		return MallocSelect(Size, Alignment, UseSmallAlloc(Size, Alignment));
 	}
 
+public:
+	UE_AUTORTFM_NOAUTORTFM
 	FORCEINLINE virtual void* Realloc(void* Ptr, SIZE_T NewSize, uint32 Alignment) override
 	{
 #if UE_MB2_ALLOCATOR_STATS_VALIDATION
@@ -372,6 +377,7 @@ public:
 #endif
 	}
 
+protected:
 	FORCEINLINE void* ReallocInline(void* Ptr, SIZE_T NewSize, uint32 Alignment) 
 	{
 #if UE_USE_VERYLARGEPAGEALLOCATOR && UE_MB2_BOOKKEEPING_AT_THE_END_OF_LARGEBLOCK
@@ -438,6 +444,8 @@ public:
 		return Result;
 	}
 
+public:
+	UE_AUTORTFM_NOAUTORTFM
 	FORCEINLINE virtual void Free(void* Ptr) override
 	{
 #if UE_MB2_ALLOCATOR_STATS_VALIDATION
@@ -464,6 +472,7 @@ public:
 #endif
 	}
 
+protected:
 	FORCEINLINE void FreeInline(void* Ptr)
 	{
 		if (!IsOSAllocation(Ptr))
@@ -484,6 +493,7 @@ public:
 		FreeExternal(Ptr);
 	}
 
+public:
 	FORCEINLINE bool GetSmallAllocationSize(void* Ptr, SIZE_T& SizeOut) const
 	{
 		if (!IsOSAllocation(Ptr))
diff --git a/Engine/Source/Runtime/Core/Public/HAL/MallocBinned3.h b/Engine/Source/Runtime/Core/Public/HAL/MallocBinned3.h
--- a/Engine/Source/Runtime/Core/Public/HAL/MallocBinned3.h
+++ b/Engine/Source/Runtime/Core/Public/HAL/MallocBinned3.h
@@ -5,6 +5,7 @@
 #include "CoreTypes.h"
 
 #if PLATFORM_HAS_FPlatformVirtualMemoryBlock
+#include "AutoRTFM.h"
 #include "HAL/Allocators/CachedOSPageAllocator.h"
 #include "HAL/Allocators/PooledVirtualMemoryAllocator.h"
 #include "HAL/CriticalSection.h"
@@ -297,6 +298,8 @@ public:
 
 	// FMalloc interface.
 	CORE_API virtual bool IsInternallyThreadSafe() const override;
+
+	UE_AUTORTFM_NOAUTORTFM
 	FORCEINLINE virtual void* Malloc(SIZE_T Size, uint32 Alignment) override
 	{
 		void* Result = nullptr;
@@ -328,6 +331,7 @@ public:
 		return Result;
 	}
 
+	UE_AUTORTFM_NOAUTORTFM
 	FORCEINLINE virtual void* Realloc(void* Ptr, SIZE_T NewSize, uint32 Alignment) override
 	{
 		if (NewSize <= UE_MB3_MAX_SMALL_POOL_SIZE && Alignment <= UE_MBC_MIN_SMALL_POOL_ALIGNMENT) // one branch, not two
@@ -391,6 +395,7 @@ public:
 		return Result;
 	}
 
+	UE_AUTORTFM_NOAUTORTFM
 	FORCEINLINE virtual void Free(void* Ptr) override
 	{
 		const uint64 PoolIndex = PoolIndexFromPtr(Ptr);
