WindowsPlatformFile: Update FAsyncBufferedFileReaderWindows::ReadAt, FFileHandleWindows::Read, FFileHandleWindows::ReadAt, FFileHandleWindows::Write to handle NumRead != requested BytesToRead, or NumWrite != requested BytesToWrite. This behavior can occur on Wine executions of Windows because the underlying linux kernel has a limit of 0x7ffff000 for read/writes.
#rnx

diff --git a/Engine/Source/Runtime/Core/Private/Windows/WindowsPlatformFile.cpp b/Engine/Source/Runtime/Core/Private/Windows/WindowsPlatformFile.cpp
--- a/Engine/Source/Runtime/Core/Private/Windows/WindowsPlatformFile.cpp
+++ b/Engine/Source/Runtime/Core/Private/Windows/WindowsPlatformFile.cpp
@@ -550,8 +550,21 @@ public:
 				}
 			}
 
-			BytesToRead -= BytesToRead32;
+			BytesToRead -= NumRead;
 			TotalNumRead += NumRead;
+			Destination += NumRead;
+
+			if (NumRead == 0)
+			{
+				// Reading less than requested bytes can occur due to an internal limit, but we always expect
+				// at least some bytes to be read. If 0 bytes were read repeatedly, we could deadlock. So return
+				// false when NumRead is 0 rather than continuing to try.
+				// Pick an error we can use to report to the caller why we failed. We arbitrarily use
+				// ERROR_CAN_NOT_COMPLETE, since it seems like one that ReadFile does not normally return.
+				SetLastError(ERROR_CAN_NOT_COMPLETE);
+				TRACE_PLATFORMFILE_END_READ(&LocalOverlappedIO, TotalNumRead);
+				return false;
+			}
 		} while (BytesToRead > 0);
 
 		TRACE_PLATFORMFILE_END_READ(&LocalOverlappedIO, TotalNumRead);
@@ -759,20 +772,24 @@ public:
 				}
 			}
 
-			BytesToRead -= BytesToRead32;
-			Destination += BytesToRead32;
+			BytesToRead -= NumRead;
 			TotalNumRead += NumRead;
+			Destination += NumRead;
 			// Update where we are in the file
 			FilePos += NumRead;
 			UE::WindowsPlatformFile::Private::UpdateOverlappedPos(OverlappedIO, FilePos);
-			
-			// Early out as a failure case if we did not read all of the bytes that we expected to read
-			if (BytesToRead32 != NumRead)
+
+			if (NumRead == 0)
 			{
+				// Reading less than requested bytes can occur due to an internal limit, but we always expect
+				// at least some bytes to be read. If 0 bytes were read repeatedly, we could deadlock. So return
+				// false when NumRead is 0 rather than continuing to try.
+				// Pick an error we can use to report to the caller why we failed. We arbitrarily use
+				// ERROR_CAN_NOT_COMPLETE, since it seems like one that ReadFile does not normally return.
+				SetLastError(ERROR_CAN_NOT_COMPLETE);
 				TRACE_PLATFORMFILE_END_READ(&OverlappedIO, TotalNumRead);
 				return false;
-			}	
-					
+			}
 		} while (BytesToRead > 0);
 		TRACE_PLATFORMFILE_END_READ(&OverlappedIO, TotalNumRead);
 		return true;
@@ -809,16 +826,21 @@ public:
 				return false;
 			}
 
-			BytesToRead -= BytesToRead32;
+			BytesToRead -= NumRead;
 			TotalNumRead += NumRead;
+			Destination += NumRead;
 
-			// Early out as a failure case if we did not read all of the bytes that we expected to read
-			if (BytesToRead32 != NumRead)
+			if (NumRead == 0)
 			{
+				// Reading less than requested bytes can occur due to an internal limit, but we always expect
+				// at least some bytes to be read. If 0 bytes were read repeatedly, we could deadlock. So return
+				// false when NumRead is 0 rather than continuing to try.
+				// Pick an error we can use to report to the caller why we failed. We arbitrarily use
+				// ERROR_CAN_NOT_COMPLETE, since it seems like one that ReadFile does not normally return.
+				SetLastError(ERROR_CAN_NOT_COMPLETE);
 				TRACE_PLATFORMFILE_END_READ(&OverlappedIO, TotalNumRead);
 				return false;
 			}
-
 		} while (BytesToRead > 0);
 
 		TRACE_PLATFORMFILE_END_READ(&LocalOverlappedIO, TotalNumRead);
@@ -854,25 +876,26 @@ public:
 					return false;
 				}
 			}
-	
-			BytesToWrite -= BytesToWrite32;
-			Source += BytesToWrite32;
+
+			BytesToWrite -= NumWritten;
+			Source += NumWritten;
 			TotalNumWritten += NumWritten;
 			// Update where we are in the file
 			FilePos += NumWritten;
 			UE::WindowsPlatformFile::Private::UpdateOverlappedPos(OverlappedIO, FilePos);
 			FileSize = FMath::Max(FilePos, FileSize);
-			
-			// Early out as a failure case if we didn't write all of the data we expected
-			if (BytesToWrite32 != NumWritten)
+
+			if (NumWritten == 0)
 			{
+				// Writing less than requested bytes can occur due to an internal limit, but we always expect
+				// at least some bytes to be written. If 0 bytes were written repeatedly, we could deadlock. So return
+				// false when NumWritten is 0 rather than continuing to try.
 				// Pick an error we can use to report to the caller why we failed. We arbitrarily use
 				// ERROR_CAN_NOT_COMPLETE, since it seems like one that WriteFile does not normally return.
 				SetLastError(ERROR_CAN_NOT_COMPLETE);
 				TRACE_PLATFORMFILE_END_WRITE(this, TotalNumWritten);
 				return false;
 			}
-			
 		} while (BytesToWrite > 0);
 
 		TRACE_PLATFORMFILE_END_WRITE(this, TotalNumWritten);
