[RaceDetector]
- Fix 256 context limit reached by hard limiting the number of core since the PSO compilation thread pool does not abide by NumberOfWorkerThreadsToSpawn

diff --git a/Engine/Source/Runtime/Core/Private/GenericPlatform/GenericPlatformMisc.cpp b/Engine/Source/Runtime/Core/Private/GenericPlatform/GenericPlatformMisc.cpp
--- a/Engine/Source/Runtime/Core/Private/GenericPlatform/GenericPlatformMisc.cpp
+++ b/Engine/Source/Runtime/Core/Private/GenericPlatform/GenericPlatformMisc.cpp
@@ -1669,6 +1669,13 @@ void FGenericPlatformMisc::GetConfiguredCoreLimits(int32 PlatformNumPhysicalCore
 		LogicalCoreLimit = LogicalCoreLimit == 0 ? LegacyCoreLimit : FMath::Min(LogicalCoreLimit, LegacyCoreLimit);
 	}
 
+#if USING_INSTRUMENTATION
+	// Limit the number of threads as we have a limit of 256 context id in the shadow memory structure.
+	constexpr int32 InstrumentationCoreLimit = 32;
+	PhysicalCoreLimit = PhysicalCoreLimit == 0 ? InstrumentationCoreLimit : FMath::Min(PhysicalCoreLimit, InstrumentationCoreLimit);
+	LogicalCoreLimit  = LogicalCoreLimit  == 0 ? InstrumentationCoreLimit : FMath::Min(LogicalCoreLimit,  InstrumentationCoreLimit);
+#endif
+
 	bOutFullyInitialized = true;
 	OutPhysicalCoreLimit = PhysicalCoreLimit;
 	OutLogicalCoreLimit = LogicalCoreLimit;
diff --git a/Engine/Source/Runtime/Core/Private/Windows/WindowsPlatformMisc.cpp b/Engine/Source/Runtime/Core/Private/Windows/WindowsPlatformMisc.cpp
--- a/Engine/Source/Runtime/Core/Private/Windows/WindowsPlatformMisc.cpp
+++ b/Engine/Source/Runtime/Core/Private/Windows/WindowsPlatformMisc.cpp
@@ -2595,13 +2595,7 @@ int32 FWindowsPlatformMisc::NumberOfWorkerThreadsToSpawn()
 {	
 	static int32 MaxServerWorkerThreads = 4;
 
-// Limit the number of threads as we have a limit of 256 context id in the shadow memory structure.
-#if USING_INSTRUMENTATION
-	int32 MaxWorkerThreads = 32;
-#else
 	int32 MaxWorkerThreads = INT32_MAX;
-#endif
-
 	int32 NumberOfCores = FWindowsPlatformMisc::NumberOfCores();
 	int32 NumberOfCoresIncludingHyperthreads = FWindowsPlatformMisc::NumberOfCoresIncludingHyperthreads();
 	int32 NumberOfThreads = 0;
