[Insights] Timing Regions: Added an optional Category parameter when tracing regions
- Added optional Category parameter to Regions Macros and `FMiscTrace::OutputBeginRegion`

diff --git a/Engine/Source/Runtime/Core/Private/ProfilingDebugging/MiscTrace.cpp b/Engine/Source/Runtime/Core/Private/ProfilingDebugging/MiscTrace.cpp
--- a/Engine/Source/Runtime/Core/Private/ProfilingDebugging/MiscTrace.cpp
+++ b/Engine/Source/Runtime/Core/Private/ProfilingDebugging/MiscTrace.cpp
@@ -30,11 +30,13 @@ UE_TRACE_EVENT_END()
 UE_TRACE_EVENT_BEGIN(Misc, RegionBegin)
 	UE_TRACE_EVENT_FIELD(uint64, Cycle)
 	UE_TRACE_EVENT_FIELD(UE::Trace::WideString, RegionName)
+	UE_TRACE_EVENT_FIELD(UE::Trace::WideString, Category)
 UE_TRACE_EVENT_END()
 
 UE_TRACE_EVENT_BEGIN(Misc, RegionBeginWithId)
 	UE_TRACE_EVENT_FIELD(uint64, CycleAndId)
 	UE_TRACE_EVENT_FIELD(UE::Trace::WideString, RegionName)
+	UE_TRACE_EVENT_FIELD(UE::Trace::WideString, Category)
 UE_TRACE_EVENT_END()
 
 UE_TRACE_EVENT_BEGIN(Misc, RegionEnd)
@@ -87,19 +89,21 @@ void FMiscTrace::OutputBookmarkSpec(const void* BookmarkPoint, const ANSICHAR* F
 		<< BookmarkSpec.FileName(File, FileNameLen);
 }
 
-void FMiscTrace::OutputBeginRegion(const TCHAR* RegionName)
+void FMiscTrace::OutputBeginRegion(const TCHAR* RegionName, const TCHAR* Category)
 {
 	UE_TRACE_LOG(Misc, RegionBegin, RegionChannel)
 		<< RegionBegin.Cycle(FPlatformTime::Cycles64())
-		<< RegionBegin.RegionName(RegionName);
+		<< RegionBegin.RegionName(RegionName)
+		<< RegionBegin.Category(Category == nullptr ? TEXT("") : Category);
 }
 
-uint64 FMiscTrace::OutputBeginRegionWithId(const TCHAR* RegionName)
+uint64 FMiscTrace::OutputBeginRegionWithId(const TCHAR* RegionName, const TCHAR* Category)
 {
 	const uint64 CycleAndId = FPlatformTime::Cycles64();
 	UE_TRACE_LOG(Misc, RegionBeginWithId, RegionChannel)
 		<< RegionBeginWithId.CycleAndId(CycleAndId)
-		<< RegionBeginWithId.RegionName(RegionName);
+		<< RegionBeginWithId.RegionName(RegionName)
+		<< RegionBeginWithId.Category(Category == nullptr ? TEXT("") : Category);
 	return CycleAndId;
 }
 
diff --git a/Engine/Source/Runtime/Core/Private/ProfilingDebugging/TraceAuxiliary.cpp b/Engine/Source/Runtime/Core/Private/ProfilingDebugging/TraceAuxiliary.cpp
--- a/Engine/Source/Runtime/Core/Private/ProfilingDebugging/TraceAuxiliary.cpp
+++ b/Engine/Source/Runtime/Core/Private/ProfilingDebugging/TraceAuxiliary.cpp
@@ -1364,7 +1364,7 @@ static void TraceRegionBegin(const TArray<FString>& Args)
 {
 	if (Args.Num() > 0)
 	{
-		TRACE_BEGIN_REGION(*FString::Join(Args, TEXT(" ")));
+		TRACE_BEGIN_REGION(*FString::Join(Args, TEXT(" ")), TEXT("ConsoleCommandRegion"));
 	}
 }
 
diff --git a/Engine/Source/Runtime/Core/Public/ProfilingDebugging/MiscTrace.h b/Engine/Source/Runtime/Core/Public/ProfilingDebugging/MiscTrace.h
--- a/Engine/Source/Runtime/Core/Public/ProfilingDebugging/MiscTrace.h
+++ b/Engine/Source/Runtime/Core/Public/ProfilingDebugging/MiscTrace.h
@@ -77,8 +77,8 @@ struct FMiscTrace
 		}
 	}
 
-	CORE_API static void OutputBeginRegion(const TCHAR* RegionName);
-	[[nodiscard]] CORE_API static uint64 OutputBeginRegionWithId(const TCHAR* RegionName);
+	CORE_API static void OutputBeginRegion(const TCHAR* RegionName, const TCHAR* Category = nullptr);
+	[[nodiscard]] CORE_API static uint64 OutputBeginRegionWithId(const TCHAR* RegionName, const TCHAR* Category = nullptr);
 	CORE_API static void OutputEndRegion(const TCHAR* RegionName);
 	CORE_API static void OutputEndRegionWithId(uint64 RegionId);
 
@@ -139,11 +139,32 @@ if (UE_TRACE_CHANNELEXPR_IS_ENABLED(BookmarkChannel)) \
 	FMiscTrace::OutputBookmarkCycles(CallstackTrace_GetCurrentId(), Cycles, &PREPROCESSOR_JOIN(__BookmarkPoint, __LINE__), ##__VA_ARGS__); \
 }
 
-#define TRACE_BEGIN_REGION(RegionName) \
-	FMiscTrace::OutputBeginRegion(RegionName);
-
-#define TRACE_BEGIN_REGION_WITH_ID(RegionName) \
-	FMiscTrace::OutputBeginRegionWithId(RegionName);
+/**
+ * Start a Timing Region, identified by its name.
+ * Prefer the _WITH_ID variant of this macro since overlapping Regions with the same name cannot be uniquely identified
+ * Example Usage:
+ * @code
+ * TRACE_BEGIN_REGION(TEXT("MyRegion"))
+ * [...]
+ * TRACE_END_REGION(TEXT("MyRegion"))
+ * @endcode 
+ * 
+ */
+#define TRACE_BEGIN_REGION(RegionName, ...) \
+FMiscTrace::OutputBeginRegion(RegionName __VA_OPT__(,) __VA_ARGS__);
+
+/**
+ * Start a Timing Region, identified by an Id
+ * Example Usage:
+ * @code
+ * uint64 OpenRegionId = TRACE_BEGIN_REGION_WITH_ID(TEXT("MyRegion"), TEXT("MyCategory"))
+ * [...]
+ * TRACE_END_REGION_WITH_ID(OpenRegionId)
+ * @endcode 
+ * 
+ */
+#define TRACE_BEGIN_REGION_WITH_ID(RegionName, ...) \
+FMiscTrace::OutputBeginRegionWithId(RegionName __VA_OPT__(,) __VA_ARGS__);
 
 #define TRACE_END_REGION(RegionName) \
 	FMiscTrace::OutputEndRegion(RegionName);
