[Core]
* Moved dll storage to methods/variables instead of on types. This is needed to get merged modules working

diff --git a/Engine/Source/Runtime/Core/Public/Async/TaskTrace.h b/Engine/Source/Runtime/Core/Public/Async/TaskTrace.h
--- a/Engine/Source/Runtime/Core/Public/Async/TaskTrace.h
+++ b/Engine/Source/Runtime/Core/Public/Async/TaskTrace.h
@@ -24,6 +24,12 @@ namespace ENamedThreads
 	enum Type : int32;
 }
 
+#if UE_TASK_TRACE_ENABLED
+#define TASK_CORE_API CORE_API
+#else
+#define TASK_CORE_API
+#endif
+
 namespace TaskTrace
 {
 	UE_TRACE_CHANNEL_EXTERN(TaskChannel, CORE_API);
@@ -48,15 +54,15 @@ namespace TaskTrace
 
 	struct FWaitingScope
 	{
-		CORE_API explicit FWaitingScope(const TArray<FId>& Tasks); // waiting for given tasks completion
-		CORE_API explicit FWaitingScope(FId TaskId);
-		CORE_API ~FWaitingScope();
+		TASK_CORE_API explicit FWaitingScope(const TArray<FId>& Tasks); // waiting for given tasks completion
+		TASK_CORE_API explicit FWaitingScope(FId TaskId);
+		TASK_CORE_API ~FWaitingScope();
 	};
 
 	struct FTaskTimingEventScope
 	{
-		CORE_API FTaskTimingEventScope(TaskTrace::FId InTaskId);
-		CORE_API ~FTaskTimingEventScope();
+		TASK_CORE_API FTaskTimingEventScope(TaskTrace::FId InTaskId);
+		TASK_CORE_API ~FTaskTimingEventScope();
 
 	private:
 		bool bIsActive = false;
diff --git a/Engine/Source/Runtime/Core/Public/Experimental/UnifiedError/UnifiedError.h b/Engine/Source/Runtime/Core/Public/Experimental/UnifiedError/UnifiedError.h
--- a/Engine/Source/Runtime/Core/Public/Experimental/UnifiedError/UnifiedError.h
+++ b/Engine/Source/Runtime/Core/Public/Experimental/UnifiedError/UnifiedError.h
@@ -140,13 +140,13 @@ namespace UE::UnifiedError
 		/// GetInnerErrorDetails; Exposes inner error details to FError, if this ErrorDetails allows inner details
 		/// </summary>
 		/// <returns></returns>
-		CORE_API virtual TRefCountPtr<const IErrorDetails> GetInnerErrorDetails() const { return nullptr; }
+		virtual TRefCountPtr<const IErrorDetails> GetInnerErrorDetails() const { return nullptr; }
 
 		/// <summary>
 		/// SetInnerErrorDetails; Exposes inner error details to FError, if this ErrorDetails allows inner details
 		/// </summary>
 		/// <returns></returns>
-		CORE_API virtual void SetInnerErrorDetails(TRefCountPtr<const IErrorDetails> ErrorDetails) { checkf(false, TEXT("SetInnerErrorDetails not implemented!")); }
+		virtual void SetInnerErrorDetails(TRefCountPtr<const IErrorDetails> ErrorDetails) { checkf(false, TEXT("SetInnerErrorDetails not implemented!")); }
 
 		/// <summary>
 		/// GetErrorDetialsTypeId; Simple type information for error details, generated using hash of details name
@@ -157,7 +157,7 @@ namespace UE::UnifiedError
 
 
 		CORE_API virtual const FAnsiStringView GetErrorDetailsTypeName() const = 0;
-		CORE_API virtual void SerializeForLog(FCbWriter& Writer) const {}
+		virtual void SerializeForLog(FCbWriter& Writer) const {}
 
 	};
 
@@ -214,7 +214,7 @@ namespace UE::UnifiedError
 
 		// IErrorDetails functions
 		CORE_API virtual void SetInnerErrorDetails(TRefCountPtr<const IErrorDetails> InInnerErrorDetails);
-		CORE_API virtual TRefCountPtr<const IErrorDetails> GetInnerErrorDetails() const { return InnerErrorDetails; }
+		virtual TRefCountPtr<const IErrorDetails> GetInnerErrorDetails() const { return InnerErrorDetails; }
 
 		/// <summary>
 		/// GetErrorFormatString; Pass through to the InnerErrorDetails.
@@ -373,14 +373,14 @@ namespace UE::UnifiedError
 		int32 ModuleId;
 		TRefCountPtr<const IErrorDetails> ErrorDetails;
 	public:
-		CORE_API FError(int32 InModuleId, int32 InErrorCode, const IErrorDetails* InErrorDetails)
+		FError(int32 InModuleId, int32 InErrorCode, const IErrorDetails* InErrorDetails)
 		{
 			ErrorCode = InErrorCode;
 			ModuleId = InModuleId;
 			ErrorDetails = InErrorDetails;
 		}
 	public:
-		CORE_API FError(FError&& InError)
+		FError(FError&& InError)
 		{
 			ErrorCode = InError.ErrorCode;
 			InError.ErrorCode = 0;
@@ -393,7 +393,7 @@ namespace UE::UnifiedError
 			}
 		}
 
-		CORE_API FError(const FError& InError)
+		FError(const FError& InError)
 		{
 			ErrorCode = InError.ErrorCode;
 			ModuleId = InError.ModuleId;
@@ -475,22 +475,20 @@ namespace UE::UnifiedError
 			ErrorDetails = nullptr;
 		}
 	private:
-		friend CORE_API bool operator==(const FError& Error, const FError& OtherError);
+		friend bool operator==(const FError& Error, const FError& OtherError)
+		{
+			if ((Error.ModuleId == OtherError.ModuleId) &&
+				(Error.ErrorCode == OtherError.ErrorCode))
+			{
+				return true;
+			}
+			return false;
+		}
 
 		TRefCountPtr<const IErrorDetails> GetInnerMostErrorDetails() const;
 	};
 
 
-	CORE_API inline bool operator==(const FError& Error, const FError& OtherError)
-	{
-		if ((Error.ModuleId == OtherError.ModuleId) &&
-			(Error.ErrorCode == OtherError.ErrorCode))
-		{
-			return true;
-		}
-		return false;
-	}
-
 	class FErrorRegistry
 	{
 	public:
diff --git a/Engine/Source/Runtime/Core/Public/GenericPlatform/GenericPlatformChunkInstall.h b/Engine/Source/Runtime/Core/Public/GenericPlatform/GenericPlatformChunkInstall.h
--- a/Engine/Source/Runtime/Core/Public/GenericPlatform/GenericPlatformChunkInstall.h
+++ b/Engine/Source/Runtime/Core/Public/GenericPlatform/GenericPlatformChunkInstall.h
@@ -386,10 +386,10 @@ PRAGMA_DISABLE_DEPRECATION_WARNINGS
 /**
  * Generic implementation of chunk based install
  */
-class CORE_API FGenericPlatformChunkInstall : public IPlatformChunkInstall
+class FGenericPlatformChunkInstall : public IPlatformChunkInstall
 {
 public:
-	virtual void AsyncInit( TFunction<void(bool/*bSuccess*/)> OnInitComplete ) override;
+	CORE_API virtual void AsyncInit( TFunction<void(bool/*bSuccess*/)> OnInitComplete ) override;
 
 	virtual bool IsAvailable() const
 	{
@@ -563,8 +563,8 @@ public:
 
 protected:
 
-	void DoNamedChunkCompleteCallbacks( const FName NamedChunk, EChunkLocation::Type Location, bool bHasSucceeded ) const;
-	void DoNamedChunkCompleteCallbacks( const TArrayView<const FName>& NamedChunks, EChunkLocation::Type Location, bool bHasSucceeded ) const;
+	CORE_API void DoNamedChunkCompleteCallbacks( const FName NamedChunk, EChunkLocation::Type Location, bool bHasSucceeded ) const;
+	CORE_API void DoNamedChunkCompleteCallbacks( const TArrayView<const FName>& NamedChunks, EChunkLocation::Type Location, bool bHasSucceeded ) const;
 
 	/** Delegates called when installation succeeds or fails */
 	FPlatformChunkInstallMultiDelegate InstallDelegate;
diff --git a/Engine/Source/Runtime/Core/Public/Misc/AutomationTest.h b/Engine/Source/Runtime/Core/Public/Misc/AutomationTest.h
--- a/Engine/Source/Runtime/Core/Public/Misc/AutomationTest.h
+++ b/Engine/Source/Runtime/Core/Public/Misc/AutomationTest.h
@@ -1960,7 +1960,7 @@ public:
 	CORE_API bool TestEqual(const TCHAR* What, const FLinearColor Actual, const FLinearColor Expected);
 	CORE_API bool TestEqual(const TCHAR* What, const TCHAR* Actual, const TCHAR* Expected);
 	CORE_API bool TestEqual(const TCHAR* What, FStringView Actual, FStringView Expected);
-	CORE_API bool TestEqual(const TCHAR* What, const FString& Actual, const FString& Expected)
+	bool TestEqual(const TCHAR* What, const FString& Actual, const FString& Expected)
 	{
 		return TestEqual(What, *Actual, *Expected);
 	}
@@ -1969,7 +1969,7 @@ public:
 	CORE_API bool TestEqual(const TCHAR* What, const FName Actual, const FName Expected);
 	CORE_API bool TestNotEqual(const TCHAR* What, const TCHAR* Actual, const TCHAR* Expected);
 	CORE_API bool TestNotEqual(const TCHAR* What, FStringView Actual, FStringView Expected);
-	CORE_API bool TestNotEqual(const TCHAR* What, const FString& Actual, const FString& Expected)
+	bool TestNotEqual(const TCHAR* What, const FString& Actual, const FString& Expected)
 	{
 		return TestNotEqual(What, *Actual, *Expected);
 	}
@@ -1992,14 +1992,14 @@ public:
 
 	CORE_API bool TestEqualSensitive(const TCHAR* What, const TCHAR* Actual, const TCHAR* Expected);
 	CORE_API bool TestEqualSensitive(const TCHAR* What, FStringView Actual, FStringView Expected);
-	CORE_API bool TestEqualSensitive(const TCHAR* What, const FString& Actual, const FString& Expected)
+	bool TestEqualSensitive(const TCHAR* What, const FString& Actual, const FString& Expected)
 	{
 		return TestEqualSensitive(What, *Actual, *Expected);
 	}
 	CORE_API bool TestEqualSensitive(const TCHAR* What, FUtf8StringView Actual, FUtf8StringView Expected);
 	CORE_API bool TestNotEqualSensitive(const TCHAR* What, const TCHAR* Actual, const TCHAR* Expected);
 	CORE_API bool TestNotEqualSensitive(const TCHAR* What, FStringView Actual, FStringView Expected);
-	CORE_API bool TestNotEqualSensitive(const TCHAR* What, const FString& Actual, const FString& Expected)
+	bool TestNotEqualSensitive(const TCHAR* What, const FString& Actual, const FString& Expected)
 	{
 		return TestNotEqualSensitive(What, *Actual, *Expected);
 	}
@@ -3820,16 +3820,16 @@ class CommandName : public IAutomationLatentCommand \
 }
 
 #define DEFINE_EXPORTED_LATENT_AUTOMATION_COMMAND(EXPORT_API, CommandName)	\
-class EXPORT_API CommandName : public IAutomationLatentCommand \
+class CommandName : public IAutomationLatentCommand \
 	{ \
 	public: \
 	virtual ~CommandName() \
 		{} \
-		virtual bool Update() override; \
+		EXPORT_API virtual bool Update() override; \
 }
 
 #define DEFINE_EXPORTED_LATENT_AUTOMATION_COMMAND_ONE_PARAMETER(EXPORT_API, CommandName,ParamType,ParamName)	\
-class EXPORT_API CommandName : public IAutomationLatentCommand \
+class CommandName : public IAutomationLatentCommand \
 	{ \
 	public: \
 	CommandName(ParamType InputParam) \
@@ -3837,13 +3837,13 @@ class EXPORT_API CommandName : public IAutomationLatentCommand \
 		{} \
 		virtual ~CommandName() \
 		{} \
-		virtual bool Update() override; \
+		EXPORT_API virtual bool Update() override; \
 	private: \
 	ParamType ParamName; \
 }
 
 #define DEFINE_EXPORTED_LATENT_AUTOMATION_COMMAND_TWO_PARAMETER(EXPORT_API, CommandName,ParamType0,ParamName0,ParamType1,ParamName1)	\
-class EXPORT_API CommandName : public IAutomationLatentCommand \
+class CommandName : public IAutomationLatentCommand \
 	{ \
 	public: \
 	CommandName(ParamType0 InputParam0, ParamType1 InputParam1) \
@@ -3852,7 +3852,7 @@ class EXPORT_API CommandName : public IAutomationLatentCommand \
 		{} \
 		virtual ~CommandName() \
 		{} \
-		virtual bool Update() override; \
+		EXPORT_API virtual bool Update() override; \
 	private: \
 	ParamType0 ParamName0; \
 	ParamType1 ParamName1; \
@@ -3865,7 +3865,7 @@ class EXPORT_API CommandName : public IAutomationLatentCommand \
 	DEFINE_EXPORTED_LATENT_AUTOMATION_COMMAND_ONE_PARAMETER(ENGINE_API, CommandName, ParamType, ParamName)
 
 #define DEFINE_EXPORTED_LATENT_AUTOMATION_COMMAND_WITH_RETRIES(EXPORT_API,CommandName,RetryCount,WaitTimeBetweenRuns)	\
-class EXPORT_API CommandName : public IAutomationLatentCommandWithRetriesAndDelays \
+class CommandName : public IAutomationLatentCommandWithRetriesAndDelays \
 	{ \
 	public: \
 	CommandName(int32 InRetryCount, double InWaitTimeBetweenRuns) \
@@ -3873,12 +3873,12 @@ class EXPORT_API CommandName : public IAutomationLatentCommandWithRetriesAndDela
 		{} \
 		virtual ~CommandName() \
 		{} \
-		virtual bool Execute() override; \
+		EXPORT_API virtual bool Execute() override; \
 	private: \
 }
 
 #define DEFINE_EXPORTED_LATENT_AUTOMATION_COMMAND_WITH_RETRIES_ONE_PARAMETER(EXPORT_API,CommandName,RetryCount,WaitTimeBetweenRuns,ParamType,ParamName)	\
-class EXPORT_API CommandName : public IAutomationLatentCommandWithRetriesAndDelays \
+class CommandName : public IAutomationLatentCommandWithRetriesAndDelays \
 	{ \
 	public: \
 	CommandName(int32 InRetryCount, double InWaitTimeBetweenRuns, ParamType ParamName) \
@@ -3887,13 +3887,13 @@ class EXPORT_API CommandName : public IAutomationLatentCommandWithRetriesAndDela
 		{} \
 		virtual ~CommandName() \
 		{} \
-		virtual bool Execute() override; \
+		EXPORT_API virtual bool Execute() override; \
 	private: \
 	ParamType ParamName; \
 }
 
 #define DEFINE_EXPORTED_LATENT_AUTOMATION_COMMAND_WITH_RETRIES_TWO_PARAMETERS(EXPORT_API,CommandName,RetryCount,WaitTimeBetweenRuns,ParamType0,ParamName0,ParamType1,ParamName1)	\
-class EXPORT_API CommandName : public IAutomationLatentCommandWithRetriesAndDelays \
+class CommandName : public IAutomationLatentCommandWithRetriesAndDelays \
 	{ \
 	public: \
 	CommandName(int32 InRetryCount, double InWaitTimeBetweenRuns, ParamType0 ParamName0, ParamType1 ParamName1) \
@@ -3903,14 +3903,14 @@ class EXPORT_API CommandName : public IAutomationLatentCommandWithRetriesAndDela
 		{} \
 		virtual ~CommandName() \
 		{} \
-		virtual bool Execute() override; \
+		EXPORT_API virtual bool Execute() override; \
 	private: \
 	ParamType0 ParamName0; \
 	ParamType1 ParamName1; \
 }
 
 #define DEFINE_EXPORTED_LATENT_AUTOMATION_COMMAND_WITH_RETRIES_THREE_PARAMETERS(EXPORT_API,CommandName,RetryCount,WaitTimeBetweenRuns,ParamType0,ParamName0,ParamType1,ParamName1,ParamType2,ParamName2)	\
-class EXPORT_API CommandName : public IAutomationLatentCommandWithRetriesAndDelays \
+class CommandName : public IAutomationLatentCommandWithRetriesAndDelays \
 	{ \
 	public: \
 	CommandName(int32 InRetryCount, double InWaitTimeBetweenRuns, ParamType0 ParamName0, ParamType1 ParamName1, ParamType2 ParamName2) \
@@ -3921,7 +3921,7 @@ class EXPORT_API CommandName : public IAutomationLatentCommandWithRetriesAndDela
 		{} \
 		virtual ~CommandName() \
 		{} \
-		virtual bool Execute() override; \
+		EXPORT_API virtual bool Execute() override; \
 	private: \
 	ParamType0 ParamName0; \
 	ParamType1 ParamName1; \
@@ -3929,7 +3929,7 @@ class EXPORT_API CommandName : public IAutomationLatentCommandWithRetriesAndDela
 }
 
 #define DEFINE_EXPORTED_LATENT_AUTOMATION_COMMAND_WITH_RETRIES_FOUR_PARAMETERS(EXPORT_API,CommandName,RetryCount,WaitTimeBetweenRuns,ParamType0,ParamName0,ParamType1,ParamName1,ParamType2,ParamName2,ParamType3,ParamName3)	\
-class EXPORT_API CommandName : public IAutomationLatentCommandWithRetriesAndDelays \
+class CommandName : public IAutomationLatentCommandWithRetriesAndDelays \
 	{ \
 	public: \
 	CommandName(int32 InRetryCount, double InWaitTimeBetweenRuns, ParamType0 ParamName0, ParamType1 ParamName1, ParamType2 ParamName2, ParamType3 ParamName3) \
@@ -3941,7 +3941,7 @@ class EXPORT_API CommandName : public IAutomationLatentCommandWithRetriesAndDela
 		{} \
 		virtual ~CommandName() \
 		{} \
-		virtual bool Execute() override; \
+		EXPORT_API virtual bool Execute() override; \
 	private: \
 	ParamType0 ParamName0; \
 	ParamType1 ParamName1; \
diff --git a/Engine/Source/Runtime/Core/Public/Misc/DateTime.h b/Engine/Source/Runtime/Core/Public/Misc/DateTime.h
--- a/Engine/Source/Runtime/Core/Public/Misc/DateTime.h
+++ b/Engine/Source/Runtime/Core/Public/Misc/DateTime.h
@@ -786,7 +786,7 @@ public:
 	 * @param DateTime The date and time value to serialize.
 	 * @return The archive.
 	 */
-	friend CORE_API FArchive& operator<<(FArchive& Ar, FDateTime& DateTime)
+	friend FArchive& operator<<(FArchive& Ar, FDateTime& DateTime)
 	{
 		return Ar << DateTime.Ticks;
 	}
@@ -797,7 +797,7 @@ public:
 	 * @param Slot The structured archive slot to serialize from or into.
 	 * @param DateTime The date and time value to serialize.
 	 */
-	friend CORE_API void operator<<(FStructuredArchive::FSlot Slot, FDateTime& DateTime)
+	friend void operator<<(FStructuredArchive::FSlot Slot, FDateTime& DateTime)
 	{
 		return Slot << DateTime.Ticks;
 	}
@@ -808,7 +808,7 @@ public:
 	 * @param DateTime The date and time to get the hash for.
 	 * @return Hash value.
 	 */
-	friend CORE_API uint32 GetTypeHash(const FDateTime& DateTime)
+	friend uint32 GetTypeHash(const FDateTime& DateTime)
 	{
 		return GetTypeHash(DateTime.Ticks);
 	}
diff --git a/Engine/Source/Runtime/Core/Public/Misc/Paths.h b/Engine/Source/Runtime/Core/Public/Misc/Paths.h
--- a/Engine/Source/Runtime/Core/Public/Misc/Paths.h
+++ b/Engine/Source/Runtime/Core/Public/Misc/Paths.h
@@ -175,7 +175,7 @@ public:
 	 * @param Platform the platform to get the extension directory for
 	 * @return engine's platform extension directory
 	 */
-	static CORE_API FString EnginePlatformExtensionDir(const TCHAR* Platform)
+	static FString EnginePlatformExtensionDir(const TCHAR* Platform)
 	{
 		return ConvertPath(FPaths::EngineDir(), EPathConversion::Engine_PlatformExtension, Platform);
 	}
@@ -194,7 +194,7 @@ public:
 	 * @param Platform the platform to get the extension directory for
 	 * @return project's platform extension directory
 	 */
-	static CORE_API FString ProjectPlatformExtensionDir(const TCHAR* Platform)
+	static FString ProjectPlatformExtensionDir(const TCHAR* Platform)
 	{
 		return ConvertPath(FPaths::ProjectDir(), EPathConversion::Project_PlatformExtension, Platform);
 	}
diff --git a/Engine/Source/Runtime/Core/Public/Serialization/CustomVersion.h b/Engine/Source/Runtime/Core/Public/Serialization/CustomVersion.h
--- a/Engine/Source/Runtime/Core/Public/Serialization/CustomVersion.h
+++ b/Engine/Source/Runtime/Core/Public/Serialization/CustomVersion.h
@@ -92,7 +92,7 @@ private:
 	mutable FName FriendlyName;
 };
 
-class CORE_API FCustomVersionRegistration;
+class FCustomVersionRegistration;
 
 /** Flags defining how the Custom Version should be applied */
 enum class ESetCustomVersionFlags : uint8
diff --git a/Engine/Source/Runtime/Core/Public/Serialization/MemoryImage.h b/Engine/Source/Runtime/Core/Public/Serialization/MemoryImage.h
--- a/Engine/Source/Runtime/Core/Public/Serialization/MemoryImage.h
+++ b/Engine/Source/Runtime/Core/Public/Serialization/MemoryImage.h
@@ -596,7 +596,7 @@ public:
 
 };
 
-class CORE_API FMemoryImageAllocatorBase
+class FMemoryImageAllocatorBase
 {
 	UE_NONCOPYABLE(FMemoryImageAllocatorBase);
 public:
@@ -607,10 +607,10 @@ public:
 	* Assumes that the allocator is currently empty, i.e. memory may be allocated but any existing elements have already been destructed (if necessary).
 	* @param Other - The allocator to move the state from.  This allocator should be left in a valid empty state.
 	*/
-	void MoveToEmpty(FMemoryImageAllocatorBase& Other);
+	CORE_API void MoveToEmpty(FMemoryImageAllocatorBase& Other);
 
 	/** Destructor. */
-	~FMemoryImageAllocatorBase();
+	CORE_API ~FMemoryImageAllocatorBase();
 
 	// FContainerAllocatorInterface
 	FORCEINLINE FScriptContainerElement* GetAllocation() const
@@ -628,10 +628,10 @@ public:
 	}
 	FORCEINLINE int64 GetFrozenOffsetFromThis() const { return Data.GetFrozenOffsetFromThis(); }
 
-	void ResizeAllocation(int32 PreviousNumElements, int32 NumElements, SIZE_T NumBytesPerElement, uint32 Alignment);
-	void WriteMemoryImage(FMemoryImageWriter& Writer, const FTypeLayoutDesc& TypeDesc, int32 NumAllocatedElements, uint32 Alignment) const;
-	void ToString(const FTypeLayoutDesc& TypeDesc, int32 NumAllocatedElements, int32 MaxAllocatedElements, const FPlatformTypeLayoutParameters& LayoutParams, FMemoryToStringContext& OutContext) const;
-	void CopyUnfrozen(const FMemoryUnfreezeContent& Context, const FTypeLayoutDesc& TypeDesc, int32 NumAllocatedElements, void* Dst) const;
+	CORE_API void ResizeAllocation(int32 PreviousNumElements, int32 NumElements, SIZE_T NumBytesPerElement, uint32 Alignment);
+	CORE_API void WriteMemoryImage(FMemoryImageWriter& Writer, const FTypeLayoutDesc& TypeDesc, int32 NumAllocatedElements, uint32 Alignment) const;
+	CORE_API void ToString(const FTypeLayoutDesc& TypeDesc, int32 NumAllocatedElements, int32 MaxAllocatedElements, const FPlatformTypeLayoutParameters& LayoutParams, FMemoryToStringContext& OutContext) const;
+	CORE_API void CopyUnfrozen(const FMemoryUnfreezeContent& Context, const FTypeLayoutDesc& TypeDesc, int32 NumAllocatedElements, void* Dst) const;
 
 private:
 	TMemoryImagePtr<FScriptContainerElement> Data;
@@ -918,7 +918,7 @@ namespace Freeze
 }
 
 
-class CORE_API FPtrTableBase
+class FPtrTableBase
 {
 public:
 	template<typename PtrType>
@@ -934,7 +934,7 @@ public:
 		}
 	}
 
-	void SavePatchesToArchive(FArchive& Ar, uint32 PtrIndex) const;
+	CORE_API void SavePatchesToArchive(FArchive& Ar, uint32 PtrIndex) const;
 
 protected:
 	struct FPatchOffset
@@ -950,7 +950,7 @@ protected:
 		uint32 NumOffsets;
 	};
 
-	void AddPatchedPointerBase(uint32 PtrIndex, uint64 Offset);
+	CORE_API void AddPatchedPointerBase(uint32 PtrIndex, uint64 Offset);
 
 	TArray<FPatchOffsetList> PatchLists;
 	TArray<FPatchOffset> PatchOffsets;
