Make the CPU Usage Utiliization CSV work on more than just Windows
- Add GetPerFrameProcessorUsage to FGenericPlatformProcess
- This follows the Windows version's arguments
- Pull out the struct that Android uses to be more generic
* I did not refactor Android to use this struct for now as it requires a broader refactor
- Use the implementation from Android for Linux (this should be refactored)
- Add a shim for Android that basically calls the already established CPU Usage function

diff --git a/Engine/Source/Runtime/Core/Private/Android/AndroidPlatformProcess.cpp b/Engine/Source/Runtime/Core/Private/Android/AndroidPlatformProcess.cpp
--- a/Engine/Source/Runtime/Core/Private/Android/AndroidPlatformProcess.cpp
+++ b/Engine/Source/Runtime/Core/Private/Android/AndroidPlatformProcess.cpp
@@ -8,6 +8,7 @@
 #include "Android/AndroidPlatform.h"
 #include "Android/AndroidPlatformRunnableThread.h"
 #include "Android/AndroidPlatformAffinity.h"
+#include "Android/AndroidPlatformMisc.h"
 #include "Async/TaskGraphInterfaces.h"
 
 #include <sys/syscall.h>
@@ -175,6 +176,18 @@ void FAndroidPlatformProcess::SetThreadName(const TCHAR* ThreadName)
 	pthread_setname_np(pthread_self(), TCHAR_TO_ANSI(ThreadName));
 }
 
+bool FAndroidPlatformProcess::GetPerFrameProcessorUsage(uint32 ProcessId, float& ProcessUsageFraction, float& IdleUsageFraction)
+{
+	(void)(ProcessId);
+	FAndroidMisc::FCPUState& AndroidCPUState = FAndroidMisc::GetCPUState();
+	ProcessUsageFraction = (float)AndroidCPUState.AverageUtilization;
+	// Cheat for now by assuming idle time + utilization always equals 1.0f
+	IdleUsageFraction = 1.0f - (float)AndroidCPUState.AverageUtilization;
+
+	// We memset the structure zero, so if we have failed for whatever reason this is our sentinel value
+	return AndroidCPUState.CoreCount != 0;
+}
+
 // Can be specified per device profile
 // android.DefaultThreadAffinity GT 0x01 RT 0x02
 TAutoConsoleVariable<FString> CVarAndroidDefaultThreadAffinity(
diff --git a/Engine/Source/Runtime/Core/Private/Unix/UnixPlatformProcess.cpp b/Engine/Source/Runtime/Core/Private/Unix/UnixPlatformProcess.cpp
--- a/Engine/Source/Runtime/Core/Private/Unix/UnixPlatformProcess.cpp
+++ b/Engine/Source/Runtime/Core/Private/Unix/UnixPlatformProcess.cpp
@@ -1799,6 +1799,97 @@ FString FUnixPlatformProcess::GetCurrentWorkingDirectory()
 	return UTF8_TO_TCHAR(CurrentDir);
 }
 
+static UE::Profiling::FCPUState CurrentCPUState;
+
+bool FUnixPlatformProcess::GetPerFrameProcessorUsage(uint32 ProcessId, float& ProcessUsageFraction, float& IdleUsageFraction)
+{
+ 	(void)(ProcessId);
+	int64_t UserTime, NiceTime, SystemTime, SoftIRQTime, IRQTime, IdleTime, WallTime, IOWaitTime;
+	int32		Index = 0;
+	ANSICHAR	Buffer[500];
+
+	CurrentCPUState.CoreCount = FMath::Min( 1, UE::Profiling::FCPUState::MaxSupportedCores);
+	FILE* FileHandle = fopen("/proc/stat", "r");
+	if (FileHandle)
+	{
+		CurrentCPUState.ActivatedCoreCount = 0;
+		for (size_t n = 0; n < CurrentCPUState.CoreCount; n++)
+		{
+			CurrentCPUState.Status[n] = 0;
+			CurrentCPUState.PreviousUsage[n] = CurrentCPUState.CurrentUsage[n];
+		}
+
+		while (fgets(Buffer, 100, FileHandle)) {
+#if PLATFORM_64BITS
+			sscanf(Buffer, "%5s %8lu %8lu %8lu %8lu %8lu %8lu %8lu", CurrentCPUState.Name,
+				&UserTime, &NiceTime, &SystemTime, &IdleTime, &IOWaitTime, &IRQTime, &SoftIRQTime);
+#else
+			sscanf(Buffer, "%5s %8llu %8llu %8llu %8llu %8llu %8llu %8llu", CurrentCPUState.Name,
+				&UserTime, &NiceTime, &SystemTime, &IdleTime, &IOWaitTime, &IRQTime, &SoftIRQTime);
+#endif
+
+			if (0 == strncmp(CurrentCPUState.Name, "cpu", 3))
+			{
+				Index = CurrentCPUState.Name[3] - '0';
+				if (Index >= 0 && Index < CurrentCPUState.CoreCount)
+				{
+					if(CurrentCPUState.Name[5] != '\0')
+					{
+						Index = atol(&CurrentCPUState.Name[3]);
+					}
+					CurrentCPUState.CurrentUsage[Index].IdleTime = IdleTime;
+					CurrentCPUState.CurrentUsage[Index].NiceTime = NiceTime;
+					CurrentCPUState.CurrentUsage[Index].SystemTime = SystemTime;
+					CurrentCPUState.CurrentUsage[Index].SoftIRQTime = SoftIRQTime;
+					CurrentCPUState.CurrentUsage[Index].IRQTime = IRQTime;
+					CurrentCPUState.CurrentUsage[Index].IOWaitTime = IOWaitTime;
+					CurrentCPUState.CurrentUsage[Index].UserTime = UserTime;
+					CurrentCPUState.CurrentUsage[Index].TotalTime = UserTime + NiceTime + SystemTime + SoftIRQTime + IRQTime + IdleTime + IOWaitTime;
+					CurrentCPUState.Status[Index] = 1;
+					CurrentCPUState.ActivatedCoreCount++;
+				}
+				if (Index == CurrentCPUState.CoreCount - 1)
+					break;
+			}
+		}
+		fclose(FileHandle);
+
+		CurrentCPUState.AverageUtilization = 0.0;
+		CurrentCPUState.AverageIdleTime = 0.0f;
+		for (size_t n = 0; n < CurrentCPUState.CoreCount; n++)
+		{
+			if (CurrentCPUState.CurrentUsage[n].TotalTime <= CurrentCPUState.PreviousUsage[n].TotalTime)
+			{
+				continue;
+			}
+			
+			WallTime = CurrentCPUState.CurrentUsage[n].TotalTime - CurrentCPUState.PreviousUsage[n].TotalTime;
+			IdleTime = CurrentCPUState.CurrentUsage[n].IdleTime - CurrentCPUState.PreviousUsage[n].IdleTime;
+
+			if (!WallTime || WallTime <= IdleTime)
+			{
+				continue;
+			}
+			CurrentCPUState.Utilization[n] = ((double)WallTime - (double)IdleTime) * 100.0 / (double)WallTime;
+			CurrentCPUState.IdleTime[n] = (double)IdleTime * 100.0 / (double)WallTime;
+			CurrentCPUState.AverageUtilization += CurrentCPUState.Utilization[n];
+			CurrentCPUState.AverageIdleTime += CurrentCPUState.IdleTime[n];
+		}
+		CurrentCPUState.AverageUtilization /= (double)CurrentCPUState.CoreCount;
+		CurrentCPUState.AverageIdleTime /= (double)CurrentCPUState.CoreCount;
+	}
+ 	else
+ 	{
+ 		UE_LOG(LogHAL, Warning, TEXT("Could not open /proc/stat, CPU usage stats will be missing"));
+		FMemory::Memzero(CurrentCPUState);
+		return false;
+	}
+	ProcessUsageFraction = (float)CurrentCPUState.AverageUtilization;
+ 	IdleUsageFraction = (float)CurrentCPUState.AverageIdleTime;
+
+ 	return true;
+}
+
 bool FUnixPlatformProcess::GetProcReturnCode(FProcHandle& ProcHandle, int32* ReturnCode)
 {
 	if (IsProcRunning(ProcHandle))
diff --git a/Engine/Source/Runtime/Core/Public/Android/AndroidPlatformProcess.h b/Engine/Source/Runtime/Core/Public/Android/AndroidPlatformProcess.h
--- a/Engine/Source/Runtime/Core/Public/Android/AndroidPlatformProcess.h
+++ b/Engine/Source/Runtime/Core/Public/Android/AndroidPlatformProcess.h
@@ -42,6 +42,7 @@ struct FAndroidPlatformProcess : public FGenericPlatformProcess
 	static CORE_API void LaunchURL(const TCHAR* URL, const TCHAR* Parms, FString* Error);
 	static CORE_API FString GetGameBundleId();
 	static CORE_API void SetThreadName(const TCHAR* ThreadName);
+	static CORE_API bool GetPerFrameProcessorUsage(uint32 ProcessId, float& ProcessUsageFraction, float& IdleUsageFraction);
 };
 
 typedef FAndroidPlatformProcess FPlatformProcess;
diff --git a/Engine/Source/Runtime/Core/Public/GenericPlatform/GenericPlatformProcess.h b/Engine/Source/Runtime/Core/Public/GenericPlatform/GenericPlatformProcess.h
--- a/Engine/Source/Runtime/Core/Public/GenericPlatform/GenericPlatformProcess.h
+++ b/Engine/Source/Runtime/Core/Public/GenericPlatform/GenericPlatformProcess.h
@@ -33,6 +33,14 @@ class FEvent;
 #	endif
 #endif
 
+// Assume ThreadRipper is the max number of cores we can have
+#if PLATFORM_DESKTOP
+#define MAX_NUM_PROCESSORS 128
+#else
+#define MAX_NUM_PROCESSORS 16
+#endif
+
+
 class Error;
 struct FProcHandle;
 template <typename FuncType> class TFunctionRef;
@@ -131,6 +139,36 @@ struct FPlatformProcessMemoryStats
 	uint64 PeakUsedVirtual;
 };
 
+// Only used on Unix-like (Linux, Unix, Android) platforms because you have to scrape /proc/stat
+namespace UE::Profiling
+{
+	struct FCPUStatTime{
+		uint64_t			TotalTime;
+		uint64_t			UserTime;
+		uint64_t			NiceTime;
+		uint64_t			SystemTime;
+		uint64_t			SoftIRQTime;
+		uint64_t			IRQTime;
+		uint64_t			IdleTime;
+		uint64_t			IOWaitTime;
+	};
+	
+	struct FCPUState
+	{
+		constexpr static int32		MaxSupportedCores = MAX_NUM_PROCESSORS; 
+		int32						CoreCount;
+		int32						ActivatedCoreCount;
+		ANSICHAR					Name[6];
+		FCPUStatTime				CurrentUsage[MaxSupportedCores]; 
+		FCPUStatTime				PreviousUsage[MaxSupportedCores];
+		int32						Status[MaxSupportedCores];
+		double						Utilization[MaxSupportedCores];
+		double						IdleTime[MaxSupportedCores];
+		double						AverageUtilization;
+		double						AverageIdleTime;
+	};
+}
+
 /**
 * Generic implementation for most platforms, these tend to be unused and unimplemented
 **/
@@ -235,6 +273,13 @@ struct FGenericPlatformProcess
 	 * @return the current hardware core.
 	 */
 	static CORE_API uint32 GetCurrentCoreNumber();
+	
+	/**
+	 * Retrieves the current CPU Utilization for the current process in terms of percentage of all capacity
+	 *
+	 * @return the current hardware core.
+	 */
+	static CORE_API bool GetPerFrameProcessorUsage(uint32 ProcessId, float& ProcessUsageFraction, float& IdleUsageFraction) {return false;}
 
 	/**	 
 	 * Change the thread processor affinity
diff --git a/Engine/Source/Runtime/Core/Public/Unix/UnixPlatformProcess.h b/Engine/Source/Runtime/Core/Public/Unix/UnixPlatformProcess.h
--- a/Engine/Source/Runtime/Core/Public/Unix/UnixPlatformProcess.h
+++ b/Engine/Source/Runtime/Core/Public/Unix/UnixPlatformProcess.h
@@ -315,6 +315,7 @@ struct FUnixPlatformProcess : public FGenericPlatformProcess
 	static CORE_API int32 TranslateThreadPriority(EThreadPriority Priority);
 	static CORE_API void SetThreadNiceValue( uint32_t ThreadId, int32 NiceValue );
 	static CORE_API void SetThreadPriority( EThreadPriority NewPriority );
+	static CORE_API bool GetPerFrameProcessorUsage(uint32 ProcessId, float& ProcessUsageFraction, float& IdleUsageFraction);
 
 	/**
 	 * @brief Releases locks that we held for IsFirstInstance check
diff --git a/Engine/Source/Runtime/Launch/Private/LaunchEngineLoop.cpp b/Engine/Source/Runtime/Launch/Private/LaunchEngineLoop.cpp
--- a/Engine/Source/Runtime/Launch/Private/LaunchEngineLoop.cpp
+++ b/Engine/Source/Runtime/Launch/Private/LaunchEngineLoop.cpp
@@ -1339,14 +1339,22 @@ static void UpdateCoreCsvStats_BeginFrame()
 	{
 		if (!IsRunningDedicatedServer())
 		{
-#if PLATFORM_WINDOWS && !UE_BUILD_SHIPPING
-		    const uint32 ProcessId = (uint32)GetCurrentProcessId();
+#if !UE_BUILD_SHIPPING
+			const uint32 ProcessId = (uint32)FPlatformProcess::GetCurrentProcessId();
 		    float ProcessUsageFraction = 0.f, IdleUsageFraction = 0.f;
-		    FWindowsPlatformProcess::GetPerFrameProcessorUsage(ProcessId, ProcessUsageFraction, IdleUsageFraction);
-    
-		    CSV_CUSTOM_STAT_GLOBAL(CPUUsage_Process, ProcessUsageFraction, ECsvCustomStatOp::Set);
-		    CSV_CUSTOM_STAT_GLOBAL(CPUUsage_Idle, IdleUsageFraction, ECsvCustomStatOp::Set);
+			if (FPlatformProcess::GetPerFrameProcessorUsage(ProcessId, ProcessUsageFraction, IdleUsageFraction))
+			{
+				// This occasionally happens on Linux and I am not entirely sure why.
+				// If we have zero idle time and zero utilization assume we have full idle time
+				if (ProcessUsageFraction == 0.0f && IdleUsageFraction == 0.0f)
+				{
+					IdleUsageFraction = 1.0f;
+				}
+				CSV_CUSTOM_STAT_GLOBAL(CPUUsage_Process, ProcessUsageFraction, ECsvCustomStatOp::Set);
+				CSV_CUSTOM_STAT_GLOBAL(CPUUsage_Idle, IdleUsageFraction, ECsvCustomStatOp::Set);
+			}
 #endif
+			
 		}
 #if CSV_PROFILER_ALLOW_DEBUG_FEATURES
 		// Handle CsvExecCmds for this frame
@@ -1421,6 +1429,7 @@ static void UpdateCoreCsvStats_EndFrame()
 	    }
 	    CSV_CUSTOM_STAT_GLOBAL(MaxFrameTime, 1000.0f / TargetFPS, ECsvCustomStatOp::Set);
 
+		// TODO: see what causes this to be true
 		if (GRHISupportsGPUUsage && GNumExplicitGPUsForRendering == 1)
 		{
 			FRHIGPUUsageFractions GPUUsage = RHIGetGPUUsage(/* GPUIndex = */ 0);
