Apply Wine limits to File read/write

diff --git a/Engine/Source/Runtime/Core/Private/Windows/WindowsPlatformFile.cpp b/Engine/Source/Runtime/Core/Private/Windows/WindowsPlatformFile.cpp
--- a/Engine/Source/Runtime/Core/Private/Windows/WindowsPlatformFile.cpp
+++ b/Engine/Source/Runtime/Core/Private/Windows/WindowsPlatformFile.cpp
@@ -528,7 +528,7 @@ public:
 		int64 TotalNumRead = 0;
 		do
 		{
-			uint32 BytesToRead32 = (uint32)FMath::Min<int64>(BytesToRead, int64(UINT32_MAX));
+			uint32 BytesToRead32 = (uint32)FMath::Min<int64>(BytesToRead, FWindowsPlatformMisc::IsWine() ? int64(0x7ffff000) : int64(UINT32_MAX));
 			uint32 NumRead = 0;
 
 			UE::WindowsPlatformFile::Private::UpdateOverlappedPos(LocalOverlappedIO, Offset);
@@ -750,7 +750,7 @@ public:
 		int64 TotalNumRead = 0;
 		do
 		{
-			uint32 BytesToRead32 = (uint32)FMath::Min<int64>(BytesToRead, int64(UINT32_MAX));
+			uint32 BytesToRead32 = (uint32)FMath::Min<int64>(BytesToRead, FWindowsPlatformMisc::IsWine() ? int64(0x7ffff000) : int64(UINT32_MAX));
 			uint32 NumRead = 0;
 
 			if (!ReadFile(FileHandle, Destination, BytesToRead32, (::DWORD*)&NumRead, &OverlappedIO))
@@ -815,7 +815,8 @@ public:
 		int64 TotalNumRead = 0;
 		do
 		{
-			uint32 BytesToRead32 = (uint32)FMath::Min<int64>(BytesToRead, int64(UINT32_MAX));
+			uint32 BytesToRead32 = (uint32)FMath::Min<int64>(BytesToRead, FWindowsPlatformMisc::IsWine() ? int64(0x7ffff000) : int64(UINT32_MAX));
+			
 			uint32 NumRead = 0;
 
 			UE::WindowsPlatformFile::Private::UpdateOverlappedPos(LocalOverlappedIO, Offset);
@@ -855,7 +856,7 @@ public:
 		int64 TotalNumWritten = 0;
 		do
 		{
-			uint32 BytesToWrite32 = (uint32)FMath::Min<int64>(BytesToWrite, int64(UINT32_MAX));
+			uint32 BytesToWrite32 = (uint32)FMath::Min<int64>(BytesToWrite, FWindowsPlatformMisc::IsWine() ? int64(0x7ffff000) : int64(UINT32_MAX));
 			uint32 NumWritten = 0;
 			// Now kick off an async write
 			if (!WriteFile(FileHandle, Source, BytesToWrite32, (::DWORD*)&NumWritten, &OverlappedIO))
