[RaceDetector]
- Centralize limitation around callstack depth to a single place dictated by FContext current callstack size
- Add safe allocator support for guard pages around virtual allocations
- Remove conditional on current callstack adjustment and rely on unaccessible virtual memory guard page to hard crash if the limit is ever reached
- Cleanups for race detector checks and make them work in all targets when activated
- Add centralized memory access throttling during shadow reset to avoid infinite loop when too many threads continue making allocations while resetting

diff --git a/Engine/Source/Runtime/Core/Internal/Instrumentation/Containers.h b/Engine/Source/Runtime/Core/Internal/Instrumentation/Containers.h
--- a/Engine/Source/Runtime/Core/Internal/Instrumentation/Containers.h
+++ b/Engine/Source/Runtime/Core/Internal/Instrumentation/Containers.h
@@ -19,6 +19,8 @@ class FInstrumentationSafeWinAllocator {
 public:
 	static INSTRUMENTATION_FUNCTION_ATTRIBUTES void* Alloc(SIZE_T Size);
 	static INSTRUMENTATION_FUNCTION_ATTRIBUTES void Free(void* Data);
+	static INSTRUMENTATION_FUNCTION_ATTRIBUTES void* AllocWithGuards(SIZE_T Size);
+	static INSTRUMENTATION_FUNCTION_ATTRIBUTES void FreeWithGuards(void* Data);
 	static INSTRUMENTATION_FUNCTION_ATTRIBUTES void* Realloc(void* Data, SIZE_T Size, SIZE_T PreviousSize);
 };
 
@@ -34,6 +36,16 @@ using TInstrumentationSafeBaseAllocator = FInstrumentationSafeWinAllocator;
 		FInstrumentationSafeWinAllocator::Free(Ptr);  \
 	} 
 
+#define SAFE_OPERATOR_NEW_DELETE_WITH_GUARDS() \
+	INSTRUMENTATION_FUNCTION_ATTRIBUTES void * operator new(SIZE_T Size)  \
+	{ \
+		return FInstrumentationSafeWinAllocator::AllocWithGuards(Size); \
+	} \
+	INSTRUMENTATION_FUNCTION_ATTRIBUTES void   operator delete(void* Ptr) \
+	{ \
+		FInstrumentationSafeWinAllocator::FreeWithGuards(Ptr);  \
+	} 
+
 #else
 class TInstrumentationSafeMallocAllocator {
 public:
diff --git a/Engine/Source/Runtime/Core/Internal/Sanitizer/RaceDetectorTypes.h b/Engine/Source/Runtime/Core/Internal/Sanitizer/RaceDetectorTypes.h
--- a/Engine/Source/Runtime/Core/Internal/Sanitizer/RaceDetectorTypes.h
+++ b/Engine/Source/Runtime/Core/Internal/Sanitizer/RaceDetectorTypes.h
@@ -228,7 +228,7 @@ namespace UE::Sanitizer::RaceDetector {
 
 		INSTRUMENTATION_FUNCTION_ATTRIBUTES FCallstackLocation GetLocation(FContextId ContextId)
 		{
-			return FCallstackLocation::FromCallstackFrames(&Locations[ContextId], 1, 2);
+			return FCallstackLocation(&Locations[ContextId], 1);
 		}
 
 	private:
@@ -339,7 +339,7 @@ namespace UE::Sanitizer::RaceDetector {
 
 	enum class EHistoryEntryType : uint8
 	{
-                Invalid       = 0,     // Just make sure we do not mistake 0 memory for something valid
+		Invalid       = 0,     // Just make sure we do not mistake 0 memory for something valid
 		FunctionEntry = 0xAA,  // Any number will do but make them stand out in the trace
 		MemoryAccess  = 0xBB,
 		FunctionExit  = 0xCC
@@ -582,9 +582,7 @@ namespace UE::Sanitizer::RaceDetector {
 				// Do not bother searching a chunk that is outside the range we're looking for.
 				if (InAccess.Clock >= Chunk->StartClock && InAccess.Clock <= Chunk->EndClock)
 				{
-					int32 StackDepth = 0;
-					void* Stack[1024];
-					Stack[StackDepth] = 0;
+					TArray<void*, TInlineAllocator<1024>> Stack;
 
 					int32 Offset = 0;
 					while (Offset < Chunk->Offset)
@@ -593,8 +591,7 @@ namespace UE::Sanitizer::RaceDetector {
 						switch (Entry->Type)
 						{
 							case EHistoryEntryType::FunctionEntry:
-								check(StackDepth < UE_ARRAY_COUNT(Stack));
-								Stack[StackDepth++] = ((FHistoryEntryFunctionEntry*)Entry)->ReturnAddress;
+								Stack.Add(((FHistoryEntryFunctionEntry*)Entry)->ReturnAddress);
 								Offset += sizeof(FHistoryEntryFunctionEntry);
 							break;
 							case EHistoryEntryType::MemoryAccess:
@@ -603,23 +600,21 @@ namespace UE::Sanitizer::RaceDetector {
 								if (MemoryAccessEntry->Pointer == InAlignedPointer &&
 									MemoryAccessEntry->Access.RawValue == InAccess.RawValue)
 								{
-									OutLocation = FCallstackLocation::FromCallstackFrames(Stack, StackDepth, StackDepth + 1);
+									OutLocation = FCallstackLocation(Stack.GetData(), Stack.Num());
 									return true;
 								}
 								Offset += sizeof(FHistoryEntryAccess);
 							}
 							break;
 							case EHistoryEntryType::FunctionExit:
-								Stack[StackDepth--] = 0;
-								check(StackDepth >= 0);
+								Stack.Pop();
 								Offset += sizeof(FHistoryEntryBase);
 							break;
 							default:
 								check(false);
 								// This should never happen, but if it does it is most likely a race condition
 								// so just restart the tracing from the beginning as a last resort.
-								StackDepth = 0;
-								Stack[StackDepth] = 0;
+								Stack.Reset();
 								Offset = 0;
 							break;
 						}
@@ -655,7 +650,7 @@ namespace UE::Sanitizer::RaceDetector {
 	// by other threads and we'd need to hold a lock while scanning
 	// the history for race report, which would be unnaceptable.
 	struct FContext : public TRefCountingMixin<FContext> {
-		SAFE_OPERATOR_NEW_DELETE()
+		SAFE_OPERATOR_NEW_DELETE_WITH_GUARDS()
 
 		INSTRUMENTATION_FUNCTION_ATTRIBUTES FContext(uint32 InThreadId)
 			: ThreadId(InThreadId)
@@ -693,11 +688,27 @@ namespace UE::Sanitizer::RaceDetector {
 		// Hazard pointer used between GetSyncObject and ResetShadow
 		FSyncObjectBank* BankHazard = nullptr;
 		// We need to keep the callstack for each thread 
-		static constexpr SIZE_T MaxCallstackSize = 1024; // Yes, callstacks can get huge in UE!
-		void* CurrentCallstack[MaxCallstackSize];
 		uint16 CurrentCallstackSize = 0;
 		// We use this to pass thread arguments to functions that dont have parameters (i.e. ExitThread).
 		void* ThreadArgs = nullptr;
+		// Avoid using UniquePtr because it's instrumented and each access has a cost.
+		FAccessHistory* AccessHistory = nullptr;
+
+		// This can be bumped again if we ever face a need for deeper callstacks since
+		// this is a virtual allocation anyway so it's not going to take physical memory
+		// until it is used. 
+		// This needs to be at the end of the allocated block as we rely on page fault
+		// to abort the program if the stack ever goes beyond this limit.
+		static constexpr SIZE_T MaxCallstackSize = 4096;
+		void* CurrentCallstack[MaxCallstackSize];
+			
+		static_assert(
+			sizeof(FHistoryChunk::Buffer) > 10 * MaxCallstackSize * sizeof(FHistoryEntryFunctionEntry), 
+			"FHistoryChunk::Buffer should be big enough to accomodate initial callstack with plenty of space left"
+		);
+
+		// [NO ACCESS GUARD PAGE]
+
 		// The clock for this context Id
 		INSTRUMENTATION_FUNCTION_ATTRIBUTES FClock& CurrentClock()
 		{ 
@@ -706,6 +717,7 @@ namespace UE::Sanitizer::RaceDetector {
 			const FContextId LocalContextId = ContextId;
 			return LocalContextId == 0 ? StandbyClock : ClockBank.Get(LocalContextId);
 		}
+		
 		INSTRUMENTATION_FUNCTION_ATTRIBUTES void IncrementClock()
 		{
 			CurrentClock()++; 
@@ -716,9 +728,6 @@ namespace UE::Sanitizer::RaceDetector {
 			}
 		}
 
-		// Avoid using UniquePtr because it's instrumented and each access has a cost.
-		FAccessHistory* AccessHistory = nullptr;
-
 		INSTRUMENTATION_FUNCTION_ATTRIBUTES static bool IsValid(FContext* Context)
 		{
 			return reinterpret_cast<int64>(Context) > 0;
diff --git a/Engine/Source/Runtime/Core/Internal/Sanitizer/Types.h b/Engine/Source/Runtime/Core/Internal/Sanitizer/Types.h
--- a/Engine/Source/Runtime/Core/Internal/Sanitizer/Types.h
+++ b/Engine/Source/Runtime/Core/Internal/Sanitizer/Types.h
@@ -75,10 +75,19 @@ namespace UE::Sanitizer {
 		};
 	};
 
-	struct FFullLocation {
+	struct FFullLocation
+	{
 		TArray<FLocation> Locations;
 
-		INSTRUMENTATION_FUNCTION_ATTRIBUTES int32 Num() const { return Locations.Num(); }
+		INSTRUMENTATION_FUNCTION_ATTRIBUTES int32 Num() const
+		{ 
+			return Locations.Num(); 
+		}
+
+		INSTRUMENTATION_FUNCTION_ATTRIBUTES void Reserve(int32 Size) 
+		{ 
+			Locations.Reserve(Size);
+		}
 
 		INSTRUMENTATION_FUNCTION_ATTRIBUTES int32 GetAlignment() const
 		{
@@ -116,51 +125,22 @@ namespace UE::Sanitizer {
 	// ------------------------------------------------------------------------------
 	// Callstack.
 	// ------------------------------------------------------------------------------
-	class FCallstackLocation {
-	private:
-		static constexpr uint32 NumCallstackFrames = 1024;
-
-		static constexpr UPTRINT NoCallstackMarker = 0;
-		static constexpr UPTRINT OverflowCallstackMarker = 0x1;
-
-		static constexpr const TCHAR* OverflowString = TEXT("[Uncaptured frames, callstack too deep]");
-
+	class FCallstackLocation 
+	{
 	public:
 		INSTRUMENTATION_FUNCTION_ATTRIBUTES FCallstackLocation()
 		{
-			Callstack[0] = NoCallstackMarker;
 		}
 
-		INSTRUMENTATION_FUNCTION_ATTRIBUTES static FCallstackLocation FromCallstackFrames(const void* InCallstack, uint32 NumFrames, uint32 MaxNumFrames)
+		INSTRUMENTATION_FUNCTION_ATTRIBUTES FCallstackLocation(void* const* InCallstack, uint32 NumFrames)
 		{
-			FCallstackLocation Location;
-
-			// Set overflow marker if needed.
-			uint32 NumMaxFrames = NumCallstackFrames;
-			if (NumFrames >= MaxNumFrames)
-			{
-				Location.Callstack[NumCallstackFrames - 1] = OverflowCallstackMarker;
-				NumMaxFrames--;
-			}
-
-			NumFrames = FMath::Min(NumFrames, MaxNumFrames);
-
-			uint32 NumPopulated = FMath::Min(NumFrames, NumMaxFrames);
-			const UPTRINT* InCallstackPtrs = (const UPTRINT*)InCallstack;
-			InCallstackPtrs += NumFrames - NumPopulated;
-
-			memcpy(Location.Callstack, InCallstackPtrs, sizeof(UPTRINT) * NumPopulated);
-			if (NumPopulated < NumMaxFrames)
-			{
-				memset(Location.Callstack + NumPopulated, 0, sizeof(UPTRINT) * (NumMaxFrames - NumPopulated));
-			}
-			return Location;
+			Callstack.Append(InCallstack, NumFrames);
 		}
 
 		INSTRUMENTATION_FUNCTION_ATTRIBUTES FFullLocation GetFullLocation() const
 		{
 			FFullLocation Loc;
-			if (Callstack[0] == NoCallstackMarker)
+			if (Callstack.IsEmpty())
 			{
 				return Loc;
 			}
@@ -171,32 +151,22 @@ namespace UE::Sanitizer {
 			// during its initialization. (i.e. Registering new FNames)
 			LowLevelTasks::Private::FOversubscriptionAllowedScope AllowOversubscription(false);
 
-			for (int32 Idx = NumCallstackFrames - 1; Idx >= 0; --Idx)
+			for (int32 Idx = Callstack.Num() - 1; Idx >= 0; --Idx)
 			{
-				UPTRINT Frame = Callstack[Idx];
-				if (Frame == NoCallstackMarker)
-				{
-					continue;
-				}
-				else if (Frame == OverflowCallstackMarker)
-				{
-					Loc.Locations.Emplace(OverflowString, OverflowString, OverflowString, 0, 0);
-				}
-				else
-				{
-					const bool bIncludeInlineFrames = true;
-					// We always record the return address of our functions using FuncEntry. What we need to properly decode
-					// inline frames is something that points on the address of the actual call, not the return address, otherwise
-					// we end up decoding inline frames for what comes after we return from the function, which doesn't make sense.
-					// By using Frame - 1, we make sure that we actually point inside the chunk of assembly that represents the function call and it is enough
-					// to make symbol decode correct since we only need to point inside a valid range of a symbol entry, not directly at the start of it.
-					FPlatformStackWalk::EnumerateSymbolInfosForProgramCounter(Frame - 1, bIncludeInlineFrames, 
-						[&Loc, Frame](FProgramCounterSymbolInfo& SymbolInfo) INSTRUMENTATION_FUNCTION_ATTRIBUTES
-						{
-							Loc.Locations.Emplace(FString(SymbolInfo.ModuleName), FString(SymbolInfo.FunctionName), FString(SymbolInfo.Filename), SymbolInfo.LineNumber, Frame);
-						}
-					);
-				}
+				UPTRINT Frame = (UPTRINT)Callstack[Idx];
+
+				const bool bIncludeInlineFrames = true;
+				// We always record the return address of our functions using FuncEntry. What we need to properly decode
+				// inline frames is something that points on the address of the actual call, not the return address, otherwise
+				// we end up decoding inline frames for what comes after we return from the function, which doesn't make sense.
+				// By using Frame - 1, we make sure that we actually point inside the chunk of assembly that represents the function call and it is enough
+				// to make symbol decode correct since we only need to point inside a valid range of a symbol entry, not directly at the start of it.
+				FPlatformStackWalk::EnumerateSymbolInfosForProgramCounter(Frame - 1, bIncludeInlineFrames, 
+					[&Loc, Frame](FProgramCounterSymbolInfo& SymbolInfo) INSTRUMENTATION_FUNCTION_ATTRIBUTES
+					{
+						Loc.Locations.Emplace(FString(SymbolInfo.ModuleName), FString(SymbolInfo.FunctionName), FString(SymbolInfo.Filename), SymbolInfo.LineNumber, Frame);
+					}
+				);
 			}
 
 			return Loc;
@@ -204,33 +174,23 @@ namespace UE::Sanitizer {
 
 		INSTRUMENTATION_FUNCTION_ATTRIBUTES uint32 GetHash() const
 		{
-			if (Callstack[0] == NoCallstackMarker)
+			if (Callstack.IsEmpty())
 			{
 				return 0;
 			}
 
-			return CityHash32((const char*)Callstack, sizeof(UPTRINT) * NumCallstackFrames);
+			return CityHash32((const char*)Callstack.GetData(), sizeof(void*) * Callstack.Num());
 		}
 
 		INSTRUMENTATION_FUNCTION_ATTRIBUTES UPTRINT GetLastFrame() const
 		{
-			for (int32 Idx = NumCallstackFrames - 1; Idx >= 0; --Idx)
-			{
-				UPTRINT Frame = Callstack[Idx];
-				if (Frame != NoCallstackMarker &&
-					Frame != OverflowCallstackMarker)
-				{
-					return Frame;
-				}
-			}
-
-			return 0;
+			return Callstack.IsEmpty() ? 0 : (UPTRINT)Callstack.Last();
 		}
 
 	private:
 		// Order is outer frames to inner frames.
 		// For example, Main -> Fn1 -> Fn2 -> FnLeaf.
-		UPTRINT Callstack[NumCallstackFrames];
+		TArray<void*, TInlineAllocator<1024>> Callstack;
 	};
 
 } // namespace UE::Sanitizer
diff --git a/Engine/Source/Runtime/Core/Private/Instrumentation/Containers.cpp b/Engine/Source/Runtime/Core/Private/Instrumentation/Containers.cpp
--- a/Engine/Source/Runtime/Core/Private/Instrumentation/Containers.cpp
+++ b/Engine/Source/Runtime/Core/Private/Instrumentation/Containers.cpp
@@ -16,6 +16,34 @@ void FInstrumentationSafeWinAllocator::Free(void* Data)
 	VirtualFree(Data, 0, MEM_RELEASE);
 }
 
+INSTRUMENTATION_FUNCTION_ATTRIBUTES static size_t GetPageSize()
+{
+	static size_t PageSize = 0;
+	if (PageSize == 0)
+	{
+		SYSTEM_INFO si;
+		GetSystemInfo(&si);
+		PageSize = si.dwPageSize;
+	}
+	return PageSize;
+}
+
+void* FInstrumentationSafeWinAllocator::AllocWithGuards(size_t Size)
+{
+	size_t PageSize = GetPageSize();
+	size_t RequestedSize = Align(Size, PageSize);
+
+	// Requires 1 page on each side of the allocation that won't be committed and will fault if accessed
+	size_t TotalSize = PageSize * 2 + Align(Size, PageSize);
+	void* BaseAddress = VirtualAlloc(NULL, TotalSize, MEM_RESERVE, PAGE_NOACCESS);
+	return VirtualAlloc((void*)((UPTRINT)BaseAddress + PageSize), RequestedSize, MEM_COMMIT, PAGE_READWRITE);
+}
+
+void FInstrumentationSafeWinAllocator::FreeWithGuards(void* Data)
+{
+	VirtualFree((void*)((UPTRINT)Data - GetPageSize()), 0, MEM_RELEASE);
+}
+
 void* FInstrumentationSafeWinAllocator::Realloc(void* Data, size_t Size, size_t PreviousSize)
 {
 	if (Data == nullptr)
diff --git a/Engine/Source/Runtime/Core/Private/Sanitizer/RaceDetector.cpp b/Engine/Source/Runtime/Core/Private/Sanitizer/RaceDetector.cpp
--- a/Engine/Source/Runtime/Core/Private/Sanitizer/RaceDetector.cpp
+++ b/Engine/Source/Runtime/Core/Private/Sanitizer/RaceDetector.cpp
@@ -38,7 +38,7 @@ UE_TRACE_CHANNEL_DEFINE(RaceDetectorChannel);
 #define WITH_RACEDETECTOR_DEBUG 0
 
 #if WITH_RACEDETECTOR_CHECK
-#define RACEDETECTOR_CHECK(cond) check(cond)
+#define RACEDETECTOR_CHECK(cond) if (!(cond)) { PLATFORM_BREAK(); }
 #else
 #define RACEDETECTOR_CHECK(cond)
 #endif
@@ -190,11 +190,9 @@ namespace UE::Sanitizer::RaceDetector {
 	// ------------------------------------------------------------------------------
 	// Callstack management.
 	// ------------------------------------------------------------------------------
-	FORCEINLINE INSTRUMENTATION_FUNCTION_ATTRIBUTES FCallstackLocation GetCurrentCallstackLocation()
+	FORCEINLINE INSTRUMENTATION_FUNCTION_ATTRIBUTES FCallstackLocation GetCurrentCallstackLocation(FContext& Context)
 	{
-		FContext& Context = *GetThreadContext();
-
-		return FCallstackLocation::FromCallstackFrames(Context.CurrentCallstack, Context.CurrentCallstackSize, FContext::MaxCallstackSize);
+		return FCallstackLocation(Context.CurrentCallstack, Context.CurrentCallstackSize);
 	}
 
 	FORCEINLINE INSTRUMENTATION_FUNCTION_ATTRIBUTES void AddCallstackFrame(FContext& Context, void* ReturnAddress)
@@ -207,13 +205,9 @@ namespace UE::Sanitizer::RaceDetector {
 			Context.AccessHistory->AddFunctionEntry(ReturnAddress);
 		}
 
-		// Don't capture extremely deep callstacks.
-		if (Context.CurrentCallstackSize < FContext::MaxCallstackSize)
-		{
-			Context.CurrentCallstack[Context.CurrentCallstackSize] = ReturnAddress;
-		}
-
-		Context.CurrentCallstackSize++;
+		// No need for conditionals here since we're going to generate a page fault if we 
+		// go outside the context allocated memory as we are using guard pages.
+		Context.CurrentCallstack[Context.CurrentCallstackSize++] = ReturnAddress;
 	}
 
 	FORCEINLINE INSTRUMENTATION_FUNCTION_ATTRIBUTES void RemoveCallstackFrame(FContext& Context)
@@ -359,7 +353,7 @@ namespace UE::Sanitizer::RaceDetector {
 		}
 
 		UPTRINT Result = ShadowClockBase + (Ptr >> 3) * sizeof(FShadowClockBankSlot);
-		RACEDETECTOR_CHECK(Result >= ShadowMemoryBase && Result < ShadowMemoryEnd);
+		RACEDETECTOR_CHECK(Result >= ShadowClockBase && Result < ShadowMemoryEnd);
 		return Result;
 	}
 
@@ -521,7 +515,7 @@ namespace UE::Sanitizer::RaceDetector {
 			return;
 		}
 
-		FCallstackLocation FirstLocation = GetCurrentCallstackLocation();
+		FCallstackLocation FirstLocation = GetCurrentCallstackLocation(Context);
 
 		const uint32 FirstThreadId = Context.ThreadId;
 		const uint32 SecondThreadId = RacingContext->ThreadId;
@@ -698,10 +692,9 @@ namespace UE::Sanitizer::RaceDetector {
 
 	INSTRUMENTATION_FUNCTION_ATTRIBUTES FORCEINLINE void InstrumentAccess(FContext& Context, void* Ptr, uint8 Size, EMemoryAccessType AccessType, FAtomicMemoryOrder Order, void* ReturnAddress, const TCHAR* OpName, bool& bHasAddedStack)
 	{
-#if WITH_RACEDETECTOR_CHECK
 		// Validate size and alignment since it should have been done before entering this function.
-		check(Size <= 8 && GetShadowMemoryAddress((UPTRINT)Ptr) == GetShadowMemoryAddress((UPTRINT)Ptr + Size - 1));
-#endif
+		RACEDETECTOR_CHECK(Size <= 8 && GetShadowMemoryAddress((UPTRINT)Ptr) == GetShadowMemoryAddress((UPTRINT)Ptr + Size - 1));
+
 		// GetShadowMemory will contain the cost of committing pages to shadow memory
 		// The cost of this is higher when run under the debugger because we're not
 		// using vectored exception handler in this case due to abysmal perf caused
@@ -799,10 +792,8 @@ namespace UE::Sanitizer::RaceDetector {
 				return;
 			}
 
-#if WITH_RACEDETECTOR_CHECK
 			// Should never happen to have identical values since we verified at entry.
-			check(PreviousAccess.RawValue != CurrentAccess.RawValue);
-#endif
+			RACEDETECTOR_CHECK(PreviousAccess.RawValue != CurrentAccess.RawValue);
 
 			// Check if there is any overlap with the previous access
 			// Each bit in Access represets 1-byte in memory so we can just
@@ -898,9 +889,8 @@ namespace UE::Sanitizer::RaceDetector {
 	template <typename AtomicOpType>
 	INSTRUMENTATION_FUNCTION_ATTRIBUTES void InstrumentAtomicAccess(FContext& Context, void* AtomicAddr, uint8 Size, EMemoryAccessType AccessType, FAtomicMemoryOrder Order, AtomicOpType&& AtomicOp, void* ReturnAddress, const TCHAR* OpName, bool& bHasAddedCallstackFrame)
 	{
-#if WITH_RACEDETECTOR_CHECK
-		check(FMath::Min((UPTRINT)Size, GetAlignedAddress((UPTRINT)AtomicAddr) + 8 - (UPTRINT)AtomicAddr) == (UPTRINT)Size);
-#endif
+		RACEDETECTOR_CHECK(FMath::Min((UPTRINT)Size, GetAlignedAddress((UPTRINT)AtomicAddr) + 8 - (UPTRINT)AtomicAddr) == (UPTRINT)Size);
+
 		InstrumentAccess(Context, AtomicAddr, Size, AccessType, Order, ReturnAddress, OpName, bHasAddedCallstackFrame);
 		if (IsAtomicOrderRelaxed(Order))
 		{
@@ -933,9 +923,8 @@ namespace UE::Sanitizer::RaceDetector {
 	template <typename AtomicOpType>
 	INSTRUMENTATION_FUNCTION_ATTRIBUTES void InstrumentAtomicAccess(FContext& Context, void* AtomicAddr, uint8 Size, EMemoryAccessType AccessType, FAtomicMemoryOrder SuccessOrder, FAtomicMemoryOrder FailureOrder, AtomicOpType&& AtomicOp, void* ReturnAddress, const TCHAR* OpName, bool& bHasAddedCallstackFrame)
 	{
-#if WITH_RACEDETECTOR_CHECK
-		check(FMath::Min((UPTRINT)Size, GetAlignedAddress((UPTRINT)AtomicAddr) + 8 - (UPTRINT)AtomicAddr) == (UPTRINT)Size);
-#endif
+		RACEDETECTOR_CHECK(FMath::Min((UPTRINT)Size, GetAlignedAddress((UPTRINT)AtomicAddr) + 8 - (UPTRINT)AtomicAddr) == (UPTRINT)Size);
+
 		if (IsAtomicOrderRelaxed(SuccessOrder) && IsAtomicOrderRelaxed(FailureOrder))
 		{
 			// Both orders are relaxed, forward either one of them.
@@ -1134,6 +1123,8 @@ namespace UE::Sanitizer::RaceDetector {
 			}
 			bWasDebuggerPresent = bIsDebuggerPresent;
 
+			HandleReports();
+
 			if (bRaceDetectorActive)
 			{
 				if (RaceDetectorTimeout.IsExpired())
@@ -1162,9 +1153,9 @@ namespace UE::Sanitizer::RaceDetector {
 					ResetShadow();
 				}
 			}
-
-			HandleReports();
 		}
+
+		HandleReports();
 	}
 
 	/**
diff --git a/Engine/Source/Runtime/Core/Private/Sanitizer/RaceDetectorWindows.cpp b/Engine/Source/Runtime/Core/Private/Sanitizer/RaceDetectorWindows.cpp
--- a/Engine/Source/Runtime/Core/Private/Sanitizer/RaceDetectorWindows.cpp
+++ b/Engine/Source/Runtime/Core/Private/Sanitizer/RaceDetectorWindows.cpp
@@ -599,14 +599,16 @@ INSTRUMENTATION_FUNCTION_ATTRIBUTES PVOID NTAPI DetouredRtlReallocateHeap(PVOID
 
 	if (lpMem)
 	{
-		ULONG OldSize = HeapSize(hHeap, 0, lpMem);
-
-		if (NewPtr)
+		const SIZE_T OldSize = HeapSize(hHeap, 0, lpMem);
+		if (OldSize > 0)
 		{
-			FMemory::Memcpy(NewPtr, lpMem, FMath::Min(OldSize, dwBytes));
-		}
+			if (NewPtr)
+			{
+				FMemory::Memcpy(NewPtr, lpMem, FMath::Min((ULONG)OldSize, dwBytes));
+			}
 
-		FreeMemoryRange(lpMem, OldSize);
+			FreeMemoryRange(lpMem, OldSize);
+		}
 		TrueRtlFreeHeap(hHeap, dwFlags, lpMem);
 	}
 	
@@ -1032,6 +1034,16 @@ namespace UE::Sanitizer::RaceDetector::Platform {
 		return ::RtlCaptureStackBackTrace(FrameToSkip, FrameToCapture, Backtrace, 0);
 	}
 
+	INSTRUMENTATION_FUNCTION_ATTRIBUTES static void ThrottleMemoryAccessesDuringShadowReset()
+	{
+		// If we're resetting the shadow memory, we delay resuming threads until the reset is done.
+		// because we don't want to mess with the bitmap while it's being cleaned up.
+		while (bIsResettingShadow)
+		{
+			Sleep(0);
+		}
+	}
+
 	// This is how shadow memory typically works, you just commit the page when there's a page fault.
 	// But under the debugger, the kernel will send events to the debugger first and all these
 	// events are taking a lot of time to process...Waaaayyy too much time. 
@@ -1046,12 +1058,7 @@ namespace UE::Sanitizer::RaceDetector::Platform {
 			// Only handle exception inside the shadow memory address space.
 			if (Address >= ShadowBase && Address < ShadowBase + ShadowSize)
 			{
-				// If we're resetting the shadow memory, we delay resuming threads until the reset is done.
-				// because we don't want to mess with the bitmap while it's being cleaned up.
-				while (bIsResettingShadow)
-				{
-					Sleep(0);
-				}
+				ThrottleMemoryAccessesDuringShadowReset();
 
 				uintptr_t BaseAddress = AlignDown(Address, PageSize);
 				uintptr_t Result = (uintptr_t)::VirtualAlloc((LPVOID)BaseAddress, PageSize, MEM_COMMIT, PAGE_READWRITE);
@@ -1104,6 +1111,8 @@ namespace UE::Sanitizer::RaceDetector::Platform {
 
 		if (bNeedMapping)
 		{
+			ThrottleMemoryAccessesDuringShadowReset();
+
 			VirtualAlloc((void*)AlignedStart, AlignedEnd - AlignedStart, MEM_COMMIT, PAGE_READWRITE);
 
 			bHasShadowMemoryMapped = true;
