Add a formal way for code that runs very early in the engine initialization to find the name of a commandlet that will run.

- Added ::GetCommandletNameFromCmdline which returns the name of the commandlet that will run once the editor process reaches that point. This allows systems to query the commandlet name even if they are run very early on in the engine initialization. However it is still preferable to use ::GetRunningCommandletClass where possible.
- Fixed code documentation for ::GetRunningCommandletClass.

diff --git a/Engine/Source/Runtime/Core/Private/Misc/CoreGlobals.cpp b/Engine/Source/Runtime/Core/Private/Misc/CoreGlobals.cpp
--- a/Engine/Source/Runtime/Core/Private/Misc/CoreGlobals.cpp
+++ b/Engine/Source/Runtime/Core/Private/Misc/CoreGlobals.cpp
@@ -975,6 +975,17 @@ bool IsRunningCookOnTheFly()
 #endif
 }
 
+CORE_API FString GetCommandletNameFromCmdline()
+{
+	FString CommandletName;
+
+#if WITH_ENGINE
+	FParse::Value(FCommandLine::Get(), TEXT("-run="), CommandletName);
+#endif
+
+	return CommandletName;
+}
+
 namespace UE
 {
 
diff --git a/Engine/Source/Runtime/Core/Public/CoreGlobals.h b/Engine/Source/Runtime/Core/Public/CoreGlobals.h
--- a/Engine/Source/Runtime/Core/Public/CoreGlobals.h
+++ b/Engine/Source/Runtime/Core/Public/CoreGlobals.h
@@ -277,8 +277,8 @@ FORCEINLINE bool IsRunningDLCCookCommandlet()
 #endif
 }
 
-/** Returns running commandlet name 
- * 
+/**
+ * Returns the UClass of the currently running commandlet.
  */
 FORCEINLINE UClass* GetRunningCommandletClass()
 {
@@ -289,6 +289,20 @@ FORCEINLINE UClass* GetRunningCommandletClass()
 #endif
 }
 
+/**
+ * Returns the name of the commandlet that was set on the command line by
+ * the arg '-Run=<CommandletName>'.
+ * 
+ * This should be used by code that needs to know the commandlet name
+ * early on during engine startup, before the commandlet has been initialized.
+ * If possible, code should prefer calling GetRunningCommandletClass()->GetName()
+ * instead.
+ * 
+ * @return The name of the commandlet that will end up running. If there is
+ * no commandlet then an empty string will be returned.
+ */
+CORE_API FString GetCommandletNameFromCmdline();
+
 /**
  * Check to see if we should initialise RHI and set up scene for rendering even when running a commandlet.
  */
