[RaceDetector]
- Remove tracing scopes when inside race report to avoid reentrancy and trace corruption

diff --git a/Engine/Source/Runtime/Core/Internal/Sanitizer/RaceDetectorInterface.h b/Engine/Source/Runtime/Core/Internal/Sanitizer/RaceDetectorInterface.h
--- a/Engine/Source/Runtime/Core/Internal/Sanitizer/RaceDetectorInterface.h
+++ b/Engine/Source/Runtime/Core/Internal/Sanitizer/RaceDetectorInterface.h
@@ -12,69 +12,6 @@
 DECLARE_LOG_CATEGORY_EXTERN(LogRaceDetector, Log, All);
 
 namespace UE::Sanitizer::RaceDetector {
-	namespace Platform {
-		// Performs any platform specific initialization.
-		bool  InitializePlatform();
-		// Performs any platform specific cleanup.
-		bool  CleanupPlatform();
-		// Prepare address space of shadow memory
-		void  InitShadowMemory();
-		// Returns the command line the process was started with.
-		const TCHAR* GetCommandLine();
-		// Returns the page size granularity.
-		UPTRINT GetPageSize();
-		// Returns the base address of the shadow memory address space. 
-		// Note: InitShadowMemory must have been called before, otherwise this is meaningless.
-		UPTRINT GetShadowMemoryBase();
-		// Returns the size of the entire shadow memory address space.
-		// Note: InitShadowMemory must have been called before, otherwise this is meaningless.
-		UPTRINT GetShadowMemorySize();
-		// Get the base of the shadow memory for clock banks.
-		UPTRINT GetShadowClockBase();
-		// Returns the number of bytes currently mapped in the shadow memory.
-		uint64 GetShadowMemoryUsage();
-		// Returns whether a particular range is already accessible in shadow memory.
-		bool IsShadowMemoryMapped(UPTRINT Base, UPTRINT Size);
-		// Maps a range in shadow memory so that it is safe to access.
-		void MapShadowMemory(UPTRINT Base, UPTRINT Size);
-		// Unmaps the entire range of shadow memory.
-		void UnmapShadowMemory();
-		// Returns whether a debugger is currently attached to our process.
-		bool IsDebuggerPresent();
-		// Returns whether there is any page currently mapped in shadow memory.
-		bool HasShadowMemoryMapped();
-		// Sends a hint to visual studio to hide first chance exception to reduce noise caused by shadow memory access.
-		void HideFirstChanceExceptionInVisualStudio();
-		// Sleeps for the given amount of milliseconds.
-		void SleepMS(uint32 Milliseconds);
-		// Capture the current callstack.
-		uint16 CaptureStackBackTrace(uint32 FrameToSkip, uint32 FrameToCapture, void** Backtrace);
-		// Allocates a Tls index.
-		uint32 AllocTlsSlot();
-		// Releases a Tls index.
-		void FreeTlsSlot(uint32 Index);
-		// Gets the value of the Tls index for the current thread.
-		void* GetTlsValue(uint32 Index);
-		// Sets the value of the Tls index for the current thread.
-		void SetTlsValue(uint32 Index, void* Value);
-		// Returns the current thread Id.
-		uint32 GetCurrentThreadId();
-
-		// Rewrites the patchable function prefix with the proper jmp to the target.
-		void PrepareTrampoline(void* PatchableFunctionAddress, void* TargetFunctionAddress, bool bUseRETBypass);
-		// Rewrites the patchable function prefix back to its original compiled NOPs.
-		void CleanupTrampoline(void* PatchableFunctionAddress);
-		// Rewrites the 2 first bytes of the function to jump at the beginning of the prefix section.
-		void ActivateTrampoline(void* PatchableFunctionAddress);
-		// Rewrites the 2 first bytes of the function to either do nothing or do an immediate return (RET bypass).
-		void DeactivateTrampoline(void* PatchableFunctionAddress, bool bUseRETBypass);
-		// This needs to be called after trampoline activation or deactivation to make sure it takes effect immediately.
-		void FlushInstructionCache();
-		// Get the limits of the stack for the current thread.
-		void GetCurrentThreadStackLimits(void** LowLimit, void** HighLimit);
-		// Check if the thread id given is currently alive
-		bool IsThreadAlive(uint32 ThreadId);
-	}
 
 	// Returns whether we should instrument depending on the current context state.
 	bool ShouldInstrument(FContext& Context);
diff --git a/Engine/Source/Runtime/Core/Internal/Sanitizer/RaceDetectorPlatform.h b/Engine/Source/Runtime/Core/Internal/Sanitizer/RaceDetectorPlatform.h
new file mode 100644
--- /dev/null
+++ b/Engine/Source/Runtime/Core/Internal/Sanitizer/RaceDetectorPlatform.h
@@ -0,0 +1,86 @@
+// Copyright Epic Games, Inc. All Rights Reserved.
+
+#pragma once
+
+#include "Misc/CoreMiscDefines.h"
+
+#if USING_INSTRUMENTATION
+
+#include "CoreTypes.h"
+
+namespace UE::Sanitizer::RaceDetector {
+	namespace Platform {
+		// Performs any platform specific initialization.
+		bool  InitializePlatform();
+		// Performs any platform specific cleanup.
+		bool  CleanupPlatform();
+		// Prepare address space of shadow memory
+		void  InitShadowMemory();
+		// Returns the command line the process was started with.
+		const TCHAR* GetCommandLine();
+		// Returns the page size granularity.
+		UPTRINT GetPageSize();
+		// Returns the base address of the shadow memory address space. 
+		// Note: InitShadowMemory must have been called before, otherwise this is meaningless.
+		UPTRINT GetShadowMemoryBase();
+		// Returns the size of the entire shadow memory address space.
+		// Note: InitShadowMemory must have been called before, otherwise this is meaningless.
+		UPTRINT GetShadowMemorySize();
+		// Get the base of the shadow memory for clock banks.
+		UPTRINT GetShadowClockBase();
+		// Returns the number of bytes currently mapped in the shadow memory.
+		uint64 GetShadowMemoryUsage();
+		// Returns whether a particular range is already accessible in shadow memory.
+		bool IsShadowMemoryMapped(UPTRINT Base, UPTRINT Size);
+		// Maps a range in shadow memory so that it is safe to access.
+		void MapShadowMemory(UPTRINT Base, UPTRINT Size);
+		// Unmaps the entire range of shadow memory.
+		void UnmapShadowMemory();
+		// Returns whether a debugger is currently attached to our process.
+		bool IsDebuggerPresent();
+		// Returns whether there is any page currently mapped in shadow memory.
+		bool HasShadowMemoryMapped();
+		// Sends a hint to visual studio to hide first chance exception to reduce noise caused by shadow memory access.
+		void HideFirstChanceExceptionInVisualStudio();
+		// Sleeps for the given amount of milliseconds.
+		void SleepMS(uint32 Milliseconds);
+		// Capture the current callstack.
+		uint16 CaptureStackBackTrace(uint32 FrameToSkip, uint32 FrameToCapture, void** Backtrace);
+		// Allocates a Tls index.
+		uint32 AllocTlsSlot();
+		// Releases a Tls index.
+		void FreeTlsSlot(uint32 Index);
+		// Gets the value of the Tls index for the current thread.
+		void* GetTlsValue(uint32 Index);
+		// Sets the value of the Tls index for the current thread.
+		void SetTlsValue(uint32 Index, void* Value);
+		// Returns the current thread Id.
+		uint32 GetCurrentThreadId();
+
+		// Rewrites the patchable function prefix with the proper jmp to the target.
+		void PrepareTrampoline(void* PatchableFunctionAddress, void* TargetFunctionAddress, bool bUseRETBypass);
+		// Rewrites the patchable function prefix back to its original compiled NOPs.
+		void CleanupTrampoline(void* PatchableFunctionAddress);
+		// Rewrites the 2 first bytes of the function to jump at the beginning of the prefix section.
+		void ActivateTrampoline(void* PatchableFunctionAddress);
+		// Rewrites the 2 first bytes of the function to either do nothing or do an immediate return (RET bypass).
+		void DeactivateTrampoline(void* PatchableFunctionAddress, bool bUseRETBypass);
+		// This needs to be called after trampoline activation or deactivation to make sure it takes effect immediately.
+		void FlushInstructionCache();
+		// Get the limits of the stack for the current thread.
+		void GetCurrentThreadStackLimits(void** LowLimit, void** HighLimit);
+		// Check if the thread id given is currently alive
+		bool IsThreadAlive(uint32 ThreadId);
+		// Provides a compilation barrier
+		FORCEINLINE void AsymmetricThreadFenceLight()
+		{
+			_ReadWriteBarrier();
+		}
+
+		// The function generates an interprocessor interrupt (IPI) to all processors that are part of the current process affinity.
+		// It guarantees the visibility of write operations performed on one processor to the other processors.
+		void AsymmetricThreadFenceHeavy();
+	}
+}
+
+#endif // USING_INSTRUMENTATION
\ No newline at end of file
diff --git a/Engine/Source/Runtime/Core/Internal/Sanitizer/RaceDetectorTypes.h b/Engine/Source/Runtime/Core/Internal/Sanitizer/RaceDetectorTypes.h
--- a/Engine/Source/Runtime/Core/Internal/Sanitizer/RaceDetectorTypes.h
+++ b/Engine/Source/Runtime/Core/Internal/Sanitizer/RaceDetectorTypes.h
@@ -9,6 +9,7 @@
 #include "CoreTypes.h"
 
 #include "Sanitizer/Types.h"
+#include "Sanitizer/RaceDetectorPlatform.h"
 #include "Instrumentation/Types.h"
 #include "Instrumentation/Containers.h"
 
@@ -537,7 +538,7 @@ namespace UE::Sanitizer::RaceDetector {
 
 			new (Tail->Buffer + Tail->Offset) EntryType(Forward<ArgsType>(Args)...);
 
-			FPlatformMisc::AsymmetricThreadFenceLight();
+			Platform::AsymmetricThreadFenceLight();
 			Tail->Offset += sizeof(EntryType);
 		}
 
@@ -562,7 +563,7 @@ namespace UE::Sanitizer::RaceDetector {
 			TReadScopeLock Scope(Lock);
 
 			// Make sure we sync with the light fence.
-			FPlatformMisc::AsymmetricThreadFenceHeavy();
+			Platform::AsymmetricThreadFenceHeavy();
 
 			OutClockRange.First = MAX_uint32;
 			OutClockRange.Last  = 0;
diff --git a/Engine/Source/Runtime/Core/Private/Sanitizer/RaceDetector.cpp b/Engine/Source/Runtime/Core/Private/Sanitizer/RaceDetector.cpp
--- a/Engine/Source/Runtime/Core/Private/Sanitizer/RaceDetector.cpp
+++ b/Engine/Source/Runtime/Core/Private/Sanitizer/RaceDetector.cpp
@@ -426,7 +426,7 @@ namespace UE::Sanitizer::RaceDetector {
 			}
 
 			Context.BankHazard = LocalBank;
-			FPlatformMisc::AsymmetricThreadFenceLight();
+			Platform::AsymmetricThreadFenceLight();
 
 			// Confirm that the entry is still in the shadow slot, otherwise we iterate again as our
 			// hazard pointer protection wouldn't be guaranteed.
@@ -999,7 +999,7 @@ namespace UE::Sanitizer::RaceDetector {
 	INSTRUMENTATION_FUNCTION_ATTRIBUTES TSet<FSyncObjectBank*> GatherBankHazards()
 	{
 		// Make sure all threads have published their BankHazard
-		FPlatformMisc::AsymmetricThreadFenceHeavy();
+		Platform::AsymmetricThreadFenceHeavy();
 
 		TSet<FSyncObjectBank*> Hazards;
 		TReadScopeLock Scope(ContextMappingLock);
@@ -1503,6 +1503,11 @@ namespace UE::Sanitizer::RaceDetector {
 		GRaceDetectorThreadFilter = Platform::GetCurrentThreadId();
 	}
 
+	INSTRUMENTATION_FUNCTION_ATTRIBUTES bool IsActive()
+	{
+		return bRaceDetectorActive;
+	}
+
 	INSTRUMENTATION_FUNCTION_ATTRIBUTES void ToggleRaceDetection(bool bEnable)
 	{
 		if (bRaceDetectorActive == bEnable)
@@ -2180,6 +2185,9 @@ INSTRUMENTATION_FUNCTION_ATTRIBUTES void InitializeBoot()
 {
 	ContextTLSIndex = Platform::AllocTlsSlot();
 	EnsureCurrentContext();
+
+	// Avoid doing that while reporting a race as it may cause reentrancy issues because of tracing.
+	FPlatformStackWalk::InitStackWalking();
 }
 
 } // UE::Sanitizer::RaceDetector
diff --git a/Engine/Source/Runtime/Core/Private/Sanitizer/RaceDetectorWindows.cpp b/Engine/Source/Runtime/Core/Private/Sanitizer/RaceDetectorWindows.cpp
--- a/Engine/Source/Runtime/Core/Private/Sanitizer/RaceDetectorWindows.cpp
+++ b/Engine/Source/Runtime/Core/Private/Sanitizer/RaceDetectorWindows.cpp
@@ -870,6 +870,13 @@ namespace UE::Sanitizer::RaceDetector::Platform {
 		::TlsFree(Index);
 	}
 
+	INSTRUMENTATION_FUNCTION_ATTRIBUTES void AsymmetricThreadFenceHeavy()
+	{
+		// The function generates an interprocessor interrupt (IPI) to all processors that are part of the current process affinity.
+		// It guarantees the visibility of write operations performed on one processor to the other processors.
+		FlushProcessWriteBuffers();
+	}
+
 	INSTRUMENTATION_FUNCTION_ATTRIBUTES bool IsDebuggerPresent()
 	{
 		return ::IsDebuggerPresent();
diff --git a/Engine/Source/Runtime/Core/Public/Sanitizer/RaceDetector.h b/Engine/Source/Runtime/Core/Public/Sanitizer/RaceDetector.h
--- a/Engine/Source/Runtime/Core/Public/Sanitizer/RaceDetector.h
+++ b/Engine/Source/Runtime/Core/Public/Sanitizer/RaceDetector.h
@@ -46,6 +46,8 @@ namespace UE::Sanitizer::RaceDetector
 	// Toggle detailed logging filter to only log activity for a specific address.
 	// NOTE: Set to nullptr to remove filtering.
 	CORE_API void ToggleFilterDetailedLogOnAddress(void* Ptr);
+	// Returns whether the race detector is currently active or not.
+	CORE_API bool IsActive();
 	
 } // UE::Sanitizer::RaceDetector
 
