Fixing bHasBeenDefined uninitialized warning

#rnx
PR #12883

diff --git a/Engine/Source/Runtime/Core/Public/Misc/AutomationTest.h b/Engine/Source/Runtime/Core/Public/Misc/AutomationTest.h
--- a/Engine/Source/Runtime/Core/Public/Misc/AutomationTest.h
+++ b/Engine/Source/Runtime/Core/Public/Misc/AutomationTest.h
@@ -3150,7 +3150,8 @@ public:
 		: FAutomationTestBase(InName, bInComplexTask)
 		, DefaultTimeout(FTimespan::FromSeconds(30))
 		, bEnableSkipIfError(true)
-		, RootDefinitionScope(MakeShareable(new FSpecDefinitionScope()))		
+		, RootDefinitionScope(MakeShareable(new FSpecDefinitionScope()))
+		, bHasBeenDefined(false)
 	{
 		DefinitionScopeStack.Push(RootDefinitionScope.ToSharedRef());
 	}
