MovieSceneMusicTimeTransform

Perform music map time remapping on subsequences by implementing the MusicTimeTransform methods for Remap and InverseRemap

#rnx

diff --git a/Engine/Source/Runtime/Core/Public/Misc/MusicalTime.h b/Engine/Source/Runtime/Core/Public/Misc/MusicalTime.h
--- a/Engine/Source/Runtime/Core/Public/Misc/MusicalTime.h
+++ b/Engine/Source/Runtime/Core/Public/Misc/MusicalTime.h
@@ -72,6 +72,21 @@ struct FMusicalTime
 		return FractionalBar() - Other.FractionalBar();
 	}
 
+	static FMusicalTime FloorBar(const FMusicalTime& InMusicalTime)
+	{
+		FMusicalTime OutMusicalTime = InMusicalTime;
+		OutMusicalTime.TickInBar = 0;
+		return OutMusicalTime;
+	}
+
+	static FMusicalTime FloorBeat(const FMusicalTime& InMusicalTime)
+	{
+		FMusicalTime OutMusicalTime = InMusicalTime;
+		int32 WholeBeatInBar = OutMusicalTime.TickInBar / OutMusicalTime.TicksPerBeat;
+		OutMusicalTime.TickInBar = WholeBeatInBar * OutMusicalTime.TicksPerBeat;
+		return OutMusicalTime;
+	}
+	
 public:
 
 	friend FArchive& operator<<(FArchive& Ar, FMusicalTime& MusicalTime)
