[iOS] exposing HasRecordPermission to FIOSPlatformMisc

diff --git a/Engine/Source/Runtime/Core/Private/IOS/IOSPlatformMisc.cpp b/Engine/Source/Runtime/Core/Private/IOS/IOSPlatformMisc.cpp
--- a/Engine/Source/Runtime/Core/Private/IOS/IOSPlatformMisc.cpp
+++ b/Engine/Source/Runtime/Core/Private/IOS/IOSPlatformMisc.cpp
@@ -1289,6 +1289,11 @@ bool FIOSPlatformMisc::IsVoiceChatEnabled()
 	return [[IOSAppDelegate GetDelegate] IsVoiceChatEnabled];
 }
 
+bool FIOSPlatformMisc::HasRecordPermission()
+{
+	return [[IOSAppDelegate GetDelegate] HasRecordPermission];
+}
+
 void FIOSPlatformMisc::RegisterForRemoteNotifications()
 {
 	if (FApp::IsUnattended())
diff --git a/Engine/Source/Runtime/Core/Public/IOS/IOSPlatformMisc.h b/Engine/Source/Runtime/Core/Public/IOS/IOSPlatformMisc.h
--- a/Engine/Source/Runtime/Core/Public/IOS/IOSPlatformMisc.h
+++ b/Engine/Source/Runtime/Core/Public/IOS/IOSPlatformMisc.h
@@ -104,6 +104,7 @@ struct CORE_API FIOSPlatformMisc : public FApplePlatformMisc
 
 	static void EnableVoiceChat(bool bEnable);
 	static bool IsVoiceChatEnabled();
+	static bool HasRecordPermission();
 
 	//////// Platform specific
 	static int GetDefaultStackSize();
