Remove semicolons from UE_MT_DETECTOR macros so it's possible to declare an array of access detectors.

diff --git a/Engine/Plugins/Runtime/StateTree/Source/StateTreeModule/Private/StateTreeInstanceData.cpp b/Engine/Plugins/Runtime/StateTree/Source/StateTreeModule/Private/StateTreeInstanceData.cpp
--- a/Engine/Plugins/Runtime/StateTree/Source/StateTreeModule/Private/StateTreeInstanceData.cpp
+++ b/Engine/Plugins/Runtime/StateTree/Source/StateTreeModule/Private/StateTreeInstanceData.cpp
@@ -373,7 +373,7 @@ void FStateTreeInstanceStorage::AcquireReadAccess()
 {
 	if (UE::StateTree::Debug::Private::RuntimeValidationThreadingAccessDetector)
 	{
-		UE_MT_ACQUIRE_READ_ACCESS(AccessDetector)
+		UE_MT_ACQUIRE_READ_ACCESS(AccessDetector);
 	}
 }
 
@@ -381,7 +381,7 @@ void FStateTreeInstanceStorage::ReleaseReadAccess()
 {
 	if (UE::StateTree::Debug::Private::RuntimeValidationThreadingAccessDetector)
 	{
-		UE_MT_RELEASE_READ_ACCESS(AccessDetector)
+		UE_MT_RELEASE_READ_ACCESS(AccessDetector);
 	}
 }
 
@@ -389,7 +389,7 @@ void FStateTreeInstanceStorage::AcquireWriteAccess()
 {
 	if (UE::StateTree::Debug::Private::RuntimeValidationThreadingAccessDetector)
 	{
-		UE_MT_ACQUIRE_WRITE_ACCESS(AccessDetector)
+		UE_MT_ACQUIRE_WRITE_ACCESS(AccessDetector);
 	}
 }
 
@@ -397,7 +397,7 @@ void FStateTreeInstanceStorage::ReleaseWriteAccess()
 {
 	if (UE::StateTree::Debug::Private::RuntimeValidationThreadingAccessDetector)
 	{
-		UE_MT_RELEASE_WRITE_ACCESS(AccessDetector)
+		UE_MT_RELEASE_WRITE_ACCESS(AccessDetector);
 	}
 }
 
diff --git a/Engine/Plugins/Runtime/StateTree/Source/StateTreeModule/Public/StateTreeInstanceData.h b/Engine/Plugins/Runtime/StateTree/Source/StateTreeModule/Public/StateTreeInstanceData.h
--- a/Engine/Plugins/Runtime/StateTree/Source/StateTreeModule/Public/StateTreeInstanceData.h
+++ b/Engine/Plugins/Runtime/StateTree/Source/StateTreeModule/Public/StateTreeInstanceData.h
@@ -334,7 +334,7 @@ protected:
 	 * The instance data supports multiple reader threads or a single writer thread.
 	 * The detector supports recursive access.
 	 */
-	UE_MT_DECLARE_MRSW_RECURSIVE_ACCESS_DETECTOR(AccessDetector)
+	UE_MT_DECLARE_MRSW_RECURSIVE_ACCESS_DETECTOR(AccessDetector);
 
 	/* True if the storage owns the event queue. */
 	bool bIsOwningEventQueue = true;
diff --git a/Engine/Source/Runtime/Core/Public/Misc/MTAccessDetector.h b/Engine/Source/Runtime/Core/Public/Misc/MTAccessDetector.h
--- a/Engine/Source/Runtime/Core/Public/Misc/MTAccessDetector.h
+++ b/Engine/Source/Runtime/Core/Public/Misc/MTAccessDetector.h
@@ -667,18 +667,18 @@ public:
 
 ///////////////////////////////////////////////////////////////////
 
-#define UE_MT_DECLARE_RW_ACCESS_DETECTOR(AccessDetector) FRWAccessDetector AccessDetector;
-#define UE_MT_DECLARE_RW_RECURSIVE_ACCESS_DETECTOR(AccessDetector) FRWRecursiveAccessDetector AccessDetector;
-#define UE_MT_DECLARE_RW_FULLY_RECURSIVE_ACCESS_DETECTOR(AccessDetector) FRWFullyRecursiveAccessDetector AccessDetector;
-#define UE_MT_DECLARE_MRSW_RECURSIVE_ACCESS_DETECTOR(AccessDetector) FMRSWRecursiveAccessDetector AccessDetector;
-
-#define UE_MT_SCOPED_READ_ACCESS(AccessDetector) const FBaseScopedAccessDetector& PREPROCESSOR_JOIN(ScopedMTAccessDetector_,__LINE__) = MakeScopedReaderAccessDetector(AccessDetector);
-#define UE_MT_SCOPED_WRITE_ACCESS(AccessDetector) const FBaseScopedAccessDetector& PREPROCESSOR_JOIN(ScopedMTAccessDetector_,__LINE__) = MakeScopedWriterAccessDetector(AccessDetector);
-
-#define UE_MT_ACQUIRE_READ_ACCESS(AccessDetector) (AccessDetector).AcquireReadAccess();
-#define UE_MT_RELEASE_READ_ACCESS(AccessDetector) (AccessDetector).ReleaseReadAccess();
-#define UE_MT_ACQUIRE_WRITE_ACCESS(AccessDetector) (AccessDetector).AcquireWriteAccess();
-#define UE_MT_RELEASE_WRITE_ACCESS(AccessDetector) (AccessDetector).ReleaseWriteAccess();
+#define UE_MT_DECLARE_RW_ACCESS_DETECTOR(AccessDetector) FRWAccessDetector AccessDetector
+#define UE_MT_DECLARE_RW_RECURSIVE_ACCESS_DETECTOR(AccessDetector) FRWRecursiveAccessDetector AccessDetector
+#define UE_MT_DECLARE_RW_FULLY_RECURSIVE_ACCESS_DETECTOR(AccessDetector) FRWFullyRecursiveAccessDetector AccessDetector
+#define UE_MT_DECLARE_MRSW_RECURSIVE_ACCESS_DETECTOR(AccessDetector) FMRSWRecursiveAccessDetector AccessDetector
+
+#define UE_MT_SCOPED_READ_ACCESS(AccessDetector) const FBaseScopedAccessDetector& PREPROCESSOR_JOIN(ScopedMTAccessDetector_,__LINE__) = MakeScopedReaderAccessDetector(AccessDetector)
+#define UE_MT_SCOPED_WRITE_ACCESS(AccessDetector) const FBaseScopedAccessDetector& PREPROCESSOR_JOIN(ScopedMTAccessDetector_,__LINE__) = MakeScopedWriterAccessDetector(AccessDetector)
+
+#define UE_MT_ACQUIRE_READ_ACCESS(AccessDetector) (AccessDetector).AcquireReadAccess()
+#define UE_MT_RELEASE_READ_ACCESS(AccessDetector) (AccessDetector).ReleaseReadAccess()
+#define UE_MT_ACQUIRE_WRITE_ACCESS(AccessDetector) (AccessDetector).AcquireWriteAccess()
+#define UE_MT_RELEASE_WRITE_ACCESS(AccessDetector) (AccessDetector).ReleaseWriteAccess()
 
 #else // ENABLE_MT_DETECTOR
 
