Added NotNullGet for retrieving the inner value from a TNotNull without an explicit cast.

diff --git a/Engine/Source/Runtime/Core/Public/Misc/NotNull.h b/Engine/Source/Runtime/Core/Public/Misc/NotNull.h
--- a/Engine/Source/Runtime/Core/Public/Misc/NotNull.h
+++ b/Engine/Source/Runtime/Core/Public/Misc/NotNull.h
@@ -260,6 +260,18 @@
 		}
 	}
 
+	// Gets the inner value of a TNotNull without requiring a cast and repeat its inner type.
+	template <typename T>
+	[[nodiscard]] T NotNullGet(const TNotNull<T>& NotNull)
+	{
+		return NotNull;
+	}
+	template <typename T>
+	[[nodiscard]] T NotNullGet(TNotNull<T>&& NotNull)
+	{
+		return (TNotNull<T>&&)NotNull;
+	}
+
 	// Allow comparison, if it's supported
 	template <typename T, typename U>
 	[[nodiscard]] auto operator==(const TNotNull<T>& Lhs, const TNotNull<U>& Rhs) -> decltype((T)Lhs == (U)Rhs)
@@ -308,4 +320,11 @@
 	template <typename T>
 	using TNotNull = T;
 
+	// Gets the inner value of a TNotNull without requiring a cast and repeat its inner type.
+	template <typename T>
+	[[nodiscard]] T&& NotNullGet(T&& NotNull)
+	{
+		return (T&&)NotNull;
+	}
+
 #endif // !UE_ENABLE_NOTNULL_WRAPPER
