[ue5] - Use a more standard and predictable Quat->Rotator conversion when gimbal locked at +/-90 degrees pitch. This handles very minor differences between platforms and compilers in the 12th and higher decimal places causing divergence between Atan2(0, X) results when X is very small negative or very small positive numbers choosing opposite Yaw directions. This was previously addressed with CLs 3873301 and 3950302 but was rolled back, I think because it wasn't being clamped to (-180,180) as expected, and because at that time atan2 and fmod functions were not as stable causing other issues to propagate.

#tests PIE and relevant automation tests with MS and clang compilers.

diff --git a/Engine/Source/Runtime/Core/Private/Math/UnrealMath.cpp b/Engine/Source/Runtime/Core/Private/Math/UnrealMath.cpp
--- a/Engine/Source/Runtime/Core/Private/Math/UnrealMath.cpp
+++ b/Engine/Source/Runtime/Core/Private/Math/UnrealMath.cpp
@@ -658,14 +658,14 @@ namespace Math
 		if (SingularityTest < -SINGULARITY_THRESHOLD)
 		{
 			Pitch = -90.f;
-			Yaw = (FMath::Atan2(YawY, YawX) * RAD_TO_DEG);
-			Roll = FRotator3f::NormalizeAxis(-Yaw - (2.f * FMath::Atan2(X, W) * RAD_TO_DEG));
+			Yaw = FRotator3f::NormalizeAxis(-2.0f * FMath::Atan2(X, W) * RAD_TO_DEG);
+			Roll = 0.f;
 		}
 		else if (SingularityTest > SINGULARITY_THRESHOLD)
 		{
 			Pitch = 90.f;
-			Yaw = (FMath::Atan2(YawY, YawX) * RAD_TO_DEG);
-			Roll = FRotator3f::NormalizeAxis(Yaw - (2.f * FMath::Atan2(X, W) * RAD_TO_DEG));
+			Yaw = FRotator3f::NormalizeAxis(2.0f * FMath::Atan2(X, W) * RAD_TO_DEG);
+			Roll = 0.f;
 		}
 		else
 		{
@@ -702,14 +702,14 @@ namespace Math
 		if (SingularityTest < -SINGULARITY_THRESHOLD)
 		{
 			Pitch = -90.0;
-			Yaw = (FMath::Atan2(YawY, YawX) * RAD_TO_DEG);
-			Roll = FRotator3d::NormalizeAxis(-Yaw - (2.0 * FMath::Atan2(X, W) * RAD_TO_DEG));
+			Yaw = FRotator3d::NormalizeAxis(-2.0 * FMath::Atan2(X, W) * RAD_TO_DEG);
+			Roll = 0.0;
 		}
 		else if (SingularityTest > SINGULARITY_THRESHOLD)
 		{
 			Pitch = 90.0;
-			Yaw = (FMath::Atan2(YawY, YawX) * RAD_TO_DEG);
-			Roll = FRotator3d::NormalizeAxis(Yaw - (2.0 * FMath::Atan2(X, W) * RAD_TO_DEG));
+			Yaw = FRotator3d::NormalizeAxis(2.0 * FMath::Atan2(X, W) * RAD_TO_DEG);
+			Roll = 0.0;
 		}
 		else
 		{
diff --git a/Engine/Source/Runtime/Core/Tests/Math/UnrealMathTest.cpp b/Engine/Source/Runtime/Core/Tests/Math/UnrealMathTest.cpp
--- a/Engine/Source/Runtime/Core/Tests/Math/UnrealMathTest.cpp
+++ b/Engine/Source/Runtime/Core/Tests/Math/UnrealMathTest.cpp
@@ -689,15 +689,15 @@ FORCENOINLINE FRotator3f TestQuaternionToRotator(const FQuat4f& Quat)
 	// and we don't want fast approximations to be used here.
 	if (SingularityTest < -SINGULARITY_THRESHOLD)
 	{
-		RotatorFromQuat.Pitch = 270.f;
-		RotatorFromQuat.Yaw = atan2f(YawY, YawX) * RAD_TO_DEG;
-		RotatorFromQuat.Roll = -RotatorFromQuat.Yaw - (2.f * atan2f(X, W) * RAD_TO_DEG);
+		RotatorFromQuat.Pitch = -90.f;
+		RotatorFromQuat.Yaw = (-2.0f * atan2f(X, W) * RAD_TO_DEG);
+		RotatorFromQuat.Roll = 0.f;
 	}
 	else if (SingularityTest > SINGULARITY_THRESHOLD)
 	{
 		RotatorFromQuat.Pitch = 90.f;
-		RotatorFromQuat.Yaw = atan2f(YawY, YawX) * RAD_TO_DEG;
-		RotatorFromQuat.Roll = RotatorFromQuat.Yaw - (2.f * atan2f(X, W) * RAD_TO_DEG);
+		RotatorFromQuat.Yaw = (2.0f * atan2f(X, W) * RAD_TO_DEG);
+		RotatorFromQuat.Roll = 0.f;
 	}
 	else
 	{
