- Fixed a deprecated header messaging

diff --git a/Engine/Source/Runtime/Core/Public/IO/PackageStore.h b/Engine/Source/Runtime/Core/Public/IO/PackageStore.h
--- a/Engine/Source/Runtime/Core/Public/IO/PackageStore.h
+++ b/Engine/Source/Runtime/Core/Public/IO/PackageStore.h
@@ -2,5 +2,10 @@
 
 #pragma once
 
-#include "Misc/CoreMiscDefines.h"
-EMIT_CUSTOM_WARNING("PackageStore.h has moved to CoreUObject module.Use #include 'Serialization/PackageStore.h'")
\ No newline at end of file
+// HEADER_UNIT_SKIP - Deprecated
+
+#include "CoreTypes.h"
+
+UE_DEPRECATED_HEADER(5.6, "PackageStore has moved from Core to CoreUObject. #include 'Serialization/PackageStore.h' instead of 'IO/PackageStore.h'. Attempting to automatically include the correct one, but if the include fails, add a dependency on CoreUObject in your .Build.cs file, and update your include line.")
+
+#include "Serialization/PackageStore.h"
