IOS simulator MemoryConstants.AddressLimit fix
-In the simulator os_proc_available_memory returns zero, which resulted in asserts around AddressLimit, which was added recently.

diff --git a/Engine/Source/Runtime/Core/Private/Apple/ApplePlatformMemory.cpp b/Engine/Source/Runtime/Core/Private/Apple/ApplePlatformMemory.cpp
--- a/Engine/Source/Runtime/Core/Private/Apple/ApplePlatformMemory.cpp
+++ b/Engine/Source/Runtime/Core/Private/Apple/ApplePlatformMemory.cpp
@@ -481,7 +481,8 @@ const FPlatformMemoryConstants& FApplePlatformMemory::GetConstants()
 	{
 		// Gather platform memory constants.
 		int64 AvailablePhysical = 0;
-#if PLATFORM_IOS
+		
+#if PLATFORM_IOS && !WITH_IOS_SIMULATOR
 		AvailablePhysical = os_proc_available_memory();
 #else
 		int Mib[] = {CTL_HW, HW_MEMSIZE};
