TVariant: Add an uninitialized type constructor for the internal use case of in-place construction of non-default-constructible types during serialization.

#rnx

diff --git a/Engine/Source/Runtime/Core/Public/Misc/TVariant.h b/Engine/Source/Runtime/Core/Public/Misc/TVariant.h
--- a/Engine/Source/Runtime/Core/Public/Misc/TVariant.h
+++ b/Engine/Source/Runtime/Core/Public/Misc/TVariant.h
@@ -12,6 +12,16 @@
 template <typename T>
 struct TInPlaceType {};
 
+/**
+ * INTERNAL USE ONLY -- Do not use this type.
+ * Its internal use case is for in-place construction of non-default-constructible types during serialization.
+ */
+namespace UE::Core::Private
+{
+template <typename T>
+struct TUninitializedType {};
+}
+
 /**
  * A special tag that can be used as the first type in a TVariant parameter pack if none of the other types can be default-constructed.
  */
@@ -223,6 +233,26 @@ public:
 		return (SIZE_T)TypeIndex;
 	}
 
+	/** 
+	 * INTERNAL USE ONLY -- Do not call this constructor, it will put the variant in a bad state.
+	 * Its internal use case is for in-place construction of non-default-constructible types during serialization.
+	 *
+	 * Construct the TVariant to store the specified uninitialized element type.
+	 * The caller must unconditionally get a pointer to the element and construct an object of the right type in that position.
+	 * This is all totally exception-unsafe.
+	 * This relies on undefined behavior by dereferencing a pointer to an object (inside Get<U>) that doesn't exist,
+	 * and then turning it back to an pointer.
+	 * Any other attempts to use or destroy the variant before an object of the right type has been constructed is completely unsafe,
+	 * it may e.g. result in trying to destroy an object that is not there.
+	 */
+	template <typename U>
+	[[nodiscard]] explicit TVariant(UE::Core::Private::TUninitializedType<U>&&)
+	{
+		constexpr SIZE_T Index = UE::Core::Private::TParameterPackTypeIndex<U, T, Ts...>::Value;
+		static_assert(Index != (SIZE_T)-1, "The TVariant is not declared to store the specified type");
+		TypeIndex = (uint8)Index;
+	}
+
 private:
 #if UE_TVARIANT_TRIVIAL_DESTRUCTOR_USING_CONCEPTS
 	/** Index into the template parameter pack for the type held. */
