[Core] Extended TEnumRange to allow a first and last value

diff --git a/Engine/Source/Runtime/Core/Public/Misc/EnumRange.h b/Engine/Source/Runtime/Core/Public/Misc/EnumRange.h
--- a/Engine/Source/Runtime/Core/Public/Misc/EnumRange.h
+++ b/Engine/Source/Runtime/Core/Public/Misc/EnumRange.h
@@ -160,13 +160,44 @@ namespace NEnumRangePrivate
 	template <typename EnumType>
 	struct TEnumRange_Impl<EnumType, 0>
 	{
-		TEnumContiguousIterator<EnumType> begin() const { return TEnumContiguousIterator<EnumType>(TEnumRangeTraits<EnumType>::Begin); }
-		TEnumContiguousIterator<EnumType> end()   const { return TEnumContiguousIterator<EnumType>(TEnumRangeTraits<EnumType>::End  ); }
+	    using IntType = typename TEnumContiguousIterator<EnumType>::IntType;
+
+		FORCEINLINE explicit TEnumRange_Impl()
+			: BeginValue(TEnumRangeTraits<EnumType>::Begin)
+			, EndValue(TEnumRangeTraits<EnumType>::End)
+		{
+		}
+
+		/**
+		 * Note that this is a closed range, so it's not currently possible (or valid) to iterate over zero elements.
+		 * By specifying the same first and last value, it will iterate over that single value.
+		 */
+		FORCEINLINE explicit TEnumRange_Impl(const EnumType InFirstValue, const EnumType InLastValue)
+			: BeginValue(static_cast<IntType>(InFirstValue))
+			, EndValue(static_cast<IntType>(InLastValue) + 1)
+		{
+			checkf(
+				BeginValue < EndValue
+				&& BeginValue >= TEnumRangeTraits<EnumType>::Begin
+				&& EndValue <= TEnumRangeTraits<EnumType>::End,
+				TEXT("The provided first and last enum values must be within the available range specified in the ENUM_RANGE macro."));
+		}
+
+		TEnumContiguousIterator<EnumType> begin() const { return TEnumContiguousIterator<EnumType>(BeginValue); }
+		TEnumContiguousIterator<EnumType> end()   const { return TEnumContiguousIterator<EnumType>(EndValue  ); }
+
+	private:
+        const IntType BeginValue;
+        const IntType EndValue;
 	};
 
 	template <typename EnumType>
 	struct TEnumRange_Impl<EnumType, 1>
 	{
+		FORCEINLINE explicit TEnumRange_Impl() = default;
+
+		FORCEINLINE explicit TEnumRange_Impl(const EnumType InFirstValue, const EnumType InLastValue) = delete;
+
 		TEnumValueArrayIterator<EnumType> begin() const { return TEnumValueArrayIterator<EnumType>(TEnumRangeTraits<EnumType>::template GetPointer<void>(false)); }
 		TEnumValueArrayIterator<EnumType> end  () const { return TEnumValueArrayIterator<EnumType>(TEnumRangeTraits<EnumType>::template GetPointer<void>(true )); }
 	};
@@ -233,6 +264,9 @@ namespace UE::EnumFlags::Private
 template <typename EnumType>
 struct TEnumRange : NEnumRangePrivate::TEnumRange_Impl<EnumType, NEnumRangePrivate::TEnumRangeTraits<EnumType>::RangeType>
 {
+	using Super = NEnumRangePrivate::TEnumRange_Impl<EnumType, NEnumRangePrivate::TEnumRangeTraits<EnumType>::RangeType>;
+
+	using Super::Super;
 };
 
 /**
