ApplePlatformFile - Refactor to use common implementation FFileHandleRegistry to support more than the limit of readonly opened file.

#tests local cook & editor + preflight

diff --git a/Engine/Source/Runtime/Core/Private/Apple/ApplePlatformFile.cpp b/Engine/Source/Runtime/Core/Private/Apple/ApplePlatformFile.cpp
--- a/Engine/Source/Runtime/Core/Private/Apple/ApplePlatformFile.cpp
+++ b/Engine/Source/Runtime/Core/Private/Apple/ApplePlatformFile.cpp
@@ -11,6 +11,7 @@
 #include "CoreGlobals.h"
 #include "HAL/CriticalSection.h"
 #include "HAL/PlatformFile.h"
+#include "HAL/PlatformFileCommon.h"
 #include "HAL/PlatformTime.h"
 #include "Misc/ScopeLock.h"
 #include "ProfilingDebugging/PlatformFileTrace.h"
@@ -20,6 +21,8 @@
 #include "Async/MappedFileHandle.h"
 #include <sys/mman.h>
 
+DEFINE_LOG_CATEGORY_STATIC(LogApplePlatformFile, Log, All);
+
 // make an FTimeSpan object that represents the "epoch" for time_t (from a stat struct)
 const FDateTime MacEpoch(1970, 1, 1);
 
@@ -45,159 +48,142 @@ namespace
 		);
 	}
 }
+/**
+ * Apple version of the file handle registry
+ */
+class FAppleFileRegistry : public FFileHandleRegistry
+{
+public:
+	FAppleFileRegistry()
+		: FFileHandleRegistry(192)
+	{
+	}
 
-/** 
- * Mac file handle implementation which limits number of open files per thread. This
- * is to prevent running out of system file handles (250). Should not be neccessary when
- * using pak file (e.g., SHIPPING?) so not particularly optimized. Only manages
- * files which are opened READ_ONLY.
-**/
-#define MANAGE_FILE_HANDLES PLATFORM_MAC // !UE_BUILD_SHIPPING
+protected:
+	virtual FRegisteredFileHandle* PlatformInitialOpenFile(const TCHAR* Filename) override;
+	virtual bool PlatformReopenFile(FRegisteredFileHandle* Handle) override;
+	virtual void PlatformCloseFile(FRegisteredFileHandle* Handle) override;
+};
+static FAppleFileRegistry GFileRegistry;
+
+static int32 AppleOpenCommon(const TCHAR* Filename);
+static void AppleCloseCommon(int32 FileHandle);
 
-class CORE_API FFileHandleApple : public IFileHandle
+class CORE_API FFileHandleApple : public FRegisteredFileHandle
 {
 	enum {READWRITE_SIZE = 1024 * 1024};
 
 public:
 	FFileHandleApple(int32 InFileHandle, const TCHAR* InFilename, bool bIsReadOnly)
 		: FileHandle(InFileHandle)
-#if MANAGE_FILE_HANDLES
 		, Filename(InFilename)
-		, HandleSlot(-1)
 		, FileOffset(0)
 		, FileSize(0)
-#endif
-        , bReadOnly(bIsReadOnly)
+		, bReadOnly(bIsReadOnly)
 	{
 		check(FileHandle > -1);
 
-#if MANAGE_FILE_HANDLES
 		// Only files opened for read will be managed
 		if (bIsReadOnly)
 		{
-			ReserveSlot();
-			ActiveHandles[HandleSlot] = this;
 			struct stat FileInfo;
 			fstat(FileHandle, &FileInfo);
 			FileSize = FileInfo.st_size;
 		}
-#endif
 	}
+	
 	virtual ~FFileHandleApple()
 	{
-#if MANAGE_FILE_HANDLES
-		if( IsManaged() )
+		if (bReadOnly)
 		{
-			if( ActiveHandles[ HandleSlot ] == this )
-			{
-				TRACE_PLATFORMFILE_BEGIN_CLOSE(FileHandle);
-				flock(FileHandle, LOCK_UN | LOCK_NB);
-				int CloseResult = close(FileHandle);
-				if (CloseResult >= 0)
-				{
-					TRACE_PLATFORMFILE_END_CLOSE(FileHandle);
-				}
-				else
-				{
-					TRACE_PLATFORMFILE_FAIL_CLOSE(FileHandle);
-					UE_LOG(LogInit, Warning, TEXT("Failed to properly close readable file: %s with errno: %d: %s"),
-						*Filename, errno, UTF8_TO_TCHAR(strerror(errno)));
-				}
-				ActiveHandles[ HandleSlot ] = nullptr;
-			}
+			GFileRegistry.UnTrackAndCloseFile(this);
 		}
 		else
-#endif
 		{
-            if (!bReadOnly)
-            {
-                int Result = fsync(FileHandle);
-				if (Result < 0)
-				{
-					UE_LOG(LogInit, Error, TEXT("Failed to properly flush writable file with errno: %d: %s"),
-						errno, UTF8_TO_TCHAR(strerror(errno)));
-				}
-            }
-			TRACE_PLATFORMFILE_BEGIN_CLOSE(FileHandle);
-			flock(FileHandle, LOCK_UN | LOCK_NB);
-			int CloseResult = close(FileHandle);
-			if (CloseResult >= 0)
+			int Result = fsync(FileHandle);
+			if (Result < 0)
 			{
-				TRACE_PLATFORMFILE_END_CLOSE(FileHandle);
-			}
-			else
-			{
-				TRACE_PLATFORMFILE_FAIL_CLOSE(FileHandle);
-				UE_LOG(LogInit, Warning, TEXT("Failed to properly close file with errno: %d: %s"),
+				UE_LOG(LogApplePlatformFile, Error, TEXT("Failed to properly flush writable file with errno: %d: %s"),
 					errno, UTF8_TO_TCHAR(strerror(errno)));
 			}
+			
+			AppleCloseCommon(FileHandle);
 		}
+		
 		FileHandle = -1;
 	}
+	
 	virtual int64 Tell() override
 	{
-#if MANAGE_FILE_HANDLES
-		if( IsManaged() )
+		if (bReadOnly)
 		{
 			return FileOffset;
 		}
 		else
-#endif
 		{
 			check(IsValid());
 			return lseek(FileHandle, 0, SEEK_CUR);
 		}
 	}
+	
 	virtual bool Seek(int64 NewPosition) override
 	{
 		check(NewPosition >= 0);
 
-#if MANAGE_FILE_HANDLES
-		if( IsManaged() )
+		if (bReadOnly)
 		{
 			FileOffset = NewPosition >= FileSize ? FileSize - 1 : NewPosition;
-			return IsValid() && ActiveHandles[ HandleSlot ] == this ? lseek(FileHandle, FileOffset, SEEK_SET) != -1 : true;
+			return true;
 		}
 		else
-#endif
 		{
 			check(IsValid());
 			return lseek(FileHandle, NewPosition, SEEK_SET) != -1;
 		}
 	}
+	
 	virtual bool SeekFromEnd(int64 NewPositionRelativeToEnd = 0) override
 	{
 		check(NewPositionRelativeToEnd <= 0);
-
-#if MANAGE_FILE_HANDLES
-		if( IsManaged() )
+		
+		if (bReadOnly)
 		{
 			FileOffset = (NewPositionRelativeToEnd >= FileSize) ? 0 : ( FileSize + NewPositionRelativeToEnd - 1 );
-			return IsValid() && ActiveHandles[ HandleSlot ] == this ? lseek(FileHandle, FileOffset, SEEK_SET) != -1 : true;
+			return true;
 		}
 		else
-#endif
 		{
 			check(IsValid());
 			return lseek(FileHandle, NewPositionRelativeToEnd, SEEK_END) != -1;
 		}
 	}
+	
 	virtual bool Read(uint8* Destination, int64 BytesToRead) override
 	{
-#if MANAGE_FILE_HANDLES
-		if( IsManaged() )
+		if (bReadOnly)
 		{
-			ActivateSlot();
+			// Handle virtual file handles (only in read mode, write mode doesn't use the file handle registry)
+			FFileHandleRegistryReadTracker TrackRead(GFileRegistry, *this, true);
+			if (!TrackRead.IsValid())
+			{
+				return false;
+			}
+
+			// seek to the offset on seek? this matches console behavior more closely
+			if (lseek(FileHandle, FileOffset, SEEK_SET) == -1)
+			{
+				return false;
+			}
 			int64 BytesRead = ReadInternal(Destination, BytesToRead);
 			FileOffset += BytesRead;
 			return BytesRead == BytesToRead;
 		}
 		else
-#endif
 		{
 			return ReadInternal(Destination, BytesToRead) == BytesToRead;
 		}
 	}
+	
 	virtual bool ReadAt(uint8* Destination, int64 BytesToRead, int64 Offset) override
 	{
 		if (BytesToRead < 0 || Offset < 0)
@@ -210,13 +196,12 @@ public:
 			return true;
 		}
 
-#if MANAGE_FILE_HANDLES
-		if (IsManaged())
+		FFileHandleRegistryReadTracker TrackRead(GFileRegistry, *this, true);
+		if (bReadOnly && !TrackRead.IsValid())
 		{
-			ActivateSlot();
+			return false;
 		}
-#endif //MANAGE_FILE_HANDLES
-
+		
 		int64 TotalBytesRead = 0;
 		TRACE_PLATFORMFILE_BEGIN_READ(this, FileHandle, Offset, BytesToRead);
 
@@ -242,6 +227,7 @@ public:
 
 		return true;
 	}
+	
 	virtual bool Write(const uint8* Source, int64 BytesToWrite) override
 	{
 		check(IsValid());
@@ -274,12 +260,12 @@ public:
 	virtual bool Flush(const bool bFullFlush = false) override
 	{
 		check(IsValid());
-#if MANAGE_FILE_HANDLES
-		if (IsManaged())
+		
+		if (bReadOnly)
 		{
 			return false;
 		}
-#endif
+		
 		if (bFullFlush)
 		{
 			// OS X needs fcntl with F_FULLFSYNC to guarantee a full flush,
@@ -292,28 +278,26 @@ public:
 		// HFS+ apparently doesn't always write the updated file size when using fdatasync, so use fsync to be safe
 		return fsync(FileHandle) == 0;
 	}
+	
 	virtual bool Truncate(int64 NewSize) override
 	{
 		check(IsValid());
-#if MANAGE_FILE_HANDLES
-		if (IsManaged())
+		
+		if (bReadOnly)
 		{
 			return false;
 		}
-#endif
 		int Result = 0;
 		do { Result = ftruncate(FileHandle, NewSize); } while (Result < 0 && errno == EINTR);
 		return Result == 0;
 	}
 	virtual int64 Size() override
 	{
-#if MANAGE_FILE_HANDLES
-		if( IsManaged() )
+		if (bReadOnly)
 		{
 			return FileSize;
 		}
 		else
-#endif
 		{
 			struct stat FileInfo;
 			fstat(FileHandle, &FileInfo);
@@ -323,91 +307,6 @@ public:
 
 private:
 
-#if MANAGE_FILE_HANDLES
-	FORCEINLINE bool IsManaged()
-	{
-		return HandleSlot != -1;
-	}
-
-	void ActivateSlot()
-	{
-		if( IsManaged() )
-		{
-			if( ActiveHandles[ HandleSlot ] != this || (ActiveHandles[ HandleSlot ] && ActiveHandles[ HandleSlot ]->FileHandle == -1) )
-			{
-				ReserveSlot();
-
-				TRACE_PLATFORMFILE_BEGIN_OPEN(*Filename);
-				FileHandle = open(TCHAR_TO_UTF8(*Filename), O_RDONLY | O_SHLOCK | O_CLOEXEC);
-				if( FileHandle != -1 )
-				{
-					TRACE_PLATFORMFILE_END_OPEN(FileHandle);
-					lseek(FileHandle, FileOffset, SEEK_SET);
-					ActiveHandles[ HandleSlot ] = this;
-				}
-				else
-				{
-					TRACE_PLATFORMFILE_FAIL_OPEN(*Filename);
-				}
-			}
-			else
-			{
-				AccessTimes[ HandleSlot ] = FPlatformTime::Seconds();
-			}
-		}
-	}
-
-	void ReserveSlot()
-	{
-		HandleSlot = -1;
-
-		// Look for non-reserved slot
-		for( int32 i = 0; i < ACTIVE_HANDLE_COUNT; ++i )
-		{
-			if( ActiveHandles[ i ] == nullptr )
-			{
-				HandleSlot = i;
-				break;
-			}
-		}
-
-		// Take the oldest handle
-		if( HandleSlot == -1 )
-		{
-			int32 Oldest = 0;
-			for( int32 i = 1; i < ACTIVE_HANDLE_COUNT; ++i )
-			{
-				if( AccessTimes[ Oldest ] > AccessTimes[ i ] )
-				{
-					Oldest = i;
-				}
-			}
-
-			int32 OldestFileHandle = ActiveHandles[ Oldest ]->FileHandle;
-			TRACE_PLATFORMFILE_BEGIN_CLOSE(OldestFileHandle);
-			flock(OldestFileHandle, LOCK_UN | LOCK_NB);
-			int CloseResult = close(OldestFileHandle);
-#if PLATFORMFILETRACE_ENABLED
-			if (CloseResult >= 0)
-			{
-				TRACE_PLATFORMFILE_END_CLOSE(OldestFileHandle);
-			}
-			else
-			{
-				TRACE_PLATFORMFILE_FAIL_CLOSE(OldestFileHandle);
-			}
-#else
-			(void)CloseResult;
-#endif
-			ActiveHandles[ Oldest ]->FileHandle = -1;
-			HandleSlot = Oldest;
-		}
-
-		ActiveHandles[ HandleSlot ] = nullptr;
-		AccessTimes[ HandleSlot ] = FPlatformTime::Seconds();
-	}
-#endif
-
 	int64 ReadInternal(uint8* Destination, int64 BytesToRead)
 	{
 		check(IsValid());
@@ -446,40 +345,27 @@ private:
 
 	// Holds the internal file handle.
 	int32 FileHandle;
-
-#if MANAGE_FILE_HANDLES
+	
 	// Holds the name of the file that this handle represents. Kept around for possible reopen of file.
 	FString Filename;
 
-    // Most recent valid slot index for this handle; >=0 for handles which are managed.
-    int32 HandleSlot;
-
     // Current file offset; valid if a managed handle.
     int64 FileOffset;
 
     // Cached file size; valid if a managed handle.
     int64 FileSize;
-
-    // Each thread keeps a collection of active handles with access times.
-    static const int32 ACTIVE_HANDLE_COUNT = 192;
-    static __thread FFileHandleApple* ActiveHandles[ ACTIVE_HANDLE_COUNT ];
-    static __thread double AccessTimes[ ACTIVE_HANDLE_COUNT ];
-#endif
-
-    // Whether the file is read-only or permits writes.
-    bool bReadOnly;
-    
+	
+	// Whether the file is read-only or permits writes.
+	bool bReadOnly;
+		
 	FORCEINLINE bool IsValid()
 	{
 		return FileHandle != -1;
 	}
+	
+	friend class FAppleFileRegistry;
 };
 
-#if MANAGE_FILE_HANDLES
-__thread FFileHandleApple* FFileHandleApple::ActiveHandles[ FFileHandleApple::ACTIVE_HANDLE_COUNT ];
-__thread double FFileHandleApple::AccessTimes[ FFileHandleApple::ACTIVE_HANDLE_COUNT ];
-#endif
-
 static SIZE_T FileMappingAlignment = FPlatformMemory::GetConstants().PageSize;
 
 class FMappedFileRegion final : public IMappedFileRegion
@@ -764,44 +650,7 @@ ESymlinkResult FApplePlatformFile::IsSymlink(const TCHAR* Filename)
 
 IFileHandle* FApplePlatformFile::OpenRead(const TCHAR* Filename, bool bAllowWrite)
 {
-	TRACE_PLATFORMFILE_BEGIN_OPEN(Filename);
-	int32 Handle = open(TCHAR_TO_UTF8(*NormalizeFilename(Filename)), O_RDONLY | O_CLOEXEC);
-	if (Handle != -1)
-	{
-		TRACE_PLATFORMFILE_END_OPEN(Handle);
-#if PLATFORM_MAC && UE_EDITOR && !UE_BUILD_SHIPPING
-		// No blocking attempt shared lock, failure means we should not have opened the file for reading, protect against multiple instances and client/server versions
-		if(flock(Handle, LOCK_NB | LOCK_SH) != 0)
-		{
-			TRACE_PLATFORMFILE_BEGIN_CLOSE(Handle);
-			int CloseResult = close(Handle);
-#if PLATFORMFILETRACE_ENABLED
-			if (CloseResult >= 0)
-			{
-				TRACE_PLATFORMFILE_END_CLOSE(Handle);
-			}
-			else
-			{
-				TRACE_PLATFORMFILE_FAIL_CLOSE(Handle);
-			}
-#else
-			(void)CloseResult;
-#endif
-			return nullptr;
-		}
-#endif
-		
-#if MANAGE_FILE_HANDLES
-		return new FFileHandleApple(Handle, *NormalizeDirectory(Filename), true);
-#else
-		return new FFileHandleApple(Handle, Filename, true);
-#endif
-	}
-	else
-	{
-		TRACE_PLATFORMFILE_FAIL_OPEN(Filename);
-		return nullptr;
-	}
+	return GFileRegistry.InitialOpenFile(*NormalizeFilename(Filename));
 }
 
 IFileHandle* FApplePlatformFile::OpenWrite(const TCHAR* Filename, bool bAppend, bool bAllowRead)
@@ -817,8 +666,10 @@ IFileHandle* FApplePlatformFile::OpenWrite(const TCHAR* Filename, bool bAppend,
 		Flags |= O_WRONLY;
 	}
 	
-	TRACE_PLATFORMFILE_BEGIN_OPEN(Filename);
-	int32 Handle = open(TCHAR_TO_UTF8(*NormalizeFilename(Filename)), Flags, S_IRUSR | S_IWUSR | S_IRGRP | S_IWGRP | S_IROTH | S_IWOTH);
+	const FString NormalizedFilename = NormalizeFilename(Filename);
+
+	TRACE_PLATFORMFILE_BEGIN_OPEN(*NormalizedFilename);
+	int32 Handle = open(TCHAR_TO_UTF8(*NormalizedFilename), Flags, S_IRUSR | S_IWUSR | S_IRGRP | S_IWGRP | S_IROTH | S_IWOTH);
 	
 	if (Handle != -1)
 	{
@@ -857,11 +708,8 @@ IFileHandle* FApplePlatformFile::OpenWrite(const TCHAR* Filename, bool bAppend,
 			ftruncate(Handle, 0);
 		}
 
-#if MANAGE_FILE_HANDLES
 		FFileHandleApple* FileHandleApple = new FFileHandleApple(Handle, *NormalizeDirectory(Filename), false);
-#else
-		FFileHandleApple* FileHandleApple = new FFileHandleApple(Handle, Filename, false);
-#endif
+
 		if (bAppend)
 		{
 			FileHandleApple->SeekFromEnd(0);
@@ -870,14 +718,14 @@ IFileHandle* FApplePlatformFile::OpenWrite(const TCHAR* Filename, bool bAppend,
 	}
 	else
 	{
-		TRACE_PLATFORMFILE_FAIL_OPEN(Filename);
+		TRACE_PLATFORMFILE_FAIL_OPEN(*NormalizedFilename);
 		return nullptr;
 	}
 }
 
 FOpenMappedResult FApplePlatformFile::OpenMappedEx(const TCHAR* Filename, EOpenReadFlags OpenOptions, int64 MaximumSize)
 {
-	FString NormalizedFilename = NormalizeFilename(Filename);
+	const FString NormalizedFilename = NormalizeFilename(Filename);
 	
 	// check the read path
 	FString FinalPath = ConvertToPlatformPath(NormalizedFilename, false, false);
@@ -1051,3 +899,99 @@ int32 FApplePlatformFile::Stat(const TCHAR* Filename, struct stat* OutFileInfo)
 {
 	return stat(TCHAR_TO_UTF8(*NormalizeFilename(Filename)), OutFileInfo);
 }
+
+int32 AppleOpenCommon(const TCHAR* Filename)
+{
+	TRACE_PLATFORMFILE_BEGIN_OPEN(Filename);
+	int32 Handle = open(TCHAR_TO_UTF8(Filename), O_RDONLY | O_CLOEXEC);
+	if (Handle != -1)
+	{
+		TRACE_PLATFORMFILE_END_OPEN(Handle);
+#if PLATFORM_MAC && UE_EDITOR && !UE_BUILD_SHIPPING
+		// No blocking attempt shared lock, failure means we should not have opened the file for reading, protect against multiple instances and client/server versions
+		if(flock(Handle, LOCK_NB | LOCK_SH) != 0)
+		{
+			TRACE_PLATFORMFILE_BEGIN_CLOSE(Handle);
+			int CloseResult = close(Handle);
+#if PLATFORMFILETRACE_ENABLED
+			if (CloseResult >= 0)
+			{
+				TRACE_PLATFORMFILE_END_CLOSE(Handle);
+			}
+			else
+			{
+				TRACE_PLATFORMFILE_FAIL_CLOSE(Handle);
+			}
+#else
+			(void)CloseResult;
+#endif
+			return -1;
+		}
+#endif // PLATFORM_MAC && UE_EDITOR && !UE_BUILD_SHIPPING
+		return Handle;
+	}
+	else
+	{
+		TRACE_PLATFORMFILE_FAIL_OPEN(Filename);
+		return -1;
+	}
+}
+
+void AppleCloseCommon(int32 FileHandle)
+{
+	TRACE_PLATFORMFILE_BEGIN_CLOSE(FileHandle);
+	flock(FileHandle, LOCK_UN | LOCK_NB);
+	int CloseResult = close(FileHandle);
+	if (CloseResult >= 0)
+	{
+		TRACE_PLATFORMFILE_END_CLOSE(FileHandle);
+	}
+	else
+	{
+		TRACE_PLATFORMFILE_FAIL_CLOSE(FileHandle);
+		UE_LOG(LogInit, Warning, TEXT("Failed to properly close file with errno: %d: %s"),
+			errno, UTF8_TO_TCHAR(strerror(errno)));
+	}
+}
+
+FRegisteredFileHandle* FAppleFileRegistry::PlatformInitialOpenFile(const TCHAR* Filename)
+{
+	const int32 Handle = AppleOpenCommon(Filename);
+	if (Handle != -1)
+	{
+		return new FFileHandleApple(Handle, Filename, true);;
+	}
+	else
+	{
+		return nullptr;
+	}
+}
+
+bool FAppleFileRegistry::PlatformReopenFile(FRegisteredFileHandle* Handle)
+{
+	FFileHandleApple* AppleHandle = static_cast<FFileHandleApple*>(Handle);
+	
+	bool bSuccess = true;
+	AppleHandle->FileHandle = AppleOpenCommon(*(AppleHandle->Filename));
+	if (AppleHandle->FileHandle != -1)
+	{
+		if (lseek(AppleHandle->FileHandle, AppleHandle->FileOffset, SEEK_SET) == -1)
+		{
+			UE_LOG(LogApplePlatformFile, Warning, TEXT("Could not seek to the previous position on handle for file '%s'"), *(AppleHandle->Filename));
+			bSuccess = false;
+		}
+	}
+	else
+	{
+		UE_LOG(LogApplePlatformFile, Warning, TEXT("Could not reopen handle for file '%s'"), *(AppleHandle->Filename));
+		bSuccess = false;
+	}
+
+	return bSuccess;
+}
+
+void FAppleFileRegistry::PlatformCloseFile(FRegisteredFileHandle* Handle)
+{
+	FFileHandleApple* AppleHandle = static_cast<FFileHandleApple*>(Handle);
+	AppleCloseCommon(AppleHandle->FileHandle);
+}
