IoStoreWriter: Improve feedback when FIoStoreTocResource::Write fails.
#rnx

diff --git a/Engine/Source/Developer/IoStoreUtilities/Private/IoStoreWriter.cpp b/Engine/Source/Developer/IoStoreUtilities/Private/IoStoreWriter.cpp
--- a/Engine/Source/Developer/IoStoreUtilities/Private/IoStoreWriter.cpp
+++ b/Engine/Source/Developer/IoStoreUtilities/Private/IoStoreWriter.cpp
@@ -1285,7 +1285,7 @@ public:
 		}
 
 		TIoStatusOr<uint64> TocSize = FIoStoreTocResource::Write(*TocFilePath, TocResource, static_cast<uint32>(WriterSettings.CompressionBlockSize), WriterSettings.MaxPartitionSize, ContainerSettings);
-		check(TocSize.IsOk());
+		checkf(TocSize.IsOk(), TEXT("FIoStoreTocResource::Write failed with error %s"), *TocSize.Status().ToString());
 		
 		Result.ContainerId = ContainerSettings.ContainerId;
 		Result.ContainerName = FPaths::GetBaseFilename(TocFilePath);
diff --git a/Engine/Source/Runtime/Core/Private/IO/IoStore.cpp b/Engine/Source/Runtime/Core/Private/IO/IoStore.cpp
--- a/Engine/Source/Runtime/Core/Private/IO/IoStore.cpp
+++ b/Engine/Source/Runtime/Core/Private/IO/IoStore.cpp
@@ -1455,6 +1455,14 @@ FIoStatus FIoStoreTocResource::Read(const TCHAR* TocFilePath, EIoStoreTocReadOpt
 	return Status;
 }
 
+static FString FormatLastErrorMessage(const TCHAR* PrefixMessage)
+{
+	TCHAR Buffer[256];
+	uint32 LastError = FPlatformMisc::GetLastError();
+	return FString::Printf(TEXT("%s Error: (%u: %s)."), PrefixMessage, LastError,
+		FPlatformMisc::GetSystemErrorMessage(Buffer, UE_ARRAY_COUNT(Buffer), LastError));
+}
+
 TIoStatusOr<uint64> FIoStoreTocResource::Write(
 	const TCHAR* TocFilePath,
 	FIoStoreTocResource& TocResource,
@@ -1538,35 +1546,35 @@ TIoStatusOr<uint64> FIoStoreTocResource::Write(
 	// Header
 	if (!TocFileHandle->Write(reinterpret_cast<const uint8*>(&TocResource.Header), sizeof(FIoStoreTocHeader)))
 	{
-		return FIoStatus(EIoErrorCode::WriteError, TEXT("Failed to write TOC header"));
+		return FIoStatus(EIoErrorCode::WriteError, FormatLastErrorMessage(TEXT("Failed to write TOC header.")));
 	}
 
 	// Chunk IDs
 	if (!WriteArray(TocFileHandle.Get(), TocResource.ChunkIds))
 	{
-		return FIoStatus(EIoErrorCode::WriteError, TEXT("Failed to write chunk ids"));
+		return FIoStatus(EIoErrorCode::WriteError, FormatLastErrorMessage(TEXT("Failed to write chunk ids.")));
 	}
 
 	// Chunk offsets
 	if (!WriteArray(TocFileHandle.Get(), TocResource.ChunkOffsetLengths))
 	{
-		return FIoStatus(EIoErrorCode::WriteError, TEXT("Failed to write chunk offsets"));
+		return FIoStatus(EIoErrorCode::WriteError, FormatLastErrorMessage(TEXT("Failed to write chunk offsets.")));
 	}
 
 	// Chunk perfect hash map
 	if (!WriteArray(TocFileHandle.Get(), TocResource.ChunkPerfectHashSeeds))
 	{
-		return FIoStatus(EIoErrorCode::WriteError, TEXT("Failed to write chunk hash seeds"));
+		return FIoStatus(EIoErrorCode::WriteError, FormatLastErrorMessage(TEXT("Failed to write chunk hash seeds.")));
 	}
 	if (!WriteArray(TocFileHandle.Get(), TocResource.ChunkIndicesWithoutPerfectHash))
 	{
-		return FIoStatus(EIoErrorCode::WriteError, TEXT("Failed to write chunk indices without perfect hash"));
+		return FIoStatus(EIoErrorCode::WriteError, FormatLastErrorMessage(TEXT("Failed to write chunk indices without perfect hash.")));
 	}
 
 	// Compression blocks
 	if (!WriteArray(TocFileHandle.Get(), TocResource.CompressionBlocks))
 	{
-		return FIoStatus(EIoErrorCode::WriteError, TEXT("Failed to write chunk block entries"));
+		return FIoStatus(EIoErrorCode::WriteError, FormatLastErrorMessage(TEXT("Failed to write chunk block entries.")));
 	}
 
 	// Compression methods
@@ -1583,7 +1591,7 @@ TIoStatusOr<uint64> FIoStoreTocResource::Write(
 
 		if (!TocFileHandle->Write(reinterpret_cast<const uint8*>(AnsiMethodName), FIoStoreTocResource::CompressionMethodNameLen))
 		{
-			return FIoStatus(EIoErrorCode::WriteError, TEXT("Failed to write compression method TOC entry"));
+			return FIoStatus(EIoErrorCode::WriteError, FormatLastErrorMessage(TEXT("Failed to write compression method TOC entry.")));
 		}
 	}
 
@@ -1614,7 +1622,7 @@ TIoStatusOr<uint64> FIoStoreTocResource::Write(
 
 		if (!WriteArray(TocFileHandle.Get(), TocResource.ChunkBlockSignatures))
 		{
-			return FIoStatus(EIoErrorCode::WriteError, TEXT("Failed to write chunk block signatures"));
+			return FIoStatus(EIoErrorCode::WriteError, FormatLastErrorMessage(TEXT("Failed to write chunk block signatures.")));
 		}
 	}
 
@@ -1623,14 +1631,14 @@ TIoStatusOr<uint64> FIoStoreTocResource::Write(
 	{
 		if (!TocFileHandle->Write(TocResource.DirectoryIndexBuffer.GetData(), TocResource.DirectoryIndexBuffer.Num()))
 		{
-			return FIoStatus(EIoErrorCode::WriteError, TEXT("Failed to write directory index buffer"));
+			return FIoStatus(EIoErrorCode::WriteError, FormatLastErrorMessage(TEXT("Failed to write directory index buffer.")));
 		}
 	}
 
 	// Meta data (EIoStoreTocReadOptions::ReadTocMeta)
 	if (!WriteArray(TocFileHandle.Get(), TocResource.ChunkMetas))
 	{
-		return FIoStatus(EIoErrorCode::WriteError, TEXT("Failed to write chunk meta data"));
+		return FIoStatus(EIoErrorCode::WriteError, FormatLastErrorMessage(TEXT("Failed to write chunk meta data.")));
 	}
 
 	TocFileHandle->Flush(true);
