Removed unused parameters from CheckVerifyImpl.

diff --git a/Engine/Source/Runtime/Core/Private/Misc/AssertionMacros.cpp b/Engine/Source/Runtime/Core/Private/Misc/AssertionMacros.cpp
--- a/Engine/Source/Runtime/Core/Private/Misc/AssertionMacros.cpp
+++ b/Engine/Source/Runtime/Core/Private/Misc/AssertionMacros.cpp
@@ -833,7 +833,7 @@ FORCENOINLINE void FDebug::DumpStackTraceToLog(const TCHAR* Heading, const ELogV
 }
 
 #if DO_ENSURE && !USING_CODE_ANALYSIS
-bool UE_COLD UE_DEBUG_SECTION VARARGS CheckVerifyImpl(std::atomic<bool>& bExecuted, bool bAlways, const ANSICHAR* File, int32 Line, void* ProgramCounter, const ANSICHAR* Expr, const TCHAR* Format, va_list Args)
+bool UE_COLD UE_DEBUG_SECTION VARARGS CheckVerifyImpl(const ANSICHAR* File, int32 Line, void* ProgramCounter, const ANSICHAR* Expr, const TCHAR* Format, va_list Args)
 {
 	FDebug::OptionallyLogFormattedEnsureMessageReturningFalse(true, Expr, File, Line, ProgramCounter, Format, Args);
 
@@ -861,7 +861,7 @@ bool UE_COLD UE_DEBUG_SECTION UE::Assert::Private::ExecCheckImplInternal(std::at
 		}
 
 		va_list Args = {};
-		const bool bResult = CheckVerifyImpl(bExecuted, bAlways, File, Line, PLATFORM_RETURN_ADDRESS(), Expr, TEXT(""), Args);
+		const bool bResult = CheckVerifyImpl(File, Line, PLATFORM_RETURN_ADDRESS(), Expr, TEXT(""), Args);
 		return bResult;
 	}
 
@@ -876,7 +876,7 @@ bool UE_COLD UE_DEBUG_SECTION UE::Assert::Private::EnsureFailedV(std::atomic<boo
 		return false;
 	}
 
-	const bool bResult = CheckVerifyImpl(bExecuted, Ensure->bAlways, Ensure->File, Ensure->Line, ProgramCounter, Ensure->Expression, Ensure->Format, Args);
+	const bool bResult = CheckVerifyImpl(Ensure->File, Ensure->Line, ProgramCounter, Ensure->Expression, Ensure->Format, Args);
 	return bResult;
 }
 
