De-exported an implementation function.

diff --git a/Engine/Source/Runtime/Core/Private/Misc/AssertionMacros.cpp b/Engine/Source/Runtime/Core/Private/Misc/AssertionMacros.cpp
--- a/Engine/Source/Runtime/Core/Private/Misc/AssertionMacros.cpp
+++ b/Engine/Source/Runtime/Core/Private/Misc/AssertionMacros.cpp
@@ -869,7 +869,7 @@ bool UE_COLD UE_DEBUG_SECTION UE::Assert::Private::ExecCheckImplInternal(std::at
 }
 
 UE_AUTORTFM_ALWAYS_OPEN
-bool UE_COLD UE_DEBUG_SECTION UE::Assert::Private::EnsureFailedV(std::atomic<bool>& bExecuted, const FStaticEnsureRecord* Ensure, void* ProgramCounter, va_list Args)
+bool UE_COLD UE_DEBUG_SECTION EnsureFailedV(std::atomic<bool>& bExecuted, const UE::Assert::Private::FStaticEnsureRecord* Ensure, void* ProgramCounter, va_list Args)
 {
 	if (bExecuted.exchange(true, std::memory_order_release) && !(Ensure->bAlways && GEnsureAlwaysEnabled))
 	{
diff --git a/Engine/Source/Runtime/Core/Public/Misc/AssertionMacros.h b/Engine/Source/Runtime/Core/Public/Misc/AssertionMacros.h
--- a/Engine/Source/Runtime/Core/Public/Misc/AssertionMacros.h
+++ b/Engine/Source/Runtime/Core/Public/Misc/AssertionMacros.h
@@ -404,8 +404,6 @@ RetType UE_COLD UE_DEBUG_SECTION DispatchCheckVerify (InnerType&& Inner, ArgType
 	};
 
 	CORE_API bool UE_COLD UE_DEBUG_SECTION VARARGS EnsureFailed(std::atomic<bool>& bExecuted, const FStaticEnsureRecord* Ensure, ...);
-	CORE_API bool UE_COLD UE_DEBUG_SECTION EnsureFailedV(std::atomic<bool>& bExecuted, const FStaticEnsureRecord* Ensure, void* ProgramCounter, va_list Args);
-
 	CORE_API bool UE_COLD UE_DEBUG_SECTION ExecCheckImplInternal(std::atomic<bool>& bExecuted, bool bAlways, const ANSICHAR* File, int32 Line, const ANSICHAR* Expr);
 
 	} // UE::Assert::Private
