Removed unnecessary capture parameter from UE_ENSURE_IMPL2.

diff --git a/Engine/Source/Runtime/Core/Public/Misc/AssertionMacros.h b/Engine/Source/Runtime/Core/Public/Misc/AssertionMacros.h
--- a/Engine/Source/Runtime/Core/Public/Misc/AssertionMacros.h
+++ b/Engine/Source/Runtime/Core/Public/Misc/AssertionMacros.h
@@ -445,8 +445,8 @@ RetType UE_COLD UE_DEBUG_SECTION DispatchCheckVerify (InnerType&& Inner, ArgType
 			&& [](){ PLATFORM_BREAK(); return false; }()))
 #endif
 
-	#define UE_ENSURE_IMPL2(Capture, Always, InExpression, InFormat, ...) \
-		(LIKELY(!!(InExpression)) || ([Capture] () UE_COLD UE_DEBUG_SECTION \
+	#define UE_ENSURE_IMPL2(Always, InExpression, InFormat, ...) \
+		(LIKELY(!!(InExpression)) || ([&] () UE_COLD UE_DEBUG_SECTION \
 		{ \
 			UE_VALIDATE_FORMAT_STRING(InFormat, ##__VA_ARGS__); \
 			static std::atomic<bool> ENSURE_bExecuted = false; \
@@ -458,10 +458,10 @@ RetType UE_COLD UE_DEBUG_SECTION DispatchCheckVerify (InnerType&& Inner, ArgType
 			return false; \
 		} ()))
 
-	#define ensure(           InExpression                ) UE_ENSURE_IMPL (   false, InExpression)
-	#define ensureMsgf(       InExpression, InFormat, ... ) UE_ENSURE_IMPL2(&, false, InExpression, InFormat, ##__VA_ARGS__)
-	#define ensureAlways(     InExpression                ) UE_ENSURE_IMPL (   true,  InExpression)
-	#define ensureAlwaysMsgf( InExpression, InFormat, ... ) UE_ENSURE_IMPL2(&, true,  InExpression, InFormat, ##__VA_ARGS__)
+	#define ensure(           InExpression                ) UE_ENSURE_IMPL (false, InExpression)
+	#define ensureMsgf(       InExpression, InFormat, ... ) UE_ENSURE_IMPL2(false, InExpression, InFormat, ##__VA_ARGS__)
+	#define ensureAlways(     InExpression                ) UE_ENSURE_IMPL (true,  InExpression)
+	#define ensureAlwaysMsgf( InExpression, InFormat, ... ) UE_ENSURE_IMPL2(true,  InExpression, InFormat, ##__VA_ARGS__)
 
 
 #else // DO_ENSURE && !USING_CODE_ANALYSIS
