Passed the line number to the hashing function for further clarity inside the macro.

diff --git a/Engine/Source/Runtime/Core/Public/Misc/AssertionMacros.h b/Engine/Source/Runtime/Core/Public/Misc/AssertionMacros.h
--- a/Engine/Source/Runtime/Core/Public/Misc/AssertionMacros.h
+++ b/Engine/Source/Runtime/Core/Public/Misc/AssertionMacros.h
@@ -411,19 +411,15 @@ RetType UE_COLD UE_DEBUG_SECTION DispatchCheckVerify (InnerType&& Inner, ArgType
 	template <uint64 Uid>
 	std::atomic<bool> bGEnsureHasExecuted = false;
 
-	UE_NODEBUG constexpr uint32 FileHashForEnsure(const char* Filename)
+	UE_NODEBUG constexpr uint32 FileLineHashForEnsure(const char* Filename, uint32 Line)
 	{
 		// Based on djb2
 		uint32 Result = 5381;
-		for (;;)
+		while (char Ch = *Filename++)
 		{
-			char Ch = *Filename++;
-			if (!Ch)
-			{
-				return Result;
-			}
 			Result = ((Result << 5) + Result) + Ch;
 		}
+		return ((uint64)Result << 32) | (uint64)Line;
 	}
 
 #if UE_USE_LITE_ENSURES
@@ -436,12 +432,12 @@ RetType UE_COLD UE_DEBUG_SECTION DispatchCheckVerify (InnerType&& Inner, ArgType
 
 	#define UE_ENSURE_IMPL(Always, InExpression) \
 		(LIKELY(!!(InExpression)) \
-			|| (::UE::Assert::Private::ExecCheckImplInternal(::bGEnsureHasExecuted<((uint64)FileHashForEnsure(__FILE__) << 32) | (uint64)__LINE__>, Always, __FILE__, __LINE__, #InExpression) \
+			|| (::UE::Assert::Private::ExecCheckImplInternal(::bGEnsureHasExecuted<FileLineHashForEnsure(__FILE__, __LINE__)>, Always, __FILE__, __LINE__, #InExpression) \
 			&& BreakAndReturnFalse()))
 #else
 	#define UE_ENSURE_IMPL(Always, InExpression) \
 		(LIKELY(!!(InExpression)) \
-			|| (::UE::Assert::Private::ExecCheckImplInternal(::bGEnsureHasExecuted<((uint64)FileHashForEnsure(__FILE__) << 32) | (uint64)__LINE__>, Always, __FILE__, __LINE__, #InExpression) \
+			|| (::UE::Assert::Private::ExecCheckImplInternal(::bGEnsureHasExecuted<FileLineHashForEnsure(__FILE__, __LINE__)>, Always, __FILE__, __LINE__, #InExpression) \
 			&& [](){ PLATFORM_BREAK(); return false; }()))
 #endif
 
