Use template global bool for 'has executed' state instead of static.

diff --git a/Engine/Source/Runtime/Core/Public/Misc/AssertionMacros.h b/Engine/Source/Runtime/Core/Public/Misc/AssertionMacros.h
--- a/Engine/Source/Runtime/Core/Public/Misc/AssertionMacros.h
+++ b/Engine/Source/Runtime/Core/Public/Misc/AssertionMacros.h
@@ -445,9 +445,8 @@ RetType UE_COLD UE_DEBUG_SECTION DispatchCheckVerify (InnerType&& Inner, ArgType
 		(LIKELY(!!(InExpression)) || ([&] () UE_COLD UE_DEBUG_SECTION \
 		{ \
 			UE_VALIDATE_FORMAT_STRING(InFormat, ##__VA_ARGS__); \
-			static std::atomic<bool> ENSURE_bExecuted = false; \
 			static constexpr ::UE::Assert::Private::FStaticEnsureRecord ENSURE_Static(InFormat, #InExpression, __builtin_FILE(), __builtin_LINE(), Always); \
-			if ((Always || !ENSURE_bExecuted.load(std::memory_order_relaxed)) && FPlatformMisc::IsEnsureAllowed() && ::UE::Assert::Private::EnsureFailed(ENSURE_bExecuted, &ENSURE_Static, ##__VA_ARGS__)) \
+			if ((Always || !::bGEnsureHasExecuted<FileLineHashForEnsure(__FILE__, __LINE__)>.load(std::memory_order_relaxed)) && FPlatformMisc::IsEnsureAllowed() && ::UE::Assert::Private::EnsureFailed(::bGEnsureHasExecuted<FileLineHashForEnsure(__FILE__, __LINE__)>, &ENSURE_Static, ##__VA_ARGS__)) \
 			{ \
 				PLATFORM_BREAK(); \
 			} \
