[Core]
* Fixed dllstorage in TaskTrace

diff --git a/Engine/Source/Runtime/Core/Public/Async/TaskTrace.h b/Engine/Source/Runtime/Core/Public/Async/TaskTrace.h
--- a/Engine/Source/Runtime/Core/Public/Async/TaskTrace.h
+++ b/Engine/Source/Runtime/Core/Public/Async/TaskTrace.h
@@ -40,17 +40,17 @@ namespace TaskTrace
 
 	inline constexpr uint32 TaskTraceVersion = 1;
 
-	FId CORE_API GenerateTaskId();
-
-	void CORE_API Init();
-	void CORE_API Created(FId TaskId, uint64 TaskSize); // optional, used only if a task was created but not launched immediately
-	void CORE_API Launched(FId TaskId, const TCHAR* DebugName, bool bTracked, ENamedThreads::Type ThreadToExecuteOn, uint64 TaskSize);
-	void CORE_API Scheduled(FId TaskId);
-	void CORE_API SubsequentAdded(FId TaskId, FId SubsequentId);
-	void CORE_API Started(FId TaskId);
-	void CORE_API Finished(FId TaskId);
-	void CORE_API Completed(FId TaskId);
-	void CORE_API Destroyed(FId TaskId);
+	FId TASK_CORE_API GenerateTaskId();
+
+	void TASK_CORE_API Init();
+	void TASK_CORE_API Created(FId TaskId, uint64 TaskSize); // optional, used only if a task was created but not launched immediately
+	void TASK_CORE_API Launched(FId TaskId, const TCHAR* DebugName, bool bTracked, ENamedThreads::Type ThreadToExecuteOn, uint64 TaskSize);
+	void TASK_CORE_API Scheduled(FId TaskId);
+	void TASK_CORE_API SubsequentAdded(FId TaskId, FId SubsequentId);
+	void TASK_CORE_API Started(FId TaskId);
+	void TASK_CORE_API Finished(FId TaskId);
+	void TASK_CORE_API Completed(FId TaskId);
+	void TASK_CORE_API Destroyed(FId TaskId);
 
 	struct FWaitingScope
 	{
@@ -88,4 +88,3 @@ namespace TaskTrace
 	inline FTaskTimingEventScope::~FTaskTimingEventScope() {}
 #endif // UE_TASK_TRACE_ENABLED
 }
-
