Moved some checks from the UE_ENSURE_IMPL2 macro into the .cpp.

diff --git a/Engine/Source/Runtime/Core/Private/Misc/AssertionMacros.cpp b/Engine/Source/Runtime/Core/Private/Misc/AssertionMacros.cpp
--- a/Engine/Source/Runtime/Core/Private/Misc/AssertionMacros.cpp
+++ b/Engine/Source/Runtime/Core/Private/Misc/AssertionMacros.cpp
@@ -882,12 +882,17 @@ bool UE_COLD UE_DEBUG_SECTION EnsureFailedV(std::atomic<bool>& bExecuted, const
 
 bool UE_COLD UE_DEBUG_SECTION VARARGS UE::Assert::Private::EnsureFailed(std::atomic<bool>& bExecuted, const FStaticEnsureRecord* Ensure, ...)
 {
-	va_list Args;
-	va_start(Args, Ensure);
-	const bool bResult = EnsureFailedV(bExecuted, Ensure, PLATFORM_RETURN_ADDRESS(), Args);
-	va_end(Args);
+	if ((Ensure->bAlways || !bExecuted.load(std::memory_order_relaxed)) && FPlatformMisc::IsEnsureAllowed())
+	{
+		va_list Args;
+		va_start(Args, Ensure);
+		const bool bResult = EnsureFailedV(bExecuted, Ensure, PLATFORM_RETURN_ADDRESS(), Args);
+		va_end(Args);
 
-	return bResult;
+		return bResult;
+	}
+
+	return false;
 }
 #endif
 
diff --git a/Engine/Source/Runtime/Core/Public/Misc/AssertionMacros.h b/Engine/Source/Runtime/Core/Public/Misc/AssertionMacros.h
--- a/Engine/Source/Runtime/Core/Public/Misc/AssertionMacros.h
+++ b/Engine/Source/Runtime/Core/Public/Misc/AssertionMacros.h
@@ -446,7 +446,7 @@ RetType UE_COLD UE_DEBUG_SECTION DispatchCheckVerify (InnerType&& Inner, ArgType
 		{ \
 			UE_VALIDATE_FORMAT_STRING(InFormat, ##__VA_ARGS__); \
 			static constexpr ::UE::Assert::Private::FStaticEnsureRecord ENSURE_Static(InFormat, #InExpression, __builtin_FILE(), __builtin_LINE(), Always); \
-			if ((Always || !::bGEnsureHasExecuted<FileLineHashForEnsure(__FILE__, __LINE__)>.load(std::memory_order_relaxed)) && FPlatformMisc::IsEnsureAllowed() && ::UE::Assert::Private::EnsureFailed(::bGEnsureHasExecuted<FileLineHashForEnsure(__FILE__, __LINE__)>, &ENSURE_Static, ##__VA_ARGS__)) \
+			if (::UE::Assert::Private::EnsureFailed(::bGEnsureHasExecuted<FileLineHashForEnsure(__FILE__, __LINE__)>, &ENSURE_Static, ##__VA_ARGS__)) \
 			{ \
 				PLATFORM_BREAK(); \
 			} \
