Fix race in doc comment for example usage of FEventCount.

diff --git a/Engine/Source/Runtime/Core/Public/Async/EventCount.h b/Engine/Source/Runtime/Core/Public/Async/EventCount.h
--- a/Engine/Source/Runtime/Core/Public/Async/EventCount.h
+++ b/Engine/Source/Runtime/Core/Public/Async/EventCount.h
@@ -48,8 +48,8 @@ private:
  *
  * On the notifying thread:
  *
- *     ++CurrentValue;
- *     if (CurrentValue == TargetValue)
+ *     uint32 NewValue = CurrentValue.fetch_add(1) + 1;
+ *     if (NewValue == TargetValue)
  *     {
  *         Event.Notify();
  *     }
