- Handle caching off other-platform cvar values when they are read in before the main CVar exists. Now, when the platform cvar is created, any "read-too-early" values will be set into the platform cvar

diff --git a/Engine/Source/Runtime/Core/Private/HAL/ConsoleManager.cpp b/Engine/Source/Runtime/Core/Private/HAL/ConsoleManager.cpp
--- a/Engine/Source/Runtime/Core/Private/HAL/ConsoleManager.cpp
+++ b/Engine/Source/Runtime/Core/Private/HAL/ConsoleManager.cpp
@@ -1402,6 +1402,9 @@ FConsoleVariable<T>* FindOrCreateTypedPlatformCVar(FConsoleVariableExtendedData<
 		
 		// cache it
 		CVar->PlatformValues.Add(PlatformKey, PlatformCVar);
+
+		// let the manager know, in case it had any cached platform values that were read in before the main cvar was created
+		GetManager().OnCreatedPlatformCVar(CVar, PlatformCVar, PlatformKey);
 	}
 
 	return static_cast<FConsoleVariable<T>*>(PlatformCVar.Get());
@@ -2665,6 +2668,36 @@ void FConsoleManager::DumpObjects(const TCHAR* Params, FOutputDevice& InAr, bool
 	delete CSV;
 }
 
+
+#if ALLOW_OTHER_PLATFORM_CONFIG
+static inline FString MakeUnknownCVarKey(const FString& CVarName, FName PlatformKey)
+{
+	return FString::Printf(TEXT("%s@%s"), *CVarName, *PlatformKey.ToString());
+}
+
+void FConsoleManager::OnCreatedPlatformCVar(IConsoleVariable* MainVariable, TSharedPtr<IConsoleVariable> PlatformVariable, FName PlatformKey)
+{
+	FString UnknownCVarKey = MakeUnknownCVarKey(FindConsoleObjectName(MainVariable), PlatformKey);
+
+	FScopeLock Lock(&CachedPlatformsAndDeviceProfilesLock);
+
+	// look to see if there are any cached off values loaded from disk but created before the cvar was created
+	TArray<TTuple<FString, EConsoleVariableFlags>*> LoadedValues;
+	UnknownCVarCache.MultiFindPointer(UnknownCVarKey, LoadedValues, true);
+	if (LoadedValues.Num() > 0)
+	{
+		for (TTuple<FString, EConsoleVariableFlags>* Value : LoadedValues)
+		{
+			PlatformVariable->Set(*Value->Key, Value->Value);
+
+			UE_LOG(LogConsoleManager, Verbose, TEXT("Setting delayed platform cvar value, %s = %s, at priority %s"), *UnknownCVarKey, *Value->Key, GetConsoleVariableSetByName(Value->Value));
+		}
+
+		UnknownCVarCache.Remove(UnknownCVarKey);
+	}
+}
+#endif
+
 static void SetUnsetCVar(const TMap<FString, IConsoleObject*>& ConsoleObjects, const TCHAR* Params, FOutputDevice& Ar, bool bSet)
 {
 	FString CVarName = FParse::Token(Params, false);
@@ -3398,12 +3431,19 @@ void FConsoleManager::LoadAllPlatformCVars(FName PlatformName, const FString& De
 	
 	// use the platform's base DeviceProfile for emulation
 	VisitPlatformCVarsForEmulation(PlatformName, DeviceProfileName.IsEmpty() ? PlatformName.ToString() : DeviceProfileName,
-		[PlatformName, PlatformKey](const FString& CVarName, const FString& CVarValue, EConsoleVariableFlags SetByAndPreview)
+		[PlatformName, PlatformKey, this](const FString& CVarName, const FString& CVarValue, EConsoleVariableFlags SetByAndPreview)
 	{
+		int32 SetBy = SetByAndPreview & ECVF_SetByMask;
+			
 		// make sure the named cvar exists
 		IConsoleVariable* CVar = IConsoleManager::Get().FindConsoleVariable(*CVarName);
 		if (CVar == nullptr)
 		{
+			// store this unknown value in a map, so when the cvar gets created later, we can grab this platform's values into the platform variable
+			FString UnknownCVarKey = MakeUnknownCVarKey(CVarName, PlatformKey);
+			UnknownCVarCache.Add(UnknownCVarKey, { CVarValue, (EConsoleVariableFlags)SetBy });
+
+			UE_LOG(LogConsoleManager, Verbose, TEXT("Loaded %s@%s = %s, but the main CVar doesn't exist!"), *CVarName, *PlatformKey.ToString(), *CVarValue);
 			return;
 		}
 		
@@ -3411,7 +3451,6 @@ void FConsoleManager::LoadAllPlatformCVars(FName PlatformName, const FString& De
 		FConsoleVariableBase* PlatformCVar = FindOrCreatePlatformCVar(CVar, PlatformKey);
 
 		// now cache the passed in value
-		int32 SetBy = SetByAndPreview & ECVF_SetByMask;
 		PlatformCVar->SetOtherPlatformValue(*CVarValue, (EConsoleVariableFlags)SetBy, NAME_None);
 		
 		UE_LOG(LogConsoleManager, Verbose, TEXT("Loading %s@%s = %s [get = %s]"), *CVarName, *PlatformKey.ToString(),
diff --git a/Engine/Source/Runtime/Core/Private/HAL/ConsoleManager.h b/Engine/Source/Runtime/Core/Private/HAL/ConsoleManager.h
--- a/Engine/Source/Runtime/Core/Private/HAL/ConsoleManager.h
+++ b/Engine/Source/Runtime/Core/Private/HAL/ConsoleManager.h
@@ -83,6 +83,10 @@ public:
 	void OnCVarChanged();
 	void DumpObjects(const TCHAR* Params, FOutputDevice& InAr, bool bDisplayCommands) const;
 
+#if ALLOW_OTHER_PLATFORM_CONFIG
+	void OnCreatedPlatformCVar(IConsoleVariable* MainVariable, TSharedPtr<IConsoleVariable> PlatformVariable, FName PlatformKey);
+#endif
+
 private:
 
 	/**
@@ -133,6 +137,7 @@ private:
 	IConsoleThreadPropagation* ThreadPropagationCallback;
 	FCriticalSection CachedPlatformsAndDeviceProfilesLock;
 	TSet<FName> CachedPlatformsAndDeviceProfiles;
+	TMultiMap<FString, TTuple<FString, EConsoleVariableFlags>> UnknownCVarCache;
 
 	/** If true the next call to CallAllConsoleVariableSinks() we will call all registered sinks */
 	bool bCallAllConsoleVariableSinks;
