Split the BreakAndReturnFalse() functionality out from UE_ENSURE_IMPL into its own macro so that it can be reused.

diff --git a/Engine/Source/Runtime/Core/Public/Misc/AssertionMacros.h b/Engine/Source/Runtime/Core/Public/Misc/AssertionMacros.h
--- a/Engine/Source/Runtime/Core/Public/Misc/AssertionMacros.h
+++ b/Engine/Source/Runtime/Core/Public/Misc/AssertionMacros.h
@@ -430,16 +430,15 @@ RetType UE_COLD UE_DEBUG_SECTION DispatchCheckVerify (InnerType&& Inner, ArgType
 		return false;
 	}
 
-	#define UE_ENSURE_IMPL(Always, InExpression) \
-		(LIKELY(!!(InExpression)) \
-			|| (::UE::Assert::Private::ExecCheckImplInternal(::bGEnsureHasExecuted<FileLineHashForEnsure(__FILE__, __LINE__)>, Always, __FILE__, __LINE__, #InExpression) \
-			&& BreakAndReturnFalse()))
+	#define UE_BREAK_AND_RETURN_FALSE() (BreakAndReturnFalse())
 #else
+	#define UE_BREAK_AND_RETURN_FALSE() [](){ PLATFORM_BREAK(); return false; }()
+#endif
+
 	#define UE_ENSURE_IMPL(Always, InExpression) \
 		(LIKELY(!!(InExpression)) \
 			|| (::UE::Assert::Private::ExecCheckImplInternal(::bGEnsureHasExecuted<FileLineHashForEnsure(__FILE__, __LINE__)>, Always, __FILE__, __LINE__, #InExpression) \
-			&& [](){ PLATFORM_BREAK(); return false; }()))
-#endif
+			&& UE_BREAK_AND_RETURN_FALSE()))
 
 	#define UE_ENSURE_IMPL2(Always, InExpression, InFormat, ...) \
 		(LIKELY(!!(InExpression)) || ([&] () UE_COLD UE_DEBUG_SECTION \
