Fix Engine test 'System.Core.Misc.TextStringification.'

ICU test cases were inadvertently broken by an incorrect macro definition. This was only an issue with the unit test code, not the actual text stringification logic.

diff --git a/Engine/Source/Runtime/Core/Tests/Internationalization/TextTest.cpp b/Engine/Source/Runtime/Core/Tests/Internationalization/TextTest.cpp
--- a/Engine/Source/Runtime/Core/Tests/Internationalization/TextTest.cpp
+++ b/Engine/Source/Runtime/Core/Tests/Internationalization/TextTest.cpp
@@ -1207,12 +1207,12 @@ static void RunTestCasesForSystem_Core_Misc_TextStringification()
 		}
 	};
 
-#define TEST(Text, Str) DoSingleTest(Text, TEXT(Str), TEXT(#Text), true)
-#define TEST_EX(Text, Str, ImportCpp) DoSingleTest(Text, TEXT(Str), TEXT(#Text), ImportCpp)
+#define TEST(Text, Str)                       DoSingleTest(Text, TEXT(Str), TEXT(#Text), true)
+#define TEST_EX(Text, Str, ImportCpp)         DoSingleTest(Text, TEXT(Str), TEXT(#Text), ImportCpp)
 
 #if UE_ENABLE_ICU
-	#define TEST_ICU(Text, Str)               TEST(Text, Str)
-	#define TEST_EX_ICU(Text, Str, ImportCpp) TEST_EX(Text, Str, ImportCpp)
+	#define TEST_ICU(Text, Str)               DoSingleTest(Text, TEXT(Str), TEXT(#Text), true)
+	#define TEST_EX_ICU(Text, Str, ImportCpp) DoSingleTest(Text, TEXT(Str), TEXT(#Text), ImportCpp)
 #else
 	#define TEST_ICU(Text, Str)              
 	#define TEST_EX_ICU(Text, Str, ImportCpp)
