Added support for Ferror using structured logging to generate error messages.

diff --git a/Engine/Source/Runtime/Core/Private/Experimental/IO/IoStatusError.cpp b/Engine/Source/Runtime/Core/Private/Experimental/IO/IoStatusError.cpp
--- a/Engine/Source/Runtime/Core/Private/Experimental/IO/IoStatusError.cpp
+++ b/Engine/Source/Runtime/Core/Private/Experimental/IO/IoStatusError.cpp
@@ -19,20 +19,45 @@ namespace UE::UnifiedError::IoStore
 		{
 			ErrorMessage = InErrorMessage;
 		}
+		/*
 
-		virtual const FUtf8String& GetErrorName(const FError& Error) const final
+		virtual UE::UnifiedError::FAnyString GetErrorCodeString(const FError& Error) const final
 		{
-			if (CachedErrorName.IsEmpty() && (ErrorMessage[0] != TEXT('\0')))
-			{
-				CachedErrorName = TCHAR_TO_UTF8(GetIoErrorText(EIoErrorCode(Error.GetErrorCode())));
-			}
-			return CachedErrorName;
+			UE::UnifiedError::FAnyString Result; 
+			Result.CopyString(TCHAR_TO_UTF8(GetIoErrorText(EIoErrorCode(Error.GetErrorCode()))));
+			return Result;
 		}
-		virtual const FText GetErrorFormatString(const FError& Error) const final
+
+		virtual UE::UnifiedError::FAnyString GetModuleIdString(const FError& Error) const final
 		{
+			UE::UnifiedError::FAnyString Result;
+			Result.ReferenceString(UE::UnifiedError::IoStore::StaticModuleName);
+			return Result;
+
+		}
+*/
+
+		virtual const FText GetErrorFormatString(const FError& Error) const final
+		{	
 			return GenericErrorFormatString;
 		}
 
+		virtual void SerializeToCb(FCbWriter& Writer, const FError& Error) const override
+		{
+			Writer.BeginObject();
+			if (Error.GetModuleId() == UE::UnifiedError::IoStore::StaticModuleId)
+			{
+				Writer.AddString(UTF8TEXTVIEW("ErrorCodeString"), GetIoErrorText(EIoErrorCode(Error.GetErrorCode())));
+				Writer.AddString(UTF8TEXTVIEW("ModuleIdString"), UE::UnifiedError::IoStore::StaticModuleName);
+				Writer.AddInteger(UTF8TEXTVIEW("ErrorCode"), Error.GetErrorCode());
+				Writer.AddInteger(UTF8TEXTVIEW("ModuleId"), Error.GetModuleId());
+			}
+			Writer.AddString(UTF8TEXTVIEW("IoStoreErrorMessage"), ErrorMessage);
+			Writer.AddString(UTF8TEXTVIEW("$format"), UTF8TEXTVIEW("{IoStoreErrorMessage}"));
+			Writer.AddString(UTF8TEXTVIEW("$type"), UTF8TEXTVIEW("FIoStoreErrorDetails"));
+			Writer.EndObject();
+		}
+
 		virtual void GetErrorProperties(const FError& Error, IErrorPropertyExtractor& OutProperties) const override
 		{
 			OutProperties.AddProperty(UTF8TEXTVIEW("ErrorCodeString"), CachedErrorName);
@@ -48,7 +73,7 @@ namespace UE::UnifiedError::IoStore
 		}
 	};
 
-	const FText FIoStoreErrorDetails::GenericErrorFormatString = NSLOCTEXT("IoStore", "GenericErrorMessage", "{IoStoreErrorMessage}");
+	const FText FIoStoreErrorDetails::GenericErrorFormatString = NSLOCTEXT("IoStore", "GenericErrorMessage", "{FIoStoreErrorDetails}");
 
 
 	CORE_API UE::UnifiedError::FError ConvertError(const FIoStatus& Status)
diff --git a/Engine/Source/Runtime/Core/Private/Experimental/UnifiedError/UnifiedError.cpp b/Engine/Source/Runtime/Core/Private/Experimental/UnifiedError/UnifiedError.cpp
--- a/Engine/Source/Runtime/Core/Private/Experimental/UnifiedError/UnifiedError.cpp
+++ b/Engine/Source/Runtime/Core/Private/Experimental/UnifiedError/UnifiedError.cpp
@@ -16,12 +16,6 @@ namespace UE::UnifiedError
 
 void FError::PushErrorDetails(TRefCountPtr<IErrorDetails> InErrorDetails)
 {
-	/*if (ErrorDetails)
-	{
-		InErrorDetails->SetInnerErrorDetails(ErrorDetails);
-	}
-	InErrorDetails->AddRef();
-	ErrorDetails->ReleaseRef();*/
 	if (!InErrorDetails.IsValid())
 	{
 		return;
@@ -50,28 +44,98 @@ FText FError::GetFormatErrorText() const
 {
 	return ErrorDetails->GetErrorFormatString(*this);
 }
+/*
 
 
-#define USE_STRUCTURED_LOG_FOR_FERRORMESSAGE 0
+UE::UnifiedError::FAnyString FError::GetErrorCodeString() const
+{
+	return ErrorDetails->GetErrorCodeString(*this);
+}
 
-FText FError::GetErrorMessage(bool bIncludeContext) const
+UE::UnifiedError::FAnyString FError::GetModuleIdString() const
 {
-#if USE_STRUCTURED_LOG_FOR_FERRORMESSAGE
-	FString FormatString = ErrorDetails->GetErrorFormatString(*this).ToString();
+	return ErrorDetails->GetModuleIdString(*this);
+}
+*/
 
-	FLogTemplate* Template = UE::CreateLogTemplate(*FormatString);
 
+#define USE_STRUCTURED_LOG_FOR_FERRORMESSAGE 1
 
-	TCbWriter<1024> Writer;
-	Writer.BeginObject();
 
-	SerializeDetailsForLog(Writer);
 
-	Writer.EndObject();
+FText FError::GetErrorMessage(bool bIncludeContext) const
+{
+#if USE_STRUCTURED_LOG_FOR_FERRORMESSAGE
+	FText FormatText = FText::Join(FText::FromString(":"), NSLOCTEXT("UnifiedError", "DefaultLogPrepend", "{Root/ModuleIdString}.{Root/ErrorCodeString}"), GetFormatErrorText());
+	FString FormatString = FormatText.ToString();
 
-	FWideStringBuilderBase OutputMessage;
-	FormatLogTo(OutputMessage,  Template, Writer.Save()); 
+	UE::FLogTemplateOptions TemplateOptions;
+	TemplateOptions.bAllowSubObjectReferences = true;
+	FLogTemplate* Template = UE::CreateLogTemplate(*FormatString, TemplateOptions);
 
+	TCbWriter<1024> Writer;
+
+	SerializeDetailsForLog(Writer);
+	FUtf8StringBuilderBase OutputMessage;
+
+	auto AppendFieldToMessage = [&TemplateOptions, &OutputMessage](FCbField& Field) {
+			FCbObject Object = Field.AsObject();
+			if (Object)
+			{
+				
+				auto OutputContextHeader = [&OutputMessage, Field]() {
+						OutputMessage.Append(UTF8TEXTVIEW(",("));
+						OutputMessage.Append(Field.GetName());
+						OutputMessage.Append(UTF8TEXTVIEW(","));
+					};
+
+				auto OutputContextFooter = [&OutputMessage]() {
+						OutputMessage.Append(UTF8TEXTVIEW(")"));
+					};
+				
+
+
+				FCbField FormatEntry = Object.Find(ANSITEXTVIEW("$format"));
+				if (FormatEntry)
+				{
+					OutputContextHeader();
+					FString ChildFormatString(FormatEntry.AsString());
+					UE::FLogTemplate* ChildTemplate = UE::CreateLogTemplate(*ChildFormatString, TemplateOptions);
+					FormatLogTo(OutputMessage, ChildTemplate, Object.CreateViewIterator());
+					UE::DestroyLogTemplate(ChildTemplate);
+					OutputContextFooter();
+				}
+				else
+				{
+					FCbField TextEntry = Object.Find(ANSITEXTVIEW("$text"));
+					if (TextEntry)
+					{
+						OutputContextHeader();
+						FUtf8String TextString(TextEntry.AsString());
+						OutputMessage.Append(TextString);
+						OutputContextFooter();
+					}
+				}
+			}
+			else
+			{
+				// todo, probably want to support some different types
+				// OutputMessage.Append(UTF8TEXTVIEW(",UnsupportedCbObject"));
+			}
+		};
+
+	FCbFieldIterator CbIterator = Writer.Save();
+
+	FormatLogTo(OutputMessage,  Template, CbIterator); 
+
+	// append the context if desired
+	if (bIncludeContext)
+	{
+		for (FCbFieldIterator It = CbIterator; It; ++It)
+		{
+			AppendFieldToMessage(*It);
+		}
+	}
 	UE::DestroyLogTemplate(Template);
 	return FText::FromString(OutputMessage.ToString());
 #else
@@ -82,16 +146,23 @@ FText FError::GetErrorMessage(bool bIncludeContext) const
 #endif
 }
 
+
 void FError::SerializeDetailsForLog(FCbWriter& Writer) const
 {
 	const IErrorDetails* DetailsIt = ErrorDetails.GetReference();
 	while (DetailsIt != nullptr)
 	{
-		Writer.SetName(DetailsIt->GetErrorDetailsTypeName());
-		DetailsIt->SerializeForLog(Writer);
-
+		Writer.SetName(DetailsIt->GetErrorDetailsTypeNameForLog());
+		DetailsIt->SerializeToCb(Writer, *this);
 
-		DetailsIt = DetailsIt->GetInnerErrorDetails().GetReference();
+		const IErrorDetails* Next = DetailsIt->GetInnerErrorDetails().GetReference();
+		if (Next == nullptr)
+		{
+			// the last one contains some special information we want to reference by name
+			Writer.SetName(ANSITEXTVIEW("Root"));
+			DetailsIt->SerializeToCb(Writer, *this);
+		}
+		DetailsIt = Next;
 	}
 }
 
@@ -196,20 +267,47 @@ void FError::GetErrorProperties(IErrorPropertyExtractor& Visitor) const
 	ErrorDetails->GetErrorProperties(*this, Visitor);
 }
 
-const FUtf8String FError::GetErrorCodeString() const
+
+
+class FAppendFormatStringDetails : public FDynamicErrorDetails
 {
-	FUtf8String Result;
-	GetDetailByKey(UTF8TEXTVIEW("ErrorCodeString"), Result);
-	return Result;
-}
+	DECLARE_FERROR_DETAILS(UnifiedError, FAppendFormatStringDetails);
+public:
+	FAppendFormatStringDetails() {}
+	FAppendFormatStringDetails(FText&& InAppendFormatString, TRefCountPtr<const IErrorDetails> InInnerErrorDetails = nullptr) : FDynamicErrorDetails(InInnerErrorDetails)
+	{
+		AppendFormatString = MoveTemp(InAppendFormatString);
+	}
 
-const FUtf8String FError::GetModuleIdString() const
+private:
+	FText AppendFormatString;
+
+
+public:
+
+	virtual void SerializeToCb(FCbWriter& Writer, const FError& Error) const override
+	{
+		Writer.BeginObject();
+		SerializeBasicToCb(Writer, UTF8TEXTVIEW("AppendFormatString"), AppendFormatString);
+		Writer.EndObject();
+	}
+
+	CORE_API virtual const FText GetErrorFormatString(const FError& Error) const override
+	{
+		const FText InnerFormatString = GetInnerErrorDetails()->GetErrorFormatString(Error);
+		return FText::Join(FText::FromString(" "), InnerFormatString, AppendFormatString);
+	}
+};
+
+
+
+void FError::AppendFormatString(FText&& InFormatString)
 {
-	FUtf8String Result;
-	GetDetailByKey(UTF8TEXTVIEW("ModuleIdString"), Result);
-	return Result;
+	TRefCountPtr<IErrorDetails> NewErrorDetails(new FAppendFormatStringDetails(MoveTemp(InFormatString)));
+	PushErrorDetails(NewErrorDetails);
 }
 
+
 // Structured logging FError integration
 class FCbWriterErrorPropertyExtractor : public IErrorPropertyExtractor
 {
@@ -258,10 +356,43 @@ void SerializeForLog(FCbWriter& Writer, const FError& Error)
 
 #if USE_STRUCTURED_LOG_FOR_FERRORMESSAGE
 	
+	// what I have to do 
+	TCbWriter<1024> InnerWriter;
+
+	Error.SerializeDetailsForLog(InnerWriter);
+
+
+	UE::FLogTemplateOptions TemplateOptions;
+	TemplateOptions.bAllowSubObjectReferences = true;
+	FText FormatText = FText::Join(FText::FromString(":"), NSLOCTEXT("UnifiedError", "DefaultLogPrepend", "{Root/ModuleIdString}.{Root/ErrorCodeString}"), Error.GetFormatErrorText());
+	FLogTemplate* Template = UE::CreateLogTemplate(*FormatText.ToString(), TemplateOptions);
+
+	FUtf8StringBuilderBase OutputMessage;
+
+
+	FCbFieldIterator CbInnerIterator = InnerWriter.Save();
 
+	FormatLogTo(OutputMessage, Template, CbInnerIterator);
+
+	Writer.AddString(ANSITEXTVIEW("$text"), OutputMessage.ToString());
+	// append all the child information into the log also
+	for (FCbFieldViewIterator It = CbInnerIterator; It; ++It)
+	{
+		Writer.SetName(It->GetName());
+		Writer.AddField(*It);
+	}
+	UE::DestroyLogTemplate(Template);
+
+
+
+
+
+	/*
+	// what I want to do
 	Error.SerializeDetailsForLog(Writer);
 
 	Writer.AddString(ANSITEXTVIEW("$text"), Error.GetFormatErrorText().ToString());
+	*/
 
 #else
 
@@ -279,9 +410,8 @@ void SerializeForLog(FCbWriter& Writer, const FError& Error)
 	FTextFormatArgsPropertyExtractor ExtractToFTextArgs(ErrorMessageArgs);
 	Error.GetErrorProperties(ExtractToFTextArgs);
 	Writer.AddString("$text", FText::Format(ErrorFormatText, MoveTemp(ErrorMessageArgs)).ToString());
-
-	Writer.EndObject();
 #endif
+	Writer.EndObject();
 }
 
 // FRefCountedErrorDetails fucntions
@@ -312,8 +442,27 @@ const FUtf8String& FDynamicErrorDetails::GetErrorName(const FError& Error) const
 	return InnerErrorDetails->GetErrorName(Error);
 }
 */
+/*
 
 
+UE::UnifiedError::FAnyString FDynamicErrorDetails::GetErrorCodeString(const FError& Error) const
+{
+	if (InnerErrorDetails.IsValid())
+	{
+		return InnerErrorDetails->GetErrorCodeString(Error);
+	}
+	return UE::UnifiedError::FAnyString();
+}
+
+UE::UnifiedError::FAnyString FDynamicErrorDetails::GetModuleIdString(const FError& Error) const
+{
+	if (InnerErrorDetails.IsValid())
+	{
+		return InnerErrorDetails->GetModuleIdString(Error);
+	}
+	return UE::UnifiedError::FAnyString();
+}*/
+
 void FDynamicErrorDetails::GetErrorProperties(const FError& Error, IErrorPropertyExtractor& OutProperties) const
 {
 	InnerErrorDetails->GetErrorProperties(Error, OutProperties);
@@ -339,22 +488,40 @@ void FStaticErrorDetails::GetErrorProperties(const FError& Error, IErrorProperty
 	OutProperties.AddProperty(UTF8TEXTVIEW("ErrorCode"), Error.GetErrorCode());
 	OutProperties.AddProperty(UTF8TEXTVIEW("ModuleId"), Error.GetModuleId());
 }
+/*
 
-const FAnsiStringView& FStaticErrorDetails::GetErrorCodeString() const
+UE::UnifiedError::FAnyString FStaticErrorDetails::GetErrorCodeString(const FError& Error) const
 {
-	return ErrorName;
+	UE::UnifiedError::FAnyString ReturnValue;
+	ReturnValue.ReferenceString(ErrorName);
+	return ReturnValue;
 }
 
-const FAnsiStringView& FStaticErrorDetails::GetModuleIdString() const
+UE::UnifiedError::FAnyString FStaticErrorDetails::GetModuleIdString(const FError& Error) const
 {
-	return ModuleName;
-}
+	UE::UnifiedError::FAnyString ReturnValue;
+	ReturnValue.ReferenceString(ModuleName);
+	return ReturnValue;
+}*/
 
 const FText FStaticErrorDetails::GetErrorFormatString(const FError& Error) const
 {
 	return ErrorFormatString;
 }
 
+
+void FStaticErrorDetails::SerializeToCb(FCbWriter& Writer, const FError& Error) const
+{
+	Writer.BeginObject();
+	Writer.AddString(UTF8TEXTVIEW("$format"), UTF8TEXTVIEW("({ModuleIdString}.{ErrorCodeString})"));
+	Writer.AddString(UTF8TEXTVIEW("$type"), UTF8TEXTVIEW("FStaticErrorDetails"));
+	Writer.AddString(UTF8TEXTVIEW("ErrorCodeString"), ErrorName);
+	Writer.AddString(UTF8TEXTVIEW("ModuleIdString"), ModuleName);
+	Writer.AddInteger(UTF8TEXTVIEW("ErrorCode"), Error.GetErrorCode());
+	Writer.AddInteger(UTF8TEXTVIEW("ModuleId"), Error.GetModuleId());
+	Writer.EndObject();
+}
+
 // FErrorDetailsRegistry
 uint32 FErrorDetailsRegistry::RegisterDetails(const FAnsiStringView& ErrorDetailsName, TFunction<IErrorDetails*()> CreationFunction)
 {
@@ -369,10 +536,6 @@ uint32 FErrorDetailsRegistry::RegisterDetails(const FAnsiStringView& ErrorDetail
 
 FString LexToString(const UE::UnifiedError::FError& Error)
 {
-	
 	// TODO: daniel, switch to using structured logging output in the future instead of this
-	return FString::Printf(TEXT("%s:%s - %s"), 
-		UTF8_TO_TCHAR(*Error.GetModuleIdString()),
-		UTF8_TO_TCHAR(*Error.GetErrorCodeString()),
-		*Error.GetErrorMessage().ToString());
+	return Error.GetErrorMessage().ToString();
 }
\ No newline at end of file
diff --git a/Engine/Source/Runtime/Core/Public/Experimental/UnifiedError/UnifiedError.h b/Engine/Source/Runtime/Core/Public/Experimental/UnifiedError/UnifiedError.h
--- a/Engine/Source/Runtime/Core/Public/Experimental/UnifiedError/UnifiedError.h
+++ b/Engine/Source/Runtime/Core/Public/Experimental/UnifiedError/UnifiedError.h
@@ -8,6 +8,7 @@
 #include "Internationalization/Text.h"
 #include "Logging/StructuredLog.h"
 #include "Serialization/CompactBinaryWriter.h"
+#include "Misc/TVariant.h"
 #include "Templates/RefCounting.h"
 
 
@@ -21,6 +22,11 @@ template<typename T>
 class TErrorStructFeatures
 {
 public:
+	static const FAnsiStringView GetErrorDetailsLogEntryNameAsString()
+	{
+		static_assert(deferred_false<T>::value, "All types are required to define DECLARE_ERRORSTRUCT_FEATURES to be used as a error context type");
+		return ANSITEXTVIEW("");
+	}
 	static const FAnsiStringView GetErrorDetailsTypeNameAsString()
 	{
 		static_assert(deferred_false<T>::value, "All types are required to define DECLARE_ERRORSTRUCT_FEATURES to be used as a error context type");
@@ -36,18 +42,22 @@ public:
 
 
 
-#define DECLARE_ERRORSTRUCT_FEATURES(DetailsTypeName)\
+#define DECLARE_ERRORSTRUCT_FEATURES(DetailsNameSpace, DetailsStructName)\
 	template<>\
-	class TErrorStructFeatures<UE::UnifiedError::DetailsTypeName>\
+	class TErrorStructFeatures<UE::UnifiedError::DetailsNameSpace::DetailsStructName>\
 	{\
 	public:\
+		static const FAnsiStringView GetErrorDetailsLogEntryNameAsString() \
+		{\
+			return ANSITEXTVIEW(#DetailsNameSpace"_"#DetailsStructName);\
+		}\
 		static const FAnsiStringView GetErrorDetailsTypeNameAsString() \
 		{\
-			return ANSITEXTVIEW("TErrorDetails<"#DetailsTypeName">");\
+			return ANSITEXTVIEW("TErrorDetails<"#DetailsNameSpace"::"#DetailsStructName">");\
 		}\
 		static const FAnsiStringView GetErrorContextTypeNameAsString() \
 		{\
-			return ANSITEXTVIEW(#DetailsTypeName);\
+			return ANSITEXTVIEW(#DetailsNameSpace"::"#DetailsStructName);\
 		}\
 	};
 
@@ -61,27 +71,41 @@ public:
 	{ \
 		return StaticGetErrorDetailsTypeId(); \
 	}\
-	virtual const FAnsiStringView GetErrorDetailsTypeName() const override \
+	static const FAnsiStringView StaticGetErrorDetailsTypeName() \
 	{\
 		return #DetailsNamespace"::"#TypeName;\
+	}\
+	virtual const FAnsiStringView GetErrorDetailsTypeName() const override \
+	{\
+		return StaticGetErrorDetailsTypeName();\
+	}\
+	static const FAnsiStringView StaticGetErrorDetailsTypeNameForLog() \
+	{\
+		return #DetailsNamespace"_"#TypeName;\
+	}\
+	virtual const FAnsiStringView GetErrorDetailsTypeNameForLog() const override\
+	{\
+		return StaticGetErrorDetailsTypeNameForLog(); \
 	}
 
 
 #define DECLARE_FERROR_DETAILS_ABSTRACT(DetailsNamespace, TypeName)\
 	public:\
-	inline static const uint32 StaticDetailsTypeId = UE::UnifiedError::FErrorDetailsRegistry::Get().RegisterDetails(ANSITEXTVIEW(#DetailsNamespace "::" #TypeName), nullptr);\
-	DECLARE_ERROR_DETAILS_INTERNAL(DetailsNamespace, TypeName);
+	DECLARE_ERROR_DETAILS_INTERNAL(DetailsNamespace, TypeName);\
+	inline static const uint32 StaticDetailsTypeId = UE::UnifiedError::FErrorDetailsRegistry::Get().RegisterDetails(StaticGetErrorDetailsTypeName(), nullptr);\
+	
 	
 
 #define DECLARE_FERROR_DETAILS(DetailsNamespace, TypeName)\
 	public:\
+	DECLARE_ERROR_DETAILS_INTERNAL(DetailsNamespace, TypeName);\
 	friend IErrorDetails* Create();\
 	inline static IErrorDetails* Create()\
 	{\
 		return new TypeName();\
 	}\
-	inline static const uint32 StaticDetailsTypeId = UE::UnifiedError::FErrorDetailsRegistry::Get().RegisterDetails(ANSITEXTVIEW(#DetailsNamespace "::" #TypeName), TFunction<IErrorDetails* ()>([]() -> IErrorDetails* { return DetailsNamespace::TypeName::Create(); }));\
-	DECLARE_ERROR_DETAILS_INTERNAL(DetailsNamespace, TypeName);
+	inline static const uint32 StaticDetailsTypeId = UE::UnifiedError::FErrorDetailsRegistry::Get().RegisterDetails(StaticGetErrorDetailsTypeName(), TFunction<IErrorDetails* ()>([]() -> IErrorDetails* { return DetailsNamespace::TypeName::Create(); }));\
+	
 
 
 
@@ -114,6 +138,36 @@ namespace UE::UnifiedError
 
 	class FError;
 
+	// helper class to make pass by string view effiecent while also allowing pass by other strings
+	/*class FAnyString
+	{
+	private:
+		TVariant<FUtf8String, FUtf8StringView> VariantString;
+	public:
+		
+		void ReferenceString(FUtf8StringView InStringView)
+		{
+			VariantString.Set<FUtf8StringView>(InStringView);
+		}
+		void CopyString(FUtf8StringView InString)
+		{
+			VariantString.Set<FUtf8String>(FUtf8String(InString));
+		}
+
+		FUtf8StringView GetUtf8StringView() const
+		{
+			if (VariantString.IsType<FUtf8StringView>())
+			{
+				return VariantString.Get<FUtf8StringView>();
+			}
+			else if (VariantString.IsType<FUtf8String>())
+			{
+				return VariantString.Get<FUtf8String>();
+			}
+			return FUtf8StringView();
+		}
+	};*/
+
 	class IErrorDetails : public IRefCountedObject
 	{
 	public:
@@ -155,9 +209,15 @@ namespace UE::UnifiedError
 		/// <returns></returns>
 		virtual uint32 GetErrorDetailsTypeId() const = 0;
 
+		/*
+
+		virtual FAnyString GetErrorCodeString(const FError& Error) const = 0;
+		virtual FAnyString GetModuleIdString(const FError& Error) const = 0;
+*/
 
-		virtual const FAnsiStringView GetErrorDetailsTypeName() const = 0;
-		virtual void SerializeForLog(FCbWriter& Writer) const {}
+		CORE_API virtual const FAnsiStringView GetErrorDetailsTypeName() const = 0;
+		CORE_API virtual const FAnsiStringView GetErrorDetailsTypeNameForLog() const = 0;
+		CORE_API virtual void SerializeToCb(FCbWriter& Writer, const FError& Error) const =0;
 
 	};
 
@@ -216,6 +276,12 @@ namespace UE::UnifiedError
 		CORE_API virtual void SetInnerErrorDetails(TRefCountPtr<const IErrorDetails> InInnerErrorDetails);
 		virtual TRefCountPtr<const IErrorDetails> GetInnerErrorDetails() const { return InnerErrorDetails; }
 
+/*
+		virtual FAnyString GetErrorCodeString(const FError& Error) const override;
+		virtual FAnyString GetModuleIdString(const FError& Error) const override;
+*/
+
+
 		/// <summary>
 		/// GetErrorFormatString; Pass through to the InnerErrorDetails.
 		/// </summary>
@@ -258,14 +324,14 @@ namespace UE::UnifiedError
 		///  See also: FError::GetErrorDetails
 		/// </summary>
 		/// <returns></returns>
-		CORE_API const FAnsiStringView& GetErrorCodeString() const;
+		//virtual UE::UnifiedError::FAnyString GetErrorCodeString(const FError& Error) const override;
 
 		/// <summary>
 		/// GetModuleIdString; accessor for ModuleName.
 		///  Can be called directly on FStaticErrorDetails object
 		///  See also: FError::GetErrorDetails
 		/// </summary>
-		CORE_API const FAnsiStringView& GetModuleIdString() const;
+		///virtual UE::UnifiedError::FAnyString GetModuleIdString(const FError& Error) const override;
 
 		// IErrorDetails implementation
 
@@ -285,6 +351,7 @@ namespace UE::UnifiedError
 		/// <returns></returns>
 		CORE_API virtual const FText GetErrorFormatString(const FError& Error) const final;
 
+		virtual void SerializeToCb(FCbWriter& Writer, const FError& Error) const final;
 
 		// IRefCountedObject implementation
 		//  FStaticErrorDetails is statically allocated, make sure it is never released.
@@ -311,9 +378,10 @@ namespace UE::UnifiedError
 	{
 	public:
 		TErrorDetails() {}
-		TErrorDetails(T&& InErrorDetail)
+		TErrorDetails(T&& InErrorDetail, bool bInAppendContextToMessage)
 		{
 			ErrorDetail = MoveTemp(InErrorDetail);
+			bAppendContextToMessage = bInAppendContextToMessage;
 		}
 		TErrorDetails(T&& InErrorDetail, TRefCountPtr<const IErrorDetails> InInnerErrorDetails) : FDynamicErrorDetails(InInnerErrorDetails)
 		{
@@ -325,7 +393,7 @@ namespace UE::UnifiedError
 		{
 			return new TErrorDetails<T>();
 		}
-		inline static const uint32 StaticDetailsTypeId = UE::UnifiedError::FErrorDetailsRegistry::Get().RegisterDetails(TErrorStructFeatures<T>::GetNameAsString(), TFunction<IErrorDetails * ()>([]() -> IErrorDetails* { return Create(); }));
+		inline static const uint32 StaticDetailsTypeId = UE::UnifiedError::FErrorDetailsRegistry::Get().RegisterDetails(TErrorStructFeatures<T>::GetErrorDetailsTypeNameAsString(), TFunction<IErrorDetails * ()>([]() -> IErrorDetails* { return Create(); }));
 		static uint32 StaticGetErrorDetailsTypeId()
 		{
 			return StaticDetailsTypeId;
@@ -334,17 +402,29 @@ namespace UE::UnifiedError
 		{
 			return StaticGetErrorDetailsTypeId();
 		}
-	private:
-		T ErrorDetail;
 
-		FAnsiStringView GetTypeNameAsString() const
+		static FAnsiStringView StaticGetErrorDetailsTypeName()
 		{
 			return TErrorStructFeatures<T>::GetErrorDetailsTypeNameAsString();
 		}
+		const FAnsiStringView GetErrorDetailsTypeName() const override
+		{
+			return TErrorStructFeatures<T>::GetErrorDetailsTypeNameAsString();
+		}
+
+		const FAnsiStringView GetErrorDetailsTypeNameForLog() const override
+		{
+			return TErrorStructFeatures<T>::GetErrorDetailsLogEntryNameAsString();
+		}
+
+	private:
+		T ErrorDetail;
+		bool bAppendContextToMessage = false;
+		
 	public:
 		const T& GetErrorContext() const
 		{
-			return &ErrorDetail;
+			return ErrorDetail;
 		}
 
 		virtual void GetErrorProperties(const FError& Error, IErrorPropertyExtractor& OutProperties) const override
@@ -353,9 +433,18 @@ namespace UE::UnifiedError
 			
 			FDynamicErrorDetails::GetErrorProperties(Error, OutProperties);
 		}
-		virtual void SerializeForLog(FCbWriter& Writer) const override
+		virtual void SerializeToCb(FCbWriter& Writer, const FError& Error) const override
+		{
+			SerializeForLog(Writer, ErrorDetail);
+		}
+
+		virtual const FText GetErrorFormatString(const FError& Error) const override
 		{
-			::SerializeForLog(Writer, ErrorDetail);
+			if (bAppendContextToMessage)
+			{
+				return FText::Join(FText::FromString(" "), FDynamicErrorDetails::GetErrorFormatString(Error), FText::FromString(FUtf8String::Printf("{%s}", GetErrorDetailsTypeNameForLog().GetData())));
+			}
+			return FDynamicErrorDetails::GetErrorFormatString(Error);
 		}
 
 		const T& GetValue() const
@@ -365,7 +454,6 @@ namespace UE::UnifiedError
 
 	};
 
-
 	class FError
 	{
 	private: 
@@ -410,49 +498,41 @@ namespace UE::UnifiedError
 		CORE_API int32 GetErrorCode() const;
 		CORE_API int32 GetModuleId() const;
 
-		CORE_API const FUtf8String GetErrorCodeString() const;
-		CORE_API const FUtf8String GetModuleIdString() const;
-
-		CORE_API void PushErrorDetails(TRefCountPtr<IErrorDetails> InErrorDetails);
-		template<typename DetailType>
-		TRefCountPtr<const DetailType> GetErrorDetails() const
-		{
-			TRefCountPtr<const IErrorDetails> CurrentIt = ErrorDetails;
-			while (CurrentIt != nullptr)
-			{
-				if (CurrentIt->GetErrorDetailsTypeId() == DetailType::StaticGetErrorDetailsTypeId())
-				{
-					return TRefCountPtr<const DetailType>((const DetailType*)CurrentIt.GetReference());
-				}
-				CurrentIt = CurrentIt->GetInnerErrorDetails();
-			}
-			return nullptr;
-		}
-
 		CORE_API FText GetErrorMessage(bool bAppendContext = false) const;
 		CORE_API FText GetFormatErrorText() const;
+
 		template<typename T>
-		void PushErrorContext(T&& InErrorStruct)
+		void PushErrorContext(T&& InErrorStruct, bool bShouldAppendToErrorMessage = false)
 		{
-			PushErrorDetails(TRefCountPtr<IErrorDetails>(new TErrorDetails<T>(MoveTemp(InErrorStruct))));
+			PushErrorDetails(TRefCountPtr<IErrorDetails>(new TErrorDetails<T>(MoveTemp(InErrorStruct), bShouldAppendToErrorMessage)));
 		}
 
 		template<typename T>
 		const T* GetErrorContext() const
 		{
-			TRefCountPtr<const IErrorDetails> CurrentIt = ErrorDetails;
+
+			TRefCountPtr<const TErrorDetails<T>> CurrentDetails = GetErrorDetails<const TErrorDetails<T>>();
+			if (const TErrorDetails<T>* Details = CurrentDetails.GetReference())
+			{
+				return &Details->GetErrorContext();
+			}
+			return nullptr;
+
+			/*TRefCountPtr<const IErrorDetails> CurrentIt = ErrorDetails;
 			while (CurrentIt != nullptr)
 			{
-				if (CurrentIt->GetErrorDetailsTypeName() == TErrorStructFeatures<T>::GetErrorDetailsTypeNameAsString())
+				if (CurrentIt->GetErrorDetailsTypeId() == GetErrorDetails<>)
 				{
 					const TErrorDetails<const T>* Details = (const TErrorDetails<const T>*)(CurrentIt.GetReference());
 					return &Details->GetErrorContext();
 				}
 				CurrentIt = CurrentIt->GetInnerErrorDetails();
 			}
-			return nullptr;
+			return nullptr;*/
 		}
 
+		CORE_API void AppendFormatString(FText&& InFormatString);
+
 		CORE_API void SerializeDetailsForLog(FCbWriter& Writer) const;
 
 		CORE_API bool GetDetailByKey(const FUtf8StringView& KeyName, FString& Result) const;
@@ -475,6 +555,24 @@ namespace UE::UnifiedError
 			ErrorDetails = nullptr;
 		}
 	private:
+
+
+		CORE_API void PushErrorDetails(TRefCountPtr<IErrorDetails> InErrorDetails);
+		template<typename DetailType>
+		TRefCountPtr<const DetailType> GetErrorDetails() const
+		{
+			TRefCountPtr<const IErrorDetails> CurrentIt = ErrorDetails;
+			while (CurrentIt != nullptr)
+			{
+				if (CurrentIt->GetErrorDetailsTypeId() == DetailType::StaticGetErrorDetailsTypeId())
+				{
+					return TRefCountPtr<const DetailType>((const DetailType*)CurrentIt.GetReference());
+				}
+				CurrentIt = CurrentIt->GetInnerErrorDetails();
+			}
+			return nullptr;
+		}
+
 		friend bool operator==(const FError& Error, const FError& OtherError)
 		{
 			if ((Error.ModuleId == OtherError.ModuleId) &&
@@ -572,7 +670,36 @@ namespace UE::UnifiedError
 		}
 	};
 
-	CORE_API void SerializeForLog(FCbWriter& Writer, const FError& OnlineError);
+	inline void SerializeBasicToCb(FCbWriter& Writer, const FUtf8StringView& PropertyName, const FStringView& Value)
+	{
+		Writer.AddString(PropertyName, Value);
+	}
+	inline void SerializeBasicToCb(FCbWriter& Writer, const FUtf8StringView& PropertyName, const FUtf8StringView& Value)
+	{
+		Writer.AddString(PropertyName, Value);
+	}
+	inline void SerializeBasicToCb(FCbWriter& Writer, const FUtf8StringView& PropertyName, const int32& Value)
+	{
+		Writer.AddInteger(PropertyName, Value);
+	}
+	inline void SerializeBasicToCb(FCbWriter& Writer, const FUtf8StringView& PropertyName, const uint32& Value)
+	{
+		Writer.AddInteger(PropertyName, Value);
+	}
+	inline void SerializeBasicToCb(FCbWriter& Writer, const FUtf8StringView& PropertyName, const int64& Value)
+	{
+		Writer.AddInteger(PropertyName, Value);
+	}
+	inline void SerializeBasicToCb(FCbWriter& Writer, const FUtf8StringView& PropertyName, const uint64& Value)
+	{
+		Writer.AddInteger(PropertyName, Value);
+	}
+	inline void SerializeBasicToCb(FCbWriter& Writer, const FUtf8StringView& PropertyName, const FText& Value)
+	{
+		Writer.AddString(PropertyName, Value.ToString());
+	}
+
+	CORE_API void SerializeForLog(FCbWriter& Writer, const FError& Error);
 
 } // namespace UnifiedError
 
@@ -647,6 +774,12 @@ CORE_API constexpr uint32 GetTypeIdHash(const TCHAR* TypeName)
 					OutProperties.AddProperty(UTF8TEXTVIEW(#ParamOneName), ParamOneName); \
 					FDynamicErrorDetails::GetErrorProperties(Error, OutProperties); \
 				} \
+				virtual void SerializeToCb(FCbWriter& Writer, const FError& Error) const override \
+				{\
+					Writer.BeginObject();\
+					SerializeBasicToCb(Writer, UTF8TEXTVIEW(#ParamOneName), ParamOneName); \
+					Writer.EndObject();\
+				}\
 				const ParamOneType& Get##ParamOneName() const \
 				{ \
 					return ParamOneName; \
diff --git a/Engine/Source/Runtime/Engine/Classes/Engine/Experimental/StreamableManagerError.h b/Engine/Source/Runtime/Engine/Classes/Engine/Experimental/StreamableManagerError.h
--- a/Engine/Source/Runtime/Engine/Classes/Engine/Experimental/StreamableManagerError.h
+++ b/Engine/Source/Runtime/Engine/Classes/Engine/Experimental/StreamableManagerError.h
@@ -57,11 +57,20 @@ namespace UE::UnifiedError::StreamableManager
 		{
 			const FText InnerFormatString = GetInnerErrorDetails()->GetErrorFormatString(Error);
 
-			FText MyFormatString = NSLOCTEXT("StreamableManager", "DetailsFormatString", "(SoftObjectPath:{SoftObjectPath})");
+			FText MyFormatString = NSLOCTEXT("StreamableManager", "DetailsFormatString", "SoftObjectPath:{StreamableManager_FStreamableManagerErrorDetails/SoftObjectPath}");
 
 			return FText::Join(NSLOCTEXT("StreamableManager", "DetailsFormatStringDelimiter", " "), InnerFormatString, MyFormatString);
 		}
 
+		virtual void SerializeToCb(FCbWriter& Writer, const FError& Error) const
+		{
+			Writer.BeginObject();
+			Writer.AddString(UTF8TEXTVIEW("$type"), UTF8TEXTVIEW("FStreamableManagerErrorDetails"));
+			Writer.AddString(UTF8TEXTVIEW("$format"), UTF8TEXTVIEW("(SoftObjectPath:{SoftObjectPath})"));
+			Writer.AddString(UTF8TEXTVIEW("SoftObjectPath"), SoftObjectPath);
+			Writer.EndObject();
+		}
+
 		const FString& GetSoftObjectPath() const { return SoftObjectPath; }
 
 		
@@ -75,26 +84,24 @@ namespace UE::UnifiedError::StreamableManager
 	public:
 		/*friend void SerializeForLog(FCbWriter& Writer, const UE::UnifiedError::StreamableManager::FStreamableManagerAdditionalContext& Context);
 	private:*/
-		FString ExtraSoftObjectPath;
+		FString SoftObjectPath;
 	};
 
 
 	inline void GatherPropertiesForError(const UE::UnifiedError::FError& Error, const UE::UnifiedError::StreamableManager::FStreamableManagerAdditionalContext& Context, UE::UnifiedError::IErrorPropertyExtractor& PropertyExtractor)
 	{
-		PropertyExtractor.AddProperty(ANSITEXTVIEW("ExtraSoftObjectPath"), Context.ExtraSoftObjectPath);
+		PropertyExtractor.AddProperty(ANSITEXTVIEW("SoftObjectPath"), Context.SoftObjectPath);
 	}
 	
 }
-
+DECLARE_ERRORSTRUCT_FEATURES(StreamableManager, FStreamableManagerAdditionalContext);
 inline void SerializeForLog(FCbWriter& Writer, const UE::UnifiedError::StreamableManager::FStreamableManagerAdditionalContext& Context)
 {
 	Writer.BeginObject();
-	Writer.AddString(ANSITEXTVIEW("$type"), ANSITEXTVIEW("FStreamableManagerAdditionalContext"));
-	FStringFormatNamedArguments NamedArguments;
-	NamedArguments.Add(TTuple<FString, FStringFormatArg>(TEXT("SoftObjectPath"), Context.ExtraSoftObjectPath));
-	Writer.AddString(ANSITEXTVIEW("$text"), FString::Format(TEXT("(SoftObjectPath: {SoftObjectPath})"), NamedArguments));
-	Writer.AddString(ANSITEXTVIEW("SoftObjectPath"), Context.ExtraSoftObjectPath);
+	Writer.AddString(ANSITEXTVIEW("$type"), TErrorStructFeatures<UE::UnifiedError::StreamableManager::FStreamableManagerAdditionalContext>::GetErrorContextTypeNameAsString());
+	Writer.AddString(ANSITEXTVIEW("$format"), TEXT("(SoftObjectPath: {SoftObjectPath})"));
+	Writer.AddString(ANSITEXTVIEW("SoftObjectPath"), Context.SoftObjectPath);
 	Writer.EndObject();
 }
-DECLARE_ERRORSTRUCT_FEATURES(StreamableManager::FStreamableManagerAdditionalContext);
+
 
diff --git a/Engine/Source/Runtime/Engine/Private/AssetManager.cpp b/Engine/Source/Runtime/Engine/Private/AssetManager.cpp
--- a/Engine/Source/Runtime/Engine/Private/AssetManager.cpp
+++ b/Engine/Source/Runtime/Engine/Private/AssetManager.cpp
@@ -16,6 +16,7 @@
 #include "Engine/BlueprintGeneratedClass.h"
 #include "Engine/GameInstance.h"
 #include "Engine/StreamableManager.h"
+#include "Engine/Experimental/StreamableManagerError.h"
 #include "HAL/PlatformApplicationMisc.h"
 #include "HAL/PlatformFileManager.h"
 #include "Interfaces/IPluginManager.h"
@@ -2114,15 +2115,15 @@ void UAssetManager::OnAssetStateChangeCompleted(TSharedPtr<FStreamableHandle> Bo
 			// If an error occured, don't update the current state with a bad state
 			if (BoundHandle->HasError())
 			{
-				TRefCountPtr<const UE::UnifiedError::StreamableManager::FStreamableManagerErrorDetails> ErrorDetails = BoundHandle->GetError()->GetErrorDetails<UE::UnifiedError::StreamableManager::FStreamableManagerErrorDetails>();
+				const UE::UnifiedError::StreamableManager::FStreamableManagerAdditionalContext* AdditionalContext = BoundHandle->GetError()->GetErrorContext<UE::UnifiedError::StreamableManager::FStreamableManagerAdditionalContext>();
 
-				if (ErrorDetails.IsValid())
+				if (AdditionalContext)
 				{
-					UE_LOG(LogAssetManager, Warning, TEXT("OnAssetStateChangeCompleted: Received error (%s) loading asset! Primary asset id: (%s) erroring asset:%s"), *LexToString(*BoundHandle->GetError()), *PrimaryAssetId.ToString(), *ErrorDetails->GetSoftObjectPath());
+					UE_LOG(LogAssetManager, Warning, TEXT("OnAssetStateChangeCompleted: Received error (%s) loading asset! Primary asset id:(%s) erroring asset:%s"), *BoundHandle->GetError()->GetErrorMessage(true).ToString(), *PrimaryAssetId.ToString(), *AdditionalContext->SoftObjectPath);
 				}
 				else
 				{
-					UE_LOG(LogAssetManager, Warning, TEXT("OnAssetStateChangeCompleted: Received error (%s) loading asset! Primary asset id:(%s)"), *LexToString(*BoundHandle->GetError()), *PrimaryAssetId.ToString());
+					UE_LOG(LogAssetManager, Warning, TEXT("OnAssetStateChangeCompleted: Received error (%s) loading asset! Primary asset id:(%s)"), *BoundHandle->GetError()->GetErrorMessage(true).ToString(), *PrimaryAssetId.ToString());
 				}
 				UE_LOGFMT(LogAssetManager, Warning, "OnAssetStateChangeCompleted: Received error ({Error}) loading asset! Primary asset id:({PrimaryAssetId})", ("Error", BoundHandle->GetError()), ("PrimaryAssetId", *PrimaryAssetId.ToString()));
 			}
diff --git a/Engine/Source/Runtime/Engine/Private/StreamableManager.cpp b/Engine/Source/Runtime/Engine/Private/StreamableManager.cpp
--- a/Engine/Source/Runtime/Engine/Private/StreamableManager.cpp
+++ b/Engine/Source/Runtime/Engine/Private/StreamableManager.cpp
@@ -1225,11 +1225,15 @@ void FStreamableHandle::AsyncLoadCallbackWrapper(const FName& PackageName, UPack
 		if (Result != EAsyncLoadingResult::Succeeded)
 		{
 			MaybeError = UE::UnifiedError::StreamableManager::GetStreamableError(Result);
-			TRefCountPtr<UE::UnifiedError::IErrorDetails> ErrorDetails(new UE::UnifiedError::StreamableManager::FStreamableManagerErrorDetails(TargetName.ToString()));
-			MaybeError->PushErrorDetails(ErrorDetails);
+			/*TRefCountPtr<UE::UnifiedError::IErrorDetails> ErrorDetails(new UE::UnifiedError::StreamableManager::FStreamableManagerErrorDetails(TargetName.ToString()));
+			MaybeError->PushErrorDetails(ErrorDetails);*/
+
+			UE::UnifiedError::StreamableManager::FStreamableManagerAdditionalContext MoreContext = {TargetName.ToString()};
+			MaybeError->PushErrorContext(MoveTemp(MoreContext));
+
+			// todo: Daniel, Devin, this shouldn't be nessisary just use PushErrorContent(T, true) after we have support for $format
+			MaybeError->AppendFormatString(NSLOCTEXT("StreamableManager", "AdditionalContextFormatString", "Missing streamable asset: {StreamableManager_FStreamableManagerAdditionalContext/SoftObjectPath}"));
 
-			/*UE::UnifiedError::StreamableManager::FStreamableManagerAdditionalContext MoreContext = {TargetName.ToString()};
-			MaybeError->PushErrorStruct(MoveTemp(MoreContext));*/
 			SetError(*MaybeError);
 		}
 
