Logging: Added support for optional sub-object references in localized structured log formatting

diff --git a/Engine/Source/Runtime/Core/Private/Logging/StructuredLog.cpp b/Engine/Source/Runtime/Core/Private/Logging/StructuredLog.cpp
--- a/Engine/Source/Runtime/Core/Private/Logging/StructuredLog.cpp
+++ b/Engine/Source/Runtime/Core/Private/Logging/StructuredLog.cpp
@@ -48,9 +48,9 @@ CORE_API bool GConvertBasicLogToLogRecord = false;
 // Experimental feature to prepend log context to the log message during formatting.
 CORE_API bool GPrependLogContextToLogMessage = false;
 
-static constexpr ANSICHAR GSubObjectFieldDelimiter = '/';
+static constexpr ANSICHAR GFieldPathDelimiter = '/';
 static constexpr FAsciiSet GValidLogFieldName("0123456789ABCDEFGHIJKLMNOPQRSTUVWXYZabcdefghijklmnopqrstuvwxyz_");
-static constexpr FAsciiSet GValidLogSubObjectFieldName = GValidLogFieldName | FAsciiSet({GSubObjectFieldDelimiter, '\0'});
+static constexpr FAsciiSet GValidLogFieldPath = GValidLogFieldName | FAsciiSet({GFieldPathDelimiter, '\0'});
 
 static constexpr FAnsiStringView GLogContextsFieldName = ANSITEXTVIEW("$Contexts");
 
@@ -58,7 +58,7 @@ static constexpr FAnsiStringView GLogContextsFieldName = ANSITEXTVIEW("$Contexts
 
 struct FLogTemplateOp
 {
-	enum EOpCode : int32 { OpEnd, OpSkip, OpText, OpName, OpPath, OpIndex, OpLocText, OpCount };
+	enum EOpCode : int32 { OpEnd, OpSkip, OpText, OpName, OpPath, OpIndex, OpLocalized, OpCount };
 
 	static constexpr int32 ValueShift = 3;
 	static_assert(OpCount <= (1 << ValueShift));
@@ -219,23 +219,23 @@ static void LogFieldValue(TStringBuilderBase<CharType>& Out, const FCbFieldView&
 	}
 }
 
-static void AddFieldValue(FFormatNamedArguments& Out, const FCbFieldView& Field)
+static void AddFieldValue(FFormatNamedArguments& Out, FAnsiStringView FieldPath, const FCbFieldView& Field)
 {
-	const FString FieldName(Field.GetName());
+	FString FieldName(FieldPath);
 
 	switch (FCbValue Accessor = Field.GetValue(); Accessor.GetType())
 	{
 	case ECbFieldType::IntegerPositive:
-		Out.Emplace(FieldName, Accessor.AsIntegerPositive());
+		Out.Emplace(MoveTemp(FieldName), Accessor.AsIntegerPositive());
 		return;
 	case ECbFieldType::IntegerNegative:
-		Out.Emplace(FieldName, Accessor.AsIntegerNegative());
+		Out.Emplace(MoveTemp(FieldName), Accessor.AsIntegerNegative());
 		return;
 	case ECbFieldType::Float32:
-		Out.Emplace(FieldName, Accessor.AsFloat32());
+		Out.Emplace(MoveTemp(FieldName), Accessor.AsFloat32());
 		return;
 	case ECbFieldType::Float64:
-		Out.Emplace(FieldName, Accessor.AsFloat64());
+		Out.Emplace(MoveTemp(FieldName), Accessor.AsFloat64());
 		return;
 	default:
 		break;
@@ -244,9 +244,77 @@ static void AddFieldValue(FFormatNamedArguments& Out, const FCbFieldView& Field)
 	// Handle anything that falls through as text.
 	TStringBuilder<128> Text;
 	LogFieldValue(Text, Field);
-	Out.Emplace(FieldName, FText::FromString(FString(Text)));
+	Out.Emplace(MoveTemp(FieldName), FText::FromString(FString(Text)));
 }
 
+class FFieldFinder
+{
+public:
+	inline FFieldFinder(const TCHAR* InFormat, const FCbFieldViewIterator& InFields)
+		: Format(InFormat)
+		, Fields(InFields)
+	{
+	}
+
+	const FCbFieldView& Find(FAnsiStringView Name, int32 IndexHint = -1)
+	{
+		if (IndexHint >= 0)
+		{
+			for (; Index < IndexHint && It; ++Index, ++It)
+			{
+			}
+			if (IndexHint < Index)
+			{
+				It = Fields;
+				for (Index = 0; Index < IndexHint && It; ++Index, ++It);
+			}
+			if (IndexHint == Index && Name.Equals(It.GetName()))
+			{
+				return It;
+			}
+		}
+		const int32 PrevIndex = Index;
+		for (; It; ++Index, ++It)
+		{
+			if (Name.Equals(It.GetName()))
+			{
+				return It;
+			}
+		}
+		It = Fields;
+		for (Index = 0; Index < PrevIndex && It; ++Index, ++It)
+		{
+			if (Name.Equals(It.GetName()))
+			{
+				return It;
+			}
+		}
+		checkf(false, TEXT("Log format requires field '%.*hs' which was not provided. [[%s]]"), Name.Len(), Name.GetData(), Format);
+		return It;
+	}
+
+	FCbFieldView FindByPath(FAnsiStringView Path, int32 IndexHint = -1)
+	{
+		FAnsiStringView Name = Path;
+		bool bMore = String::SplitFirstChar(Path, GFieldPathDelimiter, Name, Path);
+		FCbFieldView Field = Find(Name, IndexHint);
+		while (bMore)
+		{
+			Name = Path;
+			bMore = String::SplitFirstChar(Path, GFieldPathDelimiter, Name, Path);
+			Field = Field.AsObjectView().FindView(Name);
+			checkf(Field, TEXT("Log format requires field '%.*hs' which was not provided. [[%s]]"), Path.Len(), Path.GetData(), Format);
+		}
+		return Field;
+	};
+
+private:
+	const TCHAR* Format;
+	const FCbFieldViewIterator Fields;
+	FCbFieldViewIterator It{Fields};
+	int32 Index{0};
+};
+
 } // UE::Logging::Private
 
 ///////////////////////////////////////////////////////////////////////////////////////////////////////////////////////
@@ -296,9 +364,9 @@ public:
 
 private:
 	template <typename CharType>
-	void FormatLocalizedTo(TStringBuilderBase<CharType>& Out, const FCbFieldViewIterator& Fields) const;
+	void FormatLocalizedTo(TStringBuilderBase<CharType>& Out, const FCbFieldViewIterator& Fields, const uint8* NextOp) const;
 
-	FText FormatLocalizedToText(const FCbFieldViewIterator& Fields) const;
+	FText FormatLocalizedToText(const FCbFieldViewIterator& Fields, const uint8* NextOp) const;
 
 	inline constexpr explicit FLogTemplate(const TCHAR* Format)
 		: StaticFormat(Format)
@@ -328,13 +396,17 @@ FLogTemplate* FLogTemplate::Create(const TCHAR* Format, const FLogTemplateOption
 
 	TArray<FLogTemplateOp, TInlineAllocator<16>> Ops;
 
+	TAnsiStringBuilder<256> FieldPathData;
+	TArray<int32, TInlineAllocator<16>> FieldPathSizes;
+
+	int32 FieldSearchIndex = -1;
 	int32 FormatFieldCount = 0;
-	int32 BracketSearchOffset = 0;
+	int32 SymbolSearchOffset = 0;
 	for (const TCHAR* TextStart = Format;;)
 	{
 		constexpr FAsciiSet Brackets("{}");
-		const TCHAR* const TextEnd = FAsciiSet::FindFirstOrEnd(TextStart + BracketSearchOffset, Brackets);
-		BracketSearchOffset = 0;
+		const TCHAR* const TextEnd = FAsciiSet::FindFirstOrEnd(TextStart + SymbolSearchOffset, Brackets);
+		SymbolSearchOffset = 0;
 
 		// Escaped "{{" or "}}"
 		if ((TextEnd[0] == TEXT('{') && TextEnd[1] == TEXT('{')) ||
@@ -345,7 +417,7 @@ FLogTemplate* FLogTemplate::Create(const TCHAR* Format, const FLogTemplateOption
 			{
 				Ops.Add({FLogTemplateOp::OpSkip, 1});
 				TextStart = TextEnd + 1;
-				BracketSearchOffset = 1;
+				SymbolSearchOffset = 1;
 			}
 			// Text and "{{" or "}}"
 			else
@@ -361,7 +433,6 @@ FLogTemplate* FLogTemplate::Create(const TCHAR* Format, const FLogTemplateOption
 		if (TextStart != TextEnd)
 		{
 			Ops.Add({FLogTemplateOp::OpText, UE_PTRDIFF_TO_INT32(TextEnd - TextStart)});
-			TextStart = TextEnd;
 		}
 
 		// End
@@ -371,40 +442,42 @@ FLogTemplate* FLogTemplate::Create(const TCHAR* Format, const FLogTemplateOption
 			break;
 		}
 
-		// Invalid '}'
-		checkf(*TextStart == TEXT('{'), TEXT("Log format has an unexpected '}' character. Use '}}' to escape it. [[%s]]"), Format);
-
-		// Field
-		const TCHAR* const FieldStart = TextStart;
-		const TCHAR* const FieldNameEnd = FAsciiSet::Skip(FieldStart + 1, GValidLogSubObjectFieldName);
-		checkf(*FieldNameEnd, TEXT("Log format has an unterminated field reference. Use '{{' to escape '{' if needed. [[%s]]"), Format);
-		checkf(*FieldNameEnd == TEXT('}'), TEXT("Log format has invalid character '%c' in field name. "
-			"Use '{{' to escape '{' if needed. Names must match \"[A-Za-z0-9][A-Za-z0-9_]*\". [[%s]]"), *FieldNameEnd, Format);
-		const TCHAR* const FieldEnd = FieldNameEnd + 1;
+		// Parse and validate the field path.
+		const TCHAR* const FieldStart = TextEnd;
+		checkf(*FieldStart == TEXT('{'), TEXT("Log format has an unexpected '%c' character. Use '%c%c' to escape it. [[%s]]"), *FieldStart, *FieldStart, *FieldStart, Format);
+		const TCHAR* const FieldPathEnd = FAsciiSet::Skip(FieldStart + 1, GValidLogFieldPath);
+		checkf(*FieldPathEnd, TEXT("Log format has an unterminated field reference. Use '{{' to escape '{' if needed. [[%s]]"), Format);
+		checkf(*FieldPathEnd == TEXT('}'), TEXT("Log format has invalid character '%c' in field name. "
+			"Use '{{' to escape '{' if needed. Names must match \"[A-Za-z0-9][A-Za-z0-9_]*\". [[%s]]"), *FieldPathEnd, Format);
+		const TCHAR* const FieldEnd = FieldPathEnd + 1;
 		const int32 FieldLen = UE_PTRDIFF_TO_INT32(FieldEnd - FieldStart);
-		const int32 FieldNameLen = FieldLen - 2;
 		checkf(FieldStart[1] != TEXT('_'), TEXT("Log format uses reserved field name '%.*s' with leading '_'. "
-			"Names must match \"[A-Za-z0-9][A-Za-z0-9_]*\". [[%s]]"), FieldNameLen, FieldStart + 1, Format);
-		const bool bHasSubObjectReference = !!Algo::Find(FStringView(FieldStart + 1, FieldNameLen), GSubObjectFieldDelimiter);
+			"Names must match \"[A-Za-z0-9][A-Za-z0-9_]*\". [[%s]]"), FieldLen - 2, FieldStart + 1, Format);
+
+		const int32 FieldPathIndex = FieldPathData.Len();
+		FieldPathData.Append(FieldStart + 1, FieldLen - 2);
+		const FAnsiStringView FieldPath = FieldPathData.ToView().RightChop(FieldPathIndex);
+		FieldPathSizes.Add(FieldPath.Len());
+
+		const bool bHasSubObjectReference = !!Algo::Find(FieldPath, GFieldPathDelimiter);
 		checkf(!bHasSubObjectReference || Options.bAllowSubObjectReferences, TEXT("Log format has a sub-object reference (%c) "
-			"in a context that does not allow them. [[%s]]"), GSubObjectFieldDelimiter, Format);
+			"in a context that does not allow them. [[%s]]"), GFieldPathDelimiter, Format);
 
 		if (bFindFields && !bPositional)
 		{
 			bool bFoundField = false;
-			for (int32 BaseFieldIndex = 0; BaseFieldIndex < FieldCount; ++BaseFieldIndex)
+			for (int32 SearchCount = FieldCount; SearchCount > 0; --SearchCount)
 			{
-				const int32 FieldIndex = (FormatFieldCount + BaseFieldIndex) % FieldCount;
-				const ANSICHAR* FieldName = Fields[FieldIndex].Name;
-				if (FPlatformString::Strncmp(FieldName, FieldStart + 1, FieldNameLen) == 0 && !FieldName[FieldNameLen])
+				FieldSearchIndex = (FieldSearchIndex + 1) % FieldCount;
+				if (FieldPath.Equals(Fields[FieldSearchIndex].Name))
 				{
-					Ops.Add({FLogTemplateOp::OpIndex, FieldIndex});
+					Ops.Add({FLogTemplateOp::OpIndex, FieldSearchIndex});
 					bFoundField = true;
 					break;
 				}
 			}
-			checkf(bFoundField, TEXT("Log format requires field '%.*s' which was not provided. [[%s]]"),
-				FieldNameLen, FieldStart + 1, Format);
+			checkf(bFoundField, TEXT("Log format requires field '%.*hs' which was not provided. [[%s]]"),
+				FieldPath.Len(), FieldPath.GetData(), Format);
 		}
 
 		Ops.Add({bHasSubObjectReference ? FLogTemplateOp::OpPath : FLogTemplateOp::OpName, FieldLen});
@@ -434,42 +507,115 @@ FLogTemplate* FLogTemplate::CreateLocalized(
 	const FLogField* Fields,
 	const int32 FieldCount)
 {
+	// A localized format string consists of an OpLocalized op followed by a sequence of OpSkip and OpName/OpPath ops
+	// that are terminated by an OpEnd op. Only the first occurrence of each name/path is included and everything else
+	// in the format string is skipped. Anything following the last name/path is ignored and not even skipped.
+
 	using namespace Logging::Private;
 
-	checkf(!Options.bAllowSubObjectReferences, TEXT("Localized log templates do not support sub-object references yet."));
+	const bool bFindFields = !!Fields;
+	checkf(!bFindFields || !Options.bAllowSubObjectReferences,
+		TEXT("Validation of field names is not compatible with sub-object references. [[%s]]"), Format);
 
-	FTextFormat TextFormat(FText::AsLocalizable_Advanced(TextNamespace, TextKey, Format));
+	TArray<FLogTemplateOp, TInlineAllocator<16>> Ops;
+	Ops.Add({FLogTemplateOp::OpLocalized});
 
-	const bool bFindFields = !!Fields;
-	if (bFindFields)
+	// Track unique field names to avoid adding multiple ops for the same name.
+	TAnsiStringBuilder<256> FieldPathData;
+	TArray<int32, TInlineAllocator<16>> FieldPathSizes;
+
+	// Parse the format string to find unique field names and optionally validate that required fields are present.
+	int32 FieldSearchIndex = -1;
+	int32 SymbolSearchOffset = 0;
+	for (const TCHAR* TextStart = Format;;)
 	{
-		int32 FormatFieldCount = 0;
-		TArray<FString> TextFormatFieldNames;
-		TextFormat.GetFormatArgumentNames(TextFormatFieldNames);
-		for (const FString& FormatName : TextFormatFieldNames)
+		constexpr FAsciiSet Symbols("`{}");
+		const TCHAR* const TextEnd = FAsciiSet::FindFirstOrEnd(TextStart + SymbolSearchOffset, Symbols);
+		SymbolSearchOffset = 0;
+
+		// Escaped "``" or "`{" or "`}"
+		if (TextEnd[0] == TEXT('`') && (TextEnd[1] == TEXT('`') || TextEnd[1] == TEXT('{') || TextEnd[1] == TEXT('}')))
+		{
+			// Continue the search after the escaped symbol.
+			SymbolSearchOffset = UE_PTRDIFF_TO_INT32(2 + TextEnd - TextStart);
+			continue;
+		}
+
+		// End. Implicitly skips any text after the last field path.
+		if (!*TextEnd)
+		{
+			Ops.Add({FLogTemplateOp::OpEnd});
+			break;
+		}
+
+		// Parse and validate the field path.
+		const TCHAR* const FieldStart = TextEnd;
+		checkf(*FieldStart == TEXT('{'), TEXT("Log format has an unexpected '%c' character. Use '`%c' to escape it. [[%s]]"), *FieldStart, *FieldStart, Format);
+		const TCHAR* const FieldPathEnd = FAsciiSet::Skip(FieldStart + 1, GValidLogFieldPath);
+		checkf(*FieldPathEnd, TEXT("Log format has an unterminated field reference. Use '`{' to escape '{' if needed. [[%s]]"), Format);
+		checkf(*FieldPathEnd == TEXT('}'), TEXT("Log format has invalid character '%c' in field name. "
+			"Use '`{' to escape '{' if needed. Names must match \"[A-Za-z0-9][A-Za-z0-9_]*\". [[%s]]"), *FieldPathEnd, Format);
+		const TCHAR* const FieldEnd = FieldPathEnd + 1;
+		const int32 FieldLen = UE_PTRDIFF_TO_INT32(FieldEnd - FieldStart);
+		checkf(FieldStart[1] != TEXT('_'), TEXT("Log format uses reserved field name '%.*s' with leading '_'. "
+			"Names must match \"[A-Za-z0-9][A-Za-z0-9_]*\". [[%s]]"), FieldLen - 2, FieldStart + 1, Format);
+
+		const int32 FieldPathIndex = FieldPathData.Len();
+		FieldPathData.Append(FieldStart + 1, FieldLen - 2);
+		const FAnsiStringView FieldPath = FieldPathData.ToView().RightChop(FieldPathIndex);
+
+		const bool bHasSubObjectReference = !!Algo::Find(FieldPath, GFieldPathDelimiter);
+		checkf(!bHasSubObjectReference || Options.bAllowSubObjectReferences, TEXT("Log format has a sub-object reference (%c) "
+			"in a context that does not allow them. [[%s]]"), GFieldPathDelimiter, Format);
+
+		// Check if the field path has been seen and skip it if it has.
+		const bool bExistingField = !!Algo::FindByPredicate(FieldPathSizes, [FieldPath, Data = FieldPathData.GetData()](int32 Size) mutable
+		{
+			ON_SCOPE_EXIT { Data += Size; };
+			return FieldPath.Equals(MakeStringView(Data, Size));
+		});
+		if (bExistingField)
+		{
+			// Continue the search after the repeated field path.
+			SymbolSearchOffset = UE_PTRDIFF_TO_INT32(FieldEnd - TextStart);
+			FieldPathData.RemoveSuffix(FieldPath.Len());
+			continue;
+		}
+		FieldPathSizes.Add(FieldPath.Len());
+
+		// Skip the text along with any escaped symbols and repeated field paths.
+		if (TextStart != TextEnd)
+		{
+			Ops.Add({FLogTemplateOp::OpSkip, UE_PTRDIFF_TO_INT32(TextEnd - TextStart)});
+		}
+
+		if (bFindFields)
 		{
 			bool bFoundField = false;
-			const TCHAR* FormatNameStr = *FormatName;
-			const int32 FormatNameLen = FormatName.Len();
-			for (int32 BaseFieldIndex = 0; BaseFieldIndex < FieldCount; ++BaseFieldIndex)
+			for (int32 SearchCount = FieldCount; SearchCount > 0; --SearchCount)
 			{
-				const int32 FieldIndex = (FormatFieldCount + BaseFieldIndex) % FieldCount;
-				const ANSICHAR* FieldName = Fields[FieldIndex].Name;
-				if (FPlatformString::Strncmp(FieldName, FormatNameStr, FormatNameLen) == 0 && !FieldName[FormatNameLen])
+				FieldSearchIndex = (FieldSearchIndex + 1) % FieldCount;
+				if (FieldPath.Equals(Fields[FieldSearchIndex].Name))
 				{
+					Ops.Add({FLogTemplateOp::OpIndex, FieldSearchIndex});
 					bFoundField = true;
 					break;
 				}
 			}
-			checkf(bFoundField, TEXT("Log format requires field '%s' which was not provided. [[%s]]"), FormatNameStr, Format);
-			++FormatFieldCount;
+			checkf(bFoundField, TEXT("Log format requires field '%.*hs' which was not provided. [[%s]]"),
+				FieldPath.Len(), FieldPath.GetData(), Format);
 		}
+
+		Ops.Add({bHasSubObjectReference ? FLogTemplateOp::OpPath : FLogTemplateOp::OpName, FieldLen});
+
+		TextStart = FieldEnd;
 	}
 
-	const FLogTemplateOp Ops[]{{FLogTemplateOp::OpLocText}, {FLogTemplateOp::OpEnd}};
+	FTextFormat TextFormat(FText::AsLocalizable_Advanced(TextNamespace, TextKey, Format));
+
 	const uint32 TotalSize = sizeof(FTextFormat) + sizeof(FLogTemplate) + Algo::TransformAccumulate(Ops, FLogTemplateOp::SaveSize, 0);
-	FTextFormat* NewTextFormat = new(FMemory::Malloc(TotalSize, uint32(FPlatformMath::Max(alignof(FTextFormat), alignof(FLogTemplate))))) FTextFormat(MoveTemp(TextFormat));
-	FLogTemplate* const Template = new(NewTextFormat + 1) FLogTemplate(Format); //-V752
+	FTextFormat* const TemplateTextFormat = new(FMemory::Malloc(TotalSize, uint32(FPlatformMath::Max(alignof(FTextFormat), alignof(FLogTemplate))))) FTextFormat(MoveTemp(TextFormat));
+	FLogTemplate* const Template = new(TemplateTextFormat + 1) FLogTemplate(Format); //-V752
 	uint8* Data = Template->GetOpData();
 	for (const FLogTemplateOp& Op : Ops)
 	{
@@ -483,7 +629,7 @@ void FLogTemplate::Destroy(FLogTemplate* Template)
 	using namespace Logging::Private;
 
 	const uint8* NextOp = Template->GetOpData();
-	if (FLogTemplateOp::Load(NextOp).Code == FLogTemplateOp::OpLocText)
+	if (FLogTemplateOp::Load(NextOp).Code == FLogTemplateOp::OpLocalized)
 	{
 		FTextFormat* TextFormat = (FTextFormat*)Template - 1;
 		TextFormat->~FTextFormat();
@@ -500,68 +646,17 @@ void FLogTemplate::FormatTo(TStringBuilderBase<CharType>& Out, const FCbFieldVie
 {
 	using namespace Logging::Private;
 
-	auto FindField = [&Fields, It = Fields, Index = 0, Format = StaticFormat](FAnsiStringView Name, int32 IndexHint = -1) mutable -> FCbFieldView&
-	{
-		if (IndexHint >= 0)
-		{
-			for (; Index < IndexHint && It; ++Index, ++It)
-			{
-			}
-			if (IndexHint < Index)
-			{
-				It = Fields;
-				for (Index = 0; Index < IndexHint && It; ++Index, ++It);
-			}
-			if (IndexHint == Index && Name.Equals(It.GetName()))
-			{
-				return It;
-			}
-		}
-		const int32 PrevIndex = Index;
-		for (; It; ++Index, ++It)
-		{
-			if (Name.Equals(It.GetName()))
-			{
-				return It;
-			}
-		}
-		It = Fields;
-		for (Index = 0; Index < PrevIndex && It; ++Index, ++It)
-		{
-			if (Name.Equals(It.GetName()))
-			{
-				return It;
-			}
-		}
-		checkf(false, TEXT("Log format requires field '%.*hs' which was not provided. [[%s]]"), Name.Len(), Name.GetData(), Format);
-		return It;
-	};
-
-	auto FindFieldByPath = [&FindField, Format = StaticFormat](FAnsiStringView Path, int32 IndexHint = -1) mutable -> FCbFieldView
-	{
-		FAnsiStringView Name = Path;
-		bool bMore = String::SplitFirstChar(Path, GSubObjectFieldDelimiter, Name, Path);
-		FCbFieldView Field = FindField(Name, IndexHint);
-		while (bMore)
-		{
-			Name = Path;
-			bMore = String::SplitFirstChar(Path, GSubObjectFieldDelimiter, Name, Path);
-			Field = Field.AsObjectView().FindView(Name);
-			checkf(Field, TEXT("Log format requires field '%.*hs' which was not provided. [[%s]]"), Path.Len(), Path.GetData(), Format);
-		}
-		return Field;
-	};
-
 	int32 FieldIndexHint = -1;
 	const uint8* NextOp = GetOpData();
 	const TCHAR* NextFormat = StaticFormat;
+	FFieldFinder FieldFinder(StaticFormat, Fields);
 	for (;;)
 	{
 		const FLogTemplateOp Op = FLogTemplateOp::Load(NextOp);
 		switch (Op.Code)
 		{
-		case FLogTemplateOp::OpLocText:
-			return FormatLocalizedTo(Out, Fields);
+		case FLogTemplateOp::OpLocalized:
+			return FormatLocalizedTo(Out, Fields, NextOp);
 		case FLogTemplateOp::OpEnd:
 			return;
 		case FLogTemplateOp::OpText:
@@ -572,15 +667,15 @@ void FLogTemplate::FormatTo(TStringBuilderBase<CharType>& Out, const FCbFieldVie
 			break;
 		case FLogTemplateOp::OpName:
 			{
-				const auto Name = StringCast<ANSICHAR>(NextFormat + 1, Op.Value - 2);
-				LogFieldValue(Out, FindField(MakeStringView(Name.Get(), Name.Length()), FieldIndexHint));
+				const auto Name = WriteToAnsiString<32>(MakeStringView(NextFormat + 1, Op.Value - 2));
+				LogFieldValue(Out, FieldFinder.Find(Name, FieldIndexHint));
 				FieldIndexHint = -1;
 			}
 			break;
 		case FLogTemplateOp::OpPath:
 			{
-				const auto Path = StringCast<ANSICHAR>(NextFormat + 1, Op.Value - 2);
-				LogFieldValue(Out, FindFieldByPath(MakeStringView(Path.Get(), Path.Length()), FieldIndexHint));
+				const auto Path = WriteToAnsiString<32>(MakeStringView(NextFormat + 1, Op.Value - 2));
+				LogFieldValue(Out, FieldFinder.FindByPath(Path, FieldIndexHint));
 				FieldIndexHint = -1;
 			}
 			break;
@@ -594,9 +689,9 @@ FText FLogTemplate::FormatToText(const FCbFieldViewIterator& Fields) const
 	using namespace Logging::Private;
 
 	const uint8* NextOp = GetOpData();
-	if (FLogTemplateOp::Load(NextOp).Code == FLogTemplateOp::OpLocText)
+	if (FLogTemplateOp::Load(NextOp).Code == FLogTemplateOp::OpLocalized)
 	{
-		return FormatLocalizedToText(Fields);
+		return FormatLocalizedToText(Fields, NextOp);
 	}
 	else
 	{
@@ -607,19 +702,47 @@ FText FLogTemplate::FormatToText(const FCbFieldViewIterator& Fields) const
 }
 
 template <typename CharType>
-FORCENOINLINE void FLogTemplate::FormatLocalizedTo(TStringBuilderBase<CharType>& Out, const FCbFieldViewIterator& Fields) const
+FORCENOINLINE void FLogTemplate::FormatLocalizedTo(TStringBuilderBase<CharType>& Out, const FCbFieldViewIterator& Fields, const uint8* NextOp) const
 {
-	Out.Append(FormatLocalizedToText(Fields).ToString());
+	Out.Append(FormatLocalizedToText(Fields, NextOp).ToString());
 }
 
-FText FLogTemplate::FormatLocalizedToText(const FCbFieldViewIterator& Fields) const
+FText FLogTemplate::FormatLocalizedToText(const FCbFieldViewIterator& Fields, const uint8* NextOp) const
 {
 	using namespace Logging::Private;
 
 	FFormatNamedArguments TextFormatArguments;
-	for (const FCbFieldView& Field : Fields)
+
+	int32 FieldIndexHint = -1;
+	const TCHAR* NextFormat = StaticFormat;
+	FFieldFinder FieldFinder(StaticFormat, Fields);
+	while (NextOp)
 	{
-		AddFieldValue(TextFormatArguments, Field);
+		const FLogTemplateOp Op = FLogTemplateOp::Load(NextOp);
+		switch (Op.Code)
+		{
+		case FLogTemplateOp::OpEnd:
+			NextOp = nullptr;
+			break;
+		case FLogTemplateOp::OpIndex:
+			FieldIndexHint = Op.Value;
+			break;
+		case FLogTemplateOp::OpName:
+			{
+				const auto Name = WriteToAnsiString<32>(MakeStringView(NextFormat + 1, Op.Value - 2));
+				AddFieldValue(TextFormatArguments, Name, FieldFinder.Find(Name, FieldIndexHint));
+				FieldIndexHint = -1;
+			}
+			break;
+		case FLogTemplateOp::OpPath:
+			{
+				const auto Path = WriteToAnsiString<32>(MakeStringView(NextFormat + 1, Op.Value - 2));
+				AddFieldValue(TextFormatArguments, Path, FieldFinder.FindByPath(Path, FieldIndexHint));
+				FieldIndexHint = -1;
+			}
+			break;
+		}
+		NextFormat += Op.GetSkipSize();
 	}
 
 	const FTextFormat* TextFormat = (const FTextFormat*)this - 1;
diff --git a/Engine/Source/Runtime/Core/Tests/Logging/StructuredLogFormatTest.cpp b/Engine/Source/Runtime/Core/Tests/Logging/StructuredLogFormatTest.cpp
--- a/Engine/Source/Runtime/Core/Tests/Logging/StructuredLogFormatTest.cpp
+++ b/Engine/Source/Runtime/Core/Tests/Logging/StructuredLogFormatTest.cpp
@@ -5,6 +5,7 @@
 #if WITH_TESTS
 
 #include "Containers/UnrealString.h"
+#include "Internationalization/Internationalization.h"
 #include "Misc/DateTime.h"
 #include "Misc/Guid.h"
 #include "Misc/Timespan.h"
@@ -116,6 +117,11 @@ TEST_CASE_NAMED(FStructuredLogFormatTest, "Core::Logging::StructuredLogFormat",
 		CHECK(Message.ToView().Equals(Expected));
 	};
 
+	// Invalid
+	//Test(TEXT("Text } Text"), UTF8TEXTVIEW(""));
+	//Test(TEXT("Text { Text"), UTF8TEXTVIEW(""));
+	//Test(TEXT("{Missing}"), UTF8TEXTVIEW(""));
+
 	// Test each of the types that can be represented by compact binary.
 	Test(TEXT("{Object}"), UTF8TEXTVIEW(R"({"X": 1, "Y": 2, "Z": 3})"));
 	Test(TEXT("{Array}"), UTF8TEXTVIEW(R"([1, 2, 3])"));
@@ -149,6 +155,67 @@ TEST_CASE_NAMED(FStructuredLogFormatTest, "Core::Logging::StructuredLogFormat",
 	Test(TEXT("{{{Object/X}, {Object/Y}, {Object/Z}}}"), UTF8TEXTVIEW("{1, 2, 3}"), {.bAllowSubObjectReferences = true});
 	Test(TEXT("{ObjectText/X}"), UTF8TEXTVIEW("1"), {.bAllowSubObjectReferences = true});
 	Test(TEXT("{ObjectWithNestedText/X}"), UTF8TEXTVIEW("0001"), {.bAllowSubObjectReferences = true});
+
+	const auto TestLoc = [&Fields](const FText& FormatText, FUtf8StringView Expected, const FLogTemplateOptions& Options = {})
+	{
+		const TOptional<FString> Namespace = FTextInspector::GetNamespace(FormatText);
+		const TOptional<FString> Key = FTextInspector::GetKey(FormatText);
+		const FString* Format = FTextInspector::GetSourceString(FormatText);
+		FLogTemplate* Template = CreateLogTemplate(**Namespace, **Key, **Format, Options);
+		TUtf8StringBuilder<1024> Message;
+		FormatLogTo(Message, Template, Fields);
+		DestroyLogTemplate(Template);
+		// TODO: Requires a fix for CaptureExpressionsAndValues in LowLevelTestAdapter.h
+		//CAPTURE(Format, Expected, Message);
+		CHECK(Message.ToView().Equals(Expected));
+	};
+
+#define LOCTEXT_NAMESPACE "StructuredLogFormatTest"
+
+	// Invalid
+	//TestLoc(LOCTEXT("Invalid1", "Text ` Text"), UTF8TEXTVIEW(""));
+	//TestLoc(LOCTEXT("Invalid2", "Text } Text"), UTF8TEXTVIEW(""));
+	//TestLoc(LOCTEXT("Invalid3", "Text { Text"), UTF8TEXTVIEW(""));
+	//TestLoc(LOCTEXT("Invalid4", "{Missing}"), UTF8TEXTVIEW(""));
+
+	// Test with an argument modifier.
+	TestLoc(LOCTEXT("EmitterSubheaderText", "Found {IntegerPositive} {IntegerPositive}|plural(one=error,other=errors)!"), UTF8TEXTVIEW("Found 63 errors!"));
+
+	// Test each of the types that can be represented by compact binary.
+	TestLoc(LOCTEXT("Object", "{Object}"), UTF8TEXTVIEW(R"({"X": 1, "Y": 2, "Z": 3})"));
+	TestLoc(LOCTEXT("Array", "{Array}"), UTF8TEXTVIEW(R"([1, 2, 3])"));
+	TestLoc(LOCTEXT("Null", "{Null}"), UTF8TEXTVIEW(R"(null)"));
+	TestLoc(LOCTEXT("Binary", "{Binary}"), UTF8TEXTVIEW(R"("AQIDBAUGBwgJCgsMDQ4PEA==")"));
+	TestLoc(LOCTEXT("String", "{String}"), UTF8TEXTVIEW(R"("Quote" with 4 words.)"));
+	TestLoc(LOCTEXT("IntegerNegative", "{IntegerNegative}"), UTF8TEXTVIEW(R"(-64)"));
+	TestLoc(LOCTEXT("IntegerPositive", "{IntegerPositive}"), UTF8TEXTVIEW(R"(63)"));
+	TestLoc(LOCTEXT("Float", "{Float}"), UTF8TEXTVIEW(R"(128.25)"));
+	TestLoc(LOCTEXT("Double", "{Double}"), UTF8TEXTVIEW(R"(123.456)"));
+	TestLoc(LOCTEXT("False", "{False}"), UTF8TEXTVIEW(R"(false)"));
+	TestLoc(LOCTEXT("True", "{True}"), UTF8TEXTVIEW(R"(true)"));
+	TestLoc(LOCTEXT("ObjectAttachment", "{ObjectAttachment}"), UTF8TEXTVIEW(R"(cb42395cfe025324d80c31c88746d3392f330e58)"));
+	TestLoc(LOCTEXT("BinaryAttachment", "{BinaryAttachment}"), UTF8TEXTVIEW(R"(0ba7b01905a760046bacb86f092f291924c1f24a)"));
+	TestLoc(LOCTEXT("Hash", "{Hash}"), UTF8TEXTVIEW(R"(700b0783bebf169c6d473141e82dd88c67f31ce2)"));
+	TestLoc(LOCTEXT("Uuid", "{Uuid}"), UTF8TEXTVIEW(R"(b1718295-2b3d-3379-9948-ef46ee9de6c3)"));
+	TestLoc(LOCTEXT("DateTime", "{DateTime}"), UTF8TEXTVIEW(R"(2025-04-10T11:15:30.123Z)"));
+	TestLoc(LOCTEXT("TimeSpan", "{TimeSpan}"), UTF8TEXTVIEW(R"(+10.11:15:30.123456700)"));
+	TestLoc(LOCTEXT("ObjectId", "{ObjectId}"), UTF8TEXTVIEW(R"(354fa969c46eb24d1ea03026)"));
+	TestLoc(LOCTEXT("CustomById", "{CustomById}"), UTF8TEXTVIEW(R"({"Id":128,"Data":"ERITFBUWFxg="})"));
+	TestLoc(LOCTEXT("CustomByName", "{CustomByName}"), UTF8TEXTVIEW(R"({"Name":"Custom","Data":"GRobHA=="})"));
+
+	// Test an object with a $text field.
+	TestLoc(LOCTEXT("ObjectText", "{ObjectText}"), UTF8TEXTVIEW(R"(X=1;Y=2;Z=3)"));
+
+	// Test an object and an array with nested $text fields.
+	TestLoc(LOCTEXT("ObjectWithNestedText", "{ObjectWithNestedText}"), UTF8TEXTVIEW(R"({"X": 0001, "Y": 0002, "Z": 0003})"));
+	TestLoc(LOCTEXT("ArrayWithNestedText", "{ArrayWithNestedText}"), UTF8TEXTVIEW(R"([0001, 0002, 0003])"));
+
+	// Test accessing sub-object fields.
+	TestLoc(LOCTEXT("ObjectX", "`{{Object/X}, {Object/Y}, {Object/Z}`}"), UTF8TEXTVIEW("{1, 2, 3}"), {.bAllowSubObjectReferences = true});
+	TestLoc(LOCTEXT("ObjectTextX", "{ObjectText/X}"), UTF8TEXTVIEW("1"), {.bAllowSubObjectReferences = true});
+	TestLoc(LOCTEXT("ObjectWithNestedTextX", "{ObjectWithNestedText/X}"), UTF8TEXTVIEW("0001"), {.bAllowSubObjectReferences = true});
+
+#undef LOCTEXT_NAMESPACE
 }
 
 } // UE
