Add Async.ParallelFor.DisableOversubscription (defaults to false) cvar which specifies if parallel for requests can spawn additional threads when waiting for spawned tasks to finish.
If set to true, it will only use existing threads which is faster when all cores are saturated with existing tasks.

diff --git a/Engine/Source/Runtime/Core/Private/Async/ParallelFor.cpp b/Engine/Source/Runtime/Core/Private/Async/ParallelFor.cpp
--- a/Engine/Source/Runtime/Core/Private/Async/ParallelFor.cpp
+++ b/Engine/Source/Runtime/Core/Private/Async/ParallelFor.cpp
@@ -10,4 +10,11 @@ static FAutoConsoleVariableRef CVarParallelForBackgroundYieldingTimeout(
 	TEXT("Async.ParallelFor.YieldingTimeout"),
 	GParallelForBackgroundYieldingTimeoutMs,
 	TEXT("The timeout (in ms) when background priority parallel for task will yield execution to give higher priority tasks the chance to run.")
+);
+
+CORE_API bool GParallelForDisableOversubscription = true; // Change to false before checkin
+static FAutoConsoleVariableRef CVarParallelForDisableOversubscription(
+	TEXT("Async.ParallelFor.DisableOversubscription"),
+	GParallelForDisableOversubscription,
+	TEXT("If true, do not enable new threads to handle tasks while waiting for a ParallelFor to finish, because new threads can decrease overall performance.")
 );
\ No newline at end of file
diff --git a/Engine/Source/Runtime/Core/Public/Async/Fundamental/Scheduler.h b/Engine/Source/Runtime/Core/Public/Async/Fundamental/Scheduler.h
--- a/Engine/Source/Runtime/Core/Public/Async/Fundamental/Scheduler.h
+++ b/Engine/Source/Runtime/Core/Public/Async/Fundamental/Scheduler.h
@@ -290,7 +290,7 @@ namespace LowLevelTasks
 
 			friend class FOversubscriptionAllowedScope;
 
-			static bool& GetIsOversubscriptionAllowedRef();
+			CORE_API static bool& GetIsOversubscriptionAllowedRef();
 
 		public:
 			static bool IsOversubscriptionAllowed() { return bIsOversubscriptionAllowed; }
diff --git a/Engine/Source/Runtime/Core/Public/Async/ParallelFor.h b/Engine/Source/Runtime/Core/Public/Async/ParallelFor.h
--- a/Engine/Source/Runtime/Core/Public/Async/ParallelFor.h
+++ b/Engine/Source/Runtime/Core/Public/Async/ParallelFor.h
@@ -40,6 +40,7 @@
 namespace UE { namespace LLMPrivate { class FTagData; } }
 
 extern CORE_API int32 GParallelForBackgroundYieldingTimeoutMs;
+extern CORE_API bool GParallelForDisableOversubscription;
 
 // Flags controlling the ParallelFor's behavior.
 enum class EParallelForFlags
@@ -453,7 +454,17 @@ namespace ParallelForImpl
 				// FinishedSignal waits here if some other thread finishes the last item
 				// Data must live on until all of the tasks are cleared which might be long after this function exits
 				TRACE_CPUPROFILER_EVENT_SCOPE(ParallelFor.Wait);
-				FinishedSignal->Wait();
+
+				if (GParallelForDisableOversubscription)
+				{
+					LowLevelTasks::Private::FOversubscriptionAllowedScope Scope(false);
+					FinishedSignal->Wait();
+				}
+				else
+				{
+					// This can spawn new threads to handle tasks
+					FinishedSignal->Wait();
+				}
 			}
 		}
 		checkSlow(LocalExecutor.GetData()->BatchItem.load(std::memory_order_relaxed) * LocalExecutor.GetData()->BatchSize >= LocalExecutor.GetData()->Num);
