Fix module manager trying to get dll exports while under merged modules

#rnx

diff --git a/Engine/Source/Runtime/Core/Private/Modules/ModuleManager.cpp b/Engine/Source/Runtime/Core/Private/Modules/ModuleManager.cpp
--- a/Engine/Source/Runtime/Core/Private/Modules/ModuleManager.cpp
+++ b/Engine/Source/Runtime/Core/Private/Modules/ModuleManager.cpp
@@ -987,11 +987,13 @@ IModuleInterface* FModuleManager::LoadModuleWithFailureReason(const FName InModu
 				// Find our "Initialize<Name>Module" global function, which must exist for all module DLLs
 				FInitializeModuleFunctionPtr InitializeModuleFunctionPtr = FModuleInitializerEntry::FindModule(*InModuleName.ToString());
 
+#if !UE_MERGED_MODULES
 				if (!InitializeModuleFunctionPtr)
 				{
 					// If not found this might be some special case module so look for "InitializeModule" global function
 					InitializeModuleFunctionPtr = (FInitializeModuleFunctionPtr)FPlatformProcess::GetDllExport(ModuleInfo->Handle, TEXT("InitializeModule"));
 				}
+#endif //  !UE_MERGED_MODULES
 
 				if (InitializeModuleFunctionPtr != nullptr)
 				{
