Always construct VectorRegister4Double using a value type

#rnx

diff --git a/Engine/Source/Runtime/Core/Public/Math/UnrealMathNeon.h b/Engine/Source/Runtime/Core/Public/Math/UnrealMathNeon.h
--- a/Engine/Source/Runtime/Core/Public/Math/UnrealMathNeon.h
+++ b/Engine/Source/Runtime/Core/Public/Math/UnrealMathNeon.h
@@ -150,7 +150,7 @@ struct alignas(16) VectorRegister4Double
 	FORCEINLINE VectorRegister4Double(const VectorRegister2Double& xy, const VectorRegister2Double& zw) : XY(xy), ZW(zw) {}
 	FORCEINLINE constexpr VectorRegister4Double(VectorRegister2Double xy, VectorRegister2Double zw, VectorRegisterConstInit) : XY(xy), ZW(zw) {}
 
-	FORCEINLINE VectorRegister4Double(const VectorRegister4Float& From)
+	FORCEINLINE VectorRegister4Double(const VectorRegister4Float From)
 	{
 		XY = vcvt_f64_f32(*(float32x2_t*)&From);
 		ZW = vcvt_high_f64_f32(From);
