[TaskGraph/Scheduler/CoreGlobals]
- Cleanup 5.5 deprecations

diff --git a/Engine/Source/Runtime/Core/Private/Async/Fundamental/ReserveScheduler.cpp b/Engine/Source/Runtime/Core/Private/Async/Fundamental/ReserveScheduler.cpp
deleted file mode 100644
diff --git a/Engine/Source/Runtime/Core/Private/Async/Fundamental/Scheduler.cpp b/Engine/Source/Runtime/Core/Private/Async/Fundamental/Scheduler.cpp
--- a/Engine/Source/Runtime/Core/Private/Async/Fundamental/Scheduler.cpp
+++ b/Engine/Source/Runtime/Core/Private/Async/Fundamental/Scheduler.cpp
@@ -522,11 +522,6 @@ namespace LowLevelTasks
 		return LocalTlsValues.WorkerType != FSchedulerTls::EWorkerType::None && LocalTlsValues.ActiveScheduler == this;
 	}
 
-	bool FSchedulerTls::IsBusyWaiting()
-	{
-		return false;
-	}
-
 	template<typename QueueType, FTask* (QueueType::*DequeueFunction)(bool), bool bIsStandbyWorker>
 	bool FScheduler::TryExecuteTaskFrom(Private::FWaitEvent* WaitEvent, QueueType* Queue, Private::FOutOfWork& OutOfWork, bool bPermitBackgroundWork)
 	{
@@ -685,16 +680,4 @@ namespace LowLevelTasks
 		LocalTlsValues.WorkerType = FSchedulerTls::EWorkerType::None;
 		FMemory::ClearAndDisableTLSCachesOnCurrentThread();
 	}
-
-	void FScheduler::BusyWaitInternal(const FConditional& Conditional, bool ForceAllowBackgroundWork)
-	{
-		TRACE_CPUPROFILER_EVENT_SCOPE(FScheduler::BusyWaitInternal);
-		CSV_SCOPED_TIMING_STAT(Scheduler, BusyWait);
-
-		FOversubscriptionScope _;
-		while (!Conditional())
-		{
-			FPlatformProcess::YieldThread();
-		}
-	}
 }
\ No newline at end of file
diff --git a/Engine/Source/Runtime/Core/Private/Async/TaskGraph.cpp b/Engine/Source/Runtime/Core/Private/Async/TaskGraph.cpp
--- a/Engine/Source/Runtime/Core/Private/Async/TaskGraph.cpp
+++ b/Engine/Source/Runtime/Core/Private/Async/TaskGraph.cpp
@@ -126,8 +126,6 @@ static FAutoConsoleVariableRef CVarTaskDynamicThreadCreation(
 	ECVF_ReadOnly
 );
 
-UE_DEPRECATED(5.5, "This variable is no longer used and will be removed.")
-CORE_API int32 GUseNewTaskBackend = 1;
 CORE_API int32 GNumForegroundWorkers = 2;
 static FAutoConsoleVariableRef CVarNumForegroundWorkers(
 	TEXT("TaskGraph.NumForegroundWorkers"),
diff --git a/Engine/Source/Runtime/Core/Private/HAL/ThreadingBase.cpp b/Engine/Source/Runtime/Core/Private/HAL/ThreadingBase.cpp
--- a/Engine/Source/Runtime/Core/Private/HAL/ThreadingBase.cpp
+++ b/Engine/Source/Runtime/Core/Private/HAL/ThreadingBase.cpp
@@ -186,7 +186,7 @@ CORE_API bool IsInGameThread()
 PRAGMA_DISABLE_DEPRECATION_WARNINGS
 		bool newValue = FTaskTagScope::IsCurrentTag(ETaskTag::EGameThread) || FTaskTagScope::IsRunningDuringStaticInit();
 #if !UE_BUILD_SHIPPING && !UE_BUILD_TEST
-		if (!LowLevelTasks::FSchedulerTls::IsBusyWaiting() &&
+		if (
 #if WITH_CPP_COROUTINES
 			!CoroTask_Detail::FCoroLocalState::IsCoroLaunchedTask() &&
 #endif
@@ -216,7 +216,7 @@ PRAGMA_DISABLE_DEPRECATION_WARNINGS
 	// If this explicitly is a slate thread, not just the main thread running slate
 	bool newValue = FTaskTagScope::IsCurrentTag(ETaskTag::ESlateThread);
 #if !UE_BUILD_SHIPPING && !UE_BUILD_TEST
-	if (!LowLevelTasks::FSchedulerTls::IsBusyWaiting() &&
+	if (
 #if WITH_CPP_COROUTINES
 		!CoroTask_Detail::FCoroLocalState::IsCoroLaunchedTask() &&
 #endif
@@ -255,11 +255,9 @@ CORE_API bool IsInActualRenderingThread()
 PRAGMA_DISABLE_DEPRECATION_WARNINGS
 	bool newValue = FTaskTagScope::IsCurrentTag(ETaskTag::ERenderingThread);
 #if !UE_BUILD_SHIPPING && !UE_BUILD_TEST
-	if (!LowLevelTasks::FSchedulerTls::IsBusyWaiting() 
 #if WITH_CPP_COROUTINES
-		&& !CoroTask_Detail::FCoroLocalState::IsCoroLaunchedTask()
+	if (!CoroTask_Detail::FCoroLocalState::IsCoroLaunchedTask())
 #endif
-		)
 	{
 		bool oldValue = FPlatformTLS::GetCurrentThreadId() == GRenderThreadId;
 		ensureMsgf(oldValue == newValue, TEXT("oldValue(%i) newValue(%i) If this check fails make sure that there is a FTaskTagScope(ETaskTag::ERenderingThread) as deep as possible on the current callstack, you can see the current value in ActiveNamedThreads(%x)"), oldValue, newValue, FTaskTagScope::GetCurrentTag());
@@ -278,7 +276,7 @@ PRAGMA_DISABLE_DEPRECATION_WARNINGS
 		: FTaskTagScope::IsCurrentTag(ETaskTag::ERenderingThread);
 
 #if !UE_BUILD_SHIPPING && !UE_BUILD_TEST
-	if (!LowLevelTasks::FSchedulerTls::IsBusyWaiting() && 
+	if (
 #if WITH_CPP_COROUTINES
 		!CoroTask_Detail::FCoroLocalState::IsCoroLaunchedTask() && 
 #endif
@@ -311,7 +309,7 @@ PRAGMA_DISABLE_DEPRECATION_WARNINGS
 	}
 
 #if !UE_BUILD_SHIPPING && !UE_BUILD_TEST
-	if (!LowLevelTasks::FSchedulerTls::IsBusyWaiting() &&
+	if (
 #if WITH_CPP_COROUTINES
 		!CoroTask_Detail::FCoroLocalState::IsCoroLaunchedTask() &&
 #endif
diff --git a/Engine/Source/Runtime/Core/Private/Misc/CoreGlobals.cpp b/Engine/Source/Runtime/Core/Private/Misc/CoreGlobals.cpp
--- a/Engine/Source/Runtime/Core/Private/Misc/CoreGlobals.cpp
+++ b/Engine/Source/Runtime/Core/Private/Misc/CoreGlobals.cpp
@@ -388,18 +388,10 @@ const TCHAR* LexToString(ELoaderType Type)
 	}
 }
 
-PRAGMA_DISABLE_DEPRECATION_WARNINGS
-/** Whether the editor is currently loading a package or not												*/
-FIsEditorLoadingPackage	GIsEditorLoadingPackage;
-PRAGMA_ENABLE_DEPRECATION_WARNINGS
 /** Whether the cooker is currently loading a package or not												*/
 bool					GIsCookerLoadingPackage			= false;
 /** Whether GWorld points to the play in editor world														*/
 bool					GIsPlayInEditorWorld			= false;
-PRAGMA_DISABLE_DEPRECATION_WARNINGS
-/** Unique ID for multiple PIE instances running in one process												*/
-FPlayInEditorID			GPlayInEditorID;
-PRAGMA_ENABLE_DEPRECATION_WARNINGS
 /** Whether or not PIE was attempting to play from PlayerStart												*/
 bool					GIsPIEUsingPlayerStart			= false;
 /** true if the runtime needs textures to be powers of two													*/
@@ -883,19 +875,6 @@ namespace UE
 	}
 }
 
-PRAGMA_DISABLE_DEPRECATION_WARNINGS
-FPlayInEditorID& FPlayInEditorID::operator= (int32 InOther)
-{
-	UE::SetPlayInEditorID(InOther);
-	return *this;
-}
-
-FPlayInEditorID::operator int32() const
-{
-	return UE::GetPlayInEditorID();
-}
-PRAGMA_ENABLE_DEPRECATION_WARNINGS
-
 bool PRIVATE_GIsEditorLoadingPackage = false;
 
 namespace UE
@@ -941,19 +920,6 @@ namespace UE
 	}
 }
 
-PRAGMA_DISABLE_DEPRECATION_WARNINGS
-FIsEditorLoadingPackage& FIsEditorLoadingPackage::operator= (bool InValue)
-{
-	UE::SetIsEditorLoadingPackage(InValue);
-	return *this;
-}
-
-FIsEditorLoadingPackage::operator bool() const
-{
-	return UE::GetIsEditorLoadingPackage();
-}
-PRAGMA_ENABLE_DEPRECATION_WARNINGS
-
 #undef LOCTEXT_NAMESPACE
 
 bool IsRunningCookOnTheFly()
diff --git a/Engine/Source/Runtime/Core/Public/Async/Fundamental/ReserveScheduler.h b/Engine/Source/Runtime/Core/Public/Async/Fundamental/ReserveScheduler.h
deleted file mode 100644
diff --git a/Engine/Source/Runtime/Core/Public/CoreGlobals.h b/Engine/Source/Runtime/Core/Public/CoreGlobals.h
--- a/Engine/Source/Runtime/Core/Public/CoreGlobals.h
+++ b/Engine/Source/Runtime/Core/Public/CoreGlobals.h
@@ -493,19 +493,6 @@ extern CORE_API void (*SuspendTextureStreamingRenderTasks)();
 /** Resume texture updates caused by completed async IOs. */
 extern CORE_API void (*ResumeTextureStreamingRenderTasks)();
 
-class UE_DEPRECATED(5.5, "This will be removed.") FIsEditorLoadingPackage
-{
-public:
-	CORE_API FIsEditorLoadingPackage& operator= (bool InValue);
-	CORE_API operator bool() const;
-};
-
-PRAGMA_DISABLE_DEPRECATION_WARNINGS
-/** Whether the editor is currently loading a package or not */
-UE_DEPRECATED(5.5, "Use UE::GetIsEditorLoadingPackage() and UE::SetIsEditorLoadingPackage() instead.")
-extern CORE_API FIsEditorLoadingPackage GIsEditorLoadingPackage;
-PRAGMA_ENABLE_DEPRECATION_WARNINGS
-
 namespace UE
 {
 	CORE_API bool GetIsEditorLoadingPackage();
@@ -518,19 +505,6 @@ extern CORE_API bool GIsCookerLoadingPackage;
 /** Whether GWorld points to the play in editor world */
 extern CORE_API bool GIsPlayInEditorWorld;
 
-class UE_DEPRECATED(5.5, "This will be removed.") FPlayInEditorID
-{
-public:
-	CORE_API FPlayInEditorID& operator= (int32 InOther);
-	CORE_API operator int32() const;
-};
-
-PRAGMA_DISABLE_DEPRECATION_WARNINGS
-/** In the editor, this is set to the specific world context PIEInstance that is being currently processed */
-UE_DEPRECATED(5.5, "Use UE::GetPlayInEditorID() and UE::SetPlayInEditorID() instead.")
-extern CORE_API FPlayInEditorID GPlayInEditorID;
-PRAGMA_ENABLE_DEPRECATION_WARNINGS
-
 namespace UE
 {
 	CORE_API int32 GetPlayInEditorID();
diff --git a/Engine/Source/Runtime/Core/Public/Misc/QueuedThreadPoolWrapper.h b/Engine/Source/Runtime/Core/Public/Misc/QueuedThreadPoolWrapper.h
--- a/Engine/Source/Runtime/Core/Public/Misc/QueuedThreadPoolWrapper.h
+++ b/Engine/Source/Runtime/Core/Public/Misc/QueuedThreadPoolWrapper.h
@@ -443,7 +443,6 @@ private:
 			FQueuedWorkInternalData* QueuedWork = Dequeue();
 			if (QueuedWork)
 			{
-				bWakeUpWorker |= LowLevelTasks::FSchedulerTls::IsBusyWaiting();
 				verifySlow(Scheduler->TryLaunch(QueuedWork->Task, bWakeUpWorker ? LowLevelTasks::EQueuePreference::GlobalQueuePreference : LowLevelTasks::EQueuePreference::LocalQueuePreference, bWakeUpWorker));
 				TaskCount.fetch_add(1, std::memory_order_acquire);
 				bWakeUpWorker = true;
diff --git a/Engine/Source/Runtime/Core/Public/Tasks/Task.h b/Engine/Source/Runtime/Core/Public/Tasks/Task.h
--- a/Engine/Source/Runtime/Core/Public/Tasks/Task.h
+++ b/Engine/Source/Runtime/Core/Public/Tasks/Task.h
@@ -18,10 +18,6 @@ namespace UE::Tasks
 	template<typename TaskCollectionType>
 	bool Wait(const TaskCollectionType& Tasks, FTimespan InTimeout = FTimespan::MaxValue());
 
-	template<typename TaskCollectionType>
-	UE_DEPRECATED(5.5, "Use Wait instead.")
-	bool BusyWait(const TaskCollectionType& Tasks, FTimespan InTimeout = FTimespan::MaxValue());
-
 	template<typename TaskType> void AddNested(const TaskType& Nested);
 
 	namespace Private
@@ -50,9 +46,6 @@ namespace UE::Tasks
 			template<typename TaskCollectionType>
 			friend bool UE::Tasks::Wait(const TaskCollectionType& Tasks, FTimespan InTimeout);
 
-			template<typename TaskCollectionType>
-			friend bool UE::Tasks::BusyWait(const TaskCollectionType& Tasks, FTimespan InTimeout);
-
 			template<typename TaskType>
 			friend void UE::Tasks::AddNested(const TaskType& Nested);
 
@@ -100,29 +93,6 @@ namespace UE::Tasks
 				return true;
 			}
 
-			// waits for task's completion for at least the specified amount of time, while executing other tasks.
-			// the call can return much later than the given timeout
-			// @return true if the task is completed
-			UE_DEPRECATED(5.5, "Use Wait instead.")
-			bool BusyWait(FTimespan Timeout = FTimespan::MaxValue()) const
-			{
-				PRAGMA_DISABLE_DEPRECATION_WARNINGS
-				return !IsValid() || Pimpl->BusyWait(FTimeout{ Timeout });
-				PRAGMA_ENABLE_DEPRECATION_WARNINGS
-			}
-
-			// waits for task's completion or the given condition becomes true, while executing other tasks.
-			// the call can return much later than the given condition became true
-			// @return true if the task is completed
-			template<typename ConditionType>
-			UE_DEPRECATED(5.5, "This method will be removed.")
-			bool BusyWait(ConditionType&& Condition) const
-			{
-				PRAGMA_DISABLE_DEPRECATION_WARNINGS
-				return !IsValid() || Pimpl->BusyWait(Forward<ConditionType>(Condition));
-				PRAGMA_ENABLE_DEPRECATION_WARNINGS
-			}
-
 			// try to finish the task on the current thread if it's not already started.
 			// @return true if the task is completed
 			bool TryRetractAndExecute()
@@ -363,32 +333,6 @@ namespace UE::Tasks
 		Task.Wait();
 	}
 
-	// wait for multiple tasks while executing other tasks
-	template<typename TaskCollectionType>
-	UE_DEPRECATED(5.5, "Use Wait instead.")
-	bool BusyWait(const TaskCollectionType& Tasks, FTimespan InTimeout/* = FTimespan::MaxValue()*/)
-	{
-		PRAGMA_DISABLE_DEPRECATION_WARNINGS
-		TaskTrace::FWaitingScope WaitingScope(Private::GetTraceIds(Tasks));
-		TRACE_CPUPROFILER_EVENT_SCOPE(Tasks::BusyWait);
-
-		FTimeout Timeout{ InTimeout };
-
-		// ignore the result as we still have to make sure the task is completed upon returning from this function call
-		Private::TryRetractAndExecute(Tasks, Timeout);
-
-		for (auto& Task : Tasks)
-		{
-			if (Timeout.IsExpired() || !Task.BusyWait(Timeout))
-			{
-				return false;
-			}
-		}
-
-		return true;
-		PRAGMA_ENABLE_DEPRECATION_WARNINGS
-	}
-
 	using FTask = Private::FTaskHandle;
 
 	namespace Private
diff --git a/Engine/Source/Runtime/Core/Public/Tasks/TaskPrivate.h b/Engine/Source/Runtime/Core/Public/Tasks/TaskPrivate.h
--- a/Engine/Source/Runtime/Core/Public/Tasks/TaskPrivate.h
+++ b/Engine/Source/Runtime/Core/Public/Tasks/TaskPrivate.h
@@ -168,10 +168,6 @@ namespace UE::Tasks
 				// store debug name, priority and an adaptor for task execution in low-level task. The task body can't be stored as this task
 				// implementation needs to do some accounting before the task is executed (e.g. maintainance of TLS "current task")
 				LowLevelTasks::ETaskFlags LowLevelTaskFlags = LowLevelTasks::ETaskFlags::DefaultFlags;
-				if (Flags == ETaskFlags::DoNotRunInsideBusyWait)
-				{
-					LowLevelTaskFlags &= ~LowLevelTasks::ETaskFlags::AllowBusyWaiting;
-				}
 				LowLevelTask.Init(InDebugName, InPriority,
 					[
 						this,
@@ -461,30 +457,6 @@ public:
 			// named thread queue and executes them
 			CORE_API void WaitWithNamedThreadsSupport();
 
-			// waits until the task is completed or waiting timed out, while executing other tasks
-			UE_DEPRECATED(5.5, "Use Wait instead")
-			bool BusyWait(FTimeout Timeout)
-			{
-				return Wait(Timeout);
-			}
-
-			// waits until the task is completed or the condition returns true, while executing other tasks
-			template<typename ConditionType>
-			UE_DEPRECATED(5.5, "This method will be removed.")
-			bool BusyWait(ConditionType&& Condition)
-			{
-				TaskTrace::FWaitingScope WaitingScope(GetTraceId());
-				TRACE_CPUPROFILER_EVENT_SCOPE(Tasks::BusyWait);
-
-				// ignore the result as we still have to make sure the task is completed upon returning from this function call
-				TryRetractAndExecute(FTimeout::Never());
-
-				LowLevelTasks::BusyWaitUntil(
-					[this, Condition = Forward<ConditionType>(Condition)]{ return IsCompleted() || Condition(); }
-				);
-				return IsCompleted();
-			}
-
 			TaskTrace::FId GetTraceId() const
 			{
 #if UE_TASK_TRACE_ENABLED
