[trivial] Moved FNamedAESKey into its own header file

diff --git a/Engine/Source/Runtime/Core/Public/Misc/KeyChainUtilities.h b/Engine/Source/Runtime/Core/Public/Misc/KeyChainUtilities.h
--- a/Engine/Source/Runtime/Core/Public/Misc/KeyChainUtilities.h
+++ b/Engine/Source/Runtime/Core/Public/Misc/KeyChainUtilities.h
@@ -1,30 +1,17 @@
 // Copyright Epic Games, Inc. All Rights Reserved.
-
 #pragma once
 
 // HEADER_UNIT_UNSUPPORTED - Includes JsonSerializer.h which is not in Core module
 
+#include "Base64.h"
+#include "CoreDelegates.h"
 #include "CoreMinimal.h"
 #include "HAL/FileManager.h"
-#include "Misc/AES.h"
-#include "Misc/IEngineCrypto.h"
-#include "Misc/CoreDelegates.h"
-#include "Misc/Base64.h"
+#include "IEngineCrypto.h"
+#include "NamedAESKey.h"
+#include "RSA.h"
 #include "Serialization/Archive.h"
 #include "Serialization/JsonSerializer.h"
-#include "RSA.h"
-
-struct FNamedAESKey
-{
-	FString Name;
-	FGuid Guid;
-	FAES::FAESKey Key;
-
-	bool IsValid() const
-	{
-		return Key.IsValid();
-	}
-};
 
 struct FKeyChain
 {
diff --git a/Engine/Source/Runtime/Core/Public/Misc/NamedAESKey.h b/Engine/Source/Runtime/Core/Public/Misc/NamedAESKey.h
new file mode 100644
--- /dev/null
+++ b/Engine/Source/Runtime/Core/Public/Misc/NamedAESKey.h
@@ -0,0 +1,18 @@
+// Copyright Epic Games, Inc. All Rights Reserved.
+#pragma once
+
+#include "Containers/UnrealString.h"
+#include "Misc/AES.h"
+#include "Misc/Guid.h"
+
+struct FNamedAESKey
+{
+	FString Name;
+	FGuid Guid;
+	FAES::FAESKey Key;
+
+	bool IsValid() const
+	{
+		return Key.IsValid();
+	}
+};
