[trivial] Roll back CL#42010180 while I investigate weird build error

#prefligth skip

diff --git a/Engine/Source/Runtime/Core/Public/Misc/KeyChainUtilities.h b/Engine/Source/Runtime/Core/Public/Misc/KeyChainUtilities.h
--- a/Engine/Source/Runtime/Core/Public/Misc/KeyChainUtilities.h
+++ b/Engine/Source/Runtime/Core/Public/Misc/KeyChainUtilities.h
@@ -3,15 +3,27 @@
 
 // HEADER_UNIT_UNSUPPORTED - Includes JsonSerializer.h which is not in Core module
 
-#include "Base64.h"
-#include "CoreDelegates.h"
 #include "CoreMinimal.h"
 #include "HAL/FileManager.h"
-#include "IEngineCrypto.h"
-#include "NamedAESKey.h"
-#include "RSA.h"
+#include "Misc/AES.h"
+#include "Misc/IEngineCrypto.h"
+#include "Misc/CoreDelegates.h"
+#include "Misc/Base64.h"
 #include "Serialization/Archive.h"
 #include "Serialization/JsonSerializer.h"
+#include "RSA.h"
+
+struct FNamedAESKey
+{
+	FString Name;
+	FGuid Guid;
+	FAES::FAESKey Key;
+
+	bool IsValid() const
+	{
+		return Key.IsValid();
+	}
+};
 
 struct FKeyChain
 {
