Logging: Added an overload of CreateLogTemplate that takes the format as FText

diff --git a/Engine/Source/Runtime/Core/Private/Logging/StructuredLog.cpp b/Engine/Source/Runtime/Core/Private/Logging/StructuredLog.cpp
--- a/Engine/Source/Runtime/Core/Private/Logging/StructuredLog.cpp
+++ b/Engine/Source/Runtime/Core/Private/Logging/StructuredLog.cpp
@@ -334,14 +334,14 @@ public:
 
 	static FLogTemplate* CreateStatic(const Logging::Private::FStaticLocalizedLogRecord& Log, const FLogField* Fields, int32 FieldCount)
 	{
-		return CreateLocalized(Log.TextNamespace, Log.TextKey, Log.Format, {}, Fields, FieldCount);
+		const FText FormatText = FText::AsLocalizable_Advanced(Log.TextNamespace, Log.TextKey, Log.Format);
+		return CreateLocalized(FormatText, Log.Format, {}, Fields, FieldCount);
 	}
 
 	static FLogTemplate* Create(const TCHAR* Format, const FLogTemplateOptions& Options, const FLogField* Fields = nullptr, int32 FieldCount = 0);
 
 	static FLogTemplate* CreateLocalized(
-		const TCHAR* TextNamespace,
-		const TCHAR* TextKey,
+		const FText& FormatText,
 		const TCHAR* Format,
 		const FLogTemplateOptions& Options,
 		const FLogField* Fields = nullptr,
@@ -500,8 +500,7 @@ FLogTemplate* FLogTemplate::Create(const TCHAR* Format, const FLogTemplateOption
 }
 
 FLogTemplate* FLogTemplate::CreateLocalized(
-	const TCHAR* TextNamespace,
-	const TCHAR* TextKey,
+	const FText& FormatText,
 	const TCHAR* Format,
 	const FLogTemplateOptions& Options,
 	const FLogField* Fields,
@@ -611,10 +610,8 @@ FLogTemplate* FLogTemplate::CreateLocalized(
 		TextStart = FieldEnd;
 	}
 
-	FTextFormat TextFormat(FText::AsLocalizable_Advanced(TextNamespace, TextKey, Format));
-
 	const uint32 TotalSize = sizeof(FTextFormat) + sizeof(FLogTemplate) + Algo::TransformAccumulate(Ops, FLogTemplateOp::SaveSize, 0);
-	FTextFormat* const TemplateTextFormat = new(FMemory::Malloc(TotalSize, uint32(FPlatformMath::Max(alignof(FTextFormat), alignof(FLogTemplate))))) FTextFormat(MoveTemp(TextFormat));
+	FTextFormat* const TemplateTextFormat = new(FMemory::Malloc(TotalSize, uint32(FPlatformMath::Max(alignof(FTextFormat), alignof(FLogTemplate))))) FTextFormat(FormatText);
 	FLogTemplate* const Template = new(TemplateTextFormat + 1) FLogTemplate(Format); //-V752
 	uint8* Data = Template->GetOpData();
 	for (const FLogTemplateOp& Op : Ops)
@@ -756,9 +753,15 @@ FLogTemplate* CreateLogTemplate(const TCHAR* Format, const FLogTemplateOptions&
 	return FLogTemplate::Create(Format, Options);
 }
 
+FLogTemplate* CreateLogTemplate(const FText& Format, const FLogTemplateOptions& Options)
+{
+	return FLogTemplate::CreateLocalized(Format, *Format.ToString(), Options);
+}
+
 FLogTemplate* CreateLogTemplate(const TCHAR* TextNamespace, const TCHAR* TextKey, const TCHAR* Format, const FLogTemplateOptions& Options)
 {
-	return FLogTemplate::CreateLocalized(TextNamespace, TextKey, Format, Options);
+	const FText FormatText = FText::AsLocalizable_Advanced(TextNamespace, TextKey, Format);
+	return FLogTemplate::CreateLocalized(FormatText, Format, Options);
 }
 
 void DestroyLogTemplate(FLogTemplate* Template)
@@ -819,7 +822,7 @@ FORCENOINLINE static void FormatDynamicRecordMessageTo(TStringBuilderBase<CharTy
 		TEXT("Log record must have both or neither of the text namespace and text key. [[%s]]"), Format);
 
 	const FLogTemplateOptions Options;
-	FLogTemplate* LocalTemplate = TextKey ? FLogTemplate::CreateLocalized(TextNamespace, TextKey, Format, Options) : FLogTemplate::Create(Format, Options);
+	FLogTemplate* LocalTemplate = TextKey ? CreateLogTemplate(TextNamespace, TextKey, Format, Options) : CreateLogTemplate(Format, Options);
 	LocalTemplate->FormatTo(Out, Record.GetFields().CreateViewIterator());
 	FLogTemplate::Destroy(LocalTemplate);
 }
diff --git a/Engine/Source/Runtime/Core/Public/Logging/StructuredLogFormat.h b/Engine/Source/Runtime/Core/Public/Logging/StructuredLogFormat.h
--- a/Engine/Source/Runtime/Core/Public/Logging/StructuredLogFormat.h
+++ b/Engine/Source/Runtime/Core/Public/Logging/StructuredLogFormat.h
@@ -22,6 +22,7 @@ struct FLogTemplateOptions
 };
 
 UE_API FLogTemplate* CreateLogTemplate(const TCHAR* Format, const FLogTemplateOptions& Options = {});
+UE_API FLogTemplate* CreateLogTemplate(const FText& Format, const FLogTemplateOptions& Options = {});
 UE_API FLogTemplate* CreateLogTemplate(const TCHAR* TextNamespace, const TCHAR* TextKey, const TCHAR* Format, const FLogTemplateOptions& Options = {});
 UE_API void DestroyLogTemplate(FLogTemplate* Template);
 
diff --git a/Engine/Source/Runtime/Core/Tests/Logging/StructuredLogFormatTest.cpp b/Engine/Source/Runtime/Core/Tests/Logging/StructuredLogFormatTest.cpp
--- a/Engine/Source/Runtime/Core/Tests/Logging/StructuredLogFormatTest.cpp
+++ b/Engine/Source/Runtime/Core/Tests/Logging/StructuredLogFormatTest.cpp
@@ -118,43 +118,43 @@ TEST_CASE_NAMED(FStructuredLogFormatTest, "Core::Logging::StructuredLogFormat",
 	};
 
 	// Invalid
-	//Test(TEXT("Text } Text"), UTF8TEXTVIEW(""));
-	//Test(TEXT("Text { Text"), UTF8TEXTVIEW(""));
-	//Test(TEXT("{Missing}"), UTF8TEXTVIEW(""));
+	//Test(TEXT("Text } Text"), ANSITEXTVIEW(""));
+	//Test(TEXT("Text { Text"), ANSITEXTVIEW(""));
+	//Test(TEXT("{Missing}"), ANSITEXTVIEW(""));
 
 	// Test each of the types that can be represented by compact binary.
-	Test(TEXT("{Object}"), UTF8TEXTVIEW(R"({"X": 1, "Y": 2, "Z": 3})"));
-	Test(TEXT("{Array}"), UTF8TEXTVIEW(R"([1, 2, 3])"));
-	Test(TEXT("{Null}"), UTF8TEXTVIEW(R"(null)"));
-	Test(TEXT("{Binary}"), UTF8TEXTVIEW(R"("AQIDBAUGBwgJCgsMDQ4PEA==")"));
-	Test(TEXT("{String}"), UTF8TEXTVIEW(R"("Quote" with 4 words.)"));
-	Test(TEXT("{IntegerNegative}"), UTF8TEXTVIEW(R"(-64)"));
-	Test(TEXT("{IntegerPositive}"), UTF8TEXTVIEW(R"(63)"));
-	Test(TEXT("{Float}"), UTF8TEXTVIEW(R"(128.25)"));
-	Test(TEXT("{Double}"), UTF8TEXTVIEW(R"(123.456)"));
-	Test(TEXT("{False}"), UTF8TEXTVIEW(R"(false)"));
-	Test(TEXT("{True}"), UTF8TEXTVIEW(R"(true)"));
-	Test(TEXT("{ObjectAttachment}"), UTF8TEXTVIEW(R"(cb42395cfe025324d80c31c88746d3392f330e58)"));
-	Test(TEXT("{BinaryAttachment}"), UTF8TEXTVIEW(R"(0ba7b01905a760046bacb86f092f291924c1f24a)"));
-	Test(TEXT("{Hash}"), UTF8TEXTVIEW(R"(700b0783bebf169c6d473141e82dd88c67f31ce2)"));
-	Test(TEXT("{Uuid}"), UTF8TEXTVIEW(R"(b1718295-2b3d-3379-9948-ef46ee9de6c3)"));
-	Test(TEXT("{DateTime}"), UTF8TEXTVIEW(R"(2025-04-10T11:15:30.123Z)"));
-	Test(TEXT("{TimeSpan}"), UTF8TEXTVIEW(R"(+10.11:15:30.123456700)"));
-	Test(TEXT("{ObjectId}"), UTF8TEXTVIEW(R"(354fa969c46eb24d1ea03026)"));
-	Test(TEXT("{CustomById}"), UTF8TEXTVIEW(R"({"Id":128,"Data":"ERITFBUWFxg="})"));
-	Test(TEXT("{CustomByName}"), UTF8TEXTVIEW(R"({"Name":"Custom","Data":"GRobHA=="})"));
+	Test(TEXT("{Object}"), ANSITEXTVIEW(R"({"X": 1, "Y": 2, "Z": 3})"));
+	Test(TEXT("{Array}"), ANSITEXTVIEW(R"([1, 2, 3])"));
+	Test(TEXT("{Null}"), ANSITEXTVIEW(R"(null)"));
+	Test(TEXT("{Binary}"), ANSITEXTVIEW(R"("AQIDBAUGBwgJCgsMDQ4PEA==")"));
+	Test(TEXT("{String}"), ANSITEXTVIEW(R"("Quote" with 4 words.)"));
+	Test(TEXT("{IntegerNegative}"), ANSITEXTVIEW(R"(-64)"));
+	Test(TEXT("{IntegerPositive}"), ANSITEXTVIEW(R"(63)"));
+	Test(TEXT("{Float}"), ANSITEXTVIEW(R"(128.25)"));
+	Test(TEXT("{Double}"), ANSITEXTVIEW(R"(123.456)"));
+	Test(TEXT("{False}"), ANSITEXTVIEW(R"(false)"));
+	Test(TEXT("{True}"), ANSITEXTVIEW(R"(true)"));
+	Test(TEXT("{ObjectAttachment}"), ANSITEXTVIEW(R"(cb42395cfe025324d80c31c88746d3392f330e58)"));
+	Test(TEXT("{BinaryAttachment}"), ANSITEXTVIEW(R"(0ba7b01905a760046bacb86f092f291924c1f24a)"));
+	Test(TEXT("{Hash}"), ANSITEXTVIEW(R"(700b0783bebf169c6d473141e82dd88c67f31ce2)"));
+	Test(TEXT("{Uuid}"), ANSITEXTVIEW(R"(b1718295-2b3d-3379-9948-ef46ee9de6c3)"));
+	Test(TEXT("{DateTime}"), ANSITEXTVIEW(R"(2025-04-10T11:15:30.123Z)"));
+	Test(TEXT("{TimeSpan}"), ANSITEXTVIEW(R"(+10.11:15:30.123456700)"));
+	Test(TEXT("{ObjectId}"), ANSITEXTVIEW(R"(354fa969c46eb24d1ea03026)"));
+	Test(TEXT("{CustomById}"), ANSITEXTVIEW(R"({"Id":128,"Data":"ERITFBUWFxg="})"));
+	Test(TEXT("{CustomByName}"), ANSITEXTVIEW(R"({"Name":"Custom","Data":"GRobHA=="})"));
 
 	// Test an object with a $text field.
-	Test(TEXT("{ObjectText}"), UTF8TEXTVIEW(R"(X=1;Y=2;Z=3)"));
+	Test(TEXT("{ObjectText}"), ANSITEXTVIEW(R"(X=1;Y=2;Z=3)"));
 
 	// Test an object and an array with nested $text fields.
-	Test(TEXT("{ObjectWithNestedText}"), UTF8TEXTVIEW(R"({"X": 0001, "Y": 0002, "Z": 0003})"));
-	Test(TEXT("{ArrayWithNestedText}"), UTF8TEXTVIEW(R"([0001, 0002, 0003])"));
+	Test(TEXT("{ObjectWithNestedText}"), ANSITEXTVIEW(R"({"X": 0001, "Y": 0002, "Z": 0003})"));
+	Test(TEXT("{ArrayWithNestedText}"), ANSITEXTVIEW(R"([0001, 0002, 0003])"));
 
 	// Test accessing sub-object fields.
-	Test(TEXT("{{{Object/X}, {Object/Y}, {Object/Z}}}"), UTF8TEXTVIEW("{1, 2, 3}"), {.bAllowSubObjectReferences = true});
-	Test(TEXT("{ObjectText/X}"), UTF8TEXTVIEW("1"), {.bAllowSubObjectReferences = true});
-	Test(TEXT("{ObjectWithNestedText/X}"), UTF8TEXTVIEW("0001"), {.bAllowSubObjectReferences = true});
+	Test(TEXT("{{{Object/X}, {Object/Y}, {Object/Z}}}"), ANSITEXTVIEW("{1, 2, 3}"), {.bAllowSubObjectReferences = true});
+	Test(TEXT("{ObjectText/X}"), ANSITEXTVIEW("1"), {.bAllowSubObjectReferences = true});
+	Test(TEXT("{ObjectWithNestedText/X}"), ANSITEXTVIEW("0001"), {.bAllowSubObjectReferences = true});
 
 	const auto TestLoc = [&Fields](const FText& FormatText, FUtf8StringView Expected, const FLogTemplateOptions& Options = {})
 	{
@@ -173,47 +173,56 @@ TEST_CASE_NAMED(FStructuredLogFormatTest, "Core::Logging::StructuredLogFormat",
 #define LOCTEXT_NAMESPACE "StructuredLogFormatTest"
 
 	// Invalid
-	//TestLoc(LOCTEXT("Invalid1", "Text ` Text"), UTF8TEXTVIEW(""));
-	//TestLoc(LOCTEXT("Invalid2", "Text } Text"), UTF8TEXTVIEW(""));
-	//TestLoc(LOCTEXT("Invalid3", "Text { Text"), UTF8TEXTVIEW(""));
-	//TestLoc(LOCTEXT("Invalid4", "{Missing}"), UTF8TEXTVIEW(""));
+	//TestLoc(LOCTEXT("Invalid1", "Text ` Text"), ANSITEXTVIEW(""));
+	//TestLoc(LOCTEXT("Invalid2", "Text } Text"), ANSITEXTVIEW(""));
+	//TestLoc(LOCTEXT("Invalid3", "Text { Text"), ANSITEXTVIEW(""));
+	//TestLoc(LOCTEXT("Invalid4", "{Missing}"), ANSITEXTVIEW(""));
 
 	// Test with an argument modifier.
-	TestLoc(LOCTEXT("EmitterSubheaderText", "Found {IntegerPositive} {IntegerPositive}|plural(one=error,other=errors)!"), UTF8TEXTVIEW("Found 63 errors!"));
+	TestLoc(LOCTEXT("EmitterSubheaderText", "Found {IntegerPositive} {IntegerPositive}|plural(one=error,other=errors)!"), ANSITEXTVIEW("Found 63 errors!"));
 
 	// Test each of the types that can be represented by compact binary.
-	TestLoc(LOCTEXT("Object", "{Object}"), UTF8TEXTVIEW(R"({"X": 1, "Y": 2, "Z": 3})"));
-	TestLoc(LOCTEXT("Array", "{Array}"), UTF8TEXTVIEW(R"([1, 2, 3])"));
-	TestLoc(LOCTEXT("Null", "{Null}"), UTF8TEXTVIEW(R"(null)"));
-	TestLoc(LOCTEXT("Binary", "{Binary}"), UTF8TEXTVIEW(R"("AQIDBAUGBwgJCgsMDQ4PEA==")"));
-	TestLoc(LOCTEXT("String", "{String}"), UTF8TEXTVIEW(R"("Quote" with 4 words.)"));
-	TestLoc(LOCTEXT("IntegerNegative", "{IntegerNegative}"), UTF8TEXTVIEW(R"(-64)"));
-	TestLoc(LOCTEXT("IntegerPositive", "{IntegerPositive}"), UTF8TEXTVIEW(R"(63)"));
-	TestLoc(LOCTEXT("Float", "{Float}"), UTF8TEXTVIEW(R"(128.25)"));
-	TestLoc(LOCTEXT("Double", "{Double}"), UTF8TEXTVIEW(R"(123.456)"));
-	TestLoc(LOCTEXT("False", "{False}"), UTF8TEXTVIEW(R"(false)"));
-	TestLoc(LOCTEXT("True", "{True}"), UTF8TEXTVIEW(R"(true)"));
-	TestLoc(LOCTEXT("ObjectAttachment", "{ObjectAttachment}"), UTF8TEXTVIEW(R"(cb42395cfe025324d80c31c88746d3392f330e58)"));
-	TestLoc(LOCTEXT("BinaryAttachment", "{BinaryAttachment}"), UTF8TEXTVIEW(R"(0ba7b01905a760046bacb86f092f291924c1f24a)"));
-	TestLoc(LOCTEXT("Hash", "{Hash}"), UTF8TEXTVIEW(R"(700b0783bebf169c6d473141e82dd88c67f31ce2)"));
-	TestLoc(LOCTEXT("Uuid", "{Uuid}"), UTF8TEXTVIEW(R"(b1718295-2b3d-3379-9948-ef46ee9de6c3)"));
-	TestLoc(LOCTEXT("DateTime", "{DateTime}"), UTF8TEXTVIEW(R"(2025-04-10T11:15:30.123Z)"));
-	TestLoc(LOCTEXT("TimeSpan", "{TimeSpan}"), UTF8TEXTVIEW(R"(+10.11:15:30.123456700)"));
-	TestLoc(LOCTEXT("ObjectId", "{ObjectId}"), UTF8TEXTVIEW(R"(354fa969c46eb24d1ea03026)"));
-	TestLoc(LOCTEXT("CustomById", "{CustomById}"), UTF8TEXTVIEW(R"({"Id":128,"Data":"ERITFBUWFxg="})"));
-	TestLoc(LOCTEXT("CustomByName", "{CustomByName}"), UTF8TEXTVIEW(R"({"Name":"Custom","Data":"GRobHA=="})"));
+	TestLoc(LOCTEXT("Object", "{Object}"), ANSITEXTVIEW(R"({"X": 1, "Y": 2, "Z": 3})"));
+	TestLoc(LOCTEXT("Array", "{Array}"), ANSITEXTVIEW(R"([1, 2, 3])"));
+	TestLoc(LOCTEXT("Null", "{Null}"), ANSITEXTVIEW(R"(null)"));
+	TestLoc(LOCTEXT("Binary", "{Binary}"), ANSITEXTVIEW(R"("AQIDBAUGBwgJCgsMDQ4PEA==")"));
+	TestLoc(LOCTEXT("String", "{String}"), ANSITEXTVIEW(R"("Quote" with 4 words.)"));
+	TestLoc(LOCTEXT("IntegerNegative", "{IntegerNegative}"), ANSITEXTVIEW(R"(-64)"));
+	TestLoc(LOCTEXT("IntegerPositive", "{IntegerPositive}"), ANSITEXTVIEW(R"(63)"));
+	TestLoc(LOCTEXT("Float", "{Float}"), ANSITEXTVIEW(R"(128.25)"));
+	TestLoc(LOCTEXT("Double", "{Double}"), ANSITEXTVIEW(R"(123.456)"));
+	TestLoc(LOCTEXT("False", "{False}"), ANSITEXTVIEW(R"(false)"));
+	TestLoc(LOCTEXT("True", "{True}"), ANSITEXTVIEW(R"(true)"));
+	TestLoc(LOCTEXT("ObjectAttachment", "{ObjectAttachment}"), ANSITEXTVIEW(R"(cb42395cfe025324d80c31c88746d3392f330e58)"));
+	TestLoc(LOCTEXT("BinaryAttachment", "{BinaryAttachment}"), ANSITEXTVIEW(R"(0ba7b01905a760046bacb86f092f291924c1f24a)"));
+	TestLoc(LOCTEXT("Hash", "{Hash}"), ANSITEXTVIEW(R"(700b0783bebf169c6d473141e82dd88c67f31ce2)"));
+	TestLoc(LOCTEXT("Uuid", "{Uuid}"), ANSITEXTVIEW(R"(b1718295-2b3d-3379-9948-ef46ee9de6c3)"));
+	TestLoc(LOCTEXT("DateTime", "{DateTime}"), ANSITEXTVIEW(R"(2025-04-10T11:15:30.123Z)"));
+	TestLoc(LOCTEXT("TimeSpan", "{TimeSpan}"), ANSITEXTVIEW(R"(+10.11:15:30.123456700)"));
+	TestLoc(LOCTEXT("ObjectId", "{ObjectId}"), ANSITEXTVIEW(R"(354fa969c46eb24d1ea03026)"));
+	TestLoc(LOCTEXT("CustomById", "{CustomById}"), ANSITEXTVIEW(R"({"Id":128,"Data":"ERITFBUWFxg="})"));
+	TestLoc(LOCTEXT("CustomByName", "{CustomByName}"), ANSITEXTVIEW(R"({"Name":"Custom","Data":"GRobHA=="})"));
 
 	// Test an object with a $text field.
-	TestLoc(LOCTEXT("ObjectText", "{ObjectText}"), UTF8TEXTVIEW(R"(X=1;Y=2;Z=3)"));
+	TestLoc(LOCTEXT("ObjectText", "{ObjectText}"), ANSITEXTVIEW(R"(X=1;Y=2;Z=3)"));
 
 	// Test an object and an array with nested $text fields.
-	TestLoc(LOCTEXT("ObjectWithNestedText", "{ObjectWithNestedText}"), UTF8TEXTVIEW(R"({"X": 0001, "Y": 0002, "Z": 0003})"));
-	TestLoc(LOCTEXT("ArrayWithNestedText", "{ArrayWithNestedText}"), UTF8TEXTVIEW(R"([0001, 0002, 0003])"));
+	TestLoc(LOCTEXT("ObjectWithNestedText", "{ObjectWithNestedText}"), ANSITEXTVIEW(R"({"X": 0001, "Y": 0002, "Z": 0003})"));
+	TestLoc(LOCTEXT("ArrayWithNestedText", "{ArrayWithNestedText}"), ANSITEXTVIEW(R"([0001, 0002, 0003])"));
 
 	// Test accessing sub-object fields.
-	TestLoc(LOCTEXT("ObjectX", "`{{Object/X}, {Object/Y}, {Object/Z}`}"), UTF8TEXTVIEW("{1, 2, 3}"), {.bAllowSubObjectReferences = true});
-	TestLoc(LOCTEXT("ObjectTextX", "{ObjectText/X}"), UTF8TEXTVIEW("1"), {.bAllowSubObjectReferences = true});
-	TestLoc(LOCTEXT("ObjectWithNestedTextX", "{ObjectWithNestedText/X}"), UTF8TEXTVIEW("0001"), {.bAllowSubObjectReferences = true});
+	TestLoc(LOCTEXT("ObjectX", "`{{Object/X}, {Object/Y}, {Object/Z}`}"), ANSITEXTVIEW("{1, 2, 3}"), {.bAllowSubObjectReferences = true});
+	TestLoc(LOCTEXT("ObjectTextX", "{ObjectText/X}"), ANSITEXTVIEW("1"), {.bAllowSubObjectReferences = true});
+	TestLoc(LOCTEXT("ObjectWithNestedTextX", "{ObjectWithNestedText/X}"), ANSITEXTVIEW("0001"), {.bAllowSubObjectReferences = true});
+
+	// Test constructing a template from FText.
+	{
+		FLogTemplate* Template = CreateLogTemplate(LOCTEXT("TextFormat", "FText bWorks={True}"));
+		TUtf8StringBuilder<24> Message;
+		FormatLogTo(Message, Template, Fields);
+		DestroyLogTemplate(Template);
+		CHECK(Message.ToView().Equals(ANSITEXTVIEW(R"(FText bWorks=true)")));
+	}
 
 #undef LOCTEXT_NAMESPACE
 }
