[trivial] Missed change on previous commit

diff --git a/Engine/Source/Runtime/Core/Public/Misc/KeyChainUtilities.h b/Engine/Source/Runtime/Core/Public/Misc/KeyChainUtilities.h
--- a/Engine/Source/Runtime/Core/Public/Misc/KeyChainUtilities.h
+++ b/Engine/Source/Runtime/Core/Public/Misc/KeyChainUtilities.h
@@ -13,18 +13,6 @@
 #include "Serialization/Archive.h"
 #include "Serialization/JsonSerializer.h"
 
-struct FNamedAESKey
-{
-	FString Name;
-	FGuid Guid;
-	FAES::FAESKey Key;
-
-	bool IsValid() const
-	{
-		return Key.IsValid();
-	}
-};
-
 struct FKeyChain
 {
 public:
