Logging: Added support for nested $format in structured log formatting

diff --git a/Engine/Source/Runtime/Core/Private/Logging/StructuredLog.cpp b/Engine/Source/Runtime/Core/Private/Logging/StructuredLog.cpp
--- a/Engine/Source/Runtime/Core/Private/Logging/StructuredLog.cpp
+++ b/Engine/Source/Runtime/Core/Private/Logging/StructuredLog.cpp
@@ -119,10 +119,17 @@ static void LogFieldValue(TStringBuilderBase<CharType>& Out, const FCbFieldView&
 		break;
 	case ECbFieldType::Object:
 	case ECbFieldType::UniformObject:
-		if (FCbFieldView TextField = Accessor.AsObjectView().FindViewIgnoreCase(ANSITEXTVIEW("$text")))
+		if (FCbFieldView TextField = Accessor.AsObjectView().FindViewIgnoreCase(ANSITEXTVIEW("$text")); TextField.IsString())
 		{
 			Out.Append(TextField.AsString());
 		}
+		else if (FCbFieldView FormatField = Accessor.AsObjectView().FindViewIgnoreCase(ANSITEXTVIEW("$format")); FormatField.IsString())
+		{
+			TStringBuilder<128> Format(InPlace, FormatField.AsString());
+			FLogTemplate* Template = CreateLogTemplate(*Format, {.bAllowSubObjectReferences = true});
+			FormatLogTo(Out, Template, Accessor.AsObjectView().CreateViewIterator());
+			DestroyLogTemplate(Template);
+		}
 		else
 		{
 			Out.AppendChar('{');
diff --git a/Engine/Source/Runtime/Core/Tests/Logging/StructuredLogFormatTest.cpp b/Engine/Source/Runtime/Core/Tests/Logging/StructuredLogFormatTest.cpp
--- a/Engine/Source/Runtime/Core/Tests/Logging/StructuredLogFormatTest.cpp
+++ b/Engine/Source/Runtime/Core/Tests/Logging/StructuredLogFormatTest.cpp
@@ -101,6 +101,38 @@ TEST_CASE_NAMED(FStructuredLogFormatTest, "Core::Logging::StructuredLogFormat",
 		Fields.EndObject();
 		Fields.EndArray();
 
+		Fields.BeginObject(ANSITEXTVIEW("ObjectFormat"));
+		Fields.AddInteger(ANSITEXTVIEW("X"), 1);
+		Fields.AddInteger(ANSITEXTVIEW("Y"), 2);
+		Fields.AddInteger(ANSITEXTVIEW("Z"), 3);
+		Fields.AddString(ANSITEXTVIEW("$format"), ANSITEXTVIEW("X={X};Y={Y};Z={Z}"));
+		Fields.EndObject();
+
+		Fields.BeginObject(ANSITEXTVIEW("ObjectWithNestedFormat"));
+		Fields.BeginObject(ANSITEXTVIEW("Point"));
+		Fields.AddInteger(ANSITEXTVIEW("X"), 1);
+		Fields.AddInteger(ANSITEXTVIEW("Y"), 2);
+		Fields.AddInteger(ANSITEXTVIEW("Z"), 3);
+		Fields.AddString(ANSITEXTVIEW("$format"), ANSITEXTVIEW("X={X};Y={Y};Z={Z}"));
+		Fields.EndObject();
+		Fields.AddString(ANSITEXTVIEW("$format"), ANSITEXTVIEW("Target=({Point}); X={Point/X}"));
+		Fields.EndObject();
+
+		Fields.BeginArray(ANSITEXTVIEW("ArrayWithNestedFormat"));
+		Fields.BeginObject();
+		Fields.AddInteger(ANSITEXTVIEW("X"), 1);
+		Fields.AddString(ANSITEXTVIEW("$format"), ANSITEXTVIEW("X={X}"));
+		Fields.EndObject();
+		Fields.BeginObject();
+		Fields.AddInteger(ANSITEXTVIEW("Y"), 2);
+		Fields.AddString(ANSITEXTVIEW("$format"), ANSITEXTVIEW("Y={Y}"));
+		Fields.EndObject();
+		Fields.BeginObject();
+		Fields.AddInteger(ANSITEXTVIEW("Z"), 3);
+		Fields.AddString(ANSITEXTVIEW("$format"), ANSITEXTVIEW("Z={Z}"));
+		Fields.EndObject();
+		Fields.EndArray();
+
 		return Fields.Save();
 	};
 
@@ -151,6 +183,11 @@ TEST_CASE_NAMED(FStructuredLogFormatTest, "Core::Logging::StructuredLogFormat",
 	Test(TEXT("{ObjectWithNestedText}"), ANSITEXTVIEW(R"({"X": 0001, "Y": 0002, "Z": 0003})"));
 	Test(TEXT("{ArrayWithNestedText}"), ANSITEXTVIEW(R"([0001, 0002, 0003])"));
 
+	// Test $format fields.
+	Test(TEXT("{ObjectFormat}"), ANSITEXTVIEW(R"(X=1;Y=2;Z=3)"));
+	Test(TEXT("{ObjectWithNestedFormat}"), ANSITEXTVIEW(R"(Target=(X=1;Y=2;Z=3); X=1)"));
+	Test(TEXT("{ArrayWithNestedFormat}"), ANSITEXTVIEW(R"([X=1, Y=2, Z=3])"));
+
 	// Test accessing sub-object fields.
 	Test(TEXT("{{{Object/X}, {Object/Y}, {Object/Z}}}"), ANSITEXTVIEW("{1, 2, 3}"), {.bAllowSubObjectReferences = true});
 	Test(TEXT("{ObjectText/X}"), ANSITEXTVIEW("1"), {.bAllowSubObjectReferences = true});
@@ -210,6 +247,11 @@ TEST_CASE_NAMED(FStructuredLogFormatTest, "Core::Logging::StructuredLogFormat",
 	TestLoc(LOCTEXT("ObjectWithNestedText", "{ObjectWithNestedText}"), ANSITEXTVIEW(R"({"X": 0001, "Y": 0002, "Z": 0003})"));
 	TestLoc(LOCTEXT("ArrayWithNestedText", "{ArrayWithNestedText}"), ANSITEXTVIEW(R"([0001, 0002, 0003])"));
 
+	// Test an object with a $format field.
+	TestLoc(LOCTEXT("ObjectWithFormat", "{ObjectFormat}"), ANSITEXTVIEW(R"(X=1;Y=2;Z=3)"));
+	TestLoc(LOCTEXT("ObjectWithNestedFormat", "{ObjectWithNestedFormat}"), ANSITEXTVIEW(R"(Target=(X=1;Y=2;Z=3); X=1)"));
+	TestLoc(LOCTEXT("ArrayWithNestedFormat", "{ArrayWithNestedFormat}"), ANSITEXTVIEW(R"([X=1, Y=2, Z=3])"));
+
 	// Test accessing sub-object fields.
 	TestLoc(LOCTEXT("ObjectX", "`{{Object/X}, {Object/Y}, {Object/Z}`}"), ANSITEXTVIEW("{1, 2, 3}"), {.bAllowSubObjectReferences = true});
 	TestLoc(LOCTEXT("ObjectTextX", "{ObjectText/X}"), ANSITEXTVIEW("1"), {.bAllowSubObjectReferences = true});
