Add TBitArray::FindFrom overload that takes an EndIndexExclusive parameter.

#rnx

diff --git a/Engine/Source/Runtime/Core/Public/Containers/BitArray.h b/Engine/Source/Runtime/Core/Public/Containers/BitArray.h
--- a/Engine/Source/Runtime/Core/Public/Containers/BitArray.h
+++ b/Engine/Source/Runtime/Core/Public/Containers/BitArray.h
@@ -1076,12 +1076,12 @@ public:
 	 */
 	[[nodiscard]] int32 Find(bool bValue) const
 	{
-		return FindFromImpl(bValue, 0);
+		return FindFromImpl(bValue, 0, NumBits);
 	}
 
 private:
 
-	int32 FindFromImpl(bool bValue, int32 StartIndex) const
+	int32 FindFromImpl(bool bValue, int32 StartIndex, int32 EndIndexExclusive) const
 	{
 		// Produce a mask for the first iteration
 		uint32 Mask = ~0u << (StartIndex % FBitSet::BitsPerWord);
@@ -1090,8 +1090,7 @@ private:
 		const uint32 Test = bValue ? 0u : ~0u;
 
 		const uint32* RESTRICT DwordArray = GetData();
-		const int32 LocalNumBits = NumBits;
-		const int32 DwordCount = FBitSet::CalculateNumWords(LocalNumBits);
+		const int32 DwordCount = FBitSet::CalculateNumWords(EndIndexExclusive);
 		int32 DwordIndex = FMath::DivideAndRoundDown(StartIndex, NumBitsPerDWORD);
 		while (DwordIndex < DwordCount && (DwordArray[DwordIndex] & Mask) == (Test & Mask))
 		{
@@ -1105,7 +1104,7 @@ private:
 			const uint32 Bits = (bValue ? DwordArray[DwordIndex] : ~DwordArray[DwordIndex]) & Mask;
 			UE_ASSUME(Bits != 0);
 			const int32 LowestBitIndex = FMath::CountTrailingZeros(Bits) + (DwordIndex << NumBitsPerDWORDLogTwo);
-			if (LowestBitIndex < LocalNumBits)
+			if (LowestBitIndex < EndIndexExclusive)
 			{
 				return LowestBitIndex;
 			}
@@ -1171,7 +1170,29 @@ public:
 	{
 		static_assert(!std::is_same_v<IndexType, bool>, "TBitArray::FindFrom: unexpected bool passed as the StartIndex argument");
 		checkSlow((StartIndex >= 0) && (StartIndex <= NumBits));
-		return FindFromImpl(bValue, StartIndex);
+		return FindFromImpl(bValue, StartIndex, NumBits);
+	}
+
+	/**
+	* Finds the first occurrence of the specified value (true/false) in the array, starting from the given bit index, and returns the bit index.
+	* If the specified value is not found from the given index, INDEX_NONE is returned.
+	*
+	* @param  bValue      The value (true/false) to search for.
+	* @param  StartIndex  The index to start the search from.
+	* @param  EndIndexExclusive	  The index to end the search at.
+	*
+	* @pre StartIndex is expected to be in the inclusive range [0, Num()].
+	* @pre EndIndexExclusive is expected to be in the inclusive range [0, Num()].
+	*
+	* @return The index of the first occurrence of the specified value (true/false) from StartIndex to EndIndexExclusive, or INDEX_NONE if not found.
+	*/
+	template <typename IndexType>
+	[[nodiscard]] FORCEINLINE int32 FindFrom(bool bValue, IndexType StartIndex, IndexType EndIndexExclusive) const
+	{
+		static_assert(!std::is_same_v<IndexType, bool>, "TBitArray::FindFrom: unexpected bool passed as the StartIndex argument");
+		checkSlow((StartIndex >= 0) && (StartIndex <= NumBits));
+		checkSlow((EndIndexExclusive >= 0) && (EndIndexExclusive <= NumBits));
+		return FindFromImpl(bValue, StartIndex, EndIndexExclusive);
 	}
 
 	/**
@@ -1229,7 +1250,7 @@ public:
 	 */
 	int32 FindAndSetFirstZeroBit(int32 StartIndex = 0)
 	{
-		const int32 FirstZeroBitIndex = FindFromImpl(false, StartIndex);
+		const int32 FirstZeroBitIndex = FindFromImpl(false, StartIndex, NumBits);
 		if (FirstZeroBitIndex != INDEX_NONE)
 		{
 			(*this)[FirstZeroBitIndex] = true;
