[core] Fix for incremental non-unity test build in Horde (and for submit tool validations).
KeyChainUtilities.h includes headers from RSA and JSON, but is not itself included by anything in Core. This only fails for .cpp files generated for the validation steps.

diff --git a/Engine/Source/Runtime/Core/Core.Build.cs b/Engine/Source/Runtime/Core/Core.Build.cs
--- a/Engine/Source/Runtime/Core/Core.Build.cs
+++ b/Engine/Source/Runtime/Core/Core.Build.cs
@@ -42,7 +42,11 @@ public class Core : ModuleRules
 		PrivateIncludePathModuleNames.AddRange(
 			new string[] {
 				"DerivedDataCache",
-				"TargetPlatform"
+				"TargetPlatform",
+
+				"Json",
+				"RSA"
+
 			}
 			);
 
@@ -290,6 +294,9 @@ public class Core : ModuleRules
 
 		PrivateDefinitions.Add("UE_DEFINE_LEGACY_MATH_CONSTANT_MACRO_NAMES=0");
 
+		PrivateDependencyModuleNames.Add("AtomicQueue");
+		PrivateDependencyModuleNames.Add("AutoRTFM");
+
 		CppCompileWarningSettings.UnsafeTypeCastWarningLevel = WarningLevel.Error;
 
 		IWYUSupport = IWYUSupport.KeepAsIs;
