PR #11044: Fix conversion bug affecting whole numbers over INT_MAX in FCsvWriterHelper::WriteValue

#tests tested on console

diff --git a/Engine/Source/Runtime/Core/Private/ProfilingDebugging/CsvProfiler.cpp b/Engine/Source/Runtime/Core/Private/ProfilingDebugging/CsvProfiler.cpp
--- a/Engine/Source/Runtime/Core/Private/ProfilingDebugging/CsvProfiler.cpp
+++ b/Engine/Source/Runtime/Core/Private/ProfilingDebugging/CsvProfiler.cpp
@@ -1772,7 +1772,7 @@ public:
 
 		if (FMath::Frac((float)Value) == 0.0f)
 		{
-			StrLen = FCStringAnsi::Snprintf(StringBuffer, 256, "%d", int(Value));
+			StrLen = FCStringAnsi::Snprintf(StringBuffer, 256, "%.0f", Value);
 		}
 		else if (FMath::Abs(Value) < 0.1)
 		{
