Removed unused UE_MB2_BOOKKEEPING_AT_THE_END_OF_LARGEBLOCK define in MB2 and dead code paths it guards

diff --git a/Engine/Source/Runtime/Core/Private/HAL/MallocBinned2.cpp b/Engine/Source/Runtime/Core/Private/HAL/MallocBinned2.cpp
--- a/Engine/Source/Runtime/Core/Private/HAL/MallocBinned2.cpp
+++ b/Engine/Source/Runtime/Core/Private/HAL/MallocBinned2.cpp
@@ -369,13 +369,8 @@ FMallocBinned2::FPoolInfo& FMallocBinned2::FPoolList::PushNewPoolToFront(FMalloc
 		OutOfMemory(GMB2PageSize);
 	}
 
-#if !UE_USE_VERYLARGEPAGEALLOCATOR || !UE_MB2_BOOKKEEPING_AT_THE_END_OF_LARGEBLOCK
 	FFreeBlock* Free = new (FreePtr) FFreeBlock(GMB2PageSize, Table.BinSize, InPoolIndex, Allocator.CurrentCanary);
 	check(IsAligned(Free, GMB2PageSize));
-#else
-	FFreeBlock* FreeBlockPtr = GetPoolHeaderFromPointer(FreePtr);
-	FFreeBlock* Free = new (FreeBlockPtr) FFreeBlock(GMB2PageSize, Table.BinSize, InPoolIndex, Allocator.CurrentCanary);
-#endif
 
 #if UE_MB2_ALLOCATOR_STATS
 	Table.TotalAllocatedBins += Free->NumFreeBins;
diff --git a/Engine/Source/Runtime/Core/Public/HAL/MallocBinned2.h b/Engine/Source/Runtime/Core/Public/HAL/MallocBinned2.h
--- a/Engine/Source/Runtime/Core/Public/HAL/MallocBinned2.h
+++ b/Engine/Source/Runtime/Core/Public/HAL/MallocBinned2.h
@@ -34,14 +34,6 @@
 #	define UE_MB2_SMALL_POOL_COUNT			51
 #endif
 
-
-// When book keeping is at the end of FFreeBlock, MallocBinned2 cannot tell if the allocation comes from a large allocation (higher than 64KB, also named as "OSAllocation") 
-// or from VeryLargePageAllocator that fell back to FCachedOSPageAllocator. In both cases the allocation (large or small) might be aligned to 64KB.
-// bookKeeping at the end needs to be disabled if we want VeryLargePageAllocator to properly fallback to regular TCachedOSPageAllocator if needed
-#ifndef UE_MB2_BOOKKEEPING_AT_THE_END_OF_LARGEBLOCK
-#	define UE_MB2_BOOKKEEPING_AT_THE_END_OF_LARGEBLOCK 0
-#endif
-
 // If we are emulating forking on a windows server or are a linux server, enable support for avoiding dirtying pages owned by the parent. 
 #ifndef BINNED2_FORK_SUPPORT
 #	define BINNED2_FORK_SUPPORT (UE_SERVER && (PLATFORM_UNIX || DEFAULT_SERVER_FAKE_FORKS))
@@ -148,20 +140,13 @@ private:
 		FORCEINLINE void* AllocateBin()
 		{
 			--NumFreeBins;
-#if !UE_USE_VERYLARGEPAGEALLOCATOR || !UE_MB2_BOOKKEEPING_AT_THE_END_OF_LARGEBLOCK
+			// this is a pointer to a whole 64KB block
 			if (IsAligned(this, UE_MB2_LARGE_ALLOC))
 			{
 				return (uint8*)this + UE_MB2_LARGE_ALLOC - (NumFreeBins + 1) * BinSize;
 			}
-#else
-			if (IsAligned((uintptr_t)this + sizeof(FFreeBlock), UE_MB2_LARGE_ALLOC))
-			{
-				// The book keeping FreeBlock is at the end of the "page" so we align down to get to the beginning of the page
-				const uintptr_t Ptr = AlignDown((uintptr_t)this, UE_MB2_LARGE_ALLOC);
-				// And we offset the returned pointer based on how many free blocks are left.
-				return (uint8*)Ptr + (NumFreeBins * BinSize);
-			}
-#endif
+
+			// this is a pointer to a single bin within the 64KB block
 			return (uint8*)this + (NumFreeBins * BinSize);
 		}
 
@@ -249,11 +234,7 @@ private:
 
 	static FORCEINLINE FFreeBlock* GetPoolHeaderFromPointer(void* Ptr)
 	{
-#if !UE_USE_VERYLARGEPAGEALLOCATOR || !UE_MB2_BOOKKEEPING_AT_THE_END_OF_LARGEBLOCK
 		return (FFreeBlock*)AlignDown(Ptr, UE_MB2_LARGE_ALLOC);
-#else
-		return (FFreeBlock*)(AlignDown((uintptr_t) Ptr, UE_MB2_LARGE_ALLOC) + UE_MB2_LARGE_ALLOC - sizeof(FFreeBlock));
-#endif
 	}
 
 public:
@@ -319,16 +300,7 @@ protected:
 
 	FORCEINLINE static bool UseSmallAlloc(SIZE_T Size, uint32 Alignment)
 	{
-#if UE_USE_VERYLARGEPAGEALLOCATOR && UE_MB2_BOOKKEEPING_AT_THE_END_OF_LARGEBLOCK
-		if (Alignment > UE_MBC_MIN_SMALL_POOL_ALIGNMENT)
-		{
-			Size = Align(Size, Alignment);
-		}
-		const bool bResult = (Size <= UE_MB2_MAX_SMALL_POOL_SIZE);
-#else
-		const bool bResult = ((Size <= UE_MB2_MAX_SMALL_POOL_SIZE) & (Alignment <= UE_MBC_MIN_SMALL_POOL_ALIGNMENT)); // one branch, not two
-#endif
-		return bResult;
+		return ((Size <= UE_MB2_MAX_SMALL_POOL_SIZE) & (Alignment <= UE_MBC_MIN_SMALL_POOL_ALIGNMENT)); // one branch, not two
 	}
 
 	CORE_API void* MallocSelect(SIZE_T Size, uint32 Alignment, bool bUseSmallPool);
@@ -380,15 +352,7 @@ public:
 protected:
 	FORCEINLINE void* ReallocInline(void* Ptr, SIZE_T NewSize, uint32 Alignment) 
 	{
-#if UE_USE_VERYLARGEPAGEALLOCATOR && UE_MB2_BOOKKEEPING_AT_THE_END_OF_LARGEBLOCK
-		if (Alignment > UE_MBC_MIN_SMALL_POOL_ALIGNMENT && (NewSize <= UE_MB2_MAX_SMALL_POOL_SIZE))
-		{
-			NewSize = Align(NewSize, Alignment);
-		}
-		if (NewSize <= UE_MB2_MAX_SMALL_POOL_SIZE)
-#else
 		if (NewSize <= UE_MB2_MAX_SMALL_POOL_SIZE && Alignment <= UE_MBC_MIN_SMALL_POOL_ALIGNMENT) // one branch, not two
-#endif
 		{
 			FPerThreadFreeBlockLists* Lists = GMallocBinnedPerThreadCaches ? FPerThreadFreeBlockLists::Get() : nullptr;
 			if (Lists && (!Ptr || !IsOSAllocation(Ptr)))
