Logging: Added support for $locformat in structured log formatting

diff --git a/Engine/Source/Runtime/Core/Private/Logging/StructuredLog.cpp b/Engine/Source/Runtime/Core/Private/Logging/StructuredLog.cpp
--- a/Engine/Source/Runtime/Core/Private/Logging/StructuredLog.cpp
+++ b/Engine/Source/Runtime/Core/Private/Logging/StructuredLog.cpp
@@ -25,6 +25,7 @@
 #include "Misc/StringBuilder.h"
 #include "Serialization/CompactBinarySerialization.h"
 #include "Serialization/CompactBinaryValue.h"
+#include "Serialization/CompactBinaryWriter.h"
 #include "Serialization/VarInt.h"
 #include "String/Split.h"
 #include "Templates/Function.h"
@@ -119,19 +120,42 @@ static void LogFieldValue(TStringBuilderBase<CharType>& Out, const FCbFieldView&
 		break;
 	case ECbFieldType::Object:
 	case ECbFieldType::UniformObject:
-		if (FCbFieldView TextField = Accessor.AsObjectView().FindViewIgnoreCase(ANSITEXTVIEW("$text")); TextField.IsString())
-		{
-			Out.Append(TextField.AsString());
-		}
-		else if (FCbFieldView FormatField = Accessor.AsObjectView().FindViewIgnoreCase(ANSITEXTVIEW("$format")); FormatField.IsString())
-		{
-			TStringBuilder<128> Format(InPlace, FormatField.AsString());
-			FLogTemplate* Template = CreateLogTemplate(*Format, {.bAllowSubObjectReferences = true});
-			FormatLogTo(Out, Template, Accessor.AsObjectView().CreateViewIterator());
-			DestroyLogTemplate(Template);
-		}
-		else
 		{
+			FCbObjectView Object = Accessor.AsObjectView();
+
+			// Use $text if present.
+			if (FCbFieldView TextField = Object.FindViewIgnoreCase(ANSITEXTVIEW("$text")); TextField.IsString())
+			{
+				Out.Append(TextField.AsString());
+				break;
+			}
+
+			// Use $format for formatting if present.
+			if (FCbFieldView FormatField = Object.FindViewIgnoreCase(ANSITEXTVIEW("$format")); FormatField.IsString())
+			{
+				TStringBuilder<128> Format(InPlace, FormatField.AsString());
+				FLogTemplate* Template = CreateLogTemplate(*Format, {.bAllowSubObjectReferences = true});
+				FormatLogTo(Out, Template, Object.CreateViewIterator());
+				DestroyLogTemplate(Template);
+				break;
+			}
+
+			// Use $locformat/$locns/$loctext for localized formatting if present.
+			FCbFieldView LocFormatField = Object.FindViewIgnoreCase(ANSITEXTVIEW("$locformat"));
+			FCbFieldView LocNamespaceField = Object.FindViewIgnoreCase(ANSITEXTVIEW("$locns"));
+			FCbFieldView LocKeyField = Object.FindViewIgnoreCase(ANSITEXTVIEW("$lockey"));
+			if (LocFormatField.IsString() && LocNamespaceField.IsString() && LocKeyField.IsString())
+			{
+				TStringBuilder<32> Namespace(InPlace, LocNamespaceField.AsString());
+				TStringBuilder<32> Key(InPlace, LocKeyField.AsString());
+				TStringBuilder<128> Format(InPlace, LocFormatField.AsString());
+				FLogTemplate* Template = CreateLogTemplate(*Namespace, *Key, *Format, {.bAllowSubObjectReferences = true});
+				FormatLogTo(Out, Template, Object.CreateViewIterator());
+				DestroyLogTemplate(Template);
+				break;
+			}
+
+			// Write as JSON as a fallback.
 			Out.AppendChar('{');
 			bool bNeedsComma = false;
 			for (FCbFieldView It : Field)
@@ -791,6 +815,17 @@ FText FormatLogToText(const FLogTemplate* Template, const FCbFieldViewIterator&
 	return Template->FormatToText(Fields);
 }
 
+void SerializeLogFormat(FCbWriter& Writer, const FText& Format)
+{
+	const TOptional<FString> Namespace = FTextInspector::GetNamespace(Format);
+	const TOptional<FString> Key = FTextInspector::GetKey(Format);
+	const FString* Source = FTextInspector::GetSourceString(Format);
+	checkf(Namespace && Key && Source, TEXT("Serializing a localized format string requires a namespace, key, and source string. [[%s]]"), *Format.ToString());
+	Writer.AddString(ANSITEXTVIEW("$locformat"), *Source);
+	Writer.AddString(ANSITEXTVIEW("$locns"), *Namespace);
+	Writer.AddString(ANSITEXTVIEW("$lockey"), *Key);
+}
+
 ///////////////////////////////////////////////////////////////////////////////////////////////////////////////////////
 
 FLogTime FLogTime::Now()
diff --git a/Engine/Source/Runtime/Core/Public/Logging/StructuredLogFormat.h b/Engine/Source/Runtime/Core/Public/Logging/StructuredLogFormat.h
--- a/Engine/Source/Runtime/Core/Public/Logging/StructuredLogFormat.h
+++ b/Engine/Source/Runtime/Core/Public/Logging/StructuredLogFormat.h
@@ -7,6 +7,7 @@
 #define UE_API CORE_API
 
 class FCbFieldViewIterator;
+class FCbWriter;
 class FText;
 
 namespace UE { class FLogTemplate; }
@@ -30,6 +31,14 @@ UE_API void FormatLogTo(FUtf8StringBuilderBase& Out, const FLogTemplate* Templat
 UE_API void FormatLogTo(FWideStringBuilderBase& Out, const FLogTemplate* Template, const FCbFieldViewIterator& Fields);
 UE_API FText FormatLogToText(const FLogTemplate* Template, const FCbFieldViewIterator& Fields);
 
+/**
+ * Serializes a localized log format template to compact binary. Call from SerializeForLog.
+ *
+ * @param Writer   A writer to serialize $locformat, $locns, $lockey fields to.
+ * @param Format   A localized format template to serialize into the writer. Must have a namespace and key!
+ */
+UE_API void SerializeLogFormat(FCbWriter& Writer, const FText& Format);
+
 } // UE
 
 #undef UE_API
diff --git a/Engine/Source/Runtime/Core/Tests/Logging/StructuredLogFormatTest.cpp b/Engine/Source/Runtime/Core/Tests/Logging/StructuredLogFormatTest.cpp
--- a/Engine/Source/Runtime/Core/Tests/Logging/StructuredLogFormatTest.cpp
+++ b/Engine/Source/Runtime/Core/Tests/Logging/StructuredLogFormatTest.cpp
@@ -13,6 +13,8 @@
 #include "Serialization/CompactBinaryWriter.h"
 #include "Tests/TestHarnessAdapter.h"
 
+#define LOCTEXT_NAMESPACE "StructuredLogFormatTest"
+
 namespace UE
 {
 
@@ -133,6 +135,38 @@ TEST_CASE_NAMED(FStructuredLogFormatTest, "Core::Logging::StructuredLogFormat",
 		Fields.EndObject();
 		Fields.EndArray();
 
+		Fields.BeginObject(ANSITEXTVIEW("ObjectLocFormat"));
+		Fields.AddInteger(ANSITEXTVIEW("X"), 1);
+		Fields.AddInteger(ANSITEXTVIEW("Y"), 2);
+		Fields.AddInteger(ANSITEXTVIEW("Z"), 3);
+		SerializeLogFormat(Fields, LOCTEXT("ObjectLocFormat", "X={X};Y={Y};Z={Z}"));
+		Fields.EndObject();
+
+		Fields.BeginObject(ANSITEXTVIEW("ObjectWithNestedLocFormat"));
+		Fields.BeginObject(ANSITEXTVIEW("Point"));
+		Fields.AddInteger(ANSITEXTVIEW("X"), 1);
+		Fields.AddInteger(ANSITEXTVIEW("Y"), 2);
+		Fields.AddInteger(ANSITEXTVIEW("Z"), 3);
+		SerializeLogFormat(Fields, LOCTEXT("ObjectWithNestedLocFormatPoint", "X={X};Y={Y};Z={Z}"));
+		Fields.EndObject();
+		SerializeLogFormat(Fields, LOCTEXT("ObjectWithNestedLocFormat", "Target=({Point}); X={Point/X}"));
+		Fields.EndObject();
+
+		Fields.BeginArray(ANSITEXTVIEW("ArrayWithNestedLocFormat"));
+		Fields.BeginObject();
+		Fields.AddInteger(ANSITEXTVIEW("X"), 1);
+		Fields.AddString(ANSITEXTVIEW("$format"), ANSITEXTVIEW("X={X}"));
+		Fields.EndObject();
+		Fields.BeginObject();
+		Fields.AddInteger(ANSITEXTVIEW("Y"), 2);
+		Fields.AddString(ANSITEXTVIEW("$format"), ANSITEXTVIEW("Y={Y}"));
+		Fields.EndObject();
+		Fields.BeginObject();
+		Fields.AddInteger(ANSITEXTVIEW("Z"), 3);
+		Fields.AddString(ANSITEXTVIEW("$format"), ANSITEXTVIEW("Z={Z}"));
+		Fields.EndObject();
+		Fields.EndArray();
+
 		return Fields.Save();
 	};
 
@@ -188,6 +222,11 @@ TEST_CASE_NAMED(FStructuredLogFormatTest, "Core::Logging::StructuredLogFormat",
 	Test(TEXT("{ObjectWithNestedFormat}"), ANSITEXTVIEW(R"(Target=(X=1;Y=2;Z=3); X=1)"));
 	Test(TEXT("{ArrayWithNestedFormat}"), ANSITEXTVIEW(R"([X=1, Y=2, Z=3])"));
 
+	// Test $locformat fields.
+	Test(TEXT("{ObjectLocFormat}"), ANSITEXTVIEW(R"(X=1;Y=2;Z=3)"));
+	Test(TEXT("{ObjectWithNestedLocFormat}"), ANSITEXTVIEW(R"(Target=(X=1;Y=2;Z=3); X=1)"));
+	Test(TEXT("{ArrayWithNestedLocFormat}"), ANSITEXTVIEW(R"([X=1, Y=2, Z=3])"));
+
 	// Test accessing sub-object fields.
 	Test(TEXT("{{{Object/X}, {Object/Y}, {Object/Z}}}"), ANSITEXTVIEW("{1, 2, 3}"), {.bAllowSubObjectReferences = true});
 	Test(TEXT("{ObjectText/X}"), ANSITEXTVIEW("1"), {.bAllowSubObjectReferences = true});
@@ -207,8 +246,6 @@ TEST_CASE_NAMED(FStructuredLogFormatTest, "Core::Logging::StructuredLogFormat",
 		CHECK(Message.ToView().Equals(Expected));
 	};
 
-#define LOCTEXT_NAMESPACE "StructuredLogFormatTest"
-
 	// Invalid
 	//TestLoc(LOCTEXT("Invalid1", "Text ` Text"), ANSITEXTVIEW(""));
 	//TestLoc(LOCTEXT("Invalid2", "Text } Text"), ANSITEXTVIEW(""));
@@ -247,11 +284,16 @@ TEST_CASE_NAMED(FStructuredLogFormatTest, "Core::Logging::StructuredLogFormat",
 	TestLoc(LOCTEXT("ObjectWithNestedText", "{ObjectWithNestedText}"), ANSITEXTVIEW(R"({"X": 0001, "Y": 0002, "Z": 0003})"));
 	TestLoc(LOCTEXT("ArrayWithNestedText", "{ArrayWithNestedText}"), ANSITEXTVIEW(R"([0001, 0002, 0003])"));
 
-	// Test an object with a $format field.
+	// Test $format fields.
 	TestLoc(LOCTEXT("ObjectWithFormat", "{ObjectFormat}"), ANSITEXTVIEW(R"(X=1;Y=2;Z=3)"));
 	TestLoc(LOCTEXT("ObjectWithNestedFormat", "{ObjectWithNestedFormat}"), ANSITEXTVIEW(R"(Target=(X=1;Y=2;Z=3); X=1)"));
 	TestLoc(LOCTEXT("ArrayWithNestedFormat", "{ArrayWithNestedFormat}"), ANSITEXTVIEW(R"([X=1, Y=2, Z=3])"));
 
+	// Test $locformat fields.
+	TestLoc(LOCTEXT("ObjectWithLocFormatFmt", "{ObjectLocFormat}"), ANSITEXTVIEW(R"(X=1;Y=2;Z=3)"));
+	TestLoc(LOCTEXT("ObjectWithNestedLocFormatFmt", "{ObjectWithNestedLocFormat}"), ANSITEXTVIEW(R"(Target=(X=1;Y=2;Z=3); X=1)"));
+	TestLoc(LOCTEXT("ArrayWithNestedLocFormatFmt", "{ArrayWithNestedLocFormat}"), ANSITEXTVIEW(R"([X=1, Y=2, Z=3])"));
+
 	// Test accessing sub-object fields.
 	TestLoc(LOCTEXT("ObjectX", "`{{Object/X}, {Object/Y}, {Object/Z}`}"), ANSITEXTVIEW("{1, 2, 3}"), {.bAllowSubObjectReferences = true});
 	TestLoc(LOCTEXT("ObjectTextX", "{ObjectText/X}"), ANSITEXTVIEW("1"), {.bAllowSubObjectReferences = true});
@@ -265,10 +307,10 @@ TEST_CASE_NAMED(FStructuredLogFormatTest, "Core::Logging::StructuredLogFormat",
 		DestroyLogTemplate(Template);
 		CHECK(Message.ToView().Equals(ANSITEXTVIEW(R"(FText bWorks=true)")));
 	}
-
-#undef LOCTEXT_NAMESPACE
 }
 
 } // UE
 
+#undef LOCTEXT_NAMESPACE
+
 #endif // WITH_TESTS
