Moving getters from FConfigCacheIni to FConfigSection.
Removing deprecated code from ConfigCacheIni.

diff --git a/Engine/Source/Runtime/Core/Private/Misc/ConfigCacheIni.cpp b/Engine/Source/Runtime/Core/Private/Misc/ConfigCacheIni.cpp
--- a/Engine/Source/Runtime/Core/Private/Misc/ConfigCacheIni.cpp
+++ b/Engine/Source/Runtime/Core/Private/Misc/ConfigCacheIni.cpp
@@ -488,6 +488,104 @@ bool FConfigSection::HandleArrayOfKeyedStructsCommand(FName Key, FString&& Value
 	return bHandledWithKey;
 }
 
+bool FConfigSection::GetString(const TCHAR* Key, FString& Value) const
+{
+	if (const FConfigValue* ConfigValue = Find(FName(Key)))
+	{
+		Value = ConfigValue->GetValue();
+		return true;
+	}
+
+	return false;
+}
+
+bool FConfigSection::GetText(const TCHAR* Section, const TCHAR* Key, FText& Value) const
+{
+	if (const FConfigValue* ConfigValue = Find(FName(Key)))
+	{
+		return FTextStringHelper::ReadFromBuffer(*ConfigValue->GetValue(), Value, Section) != nullptr;
+	}
+
+	return false;
+}
+
+bool FConfigSection::GetInt(const TCHAR* Key, int32& Value) const
+{
+	FString Text;
+	if (GetString(Key, Text))
+	{
+		Value = FCString::Atoi(*Text);
+		return true;
+	}
+	return false;
+}
+
+bool FConfigSection::GetUInt(const TCHAR* Key, uint32& Value) const
+{
+	FString Text;
+	if (GetString(Key, Text))
+	{
+		Value = static_cast<uint32>(FCString::Atoi(*Text));
+		return true;
+	}
+	return false;
+}
+
+bool FConfigSection::GetFloat(const TCHAR* Key, float& Value) const
+{
+	FString Text;
+	if (GetString(Key, Text))
+	{
+		Value = FCString::Atof(*Text);
+		return true;
+	}
+	return false;
+}
+
+bool FConfigSection::GetDouble(const TCHAR* Key, double& Value) const
+{
+	FString Text;
+	if (GetString(Key, Text))
+	{
+		Value = FCString::Atof(*Text);
+		return true;
+	}
+	return false;
+}
+
+bool FConfigSection::GetInt64(const TCHAR* Key, int64& Value) const
+{
+	FString Text;
+	if (GetString(Key, Text))
+	{
+		Value = FCString::Atoi64(*Text);
+		return true;
+	}
+	return false;
+}
+
+bool FConfigSection::GetBool(const TCHAR* Key, bool& Value) const
+{
+	FString Text;
+	if (GetString(Key, Text))
+	{
+		Value = FCString::ToBool(*Text);
+		return true;
+	}
+	return false;
+}
+
+int32 FConfigSection::GetArray(const TCHAR* Key, TArray<FString>& Value) const
+{
+	const FName KeyName(Key);
+	Value.Empty();
+
+	MultiFind(KeyName, Value, true);
+
+	// if we have values, or we were initialized to empty
+	return Value.Num() || EmptyInitializedKeys.Contains(KeyName);
+}
+
 // Look through the file's per object config ArrayOfStruct keys and see if this section matches
 template<typename SectionType>
 void FixupArrayOfStructKeysForSection(SectionType* Section, const FString& SectionName, const TMap<FString, TMap<FName, FString> >& PerObjectConfigKeys)
@@ -1113,11 +1211,6 @@ bool FConfigFile::operator!=( const FConfigFile& Other ) const
 	return ! (FConfigFile::operator==(Other));
 }
 
-FConfigSection* FConfigFile::FindOrAddSection(const FString& SectionName)
-{
-	return FindOrAddSectionInternal(SectionName);
-}
-
 FConfigSection* FConfigFile::FindOrAddSectionInternal(const FString& SectionName)
 {
 	FConfigSection* Section = FindInternal(SectionName);
@@ -2545,101 +2638,11 @@ void FConfigFile::Dump(FOutputDevice& Ar)
 	}
 }
 
-bool FConfigFile::GetString( const TCHAR* Section, const TCHAR* Key, FString& Value ) const
-{
-	const FConfigSection* Sec = FindSection( Section );
-	if( Sec == nullptr )
-	{
-		return false;
-	}
-	const FConfigValue* PairString = Sec->Find( Key );
-	if( PairString == nullptr )
-	{
-		return false;
-	}
-	Value = PairString->GetValue();
-	return true;
-}
-
-bool FConfigFile::GetText( const TCHAR* Section, const TCHAR* Key, FText& Value ) const
-{
-	const FConfigSection* Sec = FindSection( Section );
-	if( Sec == nullptr )
-	{
-		return false;
-	}
-	const FConfigValue* PairString = Sec->Find( Key );
-	if( PairString == nullptr )
-	{
-		return false;
-	}
-	return FTextStringHelper::ReadFromBuffer( *PairString->GetValue(), Value, Section ) != nullptr;
-}
-
-bool FConfigFile::GetInt(const TCHAR* Section, const TCHAR* Key, int32& Value) const
-{
-	FString Text;
-	if (GetString(Section, Key, Text))
-	{
-		Value = FCString::Atoi(*Text);
-		return true;
-	}
-	return false;
-}
-
-bool FConfigFile::GetFloat(const TCHAR* Section, const TCHAR* Key, float& Value) const
-{
-	FString Text;
-	if (GetString(Section, Key, Text))
-	{
-		Value = FCString::Atof(*Text);
-		return true;
-	}
-	return false;
-}
-
-bool FConfigFile::GetDouble(const TCHAR* Section, const TCHAR* Key, double& Value) const
-{
-	FString Text;
-	if (GetString(Section, Key, Text))
-	{
-		Value = FCString::Atod(*Text);
-		return true;
-	}
-	return false;
-}
-
-bool FConfigFile::GetInt64( const TCHAR* Section, const TCHAR* Key, int64& Value ) const
-{
-	FString Text; 
-	if( GetString( Section, Key, Text ) )
-	{
-		Value = FCString::Atoi64(*Text);
-		return true;
-	}
-	return false;
-}
-bool FConfigFile::GetBool(const TCHAR* Section, const TCHAR* Key, bool& Value ) const
-{
-	FString Text;
-	if ( GetString(Section, Key, Text ))
-	{
-		Value = FCString::ToBool(*Text);
-		return 1;
-	}
-	return 0;
-}
-
 int32 FConfigFile::GetArray(const TCHAR* Section, const TCHAR* Key, TArray<FString>& Value) const
 {
-	Value.Empty();
-	const FConfigSection* Sec = FindSection(Section);
-	FName KeyName(Key);
-	if (Sec != nullptr)
+	if (const FConfigSection* ConfigSection = FindSection(Section))
 	{
-		Sec->MultiFind(KeyName, Value, true);
-		// if we have values, or we were initialized to empty
-		return Value.Num() || Sec->EmptyInitializedKeys.Contains(KeyName);
+		return ConfigSection->GetArray(Key, Value);
 	}
 #if !UE_BUILD_SHIPPING
 	CheckLongSectionNames(Section, this);
@@ -4325,20 +4328,6 @@ bool FConfigCacheIni::GetSection( const TCHAR* Section, TArray<FString>& Result,
 	return true;
 }
 
-FConfigSection* FConfigCacheIni::GetSectionPrivate( const TCHAR* Section, const bool Force, const bool Const, const FString& Filename )
-{
-	FConfigSection* Sec = const_cast<FConfigSection*>(GetSection(Section, Force, Filename));
-	
-	// handle the non-const case
-	if ((!Const || Force) && Sec != nullptr)
-	{
-		FConfigFile* File = Find(Filename);
-		File->Dirty = true;
-	}
-
-	return Sec;
-}
-
 const FConfigSection* FConfigCacheIni::GetSection( const TCHAR* Section, const bool Force, const FString& Filename )
 {
 	FRemoteConfig::Get()->FinishRead(*Filename); // Ensure the remote file has been loaded and processed
diff --git a/Engine/Source/Runtime/Core/Public/Misc/ConfigCacheIni.h b/Engine/Source/Runtime/Core/Public/Misc/ConfigCacheIni.h
--- a/Engine/Source/Runtime/Core/Public/Misc/ConfigCacheIni.h
+++ b/Engine/Source/Runtime/Core/Public/Misc/ConfigCacheIni.h
@@ -300,16 +300,6 @@ public:
 		return SavedValue;
 	}
 #if UE_WITH_CONFIG_TRACKING
-	UE_DEPRECATED(5.4, "No longer written. Use UE::ConfigAccessTracking::AddConfigValueReadCallback instead")
-	inline const bool HasBeenRead() const
-	{
-		return false;
-	}
-	UE_DEPRECATED(5.4, "No longer read.")
-	inline void SetHasBeenRead(bool InBRead ) const
-	{
-	}
-
 	void SetSectionAccess(UE::ConfigAccessTracking::FSection* InSectionAccess)
 	{
 		SectionAccess = InSectionAccess;
@@ -464,6 +454,15 @@ public:
 		}
 	}
 
+	CORE_API bool GetString(const TCHAR* Key, FString& Value) const;
+	CORE_API bool GetText(const TCHAR* Section, const TCHAR* Key, FText& Value) const;
+	CORE_API bool GetInt(const TCHAR* Key, int32& Value) const;
+	CORE_API bool GetUInt(const TCHAR* Key, uint32& Value) const;
+	CORE_API bool GetFloat(const TCHAR* Key, float& Value) const;
+	CORE_API bool GetDouble(const TCHAR* Key, double& Value) const;
+	CORE_API bool GetInt64(const TCHAR* Key, int64& Value) const;
+	CORE_API bool GetBool(const TCHAR* Key, bool& Value) const;
+	CORE_API int32 GetArray(const TCHAR* Key, TArray<FString>& Value) const;
 
 	// look for "array of struct" keys for overwriting single entries of an array
 	TMap<FName, FString> ArrayOfStructKeys;
@@ -618,26 +617,8 @@ public:
 	CORE_API FConfigFile& operator=(FConfigFile&& Other);
 
 	// looks for a section by name, and creates an empty one if it can't be found
-	UE_DEPRECATED(5.4, "Use FindOrAddConfigSection, and/or use the new AddToSection, etc APIs to modify sections without retrieving the section. See top of ConfigCacheIni.h for more info.")
-	CORE_API FConfigSection* FindOrAddSection(const FString& Name);
 	CORE_API const FConfigSection* FindOrAddConfigSection(const FString& Name);
 
-	///////////////////////////////////
-	// Replacement functionality of TMap so we can deprecate the direct access to FConfigSection
-	UE_DEPRECATED(5.4, "Use FindSection, and/or use the new AddToSection, etc APIs to modify sections without retrieving the section. See top of ConfigCacheIni.h for more info.")
-	FORCEINLINE const FConfigSection* Find(const FString& SectionName) const
-	{
-		UE::TReadScopeLock ScopeLock(ConfigFileMapLock);
-		return FConfigFileMap::Find(SectionName);
-	}
-
-	UE_DEPRECATED(5.4, "Use FindSection, and/or use the new AddToSection, etc APIs to modify sections without retrieving the section. See top of ConfigCacheIni.h for more info.")
-	FORCEINLINE FConfigSection* Find(const FString& SectionName)
-	{
-		UE::TReadScopeLock ScopeLock(ConfigFileMapLock);
-		return FConfigFileMap::Find(SectionName);
-	}
-
 	FORCEINLINE const FConfigSection* FindSection(const FString& SectionName) const
 	{
 		UE::TReadScopeLock ScopeLock(ConfigFileMapLock);
@@ -722,33 +703,12 @@ public:
 		FConfigFileMap::Reset();
 	}
 
-	UE_DEPRECATED(5.4, "Use FindOrAddConfigSection, and/or use the new AddToSection, etc APIs to modify sections without retrieving the section. See top of ConfigCacheIni.h for more info.")
-	FORCEINLINE ValueType& FindOrAdd(const FString& Key)
-	{
-		UE::TWriteScopeLock ScopeLock(ConfigFileMapLock);
-		return FConfigFileMap::FindOrAdd(Key);
-	}
-
-	UE_DEPRECATED(5.4, "Use const ranged for iterators, (wrap your FConfigFile variable in AsConst to force the const iterator). See top of ConfigCacheIni.h for more info.")
-	FORCEINLINE TRangedForIterator begin()
-	{
-		UE::TReadScopeLock ScopeLock(ConfigFileMapLock);
-		return TRangedForIterator(Pairs.begin());
-	}
-
 	FORCEINLINE TRangedForConstIterator begin() const
 	{
 		UE::TReadScopeLock ScopeLock(ConfigFileMapLock); 
 		return TRangedForConstIterator(Pairs.begin());
 	}
 
-	UE_DEPRECATED(5.4, "Use const ranged for iterators, (wrap your FConfigFile variable in AsConst to force the const iterator). See top of ConfigCacheIni.h for more info.")
-	FORCEINLINE TRangedForIterator end()
-	{
-		UE::TReadScopeLock ScopeLock(ConfigFileMapLock); 
-		return TRangedForIterator(Pairs.end());
-	}
-
 	FORCEINLINE TRangedForConstIterator end() const
 	{
 		UE::TReadScopeLock ScopeLock(ConfigFileMapLock);
@@ -839,13 +799,78 @@ private:
 public:
 	CORE_API void Dump(FOutputDevice& Ar);
 
-	CORE_API bool GetString( const TCHAR* Section, const TCHAR* Key, FString& Value ) const;
-	CORE_API bool GetText( const TCHAR* Section, const TCHAR* Key, FText& Value ) const;
-	CORE_API bool GetInt(const TCHAR* Section, const TCHAR* Key, int32& Value) const;
-	CORE_API bool GetFloat(const TCHAR* Section, const TCHAR* Key, float& Value) const;
-	CORE_API bool GetDouble(const TCHAR* Section, const TCHAR* Key, double& Value) const;
-	CORE_API bool GetInt64( const TCHAR* Section, const TCHAR* Key, int64& Value ) const;
-	CORE_API bool GetBool( const TCHAR* Section, const TCHAR* Key, bool& Value ) const;
+	bool GetString(const TCHAR* Section, const TCHAR* Key, FString& Value) const
+	{
+		if (const FConfigSection* ConfigSection = FindSection(Section))
+		{
+			return ConfigSection->GetString(Key, Value);
+		}
+		return false;
+	}
+
+	bool GetText(const TCHAR* Section, const TCHAR* Key, FText& Value) const
+	{
+		if (const FConfigSection* ConfigSection = FindSection(Section))
+		{
+			return ConfigSection->GetText(Section, Key, Value);
+		}
+		return false;
+	}
+
+	bool GetInt(const TCHAR* Section, const TCHAR* Key, int32& Value) const
+	{
+		if (const FConfigSection* ConfigSection = FindSection(Section))
+		{
+			return ConfigSection->GetInt(Key, Value);
+		}
+		return false;
+	}
+
+	bool GetUInt(const TCHAR* Section, const TCHAR* Key, uint32& Value) const
+	{
+		if (const FConfigSection* ConfigSection = FindSection(Section))
+		{
+			return ConfigSection->GetUInt(Key, Value);
+		}
+		return false;
+	}
+
+	bool GetFloat(const TCHAR* Section, const TCHAR* Key, float& Value) const
+	{
+		if (const FConfigSection* ConfigSection = FindSection(Section))
+		{
+			return ConfigSection->GetFloat(Key, Value);
+		}
+		return false;
+	}
+
+	bool GetDouble(const TCHAR* Section, const TCHAR* Key, double& Value) const
+	{
+		if (const FConfigSection* ConfigSection = FindSection(Section))
+		{
+			return ConfigSection->GetDouble(Key, Value);
+		}
+		return false;
+	}
+
+	bool GetInt64(const TCHAR* Section, const TCHAR* Key, int64& Value) const
+	{
+		if (const FConfigSection* ConfigSection = FindSection(Section))
+		{
+			return ConfigSection->GetInt64(Key, Value);
+		}
+		return false;
+	}
+
+	bool GetBool(const TCHAR* Section, const TCHAR* Key, bool& Value) const
+	{
+		if (const FConfigSection* ConfigSection = FindSection(Section))
+		{
+			return ConfigSection->GetBool(Key, Value);
+		}
+		return false;
+	}
+
 	CORE_API int32 GetArray(const TCHAR* Section, const TCHAR* Key, TArray<FString>& Value) const;
 
 	/* Generic versions for use with templates */
@@ -1324,13 +1349,6 @@ public:
 	CORE_API bool GetSection( const TCHAR* Section, TArray<FString>& Result, const FString& Filename );
 	CORE_API bool DoesSectionExist(const TCHAR* Section, const FString& Filename);
 
-	/**
-	 * @param Force Whether to create the Section on Filename if it did not exist previously.
-	 * @param Const If Const (and not Force), then it will not modify File->Dirty. If not Const (or Force is true), then File->Dirty will be set to true.
-	 */
-	UE_DEPRECATED(5.4, "Use GetSection instead, and/or use the new AddToSection, etc APIs to modify sections without retrieving the section. See top of ConfigCacheIni.h for more info.")
-	CORE_API FConfigSection* GetSectionPrivate( const TCHAR* Section, const bool Force, const bool Const, const FString& Filename );
-
 	CORE_API const FConfigSection* GetSection( const TCHAR* Section, const bool Force, const FString& Filename );
 	CORE_API void SetString( const TCHAR* Section, const TCHAR* Key, const TCHAR* Value, const FString& Filename );
 	CORE_API void SetText( const TCHAR* Section, const TCHAR* Key, const FText& Value, const FString& Filename );
