Increase the threshold for the stat history checks in StatsData.cpp.
- The new GPU profiler emits "end-of-pipe" stats, which are now hitting the checks in this file.
- The checks date back to 2013 in CL 1666884, which were later increased with a "+ 5" in CL 1959147. Seemingly these checks are looking to see that stat processing is progressing with engine ticks, but the checks are now too aggressive.
- Doubled the HistoryFrames multiplier.

diff --git a/Engine/Source/Runtime/Core/Private/Stats/StatsData.cpp b/Engine/Source/Runtime/Core/Private/Stats/StatsData.cpp
--- a/Engine/Source/Runtime/Core/Private/Stats/StatsData.cpp
+++ b/Engine/Source/Runtime/Core/Private/Stats/StatsData.cpp
@@ -1074,11 +1074,11 @@ void FStatsThreadState::AddToHistoryAndEmpty(FStatPacketArray& NewData)
 		}
 	}
 
-	check(History.Num() <= HistoryFrames * 2 + 5);
-	check(CondensedStackHistory.Num() <= HistoryFrames * 2 + 5);
-	check(GoodFrames.Num() <= HistoryFrames * 2 + 5);
+	check(History.Num() <= HistoryFrames * 4 + 5);
+	check(CondensedStackHistory.Num() <= HistoryFrames * 4 + 5);
+	check(GoodFrames.Num() <= HistoryFrames * 4 + 5);
 #if DO_CHECK
-	if (BadFrames.Num() > HistoryFrames * 2 + 5)
+	if (BadFrames.Num() > HistoryFrames * 4 + 5)
 	{
 		UE_LOG(LogStats, Display, TEXT("BadFrames.Num(): %d, History.Num(): %d, HistoryFrames: %d"), BadFrames.Num(), History.Num(), HistoryFrames);
 		UE_LOG(LogStats, Display, TEXT("CurrentGameFrame: %lld, CurrentRenderFrame: %lld, CurrentEndOfPipeFrame: %lld, LastFullFrameMetaAndNonFrame: %lld"), CurrentGameFrame, CurrentRenderFrame, CurrentEndOfPipeFrame, LastFullFrameMetaAndNonFrame);
@@ -1093,7 +1093,7 @@ void FStatsThreadState::AddToHistoryAndEmpty(FStatPacketArray& NewData)
 		}
 	}
 #endif
-	check(BadFrames.Num() <= HistoryFrames * 2 + 5);
+	check(BadFrames.Num() <= HistoryFrames * 4 + 5);
 }
 
 
