Add a new set of LLM tags for the UObject scope to better see where the allocated memory comes from.

A new TagSet for UObjectClasses was added to keep track of UObject allocations per UObject subclass.

To avoid flooding the Default TagSet with UObject subclasses, I added a new TagSet to group all those dynamic tags under it.

#rnx

diff --git a/Engine/Plugins/Experimental/EditorTelemetry/Source/Private/EditorTelemetry.cpp b/Engine/Plugins/Experimental/EditorTelemetry/Source/Private/EditorTelemetry.cpp
--- a/Engine/Plugins/Experimental/EditorTelemetry/Source/Private/EditorTelemetry.cpp
+++ b/Engine/Plugins/Experimental/EditorTelemetry/Source/Private/EditorTelemetry.cpp
@@ -336,6 +336,7 @@ void FEditorTelemetry::RecordEvent_MemoryLLM(const FString& Context, TArray<FAna
 		FLowLevelMemTracker::Get().GetTrackedTagsNamesWithAmount(LLMTrackedNoneMemory, ELLMTracker::Default, ELLMTagSet::None);
 		RecordLLMMemoryEvent(Context, TEXT("None"), LLMTrackedNoneMemory);
 
+#if LLM_ALLOW_ASSETS_TAGS
 		// AssetClasses TagSet
 		TMap<FName, uint64> LLMTrackedAssetClassesMemory;
 		FLowLevelMemTracker::Get().GetTrackedTagsNamesWithAmount(LLMTrackedAssetClassesMemory, ELLMTracker::Default, ELLMTagSet::AssetClasses);
@@ -345,6 +346,7 @@ void FEditorTelemetry::RecordEvent_MemoryLLM(const FString& Context, TArray<FAna
 		TMap<FName, uint64> LLMTrackedAssetMemory;
 		FLowLevelMemTracker::Get().GetTrackedTagsNamesWithAmount(LLMTrackedAssetMemory, ELLMTracker::Default, ELLMTagSet::Assets);
 		RecordLLMMemoryEvent(Context, TEXT("Assets"), LLMTrackedAssetMemory);
+#endif // LLM_ALLOW_ASSETS_TAGS
 	}
 #endif
 }
diff --git a/Engine/Plugins/Experimental/RuntimeTelemetry/Source/Private/RuntimeTelemetry.cpp b/Engine/Plugins/Experimental/RuntimeTelemetry/Source/Private/RuntimeTelemetry.cpp
--- a/Engine/Plugins/Experimental/RuntimeTelemetry/Source/Private/RuntimeTelemetry.cpp
+++ b/Engine/Plugins/Experimental/RuntimeTelemetry/Source/Private/RuntimeTelemetry.cpp
@@ -69,6 +69,8 @@ void FRuntimeTelemetry::RecordEvent_MemoryLLM(const FString& Context, TArray<FAn
 		FLowLevelMemTracker::Get().GetTrackedTagsNamesWithAmount(LLMTrackedNoneMemory, ELLMTracker::Default, ELLMTagSet::None);
 		RecordLLMMemoryEvent(Context, TEXT("None"), LLMTrackedNoneMemory);
 
+#if LLM_ALLOW_ASSETS_TAGS
+
 		// AssetClasses TagSet
 		TMap<FName, uint64> LLMTrackedAssetClassesMemory;
 		FLowLevelMemTracker::Get().GetTrackedTagsNamesWithAmount(LLMTrackedAssetClassesMemory, ELLMTracker::Default, ELLMTagSet::AssetClasses);
@@ -78,6 +80,9 @@ void FRuntimeTelemetry::RecordEvent_MemoryLLM(const FString& Context, TArray<FAn
 		TMap<FName, uint64> LLMTrackedAssetMemory;
 		FLowLevelMemTracker::Get().GetTrackedTagsNamesWithAmount(LLMTrackedAssetMemory, ELLMTracker::Default, ELLMTagSet::Assets);
 		RecordLLMMemoryEvent(Context, TEXT("Assets"), LLMTrackedAssetMemory);
+
+#endif // LLM_ALLOW_ASSETS_TAGS
+
 	}
 #endif
 }
diff --git a/Engine/Plugins/MemoryUsageQueries/Source/MemoryUsageQueries/Private/MemoryUsageInfoProviderLLM.cpp b/Engine/Plugins/MemoryUsageQueries/Source/MemoryUsageQueries/Private/MemoryUsageInfoProviderLLM.cpp
--- a/Engine/Plugins/MemoryUsageQueries/Source/MemoryUsageQueries/Private/MemoryUsageInfoProviderLLM.cpp
+++ b/Engine/Plugins/MemoryUsageQueries/Source/MemoryUsageQueries/Private/MemoryUsageInfoProviderLLM.cpp
@@ -8,20 +8,34 @@
 
 bool FMemoryUsageInfoProviderLLM::IsProviderAvailable() const
 {
+	// if !LLM_ALLOW_ASSETS_TAGS there will never be assets tags in any trackers
+	// meaning that we can return false in this case
+#if !LLM_ALLOW_ASSETS_TAGS
+
+	return false;
+
+#else // LLM_ALLOW_ASSETS_TAGS
+
 	return FLowLevelMemTracker::IsEnabled();
+
+#endif // LLM_ALLOW_ASSETS_TAGS
 }
 
 uint64 FMemoryUsageInfoProviderLLM::GetAssetMemoryUsage(FName Asset) const
 {
+#if LLM_ALLOW_ASSETS_TAGS
 	if (FLowLevelMemTracker::IsEnabled())
 	{
 		return FLowLevelMemTracker::Get().GetTagAmountForTracker(ELLMTracker::Default, Asset, ELLMTagSet::Assets, UE::LLM::ESizeParams::Default);
 	}
+#endif // LLM_ALLOW_ASSETS_TAGS
+
 	return 0U;
 }
 
 uint64 FMemoryUsageInfoProviderLLM::GetAssetsMemoryUsage(const TSet<FName>& Assets) const
 {
+#if LLM_ALLOW_ASSETS_TAGS
 	if (FLowLevelMemTracker::IsEnabled())
 	{
 		uint64 TotalSize = 0;
@@ -31,11 +45,14 @@ uint64 FMemoryUsageInfoProviderLLM::GetAssetsMemoryUsage(const TSet<FName>& Asse
 		}
 		return TotalSize;
 	}
+#endif // LLM_ALLOW_ASSETS_TAGS
+
 	return 0U;
 }
 
 uint64 FMemoryUsageInfoProviderLLM::GetAssetsMemoryUsageWithSize(const TSet<FName>& Assets, TMap<FName, uint64>& OutSizes) const
 {
+#if LLM_ALLOW_ASSETS_TAGS
 	if (FLowLevelMemTracker::IsEnabled())
 	{
 		uint64 TotalSize = 0;
@@ -47,16 +64,22 @@ uint64 FMemoryUsageInfoProviderLLM::GetAssetsMemoryUsageWithSize(const TSet<FNam
 		}
 		return TotalSize;
 	}
+#endif // LLM_ALLOW_ASSETS_TAGS
+
 	return 0U;
 }
 
 void FMemoryUsageInfoProviderLLM::GetAllAssetsWithSize(TMap<FName, uint64>& OutAssets) const
 {
+#if LLM_ALLOW_ASSETS_TAGS
+
 	OutAssets.Reset();
 	if (FLowLevelMemTracker::IsEnabled())
 	{
 		FLowLevelMemTracker::Get().GetTrackedTagsNamesWithAmount(OutAssets, ELLMTracker::Default, ELLMTagSet::Assets);
 	}
+	
+#endif // LLM_ALLOW_ASSETS_TAGS
 }
 
 void FMemoryUsageInfoProviderLLM::GetFilteredTagsWithSize(TMap<FName, uint64>& OutTags, ELLMTracker Tracker, ELLMTagSet TagSet, TArray<FLLMTagSetAllocationFilter>& Filters) const
@@ -84,6 +107,12 @@ IMemoryUsageInfoProvider* GetCurrentMemoryUsageInfoProvider()
 
 bool GetFilteredPackagesWithSize(TMap<FName, uint64>& OutPackagesWithSize, FName GroupName /* = NAME_None */, FString AssetSubstring /* = FString() */, FName ClassName /* = NAME_None */, FOutputDevice* ErrorOutput /*= GLog*/)
 {
+#if !LLM_ALLOW_ASSETS_TAGS
+
+	return false;
+
+#else // LLM_ALLOW_ASSETS_TAGS
+
 	TArray<FLLMTagSetAllocationFilter> Filters;
 
 	FName LongName = NAME_None;
@@ -117,10 +146,18 @@ bool GetFilteredPackagesWithSize(TMap<FName, uint64>& OutPackagesWithSize, FName
 	OutPackagesWithSize.ValueSort(TGreater<uint64>());
 
 	return true;
+	
+#endif // LLM_ALLOW_ASSETS_TAGS
 }
 
 bool GetFilteredPackagesCategorizedWithSize(TMap<FName, uint64>& OutPackagesWithSize, FName GroupName /*= NAME_None*/, FString AssetSubstring /*= FString()*/, FName ClassName /*= NAME_None*/, FName CategoryName /*= NAME_None*/, FOutputDevice* ErrorOutput /*= GLog*/)
 {
+#if !LLM_ALLOW_ASSETS_TAGS
+
+	return false;
+
+#else // LLM_ALLOW_ASSETS_TAGS
+
 	TArray<FLLMTagSetAllocationFilter> Filters;
 
 	FName LongName = NAME_None;
@@ -164,10 +201,18 @@ bool GetFilteredPackagesCategorizedWithSize(TMap<FName, uint64>& OutPackagesWith
 	OutPackagesWithSize.ValueSort(TGreater<uint64>());
 
 	return true;
+
+#endif // LLM_ALLOW_ASSETS_TAGS
 }
 
 bool GetFilteredClassesWithSize(TMap<FName, uint64>& OutClassesWithSize, FName GroupName /* = NAME_None */, FString AssetName /* = FString() */, FOutputDevice* ErrorOutput /* = GLog */)
 {
+#if !LLM_ALLOW_ASSETS_TAGS
+
+	return false;
+
+#else // LLM_ALLOW_ASSETS_TAGS
+
 	TArray<FLLMTagSetAllocationFilter> Filters;
 
 	FName LongName = NAME_None;
@@ -187,14 +232,21 @@ bool GetFilteredClassesWithSize(TMap<FName, uint64>& OutClassesWithSize, FName G
 	}
 
 	Internal::GMemoryUsageInfoProviderLLM.GetFilteredTagsWithSize(OutClassesWithSize, ELLMTracker::Default, ELLMTagSet::AssetClasses, Filters);
-
+	
 	OutClassesWithSize.ValueSort(TGreater<uint64>());
 
 	return true;
+#endif // LLM_ALLOW_ASSETS_TAGS
 }
 
 bool GetFilteredGroupsWithSize(TMap<FName, uint64>& OutGroupsWithSize, FString AssetName /* = FString() */, FName ClassName /* = NAME_None */, FOutputDevice* ErrorOutput /* = GLog */)
 {
+#if !LLM_ALLOW_ASSETS_TAGS
+
+	return false;
+
+#else // LLM_ALLOW_ASSETS_TAGS
+
 	TArray<FLLMTagSetAllocationFilter> Filters;
 
 	FName LongName = NAME_None;
@@ -218,6 +270,8 @@ bool GetFilteredGroupsWithSize(TMap<FName, uint64>& OutGroupsWithSize, FString A
 	OutGroupsWithSize.ValueSort(TGreater<uint64>());
 
 	return true;
+
+#endif // LLM_ALLOW_ASSETS_TAGS
 }
 
 } // namespace MemoryUsageQueries
diff --git a/Engine/Source/Runtime/Core/Private/HAL/LowLevelMemTracker.cpp b/Engine/Source/Runtime/Core/Private/HAL/LowLevelMemTracker.cpp
--- a/Engine/Source/Runtime/Core/Private/HAL/LowLevelMemTracker.cpp
+++ b/Engine/Source/Runtime/Core/Private/HAL/LowLevelMemTracker.cpp
@@ -968,6 +968,8 @@ TAutoConsoleVariable<int32> CVarLLMHeaderMaxSize(
 	TEXT("LLM.LLMHeaderMaxSize"),
 #if LLM_ALLOW_ASSETS_TAGS
 	500000, // When using asset tags, you will have MANY more LLM titles since so many are auto generated.
+#elif LLM_ALLOW_UOBJECTCLASSES_TAGS
+	30000, // When using uobjects clases, you will have at least 20K+ LLM titles.
 #else
 	5000,
 #endif
@@ -1005,6 +1007,7 @@ FAutoConsoleCommand DumpLLM(
 		}
 
 		ELLMTagSet TagSet = ELLMTagSet::None;
+#if LLM_ALLOW_ASSETS_TAGS
 		if (FParse::Param(*Command, TEXT("Assets")))
 		{
 			TagSet = ELLMTagSet::Assets;
@@ -1013,6 +1016,12 @@ FAutoConsoleCommand DumpLLM(
 		{
 			TagSet = ELLMTagSet::AssetClasses;
 		}
+#elif LLM_ALLOW_UOBJECTCLASSES_TAGS
+		if (FParse::Param(*Command, TEXT("UObjectClasses")))
+		{
+			TagSet = ELLMTagSet::UObjectClasses;
+		}
+#endif // LLM_ALLOW_UOBJECTCLASSES_TAGS
 
 		FLowLevelMemTracker::Get().DumpToLog(DumpFormat, &Ar, SizeParams, TagSet);
 	}));
@@ -1364,9 +1373,9 @@ public:
 	void PopTag(ELLMTagSet TagSet);
 	const FTagData* GetTopTag(ELLMTagSet TagSet);
 	void TrackAllocation(const void* Ptr, int64 Size, ELLMTracker Tracker, ELLMAllocType AllocType, 
-		const FTagData* TagData, const FTagData* AssetTagData, const FTagData* AssetClassTagData, bool bTrackInMemPro);
+		const FTagData* TagData, const FTagData* AssetTagData, const FTagData* AssetClassTagData, const FTagData* UObjectClassesTagData, bool bTrackInMemPro);
 	void TrackFree(const void* Ptr, int64 Size, ELLMTracker Tracker, ELLMAllocType AllocType, 
-		const FTagData* TagData, const FTagData* AssetTagData, const FTagData* AssetClassTagData, bool bTrackInMemPro);
+		const FTagData* TagData, const FTagData* AssetTagData, const FTagData* AssetClassTagData, const FTagData* UObjectClassesTagData, bool bTrackInMemPro);
 	void TrackMemory(int64 Amount, ELLMTracker Tracker, ELLMAllocType AllocType,
 		const FTagData* TagData);
 	void TrackMoved(const void* Dest, const void* Source, int64 Size, ELLMTracker Tracker, const FTagData* TagData);
@@ -1440,19 +1449,19 @@ public:
 	public:
 		FLowLevelAllocInfo()
 		{
-#if LLM_ENABLED_FULL_TAGS && LLM_ALLOW_ASSETS_TAGS
+#if LLM_ENABLED_FULL_TAGS && (LLM_ALLOW_ASSETS_TAGS || LLM_ALLOW_UOBJECTCLASSES_TAGS)
 			FMemory::Memset(Tag, InvalidCompressedTagValue);
 #endif
 		}
 
 		void SetTag(const FTagData* InTag, FLowLevelMemTracker& InLLMRef
-#if LLM_ALLOW_ASSETS_TAGS
+#if LLM_ALLOW_ASSETS_TAGS || LLM_ALLOW_UOBJECTCLASSES_TAGS
 		, ELLMTagSet InTagSet = ELLMTagSet::None
 #endif
 		)
 		{
 #if LLM_ENABLED_FULL_TAGS
-#if !LLM_ALLOW_ASSETS_TAGS
+#if !LLM_ALLOW_ASSETS_TAGS && !LLM_ALLOW_UOBJECTCLASSES_TAGS
 			Tag = InTag->GetIndex();
 #else
 			Tag[static_cast<int32>(InTagSet)] = InTag ? InTag->GetIndex() : InvalidCompressedTagValue;
@@ -1463,14 +1472,14 @@ public:
 		}
 
 		const FTagData* GetTag(FLowLevelMemTracker& InLLMRef
-#if LLM_ALLOW_ASSETS_TAGS
+#if LLM_ALLOW_ASSETS_TAGS || LLM_ALLOW_UOBJECTCLASSES_TAGS
 		, ELLMTagSet InTagSet = ELLMTagSet::None
 #endif
 		) const
 		{
 #if LLM_ENABLED_FULL_TAGS
 			FReadScopeLock TagDataScopeLock(InLLMRef.TagDataLock);
-#if !LLM_ALLOW_ASSETS_TAGS
+#if !LLM_ALLOW_ASSETS_TAGS && !LLM_ALLOW_UOBJECTCLASSES_TAGS
 			return (*InLLMRef.TagDatas)[Tag];
 #else
 			return Tag[static_cast<int32>(InTagSet)] >= 0 ? (*InLLMRef.TagDatas)[Tag[static_cast<int32>(InTagSet)]] : nullptr;
@@ -1483,7 +1492,7 @@ public:
 #if LLM_ENABLED_FULL_TAGS
 		int32 GetCompressedTag(ELLMTagSet InTagSet = ELLMTagSet::None)
 		{
-#if !LLM_ALLOW_ASSETS_TAGS
+#if !LLM_ALLOW_ASSETS_TAGS && !LLM_ALLOW_UOBJECTCLASSES_TAGS
 			return (InTagSet == ELLMTagSet::None) ? Tag : InvalidCompressedTagValue;
 #else
 			return Tag[static_cast<int32>(InTagSet)];
@@ -1492,7 +1501,7 @@ public:
 
 		void SetCompressedTag(int32 InTag, ELLMTagSet InTagSet = ELLMTagSet::None)
 		{
-#if !LLM_ALLOW_ASSETS_TAGS
+#if !LLM_ALLOW_ASSETS_TAGS && !LLM_ALLOW_UOBJECTCLASSES_TAGS
 			if (InTagSet == ELLMTagSet::None)
 			{
 				Tag = InTag;
@@ -1510,12 +1519,17 @@ public:
 #if LLM_ENABLED_FULL_TAGS
 		// Even with arbitrary tags we are still partially compressed - the allocation records the tag's index
 		// (4 bytes) rather than the full tag pointer (8 bytes) or the tag's name (12 bytes).
-#if !LLM_ALLOW_ASSETS_TAGS
+#if !LLM_ALLOW_ASSETS_TAGS && !LLM_ALLOW_UOBJECTCLASSES_TAGS
 		int32 Tag = 0;
 #else
 		int32 Tag[static_cast<int32>(ELLMTagSet::Max)];
-		static_assert(static_cast<uint32>(ELLMTagSet::Max) <= 3, "Every new TagSet enum adds 4 bytes of LLM overhead memory for every active allocation. Do not add them without measuring the cost of the overhead in the cook and runtime of large projects.");
+#if LLM_ALLOW_ASSETS_TAGS
+#define TAG_SET_MAX 3
+#else
+#define TAG_SET_MAX 2
 #endif
+		static_assert(static_cast<uint32>(ELLMTagSet::Max) <= TAG_SET_MAX, "Every new TagSet enum adds 4 bytes of LLM overhead memory for every active allocation. Do not add them without measuring the cost of the overhead in the cook and runtime of large projects.");
+#endif // LLM_ALLOW_ASSETS_TAGS
 #else
 		ELLMTag Tag = ELLMTag::Untagged;
 #endif
@@ -1657,6 +1671,7 @@ static FName GetTagName_CustomName()
 static FName TagName_Untagged(TEXT("Untagged"));
 static FName TagName_UntaggedAsset(TEXT("UntaggedAsset"));
 static FName TagName_UntaggedAssetClass(TEXT("UntaggedAssetClass"));
+static FName TagName_UntaggedUObjectClasses(TEXT("UntaggedUObjectClasses"));
 
 FName LLMGetTagUniqueName(ELLMTag Tag)
 {
@@ -1725,6 +1740,7 @@ extern const FName LLMGetUntaggedTagName(ELLMTagSet TagSet)
 		{
 			return TagName_Untagged;
 		}
+#if LLM_ALLOW_ASSETS_TAGS
 		case ELLMTagSet::Assets:
 		{
 			return TagName_UntaggedAsset;
@@ -1733,6 +1749,12 @@ extern const FName LLMGetUntaggedTagName(ELLMTagSet TagSet)
 		{
 			return TagName_UntaggedAssetClass;
 		}
+#elif LLM_ALLOW_UOBJECTCLASSES_TAGS
+		case ELLMTagSet::UObjectClasses:
+		{
+			return TagName_UntaggedUObjectClasses;
+		}
+#endif // LLM_ALLOW_UOBJECTCLASSES_TAGS
 		default:
 		{
 			return NAME_None;
@@ -2327,7 +2349,7 @@ void FLowLevelMemTracker::ProcessCommandLine(const TCHAR* CmdLine)
 
 	// activate tag sets (we ignore None set, it's always on)
 	FString SetList;
-	static_assert((uint8)ELLMTagSet::Max == 3,
+	static_assert((uint8)ELLMTagSet::Max == TAG_SET_MAX,
 		"You added a tagset, without updating FLowLevelMemTracker::ProcessCommandLine");
 	if (FParse::Value(CmdLine, TEXT("LLMTAGSETS="), SetList, false /* bShouldStopOnSeparator */))
 	{
@@ -2335,21 +2357,26 @@ void FLowLevelMemTracker::ProcessCommandLine(const TCHAR* CmdLine)
 		SetList.ParseIntoArray(Sets, TEXT(","), true);
 		for (FString& Set : Sets)
 		{
+#if LLM_ALLOW_ASSETS_TAGS
 			if (Set == TEXT("Assets"))
 			{
-#if LLM_ALLOW_ASSETS_TAGS
 				ActiveSets[static_cast<int32>(ELLMTagSet::Assets)] = true;
+			}
+			else if (Set == TEXT("AssetClasses"))
+			{
+				ActiveSets[static_cast<int32>(ELLMTagSet::AssetClasses)] = true;
+			}
+#elif LLM_ALLOW_UOBJECTCLASSES_TAGS
+			if (Set == TEXT("UObjectClasses"))
+			{
+				ActiveSets[static_cast<int32>(ELLMTagSet::UObjectClasses)] = true;
+			}
 #else
 				// Asset tracking has a per-thread memory overhead, so we have a #define to completely disable it.
 				// Warn if we try to toggle it on at runtime but it is not defined.
 				UE_LOG(LogInit, Warning,
 					TEXT("Attempted to use LLM to track assets, but LLM_ALLOW_ASSETS_TAGS is not defined to 1. You will need to enable the define"));
 #endif
-			}
-			else if (Set == TEXT("AssetClasses"))
-			{
-				ActiveSets[static_cast<int32>(ELLMTagSet::AssetClasses)] = true;
-			}
 		}
 	}
 
@@ -2594,11 +2621,15 @@ bool FLowLevelMemTracker::ShouldReduceThreads()
 #endif
 }
 
+#if LLM_ALLOW_ASSETS_TAGS
+
 static bool IsAssetTagForAssets(ELLMTagSet Set)
 {
 	return Set == ELLMTagSet::Assets || Set == ELLMTagSet::AssetClasses;
 }
 
+#endif // LLM_ALLOW_ASSETS_TAGS
+
 void FLowLevelMemTracker::RegisterCustomTagInternal(int32 Tag, ELLMTagSet TagSet, const TCHAR* InDisplayName, FName StatName,
 	FName SummaryStatName, int32 ParentTag)
 {
@@ -4819,7 +4850,7 @@ void FLLMTracker::TrackAllocation(const void* Ptr, int64 Size, const FTagData* A
 #if LLM_ALLOW_ASSETS_TAGS
 	const FTagData* AssetTagData = State->GetTopTag(ELLMTagSet::Assets);
 	const FTagData* AssetClassTagData = State->GetTopTag(ELLMTagSet::AssetClasses);
-
+	
 	if (AssetTagData == nullptr && !LLMRef.IsBootstrapping() && LLMRef.IsTagSetActive(ELLMTagSet::Assets))
 	{
 		AssetTagData = LLMRef.FindOrAddTagData(TagName_UntaggedAsset, ELLMTagSet::Assets, false, UE::LLMPrivate::ETagReferenceSource::FunctionAPI);
@@ -4831,9 +4862,16 @@ void FLLMTracker::TrackAllocation(const void* Ptr, int64 Size, const FTagData* A
 #else
 	const FTagData* AssetTagData = nullptr;
 	const FTagData* AssetClassTagData = nullptr;
+#endif // LLM_ALLOW_ASSETS_TAGS
+
+#if LLM_ALLOW_UOBJECTCLASSES_TAGS
+	const FTagData* UObjectClassesTagData = State->GetTopTag(ELLMTagSet::UObjectClasses);
+#else
+	const FTagData* UObjectClassesTagData = nullptr;
 #endif
+
 	// track on the thread state
-	State->TrackAllocation(Ptr, Size, Tracker, AllocType, ActiveTagData, AssetTagData, AssetClassTagData, bTrackInMemPro);
+	State->TrackAllocation(Ptr, Size, Tracker, AllocType, ActiveTagData, AssetTagData, AssetClassTagData, UObjectClassesTagData, bTrackInMemPro);
 
 #if !UE_ONLY_USE_PLATFORM_TRACKER
 	// tracking a nullptr with a Size is allowed, but we don't need to remember it, since we can't free it ever.
@@ -4845,6 +4883,13 @@ void FLLMTracker::TrackAllocation(const void* Ptr, int64 Size, const FTagData* A
 #if LLM_ALLOW_ASSETS_TAGS
 		AllocInfo.SetTag(AssetTagData, LLMRef, ELLMTagSet::Assets);
 		AllocInfo.SetTag(AssetClassTagData, LLMRef, ELLMTagSet::AssetClasses);
+#endif
+#if LLM_ALLOW_UOBJECTCLASSES_TAGS
+		// If there is no scoped tag for UObjectClasses, don't set it
+		if (UObjectClassesTagData != nullptr)
+		{
+			AllocInfo.SetTag(UObjectClassesTagData, LLMRef, ELLMTagSet::UObjectClasses);
+		}
 #endif
 		LLMCheck(Size <= 0x0000'ffff'ffff'ffff);
 		uint32 SizeLow = uint32(Size);
@@ -4887,19 +4932,29 @@ void FLLMTracker::TrackFree(const void* Ptr, ELLMAllocType AllocType, bool bTrac
 #if LLM_ALLOW_ASSETS_TAGS
 	const FTagData* AssetTagData = AllocInfo.GetTag(LLMRef, ELLMTagSet::Assets);
 	const FTagData* AssetClassTagData = AllocInfo.GetTag(LLMRef, ELLMTagSet::AssetClasses);
-#else
+#else 
 	const FTagData* AssetTagData = nullptr;
 	const FTagData* AssetClassTagData = nullptr;
-#endif
+#endif // LLM_ALLOW_ASSETS_TAGS
+
+#if LLM_ALLOW_UOBJECTCLASSES_TAGS	
+	const FTagData* UObjectClassesTagData = AllocInfo.GetTag(LLMRef, ELLMTagSet::UObjectClasses);
 #else
+	const FTagData* UObjectClassesTagData = nullptr;
+#endif // LLM_ALLOW_UOBJECTCLASSES_TAGS
+
+#else
+
 	FLLMThreadState* State = GetOrCreateState();
 	int64 Size = 0;
 	const FTagData* TagData = nullptr;
 	const FTagData* AssetTagData = nullptr;
 	const FTagData* AssetClassTagData = nullptr;
-#endif
+	const FTagData* UObjectClassesTagData = nullptr;
 
-	State->TrackFree(Ptr, Size, Tracker, AllocType, TagData, AssetTagData, AssetClassTagData, bTrackInMemPro);
+#endif // !UE_ONLY_USE_PLATFORM_TRACKER
+
+	State->TrackFree(Ptr, Size, Tracker, AllocType, TagData, AssetTagData, AssetClassTagData, UObjectClassesTagData, bTrackInMemPro);
 }
 
 void FLLMTracker::OnAllocMoved(const void* Dest, const void* Source, ELLMAllocType AllocType)
@@ -5615,14 +5670,14 @@ void FLLMTracker::GetTagsNamesWithAmountFiltered(TMap<FName, uint64>& OutTagsNam
 		for (const FLLMTagSetAllocationFilter& Filter : Filters)
 		{
 			FLocalTagData* Data = nullptr;
-#if LLM_ALLOW_ASSETS_TAGS
+#if LLM_ALLOW_ASSETS_TAGS || LLM_ALLOW_UOBJECTCLASSES_TAGS
 			int32 CompressedTagIndex = static_cast<int32>(Tuple.Value2.GetCompressedTag(Filter.TagSet));
 #else
 			int32 CompressedTagIndex = static_cast<int32>(Tuple.Value2.GetCompressedTag());
 #endif
 			if(TagDataCached.IsValidIndex(CompressedTagIndex))
 			{
-#if LLM_ALLOW_ASSETS_TAGS
+#if LLM_ALLOW_ASSETS_TAGS || LLM_ALLOW_UOBJECTCLASSES_TAGS
 				Data = &TagDataCached[CompressedTagIndex];
 #else
 				if (Filter.TagSet == ELLMTagSet::None)
@@ -5640,7 +5695,7 @@ void FLLMTracker::GetTagsNamesWithAmountFiltered(TMap<FName, uint64>& OutTagsNam
 			
 		if (bIncludeAllocation)
 		{
-#if LLM_ALLOW_ASSETS_TAGS
+#if LLM_ALLOW_ASSETS_TAGS || LLM_ALLOW_UOBJECTCLASSES_TAGS
 			int32 CompressedTagIndex = Tuple.Value2.GetCompressedTag(TagSet);
 #else
 			int32 CompressedTagIndex = static_cast<int32>(Tuple.Value2.GetCompressedTag());
@@ -5676,7 +5731,7 @@ bool FLLMTracker::FindTagsForPtr(void* InPtr, TArray<const FTagData *, TInlineAl
 	}
 
 	OutTags.SetNumUninitialized(static_cast<int32>(ELLMTagSet::Max));
-#if LLM_ALLOW_ASSETS_TAGS
+#if LLM_ALLOW_ASSETS_TAGS || LLM_ALLOW_UOBJECTCLASSES_TAGS
 	for (int32 TagSetIndex = 0; TagSetIndex < static_cast<int32>(ELLMTagSet::Max); TagSetIndex++)
 	{
 		OutTags[TagSetIndex] = AllocInfoPtr.GetTag(LLMRef, static_cast<ELLMTagSet>(TagSetIndex));
@@ -5781,7 +5836,7 @@ void FLLMThreadState::IncrTag(const FTagData* TagData, int64 Amount)
 }
 
 void FLLMThreadState::TrackAllocation(const void* Ptr, int64 Size, ELLMTracker Tracker, ELLMAllocType AllocType, 
-	const FTagData* TagData, const FTagData* AssetTagData, const FTagData* AssetClassTagData, bool bTrackInMemPro)
+	const FTagData* TagData, const FTagData* AssetTagData, const FTagData* AssetClassTagData, const FTagData* UObjectClassesTagData, bool bTrackInMemPro)
 {
 	FScopeLock Lock(&TagSection);
 
@@ -5800,7 +5855,14 @@ void FLLMThreadState::TrackAllocation(const void* Ptr, int64 Size, ELLMTracker T
 	{
 		IncrTag(AssetClassTagData, Size);
 	}
-#endif
+#endif // LLM_ALLOW_ASSETS_TAGS
+
+#if LLM_ALLOW_UOBJECTCLASSES_TAGS
+	if (UObjectClassesTagData)
+	{
+		IncrTag(UObjectClassesTagData, Size);
+	}
+#endif // LLM_ALLOW_UOBJECTCLASSES_TAGS
 
 	ELLMTag EnumTag = TagData ? TagData->GetContainingEnum() : ELLMTag(0);
 	if (Tracker == ELLMTracker::Default)
@@ -5817,7 +5879,7 @@ void FLLMThreadState::TrackAllocation(const void* Ptr, int64 Size, ELLMTracker T
 }
 
 void FLLMThreadState::TrackFree(const void* Ptr, int64 Size, ELLMTracker Tracker, ELLMAllocType AllocType, 
-	const FTagData* TagData, const FTagData* AssetTagData, const FTagData* AssetClassTagData, bool bTrackInMemPro)
+	const FTagData* TagData, const FTagData* AssetTagData, const FTagData* AssetClassTagData, const FTagData* UObjectClassesTagData, bool bTrackInMemPro)
 {
 	FScopeLock Lock(&TagSection);
 
@@ -5835,7 +5897,14 @@ void FLLMThreadState::TrackFree(const void* Ptr, int64 Size, ELLMTracker Tracker
 	{
 		IncrTag(AssetClassTagData, -Size);
 	}
-#endif
+#endif // LLM_ALLOW_ASSETS_TAGS
+
+#if LLM_ALLOW_UOBJECTCLASSES_TAGS
+	if (UObjectClassesTagData)
+	{
+		IncrTag(UObjectClassesTagData, -Size);
+	}
+#endif // LLM_ALLOW_UOBJECTCLASSES_TAGS
 
 	ELLMTag EnumTag = TagData ? TagData->GetContainingEnum() : ELLMTag(0);
 	if (Tracker == ELLMTracker::Default)
@@ -6296,9 +6365,20 @@ void FLLMTraceWriter::Publish(FLowLevelMemTracker& LLMRef, const FTrackerTagSize
 	if (!bTagSetSpecSent && Tracker == ELLMTracker::Default)
 	{
 		bTagSetSpecSent = true;
-		FAnsiStringView TagSetNames[] = { "System", "Assets", "AssetClasses" };
+		FAnsiStringView TagSetNames[] = { 
+			"System", 
+#if LLM_ALLOW_ASSETS_TAGS
+			"Assets", "AssetClasses", 
+#elif LLM_ALLOW_UOBJECTCLASSES_TAGS
+			"UObjectClasses" 
+#endif // LLM_ALLOW_UOBJECTCLASSES_TAGS
+		};
+#if LLM_ALLOW_ASSETS_TAGS
 		static_assert((int)ELLMTagSet::Assets == 1, "");
 		static_assert((int)ELLMTagSet::AssetClasses == 2, "");
+#elif LLM_ALLOW_UOBJECTCLASSES_TAGS
+		static_assert((int)ELLMTagSet::UObjectClasses == 1, "");
+#endif // LLM_ALLOW_UOBJECTCLASSES_TAGS
 		static_assert(UE_ARRAY_COUNT(TagSetNames) == (int)ELLMTagSet::Max, "");
 		for (uint8 TagSet = 0; TagSet < (uint8)ELLMTagSet::Max; ++TagSet)
 		{
diff --git a/Engine/Source/Runtime/Core/Private/UObject/UnrealNames.cpp b/Engine/Source/Runtime/Core/Private/UObject/UnrealNames.cpp
--- a/Engine/Source/Runtime/Core/Private/UObject/UnrealNames.cpp
+++ b/Engine/Source/Runtime/Core/Private/UObject/UnrealNames.cpp
@@ -453,8 +453,10 @@ private:
 		if (!Buffer)
 		{
 			LLM_SCOPE(ELLMTag::FName);
+#if LLM_ALLOW_ASSETS_TAGS
 			LLM_TAGSET_SCOPE_CLEAR(ELLMTagSet::Assets);
 			LLM_TAGSET_SCOPE_CLEAR(ELLMTagSet::AssetClasses);
+#endif // LLM_ALLOW_ASSETS_TAGS
 			UE_TRACE_METADATA_CLEAR_SCOPE();
 			constexpr uint32 BufferSizeBytes = NumWords * sizeof(WordType);
 			Buffer = (WordType*) FMemory::MallocZeroed(BufferSizeBytes, alignof(WordType));
@@ -769,8 +771,10 @@ private:
 	uint8* AllocBlock()
 	{
 		LLM_SCOPE(ELLMTag::FName);
+#if LLM_ALLOW_ASSETS_TAGS
 		LLM_TAGSET_SCOPE_CLEAR(ELLMTagSet::Assets);
 		LLM_TAGSET_SCOPE_CLEAR(ELLMTagSet::AssetClasses);
+#endif // LLM_ALLOW_ASSETS_TAGS
 		UE_TRACE_METADATA_CLEAR_SCOPE();
 
 #if USE_FNAME_MMAP
@@ -1174,8 +1178,10 @@ public:
 	void Initialize(FNameEntryAllocator& InEntries)
 	{
 		LLM_SCOPE(ELLMTag::FName);
+#if LLM_ALLOW_ASSETS_TAGS
 		LLM_TAGSET_SCOPE_CLEAR(ELLMTagSet::Assets);
 		LLM_TAGSET_SCOPE_CLEAR(ELLMTagSet::AssetClasses);
+#endif // LLM_ALLOW_ASSETS_TAGS
 		UE_TRACE_METADATA_CLEAR_SCOPE();
 		Entries = &InEntries;
 
@@ -1579,8 +1585,10 @@ private:
 	void Grow(const uint32 NewCapacity)
 	{
 		LLM_SCOPE(ELLMTag::FName);
+#if LLM_ALLOW_ASSETS_TAGS
 		LLM_TAGSET_SCOPE_CLEAR(ELLMTagSet::Assets);
 		LLM_TAGSET_SCOPE_CLEAR(ELLMTagSet::AssetClasses);
+#endif // LLM_ALLOW_ASSETS_TAGS
 		UE_TRACE_METADATA_CLEAR_SCOPE();
 		TArrayView<FNameSlot> OldSlots(Slots, Capacity());
 		const uint32 OldUsedSlots = UsedSlots;
diff --git a/Engine/Source/Runtime/Core/Public/HAL/LowLevelMemTracker.h b/Engine/Source/Runtime/Core/Public/HAL/LowLevelMemTracker.h
--- a/Engine/Source/Runtime/Core/Public/HAL/LowLevelMemTracker.h
+++ b/Engine/Source/Runtime/Core/Public/HAL/LowLevelMemTracker.h
@@ -34,6 +34,12 @@
 	#define LLM_ALLOW_ASSETS_TAGS 0
 #endif
 
+// LLM_ALLOW_UOBJECTCLASSES_TAGS will be enabled when LLM_ALLOW_ASSETS_TAGS is not
+// The intent is to avoid the amount of memory overhead by having 3 active TagSet at once.
+#ifndef LLM_ALLOW_UOBJECTCLASSES_TAGS
+	#define LLM_ALLOW_UOBJECTCLASSES_TAGS !LLM_ALLOW_ASSETS_TAGS
+#endif
+
 // LLM_ALLOW_STATS: Set to 1 to allow stats to be used as tags, 0 to disable.
 // When enabled LLM_SCOPED_TAG_WITH_STAT macros are enabled and create an LLM tag per stat at the cost of more LLM
 // memory usage per allocation. Turning this on uses the same amount of memory per allocation as LLM_ALLOW_NAMES_TAGS.
@@ -125,8 +131,12 @@ enum class ELLMTracker : uint8
 enum class ELLMTagSet : uint8
 {
 	None,
+#if LLM_ALLOW_ASSETS_TAGS
 	Assets,
 	AssetClasses,
+#elif LLM_ALLOW_UOBJECTCLASSES_TAGS
+	UObjectClasses,
+#endif // LLM_ALLOW_UOBJECTCLASSES_TAGS
 	
 	Max, // note: see FLowLevelMemTracker::ShouldReduceThreads and IsAssetTagForAssets if you add any asset-style tagsets
 };
@@ -411,6 +421,21 @@ extern CORE_API const FName LLMGetUntaggedTagName(ELLMTagSet TagSet);
  */
 #define LLM_DEFINE_TAG(UniqueNameWithUnderscores, ...) FLLMTagDeclaration PREPROCESSOR_JOIN(LLMTagDeclaration_, UniqueNameWithUnderscores)(TEXT(#UniqueNameWithUnderscores), ##__VA_ARGS__)
 
+/**
+ * Statically define a tag which can be used in LLM_SCOPE_BYTAG or referenced by name in other LLM_SCOPEs.
+ * @param UniqueNameWithUnderscores - Modified version of the name of the tag. Used for looking up by name,
+ *        must be unique across all tags passed to LLM_DEFINE_TAG, LLM_SCOPE, or ELLMTag.
+ *        The modification: the usual separator / for parents must be replaced with _ in LLM_DEFINE_TAGs.
+ * @param DisplayName - (Optional) - The name to display when tracing the tag; joined with "/" to the name of its
+ *        parent if it has a parent, or NAME_None to use the UniqueName.
+ * @param ParentTagName - (Optional) - The unique name of the parent tag, or NAME_None if it has no parent.
+ * @param StatName - (Optional) - The name of the stat to populate with this tag's amount when publishing LLM data each
+ *        frame, or NAME_None if no stat should be populated.
+ * @param SummaryStatName - (Optional) - The name of the stat group to add on this tag's amount when publishing LLM
+ *        data each frame, or NAME_None if no stat group should be added to.
+ */
+#define LLM_DEFINE_STATIC_TAG(UniqueNameWithUnderscores, ...) static FLLMTagDeclaration PREPROCESSOR_JOIN(LLMTagDeclaration_, UniqueNameWithUnderscores)(TEXT(#UniqueNameWithUnderscores), ##__VA_ARGS__)
+
  /**
   * Declare a tag defined by LLM_DEFINE_TAG. It is used in LLM_SCOPE_BYTAG or referenced by name in other LLM_SCOPEs.
   * @param UniqueName - The name of the tag for looking up by name, must be unique across all tags passed to
@@ -1059,6 +1084,14 @@ inline bool FLowLevelMemTracker::IsEnabled()
 
 #else
 
+#ifndef LLM_ALLOW_ASSETS_TAGS
+	#define LLM_ALLOW_ASSETS_TAGS 0
+#endif
+
+#ifndef LLM_ALLOW_UOBJECTCLASSES_TAGS
+	#define LLM_ALLOW_UOBJECTCLASSES_TAGS 0
+#endif
+
 #define LLM(...)
 #define LLM_IF_ENABLED(...)
 #define LLM_IS_ENABLED() false
@@ -1082,6 +1115,7 @@ inline bool FLowLevelMemTracker::IsEnabled()
 #define LLM_DUMP_TAG()
 #define LLM_DUMP_PLATFORM_TAG()
 #define LLM_DEFINE_TAG(...)
+#define LLM_DEFINE_STATIC_TAG(...)
 #define LLM_DECLARE_TAG(...)
 #define LLM_DECLARE_TAG_API(...)
 #define LLM_DEFINE_BOOTSTRAP_TAG(...)
diff --git a/Engine/Source/Runtime/CoreUObject/Private/Serialization/AsyncLoading.cpp b/Engine/Source/Runtime/CoreUObject/Private/Serialization/AsyncLoading.cpp
--- a/Engine/Source/Runtime/CoreUObject/Private/Serialization/AsyncLoading.cpp
+++ b/Engine/Source/Runtime/CoreUObject/Private/Serialization/AsyncLoading.cpp
@@ -98,6 +98,8 @@ DECLARE_FLOAT_ACCUMULATOR_STAT(TEXT("Total PostLoadObjects time GT"), STAT_FAsyn
 DECLARE_FLOAT_ACCUMULATOR_STAT( TEXT( "Async loading block time" ), STAT_AsyncIO_AsyncLoadingBlockingTime, STATGROUP_AsyncIO );
 DECLARE_FLOAT_ACCUMULATOR_STAT( TEXT( "Async package precache wait time" ), STAT_AsyncIO_AsyncPackagePrecacheWaitTime, STATGROUP_AsyncIO );
 
+LLM_DEFINE_TAG(UObject_FAsyncPackage);
+
 struct FEDLBootObjectState
 {
 	ENotifyRegistrationType NotifyRegistrationType;
@@ -3456,6 +3458,7 @@ void FAsyncPackage::MarkNewObjectForLoadIfItIsAnExport(UObject *Object)
 void FAsyncPackage::EventDrivenSerializeExport(int32 LocalExportIndex)
 {
 	LLM_SCOPE(ELLMTag::UObject);
+	LLM_SCOPE_BYTAG(UObject_FAsyncPackage);
 	SCOPED_LOADTIMER(Package_PreLoadObjects);
 
 	FObjectExport& Export = Linker->ExportMap[LocalExportIndex];
diff --git a/Engine/Source/Runtime/CoreUObject/Private/Serialization/AsyncLoading2.cpp b/Engine/Source/Runtime/CoreUObject/Private/Serialization/AsyncLoading2.cpp
--- a/Engine/Source/Runtime/CoreUObject/Private/Serialization/AsyncLoading2.cpp
+++ b/Engine/Source/Runtime/CoreUObject/Private/Serialization/AsyncLoading2.cpp
@@ -107,6 +107,8 @@
 CSV_DECLARE_CATEGORY_MODULE_EXTERN(CORE_API, FileIO);
 CSV_DEFINE_STAT(FileIO, FrameCompletedExportBundleLoadsKB);
 
+LLM_DEFINE_TAG(UObject_FAsyncPackage2);
+
 FArchive& operator<<(FArchive& Ar, FZenPackageVersioningInfo& VersioningInfo)
 {
 	Ar << VersioningInfo.ZenVersion;
@@ -8182,6 +8184,7 @@ void FAsyncPackage2::EventDrivenCreateExport(const FAsyncPackageHeaderData& Head
 bool FAsyncPackage2::EventDrivenSerializeExport(const FAsyncPackageHeaderData& Header, int32 LocalExportIndex, FExportArchive* Ar)
 {
 	LLM_SCOPE(ELLMTag::UObject);
+	LLM_SCOPE_BYTAG(UObject_FAsyncPackage2);
 	TRACE_CPUPROFILER_EVENT_SCOPE(SerializeExport);
 
 	const FExportMapEntry& Export = Header.ExportMap[LocalExportIndex];
@@ -9013,6 +9016,7 @@ EEventLoadNodeExecutionResult FAsyncPackage2::Event_DeferredPostLoadExportBundle
 						{
 							FScopeCycleCounterUObject ConstructorScope(Object, GET_STATID(STAT_FAsyncPackage_PostLoadObjectsGameThread));
 							LLM_SCOPE(ELLMTag::UObject);
+							LLM_SCOPE_BYTAG(UObject_FAsyncPackage2);
 							LLM_SCOPE_DYNAMIC_STAT_OBJECTPATH(Object->GetPackage(), ELLMTagSet::Assets);
 							LLM_SCOPE_DYNAMIC_STAT_OBJECTPATH(Object->GetClass(), ELLMTagSet::AssetClasses);
 							UE_TRACE_METADATA_SCOPE_ASSET(Object, Object->GetClass());
@@ -11106,8 +11110,10 @@ void FAsyncLoadingThread2::FlushLoading(TConstArrayView<int32> RequestIDs)
 	if (IsAsyncLoadingPackages())
 	{
 		LLM_SCOPE(ELLMTag::AsyncLoading);
+#if LLM_ALLOW_ASSETS_TAGS
 		LLM_TAGSET_SCOPE_CLEAR(ELLMTagSet::Assets);
 		LLM_TAGSET_SCOPE_CLEAR(ELLMTagSet::AssetClasses);
+#endif // LLM_ALLOW_ASSETS_TAGS
 		UE_TRACE_METADATA_CLEAR_SCOPE();
 		// Prevent objects still being loaded from being accessible from the game thread.
 		TGuardValue GuardVisibilityFilter(FUObjectThreadContext::Get().AsyncVisibilityFilter, EInternalObjectFlags::AsyncLoadingPhase1);
diff --git a/Engine/Source/Runtime/CoreUObject/Private/UObject/Class.cpp b/Engine/Source/Runtime/CoreUObject/Private/UObject/Class.cpp
--- a/Engine/Source/Runtime/CoreUObject/Private/UObject/Class.cpp
+++ b/Engine/Source/Runtime/CoreUObject/Private/UObject/Class.cpp
@@ -86,6 +86,8 @@ DEFINE_LOG_CATEGORY(LogClass);
 
 UE_IMPLEMENT_STRUCT("/Script/CoreUObject", TestUninitializedScriptStructMembersTest);
 
+LLM_DEFINE_TAG(UObject_UClass);
+
 #if defined(_MSC_VER) && _MSC_VER == 1900
 	#ifdef PRAGMA_DISABLE_SHADOW_VARIABLE_WARNINGS
 		PRAGMA_DISABLE_SHADOW_VARIABLE_WARNINGS
@@ -6902,6 +6904,7 @@ void UClass::ClearFunctionMapsCaches()
 UFunction* UClass::FindFunctionByName(FName InName, EIncludeSuperFlag::Type IncludeSuper) const
 {
 	LLM_SCOPE(ELLMTag::UObject);
+	LLM_SCOPE_BYTAG(UObject_UClass);
 
 	UFunction* Result = nullptr;
 
diff --git a/Engine/Source/Runtime/CoreUObject/Private/UObject/DeferredRegistry.h b/Engine/Source/Runtime/CoreUObject/Private/UObject/DeferredRegistry.h
--- a/Engine/Source/Runtime/CoreUObject/Private/UObject/DeferredRegistry.h
+++ b/Engine/Source/Runtime/CoreUObject/Private/UObject/DeferredRegistry.h
@@ -335,7 +335,9 @@ private:
 template <>
 inline typename TDeferredRegistry<FClassRegistrationInfo>::TType* TDeferredRegistry<FClassRegistrationInfo>::InnerRegister(const TDeferredRegistry<FClassRegistrationInfo>::FRegistrant& Registrant)
 {
+	LLM_DEFINE_STATIC_TAG(UObject_TDeferredRegistry);
 	LLM_SCOPE(ELLMTag::UObject);
+	LLM_SCOPE_BYTAG(UObject_TDeferredRegistry);
 	return Registrant.InnerRegisterFn();
 }
 
diff --git a/Engine/Source/Runtime/CoreUObject/Private/UObject/LinkerLoad.cpp b/Engine/Source/Runtime/CoreUObject/Private/UObject/LinkerLoad.cpp
--- a/Engine/Source/Runtime/CoreUObject/Private/UObject/LinkerLoad.cpp
+++ b/Engine/Source/Runtime/CoreUObject/Private/UObject/LinkerLoad.cpp
@@ -92,6 +92,7 @@ DECLARE_FLOAT_ACCUMULATOR_STAT(TEXT("Fixup editor-only flags time"), STAT_Editor
 FName FLinkerLoad::NAME_LoadErrors("LoadErrors");
 
 LLM_DEFINE_TAG(UObject_Linker);
+LLM_DEFINE_TAG(UObject_FLinkerLoad);
 
 /**
 * Helper function to determine and trace the most important asset class.
@@ -1686,6 +1687,7 @@ FLinkerLoad::ELinkerStatus FLinkerLoad::SerializePackageFileSummary()
 {
 	DECLARE_SCOPE_CYCLE_COUNTER( TEXT( "FLinkerLoad::SerializePackageFileSummary" ), STAT_LinkerLoad_SerializePackageFileSummary, STATGROUP_LinkerLoad );
 	LLM_SCOPE(ELLMTag::UObject);
+	LLM_SCOPE_BYTAG(UObject_FLinkerLoad);
 
 	if (bHasSerializedPackageFileSummary == false)
 	{
@@ -3672,6 +3674,7 @@ bool FLinkerLoad::VerifyImportInner(const int32 ImportIndex, FString& WarningSuf
 	auto LoadImportPackage = [this, ImportIndex](FObjectImport& Import, TOptional<FScopedSlowTask>& SlowTask) -> UPackage*
 	{
 		LLM_SCOPE(ELLMTag::UObject);
+		LLM_SCOPE_BYTAG(UObject_FLinkerLoad);
 
 		// Either this import is a UPackage or it has PackageName set.
 		check(Import.ClassName == NAME_Package || Import.HasPackageName());
diff --git a/Engine/Source/Runtime/CoreUObject/Private/UObject/Obj.cpp b/Engine/Source/Runtime/CoreUObject/Private/UObject/Obj.cpp
--- a/Engine/Source/Runtime/CoreUObject/Private/UObject/Obj.cpp
+++ b/Engine/Source/Runtime/CoreUObject/Private/UObject/Obj.cpp
@@ -83,6 +83,8 @@
 
 DEFINE_LOG_CATEGORY(LogObj);
 
+LLM_DEFINE_TAG(UObject_UObjectInternals);
+
 namespace
 {
 	static int32 GPrematureConfigLoadErrorLevel = 0;
@@ -1307,6 +1309,7 @@ bool UObject::ConditionalFinishDestroy()
 void UObject::ConditionalPostLoad()
 {
 	LLM_SCOPE(ELLMTag::UObject);
+	LLM_SCOPE_BYTAG(UObject_UObjectInternals);
 
 	ensureAlwaysMsgf((GetLoaderType() != ELoaderType::ZenLoader) || !HasAnyFlags(RF_NeedLoad), TEXT("Object '%s' does not have RF_NeedLoad cleared in PostLoad!"), *GetFullName());
 
diff --git a/Engine/Source/Runtime/CoreUObject/Private/UObject/ScriptCore.cpp b/Engine/Source/Runtime/CoreUObject/Private/UObject/ScriptCore.cpp
--- a/Engine/Source/Runtime/CoreUObject/Private/UObject/ScriptCore.cpp
+++ b/Engine/Source/Runtime/CoreUObject/Private/UObject/ScriptCore.cpp
@@ -1117,6 +1117,7 @@ void UObject::CallFunction( FFrame& Stack, RESULT_DECL, UFunction* Function )
 
 	SCOPE_CYCLE_UOBJECT(ContextScope, GVerboseScriptStats ? this : nullptr);
 	LLM_SCOPE(ELLMTag::UObject);
+	LLM_SCOPE_BYNAME(TEXT("UObject/UObjectInternals"));
 	LLM_SCOPE_DYNAMIC_STAT_OBJECTPATH(GetPackage(), ELLMTagSet::Assets);
 	LLM_SCOPE_DYNAMIC_STAT_OBJECTPATH(GetClass(), ELLMTagSet::AssetClasses);
 	UE_TRACE_METADATA_SCOPE_ASSET(this, GetClass());
@@ -1324,6 +1325,7 @@ void ProcessLocalFunction(UObject* Context, UFunction* Fn, FFrame& Stack, RESULT
 	if (Context && FLowLevelMemTracker::IsEnabled())
 	{
 		LLM_SCOPE(ELLMTag::UObject);
+		LLM_SCOPE_BYNAME(TEXT("UObject/UObjectInternals"));
 		LLM_SCOPE_DYNAMIC_STAT_OBJECTPATH(Context->GetPackage(), ELLMTagSet::Assets);
 		LLM_SCOPE_DYNAMIC_STAT_OBJECTPATH(Context->GetClass(), ELLMTagSet::AssetClasses);
 		UE_TRACE_METADATA_SCOPE_ASSET(Context, Context->GetClass());
@@ -2084,6 +2086,7 @@ void UObject::ProcessEvent( UFunction* Function, void* Parms )
 	SCOPE_CYCLE_UOBJECT(ContextScope, GVerboseScriptStats ? this : nullptr);
 
 	LLM_SCOPE(ELLMTag::UObject);
+	LLM_SCOPE_BYNAME(TEXT("UObject/UObjectInternals"));
 	LLM_SCOPE_DYNAMIC_STAT_OBJECTPATH(GetPackage(), ELLMTagSet::Assets);
 	LLM_SCOPE_DYNAMIC_STAT_OBJECTPATH(GetClass(), ELLMTagSet::AssetClasses);
 	UE_TRACE_METADATA_SCOPE_ASSET(this, GetClass());
diff --git a/Engine/Source/Runtime/CoreUObject/Private/UObject/UObjectArray.cpp b/Engine/Source/Runtime/CoreUObject/Private/UObject/UObjectArray.cpp
--- a/Engine/Source/Runtime/CoreUObject/Private/UObject/UObjectArray.cpp
+++ b/Engine/Source/Runtime/CoreUObject/Private/UObject/UObjectArray.cpp
@@ -22,6 +22,8 @@ DEFINE_LOG_CATEGORY_STATIC(LogUObjectArray, Log, All);
 
 DECLARE_CYCLE_STAT_WITH_FLAGS(TEXT("CreateStatId"), CreateStatId, STATGROUP_Object, EStatFlags::Verbose);
 
+LLM_DEFINE_TAG(UObject_UObjectArray);
+
 FUObjectClusterContainer GUObjectClusters;
 
 #if STATS || ENABLE_STATNAMEDEVENTS_UOBJECT
@@ -229,9 +231,12 @@ void FUObjectArray::DisableDisregardForGC()
 void FUObjectArray::AllocateUObjectIndex(UObjectBase* Object, EInternalObjectFlags InitialFlags, int32 AlreadyAllocatedIndex, int32 SerialNumber, FRemoteObjectId RemoteId)
 {
 	LLM_SCOPE(ELLMTag::UObject);
+	LLM_SCOPE_BYTAG(UObject_UObjectArray);
 	// Clear asset scopes
+#if LLM_ALLOW_ASSETS_TAGS
 	LLM_TAGSET_SCOPE_CLEAR(ELLMTagSet::Assets);
 	LLM_TAGSET_SCOPE_CLEAR(ELLMTagSet::AssetClasses);
+#endif // LLM_ALLOW_ASSETS_TAGS
 	UE_TRACE_METADATA_CLEAR_SCOPE();
 
 	int32 Index = INDEX_NONE;
@@ -359,6 +364,7 @@ void FUObjectArray::RemoveObjectFromDeleteListeners(UObjectBase* Object)
 void FUObjectArray::FreeUObjectIndex(UObjectBase* Object)
 {
 	LLM_SCOPE(ELLMTag::UObject);
+	LLM_SCOPE_BYTAG(UObject_UObjectArray);
 
 	// This should only be happening on the game thread (GC runs only on game thread when it's freeing objects)
 	check(IsInGameThread());
diff --git a/Engine/Source/Runtime/CoreUObject/Private/UObject/UObjectBase.cpp b/Engine/Source/Runtime/CoreUObject/Private/UObject/UObjectBase.cpp
--- a/Engine/Source/Runtime/CoreUObject/Private/UObject/UObjectBase.cpp
+++ b/Engine/Source/Runtime/CoreUObject/Private/UObject/UObjectBase.cpp
@@ -36,6 +36,8 @@ DECLARE_CYCLE_STAT(TEXT("CreateStatID"), STAT_CreateStatID, STATGROUP_StatSystem
 
 DEFINE_LOG_CATEGORY_STATIC(LogUObjectBootstrap, Display, Display);
 
+LLM_DEFINE_TAG(UObject_UObjectBase);
+
 #if CSV_PROFILER_STATS && CSV_TRACK_UOBJECT_COUNT
 namespace UObjectStats
 {
@@ -535,6 +537,7 @@ void UObjectBase::Register(const TCHAR* PackageName, const TCHAR* InName)
 void UObjectBase::Register(class UClass* (*StaticClassFn)(), const TCHAR* PackageName, const TCHAR* InName)
 {
 	LLM_SCOPE(ELLMTag::UObject);
+	LLM_SCOPE_BYTAG(UObject_UObjectBase);
 	TMap<UObjectBase*, FPendingRegistrantInfo>& PendingRegistrants = FPendingRegistrantInfo::GetMap();
 
 	FPendingRegistrant* PendingRegistration = new FPendingRegistrant(this);
@@ -624,6 +627,7 @@ static void UObjectProcessRegistrants()
 void UObjectForceRegistration(UObjectBase* Object, bool bCheckForModuleRelease)
 {
 	LLM_SCOPE(ELLMTag::UObject);
+	LLM_SCOPE_BYTAG(UObject_UObjectBase);
 	TMap<UObjectBase*, FPendingRegistrantInfo>& PendingRegistrants = FPendingRegistrantInfo::GetMap();
 
 	FPendingRegistrantInfo* Info = PendingRegistrants.Find(Object);
@@ -731,6 +735,7 @@ void RegisterCompiledInInfo(UPackage* (*InOuterRegister)(), const TCHAR* InPacka
 void RegisterCompiledInInfo(const TCHAR* PackageName, const FClassRegisterCompiledInInfo* ClassInfo, size_t NumClassInfo, const FStructRegisterCompiledInInfo* StructInfo, size_t NumStructInfo, const FEnumRegisterCompiledInInfo* EnumInfo, size_t NumEnumInfo)
 {
 	LLM_SCOPE(ELLMTag::UObject);
+	LLM_SCOPE_BYTAG(UObject_UObjectBase);
 
 	for (size_t Index = 0; Index < NumClassInfo; ++Index)
 	{
@@ -763,6 +768,7 @@ void UClassRegisterAllCompiledInClasses()
 #endif
 	SCOPED_BOOT_TIMING("UClassRegisterAllCompiledInClasses");
 	LLM_SCOPE(ELLMTag::UObject);
+	LLM_SCOPE_BYTAG(UObject_UObjectBase);
 
 	FClassDeferredRegistry& Registry = FClassDeferredRegistry::Get();
 
@@ -981,6 +987,7 @@ void ProcessNewlyLoadedUObjects(FName InModuleName, bool bCanProcessNewlyLoadedO
 		return;
 	}
 	LLM_SCOPE(ELLMTag::UObject);
+	LLM_SCOPE_BYTAG(UObject_UObjectBase);
 	DECLARE_SCOPE_CYCLE_COUNTER(TEXT("ProcessNewlyLoadedUObjects"), STAT_ProcessNewlyLoadedUObjects, STATGROUP_ObjectVerbose);
 
 	FPackageDeferredRegistry& PackageRegistry = FPackageDeferredRegistry::Get();
@@ -1060,6 +1067,7 @@ static UPackage* GetScriptPackageFromModuleName (FName InModuleName )
 static void RemoveObjectsInModules(TConstArrayView<FName> InModuleNames)
 {
 	LLM_SCOPE(ELLMTag::UObject);
+	LLM_SCOPE_BYTAG(UObject_UObjectBase);
 	DECLARE_SCOPE_CYCLE_COUNTER(TEXT("RemoveObjectsInModules"), STAT_RemoveObjectsInModules, STATGROUP_ObjectVerbose);
 
 	FPermanentObjectPoolExtents PermanentPool;
@@ -1131,6 +1139,7 @@ static bool CheckLiveObjectsInModules(TConstArrayView<FName> InModuleNames)
 static void RemoveLoadedUObjects(TConstArrayView<FName> InModuleNames)
 {
 	LLM_SCOPE(ELLMTag::UObject);
+	LLM_SCOPE_BYTAG(UObject_UObjectBase);
 	DECLARE_SCOPE_CYCLE_COUNTER(TEXT("RemoveLoadedUObjects"), STAT_RemoveLoadedUObjects, STATGROUP_ObjectVerbose);
 
 	FlushAsyncLoading();
diff --git a/Engine/Source/Runtime/CoreUObject/Private/UObject/UObjectGlobals.cpp b/Engine/Source/Runtime/CoreUObject/Private/UObject/UObjectGlobals.cpp
--- a/Engine/Source/Runtime/CoreUObject/Private/UObject/UObjectGlobals.cpp
+++ b/Engine/Source/Runtime/CoreUObject/Private/UObject/UObjectGlobals.cpp
@@ -161,6 +161,9 @@ DECLARE_CYCLE_STAT(TEXT("PostReinitProperties"), STAT_PostReinitProperties, STAT
 
 CSV_DEFINE_CATEGORY(UObject, false);
 
+LLM_DEFINE_TAG(UObject_StaticConstructObjectInternal);
+LLM_DEFINE_TAG(UObject_StaticAllocateObject);
+
 #if ENABLE_COOK_STATS
 #include "ProfilingDebugging/ScopedTimers.h"
 
@@ -2791,8 +2794,10 @@ namespace NameReuse
 			return FName();
 		}
 
+#if LLM_ALLOW_ASSETS_TAGS
 		LLM_TAGSET_SCOPE_CLEAR(ELLMTagSet::Assets);
 		LLM_TAGSET_SCOPE_CLEAR(ELLMTagSet::AssetClasses);
+#endif // LLM_ALLOW_ASSETS_TAGS
 		UE_TRACE_METADATA_CLEAR_SCOPE();
 		LLM_SCOPE(ELLMTag::FName);
 
@@ -3553,6 +3558,15 @@ UObject* StaticAllocateObject
 	FGCReconstructionGuard* GCGuard)
 {
 	LLM_SCOPE(ELLMTag::UObject);
+	LLM_SCOPE_BYTAG(UObject_StaticAllocateObject);
+	
+	// report detailed UObject classes allocations
+#if ENABLE_LOW_LEVEL_MEM_TRACKER && LLM_ALLOW_UOBJECTCLASSES_TAGS
+	FString ClassNameString = FString::Printf(TEXT("UObject/%s"), *InClass->GetName());
+	FName LLMScope_Name(*ClassNameString);
+
+	LLM_SCOPE_DYNAMIC(LLMScope_Name, ELLMTracker::Default, ELLMTagSet::UObjectClasses, FLLMDynamicTagConstructorStatString(ClassNameString));
+#endif // ENABLE_LOW_LEVEL_MEM_TRACKER
 
 	SCOPE_CYCLE_COUNTER(STAT_AllocateObject);
 	checkSlow(InOuter != INVALID_OBJECT); // not legal
@@ -4904,6 +4918,7 @@ UObject* StaticConstructObject_Internal(const FStaticConstructObjectParameters&
 	UE::RemoteObject::Serialization::FRemoteObjectConstructionOverrides* RemoteSubObjectOverrides = nullptr;
 
 	LLM_SCOPE(ELLMTag::UObject);
+	LLM_SCOPE_BYTAG(UObject_StaticConstructObjectInternal);
 
 	SCOPE_CYCLE_COUNTER(STAT_ConstructObject);
 	UObject* Result = NULL;
diff --git a/Engine/Source/Runtime/CoreUObject/Private/UObject/UObjectHash.cpp b/Engine/Source/Runtime/CoreUObject/Private/UObject/UObjectHash.cpp
--- a/Engine/Source/Runtime/CoreUObject/Private/UObject/UObjectHash.cpp
+++ b/Engine/Source/Runtime/CoreUObject/Private/UObject/UObjectHash.cpp
@@ -5,6 +5,7 @@
 =============================================================================*/
 
 #include "UObject/UObjectHash.h"
+#include "HAL/LowLevelMemTracker.h"
 #include "Misc/AssertionMacros.h"
 #include "Misc/TransactionallySafeCriticalSection.h"
 #include "UObject/Class.h"
@@ -1529,8 +1530,10 @@ FORCEINLINE static void RemoveFromRemoteIdMap(FUObjectHashTables& ThreadHash, UO
 void ShrinkUObjectHashTables()
 {
 	LLM_SCOPE_BYTAG(UObjectHash);
+#if LLM_ALLOW_ASSETS_TAGS
 	LLM_TAGSET_SCOPE_CLEAR(ELLMTagSet::Assets);
 	LLM_TAGSET_SCOPE_CLEAR(ELLMTagSet::AssetClasses);
+#endif // LLM_ALLOW_ASSETS_TAGS
 	UE_TRACE_METADATA_CLEAR_SCOPE();
 	TRACE_CPUPROFILER_EVENT_SCOPE(ShrinkUObjectHashTables);
 	FUObjectHashTables& ThreadHash = FUObjectHashTables::Get();
@@ -1965,8 +1968,10 @@ void HashObject(UObjectBase* Object)
 	if (Name != NAME_None)
 	{
 		LLM_SCOPE_BYTAG(UObjectHash);
+#if LLM_ALLOW_ASSETS_TAGS
 		LLM_TAGSET_SCOPE_CLEAR(ELLMTagSet::Assets);
 		LLM_TAGSET_SCOPE_CLEAR(ELLMTagSet::AssetClasses);
+#endif // LLM_ALLOW_ASSETS_TAGS
 		UE_TRACE_METADATA_CLEAR_SCOPE();
 #if !UE_BUILD_TEST && !UE_BUILD_SHIPPING
 		SCOPE_CYCLE_COUNTER(STAT_Hash_HashObject);
@@ -2020,8 +2025,10 @@ void UnhashObject(UObjectBase* Object)
 	if (Name != NAME_None)
 	{
 		LLM_SCOPE_BYTAG(UObjectHash);
+#if LLM_ALLOW_ASSETS_TAGS
 		LLM_TAGSET_SCOPE_CLEAR(ELLMTagSet::Assets);
 		LLM_TAGSET_SCOPE_CLEAR(ELLMTagSet::AssetClasses);
+#endif // LLM_ALLOW_ASSETS_TAGS
 		UE_TRACE_METADATA_CLEAR_SCOPE();
 #if !UE_BUILD_TEST && !UE_BUILD_SHIPPING
 		SCOPE_CYCLE_COUNTER(STAT_Hash_UnhashObject);
@@ -2064,8 +2071,10 @@ void UnhashObject(UObjectBase* Object)
 void HashObjectExternalPackage(UObjectBase* Object, UPackage* Package)
 {
 	LLM_SCOPE_BYTAG(UObjectHash);
+#if LLM_ALLOW_ASSETS_TAGS
 	LLM_TAGSET_SCOPE_CLEAR(ELLMTagSet::Assets);
 	LLM_TAGSET_SCOPE_CLEAR(ELLMTagSet::AssetClasses);
+#endif // LLM_ALLOW_ASSETS_TAGS
 	UE_TRACE_METADATA_CLEAR_SCOPE();
 	if (Package)
 	{
@@ -2091,8 +2100,10 @@ void HashObjectExternalPackage(UObjectBase* Object, UPackage* Package)
 void UnhashObjectExternalPackage(class UObjectBase* Object)
 {
 	LLM_SCOPE_BYTAG(UObjectHash);
+#if LLM_ALLOW_ASSETS_TAGS
 	LLM_TAGSET_SCOPE_CLEAR(ELLMTagSet::Assets);
 	LLM_TAGSET_SCOPE_CLEAR(ELLMTagSet::AssetClasses);
+#endif // LLM_ALLOW_ASSETS_TAGS
 	UE_TRACE_METADATA_CLEAR_SCOPE();
 	FUObjectHashTables& ThreadHash = FUObjectHashTables::Get();
 	FHashTableLock LockHash(ThreadHash);
diff --git a/Engine/Source/Runtime/Engine/Private/ContentStreaming.cpp b/Engine/Source/Runtime/Engine/Private/ContentStreaming.cpp
--- a/Engine/Source/Runtime/Engine/Private/ContentStreaming.cpp
+++ b/Engine/Source/Runtime/Engine/Private/ContentStreaming.cpp
@@ -373,8 +373,10 @@ FStreamingManagerCollection& IStreamingManager::Get()
 	if (StreamingManagerCollection == nullptr)
 	{
 		LLM_SCOPE(ELLMTag::StreamingManager);
+#if LLM_ALLOW_ASSETS_TAGS
 		LLM_TAGSET_SCOPE_CLEAR(ELLMTagSet::Assets);
 		LLM_TAGSET_SCOPE_CLEAR(ELLMTagSet::AssetClasses);
+#endif // LLM_ALLOW_ASSETS_TAGS
 		UE_TRACE_METADATA_CLEAR_SCOPE();
 		// Since this is a lazily created static global variable we create it
 		// in the open.
diff --git a/Engine/Source/Runtime/Engine/Private/PSOPrecacheMaterial.cpp b/Engine/Source/Runtime/Engine/Private/PSOPrecacheMaterial.cpp
--- a/Engine/Source/Runtime/Engine/Private/PSOPrecacheMaterial.cpp
+++ b/Engine/Source/Runtime/Engine/Private/PSOPrecacheMaterial.cpp
@@ -188,8 +188,10 @@ public:
 	FMaterialPSOPrecacheRequestID PrecachePSOs(const FMaterialPSOPrecacheParams& Params, EPSOPrecachePriority Priority, FGraphEventArray& OutGraphEvents)
 	{
 		LLM_SCOPE(ELLMTag::PSO);
+#if LLM_ALLOW_ASSETS_TAGS
 		LLM_TAGSET_SCOPE_CLEAR(ELLMTagSet::Assets);
 		LLM_TAGSET_SCOPE_CLEAR(ELLMTagSet::AssetClasses);
+#endif // LLM_ALLOW_ASSETS_TAGS
 		UE_TRACE_METADATA_CLEAR_SCOPE();
 		FMaterialPSOPrecacheRequestID RequestID = INDEX_NONE;		
 
@@ -357,8 +359,10 @@ public:
 	void MarkCollectionComplete(const FMaterialPSOPrecacheParams& Params, const FPSOPrecacheDataArray& PrecacheData, const FPSOPrecacheRequestResultArray& PrecacheRequestResults, uint32 RequestLifecycleID)
 	{
 		LLM_SCOPE(ELLMTag::PSO);
+#if LLM_ALLOW_ASSETS_TAGS
 		LLM_TAGSET_SCOPE_CLEAR(ELLMTagSet::Assets);
 		LLM_TAGSET_SCOPE_CLEAR(ELLMTagSet::AssetClasses);
+#endif // LLM_ALLOW_ASSETS_TAGS
 		UE_TRACE_METADATA_CLEAR_SCOPE();
 		FRWScopeLock WriteLock(RWLock, SLT_Write);
 		
@@ -395,8 +399,10 @@ public:
 	void ReleasePrecacheData(FMaterialPSOPrecacheRequestID MaterialPSORequestID)
 	{
 		LLM_SCOPE(ELLMTag::PSO);
+#if LLM_ALLOW_ASSETS_TAGS
 		LLM_TAGSET_SCOPE_CLEAR(ELLMTagSet::Assets);
 		LLM_TAGSET_SCOPE_CLEAR(ELLMTagSet::AssetClasses);
+#endif // LLM_ALLOW_ASSETS_TAGS
 		UE_TRACE_METADATA_CLEAR_SCOPE();
 		check(MaterialPSORequestID != INDEX_NONE);
 
@@ -411,8 +417,10 @@ public:
 	void BoostPriority(EPSOPrecachePriority NewPri, FMaterialPSOPrecacheRequestID MaterialPSORequestID)
 	{
 		LLM_SCOPE(ELLMTag::PSO);
+#if LLM_ALLOW_ASSETS_TAGS
 		LLM_TAGSET_SCOPE_CLEAR(ELLMTagSet::Assets);
 		LLM_TAGSET_SCOPE_CLEAR(ELLMTagSet::AssetClasses);
+#endif // LLM_ALLOW_ASSETS_TAGS
 		UE_TRACE_METADATA_CLEAR_SCOPE();
 		check(MaterialPSORequestID != INDEX_NONE);
 
@@ -446,8 +454,10 @@ public:
 	void ClearMaterialPSORequests()
 	{
 		LLM_SCOPE(ELLMTag::PSO);
+#if LLM_ALLOW_ASSETS_TAGS
 		LLM_TAGSET_SCOPE_CLEAR(ELLMTagSet::Assets);
 		LLM_TAGSET_SCOPE_CLEAR(ELLMTagSet::AssetClasses);
+#endif // LLM_ALLOW_ASSETS_TAGS
 		UE_TRACE_METADATA_CLEAR_SCOPE();
 		check(IsInGameThread());
 
@@ -584,8 +594,10 @@ FMaterialPSORequestManager GMaterialPSORequestManager;
 void FMaterialPSOPrecacheCollectionTask::DoTask(ENamedThreads::Type CurrentThread, const FGraphEventRef& MyCompletionGraphEvent)
 {
 	LLM_SCOPE(ELLMTag::PSO);
+#if LLM_ALLOW_ASSETS_TAGS
 	LLM_TAGSET_SCOPE_CLEAR(ELLMTagSet::Assets);
 	LLM_TAGSET_SCOPE_CLEAR(ELLMTagSet::AssetClasses);
+#endif // LLM_ALLOW_ASSETS_TAGS
 	UE_TRACE_METADATA_CLEAR_SCOPE();
 	TRACE_CPUPROFILER_EVENT_SCOPE(FMaterialPSOPrecacheCollectionTask);
 
@@ -635,8 +647,10 @@ void FMaterialPSOPrecacheCollectionTask::DoTask(ENamedThreads::Type CurrentThrea
 void FShaderMapPreloadTask::DoTask(ENamedThreads::Type CurrentThread, const FGraphEventRef& MyCompletionGraphEvent)
 {
 	LLM_SCOPE(ELLMTag::PSO);
+#if LLM_ALLOW_ASSETS_TAGS
 	LLM_TAGSET_SCOPE_CLEAR(ELLMTagSet::Assets);
 	LLM_TAGSET_SCOPE_CLEAR(ELLMTagSet::AssetClasses);
+#endif // LLM_ALLOW_ASSETS_TAGS
 	UE_TRACE_METADATA_CLEAR_SCOPE();
 	TRACE_CPUPROFILER_EVENT_SCOPE(FShaderMapPreloadTask);
 
@@ -658,8 +672,10 @@ void FShaderMapPreloadTask::DoTask(ENamedThreads::Type CurrentThread, const FGra
 void FShaderPreloadCollectionTask::DoTask(ENamedThreads::Type CurrentThread, const FGraphEventRef& MyCompletionGraphEvent)
 {
 	LLM_SCOPE(ELLMTag::PSO);
+#if LLM_ALLOW_ASSETS_TAGS
 	LLM_TAGSET_SCOPE_CLEAR(ELLMTagSet::Assets);
 	LLM_TAGSET_SCOPE_CLEAR(ELLMTagSet::AssetClasses);
+#endif // LLM_ALLOW_ASSETS_TAGS
 	UE_TRACE_METADATA_CLEAR_SCOPE();
 	TRACE_CPUPROFILER_EVENT_SCOPE(FShaderPreloadCollectionTask);
 
diff --git a/Engine/Source/Runtime/Engine/Private/ShaderCompiler/ShaderCompiler.cpp b/Engine/Source/Runtime/Engine/Private/ShaderCompiler/ShaderCompiler.cpp
--- a/Engine/Source/Runtime/Engine/Private/ShaderCompiler/ShaderCompiler.cpp
+++ b/Engine/Source/Runtime/Engine/Private/ShaderCompiler/ShaderCompiler.cpp
@@ -5265,7 +5265,10 @@ static bool TryLoadCookedGlobalShaderMap(EShaderPlatform Platform, FScopedSlowTa
 
 void CompileGlobalShaderMap(EShaderPlatform Platform, const ITargetPlatform* TargetPlatform, bool bRefreshShaderMap)
 {
+	// as there is no assets tags, there is no need to register this
+#if LLM_ALLOW_ASSETS_TAGS
 	LLM_SCOPE_RENDER_RESOURCE(TEXT("GlobalShaderMap"));
+#endif // LLM_ALLOW_ASSETS_TAGS
 
 	// No global shaders needed on dedicated server or clients that use NullRHI. Note that cook commandlet needs to have them, even if it is not allowed to render otherwise.
 	if (FPlatformProperties::IsServerOnly() || (!IsRunningCommandlet() && !FApp::CanEverRender()))
diff --git a/Engine/Source/Runtime/Engine/Private/Streaming/StreamingManagerTexture.cpp b/Engine/Source/Runtime/Engine/Private/Streaming/StreamingManagerTexture.cpp
--- a/Engine/Source/Runtime/Engine/Private/Streaming/StreamingManagerTexture.cpp
+++ b/Engine/Source/Runtime/Engine/Private/Streaming/StreamingManagerTexture.cpp
@@ -305,8 +305,10 @@ void FRenderAssetStreamingManager::MarkMountedStateDirty(FIoFilenameHash Filenam
 	if (!bRecacheAllFiles && FilenameHash != INVALID_IO_FILENAME_HASH)
 	{
 		LLM_SCOPE(ELLMTag::StreamingManager);
+#if LLM_ALLOW_ASSETS_TAGS
 		LLM_TAGSET_SCOPE_CLEAR(ELLMTagSet::Assets);
 		LLM_TAGSET_SCOPE_CLEAR(ELLMTagSet::AssetClasses);
+#endif // LLM_ALLOW_ASSETS_TAGS
 		UE_TRACE_METADATA_CLEAR_SCOPE();
 		FScopeLock ScopeLock(&MountedStateDirtyFilesCS);
 		MountedStateDirtyFiles.Emplace(FilenameHash);
@@ -370,8 +372,10 @@ bool FRenderAssetStreamingManager::StreamOutRenderAssetData( int64 RequiredMemor
 {
 	FScopeLock ScopeLock(&CriticalSection);
 	LLM_SCOPE(ELLMTag::StreamingManager);
+#if LLM_ALLOW_ASSETS_TAGS
 	LLM_TAGSET_SCOPE_CLEAR(ELLMTagSet::Assets);
 	LLM_TAGSET_SCOPE_CLEAR(ELLMTagSet::AssetClasses);
+#endif // LLM_ALLOW_ASSETS_TAGS
 	UE_TRACE_METADATA_CLEAR_SCOPE();
 
 	const int64 MaxTempMemoryAllowed = static_cast<int64>(Settings.MaxTempMemoryAllowed) * 1024 * 1024;
@@ -843,8 +847,10 @@ void FRenderAssetStreamingManager::AddLevel( ULevel* Level )
 {
 	FScopeLock ScopeLock(&CriticalSection);
 	LLM_SCOPE(ELLMTag::StreamingManager);
+#if LLM_ALLOW_ASSETS_TAGS
 	LLM_TAGSET_SCOPE_CLEAR(ELLMTagSet::Assets);
 	LLM_TAGSET_SCOPE_CLEAR(ELLMTagSet::AssetClasses);
+#endif // LLM_ALLOW_ASSETS_TAGS
 	UE_TRACE_METADATA_CLEAR_SCOPE();
 
 	check(Level);
@@ -957,8 +963,10 @@ void FRenderAssetStreamingManager::AddStreamingRenderAsset(UStreamableRenderAsse
 {
 	FScopeLock ScopeLock(&CriticalSection);
 	LLM_SCOPE(ELLMTag::StreamingManager);
+#if LLM_ALLOW_ASSETS_TAGS
 	LLM_TAGSET_SCOPE_CLEAR(ELLMTagSet::Assets);
 	LLM_TAGSET_SCOPE_CLEAR(ELLMTagSet::AssetClasses);
+#endif // LLM_ALLOW_ASSETS_TAGS
 	UE_TRACE_METADATA_CLEAR_SCOPE();
 
 	STAT(GatheredStats.CallbacksCycles = -(int32)FPlatformTime::Cycles();)
@@ -983,8 +991,10 @@ void FRenderAssetStreamingManager::RemoveStreamingRenderAsset( UStreamableRender
 {
 	FScopeLock ScopeLock(&CriticalSection);
 	LLM_SCOPE(ELLMTag::StreamingManager);
+#if LLM_ALLOW_ASSETS_TAGS
 	LLM_TAGSET_SCOPE_CLEAR(ELLMTagSet::Assets);
 	LLM_TAGSET_SCOPE_CLEAR(ELLMTagSet::AssetClasses);
+#endif // LLM_ALLOW_ASSETS_TAGS
 	UE_TRACE_METADATA_CLEAR_SCOPE();
 	TArray<FStreamingRenderAsset>& StreamingRenderAssets = GetStreamingRenderAssetsAsyncSafe();
 
diff --git a/Engine/Source/Runtime/Experimental/IoStore/OnDemand/Private/OnDemandIoStore.cpp b/Engine/Source/Runtime/Experimental/IoStore/OnDemand/Private/OnDemandIoStore.cpp
--- a/Engine/Source/Runtime/Experimental/IoStore/OnDemand/Private/OnDemandIoStore.cpp
+++ b/Engine/Source/Runtime/Experimental/IoStore/OnDemand/Private/OnDemandIoStore.cpp
@@ -2121,8 +2121,10 @@ void FOnDemandIoStore::CreateContainersFromToc(
 {
 	static FName AssetClassName(TEXT("OnDemandIoStore"));
 	LLM_SCOPE_BYNAME(TEXT("FileSystem/OnDemandIoStore"));
+#if LLM_ALLOW_ASSETS_TAGS
 	LLM_TAGSET_SCOPE(FName(MountId), ELLMTagSet::Assets);
 	LLM_TAGSET_SCOPE(AssetClassName, ELLMTagSet::AssetClasses);
+#endif // LLM_ALLOW_ASSETS_TAGS
 	UE_TRACE_METADATA_SCOPE_ASSET_FNAME(FName(MountId), AssetClassName, FName(TocPath));
 	const FOnDemandTocHeader& Header = Toc.Header;
 	const FName CompressionFormat(Header.CompressionFormat);
@@ -2234,8 +2236,10 @@ void FOnDemandIoStore::CreateContainersFromToc(
 		}
 
 		// Do not count output container memory as being allocated for utocs
+#if LLM_ALLOW_ASSETS_TAGS
 		LLM_TAGSET_SCOPE_CLEAR(ELLMTagSet::Assets);
 		LLM_TAGSET_SCOPE_CLEAR(ELLMTagSet::AssetClasses);
+#endif // LLM_ALLOW_ASSETS_TAGS
 		UE_TRACE_METADATA_CLEAR_SCOPE();
 		Out.Add(MoveTemp(Container));
 	}
diff --git a/Engine/Source/Runtime/RHI/Private/RHI.cpp b/Engine/Source/Runtime/RHI/Private/RHI.cpp
--- a/Engine/Source/Runtime/RHI/Private/RHI.cpp
+++ b/Engine/Source/Runtime/RHI/Private/RHI.cpp
@@ -524,8 +524,10 @@ void FRHIResource::BeginTrackingResource(FRHIResource* InResource)
 	if (GRHITrackingResources)
 	{
 		LLM_SCOPE_BYNAME(TEXT("RHIMisc/ResourceTracking"));
+#if LLM_ALLOW_ASSETS_TAGS
 		LLM_TAGSET_SCOPE_CLEAR(ELLMTagSet::Assets);
 		LLM_TAGSET_SCOPE_CLEAR(ELLMTagSet::AssetClasses);
+#endif // LLM_ALLOW_ASSETS_TAGS
 		UE_TRACE_METADATA_CLEAR_SCOPE();
 
 		FScopeLock Lock(&GRHIResourceTrackingCriticalSection);
@@ -541,8 +543,10 @@ void FRHIResource::EndTrackingResource(FRHIResource* InResource)
 	if (InResource->bBeingTracked)
 	{
 		LLM_SCOPE_BYNAME(TEXT("RHIMisc/ResourceTracking"));
+#if LLM_ALLOW_ASSETS_TAGS
 		LLM_TAGSET_SCOPE_CLEAR(ELLMTagSet::Assets);
 		LLM_TAGSET_SCOPE_CLEAR(ELLMTagSet::AssetClasses);
+#endif // LLM_ALLOW_ASSETS_TAGS
 		UE_TRACE_METADATA_CLEAR_SCOPE();
 		FScopeLock Lock(&GRHIResourceTrackingCriticalSection);
 		GRHITrackedResources.Remove(InResource);
@@ -558,8 +562,10 @@ void FRHIResource::StartTrackingAllResources()
 void FRHIResource::StopTrackingAllResources()
 {
 	LLM_SCOPE_BYNAME(TEXT("RHIMisc/ResourceTracking"));
+#if LLM_ALLOW_ASSETS_TAGS
 	LLM_TAGSET_SCOPE_CLEAR(ELLMTagSet::Assets);
 	LLM_TAGSET_SCOPE_CLEAR(ELLMTagSet::AssetClasses);
+#endif // LLM_ALLOW_ASSETS_TAGS
 	UE_TRACE_METADATA_CLEAR_SCOPE();
 	FScopeLock Lock(&GRHIResourceTrackingCriticalSection);
 	for (FRHIResource* Resource : GRHITrackedResources)
diff --git a/Engine/Source/Runtime/RHI/Private/TextureProfiler.cpp b/Engine/Source/Runtime/RHI/Private/TextureProfiler.cpp
--- a/Engine/Source/Runtime/RHI/Private/TextureProfiler.cpp
+++ b/Engine/Source/Runtime/RHI/Private/TextureProfiler.cpp
@@ -168,8 +168,10 @@ FTextureProfiler* FTextureProfiler::Get()
 void FTextureProfiler::AddTextureAllocation(FRHITexture* UniqueTexturePtr, size_t Size, uint32 Alignment, size_t AllocationWaste)
 {
 	LLM_SCOPE_BYNAME("Debug/TextureProfiler");
+#if LLM_ALLOW_ASSETS_TAGS
 	LLM_TAGSET_SCOPE_CLEAR(ELLMTagSet::Assets);
 	LLM_TAGSET_SCOPE_CLEAR(ELLMTagSet::AssetClasses);
+#endif // LLM_ALLOW_ASSETS_TAGS
 	UE_TRACE_METADATA_CLEAR_SCOPE();
 	FScopeLock Lock(&TextureMapCS);
 
@@ -214,8 +216,10 @@ void FTextureProfiler::AddTextureAllocation(FRHITexture* UniqueTexturePtr, size_
 void FTextureProfiler::UpdateTextureAllocation(FRHITexture* UniqueTexturePtr, size_t Size, uint32 Alignment, size_t AllocationWaste)
 {
 	LLM_SCOPE_BYNAME("Debug/TextureProfiler");
+#if LLM_ALLOW_ASSETS_TAGS
 	LLM_TAGSET_SCOPE_CLEAR(ELLMTagSet::Assets);
 	LLM_TAGSET_SCOPE_CLEAR(ELLMTagSet::AssetClasses);
+#endif // LLM_ALLOW_ASSETS_TAGS
 	UE_TRACE_METADATA_CLEAR_SCOPE();
 	FScopeLock Lock(&TextureMapCS);
 
@@ -243,8 +247,10 @@ void FTextureProfiler::UpdateTextureAllocation(FRHITexture* UniqueTexturePtr, si
 void FTextureProfiler::RemoveTextureAllocation(FRHITexture* UniqueTexturePtr)
 {
 	LLM_SCOPE_BYNAME("Debug/TextureProfiler");
+#if LLM_ALLOW_ASSETS_TAGS
 	LLM_TAGSET_SCOPE_CLEAR(ELLMTagSet::Assets);
 	LLM_TAGSET_SCOPE_CLEAR(ELLMTagSet::AssetClasses);
+#endif // LLM_ALLOW_ASSETS_TAGS
 	UE_TRACE_METADATA_CLEAR_SCOPE();
 	FScopeLock Lock(&TextureMapCS);
 	FTextureDetails Details = TexturesMap.FindAndRemoveChecked(UniqueTexturePtr);
@@ -259,8 +265,10 @@ void FTextureProfiler::RemoveTextureAllocation(FRHITexture* UniqueTexturePtr)
 void FTextureProfiler::UpdateTextureName(FRHITexture* UniqueTexturePtr)
 {
 	LLM_SCOPE_BYNAME("Debug/TextureProfiler");
+#if LLM_ALLOW_ASSETS_TAGS
 	LLM_TAGSET_SCOPE_CLEAR(ELLMTagSet::Assets);
 	LLM_TAGSET_SCOPE_CLEAR(ELLMTagSet::AssetClasses);
+#endif // LLM_ALLOW_ASSETS_TAGS
 	UE_TRACE_METADATA_CLEAR_SCOPE();
 	FScopeLock Lock(&TextureMapCS);
 
diff --git a/Engine/Source/Runtime/RenderCore/Private/RenderGraphResourcePool.cpp b/Engine/Source/Runtime/RenderCore/Private/RenderGraphResourcePool.cpp
--- a/Engine/Source/Runtime/RenderCore/Private/RenderGraphResourcePool.cpp
+++ b/Engine/Source/Runtime/RenderCore/Private/RenderGraphResourcePool.cpp
@@ -182,8 +182,10 @@ FRDGPooledBuffer* FRDGBufferPool::CreateBuffer(FRHICommandListBase& RHICmdList,
 	
 	LLM_SCOPE_BYNAME(TEXT("RHIMisc/BufferPool"));
 	UE_TRACE_METADATA_CLEAR_SCOPE(); // Do not associate a pooled buffer with specific asset
+#if LLM_ALLOW_ASSETS_TAGS
 	LLM_TAGSET_SCOPE_CLEAR(ELLMTagSet::Assets);
 	LLM_TAGSET_SCOPE_CLEAR(ELLMTagSet::AssetClasses);
+#endif // LLM_ALLOW_ASSETS_TAGS
 
 	const FRHIBufferCreateDesc CreateDesc =
 		FRHIBufferCreateDesc::Create(InDebugName, NumBytes, Desc.BytesPerElement, Desc.Usage)
