Fix warning on PVS-Studio

#rnx

diff --git a/Engine/Source/Runtime/Core/Public/HAL/LowLevelMemTracker.h b/Engine/Source/Runtime/Core/Public/HAL/LowLevelMemTracker.h
--- a/Engine/Source/Runtime/Core/Public/HAL/LowLevelMemTracker.h
+++ b/Engine/Source/Runtime/Core/Public/HAL/LowLevelMemTracker.h
@@ -361,7 +361,8 @@ extern CORE_API const FName LLMGetUntaggedTagName(ELLMTagSet TagSet);
 				{ \
 					SCOPE_NAME.TryAddTagAndActivate(UniqueNameEvaluated, Constructor); \
 				} \
-				if (TagSet == ELLMTagSet::None && Tracker == ELLMTracker::Default) \
+				ELLMTracker InTracker = Tracker; \
+				if (TagSet == ELLMTagSet::None && InTracker == ELLMTracker::Default) \
 				{ \
 					UE_MEMSCOPE_ACTIVATE(__LINE__, UniqueNameEvaluated); \
 				} \
