Second attempt to fix the PVS-Studio warning

#rnx

diff --git a/Engine/Source/Runtime/Core/Public/HAL/LowLevelMemTracker.h b/Engine/Source/Runtime/Core/Public/HAL/LowLevelMemTracker.h
--- a/Engine/Source/Runtime/Core/Public/HAL/LowLevelMemTracker.h
+++ b/Engine/Source/Runtime/Core/Public/HAL/LowLevelMemTracker.h
@@ -361,8 +361,8 @@ extern CORE_API const FName LLMGetUntaggedTagName(ELLMTagSet TagSet);
 				{ \
 					SCOPE_NAME.TryAddTagAndActivate(UniqueNameEvaluated, Constructor); \
 				} \
-				ELLMTracker InTracker = Tracker; \
-				if (TagSet == ELLMTagSet::None && InTracker == ELLMTracker::Default) \
+				constexpr int DefaultTracker = (int)ELLMTracker::Default; \
+				if (TagSet == ELLMTagSet::None && (int)Tracker == DefaultTracker) \
 				{ \
 					UE_MEMSCOPE_ACTIVATE(__LINE__, UniqueNameEvaluated); \
 				} \
