Fixed UE_STATIC_ASSERT_COMPLETE_TYPE for zero-sized arrays.

diff --git a/Engine/Source/Runtime/Core/Public/Misc/StaticAssertCompleteType.h b/Engine/Source/Runtime/Core/Public/Misc/StaticAssertCompleteType.h
--- a/Engine/Source/Runtime/Core/Public/Misc/StaticAssertCompleteType.h
+++ b/Engine/Source/Runtime/Core/Public/Misc/StaticAssertCompleteType.h
@@ -2,11 +2,45 @@
 
 #pragma once
 
+#include "CoreTypes.h"
+
 namespace UEStaticAssertCompleteType_Private
 {
 	class FIncompleteType;
 
-	template <typename T> struct TUEStaticAssertTypeChecker       { static T& Func(); };
+#if defined(__clang__) || defined(__GNUC__)
+	// Under Clang and GCC, a zero-sized array cannot be used in a template specialization, function overload, or tested for by std::is_array.
+	// Testing for sizeof(T) == 0 works though, and should be the only type this is true for, then __remove_all_extents (which also works)
+	// can be used to check the element type of the array.
+	template <typename T, SIZE_T SizeOfT = sizeof(T)>
+	struct TUEStaticAssertTypeCheckerBase
+	{
+		static T& Func();
+	};
+	template <typename T>
+	struct TUEStaticAssertTypeCheckerBase<T, 0>
+	{
+		// Use the __remove_all_extents builtin, which works, not std::remove_all_extents, which doesn't always.
+		static __remove_all_extents(T)& Func();
+	};
+#else
+	// The above trick doesn't work for zero-sized arrays under MSVC, but specialization does.
+	template <typename T>
+	struct TUEStaticAssertTypeCheckerBase
+	{
+		static T& Func();
+	};
+	template <typename T>
+	struct TUEStaticAssertTypeCheckerBase<T[0]>
+	{
+		static T& Func();
+	};
+#endif
+
+	template <typename T>
+	struct TUEStaticAssertTypeChecker : TUEStaticAssertTypeCheckerBase<T>
+	{
+	};
 
 	// Voids will give an unhelpful error message when trying to take a reference to them, so work around it
 	template <> struct TUEStaticAssertTypeChecker<               void> { static FIncompleteType Func(); };
@@ -30,19 +64,21 @@ namespace UEStaticAssertCompleteType_Private
 
 // Each of these should fail to compile
 #if 0
-	UE_STATIC_ASSERT_COMPLETE_TYPE(               void,                                     "CV void is incomplete");
-	UE_STATIC_ASSERT_COMPLETE_TYPE(      volatile void,                                     "CV void is incomplete");
-	UE_STATIC_ASSERT_COMPLETE_TYPE(const          void,                                     "CV void is incomplete");
-	UE_STATIC_ASSERT_COMPLETE_TYPE(const volatile void,                                     "CV void is incomplete");
+	UE_STATIC_ASSERT_COMPLETE_TYPE(               void,                                    "CV void is incomplete");
+	UE_STATIC_ASSERT_COMPLETE_TYPE(      volatile void,                                    "CV void is incomplete");
+	UE_STATIC_ASSERT_COMPLETE_TYPE(const          void,                                    "CV void is incomplete");
+	UE_STATIC_ASSERT_COMPLETE_TYPE(const volatile void,                                    "CV void is incomplete");
 	UE_STATIC_ASSERT_COMPLETE_TYPE(UEStaticAssertCompleteType_Private::FIncompleteType,    "A forward-declared-but-undefined class is incomplete");
 	UE_STATIC_ASSERT_COMPLETE_TYPE(UEStaticAssertCompleteType_Private::FIncompleteType[2], "An array of an incomplete class is incomplete");
-	UE_STATIC_ASSERT_COMPLETE_TYPE(int[],                                                   "An array of a complete type of unspecified bound is incomplete");
+	UE_STATIC_ASSERT_COMPLETE_TYPE(int[],                                                  "An array of a complete type of unspecified bound is incomplete");
+	UE_STATIC_ASSERT_COMPLETE_TYPE(UEStaticAssertCompleteType_Private::FIncompleteType[0], "A zero-sized array of an incomplete type is incomplete");
 #endif
 
 // Each of these should pass
 #if 0
-	UE_STATIC_ASSERT_COMPLETE_TYPE(UEStaticAssertCompleteType_Private::FIncompleteType*,                                                         "A pointer to an incomplete type is complete");
-	UE_STATIC_ASSERT_COMPLETE_TYPE(UEStaticAssertCompleteType_Private::FIncompleteType&,                                                         "A reference to an incomplete type is complete");
+	UE_STATIC_ASSERT_COMPLETE_TYPE(UEStaticAssertCompleteType_Private::FIncompleteType*,                                                        "A pointer to an incomplete type is complete");
+	UE_STATIC_ASSERT_COMPLETE_TYPE(UEStaticAssertCompleteType_Private::FIncompleteType&,                                                        "A reference to an incomplete type is complete");
 	UE_STATIC_ASSERT_COMPLETE_TYPE(UEStaticAssertCompleteType_Private::FIncompleteType   (UEStaticAssertCompleteType_Private::FIncompleteType), "A function type is not incomplete, even if it returns or takes an incomplete type");
 	UE_STATIC_ASSERT_COMPLETE_TYPE(UEStaticAssertCompleteType_Private::FIncompleteType(&)(UEStaticAssertCompleteType_Private::FIncompleteType), "References to function types must give a good error");
+	UE_STATIC_ASSERT_COMPLETE_TYPE(int[0],                                                                                                      "A zero-sized array of a complete type is complete");
 #endif
