Logging: Moved storage of FTextFormat for localized log templates into a global table

diff --git a/Engine/Source/Runtime/Core/Private/Logging/StructuredLog.cpp b/Engine/Source/Runtime/Core/Private/Logging/StructuredLog.cpp
--- a/Engine/Source/Runtime/Core/Private/Logging/StructuredLog.cpp
+++ b/Engine/Source/Runtime/Core/Private/Logging/StructuredLog.cpp
@@ -6,8 +6,11 @@
 #include "Algo/AllOf.h"
 #include "Algo/Find.h"
 #include "Algo/NoneOf.h"
+#include "Async/Mutex.h"
+#include "Async/UniqueLock.h"
 #include "AutoRTFM.h"
 #include "Containers/AnsiString.h"
+#include "Containers/SparseArray.h"
 #include "Containers/StringConv.h"
 #include "Containers/StringView.h"
 #include "CoreGlobals.h"
@@ -67,15 +70,21 @@ struct FLogTemplateOp
 	EOpCode Code = OpEnd;
 	int32 Value = 0;
 
-	inline int32 GetSkipSize() const { return Code == OpIndex ? 0 : Value; }
+	inline int32 GetSkipSize() const
+	{
+		return Code == OpIndex || Code == OpLocalized ? 0 : Value;
+	}
 
 	static inline FLogTemplateOp Load(const uint8*& Data);
 	static inline uint32 SaveSize(const FLogTemplateOp& Op) { return MeasureVarUInt(Encode(Op)); }
 	static inline void Save(const FLogTemplateOp& Op, uint8*& Data);
-	static inline uint64 Encode(const FLogTemplateOp& Op) { return uint64(Op.Code) | (uint64(Op.Value) << ValueShift); }
-	static inline FLogTemplateOp Decode(uint64 Value) { return {EOpCode(Value & ((1 << ValueShift) - 1)), int32(Value >> ValueShift)}; }
+	static constexpr uint64 Encode(const FLogTemplateOp& Op) { return uint64(Op.Code) | (uint64(Op.Value) << ValueShift); }
+	static constexpr FLogTemplateOp Decode(uint64 Value) { return {EOpCode(Value & ((1 << ValueShift) - 1)), int32(Value >> ValueShift)}; }
 };
 
+static_assert(FLogTemplateOp::Decode(FLogTemplateOp::Encode({.Value = 123})).Value == 123);
+static_assert(FLogTemplateOp::Decode(FLogTemplateOp::Encode({.Value = -123})).Value == -123);
+
 inline FLogTemplateOp FLogTemplateOp::Load(const uint8*& Data)
 {
 	uint32 ByteCount = 0;
@@ -346,10 +355,45 @@ private:
 	int32 Index{0};
 };
 
-} // UE::Logging::Private
-
 ///////////////////////////////////////////////////////////////////////////////////////////////////////////////////////
 
+/**
+ * A table of localized log formats referenced by log templates.
+ */
+class FLocalizedLogFormatTable
+{
+public:
+	int32 Add(FTextFormat&& Format)
+	{
+		TUniqueLock Lock(Mutex);
+		return Table.Emplace(MoveTemp(Format));
+	}
+
+	void RemoveAt(int32 Index)
+	{
+		TUniqueLock Lock(Mutex);
+		Table.RemoveAt(Index);
+	}
+
+	FTextFormat Get(int32 Index) const
+	{
+		TUniqueLock Lock(Mutex);
+		return Table[Index];
+	}
+
+private:
+	TSparseArray<FTextFormat> Table;
+	mutable FMutex Mutex;
+};
+
+static FLocalizedLogFormatTable& GetLocalizedLogFormatTable()
+{
+	static FLocalizedLogFormatTable Table;
+	return Table;
+}
+
+} // UE::Logging::Private
+
 namespace UE
 {
 
@@ -395,16 +439,15 @@ public:
 
 private:
 	template <typename CharType>
-	void FormatLocalizedTo(TStringBuilderBase<CharType>& Out, const FCbFieldViewIterator& Fields, const uint8* NextOp) const;
+	void FormatLocalizedTo(TStringBuilderBase<CharType>& Out, const FCbFieldViewIterator& Fields) const;
 
-	FText FormatLocalizedToText(const FCbFieldViewIterator& Fields, const uint8* NextOp) const;
+	FText FormatLocalizedToText(const FCbFieldViewIterator& Fields) const;
 
 	inline constexpr explicit FLogTemplate(const TCHAR* Format)
 		: StaticFormat(Format)
 	{
 	}
 
-	~FLogTemplate() = default;
 	FLogTemplate(const FLogTemplate&) = delete;
 	FLogTemplate& operator=(const FLogTemplate&) = delete;
 
@@ -548,7 +591,7 @@ FLogTemplate* FLogTemplate::CreateLocalized(
 		TEXT("Validation of field names is not compatible with sub-object references. [[%s]]"), Format);
 
 	TArray<FLogTemplateOp, TInlineAllocator<16>> Ops;
-	Ops.Add({FLogTemplateOp::OpLocalized});
+	Ops.Add({FLogTemplateOp::OpLocalized, GetLocalizedLogFormatTable().Add(FTextFormat(FormatText))});
 
 	// Track unique field names to avoid adding multiple ops for the same name.
 	TAnsiStringBuilder<256> FieldPathData;
@@ -641,9 +684,8 @@ FLogTemplate* FLogTemplate::CreateLocalized(
 		TextStart = FieldEnd;
 	}
 
-	const uint32 TotalSize = sizeof(FTextFormat) + sizeof(FLogTemplate) + Algo::TransformAccumulate(Ops, FLogTemplateOp::SaveSize, 0);
-	FTextFormat* const TemplateTextFormat = new(FMemory::Malloc(TotalSize, uint32(FPlatformMath::Max(alignof(FTextFormat), alignof(FLogTemplate))))) FTextFormat(FormatText);
-	FLogTemplate* const Template = new(TemplateTextFormat + 1) FLogTemplate(Format); //-V752
+	const uint32 TotalSize = sizeof(FLogTemplate) + Algo::TransformAccumulate(Ops, FLogTemplateOp::SaveSize, 0);
+	FLogTemplate* const Template = new(FMemory::Malloc(TotalSize, alignof(FLogTemplate))) FLogTemplate(Format);
 	uint8* Data = Template->GetOpData();
 	for (const FLogTemplateOp& Op : Ops)
 	{
@@ -657,16 +699,13 @@ void FLogTemplate::Destroy(FLogTemplate* Template)
 	using namespace Logging::Private;
 
 	const uint8* NextOp = Template->GetOpData();
-	if (FLogTemplateOp::Load(NextOp).Code == FLogTemplateOp::OpLocalized)
-	{
-		FTextFormat* TextFormat = (FTextFormat*)Template - 1;
-		TextFormat->~FTextFormat();
-		FMemory::Free(TextFormat);
-	}
-	else
+	if (FLogTemplateOp Op = FLogTemplateOp::Load(NextOp); Op.Code == FLogTemplateOp::OpLocalized)
 	{
-		FMemory::Free(Template);
+		GetLocalizedLogFormatTable().RemoveAt(Op.Value);
 	}
+
+	static_assert(std::is_trivially_destructible_v<FLogTemplate>);
+	FMemory::Free(Template);
 }
 
 template <typename CharType>
@@ -684,7 +723,7 @@ void FLogTemplate::FormatTo(TStringBuilderBase<CharType>& Out, const FCbFieldVie
 		switch (Op.Code)
 		{
 		case FLogTemplateOp::OpLocalized:
-			return FormatLocalizedTo(Out, Fields, NextOp);
+			return FormatLocalizedTo(Out, Fields);
 		case FLogTemplateOp::OpEnd:
 			return;
 		case FLogTemplateOp::OpText:
@@ -719,7 +758,7 @@ FText FLogTemplate::FormatToText(const FCbFieldViewIterator& Fields) const
 	const uint8* NextOp = GetOpData();
 	if (FLogTemplateOp::Load(NextOp).Code == FLogTemplateOp::OpLocalized)
 	{
-		return FormatLocalizedToText(Fields, NextOp);
+		return FormatLocalizedToText(Fields);
 	}
 	else
 	{
@@ -730,18 +769,20 @@ FText FLogTemplate::FormatToText(const FCbFieldViewIterator& Fields) const
 }
 
 template <typename CharType>
-FORCENOINLINE void FLogTemplate::FormatLocalizedTo(TStringBuilderBase<CharType>& Out, const FCbFieldViewIterator& Fields, const uint8* NextOp) const
+FORCENOINLINE void FLogTemplate::FormatLocalizedTo(TStringBuilderBase<CharType>& Out, const FCbFieldViewIterator& Fields) const
 {
-	Out.Append(FormatLocalizedToText(Fields, NextOp).ToString());
+	Out.Append(FormatLocalizedToText(Fields).ToString());
 }
 
-FText FLogTemplate::FormatLocalizedToText(const FCbFieldViewIterator& Fields, const uint8* NextOp) const
+FText FLogTemplate::FormatLocalizedToText(const FCbFieldViewIterator& Fields) const
 {
 	using namespace Logging::Private;
 
+	TOptional<FTextFormat> TextFormat;
 	FFormatNamedArguments TextFormatArguments;
 
 	int32 FieldIndexHint = -1;
+	const uint8* NextOp = GetOpData();
 	const TCHAR* NextFormat = StaticFormat;
 	FFieldFinder FieldFinder(StaticFormat, Fields);
 	while (NextOp)
@@ -749,6 +790,9 @@ FText FLogTemplate::FormatLocalizedToText(const FCbFieldViewIterator& Fields, co
 		const FLogTemplateOp Op = FLogTemplateOp::Load(NextOp);
 		switch (Op.Code)
 		{
+		case FLogTemplateOp::OpLocalized:
+			TextFormat = GetLocalizedLogFormatTable().Get(Op.Value);
+			break;
 		case FLogTemplateOp::OpEnd:
 			NextOp = nullptr;
 			break;
@@ -773,7 +817,7 @@ FText FLogTemplate::FormatLocalizedToText(const FCbFieldViewIterator& Fields, co
 		NextFormat += Op.GetSkipSize();
 	}
 
-	const FTextFormat* TextFormat = (const FTextFormat*)this - 1;
+	checkf(TextFormat, TEXT("Missing text format when formatting localized template. [[%s]]"), StaticFormat);
 	return FText::Format(*TextFormat, TextFormatArguments);
 }
 
