[StripedMap]
- Fix GetKeys that could return duplicates

diff --git a/Engine/Source/Runtime/Core/Public/Containers/StripedMap.h b/Engine/Source/Runtime/Core/Public/Containers/StripedMap.h
--- a/Engine/Source/Runtime/Core/Public/Containers/StripedMap.h
+++ b/Engine/Source/Runtime/Core/Public/Containers/StripedMap.h
@@ -749,10 +749,15 @@ public:
 	int32 GetKeys(TArray<KeyType, Allocator>& OutKeys) const
 	{
 		OutKeys.Reset();
+		TSet<KeyType> VisitedKeys;
 		ForEach(
-			[&OutKeys](const auto& Pair)
+			[&OutKeys, &VisitedKeys](const auto& Pair)
 			{
-				OutKeys.Add(Pair.Key);
+				if (!VisitedKeys.Contains(Pair.Key))
+				{
+					OutKeys.Add(Pair.Key);
+					VisitedKeys.Add(Pair.Key);
+				}
 			}
 		);
 		return OutKeys.Num();
