Define TIsBitwiseConstructible for ANSICHAR and UTF8CHAR to byte conversions.

diff --git a/Engine/Source/Runtime/Core/Public/Templates/UnrealTypeTraits.h b/Engine/Source/Runtime/Core/Public/Templates/UnrealTypeTraits.h
--- a/Engine/Source/Runtime/Core/Public/Templates/UnrealTypeTraits.h
+++ b/Engine/Source/Runtime/Core/Public/Templates/UnrealTypeTraits.h
@@ -456,6 +456,17 @@ template <> struct TIsBitwiseConstructible< int32, uint32> { enum { Value = true
 template <> struct TIsBitwiseConstructible<uint64,  int64> { enum { Value = true }; };
 template <> struct TIsBitwiseConstructible< int64, uint64> { enum { Value = true }; };
 
+// ANSICHAR can be bitwise converted to a byte or UTF8CHAR
+// conversions in the other direction (to ANSICHAR) are not generally safe
+template <> struct TIsBitwiseConstructible< int8,    ANSICHAR> { enum { Value = true }; };
+template <> struct TIsBitwiseConstructible<uint8,    ANSICHAR> { enum { Value = true }; };
+template <> struct TIsBitwiseConstructible<UTF8CHAR, ANSICHAR> { enum { Value = true }; };
+
+// UTF8CHAR can be bitwise converted to a byte
+// conversions in the other direction (to UTF8CHAR) are not generally safe
+template <> struct TIsBitwiseConstructible< int8, UTF8CHAR> { enum { Value = true }; };
+template <> struct TIsBitwiseConstructible<uint8, UTF8CHAR> { enum { Value = true }; };
+
 #define GENERATE_MEMBER_FUNCTION_CHECK(MemberName, Result, ConstModifier, ...)									\
 template <typename T>																							\
 class THasMemberFunction_##MemberName																			\
