Logging: Added TLogTemplate and control over log template allocation

diff --git a/Engine/Source/Runtime/Core/Private/Logging/StructuredLog.cpp b/Engine/Source/Runtime/Core/Private/Logging/StructuredLog.cpp
--- a/Engine/Source/Runtime/Core/Private/Logging/StructuredLog.cpp
+++ b/Engine/Source/Runtime/Core/Private/Logging/StructuredLog.cpp
@@ -143,9 +143,9 @@ static void LogFieldValue(TStringBuilderBase<CharType>& Out, const FCbFieldView&
 			if (FCbFieldView FormatField = Object.FindViewIgnoreCase(ANSITEXTVIEW("$format")); FormatField.IsString())
 			{
 				TStringBuilder<128> Format(InPlace, FormatField.AsString());
-				FLogTemplate* Template = CreateLogTemplate(*Format, {.bAllowSubObjectReferences = true});
+				FLogTemplate* Template = UE::CreateLogTemplate(*Format, {.bAllowSubObjectReferences = true});
 				FormatLogTo(Out, Template, Object.CreateViewIterator());
-				DestroyLogTemplate(Template);
+				UE::DestroyLogTemplate(Template);
 				break;
 			}
 
@@ -158,9 +158,9 @@ static void LogFieldValue(TStringBuilderBase<CharType>& Out, const FCbFieldView&
 				TStringBuilder<32> Namespace(InPlace, LocNamespaceField.AsString());
 				TStringBuilder<32> Key(InPlace, LocKeyField.AsString());
 				TStringBuilder<128> Format(InPlace, LocFormatField.AsString());
-				FLogTemplate* Template = CreateLogTemplate(*Namespace, *Key, *Format, {.bAllowSubObjectReferences = true});
+				FLogTemplate* Template = UE::CreateLogTemplate(*Namespace, *Key, *Format, {.bAllowSubObjectReferences = true});
 				FormatLogTo(Out, Template, Object.CreateViewIterator());
-				DestroyLogTemplate(Template);
+				UE::DestroyLogTemplate(Template);
 				break;
 			}
 
@@ -394,6 +394,8 @@ static FLocalizedLogFormatTable& GetLocalizedLogFormatTable()
 
 } // UE::Logging::Private
 
+///////////////////////////////////////////////////////////////////////////////////////////////////////////////////////
+
 namespace UE
 {
 
@@ -402,25 +404,20 @@ class FLogTemplate
 	using FLogField = Logging::Private::FLogField;
 
 public:
-	static FLogTemplate* CreateStatic(const Logging::Private::FStaticLogRecord& Log, const FLogField* Fields, int32 FieldCount)
-	{
-		return Create(Log.Format, {}, Fields, FieldCount);
-	}
-
-	static FLogTemplate* CreateStatic(const Logging::Private::FStaticLocalizedLogRecord& Log, const FLogField* Fields, int32 FieldCount)
-	{
-		const FText FormatText = FText::AsLocalizable_Advanced(Log.TextNamespace, Log.TextKey, Log.Format);
-		return CreateLocalized(FormatText, Log.Format, {}, Fields, FieldCount);
-	}
-
-	static FLogTemplate* Create(const TCHAR* Format, const FLogTemplateOptions& Options, const FLogField* Fields = nullptr, int32 FieldCount = 0);
+	static FLogTemplate* Create(
+		const TCHAR* Format,
+		const FLogTemplateOptions& Options,
+		const FLogField* Fields,
+		int32 FieldCount,
+		TFunctionRef<void* (int32)> Allocate);
 
 	static FLogTemplate* CreateLocalized(
 		const FText& FormatText,
 		const TCHAR* Format,
 		const FLogTemplateOptions& Options,
-		const FLogField* Fields = nullptr,
-		int32 FieldCount = 0);
+		const FLogField* Fields,
+		int32 FieldCount,
+		TFunctionRef<void* (int32)> Allocate);
 
 	static void Destroy(FLogTemplate* Template);
 
@@ -455,7 +452,14 @@ private:
 	FLogTemplate* Next = nullptr;
 };
 
-FLogTemplate* FLogTemplate::Create(const TCHAR* Format, const FLogTemplateOptions& Options, const FLogField* Fields, const int32 FieldCount)
+static_assert(std::is_trivially_destructible_v<FLogTemplate>);
+
+FLogTemplate* FLogTemplate::Create(
+	const TCHAR* Format,
+	const FLogTemplateOptions& Options,
+	const FLogField* Fields,
+	const int32 FieldCount,
+	TFunctionRef<void* (int32)> Allocate)
 {
 	using namespace Logging::Private;
 
@@ -564,7 +568,7 @@ FLogTemplate* FLogTemplate::Create(const TCHAR* Format, const FLogTemplateOption
 		TEXT("Log format requires %d fields and %d were provided. [[%s]]"), FormatFieldCount, FieldCount, Format);
 
 	const uint32 TotalSize = sizeof(FLogTemplate) + Algo::TransformAccumulate(Ops, FLogTemplateOp::SaveSize, 0);
-	FLogTemplate* const Template = new(FMemory::Malloc(TotalSize, alignof(FLogTemplate))) FLogTemplate(Format);
+	FLogTemplate* const Template = new(Allocate(TotalSize)) FLogTemplate(Format);
 	uint8* Data = Template->GetOpData();
 	for (const FLogTemplateOp& Op : Ops)
 	{
@@ -578,7 +582,8 @@ FLogTemplate* FLogTemplate::CreateLocalized(
 	const TCHAR* Format,
 	const FLogTemplateOptions& Options,
 	const FLogField* Fields,
-	const int32 FieldCount)
+	const int32 FieldCount,
+	TFunctionRef<void* (int32)> Allocate)
 {
 	// A localized format string consists of an OpLocalized op followed by a sequence of OpSkip and OpName/OpPath ops
 	// that are terminated by an OpEnd op. Only the first occurrence of each name/path is included and everything else
@@ -685,7 +690,7 @@ FLogTemplate* FLogTemplate::CreateLocalized(
 	}
 
 	const uint32 TotalSize = sizeof(FLogTemplate) + Algo::TransformAccumulate(Ops, FLogTemplateOp::SaveSize, 0);
-	FLogTemplate* const Template = new(FMemory::Malloc(TotalSize, alignof(FLogTemplate))) FLogTemplate(Format);
+	FLogTemplate* const Template = new(Allocate(TotalSize)) FLogTemplate(Format);
 	uint8* Data = Template->GetOpData();
 	for (const FLogTemplateOp& Op : Ops)
 	{
@@ -703,9 +708,6 @@ void FLogTemplate::Destroy(FLogTemplate* Template)
 	{
 		GetLocalizedLogFormatTable().RemoveAt(Op.Value);
 	}
-
-	static_assert(std::is_trivially_destructible_v<FLogTemplate>);
-	FMemory::Free(Template);
 }
 
 template <typename CharType>
@@ -823,25 +825,28 @@ FText FLogTemplate::FormatLocalizedToText(const FCbFieldViewIterator& Fields) co
 
 ///////////////////////////////////////////////////////////////////////////////////////////////////////////////////////
 
-FLogTemplate* CreateLogTemplate(const TCHAR* Format, const FLogTemplateOptions& Options)
+void Logging::Private::CreateLogTemplate(const TCHAR* Format, const FLogTemplateOptions& Options, const FLogField* Fields, const int32 FieldCount, TFunctionRef<void* (int32)> Allocate)
 {
-	return FLogTemplate::Create(Format, Options);
+	FLogTemplate::Create(Format, Options, Fields, FieldCount, Allocate);
 }
 
-FLogTemplate* CreateLogTemplate(const FText& Format, const FLogTemplateOptions& Options)
+void Logging::Private::CreateLocalizedLogTemplate(const FText& Format, const FLogTemplateOptions& Options, const FLogField* Fields, const int32 FieldCount, TFunctionRef<void* (int32)> Allocate)
 {
-	return FLogTemplate::CreateLocalized(Format, *Format.ToString(), Options);
+	FLogTemplate::CreateLocalized(Format, *Format.ToString(), Options, Fields, FieldCount, Allocate);
 }
 
-FLogTemplate* CreateLogTemplate(const TCHAR* TextNamespace, const TCHAR* TextKey, const TCHAR* Format, const FLogTemplateOptions& Options)
+void Logging::Private::CreateLocalizedLogTemplate(const TCHAR* TextNamespace, const TCHAR* TextKey, const TCHAR* Format, const FLogTemplateOptions& Options, const FLogField* Fields, const int32 FieldCount, TFunctionRef<void* (int32)> Allocate)
 {
 	const FText FormatText = FText::AsLocalizable_Advanced(TextNamespace, TextKey, Format);
-	return FLogTemplate::CreateLocalized(FormatText, Format, Options);
+	FLogTemplate::CreateLocalized(FormatText, Format, Options, Fields, FieldCount, Allocate);
 }
 
-void DestroyLogTemplate(FLogTemplate* Template)
+void Logging::Private::DestroyLogTemplate(FLogTemplate* Template)
 {
-	FLogTemplate::Destroy(Template);
+	if (Template)
+	{
+		FLogTemplate::Destroy(Template);
+	}
 }
 
 void FormatLogTo(FUtf8StringBuilderBase& Out, const FLogTemplate* Template, const FCbFieldViewIterator& Fields)
@@ -1258,9 +1263,24 @@ FORCENOINLINE static FLogTemplate& CreateLogTemplate(const FLogCategoryBase& Cat
 			return *Template;
 		}
 
-		FLogTemplate* NewTemplate = FLogTemplate::CreateStatic(Log, Fields, FieldCount);
+		struct FTemplateCreator
+		{
+			static TLogTemplate<FMemoryLogTemplateStorage> Create(const FStaticLogRecord& Log, const FLogField* Fields, int32 FieldCount)
+			{
+				return TLogTemplate<FMemoryLogTemplateStorage>(Log.Format, {}, Fields, FieldCount);
+			}
+
+			static TLogTemplate<FMemoryLogTemplateStorage> Create(const FStaticLocalizedLogRecord& Log, const FLogField* Fields, int32 FieldCount)
+			{
+				return TLogTemplate<FMemoryLogTemplateStorage>(Log.TextNamespace, Log.TextKey, Log.Format, {}, Fields, FieldCount);
+			}
+		};
+
+		TLogTemplate<FMemoryLogTemplateStorage> LocalTemplate = FTemplateCreator::Create(Log, Fields, FieldCount);
+		FLogTemplate* NewTemplate = LocalTemplate.Get();
 		if (LIKELY(Log.DynamicData.Template.compare_exchange_strong(Template, NewTemplate, std::memory_order_release, std::memory_order_acquire)))
 		{
+			LocalTemplate.Detach();
 			// Register the template to destroy on exit.
 			for (FLogTemplate* Head = GStaticLogTemplateManager.Head.load(std::memory_order_relaxed);;)
 			{
@@ -1272,7 +1292,6 @@ FORCENOINLINE static FLogTemplate& CreateLogTemplate(const FLogCategoryBase& Cat
 			}
 			return *NewTemplate;
 		}
-		FLogTemplate::Destroy(NewTemplate);
 	}
 }
 
@@ -1445,16 +1464,16 @@ public:
 			return *LocalTemplate;
 		}
 		return CreateTemplate();
-	}
+}
 
 private:
 	FORCENOINLINE FLogTemplate& CreateTemplate()
 	{
-		FLogTemplate* NewTemplate = FLogTemplate::Create(GStaticBasicLogFormat, {});
+		FLogTemplate* NewTemplate = UE::CreateLogTemplate(GStaticBasicLogFormat);
 		if (FLogTemplate* ExistingTemplate = nullptr;
 			UNLIKELY(!Template.compare_exchange_strong(ExistingTemplate, NewTemplate, std::memory_order_release, std::memory_order_acquire)))
 		{
-			FLogTemplate::Destroy(NewTemplate);
+			UE::DestroyLogTemplate(NewTemplate);
 			return *ExistingTemplate;
 		}
 		return *NewTemplate;
diff --git a/Engine/Source/Runtime/Core/Public/Logging/StructuredLogFormat.h b/Engine/Source/Runtime/Core/Public/Logging/StructuredLogFormat.h
--- a/Engine/Source/Runtime/Core/Public/Logging/StructuredLogFormat.h
+++ b/Engine/Source/Runtime/Core/Public/Logging/StructuredLogFormat.h
@@ -3,14 +3,63 @@
 #pragma once
 
 #include "Containers/StringFwd.h"
+#include "Internationalization/Text.h"
+#include "Misc/ScopeExit.h"
 
 #define UE_API CORE_API
 
 class FCbFieldViewIterator;
 class FCbWriter;
-class FText;
+template <typename FuncType> class TFunctionRef;
 
 namespace UE { class FLogTemplate; }
+namespace UE { struct FLogTemplateOptions; }
+namespace UE::Logging::Private { struct FLogField; }
+
+namespace UE::Logging::Private
+{
+
+/** Store a log template in an allocation from FMemory. */
+class FMemoryLogTemplateStorage
+{
+public:
+	FMemoryLogTemplateStorage() = default;
+	FMemoryLogTemplateStorage(const FMemoryLogTemplateStorage&) = delete;
+	FMemoryLogTemplateStorage& operator=(const FMemoryLogTemplateStorage&) = delete;
+
+	~FMemoryLogTemplateStorage()
+	{
+		FMemory::Free(Data);
+	}
+
+	inline void* Allocate(int32 Size)
+	{
+		Data = FMemory::Malloc(Size);
+		return Data;
+	}
+
+	inline const void* Get() const
+	{
+		return Data;
+	}
+
+	inline void* Detach()
+	{
+		ON_SCOPE_EXIT { Data = nullptr; };
+		return Data;
+	}
+
+private:
+	void* Data = nullptr;
+};
+
+UE_API void CreateLogTemplate(const TCHAR* Format, const FLogTemplateOptions& Options, const FLogField* Fields, int32 FieldCount, TFunctionRef<void* (int32)> Allocate);
+UE_API void CreateLocalizedLogTemplate(const FText& Format, const FLogTemplateOptions& Options, const FLogField* Fields, int32 FieldCount, TFunctionRef<void* (int32)> Allocate);
+UE_API void CreateLocalizedLogTemplate(const TCHAR* TextNamespace, const TCHAR* TextKey, const TCHAR* Format, const FLogTemplateOptions& Options, const FLogField* Fields, int32 FieldCount, TFunctionRef<void* (int32)> Allocate);
+
+UE_API void DestroyLogTemplate(FLogTemplate* Template);
+
+} // UE::Logging::Private
 
 namespace UE
 {
@@ -22,15 +71,92 @@ struct FLogTemplateOptions
 	bool bAllowSubObjectReferences = false;
 };
 
-UE_API FLogTemplate* CreateLogTemplate(const TCHAR* Format, const FLogTemplateOptions& Options = {});
-UE_API FLogTemplate* CreateLogTemplate(const FText& Format, const FLogTemplateOptions& Options = {});
-UE_API FLogTemplate* CreateLogTemplate(const TCHAR* TextNamespace, const TCHAR* TextKey, const TCHAR* Format, const FLogTemplateOptions& Options = {});
-UE_API void DestroyLogTemplate(FLogTemplate* Template);
-
 UE_API void FormatLogTo(FUtf8StringBuilderBase& Out, const FLogTemplate* Template, const FCbFieldViewIterator& Fields);
 UE_API void FormatLogTo(FWideStringBuilderBase& Out, const FLogTemplate* Template, const FCbFieldViewIterator& Fields);
 UE_API FText FormatLogToText(const FLogTemplate* Template, const FCbFieldViewIterator& Fields);
 
+/** A log template that is templated on its storage. */
+template <typename StorageType>
+class TLogTemplate
+{
+public:
+	inline explicit TLogTemplate(const TCHAR* Format, const FLogTemplateOptions& Options = {}, const Logging::Private::FLogField* Fields = nullptr, int32 FieldCount = 0)
+	{
+		Logging::Private::CreateLogTemplate(Format, Options, Fields, FieldCount, [this](int32 Size) { return Storage.Allocate(Size); });
+	}
+
+	inline explicit TLogTemplate(const FText& Format, const FLogTemplateOptions& Options = {}, const Logging::Private::FLogField* Fields = nullptr, int32 FieldCount = 0)
+	{
+		Logging::Private::CreateLocalizedLogTemplate(Format, Options, Fields, FieldCount, [this](int32 Size) { return Storage.Allocate(Size); });
+	}
+
+	inline explicit TLogTemplate(const TCHAR* TextNamespace, const TCHAR* TextKey, const TCHAR* Format, const FLogTemplateOptions& Options = {}, const Logging::Private::FLogField* Fields = nullptr, int32 FieldCount = 0)
+	{
+		Logging::Private::CreateLocalizedLogTemplate(TextNamespace, TextKey, Format, Options, Fields, FieldCount, [this](int32 Size) { return Storage.Allocate(Size); });
+	}
+
+	inline void FormatTo(FUtf8StringBuilderBase& Out, const FCbFieldViewIterator& Fields)
+	{
+		FormatLogTo(Out, Get(), Fields);
+	}
+
+	inline void FormatTo(FWideStringBuilderBase& Out, const FCbFieldViewIterator& Fields)
+	{
+		FormatLogTo(Out, Get(), Fields);
+	}
+
+	inline FText FormatToText(const FCbFieldViewIterator& Fields)
+	{
+		FormatLogToText(Get(), Fields);
+	}
+
+	inline ~TLogTemplate()
+	{
+		Logging::Private::DestroyLogTemplate(Get());
+	}
+
+	inline FLogTemplate* Get() const
+	{
+		return (FLogTemplate*)Storage.Get();
+	}
+
+	inline FLogTemplate* Detach()
+		requires requires (StorageType& S) { S.Detach(); }
+	{
+		return (FLogTemplate*)Storage.Detach();
+	}
+
+private:
+	TLogTemplate() = default;
+
+	StorageType Storage;
+};
+
+/** Create a log template on the heap. */
+inline FLogTemplate* CreateLogTemplate(const TCHAR* Format, const FLogTemplateOptions& Options = {})
+{
+	return TLogTemplate<Logging::Private::FMemoryLogTemplateStorage>(Format, Options).Detach();
+}
+
+/** Create a localized log template on the heap. */
+inline FLogTemplate* CreateLogTemplate(const FText& Format, const FLogTemplateOptions& Options = {})
+{
+	return TLogTemplate<Logging::Private::FMemoryLogTemplateStorage>(Format, Options).Detach();
+}
+
+/** Create a localized log template on the heap. */
+inline FLogTemplate* CreateLogTemplate(const TCHAR* TextNamespace, const TCHAR* TextKey, const TCHAR* Format, const FLogTemplateOptions& Options = {})
+{
+	return TLogTemplate<Logging::Private::FMemoryLogTemplateStorage>(TextNamespace, TextKey, Format, Options).Detach();
+}
+
+/** Destroy a log template allocated by CreateLogTemplate. */
+inline void DestroyLogTemplate(FLogTemplate* Template)
+{
+	Logging::Private::DestroyLogTemplate(Template);
+	FMemory::Free(Template);
+}
+
 /**
  * Serializes a localized log format template to compact binary. Call from SerializeForLog.
  *
