Simplified UE_BREAK_AND_RETURN_FALSE() for Clang/GCC, and made an alternative for MSVC based on the fact that its PLATFORM_BREAK() is an expression.

diff --git a/Engine/Source/Runtime/Core/Public/HAL/Platform.h b/Engine/Source/Runtime/Core/Public/HAL/Platform.h
--- a/Engine/Source/Runtime/Core/Public/HAL/Platform.h
+++ b/Engine/Source/Runtime/Core/Public/HAL/Platform.h
@@ -638,6 +638,9 @@
 	#define PLATFORM_FORCE_UE_LOCK_USAGE 0
 #endif
 
+#ifndef PLATFORM_BREAK_IS_EXPRESSION
+	#define PLATFORM_BREAK_IS_EXPRESSION 0
+#endif
 
 // deprecated, do not use
 #define PLATFORM_HAS_THREADSAFE_RHIGetRenderQueryResult	#
diff --git a/Engine/Source/Runtime/Core/Public/Misc/AssertionMacros.h b/Engine/Source/Runtime/Core/Public/Misc/AssertionMacros.h
--- a/Engine/Source/Runtime/Core/Public/Misc/AssertionMacros.h
+++ b/Engine/Source/Runtime/Core/Public/Misc/AssertionMacros.h
@@ -33,19 +33,6 @@
 #endif // DO_CHECK || DO_GUARD_SLOW
 #endif
 
-// Lite ensures use an alternate implementation that avoid lambdas which are expensive in the amount of
-// debug information generated. However, this comes at the cost of ensures breaking inside an implementation
-// function rather than on the line that the ensure is on. On Clang, this limitation can be mitigated by
-// using UE_NODEBUG, so that the break happens where the code is inlined rather than where it is defined,
-// but no such option exists for MSVC, so we disable lite ensures by default on that compiler.
-#ifndef UE_USE_LITE_ENSURES
-	#ifdef __clang__
-		#define UE_USE_LITE_ENSURES 1
-	#else
-		#define UE_USE_LITE_ENSURES 0
-	#endif
-#endif
-
 namespace ELogVerbosity
 {
 	enum Type : uint8;
@@ -422,16 +409,16 @@ RetType UE_COLD UE_DEBUG_SECTION DispatchCheckVerify (InnerType&& Inner, ArgType
 		return ((uint64)Result << 32) | (uint64)Line;
 	}
 
-#if UE_USE_LITE_ENSURES
-	UE_NODEBUG FORCEINLINE bool BreakAndReturnFalse()
-	{
-		// This is tagged with UE_NODEBUG and FORCEINLINE to ensure that the break happens at the usage of ensure(), not inside this function.
-		PLATFORM_BREAK();
-		return false;
-	}
-
-	#define UE_BREAK_AND_RETURN_FALSE() (BreakAndReturnFalse())
+// UE_BREAK_AND_RETURN_FALSE() - an expression which breaks into the debugger and returns false.
+#if PLATFORM_BREAK_IS_EXPRESSION
+	// If PLATFORM_BREAK is also an expression, it can be used directly.
+	#define UE_BREAK_AND_RETURN_FALSE() (PLATFORM_BREAK(), false)
+#elif defined(__clang__) || defined(__GNUC__)
+	// Clang and GCC support 'statement expressions' which allows the expansion of statements in an
+	// expression context.
+	#define UE_BREAK_AND_RETURN_FALSE() ({ PLATFORM_BREAK(); false; })
 #else
+	// Fallback to using a lambda if there is no other choice, which generates more debug information.
 	#define UE_BREAK_AND_RETURN_FALSE() [](){ PLATFORM_BREAK(); return false; }()
 #endif
 
diff --git a/Engine/Source/Runtime/Core/Public/Windows/WindowsPlatform.h b/Engine/Source/Runtime/Core/Public/Windows/WindowsPlatform.h
--- a/Engine/Source/Runtime/Core/Public/Windows/WindowsPlatform.h
+++ b/Engine/Source/Runtime/Core/Public/Windows/WindowsPlatform.h
@@ -136,6 +136,7 @@ typedef FWindowsPlatformTypes FPlatformTypes;
 // for example, which can appear as if the did not fire, leaving a programmer
 // unknowingly debugging an undefined process.
 #define PLATFORM_BREAK() (__nop(), __debugbreak())
+#define PLATFORM_BREAK_IS_EXPRESSION 1
 
 // Intrinsics for 128-bit atomics on Windows platform requires Windows 8 or higher (WINVER>=0x0602)
 // http://msdn.microsoft.com/en-us/library/windows/desktop/hh972640.aspx
