[Core] Add Compact Set and Hash Table
* Added natvis visualizers
* Add unit tests that run the same checks against both tset an tcompactset

#tests Added unit tests that include benchmarks against TSet. Have other cls that comletely replace TSet usage with it and had the game up and running in PIE

diff --git a/Engine/Source/Runtime/Core/Public/Containers/CompactHashTable.h b/Engine/Source/Runtime/Core/Public/Containers/CompactHashTable.h
new file mode 100644
--- /dev/null
+++ b/Engine/Source/Runtime/Core/Public/Containers/CompactHashTable.h
@@ -0,0 +1,386 @@
+// Copyright Epic Games, Inc. All Rights Reserved.
+
+#pragma once
+
+#include "Containers/ArrayView.h"
+
+/* Compact hash table has two distinct features to keep it small
+	1. The Index type adapts to the number of elements, so a smaller table will use a smaller type
+	2. Holes in the table are patched up so all slots up to count are guaranteed to be valid
+
+	For performance reasons we only reset the HashTable part of the lookup table since the indexes will be unused until added.
+	It's the users responsibility to maintain correct item count
+	!!Important!! It's also the users reponsibility to move items from the end of the list into their new spots when removing items.
+	
+	You are be able to maintain different views (hashes) of the same data since hole patching is deterministic between two tables.
+	This can be used to allow searching against the same data using different fields as names
+*/
+
+/* Helper to write code against specific types resolves by number of elements. Usage:
+	#define UE_COMPACTHASHTABLE_EXECUTEBYTYPE(Type) // Code
+	UE_COMPACTHASHTABLE_CALLBYTYPE(NumElements)
+	#undef UE_COMPACTHASHTABLE_EXECUTEBYTYPE
+*/
+#define UE_COMPACTHASHTABLE_CALLBYTYPE(NextIndexCount) \
+	switch (UE::Core::CompactHashTable::GetTypeSize(NextIndexCount)) { \
+	case 1: { UE_COMPACTHASHTABLE_EXECUTEBYTYPE(uint8); } break; \
+	case 2: { UE_COMPACTHASHTABLE_EXECUTEBYTYPE(uint16); } break; \
+	case 4: { UE_COMPACTHASHTABLE_EXECUTEBYTYPE(uint32); } break; \
+	default: checkNoEntry(); }
+
+namespace UE::Core::CompactHashTable::Private
+{
+	// Remove an element from the list and patch up any next index references, after this the spot is empty and needs to be filled
+	template<typename IndexType>
+	FORCEINLINE void RemoveInternal(const uint32 Index, const uint32 Key, IndexType* HashData, const uint32 HashCount, IndexType* NextIndexData, const uint32 NextIndexCount)
+	{
+		const uint32 HashIndex = Key & (HashCount - 1);
+
+		for (IndexType* Lookup = &HashData[HashIndex]; *Lookup < NextIndexCount; Lookup = &NextIndexData[*Lookup])
+		{
+			if (*Lookup == Index)
+			{
+				// Next = Next->Next
+				*Lookup = NextIndexData[Index];
+				return;
+			}
+		}
+	}
+}
+
+/* Helpers for interacting with compact hash table memory layouts */
+namespace UE::Core::CompactHashTable
+{
+	[[nodiscard]] FORCEINLINE constexpr uint32 GetTypeSize(uint32 IndexCount)
+	{
+		return 1 + int(IndexCount > 0xff) + int(IndexCount > 0xffff) * 2;
+	}
+
+	[[nodiscard]] FORCEINLINE constexpr uint32 GetTypeShift(uint32 IndexCount)
+	{
+		return 0 + int(IndexCount > 0xff) + int(IndexCount > 0xffff);
+	}
+
+	[[nodiscard]] FORCEINLINE constexpr size_t GetMemoryRequiredInBytes(uint32 IndexCount, uint32 HashCount)
+	{
+		return size_t(IndexCount + HashCount) << GetTypeShift(IndexCount);
+	}
+
+	[[nodiscard]] FORCEINLINE constexpr size_t GetMemoryAlignment()
+	{
+		return 4; // only support uint32 for now, pick highest alignment so we don't change alignment between allocations
+	}
+
+	/* Calculate the size of the hash table from the number of elements in the set */
+	[[nodiscard]] FORCEINLINE constexpr size_t GetHashCount(uint32 NumElements)
+	{
+		if (!NumElements)
+		{
+			return 0;
+		}
+		if (NumElements < 8)
+		{
+			return 4;
+		}
+
+		// Always use the power of 2 smaller than the current size to prioritize size over speed just a little bit
+		return FGenericPlatformMath::RoundUpToPowerOfTwo(NumElements / 2 + 1); // 255 == 128, 256 == 256
+	}
+
+	/* Return the first index of a key from the hash table portion */
+	template<typename IndexType>
+	[[nodiscard]] FORCEINLINE uint32 GetFirst(uint32 Key, const IndexType* HashData, const uint32 HashCount)
+	{
+		const uint32 HashIndex = Key & (HashCount - 1);
+		const IndexType FirstIndex = HashData[HashIndex];
+
+		// Convert smaller type INDEX_NONE to uint32 INDEX_NONE if data is invalid
+		constexpr IndexType InvalidIndex = (IndexType)INDEX_NONE;
+		return FirstIndex == InvalidIndex ? (uint32)INDEX_NONE : (uint32)FirstIndex;
+	}
+
+	/* Return the first index of a key from the hash table portion */
+	template<typename IndexType>
+	[[nodiscard]] FORCEINLINE uint32 GetFirstByIndex(uint32 HashIndex, const IndexType* HashData, const uint32 HashCount)
+	{
+		checkSlow(HashIndex <= HashCount);
+		const IndexType FirstIndex = HashData[HashIndex];
+
+		// Convert smaller type INDEX_NONE to uint32 INDEX_NONE if data is invalid
+		constexpr IndexType InvalidIndex = (IndexType)INDEX_NONE;
+		return FirstIndex == InvalidIndex ? (uint32)INDEX_NONE : (uint32)FirstIndex;
+	}
+
+	/* Given an existing index, return the next index in case there was a collision in the has table */
+	template<typename IndexType>
+	[[nodiscard]] FORCEINLINE uint32 GetNext(uint32 Index, const IndexType* NextIndexData, const uint32 NextIndexCount)
+	{
+		checkSlow(Index < NextIndexCount);
+		const IndexType NextIndex = NextIndexData[Index];
+		
+		// Convert smaller type INDEX_NONE to uint32 INDEX_NONE if data is invalid
+		constexpr IndexType InvalidIndex = (IndexType)INDEX_NONE;
+		return NextIndex == InvalidIndex ? (uint32)INDEX_NONE : (uint32)NextIndex;
+	}
+
+	/* Do a full search for an existing element in the table given a function to compare if a found element is what you're looking for */
+	template<typename IndexType, typename PredicateType>
+	[[nodiscard]] FORCEINLINE uint32 Find(uint32 Key, const IndexType* HashData, const uint32 HashCount, const IndexType* NextIndexData, const uint32 NextIndexCount, const PredicateType& Predicate)
+	{
+		for (IndexType ElementIndex = HashData[Key & (HashCount - 1)]; ElementIndex != (IndexType)INDEX_NONE; ElementIndex = NextIndexData[ElementIndex])
+		{
+			if (Predicate(ElementIndex))
+			{
+				// Return the first match, regardless of whether the set has multiple matches for the key or not.
+				return ElementIndex;
+			}
+			checkSlow(ElementIndex < NextIndexCount);
+		}
+		return INDEX_NONE;
+	}
+	
+	/* Insert new element into the hash table */
+	template<typename IndexType>
+	FORCEINLINE void Add(uint32 Index, uint32 Key, IndexType* HashData, const uint32 HashCount, IndexType* NextIndexData, const uint32 NextIndexCount)
+	{
+		checkSlow(Index < NextIndexCount);
+
+		const uint32 HashIndex = Key & (HashCount - 1);
+		NextIndexData[Index] = HashData[HashIndex];
+		HashData[HashIndex] = (IndexType)Index;
+	}
+
+	/* Remove an element from the list, move the last element into the now empty slot
+	 If the item to remove is the last element then the last element's key will be ignored (you can skip calculating it if it's expensive)
+	*/
+	template<typename IndexType>
+	FORCEINLINE void Remove(const uint32 Index, const uint32 Key, const uint32 LastIndex, uint32 OptLastKey, IndexType* HashData, const uint32 HashCount, IndexType* NextIndexData, const uint32 NextIndexCount)
+	{
+		checkSlow(LastIndex < NextIndexCount && Index <= LastIndex);
+
+		UE::Core::CompactHashTable::Private::RemoveInternal<IndexType>(Index, Key, HashData, HashCount, NextIndexData, NextIndexCount);
+
+		if (Index != LastIndex)
+		{
+			// Remove the last element and add it into the empty spot
+			UE::Core::CompactHashTable::Private::RemoveInternal<IndexType>(LastIndex, OptLastKey, HashData, HashCount, NextIndexData, NextIndexCount);
+
+			const uint32 HashIndex = OptLastKey & (HashCount - 1);
+			NextIndexData[Index] = HashData[HashIndex];
+			HashData[HashIndex] = (IndexType)Index;
+		}
+	}
+
+	/* Remove an element from the list, shift all indexes down to preserve the order of elements in a list
+	 This is a very expensive operation so it should only be used if absolutely necessary (i.e. generally only for user facing data)
+	*/
+	template<typename IndexType>
+	FORCEINLINE void RemoveStable(const uint32 Index, const uint32 Key, IndexType* HashData, const uint32 HashCount, IndexType* NextIndexData, const uint32 NextIndexCount)
+	{
+		checkSlow(Index < NextIndexCount);
+
+		UE::Core::CompactHashTable::Private::RemoveInternal<IndexType>(Index, Key, HashData, HashCount, NextIndexData, NextIndexCount);
+
+		// For the hash indexes, just decrement any that are bigger than the removed element
+		for (uint32 HashIndex = 0; HashIndex < HashCount; ++HashIndex)
+		{
+			if (HashData[HashIndex] > Index && HashData[HashIndex] != (IndexType)INDEX_NONE)
+			{
+				--HashData[HashIndex];
+			}
+		}
+
+		// Decrement values for all next index elements that are before the removed element
+		for (uint32 NextIndexIndex = 0; NextIndexIndex < Index; ++NextIndexIndex)
+		{
+			if (NextIndexData[NextIndexIndex] > Index && NextIndexData[NextIndexIndex] != (IndexType)INDEX_NONE)
+			{
+				--NextIndexData[NextIndexIndex];
+			}
+		}
+
+		// Move AND Decrement values for all next index elements that are after the removed element
+		for (uint32 NextIndexIndex = Index + 1; NextIndexIndex < NextIndexCount; ++NextIndexIndex)
+		{
+			if (NextIndexData[NextIndexIndex] > Index && NextIndexData[NextIndexIndex] != (IndexType)INDEX_NONE)
+			{
+				NextIndexData[NextIndexIndex - 1] = NextIndexData[NextIndexIndex] - 1;
+			}
+			else
+			{
+				NextIndexData[NextIndexIndex - 1] = NextIndexData[NextIndexIndex];
+			}
+		}
+	}
+}
+
+/* Helper to lookup a type from a type size, todo: Add support for 3 byte lookup */
+template<uint32 TypeSize> struct TCompactHashTypeLookupBySize;
+template<> struct TCompactHashTypeLookupBySize<1> { using Type = uint8; };
+template<> struct TCompactHashTypeLookupBySize<2> { using Type = uint16; };
+template<> struct TCompactHashTypeLookupBySize<4> { using Type = uint32; };
+
+/* Helper for making a fixed size hash table that manages its own memory */
+template<uint32 ElementCount, uint32 HashCount = UE::Core::CompactHashTable::GetHashCount(ElementCount)>
+class TStaticCompactHashTable
+{
+public:
+	using IndexType = typename TCompactHashTypeLookupBySize<UE::Core::CompactHashTable::GetTypeSize(ElementCount)>::Type;
+
+	TStaticCompactHashTable()
+	{
+		static_assert(FMath::IsPowerOfTwo(HashCount));
+		Reset();
+	}
+	TStaticCompactHashTable(ENoInit) {}
+	
+	FORCEINLINE void Reset()
+	{
+		FMemory::Memset(HashData, 0xff, sizeof(HashData));
+	}
+
+	[[nodiscard]] FORCEINLINE uint32 GetFirst(uint32 Key) const
+	{
+		return UE::Core::CompactHashTable::GetFirst(Key, HashData, HashCount);
+	}
+
+	[[nodiscard]] FORCEINLINE uint32 GetFirstByIndex(uint32 HashIndex) const
+	{
+		return UE::Core::CompactHashTable::GetFirstByIndex(HashIndex, HashData, HashCount);
+	}
+
+	[[nodiscard]] FORCEINLINE uint32 GetNext(uint32 Index, uint32 CurrentCount) const
+	{
+		checkSlow(CurrentCount <= ElementCount);
+		return UE::Core::CompactHashTable::GetNext(Index, NextIndexData, CurrentCount);
+	}
+
+	template<typename PredicateType>
+	[[nodiscard]] FORCEINLINE uint32 Find(uint32 Key, uint32 CurrentCount, const PredicateType& Predicate) const
+	{
+		checkSlow(CurrentCount <= ElementCount);
+		return UE::Core::CompactHashTable::Find(Key, HashData, HashCount, NextIndexData, CurrentCount, Predicate);
+	}
+	
+	FORCEINLINE void Add(uint32 CurrentCount, uint32 Key)
+	{
+		checkSlow(CurrentCount < ElementCount);
+		UE::Core::CompactHashTable::Add(CurrentCount, Key, HashData, HashCount, NextIndexData, CurrentCount + 1);
+	}
+
+	FORCEINLINE void Remove(uint32 Index, uint32 Key, uint32 LastIndex, uint32 OptLastKey)
+	{
+		UE::Core::CompactHashTable::Remove(Index, Key, LastIndex, OptLastKey, HashData, HashCount, NextIndexData, ElementCount);
+	}
+
+protected:
+	IndexType NextIndexData[ElementCount]; // Collision redirector to next index for keys that hash to the same initial index
+	IndexType HashData[HashCount]; // First index lookup from key
+};
+
+/* Helper for interacting with existing hash table memory */
+class FConstCompactHashTableView
+{
+public:
+	explicit FConstCompactHashTableView() = default;
+
+	FORCEINLINE explicit FConstCompactHashTableView(const uint8* Memory, uint32 InNextIndexCount, uint32 InHashCount, size_t MemorySize)
+	: NextIndexData(Memory)
+	, HashData(Memory + (size_t(InNextIndexCount) << UE::Core::CompactHashTable::GetTypeShift(InNextIndexCount)))
+	, NextIndexCount(InNextIndexCount)
+	, HashCount(InHashCount)
+	{
+		checkSlow(Memory != nullptr && InNextIndexCount > 0 && InHashCount > 0 && MemorySize > 0);
+		checkfSlow(MemorySize == UE::Core::CompactHashTable::GetMemoryRequiredInBytes(NextIndexCount, HashCount), TEXT("Got %d bytes, expected %d bytes"), MemorySize, UE::Core::CompactHashTable::GetMemoryRequiredInBytes(NextIndexCount, HashCount));
+		checkSlow(FMath::IsPowerOfTwo(HashCount));
+	}
+
+	[[nodiscard]] uint32 GetHashCount() const
+	{
+		return HashCount;
+	}
+
+	// Functions used to search
+	[[nodiscard]] FORCEINLINE uint32 GetFirst(uint32 Key) const
+	{
+		#define UE_COMPACTHASHTABLE_EXECUTEBYTYPE(Type) return UE::Core::CompactHashTable::GetFirst(Key, (const Type *)HashData, HashCount)
+		UE_COMPACTHASHTABLE_CALLBYTYPE(NextIndexCount);
+		#undef UE_COMPACTHASHTABLE_EXECUTEBYTYPE
+		return INDEX_NONE;
+	}
+
+	// Advanced used for manual inspection of the hash data
+	[[nodiscard]] FORCEINLINE uint32 GetFirstByIndex(uint32 HashIndex) const
+	{
+		#define UE_COMPACTHASHTABLE_EXECUTEBYTYPE(Type) return UE::Core::CompactHashTable::GetFirstByIndex(HashIndex, (const Type *)HashData, HashCount)
+		UE_COMPACTHASHTABLE_CALLBYTYPE(NextIndexCount);
+		#undef UE_COMPACTHASHTABLE_EXECUTEBYTYPE
+		return INDEX_NONE;
+	}
+
+	[[nodiscard]] FORCEINLINE uint32 GetNext(uint32 Index, uint32 CurrentCount) const
+	{
+		checkSlow(CurrentCount <= NextIndexCount);
+	
+		#define UE_COMPACTHASHTABLE_EXECUTEBYTYPE(Type) return UE::Core::CompactHashTable::GetNext(Index, (const Type *)NextIndexData, CurrentCount)
+		UE_COMPACTHASHTABLE_CALLBYTYPE(NextIndexCount);
+		#undef UE_COMPACTHASHTABLE_EXECUTEBYTYPE
+		return INDEX_NONE;
+	}
+
+	template<typename PredicateType>
+	[[nodiscard]] FORCEINLINE uint32 Find(uint32 Key, uint32 CurrentCount, const PredicateType& Predicate) const
+	{
+		checkSlow(CurrentCount <= NextIndexCount);
+		#define UE_COMPACTHASHTABLE_EXECUTEBYTYPE(Type) return UE::Core::CompactHashTable::Find(Key, (const Type *)HashData, HashCount, (const Type *)NextIndexData, CurrentCount, Predicate)
+		UE_COMPACTHASHTABLE_CALLBYTYPE(NextIndexCount);
+		#undef UE_COMPACTHASHTABLE_EXECUTEBYTYPE
+		return INDEX_NONE;
+	}
+
+protected:
+	const uint8* NextIndexData = nullptr;
+	const uint8* HashData = nullptr;
+	uint32 NextIndexCount = 0;
+	uint32 HashCount = 0;
+};
+
+/* Helper for interacting with existing hash table memory */
+class FCompactHashTableView : public FConstCompactHashTableView
+{
+public:
+	FORCEINLINE explicit FCompactHashTableView(uint8* Memory, uint32 InNextIndexCount, uint32 InHashCount, size_t MemorySize)
+	: FConstCompactHashTableView(Memory, InNextIndexCount, InHashCount, MemorySize)
+	{
+	}
+
+	FORCEINLINE void Reset() const
+	{
+		const uint32 TypeShift = UE::Core::CompactHashTable::GetTypeShift(NextIndexCount);
+
+		// The const_casts are a little gross but I don't wan't to maintain duplicate copies of the const functions
+		FMemory::Memset(const_cast<uint8*>(HashData), 0xff, size_t(HashCount) << TypeShift);
+	}
+	
+	FORCEINLINE void Add(uint32 Index, uint32 Key) const
+	{
+		checkSlow(Index < NextIndexCount);
+		#define UE_COMPACTHASHTABLE_EXECUTEBYTYPE(Type) UE::Core::CompactHashTable::Add(Index, Key, (Type*)(HashData), HashCount, (Type*)(NextIndexData), NextIndexCount)
+		UE_COMPACTHASHTABLE_CALLBYTYPE(NextIndexCount);
+		#undef UE_COMPACTHASHTABLE_EXECUTEBYTYPE
+	}
+
+	FORCEINLINE void Remove(uint32 Index, uint32 Key, uint32 LastIndex, uint32 OptLastKey) const
+	{
+		#define UE_COMPACTHASHTABLE_EXECUTEBYTYPE(Type) return UE::Core::CompactHashTable::Remove(Index, Key, LastIndex, OptLastKey, (Type*)(HashData), HashCount, (Type*)(NextIndexData), NextIndexCount)
+		UE_COMPACTHASHTABLE_CALLBYTYPE(NextIndexCount);
+		#undef UE_COMPACTHASHTABLE_EXECUTEBYTYPE
+	}
+
+	FORCEINLINE void RemoveStable(uint32 Index, uint32 Key) const
+	{
+		#define UE_COMPACTHASHTABLE_EXECUTEBYTYPE(Type) return UE::Core::CompactHashTable::RemoveStable(Index, Key, (Type*)(HashData), HashCount, (Type*)(NextIndexData), NextIndexCount)
+		UE_COMPACTHASHTABLE_CALLBYTYPE(NextIndexCount);
+		#undef UE_COMPACTHASHTABLE_EXECUTEBYTYPE
+	}
+};
diff --git a/Engine/Source/Runtime/Core/Public/Containers/CompactSet.h b/Engine/Source/Runtime/Core/Public/Containers/CompactSet.h
new file mode 100644
--- /dev/null
+++ b/Engine/Source/Runtime/Core/Public/Containers/CompactSet.h
@@ -0,0 +1,1669 @@
+// Copyright Epic Games, Inc. All Rights Reserved.
+
+#pragma once
+
+#include "ContainersFwd.h"
+#include "Containers/CompactSetBase.h"
+#include "Containers/ContainerAllocationPolicies.h"
+#include "Containers/ContainerElementTypeCompatibility.h"
+#include "Containers/SetUtilities.h"
+#include "Containers/UnrealString.h"
+#include "Serialization/StructuredArchive.h"
+#include "Serialization/MemoryImageWriter.h"
+#include "Templates/AlignmentTemplates.h"
+#include "Templates/RetainedRef.h"
+#include "Templates/TypeHash.h"
+
+#include <initializer_list>
+#include <type_traits>
+
+namespace UE::Core::Private
+{
+	[[noreturn]] CORE_API void OnInvalidSetNum(unsigned long long NewNum);
+}
+
+/**
+ * A set with an optional KeyFuncs parameters for customizing how the elements are compared and searched.  
+ * E.g. You can specify a mapping from elements to keys if you want to find elements by specifying a subset of 
+ * the element type.  It uses a TSparseArray of the elements, and also links the elements into a hash with a 
+ * number of buckets proportional to the number of elements.  Addition, removal, and finding are O(1).
+ *
+ * The ByHash() functions are somewhat dangerous but particularly useful in two scenarios:
+ * -- Heterogeneous lookup to avoid creating expensive keys like FString when looking up by const TCHAR*.
+ *	  You must ensure the hash is calculated in the same way as ElementType is hashed.
+ *    If possible put both ComparableKey and ElementType hash functions next to each other in the same header
+ *    to avoid bugs when the ElementType hash function is changed.
+ * -- Reducing contention around hash tables protected by a lock. It is often important to incur
+ *    the cache misses of reading key data and doing the hashing *before* acquiring the lock.
+ *
+ * Here's a visual example of the data layout for the compact set
+ *  ____________ ___________ ________________ ____________
+ * |            |           |                |            |
+ * | Data Array | Hash Size | Collision List | Hash Table |
+ * |____________|___________|________________|____________|
+ * 
+ * Data Array - Payload of the set. This is a just a regular array of items without any empty spots
+ * Hash Size - 4 byte integer to reference how big the hash table is. Storing this is significantly faster than trying to recalculate it
+ * Collision List - For each entry in the Data Array there is an entry in this list that may contain a valid index to the next item to consider for hash table collisions
+ * Hash Table - Power of 2 table to lookup the first index of a given hash value
+ *
+ **/
+template<typename InElementType, typename KeyFuncs /*= DefaultKeyFuncs<ElementType>*/, typename Allocator /*= FDefaultSetAllocator*/>
+class alignas(Allocator::template AllocatorAlignment<InElementType>::Value) TCompactSet : public TCompactSetBase<typename Allocator::template ElementAllocator<sizeof(InElementType)>>
+{
+public:
+	using ElementType = InElementType;
+	using KeyFuncsType = KeyFuncs;
+
+private:
+	using Super = TCompactSetBase<typename Allocator::template ElementAllocator<sizeof(InElementType)>>;
+	using typename Super::SizeType;
+	using typename Super::HashCountType;
+	using typename Super::AllocatorType;
+	using USizeType = std::make_unsigned_t<SizeType>;
+	using KeyInitType = typename KeyFuncs::KeyInitType;
+	using ElementInitType = typename KeyFuncs::ElementInitType;
+
+	// Required so TScriptCompactSet can validate it's layout matches this type
+	template <typename ScriptAllocator>
+	friend class TScriptCompactSet;
+
+	// Private helper functions to help binding to global functions
+	friend struct TCompactSetPrivateFriend;
+
+	/** The base type of whole set iterators. */
+	template<bool bConst>
+	class TBaseIterator
+	{
+	private:
+		using SetType = std::conditional_t<bConst, const TCompactSet, TCompactSet>;
+
+	public:
+		using ElementItType = std::conditional_t<bConst, const ElementType, ElementType>;
+
+		[[nodiscard]] FORCEINLINE TBaseIterator(SetType& InSet)
+		: TBaseIterator(InSet, 0)
+		{
+		}
+
+		[[nodiscard]] FORCEINLINE TBaseIterator(SetType& InSet, SizeType StartIndex)
+		: Set(InSet)
+		, Index(StartIndex)
+#if DO_CHECK
+		, InitialNum(InSet.Num())
+#endif
+		{
+#if DO_CHECK
+			check(StartIndex >= 0 && StartIndex <= InitialNum);
+#endif
+		}
+
+		[[nodiscard]] FORCEINLINE ElementItType & operator*() const
+		{
+			return Set[FSetElementId::FromInteger(Index)];
+		}
+
+		[[nodiscard]] FORCEINLINE ElementItType* operator->() const
+		{
+			return &Set[FSetElementId::FromInteger(Index)];
+		}
+
+		FORCEINLINE TBaseIterator& operator++()
+		{
+			++Index;
+			return *this;
+		}
+
+		/** conversion to "bool" returning true if the iterator is valid. */
+		[[nodiscard]] FORCEINLINE explicit operator bool() const
+		{
+			return Set.IsValidId(this->GetId());
+		}
+
+		[[nodiscard]] FORCEINLINE bool operator==(const TBaseIterator& Rhs) const
+		{
+#if DO_CHECK
+			checkf(Set.Num() == InitialNum, TEXT("Sets/Maps should never be mutated during range-for iteration. Use CreateIterator form to allow more flexibility. InitialNum %d CurrentNum %d"), InitialNum, Set.Num());
+#endif
+			checkSlow(&Set == &Rhs.Set);
+			return Index == Rhs.Index;
+		}
+
+		[[nodiscard]] FORCEINLINE bool operator!=(const TBaseIterator& Rhs) const
+		{
+			return !(*this == Rhs);
+		}
+
+		// Accessors.
+		[[nodiscard]] FORCEINLINE FSetElementId GetId() const
+		{
+			return FSetElementId::FromInteger(Index);
+		}
+		
+	protected:
+		SetType& Set;
+		SizeType Index;
+#if DO_CHECK
+		SizeType InitialNum;
+#endif
+	};
+
+	/** The base type of whole set iterators. */
+	template<bool bConst>
+	class TBaseKeyIterator
+	{
+	private:
+		typedef std::conditional_t < bConst, const TCompactSet, TCompactSet > SetType;
+		typedef std::conditional_t < bConst, const ElementType, ElementType > ItElementType;
+		typedef typename TTypeTraits<typename KeyFuncs::KeyType>::ConstPointerType ReferenceOrValueType;
+
+	public:
+		using KeyArgumentType = std::conditional_t <
+			std::is_reference_v<ReferenceOrValueType> ,
+			TRetainedRef<std::remove_reference_t<ReferenceOrValueType>> ,
+			KeyInitType
+		>;
+
+		/** Initialization constructor. */
+		[[nodiscard]] FORCEINLINE TBaseKeyIterator(SetType& InSet, KeyArgumentType InKey)
+		: Set  (InSet)
+		, Key  (InKey)
+		, HashTable(InSet.Num() ? Set.GetConstHashTableView() : FConstCompactHashTableView())
+		, NextIndex(InSet.Num() ? HashTable.GetFirst(KeyFuncs::GetKeyHash(Key)) : INDEX_NONE)
+#if DO_CHECK
+		, InitialNum(InSet.Num())
+#endif
+		{
+			++(*this);
+		}
+
+		/** Advances the iterator to the next element. */
+		FORCEINLINE TBaseKeyIterator& operator++()
+		{
+			const int32 SetNum = Set.Num();
+
+			// Note: Adding new elements is safe, however they will be guaranteed to be missed by the current iteration
+#if DO_CHECK
+			checkf(SetNum >= InitialNum, TEXT("Sets/Maps should never have elements removed during iteration outside of Iterator.RemoveCurrent(). InitialNum %d CurrentNum %d"), InitialNum, Set.Num());
+#endif
+
+			Index = NextIndex;
+
+			while (Index != INDEX_NONE)
+			{
+				NextIndex = HashTable.GetNext(Index, SetNum);
+				checkSlow(Index != NextIndex);
+
+				if (KeyFuncs::Matches(KeyFuncs::GetSetKey(Set[GetId()]), Key))
+				{
+					break;
+				}
+
+				Index = NextIndex;
+			}
+
+			return *this;
+		}
+
+		/** conversion to "bool" returning true if the iterator is valid. */
+		[[nodiscard]] FORCEINLINE explicit operator bool() const
+		{
+			return Index != INDEX_NONE;
+		}
+
+		// Accessors.
+		[[nodiscard]] FORCEINLINE FSetElementId GetId() const
+		{
+			return FSetElementId::FromInteger(this->Index);
+		}
+
+		[[nodiscard]] FORCEINLINE ItElementType* operator->() const
+		{
+			return &Set[GetId()];
+		}
+
+		[[nodiscard]] FORCEINLINE ItElementType& operator*() const
+		{
+			return Set[GetId()];
+		}
+
+	protected:
+		SetType& Set;
+		ReferenceOrValueType Key;
+		FConstCompactHashTableView HashTable;
+		SizeType Index;
+		SizeType NextIndex;
+
+#if DO_CHECK
+		SizeType InitialNum;
+#endif
+	};
+
+public:
+
+	/** Used to iterate over the elements of a const TCompactSet. */
+	using TConstIterator = TBaseIterator<true>;
+
+	/** Used to iterate over the elements of a TCompactSet. */
+	class TIterator : public TBaseIterator<false>
+	{
+	private:
+		using Super = TBaseIterator<false>;
+
+	public:
+		using Super::Super;
+
+		/** Removes the current element from the set. */
+		FORCEINLINE void RemoveCurrent()
+		{
+			this->Set.Remove(this->GetId());
+			--this->Index;
+#if DO_CHECK
+			--this->InitialNum;
+#endif
+		}
+	};
+
+	/** Used to iterate over the elements of a const TCompactSet. */
+	using TConstKeyIterator = TBaseKeyIterator<true>;
+
+	/** Used to iterate over the elements of a TCompactSet. */
+	class TKeyIterator : public TBaseKeyIterator<false>
+	{
+	private:
+		using Super = TBaseKeyIterator<false>;
+
+	public:
+		using KeyArgumentType = typename Super::KeyArgumentType;
+		using Super::Super;
+
+		/** Removes the current element from the set. */
+		FORCEINLINE void RemoveCurrent()
+		{
+			this->Set.Remove(this->GetId());
+#if DO_CHECK
+			--this->InitialNum;
+#endif
+
+			// If the next element was the last in the set then it will get remapped to the current index, so fix that up
+			if (this->NextIndex == TBaseKeyIterator<false>::Set.Num())
+			{
+				this->NextIndex = TBaseKeyIterator<false>::Index;
+			}
+
+			this->Index = INDEX_NONE;
+		}
+	};
+
+	/* Constructors */
+
+	[[nodiscard]] FORCEINLINE TCompactSet() = default;
+
+	[[nodiscard]] FORCEINLINE TCompactSet(const TCompactSet& Copy)
+	{
+		*this = Copy;
+	}
+
+	[[nodiscard]] FORCEINLINE explicit TCompactSet(TArrayView<const ElementType> InArrayView)
+	{
+		Append(InArrayView);
+	}
+
+	[[nodiscard]] FORCEINLINE explicit TCompactSet(TArray<ElementType>&& InArray)
+	{
+		Append(MoveTemp(InArray));
+	}
+
+	[[nodiscard]] TCompactSet(std::initializer_list<ElementType> InitList)
+	{
+		Append(InitList);
+	}
+
+	[[nodiscard]] TCompactSet(TCompactSet&& Other)
+	{
+		*this = MoveTemp(Other);
+	}
+
+	/** Constructor for moving elements from a TCompactSet with a different SetAllocator */
+	template<typename OtherAllocator>
+	[[nodiscard]] TCompactSet(TCompactSet<ElementType, KeyFuncs, OtherAllocator>&& Other)
+	{
+		Append(MoveTemp(Other));
+	}
+
+	/** Constructor for copying elements from a TCompactSet with a different SetAllocator */
+	template<typename OtherAllocator>
+	[[nodiscard]] TCompactSet(const TCompactSet<ElementType, KeyFuncs, OtherAllocator>& Other)
+	{
+		Append(Other);
+	}
+
+	FORCEINLINE ~TCompactSet()
+	{
+		Empty(0);
+	}
+
+	////////////////////////////////////////////////////
+	// Start - intrusive TOptional<TCompactSet> state //
+	////////////////////////////////////////////////////
+	constexpr static bool bHasIntrusiveUnsetOptionalState = true;
+	using IntrusiveUnsetOptionalStateType = TCompactSet;
+
+	[[nodiscard]] explicit TCompactSet(FIntrusiveUnsetOptionalState Tag) : Super(Tag)
+	{
+	}
+	//////////////////////////////////////////////////
+	// End - intrusive TOptional<TCompactSet> state //
+	//////////////////////////////////////////////////
+
+	/* Assignment operators */
+
+	TCompactSet& operator=(const TCompactSet& Copy)
+	{
+		if (this != &Copy)
+		{
+			// This could either take the full memory size of the Copy and prevent a rehash or
+			// only allocate the required set but have to hash everything again (i.e. perf vs memory)
+			// Could try a middle ground of allowing extra memory if it's within slack margins.
+			// Going for memory savings for now
+
+			// Not using Empty(NumElements) to avoid clearing the hash memory since we'll rebuild it anyway
+			// We need to make sure the relevant parts are cleared so ResizeAllocation can run safely and with minimal cost
+			DestructItems(GetData(), this->NumElements);
+			this->NumElements = 0;
+
+			ResizeAllocation(Copy.NumElements);
+
+			this->NumElements = Copy.NumElements;
+			ConstructItems<ElementType>(GetData(), Copy.GetData(), this->NumElements);
+
+			Rehash();
+		}
+		return *this;
+	}
+
+	TCompactSet& operator=(TCompactSet&& Other)
+	{
+		if (this != &Other)
+		{
+			this->Elements.MoveToEmpty(Other.Elements);
+			this->NumElements = Other.NumElements;
+			this->MaxElements = Other.MaxElements;
+
+			Other.NumElements = 0;
+			Other.MaxElements = 0;
+		}
+
+		return *this;
+	}
+
+	/** Assignment operator for moving elements from a TCompactSet with a different SetAllocator */
+	template<typename OtherAllocator>
+	TCompactSet& operator=(TCompactSet<ElementType, KeyFuncs, OtherAllocator> && Other)
+	{
+		Reset();
+		Append(MoveTemp(Other));
+		return *this;
+	}
+
+	/** Assignment operator for copying elements from a TCompactSet with a different SetAllocator */
+	template<typename OtherAllocator>
+	TCompactSet& operator=(const TCompactSet<ElementType, KeyFuncs, OtherAllocator>& Other)
+	{
+		Reset();
+		Append(Other);
+		return *this;
+	}
+
+	TCompactSet& operator=(std::initializer_list<ElementType> InitList)
+	{
+		Reset();
+		Append(InitList);
+		return *this;
+	}
+
+	/**
+	 * Move assignment operator.
+	 * Compatible element type version.
+	 *
+	 * @param Other Set to assign and move from.
+	 */
+	template <
+		typename OtherKeyFuncs,
+		typename AliasElementType = ElementType
+		UE_REQUIRES(TIsContainerElementTypeCopyable_V<AliasElementType>)
+	>
+	TCompactSet& operator=(TCompactSet<typename TContainerElementTypeCompatibility<ElementType>::CopyFromOtherType, OtherKeyFuncs, Allocator> && Other)
+	{
+		TContainerElementTypeCompatibility<ElementType>::CopyingFromOtherType();
+		Reset();
+		Append(MoveTemp(Other));
+		return *this;
+	}
+
+	/**
+	 * Assignment operator. First deletes all currently contained elements
+	 * and then copies from other set.
+	 * Compatible element type version.
+	 *
+	 * @param Other The source set to assign from.
+	 */
+	template <
+		typename OtherKeyFuncs,
+		typename OtherAllocator,
+		typename AliasElementType = ElementType
+		UE_REQUIRES(TIsContainerElementTypeCopyable_V<AliasElementType>)
+	>
+	TCompactSet& operator=(const TCompactSet<typename TContainerElementTypeCompatibility<ElementType>::CopyFromOtherType, OtherKeyFuncs, OtherAllocator>& Other)
+	{
+		TContainerElementTypeCompatibility<ElementType>::CopyingFromOtherType();
+		Reset();
+		Append(Other);
+		return *this;
+	}
+
+	/**
+	 * Removes all elements from the set, potentially leaving space allocated for an expected number of elements about to be added.
+	 * @param ExpectedNumElements - The number of elements about to be added to the set.
+	 */
+	void Empty(int32 ExpectedNumElements = 0)
+	{
+		DestructItems(GetData(), this->NumElements);
+		this->NumElements = 0;
+
+		ResizeAllocation(ExpectedNumElements);
+		if (this->MaxElements > 0)
+		{
+			GetHashTableView().Reset();
+		}
+	}
+
+	/** Efficiently empties out the set but preserves all allocations and capacities */
+	void Reset()
+	{
+		if (this->NumElements > 0)
+		{
+			DestructItems(GetData(), this->NumElements);
+			this->NumElements = 0;
+			
+			GetHashTableView().Reset();
+		}
+	}
+
+	/** Shrinks the set's element storage to avoid slack. */
+	void Shrink()
+	{
+		if (this->NumElements != this->MaxElements)
+		{
+			if (ResizeAllocationPreserveData(this->NumElements))
+			{
+				Rehash();
+			}
+		}
+	}
+
+	/** Preallocates enough memory to contain Number elements */
+	void Reserve(int32 Number)
+	{
+		// makes sense only when Number > Elements.Num() since TSparseArray::Reserve 
+		// does any work only if that's the case
+		if ((USizeType)Number > (USizeType)this->MaxElements)
+		{
+			// Trap negative reserves
+			if (Number < 0)
+			{
+				UE::Core::Private::OnInvalidSetNum((unsigned long long)Number);
+			}
+
+			if (ResizeAllocationPreserveData(Number))
+			{
+				Rehash();
+			}
+		}
+	}
+
+	/** Deprecated - default behavior now, keeping this here so TCompactSet can be swapped with TSet without changing code */
+	void Compact()
+	{
+	}
+
+	/** Deprecated - use sparse array if this behavior is required, keeping this here so TCompactSet can be swapped with TSet without changing code */
+	void CompactStable()
+	{
+		ensureMsgf(false, TEXT("Compact sets are always compact so CompactStable will not do anything. If you hit this then you likely need to use a different pattern to maintain order, see RemoveStable"));
+	}
+
+	/** Deprecated - unnecessary, keeping this here so TCompactSet can be swapped with TSet without changing code */
+	void SortFreeList()
+	{
+	}
+
+	/** Deprecated - default behavior now, keeping this here so TCompactSet can be swapped with TSet without changing code */
+	void Relax()
+	{
+	}
+
+	/** 
+	 * Helper function to return the amount of memory allocated by this container 
+	 * Only returns the size of allocations made directly by the container, not the elements themselves.
+	 * @return number of bytes allocated by this container
+	 */
+	[[nodiscard]] FORCEINLINE SIZE_T GetAllocatedSize() const
+	{
+		return this->Elements.GetAllocationSize();
+	}
+
+	/** Tracks the container's memory use through an archive. */
+	FORCEINLINE void CountBytes(FArchive& Ar) const
+	{
+		const FCompactSetLayout Layout = GetSetLayout();
+		Ar.CountBytes(Super::GetTotalMemoryRequiredInBytes(this->NumElements, Layout), Super::GetTotalMemoryRequiredInBytes(this->MaxElements, Layout));
+	}
+
+	/* Calculate the size of the hash table from the number of elements in the set assuming the default number of hash elements */
+	[[nodiscard]] FORCEINLINE static constexpr size_t GetTotalMemoryRequiredInBytes(uint32 NumElements)
+	{
+		return Super::GetTotalMemoryRequiredInBytes(NumElements, GetSetLayout());
+	}
+
+
+	/* Get the alignment required for the allocation */
+	[[nodiscard]] FORCEINLINE static constexpr size_t GetAlignment()
+	{
+		return FGenericPlatformMath::Max<size_t>(alignof(ElementType), UE::Core::CompactHashTable::GetMemoryAlignment());
+	}
+
+	[[nodiscard]] FORCEINLINE ElementType* GetData()
+	{
+		return (ElementType*)this->Elements.GetAllocation();
+	}
+
+	[[nodiscard]] FORCEINLINE const ElementType* GetData() const
+	{
+		return (const ElementType*)this->Elements.GetAllocation();
+	}
+
+	/**
+	 * Checks whether an element id is valid.
+	 * @param Id - The element id to check.
+	 * @return true if the element identifier refers to a valid element in this set.
+	 */
+	[[nodiscard]] FORCEINLINE bool IsValidId(FSetElementId Id) const
+	{
+		return Id.AsInteger() >= 0 && Id.AsInteger() < this->NumElements;
+	}
+
+	/**
+	 * Checks array invariants: if array size is greater than or equal to zero and less
+	 * than or equal to the maximum.
+	 */
+	FORCEINLINE void CheckInvariants() const
+	{
+		checkSlow((this->NumElements >= 0) & (this->MaxElements >= this->NumElements)); // & for one branch
+	}
+
+	/**
+		* Checks if index is in array range.
+		*
+		* @param Index Index to check.
+		*/
+	FORCEINLINE void RangeCheck(FSetElementId Id) const
+	{
+		CheckInvariants();
+
+		// Template property, branch will be optimized out
+		if constexpr (AllocatorType::RequireRangeCheck)
+		{
+			checkf(IsValidId(Id), TEXT("Array index out of bounds: %d into an array of size %lld"), Id.AsInteger(), (long long)this->NumElements); // & for one branch
+		}
+	}
+
+	/** Accesses the identified element's value. Element must be valid (see @IsValidId). */
+	[[nodiscard]] FORCEINLINE ElementType & operator[](FSetElementId Id)
+	{
+		RangeCheck(Id);
+		return GetData()[Id.AsInteger()];
+	}
+
+	/** Accesses the identified element's value. Element must be valid (see @IsValidId). */
+	[[nodiscard]] FORCEINLINE const ElementType & operator[](FSetElementId Id) const
+	{
+		RangeCheck(Id);
+		return GetData()[Id.AsInteger()];
+	}
+
+	/** Accesses the identified element's value. Element must be valid (see @IsValidId). */
+	[[nodiscard]] FORCEINLINE ElementType& Get(FSetElementId Id)
+	{
+		RangeCheck(Id);
+		return GetData()[Id.AsInteger()];
+	}
+
+	/** Accesses the identified element's value. Element must be valid (see @IsValidId). */
+	[[nodiscard]] FORCEINLINE const ElementType& Get(FSetElementId Id) const
+	{
+		RangeCheck(Id);
+		return GetData()[Id.AsInteger()];
+	}
+
+	/**
+	 * Adds an element to the set.
+	 *
+	 * @param	InElement					Element to add to set
+	 * @param	bIsAlreadyInSetPtr	[out]	Optional pointer to bool that will be set depending on whether element is already in set
+	 * @return	A pointer to the element stored in the set.
+	 */
+	FORCEINLINE FSetElementId Add(const ElementType&  InElement, bool* bIsAlreadyInSetPtr = nullptr)
+	{
+		return EmplaceByHash(KeyFuncs::GetKeyHash(KeyFuncs::GetSetKey(InElement)), InElement, bIsAlreadyInSetPtr);
+	}
+	FORCEINLINE FSetElementId Add(ElementType && InElement, bool* bIsAlreadyInSetPtr = nullptr)
+	{
+		return EmplaceByHash(KeyFuncs::GetKeyHash(KeyFuncs::GetSetKey(InElement)), MoveTempIfPossible(InElement), bIsAlreadyInSetPtr);
+	}
+
+	/**
+	 * Adds an element to the set.
+	 *
+	 * @see		Class documentation section on ByHash() functions
+	 * @param	InElement					Element to add to set
+	 * @param	bIsAlreadyInSetPtr	[out]	Optional pointer to bool that will be set depending on whether element is already in set
+     * @return  A handle to the element stored in the set
+     */
+	FORCEINLINE FSetElementId AddByHash(uint32 KeyHash, const ElementType& InElement, bool* bIsAlreadyInSetPtr = nullptr)
+	{
+		return EmplaceByHash(KeyHash, InElement, bIsAlreadyInSetPtr);
+	}
+	FORCEINLINE FSetElementId AddByHash(uint32 KeyHash,		 ElementType && InElement, bool* bIsAlreadyInSetPtr = nullptr)
+	{
+		return EmplaceByHash(KeyHash, MoveTempIfPossible(InElement), bIsAlreadyInSetPtr);
+	}
+
+	/**
+	 * Adds an element to the set.
+	 *
+	 * @param	Args						The argument(s) to be forwarded to the set element's constructor.
+	 * @param	bIsAlreadyInSetPtr	[out]	Optional pointer to bool that will be set depending on whether element is already in set
+	 * @return	A handle to the element stored in the set.
+	 */
+	template<typename ArgsType = ElementType>
+	FSetElementId Emplace(ArgsType && Args, bool* bIsAlreadyInSetPtr = nullptr)
+	{
+		alignas(alignof(ElementType))char StackElement[sizeof(ElementType)];
+		ElementType& TempElement = *new (StackElement) ElementType(Forward<ArgsType>(Args));
+		SizeType ExistingIndex(INDEX_NONE);
+
+		const KeyInitType Key = KeyFuncs::GetSetKey(TempElement);
+		const uint32 KeyHash = KeyFuncs::GetKeyHash(Key);
+
+		if constexpr (!KeyFuncs::bAllowDuplicateKeys)
+		{
+			ExistingIndex = FindIndexByHash(KeyHash, Key);
+		}
+
+		if (bIsAlreadyInSetPtr)
+		{
+			*bIsAlreadyInSetPtr = ExistingIndex != INDEX_NONE;
+		}
+
+		if (ExistingIndex == INDEX_NONE)
+		{
+			ExistingIndex = this->NumElements;
+			ElementType& NewElement = AddUninitialized(KeyHash);
+			RelocateConstructItem<ElementType, ElementType, SizeType>(&NewElement, &TempElement);
+		}
+		else
+		{
+			ElementType& ExistingElement = GetData()[ExistingIndex];
+			MoveByRelocate(ExistingElement, TempElement);
+		}
+
+		return FSetElementId::FromInteger(ExistingIndex);
+	}
+	
+	/**
+	 * Adds an element to the set.
+	 *
+	 * @see		Class documentation section on ByHash() functions
+	 * @param	Args						The argument(s) to be forwarded to the set element's constructor.
+	 * @param	bIsAlreadyInSetPtr	[out]	Optional pointer to bool that will be set depending on whether element is already in set
+	 * @return	A handle to the element stored in the set.
+	 */
+	template<typename ArgsType = ElementType>
+	FSetElementId EmplaceByHash(uint32 KeyHash, ArgsType && Args, bool* bIsAlreadyInSetPtr = nullptr)
+	{
+		alignas(alignof(ElementType))char StackElement[sizeof(ElementType)];
+		ElementType& TempElement = *new (StackElement) ElementType(Forward<ArgsType>(Args));
+		SizeType ExistingIndex(INDEX_NONE);
+
+		if constexpr (!KeyFuncs::bAllowDuplicateKeys)
+		{
+			ExistingIndex = FindIndexByHash(KeyHash, KeyFuncs::GetSetKey(TempElement));
+		}
+
+		if (bIsAlreadyInSetPtr)
+		{
+			*bIsAlreadyInSetPtr = ExistingIndex != INDEX_NONE;
+		}
+
+		if (ExistingIndex == INDEX_NONE)
+		{
+			ExistingIndex = this->NumElements;
+			ElementType& NewElement = AddUninitialized(KeyHash);
+			RelocateConstructItem<ElementType, ElementType, SizeType>(&NewElement, &TempElement);
+		}
+		else
+		{
+			ElementType& ExistingElement = GetData()[ExistingIndex];
+			MoveByRelocate(ExistingElement, TempElement);
+		}
+
+		return FSetElementId::FromInteger(ExistingIndex);
+	}
+
+	/**
+		* Adds an element to the set if not already present and returns a reference to the added or existing element.
+		*
+		* @param	InElement					Element to add to set
+		* @param	bIsAlreadyInSetPtr	[out]	Optional pointer to bool that will be set depending on whether element is already in set
+		* @return	A reference to the element stored in the set.
+		*/
+	FORCEINLINE ElementType& FindOrAdd(const InElementType& InElement, bool* bIsAlreadyInSetPtr = nullptr)
+	{
+		return FindOrAddByHash(KeyFuncs::GetKeyHash(KeyFuncs::GetSetKey(InElement)), InElement, bIsAlreadyInSetPtr);
+	}
+	FORCEINLINE ElementType& FindOrAdd(InElementType && InElement, bool* bIsAlreadyInSetPtr = nullptr)
+	{
+		return FindOrAddByHash(KeyFuncs::GetKeyHash(KeyFuncs::GetSetKey(InElement)), MoveTempIfPossible(InElement), bIsAlreadyInSetPtr);
+	}
+
+	/**
+	 * Adds an element to the set if not already present and returns a reference to the added or existing element.
+	 *
+	 * @see		Class documentation section on ByHash() functions
+	 * @param	InElement					Element to add to set
+	 * @param	bIsAlreadyInSetPtr	[out]	Optional pointer to bool that will be set depending on whether element is already in set
+	 * @return  A reference to the element stored in the set
+	 */
+	template<typename ElementReferenceType>
+	ElementType& FindOrAddByHash(uint32 KeyHash, ElementReferenceType && InElement, bool* bIsAlreadyInSetPtr = nullptr)
+	{
+		SizeType ExistingIndex = FindIndexByHash(KeyHash, KeyFuncs::GetSetKey(InElement));
+		bool bIsAlreadyInSet = ExistingIndex != INDEX_NONE;
+		if (bIsAlreadyInSetPtr)
+		{
+			*bIsAlreadyInSetPtr = bIsAlreadyInSet;
+		}
+		if (bIsAlreadyInSet)
+		{
+			return GetData()[ExistingIndex];
+		}
+
+		ElementType& NewElement = AddUninitialized(KeyHash);
+		return *(new ((void*)&NewElement) ElementType(Forward<ElementReferenceType>(InElement)));
+	}
+
+	void Append(TArrayView<const ElementType> InElements)
+	{
+		Reserve(this->NumElements + InElements.Num());
+		for (const ElementType& Element : InElements)
+		{
+			Add(Element);
+		}
+	}
+
+	template<typename ArrayAllocator>
+	void Append(TArray<ElementType, ArrayAllocator> && InElements)
+	{
+		Reserve(this->NumElements + InElements.Num());
+		for (ElementType& Element : InElements)
+		{
+			Add(MoveTempIfPossible(Element));
+		}
+		InElements.Reset();
+	}
+
+	/**
+	 * Add all items from another set to our set (union without creating a new set)
+	 * @param OtherSet - The other set of items to add.
+	 */
+	template<typename OtherAllocator>
+	void Append(const TCompactSet<ElementType, KeyFuncs, OtherAllocator>& OtherSet)
+	{
+		Reserve(this->NumElements + OtherSet.Num());
+		for (const ElementType& Element : OtherSet)
+		{
+			Add(Element);
+		}
+	}
+
+	template<typename OtherAllocator>
+	void Append(TCompactSet<ElementType, KeyFuncs, OtherAllocator> && OtherSet)
+	{
+		Reserve(this->NumElements + OtherSet.Num());
+		for (ElementType& Element : OtherSet)
+		{
+			Add(MoveTempIfPossible(Element));
+		}
+		OtherSet.Reset();
+	}
+
+	void Append(std::initializer_list<ElementType> InitList)
+	{
+		Reserve(this->NumElements + (int32)InitList.size());
+		for (const ElementType& Element : InitList)
+		{
+			Add(Element);
+		}
+	}
+
+	/**
+	 * Add all items from another set to our set (union without creating a new set)
+	 * Compatible element type version.
+	 * @param OtherSet - The other set of items to add.
+	 */
+	template <
+		typename OtherKeyFuncs,
+		typename OtherAllocator,
+		typename AliasElementType = ElementType
+		UE_REQUIRES(TIsContainerElementTypeCopyable_V<AliasElementType>)
+	>
+	void Append(const TCompactSet<typename TContainerElementTypeCompatibility<ElementType>::CopyFromOtherType, OtherKeyFuncs, OtherAllocator>& OtherSet)
+	{
+		TContainerElementTypeCompatibility<ElementType>::CopyingFromOtherType();
+		Reserve(this->NumElements + OtherSet.Num());
+		for (const ElementType& Element : OtherSet)
+		{
+			Add(Element);
+		}
+	}
+
+	/**
+	 * Add all items from another set to our set (union without creating a new set)
+	 * Compatible element type version.
+	 * @param OtherSet - The other set of items to add.
+	 */
+	template <
+		typename OtherKeyFuncs,
+		typename AliasElementType = ElementType
+		UE_REQUIRES(TIsContainerElementTypeCopyable_V<AliasElementType>)
+	>
+	void Append(TCompactSet<typename TContainerElementTypeCompatibility<ElementType>::CopyFromOtherType, OtherKeyFuncs, Allocator> && OtherSet)
+	{
+		TContainerElementTypeCompatibility<ElementType>::CopyingFromOtherType();
+		Reserve(this->NumElements + OtherSet.Num());
+		for (ElementType& Element : OtherSet)
+		{
+			Add(MoveTempIfPossible(Element));
+		}
+		OtherSet.Reset();
+	}
+
+	/**
+		* Finds any element in the set and returns a pointer to it.
+		* Callers should not depend on particular patterns in the behaviour of this function.
+		* @return A pointer to an arbitrary element, or nullptr if the container is empty.
+		*/
+	[[nodiscard]] ElementType* FindArbitraryElement()
+	{
+		return (this->NumElements > 0) ? GetData() : nullptr;
+	}
+	[[nodiscard]] const ElementType* FindArbitraryElement() const
+	{
+		return const_cast<TCompactSet*>(this)->FindArbitraryElement();
+	}
+
+
+	/**
+	 * Finds an element with the given key in the set.
+	 * @param Key - The key to search for.
+	 * @return The id of the set element matching the given key, or the NULL id if none matches.
+	 */
+	[[nodiscard]] FSetElementId FindId(KeyInitType Key) const
+	{
+		return FSetElementId::FromInteger(FindIndexByHash(KeyFuncs::GetKeyHash(Key), Key));
+	}
+
+	/**
+	 * Finds an element with a pre-calculated hash and a key that can be compared to KeyType
+	 * @see	Class documentation section on ByHash() functions
+	 * @return The element id that matches the key and hash or an invalid element id
+	 */
+	template<typename ComparableKey>
+	[[nodiscard]] FSetElementId FindIdByHash(uint32 KeyHash, const ComparableKey& Key) const
+	{
+		checkSlow(KeyHash == KeyFuncs::GetKeyHash(Key));
+		return FSetElementId::FromInteger(FindIndexByHash(KeyHash, Key));
+	}
+
+	/**
+	 * Finds an element with the given key in the set.
+	 * @param Key - The key to search for.
+	 * @return A pointer to an element with the given key.  If no element in the set has the given key, this will return NULL.
+	 */
+	[[nodiscard]] FORCEINLINE ElementType* Find(KeyInitType Key)
+	{
+		SizeType ElementIndex = FindIndexByHash(KeyFuncs::GetKeyHash(Key), Key);
+		if (ElementIndex != INDEX_NONE)
+		{
+			return GetData() + ElementIndex;
+		}
+		else
+		{
+			return nullptr;
+		}
+	}
+	
+	/**
+	 * Finds an element with the given key in the set.
+	 * @param Key - The key to search for.
+	 * @return A const pointer to an element with the given key.  If no element in the set has the given key, this will return NULL.
+	 */
+	[[nodiscard]] FORCEINLINE const ElementType* Find(KeyInitType Key) const
+	{
+		return const_cast<TCompactSet*>(this)->Find(Key);
+	}
+
+	/**
+	 * Finds an element with a pre-calculated hash and a key that can be compared to KeyType.
+	 * @see	Class documentation section on ByHash() functions
+	 * @return A pointer to the contained element or nullptr.
+	 */
+	template<typename ComparableKey>
+	[[nodiscard]] ElementType* FindByHash(uint32 KeyHash, const ComparableKey& Key)
+	{
+		SizeType ElementIndex = FindIndexByHash(KeyHash, Key);
+		if (ElementIndex != INDEX_NONE)
+		{
+			return GetData() + ElementIndex;
+		}
+		else
+		{
+			return nullptr;
+		}
+	}
+
+	template<typename ComparableKey>
+	[[nodiscard]] const ElementType* FindByHash(uint32 KeyHash, const ComparableKey& Key) const
+	{
+		return const_cast<TCompactSet*>(this)->FindByHash(KeyHash, Key);
+	}
+
+	/**
+	 * Removes an element from the set.
+	 * @param Element - A pointer to the element in the set, as returned by Add or Find.
+	 */
+	void Remove(FSetElementId ElementId)
+	{
+		RemoveByIndex(ElementId.AsInteger());
+	}
+
+	/**
+	 * Removes all elements from the set matching the specified key.
+	 * @param Key - The key to match elements against.
+	 * @return The number of elements removed.
+	 */
+	int32 Remove(KeyInitType Key)
+	{
+		if (this->NumElements)
+		{
+			return RemoveImpl(KeyFuncs::GetKeyHash(Key), Key);
+		}
+
+		return 0;
+	}
+
+	/**
+	 * Removes an element from the set while maintaining set order.
+	 * @param Element - A pointer to the element in the set, as returned by Add or Find.
+	 */
+	void RemoveStable(FSetElementId ElementId)
+	{
+		RemoveByIndex<true>(ElementId.AsInteger());
+	}
+
+	/**
+	 * Removes all elements from the set matching the specified key.
+	 * @param Key - The key to match elements against.
+	 * @return The number of elements removed.
+	 */
+	int32 RemoveStable(KeyInitType Key)
+	{
+		if (this->NumElements)
+		{
+			return RemoveImplStable(KeyFuncs::GetKeyHash(Key), Key);
+		}
+
+		return 0;
+	}
+
+	/**
+	 * Removes all elements from the set matching the specified key.
+	 *
+	 * @see		Class documentation section on ByHash() functions
+	 * @param	Key - The key to match elements against.
+	 * @return	The number of elements removed.
+	 */
+	template<typename ComparableKey>
+	int32 RemoveByHash(uint32 KeyHash, const ComparableKey& Key)
+	{
+		checkSlow(KeyHash == KeyFuncs::GetKeyHash(Key));
+
+		if (this->NumElements)
+		{
+			return RemoveImpl(KeyHash, Key);
+		}
+
+		return 0;
+	}
+
+	/**
+	 * Checks if the element contains an element with the given key.
+	 * @param Key - The key to check for.
+	 * @return true if the set contains an element with the given key.
+	 */
+	[[nodiscard]] FORCEINLINE bool Contains(KeyInitType Key) const
+	{
+		return FindIndexByHash(KeyFuncs::GetKeyHash(Key), Key) != INDEX_NONE;
+	}
+
+	/**
+	 * Checks if the element contains an element with the given key.
+	 *
+	 * @see	Class documentation section on ByHash() functions
+	 */
+	template<typename ComparableKey>
+	[[nodiscard]] FORCEINLINE bool ContainsByHash(uint32 KeyHash, const ComparableKey& Key) const
+	{
+		checkSlow(KeyHash == KeyFuncs::GetKeyHash(Key));
+		return FindIndexByHash(KeyHash, Key) != INDEX_NONE;
+	}
+
+	/**
+	 * Sorts the set's elements using the provided comparison class.
+	 */
+	template<typename PredicateType>
+	void Sort(const PredicateType& Predicate)
+	{
+		TArrayView<ElementType>(GetData(), this->NumElements).Sort(Predicate);
+		Rehash();
+	}
+
+	/**
+	 * Stable sorts the set's elements using the provided comparison class.
+	 */
+	template<typename PredicateType>
+	void StableSort(const PredicateType& Predicate)
+	{
+		TArrayView<ElementType>(GetData(), this->NumElements).StableSort(Predicate);
+		Rehash();
+	}
+
+	/**
+	 * Describes the set's contents through an output device.
+	 * @param Ar - The output device to describe the set's contents through.
+	 */
+	void Dump(FOutputDevice& Ar) const
+	{
+		if (this->MaxElements == 0)
+		{
+			Ar.Logf(TEXT("TCompactSet: empty"), this->NumElements, this->MaxElements);
+		}
+		else if (this->MaxElements > 0)
+		{
+			FConstCompactHashTableView HashTable = this->GetConstHashTableView();
+
+			Ar.Logf(TEXT("TCompactSet: %i elements, %i max elements, %i hash slots"), this->NumElements, this->MaxElements, HashTable.GetHashCount());
+
+			for (uint32 HashIndex = 0; HashIndex < HashTable.GetHashCount(); ++HashIndex)
+			{
+				// Count the numTableIndexber of elements in this hash bucket.
+				int32 NumElementsInBucket = 0;
+				for (uint32 ElementIndex = HashTable.GetFirstByIndex(HashIndex); ElementIndex != INDEX_NONE; ElementIndex = HashTable.GetNext(ElementIndex, this->NumElements))
+				{
+					NumElementsInBucket++;
+				}
+
+				Ar.Logf(TEXT("   Hash[%i] = %i"), HashIndex, NumElementsInBucket);
+			}
+		}
+		else
+		{
+			// MaxElements == INDEX_NONE is a TOptional that's null, anything else is just garbage data
+			checkNoEntry();
+		}
+	}
+
+	/** @return the intersection of two sets. (A AND B)*/
+	[[nodiscard]] TCompactSet Intersect(const TCompactSet& OtherSet) const
+	{
+		const bool bOtherSmaller = (this->NumElements > OtherSet.NumElements);
+		const TCompactSet& A = (bOtherSmaller ? OtherSet : *this);
+		const TCompactSet& B = (bOtherSmaller ? *this : OtherSet);
+
+		TCompactSet Result;
+		Result.Reserve(A.NumElements); // Worst case is everything in smaller is in larger
+
+		for (const ElementType& Element : A)
+		{
+			if (B.Contains(KeyFuncs::GetSetKey(Element)))
+			{
+				Result.Add(Element);
+			}
+		}
+		return Result;
+	}
+
+	/** @return the union of two sets. (A OR B)*/
+	[[nodiscard]] TCompactSet Union(const TCompactSet& OtherSet) const
+	{
+		TCompactSet Result;
+		Result.Reserve(this->NumElements + OtherSet.NumElements); // Worst case is 2 totally unique Sets
+
+		for (const ElementType& Element : *this)
+		{
+			Result.Add(Element);
+		}
+		for (const ElementType& Element : OtherSet)
+		{
+			Result.Add(Element);
+		}
+		return Result;
+	}
+
+	/** @return the complement of two sets. (A not in B where A is this and B is Other)*/
+	[[nodiscard]] TCompactSet Difference(const TCompactSet& OtherSet) const
+	{
+		TCompactSet Result;
+		Result.Reserve(this->NumElements); // Worst case is no elements of this are in Other
+
+		for (const ElementType& Element : *this)
+		{
+			if (!OtherSet.Contains(KeyFuncs::GetSetKey(Element)))
+			{
+				Result.Add(Element);
+			}
+		}
+		return Result;
+	}
+
+	/**
+	 * Determine whether the specified set is entirely included within this set
+	 * 
+	 * @param OtherSet	Set to check
+	 * 
+	 * @return True if the other set is entirely included in this set, false if it is not
+	 */
+	[[nodiscard]] bool Includes(const TCompactSet<ElementType, KeyFuncs, Allocator>& OtherSet) const
+	{
+		bool bIncludesSet = true;
+		if (OtherSet.NumElements <= this->NumElements)
+		{
+			for (const ElementType& Element : OtherSet)
+			{
+				if (!Contains(KeyFuncs::GetSetKey(Element)))
+				{
+					bIncludesSet = false;
+					break;
+				}
+			}
+		}
+		else
+		{
+			// Not possible to include if it is bigger than us
+			bIncludesSet = false;
+		}
+		return bIncludesSet;
+	}
+
+	/** @return a TArray that's a copy of all the elements in this set, prefer using ArrayView where possible. */
+	[[nodiscard]] TArray<ElementType> Array() const
+	{
+		return TArray<ElementType>(GetData(), this->NumElements);
+	}
+
+	/** @return a readonly TArrayView of the elements */
+	[[nodiscard]] TArrayView<const ElementType> ArrayView() const
+	{
+		return TArrayView<const ElementType>(GetData(), this->NumElements);
+	}
+
+	/** Creates an iterator for the contents of this set */
+	[[nodiscard]] FORCEINLINE TIterator CreateIterator()
+	{
+		return TIterator(*this);
+	}
+
+	/** Creates a const iterator for the contents of this set */
+	[[nodiscard]] FORCEINLINE TConstIterator CreateConstIterator() const
+	{
+		return TConstIterator(*this);
+	}
+
+	/**
+	* DO NOT USE DIRECTLY
+	* STL-like iterators to enable range-based for loop support.
+	*/
+	#if TARRAY_RANGED_FOR_CHECKS // Only use iterators with checks enabled to get modification check, otherwise it's just a waste of compile/runtime resources
+	using TRangedForIterator             = TCheckedPointerIterator<      ElementType, SizeType, false>;
+	using TRangedForConstIterator        = TCheckedPointerIterator<const ElementType, SizeType, false>;
+
+	[[nodiscard]] FORCEINLINE TRangedForIterator      begin() { return TRangedForIterator(this->NumElements, GetData()); }
+	[[nodiscard]] FORCEINLINE TRangedForConstIterator begin() const { return TRangedForConstIterator(this->NumElements, GetData()); }
+	[[nodiscard]] FORCEINLINE TRangedForIterator      end() { return TRangedForIterator(this->NumElements, GetData() + this->NumElements); }
+	[[nodiscard]] FORCEINLINE TRangedForConstIterator end() const { return TRangedForConstIterator(this->NumElements, GetData() + this->NumElements); }
+	#else
+	using TRangedForIterator = ElementType*;
+	using TRangedForConstIterator = const ElementType*;
+
+	[[nodiscard]] FORCEINLINE ElementType*       begin() { return GetData(); }
+	[[nodiscard]] FORCEINLINE const ElementType* begin() const { return GetData(); }
+	[[nodiscard]] FORCEINLINE ElementType*       end() { return GetData() + this->NumElements; }
+	[[nodiscard]] FORCEINLINE const ElementType* end() const { return GetData() + this->NumElements; }
+	#endif
+
+	// Sets are deliberately prevented from being hashed or compared, because this would hide potentially major performance problems behind default operations.
+	friend uint32 GetTypeHash(const TCompactSet& Set) = delete;
+	friend bool operator==(const TCompactSet&, const TCompactSet&) = delete;
+	friend bool operator!=(const TCompactSet&, const TCompactSet&) = delete;
+
+	void WriteMemoryImage(FMemoryImageWriter& Writer) const
+	{
+		checkf(!Writer.Is32BitTarget(), TEXT("TCompactSet does not currently support freezing for 32bits"));
+		if constexpr (TAllocatorTraits<Allocator>::SupportsFreezeMemoryImage && THasTypeLayout<InElementType>::Value)
+		{
+			if (this->NumElements > 0)
+			{
+				const ElementType* Data = GetData();
+				check(Data);
+				FMemoryImageWriter ArrayWriter = Writer.WritePointer(StaticGetTypeLayoutDesc<ElementType>());
+				ArrayWriter.WriteAlignment(GetAlignment());
+
+				// Write active element data
+				ArrayWriter.WriteObjectArray(Data, StaticGetTypeLayoutDesc<ElementType>(), this->NumElements);
+				ArrayWriter.WritePaddingToSize(this->MaxElements * sizeof(ElementType));
+
+				// Write remaining byte and hash table data
+				const FCompactSetLayout Layout = GetSetLayout();
+				const HashCountType* HashTable = this->GetHashTableMemory(Layout);
+				const uint8* HashTableDataEnd = (const uint8*)Data + Super::GetTotalMemoryRequiredInBytes(this->MaxElements, *HashTable, Layout);
+				ArrayWriter.WriteBytes(HashTable, HashTableDataEnd - (const uint8*)HashTable);
+
+				Writer.WriteBytes(this->NumElements);
+				Writer.WriteBytes(this->MaxElements);
+			}
+			else
+			{
+				Writer.WriteBytes(TCompactSet());
+			}
+		}
+		else
+		{
+			Writer.WriteBytes(TCompactSet());
+		}
+	}
+
+	void CopyUnfrozen(const FMemoryUnfreezeContent& Context, void* Dst) const
+	{
+		::new(Dst) TCompactSet();
+
+		if constexpr (TAllocatorTraits<Allocator>::SupportsFreezeMemoryImage && THasTypeLayout<InElementType>::Value)
+		{
+			const FTypeLayoutDesc& ElementTypeDesc = StaticGetTypeLayoutDesc<ElementType>();
+
+			TCompactSet* DstObject = static_cast<TCompactSet*>(Dst);
+
+			{
+				DstObject->ResizeAllocation(this->MaxElements);
+				DstObject->NumElements = this->NumElements;
+
+				const ElementType* SrcData = GetData();
+				ElementType* DstData = DstObject->GetData();
+
+				for (int32 Index = 0; Index < this->NumElements; ++Index)
+				{
+					Context.UnfreezeObject(SrcData + Index, ElementTypeDesc, DstData + Index);
+				}
+
+				const FCompactSetLayout Layout = GetSetLayout();
+				const HashCountType* SrcHashTable = this->GetHashTableMemory(Layout);
+				const uint8* SrcHashTableEnd = (const uint8 *)SrcData + Super::GetTotalMemoryRequiredInBytes(this->MaxElements, *SrcHashTable, Layout);
+
+				HashCountType* DstHashTable = (HashCountType*)DstObject->GetHashTableMemory(Layout);
+				FMemory::Memcpy(DstHashTable, SrcHashTable, SrcHashTableEnd - (const uint8*)SrcHashTable);
+			}
+		}
+	}
+
+	static void AppendHash(const FPlatformTypeLayoutParameters& LayoutParams, FSHA1& Hasher)
+	{
+		if constexpr (TAllocatorTraits<Allocator>::SupportsFreezeMemoryImage && THasTypeLayout<ElementType>::Value)
+		{
+			Freeze::AppendHash(StaticGetTypeLayoutDesc<ElementType>(), LayoutParams, Hasher);
+		}
+	}
+
+private:
+	[[nodiscard]] FORCEINLINE static constexpr FCompactSetLayout GetSetLayout()
+	{
+		return { sizeof(ElementType), GetAlignment() };
+	}
+
+	[[nodiscard]] FORCEINLINE FCompactHashTableView GetHashTableView()
+	{
+		return Super::GetHashTableView(GetSetLayout());
+	}
+
+	[[nodiscard]] FORCEINLINE FConstCompactHashTableView GetConstHashTableView() const
+	{
+		return Super::GetConstHashTableView(GetSetLayout());
+	}
+
+	// Use if you're going to reset/rehash regardless of the results
+	FORCEINLINE void ResizeAllocation(SizeType NewMaxElements)
+	{
+		(void)Super::ResizeAllocationPreserveData(NewMaxElements, GetSetLayout(), false);
+	}
+
+	// Use this if you'll be keeping the element data
+	[[nodiscard]] bool ResizeAllocationPreserveData(SizeType NewMaxElements, bool bPreserveHashData = true)
+	{
+		return Super::ResizeAllocationPreserveData(NewMaxElements, GetSetLayout(), bPreserveHashData);
+	}
+	
+	/** Recalculate the lookup table, used after bulk operations are performed */
+	void Rehash()
+	{
+		if (const ElementType* ElementData = GetData())
+		{
+			FCompactHashTableView HashTable = GetHashTableView();
+
+			HashTable.Reset();
+
+			for (int32 Index = 0; Index < this->NumElements; ++Index)
+			{
+				HashTable.Add(Index, KeyFuncs::GetKeyHash(KeyFuncs::GetSetKey(ElementData[Index])));
+			}
+		}
+	}
+
+	[[nodiscard]] ElementType& AddUninitialized(uint32 KeyHash)
+	{
+		checkSlow(this->MaxElements >= 0);
+		if (this->NumElements == this->MaxElements)
+		{
+			Reserve(this->AllocatorCalculateSlackGrow(this->NumElements + 1, GetSetLayout()));
+		}
+		GetHashTableView().Add(this->NumElements, KeyHash);
+		return GetData()[this->NumElements++];
+	}
+
+	/**
+		* Finds an element with a pre-calculated hash and a key that can be compared to KeyType
+		* @see	Class documentation section on ByHash() functions
+		* @return The element id that matches the key and hash or an invalid element id
+		*/
+	template<typename ComparableKey>
+	[[nodiscard]] SizeType FindIndexByHash(uint32 KeyHash, const ComparableKey& Key) const
+	{
+		if (this->NumElements == 0)
+		{
+			return INDEX_NONE;
+		}
+
+		const ElementType* ElementData = GetData();
+		const HashCountType* HashTable = this->GetHashTableMemory(GetSetLayout());
+		const uint8* NextIndicesData = (const uint8*)(HashTable + 1);
+		const uint32 HashCount = *HashTable;
+		
+		// Inlining this can save up to 40% perf compared to GetHashTableView().Find()
+		#define UE_COMPACTHASHTABLE_EXECUTEBYTYPE(Type) \
+				const Type* NextIndices = (const Type *)NextIndicesData; \
+				const Type* HashIndices = NextIndices + this->MaxElements; \
+				for (Type Index = HashIndices[KeyHash & (HashCount-1)]; Index != (Type)INDEX_NONE; Index = NextIndices[Index]) \
+				{ \
+					checkSlow((SizeType)Index < this->NumElements); \
+						if (KeyFuncs::Matches(KeyFuncs::GetSetKey(ElementData[Index]), Key)) \
+						{ \
+							return Index; \
+						} \
+				}
+		
+		UE_COMPACTHASHTABLE_CALLBYTYPE(this->MaxElements)
+		#undef UE_COMPACTHASHTABLE_EXECUTEBYTYPE
+
+		return INDEX_NONE;
+	}
+
+	template <bool IsStable = false>
+	void RemoveByIndex(const SizeType ElementIndex)
+	{
+		checkf(ElementIndex >= 0 && ElementIndex < this->NumElements, TEXT("Invalid ElementIndex passed to TCompactSet::RemoveByIndex"));
+		RemoveByIndexAndHash<IsStable>(ElementIndex, KeyFuncs::GetKeyHash(KeyFuncs::GetSetKey(GetData()[ElementIndex])));
+	}
+
+	template <bool IsStable = false>
+	void RemoveByIndexAndHash(const SizeType ElementIndex, const uint32 KeyHash)
+	{
+		checkf(ElementIndex >= 0 && ElementIndex < this->NumElements, TEXT("Invalid ElementIndex passed to TCompactSet::RemoveByIndex"));
+
+		ElementType* ElementsData = GetData();
+		FCompactHashTableView HashTable = GetHashTableView();
+
+		const SizeType LastElementIndex = this->NumElements - 1;
+		if (ElementIndex == LastElementIndex)
+		{
+			HashTable.Remove(ElementIndex, KeyHash, ElementIndex, KeyHash);
+			ElementsData[LastElementIndex].~ElementType();
+		}
+		else
+		{
+			if constexpr (IsStable)
+			{
+				HashTable.RemoveStable(ElementIndex, KeyHash);
+
+				ElementsData[ElementIndex].~ElementType();
+				RelocateConstructItems<ElementType>(ElementsData + ElementIndex, ElementsData + ElementIndex + 1, LastElementIndex - ElementIndex);
+			}
+			else
+			{
+				const uint32 LastElementHash = KeyFuncs::GetKeyHash(KeyFuncs::GetSetKey(ElementsData[LastElementIndex]));
+				HashTable.Remove(ElementIndex, KeyHash, LastElementIndex, LastElementHash);
+
+				MoveByRelocate(ElementsData[ElementIndex], ElementsData[LastElementIndex]);
+			}
+		}
+
+		--this->NumElements;
+	}
+
+	template<typename ComparableKey, bool IsStable = false>
+	int32 RemoveImpl(uint32 KeyHash, const ComparableKey& Key)
+	{
+		checkSlow(this->NumElements > 0);
+		int32 NumRemovedElements = 0;
+
+		const ElementType* ElementsData = GetData();
+		FCompactHashTableView HashTable = GetHashTableView();
+
+		SizeType LastElementIndex = INDEX_NONE;
+		SizeType ElementIndex = HashTable.GetFirst(KeyHash);
+
+		while (ElementIndex != INDEX_NONE)
+		{
+			if (KeyFuncs::Matches(KeyFuncs::GetSetKey(ElementsData[ElementIndex]), Key))
+			{
+				RemoveByIndexAndHash<IsStable>(ElementIndex, KeyHash);
+				NumRemovedElements++;
+
+				if constexpr (!KeyFuncs::bAllowDuplicateKeys)
+				{
+					// If the hash disallows duplicate keys, we're done removing after the first matched key.
+					break;
+				}
+				else
+				{
+					if (LastElementIndex == INDEX_NONE)
+					{
+						ElementIndex = HashTable.GetFirst(KeyHash);
+					}
+					else
+					{
+						if (LastElementIndex == this->NumElements) // Would have been remapped to ElementIndex
+						{
+							LastElementIndex = ElementIndex;
+						}
+
+						ElementIndex = HashTable.GetNext(LastElementIndex, this->NumElements);
+					}
+				}
+			}
+			else
+			{
+				LastElementIndex = ElementIndex;
+				ElementIndex = HashTable.GetNext(LastElementIndex, this->NumElements);
+			}
+		}
+
+		return NumRemovedElements;
+	}
+
+	template<typename ComparableKey>
+	FORCEINLINE int32 RemoveImplStable(uint32 KeyHash, const ComparableKey& Key)
+	{
+		return RemoveImpl<ComparableKey, true>(KeyHash, Key);
+	}
+};
+
+template<typename RangeType>
+TCompactSet(RangeType &&)->TCompactSet< TElementType_T<RangeType> >;
+
+namespace Freeze
+{
+	template<typename ElementType, typename KeyFuncs, typename Allocator>
+	void IntrinsicWriteMemoryImage(FMemoryImageWriter& Writer, const TCompactSet<ElementType, KeyFuncs, Allocator>& Object, const FTypeLayoutDesc&)
+	{
+		Object.WriteMemoryImage(Writer);
+	}
+
+	template<typename ElementType, typename KeyFuncs, typename Allocator>
+	uint32 IntrinsicUnfrozenCopy(const FMemoryUnfreezeContent& Context, const TCompactSet<ElementType, KeyFuncs, Allocator>& Object, void* OutDst)
+	{
+		Object.CopyUnfrozen(Context, OutDst);
+		return sizeof(Object);
+	}
+
+	template<typename ElementType, typename KeyFuncs, typename Allocator>
+	uint32 IntrinsicAppendHash(const TCompactSet<ElementType, KeyFuncs, Allocator>* DummyObject, const FTypeLayoutDesc& TypeDesc, const FPlatformTypeLayoutParameters& LayoutParams, FSHA1& Hasher)
+	{
+		TCompactSet<ElementType, KeyFuncs, Allocator>::AppendHash(LayoutParams, Hasher);
+		return DefaultAppendHash(TypeDesc, LayoutParams, Hasher);
+	}
+}
+
+DECLARE_TEMPLATE_INTRINSIC_TYPE_LAYOUT((template<typename ElementType, typename KeyFuncs, typename Allocator>), (TCompactSet<ElementType, KeyFuncs, Allocator>));
+
+struct TCompactSetPrivateFriend
+{
+	/** Serializer. */
+	template<typename ElementType, typename KeyFuncs, typename Allocator>
+	static FArchive& Serialize(FArchive& Ar, TCompactSet<ElementType, KeyFuncs, Allocator>& Set)
+	{
+		Set.CountBytes(Ar);
+
+		int32 NumElements = Set.Num();
+		Ar << NumElements;
+
+		if (Ar.IsLoading())
+		{
+			// We can skip the reset on Empty and do it once at the end with Rehash
+			DestructItems(Set.GetData(), Set.Num());
+			Set.NumElements = 0;
+			Set.ResizeAllocation(NumElements);
+
+			ElementType* Data = Set.GetData();
+			for (int32 ElementIndex = 0; ElementIndex < NumElements; ElementIndex++)
+			{
+				Ar << *::new((void*)(Data + ElementIndex)) ElementType;
+			}
+
+			Set.NumElements = NumElements;
+			Set.Rehash();
+		}
+		else
+		{
+			for (ElementType& Element : Set)
+			{
+				Ar << Element;
+			}
+		}
+		return Ar;
+	}
+
+	/** Structured archive serializer. */
+	template<typename ElementType, typename KeyFuncs, typename Allocator>
+	static void SerializeStructured(FStructuredArchive::FSlot Slot, TCompactSet<ElementType, KeyFuncs, Allocator>& Set)
+	{
+		int32 NumElements = Set.Num();
+		FStructuredArchive::FArray Array = Slot.EnterArray(NumElements);
+
+		if (Slot.GetUnderlyingArchive().IsLoading())
+		{
+			// We can skip the reset on Empty and do it once at the end with Rehash
+			DestructItems(Set.GetData(), Set.Num());
+			Set.NumElements = 0;
+			Set.ResizeAllocation(NumElements);
+
+			ElementType* Data = Set.GetData();
+			for (int32 ElementIndex = 0; ElementIndex < NumElements; ElementIndex++)
+			{
+				FStructuredArchive::FSlot ElementSlot = Array.EnterElement();
+				ElementSlot << *::new((void*)(Data + ElementIndex)) ElementType;
+			}
+
+			Set.NumElements = NumElements;
+			Set.Rehash();
+		}
+		else
+		{
+			for (ElementType& Element : Set)
+			{
+				FStructuredArchive::FSlot ElementSlot = Array.EnterElement();
+				ElementSlot << Element;
+			}
+		}
+	}
+
+	// Legacy comparison operators.  Note that these also test whether the set's elements were added in the same order!
+	template<typename ElementType, typename KeyFuncs, typename Allocator>
+	[[nodiscard]] static bool LegacyCompareEqual(const TCompactSet<ElementType, KeyFuncs, Allocator>& A, const TCompactSet<ElementType, KeyFuncs, Allocator>& B)
+	{
+		return A.Num() == B.Num() && CompareItems(A.GetData(), B.GetData(), A.Num());
+	}
+};
+
+/** Serializer. */
+template<typename ElementType, typename KeyFuncs, typename Allocator>
+FArchive& operator<<(FArchive& Ar, TCompactSet<ElementType, KeyFuncs, Allocator>& Set)
+{
+	return TCompactSetPrivateFriend::Serialize(Ar, Set);
+}
+
+/** Structured archive serializer. */
+template<typename ElementType, typename KeyFuncs, typename Allocator>
+void operator<<(FStructuredArchive::FSlot& Ar, TCompactSet<ElementType, KeyFuncs, Allocator>& Set)
+{
+	TCompactSetPrivateFriend::SerializeStructured(Ar, Set);
+}
+
+// Legacy comparison operators.  Note that these also test whether the set's elements were added in the same order!
+template<typename ElementType, typename KeyFuncs, typename Allocator>
+[[nodiscard]] bool LegacyCompareEqual(const TCompactSet<ElementType, KeyFuncs, Allocator>& A, const TCompactSet<ElementType, KeyFuncs, Allocator>& B)
+{
+	return TCompactSetPrivateFriend::LegacyCompareEqual(A, B);
+}
+template<typename ElementType, typename KeyFuncs, typename Allocator>
+[[nodiscard]] bool LegacyCompareNotEqual(const TCompactSet<ElementType, KeyFuncs, Allocator>& A, const TCompactSet<ElementType, KeyFuncs, Allocator>& B)
+{
+	return !TCompactSetPrivateFriend::LegacyCompareEqual(A, B);
+}
\ No newline at end of file
diff --git a/Engine/Source/Runtime/Core/Public/Containers/CompactSetBase.h b/Engine/Source/Runtime/Core/Public/Containers/CompactSetBase.h
new file mode 100644
--- /dev/null
+++ b/Engine/Source/Runtime/Core/Public/Containers/CompactSetBase.h
@@ -0,0 +1,328 @@
+// Copyright Epic Games, Inc. All Rights Reserved.
+
+#pragma once
+
+#include "Containers/ContainerAllocationPolicies.h"
+#include "Containers/CompactHashTable.h"
+#include "Templates/AlignmentTemplates.h"
+
+//
+// Compact set allocation definitions - It operates entirely with uint8 because it includes its lookup table inline with the overall allocation
+// Todo: move to ContainerAllocationPolicies.h
+//
+
+namespace UE::Core::Private::CompactSetAllocatorHelpers
+{
+	template <uint32 NumInlineElements, int32 ElementSize>
+	constexpr int32 CalculateRequiredBytes()
+	{
+		constexpr uint32 TypeSize = 1 + (NumInlineElements > 0xff) + (NumInlineElements > 0xffff) * 2;
+		constexpr uint32 HashSize = NumInlineElements < 8 ? 4 : FGenericPlatformMath::RoundUpToPowerOfTwo((NumInlineElements / 2) + 1);
+		return Align(NumInlineElements*ElementSize, 4) + 4 + (NumInlineElements + HashSize) * TypeSize;
+	}
+
+	/** Regular iterator */
+	template<typename InElementAllocator = FDefaultAllocator>
+	struct TCompactSetAllocator
+	{
+		/** CompactSets use raw uint8, so we need extra alignment information for inline allocators */
+		template <typename ElementType>
+		struct AllocatorAlignment
+		{
+			static constexpr size_t Value = alignof(typename InElementAllocator::template ForElementType<uint8>);
+		};
+
+		template<int32 ElementSize>
+		using ElementAllocator = InElementAllocator;
+	};
+
+	/** An inline allocator that allows sizing of the inline allocations for a set number of elements. */
+	template<uint32 NumInlineElements, typename SecondaryAllocator = TCompactSetAllocator<>>
+	struct TInlineCompactSetAllocator
+	{
+		template<int32 ElementSize>
+		using ElementAllocator = TInlineAllocator<UE::Core::Private::CompactSetAllocatorHelpers::CalculateRequiredBytes<NumInlineElements, ElementSize>(),typename SecondaryAllocator::template ElementAllocator<ElementSize>>;
+
+		/** CompactSets use raw uint8, so we need extra alignment information for inline allocators */
+		template <typename ElementType>
+		struct AllocatorAlignment
+		{
+			static constexpr size_t ElementAlignof = alignof(ElementType);
+			static constexpr size_t AllocatorAlignof = alignof(typename ElementAllocator<sizeof(ElementType)>::template ForElementType<uint8>);
+			static constexpr size_t Value = FGenericPlatformMath::Max(ElementAlignof, AllocatorAlignof);
+		};
+	};
+
+	/** An inline allocator that doesn't have any secondary storage. */
+	template <uint32 NumInlineElements>
+	struct TFixedCompactSetAllocator
+	{
+		template<int32 ElementSize>
+		using ElementAllocator = TFixedAllocator<UE::Core::Private::CompactSetAllocatorHelpers::CalculateRequiredBytes<NumInlineElements, ElementSize>()>;
+	
+		/** CompactSets use raw uint8, so we need extra alignment information for inline allocators */
+		template <typename ElementType>
+		struct AllocatorAlignment
+		{
+			static constexpr size_t ElementAlignof = alignof(ElementType);
+			static constexpr size_t AllocatorAlignof = alignof(typename ElementAllocator<sizeof(ElementType)>::template ForElementType<uint8>);
+			static constexpr size_t Value = FGenericPlatformMath::Max(ElementAlignof, AllocatorAlignof);
+		};
+	};
+}
+
+template<typename InElementAllocator>
+struct TAllocatorTraits<UE::Core::Private::CompactSetAllocatorHelpers::TCompactSetAllocator<InElementAllocator>> : TAllocatorTraitsBase<UE::Core::Private::CompactSetAllocatorHelpers::TCompactSetAllocator<InElementAllocator>>
+{
+	enum
+	{
+		SupportsFreezeMemoryImage = TAllocatorTraits<InElementAllocator>::SupportsFreezeMemoryImage,
+	};
+};
+
+class FDefaultCompactSetAllocator  : public UE::Core::Private::CompactSetAllocatorHelpers::TCompactSetAllocator<>  { public: typedef UE::Core::Private::CompactSetAllocatorHelpers::TCompactSetAllocator<>  Typedef; };
+
+// End Compact set allocation definitions
+
+// Todo: Move to ContainersFwd.h
+class FDefaultCompactSetAllocator;
+template<typename InElementType, typename KeyFuncs = DefaultKeyFuncs<InElementType>, typename Allocator = FDefaultCompactSetAllocator> class TCompactSet;
+// End Move to ContainersFwd.h
+
+/** Used to describe the data layout of the contents of the set */
+struct FCompactSetLayout
+{
+	int32 Size;
+	int32 Alignment;
+};
+
+/** Base class for the compact set provides common functionality to manage the data using FCompactSetLayout to describe it */
+template<typename Allocator>
+class TCompactSetBase
+{
+public:
+	using AllocatorType = Allocator;
+	using SizeType = typename AllocatorType::SizeType;
+
+	[[nodiscard]] bool operator==(FIntrusiveUnsetOptionalState Tag) const
+	{
+		return MaxElements == INDEX_NONE;
+	}
+
+	[[nodiscard]] FORCEINLINE bool IsEmpty() const
+	{
+		return NumElements == 0;
+	}
+
+	[[nodiscard]] FORCEINLINE int32 Num() const
+	{
+		return NumElements;
+	}
+
+	/** Deprecated - unnecessary, keeping this here so TCompactSet can be swapped with TSet without changing code */
+	[[nodiscard]] FORCEINLINE int32 GetMaxIndex() const
+	{
+		return NumElements;
+	}
+
+protected:
+	using HashCountType = uint32;
+	static constexpr size_t HashCountSize = sizeof(HashCountType);
+	using ElementAllocatorType = typename AllocatorType::template ForElementType<uint8>;
+
+	static_assert(std::is_same_v<SizeType, int32>, "TCompactSet currently only supports 32-bit allocators");
+	static_assert(sizeof(HashCountType) == UE::Core::CompactHashTable::GetMemoryAlignment(), "Hashtable alignment changed, need to make sure HashCountType is stored correctly");
+
+	[[nodiscard]] FORCEINLINE TCompactSetBase() = default;	
+
+	[[nodiscard]] explicit TCompactSetBase(FIntrusiveUnsetOptionalState Tag)
+	: NumElements(0)
+	, MaxElements(INDEX_NONE)
+	{
+	}
+
+	/* Calculate the size of the hash table from the number of elements in the set */
+	[[nodiscard]] FORCEINLINE const HashCountType* GetHashTableMemory(const FCompactSetLayout Layout) const
+	{
+		return (const HashCountType*)(Elements.GetAllocation() + GetElementsSizeInBytes(MaxElements, Layout));
+	}
+
+	[[nodiscard]] FORCEINLINE FCompactHashTableView GetHashTableView(const FCompactSetLayout Layout)
+	{
+		checkSlow(MaxElements > 0);
+		const HashCountType* HashTable = GetHashTableMemory(Layout);
+		return FCompactHashTableView((uint8*)(HashTable + 1), MaxElements, *HashTable, UE::Core::CompactHashTable::GetMemoryRequiredInBytes(MaxElements, *HashTable));
+	}
+
+	[[nodiscard]] FORCEINLINE FConstCompactHashTableView GetConstHashTableView(const FCompactSetLayout Layout) const
+	{
+		checkSlow(MaxElements > 0);
+		const HashCountType* HashTable = GetHashTableMemory(Layout);
+		return FConstCompactHashTableView((const uint8*)(HashTable + 1), MaxElements, *HashTable, UE::Core::CompactHashTable::GetMemoryRequiredInBytes(MaxElements, *HashTable));
+	}
+
+	/* Calculate the size of the hash table from the number of elements in the set */
+	[[nodiscard]] FORCEINLINE static constexpr SizeType GetHashCount(uint32 NumElements)
+	{
+		return UE::Core::CompactHashTable::GetHashCount(NumElements);
+	}
+
+	/* Calculate bytes required to store the elements, includes any padding required for the hash table */
+	[[nodiscard]] FORCEINLINE static constexpr size_t GetElementsSizeInBytes(uint32 NumElements, const FCompactSetLayout Layout)
+	{
+		return Align(Layout.Size * NumElements, UE::Core::CompactHashTable::GetMemoryAlignment());
+	}
+
+	/* Get the total memory required for the compact set for the given number of elements and size of the number of hash elements */
+	[[nodiscard]] FORCEINLINE static constexpr size_t GetTotalMemoryRequiredInBytes(uint32 NumElements, uint32 HashCount, const FCompactSetLayout Layout)
+	{
+		return NumElements ? GetElementsSizeInBytes(NumElements, Layout) + UE::Core::CompactHashTable::GetMemoryRequiredInBytes(NumElements, HashCount) + HashCountSize : 0;
+	}
+
+	/* Calculate the size of the hash table from the number of elements in the set assuming the default number of hash elements */
+	[[nodiscard]] FORCEINLINE static constexpr size_t GetTotalMemoryRequiredInBytes(uint32 NumElements, const FCompactSetLayout Layout)
+	{
+		return NumElements ? GetElementsSizeInBytes(NumElements, Layout) + UE::Core::CompactHashTable::GetMemoryRequiredInBytes(NumElements, GetHashCount(NumElements)) + HashCountSize : 0;
+	}
+
+	/* The total amount of memory required by the has set to store N elements with a hash table size of Y */
+	[[nodiscard]] static constexpr SizeType GetMaxElementsForAvailableSpace(size_t TotalBytes, uint32 HashCount, uint32 MinElementCount, const FCompactSetLayout Layout)
+	{
+		// Given some space in memory and a requested size for hash table, figure out how many elements we can fit in the remaining space
+		const uint32 TypeSize = UE::Core::CompactHashTable::GetTypeSize(MinElementCount);
+		const uint32 TypeShift = UE::Core::CompactHashTable::GetTypeShift(MinElementCount);
+		const size_t AvailableBytes = TotalBytes - sizeof(HashCountType) - ((size_t)HashCount << TypeShift); // Remove hashtable and HashCount data
+		const size_t MaxElements = AvailableBytes / (Layout.Size + TypeSize); // Calculate the max available ignoring the fact that the hash data has to be aligned
+		const size_t RealAvailableBytes = AlignDown(AvailableBytes - (MaxElements << TypeShift), UE::Core::CompactHashTable::GetMemoryAlignment()); // Remove the max required indexes and align down
+		return FMath::Min<SizeType>(MaxElements, RealAvailableBytes / Layout.Size); // Now we can get the true number of elements we could potentially use within the aligned space
+	}
+
+	[[nodiscard]] int32 AllocatorCalculateSlackGrow(int32 NewMaxElements, const FCompactSetLayout& Layout) const
+	{
+		const size_t OldHashCount = MaxElements > 0 ? *(uint32*)(Elements.GetAllocation() + GetElementsSizeInBytes(MaxElements, Layout)) : 0;
+		const size_t OldSize = MaxElements > 0 ? GetTotalMemoryRequiredInBytes(MaxElements, OldHashCount, Layout) : 0;
+
+		const size_t NewHashCount = NewMaxElements > 0 ? GetHashCount(NewMaxElements) : 0;
+		const size_t NewSize = NewMaxElements > 0 ? GetTotalMemoryRequiredInBytes(NewMaxElements, NewHashCount, Layout) : 0;
+		size_t NewSlackSize = 0;
+
+		if constexpr (TAllocatorTraits<AllocatorType>::SupportsElementAlignment)
+		{
+			NewSlackSize = Elements.CalculateSlackGrow(NewSize, OldSize, 1, Layout.Alignment);
+		}
+		else
+		{
+			NewSlackSize = Elements.CalculateSlackGrow(NewSize, OldSize, 1);
+		}
+
+		if (NewSlackSize == NewSize)
+		{
+			// No need for expensive slack calculations if there is none
+			return NewMaxElements;
+		}
+
+		// Calculate the number of items we can fit in the allotted space
+		SizeType SlackNumElements = GetMaxElementsForAvailableSpace(NewSlackSize, NewHashCount, NewMaxElements, Layout);
+		if (SlackNumElements <= NewMaxElements)
+		{
+			// This can happen if we're basically at the slack limit, when we got to calculate the results, even adding one element tracking index can
+			// cause us to lose space due to alignment changes, in this case we just take the user input as the source of truth
+			return NewMaxElements;
+		}
+
+		size_t SlackHashCount = GetHashCount(SlackNumElements);
+		if (SlackHashCount > NewHashCount)
+		{
+			// New Size required to much space. Cut the new hash count in half and restrict the elements to the max for that hash count
+			SlackNumElements = SlackHashCount - 1;
+			SlackHashCount /= 2;
+		}
+
+		checkSlow(SlackNumElements >= NewMaxElements);
+		checkSlow(GetTotalMemoryRequiredInBytes(SlackNumElements, SlackHashCount, Layout) <= NewSlackSize);
+
+		return SlackNumElements;
+	}
+
+	void ResizeAllocation(const int32 NewMaxElements, const FCompactSetLayout& Layout)
+	{
+		(void)ResizeAllocationPreserveData(NewMaxElements, Layout, false);
+	}
+
+	// Use this is you'll be keeping the element data
+	[[nodiscard]] bool ResizeAllocationPreserveData(const int32 NewMaxElements, const FCompactSetLayout& Layout, bool bPreserve = true)
+	{
+		bool bRequiresRehash = false;
+
+		if (NewMaxElements != MaxElements)
+		{
+			const size_t OldHashCount = MaxElements > 0 ? *(uint32*)(Elements.GetAllocation() + GetElementsSizeInBytes(MaxElements, Layout)) : 0;
+			const size_t OldSize = MaxElements > 0 ? GetTotalMemoryRequiredInBytes(MaxElements, OldHashCount, Layout) : 0;
+
+			const size_t NewHashCount = NewMaxElements > 0 ? UE::Core::CompactHashTable::GetHashCount(NewMaxElements) : 0;
+			const size_t NewSize = NewMaxElements > 0 ? GetTotalMemoryRequiredInBytes(NewMaxElements, NewHashCount, Layout) : 0;
+
+			if (bPreserve && NewMaxElements > MaxElements && OldHashCount == NewHashCount)
+			{
+				checkf(NewSize >= 0 && NewSize <= MAX_int32, TEXT("Invalid size for TSet[%d]: NewMaxElements[%d] ElementSize[%d] HashCount[%d]"), NewSize, NewMaxElements, Layout.Size, NewHashCount);
+
+				// If preserving then we copy over all the data
+				if constexpr (TAllocatorTraits<AllocatorType>::SupportsElementAlignment)
+				{
+					Elements.ResizeAllocation(OldSize, NewSize, 1, Layout.Alignment);
+				}
+				else
+				{
+					Elements.ResizeAllocation(OldSize, NewSize, 1);
+				}
+
+				const uint32 NewTypeShift = UE::Core::CompactHashTable::GetTypeShift(NewMaxElements);
+
+				// This should always be true since our type size will change on Powerof2 barriers (256, 65536, etc)
+				check(NewTypeShift == UE::Core::CompactHashTable::GetTypeShift(MaxElements));
+
+				const uint8* OldHashTable = Elements.GetAllocation() + GetElementsSizeInBytes(MaxElements, Layout);
+				uint8* NewHashTable = Elements.GetAllocation() + GetElementsSizeInBytes(NewMaxElements, Layout);
+
+				const uint8* OldHashLocation = OldHashTable + ((size_t)MaxElements << NewTypeShift) + HashCountSize;
+				uint8* NewHashLocation = NewHashTable + ((size_t)NewMaxElements << NewTypeShift) + HashCountSize;
+
+				// Copy hash lookup table (first so to free up the space for the other data to move)
+				FMemory::Memmove(NewHashLocation, OldHashLocation, (NewHashCount << NewTypeShift));
+
+				// Copy hash size + next index table
+				FMemory::Memmove(NewHashTable, OldHashTable, ((size_t)NumElements << NewTypeShift) + HashCountSize);
+			}
+			else
+			{
+				// If not preserving (or shrinking) then we only need to copy the array data
+				if constexpr (TAllocatorTraits<AllocatorType>::SupportsElementAlignment)
+				{
+					Elements.ResizeAllocation(NumElements * Layout.Size, NewSize, 1, Layout.Alignment);
+				}
+				else
+				{
+					Elements.ResizeAllocation(NumElements * Layout.Size, NewSize, 1);
+				}
+
+				// Update hash size
+				if (NewMaxElements > 0)
+				{
+					uint32* NewHashTable = (uint32*)(Elements.GetAllocation() + GetElementsSizeInBytes(NewMaxElements, Layout));
+
+					*NewHashTable = NewHashCount;
+					bRequiresRehash = true;
+				}
+			}
+
+			MaxElements = NewMaxElements;
+		}
+
+		return bRequiresRehash;
+	}
+
+	ElementAllocatorType Elements;
+
+	SizeType NumElements = 0;
+	SizeType MaxElements = 0;
+};
diff --git a/Engine/Source/Runtime/Core/Public/Containers/ScriptCompactSet.h b/Engine/Source/Runtime/Core/Public/Containers/ScriptCompactSet.h
new file mode 100644
--- /dev/null
+++ b/Engine/Source/Runtime/Core/Public/Containers/ScriptCompactSet.h
@@ -0,0 +1,243 @@
+// Copyright Epic Games, Inc. All Rights Reserved.
+
+#pragma once
+
+#include "ContainersFwd.h"
+#include "Containers/ContainerAllocationPolicies.h"
+#include "Containers/CompactSet.h"
+#include "Containers/SetUtilities.h"
+#include "Templates/Function.h"
+
+using FScriptCompactSetLayout = FCompactSetLayout;
+
+// Untyped set type for accessing TCompactSet data, like FScriptArray for TArray.
+// Must have the same memory representation as a TCompactSet.
+template <typename Allocator>
+class TScriptCompactSet : public TCompactSetBase<typename Allocator::template ElementAllocator<sizeof(uint8)>>
+{
+	using Super = TCompactSetBase<typename Allocator::template ElementAllocator<sizeof(uint8)>>;
+public:
+	[[nodiscard]] static FScriptCompactSetLayout GetScriptLayout(int32 ElementSize, int32 ElementAlignment)
+	{
+		return { ElementSize, FGenericPlatformMath::Max<int32>(ElementAlignment, UE::Core::CompactHashTable::GetMemoryAlignment()) };
+	}
+
+	//////////////////////////////////////////////////////////
+	// Start - intrusive TOptional<TScriptCompactSet> state //
+	//////////////////////////////////////////////////////////
+	constexpr static bool bHasIntrusiveUnsetOptionalState = true;
+	using IntrusiveUnsetOptionalStateType = TScriptCompactSet;
+
+	[[nodiscard]] explicit TScriptCompactSet(FIntrusiveUnsetOptionalState Tag) : Super(Tag)
+	{
+	}
+	////////////////////////////////////////////////////////
+	// End - intrusive TOptional<TScriptCompactSet> state //
+	////////////////////////////////////////////////////////
+
+	[[nodiscard]] TScriptCompactSet() = default;
+
+	[[nodiscard]] bool IsValidIndex(int32 Index) const
+	{
+		return Index >= 0 && Index < this->NumElements;
+	}
+
+	[[nodiscard]] int32 NumUnchecked() const
+	{
+		return this->NumElements;
+	}
+
+	[[nodiscard]] void* GetData(int32 Index, const FScriptCompactSetLayout& Layout)
+	{
+		return (uint8*)this->Elements.GetAllocation() + (Layout.Size * Index);
+	}
+
+	[[nodiscard]] const void* GetData(int32 Index, const FScriptCompactSetLayout& Layout) const
+	{
+		return (const uint8 *)this->Elements.GetAllocation() + (Layout.Size * Index);
+	}
+
+	void MoveAssign(TScriptCompactSet& Other, const FScriptCompactSetLayout& Layout)
+	{
+		checkSlow(this != &Other);
+
+		this->Elements.MoveToEmpty(Other.Elements);
+		this->NumElements = Other.NumElements;
+		this->MaxElements = Other.MaxElements;
+
+		Other.NumElements = 0;
+		Other.MaxElements = 0;
+	}
+
+	void Empty(int32 Slack, const FScriptCompactSetLayout& Layout)
+	{
+		this->ResizeAllocation(Slack, Layout);
+		if (this->MaxElements)
+		{
+			this->GetHashTableView(Layout).Reset();
+		}
+		this->NumElements = 0;
+	}
+
+	void RemoveAt(int32 Index, const FScriptCompactSetLayout& Layout, TFunctionRef<uint32 (const void*)> GetKeyHash)
+	{
+		check(IsValidIndex(Index));
+
+		const uint32 KeyHash = GetKeyHash(GetData(Index, Layout));
+
+		if (Index == this->NumElements - 1)
+		{
+			this->GetHashTableView(Layout).Remove(Index, KeyHash, Index, 0);
+		}
+		else
+		{
+			this->GetHashTableView(Layout).Remove(Index, KeyHash, this->NumElements - 1, GetKeyHash(GetData(this->NumElements - 1, Layout)));
+		}
+
+		--this->NumElements;
+	}
+
+	/**
+	 * Adds an uninitialized object to the set.
+	 * The set will need rehashing at some point after this call to make it valid.
+	 *
+	 * @return  The index of the added element.
+	 */
+	int32 AddUninitialized(const FScriptCompactSetLayout& Layout)
+	{
+		checkSlow(this->NumElements >= 0 && this->MaxElements >= INDEX_NONE);
+		if (this->NumElements == this->MaxElements)
+		{
+			this->ResizeAllocation(this->AllocatorCalculateSlackGrow(this->NumElements + 1, Layout), Layout);
+		}
+		++this->NumElements;
+		return this->NumElements - 1;
+	}
+
+	void RemoveAtUninitialized(const FScriptCompactSetLayout& Layout, int32 Index)
+	{
+		// Can only be pairs with AddUninitialized for now
+		check(Index == this->NumElements - 1);
+		--this->NumElements;
+	}
+
+	void CommitLastUninitialized(const FScriptCompactSetLayout& Layout, TFunctionRef<uint32 (const void*)> GetKeyHash)
+	{
+		const uint32 KeyHash = GetKeyHash(GetData(this->NumElements - 1, Layout));
+		this->GetHashTableView(Layout).Add(this->NumElements - 1, KeyHash);
+	}
+
+	void CommitAllUninitialized(const FScriptCompactSetLayout& Layout, TFunctionRef<uint32 (const void*)> GetKeyHash)
+	{
+		// Keep consistent interface with ScriptSparseSet
+	}
+
+	void Rehash(const FScriptCompactSetLayout& Layout, TFunctionRef<uint32 (const void*)> GetKeyHash)
+	{
+		if (this->MaxElements > 0)
+		{
+			const uint8* ElementData = this->Elements.GetAllocation();
+			FCompactHashTableView HashTable = this->GetHashTableView(Layout);
+
+			HashTable.Reset();
+
+			for (int32 Index = 0; Index < this->NumElements; ++Index)
+			{
+				HashTable.Add(Index, GetKeyHash(ElementData + Layout.Size * Index));
+			}
+		}
+	}
+
+	[[nodiscard]] int32 FindIndex(const void* Element, const FScriptCompactSetLayout& Layout, TFunctionRef<uint32 (const void*)> GetKeyHash, TFunctionRef<bool (const void*, const void*)> EqualityFn) const
+	{
+		if (this->NumElements)
+		{
+			return FindIndexImpl(Element, Layout, GetKeyHash(Element), EqualityFn);
+		}
+
+		return INDEX_NONE;
+	}
+
+	[[nodiscard]] int32 FindIndexByHash(const void* Element, const FScriptCompactSetLayout& Layout, uint32 KeyHash, TFunctionRef<bool (const void*, const void*)> EqualityFn) const
+	{
+		if (this->NumElements)
+		{
+			return FindIndexImpl(Element, Layout, KeyHash, EqualityFn);
+		}
+
+		return INDEX_NONE;
+	}
+
+	int32 FindOrAdd(const void* Element, const FScriptCompactSetLayout& Layout, TFunctionRef<uint32(const void*)> GetKeyHash, TFunctionRef<bool(const void*, const void*)> EqualityFn, TFunctionRef<void(void*)> ConstructFn)
+	{
+		uint32 KeyHash = GetKeyHash(Element);
+		int32 OldElementIndex = FindIndexByHash(Element, Layout, KeyHash, EqualityFn);
+		if (OldElementIndex != INDEX_NONE)
+		{
+			return OldElementIndex;
+		}
+
+		return AddNewElement(Layout, GetKeyHash, KeyHash, ConstructFn);
+	}
+
+	void Add(const void* Element, const FScriptCompactSetLayout& Layout, TFunctionRef<uint32(const void*)> GetKeyHash, TFunctionRef<bool(const void*, const void*)> EqualityFn, TFunctionRef<void(void*)> ConstructFn, TFunctionRef<void(void*)> DestructFn)
+	{
+		uint32 KeyHash = GetKeyHash(Element);
+		int32 OldElementIndex = FindIndexByHash(Element, Layout, KeyHash, EqualityFn);
+		if (OldElementIndex != INDEX_NONE)
+		{
+			void* ElementPtr = GetData(OldElementIndex, Layout);
+
+			DestructFn(ElementPtr);
+			ConstructFn(ElementPtr);
+
+			// We don't update the hash because we don't need to - the new element
+			// should have the same hash, but let's just check.
+			PRAGMA_DISABLE_DEPRECATION_WARNINGS
+			// Disable deprecations warnings to stop warnings being thrown by our check macro.
+			checkSlow(KeyHash == GetKeyHash(ElementPtr));
+			PRAGMA_ENABLE_DEPRECATION_WARNINGS
+		}
+		else
+		{
+			AddNewElement(Layout, GetKeyHash, KeyHash, ConstructFn);
+		}
+	}
+
+private:
+	[[nodiscard]] int32 FindIndexImpl(const void* Element, const FScriptCompactSetLayout& Layout, uint32 KeyHash, TFunctionRef<bool (const void*, const void*)> EqualityFn) const
+	{
+		return this->GetConstHashTableView(Layout).Find(KeyHash, this->NumElements, [Element, &Layout, &EqualityFn, this](uint32 Index) { return EqualityFn(Element, GetData(Index, Layout)); });
+	}
+
+	int32 AddNewElement(const FScriptCompactSetLayout& Layout, TFunctionRef<uint32(const void*)> GetKeyHash, uint32 KeyHash, TFunctionRef<void(void*)> ConstructFn)
+	{
+		checkSlow(this->NumElements >= 0 && this->MaxElements >= 0);
+		if (this->NumElements == this->MaxElements)
+		{
+			if (this->ResizeAllocationPreserveData(this->AllocatorCalculateSlackGrow(this->NumElements + 1, Layout), Layout))
+			{
+				Rehash(Layout, GetKeyHash);
+			}
+		}
+
+		this->GetHashTableView(Layout).Add(this->NumElements, KeyHash);
+
+		ConstructFn(GetData(this->NumElements, Layout));
+		return this->NumElements++;
+	}
+
+public:
+	// These should really be private, because they shouldn't be called, but there's a bunch of code
+	// that needs to be fixed first.
+	TScriptCompactSet(const TScriptCompactSet&) { check(false); }
+	void operator=(const TScriptCompactSet&) { check(false); }
+};
+
+template <typename AllocatorType>
+struct TIsZeroConstructType<TScriptCompactSet<AllocatorType>>
+{
+	enum { Value = true };
+};
+
+using FScriptCompactSet = TScriptCompactSet<FDefaultCompactSetAllocator>;
diff --git a/Engine/Source/Runtime/Core/Tests/Containers/CompactHashTableTests.cpp b/Engine/Source/Runtime/Core/Tests/Containers/CompactHashTableTests.cpp
new file mode 100644
--- /dev/null
+++ b/Engine/Source/Runtime/Core/Tests/Containers/CompactHashTableTests.cpp
@@ -0,0 +1,286 @@
+// Copyright Epic Games, Inc. All Rights Reserved.
+
+#if WITH_TESTS
+
+#include "Containers/CompactHashTable.h"
+#include "Tests/TestHarnessAdapter.h"
+
+struct FCommonCompactHashTableTests
+{
+	template<typename HashTableType>
+	static uint32 TestFindExpectedIndex(const HashTableType& HashTable, uint32 Key, uint32 Index, uint32 CurrentSize)
+	{
+		uint32 Steps = 0;
+
+		uint32 FindIndex = HashTable.GetFirst(Key);
+		while (FindIndex != Index && FindIndex != INDEX_NONE)
+		{
+			++Steps;
+			FindIndex = HashTable.GetNext(FindIndex, CurrentSize);
+		}
+
+		CHECK(FindIndex == Index);
+
+		return Steps;
+	}
+
+	template<uint32 Size, typename HashTableType>
+	static void TestAddRemove(HashTableType& HashTable)
+	{
+		uint32 Data[Size];
+		for (uint32 Index = 0; Index < Size; ++Index)
+		{
+			Data[Index] = (Index + 1) * Size - 7; // Math isn't important, it's unique and lets us hit collisions, that's the important part
+			HashTable.Add(Index, Data[Index]);
+		}
+
+		uint32 CurrentSize = Size;
+		uint32 Collisions = 0;
+		for (uint32 FullIndex = 0; FullIndex < Size; ++FullIndex)
+		{
+			// Pick a consistent random element to remove
+			const uint32 Index = Data[FullIndex % CurrentSize] % CurrentSize;
+
+			Collisions += TestFindExpectedIndex(HashTable, Data[Index], Index, CurrentSize);
+
+			const uint32 LastElement = --CurrentSize;
+			HashTable.Remove(Index, Data[Index], LastElement, Data[LastElement]);
+
+			if (Index != LastElement)
+			{
+				// User needs to also move data
+				Data[Index] = Data[LastElement];
+
+				Collisions += TestFindExpectedIndex(HashTable, Data[Index], Index, CurrentSize);
+			}
+		}
+
+		CHECK(Collisions > 0); // I'd expect at least a few, make sure we hit those
+	}
+
+	template<uint32 Size, typename HashTableType>
+	static void TestSynchedTables(HashTableType& HashTableKeys, HashTableType& HashTableValues)
+	{
+		struct Pair
+		{
+			uint32 Key;
+			uint32 Value;
+		};
+
+		Pair Data[Size];
+
+		for (uint32 Index = 0; Index < Size; ++Index)
+		{
+			Data[Index].Key = (Index + 1) * Size - 7;
+			Data[Index].Value = Data[Index].Key + 1;
+
+			HashTableKeys.Add(Index, Data[Index].Key);
+			HashTableValues.Add(Index, Data[Index].Value);
+		}
+
+		uint32 CurrentSize = Size;
+		uint32 Collisions = 0;
+
+		for (uint32 FullIndex = 0; FullIndex < Size; ++FullIndex)
+		{
+			// Pick a consistent random element to remove
+			const uint32 Index = Data[FullIndex % CurrentSize].Key % CurrentSize;
+
+			CHECK(TestFindExpectedIndex(HashTableKeys, Data[Index].Key, Index, CurrentSize) == TestFindExpectedIndex(HashTableValues, Data[Index].Value, Index, CurrentSize));
+
+			const uint32 LastElement = --CurrentSize;
+			HashTableKeys.Remove(Index, Data[Index].Key, LastElement, Data[LastElement].Key);
+			HashTableValues.Remove(Index, Data[Index].Value, LastElement, Data[LastElement].Value);
+
+			if (Index != LastElement)
+			{
+				// User needs to also move data
+				Data[Index] = Data[LastElement];
+
+				CHECK(TestFindExpectedIndex(HashTableKeys, Data[Index].Key, Index, CurrentSize) == TestFindExpectedIndex(HashTableValues, Data[Index].Value, Index, CurrentSize));
+			}
+		}
+	}
+
+
+};
+
+TEST_CASE_NAMED(FStaticCompactHashTableTest, "System::Core::Containers::CompactHashTable::Static", "[Core][Containers][CompactHashTable]")
+{
+	SECTION("Constructor Behavior")
+	{
+		static_assert(sizeof(TStaticCompactHashTable<16>::IndexType) == sizeof(uint8));
+		static_assert(sizeof(TStaticCompactHashTable<256>::IndexType) == sizeof(uint16));
+		static_assert(sizeof(TStaticCompactHashTable<65536>::IndexType) == sizeof(uint32));
+
+		static_assert(sizeof(TStaticCompactHashTable<16>) == (16 + 16) * sizeof(uint8));
+		static_assert(sizeof(TStaticCompactHashTable<16, 32>) == (16 + 32) * sizeof(uint8));
+		static_assert(sizeof(TStaticCompactHashTable<256>) == (256 + 256) * sizeof(uint16));
+		static_assert(sizeof(TStaticCompactHashTable<256, 512>) == (256 + 512) * sizeof(uint16));
+		static_assert(sizeof(TStaticCompactHashTable<65536>) == (65536 + 65536) * sizeof(uint32));
+		static_assert(sizeof(TStaticCompactHashTable<65536, 65537>) == (65536 + 65537) * sizeof(uint32));
+
+		const uint8 Sentinal = 0xf0;
+		const uint8 Clear = 0xff;
+		
+		auto TestTableType = []<uint32 Size, uint32 HashSize = Size>()
+		{
+			using TestType = TStaticCompactHashTable<Size, HashSize>;
+
+			uint8 TestMemory[sizeof(TestType)];
+			FMemory::Memset(TestMemory, Sentinal);
+
+			new (TestMemory) TestType(NoInit);
+
+			for (uint8 Test : TestMemory)
+			{
+				CHECK(Test == Sentinal);
+			}
+
+			new (TestMemory) TestType();
+
+			for (int32 Index = 0; Index < sizeof(TestType); ++Index)
+			{
+				if (Index < Size * sizeof(typename TestType::IndexType))
+				{
+					CHECK(TestMemory[Index] == Sentinal);
+				}
+				else
+				{
+					CHECK(TestMemory[Index] == Clear);
+				}
+			}
+		};
+
+		TestTableType.operator()<16>();
+		TestTableType.operator()<16, 32>();
+		TestTableType.operator()<256>();
+		TestTableType.operator()<255, 128>();
+	}
+
+	SECTION("Add Behavior")
+	{
+		auto TestTableType = []<uint32 Size, uint32 HashSize = Size>()
+		{
+			TStaticCompactHashTable<Size, HashSize> HashTable;
+			FCommonCompactHashTableTests::TestAddRemove<Size>(HashTable);
+		};
+
+		TestTableType.operator()<16>();
+		TestTableType.operator()<16, 32>();
+		TestTableType.operator()<256>();
+		TestTableType.operator()<255, 128>();
+	}
+
+	SECTION("Synched Tables Behavior")
+	{
+		auto TestTableType = []<uint32 Size, uint32 HashSize = Size>()
+		{
+			TStaticCompactHashTable<Size, HashSize> HashTableKeys;
+			TStaticCompactHashTable<Size, HashSize> HashTableValues;
+			FCommonCompactHashTableTests::TestSynchedTables<Size>(HashTableKeys, HashTableValues);
+		};
+
+		TestTableType.operator()<16>();
+		TestTableType.operator()<16, 32>();
+		TestTableType.operator()<256>();
+		TestTableType.operator()<255, 128>();
+	}
+}
+
+TEST_CASE_NAMED(FDynamicCompactHashTableTest, "System::Core::Containers::CompactHashTable::Dynamic", "[Core][Containers][CompactHashTable]")
+{
+	SECTION("Constructor Behavior")
+	{
+		const uint8 Sentinal = 0xf0;
+		const uint8 Clear = 0xff;
+
+		auto TestTableType = []<uint32 Size, uint32 HashSize = Size>()
+		{
+			using IndexType = typename TStaticCompactHashTable<Size>::IndexType;
+			using TestType = FCompactHashTableView;
+
+			constexpr size_t RequiredSize = UE::Core::CompactHashTable::GetMemoryRequiredInBytes(Size, HashSize);
+			static_assert(RequiredSize == sizeof(typename TStaticCompactHashTable<Size>::IndexType) * (Size + HashSize));
+
+			uint8 TestMemory[RequiredSize];
+			FMemory::Memset(TestMemory, Sentinal, RequiredSize);
+
+			TestType HashTable(TestMemory, Size, HashSize, RequiredSize);
+
+			for (uint8 Test : TestMemory)
+			{
+				CHECK(Test == Sentinal);
+			}
+
+			HashTable.Reset();
+			
+			for (int32 Index = 0; Index < RequiredSize; ++Index)
+			{
+				if (Index < Size * sizeof(IndexType))
+				{
+					CHECK(TestMemory[Index] == Sentinal);
+				}
+				else
+				{
+					CHECK(TestMemory[Index] == Clear);
+				}
+			}
+		};
+
+		TestTableType.operator()<16>();
+		TestTableType.operator()<16, 32>();
+		TestTableType.operator()<256>();
+		TestTableType.operator()<255, 128>();
+	}
+
+	SECTION("Add Behavior")
+	{
+		auto TestTableType = []<uint32 Size, uint32 HashSize = Size>()
+		{
+			using IndexType = typename TStaticCompactHashTable<Size>::IndexType;
+			using TestType = FCompactHashTableView;
+
+			constexpr size_t RequiredSize = UE::Core::CompactHashTable::GetMemoryRequiredInBytes(Size, HashSize);
+
+			uint8 TestMemory[RequiredSize];
+			TestType HashTable(TestMemory, Size, HashSize, RequiredSize);
+			HashTable.Reset();
+
+			FCommonCompactHashTableTests::TestAddRemove<Size>(HashTable);
+		};
+
+		TestTableType.operator()<16>();
+		TestTableType.operator()<16, 32>();
+		TestTableType.operator()<256>();
+		TestTableType.operator()<255, 128>();
+	}
+
+	SECTION("Synched Tables Behavior")
+	{
+		auto TestTableType = []<uint32 Size, uint32 HashSize = Size>()
+		{
+			using IndexType = typename TStaticCompactHashTable<Size>::IndexType;
+			using TestType = FCompactHashTableView;
+
+			constexpr size_t RequiredSize = UE::Core::CompactHashTable::GetMemoryRequiredInBytes(Size, HashSize);
+
+			uint8 TestMemoryKeys[RequiredSize];
+			TestType HashTableKeys(TestMemoryKeys, Size, HashSize, RequiredSize);
+			HashTableKeys.Reset();
+
+			uint8 TestMemoryValues[RequiredSize];
+			TestType HashTableValues(TestMemoryValues, Size, HashSize, RequiredSize);
+			HashTableValues.Reset();
+
+			FCommonCompactHashTableTests::TestSynchedTables<Size>(HashTableKeys, HashTableValues);
+		};
+
+		TestTableType.operator()<16>();
+		TestTableType.operator()<16, 32>();
+		TestTableType.operator()<256>();
+		TestTableType.operator()<255, 128>();
+	}
+}
+
+#endif // WITH_TESTS
diff --git a/Engine/Source/Runtime/Core/Tests/Containers/SetTests.cpp b/Engine/Source/Runtime/Core/Tests/Containers/SetTests.cpp
new file mode 100644
--- /dev/null
+++ b/Engine/Source/Runtime/Core/Tests/Containers/SetTests.cpp
@@ -0,0 +1,579 @@
+// Copyright Epic Games, Inc. All Rights Reserved.
+
+#if WITH_TESTS
+
+#include "Algo/IndexOf.h"
+#include "Algo/RandomShuffle.h"
+#include "Containers/CompactSet.h"
+#include "Containers/Set.h"
+#include "Containers/ScriptCompactSet.h"
+#include "Containers/Map.h"
+#include "Tests/TestHarnessAdapter.h"
+
+#define WITH_SET_BENCHMARKS 0 && WITH_LOW_LEVEL_TESTS
+
+#if WITH_SET_BENCHMARKS
+#include <catch2/benchmark/catch_benchmark.hpp>
+#endif
+
+struct FSparseSetResolver
+{
+	template<typename InElementType, typename KeyFuncs = DefaultKeyFuncs<InElementType>, typename Allocator = FDefaultSetAllocator>
+	using SetType = TSet<InElementType, KeyFuncs, Allocator>;
+
+	template <typename Type, int32 Size>
+	using FixedSetType = TSet<Type, DefaultKeyFuncs<Type>, TFixedSetAllocator<Size>>;
+};
+
+struct FCompactSetResolver
+{
+	template<typename InElementType, typename KeyFuncs = DefaultKeyFuncs<InElementType>, typename Allocator = FDefaultCompactSetAllocator>
+	using SetType = TCompactSet<InElementType, KeyFuncs, Allocator>;
+
+	template <typename Type, int32 Size>
+	using FixedSetType = TCompactSet<Type, DefaultKeyFuncs<Type>, UE::Core::Private::CompactSetAllocatorHelpers::TFixedCompactSetAllocator<Size>>;
+};
+
+template<typename Type, typename SetType>
+void SetTestExists(const Type& Key, const SetType& Set)
+{
+	const Type* Value = Set.Find(Key);
+	CHECK(Value);
+	if (Value)
+	{
+		CHECK(*Value == Key);
+	}
+}
+
+struct FSetTestNonCopyableTestData
+{
+	FSetTestNonCopyableTestData() = default;
+	FSetTestNonCopyableTestData(uint32 V)
+	: Value(V)
+	{
+	}
+
+	// Not using noncopyable helpers since we're explicitly want to catch incorrect copy operations
+	FSetTestNonCopyableTestData(const FSetTestNonCopyableTestData&) = delete;
+	FSetTestNonCopyableTestData& operator=(const FSetTestNonCopyableTestData&) = delete;
+
+	// But allow for proper use of move operations
+	FSetTestNonCopyableTestData(FSetTestNonCopyableTestData&&) = default;
+	FSetTestNonCopyableTestData& operator=(FSetTestNonCopyableTestData&&) = default;
+	
+	bool operator==(const FSetTestNonCopyableTestData& Other) const
+	{
+		return Value == Other.Value;
+	}
+
+	uint32 Value = 0;
+};
+
+uint32 GetTypeHash(const FSetTestNonCopyableTestData& A)
+{
+	return A.Value;
+}
+
+
+template<typename SetResolver = FCompactSetResolver>
+struct TSetTestHelper
+{
+	template <typename Type>
+	using SetType = typename SetResolver::template SetType<Type>;
+
+	template<typename T>
+	using DuplicateKeyFunc = DefaultKeyFuncs<T, true>;
+
+	template <typename Type>
+	using DuplicateSetType = typename SetResolver::template SetType<Type, DuplicateKeyFunc<Type>>;
+
+	template <typename Type, int32 Size>
+	using FixedSetType = typename SetResolver::template FixedSetType<Type, Size>;
+
+	TSetTestHelper()
+	{
+		SECTION("Default Constructor")
+		{
+			SetType<int32> Set;
+			CHECK(Set.Num() == 0);
+		}
+
+		SECTION("Basic Tests")
+		{
+			FString StringTestItems[] = { TEXT("Zero"), TEXT("One"), TEXT("Two"), TEXT("Three") };
+			SetType<FString> Set( { StringTestItems[0], StringTestItems[1], StringTestItems[2], StringTestItems[3] });
+			CHECK(Set.Num() == 4);
+
+			for (int32 Index = 0; Index < UE_ARRAY_COUNT(StringTestItems); ++Index)
+			{
+				CHECK(Set[FSetElementId::FromInteger(Index)] == StringTestItems[Index]);
+				SetTestExists(StringTestItems[Index], Set);
+			}
+
+			CHECK(Set.Find( { TEXT("Invalid") }) == nullptr);
+			CHECK(Set.Find({}) == nullptr);
+		}
+
+		SECTION("NonCopyable Tests")
+		{
+			SetType<FSetTestNonCopyableTestData> Set;
+			for (uint32 Index = 0; Index < 10; ++Index)
+			{
+				Set.Emplace(Index);
+
+				FSetTestNonCopyableTestData Tmp { Index };
+				Set.Emplace(MoveTemp(Tmp));
+			}
+		}
+
+		SECTION("Advanced Tests")
+		{
+			int32 IntTestItems[256];
+
+			{
+				int32 IntInit = 1;
+				for (int32& Int : IntTestItems)
+				{
+					Int = HashCombine(IntInit++, 0);
+				}
+			}
+
+			SetType<int32> Set1(IntTestItems);
+			SetType<int32> Set2;
+
+			for (int32 Item : ReverseIterate(IntTestItems))
+			{
+				Set2.Add(Item);
+			}
+
+			for (int32 TestItem : IntTestItems)
+			{
+				SetTestExists(TestItem, Set1);
+				SetTestExists(TestItem, Set2);
+			}
+
+			CHECK(Set1.Includes(Set2));
+			CHECK(Set2.Includes(Set1));
+			CHECK(Set1.Difference(Set2).IsEmpty());
+
+			CHECK(Set1.Union(Set2).Difference(Set2).IsEmpty());
+			CHECK(Set1.Intersect(Set2).Difference(Set2).IsEmpty());
+
+			SetType<int32> Set3;
+			for (int32 Index = 0; Index < UE_ARRAY_COUNT(IntTestItems) / 2; ++Index)
+			{
+				Set3.Add(IntTestItems[Index]);
+			}
+
+			CHECK(Set1.Difference(Set3).Num() == UE_ARRAY_COUNT(IntTestItems) / 2);
+			CHECK(Set1.Intersect(Set3).Difference(Set3).IsEmpty());
+			CHECK(Set1.Difference(Set3).Difference(Set2.Difference(Set3)).IsEmpty());
+
+			CHECK(Set1.Includes(Set3));
+			CHECK(!Set3.Includes(Set1));
+		}
+
+		SECTION("Copy/Move")
+		{
+			FString StringTestItems[] = { TEXT("Zero"), TEXT("One"), TEXT("Two"), TEXT("Three") };
+			SetType<FString> Set1( { StringTestItems[0], StringTestItems[1], StringTestItems[2], StringTestItems[3] });
+			SetType<FString> Set2(MoveTemp(Set1));
+
+			CHECK(Set1.IsEmpty());
+			CHECK(Set1.Difference(Set2).IsEmpty());
+		}
+
+		SECTION("Growth Tests")
+		{
+			SetType<int32> Set;
+
+			int32 ValuesToAdd[66000]; // include at least one hash size growth
+			int32 IndexToAdd = 0;
+			for (int32& Value : ValuesToAdd)
+			{
+				IndexToAdd += 1 + FMath::RandHelper(4);
+				Value = IndexToAdd;
+			}
+			Algo::RandomShuffle(ValuesToAdd);
+
+			for (int32 Value : ValuesToAdd)
+			{
+				CHECK(Set.Contains(Value) == false);
+				Set.Add(Value);
+				CHECK(Set.Contains(Value) == true);
+			}
+
+			for (int32 Value : ValuesToAdd)
+			{
+				CHECK(Set.Contains(Value) == true);
+			}
+		}
+
+		SECTION("Iteration")
+		{
+			int32 IntTestItems[256];
+
+			{
+				int32 IntInit = 1;
+				for (int32& Int : IntTestItems)
+				{
+					Int = HashCombine(IntInit++, 0);
+				}
+			}
+
+			SetType<int32> Set(IntTestItems);
+
+			for (int32 Index = 0; Index < Set.Num(); ++Index)
+			{
+				CHECK(IntTestItems[Index] == Set[FSetElementId::FromInteger(Index)]);
+			}
+
+			int32 RemoveIndex = 0;
+			for (auto Iter = Set.CreateIterator(); Iter; ++Iter, ++RemoveIndex)
+			{
+				if (RemoveIndex % 2)
+				{
+					Iter.RemoveCurrent();
+				}
+			}
+
+			CHECK(Set.Num() == UE_ARRAY_COUNT(IntTestItems) / 2);
+		}
+
+		SECTION("Fixed Allocator")
+		{
+			FixedSetType<int32, 64> Set;
+
+			for (int32 Index = 0; Index < 64; ++Index)
+			{
+				Set.Add(Index);
+			}
+		}
+
+		SECTION("Duplicate Support")
+		{
+			FString StringTestItems[] = { TEXT("Zero"), TEXT("One"), TEXT("Two"), TEXT("Three") };
+			DuplicateSetType<FString> Set;
+
+			for (int32 Index = 0; Index < 32; ++Index)
+			{
+				for (const FString& TestItem : StringTestItems)
+				{
+					Set.Add(TestItem);
+				}
+			}
+
+			CHECK(Set.Num() == 32 * 4);
+
+			{
+				int32 KeyIteratorIndex = 0;
+				for (typename DuplicateSetType<FString>::TConstKeyIterator KeyIter(Set, StringTestItems[1]); KeyIter; ++KeyIter, ++KeyIteratorIndex)
+				{
+					CHECK(KeyIter.GetId().AsInteger() % 4 == 1);
+					CHECK(*KeyIter == StringTestItems[1]);
+				}
+
+				CHECK(KeyIteratorIndex == 32);
+			}
+
+			Set.Sort([&StringTestItems](const FString& Lhs, const FString& Rhs)
+				{
+					return Algo::IndexOf(StringTestItems, Lhs) < Algo::IndexOf(StringTestItems, Rhs);
+				});
+
+			{
+				int32 KeyIteratorIndex = 0;
+				for (typename DuplicateSetType<FString>::TConstKeyIterator KeyIter(Set, StringTestItems[1]); KeyIter; ++KeyIter, ++KeyIteratorIndex)
+				{
+					CHECK((KeyIter.GetId().AsInteger() >= 32 && KeyIter.GetId().AsInteger() < 64));
+					CHECK(*KeyIter == StringTestItems[1]);
+				}
+				CHECK(KeyIteratorIndex == 32);
+			}
+
+			// Remove half of items
+			{
+				int32 KeyIteratorIndex = 0;
+				for (typename DuplicateSetType<FString>::TKeyIterator KeyIter(Set, StringTestItems[1]); KeyIter; ++KeyIter, ++KeyIteratorIndex)
+				{
+					if ((KeyIteratorIndex % 2) == 0)
+					{
+						// Sorting breaks here
+						KeyIter.RemoveCurrent();
+					}
+				}
+				CHECK(KeyIteratorIndex == 32);
+			}
+
+			// Check remaining
+			{
+				int32 KeyIteratorIndex = 0;
+				for (typename DuplicateSetType<FString>::TConstKeyIterator KeyIter(Set, StringTestItems[1]); KeyIter; ++KeyIter, ++KeyIteratorIndex)
+				{
+					// Items will still fall in the old range since they remained in place while items from the end filled the gap
+					CHECK((KeyIter.GetId().AsInteger() >= 32 && KeyIter.GetId().AsInteger() < 64));
+					CHECK(*KeyIter == StringTestItems[1]);
+				}
+				CHECK(KeyIteratorIndex == 16);
+			}
+
+			CHECK(Set.Remove(StringTestItems[1]) == 16);
+			CHECK(Set.Num() == 32 * 3);
+			CHECK(Set.Find(StringTestItems[1]) == nullptr);
+		}
+
+		SECTION("Optional Support")
+		{
+			TOptional<SetType<int32>> OptionalSet;
+
+			const SetType<int32>* OptionalSetInternals = (const SetType<int32>*)&OptionalSet;
+			const FIntrusiveUnsetOptionalState& OptionalCheck = *(FIntrusiveUnsetOptionalState*)&OptionalSet; // hack to access this type for internal comparison
+			CHECK(OptionalSetInternals->Num() == 0);
+			CHECK(*OptionalSetInternals == OptionalCheck);
+
+			OptionalSet.Emplace();
+			OptionalSet.GetValue().Add(1);
+			OptionalSet.GetValue().Add(2);
+			CHECK(OptionalSet.GetValue().Num() == 2);
+			CHECK(!(*OptionalSetInternals == OptionalCheck));
+
+			OptionalSet.Reset();
+			CHECK(OptionalSetInternals->Num() == 0);
+			CHECK(*OptionalSetInternals == OptionalCheck);
+		}
+
+#if WITH_SET_BENCHMARKS
+		SECTION("Benchmark")
+		{
+			BENCHMARK("Add (With Alloc)")
+			{
+				SetType<int32> Set;
+
+				for (int32 Index = 0; Index < 10000; ++Index)
+				{
+					Set.Add(Index * 1023);
+				}
+			};
+
+			BENCHMARK("Add (Without Alloc)")
+			{
+				SetType<int32> Set;
+				Set.Reserve(1000000);
+
+				for (int32 Index = 0; Index < 100000; ++Index)
+				{
+					Set.Add(Index * 1023);
+				}
+			};
+
+			BENCHMARK_ADVANCED("Remove (Forward)")(Catch::Benchmark::Chronometer meter)
+			{
+				SetType<int32> Set;
+				Set.Reserve(10000);
+
+				for (int32 Index = 0; Index < 10000; ++Index)
+				{
+					Set.Add(Index * 1023);
+				}
+
+				meter.measure([&]
+					{
+						for (int32 Index = 0; Index < 10000; ++Index)
+						{
+							Set.Remove(Index * 1023);
+						}
+					});
+			};
+
+			BENCHMARK_ADVANCED("Remove (Reverse)")(Catch::Benchmark::Chronometer meter)
+			{
+				SetType<int32> Set;
+				Set.Reserve(10000);
+
+				for (int32 Index = 0; Index < 10000; ++Index)
+				{
+					Set.Add(Index * 1023);
+				}
+
+				meter.measure([&]
+					{
+						for (int32 Index = 10000 - 1; Index >= 0; --Index)
+						{
+							Set.Remove(Index * 1023);
+						}
+					});
+			};
+
+			BENCHMARK_ADVANCED("Lookup")(Catch::Benchmark::Chronometer meter)
+			{
+				SetType<int32> Set;
+				Set.Reserve(10000);
+
+				for (int32 Index = 0; Index < 10000; ++Index)
+				{
+					Set.Add(Index * 1023);
+				}
+
+				meter.measure([&]
+					{
+						int32 Lookups = 0;
+						for (int32 Index = 10000 - 1; Index >= 0; --Index)
+						{
+							Lookups += *Set.Find(Index * 1023) == Index * 1023;
+						}
+						return Lookups;
+					});
+			};
+
+			BENCHMARK_ADVANCED("Lookup Small")(Catch::Benchmark::Chronometer meter)
+			{
+				SetType<int32> Set;
+				Set.Reserve(128);
+
+				for (int32 Index = 0; Index < 128; ++Index)
+				{
+					Set.Add(Index * 1023);
+				}
+
+				meter.measure([&]
+					{
+						int32 Lookups = 0;
+						for (int32 Index = 0; Index < 10000; ++Index)
+						{
+							Lookups += *Set.Find((Index % 128) * 1023) == (Index % 128) * 1023;
+						}
+						return Lookups;
+					});
+			};
+		}
+#endif
+	}
+};
+
+TEST_CASE_NAMED(FCompactSetTest, "System::Core::Containers::CompactSet", "[Core][Containers][CompactSet]")
+{
+	static_assert(TCompactSet<FString>::GetTotalMemoryRequiredInBytes(97) == UE::Core::Private::CompactSetAllocatorHelpers::CalculateRequiredBytes<97, sizeof(FString)>());
+
+	TSetTestHelper<FCompactSetResolver>();
+
+	// Todo: Move into TSetTestHelper when Stable moved to TSparseSet
+	SECTION("Stability Support")
+	{
+		TArray<FString> StringTestItems = { TEXT("Zero"), TEXT("One"), TEXT("Two"), TEXT("Three") };
+		TSetTestHelper<FCompactSetResolver>::DuplicateSetType<FString> Set(StringTestItems);
+
+		Set.RemoveStable(TEXT("One"));
+		StringTestItems.RemoveAt(1);
+
+		{
+			auto Itr = Set.CreateConstIterator();
+			for (const FString& TestString : StringTestItems)
+			{
+				CHECK(TestString == *Itr);
+				++Itr;
+			}
+		}
+
+		Set.RemoveStable(FSetElementId::FromInteger(0));
+		StringTestItems.RemoveAt(0);
+
+		{
+			auto Itr = Set.CreateConstIterator();
+			for (const FString& TestString : StringTestItems)
+			{
+				CHECK(TestString == *Itr);
+				++Itr;
+			}
+		}
+	}
+}
+
+TEST_CASE_NAMED(FSparseSetTest, "System::Core::Containers::SparseSet", "[Core][Containers][CompactSet]")
+{
+	TSetTestHelper<FSparseSetResolver>();
+}
+
+TEST_CASE_NAMED(FScriptCompactSetTest, "System::Core::Containers::ScriptCompactSet", "[Core][Containers][CompactSet]")
+{
+	SECTION("Basic Tests")
+	{
+		FScriptCompactSet Set;
+		FScriptCompactSetLayout Layout = { sizeof(FString), alignof(FString) };
+
+		FString StringTestItems[] = { TEXT("Zero"), TEXT("One"), TEXT("Two"), TEXT("Three") };
+		auto GetKeyHash = [](const void* Data) { return GetTypeHash( * (const FString*)Data); };
+		auto EqualityFn = [](const void* Lhs, const void* Rhs) { return * (const FString*)Lhs == * (const FString*)Rhs; };
+		auto DestructFn = [](void* Data) { ((FString*)Data)->~FString(); };
+
+		for (const FString& TestItem : StringTestItems)
+		{
+			Set.Add(&TestItem, Layout, GetKeyHash, EqualityFn, [&TestItem](void* Data) { new (Data)FString(TestItem); }, DestructFn);
+		}
+
+		CHECK(Set.Num() == 4);
+
+		for (int32 Index = 0; Index < UE_ARRAY_COUNT(StringTestItems); ++Index)
+		{
+			CHECK(EqualityFn(Set.GetData(Index, Layout), &StringTestItems[Index]));
+
+			int32 FindIndex = Set.FindIndex(&StringTestItems[Index], Layout, GetKeyHash, EqualityFn);
+			CHECK(FindIndex == Index);
+		}
+
+		FString Invalid;
+		CHECK(Set.FindIndex(&Invalid, Layout, GetKeyHash, EqualityFn) == INDEX_NONE);
+		Invalid = TEXT("Invalid");
+		CHECK(Set.FindIndex(&Invalid, Layout, GetKeyHash, EqualityFn) == INDEX_NONE);
+
+		for (int32 Index = 0; Index < UE_ARRAY_COUNT(StringTestItems); ++Index)
+		{
+			DestructFn(Set.GetData(Index, Layout));
+		}
+	}
+
+	SECTION("Advanced Tests")
+	{
+		int32 IntTestItems[256];
+
+		{
+			int32 IntInit = 1;
+			for (int32& Int : IntTestItems)
+			{
+				Int = HashCombine(IntInit++, 0);
+			}
+		}
+
+		FScriptCompactSet Set;
+		FScriptCompactSetLayout Layout = { sizeof(int32), alignof(int32) };
+
+		auto GetKeyHash = [](const void* Data) { return * (const int32 *)Data; };
+		auto EqualityFn = [](const void* Lhs, const void* Rhs) { return * (const int32 *)Lhs == * (const int32 *)Rhs; };
+
+		for (int32 TestItem : IntTestItems)
+		{
+			*(int32 *)Set.GetData(Set.AddUninitialized(Layout), Layout) = TestItem;
+		}
+		Set.Rehash(Layout, GetKeyHash);
+
+		for (int32 TestItem : IntTestItems)
+		{
+			CHECK(Set.FindIndex(&TestItem, Layout, GetKeyHash, EqualityFn) != INDEX_NONE);
+		}
+
+		for (int32 Index = 0; Index < Set.Num(); ++Index)
+		{
+			if (Index % 2)
+			{
+				Set.RemoveAt(Index, Layout, GetKeyHash);
+				*(int32 *)Set.GetData(Index, Layout) = *(int32 *)Set.GetData(Set.Num() - 1, Layout);
+				CHECK(Set.FindIndex(&IntTestItems[Index], Layout, GetKeyHash, EqualityFn) == INDEX_NONE);
+			}
+			else
+			{
+				CHECK(Set.FindIndex(&IntTestItems[Index], Layout, GetKeyHash, EqualityFn) != INDEX_NONE);
+			}
+		}
+	}
+}
+
+#endif // WITH_TESTS
