Added IsInWorkerThread that returns true if called from a worker thread.

diff --git a/Engine/Source/Runtime/Core/Private/HAL/ThreadingBase.cpp b/Engine/Source/Runtime/Core/Private/HAL/ThreadingBase.cpp
--- a/Engine/Source/Runtime/Core/Private/HAL/ThreadingBase.cpp
+++ b/Engine/Source/Runtime/Core/Private/HAL/ThreadingBase.cpp
@@ -210,6 +210,11 @@ CORE_API bool IsInParallelGameThread()
 	return FTaskTagScope::IsCurrentTag(ETaskTag::EParallelGameThread);
 }
 
+CORE_API bool IsInWorkerThread()
+{
+	return LowLevelTasks::FScheduler::Get().IsWorkerThread();
+}
+
 CORE_API bool IsInSlateThread()
 {
 PRAGMA_DISABLE_DEPRECATION_WARNINGS
diff --git a/Engine/Source/Runtime/Core/Public/CoreGlobals.h b/Engine/Source/Runtime/Core/Public/CoreGlobals.h
--- a/Engine/Source/Runtime/Core/Public/CoreGlobals.h
+++ b/Engine/Source/Runtime/Core/Public/CoreGlobals.h
@@ -723,6 +723,9 @@ extern CORE_API bool IsInGameThread();
 /** @return True if called from the game thread in a parallel for. */
 extern CORE_API bool IsInParallelGameThread();
 
+/** @return True if called from a worker thread. */
+extern CORE_API bool IsInWorkerThread();
+
 extern CORE_API bool IsAudioThreadRunning();
 
 /** @return True if called from the audio thread, and not merely a thread calling audio functions. */
