Fix CIS static analysis issue

diff --git a/Engine/Source/Runtime/Core/Public/Containers/CompactHashTable.h b/Engine/Source/Runtime/Core/Public/Containers/CompactHashTable.h
--- a/Engine/Source/Runtime/Core/Public/Containers/CompactHashTable.h
+++ b/Engine/Source/Runtime/Core/Public/Containers/CompactHashTable.h
@@ -63,7 +63,7 @@ namespace UE::Core::CompactHashTable
 
 	[[nodiscard]] FORCEINLINE constexpr size_t GetMemoryRequiredInBytes(uint32 IndexCount, uint32 HashCount)
 	{
-		return size_t(IndexCount + HashCount) << GetTypeShift(IndexCount);
+		return ((size_t)IndexCount + HashCount) << GetTypeShift(IndexCount);
 	}
 
 	[[nodiscard]] FORCEINLINE constexpr size_t GetMemoryAlignment()
