Added IsLocked and IsLockedShared to the mutex types

diff --git a/Engine/Source/Runtime/Core/Public/Async/ExternalMutex.h b/Engine/Source/Runtime/Core/Public/Async/ExternalMutex.h
--- a/Engine/Source/Runtime/Core/Public/Async/ExternalMutex.h
+++ b/Engine/Source/Runtime/Core/Public/Async/ExternalMutex.h
@@ -39,7 +39,7 @@ public:
 
 	inline bool IsLocked() const
 	{
-		return (State.load(std::memory_order_relaxed) & IsLockedFlag);
+		return !!(State.load(std::memory_order_relaxed) & IsLockedFlag);
 	}
 
 	inline bool TryLock()
diff --git a/Engine/Source/Runtime/Core/Public/Async/Mutex.h b/Engine/Source/Runtime/Core/Public/Async/Mutex.h
--- a/Engine/Source/Runtime/Core/Public/Async/Mutex.h
+++ b/Engine/Source/Runtime/Core/Public/Async/Mutex.h
@@ -31,9 +31,9 @@ public:
 
 	inline bool IsLocked() const
 	{
-		return (State.load(std::memory_order_relaxed) & IsLockedFlag);
+		return !!(State.load(std::memory_order_relaxed) & IsLockedFlag);
 	}
-	
+
 	inline bool TryLock()
 	{
 		uint8 Expected = State.load(std::memory_order_relaxed);
diff --git a/Engine/Source/Runtime/Core/Public/Async/RecursiveMutex.h b/Engine/Source/Runtime/Core/Public/Async/RecursiveMutex.h
--- a/Engine/Source/Runtime/Core/Public/Async/RecursiveMutex.h
+++ b/Engine/Source/Runtime/Core/Public/Async/RecursiveMutex.h
@@ -23,6 +23,11 @@ public:
 	FRecursiveMutex(const FRecursiveMutex&) = delete;
 	FRecursiveMutex& operator=(const FRecursiveMutex&) = delete;
 
+	inline bool IsLocked() const
+	{
+		return !!(State.load(std::memory_order_relaxed) & LockCountMask);
+	}
+
 	UE_API bool TryLock();
 	UE_API void Lock();
 	UE_API void Unlock();
diff --git a/Engine/Source/Runtime/Core/Public/Async/SharedMutex.h b/Engine/Source/Runtime/Core/Public/Async/SharedMutex.h
--- a/Engine/Source/Runtime/Core/Public/Async/SharedMutex.h
+++ b/Engine/Source/Runtime/Core/Public/Async/SharedMutex.h
@@ -26,6 +26,11 @@ public:
 	FSharedMutex(const FSharedMutex&) = delete;
 	FSharedMutex& operator=(const FSharedMutex&) = delete;
 
+	inline bool IsLocked() const
+	{
+		return !!(State.load(std::memory_order_relaxed) & IsLockedFlag);
+	}
+
 	inline bool TryLock()
 	{
 		uint32 Expected = State.load(std::memory_order_relaxed);
@@ -56,6 +61,11 @@ public:
 		WakeWaitingThreads(LastState);
 	}
 
+	inline bool IsLockShared() const
+	{
+		return !!(State.load(std::memory_order_relaxed) & SharedLockCountMask);
+	}
+
 	inline bool TryLockShared()
 	{
 		uint32 Expected = State.load(std::memory_order_relaxed);
diff --git a/Engine/Source/Runtime/Core/Public/Async/SharedRecursiveMutex.h b/Engine/Source/Runtime/Core/Public/Async/SharedRecursiveMutex.h
--- a/Engine/Source/Runtime/Core/Public/Async/SharedRecursiveMutex.h
+++ b/Engine/Source/Runtime/Core/Public/Async/SharedRecursiveMutex.h
@@ -44,10 +44,20 @@ public:
 	FSharedRecursiveMutex(const FSharedRecursiveMutex&) = delete;
 	FSharedRecursiveMutex& operator=(const FSharedRecursiveMutex&) = delete;
 
+	inline bool IsLocked() const
+	{
+		return !!(State.load(std::memory_order_relaxed) & LockCountMask);
+	}
+
 	UE_API bool TryLock();
 	UE_API void Lock();
 	UE_API void Unlock();
 
+	inline bool IsLockShared() const
+	{
+		return !!(State.load(std::memory_order_relaxed) & SharedLockCountMask);
+	}
+
 	// Use TSharedLock or TDynamicSharedLock to acquire a shared lock.
 	UE_API bool TryLockShared(Core::Private::FSharedRecursiveMutexLink& Link);
 	UE_API void LockShared(Core::Private::FSharedRecursiveMutexLink& Link);
