StructuredLogging: Temporary fix for crash on shutdown due to GStaticLogTemplateManager destructor trying to access other global static data after c++ main exit that was destroyed earlier in the shutdown order.
#rnx

diff --git a/Engine/Source/Runtime/Core/Private/Logging/StructuredLog.cpp b/Engine/Source/Runtime/Core/Private/Logging/StructuredLog.cpp
--- a/Engine/Source/Runtime/Core/Private/Logging/StructuredLog.cpp
+++ b/Engine/Source/Runtime/Core/Private/Logging/StructuredLog.cpp
@@ -1153,12 +1153,16 @@ struct FStaticLogTemplateManager
 
 	~FStaticLogTemplateManager()
 	{
+		// Skip this at-exit shutdown code, because in some circumstances the GetLocalizedLogFormatTable used by
+		// FLogTemplate::Destroy can already be torn down. This class will be removed in an upcoming refactor.
+#if 0
 		for (FLogTemplate* Template = Head.exchange(nullptr); Template;)
 		{
 			FLogTemplate* NextTemplate = Template->GetNext();
 			FLogTemplate::Destroy(Template);
 			Template = NextTemplate;
 		}
+#endif
 	}
 };
 
