Add macro to optionally disable named events

diff --git a/Engine/Source/Runtime/Core/Public/HAL/PlatformMisc.h b/Engine/Source/Runtime/Core/Public/HAL/PlatformMisc.h
--- a/Engine/Source/Runtime/Core/Public/HAL/PlatformMisc.h
+++ b/Engine/Source/Runtime/Core/Public/HAL/PlatformMisc.h
@@ -29,8 +29,12 @@
 #define PLATFORM_EMPTY_BASES
 #endif
 
+#ifndef ALLOW_NAMED_EVENTS_IN_TEST
+#define ALLOW_NAMED_EVENTS_IN_TEST 1	// enabled in Test by default to facilitate auto perf testing
+#endif
+
 // Master switch for scoped named events
-#define ENABLE_NAMED_EVENTS (!UE_BUILD_SHIPPING && 1)
+#define ENABLE_NAMED_EVENTS (!UE_BUILD_SHIPPING && (!UE_BUILD_TEST || ALLOW_NAMED_EVENTS_IN_TEST))
 
 #if PLATFORM_USES_ANSI_STRING_FOR_EXTERNAL_PROFILING
 	#define NAMED_EVENT_STR(x) x
