Fix missing return
#rnx

diff --git a/Engine/Source/Runtime/Core/Public/Logging/StructuredLogFormat.h b/Engine/Source/Runtime/Core/Public/Logging/StructuredLogFormat.h
--- a/Engine/Source/Runtime/Core/Public/Logging/StructuredLogFormat.h
+++ b/Engine/Source/Runtime/Core/Public/Logging/StructuredLogFormat.h
@@ -107,7 +107,7 @@ public:
 
 	inline FText FormatToText(const FCbFieldViewIterator& Fields)
 	{
-		FormatLogToText(Get(), Fields);
+		return FormatLogToText(Get(), Fields);
 	}
 
 	inline ~TLogTemplate()
