Timecode: Cvar to force or remove subframes from its ToString function.

Useful for debugging purposes where one wants to check if subframes are available or not.

Also changed the interpretation of the 2nd argument regarding subframes. Displaying subframes only when subframe > 0 is never really desired or used because it will cause flickering. Instead the argument now specifices whether you want them or not in the resulting string.

Also changed a couple of places where timecode was being converted to string manually instead of using the built in function.

#rnx

diff --git a/Engine/Plugins/VirtualProduction/StageMonitoring/Source/StageMonitorEditor/Private/Widgets/SDataProviderActivities.cpp b/Engine/Plugins/VirtualProduction/StageMonitoring/Source/StageMonitorEditor/Private/Widgets/SDataProviderActivities.cpp
--- a/Engine/Plugins/VirtualProduction/StageMonitoring/Source/StageMonitorEditor/Private/Widgets/SDataProviderActivities.cpp
+++ b/Engine/Plugins/VirtualProduction/StageMonitoring/Source/StageMonitorEditor/Private/Widgets/SDataProviderActivities.cpp
@@ -378,7 +378,7 @@ FText SDataProviderActivitiesTableRow::GetTimecode() const
 		FStageProviderMessage* Data = reinterpret_cast<FStageProviderMessage*>(Item->Data->GetStructMemory());
 		if (Data->FrameTime.AsSeconds() >= 0)
 		{
-			return FText::FromString(FTimecode::FromFrameNumber(Data->FrameTime.Time.GetFrame(), Data->FrameTime.Rate).ToString());
+			return FText::FromString(Data->FrameTime.ToTimecode().ToString());
 		}
 		else
 		{
diff --git a/Engine/Plugins/VirtualProduction/TimedDataMonitor/Source/TimedDataMonitorEditor/Private/STimedDataListView.cpp b/Engine/Plugins/VirtualProduction/TimedDataMonitor/Source/TimedDataMonitorEditor/Private/STimedDataListView.cpp
--- a/Engine/Plugins/VirtualProduction/TimedDataMonitor/Source/TimedDataMonitorEditor/Private/STimedDataListView.cpp
+++ b/Engine/Plugins/VirtualProduction/TimedDataMonitor/Source/TimedDataMonitorEditor/Private/STimedDataListView.cpp
@@ -154,15 +154,20 @@ struct FTimedDataInputTableRowData : TSharedFromThis<FTimedDataInputTableRowData
 			{
 			case ETimedDataInputEvaluationType::Timecode:
 			{
-				const FTimecode Timecode = FTimecode::FromFrameNumber(NewestDataTime.Timecode.Time.GetFrame(), NewestDataTime.Timecode.Rate);
+				const FTimecode Timecode = NewestDataTime.Timecode.ToTimecode();
+
 				if (Timecode == FTimecode())
-				{
+				{ 
 					// Timecode isn't valid.
 					CachedDescription = LOCTEXT("NoTimecodeLabel", "No Timecode");
 				}
 				else
 				{
-					CachedDescription = FText::Format(INVTEXT("{0}@{1}"), FText::FromString(Timecode.ToString()), NewestDataTime.Timecode.Rate.ToPrettyText());
+					CachedDescription = FText::Format(
+						INVTEXT("{0} @ {1}"), 
+						FText::FromString(Timecode.ToString()),
+						NewestDataTime.Timecode.Rate.ToPrettyText()
+					);
 				}
 			}
 			break;
diff --git a/Engine/Plugins/VirtualProduction/TimedDataMonitor/Source/TimedDataMonitorEditor/Private/STimedDataTimecodeProvider.cpp b/Engine/Plugins/VirtualProduction/TimedDataMonitor/Source/TimedDataMonitorEditor/Private/STimedDataTimecodeProvider.cpp
--- a/Engine/Plugins/VirtualProduction/TimedDataMonitor/Source/TimedDataMonitorEditor/Private/STimedDataTimecodeProvider.cpp
+++ b/Engine/Plugins/VirtualProduction/TimedDataMonitor/Source/TimedDataMonitorEditor/Private/STimedDataTimecodeProvider.cpp
@@ -189,8 +189,8 @@ void STimedDataTimecodeProvider::UpdateCachedValue()
 	CachedFrameTimeOptional = FApp::GetCurrentFrameTime();
 	if (CachedFrameTimeOptional.IsSet())
 	{
-		CachedTimecode = FText::Format(LOCTEXT("TimecodeFormat", "{0}@{1}")
-			, FText::FromString(FTimecode::FromFrameNumber(CachedFrameTimeOptional->Time.GetFrame(), CachedFrameTimeOptional->Rate).ToString())
+		CachedTimecode = FText::Format(LOCTEXT("TimecodeFormat", "{0} @ {1}")
+			, FText::FromString(CachedFrameTimeOptional->ToTimecode().ToString())
 			, CachedFrameTimeOptional->Rate.ToPrettyText());
 	}
 	else
diff --git a/Engine/Source/Runtime/Core/Private/Misc/Timecode.cpp b/Engine/Source/Runtime/Core/Private/Misc/Timecode.cpp
--- a/Engine/Source/Runtime/Core/Private/Misc/Timecode.cpp
+++ b/Engine/Source/Runtime/Core/Private/Misc/Timecode.cpp
@@ -10,6 +10,13 @@ static TAutoConsoleVariable<bool> CVarUseDropFormatTimecodeByDefaultWhenSupporte
 	TEXT("By default, should we generate a timecode in drop frame format when the frame rate does support it."),
 	ECVF_Default);
 
+
+static TAutoConsoleVariable<int32> CVarForceStringifyTimecodeSubframes(
+	TEXT("timecode.ForceStringifyTimecodeSubframes"),
+	0,
+	TEXT("Should Timecode.ToString() be forced to include subframes. 0 - Don't force. 1 - Force show. 2 - Force hide"),
+	ECVF_Default);
+
 /* FTimecode interface
  *****************************************************************************/
 bool FTimecode::UseDropFormatTimecodeByDefaultWhenSupported()
@@ -207,4 +214,57 @@ TOptional<FTimecode> FTimecode::ParseTimecode(const FStringView InTimecodeString
 	using namespace UE::Core::Private;
 	FTimecodeParser& StaticTimecodeParser = FTimecodeParser::Get();
 	return StaticTimecodeParser.Evaluate(InTimecodeString);
+}
+
+
+FString FTimecode::ToString(bool bForceSignDisplay, bool bDisplaySubframe) const
+{
+	bool bHasNegativeComponent = Hours < 0 || Minutes < 0 || Seconds < 0 || Frames < 0;
+
+	const TCHAR* NegativeSign = TEXT("- ");
+	const TCHAR* PositiveSign = TEXT("+ ");
+	const TCHAR* SignText = TEXT("");
+
+	if (bHasNegativeComponent)
+	{
+		SignText = NegativeSign;
+	}
+	else if (bForceSignDisplay)
+	{
+		SignText = PositiveSign;
+	}
+
+	// Use a buffer that will hold 64 chars to account for maximum int sizes for hours, min, seconds, frames + 3 charcs for subframe +/- 3 chars for sign.
+	//
+	TStringBuilder<64> Builder;
+	if (bDropFrameFormat)
+	{
+		Builder.Appendf(TEXT("%s%02d:%02d:%02d;%02d"), SignText, FMath::Abs(Hours), FMath::Abs(Minutes), FMath::Abs(Seconds), FMath::Abs(Frames));
+	}
+	else
+	{
+		Builder.Appendf(TEXT("%s%02d:%02d:%02d:%02d"), SignText, FMath::Abs(Hours), FMath::Abs(Minutes), FMath::Abs(Seconds), FMath::Abs(Frames));
+	}
+
+	switch (CVarForceStringifyTimecodeSubframes.GetValueOnAnyThread())
+	{
+	case 1:
+		bDisplaySubframe = true;
+		break;
+
+	case 2:
+		bDisplaySubframe = false;
+		break;
+
+	default:
+		break;
+	}
+
+	if (bDisplaySubframe)
+	{
+		int32 ClampedSubframe = static_cast<int32>(FMath::Clamp(100 * Subframe, 0, 99));
+		Builder.Appendf(TEXT(".%02d"), ClampedSubframe);
+	}
+
+	return FString(Builder);
 }
\ No newline at end of file
diff --git a/Engine/Source/Runtime/Core/Public/Misc/Timecode.h b/Engine/Source/Runtime/Core/Public/Misc/Timecode.h
--- a/Engine/Source/Runtime/Core/Public/Misc/Timecode.h
+++ b/Engine/Source/Runtime/Core/Public/Misc/Timecode.h
@@ -345,43 +345,9 @@ public:
 	 * Get the Qualified Timecode formatted in HH:MM:SS:FF or HH:MM:SS;FF depending on if this represents drop-frame timecode or not.
 	 * @param bForceSignDisplay - Forces the timecode to be prepended with a positive or negative sign.
 								  Standard behavior is to only show the sign when the value is negative.
-	 * @param bAlwaysDisplaySubframe - Forces the subframe value to be included in the string result.
+	 * @param bDisplaySubframe - If subframe decimal digits are desired.
 	 */
-	FString ToString(bool bForceSignDisplay = false, bool bAlwaysDisplaySubframe = false) const
-	{
-		bool bHasNegativeComponent = Hours < 0 || Minutes < 0 || Seconds < 0 || Frames < 0;
-
-		const TCHAR* NegativeSign = TEXT("- ");
-		const TCHAR* PositiveSign = TEXT("+ ");
-		const TCHAR* SignText = TEXT("");
-		if (bHasNegativeComponent)
-		{
-			SignText = NegativeSign;
-		}
-		else if (bForceSignDisplay)
-		{
-			SignText = PositiveSign;
-		}
-
-		// Use a buffer that will hold 64 chars to account for maximum int sizes for hours, min, seconds, frames + 3 charcs for subframe +/- 3 chars for sign.
-		//
-		TStringBuilder<64> Builder;
-		if (bDropFrameFormat)
-		{
-			Builder.Appendf(TEXT("%s%02d:%02d:%02d;%02d"), SignText, FMath::Abs(Hours), FMath::Abs(Minutes), FMath::Abs(Seconds), FMath::Abs(Frames));
-		}
-		else
-		{
-			Builder.Appendf(TEXT("%s%02d:%02d:%02d:%02d"), SignText, FMath::Abs(Hours), FMath::Abs(Minutes), FMath::Abs(Seconds), FMath::Abs(Frames));
-		}
-
-		if (bAlwaysDisplaySubframe || Subframe > 0)
-		{
-			int32 ClampedSubframe = static_cast<int32>(FMath::Clamp(100*Subframe,0,99));
-			Builder.Appendf(TEXT(".%02d"), ClampedSubframe);
-		}
-		return FString(Builder);
-	}
+	CORE_API FString ToString(bool bForceSignDisplay = false, bool bDisplaySubframe = false) const;
 
 	/**
 	 * Parses a string representation of a timecode.
diff --git a/Engine/Source/Runtime/TimeManagement/Private/Widgets/STimecode.cpp b/Engine/Source/Runtime/TimeManagement/Private/Widgets/STimecode.cpp
--- a/Engine/Source/Runtime/TimeManagement/Private/Widgets/STimecode.cpp
+++ b/Engine/Source/Runtime/TimeManagement/Private/Widgets/STimecode.cpp
@@ -48,12 +48,9 @@ int32 STimecode::OnPaint(const FPaintArgs& Args, const FGeometry& AllottedGeomet
 {
 	const ESlateDrawEffect DrawEffects = ESlateDrawEffect::None;
 
-	FTimecode TimecodeToPaint = Timecode.Get();
-	FString TimecodeToPaintString = TimecodeToPaint.ToString();
-	if (!bDisplaySubframes && TimecodeToPaint.Subframe != 0)
-	{
-		TimecodeToPaintString.LeftChopInline(3);
-	}
+	const FTimecode TimecodeToPaint = Timecode.Get();
+	const FString TimecodeToPaintString = TimecodeToPaint.ToString(false /* bForceSignDisplay */, bDisplaySubframes);
+
 	const FSlateFontInfo& TimecodeFontInfo = TimecodeFont.Get();
 
 	int32 NewLayerId = LayerId+1;
@@ -123,11 +120,7 @@ int32 STimecode::OnPaint(const FPaintArgs& Args, const FGeometry& AllottedGeomet
 FVector2D STimecode::ComputeDesiredSize(float LayoutScale) const
 {
 	const FTimecode TimecodeToPaint = Timecode.Get();
-	FString TimecodeString = Timecode.Get().ToString();
-	if (!bDisplaySubframes && TimecodeToPaint.Subframe != 0)
-	{
-		TimecodeString.LeftChopInline(3);
-	}
+	const FString TimecodeString = Timecode.Get().ToString(false /* bForceSignDisplay */, bDisplaySubframes);
 
 	if (bDisplayLabel.Get())
 	{
