[AutoRTFM] Mark FOutputDeviceRedirector logging methods with UE_AUTORTFM_ALWAYS_OPEN

We do not want to revert log messages, so mark the FOutputDeviceRedirector logging methods as always-open.
Remove the partial handling of closed transactions from FOutputDevicesReadScopeLock as this is no longer used in the closed.

We were hesitant to make this logic always-open due to transitive open calls having the potential to perform mixed open / close writes, but this ship has long sailed as BasicLogV() is already marked with UE_AUTORTFM_ALWAYS_OPEN, which is the caller of most logging functions.

diff --git a/Engine/Source/Programs/AutoRTFMTests/Private/UECoreTests.cpp b/Engine/Source/Programs/AutoRTFMTests/Private/UECoreTests.cpp
--- a/Engine/Source/Programs/AutoRTFMTests/Private/UECoreTests.cpp
+++ b/Engine/Source/Programs/AutoRTFMTests/Private/UECoreTests.cpp
@@ -1622,41 +1622,40 @@ TEST_CASE("UECore.UE_LOGFMT")
 
 TEST_CASE("UECore.FOutputDeviceRedirector")
 {
+	FOutputDeviceRedirector Redirector;
 	SECTION("Commit")
 	{
 		FStringOutputDevice StringLog;
-		GLog->AddOutputDevice(&StringLog);
+		Redirector.AddOutputDevice(&StringLog);
 
 		AutoRTFM::Testing::Commit([&]
 		{
 			// This test will actually be run twice, because we test with SetRetryTransaction enabled.
-			// Logging always runs in the open because UE::Private::Log::BasicLog is marked ALWAYS_OPEN,
-			// so the log won't be undone when the transaction is rolled back before being retried.
+			// Logging always runs in the open so the log won't be undone when the transaction is
+			// rolled back before being retried.
 			// We handle this by making sure that the string log has "Commit" appended to it, rather
 			// than verifying that it contains "Commit" exactly.
 			FString PreviousStringLog = StringLog;
-			UE_LOG(LogAutoRTFM_UECoreTests, Display, TEXT("Commit"));
+			Redirector.Log(TEXT("Commit"));
+			Redirector.Flush();
 			REQUIRE(StringLog == PreviousStringLog + FString("Commit"));
 		});
-
-		GLog->RemoveOutputDevice(&StringLog);
 	}
 
 	SECTION("Abort")
 	{
 		FStringOutputDevice StringLog;
-		GLog->AddOutputDevice(&StringLog);
+		Redirector.AddOutputDevice(&StringLog);
 
 		AutoRTFM::Testing::Abort([&]
 		{
 			FString PreviousStringLog = StringLog;
-			UE_LOG(LogAutoRTFM_UECoreTests, Display, TEXT("Abort"));
+			Redirector.Log(TEXT("Abort"));
+			Redirector.Flush();
 			REQUIRE(StringLog == PreviousStringLog + FString("Abort"));
 
 			AutoRTFM::AbortTransaction();
 		});
-
-		GLog->RemoveOutputDevice(&StringLog);
 	}
 }
 
diff --git a/Engine/Source/Runtime/Core/Private/Logging/LogScopedCategoryAndVerbosityOverride.cpp b/Engine/Source/Runtime/Core/Private/Logging/LogScopedCategoryAndVerbosityOverride.cpp
--- a/Engine/Source/Runtime/Core/Private/Logging/LogScopedCategoryAndVerbosityOverride.cpp
+++ b/Engine/Source/Runtime/Core/Private/Logging/LogScopedCategoryAndVerbosityOverride.cpp
@@ -2,6 +2,7 @@
 
 #include "Logging/LogScopedCategoryAndVerbosityOverride.h"
 
+#include "AutoRTFM.h"
 #include "CoreTypes.h"
 #include "HAL/PlatformTLS.h"
 
@@ -28,6 +29,7 @@ uint32 GetOverrideTLSID()
 	return OverrrideTLSID;
 }
 
+UE_AUTORTFM_ALWAYS_OPEN
 FLogScopedCategoryAndVerbosityOverride::FOverride* FLogScopedCategoryAndVerbosityOverride::GetTLSCurrent()
 {
 	FOverride* TLS = (FOverride*)FPlatformTLS::GetTlsValue(GetOverrideTLSID());
diff --git a/Engine/Source/Runtime/Core/Private/Misc/OutputDeviceRedirector.cpp b/Engine/Source/Runtime/Core/Private/Misc/OutputDeviceRedirector.cpp
--- a/Engine/Source/Runtime/Core/Private/Misc/OutputDeviceRedirector.cpp
+++ b/Engine/Source/Runtime/Core/Private/Misc/OutputDeviceRedirector.cpp
@@ -232,23 +232,12 @@ public:
 	FORCEINLINE explicit FOutputDevicesReadScopeLock(FOutputDeviceRedirectorState& InState)
 		: State(InState)
 	{
-		UE_AUTORTFM_OPEN
-		{
-			AcquireLock();
-		};
-		AutoRTFM::PushOnAbortHandler(this, [this]
-		{
-			ReleaseLock();
-		});
+		AcquireLock();
 	}
 
 	FORCEINLINE ~FOutputDevicesReadScopeLock()
 	{
-		UE_AUTORTFM_OPEN
-		{
-			ReleaseLock();
-		};
-		AutoRTFM::PopOnAbortHandler(this);
+		ReleaseLock();
 	}
 
 private:
@@ -393,6 +382,7 @@ private:
 	bool bLocked = false;
 };
 
+UE_AUTORTFM_NOAUTORTFM
 void FOutputDeviceRedirectorState::AddOutputDevice(FOutputDevice* OutputDevice)
 {
 	const auto AddTo = [OutputDevice](TArray<FOutputDevice*>& OutputDevices, TBitArray<TInlineAllocator<1>>& Flags)
@@ -414,6 +404,7 @@ void FOutputDeviceRedirectorState::AddOutputDevice(FOutputDevice* OutputDevice)
 	}
 }
 
+UE_AUTORTFM_NOAUTORTFM
 void FOutputDeviceRedirectorState::RemoveOutputDevice(FOutputDevice* OutputDevice)
 {
 	const auto RemoveFrom = [OutputDevice](TArray<FOutputDevice*>& OutputDevices, TBitArray<TInlineAllocator<1>>& Flags)
@@ -429,6 +420,7 @@ void FOutputDeviceRedirectorState::RemoveOutputDevice(FOutputDevice* OutputDevic
 	RemoveFrom(UnbufferedOutputDevices, UnbufferedOutputDevicesCanBeUsedOnPanicThread);
 }
 
+UE_AUTORTFM_NOAUTORTFM
 bool FOutputDeviceRedirectorState::TryStartThread()
 {
 	if (FWriteScopeLock ThreadScopeLock(ThreadLock); !bThreadStarted.exchange(true, std::memory_order_relaxed))
@@ -438,6 +430,7 @@ bool FOutputDeviceRedirectorState::TryStartThread()
 	return true;
 }
 
+UE_AUTORTFM_NOAUTORTFM
 bool FOutputDeviceRedirectorState::TryStopThread()
 {
 	if (FWriteScopeLock ThreadScopeLock(ThreadLock); bThreadStarted.exchange(false, std::memory_order_relaxed))
@@ -448,6 +441,7 @@ bool FOutputDeviceRedirectorState::TryStopThread()
 	return true;
 }
 
+UE_AUTORTFM_NOAUTORTFM
 void FOutputDeviceRedirectorState::ThreadLoop()
 {
 	const uint32 ThreadId = FPlatformTLS::GetCurrentThreadId();
@@ -476,6 +470,7 @@ void FOutputDeviceRedirectorState::ThreadLoop()
 	}
 }
 
+UE_AUTORTFM_NOAUTORTFM
 void FOutputDeviceRedirectorState::FlushBufferedItems()
 {
 	using namespace UE;
@@ -522,6 +517,7 @@ void FOutputDeviceRedirectorState::FlushBufferedItems()
 	});
 }
 
+UE_AUTORTFM_NOAUTORTFM
 void FOutputDeviceRedirectorState::AddToBacklog(const FLogRecord& Record)
 {
 	TStringBuilder<512> Text;
@@ -529,6 +525,7 @@ void FOutputDeviceRedirectorState::AddToBacklog(const FLogRecord& Record)
 	AddToBacklog(*Text, Record.GetVerbosity(), Record.GetCategory(), FPlatformTime::Seconds() - GStartTime);
 }
 
+UE_AUTORTFM_NOAUTORTFM
 void FOutputDeviceRedirectorState::AddToBacklog(
 	const TCHAR* const Data,
 	const ELogVerbosity::Type Verbosity,
@@ -541,6 +538,7 @@ void FOutputDeviceRedirectorState::AddToBacklog(
 
 } // UE::Private
 
+UE_AUTORTFM_NOAUTORTFM // FOutputDeviceRedirector cannot be constructed in a closed transaction
 FOutputDeviceRedirector::FOutputDeviceRedirector()
 	: State(MakePimpl<UE::Private::FOutputDeviceRedirectorState>())
 {
@@ -548,6 +546,7 @@ FOutputDeviceRedirector::FOutputDeviceRedirector()
 
 static bool GWasGLogDestroyed = false;
 
+UE_AUTORTFM_NOAUTORTFM // FOutputDeviceRedirector cannot be destructed in a closed transaction
 FOutputDeviceRedirector::~FOutputDeviceRedirector()
 {
 	if (this == FOutputDeviceRedirector::Get())
@@ -562,6 +561,7 @@ FOutputDeviceRedirector* FOutputDeviceRedirector::Get()
 	return &Singleton;
 }
 
+UE_AUTORTFM_NOAUTORTFM // Not currently supported in a closed transaction
 void FOutputDeviceRedirector::AddOutputDevice(FOutputDevice* OutputDevice)
 {
 	if (OutputDevice)
@@ -570,6 +570,7 @@ void FOutputDeviceRedirector::AddOutputDevice(FOutputDevice* OutputDevice)
 	}
 }
 
+UE_AUTORTFM_NOAUTORTFM // Not currently supported in a closed transaction
 void FOutputDeviceRedirector::RemoveOutputDevice(FOutputDevice* OutputDevice)
 {
 	if (OutputDevice)
@@ -578,12 +579,14 @@ void FOutputDeviceRedirector::RemoveOutputDevice(FOutputDevice* OutputDevice)
 	}
 }
 
+UE_AUTORTFM_ALWAYS_OPEN
 bool FOutputDeviceRedirector::IsRedirectingTo(FOutputDevice* OutputDevice)
 {
 	UE::Private::FOutputDevicesReadScopeLock Lock(*State);
 	return State->BufferedOutputDevices.Contains(OutputDevice) || State->UnbufferedOutputDevices.Contains(OutputDevice);
 }
 
+UE_AUTORTFM_ALWAYS_OPEN
 void FOutputDeviceRedirector::FlushThreadedLogs(EOutputDeviceRedirectorFlushOptions Options)
 {
 	if (FReadScopeLock ThreadLock(State->ThreadLock); State->bThreadStarted.load(std::memory_order_relaxed))
@@ -604,6 +607,7 @@ void FOutputDeviceRedirector::FlushThreadedLogs(EOutputDeviceRedirectorFlushOpti
 	}
 }
 
+UE_AUTORTFM_ALWAYS_OPEN
 void FOutputDeviceRedirector::SerializeBacklog(FOutputDevice* OutputDevice)
 {
 	FReadScopeLock ScopeLock(State->BacklogLock);
@@ -613,6 +617,7 @@ void FOutputDeviceRedirector::SerializeBacklog(FOutputDevice* OutputDevice)
 	}
 }
 
+UE_AUTORTFM_NOAUTORTFM // Not currently supported in a closed transaction
 void FOutputDeviceRedirector::EnableBacklog(bool bEnable)
 {
 	FWriteScopeLock ScopeLock(State->BacklogLock);
@@ -623,6 +628,7 @@ void FOutputDeviceRedirector::EnableBacklog(bool bEnable)
 	}
 }
 
+UE_AUTORTFM_NOAUTORTFM // Not currently supported in a closed transaction
 void FOutputDeviceRedirector::SetCurrentThreadAsPrimaryThread()
 {
 	const uint32 ThreadId = FPlatformTLS::GetCurrentThreadId();
@@ -640,11 +646,13 @@ void FOutputDeviceRedirector::SetCurrentThreadAsPrimaryThread()
 	State->TryStopThread();
 }
 
+UE_AUTORTFM_NOAUTORTFM // Not currently supported in a closed transaction
 bool FOutputDeviceRedirector::TryStartDedicatedPrimaryThread()
 {
 	return FApp::ShouldUseThreadingForPerformance() && State->TryStartThread();
 }
 
+UE_AUTORTFM_ALWAYS_OPEN
 void FOutputDeviceRedirector::SerializeRecord(const UE::FLogRecord& Record)
 {
 	using namespace UE::Private;
@@ -690,6 +698,7 @@ void FOutputDeviceRedirector::SerializeRecord(const UE::FLogRecord& Record)
 	}
 }
 
+UE_AUTORTFM_ALWAYS_OPEN
 void FOutputDeviceRedirector::Serialize(const TCHAR* const Data, const ELogVerbosity::Type Verbosity, const FName& Category, const double Time)
 {
 	using namespace UE::Private;
@@ -766,6 +775,7 @@ void FOutputDeviceRedirector::RedirectLog(const FLazyName& Category, ELogVerbosi
 	Serialize(Data, Verbosity, Category, -1.0);
 }
 
+UE_AUTORTFM_ALWAYS_OPEN
 void FOutputDeviceRedirector::Flush()
 {
 	if (UE::Private::FOutputDevicesPrimaryScopeLock Lock(*State); Lock.IsLocked())
@@ -777,6 +787,7 @@ void FOutputDeviceRedirector::Flush()
 	}
 }
 
+UE_AUTORTFM_ALWAYS_OPEN
 void FOutputDeviceRedirector::Panic()
 {
 	uint32 PreviousThreadId = MAX_uint32;
@@ -814,6 +825,7 @@ void FOutputDeviceRedirector::Panic()
 	}
 }
 
+UE_AUTORTFM_NOAUTORTFM // Not currently supported in a closed transaction
 void FOutputDeviceRedirector::TearDown()
 {
 	SetCurrentThreadAsPrimaryThread();
@@ -844,6 +856,7 @@ void FOutputDeviceRedirector::TearDown()
 	}
 }
 
+UE_AUTORTFM_ALWAYS_OPEN
 bool FOutputDeviceRedirector::IsBacklogEnabled() const
 {
 	FReadScopeLock Lock(State->BacklogLock);
