Logging: Added FInlineLogTemplate and FDefaultLogTemplate to provide more control over log template allocations

diff --git a/Engine/Source/Runtime/Core/Private/Experimental/UnifiedError/UnifiedError.cpp b/Engine/Source/Runtime/Core/Private/Experimental/UnifiedError/UnifiedError.cpp
--- a/Engine/Source/Runtime/Core/Private/Experimental/UnifiedError/UnifiedError.cpp
+++ b/Engine/Source/Runtime/Core/Private/Experimental/UnifiedError/UnifiedError.cpp
@@ -59,9 +59,6 @@ FText FError::GetErrorMessage(bool bIncludeContext) const
 #if USE_STRUCTURED_LOG_FOR_FERRORMESSAGE
 	FString FormatString = ErrorDetails->GetErrorFormatString(*this).ToString();
 
-	FLogTemplate* Template = UE::CreateLogTemplate(*FormatString);
-
-
 	TCbWriter<1024> Writer;
 	Writer.BeginObject();
 
@@ -70,9 +67,9 @@ FText FError::GetErrorMessage(bool bIncludeContext) const
 	Writer.EndObject();
 
 	FWideStringBuilderBase OutputMessage;
-	FormatLogTo(OutputMessage,  Template, Writer.Save()); 
+	FInlineLogTemplate Template(*FormatString);
+	Template.FormatTo(OutputMessage, Writer.Save());
 
-	UE::DestroyLogTemplate(Template);
 	return FText::FromString(OutputMessage.ToString());
 #else
 	FFormatNamedArguments Args;
diff --git a/Engine/Source/Runtime/Core/Private/Logging/StructuredLog.cpp b/Engine/Source/Runtime/Core/Private/Logging/StructuredLog.cpp
--- a/Engine/Source/Runtime/Core/Private/Logging/StructuredLog.cpp
+++ b/Engine/Source/Runtime/Core/Private/Logging/StructuredLog.cpp
@@ -143,9 +143,8 @@ static void LogFieldValue(TStringBuilderBase<CharType>& Out, const FCbFieldView&
 			if (FCbFieldView FormatField = Object.FindViewIgnoreCase(ANSITEXTVIEW("$format")); FormatField.IsString())
 			{
 				TStringBuilder<128> Format(InPlace, FormatField.AsString());
-				FLogTemplate* Template = UE::CreateLogTemplate(*Format, {.bAllowSubObjectReferences = true});
-				FormatLogTo(Out, Template, Object.CreateViewIterator());
-				UE::DestroyLogTemplate(Template);
+				FInlineLogTemplate Template(*Format, {.bAllowSubObjectReferences = true});
+				Template.FormatTo(Out, Object.CreateViewIterator());
 				break;
 			}
 
@@ -158,9 +157,8 @@ static void LogFieldValue(TStringBuilderBase<CharType>& Out, const FCbFieldView&
 				TStringBuilder<32> Namespace(InPlace, LocNamespaceField.AsString());
 				TStringBuilder<32> Key(InPlace, LocKeyField.AsString());
 				TStringBuilder<128> Format(InPlace, LocFormatField.AsString());
-				FLogTemplate* Template = UE::CreateLogTemplate(*Namespace, *Key, *Format, {.bAllowSubObjectReferences = true});
-				FormatLogTo(Out, Template, Object.CreateViewIterator());
-				UE::DestroyLogTemplate(Template);
+				FInlineLogTemplate Template(*Namespace, *Key, *Format, {.bAllowSubObjectReferences = true});
+				Template.FormatTo(Out, Object.CreateViewIterator());
 				break;
 			}
 
@@ -912,10 +910,9 @@ FORCENOINLINE static void FormatDynamicRecordMessageTo(TStringBuilderBase<CharTy
 	checkf(!TextNamespace == !TextKey,
 		TEXT("Log record must have both or neither of the text namespace and text key. [[%s]]"), Format);
 
-	const FLogTemplateOptions Options;
-	FLogTemplate* LocalTemplate = TextKey ? CreateLogTemplate(TextNamespace, TextKey, Format, Options) : CreateLogTemplate(Format, Options);
+	TOptional<FInlineLogTemplate> LocalTemplate;
+	TextKey ? LocalTemplate.Emplace(TextNamespace, TextKey, Format) : LocalTemplate.Emplace(Format);
 	LocalTemplate->FormatTo(Out, Record.GetFields().CreateViewIterator());
-	FLogTemplate::Destroy(LocalTemplate);
 }
 
 template <typename CharType>
@@ -1450,44 +1447,12 @@ namespace UE::Logging::Private
 
 static constexpr const TCHAR* const GStaticBasicLogFormat = TEXT("{Message}");
 
-class FStaticBasicLogTemplateManager
+static FLogTemplate* GetStaticBasicLogTemplate()
 {
-public:
-	~FStaticBasicLogTemplateManager()
-	{
-		if (FLogTemplate* LocalTemplate = Template.exchange(nullptr))
-		{
-			FLogTemplate::Destroy(LocalTemplate);
-		}
-	}
-
-	inline FLogTemplate& EnsureTemplate()
-	{
-		if (FLogTemplate* LocalTemplate = Template.load(std::memory_order_acquire); LIKELY(LocalTemplate))
-		{
-			return *LocalTemplate;
-		}
-		return CreateTemplate();
+	static FInlineLogTemplate Template(GStaticBasicLogFormat);
+	return Template.Get();
 }
 
-private:
-	FORCENOINLINE FLogTemplate& CreateTemplate()
-	{
-		FLogTemplate* NewTemplate = UE::CreateLogTemplate(GStaticBasicLogFormat);
-		if (FLogTemplate* ExistingTemplate = nullptr;
-			UNLIKELY(!Template.compare_exchange_strong(ExistingTemplate, NewTemplate, std::memory_order_release, std::memory_order_acquire)))
-		{
-			UE::DestroyLogTemplate(NewTemplate);
-			return *ExistingTemplate;
-		}
-		return *NewTemplate;
-	}
-
-	std::atomic<FLogTemplate*> Template = nullptr;
-};
-
-static FStaticBasicLogTemplateManager GStaticBasicLogTemplateManager;
-
 // Serializing the log to compact binary happens in its own function because that allows stack
 // space for the writer to be returned before calling into the output devices.
 FORCENOINLINE static FCbObject SerializeBasicLogMessage(const FStaticBasicLogRecord& Log, va_list Args)
@@ -1546,7 +1511,7 @@ static void BasicLogV(const FLogCategoryBase& Category, const FStaticBasicLogRec
 
 		FLogRecord Record;
 		Record.SetFormat(GStaticBasicLogFormat);
-		Record.SetTemplate(&GStaticBasicLogTemplateManager.EnsureTemplate());
+		Record.SetTemplate(GetStaticBasicLogTemplate());
 		Record.SetFields(MoveTemp(Fields));
 		Record.SetFile(Log->File);
 		Record.SetLine(Log->Line);
diff --git a/Engine/Source/Runtime/Core/Public/Logging/StructuredLogFormat.h b/Engine/Source/Runtime/Core/Public/Logging/StructuredLogFormat.h
--- a/Engine/Source/Runtime/Core/Public/Logging/StructuredLogFormat.h
+++ b/Engine/Source/Runtime/Core/Public/Logging/StructuredLogFormat.h
@@ -2,9 +2,11 @@
 
 #pragma once
 
+#include "Containers/Array.h"
 #include "Containers/StringFwd.h"
 #include "Internationalization/Text.h"
 #include "Misc/ScopeExit.h"
+#include "Templates/UniquePtr.h"
 
 #define UE_API CORE_API
 
@@ -19,7 +21,50 @@ namespace UE::Logging::Private { struct FLogField; }
 namespace UE::Logging::Private
 {
 
-/** Store a log template in an allocation from FMemory. */
+/** Store a log template in an inline byte array. */
+class FInlineLogTemplateStorage
+{
+public:
+	inline void* Allocate(int32 Size)
+	{
+		Data.SetNum(Size);
+		return Data.GetData();
+	}
+
+	inline const void* Get() const
+	{
+		return Data.GetData();
+	}
+
+private:
+	// FLogTemplate is 8 bytes plus the encoded ops.
+	// There are 2 ops for each field. There is 1 op for each contiguous region of literal text.
+	// There is 1 op for each escaped character. There is 1 op to mark the end.
+	// Most ops encode as 1 byte. A format string that has 12 fields with text surrounding each
+	// field will typically be encoded in 46 bytes.
+	TArray<uint8, TInlineAllocator<48>> Data;
+};
+
+/** Store a log template in a movable heap-allocated byte array. */
+class FUniqueLogTemplateStorage
+{
+public:
+	inline void* Allocate(int32 Size)
+	{
+		Data = MakeUnique<uint8[]>(Size);
+		return Data.Get();
+	}
+
+	inline const void* Get() const
+	{
+		return Data.Get();
+	}
+
+private:
+	TUniquePtr<uint8[]> Data;
+};
+
+/** Store a log template in a detachable allocation from FMemory. */
 class FMemoryLogTemplateStorage
 {
 public:
@@ -29,11 +74,12 @@ public:
 
 	~FMemoryLogTemplateStorage()
 	{
-		FMemory::Free(Data);
+		Free(Data);
 	}
 
 	inline void* Allocate(int32 Size)
 	{
+		Free(Data);
 		Data = FMemory::Malloc(Size);
 		return Data;
 	}
@@ -49,6 +95,14 @@ public:
 		return Data;
 	}
 
+	inline static void Free(void* D)
+	{
+		if (D)
+		{
+			FMemory::Free(D);
+		}
+	}
+
 private:
 	void* Data = nullptr;
 };
@@ -75,7 +129,7 @@ UE_API void FormatLogTo(FUtf8StringBuilderBase& Out, const FLogTemplate* Templat
 UE_API void FormatLogTo(FWideStringBuilderBase& Out, const FLogTemplate* Template, const FCbFieldViewIterator& Fields);
 UE_API FText FormatLogToText(const FLogTemplate* Template, const FCbFieldViewIterator& Fields);
 
-/** A log template that is templated on its storage. */
+/** A log template that is templated on its storage. Use FInlineLogTemplate or FSharedLogTemplate. */
 template <typename StorageType>
 class TLogTemplate
 {
@@ -132,29 +186,50 @@ private:
 	StorageType Storage;
 };
 
-/** Create a log template on the heap. */
+/** A log template that stores its data inline. Best choice for a stack-based temporary. */
+using FInlineLogTemplate = TLogTemplate<Logging::Private::FInlineLogTemplateStorage>;
+
+/** A log template that stores its data on the heap. A reasonable default when for non-temporary templates. */
+using FUniqueLogTemplate = TLogTemplate<Logging::Private::FUniqueLogTemplateStorage>;
+
+/**
+ * Create a log template with explicit lifetime management.
+ *
+ * Destroy with DestroyLogTemplate.
+ * Prefer FInlineLogTemplate and FUniqueLogTemplate.
+ */
 inline FLogTemplate* CreateLogTemplate(const TCHAR* Format, const FLogTemplateOptions& Options = {})
 {
 	return TLogTemplate<Logging::Private::FMemoryLogTemplateStorage>(Format, Options).Detach();
 }
 
-/** Create a localized log template on the heap. */
+/**
+ * Create a localized log template with explicit lifetime management.
+ *
+ * Destroy with DestroyLogTemplate.
+ * Prefer FInlineLogTemplate and FUniqueLogTemplate.
+ */
 inline FLogTemplate* CreateLogTemplate(const FText& Format, const FLogTemplateOptions& Options = {})
 {
 	return TLogTemplate<Logging::Private::FMemoryLogTemplateStorage>(Format, Options).Detach();
 }
 
-/** Create a localized log template on the heap. */
+/**
+ * Create a localized log template with explicit lifetime management.
+ *
+ * Destroy with DestroyLogTemplate.
+ * Prefer FInlineLogTemplate and FUniqueLogTemplate.
+ */
 inline FLogTemplate* CreateLogTemplate(const TCHAR* TextNamespace, const TCHAR* TextKey, const TCHAR* Format, const FLogTemplateOptions& Options = {})
 {
 	return TLogTemplate<Logging::Private::FMemoryLogTemplateStorage>(TextNamespace, TextKey, Format, Options).Detach();
 }
 
-/** Destroy a log template allocated by CreateLogTemplate. */
+/** Destroy a log template created by CreateLogTemplate. */
 inline void DestroyLogTemplate(FLogTemplate* Template)
 {
 	Logging::Private::DestroyLogTemplate(Template);
-	FMemory::Free(Template);
+	Logging::Private::FMemoryLogTemplateStorage::Free(Template);
 }
 
 /**
diff --git a/Engine/Source/Runtime/Core/Tests/Logging/StructuredLogFormatTest.cpp b/Engine/Source/Runtime/Core/Tests/Logging/StructuredLogFormatTest.cpp
--- a/Engine/Source/Runtime/Core/Tests/Logging/StructuredLogFormatTest.cpp
+++ b/Engine/Source/Runtime/Core/Tests/Logging/StructuredLogFormatTest.cpp
@@ -174,10 +174,9 @@ TEST_CASE_NAMED(FStructuredLogFormatTest, "Core::Logging::StructuredLogFormat",
 
 	const auto Test = [&Fields](const TCHAR* Format, FUtf8StringView Expected, const FLogTemplateOptions& Options = {})
 	{
-		FLogTemplate* Template = CreateLogTemplate(Format, Options);
+		FInlineLogTemplate Template(Format, Options);
 		TUtf8StringBuilder<1024> Message;
-		FormatLogTo(Message, Template, Fields);
-		DestroyLogTemplate(Template);
+		Template.FormatTo(Message, Fields);
 		// TODO: Requires a fix for CaptureExpressionsAndValues in LowLevelTestAdapter.h
 		//CAPTURE(Format, Expected, Message);
 		CHECK(Message.ToView().Equals(Expected));
@@ -237,10 +236,9 @@ TEST_CASE_NAMED(FStructuredLogFormatTest, "Core::Logging::StructuredLogFormat",
 		const TOptional<FString> Namespace = FTextInspector::GetNamespace(FormatText);
 		const TOptional<FString> Key = FTextInspector::GetKey(FormatText);
 		const FString* Format = FTextInspector::GetSourceString(FormatText);
-		FLogTemplate* Template = CreateLogTemplate(**Namespace, **Key, **Format, Options);
+		FInlineLogTemplate Template(**Namespace, **Key, **Format, Options);
 		TUtf8StringBuilder<1024> Message;
-		FormatLogTo(Message, Template, Fields);
-		DestroyLogTemplate(Template);
+		Template.FormatTo(Message, Fields);
 		// TODO: Requires a fix for CaptureExpressionsAndValues in LowLevelTestAdapter.h
 		//CAPTURE(Format, Expected, Message);
 		CHECK(Message.ToView().Equals(Expected));
@@ -301,10 +299,9 @@ TEST_CASE_NAMED(FStructuredLogFormatTest, "Core::Logging::StructuredLogFormat",
 
 	// Test constructing a template from FText.
 	{
-		FLogTemplate* Template = CreateLogTemplate(LOCTEXT("TextFormat", "FText bWorks={True}"));
+		FUniqueLogTemplate Template(LOCTEXT("TextFormat", "FText bWorks={True}"));
 		TUtf8StringBuilder<24> Message;
-		FormatLogTo(Message, Template, Fields);
-		DestroyLogTemplate(Template);
+		Template.FormatTo(Message, Fields);
 		CHECK(Message.ToView().Equals(ANSITEXTVIEW(R"(FText bWorks=true)")));
 	}
 }
