StringBuilder: Added support for joining non-contiguous ranges

diff --git a/Engine/Source/Runtime/Core/Public/Misc/StringBuilder.h b/Engine/Source/Runtime/Core/Public/Misc/StringBuilder.h
--- a/Engine/Source/Runtime/Core/Public/Misc/StringBuilder.h
+++ b/Engine/Source/Runtime/Core/Public/Misc/StringBuilder.h
@@ -20,6 +20,7 @@
 #include "Traits/IsCharType.h"
 #include "Traits/IsContiguousContainer.h"
 
+#include <iterator>
 #include <type_traits>
 
 template <typename T> struct TIsContiguousContainer;
@@ -66,14 +67,6 @@ public:
 	/** The string view type that this builder is compatible with. */
 	using ViewType = TStringView<ElementType>;
 
-	/** Whether the given type can be appended to this builder using the append operator. */
-	template <typename AppendType>
-	inline constexpr static bool TCanAppend_V = std::is_same_v<BuilderType&, decltype(DeclVal<BuilderType&>() << DeclVal<AppendType>())>;
-
-	/** Whether the given range type can have its elements appended to the builder using the append operator. */
-	template <typename RangeType>
-	inline constexpr static bool TCanAppendRange_V = TIsContiguousContainer<RangeType>::Value && TCanAppend_V<decltype(*::GetData(DeclVal<RangeType>()))>;
-
 				TStringBuilderBase() = default;
 	CORE_API	~TStringBuilderBase();
 
@@ -331,9 +324,13 @@ public:
 	 *
 	 * @return The builder, to allow additional operations to be composed with this one.
 	 */
-	template <typename RangeType, typename DelimiterType,
-		std::enable_if_t<TCanAppendRange_V<RangeType&&> && TCanAppend_V<DelimiterType&&>>* = nullptr>
-	inline BuilderType& Join(RangeType&& InRange, DelimiterType&& InDelimiter)
+	template <typename RangeType, typename DelimiterType>
+		requires requires (BuilderType& Builder, RangeType&& Range, const DelimiterType& Delimiter)
+		{
+			Builder << *std::begin(Range);
+			Builder << Delimiter;
+		}
+	inline BuilderType& Join(RangeType&& InRange, const DelimiterType& InDelimiter)
 	{
 		bool bFirst = true;
 		for (auto&& Elem : Forward<RangeType>(InRange))
@@ -364,9 +361,13 @@ public:
 	 *
 	 * @return The builder, to allow additional operations to be composed with this one.
 	 */
-	template <typename RangeType, typename DelimiterType, typename QuoteType,
-		std::enable_if_t<TCanAppendRange_V<RangeType> && TCanAppend_V<DelimiterType&&> && TCanAppend_V<QuoteType&&>>* = nullptr>
-	inline BuilderType& JoinQuoted(RangeType&& InRange, DelimiterType&& InDelimiter, QuoteType&& InQuote)
+	template <typename RangeType, typename DelimiterType, typename QuoteType>
+		requires requires (BuilderType& Builder, RangeType&& Range, const DelimiterType& Delimiter, const QuoteType& Quote)
+	{
+		Builder << Quote << *std::begin(Range) << Quote;
+		Builder << Delimiter;
+	}
+	inline BuilderType& JoinQuoted(RangeType&& InRange, const DelimiterType& InDelimiter, const QuoteType& InQuote)
 	{
 		bool bFirst = true;
 		for (auto&& Elem : Forward<RangeType>(InRange))
diff --git a/Engine/Source/Runtime/Core/Tests/Misc/StringBuilderTest.cpp b/Engine/Source/Runtime/Core/Tests/Misc/StringBuilderTest.cpp
--- a/Engine/Source/Runtime/Core/Tests/Misc/StringBuilderTest.cpp
+++ b/Engine/Source/Runtime/Core/Tests/Misc/StringBuilderTest.cpp
@@ -4,12 +4,43 @@
 
 #include "Misc/StringBuilder.h"
 
+#include "Containers/Array.h"
+#include "Containers/Set.h"
 #include "Containers/UnrealString.h"
 #include "String/Find.h"
 #include "TestHarness.h"
 
 #include <catch2/generators/catch_generators.hpp>
 
+namespace UE::Core::Private
+{
+
+template <typename T>
+	requires requires (FStringBuilderBase& Builder, const T& Value) { Builder << Value; }
+struct TAppendable
+{
+	T Value{};
+
+	inline bool operator==(const TAppendable& Other) const
+		requires requires (const T& V) { V == V; }
+	{
+		return Value == Other.Value;
+	}
+
+	friend uint32 GetTypeHash(const TAppendable& Appendable)
+		requires requires (const T& V) { GetTypeHash(V); }
+	{
+		return GetTypeHash(Appendable.Value);
+	}
+
+	friend FStringBuilderBase& operator<<(FStringBuilderBase& Builder, const TAppendable& Appendable)
+	{
+		return Builder << Appendable.Value;
+	}
+};
+
+} // UE::Core::Private
+
 TEST_CASE("Core::String::StringBuilder", "[Core][String][Smoke]")
 {
 	SECTION("Static")
@@ -282,6 +313,32 @@ TEST_CASE("Core::String::StringBuilder", "[Core][String][Smoke]")
 		Builder.RemoveAt(0, 2);
 		CHECK(Builder == ANSITEXTVIEW(""));
 	}
+
+	SECTION("Join")
+	{
+		using namespace UE::Core::Private;
+
+		TStringBuilder<64> Builder;
+		Builder.Join(TArray<TAppendable<int32>>{{1}, {2}, {3}}, TAppendable{", "});
+		CHECK(Builder == TEXTVIEW("1, 2, 3"));
+
+		Builder.Reset();
+		Builder.Join(TSet<TAppendable<int32>>{{1}, {2}, {3}}, TAppendable{", "});
+		CHECK(Builder == TEXTVIEW("1, 2, 3"));
+	}
+
+	SECTION("JoinQuoted")
+	{
+		using namespace UE::Core::Private;
+
+		TStringBuilder<64> Builder;
+		Builder.JoinQuoted(TArray<TAppendable<int32>>{{1}, {2}, {3}}, TAppendable{", "}, TAppendable{'"'});
+		CHECK(Builder == TEXTVIEW(R"("1", "2", "3")"));
+
+		Builder.Reset();
+		Builder.JoinQuoted(TSet<TAppendable<int32>>{{1}, {2}, {3}}, TAppendable{", "}, TAppendable{'"'});
+		CHECK(Builder == TEXTVIEW(R"("1", "2", "3")"));
+	}
 }
 
 #endif
\ No newline at end of file
