Add way to get the landscape spline segment key interval and evaluate some key to get the location

diff --git a/Engine/Source/Runtime/Core/Public/Math/InterpCurve.h b/Engine/Source/Runtime/Core/Public/Math/InterpCurve.h
--- a/Engine/Source/Runtime/Core/Public/Math/InterpCurve.h
+++ b/Engine/Source/Runtime/Core/Public/Math/InterpCurve.h
@@ -77,7 +77,7 @@ public:
 	void ClearLoopKey();
 
 	/** 
-	 *	Evaluate the output for an arbitary input value. 
+	 *	Evaluate the output for an arbitrary input value. 
 	 *	For inputs outside the range of the keys, the first/last key value is assumed.
 	 */
 	[[nodiscard]] T Eval( const float InVal, const T& Default = T(ForceInit) ) const;
@@ -140,6 +140,9 @@ public:
 	/** Calculate the min/max out value that can be returned by this InterpCurve. */
 	void CalcBounds(T& OutMin, T& OutMax, const T& Default = T(ForceInit)) const;
 
+	/** Returns the interval of key for the curve points */
+	FFloatInterval GetKeyInterval() const;
+
 public:
 
 	/**
@@ -769,7 +772,16 @@ void FInterpCurve<T>::CalcBounds(T& OutMin, T& OutMax, const T& Default) const
 	}
 }
 
+template <class T>
+FFloatInterval FInterpCurve<T>::GetKeyInterval() const
+{
+	if (Points.IsEmpty())
+	{
+		return {};
+	}
 
+	return FFloatInterval(Points[0].InVal, Points.Last().InVal);
+}
 
 /* Common type definitions
  *****************************************************************************/
diff --git a/Engine/Source/Runtime/Landscape/Classes/LandscapeSplineSegment.h b/Engine/Source/Runtime/Landscape/Classes/LandscapeSplineSegment.h
--- a/Engine/Source/Runtime/Landscape/Classes/LandscapeSplineSegment.h
+++ b/Engine/Source/Runtime/Landscape/Classes/LandscapeSplineSegment.h
@@ -362,6 +362,12 @@ public:
 
 	virtual void FindNearest(const FVector& InLocation, float& t, FVector& OutLocation, FVector& OutTangent);
 
+	/** Return the key interval for the segment curve points */
+	LANDSCAPE_API FFloatInterval GetKeyInterval() const;
+
+	/** Evaluate the output for an arbitrary input value. For keys outside the range, the first/last key value is assumed. */
+	LANDSCAPE_API FVector Evaluate(const float Key, const FVector& Default = FVector::ZeroVector) const;
+
 	//~ Begin UObject Interface
 	virtual void Serialize(FArchive& Ar) override;
 	virtual void PostLoad() override;
diff --git a/Engine/Source/Runtime/Landscape/Private/LandscapeSplines.cpp b/Engine/Source/Runtime/Landscape/Private/LandscapeSplines.cpp
--- a/Engine/Source/Runtime/Landscape/Private/LandscapeSplines.cpp
+++ b/Engine/Source/Runtime/Landscape/Private/LandscapeSplines.cpp
@@ -3628,6 +3628,16 @@ void ULandscapeSplineSegment::FindNearest( const FVector& InLocation, float& t,
 	OutTangent = SplineInfo.EvalDerivative(t, FVector::ZeroVector);
 }
 
+FFloatInterval ULandscapeSplineSegment::GetKeyInterval() const
+{
+	return SplineInfo.GetKeyInterval();
+}
+
+FVector ULandscapeSplineSegment::Evaluate(const float Key, const FVector& Default) const
+{
+	return SplineInfo.Eval(Key, Default);
+}
+
 #if WITH_EDITOR
 void ULandscapeSplineSegment::PostEditUndo()
 {
