- Fix bExposed differences in the new compilation mode vs the old compilation mode
- The old compilation mode would call IsExposed() which verified the parameter type vs just checking the bool
- This only appeared in some newer content as old content had the correct set of options checked (perhaps by luck) but new content did not

#rnx

diff --git a/Engine/Plugins/FX/Niagara/Source/NiagaraEditor/Private/NiagaraGraphDigest.cpp b/Engine/Plugins/FX/Niagara/Source/NiagaraEditor/Private/NiagaraGraphDigest.cpp
--- a/Engine/Plugins/FX/Niagara/Source/NiagaraEditor/Private/NiagaraGraphDigest.cpp
+++ b/Engine/Plugins/FX/Niagara/Source/NiagaraEditor/Private/NiagaraGraphDigest.cpp
@@ -3280,8 +3280,8 @@ FNiagaraCompilationNodeInput::FNiagaraCompilationNodeInput(const UNiagaraNodeInp
 		ObjectAssetPath = FSoftObjectPath(SourceObjectAsset);
 	}
 
+	bExposed = InNode->IsExposed();
 	bRequired = InNode->ExposureOptions.bRequired;
-	bExposed = InNode->ExposureOptions.bExposed;
 	bCanAutoBind = InNode->ExposureOptions.bCanAutoBind;
 }
 
diff --git a/Engine/Source/Runtime/Core/Private/UObject/UE5ReleaseStreamObjectVersion.cpp b/Engine/Source/Runtime/Core/Private/UObject/UE5ReleaseStreamObjectVersion.cpp
--- a/Engine/Source/Runtime/Core/Private/UObject/UE5ReleaseStreamObjectVersion.cpp
+++ b/Engine/Source/Runtime/Core/Private/UObject/UE5ReleaseStreamObjectVersion.cpp
@@ -12,7 +12,7 @@ TMap<FGuid, FGuid> FUE5ReleaseStreamObjectVersion::GetSystemGuids()
 	SystemGuids.Add(DevGuids.MATERIALSHADERMAP_DERIVEDDATA_VER, FGuid("79090A66F9D94B5BB285D49E5D39468E"));
 	SystemGuids.Add(DevGuids.NANITE_DERIVEDDATA_VER, FGuid("AC88CFBDDA614A9C8488F64D2DAF699D"));
 	SystemGuids.Add(DevGuids.NIAGARASHADERMAP_DERIVEDDATA_VER, FGuid("7BCE96ABAA6E46D689306FE673DED4DE"));
-	SystemGuids.Add(DevGuids.Niagara_LatestScriptCompileVersion, FGuid("85D2347062C644AAAD31E3E6CBA1FDA0"));
+	SystemGuids.Add(DevGuids.Niagara_LatestScriptCompileVersion, FGuid("E5A666769537104AB3D9C7A590BC7FAC"));
 	SystemGuids.Add(DevGuids.SkeletalMeshDerivedDataVersion, FGuid("BA91E71F642E4813A1B74B80D0448928"));
 	SystemGuids.Add(DevGuids.STATICMESH_DERIVEDDATA_VER, FGuid("7A556175518D458AA8786ED69AB8DDE7"));
 
