Logging: Optimized allocation of static log templates

Pack static log templates into 4 KiB blocks that are freed on exit without requiring destruction of the individual templates.

diff --git a/Engine/Source/Runtime/Core/Private/Logging/StructuredLog.cpp b/Engine/Source/Runtime/Core/Private/Logging/StructuredLog.cpp
--- a/Engine/Source/Runtime/Core/Private/Logging/StructuredLog.cpp
+++ b/Engine/Source/Runtime/Core/Private/Logging/StructuredLog.cpp
@@ -421,9 +421,6 @@ public:
 
 	const TCHAR* GetFormat() const { return StaticFormat; }
 
-	FLogTemplate* GetNext() { return Next; }
-	void SetNext(FLogTemplate* Template) { Next = Template; }
-
 	uint8* GetOpData() { return (uint8*)(this + 1); }
 	const uint8* GetOpData() const { return (const uint8*)(this + 1); }
 
@@ -447,7 +444,6 @@ private:
 	FLogTemplate& operator=(const FLogTemplate&) = delete;
 
 	const TCHAR* StaticFormat = nullptr;
-	FLogTemplate* Next = nullptr;
 };
 
 static_assert(std::is_trivially_destructible_v<FLogTemplate>);
@@ -1144,26 +1140,82 @@ FORCENOINLINE static void BasicLogToTrace(const void* LogPoint, const TCHAR* For
 
 ///////////////////////////////////////////////////////////////////////////////////////////////////////////////////////
 
-struct FStaticLogTemplateManager
+class FStaticLogTemplateAllocator
 {
-	std::atomic<FLogTemplate*> Head = nullptr;
+public:
+	~FStaticLogTemplateAllocator()
+	{
+		for (FBlock* Block = Tail; Block;)
+		{
+			FBlock* Previous = Block->Previous;
+			FMemory::Free(Block);
+			Block = Previous;
+		}
+	}
 
-	~FStaticLogTemplateManager()
+	void* Allocate(int32 Size)
 	{
-		// Skip this at-exit shutdown code, because in some circumstances the GetLocalizedLogFormatTable used by
-		// FLogTemplate::Destroy can already be torn down. This class will be removed in an upcoming refactor.
-#if 0
-		for (FLogTemplate* Template = Head.exchange(nullptr); Template;)
+		TUniqueLock Lock(Mutex);
+		if (!Tail || TailOffset + Size > Tail->Size)
 		{
-			FLogTemplate* NextTemplate = Template->GetNext();
-			FLogTemplate::Destroy(Template);
-			Template = NextTemplate;
+			int32 NewSize = (sizeof(FBlock) + Size + BlockSize - 1) & ~(BlockSize - 1);
+			Tail = new(FMemory::Malloc(NewSize, alignof(FBlock))) FBlock{Tail, NewSize};
+			TailOffset = sizeof(FBlock);
 		}
-#endif
+		void* Address = (uint8*)Tail + TailOffset;
+		// TODO: Aligned to 8 until unaligned pointer access has been tested on every platform.
+		TailOffset += (Size + 7) & ~7;
+		return Address;
 	}
+
+private:
+	inline constexpr static int32 BlockSize = 4096;
+
+	struct FBlock
+	{
+		FBlock* Previous = nullptr;
+		int32 Size = 0;
+	};
+
+	FBlock* Tail = nullptr;
+	int32 TailOffset = 0;
+	FMutex Mutex;
 };
 
-static FStaticLogTemplateManager GStaticLogTemplateManager;
+static FStaticLogTemplateAllocator& GetStaticLogTemplateAllocator()
+{
+	static FStaticLogTemplateAllocator Allocator;
+	return Allocator;
+}
+
+class FStaticLogTemplateStorage
+{
+public:
+	FStaticLogTemplateStorage() = default;
+	FStaticLogTemplateStorage(const FStaticLogTemplateStorage&) = delete;
+	FStaticLogTemplateStorage& operator=(const FStaticLogTemplateStorage&) = delete;
+
+	inline void* Allocate(int32 Size)
+	{
+		Data = GetStaticLogTemplateAllocator().Allocate(Size);
+		return Data;
+	}
+
+	inline const void* Get() const
+	{
+		return Data;
+	}
+
+	inline void* Detach()
+	{
+		ON_SCOPE_EXIT { Data = nullptr; };
+		return Data;
+	}
+
+private:
+	// Intentionally leaked because the allocator will free the memory on exit.
+	void* Data = nullptr;
+};
 
 // Serializing log fields to compact binary happens in its own function because that allows stack
 // space for the writer to be returned before calling into the output devices.
@@ -1266,31 +1318,22 @@ FORCENOINLINE static FLogTemplate& CreateLogTemplate(const FLogCategoryBase& Cat
 
 		struct FTemplateCreator
 		{
-			static TLogTemplate<FMemoryLogTemplateStorage> Create(const FStaticLogRecord& Log, const FLogField* Fields, int32 FieldCount)
+			static TLogTemplate<FStaticLogTemplateStorage> Create(const FStaticLogRecord& Log, const FLogField* Fields, int32 FieldCount)
 			{
-				return TLogTemplate<FMemoryLogTemplateStorage>(Log.Format, {}, Fields, FieldCount);
+				return TLogTemplate<FStaticLogTemplateStorage>(Log.Format, {}, Fields, FieldCount);
 			}
 
-			static TLogTemplate<FMemoryLogTemplateStorage> Create(const FStaticLocalizedLogRecord& Log, const FLogField* Fields, int32 FieldCount)
+			static TLogTemplate<FStaticLogTemplateStorage> Create(const FStaticLocalizedLogRecord& Log, const FLogField* Fields, int32 FieldCount)
 			{
-				return TLogTemplate<FMemoryLogTemplateStorage>(Log.TextNamespace, Log.TextKey, Log.Format, {}, Fields, FieldCount);
+				return TLogTemplate<FStaticLogTemplateStorage>(Log.TextNamespace, Log.TextKey, Log.Format, {}, Fields, FieldCount);
 			}
 		};
 
-		TLogTemplate<FMemoryLogTemplateStorage> LocalTemplate = FTemplateCreator::Create(Log, Fields, FieldCount);
+		TLogTemplate<FStaticLogTemplateStorage> LocalTemplate = FTemplateCreator::Create(Log, Fields, FieldCount);
 		FLogTemplate* NewTemplate = LocalTemplate.Get();
 		if (LIKELY(Log.DynamicData.Template.compare_exchange_strong(Template, NewTemplate, std::memory_order_release, std::memory_order_acquire)))
 		{
 			LocalTemplate.Detach();
-			// Register the template to destroy on exit.
-			for (FLogTemplate* Head = GStaticLogTemplateManager.Head.load(std::memory_order_relaxed);;)
-			{
-				NewTemplate->SetNext(Head);
-				if (LIKELY(GStaticLogTemplateManager.Head.compare_exchange_weak(Head, NewTemplate, std::memory_order_release, std::memory_order_relaxed)))
-				{
-					break;
-				}
-			}
 			return *NewTemplate;
 		}
 	}
