Fix issue with asset with starting offset other than 0

#android

diff --git a/Engine/Source/Runtime/Core/Private/Android/AndroidPlatformFile.cpp b/Engine/Source/Runtime/Core/Private/Android/AndroidPlatformFile.cpp
--- a/Engine/Source/Runtime/Core/Private/Android/AndroidPlatformFile.cpp
+++ b/Engine/Source/Runtime/Core/Private/Android/AndroidPlatformFile.cpp
@@ -304,15 +304,12 @@ public:
 		: File(MakeShareable(new FileReference(path, asset)))
 		, Start(0), Length(0), CurrentOffset(0)
 	{
-#if UE_ANDROID_FILE_64 && PLATFORM_32BITS
-		File->Handle = AAsset_openFileDescriptor64(File->Asset, &Start, &Length);
-#else
-		off_t OutStart = Start;
-		off_t OutLength = Length;
-		File->Handle = AAsset_openFileDescriptor(File->Asset, &OutStart, &OutLength);
+		off64_t OutStart = Start;
+		off64_t OutLength = Length;
+		File->Handle = AAsset_openFileDescriptor64(File->Asset, &OutStart, &OutLength);
 		Start = OutStart;
 		Length = OutLength;
-#endif
+		CurrentOffset = Start;
 		CheckValid();
 		LogInfo();
 	}
