Fix potential stalls on certain iOS caused by querying NSURLVolumeAvailableCapacityForImportantUsageKey. In the default case we will not query NSURLVolumeAvailableCapacityKey and only optionally query "For Important Usage" if required by client application.

diff --git a/Engine/Source/Runtime/Core/Private/IOS/IOSPlatformMisc.cpp b/Engine/Source/Runtime/Core/Private/IOS/IOSPlatformMisc.cpp
--- a/Engine/Source/Runtime/Core/Private/IOS/IOSPlatformMisc.cpp
+++ b/Engine/Source/Runtime/Core/Private/IOS/IOSPlatformMisc.cpp
@@ -933,7 +933,7 @@ FString FIOSPlatformMisc::GetOSVersion()
 	return FString([[UIDevice currentDevice] systemVersion]);
 }
 
-bool FIOSPlatformMisc::GetDiskTotalAndFreeSpace(const FString& InPath, uint64& TotalNumberOfBytes, uint64& NumberOfFreeBytes)
+bool FIOSPlatformMisc::GetDiskTotalAndFreeSpace(const FString& InPath, uint64& TotalNumberOfBytes, uint64& NumberOfFreeBytes, bool IncludeFreeable)
 {
     //On iOS 11 use new method to return disk space available for important usages
 #if !PLATFORM_TVOS
@@ -941,7 +941,7 @@ bool FIOSPlatformMisc::GetDiskTotalAndFreeSpace(const FString& InPath, uint64& T
     
     NSNumber *FreeBytes = nil;
     NSURL *URL = [NSURL fileURLWithPath : NSHomeDirectory()];
-    GetValueSuccess = [URL getResourceValue : &FreeBytes forKey : NSURLVolumeAvailableCapacityForImportantUsageKey error : nil];
+    GetValueSuccess = [URL getResourceValue : &FreeBytes forKey : (IncludeFreeable ? NSURLVolumeAvailableCapacityForImportantUsageKey : NSURLVolumeAvailableCapacityKey) error : nil];
     if (FreeBytes)
     {
         NumberOfFreeBytes = [FreeBytes longLongValue];
diff --git a/Engine/Source/Runtime/Core/Public/IOS/IOSPlatformMisc.h b/Engine/Source/Runtime/Core/Public/IOS/IOSPlatformMisc.h
--- a/Engine/Source/Runtime/Core/Public/IOS/IOSPlatformMisc.h
+++ b/Engine/Source/Runtime/Core/Public/IOS/IOSPlatformMisc.h
@@ -117,7 +117,7 @@ struct CORE_API FIOSPlatformMisc : public FApplePlatformMisc
 	static FString GetDeviceId();
 	static FString GetOSVersion();
 	static FString GetUniqueAdvertisingId();
-	static bool GetDiskTotalAndFreeSpace(const FString& InPath, uint64& TotalNumberOfBytes, uint64& NumberOfFreeBytes);
+	static bool GetDiskTotalAndFreeSpace(const FString& InPath, uint64& TotalNumberOfBytes, uint64& NumberOfFreeBytes, bool IncludeFreeable = false);
 	
 	static void RequestStoreReview();
 
