- Fix CIS failure due to MarketplaceKit dependency not being contained to IOS platform

diff --git a/Engine/Source/Runtime/Core/Core.Build.cs b/Engine/Source/Runtime/Core/Core.Build.cs
--- a/Engine/Source/Runtime/Core/Core.Build.cs
+++ b/Engine/Source/Runtime/Core/Core.Build.cs
@@ -42,7 +42,6 @@ public class Core : ModuleRules
 		PrivateIncludePathModuleNames.AddRange(
 			new string[] {
 				"DerivedDataCache",
-				"MarketplaceKit",
 				"TargetPlatform",
 
 				"Json",
@@ -144,6 +143,8 @@ public class Core : ModuleRules
 				AddEngineThirdPartyPrivateStaticDependencies(Target,
 					"PLCrashReporter"
 					);
+
+				PrivateIncludePathModuleNames.Add("MarketplaceKit");
 			}
 			if (Target.Platform == UnrealTargetPlatform.VisionOS)
 			{
