- Fix CIS failure due to including MarketplaceKit.h in TVOS builds.

diff --git a/Engine/Source/Runtime/Core/Private/IOS/IOSPlatformMisc.cpp b/Engine/Source/Runtime/Core/Private/IOS/IOSPlatformMisc.cpp
--- a/Engine/Source/Runtime/Core/Private/IOS/IOSPlatformMisc.cpp
+++ b/Engine/Source/Runtime/Core/Private/IOS/IOSPlatformMisc.cpp
@@ -20,7 +20,6 @@
 #include "IOS/IOSPlatformCrashContext.h"
 #include "IOS/IOSPlatformPLCrashReporterIncludes.h"
 #include "IOS/IOSView.h"
-#include "MarketplaceKit.h"
 #include "Misc/App.h"
 #include "Misc/CommandLine.h"
 #include "Misc/ConfigCacheIni.h"
@@ -39,7 +38,8 @@
 
 #if !PLATFORM_TVOS && !PLATFORM_VISIONOS
 #include <AdSupport/ASIdentifierManager.h>
-#endif // !PLATFORM_TVOS
+#include "MarketplaceKit.h"
+#endif // !PLATFORM_TVOS && !PLATFORM_VISIONOS
 
 #import <DeviceCheck/DeviceCheck.h>
 #import <Foundation/Foundation.h>
