Improve module manager logging

#rnx

diff --git a/Engine/Source/Runtime/Core/Private/Modules/ModuleManager.cpp b/Engine/Source/Runtime/Core/Private/Modules/ModuleManager.cpp
--- a/Engine/Source/Runtime/Core/Private/Modules/ModuleManager.cpp
+++ b/Engine/Source/Runtime/Core/Private/Modules/ModuleManager.cpp
@@ -694,7 +694,12 @@ void FModuleManager::OnObjectCleanup()
 			if (HandlePtr)
 			{
 				UE_LOG(LogModuleManager, Log, TEXT("OnObjectCleanup: unloading dynamic library '%s'"), *LibraryAndModulesIterator.Key());
+
+				const uint64_t InitialMemUsed = FPlatformMemory::GetMemoryUsedFast();
 				FPlatformProcess::FreeDllHandle(*HandlePtr);
+				const uint64_t FinalMemUsed = FPlatformMemory::GetMemoryUsedFast();
+
+				UE_LOG(LogModuleManager, Verbose, TEXT("OnObjectCleanup: freed about %llu KB"), (InitialMemUsed - FinalMemUsed) / 1024);
 			}
 
 			LibraryAndModulesIterator.RemoveCurrent();
