Expose network changed delegate for iOS

diff --git a/Engine/Source/Runtime/Core/Private/IOS/IOSPlatformMisc.cpp b/Engine/Source/Runtime/Core/Private/IOS/IOSPlatformMisc.cpp
--- a/Engine/Source/Runtime/Core/Private/IOS/IOSPlatformMisc.cpp
+++ b/Engine/Source/Runtime/Core/Private/IOS/IOSPlatformMisc.cpp
@@ -92,6 +92,70 @@ namespace UE::FIOSPlatformMisc
 	// Note: defined here to not force objcisims into the header which is included on non-objc platforms
 	nw_path_monitor_t PathMonitor = nullptr;
 	nw_path_t CachedNetworkPath = nullptr;
+	ENetworkConnectionType LastReportedNetworkType = ENetworkConnectionType::Unknown;
+
+	static ENetworkConnectionType CheckNetworkConnectionType()
+	{
+		nw_path_t MyPath = UE::FIOSPlatformMisc::CachedNetworkPath;
+		if (MyPath == nullptr)
+		{
+			UE_LOG(LogInit, Warning, TEXT("Network path currently NOT set"));
+			return ENetworkConnectionType::Unknown;
+		}
+		
+		nw_path_status_t PathStatus = nw_path_get_status(MyPath);
+		if (PathStatus != nw_path_status_satisfied)
+		{
+			if (PathStatus == nw_path_status_invalid)
+			{
+				UE_LOG(LogInit, Warning, TEXT("Network path is invalid"));
+				return ENetworkConnectionType::None;
+			}
+			
+			// Basically, if there's _no_ network interfaces available, assume it's in airplane mode.
+			// Note: if you have wifi enabled - but airplane mode also enabled, that will NOT be detected
+			//       as airplane mode as you still have a valid network path.
+			__block bool HasAnInterface = false;
+			nw_path_enumerate_interfaces(MyPath, (nw_path_enumerate_interfaces_block_t) ^ (nw_interface_t Interface)
+			{
+				const char* Name = nw_interface_get_name(Interface);
+				// UE_LOG(LogIOS, Log, TEXT("NW path monitor available interface: %s"), *FString(Name));
+				HasAnInterface = true;
+				return false;
+			});
+			
+			if (!HasAnInterface)
+			{
+				return ENetworkConnectionType::AirplaneMode;
+			}
+			return ENetworkConnectionType::None;
+		}
+		else
+		{
+			// If satisfied detect the type of connection or maybe use isExpensive???
+			bool bHasActiveWiFiConnection = nw_path_uses_interface_type(MyPath, nw_interface_type_wifi);
+			bool bHasActiveCellConnection = nw_path_uses_interface_type(MyPath, nw_interface_type_cellular);
+			bool bHasActiveWiredConnection = nw_path_uses_interface_type(MyPath, nw_interface_type_wired);
+
+			if (bHasActiveWiFiConnection)
+			{
+				return ENetworkConnectionType::WiFi;
+			}
+			else if (bHasActiveCellConnection)
+			{
+				return ENetworkConnectionType::Cell;
+			}
+			else if (bHasActiveWiredConnection)
+			{
+				return ENetworkConnectionType::Ethernet;
+			}
+			else
+			{
+				// For now, mw_interface_type_loopback and nw_interface_type_other return "Unknown"
+				return ENetworkConnectionType::Unknown;
+			}
+		}
+	}
 }
 
 void FIOSPlatformMisc::PlatformInit()
@@ -183,16 +247,27 @@ void FIOSPlatformMisc::PlatformInit()
 	nw_path_monitor_set_update_handler(
 		UE::FIOSPlatformMisc::PathMonitor, 
 		^(nw_path_t InPath)
-		{ 
+		{
 			UE_LOG(LogInit, Log, TEXT("Network path changed"));
 			
 			// cache the path we get
 			if (UE::FIOSPlatformMisc::PathMonitor != nullptr)
 			{
-				nw_release(UE::FIOSPlatformMisc::CachedNetworkPath);	
+				nw_release(UE::FIOSPlatformMisc::CachedNetworkPath);
 			}
 			UE::FIOSPlatformMisc::CachedNetworkPath = InPath;
 			nw_retain(UE::FIOSPlatformMisc::CachedNetworkPath);
+			if (FTaskGraphInterface::IsRunning())
+			{
+				ENetworkConnectionType ConnectionType = UE::FIOSPlatformMisc::CheckNetworkConnectionType();
+				if (ConnectionType != UE::FIOSPlatformMisc::LastReportedNetworkType)
+				{
+					UE::FIOSPlatformMisc::LastReportedNetworkType = ConnectionType;
+					FFunctionGraphTask::CreateAndDispatchWhenReady([ConnectionType]() {
+						FCoreDelegates::OnNetworkConnectionChanged.Broadcast(ConnectionType);
+					}, TStatId(), NULL, ENamedThreads::GameThread);
+				}
+			}
 		}
 	);
 	
@@ -1624,69 +1699,6 @@ void FIOSPlatformMisc::GetValidTargetPlatforms(TArray<FString>& TargetPlatformNa
 #endif
 }
 
-static inline ENetworkConnectionType CheckNetworkConnectionType()
-{
-	nw_path_t MyPath = UE::FIOSPlatformMisc::CachedNetworkPath;
-	if (MyPath == nullptr)
-	{
-		UE_LOG(LogInit, Warning, TEXT("Network path currently NOT set"));
-		return ENetworkConnectionType::Unknown;
-	}
-	
-	nw_path_status_t PathStatus = nw_path_get_status(MyPath);
-	if (PathStatus != nw_path_status_satisfied)
-	{
-		if (PathStatus == nw_path_status_invalid)
-		{
-			UE_LOG(LogInit, Warning, TEXT("Network path is invalid"));
-			return ENetworkConnectionType::None;	
-		}
-		
-		// Basically, if there's _no_ network interfaces available, assume it's in airplane mode.
-		// Note: if you have wifi enabled - but airplane mode also enabled, that will NOT be detected
-		//       as airplane mode as you still have a valid network path.
-		__block bool HasAnInterface = false;
-		nw_path_enumerate_interfaces(MyPath, (nw_path_enumerate_interfaces_block_t) ^ (nw_interface_t Interface)
-		{
-			const char* Name = nw_interface_get_name(Interface);
-			// UE_LOG(LogIOS, Log, TEXT("NW path monitor available interface: %s"), *FString(Name));
-			HasAnInterface = true;
-			return false;
-		});
-		
-		if (!HasAnInterface)
-		{
-			return ENetworkConnectionType::AirplaneMode;
-		}
-		return ENetworkConnectionType::None;
-	}
-	else 
-	{
-		// If satisfied detect the type of connection or maybe use isExpensive???
-		bool bHasActiveWiFiConnection = nw_path_uses_interface_type(MyPath, nw_interface_type_wifi);
-		bool bHasActiveCellConnection = nw_path_uses_interface_type(MyPath, nw_interface_type_cellular);
-		bool bHasActiveWiredConnection = nw_path_uses_interface_type(MyPath, nw_interface_type_wired);	
-
-		if (bHasActiveWiFiConnection)
-		{
-			return ENetworkConnectionType::WiFi;
-		}
-		else if (bHasActiveCellConnection)
-		{
-			return ENetworkConnectionType::Cell;
-		}
-		else if (bHasActiveWiredConnection)
-		{
-			return ENetworkConnectionType::Ethernet;
-		}
-		else
-		{
-			// For now, mw_interface_type_loopback and nw_interface_type_other return "Unknown"
-			return ENetworkConnectionType::Unknown;
-		}
-	}
-}
-
 ENetworkConnectionType FIOSPlatformMisc::GetNetworkConnectionType()
 {
 	static TOptional<ENetworkConnectionType> ConnectionType = {};
@@ -1697,7 +1709,7 @@ ENetworkConnectionType FIOSPlatformMisc::GetNetworkConnectionType()
 
 	if (!ConnectionType.IsSet() || CurrentTime >= LastCheckTime + CheckInterval)
 	{
-		ConnectionType = CheckNetworkConnectionType();
+		ConnectionType = UE::FIOSPlatformMisc::CheckNetworkConnectionType();
 		LastCheckTime = CurrentTime;
 	}
 
