[Scheduler]
- Remove 5.5 deprecations

diff --git a/Engine/Source/Runtime/Core/Public/Async/Fundamental/Scheduler.h b/Engine/Source/Runtime/Core/Public/Async/Fundamental/Scheduler.h
--- a/Engine/Source/Runtime/Core/Public/Async/Fundamental/Scheduler.h
+++ b/Engine/Source/Runtime/Core/Public/Async/Fundamental/Scheduler.h
@@ -34,82 +34,6 @@ namespace LowLevelTasks
 		DefaultPreference = LocalQueuePreference,
 	};
 
-	//implementation of a treiber stack
-	//(https://en.wikipedia.org/wiki/Treiber_stack)
-	template<typename NodeType>
-	class UE_DEPRECATED(5.5, "This class will be removed.") TEventStack
-	{
-		static constexpr uint32 EVENT_INDEX_NONE = ~0u;
-
-		struct FTopNode
-		{
-			uint32 EventIndex;		// All events are stored in an array, so this allows us to pack node info into 32 bit as array index
-			uint32 Revision;		// Tagging is used to avoid ABA (https://en.wikipedia.org/wiki/ABA_problem#Tagged_state_reference)
-		};
-		std::atomic<FTopNode> Top { FTopNode{EVENT_INDEX_NONE, 0} };
-		TAlignedArray<NodeType>& NodesArray;
-
-		uint32 GetNodeIndex(const NodeType* Node) const
-		{
-			return (Node == nullptr) ? EVENT_INDEX_NONE : uint32(Node - NodesArray.GetData());
-		}
-
-	public:
-		TEventStack(TAlignedArray<NodeType>& NodesArray)
-			: NodesArray(NodesArray)
-		{
-		}
-
-		NodeType* Pop()
-		{
-			FTopNode LocalTop = Top.load(std::memory_order_acquire);
-			while (true) 
-			{			
-				if (LocalTop.EventIndex == EVENT_INDEX_NONE)
-				{
-					return nullptr;
-				}
-#if DO_CHECK
-				const int64 LastRevision = int64(LocalTop.Revision);
-#endif
-				NodeType* Item = &NodesArray[LocalTop.EventIndex];
-				if (Top.compare_exchange_weak(LocalTop, FTopNode { GetNodeIndex(Item->Next.load(std::memory_order_relaxed)), uint32(LocalTop.Revision + 1) }, std::memory_order_relaxed, std::memory_order_relaxed))
-				{
-					Item->Next.store(nullptr, std::memory_order_relaxed);
-					return Item;
-				}
-#if DO_CHECK
-				const int64 NewRevision = int64(LocalTop.Revision) < LastRevision ? (int64(UINT32_MAX) + int64(LocalTop.Revision)) : int64(LocalTop.Revision);
-				ensureMsgf((NewRevision - LastRevision) < (1ll << 31), TEXT("Dangerously close to the wraparound: %" INT64_FMT ", %" INT64_FMT), LastRevision, NewRevision);
-#endif
-			}
-		}
-
-		void Push(NodeType* Item)
-		{
-			checkSlow(Item != nullptr);
-			checkSlow(Item->Next.load(std::memory_order_relaxed) == nullptr);
-			checkfSlow((Item >= NodesArray.GetData()) && (Item < (NodesArray.GetData() + NodesArray.Num())), TEXT("Item doesn't belong to a Nodes Array"));
-
-			FTopNode LocalTop = Top.load(std::memory_order_relaxed);
-			while (true) 
-			{
-#if DO_CHECK
-				const int64 LastRevision = int64(LocalTop.Revision);
-#endif
-				Item->Next.store((LocalTop.EventIndex == EVENT_INDEX_NONE) ? nullptr : &NodesArray[LocalTop.EventIndex], std::memory_order_relaxed);
-				if (Top.compare_exchange_weak(LocalTop, FTopNode { GetNodeIndex(Item), uint32(LocalTop.Revision + 1) }, std::memory_order_release, std::memory_order_relaxed))
-				{
-					return;
-				}
-#if DO_CHECK
-				const int64 NewRevision = int64(LocalTop.Revision) < LastRevision ? (int64(UINT32_MAX) + int64(LocalTop.Revision)) : int64(LocalTop.Revision);
-				ensureMsgf((NewRevision - LastRevision) < (1ll << 31), TEXT("Dangerously close to the wraparound: %" INT64_FMT ", %" INT64_FMT), LastRevision, NewRevision);
-#endif
-			}
-		}
-	};
-
 	class FSchedulerTls
 	{
 	protected:
@@ -179,9 +103,6 @@ namespace LowLevelTasks
 	public:
 		CORE_API bool IsWorkerThread() const;
 
-		// returns true if the current thread execution is in the context of busy-waiting
-		CORE_API static bool IsBusyWaiting();
-
 	protected:
 		CORE_API bool HasPendingWakeUp() const;
 
@@ -209,22 +130,6 @@ namespace LowLevelTasks
 		//try to launch the task, the return value will specify if the task was in the ready state and has been launched
 		inline bool TryLaunch(FTask& Task, EQueuePreference QueuePreference = EQueuePreference::DefaultPreference, bool bWakeUpWorker = true);
 
-		//tries to do some work until the Task is completed
-		template<typename TaskType>
-		UE_DEPRECATED(5.5, "This method will be removed.")
-		inline void BusyWait(const TaskType& Task, bool ForceAllowBackgroundWork = false);
-
-		//tries to do some work until the Conditional return true
-		template<typename Conditional>
-		UE_DEPRECATED(5.5, "This method will be removed.")
-		inline void BusyWaitUntil(Conditional&& Cond, bool ForceAllowBackgroundWork = false);
-
-		//tries to do some work until all the Tasks are completed
-		//the template parameter can be any Type that has a const conversion operator to FTask
-		template<typename TaskType>
-		UE_DEPRECATED(5.5, "This method will be removed.")
-		inline void BusyWait(const TArrayView<const TaskType>& Tasks, bool ForceAllowBackgroundWork = false);
-
 		//number of instantiated workers
 		inline uint32 GetNumWorkers() const;
 
@@ -253,7 +158,6 @@ namespace LowLevelTasks
 		void StandbyLoop(Private::FWaitEvent* WorkerEvent, FSchedulerTls::FLocalQueueType* ExternalWorkerLocalQueue, uint32 WaitCycles, bool bPermitBackgroundWork);
 		void WorkerLoop(Private::FWaitEvent* WorkerEvent, FSchedulerTls::FLocalQueueType* ExternalWorkerLocalQueue, uint32 WaitCycles, bool bPermitBackgroundWork);
 		CORE_API void LaunchInternal(FTask& Task, EQueuePreference QueuePreference, bool bWakeUpWorker);
-		CORE_API void BusyWaitInternal(const FConditional& Conditional, bool ForceAllowBackgroundWork);
 		inline bool WakeUpWorker(bool bBackgroundWorker);
 		CORE_API void IncrementOversubscription();
 		CORE_API void DecrementOversubscription();
@@ -367,32 +271,6 @@ namespace LowLevelTasks
 		return FScheduler::Get().TryLaunch(Task, QueuePreference, bWakeUpWorker);
 	}
 
-	UE_DEPRECATED(5.5, "This method will be removed.")
-	FORCEINLINE_DEBUGGABLE void BusyWaitForTask(const FTask& Task, bool ForceAllowBackgroundWork = false)
-	{
-		PRAGMA_DISABLE_DEPRECATION_WARNINGS
-		FScheduler::Get().BusyWait(Task, ForceAllowBackgroundWork);
-		PRAGMA_ENABLE_DEPRECATION_WARNINGS
-	}
-
-	template<typename Conditional>
-	UE_DEPRECATED(5.5, "This method will be removed.")
-	FORCEINLINE_DEBUGGABLE void BusyWaitUntil(Conditional&& Cond, bool ForceAllowBackgroundWork = false)
-	{
-		PRAGMA_DISABLE_DEPRECATION_WARNINGS
-		FScheduler::Get().BusyWaitUntil<Conditional>(Forward<Conditional>(Cond), ForceAllowBackgroundWork);
-		PRAGMA_ENABLE_DEPRECATION_WARNINGS
-	}
-
-	template<typename TaskType>
-	UE_DEPRECATED(5.5, "This method will be removed.")
-	FORCEINLINE_DEBUGGABLE void BusyWaitForTasks(const TArrayView<const TaskType>& Tasks, bool ForceAllowBackgroundWork = false)
-	{
-		PRAGMA_DISABLE_DEPRECATION_WARNINGS
-		FScheduler::Get().BusyWait<TaskType>(Tasks, ForceAllowBackgroundWork);
-		PRAGMA_ENABLE_DEPRECATION_WARNINGS
-	}
-
 	/******************
 	* IMPLEMENTATION *
 	******************/
@@ -411,49 +289,6 @@ namespace LowLevelTasks
 		return ActiveWorkers.load(std::memory_order_relaxed);
 	}
 
-	template<typename TaskType>
-	inline void FScheduler::BusyWait(const TaskType& Task, bool ForceAllowBackgroundWork)
-	{
-		if(!Task.IsCompleted())
-		{
-			FScheduler::BusyWaitInternal([&Task](){ return Task.IsCompleted(); }, ForceAllowBackgroundWork);
-		}
-	}
-
-	template<typename Conditional>
-	inline void FScheduler::BusyWaitUntil(Conditional&& Cond, bool ForceAllowBackgroundWork)
-	{
-		static_assert(TIsInvocable<Conditional>::Value, "Conditional is not invocable");
-		static_assert(std::is_same_v<decltype(Cond()), bool>, "Conditional must return a boolean");
-		
-		if(!Cond())
-		{
-			FScheduler::BusyWaitInternal(Forward<Conditional>(Cond), ForceAllowBackgroundWork);
-		}
-	}
-
-	template<typename TaskType>
-	inline void FScheduler::BusyWait(const TArrayView<const TaskType>& Tasks, bool ForceAllowBackgroundWork)
-	{
-		auto AllTasksCompleted = [Index(0), &Tasks]() mutable
-		{
-			while (Index < Tasks.Num())
-			{
-				if (!Tasks[Index].IsCompleted())
-				{
-					return false;
-				}
-				Index++;
-			}
-			return true;
-		};
-
-		if (!AllTasksCompleted())
-		{
-			FScheduler::BusyWaitInternal([&AllTasksCompleted](){ return AllTasksCompleted(); }, ForceAllowBackgroundWork);
-		}
-	}
-
 	inline bool FScheduler::WakeUpWorker(bool bBackgroundWorker)
 	{
 		return WaitingQueue[bBackgroundWorker].Notify() != 0;
