Fixed FCountersTrace::TCounter UB from storing invalid enum value in enum-typed field.

diff --git a/Engine/Source/Runtime/Core/Public/ProfilingDebugging/CountersTrace.h b/Engine/Source/Runtime/Core/Public/ProfilingDebugging/CountersTrace.h
--- a/Engine/Source/Runtime/Core/Public/ProfilingDebugging/CountersTrace.h
+++ b/Engine/Source/Runtime/Core/Public/ProfilingDebugging/CountersTrace.h
@@ -21,19 +21,19 @@
 #endif
 #endif
 
-enum ETraceCounterType
+enum ETraceCounterType : uint8
 {
 	TraceCounterType_Int = 0,
 	TraceCounterType_Float = 1,
 };
 
-enum ETraceCounterDisplayHint
+enum ETraceCounterDisplayHint : uint8
 {
 	TraceCounterDisplayHint_None = 0,
 	TraceCounterDisplayHint_Memory = 1,
 };
 
-enum ETraceCounterNameType
+enum ETraceCounterNameType : uint8
 {
 	TraceCounterNameType_Static = 0, // TCounter is allowed to keep a pointer to InCounterName string
 	TraceCounterNameType_Dynamic = 0x10, // TCounter needs to copy the InCounterName string
@@ -64,30 +64,30 @@ struct FCountersTrace
 			: Value(0)
 			, CounterName(InCounterName) // assumes that InCounterName is a static string
 			, CounterId(0)
-			, CounterDisplayHint(InCounterDisplayHint)
+			, CounterFlags(uint8(InCounterDisplayHint))
 		{
-			CounterId = OutputInitCounter(InCounterName, CounterType, CounterDisplayHint);
+			CounterId = OutputInitCounter(InCounterName, CounterType, InCounterDisplayHint);
 		}
 
 		TCounter(ETraceCounterNameType InCounterNameType, const TCHAR* InCounterName, ETraceCounterDisplayHint InCounterDisplayHint)
 			: Value(0)
 			, CounterName(InCounterName)
 			, CounterId(0)
-			, CounterDisplayHint(InCounterDisplayHint)
+			, CounterFlags(uint8(InCounterDisplayHint))
 		{
-			CounterId = OutputInitCounter(InCounterName, CounterType, CounterDisplayHint);
+			CounterId = OutputInitCounter(InCounterName, CounterType, InCounterDisplayHint);
 
 			if (CounterId == 0 && InCounterNameType == TraceCounterNameType_Dynamic)
 			{
 				// Store counter name for late init. Needs a copy as InCounterName pointer might not be valid later.
 				CounterName = AllocAndCopyCounterName(InCounterName);
-				CounterDisplayHint = ETraceCounterDisplayHint(uint8(CounterDisplayHint) | uint8(TraceCounterNameType_AllocNameCopy));
+				CounterFlags |= uint8(TraceCounterNameType_AllocNameCopy);
 			}
 		}
 
 		~TCounter()
 		{
-			if (uint8(CounterDisplayHint) & uint8(TraceCounterNameType_AllocNameCopy))
+			if (CounterFlags & uint8(TraceCounterNameType_AllocNameCopy))
 			{
 				FreeCounterName(CounterName);
 			}
@@ -98,7 +98,7 @@ struct FCountersTrace
 			uint32 OldId = CounterId.load();
 			if (!OldId)
 			{
-				uint32 NewId = OutputInitCounter(CounterName, CounterType, ETraceCounterDisplayHint(uint8(CounterDisplayHint) & 0xF));
+				uint32 NewId = OutputInitCounter(CounterName, CounterType, ETraceCounterDisplayHint(CounterFlags & 0xF));
 				CounterId.compare_exchange_weak(OldId, NewId);
 			}
 		}
@@ -240,7 +240,7 @@ struct FCountersTrace
 		StoredType Value;
 		const TCHAR* CounterName;
 		std::atomic<uint32> CounterId;
-		ETraceCounterDisplayHint CounterDisplayHint;
+		uint8 CounterFlags; // stores a combination of ETraceCounterDisplayHint and ETraceCounterNameType flags
 	};
 
 	using FCounterInt = TCounter<int64, TraceCounterType_Int>;
