Logging: Documented custom structured log formatting conventions and deprecated the original log template API

diff --git a/Engine/Source/Runtime/Core/Public/Logging/StructuredLog.h b/Engine/Source/Runtime/Core/Public/Logging/StructuredLog.h
--- a/Engine/Source/Runtime/Core/Public/Logging/StructuredLog.h
+++ b/Engine/Source/Runtime/Core/Public/Logging/StructuredLog.h
@@ -276,6 +276,15 @@ UE_API void VisitLogContext(TFunctionRef<void (const FCbField&)> Visitor);
  * Serializes the value to be used in a log message.
  *
  * Overload this when the log behavior needs to differ from general serialization to compact binary.
+ *
+ * There are three ways to perform custom formatting for values that are serialized as an object:
+ * 1. Add a $text field that is a string with the exact text to display.
+ * 2. Add a $format field that is a format string that may reference fields of the object and their sub-objects.
+ * 3. Add a $locformat field that is a localized format string that may reference fields of the object and their
+ *    sub-objects. The namespace and key must be included in $locns and $lockey string fields. Serialize FText
+ *    to the writer using SerializeLogFormat().
+ *
+ * Arrays and objects without custom formatting are converted to JSON.
  */
 template <typename ValueType UE_REQUIRES(TModels_V<CInsertable<FCbWriter&>, ValueType>)>
 inline void SerializeForLog(FCbWriter& Writer, ValueType&& Value)
diff --git a/Engine/Source/Runtime/Core/Public/Logging/StructuredLogFormat.h b/Engine/Source/Runtime/Core/Public/Logging/StructuredLogFormat.h
--- a/Engine/Source/Runtime/Core/Public/Logging/StructuredLogFormat.h
+++ b/Engine/Source/Runtime/Core/Public/Logging/StructuredLogFormat.h
@@ -192,40 +192,25 @@ using FInlineLogTemplate = TLogTemplate<Logging::Private::FInlineLogTemplateStor
 /** A log template that stores its data on the heap. A reasonable default when for non-temporary templates. */
 using FUniqueLogTemplate = TLogTemplate<Logging::Private::FUniqueLogTemplateStorage>;
 
-/**
- * Create a log template with explicit lifetime management.
- *
- * Destroy with DestroyLogTemplate.
- * Prefer FInlineLogTemplate and FUniqueLogTemplate.
- */
+UE_DEPRECATED(5.6, "Use FInlineLogTemplate or FUniqueLogTemplate.")
 inline FLogTemplate* CreateLogTemplate(const TCHAR* Format, const FLogTemplateOptions& Options = {})
 {
 	return TLogTemplate<Logging::Private::FMemoryLogTemplateStorage>(Format, Options).Detach();
 }
 
-/**
- * Create a localized log template with explicit lifetime management.
- *
- * Destroy with DestroyLogTemplate.
- * Prefer FInlineLogTemplate and FUniqueLogTemplate.
- */
+UE_DEPRECATED(5.6, "Use FInlineLogTemplate or FUniqueLogTemplate.")
 inline FLogTemplate* CreateLogTemplate(const FText& Format, const FLogTemplateOptions& Options = {})
 {
 	return TLogTemplate<Logging::Private::FMemoryLogTemplateStorage>(Format, Options).Detach();
 }
 
-/**
- * Create a localized log template with explicit lifetime management.
- *
- * Destroy with DestroyLogTemplate.
- * Prefer FInlineLogTemplate and FUniqueLogTemplate.
- */
+UE_DEPRECATED(5.6, "Use FInlineLogTemplate or FUniqueLogTemplate.")
 inline FLogTemplate* CreateLogTemplate(const TCHAR* TextNamespace, const TCHAR* TextKey, const TCHAR* Format, const FLogTemplateOptions& Options = {})
 {
 	return TLogTemplate<Logging::Private::FMemoryLogTemplateStorage>(TextNamespace, TextKey, Format, Options).Detach();
 }
 
-/** Destroy a log template created by CreateLogTemplate. */
+UE_DEPRECATED(5.6, "Use FInlineLogTemplate or FUniqueLogTemplate.")
 inline void DestroyLogTemplate(FLogTemplate* Template)
 {
 	Logging::Private::DestroyLogTemplate(Template);
