Add missing visionOS exclude for MarketplaceKit check

diff --git a/Engine/Source/Runtime/Core/Private/IOS/IOSPlatformMisc.cpp b/Engine/Source/Runtime/Core/Private/IOS/IOSPlatformMisc.cpp
--- a/Engine/Source/Runtime/Core/Private/IOS/IOSPlatformMisc.cpp
+++ b/Engine/Source/Runtime/Core/Private/IOS/IOSPlatformMisc.cpp
@@ -369,7 +369,7 @@ float FIOSPlatformMisc::GetBrightness()
 	return (float)[[[IOSAppDelegate GetDelegate] window] screen].brightness;
 #else
 	return 1.0f;
-#endif // !PLATFORM_TVOS
+#endif // !PLATFORM_TVOS && !PLATFORM_VISIONOS
 }
 
 void FIOSPlatformMisc::SetBrightness(float Brightness)
@@ -378,7 +378,7 @@ void FIOSPlatformMisc::SetBrightness(float Brightness)
 	dispatch_async(dispatch_get_main_queue(), ^{
 		[[[IOSAppDelegate GetDelegate] window] screen].brightness = Brightness;
 	});
-#endif // !PLATFORM_TVOS
+#endif // !PLATFORM_TVOS && !PLATFORM_VISIONOS
 }
 
 bool FIOSPlatformMisc::IsRunningOnBattery()
@@ -1056,7 +1056,7 @@ bool FIOSPlatformMisc::GetDiskTotalAndFreeSpace(const FString& InPath, uint64& T
 
 void FIOSPlatformMisc::RequestStoreReview()
 {
-#if UE_WITH_STORE_KIT && !PLATFORM_TVOS
+#if UE_WITH_STORE_KIT && !PLATFORM_TVOS && !PLATFORM_VISIONOS
 	if (FMarketplaceKitModule::GetCurrentTypeStatic() == EMarketplaceType::AppStore)
 	{
 		// Deprecated for iOS18.0 & visionOS 2.0.  To be replaced with Swift only StoreKit::RequestReviewAction call (UE-228925)
