StringBuilder: Fixed a compile error in the tests

#rnx

diff --git a/Engine/Source/Runtime/Core/Tests/Misc/StringBuilderTest.cpp b/Engine/Source/Runtime/Core/Tests/Misc/StringBuilderTest.cpp
--- a/Engine/Source/Runtime/Core/Tests/Misc/StringBuilderTest.cpp
+++ b/Engine/Source/Runtime/Core/Tests/Misc/StringBuilderTest.cpp
@@ -39,6 +39,9 @@ struct TAppendable
 	}
 };
 
+template <typename T> TAppendable(const T&) -> TAppendable<T>;
+template <typename T> TAppendable(T&&) -> TAppendable<T>;
+
 } // UE::Core::Private
 
 TEST_CASE("Core::String::StringBuilder", "[Core][String][Smoke]")
