Modify TOptional to allow in-place construction in constexpr contexts which requires us not to use reinterpret_cast.
We can't use std::construct_at on some platforms so only support in-place construction from values, and  not other modifications in a constexpr context.
No support for copying or moving other TOptionals.
Added compile time tests for supported constexpr operations.

diff --git a/Engine/Source/Runtime/Core/Public/HAL/Platform.h b/Engine/Source/Runtime/Core/Public/HAL/Platform.h
--- a/Engine/Source/Runtime/Core/Public/HAL/Platform.h
+++ b/Engine/Source/Runtime/Core/Public/HAL/Platform.h
@@ -751,7 +751,9 @@
 
 /* Use before a class data member declaration allow it to be overlapped with other non-static data members or base class subobjects of its class. */
 #if !defined(UE_NO_UNIQUE_ADDRESS) && defined(__has_cpp_attribute)
-	#if __has_cpp_attribute(no_unique_address)
+	#if __has_cpp_attribute(msvc::no_unique_address)
+		#define UE_NO_UNIQUE_ADDRESS [[msvc::no_unique_address]] 
+	#elif __has_cpp_attribute(no_unique_address)
 		#define UE_NO_UNIQUE_ADDRESS [[no_unique_address]]
 	#endif
 #endif
diff --git a/Engine/Source/Runtime/Core/Public/Misc/Optional.h b/Engine/Source/Runtime/Core/Public/Misc/Optional.h
--- a/Engine/Source/Runtime/Core/Public/Misc/Optional.h
+++ b/Engine/Source/Runtime/Core/Public/Misc/Optional.h
@@ -22,11 +22,9 @@ namespace UE::Core::Private
 		);
 	};
 
-	template <size_t N>
-	struct TNonIntrusiveOptionalStorage
+	// Empty structure used for bIsSet member for intrusive optional types
+	struct FEmptyOptionalFlag
 	{
-		uint8 Storage[N];
-		bool bIsSet;
 	};
 }
 
@@ -62,16 +60,17 @@ public:
 	using ElementType = OptionalType;
 
 	/** Construct an OptionalType with a valid value. */
-	[[nodiscard]] TOptional(const OptionalType& InValue)
+	[[nodiscard]] constexpr TOptional(const OptionalType& InValue)
 		: TOptional(InPlace, InValue)
 	{
 	}
-	[[nodiscard]] TOptional(OptionalType&& InValue)
+	[[nodiscard]] constexpr TOptional(OptionalType&& InValue)
 		: TOptional(InPlace, MoveTempIfPossible(InValue))
 	{
 	}
 	template <typename... ArgTypes>
-	[[nodiscard]] explicit TOptional(EInPlace, ArgTypes&&... Args)
+	[[nodiscard]] explicit constexpr TOptional(EInPlace, ArgTypes&&... Args)
+		: TypedValue(Forward<ArgTypes>(Args)...)
 	{
 		// If this fails to compile when trying to call TOptional(EInPlace, ...) with a non-public constructor,
 		// do not make TOptional a friend.
@@ -93,11 +92,10 @@ public:
 		//     // Won't compile if the caller doesn't have access to FMyType::FPrivateToken
 		//     TOptional<FMyType> Opt(InPlace, FMyType::FPrivateToken{}, 5, 3.14f, TEXT("Banana"));
 		//
-		::new((void*)&Value) OptionalType(Forward<ArgTypes>(Args)...);
 
 		if constexpr (!bUsingIntrusiveUnsetState)
 		{
-			Value.bIsSet = true;
+			bIsSet = true;
 		}
 		else
 		{
@@ -107,67 +105,66 @@ public:
 	}
 	
 	/** Construct an OptionalType with an invalid value. */
-	[[nodiscard]] TOptional(FNullOpt)
+	[[nodiscard]] constexpr TOptional(FNullOpt)
 		: TOptional()
 	{
 	}
 
+	/** Construct an OptionalType intrusively with no value; i.e. unset */
+	[[nodiscard]] constexpr TOptional() requires bUsingIntrusiveUnsetState
+		: TypedValue(FIntrusiveUnsetOptionalState{})
+	{
+	}
+
 	/** Construct an OptionalType with no value; i.e. unset */
-	[[nodiscard]] TOptional()
+	[[nodiscard]] constexpr TOptional() requires (!bUsingIntrusiveUnsetState)
+		: bIsSet(false)
+	{
+	}
+
+	// Allow trivial destruction if contained type is trivially destructible
+	constexpr ~TOptional() requires std::is_trivially_destructible_v<OptionalType> = default;
+	constexpr ~TOptional() requires (!std::is_trivially_destructible_v<OptionalType>)
 	{
 		if constexpr (bUsingIntrusiveUnsetState)
 		{
-			::new ((void*)&Value) OptionalType(FIntrusiveUnsetOptionalState{});
+			// For intrusive optionals, destroy the value regardless of whether it is set
+			DestroyValue();
 		}
-		else
+		else if (IsSet())
 		{
-			Value.bIsSet = false;
+			DestroyValue();
 		}
 	}
 
-	~TOptional()
-	{
-		Reset();
-	}
-
 	/** Copy/Move construction */
 	[[nodiscard]] TOptional(const TOptional& Other)
+		: TOptional()
 	{
 		bool bLocalIsSet = Other.IsSet();
 		if constexpr (!bUsingIntrusiveUnsetState)
 		{
-			Value.bIsSet = bLocalIsSet;
+			bIsSet = bLocalIsSet;
 		}
 		if (bLocalIsSet)
 		{
-			::new((void*)&Value) OptionalType(*(const OptionalType*)&Other.Value);
-		}
-		else
-		{
-			if constexpr (bUsingIntrusiveUnsetState)
-			{
-				::new ((void*)&Value) OptionalType(FIntrusiveUnsetOptionalState{});
-			}
+			::new((void*)std::addressof(TypedValue)) OptionalType(Other.TypedValue);
 		}
+		// Default constructor initialized unset intrusive optional if necessary
 	}
 	[[nodiscard]] TOptional(TOptional&& Other)
+		: TOptional()
 	{
 		bool bLocalIsSet = Other.IsSet();
 		if constexpr (!bUsingIntrusiveUnsetState)
 		{
-			Value.bIsSet = bLocalIsSet;
+			bIsSet = bLocalIsSet;
 		}
 		if (bLocalIsSet)
 		{
-			::new((void*)&Value) OptionalType(MoveTempIfPossible(*(OptionalType*)&Other.Value));
-		}
-		else
-		{
-			if constexpr (bUsingIntrusiveUnsetState)
-			{
-				::new ((void*)&Value) OptionalType(FIntrusiveUnsetOptionalState{});
-			}
+			::new((void*)std::addressof(TypedValue)) OptionalType(MoveTempIfPossible(Other.TypedValue));
 		}
+		// Default constructor initialized unset intrusive optional if necessary
 	}
 
 	TOptional& operator=(const TOptional& Other)
@@ -203,7 +200,7 @@ public:
 
 	TOptional& operator=(const OptionalType& InValue)
 	{
-		if (&InValue != (const OptionalType*)&Value)
+		if (std::addressof(InValue) != std::addressof(TypedValue))
 		{
 			Emplace(InValue);
 		}
@@ -211,7 +208,7 @@ public:
 	}
 	TOptional& operator=(OptionalType&& InValue)
 	{
-		if (&InValue != (const OptionalType*)&Value)
+		if (std::addressof(InValue) != std::addressof(TypedValue))
 		{
 			Emplace(MoveTempIfPossible(InValue));
 		}
@@ -225,11 +222,11 @@ public:
 			DestroyValue();
 			if constexpr (bUsingIntrusiveUnsetState)
 			{
-				::new((void*)&Value) OptionalType(FIntrusiveUnsetOptionalState{});
+				::new((void*)std::addressof(TypedValue)) OptionalType(FIntrusiveUnsetOptionalState{});
 			}
 			else
 			{
-				Value.bIsSet = false;
+				bIsSet = false;
 			}
 		}
 	}
@@ -272,11 +269,11 @@ public:
 		//     // Won't compile if the caller doesn't have access to FMyType::FPrivateToken
 		//     Opt.Emplace(FMyType::FPrivateToken{}, 5, 3.14f, TEXT("Banana"));
 		//
-		OptionalType* Result = ::new((void*)&Value) OptionalType(Forward<ArgsType>(Args)...);
+		OptionalType* Result = ::new((void*)std::addressof(TypedValue)) OptionalType(Forward<ArgsType>(Args)...);
 
 		if constexpr (!bUsingIntrusiveUnsetState)
 		{
-			Value.bIsSet = true;
+			bIsSet = true;
 		}
 		else
 		{
@@ -287,7 +284,7 @@ public:
 		return *Result;
 	}
 
-	[[nodiscard]] friend bool operator==(const TOptional& Lhs, const TOptional& Rhs)
+	[[nodiscard]] friend constexpr bool operator==(const TOptional& Lhs, const TOptional& Rhs)
 	{
 		bool bIsLhsSet = Lhs.IsSet();
 		bool bIsRhsSet = Rhs.IsSet();
@@ -301,10 +298,10 @@ public:
 			return true;
 		}
 
-		return (*(const OptionalType*)&Lhs.Value) == (*(const OptionalType*)&Rhs.Value);
+		return Lhs.TypedValue == Rhs.TypedValue;
 	}
 
-	[[nodiscard]] friend bool operator!=(const TOptional& Lhs, const TOptional& Rhs)
+	[[nodiscard]] friend constexpr bool operator!=(const TOptional& Lhs, const TOptional& Rhs)
 	{
 		return !(Lhs == Rhs);
 	}
@@ -340,63 +337,63 @@ public:
 	}
 
 	/** @return true when the value is meaningful; false if calling GetValue() is undefined. */
-	[[nodiscard]] bool IsSet() const
+	[[nodiscard]] constexpr bool IsSet() const
 	{
 		if constexpr (bUsingIntrusiveUnsetState)
 		{
-			return !(*(const OptionalType*)&Value == FIntrusiveUnsetOptionalState{});
+			return !(TypedValue == FIntrusiveUnsetOptionalState{});
 		}
 		else
 		{
-			return Value.bIsSet;
+			return bIsSet;
 		}
 	}
-	[[nodiscard]] FORCEINLINE explicit operator bool() const
+	[[nodiscard]] FORCEINLINE explicit constexpr operator bool() const
 	{
 		return IsSet();
 	}
 
 	/** @return The optional value; undefined when IsSet() returns false. */
-	[[nodiscard]] OptionalType& GetValue()
+	[[nodiscard]] constexpr OptionalType& GetValue()
 	{
 		checkf(IsSet(), TEXT("It is an error to call GetValue() on an unset TOptional. Please either check IsSet() or use Get(DefaultValue) instead."));
-		return *(OptionalType*)&Value;
+		return TypedValue;
 	}
-	[[nodiscard]] FORCEINLINE const OptionalType& GetValue() const
+	[[nodiscard]] FORCEINLINE constexpr const OptionalType& GetValue() const
 	{
 		return const_cast<TOptional*>(this)->GetValue();
 	}
 
-	[[nodiscard]] OptionalType* operator->()
+	[[nodiscard]] constexpr OptionalType* operator->()
 	{
-		return &GetValue();
+		return std::addressof(GetValue());
 	}
-	[[nodiscard]] FORCEINLINE const OptionalType* operator->() const
+	[[nodiscard]] FORCEINLINE constexpr const OptionalType* operator->() const
 	{
 		return const_cast<TOptional*>(this)->operator->();
 	}
 
-	[[nodiscard]] OptionalType& operator*()
+	[[nodiscard]] constexpr OptionalType& operator*()
 	{
 		return GetValue();
 	}
-	[[nodiscard]] FORCEINLINE const OptionalType& operator*() const
+	[[nodiscard]] FORCEINLINE constexpr const OptionalType& operator*() const
 	{
 		return const_cast<TOptional*>(this)->operator*();
 	}
 
 	/** @return The optional value when set; DefaultValue otherwise. */
-	[[nodiscard]] const OptionalType& Get(const OptionalType& DefaultValue UE_LIFETIMEBOUND) const UE_LIFETIMEBOUND
+	[[nodiscard]] constexpr const OptionalType& Get(const OptionalType& DefaultValue UE_LIFETIMEBOUND) const UE_LIFETIMEBOUND
 	{
-		return IsSet() ? *(const OptionalType*)&Value : DefaultValue;
+		return IsSet() ? TypedValue : DefaultValue;
 	}
 
 	/** @return A pointer to the optional value when set, nullptr otherwise. */
-	[[nodiscard]] OptionalType* GetPtrOrNull()
+	[[nodiscard]] constexpr OptionalType* GetPtrOrNull()
 	{
-		return IsSet() ? (OptionalType*)&Value : nullptr;
+		return IsSet() ? std::addressof(TypedValue) : nullptr;
 	}
-	[[nodiscard]] FORCEINLINE const OptionalType* GetPtrOrNull() const
+	[[nodiscard]] FORCEINLINE constexpr const OptionalType* GetPtrOrNull() const
 	{
 		return const_cast<TOptional*>(this)->GetPtrOrNull();
 	}
@@ -407,18 +404,15 @@ private:
 	 */
 	FORCEINLINE void DestroyValue()
 	{
-		if constexpr (bUsingIntrusiveUnsetState)
-		{
-			DestructItem((OptionalType*)&Value);
-		}
-		else
-		{
-			DestructItem((OptionalType*)Value.Storage);
-		}
+		DestructItem(std::addressof(TypedValue));
 	}
 
-	using ValueStorageType = std::conditional_t<bUsingIntrusiveUnsetState, uint8[sizeof(OptionalType)], UE::Core::Private::TNonIntrusiveOptionalStorage<sizeof(OptionalType)>>;
-	alignas(OptionalType) ValueStorageType Value;
+	union
+	{
+		OptionalType TypedValue;
+	};
+	// Set flag takes up no space if bIsUsingIntrusiveUnsetState
+	UE_NO_UNIQUE_ADDRESS std::conditional_t<bUsingIntrusiveUnsetState, UE::Core::Private::FEmptyOptionalFlag, bool> bIsSet;
 };
 
 template<typename OptionalType>
diff --git a/Engine/Source/Runtime/Core/Tests/Misc/OptionalTest.cpp b/Engine/Source/Runtime/Core/Tests/Misc/OptionalTest.cpp
new file mode 100644
--- /dev/null
+++ b/Engine/Source/Runtime/Core/Tests/Misc/OptionalTest.cpp
@@ -0,0 +1,98 @@
+// Copyright Epic Games, Inc. All Rights Reserved.
+
+#if WITH_TESTS
+
+#include "CoreTypes.h"
+#include "Containers/UnrealString.h"
+#include "Math/Vector.h"
+#include "Misc/Optional.h"
+#include "Templates/AlignmentTemplates.h"
+#include "Tests/TestHarnessAdapter.h"
+
+// Tests for traits and constexpr use of TOptional
+namespace UE::Core::Private
+{
+	// TIsTOptional_V of optionals
+	static_assert(TIsTOptional_V<      TOptional<float>>, "Expected TIsTOptional_V<TOptional> to be true");
+	static_assert(TIsTOptional_V<const TOptional<float>>, "Expected TIsTOptional_V<const TOptional> to be true");
+
+	// TIsTOptional_V of non-optionals (references are not optionals)
+	static_assert(TIsTOptional_V<      TOptional<float>&>  == false, "Expected TIsTOptional_V<TOptional&> to be false");
+	static_assert(TIsTOptional_V<const TOptional<float>&>  == false, "Expected TIsTOptional_V<const TOptional&> to be false");
+	static_assert(TIsTOptional_V<      TOptional<float>&&> == false, "Expected TIsTOptional_V<TOptional&&> to be false");
+	static_assert(TIsTOptional_V<      bool>               == false, "Expected TIsTOptional_V<non-TOptional> to be false");
+
+	// Ensure that TOptional has a trivial destructor if the inner type does
+	static_assert( std::is_trivially_destructible_v<TOptional<int>>);
+	static_assert( std::is_trivially_destructible_v<TOptional<FVector>>);
+	static_assert(!std::is_trivially_destructible_v<TOptional<FString>>);
+
+    template<bool bIntrusive>
+	struct TTestValue
+	{
+		static constexpr inline bool bHasIntrusiveUnsetOptionalState = bIntrusive;
+		using IntrusiveUnsetOptionalStateType = TTestValue<true>;
+		constexpr TTestValue(FIntrusiveUnsetOptionalState) requires bIntrusive
+		: Value(0)
+		{
+		}
+		constexpr TTestValue(int32 InValue)
+		: Value(InValue)
+		{
+		}
+        constexpr TTestValue(const TTestValue&) = default;
+        constexpr TTestValue& operator=(const TTestValue&) = default;
+        constexpr TTestValue(TTestValue&&) = default;
+        constexpr TTestValue& operator=(TTestValue&&) = default;
+
+        [[nodiscard]] constexpr bool operator==(FIntrusiveUnsetOptionalState) const requires bIntrusive
+        {
+            return Value == 0;
+        }
+
+        [[nodiscard]] constexpr bool operator==(const TTestValue&) const = default;
+
+		int32 Value;
+	};
+
+    using FIntrusive = TTestValue<true>;
+	using FNonIntrusive = TTestValue<false>;
+
+	static_assert(sizeof(TOptional<FIntrusive>) == sizeof(FIntrusive), "Expected intrusive optional to be the same size as the inner type");
+	static_assert(sizeof(TOptional<FNonIntrusive>) > sizeof(FNonIntrusive), "Expected non-intrusive optional to be larger than the inner type");
+	static_assert(sizeof(TOptional<FNonIntrusive>) <= Align(sizeof(FNonIntrusive) + 1, alignof(FNonIntrusive)),
+        "Expected non-intrusive optional to be larger than the inner type by at most the alignment of the inner type");
+
+    // Tests for constexpr use of optional
+    template<typename InnerType>
+    constexpr bool TestConstexprOptionalUse()
+    {
+        InnerType Value{32};
+        TOptional<InnerType> O1; // Default construction
+        TOptional<InnerType> O2{NullOpt}; // Explicit empty construction;
+        TOptional<InnerType> O3(Value); // Copying inner type
+        TOptional<InnerType> O4(Forward<InnerType>(InnerType(30))); // Moving inner type
+        TOptional<InnerType> O5(InPlace, 29); // In-place construction
+        if constexpr (false) // At present, copy and move construction are not constexpr-safe
+        {
+            TOptional<InnerType> O6(O3); // Copy construction
+            TOptional<InnerType> O7(MoveTemp(O4)); // Move construction
+        }
+        bool bTestResult = (O1 == O2) && (O4 != O5);
+        bTestResult = !O1.IsSet() && O3.IsSet() && bTestResult;
+        const TOptional<InnerType> C1(InPlace, 33);
+        const InnerType& I1 = C1.GetValue();
+        O5.GetValue().Value++;
+        O5->Value = C1->Value;
+        bTestResult = ((*O5) == (*C1)) && bTestResult;
+        const InnerType& I2 = O1.Get(Value);
+        *O4.GetPtrOrNull() = *C1.GetPtrOrNull();
+
+        return bTestResult;
+    }
+
+    static_assert(TestConstexprOptionalUse<FIntrusive>(), "Constexpr tests for TOptional<FIntrusive> failed");
+    static_assert(TestConstexprOptionalUse<FNonIntrusive>(), "Constexpr tests for TOptional<FNonIntrusive> failed");
+};
+
+#endif // WITH_TESTS
\ No newline at end of file
diff --git a/Engine/Source/Runtime/CoreUObject/Public/UObject/Class.h b/Engine/Source/Runtime/CoreUObject/Public/UObject/Class.h
--- a/Engine/Source/Runtime/CoreUObject/Public/UObject/Class.h
+++ b/Engine/Source/Runtime/CoreUObject/Public/UObject/Class.h
@@ -1694,19 +1694,32 @@ public:
 		/** Construct an unset optional value */
 		virtual void InitializeIntrusiveUnsetOptionalValue(void* Data) const override
 		{
-			new (Data) TOptional<CPPSTRUCT>();
+			if constexpr(::HasIntrusiveUnsetOptionalState<CPPSTRUCT>())
+			{
+				new (Data) TOptional<CPPSTRUCT>();
+			}
 		}
 
 		/** Return true if the optional value at Data is in an unset state */
 		virtual bool IsIntrusiveOptionalValueSet(const void* Data) const override
 		{
-			return reinterpret_cast<const TOptional<CPPSTRUCT>*>(Data)->IsSet();
+			if constexpr(::HasIntrusiveUnsetOptionalState<CPPSTRUCT>())
+			{
+				return reinterpret_cast<const TOptional<CPPSTRUCT>*>(Data)->IsSet();
+			}
+			else 
+			{
+				return false;
+			}
 		}
 
 		/** Reset an optional value to its unset state */
 		virtual void ClearIntrusiveOptionalValue(void* Data) const override
 		{
-			reinterpret_cast<TOptional<CPPSTRUCT>*>(Data)->Reset();
+			if constexpr(::HasIntrusiveUnsetOptionalState<CPPSTRUCT>())
+			{
+				reinterpret_cast<TOptional<CPPSTRUCT>*>(Data)->Reset();
+			}
 		}
 
 		/** 
