[Backout] - CL42245538
#rnx
Original CL Desc
-----------------------------------------------------------------
Modify TOptional to allow in-place construction in constexpr contexts which requires us not to use reinterpret_cast.
We can't use std::construct_at on some platforms so only support in-place construction from values, and  not other modifications in a constexpr context.
No support for copying or moving other TOptionals.
Added compile time tests for supported constexpr operations.

diff --git a/Engine/Source/Runtime/Core/Public/HAL/Platform.h b/Engine/Source/Runtime/Core/Public/HAL/Platform.h
--- a/Engine/Source/Runtime/Core/Public/HAL/Platform.h
+++ b/Engine/Source/Runtime/Core/Public/HAL/Platform.h
@@ -751,9 +751,7 @@
 
 /* Use before a class data member declaration allow it to be overlapped with other non-static data members or base class subobjects of its class. */
 #if !defined(UE_NO_UNIQUE_ADDRESS) && defined(__has_cpp_attribute)
-	#if __has_cpp_attribute(msvc::no_unique_address)
-		#define UE_NO_UNIQUE_ADDRESS [[msvc::no_unique_address]] 
-	#elif __has_cpp_attribute(no_unique_address)
+	#if __has_cpp_attribute(no_unique_address)
 		#define UE_NO_UNIQUE_ADDRESS [[no_unique_address]]
 	#endif
 #endif
diff --git a/Engine/Source/Runtime/Core/Public/Misc/Optional.h b/Engine/Source/Runtime/Core/Public/Misc/Optional.h
--- a/Engine/Source/Runtime/Core/Public/Misc/Optional.h
+++ b/Engine/Source/Runtime/Core/Public/Misc/Optional.h
@@ -22,9 +22,11 @@ namespace UE::Core::Private
 		);
 	};
 
-	// Empty structure used for bIsSet member for intrusive optional types
-	struct FEmptyOptionalFlag
+	template <size_t N>
+	struct TNonIntrusiveOptionalStorage
 	{
+		uint8 Storage[N];
+		bool bIsSet;
 	};
 }
 
@@ -60,17 +62,16 @@ public:
 	using ElementType = OptionalType;
 
 	/** Construct an OptionalType with a valid value. */
-	[[nodiscard]] constexpr TOptional(const OptionalType& InValue)
+	[[nodiscard]] TOptional(const OptionalType& InValue)
 		: TOptional(InPlace, InValue)
 	{
 	}
-	[[nodiscard]] constexpr TOptional(OptionalType&& InValue)
+	[[nodiscard]] TOptional(OptionalType&& InValue)
 		: TOptional(InPlace, MoveTempIfPossible(InValue))
 	{
 	}
 	template <typename... ArgTypes>
-	[[nodiscard]] explicit constexpr TOptional(EInPlace, ArgTypes&&... Args)
-		: TypedValue(Forward<ArgTypes>(Args)...)
+	[[nodiscard]] explicit TOptional(EInPlace, ArgTypes&&... Args)
 	{
 		// If this fails to compile when trying to call TOptional(EInPlace, ...) with a non-public constructor,
 		// do not make TOptional a friend.
@@ -92,10 +93,11 @@ public:
 		//     // Won't compile if the caller doesn't have access to FMyType::FPrivateToken
 		//     TOptional<FMyType> Opt(InPlace, FMyType::FPrivateToken{}, 5, 3.14f, TEXT("Banana"));
 		//
+		::new((void*)&Value) OptionalType(Forward<ArgTypes>(Args)...);
 
 		if constexpr (!bUsingIntrusiveUnsetState)
 		{
-			bIsSet = true;
+			Value.bIsSet = true;
 		}
 		else
 		{
@@ -105,66 +107,67 @@ public:
 	}
 	
 	/** Construct an OptionalType with an invalid value. */
-	[[nodiscard]] constexpr TOptional(FNullOpt)
+	[[nodiscard]] TOptional(FNullOpt)
 		: TOptional()
 	{
 	}
 
-	/** Construct an OptionalType intrusively with no value; i.e. unset */
-	[[nodiscard]] constexpr TOptional() requires bUsingIntrusiveUnsetState
-		: TypedValue(FIntrusiveUnsetOptionalState{})
-	{
-	}
-
 	/** Construct an OptionalType with no value; i.e. unset */
-	[[nodiscard]] constexpr TOptional() requires (!bUsingIntrusiveUnsetState)
-		: bIsSet(false)
-	{
-	}
-
-	// Allow trivial destruction if contained type is trivially destructible
-	constexpr ~TOptional() requires std::is_trivially_destructible_v<OptionalType> = default;
-	constexpr ~TOptional() requires (!std::is_trivially_destructible_v<OptionalType>)
+	[[nodiscard]] TOptional()
 	{
 		if constexpr (bUsingIntrusiveUnsetState)
 		{
-			// For intrusive optionals, destroy the value regardless of whether it is set
-			DestroyValue();
+			::new ((void*)&Value) OptionalType(FIntrusiveUnsetOptionalState{});
 		}
-		else if (IsSet())
+		else
 		{
-			DestroyValue();
+			Value.bIsSet = false;
 		}
 	}
 
+	~TOptional()
+	{
+		Reset();
+	}
+
 	/** Copy/Move construction */
 	[[nodiscard]] TOptional(const TOptional& Other)
-		: TOptional()
 	{
 		bool bLocalIsSet = Other.IsSet();
 		if constexpr (!bUsingIntrusiveUnsetState)
 		{
-			bIsSet = bLocalIsSet;
+			Value.bIsSet = bLocalIsSet;
 		}
 		if (bLocalIsSet)
 		{
-			::new((void*)std::addressof(TypedValue)) OptionalType(Other.TypedValue);
+			::new((void*)&Value) OptionalType(*(const OptionalType*)&Other.Value);
+		}
+		else
+		{
+			if constexpr (bUsingIntrusiveUnsetState)
+			{
+				::new ((void*)&Value) OptionalType(FIntrusiveUnsetOptionalState{});
+			}
 		}
-		// Default constructor initialized unset intrusive optional if necessary
 	}
 	[[nodiscard]] TOptional(TOptional&& Other)
-		: TOptional()
 	{
 		bool bLocalIsSet = Other.IsSet();
 		if constexpr (!bUsingIntrusiveUnsetState)
 		{
-			bIsSet = bLocalIsSet;
+			Value.bIsSet = bLocalIsSet;
 		}
 		if (bLocalIsSet)
 		{
-			::new((void*)std::addressof(TypedValue)) OptionalType(MoveTempIfPossible(Other.TypedValue));
+			::new((void*)&Value) OptionalType(MoveTempIfPossible(*(OptionalType*)&Other.Value));
+		}
+		else
+		{
+			if constexpr (bUsingIntrusiveUnsetState)
+			{
+				::new ((void*)&Value) OptionalType(FIntrusiveUnsetOptionalState{});
+			}
 		}
-		// Default constructor initialized unset intrusive optional if necessary
 	}
 
 	TOptional& operator=(const TOptional& Other)
@@ -200,7 +203,7 @@ public:
 
 	TOptional& operator=(const OptionalType& InValue)
 	{
-		if (std::addressof(InValue) != std::addressof(TypedValue))
+		if (&InValue != (const OptionalType*)&Value)
 		{
 			Emplace(InValue);
 		}
@@ -208,7 +211,7 @@ public:
 	}
 	TOptional& operator=(OptionalType&& InValue)
 	{
-		if (std::addressof(InValue) != std::addressof(TypedValue))
+		if (&InValue != (const OptionalType*)&Value)
 		{
 			Emplace(MoveTempIfPossible(InValue));
 		}
@@ -222,11 +225,11 @@ public:
 			DestroyValue();
 			if constexpr (bUsingIntrusiveUnsetState)
 			{
-				::new((void*)std::addressof(TypedValue)) OptionalType(FIntrusiveUnsetOptionalState{});
+				::new((void*)&Value) OptionalType(FIntrusiveUnsetOptionalState{});
 			}
 			else
 			{
-				bIsSet = false;
+				Value.bIsSet = false;
 			}
 		}
 	}
@@ -269,11 +272,11 @@ public:
 		//     // Won't compile if the caller doesn't have access to FMyType::FPrivateToken
 		//     Opt.Emplace(FMyType::FPrivateToken{}, 5, 3.14f, TEXT("Banana"));
 		//
-		OptionalType* Result = ::new((void*)std::addressof(TypedValue)) OptionalType(Forward<ArgsType>(Args)...);
+		OptionalType* Result = ::new((void*)&Value) OptionalType(Forward<ArgsType>(Args)...);
 
 		if constexpr (!bUsingIntrusiveUnsetState)
 		{
-			bIsSet = true;
+			Value.bIsSet = true;
 		}
 		else
 		{
@@ -284,7 +287,7 @@ public:
 		return *Result;
 	}
 
-	[[nodiscard]] friend constexpr bool operator==(const TOptional& Lhs, const TOptional& Rhs)
+	[[nodiscard]] friend bool operator==(const TOptional& Lhs, const TOptional& Rhs)
 	{
 		bool bIsLhsSet = Lhs.IsSet();
 		bool bIsRhsSet = Rhs.IsSet();
@@ -298,10 +301,10 @@ public:
 			return true;
 		}
 
-		return Lhs.TypedValue == Rhs.TypedValue;
+		return (*(const OptionalType*)&Lhs.Value) == (*(const OptionalType*)&Rhs.Value);
 	}
 
-	[[nodiscard]] friend constexpr bool operator!=(const TOptional& Lhs, const TOptional& Rhs)
+	[[nodiscard]] friend bool operator!=(const TOptional& Lhs, const TOptional& Rhs)
 	{
 		return !(Lhs == Rhs);
 	}
@@ -337,63 +340,63 @@ public:
 	}
 
 	/** @return true when the value is meaningful; false if calling GetValue() is undefined. */
-	[[nodiscard]] constexpr bool IsSet() const
+	[[nodiscard]] bool IsSet() const
 	{
 		if constexpr (bUsingIntrusiveUnsetState)
 		{
-			return !(TypedValue == FIntrusiveUnsetOptionalState{});
+			return !(*(const OptionalType*)&Value == FIntrusiveUnsetOptionalState{});
 		}
 		else
 		{
-			return bIsSet;
+			return Value.bIsSet;
 		}
 	}
-	[[nodiscard]] FORCEINLINE explicit constexpr operator bool() const
+	[[nodiscard]] FORCEINLINE explicit operator bool() const
 	{
 		return IsSet();
 	}
 
 	/** @return The optional value; undefined when IsSet() returns false. */
-	[[nodiscard]] constexpr OptionalType& GetValue()
+	[[nodiscard]] OptionalType& GetValue()
 	{
 		checkf(IsSet(), TEXT("It is an error to call GetValue() on an unset TOptional. Please either check IsSet() or use Get(DefaultValue) instead."));
-		return TypedValue;
+		return *(OptionalType*)&Value;
 	}
-	[[nodiscard]] FORCEINLINE constexpr const OptionalType& GetValue() const
+	[[nodiscard]] FORCEINLINE const OptionalType& GetValue() const
 	{
 		return const_cast<TOptional*>(this)->GetValue();
 	}
 
-	[[nodiscard]] constexpr OptionalType* operator->()
+	[[nodiscard]] OptionalType* operator->()
 	{
-		return std::addressof(GetValue());
+		return &GetValue();
 	}
-	[[nodiscard]] FORCEINLINE constexpr const OptionalType* operator->() const
+	[[nodiscard]] FORCEINLINE const OptionalType* operator->() const
 	{
 		return const_cast<TOptional*>(this)->operator->();
 	}
 
-	[[nodiscard]] constexpr OptionalType& operator*()
+	[[nodiscard]] OptionalType& operator*()
 	{
 		return GetValue();
 	}
-	[[nodiscard]] FORCEINLINE constexpr const OptionalType& operator*() const
+	[[nodiscard]] FORCEINLINE const OptionalType& operator*() const
 	{
 		return const_cast<TOptional*>(this)->operator*();
 	}
 
 	/** @return The optional value when set; DefaultValue otherwise. */
-	[[nodiscard]] constexpr const OptionalType& Get(const OptionalType& DefaultValue UE_LIFETIMEBOUND) const UE_LIFETIMEBOUND
+	[[nodiscard]] const OptionalType& Get(const OptionalType& DefaultValue UE_LIFETIMEBOUND) const UE_LIFETIMEBOUND
 	{
-		return IsSet() ? TypedValue : DefaultValue;
+		return IsSet() ? *(const OptionalType*)&Value : DefaultValue;
 	}
 
 	/** @return A pointer to the optional value when set, nullptr otherwise. */
-	[[nodiscard]] constexpr OptionalType* GetPtrOrNull()
+	[[nodiscard]] OptionalType* GetPtrOrNull()
 	{
-		return IsSet() ? std::addressof(TypedValue) : nullptr;
+		return IsSet() ? (OptionalType*)&Value : nullptr;
 	}
-	[[nodiscard]] FORCEINLINE constexpr const OptionalType* GetPtrOrNull() const
+	[[nodiscard]] FORCEINLINE const OptionalType* GetPtrOrNull() const
 	{
 		return const_cast<TOptional*>(this)->GetPtrOrNull();
 	}
@@ -404,15 +407,18 @@ private:
 	 */
 	FORCEINLINE void DestroyValue()
 	{
-		DestructItem(std::addressof(TypedValue));
+		if constexpr (bUsingIntrusiveUnsetState)
+		{
+			DestructItem((OptionalType*)&Value);
+		}
+		else
+		{
+			DestructItem((OptionalType*)Value.Storage);
+		}
 	}
 
-	union
-	{
-		OptionalType TypedValue;
-	};
-	// Set flag takes up no space if bIsUsingIntrusiveUnsetState
-	UE_NO_UNIQUE_ADDRESS std::conditional_t<bUsingIntrusiveUnsetState, UE::Core::Private::FEmptyOptionalFlag, bool> bIsSet;
+	using ValueStorageType = std::conditional_t<bUsingIntrusiveUnsetState, uint8[sizeof(OptionalType)], UE::Core::Private::TNonIntrusiveOptionalStorage<sizeof(OptionalType)>>;
+	alignas(OptionalType) ValueStorageType Value;
 };
 
 template<typename OptionalType>
diff --git a/Engine/Source/Runtime/Core/Tests/Misc/OptionalTest.cpp b/Engine/Source/Runtime/Core/Tests/Misc/OptionalTest.cpp
deleted file mode 100644
diff --git a/Engine/Source/Runtime/CoreUObject/Public/UObject/Class.h b/Engine/Source/Runtime/CoreUObject/Public/UObject/Class.h
--- a/Engine/Source/Runtime/CoreUObject/Public/UObject/Class.h
+++ b/Engine/Source/Runtime/CoreUObject/Public/UObject/Class.h
@@ -1694,32 +1694,19 @@ public:
 		/** Construct an unset optional value */
 		virtual void InitializeIntrusiveUnsetOptionalValue(void* Data) const override
 		{
-			if constexpr(::HasIntrusiveUnsetOptionalState<CPPSTRUCT>())
-			{
-				new (Data) TOptional<CPPSTRUCT>();
-			}
+			new (Data) TOptional<CPPSTRUCT>();
 		}
 
 		/** Return true if the optional value at Data is in an unset state */
 		virtual bool IsIntrusiveOptionalValueSet(const void* Data) const override
 		{
-			if constexpr(::HasIntrusiveUnsetOptionalState<CPPSTRUCT>())
-			{
-				return reinterpret_cast<const TOptional<CPPSTRUCT>*>(Data)->IsSet();
-			}
-			else 
-			{
-				return false;
-			}
+			return reinterpret_cast<const TOptional<CPPSTRUCT>*>(Data)->IsSet();
 		}
 
 		/** Reset an optional value to its unset state */
 		virtual void ClearIntrusiveOptionalValue(void* Data) const override
 		{
-			if constexpr(::HasIntrusiveUnsetOptionalState<CPPSTRUCT>())
-			{
-				reinterpret_cast<TOptional<CPPSTRUCT>*>(Data)->Reset();
-			}
+			reinterpret_cast<TOptional<CPPSTRUCT>*>(Data)->Reset();
 		}
 
 		/** 
