Added support for Ferror using structured logging to generate error messages.

diff --git a/Engine/Plugins/FX/Niagara/Source/NiagaraEditor/Private/NiagaraEditorUtilities.cpp b/Engine/Plugins/FX/Niagara/Source/NiagaraEditor/Private/NiagaraEditorUtilities.cpp
--- a/Engine/Plugins/FX/Niagara/Source/NiagaraEditor/Private/NiagaraEditorUtilities.cpp
+++ b/Engine/Plugins/FX/Niagara/Source/NiagaraEditor/Private/NiagaraEditorUtilities.cpp
@@ -51,12 +51,14 @@
 #include "ViewModels/Stack/NiagaraStackNote.h"
 #include "Widgets/SNiagaraParameterName.h"
 
+
 #include "EdGraph/EdGraphPin.h"
 #include "Editor/EditorEngine.h"
 #include "Framework/MultiBox/MultiBoxBuilder.h"
 #include "Framework/Notifications/NotificationManager.h"
 #include "HAL/PlatformFileManager.h"
 #include "Interfaces/IPluginManager.h"
+#include "Logging/StructuredLog.h"
 #include "NiagaraSystem.h"
 #include "NiagaraSystemEditorData.h"
 #include "Toolkits/NiagaraSystemToolkit.h"
diff --git a/Engine/Source/Runtime/Core/Private/Experimental/IO/IoStatusError.cpp b/Engine/Source/Runtime/Core/Private/Experimental/IO/IoStatusError.cpp
--- a/Engine/Source/Runtime/Core/Private/Experimental/IO/IoStatusError.cpp
+++ b/Engine/Source/Runtime/Core/Private/Experimental/IO/IoStatusError.cpp
@@ -19,20 +19,45 @@ namespace UE::UnifiedError::IoStore
 		{
 			ErrorMessage = InErrorMessage;
 		}
+		/*
 
-		virtual const FUtf8String& GetErrorName(const FError& Error) const final
+		virtual UE::UnifiedError::FAnyString GetErrorCodeString(const FError& Error) const final
 		{
-			if (CachedErrorName.IsEmpty() && (ErrorMessage[0] != TEXT('\0')))
-			{
-				CachedErrorName = TCHAR_TO_UTF8(GetIoErrorText(EIoErrorCode(Error.GetErrorCode())));
-			}
-			return CachedErrorName;
+			UE::UnifiedError::FAnyString Result; 
+			Result.CopyString(TCHAR_TO_UTF8(GetIoErrorText(EIoErrorCode(Error.GetErrorCode()))));
+			return Result;
 		}
-		virtual const FText GetErrorFormatString(const FError& Error) const final
+
+		virtual UE::UnifiedError::FAnyString GetModuleIdString(const FError& Error) const final
 		{
+			UE::UnifiedError::FAnyString Result;
+			Result.ReferenceString(UE::UnifiedError::IoStore::StaticModuleName);
+			return Result;
+
+		}
+*/
+
+		virtual const FText GetErrorFormatString(const FError& Error) const final
+		{	
 			return GenericErrorFormatString;
 		}
 
+		virtual void SerializeToCb(FCbWriter& Writer, const FError& Error) const override
+		{
+			Writer.BeginObject();
+			if (Error.GetModuleId() == UE::UnifiedError::IoStore::StaticModuleId)
+			{
+				Writer.AddString(UTF8TEXTVIEW("ErrorCodeString"), GetIoErrorText(EIoErrorCode(Error.GetErrorCode())));
+				Writer.AddString(UTF8TEXTVIEW("ModuleIdString"), UE::UnifiedError::IoStore::StaticModuleName);
+				Writer.AddInteger(UTF8TEXTVIEW("ErrorCode"), Error.GetErrorCode());
+				Writer.AddInteger(UTF8TEXTVIEW("ModuleId"), Error.GetModuleId());
+			}
+			Writer.AddString(UTF8TEXTVIEW("IoStoreErrorMessage"), ErrorMessage);
+			Writer.AddString(UTF8TEXTVIEW("$format"), UTF8TEXTVIEW("{IoStoreErrorMessage}"));
+			Writer.AddString(UTF8TEXTVIEW("$type"), UTF8TEXTVIEW("FIoStoreErrorDetails"));
+			Writer.EndObject();
+		}
+
 		virtual void GetErrorProperties(const FError& Error, IErrorPropertyExtractor& OutProperties) const override
 		{
 			OutProperties.AddProperty(UTF8TEXTVIEW("ErrorCodeString"), CachedErrorName);
@@ -42,13 +67,17 @@ namespace UE::UnifiedError::IoStore
 			OutProperties.AddProperty(UTF8TEXTVIEW("IoStoreErrorMessage"), ErrorMessage);
 		}
 
-		virtual const FAnsiStringView GetErrorDetailsTypeName() const
+		virtual const FAnsiStringView GetErrorDetailsTypeName() const override
+		{
+			return ANSITEXTVIEW("FIoStoreErrorDetails");
+		}
+		virtual const FAnsiStringView GetErrorDetailsTypeNameForLog() const override 
 		{
 			return ANSITEXTVIEW("FIoStoreErrorDetails");
 		}
 	};
 
-	const FText FIoStoreErrorDetails::GenericErrorFormatString = NSLOCTEXT("IoStore", "GenericErrorMessage", "{IoStoreErrorMessage}");
+	const FText FIoStoreErrorDetails::GenericErrorFormatString = NSLOCTEXT("IoStore", "GenericErrorMessage", "{FIoStoreErrorDetails}");
 
 
 	CORE_API UE::UnifiedError::FError ConvertError(const FIoStatus& Status)
diff --git a/Engine/Source/Runtime/Core/Private/Experimental/UnifiedError/UnifiedError.cpp b/Engine/Source/Runtime/Core/Private/Experimental/UnifiedError/UnifiedError.cpp
--- a/Engine/Source/Runtime/Core/Private/Experimental/UnifiedError/UnifiedError.cpp
+++ b/Engine/Source/Runtime/Core/Private/Experimental/UnifiedError/UnifiedError.cpp
@@ -16,12 +16,6 @@ namespace UE::UnifiedError
 
 void FError::PushErrorDetails(TRefCountPtr<IErrorDetails> InErrorDetails)
 {
-	/*if (ErrorDetails)
-	{
-		InErrorDetails->SetInnerErrorDetails(ErrorDetails);
-	}
-	InErrorDetails->AddRef();
-	ErrorDetails->ReleaseRef();*/
 	if (!InErrorDetails.IsValid())
 	{
 		return;
@@ -50,27 +44,69 @@ FText FError::GetFormatErrorText() const
 {
 	return ErrorDetails->GetErrorFormatString(*this);
 }
+/*
+
+
+UE::UnifiedError::FAnyString FError::GetErrorCodeString() const
+{
+	return ErrorDetails->GetErrorCodeString(*this);
+}
+
+UE::UnifiedError::FAnyString FError::GetModuleIdString() const
+{
+	return ErrorDetails->GetModuleIdString(*this);
+}
+*/
+
+
+#define USE_STRUCTURED_LOG_FOR_FERRORMESSAGE 1
+#define USE_LOCALIZED_STRUCTURED_LOG_FOR_FERRORMESSAGE 1
 
 
-#define USE_STRUCTURED_LOG_FOR_FERRORMESSAGE 0
 
 FText FError::GetErrorMessage(bool bIncludeContext) const
 {
+
 #if USE_STRUCTURED_LOG_FOR_FERRORMESSAGE
-	FString FormatString = ErrorDetails->GetErrorFormatString(*this).ToString();
 
-	TCbWriter<1024> Writer;
-	Writer.BeginObject();
+	FText FormatText = FText::Join(FText::FromString(": "), NSLOCTEXT("UnifiedError", "DefaultLogPrepend", "{Root/ModuleIdString}.{Root/ErrorCodeString}"), GetFormatErrorText());
 
-	SerializeDetailsForLog(Writer);
+	if (bIncludeContext)
+	{
+		TArray<FText> AdditionalContext;
+		AdditionalContext.Add(FormatText);
 
-	Writer.EndObject();
+		const IErrorDetails* DetailsIt = ErrorDetails.GetReference();
+		while (DetailsIt != nullptr)
+		{
+			if (DetailsIt->ShouldIncludeInContextMessage())
+			{
+				AdditionalContext.Add(FText::FromString(FUtf8String::Printf("(%s:{%s})", DetailsIt->GetErrorDetailsTypeNameForLog().GetData(), DetailsIt->GetErrorDetailsTypeNameForLog().GetData())));
+				
+			}
+			const IErrorDetails* Next = DetailsIt->GetInnerErrorDetails().GetReference(); 
+			DetailsIt = Next;
+		}
+		FormatText = FText::Join(FText::FromString(", "), AdditionalContext);
+	}
+
+#if USE_LOCALIZED_STRUCTURED_LOG_FOR_FERRORMESSAGE
+	UE::FLogTemplateOptions TemplateOptions;
+	TemplateOptions.bAllowSubObjectReferences = true;
+
+	FInlineLogTemplate Template(FormatText, TemplateOptions);
+#else
+	FString FormatString = FormatText.ToString();
+	UE::FLogTemplateOptions TemplateOptions;
+	TemplateOptions.bAllowSubObjectReferences = true;
+	FInlineLogTemplate Template(*FormatString, TemplateOptions);
+#endif
+	TCbWriter<1024> Writer;
 
-	FWideStringBuilderBase OutputMessage;
-	FInlineLogTemplate Template(*FormatString);
-	Template.FormatTo(OutputMessage, Writer.Save());
+	SerializeDetailsForLog(Writer);
+	FUtf8StringBuilderBase OutputMessage;
 
-	return FText::FromString(OutputMessage.ToString());
+	return Template.FormatToText(Writer.Save());
 #else
 	FFormatNamedArguments Args;
 	FTextFormatArgsPropertyExtractor PropertyExtractor(Args);
@@ -79,16 +115,23 @@ FText FError::GetErrorMessage(bool bIncludeContext) const
 #endif
 }
 
+
 void FError::SerializeDetailsForLog(FCbWriter& Writer) const
 {
 	const IErrorDetails* DetailsIt = ErrorDetails.GetReference();
 	while (DetailsIt != nullptr)
 	{
-		Writer.SetName(DetailsIt->GetErrorDetailsTypeName());
-		DetailsIt->SerializeForLog(Writer);
-
+		Writer.SetName(DetailsIt->GetErrorDetailsTypeNameForLog());
+		DetailsIt->SerializeToCb(Writer, *this);
 
-		DetailsIt = DetailsIt->GetInnerErrorDetails().GetReference();
+		const IErrorDetails* Next = DetailsIt->GetInnerErrorDetails().GetReference();
+		if (Next == nullptr)
+		{
+			// the last one contains some special information we want to reference by name
+			Writer.SetName(ANSITEXTVIEW("Root"));
+			DetailsIt->SerializeToCb(Writer, *this);
+		}
+		DetailsIt = Next;
 	}
 }
 
@@ -193,20 +236,49 @@ void FError::GetErrorProperties(IErrorPropertyExtractor& Visitor) const
 	ErrorDetails->GetErrorProperties(*this, Visitor);
 }
 
-const FUtf8String FError::GetErrorCodeString() const
+
+
+class FAppendFormatStringDetails : public FDynamicErrorDetails
 {
-	FUtf8String Result;
-	GetDetailByKey(UTF8TEXTVIEW("ErrorCodeString"), Result);
-	return Result;
-}
+	DECLARE_FERROR_DETAILS(UnifiedError, FAppendFormatStringDetails);
+public:
+	FAppendFormatStringDetails() {}
+	FAppendFormatStringDetails(FText&& InAppendFormatString, TRefCountPtr<const IErrorDetails> InInnerErrorDetails = nullptr) : FDynamicErrorDetails(InInnerErrorDetails)
+	{
+		AppendFormatString = MoveTemp(InAppendFormatString);
+	}
 
-const FUtf8String FError::GetModuleIdString() const
+private:
+	FText AppendFormatString;
+
+
+public:
+
+	virtual void SerializeToCb(FCbWriter& Writer, const FError& Error) const override
+	{
+		Writer.BeginObject();
+		SerializeBasicToCb(Writer, UTF8TEXTVIEW("AppendFormatString"), AppendFormatString);
+		Writer.EndObject();
+	}
+
+	CORE_API virtual const FText GetErrorFormatString(const FError& Error) const override
+	{
+		const FText InnerFormatString = GetInnerErrorDetails()->GetErrorFormatString(Error);
+		return FText::Join(FText::FromString(" "), InnerFormatString, AppendFormatString);
+	}
+
+	virtual bool ShouldIncludeInContextMessage() const { return false; }
+};
+
+
+
+void FError::AppendFormatString(FText&& InFormatString)
 {
-	FUtf8String Result;
-	GetDetailByKey(UTF8TEXTVIEW("ModuleIdString"), Result);
-	return Result;
+	TRefCountPtr<IErrorDetails> NewErrorDetails(new FAppendFormatStringDetails(MoveTemp(InFormatString)));
+	PushErrorDetails(NewErrorDetails);
 }
 
+
 // Structured logging FError integration
 class FCbWriterErrorPropertyExtractor : public IErrorPropertyExtractor
 {
@@ -255,10 +327,29 @@ void SerializeForLog(FCbWriter& Writer, const FError& Error)
 
 #if USE_STRUCTURED_LOG_FOR_FERRORMESSAGE
 	
+	// what I have to do 
+	TCbWriter<1024> InnerWriter;
+
+	Error.SerializeDetailsForLog(InnerWriter);
+
+	FText FormatText = FText::Join(FText::FromString(":"), NSLOCTEXT("UnifiedError", "DefaultLogPrepend", "{Root/ModuleIdString}.{Root/ErrorCodeString}"), Error.GetFormatErrorText());
+#if USE_LOCALIZED_STRUCTURED_LOG_FOR_FERRORMESSAGE
+	UE::FLogTemplateOptions TemplateOptions;
+	TemplateOptions.bAllowSubObjectReferences = true;
+	FInlineLogTemplate Template(FormatText, TemplateOptions);
+#else
+	UE::FLogTemplateOptions TemplateOptions;
+	TemplateOptions.bAllowSubObjectReferences = true;
+	FInlineLogTemplate Template(*FormatText.ToString(), TemplateOptions);
+#endif
+	FUtf8StringBuilderBase OutputMessage;
 
-	Error.SerializeDetailsForLog(Writer);
 
-	Writer.AddString(ANSITEXTVIEW("$text"), Error.GetFormatErrorText().ToString());
+	FCbFieldIterator CbInnerIterator = InnerWriter.Save();
+
+	Template.FormatTo(OutputMessage, CbInnerIterator);
+
+	Writer.AddString(ANSITEXTVIEW("$text"), OutputMessage.ToString());
 
 #else
 
@@ -276,9 +367,8 @@ void SerializeForLog(FCbWriter& Writer, const FError& Error)
 	FTextFormatArgsPropertyExtractor ExtractToFTextArgs(ErrorMessageArgs);
 	Error.GetErrorProperties(ExtractToFTextArgs);
 	Writer.AddString("$text", FText::Format(ErrorFormatText, MoveTemp(ErrorMessageArgs)).ToString());
-
-	Writer.EndObject();
 #endif
+	Writer.EndObject();
 }
 
 // FRefCountedErrorDetails fucntions
@@ -309,13 +399,28 @@ const FUtf8String& FDynamicErrorDetails::GetErrorName(const FError& Error) const
 	return InnerErrorDetails->GetErrorName(Error);
 }
 */
+/*
 
 
-void FDynamicErrorDetails::GetErrorProperties(const FError& Error, IErrorPropertyExtractor& OutProperties) const
+UE::UnifiedError::FAnyString FDynamicErrorDetails::GetErrorCodeString(const FError& Error) const
 {
-	InnerErrorDetails->GetErrorProperties(Error, OutProperties);
+	if (InnerErrorDetails.IsValid())
+	{
+		return InnerErrorDetails->GetErrorCodeString(Error);
+	}
+	return UE::UnifiedError::FAnyString();
 }
 
+UE::UnifiedError::FAnyString FDynamicErrorDetails::GetModuleIdString(const FError& Error) const
+{
+	if (InnerErrorDetails.IsValid())
+	{
+		return InnerErrorDetails->GetModuleIdString(Error);
+	}
+	return UE::UnifiedError::FAnyString();
+}*/
+
+
 void FDynamicErrorDetails::SetInnerErrorDetails(TRefCountPtr<const IErrorDetails> InInnerErrorDetails)
 {
 	InnerErrorDetails = InInnerErrorDetails;
@@ -329,29 +434,28 @@ FStaticErrorDetails::FStaticErrorDetails(const FAnsiStringView& InErrorName, con
 
 
 
-void FStaticErrorDetails::GetErrorProperties(const FError& Error, IErrorPropertyExtractor& OutProperties) const
-{
-	OutProperties.AddProperty(UTF8TEXTVIEW("ErrorCodeString"), ErrorName);
-	OutProperties.AddProperty(UTF8TEXTVIEW("ModuleIdString"), ModuleName);
-	OutProperties.AddProperty(UTF8TEXTVIEW("ErrorCode"), Error.GetErrorCode());
-	OutProperties.AddProperty(UTF8TEXTVIEW("ModuleId"), Error.GetModuleId());
-}
+/*
 
-const FAnsiStringView& FStaticErrorDetails::GetErrorCodeString() const
+UE::UnifiedError::FAnyString FStaticErrorDetails::GetErrorCodeString(const FError& Error) const
 {
-	return ErrorName;
+	UE::UnifiedError::FAnyString ReturnValue;
+	ReturnValue.ReferenceString(ErrorName);
+	return ReturnValue;
 }
 
-const FAnsiStringView& FStaticErrorDetails::GetModuleIdString() const
+UE::UnifiedError::FAnyString FStaticErrorDetails::GetModuleIdString(const FError& Error) const
 {
-	return ModuleName;
-}
+	UE::UnifiedError::FAnyString ReturnValue;
+	ReturnValue.ReferenceString(ModuleName);
+	return ReturnValue;
+}*/
 
 const FText FStaticErrorDetails::GetErrorFormatString(const FError& Error) const
 {
 	return ErrorFormatString;
 }
 
+
 // FErrorDetailsRegistry
 uint32 FErrorDetailsRegistry::RegisterDetails(const FAnsiStringView& ErrorDetailsName, TFunction<IErrorDetails*()> CreationFunction)
 {
@@ -366,10 +470,6 @@ uint32 FErrorDetailsRegistry::RegisterDetails(const FAnsiStringView& ErrorDetail
 
 FString LexToString(const UE::UnifiedError::FError& Error)
 {
-	
 	// TODO: daniel, switch to using structured logging output in the future instead of this
-	return FString::Printf(TEXT("%s:%s - %s"), 
-		UTF8_TO_TCHAR(*Error.GetModuleIdString()),
-		UTF8_TO_TCHAR(*Error.GetErrorCodeString()),
-		*Error.GetErrorMessage().ToString());
+	return Error.GetErrorMessage().ToString();
 }
\ No newline at end of file
diff --git a/Engine/Source/Runtime/Core/Public/Experimental/UnifiedError/UnifiedError.h b/Engine/Source/Runtime/Core/Public/Experimental/UnifiedError/UnifiedError.h
--- a/Engine/Source/Runtime/Core/Public/Experimental/UnifiedError/UnifiedError.h
+++ b/Engine/Source/Runtime/Core/Public/Experimental/UnifiedError/UnifiedError.h
@@ -6,8 +6,8 @@
 #include "Containers/StringView.h"
 #include "Internationalization/Internationalization.h"
 #include "Internationalization/Text.h"
-#include "Logging/StructuredLog.h"
 #include "Serialization/CompactBinaryWriter.h"
+#include "Misc/TVariant.h"
 #include "Templates/RefCounting.h"
 
 
@@ -21,6 +21,11 @@
 class TErrorStructFeatures
 {
 public:
+	static const FAnsiStringView GetErrorDetailsLogEntryNameAsString()
+	{
+		static_assert(deferred_false<T>::value, "All types are required to define DECLARE_ERRORSTRUCT_FEATURES to be used as a error context type");
+		return ANSITEXTVIEW("");
+	}
 	static const FAnsiStringView GetErrorDetailsTypeNameAsString()
 	{
 		static_assert(deferred_false<T>::value, "All types are required to define DECLARE_ERRORSTRUCT_FEATURES to be used as a error context type");
@@ -36,18 +41,22 @@
 
 
 
-#define DECLARE_ERRORSTRUCT_FEATURES(DetailsTypeName)\
+#define DECLARE_ERRORSTRUCT_FEATURES(DetailsNameSpace, DetailsStructName)\
 	template<>\
-	class TErrorStructFeatures<UE::UnifiedError::DetailsTypeName>\
+	class TErrorStructFeatures<UE::UnifiedError::DetailsNameSpace::DetailsStructName>\
 	{\
 	public:\
+		static const FAnsiStringView GetErrorDetailsLogEntryNameAsString() \
+		{\
+			return ANSITEXTVIEW(#DetailsNameSpace"_"#DetailsStructName);\
+		}\
 		static const FAnsiStringView GetErrorDetailsTypeNameAsString() \
 		{\
-			return ANSITEXTVIEW("TErrorDetails<"#DetailsTypeName">");\
+			return ANSITEXTVIEW("TErrorDetails<"#DetailsNameSpace"::"#DetailsStructName">");\
 		}\
 		static const FAnsiStringView GetErrorContextTypeNameAsString() \
 		{\
-			return ANSITEXTVIEW(#DetailsTypeName);\
+			return ANSITEXTVIEW(#DetailsNameSpace"::"#DetailsStructName);\
 		}\
 	};
 
@@ -61,27 +70,41 @@
 	{ \
 		return StaticGetErrorDetailsTypeId(); \
 	}\
-	virtual const FAnsiStringView GetErrorDetailsTypeName() const override \
+	static const FAnsiStringView StaticGetErrorDetailsTypeName() \
 	{\
 		return #DetailsNamespace"::"#TypeName;\
+	}\
+	virtual const FAnsiStringView GetErrorDetailsTypeName() const override \
+	{\
+		return StaticGetErrorDetailsTypeName();\
+	}\
+	static const FAnsiStringView StaticGetErrorDetailsTypeNameForLog() \
+	{\
+		return #DetailsNamespace"_"#TypeName;\
+	}\
+	virtual const FAnsiStringView GetErrorDetailsTypeNameForLog() const override\
+	{\
+		return StaticGetErrorDetailsTypeNameForLog(); \
 	}
 
 
 #define DECLARE_FERROR_DETAILS_ABSTRACT(DetailsNamespace, TypeName)\
 	public:\
-	inline static const uint32 StaticDetailsTypeId = UE::UnifiedError::FErrorDetailsRegistry::Get().RegisterDetails(ANSITEXTVIEW(#DetailsNamespace "::" #TypeName), nullptr);\
-	DECLARE_ERROR_DETAILS_INTERNAL(DetailsNamespace, TypeName);
+	DECLARE_ERROR_DETAILS_INTERNAL(DetailsNamespace, TypeName);\
+	inline static const uint32 StaticDetailsTypeId = UE::UnifiedError::FErrorDetailsRegistry::Get().RegisterDetails(StaticGetErrorDetailsTypeName(), nullptr);\
+	
 	
 
 #define DECLARE_FERROR_DETAILS(DetailsNamespace, TypeName)\
 	public:\
+	DECLARE_ERROR_DETAILS_INTERNAL(DetailsNamespace, TypeName);\
 	friend IErrorDetails* Create();\
 	inline static IErrorDetails* Create()\
 	{\
 		return new TypeName();\
 	}\
-	inline static const uint32 StaticDetailsTypeId = UE::UnifiedError::FErrorDetailsRegistry::Get().RegisterDetails(ANSITEXTVIEW(#DetailsNamespace "::" #TypeName), TFunction<IErrorDetails* ()>([]() -> IErrorDetails* { return DetailsNamespace::TypeName::Create(); }));\
-	DECLARE_ERROR_DETAILS_INTERNAL(DetailsNamespace, TypeName);
+	inline static const uint32 StaticDetailsTypeId = UE::UnifiedError::FErrorDetailsRegistry::Get().RegisterDetails(StaticGetErrorDetailsTypeName(), TFunction<IErrorDetails* ()>([]() -> IErrorDetails* { return DetailsNamespace::TypeName::Create(); }));\
+	
 
 
 
@@ -114,6 +137,36 @@
 
 	class FError;
 
+	// helper class to make pass by string view effiecent while also allowing pass by other strings
+	/*class FAnyString
+	{
+	private:
+		TVariant<FUtf8String, FUtf8StringView> VariantString;
+	public:
+		
+		void ReferenceString(FUtf8StringView InStringView)
+		{
+			VariantString.Set<FUtf8StringView>(InStringView);
+		}
+		void CopyString(FUtf8StringView InString)
+		{
+			VariantString.Set<FUtf8String>(FUtf8String(InString));
+		}
+
+		FUtf8StringView GetUtf8StringView() const
+		{
+			if (VariantString.IsType<FUtf8StringView>())
+			{
+				return VariantString.Get<FUtf8StringView>();
+			}
+			else if (VariantString.IsType<FUtf8String>())
+			{
+				return VariantString.Get<FUtf8String>();
+			}
+			return FUtf8StringView();
+		}
+	};*/
+
 	class IErrorDetails : public IRefCountedObject
 	{
 	public:
@@ -134,7 +187,7 @@
 		/// </summary>
 		/// <param name="Error"></param>
 		/// <param name="OutProperties"></param>
-		virtual void GetErrorProperties(const FError& Error, IErrorPropertyExtractor& OutProperties) const = 0;
+		CORE_API virtual void GetErrorProperties(const FError& Error, IErrorPropertyExtractor& OutProperties) const = 0;
 
 		/// <summary>
 		/// GetInnerErrorDetails; Exposes inner error details to FError, if this ErrorDetails allows inner details
@@ -155,10 +208,18 @@
 		/// <returns></returns>
 		virtual uint32 GetErrorDetailsTypeId() const = 0;
 
+		/*
 
-		virtual const FAnsiStringView GetErrorDetailsTypeName() const = 0;
-		virtual void SerializeForLog(FCbWriter& Writer) const {}
+		virtual FAnyString GetErrorCodeString(const FError& Error) const = 0;
+		virtual FAnyString GetModuleIdString(const FError& Error) const = 0;
+*/
 
+		CORE_API virtual const FAnsiStringView GetErrorDetailsTypeName() const = 0;
+		CORE_API virtual const FAnsiStringView GetErrorDetailsTypeNameForLog() const = 0;
+		virtual void SerializeToCb(FCbWriter& Writer, const FError& Error) const =0;
+
+
+		virtual bool ShouldIncludeInContextMessage() const { return true; }
 	};
 
 
@@ -216,6 +277,12 @@
 		CORE_API virtual void SetInnerErrorDetails(TRefCountPtr<const IErrorDetails> InInnerErrorDetails);
 		virtual TRefCountPtr<const IErrorDetails> GetInnerErrorDetails() const { return InnerErrorDetails; }
 
+/*
+		virtual FAnyString GetErrorCodeString(const FError& Error) const override;
+		virtual FAnyString GetModuleIdString(const FError& Error) const override;
+*/
+
+
 		/// <summary>
 		/// GetErrorFormatString; Pass through to the InnerErrorDetails.
 		/// </summary>
@@ -225,72 +292,14 @@
 
 		/// <summary>
 		/// GetErrorProperties; by default pass through to the InnerErrorDetails.
-		///  It's expected FDynamicErrorDetails will be derived from and implement GetErrorProperties. 
+		///  It's expected FDynamicErrorDetails will be derived from and implement GetErrorProperties.  Declared in header to avoid issues with dll exports
 		/// </summary>
 		/// <param name="Error"></param>
 		/// <param name="OutProperties"></param>
-		CORE_API virtual void GetErrorProperties(const FError& Error, IErrorPropertyExtractor& OutProperties) const override;
+		virtual void GetErrorProperties(const FError& Error, IErrorPropertyExtractor& OutProperties) const override
+		{
+			InnerErrorDetails->GetErrorProperties(Error, OutProperties);
+		}
 	};
 
 
-	/// <summary>
-	/// FStaticErrorDetails; static error details and members are statically allocated
-	///  Every error which uses DEFINE_ERROR will have FStaticErrorDetails generated for it
-	///  Can not rely on it to be available for every error as some Error conversion functions will not use pregenerated errors or error codesF
-	///  Use FError::GetErrorDetails to discover FStaticErrorDetails
-	/// </summary>
-	class FStaticErrorDetails : public IErrorDetails
-	{
-		DECLARE_FERROR_DETAILS_ABSTRACT(UnifiedError,FStaticErrorDetails);
-	private:
-		const FAnsiStringView& ErrorName;
-		const FAnsiStringView& ModuleName;
-		FText ErrorFormatString;
-	public:
-		// TODO: convert AnsiStringView to UTF8StringView when c++ 20 is supported
-		CORE_API FStaticErrorDetails(const FAnsiStringView& InErrorName, const FAnsiStringView& InModuleName, const FText& InErrorFormatString);
-
-		virtual ~FStaticErrorDetails() = default;
-
-		/// <summary>
-		/// GetErrorCodeString; Accessor for ErrorName.  
-		///  Can be called directly on FStaticErrorDetails object.  
-		///  See also: FError::GetErrorDetails
-		/// </summary>
-		/// <returns></returns>
-		CORE_API const FAnsiStringView& GetErrorCodeString() const;
-
-		/// <summary>
-		/// GetModuleIdString; accessor for ModuleName.
-		///  Can be called directly on FStaticErrorDetails object
-		///  See also: FError::GetErrorDetails
-		/// </summary>
-		CORE_API const FAnsiStringView& GetModuleIdString() const;
-
-		// IErrorDetails implementation
-
-		/// <summary>
-		/// GetErrorProperties;  Adds error properties exposed by this error details object to the OutProperties value
-		///  FStaticErrorDetails includes error details all errors have ErrorCode, ModuleId, ErrorCodeString, ModuleIdString
-		///  These can be used in any error format string see also GetErrorFormatString
-		/// </summary>
-		/// <param name="Error"></param>
-		/// <param name="OutProperties"></param>
-		CORE_API virtual void GetErrorProperties(const FError& Error, IErrorPropertyExtractor& OutProperties) const;
-
-		/// <summary>
-		/// GetErrorFormatString; return the localized format text generated in DECLARE_ERROR macro.
-		/// </summary>
-		/// <param name="Error"></param>
-		/// <returns></returns>
-		CORE_API virtual const FText GetErrorFormatString(const FError& Error) const final;
-
-
-		// IRefCountedObject implementation
-		//  FStaticErrorDetails is statically allocated, make sure it is never released.
-		virtual FReturnedRefCountValue AddRef() const final { return FReturnedRefCountValue(10); }
-		virtual uint32 Release() const final { return 10; }
-		virtual uint32 GetRefCount() const final { return 10; }
-	};
-
-
@@ -311,9 +320,10 @@
 	{
 	public:
 		TErrorDetails() {}
-		TErrorDetails(T&& InErrorDetail)
+		TErrorDetails(T&& InErrorDetail, bool bInAppendContextToMessage)
 		{
 			ErrorDetail = MoveTemp(InErrorDetail);
+			bAppendContextToMessage = bInAppendContextToMessage;
 		}
 		TErrorDetails(T&& InErrorDetail, TRefCountPtr<const IErrorDetails> InInnerErrorDetails) : FDynamicErrorDetails(InInnerErrorDetails)
 		{
@@ -325,7 +335,7 @@
 		{
 			return new TErrorDetails<T>();
 		}
-		inline static const uint32 StaticDetailsTypeId = UE::UnifiedError::FErrorDetailsRegistry::Get().RegisterDetails(TErrorStructFeatures<T>::GetNameAsString(), TFunction<IErrorDetails * ()>([]() -> IErrorDetails* { return Create(); }));
+		inline static const uint32 StaticDetailsTypeId = UE::UnifiedError::FErrorDetailsRegistry::Get().RegisterDetails(TErrorStructFeatures<T>::GetErrorDetailsTypeNameAsString(), TFunction<IErrorDetails * ()>([]() -> IErrorDetails* { return Create(); }));
 		static uint32 StaticGetErrorDetailsTypeId()
 		{
 			return StaticDetailsTypeId;
@@ -334,17 +344,29 @@
 		{
 			return StaticGetErrorDetailsTypeId();
 		}
+
+		static FAnsiStringView StaticGetErrorDetailsTypeName()
+		{
+			return TErrorStructFeatures<T>::GetErrorDetailsTypeNameAsString();
+		}
+		const FAnsiStringView GetErrorDetailsTypeName() const override
+		{
+			return TErrorStructFeatures<T>::GetErrorDetailsTypeNameAsString();
+		}
+
+		const FAnsiStringView GetErrorDetailsTypeNameForLog() const override
+		{
+			return TErrorStructFeatures<T>::GetErrorDetailsLogEntryNameAsString();
+		}
+
 	private:
 		T ErrorDetail;
-
-		FAnsiStringView GetTypeNameAsString() const
-		{
-			return TErrorStructFeatures<T>::GetErrorDetailsTypeNameAsString();
-		}
+		bool bAppendContextToMessage = false;
+		
 	public:
 		const T& GetErrorContext() const
 		{
-			return &ErrorDetail;
+			return ErrorDetail;
 		}
 
 		virtual void GetErrorProperties(const FError& Error, IErrorPropertyExtractor& OutProperties) const override
@@ -353,9 +375,18 @@
 			
 			FDynamicErrorDetails::GetErrorProperties(Error, OutProperties);
 		}
-		virtual void SerializeForLog(FCbWriter& Writer) const override
+		virtual void SerializeToCb(FCbWriter& Writer, const FError& Error) const override
 		{
-			::SerializeForLog(Writer, ErrorDetail);
+			SerializeForLog(Writer, ErrorDetail);
 		}
 
+		virtual const FText GetErrorFormatString(const FError& Error) const override
+		{
+			if (bAppendContextToMessage)
+			{
+				return FText::Join(FText::FromString(" "), FDynamicErrorDetails::GetErrorFormatString(Error), FText::FromString(FUtf8String::Printf("{%s}", GetErrorDetailsTypeNameForLog().GetData())));
+			}
+			return FDynamicErrorDetails::GetErrorFormatString(Error);
+		}
+
 		const T& GetValue() const
@@ -365,7 +396,6 @@
 
 	};
 
-
 	class FError
 	{
 	private: 
@@ -405,54 +435,46 @@
 			Invalidate();
 		}
 
-		CORE_API void GetErrorProperties(IErrorPropertyExtractor& Visitor) const;
+		void GetErrorProperties(IErrorPropertyExtractor& Visitor) const;
 
 		CORE_API int32 GetErrorCode() const;
 		CORE_API int32 GetModuleId() const;
 
-		CORE_API const FUtf8String GetErrorCodeString() const;
-		CORE_API const FUtf8String GetModuleIdString() const;
-
-		CORE_API void PushErrorDetails(TRefCountPtr<IErrorDetails> InErrorDetails);
-		template<typename DetailType>
-		TRefCountPtr<const DetailType> GetErrorDetails() const
-		{
-			TRefCountPtr<const IErrorDetails> CurrentIt = ErrorDetails;
-			while (CurrentIt != nullptr)
-			{
-				if (CurrentIt->GetErrorDetailsTypeId() == DetailType::StaticGetErrorDetailsTypeId())
-				{
-					return TRefCountPtr<const DetailType>((const DetailType*)CurrentIt.GetReference());
-				}
-				CurrentIt = CurrentIt->GetInnerErrorDetails();
-			}
-			return nullptr;
-		}
-
 		CORE_API FText GetErrorMessage(bool bAppendContext = false) const;
 		CORE_API FText GetFormatErrorText() const;
+
 		template<typename T>
-		void PushErrorContext(T&& InErrorStruct)
+		void PushErrorContext(T&& InErrorStruct, bool bShouldAppendToErrorMessage = false)
 		{
-			PushErrorDetails(TRefCountPtr<IErrorDetails>(new TErrorDetails<T>(MoveTemp(InErrorStruct))));
+			PushErrorDetails(TRefCountPtr<IErrorDetails>(new TErrorDetails<T>(MoveTemp(InErrorStruct), bShouldAppendToErrorMessage)));
 		}
 
 		template<typename T>
 		const T* GetErrorContext() const
 		{
-			TRefCountPtr<const IErrorDetails> CurrentIt = ErrorDetails;
+
+			TRefCountPtr<const TErrorDetails<T>> CurrentDetails = GetErrorDetails<const TErrorDetails<T>>();
+			if (const TErrorDetails<T>* Details = CurrentDetails.GetReference())
+			{
+				return &Details->GetErrorContext();
+			}
+			return nullptr;
+
+			/*TRefCountPtr<const IErrorDetails> CurrentIt = ErrorDetails;
 			while (CurrentIt != nullptr)
 			{
-				if (CurrentIt->GetErrorDetailsTypeName() == TErrorStructFeatures<T>::GetErrorDetailsTypeNameAsString())
+				if (CurrentIt->GetErrorDetailsTypeId() == GetErrorDetails<>)
 				{
 					const TErrorDetails<const T>* Details = (const TErrorDetails<const T>*)(CurrentIt.GetReference());
 					return &Details->GetErrorContext();
 				}
 				CurrentIt = CurrentIt->GetInnerErrorDetails();
 			}
-			return nullptr;
+			return nullptr;*/
 		}
 
+		CORE_API void AppendFormatString(FText&& InFormatString);
+
 		CORE_API void SerializeDetailsForLog(FCbWriter& Writer) const;
 
 		CORE_API bool GetDetailByKey(const FUtf8StringView& KeyName, FString& Result) const;
@@ -475,6 +497,24 @@
 			ErrorDetails = nullptr;
 		}
 	private:
+
+
+		CORE_API void PushErrorDetails(TRefCountPtr<IErrorDetails> InErrorDetails);
+		template<typename DetailType>
+		TRefCountPtr<const DetailType> GetErrorDetails() const
+		{
+			TRefCountPtr<const IErrorDetails> CurrentIt = ErrorDetails;
+			while (CurrentIt != nullptr)
+			{
+				if (CurrentIt->GetErrorDetailsTypeId() == DetailType::StaticGetErrorDetailsTypeId())
+				{
+					return TRefCountPtr<const DetailType>((const DetailType*)CurrentIt.GetReference());
+				}
+				CurrentIt = CurrentIt->GetInnerErrorDetails();
+			}
+			return nullptr;
+		}
+
 		friend bool operator==(const FError& Error, const FError& OtherError)
 		{
 			if ((Error.ModuleId == OtherError.ModuleId) &&
@@ -527,6 +567,85 @@
 	};
 
 
+	/// <summary>
+	/// FStaticErrorDetails; static error details and members are statically allocated
+	///  Every error which uses DEFINE_ERROR will have FStaticErrorDetails generated for it
+	///  Can not rely on it to be available for every error as some Error conversion functions will not use pregenerated errors or error codesF
+	///  Use FError::GetErrorDetails to discover FStaticErrorDetails
+	/// </summary>
+	class FStaticErrorDetails : public IErrorDetails
+	{
+		DECLARE_FERROR_DETAILS_ABSTRACT(UnifiedError, FStaticErrorDetails);
+	private:
+		const FAnsiStringView& ErrorName;
+		const FAnsiStringView& ModuleName;
+		FText ErrorFormatString;
+	public:
+		// TODO: convert AnsiStringView to UTF8StringView when c++ 20 is supported
+		CORE_API FStaticErrorDetails(const FAnsiStringView& InErrorName, const FAnsiStringView& InModuleName, const FText& InErrorFormatString);
+
+		virtual ~FStaticErrorDetails() = default;
+
+		/// <summary>
+		/// GetErrorCodeString; Accessor for ErrorName.  
+		///  Can be called directly on FStaticErrorDetails object.  
+		///  See also: FError::GetErrorDetails
+		/// </summary>
+		/// <returns></returns>
+		//virtual UE::UnifiedError::FAnyString GetErrorCodeString(const FError& Error) const override;
+
+		/// <summary>
+		/// GetModuleIdString; accessor for ModuleName.
+		///  Can be called directly on FStaticErrorDetails object
+		///  See also: FError::GetErrorDetails
+		/// </summary>
+		///virtual UE::UnifiedError::FAnyString GetModuleIdString(const FError& Error) const override;
+
+		// IErrorDetails implementation
+
+		/// <summary>
+		/// GetErrorProperties;  Adds error properties exposed by this error details object to the OutProperties value
+		///  FStaticErrorDetails includes error details all errors have ErrorCode, ModuleId, ErrorCodeString, ModuleIdString
+		///  These can be used in any error format string see also GetErrorFormatString
+		/// </summary>
+		/// <param name="Error"></param>
+		/// <param name="OutProperties"></param>
+		virtual void GetErrorProperties(const FError& Error, IErrorPropertyExtractor& OutProperties) const
+		{
+			OutProperties.AddProperty(UTF8TEXTVIEW("ErrorCodeString"), ErrorName);
+			OutProperties.AddProperty(UTF8TEXTVIEW("ModuleIdString"), ModuleName);
+			OutProperties.AddProperty(UTF8TEXTVIEW("ErrorCode"), Error.GetErrorCode());
+			OutProperties.AddProperty(UTF8TEXTVIEW("ModuleId"), Error.GetModuleId());
+		}
+
+		/// <summary>
+		/// GetErrorFormatString; return the localized format text generated in DECLARE_ERROR macro.
+		/// </summary>
+		/// <param name="Error"></param>
+		/// <returns></returns>
+		CORE_API virtual const FText GetErrorFormatString(const FError& Error) const final;
+
+		virtual void SerializeToCb(FCbWriter& Writer, const FError& Error) const final
+		{
+			Writer.BeginObject();
+			Writer.AddString(UTF8TEXTVIEW("$format"), UTF8TEXTVIEW("({ModuleIdString}.{ErrorCodeString})"));
+			Writer.AddString(UTF8TEXTVIEW("$type"), UTF8TEXTVIEW("FStaticErrorDetails"));
+			Writer.AddString(UTF8TEXTVIEW("ErrorCodeString"), ErrorName);
+			Writer.AddString(UTF8TEXTVIEW("ModuleIdString"), ModuleName);
+			Writer.AddInteger(UTF8TEXTVIEW("ErrorCode"), Error.GetErrorCode());
+			Writer.AddInteger(UTF8TEXTVIEW("ModuleId"), Error.GetModuleId());
+			Writer.EndObject();
+		}
+
+		// IRefCountedObject implementation
+		//  FStaticErrorDetails is statically allocated, make sure it is never released.
+		virtual FReturnedRefCountValue AddRef() const final { return FReturnedRefCountValue(10); }
+		virtual uint32 Release() const final { return 10; }
+		virtual uint32 GetRefCount() const final { return 10; }
+
+		virtual bool ShouldIncludeInContextMessage() const { return false; }
+	};
+
 	class FTextFormatArgsPropertyExtractor : public IErrorPropertyExtractor
 	{
 	private:
@@ -572,7 +691,36 @@
 		}
 	};
 
-	CORE_API void SerializeForLog(FCbWriter& Writer, const FError& OnlineError);
+	inline void SerializeBasicToCb(FCbWriter& Writer, const FUtf8StringView& PropertyName, const FStringView& Value)
+	{
+		Writer.AddString(PropertyName, Value);
+	}
+	inline void SerializeBasicToCb(FCbWriter& Writer, const FUtf8StringView& PropertyName, const FUtf8StringView& Value)
+	{
+		Writer.AddString(PropertyName, Value);
+	}
+	inline void SerializeBasicToCb(FCbWriter& Writer, const FUtf8StringView& PropertyName, const int32& Value)
+	{
+		Writer.AddInteger(PropertyName, Value);
+	}
+	inline void SerializeBasicToCb(FCbWriter& Writer, const FUtf8StringView& PropertyName, const uint32& Value)
+	{
+		Writer.AddInteger(PropertyName, Value);
+	}
+	inline void SerializeBasicToCb(FCbWriter& Writer, const FUtf8StringView& PropertyName, const int64& Value)
+	{
+		Writer.AddInteger(PropertyName, Value);
+	}
+	inline void SerializeBasicToCb(FCbWriter& Writer, const FUtf8StringView& PropertyName, const uint64& Value)
+	{
+		Writer.AddInteger(PropertyName, Value);
+	}
+	inline void SerializeBasicToCb(FCbWriter& Writer, const FUtf8StringView& PropertyName, const FText& Value)
+	{
+		Writer.AddString(PropertyName, Value.ToString());
+	}
 
+	CORE_API void SerializeForLog(FCbWriter& Writer, const FError& Error);
+
 } // namespace UnifiedError
 
@@ -647,6 +795,12 @@
 					OutProperties.AddProperty(UTF8TEXTVIEW(#ParamOneName), ParamOneName); \
 					FDynamicErrorDetails::GetErrorProperties(Error, OutProperties); \
 				} \
+				virtual void SerializeToCb(FCbWriter& Writer, const FError& Error) const override \
+				{\
+					Writer.BeginObject();\
+					SerializeBasicToCb(Writer, UTF8TEXTVIEW(#ParamOneName), ParamOneName); \
+					Writer.EndObject();\
+				}\
 				const ParamOneType& Get##ParamOneName() const \
 				{ \
 					return ParamOneName; \
diff --git a/Engine/Source/Runtime/Engine/Classes/Engine/Experimental/StreamableManagerError.h b/Engine/Source/Runtime/Engine/Classes/Engine/Experimental/StreamableManagerError.h
--- a/Engine/Source/Runtime/Engine/Classes/Engine/Experimental/StreamableManagerError.h
+++ b/Engine/Source/Runtime/Engine/Classes/Engine/Experimental/StreamableManagerError.h
@@ -57,11 +57,20 @@ namespace UE::UnifiedError::StreamableManager
 		{
 			const FText InnerFormatString = GetInnerErrorDetails()->GetErrorFormatString(Error);
 
-			FText MyFormatString = NSLOCTEXT("StreamableManager", "DetailsFormatString", "(SoftObjectPath:{SoftObjectPath})");
+			FText MyFormatString = NSLOCTEXT("StreamableManager", "DetailsFormatString", "SoftObjectPath:{StreamableManager_FStreamableManagerErrorDetails/SoftObjectPath}");
 
 			return FText::Join(NSLOCTEXT("StreamableManager", "DetailsFormatStringDelimiter", " "), InnerFormatString, MyFormatString);
 		}
 
+		virtual void SerializeToCb(FCbWriter& Writer, const FError& Error) const
+		{
+			Writer.BeginObject();
+			Writer.AddString(UTF8TEXTVIEW("$type"), UTF8TEXTVIEW("FStreamableManagerErrorDetails"));
+			Writer.AddString(UTF8TEXTVIEW("$format"), UTF8TEXTVIEW("(SoftObjectPath:{SoftObjectPath})"));
+			Writer.AddString(UTF8TEXTVIEW("SoftObjectPath"), SoftObjectPath);
+			Writer.EndObject();
+		}
+
 		const FString& GetSoftObjectPath() const { return SoftObjectPath; }
 
 		
@@ -75,26 +84,24 @@ namespace UE::UnifiedError::StreamableManager
 	public:
 		/*friend void SerializeForLog(FCbWriter& Writer, const UE::UnifiedError::StreamableManager::FStreamableManagerAdditionalContext& Context);
 	private:*/
-		FString ExtraSoftObjectPath;
+		FString SoftObjectPath;
 	};
 
 
 	inline void GatherPropertiesForError(const UE::UnifiedError::FError& Error, const UE::UnifiedError::StreamableManager::FStreamableManagerAdditionalContext& Context, UE::UnifiedError::IErrorPropertyExtractor& PropertyExtractor)
 	{
-		PropertyExtractor.AddProperty(ANSITEXTVIEW("ExtraSoftObjectPath"), Context.ExtraSoftObjectPath);
+		PropertyExtractor.AddProperty(ANSITEXTVIEW("SoftObjectPath"), Context.SoftObjectPath);
 	}
 	
 }
-
+DECLARE_ERRORSTRUCT_FEATURES(StreamableManager, FStreamableManagerAdditionalContext);
 inline void SerializeForLog(FCbWriter& Writer, const UE::UnifiedError::StreamableManager::FStreamableManagerAdditionalContext& Context)
 {
 	Writer.BeginObject();
-	Writer.AddString(ANSITEXTVIEW("$type"), ANSITEXTVIEW("FStreamableManagerAdditionalContext"));
-	FStringFormatNamedArguments NamedArguments;
-	NamedArguments.Add(TTuple<FString, FStringFormatArg>(TEXT("SoftObjectPath"), Context.ExtraSoftObjectPath));
-	Writer.AddString(ANSITEXTVIEW("$text"), FString::Format(TEXT("(SoftObjectPath: {SoftObjectPath})"), NamedArguments));
-	Writer.AddString(ANSITEXTVIEW("SoftObjectPath"), Context.ExtraSoftObjectPath);
+	Writer.AddString(ANSITEXTVIEW("$type"), TErrorStructFeatures<UE::UnifiedError::StreamableManager::FStreamableManagerAdditionalContext>::GetErrorContextTypeNameAsString());
+	Writer.AddString(ANSITEXTVIEW("$format"), TEXT("(SoftObjectPath: {SoftObjectPath})"));
+	Writer.AddString(ANSITEXTVIEW("SoftObjectPath"), Context.SoftObjectPath);
 	Writer.EndObject();
 }
-DECLARE_ERRORSTRUCT_FEATURES(StreamableManager::FStreamableManagerAdditionalContext);
+
 
diff --git a/Engine/Source/Runtime/Engine/Private/AssetManager.cpp b/Engine/Source/Runtime/Engine/Private/AssetManager.cpp
--- a/Engine/Source/Runtime/Engine/Private/AssetManager.cpp
+++ b/Engine/Source/Runtime/Engine/Private/AssetManager.cpp
@@ -17,6 +17,7 @@
 #include "Engine/BlueprintGeneratedClass.h"
 #include "Engine/GameInstance.h"
 #include "Engine/StreamableManager.h"
+#include "Engine/Experimental/StreamableManagerError.h"
 #include "HAL/PlatformApplicationMisc.h"
 #include "HAL/PlatformFileManager.h"
 #include "Interfaces/IPluginManager.h"
@@ -2114,15 +2115,15 @@ void UAssetManager::OnAssetStateChangeCompleted(TSharedPtr<FStreamableHandle> Bo
 			// If an error occured, don't update the current state with a bad state
 			if (BoundHandle->HasError())
 			{
-				TRefCountPtr<const UE::UnifiedError::StreamableManager::FStreamableManagerErrorDetails> ErrorDetails = BoundHandle->GetError()->GetErrorDetails<UE::UnifiedError::StreamableManager::FStreamableManagerErrorDetails>();
+				const UE::UnifiedError::StreamableManager::FStreamableManagerAdditionalContext* AdditionalContext = BoundHandle->GetError()->GetErrorContext<UE::UnifiedError::StreamableManager::FStreamableManagerAdditionalContext>();
 
-				if (ErrorDetails.IsValid())
+				if (AdditionalContext)
 				{
-					UE_LOG(LogAssetManager, Warning, TEXT("OnAssetStateChangeCompleted: Received error (%s) loading asset! Primary asset id: (%s) erroring asset:%s"), *LexToString(*BoundHandle->GetError()), *PrimaryAssetId.ToString(), *ErrorDetails->GetSoftObjectPath());
+					UE_LOG(LogAssetManager, Warning, TEXT("OnAssetStateChangeCompleted: Received error (%s) loading asset! Primary asset id:(%s) erroring asset:%s"), *BoundHandle->GetError()->GetErrorMessage(true).ToString(), *PrimaryAssetId.ToString(), *AdditionalContext->SoftObjectPath);
 				}
 				else
 				{
-					UE_LOG(LogAssetManager, Warning, TEXT("OnAssetStateChangeCompleted: Received error (%s) loading asset! Primary asset id:(%s)"), *LexToString(*BoundHandle->GetError()), *PrimaryAssetId.ToString());
+					UE_LOG(LogAssetManager, Warning, TEXT("OnAssetStateChangeCompleted: Received error (%s) loading asset! Primary asset id:(%s)"), *BoundHandle->GetError()->GetErrorMessage(true).ToString(), *PrimaryAssetId.ToString());
 				}
 				UE_LOGFMT(LogAssetManager, Warning, "OnAssetStateChangeCompleted: Received error ({Error}) loading asset! Primary asset id:({PrimaryAssetId})", ("Error", BoundHandle->GetError()), ("PrimaryAssetId", *PrimaryAssetId.ToString()));
 			}
diff --git a/Engine/Source/Runtime/Engine/Private/StreamableManager.cpp b/Engine/Source/Runtime/Engine/Private/StreamableManager.cpp
--- a/Engine/Source/Runtime/Engine/Private/StreamableManager.cpp
+++ b/Engine/Source/Runtime/Engine/Private/StreamableManager.cpp
@@ -1225,13 +1225,21 @@ void FStreamableHandle::AsyncLoadCallbackWrapper(const FName& PackageName, UPack
 		if (Result != EAsyncLoadingResult::Succeeded)
 		{
 			MaybeError = UE::UnifiedError::StreamableManager::GetStreamableError(Result);
-			TRefCountPtr<UE::UnifiedError::IErrorDetails> ErrorDetails(new UE::UnifiedError::StreamableManager::FStreamableManagerErrorDetails(TargetName.ToString()));
-			MaybeError->PushErrorDetails(ErrorDetails);
+			/*TRefCountPtr<UE::UnifiedError::IErrorDetails> ErrorDetails(new UE::UnifiedError::StreamableManager::FStreamableManagerErrorDetails(TargetName.ToString()));
+			MaybeError->PushErrorDetails(ErrorDetails);*/
+
+			UE::UnifiedError::StreamableManager::FStreamableManagerAdditionalContext MoreContext = {TargetName.ToString()};
+			MaybeError->PushErrorContext(MoveTemp(MoreContext));
+
+			// todo: Daniel, Devin, this shouldn't be nessisary just use PushErrorContent(T, true) after we have support for $format
+			MaybeError->AppendFormatString(NSLOCTEXT("StreamableManager", "AdditionalContextFormatString", "Missing streamable asset: {StreamableManager_FStreamableManagerAdditionalContext/SoftObjectPath}"));
 
-			/*UE::UnifiedError::StreamableManager::FStreamableManagerAdditionalContext MoreContext = {TargetName.ToString()};
-			MaybeError->PushErrorStruct(MoveTemp(MoreContext));*/
 			SetError(*MaybeError);
 		}
+		else
+		{
+			//if (TargetName.ToString())
+		}
 
 		OwningManager->AsyncLoadCallback(MoveTemp(TargetName), Package, MaybeError);
 
