Fix some command line arguments documentation and allow override of TraceThreadSleepTime

#rnx

diff --git a/Engine/Source/Runtime/Core/Private/ProfilingDebugging/TraceAuxiliary.cpp b/Engine/Source/Runtime/Core/Private/ProfilingDebugging/TraceAuxiliary.cpp
--- a/Engine/Source/Runtime/Core/Private/ProfilingDebugging/TraceAuxiliary.cpp
+++ b/Engine/Source/Runtime/Core/Private/ProfilingDebugging/TraceAuxiliary.cpp
@@ -1073,7 +1073,7 @@ static void TraceAuxiliarySetupInitFromConfig(UE::Trace::FInitializeDesc& OutDes
 	}
 
 	// Note that these options can only be used when tracing from a forked process (e.g. server).
-	// For a regular process use command line argument -TraceThreadSleepTime and -TraceTailSizeMb
+	// For a regular process use command line argument -TraceThreadSleepTime and -TraceTailMb
 
 	int32 SleepTimeConfig = 0;
 	if (GConfig->GetInt(GTraceConfigSection, TEXT("SleepTimeInMS"), SleepTimeConfig, GEngineIni))
@@ -1927,23 +1927,26 @@ void FTraceAuxiliary::Initialize(const TCHAR* CommandLine)
 		Desc.TailSizeBytes <<= 20;
 	}
 
-	// Memory tracing is very chatty. To reduce load on trace we'll speed up the
-	// worker thread so it can clear events faster.
-	extern bool MemoryTrace_IsActive();
-	if (MemoryTrace_IsActive())
+	if (!FParse::Value(CommandLine, TEXT("-tracethreadsleeptime="), Desc.ThreadSleepTimeInMS))
 	{
-		int32 SleepTimeMs = 5;
-		if (GConfig)
+		// Memory tracing is very chatty. To reduce load on trace we'll speed up the
+		// worker thread so it can clear events faster.
+		extern bool MemoryTrace_IsActive();
+		if (MemoryTrace_IsActive())
 		{
-			GConfig->GetInt(GTraceConfigSection, TEXT("SleepTimeWhenMemoryTracingInMS"), SleepTimeMs, GEngineIni);
-		}
+			int32 SleepTimeMs = 5;
+			if (GConfig)
+			{
+				GConfig->GetInt(GTraceConfigSection, TEXT("SleepTimeWhenMemoryTracingInMS"), SleepTimeMs, GEngineIni);
+			}
 
-		if (Desc.ThreadSleepTimeInMS)
-		{
-			SleepTimeMs = FMath::Min<uint32>(Desc.ThreadSleepTimeInMS, SleepTimeMs);
-		}
+			if (Desc.ThreadSleepTimeInMS)
+			{
+				SleepTimeMs = FMath::Min<uint32>(Desc.ThreadSleepTimeInMS, SleepTimeMs);
+			}
 
-		Desc.ThreadSleepTimeInMS = SleepTimeMs;
+			Desc.ThreadSleepTimeInMS = SleepTimeMs;
+		}
 	}
 
 	UE::Trace::Initialize(Desc);
