Fix stats system asserts when running with the new GPU profiler
- The EStatOperation::AdvanceFrameEventEndOfPipe operation was not always being triggered because the "StatsFrame" integer may be negative, and was being skipped over in "*LastStatsFrame < *FrameState.StatsFrame". Check needs to be !=.
- Also moved the StatsFrameRT integer into the stats system itself. AdvanceFrameEventEndOfPipe is only emitted in one location now.

diff --git a/Engine/Source/Runtime/Core/Private/Stats/StatsSystem.cpp b/Engine/Source/Runtime/Core/Private/Stats/StatsSystem.cpp
--- a/Engine/Source/Runtime/Core/Private/Stats/StatsSystem.cpp
+++ b/Engine/Source/Runtime/Core/Private/Stats/StatsSystem.cpp
@@ -22,6 +22,13 @@ DEFINE_STAT(STAT_SecondsPerCycle);
 
 namespace UE::Stats
 {
+#if STATS
+	// The engine stats system maintains its own frame counter. We need to propagate this to
+	// the end-of-pipe thread so that stat data it submits is attributed to the correct frame.
+	// This is done inside AdvanceRenderingThreadStats.
+	CORE_API int64 FStats::StatsFrameRT = 0;
+#endif
+
 	TAtomic<int32> FStats::GameThreadStatsFrame(1);
 
 	namespace Private
@@ -83,7 +90,7 @@ namespace UE::Stats
 		{
 			// There is no rendering thread, so this message is sufficient to make stats happy and don't leak memory.
 			FThreadStats::AddMessage(FStatConstants::AdvanceFrame.GetEncodedName(), EStatOperation::AdvanceFrameEventRenderThread, Frame);
-			FThreadStats::AddMessage(FStatConstants::AdvanceFrame.GetEncodedName(), EStatOperation::AdvanceFrameEventEndOfPipe, Frame);
+			FStats::StatsFrameRT = Frame;
 		}
 
 		FThreadStats::ExplicitFlush(bDiscardCallstack);
diff --git a/Engine/Source/Runtime/Core/Public/Stats/StatsSystem.h b/Engine/Source/Runtime/Core/Public/Stats/StatsSystem.h
--- a/Engine/Source/Runtime/Core/Public/Stats/StatsSystem.h
+++ b/Engine/Source/Runtime/Core/Public/Stats/StatsSystem.h
@@ -42,6 +42,13 @@ namespace UE::Stats
 
 		/** Current game thread stats frame. */
 		static CORE_API TAtomic<int32> GameThreadStatsFrame;
+
+	#if STATS
+		// The engine stats system maintains its own frame counter. We need to propagate this to
+		// the end-of-pipe thread so that stat data it submits is attributed to the correct frame.
+		// This is done inside AdvanceRenderingThreadStats.
+		static CORE_API int64 StatsFrameRT;
+	#endif
 	};
 } // namespace UE::Stats
 
diff --git a/Engine/Source/Runtime/RHI/Private/GPUProfiler.cpp b/Engine/Source/Runtime/RHI/Private/GPUProfiler.cpp
--- a/Engine/Source/Runtime/RHI/Private/GPUProfiler.cpp
+++ b/Engine/Source/Runtime/RHI/Private/GPUProfiler.cpp
@@ -2343,7 +2343,7 @@ namespace UE::RHI::GPUProfiler
 			check(FrameState.StatsFrame.IsSet());
 			FEndOfPipeStats* Stats = FEndOfPipeStats::Get();
 
-			if (!LastStatsFrame.IsSet() || *LastStatsFrame < *FrameState.StatsFrame)
+			if (!LastStatsFrame.IsSet() || *LastStatsFrame != *FrameState.StatsFrame)
 			{
 				Stats->AddMessage(FStatConstants::AdvanceFrame.GetEncodedName(), EStatOperation::AdvanceFrameEventEndOfPipe, *FrameState.StatsFrame);
 				LastStatsFrame = *FrameState.StatsFrame;
diff --git a/Engine/Source/Runtime/RHI/Private/RHICommandList.cpp b/Engine/Source/Runtime/RHI/Private/RHICommandList.cpp
--- a/Engine/Source/Runtime/RHI/Private/RHICommandList.cpp
+++ b/Engine/Source/Runtime/RHI/Private/RHICommandList.cpp
@@ -14,6 +14,7 @@
 #include "RHIContext.h"
 #include "RHIFwd.h"
 #include "RHITransition.h"
+#include "Stats/StatsSystem.h"
 #include "Stats/StatsTrace.h"
 #include "Stats/ThreadIdleStats.h"
 #include "HAL/PlatformMisc.h"
@@ -1429,7 +1430,7 @@ RHI_API FGraphEventRef FRHICommandListExecutor::Submit(TConstArrayView<FRHIComma
 				  this
 				, State = SubmitState
 			#if STATS
-				, LocalStatsFrame = StatsFrame
+				, LocalStatsFrame = UE::Stats::FStats::StatsFrameRT
 			#endif
 			]() mutable
 			{
diff --git a/Engine/Source/Runtime/RHI/Public/RHICommandList.h b/Engine/Source/Runtime/RHI/Public/RHICommandList.h
--- a/Engine/Source/Runtime/RHI/Public/RHICommandList.h
+++ b/Engine/Source/Runtime/RHI/Public/RHICommandList.h
@@ -5477,19 +5477,6 @@ private:
 	// Counts the number of calls to RHIEndFrame, and is used in GPU profiler frame boundary events.
 	uint32 FrameNumber = 0;
 
-#if STATS
-	// The engine stats system maintains its own frame counter. We need to propagate this to
-	// the end-of-pipe thread so that stat data it submits is attributed to the correct frame.
-	// This is done inside AdvanceRenderingThreadStats.
-	int64 StatsFrame = 0;
-public:
-	void SetStatsFrame(int64 InStatsFrame)
-	{
-		StatsFrame = InStatsFrame;
-	}
-private:
-#endif
-
 	bool AllowParallel() const;
 };
 
diff --git a/Engine/Source/Runtime/RenderCore/Private/RenderingThread.cpp b/Engine/Source/Runtime/RenderCore/Private/RenderingThread.cpp
--- a/Engine/Source/Runtime/RenderCore/Private/RenderingThread.cpp
+++ b/Engine/Source/Runtime/RenderCore/Private/RenderingThread.cpp
@@ -21,6 +21,7 @@
 #include "RenderCommandFence.h"
 #include "RenderDeferredCleanup.h"
 #include "TickableObjectRenderThread.h"
+#include "Stats/StatsSystem.h"
 #include "Stats/StatsData.h"
 #include "HAL/ThreadHeartBeat.h"
 #include "RenderResource.h"
@@ -347,7 +348,7 @@ static void AdvanceRenderingThreadStats(int64 StatsFrame, int32 DisableChangeTag
 
 #if RHI_NEW_GPU_PROFILER
 	// Propagate the stats frame value down to the end-of-pipe thread.
-	GRHICommandList.SetStatsFrame(Frame);
+	UE::Stats::FStats::StatsFrameRT = Frame;
 #else
 	FThreadStats::AddMessage(FStatConstants::AdvanceFrame.GetEncodedName(), EStatOperation::AdvanceFrameEventEndOfPipe, Frame);
 #endif
