Fixed UE_STATIC_ASSERT_COMPLETE_TYPE, GetData(), GetNum() and TIsContiguousContainer for zero-sized arrays, which fixes MakeArrayView() for zero-sized arrays.

diff --git a/Engine/Source/Runtime/Core/Public/Misc/StaticAssertCompleteType.h b/Engine/Source/Runtime/Core/Public/Misc/StaticAssertCompleteType.h
--- a/Engine/Source/Runtime/Core/Public/Misc/StaticAssertCompleteType.h
+++ b/Engine/Source/Runtime/Core/Public/Misc/StaticAssertCompleteType.h
@@ -2,11 +2,46 @@
 
 #pragma once
 
+#include "CoreTypes.h"
+#include <utility>
+
 namespace UEStaticAssertCompleteType_Private
 {
 	class FIncompleteType;
 
-	template <typename T> struct TUEStaticAssertTypeChecker       { static T& Func(); };
+#if defined(__clang__) || defined(__GNUC__)
+	// Under Clang, a zero-sized array cannot be used in a template specialization, and also under GCC cannot be used as a function overload
+	// or used with traits like std::is_array or std::remove_extent.
+	// Testing for sizeof(T) == 0 works though, and should be the only type this is true for, then operator+ is used to coerce this into a
+	// pointer and dereferenced to give a reference which can be used to check the completeness of the element type of the array.
+	template <typename T, SIZE_T SizeOfT = sizeof(T)>
+	struct TUEStaticAssertTypeCheckerBase
+	{
+		static T& Func();
+	};
+	template <typename T>
+	struct TUEStaticAssertTypeCheckerBase<T, 0>
+	{
+		static decltype(*+std::declval<T>())& Func();
+	};
+#else
+	// The above trick doesn't work for zero-sized arrays under MSVC, where sizeof(T[0]) is illegal, but specialization does.
+	template <typename T>
+	struct TUEStaticAssertTypeCheckerBase
+	{
+		static T& Func();
+	};
+	template <typename T>
+	struct TUEStaticAssertTypeCheckerBase<T[0]>
+	{
+		static T& Func();
+	};
+#endif
+
+	template <typename T>
+	struct TUEStaticAssertTypeChecker : TUEStaticAssertTypeCheckerBase<T>
+	{
+	};
 
 	// Voids will give an unhelpful error message when trying to take a reference to them, so work around it
 	template <> struct TUEStaticAssertTypeChecker<               void> { static FIncompleteType Func(); };
@@ -30,19 +65,21 @@ namespace UEStaticAssertCompleteType_Private
 
 // Each of these should fail to compile
 #if 0
-	UE_STATIC_ASSERT_COMPLETE_TYPE(               void,                                     "CV void is incomplete");
-	UE_STATIC_ASSERT_COMPLETE_TYPE(      volatile void,                                     "CV void is incomplete");
-	UE_STATIC_ASSERT_COMPLETE_TYPE(const          void,                                     "CV void is incomplete");
-	UE_STATIC_ASSERT_COMPLETE_TYPE(const volatile void,                                     "CV void is incomplete");
+	UE_STATIC_ASSERT_COMPLETE_TYPE(               void,                                    "CV void is incomplete");
+	UE_STATIC_ASSERT_COMPLETE_TYPE(      volatile void,                                    "CV void is incomplete");
+	UE_STATIC_ASSERT_COMPLETE_TYPE(const          void,                                    "CV void is incomplete");
+	UE_STATIC_ASSERT_COMPLETE_TYPE(const volatile void,                                    "CV void is incomplete");
 	UE_STATIC_ASSERT_COMPLETE_TYPE(UEStaticAssertCompleteType_Private::FIncompleteType,    "A forward-declared-but-undefined class is incomplete");
 	UE_STATIC_ASSERT_COMPLETE_TYPE(UEStaticAssertCompleteType_Private::FIncompleteType[2], "An array of an incomplete class is incomplete");
-	UE_STATIC_ASSERT_COMPLETE_TYPE(int[],                                                   "An array of a complete type of unspecified bound is incomplete");
+	UE_STATIC_ASSERT_COMPLETE_TYPE(int[],                                                  "An array of a complete type of unspecified bound is incomplete");
+	UE_STATIC_ASSERT_COMPLETE_TYPE(UEStaticAssertCompleteType_Private::FIncompleteType[0], "A zero-sized array of an incomplete type is incomplete");
 #endif
 
 // Each of these should pass
 #if 0
-	UE_STATIC_ASSERT_COMPLETE_TYPE(UEStaticAssertCompleteType_Private::FIncompleteType*,                                                         "A pointer to an incomplete type is complete");
-	UE_STATIC_ASSERT_COMPLETE_TYPE(UEStaticAssertCompleteType_Private::FIncompleteType&,                                                         "A reference to an incomplete type is complete");
+	UE_STATIC_ASSERT_COMPLETE_TYPE(UEStaticAssertCompleteType_Private::FIncompleteType*,                                                        "A pointer to an incomplete type is complete");
+	UE_STATIC_ASSERT_COMPLETE_TYPE(UEStaticAssertCompleteType_Private::FIncompleteType&,                                                        "A reference to an incomplete type is complete");
 	UE_STATIC_ASSERT_COMPLETE_TYPE(UEStaticAssertCompleteType_Private::FIncompleteType   (UEStaticAssertCompleteType_Private::FIncompleteType), "A function type is not incomplete, even if it returns or takes an incomplete type");
 	UE_STATIC_ASSERT_COMPLETE_TYPE(UEStaticAssertCompleteType_Private::FIncompleteType(&)(UEStaticAssertCompleteType_Private::FIncompleteType), "References to function types must give a good error");
+	UE_STATIC_ASSERT_COMPLETE_TYPE(int[0],                                                                                                      "A zero-sized array of a complete type is complete");
 #endif
diff --git a/Engine/Source/Runtime/Core/Public/Templates/UnrealTemplate.h b/Engine/Source/Runtime/Core/Public/Templates/UnrealTemplate.h
--- a/Engine/Source/Runtime/Core/Public/Templates/UnrealTemplate.h
+++ b/Engine/Source/Runtime/Core/Public/Templates/UnrealTemplate.h
@@ -69,6 +69,39 @@ FORCEINLINE void Move(T& A,typename TMoveSupportTraits<T>::Move B)
 	::new((void*)&A) T(MoveTemp(B));
 }
 
+namespace UE::Core::Private
+{
+	template <typename T>
+	[[nodiscard]] constexpr auto GetDataImpl(T&& Container) -> decltype(Container.GetData())
+	{
+		return Container.GetData();
+	}
+
+	template <
+		typename T
+		UE_REQUIRES(std::is_pointer_v<decltype(+std::declval<T>())>)
+	>
+	[[nodiscard]] constexpr auto GetDataImpl(T&& Container)
+	{
+		return +Container;
+	}
+
+	template <typename T>
+	[[nodiscard]] constexpr auto GetNumImpl(const T& Container) -> decltype(Container.Num())
+	{
+		return Container.Num();
+	}
+
+	template <
+		typename T
+		UE_REQUIRES(std::is_pointer_v<decltype(+std::declval<T>())>)
+	>
+	[[nodiscard]] constexpr auto GetNumImpl(const T& Container) -> SIZE_T
+	{
+		return sizeof(Container) / sizeof(*Container);
+	}
+}
+
 /**
  * Generically gets the data pointer of a contiguous container
  */
@@ -76,16 +109,11 @@ template <
 	typename T
 	UE_REQUIRES(TIsContiguousContainer<T>::Value)
 >
-constexpr auto GetData(T&& Container) -> decltype(Container.GetData())
+[[nodiscard]] constexpr auto GetData(T&& Container) -> decltype(UE::Core::Private::GetDataImpl((T&&)Container))
 {
-	return Container.GetData();
+	return UE::Core::Private::GetDataImpl((T&&)Container);
 }
 
-template <typename T, SIZE_T N> constexpr T* GetData(      T (& Container)[N]) { return Container; }
-template <typename T, SIZE_T N> constexpr T* GetData(      T (&&Container)[N]) { return Container; }
-template <typename T, SIZE_T N> constexpr const T* GetData(const T (& Container)[N]) { return Container; }
-template <typename T, SIZE_T N> constexpr const T* GetData(const T (&&Container)[N]) { return Container; }
-
 template <typename T>
 constexpr const T* GetData(std::initializer_list<T> List)
 {
@@ -99,16 +127,11 @@ template <
 	typename T
 	UE_REQUIRES(TIsContiguousContainer<T>::Value)
 >
-constexpr auto GetNum(T&& Container) -> decltype(Container.Num())
+[[nodiscard]] constexpr auto GetNum(const T& Container) -> decltype(UE::Core::Private::GetNumImpl(Container))
 {
-	return Container.Num();
+	return UE::Core::Private::GetNumImpl(Container);
 }
 
-template <typename T, SIZE_T N> constexpr SIZE_T GetNum(      T (& Container)[N]) { return N; }
-template <typename T, SIZE_T N> constexpr SIZE_T GetNum(      T (&&Container)[N]) { return N; }
-template <typename T, SIZE_T N> constexpr SIZE_T GetNum(const T (& Container)[N]) { return N; }
-template <typename T, SIZE_T N> constexpr SIZE_T GetNum(const T (&&Container)[N]) { return N; }
-
 /**
  * Gets the number of items in an initializer list.
  *
diff --git a/Engine/Source/Runtime/Core/Public/Traits/IsContiguousContainer.h b/Engine/Source/Runtime/Core/Public/Traits/IsContiguousContainer.h
--- a/Engine/Source/Runtime/Core/Public/Traits/IsContiguousContainer.h
+++ b/Engine/Source/Runtime/Core/Public/Traits/IsContiguousContainer.h
@@ -15,7 +15,9 @@ template <typename T>
 struct TIsContiguousContainer
 {
 	UE_STATIC_ASSERT_COMPLETE_TYPE(T, "TIsContiguousContainer instantiated with an incomplete type");
-	enum { Value = false };
+	// This sizeof(T) == 0 test will let Clang treat zero-sized arrays as contiguous containers.
+	// An alternate fix using a specialization of TIsContiguousContainer (which Clang doesn't support) can be found below.
+	enum { Value = sizeof(T) == 0 };
 };
 
 template <typename T> struct TIsContiguousContainer<             T& > : TIsContiguousContainer<T> {};
@@ -32,6 +34,16 @@ template <typename T, size_t N> struct TIsContiguousContainer<const          T[N
 template <typename T, size_t N> struct TIsContiguousContainer<      volatile T[N]> { enum { Value = true }; };
 template <typename T, size_t N> struct TIsContiguousContainer<const volatile T[N]> { enum { Value = true }; };
 
+/**
+ * Specialization for zero-sized arrays - which are not handled by the above Clang-specific fix in the primary TIsContiguousContainer template.
+ */
+#ifndef __clang__
+	template <typename T> struct TIsContiguousContainer<               T[0]> { enum { Value = true }; };
+	template <typename T> struct TIsContiguousContainer<const          T[0]> { enum { Value = true }; };
+	template <typename T> struct TIsContiguousContainer<      volatile T[0]> { enum { Value = true }; };
+	template <typename T> struct TIsContiguousContainer<const volatile T[0]> { enum { Value = true }; };
+#endif
+
 /**
  * Specialization for unbounded C arrays (never contiguous - should be treated as pointers which are not regarded as contiguous containers)
  */
