Improved readability changes ensure and ensureMsgf by putting each expression and matching braces/brackets on their own lines.

diff --git a/Engine/Source/Runtime/Core/Public/Misc/AssertionMacros.h b/Engine/Source/Runtime/Core/Public/Misc/AssertionMacros.h
--- a/Engine/Source/Runtime/Core/Public/Misc/AssertionMacros.h
+++ b/Engine/Source/Runtime/Core/Public/Misc/AssertionMacros.h
@@ -423,21 +423,34 @@ RetType UE_COLD UE_DEBUG_SECTION DispatchCheckVerify (InnerType&& Inner, ArgType
 #endif
 
 	#define UE_ENSURE_IMPL(Always, InExpression) \
-		(LIKELY(!!(InExpression)) \
-			|| (::UE::Assert::Private::ExecCheckImplInternal(::bGEnsureHasExecuted<FileLineHashForEnsure(__FILE__, __LINE__)>, Always, __FILE__, __LINE__, #InExpression) \
-			&& UE_BREAK_AND_RETURN_FALSE()))
+		( \
+			LIKELY(!!(InExpression)) || \
+			( \
+				::UE::Assert::Private::ExecCheckImplInternal( \
+					::bGEnsureHasExecuted<FileLineHashForEnsure(__FILE__, __LINE__)>, \
+					Always, \
+					__FILE__, \
+					__LINE__, \
+					#InExpression \
+				) && \
+				UE_BREAK_AND_RETURN_FALSE() \
+			) \
+		)
 
 	#define UE_ENSURE_IMPL2(Always, InExpression, InFormat, ...) \
-		(LIKELY(!!(InExpression)) || ([&] () UE_COLD UE_DEBUG_SECTION \
-		{ \
-			UE_VALIDATE_FORMAT_STRING(InFormat, ##__VA_ARGS__); \
-			static constexpr ::UE::Assert::Private::FStaticEnsureRecord ENSURE_Static(InFormat, #InExpression, __builtin_FILE(), __builtin_LINE(), Always); \
-			if ((Always || !::bGEnsureHasExecuted<FileLineHashForEnsure(__FILE__, __LINE__)>.load(std::memory_order_relaxed)) && FPlatformMisc::IsEnsureAllowed() && ::UE::Assert::Private::EnsureFailed(::bGEnsureHasExecuted<FileLineHashForEnsure(__FILE__, __LINE__)>, &ENSURE_Static, ##__VA_ARGS__)) \
+		( \
+			LIKELY(!!(InExpression)) || \
+			[&]() UE_COLD UE_DEBUG_SECTION \
 			{ \
-				PLATFORM_BREAK(); \
-			} \
-			return false; \
-		} ()))
+				UE_VALIDATE_FORMAT_STRING(InFormat, ##__VA_ARGS__); \
+				static constexpr ::UE::Assert::Private::FStaticEnsureRecord ENSURE_Static(InFormat, #InExpression, __builtin_FILE(), __builtin_LINE(), Always); \
+				if ((Always || !::bGEnsureHasExecuted<FileLineHashForEnsure(__FILE__, __LINE__)>.load(std::memory_order_relaxed)) && FPlatformMisc::IsEnsureAllowed() && ::UE::Assert::Private::EnsureFailed(::bGEnsureHasExecuted<FileLineHashForEnsure(__FILE__, __LINE__)>, &ENSURE_Static, ##__VA_ARGS__)) \
+				{ \
+					PLATFORM_BREAK(); \
+				} \
+				return false; \
+			}() \
+		)
 
 	#define ensure(           InExpression                ) UE_ENSURE_IMPL (false, InExpression)
 	#define ensureMsgf(       InExpression, InFormat, ... ) UE_ENSURE_IMPL2(false, InExpression, InFormat, ##__VA_ARGS__)
