Improve doc comments around `TFunctionRef`.

It wasn't entirely clear whether `MoveTemp` was meaningful when applied to a `TFunctionRef` without reading the code. The doc comments came close but didn't quite spell it out.

diff --git a/Engine/Source/Runtime/Core/Public/Templates/Function.h b/Engine/Source/Runtime/Core/Public/Templates/Function.h
--- a/Engine/Source/Runtime/Core/Public/Templates/Function.h
+++ b/Engine/Source/Runtime/Core/Public/Templates/Function.h
@@ -577,7 +577,8 @@ namespace UE::Core::Private::Function
  *
  * If you also want to take ownership of the callable thing, e.g. you want to return a lambda from a
  * function, you should use TFunction.  TFunctionRef does not concern itself with ownership because it's
- * intended to be FAST.
+ * intended to be faster than TFunction.  A TFunctionRef is lightweight and meant to be passed and copied
+ * by value; there is no benefit to using MoveTemp.
  *
  * TFunctionRef is most useful when you want to parameterize a function with some caller-defined code
  * without making it a template.
