ParkingLot: Removed 2048b alignment of FThread

Changing the alignment of FThread no longer impacts lock/unlock performance.

#rnx

diff --git a/Engine/Source/Runtime/Core/Private/Async/ParkingLot.cpp b/Engine/Source/Runtime/Core/Private/Async/ParkingLot.cpp
--- a/Engine/Source/Runtime/Core/Private/Async/ParkingLot.cpp
+++ b/Engine/Source/Runtime/Core/Private/Async/ParkingLot.cpp
@@ -80,11 +80,7 @@ struct alignas(PLATFORM_CACHE_LINE_SIZE) FThread final
 
 	inline static FThread* New()
 	{
-		// Performance is very sensitive to how this is aligned.
-		// This value likely needs to be tuned for different hardware/platforms.
-		constexpr SIZE_T ThreadAlignment = 2048;
-
-		return new(Allocator.Malloc(sizeof(FThread), ThreadAlignment)) FThread;
+		return new(Allocator.Malloc(sizeof(FThread), alignof(FThread))) FThread;
 	}
 
 private:
