Partially revert 42273542. The change fixed one set of asserts while introducing another.
- Keep only the change from "*LastStatsFrame < *FrameState.StatsFrame" to "*LastStatsFrame != *FrameState.StatsFrame" in GPUProfiler.cpp to handle the "bad frame" markers the stats system uses, where the frame index is made negative.

diff --git a/Engine/Source/Runtime/Core/Private/Stats/StatsSystem.cpp b/Engine/Source/Runtime/Core/Private/Stats/StatsSystem.cpp
--- a/Engine/Source/Runtime/Core/Private/Stats/StatsSystem.cpp
+++ b/Engine/Source/Runtime/Core/Private/Stats/StatsSystem.cpp
@@ -22,13 +22,6 @@ DEFINE_STAT(STAT_SecondsPerCycle);
 
 namespace UE::Stats
 {
-#if STATS
-	// The engine stats system maintains its own frame counter. We need to propagate this to
-	// the end-of-pipe thread so that stat data it submits is attributed to the correct frame.
-	// This is done inside AdvanceRenderingThreadStats.
-	CORE_API int64 FStats::StatsFrameRT = 0;
-#endif
-
 	TAtomic<int32> FStats::GameThreadStatsFrame(1);
 
 	namespace Private
@@ -90,7 +83,7 @@ namespace UE::Stats
 		{
 			// There is no rendering thread, so this message is sufficient to make stats happy and don't leak memory.
 			FThreadStats::AddMessage(FStatConstants::AdvanceFrame.GetEncodedName(), EStatOperation::AdvanceFrameEventRenderThread, Frame);
-			FStats::StatsFrameRT = Frame;
+			FThreadStats::AddMessage(FStatConstants::AdvanceFrame.GetEncodedName(), EStatOperation::AdvanceFrameEventEndOfPipe, Frame);
 		}
 
 		FThreadStats::ExplicitFlush(bDiscardCallstack);
diff --git a/Engine/Source/Runtime/Core/Public/Stats/StatsSystem.h b/Engine/Source/Runtime/Core/Public/Stats/StatsSystem.h
--- a/Engine/Source/Runtime/Core/Public/Stats/StatsSystem.h
+++ b/Engine/Source/Runtime/Core/Public/Stats/StatsSystem.h
@@ -42,13 +42,6 @@ namespace UE::Stats
 
 		/** Current game thread stats frame. */
 		static CORE_API TAtomic<int32> GameThreadStatsFrame;
-
-	#if STATS
-		// The engine stats system maintains its own frame counter. We need to propagate this to
-		// the end-of-pipe thread so that stat data it submits is attributed to the correct frame.
-		// This is done inside AdvanceRenderingThreadStats.
-		static CORE_API int64 StatsFrameRT;
-	#endif
 	};
 } // namespace UE::Stats
 
diff --git a/Engine/Source/Runtime/RHI/Private/RHICommandList.cpp b/Engine/Source/Runtime/RHI/Private/RHICommandList.cpp
--- a/Engine/Source/Runtime/RHI/Private/RHICommandList.cpp
+++ b/Engine/Source/Runtime/RHI/Private/RHICommandList.cpp
@@ -14,7 +14,6 @@
 #include "RHIContext.h"
 #include "RHIFwd.h"
 #include "RHITransition.h"
-#include "Stats/StatsSystem.h"
 #include "Stats/StatsTrace.h"
 #include "Stats/ThreadIdleStats.h"
 #include "HAL/PlatformMisc.h"
@@ -1430,7 +1429,7 @@ RHI_API FGraphEventRef FRHICommandListExecutor::Submit(TConstArrayView<FRHIComma
 				  this
 				, State = SubmitState
 			#if STATS
-				, LocalStatsFrame = UE::Stats::FStats::StatsFrameRT
+				, LocalStatsFrame = StatsFrame
 			#endif
 			]() mutable
 			{
diff --git a/Engine/Source/Runtime/RHI/Public/RHICommandList.h b/Engine/Source/Runtime/RHI/Public/RHICommandList.h
--- a/Engine/Source/Runtime/RHI/Public/RHICommandList.h
+++ b/Engine/Source/Runtime/RHI/Public/RHICommandList.h
@@ -5477,6 +5477,19 @@ private:
 	// Counts the number of calls to RHIEndFrame, and is used in GPU profiler frame boundary events.
 	uint32 FrameNumber = 0;
 
+#if STATS
+	// The engine stats system maintains its own frame counter. We need to propagate this to
+	// the end-of-pipe thread so that stat data it submits is attributed to the correct frame.
+	// This is done inside AdvanceRenderingThreadStats.
+	int64 StatsFrame = 0;
+public:
+	void SetStatsFrame(int64 InStatsFrame)
+	{
+		StatsFrame = InStatsFrame;
+	}
+private:
+#endif
+
 	bool AllowParallel() const;
 };
 
diff --git a/Engine/Source/Runtime/RenderCore/Private/RenderingThread.cpp b/Engine/Source/Runtime/RenderCore/Private/RenderingThread.cpp
--- a/Engine/Source/Runtime/RenderCore/Private/RenderingThread.cpp
+++ b/Engine/Source/Runtime/RenderCore/Private/RenderingThread.cpp
@@ -21,7 +21,6 @@
 #include "RenderCommandFence.h"
 #include "RenderDeferredCleanup.h"
 #include "TickableObjectRenderThread.h"
-#include "Stats/StatsSystem.h"
 #include "Stats/StatsData.h"
 #include "HAL/ThreadHeartBeat.h"
 #include "RenderResource.h"
@@ -348,7 +347,7 @@ static void AdvanceRenderingThreadStats(int64 StatsFrame, int32 DisableChangeTag
 
 #if RHI_NEW_GPU_PROFILER
 	// Propagate the stats frame value down to the end-of-pipe thread.
-	UE::Stats::FStats::StatsFrameRT = Frame;
+	GRHICommandList.SetStatsFrame(Frame);
 #else
 	FThreadStats::AddMessage(FStatConstants::AdvanceFrame.GetEncodedName(), EStatOperation::AdvanceFrameEventEndOfPipe, Frame);
 #endif
