Fix for Windows callstacks when compiled with Clang
#tests verified memory callstacks captured successfully with Clang

diff --git a/Engine/Source/Runtime/Core/Public/Windows/WindowsPlatform.h b/Engine/Source/Runtime/Core/Public/Windows/WindowsPlatform.h
--- a/Engine/Source/Runtime/Core/Public/Windows/WindowsPlatform.h
+++ b/Engine/Source/Runtime/Core/Public/Windows/WindowsPlatform.h
@@ -121,8 +121,13 @@ typedef FWindowsPlatformTypes FPlatformTypes;
 
 #define PLATFORM_SUPPORTS_BORDERLESS_WINDOW					1
 
+#if defined(__clang__)
+#define PLATFORM_RETURN_ADDRESS_FOR_CALLSTACKTRACING		PLATFORM_RETURN_ADDRESS
+#define PLATFORM_USE_CALLSTACK_ADDRESS_POINTER				0
+#else
 #define PLATFORM_RETURN_ADDRESS_FOR_CALLSTACKTRACING		PLATFORM_RETURN_ADDRESS_POINTER
 #define PLATFORM_USE_CALLSTACK_ADDRESS_POINTER				1
+#endif
 
 #define WINDOWS_USE_FEATURE_APPLICATIONMISC_CLASS			1
 #define WINDOWS_USE_FEATURE_PLATFORMMISC_CLASS				1
