Core: Moved AtomicQueue module dependency from private dependency list to public dependency list, so that modules that modules that tried to use FTaskConcurrencyLimiter would not have to manually add AtomicQueue as a dependency to allow compiling.

diff --git a/Engine/Source/Runtime/Core/Core.Build.cs b/Engine/Source/Runtime/Core/Core.Build.cs
--- a/Engine/Source/Runtime/Core/Core.Build.cs
+++ b/Engine/Source/Runtime/Core/Core.Build.cs
@@ -38,6 +38,7 @@ public class Core : ModuleRules
 
 		PublicDependencyModuleNames.Add("GuidelinesSupportLibrary");
 		PublicDependencyModuleNames.Add("TraceLog");
+		PublicDependencyModuleNames.Add("AtomicQueue");
 
 		PrivateIncludePathModuleNames.AddRange(
 			new string[] {
@@ -296,7 +297,6 @@ public class Core : ModuleRules
 
 		PrivateDefinitions.Add("UE_DEFINE_LEGACY_MATH_CONSTANT_MACRO_NAMES=0");
 
-		PrivateDependencyModuleNames.Add("AtomicQueue");
 		PrivateDependencyModuleNames.Add("AutoRTFM");
 
 		CppCompileWarningSettings.UnsafeTypeCastWarningLevel = WarningLevel.Error;
