Fixing an issue that can occur after a FConfigBranch is safe reloaded. FindConfigFileWithBaseName can fail to match the BaseName to the branch due to it comparing Branch::InMemoryFile::Name instead of Branch::IniName

diff --git a/Engine/Source/Runtime/Core/Private/Misc/ConfigCacheIni.cpp b/Engine/Source/Runtime/Core/Private/Misc/ConfigCacheIni.cpp
--- a/Engine/Source/Runtime/Core/Private/Misc/ConfigCacheIni.cpp
+++ b/Engine/Source/Runtime/Core/Private/Misc/ConfigCacheIni.cpp
@@ -3895,7 +3895,7 @@ FConfigFile* FConfigCacheIni::FindConfigFileWithBaseName(FName BaseName)
 	{
 		for (TPair<FString,FConfigBranch*>& CurrentFilePair : OtherFiles)
 		{
-			if (CurrentFilePair.Value->InMemoryFile.Name == BaseName)
+			if (CurrentFilePair.Value->IniName == BaseName)
 			{
 				Result = CurrentFilePair.Value;
 				break;
