Adding analytics fields to help identify and monitor laptop performance.
Implementing GetBatteryLevel() for the Windows platform. This will continue to return -1 (not implemented) for systems without a battery.
Adding FWindowsPlatformMisc::IsUsingBatterySaver() to provide a way to check if the device has battery saver enabled.

diff --git a/Engine/Source/Runtime/Core/Private/Windows/WindowsPlatformCrashContext.cpp b/Engine/Source/Runtime/Core/Private/Windows/WindowsPlatformCrashContext.cpp
--- a/Engine/Source/Runtime/Core/Private/Windows/WindowsPlatformCrashContext.cpp
+++ b/Engine/Source/Runtime/Core/Private/Windows/WindowsPlatformCrashContext.cpp
@@ -311,6 +311,7 @@ void FWindowsPlatformCrashContext::AddPlatformSpecificProperties() const
 	AddCrashProperty(TEXT("PlatformIsRunningWine"), FWindowsPlatformMisc::IsWine());
 	AddCrashProperty(TEXT("PlatformSupportsWindows11"), FWindowsPlatformMisc::SupportsWindows11());
 	AddCrashProperty(TEXT("IsRunningOnBattery"), FPlatformMisc::IsRunningOnBattery());
+	AddCrashProperty(TEXT("IsUsingBatterySaver"), FWindowsPlatformMisc::IsUsingBatterySaver());
 	WIDECHAR DriveName = 0;
 	const TCHAR* BaseDir = FWindowsPlatformProcess::BaseDir();
 	if (BaseDir && *BaseDir)
diff --git a/Engine/Source/Runtime/Core/Private/Windows/WindowsPlatformMisc.cpp b/Engine/Source/Runtime/Core/Private/Windows/WindowsPlatformMisc.cpp
--- a/Engine/Source/Runtime/Core/Private/Windows/WindowsPlatformMisc.cpp
+++ b/Engine/Source/Runtime/Core/Private/Windows/WindowsPlatformMisc.cpp
@@ -4201,14 +4201,8 @@ FText FWindowsPlatformMisc::GetFileManagerName()
 	return NSLOCTEXT("WindowsPlatform", "FileManagerName", "Explorer");
 }
 
-bool FWindowsPlatformMisc::IsRunningOnBattery()
+static bool IsRunningOnBatteryInternal(SYSTEM_POWER_STATUS Status)
 {
-	SYSTEM_POWER_STATUS status;
-	if (GetSystemPowerStatus(&status) == 0)
-	{
-		return false;
-	}
-
 	BYTE BatteryPresentFlags = 4  // Critical-the battery capacity is at less than five percent.
 		| 2  // Low-the battery capacity is at less than 33 percent.
 		| 1  // High-the battery capacity is at more than 66 percent.
@@ -4216,10 +4210,46 @@ bool FWindowsPlatformMisc::IsRunningOnBattery()
 
 	BYTE BatteryNotPresentFlags = 128; // No system battery.
 
-	bool BatteryPresent = ((status.BatteryFlag & BatteryPresentFlags) != 0) && ((status.BatteryFlag & BatteryNotPresentFlags) == 0);
+	bool BatteryPresent = ((Status.BatteryFlag & BatteryPresentFlags) != 0) && ((Status.BatteryFlag & BatteryNotPresentFlags) == 0);
 
 	// We also need to check for 0, which means "the battery is not being charged and the battery capacity is between low and high".
-	return BatteryPresent || status.BatteryFlag == 0;
+	return BatteryPresent || Status.BatteryFlag == 0;
+}
+
+bool FWindowsPlatformMisc::IsRunningOnBattery()
+{
+	SYSTEM_POWER_STATUS Status;
+	if (GetSystemPowerStatus(&Status) == 0)
+	{
+		return false;
+	}
+	return IsRunningOnBatteryInternal(Status);
+}
+
+int FWindowsPlatformMisc::GetBatteryLevel()
+{
+	SYSTEM_POWER_STATUS Status;
+	if (GetSystemPowerStatus(&Status) == 0 || !IsRunningOnBatteryInternal(Status))
+	{
+		return -1; 
+	}
+	if (Status.BatteryLifePercent == 255)
+	{
+		return -1; // Unknown battery status
+	}
+	return Status.BatteryLifePercent;
+}
+
+bool FWindowsPlatformMisc::IsUsingBatterySaver()
+{
+	SYSTEM_POWER_STATUS Status;
+	if (GetSystemPowerStatus(&Status) == 0 || !IsRunningOnBatteryInternal(Status))
+	{
+		return false;
+	}
+	// A SystemStatusFlag value of 1 indicates battery saver is enabled
+	// https://learn.microsoft.com/en-us/windows/win32/api/winbase/ns-winbase-system_power_status
+	return Status.SystemStatusFlag == 1;
 }
 
 FString FWindowsPlatformMisc::GetOperatingSystemId()
diff --git a/Engine/Source/Runtime/Core/Public/Windows/WindowsPlatformMisc.h b/Engine/Source/Runtime/Core/Public/Windows/WindowsPlatformMisc.h
--- a/Engine/Source/Runtime/Core/Public/Windows/WindowsPlatformMisc.h
+++ b/Engine/Source/Runtime/Core/Public/Windows/WindowsPlatformMisc.h
@@ -329,10 +329,20 @@ struct FWindowsPlatformMisc
 	}
 
 	/**
-	 * Returns whether the platform is running on battery power or not.
+	 * Returns whether the platform is running on battery power or not. For Windows laptops this is true regardless of charging state.
 	 */
 	static CORE_API bool IsRunningOnBattery();
 
+	/**
+	 * Returns the current battery level between 0 and 100, or -1 if a battery is not present or if the current battery level cannot be retrieved.
+	 */
+	static CORE_API int GetBatteryLevel();
+
+	/**
+	 * Returns whether the platform is currently running with battery saver enabled.
+	 */
+	static CORE_API bool IsUsingBatterySaver();
+
 	FORCEINLINE static void ChooseHDRDeviceAndColorGamut(uint32 DeviceId, uint32 DisplayNitLevel, EDisplayOutputFormat& OutputDevice, EDisplayColorGamut& ColorGamut)
 	{
 		// needs to match GRHIHDRDisplayOutputFormat chosen in FD3D12DynamicRHI::Init
