Allocators: Removed dependency on FCriticalSection and ParkingLot

ParkingLot depends on AnsiMalloc which can be redirected into MB2 or MB3. Remove the dependency of these allocators on mutexes that use ParkingLot, including FCriticalSection which can be an alias for FRecursiveMutex.

diff --git a/Engine/Source/Runtime/Core/Private/HAL/Allocators/CachedOSPageAllocator.cpp b/Engine/Source/Runtime/Core/Private/HAL/Allocators/CachedOSPageAllocator.cpp
--- a/Engine/Source/Runtime/Core/Private/HAL/Allocators/CachedOSPageAllocator.cpp
+++ b/Engine/Source/Runtime/Core/Private/HAL/Allocators/CachedOSPageAllocator.cpp
@@ -1,13 +1,14 @@
 // Copyright Epic Games, Inc. All Rights Reserved.
 
 #include "HAL/Allocators/CachedOSPageAllocator.h"
-#include "HAL/UnrealMemory.h"
-#include "Logging/LogMacros.h"
 #include "CoreGlobals.h"
 #include "HAL/LowLevelMemTracker.h"
+#include "HAL/UnrealMemory.h"
+#include "Logging/LogMacros.h"
+#include "Misc/ScopeLock.h"
 
 
-void* FCachedOSPageAllocator::AllocateImpl(SIZE_T Size, uint32 CachedByteLimit, FFreePageBlock* First, FFreePageBlock* Last, uint32& FreedPageBlocksNum, SIZE_T& CachedTotal, FCriticalSection* Mutex)
+void* FCachedOSPageAllocator::AllocateImpl(SIZE_T Size, uint32 CachedByteLimit, FFreePageBlock* First, FFreePageBlock* Last, uint32& FreedPageBlocksNum, SIZE_T& CachedTotal, UE::FPlatformRecursiveMutex* Mutex)
 {
 	if (!IsOSAllocation(Size, CachedByteLimit))
 	{
@@ -57,7 +58,7 @@ void* FCachedOSPageAllocator::AllocateImpl(SIZE_T Size, uint32 CachedByteLimit,
 
 			{
 #if UE_ALLOW_OSMEMORYLOCKFREE
-				FScopeUnlock ScopeUnlock(Mutex);
+				UE::TScopeUnlock ScopeUnlock(Mutex);
 #endif
 				LLM_PLATFORM_SCOPE(ELLMTag::FMalloc);
 				void* Ptr = FPlatformMemory::BinnedAllocFromOS(Size);
@@ -82,7 +83,7 @@ void* FCachedOSPageAllocator::AllocateImpl(SIZE_T Size, uint32 CachedByteLimit,
 	void* Ret = nullptr;
 	{
 #if UE_ALLOW_OSMEMORYLOCKFREE
-		FScopeUnlock ScopeUnlock(Mutex);
+		UE::TScopeUnlock ScopeUnlock(Mutex);
 #endif
 		LLM_PLATFORM_SCOPE(ELLMTag::FMalloc);
 		Ret = FPlatformMemory::BinnedAllocFromOS(Size);
@@ -90,12 +91,12 @@ void* FCachedOSPageAllocator::AllocateImpl(SIZE_T Size, uint32 CachedByteLimit,
 	return Ret;
 }
 
-void FCachedOSPageAllocator::FreeImpl(void* Ptr, SIZE_T Size, uint32 NumCacheBlocks, uint32 CachedByteLimit, FFreePageBlock* First, uint32& FreedPageBlocksNum, SIZE_T& CachedTotal, FCriticalSection* Mutex, bool ThreadIsTimeCritical)
+void FCachedOSPageAllocator::FreeImpl(void* Ptr, SIZE_T Size, uint32 NumCacheBlocks, uint32 CachedByteLimit, FFreePageBlock* First, uint32& FreedPageBlocksNum, SIZE_T& CachedTotal, UE::FPlatformRecursiveMutex* Mutex, bool ThreadIsTimeCritical)
 {
 	if (IsOSAllocation(Size, CachedByteLimit))
 	{
 #if UE_ALLOW_OSMEMORYLOCKFREE
-		FScopeUnlock ScopeUnlock(Mutex);
+		UE::TScopeUnlock ScopeUnlock(Mutex);
 #endif
 		FPlatformMemory::BinnedFreeToOS(Ptr, Size);
 		return;
@@ -114,7 +115,7 @@ void FCachedOSPageAllocator::FreeImpl(void* Ptr, SIZE_T Size, uint32 NumCacheBlo
 		}
 		{
 #if UE_ALLOW_OSMEMORYLOCKFREE
-			FScopeUnlock ScopeUnlock(Mutex);
+			UE::TScopeUnlock ScopeUnlock(Mutex);
 #endif
 			FPlatformMemory::BinnedFreeToOS(FreePtr, FreeSize);
 		}
@@ -127,7 +128,7 @@ void FCachedOSPageAllocator::FreeImpl(void* Ptr, SIZE_T Size, uint32 NumCacheBlo
 	++FreedPageBlocksNum;
 }
 
-void FCachedOSPageAllocator::FreeAllImpl(FFreePageBlock* First, uint32& FreedPageBlocksNum, SIZE_T& CachedTotal, FCriticalSection* Mutex)
+void FCachedOSPageAllocator::FreeAllImpl(FFreePageBlock* First, uint32& FreedPageBlocksNum, SIZE_T& CachedTotal, UE::FPlatformRecursiveMutex* Mutex)
 {
 	while (FreedPageBlocksNum)
 	{
diff --git a/Engine/Source/Runtime/Core/Private/HAL/Allocators/CachedOSVeryLargePageAllocator.cpp b/Engine/Source/Runtime/Core/Private/HAL/Allocators/CachedOSVeryLargePageAllocator.cpp
--- a/Engine/Source/Runtime/Core/Private/HAL/Allocators/CachedOSVeryLargePageAllocator.cpp
+++ b/Engine/Source/Runtime/Core/Private/HAL/Allocators/CachedOSVeryLargePageAllocator.cpp
@@ -1,6 +1,7 @@
 // Copyright Epic Games, Inc. All Rights Reserved.
 
 #include "HAL/Allocators/CachedOSVeryLargePageAllocator.h"
+#include "Async/UniqueLock.h"
 #include "HAL/IConsoleManager.h"
 #include "HAL/UnrealMemory.h"
 #include "Logging/LogMacros.h"
@@ -169,7 +170,7 @@ void FCachedOSVeryLargePageAllocator::Refresh()
 	}
 }
 
-void* FCachedOSVeryLargePageAllocator::Allocate(SIZE_T Size, uint32 AllocationHint, FCriticalSection* Mutex)
+void* FCachedOSVeryLargePageAllocator::Allocate(SIZE_T Size, uint32 AllocationHint, UE::FPlatformRecursiveMutex* Mutex)
 {
 	Size = Align(Size, 4096);
 
@@ -210,7 +211,7 @@ void* FCachedOSVeryLargePageAllocator::Allocate(SIZE_T Size, uint32 AllocationHi
 }
 
 
-FCachedOSVeryLargePageAllocator::FLargePage* FCachedOSVeryLargePageAllocator::GetOrAllocLargePage(uint32 AllocationHint, FCriticalSection* Mutex)
+FCachedOSVeryLargePageAllocator::FLargePage* FCachedOSVeryLargePageAllocator::GetOrAllocLargePage(uint32 AllocationHint, UE::FPlatformRecursiveMutex* Mutex)
 {
 	// Use an existing page with space if one is available
 	FLargePage* LargePage = UsedLargePagesWithSpaceHead[AllocationHint];
@@ -243,7 +244,7 @@ FCachedOSVeryLargePageAllocator::FLargePage* FCachedOSVeryLargePageAllocator::Ge
 	return LargePage;
 }
 
-FCachedOSVeryLargePageAllocator::FLargePage* FCachedOSVeryLargePageAllocator::AllocNewLargePage(uint32 AllocationHint, FCriticalSection* Mutex, bool& bOutCommitFailure)
+FCachedOSVeryLargePageAllocator::FLargePage* FCachedOSVeryLargePageAllocator::AllocNewLargePage(uint32 AllocationHint, UE::FPlatformRecursiveMutex* Mutex, bool& bOutCommitFailure)
 {
 	// If there's a limit for this allocation type then check first
 	bOutCommitFailure = false;
@@ -259,7 +260,7 @@ FCachedOSVeryLargePageAllocator::FLargePage* FCachedOSVeryLargePageAllocator::Al
 		LargePage->Unlink();
 		{
 #if UE_ALLOW_OSMEMORYLOCKFREE
-			FScopeUnlock ScopeUnlock(Mutex);
+			UE::TScopeUnlock ScopeUnlock(Mutex);
 #endif
 			LLM_PLATFORM_SCOPE(ELLMTag::FMalloc);
 			if (!Block.Commit(LargePage->BaseAddress - AddressSpaceReserved, SizeOfLargePage, false))
@@ -267,7 +268,7 @@ FCachedOSVeryLargePageAllocator::FLargePage* FCachedOSVeryLargePageAllocator::Al
 #if UE_ALLOW_OSMEMORYLOCKFREE
 				if (Mutex != nullptr)
 				{
-					FScopeLock Lock(Mutex);
+					UE::TUniqueLock Lock(*Mutex);
 					LargePage->LinkHead(FreeLargePagesHead[AllocationHint]);
 				}
 				else
@@ -290,7 +291,7 @@ FCachedOSVeryLargePageAllocator::FLargePage* FCachedOSVeryLargePageAllocator::Al
 	return LargePage;
 }
 
-void FCachedOSVeryLargePageAllocator::Free(void* Ptr, SIZE_T Size, FCriticalSection* Mutex, bool ThreadIsTimeCritical)
+void FCachedOSVeryLargePageAllocator::Free(void* Ptr, SIZE_T Size, UE::FPlatformRecursiveMutex* Mutex, bool ThreadIsTimeCritical)
 {
 	Size = Align(Size, 4096);
 	uint64 Index = ((uintptr_t)Ptr - (uintptr_t)AddressSpaceReserved) / SizeOfLargePage;
@@ -318,7 +319,7 @@ void FCachedOSVeryLargePageAllocator::Free(void* Ptr, SIZE_T Size, FCriticalSect
 				// need to move which list we are in and remove the backing store
 				{
 #if UE_ALLOW_OSMEMORYLOCKFREE
-					FScopeUnlock ScopeUnlock(Mutex);
+					UE::TScopeUnlock ScopeUnlock(Mutex);
 #endif
 					Block.Decommit(LargePage->BaseAddress - AddressSpaceReserved, SizeOfLargePage);
 					LLM_IF_ENABLED(FLowLevelMemTracker::Get().OnLowLevelFree(ELLMTracker::Platform, (void*)LargePage->BaseAddress));
@@ -388,7 +389,7 @@ void FCachedOSVeryLargePageAllocator::ShrinkEmptyBackStore(int32 NewEmptyBackSto
 }
 
 
-void FCachedOSVeryLargePageAllocator::FreeAll(FCriticalSection* Mutex)
+void FCachedOSVeryLargePageAllocator::FreeAll(UE::FPlatformRecursiveMutex* Mutex)
 {
 	for (int i = 0; i < FMemory::AllocationHints::Max; i++)
 	{
diff --git a/Engine/Source/Runtime/Core/Private/HAL/Allocators/PooledVirtualMemoryAllocator.cpp b/Engine/Source/Runtime/Core/Private/HAL/Allocators/PooledVirtualMemoryAllocator.cpp
--- a/Engine/Source/Runtime/Core/Private/HAL/Allocators/PooledVirtualMemoryAllocator.cpp
+++ b/Engine/Source/Runtime/Core/Private/HAL/Allocators/PooledVirtualMemoryAllocator.cpp
@@ -1,6 +1,7 @@
 // Copyright Epic Games, Inc. All Rights Reserved.
 
 #include "HAL/Allocators/PooledVirtualMemoryAllocator.h"
+#include "Async/UniqueLock.h"
 #include "HAL/UnrealMemory.h"
 #include "Logging/LogMacros.h"
 #include "CoreGlobals.h"
@@ -47,7 +48,7 @@ FPooledVirtualMemoryAllocator::FPooledVirtualMemoryAllocator()
 	}
 }
 
-void* FPooledVirtualMemoryAllocator::Allocate(SIZE_T Size, uint32 AllocationHint, FCriticalSection* Mutex)
+void* FPooledVirtualMemoryAllocator::Allocate(SIZE_T Size, uint32 AllocationHint, UE::FPlatformRecursiveMutex* Mutex)
 {
 	if (Size > Limits::MaxAllocationSizeToPool)
 	{
@@ -55,14 +56,14 @@ void* FPooledVirtualMemoryAllocator::Allocate(SIZE_T Size, uint32 AllocationHint
 		{
 			// do not report to LLM here, the platform functions will do that
 #if UE_ALLOW_OSMEMORYLOCKFREE
-			FScopeUnlock FScopeUnlock(Mutex);
+			UE::TScopeUnlock ScopeUnlock(Mutex);
 #endif
 			return FPlatformMemory::BinnedAllocFromOS(Size);
 		}
 		else
 		{
 			// do not report to LLM here, the platform functions will do that
-			FScopeLock Lock(&OsAllocatorCacheLock);
+			UE::TUniqueLock Lock(OsAllocatorCacheLock);
 			void* Ptr = OsAllocatorCache.Allocate(Size);
 			return Ptr;
 		}
@@ -74,7 +75,7 @@ void* FPooledVirtualMemoryAllocator::Allocate(SIZE_T Size, uint32 AllocationHint
 		LLM_PLATFORM_SCOPE(ELLMTag::FMalloc);
 
 		// [RCL] TODO: find a way to convert to lock-free
-		FScopeLock Lock(&ClassesLocks[SizeClass]);
+		UE::TUniqueLock Lock(ClassesLocks[SizeClass]);
 
 		// follow the list until we can allocate
 		for(FPoolDescriptorBase* BaseDesc = ClassesListHeads[SizeClass]; BaseDesc; BaseDesc = BaseDesc->Next)
@@ -117,7 +118,7 @@ void* FPooledVirtualMemoryAllocator::Allocate(SIZE_T Size, uint32 AllocationHint
 	}
 };
 
-void FPooledVirtualMemoryAllocator::Free(void* Ptr, SIZE_T Size, FCriticalSection* Mutex, bool /*ThreadIsTimeCritical = false*/)
+void FPooledVirtualMemoryAllocator::Free(void* Ptr, SIZE_T Size, UE::FPlatformRecursiveMutex* Mutex, bool /*ThreadIsTimeCritical = false*/)
 {
 	if (Size > Limits::MaxAllocationSizeToPool)
 	{
@@ -125,14 +126,14 @@ void FPooledVirtualMemoryAllocator::Free(void* Ptr, SIZE_T Size, FCriticalSectio
 		{
 			// do not report to LLM here, the platform functions will do that
 #if UE_ALLOW_OSMEMORYLOCKFREE
-			FScopeUnlock FScopeUnlock(Mutex);
+			UE::TScopeUnlock ScopeUnlock(Mutex);
 #endif
 			FPlatformMemory::BinnedFreeToOS(Ptr, Size);
 		}
 		else
 		{
 			// do not report to LLM here, the platform functions will do that
-			FScopeLock Lock(&OsAllocatorCacheLock);
+			UE::TUniqueLock Lock(OsAllocatorCacheLock);
 			OsAllocatorCache.Free(Ptr, Size);
 		}
 	}
@@ -141,7 +142,7 @@ void FPooledVirtualMemoryAllocator::Free(void* Ptr, SIZE_T Size, FCriticalSectio
 		int32 SizeClass = GetAllocationSizeClass(Size);
 
 		// [RCL] TODO: find a way to convert to lock-free
-		FScopeLock Lock(&ClassesLocks[SizeClass]);
+		UE::TUniqueLock Lock(ClassesLocks[SizeClass]);
 
 		// follow the list until we can find the pool it came from
 		FPoolDescriptorBase* PrevBaseDesc = nullptr;
@@ -269,9 +270,9 @@ void FPooledVirtualMemoryAllocator::DestroyPool(FPoolDescriptorBase* Pool)
 	VMBlock.FreeVirtual();
 }
 
-void FPooledVirtualMemoryAllocator::FreeAll(FCriticalSection* /*Mutex = nullptr*/)
+void FPooledVirtualMemoryAllocator::FreeAll(UE::FPlatformRecursiveMutex* /*Mutex = nullptr*/)
 {
-	FScopeLock Lock(&OsAllocatorCacheLock);
+	UE::TUniqueLock Lock(OsAllocatorCacheLock);
 	OsAllocatorCache.FreeAll();
 
 	// Currently, there's nothing else to trim.
@@ -285,7 +286,7 @@ uint64 FPooledVirtualMemoryAllocator::GetCachedFreeTotal() const
 
 	for (int32 IdxSizeClass = 0; IdxSizeClass < Limits::NumAllocationSizeClasses; ++IdxSizeClass)
 	{
-		FScopeLock Lock(&ClassesLocks[IdxSizeClass]);
+		UE::TUniqueLock Lock(ClassesLocks[IdxSizeClass]);
 
 		for(FPoolDescriptorBase* BaseDesc = ClassesListHeads[IdxSizeClass]; BaseDesc; BaseDesc = BaseDesc->Next)
 		{
@@ -304,7 +305,7 @@ void FPooledVirtualMemoryAllocator::DumpAllocatorStats(FOutputDevice& Ar)
 {
 	for (int32 IdxSizeClass = 0; IdxSizeClass < Limits::NumAllocationSizeClasses; ++IdxSizeClass)
 	{
-		FScopeLock Lock(&ClassesLocks[IdxSizeClass]);
+		UE::TUniqueLock Lock(ClassesLocks[IdxSizeClass]);
 		SIZE_T AllocationSizeForClass = CalculateAllocationSizeFromClass(IdxSizeClass);
 		Ar.Logf(TEXT("PooledVirtualMemoryAllocator Index: %d, SizeClass: %.2fKB"),
 			IdxSizeClass,
@@ -320,7 +321,7 @@ void FPooledVirtualMemoryAllocator::DumpAllocatorStats(FOutputDevice& Ar)
 		}
 	}
 
-	FScopeLock Lock(&OsAllocatorCacheLock);
+	UE::TUniqueLock Lock(OsAllocatorCacheLock);
 	OsAllocatorCache.DumpAllocatorStats(Ar);
 }
 #endif
diff --git a/Engine/Source/Runtime/Core/Private/HAL/MallocBinned.cpp b/Engine/Source/Runtime/Core/Private/HAL/MallocBinned.cpp
--- a/Engine/Source/Runtime/Core/Private/HAL/MallocBinned.cpp
+++ b/Engine/Source/Runtime/Core/Private/HAL/MallocBinned.cpp
@@ -1,7 +1,7 @@
 // Copyright Epic Games, Inc. All Rights Reserved.
 
 #include "HAL/MallocBinned.h"
-#include "Misc/ScopeLock.h"
+#include "Async/UniqueLock.h"
 #include "Misc/BufferedOutputDevice.h"
 
 #include "HAL/MemoryMisc.h"
@@ -203,9 +203,9 @@ struct FMallocBinned::Private
 		uint64 AddressOffset = ((uint64)ptr) - SmallBlockStartPtr;
 		return ((AddressOffset % PAGE_SIZE_LIMIT) == 0);
 	}
-	
-	static FCriticalSection SmallLock;
-	
+
+	inline static UE::FPlatformRecursiveMutex SmallLock;
+
 	static FFreeMem* GetAllocFromSmallBlockGrab(FMallocBinned& Allocator, UPTRINT OsBytes)
 	{
 		//Make sure we can hold this memory in a free SmallBlockGrab
@@ -215,7 +215,7 @@ struct FMallocBinned::Private
 			return nullptr;
 		}
 		
-		FScopeLock ScopedLock( &SmallLock );
+		UE::TUniqueLock ScopedLock(SmallLock);
 		if (Private::NumFreeSmallBlockGrabAllocations == 0)
 		{
 			//FPlatformMisc::LowLevelOutputDebugStringf(TEXT("OSBytes: %d  FreeAllocations: %d"), OsBytes, Private::NumFreeSmallBlockGrabAllocations);
@@ -237,7 +237,7 @@ struct FMallocBinned::Private
 	
 	static void FreeSmallBlockGrab(void* ptr, FMallocBinned& Allocator, SIZE_T Size)
 	{
-		FScopeLock ScopedLock( &SmallLock );
+		UE::TUniqueLock ScopedLock(SmallLock);
 
 		check(IsSmallBlockGrabAllocation(ptr));
 		check(!SMALL_BLOCK_GRAB_ENSURE_PAGE_ALIGNMENT || IsNanoMallocPageAligned(ptr));
@@ -638,7 +638,7 @@ struct FMallocBinned::Private
 		FPoolInfo* Pool;
 		{
 #ifdef USE_FINE_GRAIN_LOCKS
-			FScopeLock PoolInfoLock(&Allocator.AccessGuard);
+			UE::TUniqueLock PoolInfoLock(Allocator.AccessGuard);
 #endif
 			Pool = GetPoolInfo(Allocator, (UPTRINT)Free);
 			for (UPTRINT i = (UPTRINT)PageSize, Offset = 0; i < OsBytes; i += PageSize, ++Offset)
@@ -719,7 +719,7 @@ struct FMallocBinned::Private
 		{
 			FPoolTable* Table = Allocator.MemSizeToPoolTable[Pool->TableIndex];
 #ifdef USE_FINE_GRAIN_LOCKS
-			FScopeLock TableLock(&Table->CriticalSection);
+			UE::TUniqueLock TableLock(Table->CriticalSection);
 #endif
 #if STATS
 			Table->ActiveRequests--;
@@ -784,7 +784,7 @@ struct FMallocBinned::Private
 		Allocator.PendingFreeList->Push(Ptr);
 #else
 #ifdef USE_COARSE_GRAIN_LOCKS
-		FScopeLock ScopedLock(&Allocator.AccessGuard);
+		UE::TUniqueLock ScopedLock(Allocator.AccessGuard);
 #endif
 		FreeInternal(Allocator, Ptr);
 #endif
@@ -823,7 +823,7 @@ struct FMallocBinned::Private
 	{
 #ifdef CACHE_FREED_OS_ALLOCS
 #ifdef USE_FINE_GRAIN_LOCKS
-		FScopeLock MainLock(&Allocator.AccessGuard);
+		UE::TUniqueLock MainLock(Allocator.AccessGuard);
 #endif
 		if (Size > MAX_CACHED_OS_FREES_BYTE_LIMIT / 4)
 		{
@@ -899,7 +899,7 @@ struct FMallocBinned::Private
 #ifdef USE_FINE_GRAIN_LOCKS
 			// We want to hold the lock a little as possible so release it
 			// before the big call to the OS
-			FScopeLock MainLock(&Allocator.AccessGuard);
+			UE::TUniqueLock MainLock(Allocator.AccessGuard);
 #endif
 			for (uint32 i=0; i < Allocator.FreedPageBlocksNum; ++i)
 			{
@@ -978,7 +978,7 @@ struct FMallocBinned::Private
 	static void FlushAllocCache(FMallocBinned& Allocator)
 	{
 #ifdef USE_FINE_GRAIN_LOCKS
-		FScopeLock MainLock(&Allocator.AccessGuard);
+		UE::TUniqueLock MainLock(Allocator.AccessGuard);
 #endif
 		for (int i = 0, n = Allocator.FreedPageBlocksNum; i<n; ++i) 
 		{
@@ -1013,7 +1013,6 @@ struct FMallocBinned::Private
 #if USE_OS_SMALL_BLOCK_GRAB_MEMORY_FROM_OS
 uint64 FMallocBinned::Private::SmallBlockGrab_FreeStartPointers[MAX_NUM_BLOCK_START_ADDRESSES];
 int FMallocBinned::Private::NumFreeSmallBlockGrabAllocations = 0;
-FCriticalSection FMallocBinned::Private::SmallLock;
 #endif
 
 bool FMallocBinned::Private::bHasInitializedStatsMetadata = false;
@@ -1047,7 +1046,7 @@ void FMallocBinned::GetAllocatorStats( FGenericMemoryStats& out_Stats )
 	
 	{
 #ifdef USE_INTERNAL_LOCKS
-		FScopeLock ScopedLock( &AccessGuard );
+		UE::TUniqueLock ScopedLock(AccessGuard);
 #endif
 
 		Private::UpdateSlackStat(*this);
@@ -1287,10 +1286,10 @@ bool FMallocBinned::IsInternallyThreadSafe() const
 void* FMallocBinned::Malloc(SIZE_T Size, uint32 Alignment)
 {
 #ifdef USE_COARSE_GRAIN_LOCKS
-	FScopeLock ScopedLock(&AccessGuard);
+	UE::TUniqueLock ScopedLock(AccessGuard);
 #endif
-    
-    ensure(Alignment <= PageSize);
+
+	ensure(Alignment <= PageSize);
 
 	Private::FlushPendingFrees(*this);
 
@@ -1298,7 +1297,7 @@ void* FMallocBinned::Malloc(SIZE_T Size, uint32 Alignment)
 	
 	BINNED_INCREMENT_STATCOUNTER(CurrentAllocs);
 	BINNED_INCREMENT_STATCOUNTER(TotalAllocs);
-	
+
 	Size = FMath::Max(Size, (SIZE_T)1);
 	FFreeMem* Free = nullptr;
 	bool bUsePools = true;
@@ -1338,7 +1337,7 @@ void* FMallocBinned::Malloc(SIZE_T Size, uint32 Alignment)
 			// Allocate from pool.
 			FPoolTable* Table = MemSizeToPoolTable[Size];
 #ifdef USE_FINE_GRAIN_LOCKS
-			FScopeLock TableLock(&Table->CriticalSection);
+			UE::TUniqueLock TableLock(Table->CriticalSection);
 #endif
 			Private::TrackStats(Table, (uint32)Size);
 
@@ -1358,7 +1357,7 @@ void* FMallocBinned::Malloc(SIZE_T Size, uint32 Alignment)
 			uint32 PageCount = 3*BinType + 3;
 			FPoolTable* Table = &PagePoolTable[BinType];
 #ifdef USE_FINE_GRAIN_LOCKS
-			FScopeLock TableLock(&Table->CriticalSection);
+			UE::TUniqueLock TableLock(Table->CriticalSection);
 #endif
 			Private::TrackStats(Table, (uint32)Size);
 
@@ -1387,7 +1386,7 @@ void* FMallocBinned::Malloc(SIZE_T Size, uint32 Alignment)
 			FPoolInfo* Pool;
 			{
 #ifdef USE_FINE_GRAIN_LOCKS
-				FScopeLock PoolInfoLock(&AccessGuard);
+				UE::TUniqueLock PoolInfoLock(AccessGuard);
 #endif
 				Pool = Private::GetPoolInfo(*this, (UPTRINT)Free);
 
@@ -1483,7 +1482,7 @@ void* FMallocBinned::Realloc( void* Ptr, SIZE_T NewSize, uint32 Alignment )
 				{
 //need a lock to cover the SetAllocationSizes()
 #ifdef USE_FINE_GRAIN_LOCKS
-					FScopeLock PoolInfoLock(&AccessGuard);
+					UE::TUniqueLock PoolInfoLock(AccessGuard);
 #endif
 					int32 UsedChange = (NewSize - Pool->GetBytes());
 				
@@ -1604,13 +1603,13 @@ SIZE_T FMallocBinned::QuantizeSize(SIZE_T Size, uint32 Alignment)
 bool FMallocBinned::ValidateHeap()
 {
 #ifdef USE_COARSE_GRAIN_LOCKS
-	FScopeLock ScopedLock(&AccessGuard);
+	UE::TUniqueLock ScopedLock(AccessGuard);
 #endif
 	for( int32 i = 0; i < POOL_COUNT; i++ )
 	{
 		FPoolTable* Table = &PoolTable[i];
 #ifdef USE_FINE_GRAIN_LOCKS
-		FScopeLock TableLock(&Table->CriticalSection);
+		UE::TUniqueLock TableLock(Table->CriticalSection);
 #endif
 		for( FPoolInfo** PoolPtr = &Table->FirstPool; *PoolPtr; PoolPtr = &(*PoolPtr)->Next )
 		{
@@ -1656,7 +1655,7 @@ void FMallocBinned::UpdateStats()
 	
 	{
 #ifdef USE_INTERNAL_LOCKS
-		FScopeLock ScopedLock( &AccessGuard );
+		UE::TUniqueLock ScopedLock(AccessGuard);
 #endif
 		Private::UpdateSlackStat(*this);
 
@@ -1703,7 +1702,7 @@ void FMallocBinned::DumpAllocatorStats( class FOutputDevice& Ar )
 	FBufferedOutputDevice BufferedOutput;
 	{
 #ifdef USE_COARSE_GRAIN_LOCKS
-		FScopeLock ScopedLock( &AccessGuard );
+		UE::TUniqueLock ScopedLock(AccessGuard);
 #endif
 		ValidateHeap();
 #if STATS
diff --git a/Engine/Source/Runtime/Core/Private/HAL/MallocBinned2.cpp b/Engine/Source/Runtime/Core/Private/HAL/MallocBinned2.cpp
--- a/Engine/Source/Runtime/Core/Private/HAL/MallocBinned2.cpp
+++ b/Engine/Source/Runtime/Core/Private/HAL/MallocBinned2.cpp
@@ -38,7 +38,7 @@ namespace
 
 #if UE_MB2_ALLOCATOR_STATS_VALIDATION
 	int64 AllocatedSmallPoolMemoryValidation = 0;
-	FCriticalSection ValidationCriticalSection;
+	UE::FPlatformRecursiveMutex ValidationCriticalSection;
 	int32 RecursionCounter = 0;
 #endif
 
@@ -237,7 +237,7 @@ struct FMallocBinned2::Private
 	{
 		FPoolTable& Table = Allocator.SmallPoolTables[InPoolIndex];
 
-		FScopeLock Lock(&Table.Mutex);
+		UE::TUniqueLock Lock(Table.Mutex);
 
 		FBundleNode* Bundle = BundlesToRecycle;
 		while (Bundle)
@@ -287,7 +287,7 @@ struct FMallocBinned2::Private
 					// Free the OS memory.
 					NodePool->Unlink();
 					{
-						FScopeLock InnerLock(&Allocator.ExternalAllocMutex);
+						UE::TUniqueLock InnerLock(Allocator.ExternalAllocMutex);
 						Allocator.CachedOSPageAllocator.Free(BasePtrOfNode, GMB2PageSize);
 					}
 
@@ -316,7 +316,7 @@ struct FMallocBinned2::Private
 	{
 		if(GMallocBinnedPerThreadCaches)
 		{
-			FScopeLock Lock(&GetFreeBlockListsRegistrationMutex());
+			UE::TUniqueLock Lock(GetFreeBlockListsRegistrationMutex());
 			TArray<FPerThreadFreeBlockLists*>& List = GetRegisteredFreeBlockLists();
 			UE_CLOG(List.Num() == 1 && List[0] != FPerThreadFreeBlockLists::Get(), LogMemory, Fatal, TEXT("There was a thread-local free list at fork time which did not belong to the main forking thread. No other threads should be alive at fork time. If threads are spawned before forking, they must be killed and FMallocBinned2::ClearAndDisableTLSCachesOnCurrentThread() must be called."));
 			UE_CLOG(List.Num() > 1, LogMemory, Fatal, TEXT("There were multiple thread-local free lists at fork time. No other threads should be alive at fork time. If threads are spawned before forking, they must be killed and FMallocBinned2::ClearAndDisableTLSCachesOnCurrentThread() must be called."));
@@ -361,7 +361,7 @@ FMallocBinned2::FPoolInfo& FMallocBinned2::FPoolList::PushNewPoolToFront(FMalloc
 
 	void* FreePtr;
 	{
-		FScopeLock Lock(&Allocator.ExternalAllocMutex);
+		UE::TUniqueLock Lock(Allocator.ExternalAllocMutex);
 		FreePtr = Allocator.CachedOSPageAllocator.Allocate(GMB2PageSize, FMemory::AllocationHints::SmallPool);
 	}
 	if (!FreePtr)
@@ -378,7 +378,7 @@ FMallocBinned2::FPoolInfo& FMallocBinned2::FPoolList::PushNewPoolToFront(FMalloc
 	AllocatedOSSmallPoolMemory.fetch_add((int64)GMB2PageSize, std::memory_order_relaxed);
 #endif
 
-	FScopeLock Lock(&Allocator.ExternalAllocMutex);
+	UE::TUniqueLock Lock(Allocator.ExternalAllocMutex);
 	// Create pool
 	FPoolInfo* Result = Internal::GetOrCreatePoolInfo(Allocator, Free, FPoolInfo::ECanary::FirstFreeBlockIsPtr);
 	Result->Link(Front);
@@ -455,14 +455,14 @@ void FMallocBinned2::OnMallocInitialized()
 #if UE_USE_VERYLARGEPAGEALLOCATOR
 	FCoreDelegates::GetLowLevelAllocatorMemoryTrimDelegate().AddLambda([this]()
 		{
-			FScopeLock Lock(&ExternalAllocMutex);
+			UE::TUniqueLock Lock(ExternalAllocMutex);
 			CachedOSPageAllocator.FreeAll(&ExternalAllocMutex);
 		}
 	);
 
 	FCoreDelegates::GetRefreshLowLevelAllocatorDelegate().AddLambda([this]()
 		{
-			FScopeLock Lock(&ExternalAllocMutex);
+			UE::TUniqueLock Lock(ExternalAllocMutex);
 			CachedOSPageAllocator.Refresh();
 		}
 	);
@@ -490,7 +490,7 @@ void FMallocBinned2::OnPreFork()
 	}
 
 #if !UE_USE_VERYLARGEPAGEALLOCATOR
-	FScopeLock Lock(&ExternalAllocMutex);
+	UE::TUniqueLock Lock(ExternalAllocMutex);
 	CachedOSPageAllocator.FreeAll(&ExternalAllocMutex);
 #endif
 
@@ -514,7 +514,7 @@ void FMallocBinned2::OnPostFork()
 
 	for (FPoolTable& Table : SmallPoolTables)
 	{
-		FScopeLock Lock(&Table.Mutex);
+		UE::TUniqueLock Lock(Table.Mutex);
 		// Clear our list of partially used pages so we don't dirty them and cause them to become unshared with the parent process
 		Table.ActivePools.Clear();
 		Table.ExhaustedPools.Clear();
@@ -565,7 +565,7 @@ void* FMallocBinned2::MallocExternalSmall(SIZE_T Size, uint32 Alignment)
 	// Allocate from small object pool.
 	FPoolTable& Table = SmallPoolTables[PoolIndex];
 
-	FScopeLock Lock(&Table.Mutex);
+	UE::TUniqueLock Lock(Table.Mutex);
 
 	FPoolInfo* Pool;
 	if (!Table.ActivePools.IsEmpty())
@@ -627,7 +627,7 @@ void* FMallocBinned2::MallocExternalLarge(SIZE_T Size, uint32 Alignment)
 	{
 		NOALLOC_SCOPE_CYCLE_COUNTER(STAT_FMallocBinned2_MallocExternalLarge);
 
-		FScopeLock Lock(&ExternalAllocMutex);
+		UE::TUniqueLock Lock(ExternalAllocMutex);
 
 		// Use OS for non-pooled allocations.
 		Result = CachedOSPageAllocator.Allocate(AlignedSize, 0, &ExternalAllocMutex);
@@ -803,7 +803,7 @@ void FMallocBinned2::FreeExternal(void* Ptr)
 	}
 	else if (Ptr)
 	{
-		FScopeLock Lock(&ExternalAllocMutex);
+		UE::TUniqueLock Lock(ExternalAllocMutex);
 		FPoolInfo* Pool = Internal::FindPoolInfo(*this, Ptr);
 		if (!Pool)
 		{
@@ -854,7 +854,7 @@ bool FMallocBinned2::ValidateHeap()
 	
 	for (FPoolTable& Table : SmallPoolTables)
 	{
-		FScopeLock Lock(&Table.Mutex);
+		UE::TUniqueLock Lock(Table.Mutex);
 		Table.ActivePools.ValidateActivePools();
 		Table.ExhaustedPools.ValidateExhaustedPools();
 	}
@@ -879,7 +879,7 @@ void FMallocBinned2::Trim(bool bTrimThreadCaches)
 		FMallocBinnedCommonUtils::Trim(*this);
 
 #if !UE_USE_VERYLARGEPAGEALLOCATOR
-		FScopeLock Lock(&ExternalAllocMutex);
+		UE::TUniqueLock Lock(ExternalAllocMutex);
 		// this cache is recycled anyway, if you need to trim it based on being OOM, it's already too late.
 		CachedOSPageAllocator.FreeAll(&ExternalAllocMutex);
 #endif
diff --git a/Engine/Source/Runtime/Core/Private/HAL/MallocBinned3.cpp b/Engine/Source/Runtime/Core/Private/HAL/MallocBinned3.cpp
--- a/Engine/Source/Runtime/Core/Private/HAL/MallocBinned3.cpp
+++ b/Engine/Source/Runtime/Core/Private/HAL/MallocBinned3.cpp
@@ -5,7 +5,7 @@
 PRAGMA_DISABLE_UNSAFE_TYPECAST_WARNINGS
 
 #if PLATFORM_HAS_FPlatformVirtualMemoryBlock
-#include "Misc/ScopeLock.h"
+#include "Async/UniqueLock.h"
 #include "Templates/Function.h"
 #include "GenericPlatform/GenericPlatformProcess.h"
 #include "Stats/Stats.h"
@@ -311,7 +311,7 @@ struct FMallocBinned3::Private
 	static void FreeBundles(FMallocBinned3& Allocator, FBundleNode* BundlesToRecycle, uint32 InBinSize, uint32 InPoolIndex)
 	{
 		FPoolTable& Table = Allocator.SmallPoolTables[InPoolIndex];
-		FScopeLock Lock(&Table.Mutex);
+		UE::TUniqueLock Lock(Table.Mutex);
 
 		FBundleNode* Bundle = BundlesToRecycle;
 		while (Bundle)
@@ -710,7 +710,7 @@ void* FMallocBinned3::MallocExternal(SIZE_T Size, uint32 Alignment)
 		// Allocate from small object pool.
 		FPoolTable& Table = SmallPoolTables[PoolIndex];
 
-		FScopeLock Lock(&Table.Mutex);
+		UE::TUniqueLock Lock(Table.Mutex);
 
 		uint32 BlockIndex = MAX_uint32;
 		FPoolInfoSmall* Pool = GetFrontPool(Table, PoolIndex, BlockIndex);
@@ -780,7 +780,7 @@ void* FMallocBinned3::MallocExternal(SIZE_T Size, uint32 Alignment)
 	NOALLOC_SCOPE_CYCLE_COUNTER(STAT_FMallocBinned3_MallocExternalLarge);
 
 #if UE_MB3_USE_CACHED_PAGE_ALLOCATOR_FOR_LARGE_ALLOCS
-	FScopeLock Lock(&Mutex);
+	UE::TUniqueLock Lock(Mutex);
 	void* Result = GetCachedOSPageAllocator().Allocate(AlignedSize);
 	check(IsAligned(Result, Alignment));
 #else
@@ -810,7 +810,7 @@ void* FMallocBinned3::MallocExternal(SIZE_T Size, uint32 Alignment)
 	check(IsOSAllocation(Result));
 
 #if! UE_MB3_USE_CACHED_PAGE_ALLOCATOR_FOR_LARGE_ALLOCS
-	FScopeLock Lock(&ExternalAllocMutex);
+	UE::TUniqueLock Lock(ExternalAllocMutex);
 #endif
 
 #if UE_MB3_ALLOCATOR_STATS
@@ -946,12 +946,12 @@ void FMallocBinned3::FreeExternal(void* Ptr, uint64 PoolIndex)
 	else if (Ptr)
 	{
 #if UE_MB3_USE_CACHED_PAGE_ALLOCATOR_FOR_LARGE_ALLOCS
-		FScopeLock Lock(&Mutex);
+		UE::TUniqueLock Lock(Mutex);
 #endif
 		uint32 VMPages;
 		{
 #if !UE_MB3_USE_CACHED_PAGE_ALLOCATOR_FOR_LARGE_ALLOCS
-			FScopeLock Lock(&ExternalAllocMutex);
+			UE::TUniqueLock Lock(ExternalAllocMutex);
 #endif
 			FPoolInfo* Pool = Internal::FindPoolInfo(*this, Ptr);
 			if (!Pool)
diff --git a/Engine/Source/Runtime/Core/Private/HAL/MallocBinnedGPU.cpp b/Engine/Source/Runtime/Core/Private/HAL/MallocBinnedGPU.cpp
--- a/Engine/Source/Runtime/Core/Private/HAL/MallocBinnedGPU.cpp
+++ b/Engine/Source/Runtime/Core/Private/HAL/MallocBinnedGPU.cpp
@@ -800,8 +800,8 @@ void FMallocBinnedGPU::Decommit(uint32 InPoolIndex, void *Ptr, SIZE_T Size)
 
 FMallocBinnedGPU::~FMallocBinnedGPU()
 {
-	FScopeLock Lock(&Mutex);
-	FScopeLock Lock2(&GetFreeBlockListsRegistrationMutex());
+	UE::TUniqueLock Lock(Mutex);
+	UE::TUniqueLock Lock2(GetFreeBlockListsRegistrationMutex());
 
 	MBG_STAT(
 		UE_CLOG(BinnedGPUAllocatedOSSmallPoolMemory > 0, LogCore, Error, TEXT("FMallocBinnedGPU leaked small block memory: %fmb"), ((double)BinnedGPUAllocatedOSSmallPoolMemory) / (1024.0f * 1024.0f));
@@ -872,7 +872,7 @@ void* FMallocBinnedGPU::MallocExternal(SIZE_T Size, uint32 Alignment)
 			}
 		}
 
-		FScopeLock Lock(&Mutex);
+		UE::TUniqueLock Lock(Mutex);
 
 		// Allocate from small object pool.
 		FPoolTable& Table = SmallPoolTables[PoolIndex];
@@ -923,7 +923,7 @@ void* FMallocBinnedGPU::MallocExternal(SIZE_T Size, uint32 Alignment)
 		Private::OutOfMemory(Size);
 	}
 	check(IsOSAllocation(Result));
-	FScopeLock Lock(&Mutex);
+	UE::TUniqueLock Lock(Mutex);
 
 	MBG_STAT(BinnedGPUAllocatedLargePoolMemory += (int64)Size;)
 	MBG_STAT(BinnedGPUAllocatedLargePoolMemoryWAlignment += CommitSize;)
@@ -964,7 +964,7 @@ void FMallocBinnedGPU::FreeExternal(void* Ptr)
 		if (BundlesToRecycle)
 		{
 			BundlesToRecycle->NextBundle = nullptr;
-			FScopeLock Lock(&Mutex);
+			UE::TUniqueLock Lock(Mutex);
 			Private::FreeBundles(*this, BundlesToRecycle, BlockSize, PoolIndex);
 			if (!Lists)
 			{
@@ -978,7 +978,7 @@ void FMallocBinnedGPU::FreeExternal(void* Ptr)
 		uint32 VMSizeDivVirtualSizeAlignment;
 
 		{
-			FScopeLock Lock(&Mutex);
+			UE::TUniqueLock Lock(Mutex);
 			FPoolInfoLarge* Pool = Private::FindPoolInfo(*this, Ptr);
 			if (!Pool)
 			{
@@ -1012,7 +1012,7 @@ bool FMallocBinnedGPU::GetAllocationSizeExternal(void* Ptr, SIZE_T& SizeOut)
 	{
 		return false;
 	}
-	FScopeLock Lock(&Mutex);
+	UE::TUniqueLock Lock(Mutex);
 	FPoolInfoLarge* Pool = Private::FindPoolInfo(*this, Ptr);
 	if (!Pool)
 	{
@@ -1048,7 +1048,7 @@ void FMallocBinnedGPU::FlushCurrentThreadCache()
 
 	if (Lists)
 	{
-		FScopeLock Lock(&Mutex);
+		UE::TUniqueLock Lock(Mutex);
 		WaitForMutexTime = FPlatformTime::Seconds() - StartTimeInner;
 		for (int32 PoolIndex = 0; PoolIndex != ArenaParams.PoolCount; ++PoolIndex)
 		{
@@ -1216,7 +1216,7 @@ int64 FMallocBinnedGPU::GetTotalAllocatedSmallPoolMemory()
 {
 	int64 FreeBlockAllocatedMemory = 0;
 	{
-		FScopeLock Lock(&GetFreeBlockListsRegistrationMutex());
+		UE::TUniqueLock Lock(GetFreeBlockListsRegistrationMutex());
 		for (const FPerThreadFreeBlockLists* FreeBlockLists : GetRegisteredFreeBlockLists())
 		{
 			FreeBlockAllocatedMemory += FreeBlockLists->AllocatedMemory;
diff --git a/Engine/Source/Runtime/Core/Public/HAL/Allocators/CachedOSPageAllocator.h b/Engine/Source/Runtime/Core/Public/HAL/Allocators/CachedOSPageAllocator.h
--- a/Engine/Source/Runtime/Core/Public/HAL/Allocators/CachedOSPageAllocator.h
+++ b/Engine/Source/Runtime/Core/Public/HAL/Allocators/CachedOSPageAllocator.h
@@ -4,8 +4,8 @@
 
 #include "CoreTypes.h"
 #include "Misc/OutputDevice.h"
-#include "HAL/CriticalSection.h"
 #include "HAL/PlatformMemory.h"
+#include "HAL/PlatformMutex.h"
 
 struct FCachedOSPageAllocator
 {
@@ -22,9 +22,9 @@ protected:
 		}
 	};
 
-	void* AllocateImpl(SIZE_T Size, uint32 CachedByteLimit, FFreePageBlock* First, FFreePageBlock* Last, uint32& FreedPageBlocksNum, SIZE_T& CachedTotal, FCriticalSection* Mutex);
-	void FreeImpl(void* Ptr, SIZE_T Size, uint32 NumCacheBlocks, uint32 CachedByteLimit, FFreePageBlock* First, uint32& FreedPageBlocksNum, SIZE_T& CachedTotal, FCriticalSection* Mutex, bool ThreadIsTimeCritical);
-	void FreeAllImpl(FFreePageBlock* First, uint32& FreedPageBlocksNum, SIZE_T& CachedTotal, FCriticalSection* Mutex);
+	void* AllocateImpl(SIZE_T Size, uint32 CachedByteLimit, FFreePageBlock* First, FFreePageBlock* Last, uint32& FreedPageBlocksNum, SIZE_T& CachedTotal, UE::FPlatformRecursiveMutex* Mutex);
+	void FreeImpl(void* Ptr, SIZE_T Size, uint32 NumCacheBlocks, uint32 CachedByteLimit, FFreePageBlock* First, uint32& FreedPageBlocksNum, SIZE_T& CachedTotal, UE::FPlatformRecursiveMutex* Mutex, bool ThreadIsTimeCritical);
+	void FreeAllImpl(FFreePageBlock* First, uint32& FreedPageBlocksNum, SIZE_T& CachedTotal, UE::FPlatformRecursiveMutex* Mutex);
 
 	static bool IsOSAllocation(SIZE_T Size, uint32 CachedByteLimit)
 	{
@@ -41,28 +41,30 @@ struct TCachedOSPageAllocator : private FCachedOSPageAllocator
 	{
 	}
 
-	FORCEINLINE void* Allocate(SIZE_T Size, uint32 AllocationHint = 0, FCriticalSection* Mutex = nullptr)
+	FORCEINLINE void* Allocate(SIZE_T Size, uint32 AllocationHint = 0, UE::FPlatformRecursiveMutex* Mutex = nullptr)
 	{
 		return AllocateImpl(Size, CachedByteLimit, FreedPageBlocks, FreedPageBlocks + FreedPageBlocksNum, FreedPageBlocksNum, CachedTotal, Mutex);
 	}
 
-	void Free(void* Ptr, SIZE_T Size, FCriticalSection* Mutex = nullptr, bool ThreadIsTimeCritical = false)
+	void Free(void* Ptr, SIZE_T Size, UE::FPlatformRecursiveMutex* Mutex = nullptr, bool ThreadIsTimeCritical = false)
 	{
 		return FreeImpl(Ptr, Size, ThreadIsTimeCritical ? NumCacheBlocks*2 : NumCacheBlocks, CachedByteLimit, FreedPageBlocks, FreedPageBlocksNum, CachedTotal, Mutex, ThreadIsTimeCritical);
 	}
-	void FreeAll(FCriticalSection* Mutex = nullptr)
+
+	void FreeAll(UE::FPlatformRecursiveMutex* Mutex = nullptr)
 	{
 		return FreeAllImpl(FreedPageBlocks, FreedPageBlocksNum, CachedTotal, Mutex);
 	}
+
 	// Refresh cached os allocator if needed. Does nothing for this implementation
 	void Refresh()
 	{
-
 	}
+
 	void UpdateStats()
 	{
-
 	}
+
 	uint64 GetCachedFreeTotal() const
 	{
 		return CachedTotal;
@@ -88,3 +90,7 @@ private:
 	SIZE_T         CachedTotal;
 	uint32         FreedPageBlocksNum;
 };
+
+#if UE_ENABLE_INCLUDE_ORDER_DEPRECATED_IN_5_7
+#include "HAL/CriticalSection.h"
+#endif
diff --git a/Engine/Source/Runtime/Core/Public/HAL/Allocators/CachedOSVeryLargePageAllocator.h b/Engine/Source/Runtime/Core/Public/HAL/Allocators/CachedOSVeryLargePageAllocator.h
--- a/Engine/Source/Runtime/Core/Public/HAL/Allocators/CachedOSVeryLargePageAllocator.h
+++ b/Engine/Source/Runtime/Core/Public/HAL/Allocators/CachedOSVeryLargePageAllocator.h
@@ -4,8 +4,8 @@
 #include "Containers/List.h"
 #include "CoreTypes.h"
 #include "HAL/Allocators/CachedOSPageAllocator.h"
-#include "HAL/CriticalSection.h"
 #include "HAL/PlatformMemory.h"
+#include "HAL/PlatformMutex.h"
 #include "HAL/UnrealMemory.h"
 
 
@@ -51,11 +51,11 @@ public:
 		// this leaks everything!
 	}
 
-	void* Allocate(SIZE_T Size, uint32 AllocationHint = 0, FCriticalSection* Mutex = nullptr);
+	void* Allocate(SIZE_T Size, uint32 AllocationHint = 0, UE::FPlatformRecursiveMutex* Mutex = nullptr);
 
-	void Free(void* Ptr, SIZE_T Size, FCriticalSection* Mutex = nullptr, bool ThreadIsTimeCritical = false);
+	void Free(void* Ptr, SIZE_T Size, UE::FPlatformRecursiveMutex* Mutex = nullptr, bool ThreadIsTimeCritical = false);
 
-	void FreeAll(FCriticalSection* Mutex = nullptr);
+	void FreeAll(UE::FPlatformRecursiveMutex* Mutex = nullptr);
 
 	// Refresh cached os allocator if needed. Will preallocate / reduce backstore if preallocation is enabled
 	void Refresh();
@@ -120,8 +120,8 @@ private:
 
 	void Init();
 	void ShrinkEmptyBackStore(int32 NewEmptyBackStoreSize, FMemory::AllocationHints AllocationHint);
-	FLargePage* GetOrAllocLargePage(uint32 AllocationHint, FCriticalSection* Mutex);
-	FLargePage* AllocNewLargePage(uint32 AllocationHint, FCriticalSection* Mutex, bool& bOutCommitFailure);
+	FLargePage* GetOrAllocLargePage(uint32 AllocationHint, UE::FPlatformRecursiveMutex* Mutex);
+	FLargePage* AllocNewLargePage(uint32 AllocationHint, UE::FPlatformRecursiveMutex* Mutex, bool& bOutCommitFailure);
 
 	bool bEnabled;
 	uintptr_t	AddressSpaceReserved;
@@ -149,4 +149,8 @@ private:
 };
 CORE_API extern bool GEnableVeryLargePageAllocator;
 
+#if UE_ENABLE_INCLUDE_ORDER_DEPRECATED_IN_5_7
+#include "HAL/CriticalSection.h"
+#endif
+
 #endif // UE_USE_VERYLARGEPAGEALLOCATOR
\ No newline at end of file
diff --git a/Engine/Source/Runtime/Core/Public/HAL/Allocators/PooledVirtualMemoryAllocator.h b/Engine/Source/Runtime/Core/Public/HAL/Allocators/PooledVirtualMemoryAllocator.h
--- a/Engine/Source/Runtime/Core/Public/HAL/Allocators/PooledVirtualMemoryAllocator.h
+++ b/Engine/Source/Runtime/Core/Public/HAL/Allocators/PooledVirtualMemoryAllocator.h
@@ -3,8 +3,8 @@
 #pragma once
 
 #include "CoreTypes.h"
-#include "HAL/CriticalSection.h"
 #include "HAL/Allocators/CachedOSPageAllocator.h"
+#include "HAL/PlatformMutex.h"
 
 /**
  * This class pools OS allocations made from FMallocBinned2.
@@ -61,9 +61,9 @@ struct FPooledVirtualMemoryAllocator
 {
 	FPooledVirtualMemoryAllocator();
 
-	void* Allocate(SIZE_T Size, uint32 AllocationHint = 0, FCriticalSection* Mutex = nullptr);
-	void Free(void* Ptr, SIZE_T Size, FCriticalSection* Mutex = nullptr, bool ThreadIsTimeCritical = false);
-	void FreeAll(FCriticalSection* Mutex = nullptr);
+	void* Allocate(SIZE_T Size, uint32 AllocationHint = 0, UE::FPlatformRecursiveMutex* Mutex = nullptr);
+	void Free(void* Ptr, SIZE_T Size, UE::FPlatformRecursiveMutex* Mutex = nullptr, bool ThreadIsTimeCritical = false);
+	void FreeAll(UE::FPlatformRecursiveMutex* Mutex = nullptr);
 
 	/** A structure that describes a pool of a particular size */
 	struct FPoolDescriptorBase
@@ -138,7 +138,7 @@ private:
 	FPoolDescriptorBase* ClassesListHeads[Limits::NumAllocationSizeClasses];
 
 	/** Per-class locks */
-	mutable FCriticalSection ClassesLocks[Limits::NumAllocationSizeClasses];
+	mutable UE::FPlatformRecursiveMutex ClassesLocks[Limits::NumAllocationSizeClasses];
 
 	/** Increases the number of pooled allocations next time we need a pool
 	 *
@@ -153,8 +153,12 @@ private:
 	void DestroyPool(FPoolDescriptorBase* Pool);
 
 	/** Lock for accessing the caching allocator for larger allocs */
-	FCriticalSection OsAllocatorCacheLock;
+	UE::FPlatformRecursiveMutex OsAllocatorCacheLock;
 
 	/** Caching allocator for larger allocs */
 	TCachedOSPageAllocator<MaxOSAllocsCached, MaxOSAllocCacheSize> OsAllocatorCache;
 };
+
+#if UE_ENABLE_INCLUDE_ORDER_DEPRECATED_IN_5_7
+#	include "HAL/CriticalSection.h"
+#endif
diff --git a/Engine/Source/Runtime/Core/Public/HAL/MallocBinned.h b/Engine/Source/Runtime/Core/Public/HAL/MallocBinned.h
--- a/Engine/Source/Runtime/Core/Public/HAL/MallocBinned.h
+++ b/Engine/Source/Runtime/Core/Public/HAL/MallocBinned.h
@@ -4,12 +4,12 @@
 
 #include "Containers/Array.h"
 #include "CoreTypes.h"
-#include "HAL/CriticalSection.h"
 #include "HAL/LowLevelMemStats.h"
 #include "HAL/LowLevelMemTracker.h"
 #include "HAL/MallocJemalloc.h"
 #include "HAL/MemoryBase.h"
 #include "HAL/PlatformAtomics.h"
+#include "HAL/PlatformMutex.h"
 #include "Math/UnrealMathUtility.h"
 #include "Stats/Stats.h"
 
@@ -155,7 +155,7 @@ private:
 		FPoolInfo*			ExhaustedPool;
 		uint32				BlockSize;
 #ifdef USE_FINE_GRAIN_LOCKS
-		FCriticalSection	CriticalSection;
+		UE::FPlatformRecursiveMutex CriticalSection;
 #endif
 #if STATS
 		/** Number of currently active pools */
@@ -215,7 +215,7 @@ private:
 	bool							bDoneFreeListInit;
 #endif
 
-	FCriticalSection	AccessGuard;
+	UE::FPlatformRecursiveMutex AccessGuard;
 
 	// PageSize dependent constants
 	uint64 MaxHashBuckets;
@@ -340,3 +340,7 @@ public:
 
 	virtual const TCHAR* GetDescriptiveName() override;
 };
+
+#if UE_ENABLE_INCLUDE_ORDER_DEPRECATED_IN_5_7
+#	include "HAL/CriticalSection.h"
+#endif
diff --git a/Engine/Source/Runtime/Core/Public/HAL/MallocBinned2.h b/Engine/Source/Runtime/Core/Public/HAL/MallocBinned2.h
--- a/Engine/Source/Runtime/Core/Public/HAL/MallocBinned2.h
+++ b/Engine/Source/Runtime/Core/Public/HAL/MallocBinned2.h
@@ -3,13 +3,14 @@
 #pragma once
 
 #include "CoreTypes.h"
+#include "Async/UniqueLock.h"
 #include "AutoRTFM.h"
 #include "HAL/Allocators/CachedOSPageAllocator.h"
 #include "HAL/Allocators/CachedOSVeryLargePageAllocator.h"
 #include "HAL/Allocators/PooledVirtualMemoryAllocator.h"
-#include "HAL/CriticalSection.h"
 #include "HAL/LowLevelMemTracker.h"
 #include "HAL/MallocBinnedCommon.h"
+#include "HAL/PlatformMutex.h"
 #include "HAL/UnrealMemory.h"
 #include "Math/NumericLimits.h"
 #include "Misc/AssertionMacros.h"
@@ -43,10 +44,8 @@
 #define UE_MB2_ALLOCATOR_STATS_VALIDATION	(UE_MB2_ALLOCATOR_STATS && 0)
 
 #if UE_MB2_ALLOCATOR_STATS_VALIDATION
-#	include "Misc/ScopeLock.h"
-
 	extern int64 AllocatedSmallPoolMemoryValidation;
-	extern FCriticalSection ValidationCriticalSection;
+	extern UE::FPlatformRecursiveMutex ValidationCriticalSection;
 	extern int32 RecursionCounter;
 #endif
 
@@ -197,7 +196,7 @@ public:
 		uint32 TotalAllocatedMem = 0;
 #endif
 
-		FCriticalSection Mutex;
+		UE::FPlatformRecursiveMutex Mutex;
 
 		FPoolTable() = default;
 	};
@@ -248,7 +247,7 @@ public:
 	FORCEINLINE virtual void* Malloc(SIZE_T Size, uint32 Alignment) override
 	{
 #if UE_MB2_ALLOCATOR_STATS_VALIDATION
-		FScopeLock Lock(&ValidationCriticalSection);
+		UE::TUniqueLock Lock(ValidationCriticalSection);
 		++RecursionCounter;
 		void *Result = MallocInline(Size, Alignment);
 		if (!IsOSAllocation(Result))
@@ -321,7 +320,7 @@ public:
 			ensure(GetAllocationSize(Ptr, OldSize));
 		}
 
-		FScopeLock Lock(&ValidationCriticalSection);
+		UE::TUniqueLock Lock(ValidationCriticalSection);
 		++RecursionCounter;
 		void *Result = ReallocInline(Ptr, NewSize, Alignment);
 		if (!bOldIsOsAllocation)
@@ -413,7 +412,7 @@ public:
 	FORCEINLINE virtual void Free(void* Ptr) override
 	{
 #if UE_MB2_ALLOCATOR_STATS_VALIDATION
-		FScopeLock Lock(&ValidationCriticalSection);
+		UE::TUniqueLock Lock(ValidationCriticalSection);
 		++RecursionCounter;
 		if (!IsOSAllocation(Ptr))
 		{
@@ -550,3 +549,9 @@ public:
 	#endif
 #endif
 
+#if UE_ENABLE_INCLUDE_ORDER_DEPRECATED_IN_5_7
+#	include "HAL/CriticalSection.h"
+#	if UE_MB2_ALLOCATOR_STATS_VALIDATION
+#		include "Misc/ScopeLock.h"
+#	endif
+#endif
diff --git a/Engine/Source/Runtime/Core/Public/HAL/MallocBinned3.h b/Engine/Source/Runtime/Core/Public/HAL/MallocBinned3.h
--- a/Engine/Source/Runtime/Core/Public/HAL/MallocBinned3.h
+++ b/Engine/Source/Runtime/Core/Public/HAL/MallocBinned3.h
@@ -8,10 +8,10 @@
 #include "AutoRTFM.h"
 #include "HAL/Allocators/CachedOSPageAllocator.h"
 #include "HAL/Allocators/PooledVirtualMemoryAllocator.h"
-#include "HAL/CriticalSection.h"
 #include "HAL/LowLevelMemTracker.h"
 #include "HAL/MallocBinnedCommon.h"
 #include "HAL/PlatformMemory.h"
+#include "HAL/PlatformMutex.h"
 #include "HAL/UnrealMemory.h"
 #include "Math/NumericLimits.h"
 #include "Misc/AssertionMacros.h"
@@ -149,7 +149,7 @@ public:
 		uint32 TotalAllocatedMem = 0;
 #endif
 
-		FCriticalSection Mutex;
+		UE::FPlatformRecursiveMutex Mutex;
 
 #if UE_M3_ALLOCATOR_PER_BIN_STATS
 		// these are "head end" stats, above the TLS cache
@@ -503,4 +503,8 @@ PRAGMA_RESTORE_UNSAFE_TYPECAST_WARNINGS
 	#endif
 #endif
 
+#if UE_ENABLE_INCLUDE_ORDER_DEPRECATED_IN_5_7
+#	include "HAL/CriticalSection.h"
+#endif
+
 #endif	//~PLATFORM_HAS_FPlatformVirtualMemoryBlock
diff --git a/Engine/Source/Runtime/Core/Public/HAL/MallocBinnedCommon.h b/Engine/Source/Runtime/Core/Public/HAL/MallocBinnedCommon.h
--- a/Engine/Source/Runtime/Core/Public/HAL/MallocBinnedCommon.h
+++ b/Engine/Source/Runtime/Core/Public/HAL/MallocBinnedCommon.h
@@ -6,9 +6,10 @@
 #include "CoreGlobals.h"
 #include <atomic>
 #include "HAL/MemoryBase.h"
+#include "HAL/PlatformMutex.h"
 #include "HAL/PlatformTLS.h"
-#include "HAL/CriticalSection.h"
-#include "Async/Mutex.h"
+#include "Async/UniqueLock.h"
+#include "Async/WordMutex.h"
 #include "Misc/ScopeLock.h"
 #include "Templates/AlignmentTemplates.h"
 #include "Logging/LogMacros.h"
@@ -386,7 +387,7 @@ protected:
 	FPtrToPoolMapping PtrToPoolMapping;
 	uint64 NumPoolsPerPage;		// Number of AllocType::FPoolInfo
 	static uint32 OsAllocationGranularity;
-	FCriticalSection ExternalAllocMutex;
+	UE::FPlatformRecursiveMutex ExternalAllocMutex;
 
 	static CORE_API uint32 BinnedTlsSlot;
 
@@ -606,7 +607,7 @@ protected:
 		int64 AllocatedMemory = 0;
 #endif
 	private:
-		UE::FMutex Mutex;
+		UE::FWordMutex Mutex;
 		uint64 MemoryTrimEpoch = 0;
 		FFreeBlockList FreeLists[NUM_SMALL_POOLS];
 		bool bLockedByOwnerThread = false;
@@ -729,9 +730,9 @@ protected:
 		}
 	};
 
-	static FCriticalSection& GetFreeBlockListsRegistrationMutex()
+	static UE::FPlatformRecursiveMutex& GetFreeBlockListsRegistrationMutex()
 	{
-		static FCriticalSection FreeBlockListsRegistrationMutex;
+		static UE::FPlatformRecursiveMutex FreeBlockListsRegistrationMutex;
 		return FreeBlockListsRegistrationMutex;
 	}
 
@@ -884,7 +885,7 @@ protected:
 		typename AllocType::FPoolInfo* Pool;
 		{
 			//NOALLOC_SCOPE_CYCLE_COUNTER(STAT_FMallocBinned_GetAllocationSizeExternal);
-			FScopeLock Lock(&ExternalAllocMutex);
+			UE::TUniqueLock Lock(ExternalAllocMutex);
 			Pool = Internal::FindPoolInfo(*(AllocType*)this, Ptr);
 		}
 
@@ -905,7 +906,7 @@ protected:
 		//NOALLOC_SCOPE_CYCLE_COUNTER(STAT_FMallocBinned_GetTotalAllocatedSmallPoolMemory);
 		int64 FreeBlockAllocatedMemory = 0;
 		{
-			FScopeLock Lock(&GetFreeBlockListsRegistrationMutex());
+			UE::TUniqueLock Lock(GetFreeBlockListsRegistrationMutex());
 			for (const FPerThreadFreeBlockLists* FreeBlockLists : GetRegisteredFreeBlockLists())
 			{
 				FreeBlockAllocatedMemory += FreeBlockLists->AllocatedMemory;
@@ -988,15 +989,19 @@ private:
 	static void RegisterThreadFreeBlockLists(FPerThreadFreeBlockLists* FreeBlockLists)
 	{
 		//NOALLOC_SCOPE_CYCLE_COUNTER(STAT_FMallocBinned_RegisterThreadFreeBlockLists);
-		FScopeLock Lock(&GetFreeBlockListsRegistrationMutex());
+		UE::TUniqueLock Lock(GetFreeBlockListsRegistrationMutex());
 		GetRegisteredFreeBlockLists().Add(FreeBlockLists);
 	}
 
 	static void UnregisterThreadFreeBlockLists(FPerThreadFreeBlockLists* FreeBlockLists)
 	{
 		//NOALLOC_SCOPE_CYCLE_COUNTER(STAT_FMallocBinned_UnregisterThreadFreeBlockLists);
-		FScopeLock Lock(&GetFreeBlockListsRegistrationMutex());
+		UE::TUniqueLock Lock(GetFreeBlockListsRegistrationMutex());
 		GetRegisteredFreeBlockLists().Remove(FreeBlockLists);
 		UE_MBC_UPDATE_STATS(ConsolidatedMemory.fetch_add(FreeBlockLists->AllocatedMemory, std::memory_order_relaxed));
 	}
 };
+
+#if UE_ENABLE_INCLUDE_ORDER_DEPRECATED_IN_5_7
+#	include "HAL/CriticalSection.h"
+#endif
diff --git a/Engine/Source/Runtime/Core/Public/HAL/MallocBinnedCommonUtils.h b/Engine/Source/Runtime/Core/Public/HAL/MallocBinnedCommonUtils.h
--- a/Engine/Source/Runtime/Core/Public/HAL/MallocBinnedCommonUtils.h
+++ b/Engine/Source/Runtime/Core/Public/HAL/MallocBinnedCommonUtils.h
@@ -149,7 +149,7 @@ public:
 		// Skip on desktop as we may have too many threads and this could cause some hitches.
 		if (!PLATFORM_DESKTOP && GMallocBinnedFlushRegisteredThreadCachesOnOneThread != 0)
 		{
-			FScopeLock FreeBlockLock(&AllocType::GetFreeBlockListsRegistrationMutex());
+			UE::TUniqueLock FreeBlockLock(AllocType::GetFreeBlockListsRegistrationMutex());
 			for (typename AllocType::FPerThreadFreeBlockLists* BlockList : AllocType::GetRegisteredFreeBlockLists())
 			{
 				// If we're unable to lock, it's because the thread is currently active so it
diff --git a/Engine/Source/Runtime/Core/Public/HAL/MallocBinnedGPU.h b/Engine/Source/Runtime/Core/Public/HAL/MallocBinnedGPU.h
--- a/Engine/Source/Runtime/Core/Public/HAL/MallocBinnedGPU.h
+++ b/Engine/Source/Runtime/Core/Public/HAL/MallocBinnedGPU.h
@@ -11,13 +11,14 @@
 #include "Templates/MemoryOps.h"
 
 #if PLATFORM_64BITS && PLATFORM_HAS_FPlatformVirtualMemoryBlock
+#include "Async/UniqueLock.h"
 #include "HAL/Allocators/CachedOSPageAllocator.h"
 #include "HAL/Allocators/PooledVirtualMemoryAllocator.h"
-#include "HAL/CriticalSection.h"
 #include "HAL/LowLevelMemTracker.h"
 #include "HAL/MallocBinnedCommon.h"
 #include "HAL/MemoryBase.h"
 #include "HAL/PlatformMath.h"
+#include "HAL/PlatformMutex.h"
 #include "HAL/PlatformTLS.h"
 #include "HAL/UnrealMemory.h"
 #include "Math/NumericLimits.h"
@@ -647,7 +648,7 @@ public:
 	PoolHashBucket* HashBucketFreeList;
 	uint64 NumLargePoolsPerPage;
 
-	FCriticalSection Mutex;
+	UE::FPlatformRecursiveMutex Mutex;
 	FGlobalRecycler GGlobalRecycler;
 	FPtrToPoolMapping PtrToPoolMapping;
 
@@ -678,8 +679,8 @@ public:
 		TAtomic<int64> GPUProxyMemory;
 	)
 
-	FCriticalSection FreeBlockListsRegistrationMutex;
-	FCriticalSection& GetFreeBlockListsRegistrationMutex()
+	UE::FPlatformRecursiveMutex FreeBlockListsRegistrationMutex;
+	UE::FPlatformRecursiveMutex& GetFreeBlockListsRegistrationMutex()
 	{
 		return FreeBlockListsRegistrationMutex;
 	}
@@ -690,12 +691,12 @@ public:
 	}
 	void RegisterThreadFreeBlockLists(FPerThreadFreeBlockLists* FreeBlockLists)
 	{
-		FScopeLock Lock(&GetFreeBlockListsRegistrationMutex());
+		UE::TUniqueLock Lock(GetFreeBlockListsRegistrationMutex());
 		GetRegisteredFreeBlockLists().Add(FreeBlockLists);
 	}
 	int64 UnregisterThreadFreeBlockLists(FPerThreadFreeBlockLists* FreeBlockLists)
 	{
-		FScopeLock Lock(&GetFreeBlockListsRegistrationMutex());
+		UE::TUniqueLock Lock(GetFreeBlockListsRegistrationMutex());
 		GetRegisteredFreeBlockLists().Remove(FreeBlockLists);
 		return FreeBlockLists->AllocatedMemory;
 	}
@@ -705,4 +706,8 @@ public:
 
 PRAGMA_RESTORE_UNSAFE_TYPECAST_WARNINGS
 
+#if UE_ENABLE_INCLUDE_ORDER_DEPRECATED_IN_5_7
+#	include "HAL/CriticalSection.h"
+#endif
+
 #endif
