Changed FAlignedAllocator in ConcurrentLinearAllocator to avoid overallocating and aligning memory internally and instead use default allocator's alignment parameter
Deprecated FOsAllocator as it's obsolete now

diff --git a/Engine/Plugins/Experimental/EditorDataStorage/Source/TedsCore/Private/TypedElementDatabaseScratchBuffer.h b/Engine/Plugins/Experimental/EditorDataStorage/Source/TedsCore/Private/TypedElementDatabaseScratchBuffer.h
--- a/Engine/Plugins/Experimental/EditorDataStorage/Source/TedsCore/Private/TypedElementDatabaseScratchBuffer.h
+++ b/Engine/Plugins/Experimental/EditorDataStorage/Source/TedsCore/Private/TypedElementDatabaseScratchBuffer.h
@@ -21,7 +21,7 @@ namespace UE::Editor::DataStorage
 			static constexpr bool InlineBlockAllocation = true;  // Inline or Noinline the BlockAllocation which can have an impact on Performance
 			static constexpr const char* TagName = "TedsScratchBuffer";
 	
-			using Allocator = TBlockAllocationLockFreeCache<BlockSize, FOsAllocator>;
+			using Allocator = TBlockAllocationLockFreeCache<BlockSize, FAlignedAllocator>;
 		};
 		using MemoryAllocator = TConcurrentLinearBulkObjectAllocator<FAllocationTag>;
 	
diff --git a/Engine/Source/Runtime/Core/Internal/IO/IoStore.h b/Engine/Source/Runtime/Core/Internal/IO/IoStore.h
--- a/Engine/Source/Runtime/Core/Internal/IO/IoStore.h
+++ b/Engine/Source/Runtime/Core/Internal/IO/IoStore.h
@@ -217,7 +217,7 @@ private:
 		static constexpr bool InlineBlockAllocation = false;
 		static constexpr const char* TagName = "IOStoreLinear";
 
-		using Allocator = TBlockAllocationLockFreeCache<BlockSize, FOsAllocator>;
+		using Allocator = TBlockAllocationLockFreeCache<BlockSize, FAlignedAllocator>;
 	};
 	using FDataAllocator = TConcurrentLinearArrayAllocator<FIOStoreAllocationTag>;
 	using FLocalAllocator = TInlineAllocator<10>;
diff --git a/Engine/Source/Runtime/Core/Private/Experimental/ConcurrentLinearAllocator.cpp b/Engine/Source/Runtime/Core/Private/Experimental/ConcurrentLinearAllocator.cpp
--- a/Engine/Source/Runtime/Core/Private/Experimental/ConcurrentLinearAllocator.cpp
+++ b/Engine/Source/Runtime/Core/Private/Experimental/ConcurrentLinearAllocator.cpp
@@ -2,21 +2,6 @@
 
 #include "Experimental/ConcurrentLinearAllocator.h"
 
-void* FOsAllocator::Malloc(SIZE_T Size, uint32 Alignment)
-{
-	if (UNLIKELY(!UE::Private::GMalloc))
-	{
-		// There is no public function to create GMalloc and this will do it for us.
-		FMemory::Free(FMemory::Malloc(0));
-		check(UE::Private::GMalloc);
-	}
-	return UE::Private::GMalloc->Malloc(Size, Alignment);
-}
-
-void FOsAllocator::Free(void* Pointer, SIZE_T Size)
-{
-	UE::Private::GMalloc->Free(Pointer);
-}
 
 FORCENOINLINE void UE::Core::Private::OnInvalidConcurrentLinearArrayAllocatorNum(int32 NewNum, SIZE_T NumBytesPerElement)
 {
diff --git a/Engine/Source/Runtime/Core/Private/Misc/OutputDeviceRedirector.cpp b/Engine/Source/Runtime/Core/Private/Misc/OutputDeviceRedirector.cpp
--- a/Engine/Source/Runtime/Core/Private/Misc/OutputDeviceRedirector.cpp
+++ b/Engine/Source/Runtime/Core/Private/Misc/OutputDeviceRedirector.cpp
@@ -41,7 +41,7 @@ struct FOutputDeviceBlockAllocationTag : FDefaultBlockAllocationTag
 {
 	static constexpr const char* TagName = "OutputDeviceLinear";
 
-	using Allocator = FOsAllocator;
+	using Allocator = FAlignedAllocator;
 };
 
 struct FOutputDeviceLinearAllocator
diff --git a/Engine/Source/Runtime/Core/Public/Experimental/ConcurrentLinearAllocator.h b/Engine/Source/Runtime/Core/Public/Experimental/ConcurrentLinearAllocator.h
--- a/Engine/Source/Runtime/Core/Public/Experimental/ConcurrentLinearAllocator.h
+++ b/Engine/Source/Runtime/Core/Public/Experimental/ConcurrentLinearAllocator.h
@@ -4,20 +4,12 @@
 #include <atomic>
 #include "HAL/UnrealMemory.h"
 #include "HAL/LowLevelMemTracker.h"
+#include "HAL/MallocBinnedCommon.h"
 #include "Templates/UniquePtr.h"
 #include "Templates/UnrealTypeTraits.h"
 #include "Containers/LockFreeFixedSizeAllocator.h"
 #include "Misc/MemStack.h"
 
-#ifdef USE_MALLOC_BINNED3
-#	if USE_MALLOC_BINNED3
-#		define SUPPORTS_VERY_LARGE_ALIGNMENTS 0
-#	else
-#		define SUPPORTS_VERY_LARGE_ALIGNMENTS 1
-#	endif
-#else
-#	define SUPPORTS_VERY_LARGE_ALIGNMENTS 1
-#endif
 
 #if PLATFORM_HAS_ASAN_INCLUDE
 #	include <sanitizer/asan_interface.h>
@@ -38,54 +30,28 @@
 #	define IS_ASAN_ENABLED 0
 #endif
 
-struct FOsAllocator
-{
-	static constexpr bool SupportsAlignment = false;
-	static constexpr uint32 MaxAlignment	= 16;
-
-	// Indicates whether memory is acquired from the same source as used by FMalloc i.e. FMemory/FPlatformMemory::BaseAllocator().
-	// It's needed so that ConcurrentLinearAllocator could correctly account for this memory in the memory stats and tracking systems.
-	// It is 'true' by default, which is suitable for most implementations. However, if you create a custom allocator that
-	// acquires memory in a different way e.g. directly from the OS using FPlatformMemory::BinnedAllocFromOS, then you
-	// should set this variable to 'false'.
-	static constexpr bool UsesFMalloc = true;
-
-	CORE_API static void* Malloc(SIZE_T Size, uint32 Alignment);
-	CORE_API static void Free(void* Pointer, SIZE_T Size);
-};
 
 struct FAlignedAllocator
 {
 	static constexpr bool SupportsAlignment = true;
 	static constexpr bool UsesFMalloc       = true;
-	static constexpr uint32 MaxAlignment	= 256;
+	static constexpr uint32 MaxAlignment	= UE_MBC_MAX_SMALL_POOL_ALIGNMENT;
 
 	FORCEINLINE static void* Malloc(SIZE_T Size, uint32 Alignment)
 	{
-		Alignment = FMath::Max(Alignment, (uint32)sizeof(void*));
-		Size += Alignment;
-		uint8* Result = nullptr;
-		uint8* Ptr = (uint8*)FMemory::Malloc(Size);
-		if (Ptr)
-		{
-			Result = Align(Ptr + sizeof(void*), Alignment);
-			*((void**)(Result - sizeof(void*))) = Ptr;
-		}
-
-		return Result;
+		return FMemory::Malloc(Size, Alignment);
 	}
 
 	FORCEINLINE static void Free(void* Pointer, SIZE_T Size)
 	{
-		if (Pointer)
-		{
-			void* OriginalPtr = *((void**)((uint8*)Pointer - sizeof(void*)));
-			FMemory::Free(OriginalPtr);
-		}
+		FMemory::Free(Pointer);
 	}
 };
 
-template<uint32 BlockSize, typename Allocator = FOsAllocator>
+UE_DEPRECATED(5.6, "FOsAllocator is deprecated, use FAlignedAllocator instead")
+typedef FAlignedAllocator FOsAllocator;
+
+template<uint32 BlockSize, typename Allocator = FAlignedAllocator>
 class TBlockAllocationCache
 {
 	struct FTLSCleanup
@@ -141,7 +107,7 @@ public:
 	}
 };
 
-template <uint32 BlockSize, typename Allocator = FOsAllocator>
+template <uint32 BlockSize, typename Allocator = FAlignedAllocator>
 class TBlockAllocationLockFreeCache
 {
 public:
@@ -183,7 +149,7 @@ struct FDefaultBlockAllocationTag
 	static constexpr bool InlineBlockAllocation = false;  // Inline or Noinline the BlockAllocation which can have an impact on Performance
 	static constexpr const char* TagName = "DefaultLinear";
 
-	using Allocator = TBlockAllocationLockFreeCache<BlockSize, FOsAllocator>;
+	using Allocator = TBlockAllocationLockFreeCache<BlockSize, FAlignedAllocator>;
 };
 
 // This concurrent fast linear Allocator can be used for temporary allocations that are usually produced and consumed on different threads and within the lifetime of a frame.
@@ -197,7 +163,7 @@ struct FDefaultBlockAllocationTag
 template<typename BlockAllocationTag>
 class TConcurrentLinearAllocator
 {
-	static constexpr bool SupportsFastPath = ((BlockAllocationTag::BlockSize <= (64 * 1024)) || SUPPORTS_VERY_LARGE_ALIGNMENTS) //MallocBinned3 only supports Alignment up to (64 * 1024)
+	static constexpr bool SupportsFastPath = ((BlockAllocationTag::BlockSize <= (64 * 1024)) || !USE_MALLOC_BINNED3) //MallocBinned3 doesn't guarantee 64KB alignment
 		&& FMath::IsPowerOfTwo(BlockAllocationTag::BlockSize) //Aligndown only works with Pow2
 		&& !IS_ASAN_ENABLED && !BlockAllocationTag::RequiresAccurateSize // Only enabled when not using ASAN or when the Size of an allocation is not required 
 		&& BlockAllocationTag::Allocator::SupportsAlignment; //The allocator needs to align the BlockHeader by Blocksize to find the BlockHeader
diff --git a/Engine/Source/Runtime/Core/Public/HAL/MallocBinnedCommon.h b/Engine/Source/Runtime/Core/Public/HAL/MallocBinnedCommon.h
--- a/Engine/Source/Runtime/Core/Public/HAL/MallocBinnedCommon.h
+++ b/Engine/Source/Runtime/Core/Public/HAL/MallocBinnedCommon.h
@@ -45,7 +45,7 @@
 
 #define UE_MBC_MIN_SMALL_POOL_ALIGNMENT					16	// Minimum alignment of bins
 #define UE_MBC_MIN_SMALL_POOL_ALIGNMENT_SHIFT			4	// Alignment of bins, expressed as a shift
-#define UE_MBC_MAX_SMALL_POOL_ALIGNMENT					128
+#define UE_MBC_MAX_SMALL_POOL_ALIGNMENT					256
 
 #if !defined(AGGRESSIVE_MEMORY_SAVING)
 #	error "AGGRESSIVE_MEMORY_SAVING must be defined"
diff --git a/Engine/Source/Runtime/RHI/Public/RHIAllocators.h b/Engine/Source/Runtime/RHI/Public/RHIAllocators.h
--- a/Engine/Source/Runtime/RHI/Public/RHIAllocators.h
+++ b/Engine/Source/Runtime/RHI/Public/RHIAllocators.h
@@ -17,7 +17,7 @@ struct FCommandListBaseBlockAllocationTag
 	static constexpr bool InlineBlockAllocation = false;  // Inline or Noinline the BlockAllocation which can have an impact on Performance
 	static constexpr const char* TagName = "RHICommandListBaseAllocator";
 
-	using Allocator = TBlockAllocationLockFreeCache<BlockSize, FOsAllocator>;
+	using Allocator = TBlockAllocationLockFreeCache<BlockSize, FAlignedAllocator>;
 };
 
 using FRHICmdListBaseArrayAllocator = TConcurrentLinearArrayAllocator<FCommandListBaseBlockAllocationTag>;
diff --git a/Engine/Source/Runtime/RenderCore/Public/RendererInterface.h b/Engine/Source/Runtime/RenderCore/Public/RendererInterface.h
--- a/Engine/Source/Runtime/RenderCore/Public/RendererInterface.h
+++ b/Engine/Source/Runtime/RenderCore/Public/RendererInterface.h
@@ -60,7 +60,7 @@ struct FSceneRenderingBlockAllocationTag
 	static constexpr bool InlineBlockAllocation = false;  // Inline or Noinline the BlockAllocation which can have an impact on Performance
 	static constexpr const char* TagName = "SceneRenderingAllocator";
 
-	using Allocator = TBlockAllocationLockFreeCache<BlockSize, FOsAllocator>;
+	using Allocator = TBlockAllocationLockFreeCache<BlockSize, FAlignedAllocator>;
 };
 
 using FSceneRenderingBulkObjectAllocator = TConcurrentLinearBulkObjectAllocator<FSceneRenderingBlockAllocationTag>;
