Removed MallocAnsi from the ParkingLot

diff --git a/Engine/Source/Runtime/Core/Private/Async/ParkingLot.cpp b/Engine/Source/Runtime/Core/Private/Async/ParkingLot.cpp
--- a/Engine/Source/Runtime/Core/Private/Async/ParkingLot.cpp
+++ b/Engine/Source/Runtime/Core/Private/Async/ParkingLot.cpp
@@ -7,8 +7,6 @@
 #include "Async/WordMutex.h"
 #include "Containers/Array.h"
 #include "Containers/ArrayView.h"
-#include "HAL/Allocators/AnsiAllocator.h"
-#include "HAL/MallocAnsi.h"
 #include "HAL/PlatformManualResetEvent.h"
 #include "Templates/RefCounting.h"
 #include <atomic>
@@ -18,41 +16,6 @@ namespace UE::ParkingLot::Private
 
 ///////////////////////////////////////////////////////////////////////////////////////////////////////////////////////
 
-/** A low-level allocator that maintains a free list for a given object type. Memory is never released. */
-template <typename ObjectType>
-class TObjectAllocator
-{
-	static_assert(sizeof(ObjectType) >= sizeof(void*));
-	static_assert(alignof(ObjectType) >= alignof(void*));
-
-public:
-	static void* Malloc(SIZE_T Size = sizeof(ObjectType), uint32 Alignment = alignof(ObjectType))
-	{
-		if (TUniqueLock Lock(FreeListMutex); void* Mem = FreeListHead)
-		{
-			FreeListHead = *(void**)Mem;
-			return Mem;
-		}
-		constexpr SIZE_T MinSize = sizeof(ObjectType);
-		constexpr uint32 MinAlignment = alignof(ObjectType);
-		return AnsiMalloc(FPlatformMath::Max(Size, MinSize), FPlatformMath::Max(Alignment, MinAlignment));
-	}
-
-	static void Free(void* Mem)
-	{
-		TUniqueLock Lock(FreeListMutex);
-		void* Next = FreeListHead;
-		*(void**)Mem = Next;
-		FreeListHead = Mem;
-	}
-
-private:
-	inline static FWordMutex FreeListMutex;
-	inline static void* FreeListHead = nullptr;
-};
-
-///////////////////////////////////////////////////////////////////////////////////////////////////////////////////////
-
 /**
  * A thread as stored in the wait queue.
  */
@@ -73,28 +36,17 @@ struct alignas(PLATFORM_CACHE_LINE_SIZE) FThread final
 	{
 		if (ReferenceCount.fetch_sub(1, std::memory_order_acq_rel) == 1)
 		{
-			this->~FThread();
-			Allocator.Free(this);
+			delete this;
 		}
 	}
 
-	inline static FThread* New()
-	{
-		return new(Allocator.Malloc(sizeof(FThread), alignof(FThread))) FThread;
-	}
-
-private:
 	FThread() = default;
 	~FThread() = default;
 
 	FThread(const FThread&) = delete;
 	FThread& operator=(const FThread&) = delete;
-
-	static TObjectAllocator<FThread> Allocator;
 };
 
-TObjectAllocator<FThread> FThread::Allocator;
-
 class FThreadLocalData
 {
 public:
@@ -104,8 +56,6 @@ public:
 
 	static TRefCountPtr<FThread> Get()
 	{
-        // This static will appear as memory leaks when comparing AnsiMalloc and AnsiFree
-        // But since this only grows as new threads created, it won't grow indefinitely
 		static thread_local FThreadLocalData ThreadLocalData;
 		FThreadLocalData& LocalThreadLocalData = ThreadLocalData;
 		if (LIKELY(LocalThreadLocalData.Thread))
@@ -114,9 +64,9 @@ public:
 		}
 		if (LocalThreadLocalData.bDestroyed)
 		{
-			return FThread::New();
+			return new FThread;
 		}
-		LocalThreadLocalData.Thread = FThread::New();
+		LocalThreadLocalData.Thread = new FThread;
 		return LocalThreadLocalData.Thread;
 	}
 
@@ -154,17 +104,6 @@ enum class EQueueAction
 class alignas(PLATFORM_CACHE_LINE_SIZE) FBucket final
 {
 public:
-	static FBucket* Create()
-	{
-		return new(Allocator.Malloc()) FBucket;
-	}
-
-	static void Destroy(FBucket* Bucket)
-	{
-		Bucket->~FBucket();
-		Allocator.Free(Bucket);
-	}
-
 	[[nodiscard]] inline TDynamicUniqueLock<FWordMutex> LockDynamic() { return TDynamicUniqueLock(Mutex); }
 
 	inline void Lock() { Mutex.Lock(); }
@@ -186,22 +125,18 @@ public:
 	template <typename VisitorType>
 	void DequeueIf(VisitorType&& Visitor);
 
-private:
 	FBucket() = default;
 	~FBucket() = default;
 
+private:
 	FBucket(const FBucket&) = delete;
 	FBucket& operator=(const FBucket&) = delete;
 
 	FWordMutex Mutex;
 	FThread* Head = nullptr;
 	FThread* Tail = nullptr;
-
-	static TObjectAllocator<FBucket> Allocator;
 };
 
-TObjectAllocator<FBucket> FBucket::Allocator;
-
 void FBucket::Enqueue(FThread* Thread)
 {
 	checkSlow(Thread);
@@ -304,7 +239,7 @@ private:
 	static void Destroy(FTable& Table);
 
 	/** Try to lock the whole table by locking every one of its buckets. */
-	static bool TryLock(FTable& Table, TArray<FBucket*, FAnsiAllocator>& OutBuckets);
+	static bool TryLock(FTable& Table, TArray<FBucket*>& OutBuckets);
 	/** Unlock an array of buckets that was filled by TryLock. */
 	static void Unlock(TConstArrayView<FBucket*> LockedBuckets);
 
@@ -336,7 +271,7 @@ FBucket& FTable::FindOrCreateBucket(const void* Address, TDynamicUniqueLock<FWor
 	{
 		FTable& Table = CreateOrGet();
 		const uint32 Index = Hash % Table.BucketCount;
-		FBucket& Bucket = Table.FindOrCreateBucket(Index, [] { return FBucket::Create(); });
+		FBucket& Bucket = Table.FindOrCreateBucket(Index, [] { return new FBucket; });
 		OutLock = Bucket.LockDynamic();
 
 		if (LIKELY(&Table == GlobalTable.load(std::memory_order_acquire)))
@@ -390,7 +325,7 @@ FBucket& FTable::FindOrCreateBucket(uint32 Index, AllocatorType&& BucketAllocato
 		}
 		else
 		{
-			FBucket::Destroy(NewBucket);
+			delete NewBucket;
 		}
 		checkSlow(Bucket);
 	}
@@ -422,7 +357,7 @@ FTable& FTable::CreateOrGet()
 FTable& FTable::Create(const uint32 Size)
 {
 	const SIZE_T MemorySize = sizeof(FTable) + sizeof(FBucket*) * SIZE_T(Size);
-	void* const Memory = AnsiMalloc(MemorySize, alignof(FTable));
+	void* const Memory = FMemory::Malloc(MemorySize, alignof(FTable));
 	FMemory::Memzero(Memory, MemorySize);
 	FTable& Table = *new(Memory) FTable;
 	Table.BucketCount = Size;
@@ -432,13 +367,13 @@ FTable& FTable::Create(const uint32 Size)
 void FTable::Destroy(FTable& Table)
 {
 	Table.~FTable();
-	AnsiFree(&Table);
+	FMemory::Free(&Table);
 }
 
 void FTable::Reserve(const uint32 ThreadCount)
 {
 	const uint32 TargetBucketCount = FMath::RoundUpToPowerOfTwo(ThreadCount);
-	TArray<FBucket*, FAnsiAllocator> ExistingBuckets;
+	TArray<FBucket*> ExistingBuckets;
 
 	for (;;)
 	{
@@ -458,7 +393,7 @@ void FTable::Reserve(const uint32 ThreadCount)
 
 		// Gather waiting threads to be redistributed into the buckets of the new table.
 		// Threads with the same address remain in the same relative order as they were queued.
-		TArray<FThread*, FAnsiAllocator> Threads;
+		TArray<FThread*> Threads;
 		for (FBucket* Bucket : ExistingBuckets)
 		{
 			while (FThread* Thread = Bucket->Dequeue())
@@ -470,10 +405,10 @@ void FTable::Reserve(const uint32 ThreadCount)
 		FTable& NewTable = Create(TargetBucketCount);
 
 		// Reuse existing now-empty buckets when populating the new table.
-		TArray<FBucket*, FAnsiAllocator> AvailableBuckets = ExistingBuckets;
+		TArray<FBucket*> AvailableBuckets = ExistingBuckets;
 		const auto AllocateBucket = [&AvailableBuckets]() -> FBucket*
 		{
-			return !AvailableBuckets.IsEmpty() ? AvailableBuckets.Pop(EAllowShrinking::No) : FBucket::Create();
+			return !AvailableBuckets.IsEmpty() ? AvailableBuckets.Pop(EAllowShrinking::No) : new FBucket;
 		};
 
 		// Add waiting threads to the new table.
@@ -502,14 +437,14 @@ void FTable::Reserve(const uint32 ThreadCount)
 	}
 }
 
-bool FTable::TryLock(FTable& Table, TArray<FBucket*, FAnsiAllocator>& OutBuckets)
+bool FTable::TryLock(FTable& Table, TArray<FBucket*>& OutBuckets)
 {
 	OutBuckets.Reset(Table.BucketCount);
 
 	// Gather buckets from the table, creating them as needed because the lock is on the bucket.
 	for (uint32 Index = 0; Index < Table.BucketCount; ++Index)
 	{
-		OutBuckets.Add(&Table.FindOrCreateBucket(Index, [] { return FBucket::Create(); }));
+		OutBuckets.Add(&Table.FindOrCreateBucket(Index, [] { return new FBucket; }));
 	}
 
 	// Lock the buckets in order by address to ensure consistent ordering regardless of the table being locked.
@@ -700,7 +635,7 @@ uint32 WakeMultiple(const void* const Address, const uint32 WakeCount)
 {
 	using namespace Private;
 
-	TArray<TRefCountPtr<FThread>, TInlineAllocator<128, FAnsiAllocator>> WakeThreads;
+	TArray<TRefCountPtr<FThread>, TInlineAllocator<128>> WakeThreads;
 
 	if (TDynamicUniqueLock<FWordMutex> BucketLock; FBucket* Bucket = FTable::FindBucket(Address, BucketLock))
 	{
