IncrementalCook: ConfigAccessTracking: Record the configfiles added through GConfig->Find as LocalSingleIniFile. This allows them to be loaded and compared later even if removed from GConfig, and some systems want to remove them from GConfig because they are only needed for a single call to UObject::LoadConfig.
#rnx
#incrementalcook

diff --git a/Engine/Source/Runtime/Core/Private/Misc/ConfigCacheIni.cpp b/Engine/Source/Runtime/Core/Private/Misc/ConfigCacheIni.cpp
--- a/Engine/Source/Runtime/Core/Private/Misc/ConfigCacheIni.cpp
+++ b/Engine/Source/Runtime/Core/Private/Misc/ConfigCacheIni.cpp
@@ -3869,10 +3869,18 @@ FConfigFile* FConfigCacheIni::Find(const FString& Filename)
 				Result = &Add(UnrealFileName, FConfigFile());
 				UE_LOG(LogConfig, Verbose, TEXT("GConfig::Find is looking for file:  %s"), *UnrealFileName);
 				{
-					Result->Read(UnrealFileName);
 					// Files added through Find are treated the same as ReadSingleIntoConfigSystem contexts,
 					// and do not use a hierarchy so they do not use a generatedini and should never be saved.
 					Result->NoSave = true;
+#if UE_WITH_CONFIG_TRACKING
+					Result->LoadType = UE::ConfigAccessTracking::ELoadType::LocalSingleIniFile;
+					UE::ConfigAccessTracking::FFile* FileAccess = Result->GetFileAccess();
+					if (FileAccess)
+					{
+						FileAccess->OverrideFilenameToLoad = FName(FStringView(UnrealFileName));
+					}
+#endif
+					Result->Read(UnrealFileName);
 					UE_LOG(LogConfig, Verbose, TEXT("GConfig::Find has loaded file:  %s"), *UnrealFileName);
 				}
 			}
