[Core/Misc]

If Ensures have been configured to be warnings instead of errors, don't send error reports since we are logging warnings, not errors.

#rnx

diff --git a/Engine/Source/Runtime/Core/Private/Misc/AssertionMacros.cpp b/Engine/Source/Runtime/Core/Private/Misc/AssertionMacros.cpp
--- a/Engine/Source/Runtime/Core/Private/Misc/AssertionMacros.cpp
+++ b/Engine/Source/Runtime/Core/Private/Misc/AssertionMacros.cpp
@@ -549,6 +549,7 @@ FORCENOINLINE void FDebug::EnsureFailed(const ANSICHAR* Expr, const ANSICHAR* Fi
 #endif
 			
 			// Submit the error report to the server! (and display a balloon in the system tray)
+			if (GEnsuresAreErrors)
 			{
 				// How many unique previous errors we should keep track of
 				const uint32 MaxPreviousErrorsToTrack = 4;
@@ -600,7 +601,7 @@ FORCENOINLINE void FDebug::EnsureFailed(const ANSICHAR* Expr, const ANSICHAR* Fi
 		{
 			// If we fail to generate the string to identify the crash we don't know if we should skip sending the report,
 			// so we will just send the report anyway.
-			bShouldSendNewReport = true;
+			bShouldSendNewReport = !!GEnsuresAreErrors;
 
 			// Add message to log even without stacktrace. It is useful for testing fail on ensure.
 #if !NO_LOGGING
