[HAL]

Adds forwarding implementations for wrapper IPlatformFile type's (FileLogWrapper, FileOpenLogWrapper, FileCachedWrapper)  journal methods so they can forward support for journaling to the lower level implementation being wrapped.

These wrapper types can be enabled via commandline options and were subtly disabling asset registry journal caching support by not forwarding the implementation.

The journal methods likely should have been made pure virtual originally however since that would be a breaking change (and annoying since most platforms don't have journaling support), it is now up to implementers to keep an eye on new methods added to the IPlatformFile interface and ensure they forward appropriately, which is a bit unfortunate.

#rnx

diff --git a/Engine/Source/Runtime/Core/Public/HAL/IPlatformFileCachedWrapper.h b/Engine/Source/Runtime/Core/Public/HAL/IPlatformFileCachedWrapper.h
--- a/Engine/Source/Runtime/Core/Public/HAL/IPlatformFileCachedWrapper.h
+++ b/Engine/Source/Runtime/Core/Public/HAL/IPlatformFileCachedWrapper.h
@@ -464,5 +464,41 @@ public:
 	{
 		LowerLevel->SetAsyncMinimumPriority(MinPriority);
 	}
-
+	virtual bool		FileJournalIsAvailable(const TCHAR* VolumeOrPath = nullptr,
+		ELogVerbosity::Type* OutErrorLevel = nullptr, FString* OutError = nullptr) override
+	{
+		return LowerLevel->FileJournalIsAvailable(VolumeOrPath, OutErrorLevel, OutError);
+	}
+	virtual bool		FileJournalIterateDirectory(const TCHAR* Directory,
+		FDirectoryJournalVisitorFunc Visitor, FString* OutError = nullptr) override
+	{
+		return LowerLevel->FileJournalIterateDirectory(Directory, Visitor, OutError);
+	}
+	virtual EFileJournalResult FileJournalGetLatestEntry(const TCHAR* VolumeName, FFileJournalId& OutJournalId,
+		FFileJournalEntryHandle& OutEntryHandle, FString* OutError = nullptr) override
+	{
+		return LowerLevel->FileJournalGetLatestEntry(VolumeName, OutJournalId, OutEntryHandle, OutError);
+	}
+	virtual EFileJournalResult FileJournalReadModified(const TCHAR* VolumeName,
+		const FFileJournalId& JournalIdOfStartingEntry, const FFileJournalEntryHandle& StartingJournalEntry,
+		TMap<FFileJournalFileHandle, FString>& KnownDirectories, TSet<FString>& OutModifiedDirectories,
+		FFileJournalEntryHandle& OutNextJournalEntry, FString* OutError = nullptr) override
+	{
+		return LowerLevel->FileJournalReadModified(VolumeName, JournalIdOfStartingEntry, StartingJournalEntry,
+			KnownDirectories, OutModifiedDirectories, OutNextJournalEntry, OutError);
+	}
+	virtual FFileJournalData FileJournalGetFileData(const TCHAR* FilenameOrDirectory,
+		FString* OutError = nullptr) override
+	{
+		return LowerLevel->FileJournalGetFileData(FilenameOrDirectory, OutError);
+	}
+	virtual FString		FileJournalGetVolumeName(FStringView InPath) override
+	{
+		return LowerLevel->FileJournalGetVolumeName(InPath);
+	}
+	virtual uint64		FileJournalGetMaximumSize(const TCHAR* VolumeOrPath = nullptr,
+		ELogVerbosity::Type* OutErrorLevel = nullptr, FString* OutError = nullptr) const override
+	{
+		return LowerLevel->FileJournalGetMaximumSize(VolumeOrPath, OutErrorLevel, OutError);
+	}
 };
diff --git a/Engine/Source/Runtime/Core/Public/HAL/IPlatformFileLogWrapper.h b/Engine/Source/Runtime/Core/Public/HAL/IPlatformFileLogWrapper.h
--- a/Engine/Source/Runtime/Core/Public/HAL/IPlatformFileLogWrapper.h
+++ b/Engine/Source/Runtime/Core/Public/HAL/IPlatformFileLogWrapper.h
@@ -509,5 +509,41 @@ public:
 	{
 		LowerLevel->SetAsyncMinimumPriority(Priority);
 	}
-
+	virtual bool		FileJournalIsAvailable(const TCHAR* VolumeOrPath = nullptr,
+		ELogVerbosity::Type* OutErrorLevel = nullptr, FString* OutError = nullptr) override
+	{
+		return LowerLevel->FileJournalIsAvailable(VolumeOrPath, OutErrorLevel, OutError);
+	}
+	virtual bool		FileJournalIterateDirectory(const TCHAR* Directory,
+		FDirectoryJournalVisitorFunc Visitor, FString* OutError = nullptr) override
+	{
+		return LowerLevel->FileJournalIterateDirectory(Directory, Visitor, OutError);
+	}
+	virtual EFileJournalResult FileJournalGetLatestEntry(const TCHAR* VolumeName, FFileJournalId& OutJournalId,
+		FFileJournalEntryHandle& OutEntryHandle, FString* OutError = nullptr) override
+	{
+		return LowerLevel->FileJournalGetLatestEntry(VolumeName, OutJournalId, OutEntryHandle, OutError);
+	}
+	virtual EFileJournalResult FileJournalReadModified(const TCHAR* VolumeName,
+		const FFileJournalId& JournalIdOfStartingEntry, const FFileJournalEntryHandle& StartingJournalEntry,
+		TMap<FFileJournalFileHandle, FString>& KnownDirectories, TSet<FString>& OutModifiedDirectories,
+		FFileJournalEntryHandle& OutNextJournalEntry, FString* OutError = nullptr) override
+	{
+		return LowerLevel->FileJournalReadModified(VolumeName, JournalIdOfStartingEntry, StartingJournalEntry,
+			KnownDirectories, OutModifiedDirectories, OutNextJournalEntry, OutError);
+	}
+	virtual FFileJournalData FileJournalGetFileData(const TCHAR* FilenameOrDirectory,
+		FString* OutError = nullptr) override
+	{
+		return LowerLevel->FileJournalGetFileData(FilenameOrDirectory, OutError);
+	}
+	virtual FString		FileJournalGetVolumeName(FStringView InPath) override
+	{
+		return LowerLevel->FileJournalGetVolumeName(InPath);
+	}
+	virtual uint64		FileJournalGetMaximumSize(const TCHAR* VolumeOrPath = nullptr,
+		ELogVerbosity::Type* OutErrorLevel = nullptr, FString* OutError = nullptr) const override
+	{
+		return LowerLevel->FileJournalGetMaximumSize(VolumeOrPath, OutErrorLevel, OutError);
+	}
 };
diff --git a/Engine/Source/Runtime/Core/Public/HAL/IPlatformFileOpenLogWrapper.h b/Engine/Source/Runtime/Core/Public/HAL/IPlatformFileOpenLogWrapper.h
--- a/Engine/Source/Runtime/Core/Public/HAL/IPlatformFileOpenLogWrapper.h
+++ b/Engine/Source/Runtime/Core/Public/HAL/IPlatformFileOpenLogWrapper.h
@@ -218,6 +218,43 @@ public:
 	{
 		return LowerLevel->SendMessageToServer(Message, Handler);
 	}
+	virtual bool		FileJournalIsAvailable(const TCHAR* VolumeOrPath = nullptr,
+		ELogVerbosity::Type* OutErrorLevel = nullptr, FString* OutError = nullptr) override
+	{
+		return LowerLevel->FileJournalIsAvailable(VolumeOrPath, OutErrorLevel, OutError);
+	}
+	virtual bool		FileJournalIterateDirectory(const TCHAR* Directory,
+		FDirectoryJournalVisitorFunc Visitor, FString* OutError = nullptr) override
+	{
+		return LowerLevel->FileJournalIterateDirectory(Directory, Visitor, OutError);
+	}
+	virtual EFileJournalResult FileJournalGetLatestEntry(const TCHAR* VolumeName, FFileJournalId& OutJournalId,
+		FFileJournalEntryHandle& OutEntryHandle, FString* OutError = nullptr) override
+	{
+		return LowerLevel->FileJournalGetLatestEntry(VolumeName, OutJournalId, OutEntryHandle, OutError);
+	}
+	virtual EFileJournalResult FileJournalReadModified(const TCHAR* VolumeName,
+		const FFileJournalId& JournalIdOfStartingEntry, const FFileJournalEntryHandle& StartingJournalEntry,
+		TMap<FFileJournalFileHandle, FString>& KnownDirectories, TSet<FString>& OutModifiedDirectories,
+		FFileJournalEntryHandle& OutNextJournalEntry, FString* OutError = nullptr) override
+	{
+		return LowerLevel->FileJournalReadModified(VolumeName, JournalIdOfStartingEntry, StartingJournalEntry,
+			KnownDirectories, OutModifiedDirectories, OutNextJournalEntry, OutError);
+	}
+	virtual FFileJournalData FileJournalGetFileData(const TCHAR* FilenameOrDirectory,
+		FString* OutError = nullptr) override
+	{
+		return LowerLevel->FileJournalGetFileData(FilenameOrDirectory, OutError);
+	}
+	virtual FString		FileJournalGetVolumeName(FStringView InPath) override
+	{
+		return LowerLevel->FileJournalGetVolumeName(InPath);
+	}
+	virtual uint64		FileJournalGetMaximumSize(const TCHAR* VolumeOrPath = nullptr,
+		ELogVerbosity::Type* OutErrorLevel = nullptr, FString* OutError = nullptr) const override
+	{
+		return LowerLevel->FileJournalGetMaximumSize(VolumeOrPath, OutErrorLevel, OutError);
+	}
 	virtual IAsyncReadFileHandle* OpenAsyncRead(const TCHAR* Filename, bool bAllowWrite = false) override
 	{
 		// we must not record the "open" here...what matters is when we start reading the file!
