Added a fatal when attempting to build SM6 shaders and not using XCode 16

diff --git a/Engine/Source/Developer/Apple/MetalShaderFormat/Private/MetalShaderCompiler.cpp b/Engine/Source/Developer/Apple/MetalShaderFormat/Private/MetalShaderCompiler.cpp
--- a/Engine/Source/Developer/Apple/MetalShaderFormat/Private/MetalShaderCompiler.cpp
+++ b/Engine/Source/Developer/Apple/MetalShaderFormat/Private/MetalShaderCompiler.cpp
@@ -524,6 +524,15 @@ void BuildMetalShaderOutput(
     if (bUseMetalShaderConverter)
     {
 #if PLATFORM_MAC || PLATFORM_WINDOWS
+
+	#if PLATFORM_MAC
+		// Xcode 16 is the minimum supported version for MSC
+		if(FPlatformMisc::XCodeVersionCompare(16, 0, 0) < 0)
+		{
+			UE_LOG(LogMetalShaderCompiler, Fatal, TEXT("Compiling shaders for Metal SM6 requires XCode 16.0 or above"));
+		}
+		
+	#endif
 		EnumAddFlags(Header.Bindings.Flags, EMetalBindingsFlags::UseMetalShaderConverter);
 				
 		// Only needed for VS Input (to generate the stage-in function used to convert inputs).
diff --git a/Engine/Source/Runtime/Core/Private/Mac/MacPlatformMisc.cpp b/Engine/Source/Runtime/Core/Private/Mac/MacPlatformMisc.cpp
--- a/Engine/Source/Runtime/Core/Private/Mac/MacPlatformMisc.cpp
+++ b/Engine/Source/Runtime/Core/Private/Mac/MacPlatformMisc.cpp
@@ -1849,6 +1849,26 @@ FString FMacPlatformMisc::GetXcodePath()
 	return GMacAppInfo.XcodePath;
 }
 
+int32 FMacPlatformMisc::XCodeVersionCompare(uint8 Major, uint8 Minor, uint8 Revision)
+{
+	uint8 TargetValues[3] = {Major, Minor, Revision};
+	NSInteger ComponentValues[3] = {GMacAppInfo.XcodeVersion.majorVersion, GMacAppInfo.XcodeVersion.minorVersion, GMacAppInfo.XcodeVersion.patchVersion};
+	
+	for(uint32 i = 0; i < 3; i++)
+	{
+		if(ComponentValues[i] < TargetValues[i])
+		{
+			return -1;
+		}
+		else if(ComponentValues[i] > TargetValues[i])
+		{
+			return 1;
+		}
+	}
+	
+	return 0;
+}
+
 bool FMacPlatformMisc::IsSupportedXcodeVersionInstalled()
 {
 	// We need Xcode 9.4 or newer to be able to compile Metal shaders correctly
diff --git a/Engine/Source/Runtime/Core/Public/Mac/MacPlatformMisc.h b/Engine/Source/Runtime/Core/Public/Mac/MacPlatformMisc.h
--- a/Engine/Source/Runtime/Core/Public/Mac/MacPlatformMisc.h
+++ b/Engine/Source/Runtime/Core/Public/Mac/MacPlatformMisc.h
@@ -114,6 +114,11 @@ struct CORE_API FMacPlatformMisc : public FApplePlatformMisc
 
 	static FString GetXcodePath();
 
+	/**
+	 * Returns if current < target returns 1, if current > target returns 1, else current == target and returns 0.
+	 */
+	static int32 XCodeVersionCompare(uint8 Major, uint8 Minor, uint8 Revision);
+	
 	static bool IsSupportedXcodeVersionInstalled();
 
 #if WITH_EDITOR
