[StrongObjectPtr]
- Simplify move-constructor

diff --git a/Engine/Source/Runtime/Core/Public/UObject/StrongObjectPtrTemplates.h b/Engine/Source/Runtime/Core/Public/UObject/StrongObjectPtrTemplates.h
--- a/Engine/Source/Runtime/Core/Public/UObject/StrongObjectPtrTemplates.h
+++ b/Engine/Source/Runtime/Core/Public/UObject/StrongObjectPtrTemplates.h
@@ -28,9 +28,8 @@ public:
 	using ElementType = ObjectType;
 
 	[[nodiscard]] TStrongObjectPtr(TStrongObjectPtr&& InOther)
+		: Object(InOther.Object)
 	{
-		Reset();
-		Object = InOther.Object;
 		InOther.Object = nullptr;
 	}
 
