[EventCount]
- Add a safe by default barrier in Notify even for x86 to avoid load-store reordering if the value we're synchronizing with doesn't provide its own barrier
- Add another NotifyWeak function that doesn't provide a barrier and can be used with strongly ordered external values
- Modify the taskgraph to use the weak version since it already has strong ordering
- Slight optimization of PrepareWait to avoid a CAS on x86

diff --git a/Engine/Source/Runtime/Core/Private/Tasks/TaskPrivate.cpp b/Engine/Source/Runtime/Core/Private/Tasks/TaskPrivate.cpp
--- a/Engine/Source/Runtime/Core/Private/Tasks/TaskPrivate.cpp
+++ b/Engine/Source/Runtime/Core/Private/Tasks/TaskPrivate.cpp
@@ -36,7 +36,7 @@ namespace UE::Tasks
 			
 			// In case a thread is waiting on us to perform retraction, now is the time to try retraction again.
 			// This needs to be before the launch as performing the execution can destroy the task.
-			StateChangeEvent.Notify();
+			StateChangeEvent.NotifyWeak();
 
 			// This needs to be the last line touching any of the task's properties.
 			bWakeUpWorker |= LowLevelTasks::FScheduler::Get().TryLaunch(LowLevelTask, bWakeUpWorker ? LowLevelTasks::EQueuePreference::GlobalQueuePreference : LowLevelTasks::EQueuePreference::LocalQueuePreference, bWakeUpWorker);
diff --git a/Engine/Source/Runtime/Core/Public/Async/EventCount.h b/Engine/Source/Runtime/Core/Public/Async/EventCount.h
--- a/Engine/Source/Runtime/Core/Public/Async/EventCount.h
+++ b/Engine/Source/Runtime/Core/Public/Async/EventCount.h
@@ -48,8 +48,7 @@ private:
  *
  * On the notifying thread:
  *
- *     uint32 NewValue = CurrentValue.fetch_add(1) + 1;
- *     if (NewValue == TargetValue)
+ *     if (++CurrentValue == TargetValue)
  *     {
  *         Event.Notify();
  *     }
@@ -78,7 +77,17 @@ public:
 	inline TEventCountToken<CounterType> PrepareWait()
 	{
 		TEventCountToken<CounterType> Token;
+
+#if PLATFORM_CPU_X86_FAMILY
+		// Doing a relaxed load here because fetch_or on x86 cannot return the previous value
+		// so when we use the return value of fetch_or, it gets compiled into a compare_exchange.
+		// The worst that can happen here is that we get a stale token value and end up 
+		// not waiting for an iteration. We would then get the proper value on the next iteration.
+		Token.Value = Count.load(std::memory_order_relaxed) & ~CounterType(1);
+		Count.fetch_or(1, std::memory_order_acq_rel);
+#else
 		Token.Value = Count.fetch_or(1, std::memory_order_acq_rel) & ~CounterType(1);
+#endif
 		return Token;
 	}
 
@@ -141,6 +150,37 @@ public:
 	 */
 	inline void Notify()
 	{
+		//
+		// .fetch_add(0, acq_rel) is used to have a StoreLoad barrier,
+		// which we can't express in C++. That works by making the load
+		// also be store (via RMW) and relying on a StoreStore barrier to
+		// get the desired ordering.
+		//
+		// Previously, this code was:
+		//   CounterType Value = Count.load(std::memory_order_relaxed);
+		//
+		// which had a memory re-ordering and stale values being read,
+		// leading to a missed Wake and dead-locked waiter, as a result.
+		//
+		CounterType Value = Count.fetch_add(0, std::memory_order_acq_rel);
+		if ((Value & 1) && Count.compare_exchange_strong(Value, Value + 1, std::memory_order_release))
+		{
+			ParkingLot::WakeAll(&Count);
+		}
+	}
+
+	/**
+	 * Notifies all waiting threads.
+	 *
+	 * Any threads that have called PrepareWait() and not yet waited will be notified immediately
+	 * if they do wait on a token from a call to PrepareWait() that preceded this call.
+         * Note: This version doesn't provide a memory barrier, you are responsible
+         *       for the memory ordering for the value you're synchronizing this eventcount with.
+	 */
+	inline void NotifyWeak()
+	{
+		// We probably don't need the weakly consistent case anymore
+		// for this weak version of notify, but will need more tests before removing.
 #if PLATFORM_WEAKLY_CONSISTENT_MEMORY
 		//
 		// .fetch_add(0, acq_rel) is used to have a StoreLoad barrier,
diff --git a/Engine/Source/Runtime/Core/Public/Tasks/TaskPrivate.h b/Engine/Source/Runtime/Core/Public/Tasks/TaskPrivate.h
--- a/Engine/Source/Runtime/Core/Public/Tasks/TaskPrivate.h
+++ b/Engine/Source/Runtime/Core/Public/Tasks/TaskPrivate.h
@@ -551,7 +551,7 @@ public:
 				TaskTrace::Completed(GetTraceId());
 
 				// In case a thread is waiting on us to perform retraction, now is the time to try retraction again.
-				StateChangeEvent.Notify();
+				StateChangeEvent.NotifyWeak();
 			}
 
 			CORE_API void ClearPipe();
