For editor builds, adding the flag FILE_SHARE_DELETE to various CreateFile calls in order to allow the editor to save and rename assets that are open in a -game or -server process. Also added an additional trace scope for IterateDirectory.

diff --git a/Engine/Source/Runtime/Core/Private/Windows/WindowsPlatformFile.cpp b/Engine/Source/Runtime/Core/Private/Windows/WindowsPlatformFile.cpp
--- a/Engine/Source/Runtime/Core/Private/Windows/WindowsPlatformFile.cpp
+++ b/Engine/Source/Runtime/Core/Private/Windows/WindowsPlatformFile.cpp
@@ -31,9 +31,30 @@
 namespace FileConstants
 {
 uint32 WIN_INVALID_SET_FILE_POINTER = INVALID_SET_FILE_POINTER;
+
+bool GEnableFileShareDeleteFlag = false;
+static FAutoConsoleVariableRef CVarEnableFileShareDeleteFlag(
+	TEXT("file.allowdeleteopenfiles"),
+	GEnableFileShareDeleteFlag,
+	TEXT("Sets the the FILE_SHARE_DELETE flag on all file open operations so that files can be replaced while open. Only supported in editor builds."),
+	ECVF_Default
+);
+
 } // namespace FileConstants
 #include "Windows/HideWindowsPlatformTypes.h"
 
+
+static int32 GetFileShareDeleteFlagIfEnabled()
+{
+#if WITH_EDITOR
+	if (FileConstants::GEnableFileShareDeleteFlag)
+	{
+		return FILE_SHARE_DELETE;
+	}
+#endif
+	return 0;
+}
+
 TLockFreePointerListUnordered<void, PLATFORM_CACHE_LINE_SIZE> MicrosoftAsyncIOEventPool;
 bool GTriggerFailedMicrosoftRead = false;
 
@@ -1386,7 +1407,7 @@ bool FWindowsPlatformFile::SetMarkOfTheWeb(FStringView Filename, bool bNewStatus
 IAsyncReadFileHandle* FWindowsPlatformFile::OpenAsyncRead(const TCHAR* Filename, bool bAllowWrite /*= false*/)
 {
 	uint32  Access = GENERIC_READ;
-	uint32  WinFlags = FILE_SHARE_READ | (bAllowWrite ? FILE_SHARE_WRITE : 0);
+	uint32  WinFlags = FILE_SHARE_READ | (bAllowWrite ? FILE_SHARE_WRITE : 0) | GetFileShareDeleteFlagIfEnabled();
 	uint32  Create = OPEN_EXISTING;
 
 
@@ -1413,7 +1434,7 @@ IAsyncReadFileHandle* FWindowsPlatformFile::OpenAsyncRead(const TCHAR* Filename,
 FFileOpenResult FWindowsPlatformFile::OpenRead(const TCHAR* Filename, EOpenReadFlags Flags)
 {
 	const uint32  Access	= GENERIC_READ;
-	const uint32  WinFlags  =	FILE_SHARE_READ |
+	const uint32  WinFlags  =	FILE_SHARE_READ | GetFileShareDeleteFlagIfEnabled() |
 								(EnumHasAnyFlags(Flags, EOpenReadFlags::AllowWrite) ? FILE_SHARE_WRITE : 0) |
 								(EnumHasAnyFlags(Flags, EOpenReadFlags::AllowDelete) ? FILE_SHARE_DELETE : 0);
 	const uint32  Create	= OPEN_EXISTING;
@@ -1471,7 +1492,7 @@ IFileHandle* FWindowsPlatformFile::OpenRead(const TCHAR* Filename, bool bAllowWr
 IFileHandle* FWindowsPlatformFile::OpenReadNoBuffering(const TCHAR* Filename, bool bAllowWrite)
 {
 	uint32  Access = GENERIC_READ;
-	uint32  WinFlags = FILE_SHARE_READ | (bAllowWrite ? FILE_SHARE_WRITE : 0);
+	uint32  WinFlags = FILE_SHARE_READ | (bAllowWrite ? FILE_SHARE_WRITE : 0) | GetFileShareDeleteFlagIfEnabled();
 	uint32  Create = OPEN_EXISTING;
 	TRACE_PLATFORMFILE_BEGIN_OPEN(Filename);
 	HANDLE Handle = CreateFileW(*FNormalizedFilename(Filename), Access, WinFlags, NULL, Create, FILE_ATTRIBUTE_NORMAL | FILE_FLAG_OVERLAPPED, NULL);
@@ -1505,6 +1526,8 @@ IFileHandle* FWindowsPlatformFile::OpenWrite(const TCHAR* Filename, bool bAppend
 {
 	uint32  Access    = GENERIC_WRITE | (bAllowRead ? GENERIC_READ : 0);
 	uint32  WinFlags  = bAllowRead ? FILE_SHARE_READ : 0;
+	WinFlags |= GetFileShareDeleteFlagIfEnabled();
+
 	uint32  Create    = bAppend ? OPEN_ALWAYS : CREATE_ALWAYS;
 	TRACE_PLATFORMFILE_BEGIN_OPEN(Filename);
 	HANDLE Handle    = CreateFileW(*FNormalizedFilename(Filename), Access, WinFlags, NULL, Create, FILE_ATTRIBUTE_NORMAL, NULL);
@@ -1541,14 +1564,14 @@ FOpenMappedResult FWindowsPlatformFile::OpenMappedEx(const TCHAR* Filename, EOpe
 {
 	int64 Size = FileSize(Filename);
 	uint32  Access = GENERIC_READ;
-	uint32  WinFlags = FILE_SHARE_READ;
+	uint32  WinFlags = FILE_SHARE_READ | GetFileShareDeleteFlagIfEnabled();
 	uint32  Create = OPEN_EXISTING;
 	DWORD	Protect = PAGE_READONLY;
 
 	if (EnumHasAnyFlags(OpenOptions, EOpenReadFlags::AllowWrite))
 	{
 		Access = GENERIC_READ | GENERIC_WRITE;
-		WinFlags = FILE_SHARE_READ | FILE_SHARE_WRITE;
+		WinFlags = FILE_SHARE_READ | FILE_SHARE_WRITE | GetFileShareDeleteFlagIfEnabled();
 		Create = CREATE_ALWAYS;
 		Protect = PAGE_READWRITE;
 		Size = MaximumSize;
@@ -1714,6 +1737,8 @@ FORCENOINLINE HANDLE FWindowsPlatformFile::FindFirstFileWithWildcard(const TCHAR
 
 bool FWindowsPlatformFile::IterateDirectoryCommon(const TCHAR* Directory, const TFunctionRef<bool(const WIN32_FIND_DATAW&)>& Visitor)
 {
+	TRACE_CPUPROFILER_EVENT_SCOPE(FWindowsPlatformFile_IterateDirectoryCommon_WithCallback);
+
 	bool bResult = true;
 	WIN32_FIND_DATAW Data;
 	HANDLE Handle = FindFirstFileWithWildcard(Directory, /* Out */ Data);
