Adding -nocrashreports support for windows. When provided, we will log crashes/ensures/etc but will not perform more complex reporting.

diff --git a/Engine/Source/Runtime/Core/Private/HAL/ExceptionHandling.cpp b/Engine/Source/Runtime/Core/Private/HAL/ExceptionHandling.cpp
--- a/Engine/Source/Runtime/Core/Private/HAL/ExceptionHandling.cpp
+++ b/Engine/Source/Runtime/Core/Private/HAL/ExceptionHandling.cpp
@@ -17,6 +17,7 @@
 #endif
 
 /** Whether we should generate crash reports even if the debugger is attached. */
+/** NOTE: If -nocrashreports is passed in, this will be ignored. */
 CORE_API bool GAlwaysReportCrash = false;
 
 /** Whether to use ClientReportClient rather than the old AutoReporter. */
diff --git a/Engine/Source/Runtime/Core/Private/Misc/AssertionMacros.cpp b/Engine/Source/Runtime/Core/Private/Misc/AssertionMacros.cpp
--- a/Engine/Source/Runtime/Core/Private/Misc/AssertionMacros.cpp
+++ b/Engine/Source/Runtime/Core/Private/Misc/AssertionMacros.cpp
@@ -547,7 +547,11 @@ FORCENOINLINE void FDebug::EnsureFailed(const ANSICHAR* Expr, const ANSICHAR* Fi
 				TRACE_LOG_MESSAGE(LogOutputDevice, Warning, TEXT("%s"), ErrorMsg);
 			}
 #endif
-			
+			if (FParse::Param(FCommandLine::Get(), TEXT("nocrashreports")))
+			{
+				return;
+			}
+
 			// Submit the error report to the server! (and display a balloon in the system tray)
 			if (GEnsuresAreErrors)
 			{
@@ -599,6 +603,11 @@ FORCENOINLINE void FDebug::EnsureFailed(const ANSICHAR* Expr, const ANSICHAR* Fi
 		}
 		else
 		{
+			if (FParse::Param(FCommandLine::Get(), TEXT("nocrashreports")))
+			{
+				return;
+			}
+
 			// If we fail to generate the string to identify the crash we don't know if we should skip sending the report,
 			// so we will just send the report anyway.
 			bShouldSendNewReport = !!GEnsuresAreErrors;
