Fixed return type of FileLineHashForEnsure.

diff --git a/Engine/Source/Runtime/Core/Public/Misc/AssertionMacros.h b/Engine/Source/Runtime/Core/Public/Misc/AssertionMacros.h
--- a/Engine/Source/Runtime/Core/Public/Misc/AssertionMacros.h
+++ b/Engine/Source/Runtime/Core/Public/Misc/AssertionMacros.h
@@ -398,7 +398,7 @@ RetType UE_COLD UE_DEBUG_SECTION DispatchCheckVerify (InnerType&& Inner, ArgType
 	template <uint64 Uid>
 	std::atomic<bool> bGEnsureHasExecuted = false;
 
-	UE_NODEBUG constexpr uint32 FileLineHashForEnsure(const char* Filename, uint32 Line)
+	UE_NODEBUG constexpr uint64 FileLineHashForEnsure(const char* Filename, uint32 Line)
 	{
 		// Based on djb2
 		uint32 Result = 5381;
