Mutex: Removed the delayed unlock path

#rnx

diff --git a/Engine/Source/Runtime/Core/Private/Async/ExternalMutex.cpp b/Engine/Source/Runtime/Core/Private/Async/ExternalMutex.cpp
--- a/Engine/Source/Runtime/Core/Private/Async/ExternalMutex.cpp
+++ b/Engine/Source/Runtime/Core/Private/Async/ExternalMutex.cpp
@@ -60,35 +60,11 @@ void FExternalMutex::LockSlow()
 
 FORCENOINLINE void FExternalMutex::WakeWaitingThread()
 {
-	if constexpr (!bUnlockImmediately)
-	{
-		uint8 CurrentState = State.load(std::memory_order_relaxed);
-		checkSlow(CurrentState & IsLockedFlag);
-
-		// Spin on the fast path because there may be spurious failures.
-		while (LIKELY(!(CurrentState & MayHaveWaitingLockFlag)))
-		{
-			if (LIKELY(State.compare_exchange_weak(CurrentState, CurrentState & ~IsLockedFlag, std::memory_order_release, std::memory_order_relaxed)))
-			{
-				return;
-			}
-		}
-	}
-
 	ParkingLot::WakeOne(&State, [this](ParkingLot::FWakeState WakeState) -> uint64
 	{
-		if constexpr (bUnlockImmediately)
-		{
-			if (!WakeState.bHasWaitingThreads)
-			{
-				State.fetch_and(~MayHaveWaitingLockFlag, std::memory_order_relaxed);
-			}
-		}
-		else
+		if (!WakeState.bHasWaitingThreads)
 		{
-			const uint8 ClearFlags = IsLockedFlag | (WakeState.bHasWaitingThreads ? 0 : MayHaveWaitingLockFlag);
-			const uint8 OldState = State.fetch_and(~ClearFlags, std::memory_order_release);
-			checkSlow((OldState & IsLockedFlag) && (OldState & MayHaveWaitingLockFlag));
+			State.fetch_and(~MayHaveWaitingLockFlag, std::memory_order_relaxed);
 		}
 		return 0;
 	});
diff --git a/Engine/Source/Runtime/Core/Private/Async/Mutex.cpp b/Engine/Source/Runtime/Core/Private/Async/Mutex.cpp
--- a/Engine/Source/Runtime/Core/Private/Async/Mutex.cpp
+++ b/Engine/Source/Runtime/Core/Private/Async/Mutex.cpp
@@ -57,35 +57,11 @@ void FMutex::LockSlow()
 
 FORCENOINLINE void FMutex::WakeWaitingThread()
 {
-	if constexpr (!bUnlockImmediately)
-	{
-		uint8 CurrentState = State.load(std::memory_order_relaxed);
-		checkSlow(CurrentState & IsLockedFlag);
-
-		// Spin on the fast path because there may be spurious failures.
-		while (LIKELY(!(CurrentState & MayHaveWaitingLockFlag)))
-		{
-			if (LIKELY(State.compare_exchange_weak(CurrentState, 0, std::memory_order_release, std::memory_order_relaxed)))
-			{
-				return;
-			}
-		}
-	}
-
 	ParkingLot::WakeOne(&State, [this](ParkingLot::FWakeState WakeState) -> uint64
 	{
-		if constexpr (bUnlockImmediately)
-		{
-			if (!WakeState.bHasWaitingThreads)
-			{
-				State.fetch_and(~MayHaveWaitingLockFlag, std::memory_order_relaxed);
-			}
-		}
-		else
+		if (!WakeState.bHasWaitingThreads)
 		{
-			const uint8 NewState = WakeState.bHasWaitingThreads ? MayHaveWaitingLockFlag : 0;
-			const uint8 OldState = State.exchange(NewState, std::memory_order_release);
-			checkSlow((OldState & IsLockedFlag) && (OldState & MayHaveWaitingLockFlag));
+			State.fetch_and(~MayHaveWaitingLockFlag, std::memory_order_relaxed);
 		}
 		return 0;
 	});
diff --git a/Engine/Source/Runtime/Core/Private/Async/RecursiveMutex.cpp b/Engine/Source/Runtime/Core/Private/Async/RecursiveMutex.cpp
--- a/Engine/Source/Runtime/Core/Private/Async/RecursiveMutex.cpp
+++ b/Engine/Source/Runtime/Core/Private/Async/RecursiveMutex.cpp
@@ -10,8 +10,6 @@
 namespace UE
 {
 
-static constexpr bool bUnlockRecursiveMutexImmediately = true;
-
 bool FRecursiveMutex::TryLock()
 {
 	const uint32 CurrentThreadId = FPlatformTLS::GetCurrentThreadId();
@@ -125,27 +123,12 @@ void FRecursiveMutex::Unlock()
 		// Remove the association with this thread before unlocking.
 		ThreadId.store(0, std::memory_order_relaxed);
 
-		if constexpr (bUnlockRecursiveMutexImmediately)
-		{
-			// Unlock immediately to allow other threads to acquire the lock while this thread looks for a thread to wake.
-			const uint32 LastState = State.fetch_sub(1 << LockCountShift, std::memory_order_release);
+		// Unlock immediately to allow other threads to acquire the lock while this thread looks for a thread to wake.
+		const uint32 LastState = State.fetch_sub(1 << LockCountShift, std::memory_order_release);
 
-			// Wake one exclusive waiter if there are waiting threads.
-			if (UNLIKELY(LastState & MayHaveWaitingLockFlag))
-			{
-				WakeWaitingThread();
-			}
-		}
-		else
+		// Wake one exclusive waiter if there are waiting threads.
+		if (UNLIKELY(LastState & MayHaveWaitingLockFlag))
 		{
-			if (LIKELY(!(CurrentState & MayHaveWaitingLockFlag)))
-			{
-				if (LIKELY(State.compare_exchange_strong(CurrentState, 0, std::memory_order_release, std::memory_order_relaxed)))
-				{
-					return;
-				}
-			}
-
 			WakeWaitingThread();
 		}
 	}
@@ -160,18 +143,9 @@ FORCENOINLINE void FRecursiveMutex::WakeWaitingThread()
 {
 	ParkingLot::WakeOne(&State, [this](ParkingLot::FWakeState WakeState) -> uint64
 	{
-		if constexpr (bUnlockRecursiveMutexImmediately)
-		{
-			if (!WakeState.bHasWaitingThreads)
-			{
-				State.fetch_and(~MayHaveWaitingLockFlag, std::memory_order_relaxed);
-			}
-		}
-		else
+		if (!WakeState.bHasWaitingThreads)
 		{
-			const uint32 NewState = WakeState.bHasWaitingThreads ? MayHaveWaitingLockFlag : 0;
-			const uint32 OldState = State.exchange(NewState, std::memory_order_release);
-			checkSlow((OldState & LockCountMask) && (OldState & MayHaveWaitingLockFlag));
+			State.fetch_and(~MayHaveWaitingLockFlag, std::memory_order_relaxed);
 		}
 		return 0;
 	});
diff --git a/Engine/Source/Runtime/Core/Public/Async/ExternalMutex.h b/Engine/Source/Runtime/Core/Public/Async/ExternalMutex.h
--- a/Engine/Source/Runtime/Core/Public/Async/ExternalMutex.h
+++ b/Engine/Source/Runtime/Core/Public/Async/ExternalMutex.h
@@ -61,25 +61,13 @@ public:
 
 	inline void Unlock()
 	{
-		if constexpr (bUnlockImmediately)
+		// Unlock immediately to allow other threads to acquire the lock while this thread looks for a thread to wake.
+		const uint8 LastState = State.fetch_sub(IsLockedFlag, std::memory_order_release);
+		if (LIKELY(!(LastState & MayHaveWaitingLockFlag)))
 		{
-			// Unlock immediately to allow other threads to acquire the lock while this thread looks for a thread to wake.
-			const uint8 LastState = State.fetch_sub(IsLockedFlag, std::memory_order_release);
-			if (LIKELY(!(LastState & MayHaveWaitingLockFlag)))
-			{
-				return;
-			}
-			WakeWaitingThread();
-		}
-		else
-		{
-			uint8 Expected = (State.load(std::memory_order_relaxed) | IsLockedFlag) & ~MayHaveWaitingLockFlag;
-			if (LIKELY(State.compare_exchange_weak(Expected, Expected & ~IsLockedFlag, std::memory_order_release, std::memory_order_relaxed)))
-			{
-				return;
-			}
-			WakeWaitingThread();
+			return;
 		}
+		WakeWaitingThread();
 	}
 
 private:
@@ -89,8 +77,6 @@ private:
 	static constexpr uint8 IsLockedFlag = 1 << 0;
 	static constexpr uint8 MayHaveWaitingLockFlag = 1 << 1;
 
-	static constexpr bool bUnlockImmediately = true;
-
 	std::atomic<uint8>& State;
 };
 
diff --git a/Engine/Source/Runtime/Core/Public/Async/Mutex.h b/Engine/Source/Runtime/Core/Public/Async/Mutex.h
--- a/Engine/Source/Runtime/Core/Public/Async/Mutex.h
+++ b/Engine/Source/Runtime/Core/Public/Async/Mutex.h
@@ -53,33 +53,19 @@ public:
 
 	inline void Unlock()
 	{
-		if constexpr (bUnlockImmediately)
+		// Unlock immediately to allow other threads to acquire the lock while this thread looks for a thread to wake.
+		const uint8 LastState = State.fetch_sub(IsLockedFlag, std::memory_order_release);
+		if (LIKELY(!(LastState & MayHaveWaitingLockFlag)))
 		{
-			// Unlock immediately to allow other threads to acquire the lock while this thread looks for a thread to wake.
-			const uint8 LastState = State.fetch_sub(IsLockedFlag, std::memory_order_release);
-			if (LIKELY(!(LastState & MayHaveWaitingLockFlag)))
-			{
-				return;
-			}
-			WakeWaitingThread();
-		}
-		else
-		{
-			uint8 Expected = IsLockedFlag;
-			if (LIKELY(State.compare_exchange_weak(Expected, 0, std::memory_order_release, std::memory_order_relaxed)))
-			{
-				return;
-			}
-			WakeWaitingThread();
+			return;
 		}
+		WakeWaitingThread();
 	}
 
 private:
 	UE_API void LockSlow();
 	UE_API void WakeWaitingThread();
 
-	static constexpr bool bUnlockImmediately = true;
-
 	static constexpr uint8 IsLockedFlag = 1 << 0;
 	static constexpr uint8 MayHaveWaitingLockFlag = 1 << 1;
 
