- Added a lock around OtherFiles add/remove

diff --git a/Engine/Source/Runtime/Core/Private/Misc/ConfigCacheIni.cpp b/Engine/Source/Runtime/Core/Private/Misc/ConfigCacheIni.cpp
--- a/Engine/Source/Runtime/Core/Private/Misc/ConfigCacheIni.cpp
+++ b/Engine/Source/Runtime/Core/Private/Misc/ConfigCacheIni.cpp
@@ -3690,9 +3690,11 @@ void FConfigCacheIni::Tick(float DeltaSeconds)
 	{
 		OtherFileToCheckForUnload = 0;
 	}
-	BranchesToCheck[1] = OtherFiles.FindRef(OtherFileNames[OtherFileToCheckForUnload++]);
-
-	checkf(OtherFileNames.Num() == OtherFiles.Num(), TEXT("OtherFiles and OtherFileNames are out of sync! %d other files, %d other file names!"), OtherFileNames.Num(), OtherFiles.Num());
+	{
+		FScopeLock Lock(&OtherFilesLock);
+		BranchesToCheck[1] = OtherFiles.FindRef(OtherFileNames[OtherFileToCheckForUnload++]);
+		checkf(OtherFileNames.Num() == OtherFiles.Num(), TEXT("OtherFiles and OtherFileNames are out of sync! %d other files, %d other file names!"), OtherFileNames.Num(), OtherFiles.Num());
+	}
 	
 	// now check for unused files
 	double Now = FPlatformTime::Seconds();
@@ -3740,6 +3742,7 @@ FConfigBranch* FConfigCacheIni::FindBranchWithNoReload(FName BaseIniName, const
 	}
 	if (Branch == nullptr)
 	{
+		FScopeLock Lock(&OtherFilesLock);
 		Branch = OtherFiles.FindRef(Filename);
 		if (Branch == nullptr)
 		{
@@ -3795,6 +3798,8 @@ FConfigBranch& FConfigCacheIni::AddNewBranch(const FString& Filename)
 		FileAccess->OverrideFilenameToLoad = FName(FStringView(Filename));
 	}
 #endif
+	
+	FScopeLock Lock(&OtherFilesLock);
 	if (OtherFiles.Find(Filename) == nullptr)
 	{
 		OtherFileNames.Add(Filename);
@@ -3807,6 +3812,8 @@ FConfigBranch& FConfigCacheIni::AddNewBranch(const FString& Filename)
 
 int32 FConfigCacheIni::Remove(const FString& Filename)
 {
+	FScopeLock Lock(&OtherFilesLock);
+
 	OtherFileNames.Remove(Filename);
 	delete OtherFiles.FindRef(Filename);
 	return OtherFiles.Remove(Filename);
@@ -3822,6 +3829,7 @@ FConfigFile* FConfigCacheIni::FindConfigFile( const FString& Filename )
 	}
 	else
 	{
+		FScopeLock Lock(&OtherFilesLock);
 		Result = OtherFiles.FindRef(Filename);
 	}
 
@@ -3901,6 +3909,7 @@ FConfigFile* FConfigCacheIni::FindConfigFileWithBaseName(FName BaseName)
 	FConfigBranch* Result = KnownFiles.GetBranch(BaseName);
 	if (Result == nullptr)
 	{
+		FScopeLock Lock(&OtherFilesLock);
 		for (TPair<FString,FConfigBranch*>& CurrentFilePair : OtherFiles)
 		{
 			if (CurrentFilePair.Value->IniName == BaseName)
@@ -3941,6 +3950,9 @@ FConfigFile& FConfigCacheIni::Add(const FString& Filename, const FConfigFile& Fi
 		FileAccess->OverrideFilenameToLoad = FName(FStringView(Filename));
 	}
 #endif
+
+	FScopeLock Lock(&OtherFilesLock);
+
 	if (OtherFiles.Find(Filename) == nullptr)
 	{
 		OtherFileNames.Add(Filename);
@@ -3958,11 +3970,14 @@ bool FConfigCacheIni::ContainsConfigFile(const FConfigFile* ConfigFile) const
 	// since the point at which the caller received the ConfigFile pointer
 	// they are testing. It is the caller's responsibility to not try to hold
 	// on to the ConfigFile pointer during writes to this ConfigCacheIni
-	for (const TPair<FString, FConfigBranch*>& CurrentFilePair : OtherFiles)
 	{
-		if (ConfigFile == &CurrentFilePair.Value->InMemoryFile)
+		FScopeLock Lock(&OtherFilesLock);
+		for (const TPair<FString, FConfigBranch*>& CurrentFilePair : OtherFiles)
 		{
-			return true;
+			if (ConfigFile == &CurrentFilePair.Value->InMemoryFile)
+			{
+				return true;
+			}
 		}
 	}
 	// Check the known inis
@@ -3979,8 +3994,12 @@ bool FConfigCacheIni::ContainsConfigFile(const FConfigFile* ConfigFile) const
 
 
 TArray<FString> FConfigCacheIni::GetFilenames()
-{
-	TArray<FString> Result = OtherFileNames;
+{	
+	TArray<FString> Result;
+	{
+		FScopeLock Lock(&OtherFilesLock);
+		Result = OtherFileNames;
+	}
 
 	for (const FConfigBranch& Branch : KnownFiles.Branches)
 	{
@@ -4012,6 +4031,7 @@ void FConfigCacheIni::Flush(bool bRemoveFromCache, const FString& Filename )
 			}
 			else
 			{
+				FScopeLock Lock(&OtherFilesLock);
 				// flush all files
 				for (TPair<FString, FConfigBranch*>& Pair : OtherFiles)
 				{
@@ -4040,6 +4060,7 @@ void FConfigCacheIni::Flush(bool bRemoveFromCache, const FString& Filename )
 		}
 		else
 		{
+			FScopeLock Lock(&OtherFilesLock);
 			for (TPair<FString, FConfigBranch*>& It : OtherFiles)
 			{
 				delete It.Value;
@@ -4670,6 +4691,8 @@ void FConfigCacheIni::Dump(FOutputDevice& Ar, const TCHAR* BaseIniName)
 		}
 	}
 
+	FScopeLock Lock(&OtherFilesLock);
+
 	// sort the non-known files for easier diffing
 	TArray<FString> Keys;
 	OtherFiles.GetKeys(Keys);
@@ -5257,6 +5280,7 @@ struct FConfigMemoryData
  */
 void FConfigCacheIni::ShowMemoryUsage( FOutputDevice& Ar )
 {
+	FScopeLock Lock(&OtherFilesLock);
 	FConfigMemoryData ConfigCacheMemoryData;
 
 	for (TPair<FString, FConfigBranch*>& Pair : OtherFiles)
@@ -5319,6 +5343,8 @@ void FConfigCacheIni::ShowMemoryUsage( FOutputDevice& Ar )
 
 SIZE_T FConfigCacheIni::GetMaxMemoryUsage()
 {
+	FScopeLock Lock(&OtherFilesLock);
+
 	// record the memory used by the FConfigCacheIni's TMap
 	FArchiveCountConfigMem MemAr;
 	OtherFiles.CountBytes(MemAr);
diff --git a/Engine/Source/Runtime/Core/Public/Misc/ConfigCacheIni.h b/Engine/Source/Runtime/Core/Public/Misc/ConfigCacheIni.h
--- a/Engine/Source/Runtime/Core/Public/Misc/ConfigCacheIni.h
+++ b/Engine/Source/Runtime/Core/Public/Misc/ConfigCacheIni.h
@@ -1987,6 +1987,7 @@ private:
 	FKnownConfigFiles KnownFiles;
 
 	TMap<FString, FConfigBranch*> OtherFiles;
+	mutable FCriticalSection OtherFilesLock;
 	
 	/** Parallel array to OtherFiles */
 	TArray<FString> OtherFileNames;
