Scheduler: Removed inline implementation of FOversubscriptionScope to reduce complexity where it is used

Reduced FMicrosoftPlatformManualResetEvent from 138 instructions to 93 instructions.

#rnx

diff --git a/Engine/Source/Runtime/Core/Private/Async/Fundamental/Scheduler.cpp b/Engine/Source/Runtime/Core/Private/Async/Fundamental/Scheduler.cpp
--- a/Engine/Source/Runtime/Core/Private/Async/Fundamental/Scheduler.cpp
+++ b/Engine/Source/Runtime/Core/Private/Async/Fundamental/Scheduler.cpp
@@ -33,6 +33,38 @@ namespace LowLevelTasks
 		return bIsOversubscriptionAllowed;
 	}
 
+	FORCENOINLINE void FOversubscriptionScope::TryIncrementOversubscription()
+	{
+		if (Private::FOversubscriptionTls::IsOversubscriptionAllowed())
+		{
+		#if CPUPROFILERTRACE_ENABLED
+			if (CpuChannel)
+			{
+				static uint32 OversubscriptionTraceId = FCpuProfilerTrace::OutputEventType("Oversubscription");
+				FCpuProfilerTrace::OutputBeginEvent(OversubscriptionTraceId);
+				bCpuBeginEventEmitted = true;
+			}
+		#endif
+
+			bIncrementOversubscriptionEmitted = true;
+			FScheduler::Get().IncrementOversubscription();
+		}
+	}
+
+	FORCENOINLINE void FOversubscriptionScope::DecrementOversubscription()
+	{
+		FScheduler::Get().DecrementOversubscription();
+		bIncrementOversubscriptionEmitted = false;
+
+	#if CPUPROFILERTRACE_ENABLED
+		if (bCpuBeginEventEmitted)
+		{
+			FCpuProfilerTrace::OutputEndEvent();
+			bCpuBeginEventEmitted = false;
+		}
+	#endif
+	}
+
 	FORCENOINLINE FSchedulerTls::FTlsValues& FSchedulerTls::GetTlsValuesRef()
 	{
 		return *TlsValuesHolder.TlsValues;
diff --git a/Engine/Source/Runtime/Core/Public/Async/Fundamental/Scheduler.h b/Engine/Source/Runtime/Core/Public/Async/Fundamental/Scheduler.h
--- a/Engine/Source/Runtime/Core/Public/Async/Fundamental/Scheduler.h
+++ b/Engine/Source/Runtime/Core/Public/Async/Fundamental/Scheduler.h
@@ -228,21 +228,11 @@ namespace LowLevelTasks
 		UE_NONCOPYABLE(FOversubscriptionScope);
 
 	public:
-		FOversubscriptionScope(bool bCondition = true)
+		explicit FOversubscriptionScope(bool bCondition = true)
 		{
-			if (bCondition && Private::FOversubscriptionTls::IsOversubscriptionAllowed())
+			if (bCondition)
 			{
-				bIncrementOversubscriptionEmitted = true;
-
-#if CPUPROFILERTRACE_ENABLED
-				if (CpuChannel)
-				{
-					static uint32 OversubscriptionTraceId = FCpuProfilerTrace::OutputEventType("Oversubscription");
-					FCpuProfilerTrace::OutputBeginEvent(OversubscriptionTraceId);
-					bCpuBeginEventEmitted = true;
-				}
-#endif
-				FScheduler::Get().IncrementOversubscription();
+				TryIncrementOversubscription();
 			}
 		}
 
@@ -250,18 +240,14 @@ namespace LowLevelTasks
 		{
 			if (bIncrementOversubscriptionEmitted)
 			{
-				FScheduler::Get().DecrementOversubscription();
-
-#if CPUPROFILERTRACE_ENABLED
-				if (bCpuBeginEventEmitted)
-				{
-					FCpuProfilerTrace::OutputEndEvent();
-					bCpuBeginEventEmitted = false;
-				}
-#endif
+				DecrementOversubscription();
 			}
 		}
+
 	private:
+		CORE_API void TryIncrementOversubscription();
+		CORE_API void DecrementOversubscription();
+
 		bool bIncrementOversubscriptionEmitted = false;
 		bool bCpuBeginEventEmitted = false;
 	};
