SharedMutex: Retry in TryLockShared while there is not an exclusive lock

diff --git a/Engine/Source/Runtime/Core/Private/Async/SharedRecursiveMutex.cpp b/Engine/Source/Runtime/Core/Private/Async/SharedRecursiveMutex.cpp
--- a/Engine/Source/Runtime/Core/Private/Async/SharedRecursiveMutex.cpp
+++ b/Engine/Source/Runtime/Core/Private/Async/SharedRecursiveMutex.cpp
@@ -215,7 +215,7 @@ bool FSharedRecursiveMutex::TryLockShared(Core::Private::FSharedRecursiveMutexLi
 {
 	uint32 CurrentState = State.load(std::memory_order_relaxed);
 	// Recursive shared locks are quick to acquire. Check for that case and add 1 to the count.
-	if ((CurrentState & SharedLockCountMask) && (!(CurrentState & MayHaveWaitingLockFlag) || Core::Private::FSharedRecursiveMutexLink::Owns(this)))
+	if ((CurrentState & SharedLockCountMask) && Core::Private::FSharedRecursiveMutexLink::Owns(this))
 	{
 		[[maybe_unused]] uint32 LastState = State.fetch_add(1 << SharedLockCountShift, std::memory_order_relaxed);
 		checkSlow((LastState & SharedLockCountMask) != SharedLockCountMask);
@@ -223,11 +223,17 @@ bool FSharedRecursiveMutex::TryLockShared(Core::Private::FSharedRecursiveMutexLi
 		return true;
 	}
 	// Try to acquire a shared lock if there is no active or waiting exclusive lock.
-	else if (!(CurrentState & (LockCountMask | MayHaveWaitingLockFlag)) &&
-		State.compare_exchange_strong(CurrentState, CurrentState + (1 << SharedLockCountShift), std::memory_order_acquire, std::memory_order_relaxed))
+	else
 	{
-		Link.Push(this);
-		return true;
+		while (!(CurrentState & (LockCountMask | MayHaveWaitingLockFlag)))
+		{
+			checkSlow((CurrentState & SharedLockCountMask) != SharedLockCountMask);
+			if (State.compare_exchange_weak(CurrentState, CurrentState + (1 << SharedLockCountShift), std::memory_order_acquire, std::memory_order_relaxed))
+			{
+				Link.Push(this);
+				return true;
+			}
+		}
 	}
 	return false;
 }
@@ -236,7 +242,7 @@ void FSharedRecursiveMutex::LockShared(Core::Private::FSharedRecursiveMutexLink&
 {
 	uint32 CurrentState = State.load(std::memory_order_relaxed);
 	// Recursive shared locks are quick to acquire. Check for that case and add 1 to the count.
-	if ((CurrentState & SharedLockCountMask) && (!(CurrentState & MayHaveWaitingLockFlag) || Core::Private::FSharedRecursiveMutexLink::Owns(this)))
+	if ((CurrentState & SharedLockCountMask) && Core::Private::FSharedRecursiveMutexLink::Owns(this))
 	{
 		[[maybe_unused]] uint32 LastState = State.fetch_add(1 << SharedLockCountShift, std::memory_order_relaxed);
 		checkSlow((LastState & SharedLockCountMask) != SharedLockCountMask);
@@ -244,11 +250,14 @@ void FSharedRecursiveMutex::LockShared(Core::Private::FSharedRecursiveMutexLink&
 		return;
 	}
 	// Try to acquire a shared lock if there is no active or waiting exclusive lock.
-	else if (!(CurrentState & (LockCountMask | MayHaveWaitingLockFlag)) &&
-		State.compare_exchange_weak(CurrentState, CurrentState + (1 << SharedLockCountShift), std::memory_order_acquire, std::memory_order_relaxed))
+	else if (!(CurrentState & (LockCountMask | MayHaveWaitingLockFlag)))
 	{
-		Link.Push(this);
-		return;
+		checkSlow((CurrentState & SharedLockCountMask) != SharedLockCountMask);
+		if (State.compare_exchange_weak(CurrentState, CurrentState + (1 << SharedLockCountShift), std::memory_order_acquire, std::memory_order_relaxed))
+		{
+			Link.Push(this);
+			return;
+		}
 	}
 	LockSharedSlow(Link);
 }
diff --git a/Engine/Source/Runtime/Core/Public/Async/SharedMutex.h b/Engine/Source/Runtime/Core/Public/Async/SharedMutex.h
--- a/Engine/Source/Runtime/Core/Public/Async/SharedMutex.h
+++ b/Engine/Source/Runtime/Core/Public/Async/SharedMutex.h
@@ -69,9 +69,15 @@ public:
 	inline bool TryLockShared()
 	{
 		uint32 Expected = State.load(std::memory_order_relaxed);
-		return !(Expected & (IsLockedFlag | MayHaveWaitingLockFlag)) &&
-			State.compare_exchange_strong(Expected, Expected + (1 << SharedLockCountShift),
-				std::memory_order_acquire, std::memory_order_relaxed);
+		while (LIKELY(!(Expected & (IsLockedFlag | MayHaveWaitingLockFlag))))
+		{
+			if (LIKELY(State.compare_exchange_weak(Expected, Expected + (1 << SharedLockCountShift),
+					std::memory_order_acquire, std::memory_order_relaxed)))
+			{
+				return true;
+			}
+		}
+		return false;
 	}
 
 	inline void LockShared()
