Fixed alignment of TestMemory buffer and fixed a spelling error.

#rnx

diff --git a/Engine/Source/Runtime/Core/Tests/Containers/CompactHashTableTests.cpp b/Engine/Source/Runtime/Core/Tests/Containers/CompactHashTableTests.cpp
--- a/Engine/Source/Runtime/Core/Tests/Containers/CompactHashTableTests.cpp
+++ b/Engine/Source/Runtime/Core/Tests/Containers/CompactHashTableTests.cpp
@@ -120,21 +120,21 @@ TEST_CASE_NAMED(FStaticCompactHashTableTest, "System::Core::Containers::CompactH
 		static_assert(sizeof(TStaticCompactHashTable<65536>) == (65536 + 65536) * sizeof(uint32));
 		static_assert(sizeof(TStaticCompactHashTable<65536, 65537>) == (65536 + 65537) * sizeof(uint32));
 
-		const uint8 Sentinal = 0xf0;
+		const uint8 Sentinel = 0xf0;
 		const uint8 Clear = 0xff;
 		
 		auto TestTableType = []<uint32 Size, uint32 HashSize = Size>()
 		{
 			using TestType = TStaticCompactHashTable<Size, HashSize>;
 
-			uint8 TestMemory[sizeof(TestType)];
-			FMemory::Memset(TestMemory, Sentinal);
+			alignas(TestType) uint8 TestMemory[sizeof(TestType)];
+			FMemory::Memset(TestMemory, Sentinel);
 
 			new (TestMemory) TestType(NoInit);
 
 			for (uint8 Test : TestMemory)
 			{
-				CHECK(Test == Sentinal);
+				CHECK(Test == Sentinel);
 			}
 
 			new (TestMemory) TestType();
@@ -143,7 +143,7 @@ TEST_CASE_NAMED(FStaticCompactHashTableTest, "System::Core::Containers::CompactH
 			{
 				if (Index < Size * sizeof(typename TestType::IndexType))
 				{
-					CHECK(TestMemory[Index] == Sentinal);
+					CHECK(TestMemory[Index] == Sentinel);
 				}
 				else
 				{
@@ -192,7 +192,7 @@ TEST_CASE_NAMED(FDynamicCompactHashTableTest, "System::Core::Containers::Compact
 {
 	SECTION("Constructor Behavior")
 	{
-		const uint8 Sentinal = 0xf0;
+		const uint8 Sentinel = 0xf0;
 		const uint8 Clear = 0xff;
 
 		auto TestTableType = []<uint32 Size, uint32 HashSize = Size>()
@@ -204,13 +204,13 @@ TEST_CASE_NAMED(FDynamicCompactHashTableTest, "System::Core::Containers::Compact
 			static_assert(RequiredSize == sizeof(typename TStaticCompactHashTable<Size>::IndexType) * (Size + HashSize));
 
 			uint8 TestMemory[RequiredSize];
-			FMemory::Memset(TestMemory, Sentinal, RequiredSize);
+			FMemory::Memset(TestMemory, Sentinel, RequiredSize);
 
 			TestType HashTable(TestMemory, Size, HashSize, RequiredSize);
 
 			for (uint8 Test : TestMemory)
 			{
-				CHECK(Test == Sentinal);
+				CHECK(Test == Sentinel);
 			}
 
 			HashTable.Reset();
@@ -219,7 +219,7 @@ TEST_CASE_NAMED(FDynamicCompactHashTableTest, "System::Core::Containers::Compact
 			{
 				if (Index < Size * sizeof(IndexType))
 				{
-					CHECK(TestMemory[Index] == Sentinal);
+					CHECK(TestMemory[Index] == Sentinel);
 				}
 				else
 				{
