[CoreGlobals]
- Add function to know if we're in any of the rendering thread

diff --git a/Engine/Source/Runtime/Core/Private/HAL/ThreadingBase.cpp b/Engine/Source/Runtime/Core/Private/HAL/ThreadingBase.cpp
--- a/Engine/Source/Runtime/Core/Private/HAL/ThreadingBase.cpp
+++ b/Engine/Source/Runtime/Core/Private/HAL/ThreadingBase.cpp
@@ -297,6 +297,14 @@ PRAGMA_DISABLE_DEPRECATION_WARNINGS
 PRAGMA_ENABLE_DEPRECATION_WARNINGS
 }
 
+CORE_API bool IsInAnyRenderingThread()
+{
+	return FTaskTagScope::IsCurrentTag(ETaskTag::EParallelRenderingThread)
+		|| FTaskTagScope::IsCurrentTag(ETaskTag::ERenderingThread)
+		|| FTaskTagScope::IsCurrentTag(ETaskTag::EParallelRhiThread)
+		|| FTaskTagScope::IsCurrentTag(ETaskTag::ERhiThread);
+}
+
 CORE_API bool IsInParallelRenderingThread()
 {
 PRAGMA_DISABLE_DEPRECATION_WARNINGS
diff --git a/Engine/Source/Runtime/Core/Public/CoreGlobals.h b/Engine/Source/Runtime/Core/Public/CoreGlobals.h
--- a/Engine/Source/Runtime/Core/Public/CoreGlobals.h
+++ b/Engine/Source/Runtime/Core/Public/CoreGlobals.h
@@ -740,6 +740,9 @@ extern CORE_API bool IsInRenderingThread();
 /** @return True if called from the rendering thread, or if called from ANY thread that isn't the game thread, except that during single threaded rendering the game thread is ok too.*/
 extern CORE_API bool IsInParallelRenderingThread();
 
+/** @return True if called from any of the rendering thread. Can be called from any thread without asserting. */
+extern CORE_API bool IsInAnyRenderingThread();
+
 /** @return True if called from the rendering thread. */
 // Unlike IsInRenderingThread, this will always return false if we are running single threaded. It only returns true if this is actually a separate rendering thread. Mostly useful for checks
 extern CORE_API bool IsInActualRenderingThread();
