Scheduler: Moved the oversubscription scopes into a new header

#rnx

diff --git a/Engine/Source/Runtime/Core/Public/Async/Fundamental/Oversubscription.h b/Engine/Source/Runtime/Core/Public/Async/Fundamental/Oversubscription.h
new file mode 100644
--- /dev/null
+++ b/Engine/Source/Runtime/Core/Public/Async/Fundamental/Oversubscription.h
@@ -0,0 +1,78 @@
+// Copyright Epic Games, Inc. All Rights Reserved.
+
+#pragma once
+
+#include "CoreTypes.h"
+
+namespace LowLevelTasks::Private
+{
+
+class FOversubscriptionTls
+{
+	static thread_local bool bIsOversubscriptionAllowed;
+
+	friend class FOversubscriptionAllowedScope;
+
+	CORE_API static bool& GetIsOversubscriptionAllowedRef();
+
+public:
+	static bool IsOversubscriptionAllowed() { return bIsOversubscriptionAllowed; }
+};
+
+class FOversubscriptionAllowedScope
+{
+	UE_NONCOPYABLE(FOversubscriptionAllowedScope);
+
+public:
+	explicit FOversubscriptionAllowedScope(bool bIsOversubscriptionAllowed)
+		: bValue(FOversubscriptionTls::GetIsOversubscriptionAllowedRef())
+		, bPreviousValue(bValue)
+	{
+		bValue = bIsOversubscriptionAllowed;
+	}
+
+	~FOversubscriptionAllowedScope()
+	{
+		bValue = bPreviousValue;
+	}
+
+private:
+	bool& bValue;
+	bool bPreviousValue;
+};
+
+} // LowLevelTasks::Private
+
+namespace LowLevelTasks
+{
+
+class FOversubscriptionScope
+{
+	UE_NONCOPYABLE(FOversubscriptionScope);
+
+public:
+	explicit FOversubscriptionScope(bool bCondition = true)
+	{
+		if (bCondition)
+		{
+			TryIncrementOversubscription();
+		}
+	}
+
+	~FOversubscriptionScope()
+	{
+		if (bIncrementOversubscriptionEmitted)
+		{
+			DecrementOversubscription();
+		}
+	}
+
+private:
+	CORE_API void TryIncrementOversubscription();
+	CORE_API void DecrementOversubscription();
+
+	bool bIncrementOversubscriptionEmitted = false;
+	bool bCpuBeginEventEmitted = false;
+};
+
+} // LowLevelTests
diff --git a/Engine/Source/Runtime/Core/Public/Async/Fundamental/Scheduler.h b/Engine/Source/Runtime/Core/Public/Async/Fundamental/Scheduler.h
--- a/Engine/Source/Runtime/Core/Public/Async/Fundamental/Scheduler.h
+++ b/Engine/Source/Runtime/Core/Public/Async/Fundamental/Scheduler.h
@@ -1,6 +1,8 @@
 // Copyright Epic Games, Inc. All Rights Reserved.
 
-#pragma once 
+#pragma once
+
+#include "Async/Fundamental/Oversubscription.h"
 #include "Async/Fundamental/Task.h"
 #include "Async/Fundamental/TaskShared.h"
 #include "Async/Fundamental/TaskDelegate.h"
@@ -186,72 +188,6 @@ namespace LowLevelTasks
 		FOversubscriptionLimitReached                  OversubscriptionLimitReachedEvent;
 	};
 
-	namespace Private
-	{
-		class FOversubscriptionTls
-		{
-			static thread_local bool bIsOversubscriptionAllowed;
-
-			friend class FOversubscriptionAllowedScope;
-
-			CORE_API static bool& GetIsOversubscriptionAllowedRef();
-
-		public:
-			static bool IsOversubscriptionAllowed() { return bIsOversubscriptionAllowed; }
-		};
-
-		class FOversubscriptionAllowedScope
-		{
-			UE_NONCOPYABLE(FOversubscriptionAllowedScope);
-
-		public:
-			explicit FOversubscriptionAllowedScope(bool bIsOversubscriptionAllowed)
-				: bValue(FOversubscriptionTls::GetIsOversubscriptionAllowedRef())
-				, bPreviousValue(bValue)
-			{
-				bValue = bIsOversubscriptionAllowed;
-			}
-
-			~FOversubscriptionAllowedScope()
-			{
-				bValue = bPreviousValue;
-			}
-
-		private:
-			bool& bValue;
-			bool bPreviousValue;
-		};
-	}
-
-	class FOversubscriptionScope
-	{
-		UE_NONCOPYABLE(FOversubscriptionScope);
-
-	public:
-		explicit FOversubscriptionScope(bool bCondition = true)
-		{
-			if (bCondition)
-			{
-				TryIncrementOversubscription();
-			}
-		}
-
-		~FOversubscriptionScope()
-		{
-			if (bIncrementOversubscriptionEmitted)
-			{
-				DecrementOversubscription();
-			}
-		}
-
-	private:
-		CORE_API void TryIncrementOversubscription();
-		CORE_API void DecrementOversubscription();
-
-		bool bIncrementOversubscriptionEmitted = false;
-		bool bCpuBeginEventEmitted = false;
-	};
-
 	FORCEINLINE_DEBUGGABLE bool TryLaunch(FTask& Task, EQueuePreference QueuePreference = EQueuePreference::DefaultPreference, bool bWakeUpWorker = true)
 	{
 		return FScheduler::Get().TryLaunch(Task, QueuePreference, bWakeUpWorker);
