IntrusiveMutex: Added a configurable intrusive unfair non-recursive mutex

diff --git a/Engine/Source/Runtime/Core/Public/Async/IntrusiveMutex.h b/Engine/Source/Runtime/Core/Public/Async/IntrusiveMutex.h
new file mode 100644
--- /dev/null
+++ b/Engine/Source/Runtime/Core/Public/Async/IntrusiveMutex.h
@@ -0,0 +1,151 @@
+// Copyright Epic Games, Inc. All Rights Reserved.
+
+#pragma once
+
+#include "Async/Fundamental/Oversubscription.h"
+#include "Async/ParkingLot.h"
+#include "CoreTypes.h"
+#include "HAL/PlatformProcess.h"
+
+#include <atomic>
+
+namespace UE
+{
+
+template <typename StateType>
+struct TIntrusiveMutexParams
+{
+	/** Flag that is written into the state with a bitwise or when locking the mutex. */
+	StateType IsLockedFlag = StateType(1) << 0;
+
+	/** Mask that is tested against the state when testing if the mutex is locked. 0 uses IsLockedFlag. */
+	StateType IsLockedMask = 0;
+
+	/** Flag that is set, cleared, and tested in the state to indicate that a lock may be waiting. */
+	StateType MayHaveWaitingLockFlag = StateType(1) << 1;
+
+	/** How many times to spin before parking the thread. */
+	int32 SpinLimit = 40;
+};
+
+/**
+ * An 2-bit intrusive mutex that is not fair and does not support recursive locking.
+ */
+template <typename StateType, TIntrusiveMutexParams<StateType> Params = {}>
+class TIntrusiveMutex
+{
+	TIntrusiveMutex() = delete;
+
+	constexpr static StateType IsLockedFlag = Params.IsLockedFlag;
+	constexpr static StateType IsLockedMask = Params.IsLockedMask ? Params.IsLockedMask : Params.IsLockedFlag;
+	constexpr static StateType MayHaveWaitingLockFlag = Params.MayHaveWaitingLockFlag;
+	constexpr static int32 SpinLimit = Params.SpinLimit;
+
+public:
+	FORCEINLINE static bool IsLocked(const std::atomic<StateType>& State)
+	{
+		return !!(State.load(std::memory_order_relaxed) & IsLockedFlag);
+	}
+
+	FORCEINLINE static bool TryLock(std::atomic<StateType>& State)
+	{
+		StateType Expected = State.load(std::memory_order_relaxed);
+		return !(Expected & IsLockedMask) &&
+			State.compare_exchange_strong(Expected, Expected | IsLockedFlag, std::memory_order_acquire, std::memory_order_relaxed);
+	}
+
+	FORCEINLINE static void Lock(std::atomic<StateType>& State)
+	{
+		StateType Expected = State.load(std::memory_order_relaxed) & ~IsLockedMask & ~MayHaveWaitingLockFlag;
+		if (LIKELY(State.compare_exchange_weak(Expected, Expected | IsLockedFlag, std::memory_order_acquire, std::memory_order_relaxed)))
+		{
+			return;
+		}
+		LockSlow(State);
+	}
+
+	FORCEINLINE static void LockLoop(std::atomic<StateType>& State)
+	{
+		int32 SpinCount = 0;
+		for (StateType CurrentState = State.load(std::memory_order_relaxed);;)
+		{
+			// Try to acquire the lock if it was unlocked, even if there are waiting threads.
+			// Acquiring the lock despite the waiting threads means that this lock is not FIFO and thus not fair.
+			if (LIKELY(!(CurrentState & IsLockedMask)))
+			{
+				if (LIKELY(State.compare_exchange_weak(CurrentState, CurrentState | IsLockedFlag, std::memory_order_acquire, std::memory_order_relaxed)))
+				{
+					return;
+				}
+				continue;
+			}
+
+			// Spin up to the spin limit while there are no waiting threads.
+			if (LIKELY(!(CurrentState & MayHaveWaitingLockFlag) && SpinCount < SpinLimit))
+			{
+				FPlatformProcess::YieldThread();
+				++SpinCount;
+				CurrentState = State.load(std::memory_order_relaxed);
+				continue;
+			}
+
+			// Store that there are waiting threads. Restart if the state has changed since it was loaded.
+			if (LIKELY(!(CurrentState & MayHaveWaitingLockFlag)))
+			{
+				if (UNLIKELY(!State.compare_exchange_weak(CurrentState, CurrentState | MayHaveWaitingLockFlag, std::memory_order_relaxed)))
+				{
+					continue;
+				}
+				CurrentState |= MayHaveWaitingLockFlag;
+			}
+
+			// Do not enter oversubscription during a wait on a mutex since the wait is generally too short
+			// for it to matter and it can worsen performance a lot for heavily contended locks.
+			LowLevelTasks::Private::FOversubscriptionAllowedScope _(false);
+
+			// Wait if the state has not changed. Either way, loop back and try to acquire the lock after trying to wait.
+			ParkingLot::Wait(&State, [&State]
+			{
+				const StateType NewState = State.load(std::memory_order_relaxed);
+				return (NewState & IsLockedMask) && (NewState & MayHaveWaitingLockFlag);
+			}, []{});
+			CurrentState = State.load(std::memory_order_relaxed);
+		}
+	}
+
+	FORCEINLINE static void Unlock(std::atomic<StateType>& State)
+	{
+		// Unlock immediately to allow other threads to acquire the lock while this thread looks for a thread to wake.
+		const StateType LastState = State.fetch_sub(IsLockedFlag, std::memory_order_release);
+		if (LIKELY(!(LastState & MayHaveWaitingLockFlag)))
+		{
+			return;
+		}
+		UnlockSlow(State);
+	}
+
+	FORCEINLINE static void WakeWaitingThread(std::atomic<StateType>& State)
+	{
+		ParkingLot::WakeOne(&State, [&State](ParkingLot::FWakeState WakeState) -> uint64
+		{
+			if (!WakeState.bHasWaitingThreads)
+			{
+				State.fetch_and(~MayHaveWaitingLockFlag, std::memory_order_relaxed);
+			}
+			return 0;
+		});
+	}
+
+private:
+	FORCENOINLINE static void LockSlow(std::atomic<StateType>& State)
+	{
+		LockLoop(State);
+	}
+
+	FORCENOINLINE static void UnlockSlow(std::atomic<StateType>& State)
+	{
+		WakeWaitingThread(State);
+	}
+};
+
+} // UE
