Fix initialization order TLS bug with the CSV profiler. Initialize EOP ThreadData lazily instead of on startup via a singleton
#tests tested in a cooked build
#rnx

diff --git a/Engine/Source/Runtime/Core/Private/ProfilingDebugging/CsvProfiler.cpp b/Engine/Source/Runtime/Core/Private/ProfilingDebugging/CsvProfiler.cpp
--- a/Engine/Source/Runtime/Core/Private/ProfilingDebugging/CsvProfiler.cpp
+++ b/Engine/Source/Runtime/Core/Private/ProfilingDebugging/CsvProfiler.cpp
@@ -2384,9 +2384,6 @@ private:
 		FCriticalSection TlsCS;
 		TArray<FWeakPtr> TlsInstances;
 		uint32 TlsSlot = FPlatformTLS::InvalidTlsSlot;
-
-		FSharedPtr EndOfPipeData = MakeShared<FCsvProfilerThreadData>(ECsvTimeline::EndOfPipe);
-
 	} static Singleton;
 
 public:
@@ -2432,9 +2429,10 @@ public:
 		return *ProfilerThread;
 	}
 
-	static CSV_PROFILER_INLINE FCsvProfilerThreadData& GetEndOfPipe()
+	static CSV_PROFILER_INLINE FSharedPtr GetEndOfPipe()
 	{
-		return *Singleton.EndOfPipeData.Get();
+		static FSharedPtr EndOfPipeData = MakeShared<FCsvProfilerThreadData>(ECsvTimeline::EndOfPipe);
+		return EndOfPipeData;
 	}
 
 	static inline void GetTlsInstances(TArray<FSharedPtr>& OutTlsInstances)
@@ -2454,7 +2452,7 @@ public:
 		}
 
 		// Always add the end-of-pipe data
-		OutTlsInstances.Add(Singleton.EndOfPipeData);
+		OutTlsInstances.Add(FCsvProfilerThreadData::GetEndOfPipe());
 	}
 
 	FCsvProfilerThreadData(TOptional<ECsvTimeline::Type> InCsvTimeline = NullOpt)
@@ -4554,7 +4552,7 @@ void FCsvProfiler::RecordEndOfPipeCustomStat(const FName& StatName, uint32 Categ
 	if (GCsvProfilerIsCapturing && GCsvCategoriesEnabled[CategoryIndex])
 	{
 		ValidateFName(StatName);
-		FCsvProfilerThreadData::GetEndOfPipe().AddCustomStat(StatName, CategoryIndex, (float)Value, CustomStatOp);
+		FCsvProfilerThreadData::GetEndOfPipe()->AddCustomStat(StatName, CategoryIndex, (float)Value, CustomStatOp);
 	}
 #endif
 }
