Do not pass vector register by const reference as this generates an additional load instruction and in some cases compiler generates wrong asm

diff --git a/Engine/Plugins/Interchange/Runtime/Source/Parsers/GLTFCore/Private/GLTF/GLTFReader.cpp b/Engine/Plugins/Interchange/Runtime/Source/Parsers/GLTFCore/Private/GLTF/GLTFReader.cpp
--- a/Engine/Plugins/Interchange/Runtime/Source/Parsers/GLTFCore/Private/GLTF/GLTFReader.cpp
+++ b/Engine/Plugins/Interchange/Runtime/Source/Parsers/GLTFCore/Private/GLTF/GLTFReader.cpp
@@ -19,8 +19,6 @@
 
 #include "GenericPlatform/GenericPlatformHttp.h"
 
-#include "Math/UnrealMathVectorConstants.h"
-
 
 #define LOCTEXT_NAMESPACE "InterchangeGLTFReader"
 
diff --git a/Engine/Plugins/Runtime/GeometryProcessing/Source/DynamicMesh/Private/Operations/GeodesicPath.cpp b/Engine/Plugins/Runtime/GeometryProcessing/Source/DynamicMesh/Private/Operations/GeodesicPath.cpp
--- a/Engine/Plugins/Runtime/GeometryProcessing/Source/DynamicMesh/Private/Operations/GeodesicPath.cpp
+++ b/Engine/Plugins/Runtime/GeometryProcessing/Source/DynamicMesh/Private/Operations/GeodesicPath.cpp
@@ -1,7 +1,6 @@
 // Copyright Epic Games, Inc. All Rights Reserved.
 
 #include "Operations/GeodesicPath.h"
-#include "Math/UnrealMathVectorCommon.h"
 #include "Operations/MeshGeodesicSurfaceTracer.h" 
 
 
diff --git a/Engine/Source/Editor/Persona/Private/PhysicsAssetRenderUtils.cpp b/Engine/Source/Editor/Persona/Private/PhysicsAssetRenderUtils.cpp
--- a/Engine/Source/Editor/Persona/Private/PhysicsAssetRenderUtils.cpp
+++ b/Engine/Source/Editor/Persona/Private/PhysicsAssetRenderUtils.cpp
@@ -7,7 +7,6 @@
 #include "Materials/MaterialInstanceDynamic.h"
 #include "Components/SkeletalMeshComponent.h"
 #include "Materials/Material.h"
-#include "Math/UnrealMathVectorCommon.h"
 #include "Preferences/PhysicsAssetEditorOptions.h"
 #include "PhysicsEngine/PhysicsAsset.h"
 #include "PhysicsEngine/PhysicsConstraintTemplate.h"
diff --git a/Engine/Source/Runtime/Core/Public/Math/Quat.h b/Engine/Source/Runtime/Core/Public/Math/Quat.h
--- a/Engine/Source/Runtime/Core/Public/Math/Quat.h
+++ b/Engine/Source/Runtime/Core/Public/Math/Quat.h
@@ -103,11 +103,11 @@ protected:
 	 *
 	 * @param V XYZW components of the quaternion packed into a single VectorRegister4Float
 	 */
-	[[nodiscard]] explicit TQuat(const QuatVectorRegister& V);
+	[[nodiscard]] explicit TQuat(QuatVectorRegister V);
 
 public:
 
-	[[nodiscard]] FORCEINLINE static TQuat<T> MakeFromVectorRegister(const QuatVectorRegister& V)
+	[[nodiscard]] FORCEINLINE static TQuat<T> MakeFromVectorRegister(QuatVectorRegister V)
 	{
 		return TQuat<T>(V);
 	}
@@ -882,7 +882,7 @@ FORCEINLINE TQuat<T>::TQuat(T InX, T InY, T InZ, T InW)
 }
 
 template<typename T>
-FORCEINLINE TQuat<T>::TQuat(const QuatVectorRegister& V)
+FORCEINLINE TQuat<T>::TQuat(QuatVectorRegister V)
 {
 	VectorStoreAligned(V, this);
 	DiagnosticCheckNaN();
@@ -1041,8 +1041,7 @@ FORCEINLINE TQuat<T> TQuat<T>::operator*=(const TQuat<T>& Q)
 {
 	QuatVectorRegister A = VectorLoadAligned(this);
 	QuatVectorRegister B = VectorLoadAligned(&Q);
-	QuatVectorRegister Result;
-	VectorQuaternionMultiply(&Result, &A, &B);
+	QuatVectorRegister Result = VectorQuaternionMultiply2(A, B);
 	VectorStoreAligned(Result, this);
 
 	DiagnosticCheckNaN();
diff --git a/Engine/Source/Runtime/Core/Public/Math/TransformVectorized.h b/Engine/Source/Runtime/Core/Public/Math/TransformVectorized.h
--- a/Engine/Source/Runtime/Core/Public/Math/TransformVectorized.h
+++ b/Engine/Source/Runtime/Core/Public/Math/TransformVectorized.h
@@ -11,9 +11,7 @@
 #include "Math/Quat.h"
 #include "Math/Rotator.h"
 #include "Math/ScalarRegister.h"
-#include "Math/UnrealMathSSE.h"
 #include "Math/UnrealMathUtility.h"
-#include "Math/UnrealMathVectorConstants.h"
 #include "Math/Vector.h"
 #include "Math/Vector4.h"
 #include "Math/VectorRegister.h"
@@ -228,7 +226,7 @@ public:
 	 * @param InTranslation The value to use for the translation component
 	 * @param InScale3D The value to use for the scale component
 	 */
-	[[nodiscard]] FORCEINLINE_DEBUGGABLE TTransform(const TransformVectorRegister& InRotation, const TransformVectorRegister& InTranslation, const TransformVectorRegister& InScale3D)
+	[[nodiscard]] FORCEINLINE_DEBUGGABLE TTransform(TransformVectorRegister InRotation, TransformVectorRegister InTranslation, TransformVectorRegister InScale3D)
 		: Rotation(InRotation),
 		Translation(InTranslation),
 		Scale3D(InScale3D)
@@ -719,12 +717,12 @@ public:
 
 private:
 
-	FORCEINLINE_DEBUGGABLE static bool Private_AnyHasNegativeScale(const TransformVectorRegister& InScale3D, const TransformVectorRegister& InOtherScale3D)
+	FORCEINLINE_DEBUGGABLE static bool Private_AnyHasNegativeScale(TransformVectorRegister InScale3D, TransformVectorRegister InOtherScale3D)
 	{
 		return !!VectorAnyLesserThan(VectorMin(InScale3D, InOtherScale3D), (TransformVectorRegister)GlobalVectorConstants::FloatZero);
 	}
 
-	FORCEINLINE_DEBUGGABLE bool Private_RotationEquals( const TransformVectorRegister& InRotation, const FReal ToleranceScalar = UE_KINDA_SMALL_NUMBER) const
+	FORCEINLINE_DEBUGGABLE bool Private_RotationEquals(TransformVectorRegister InRotation, const FReal ToleranceScalar = UE_KINDA_SMALL_NUMBER) const
 	{			
 		const TransformVectorRegister MyRotation = Rotation; // Load from persistent value, to avoid repeated loads.
 		const TransformVectorRegister Tolerance = VectorLoadFloat1(&ToleranceScalar);
@@ -735,7 +733,7 @@ private:
 		return !VectorAnyGreaterThan(RotationSub, Tolerance) || !VectorAnyGreaterThan(RotationAdd, Tolerance);
 	}
 
-	FORCEINLINE_DEBUGGABLE bool Private_TranslationEquals( const TransformVectorRegister& InTranslation, const FReal ToleranceScalar = UE_KINDA_SMALL_NUMBER) const
+	FORCEINLINE_DEBUGGABLE bool Private_TranslationEquals(TransformVectorRegister InTranslation, const FReal ToleranceScalar = UE_KINDA_SMALL_NUMBER) const
 	{			
 		const TransformVectorRegister Tolerance = VectorLoadFloat1(&ToleranceScalar);
 		// !( (FMath::Abs(X-V.X) > Tolerance) || (FMath::Abs(Y-V.Y) > Tolerance) || (FMath::Abs(Z-V.Z) > Tolerance) )
@@ -743,7 +741,7 @@ private:
 		return !VectorAnyGreaterThan(TranslationDiff, Tolerance);
 	}
 
-	FORCEINLINE_DEBUGGABLE bool Private_Scale3DEquals( const TransformVectorRegister& InScale3D, const FReal ToleranceScalar = UE_KINDA_SMALL_NUMBER) const
+	FORCEINLINE_DEBUGGABLE bool Private_Scale3DEquals(TransformVectorRegister InScale3D, const FReal ToleranceScalar = UE_KINDA_SMALL_NUMBER) const
 	{
 		const TransformVectorRegister Tolerance = VectorLoadFloat1(&ToleranceScalar);
 		// !( (FMath::Abs(X-V.X) > Tolerance) || (FMath::Abs(Y-V.Y) > Tolerance) || (FMath::Abs(Z-V.Z) > Tolerance) )
@@ -971,11 +969,11 @@ public:
 	}
 	
 	// For low-level VectorRegister programming
-	[[nodiscard]] const TPersistentVectorRegisterType<T>& GetTranslationRegister() const
+	[[nodiscard]] TPersistentVectorRegisterType<T> GetTranslationRegister() const
 	{
 		return Translation;
 	}
-	[[nodiscard]] const TPersistentVectorRegisterType<T>& GetRotationRegister() const
+	[[nodiscard]] TPersistentVectorRegisterType<T> GetRotationRegister() const
 	{
 		return Rotation;
 	}
@@ -1300,7 +1298,7 @@ public:
 	}
 
 private:
-	FORCEINLINE_DEBUGGABLE void ToMatrixInternal( TransformVectorRegister& OutDiagonals, TransformVectorRegister& OutAdds, TransformVectorRegister& OutSubtracts ) const
+	FORCEINLINE_DEBUGGABLE void ToMatrixInternal(TransformVectorRegister& OutDiagonals, TransformVectorRegister& OutAdds, TransformVectorRegister& OutSubtracts) const
 	{
 #if !(UE_BUILD_SHIPPING || UE_BUILD_TEST) && WITH_EDITORONLY_DATA
 		// Make sure Rotation is normalized when we turn it into a matrix.
@@ -1356,7 +1354,7 @@ private:
 		OutSubtracts = VectorMultiply(Subtracts , Scale3DYZXW);
 	}
 
-	FORCEINLINE_DEBUGGABLE void ToMatrixInternalNoScale( TransformVectorRegister& OutDiagonals, TransformVectorRegister& OutAdds, TransformVectorRegister& OutSubtracts ) const
+	FORCEINLINE_DEBUGGABLE void ToMatrixInternalNoScale(TransformVectorRegister& OutDiagonals, TransformVectorRegister& OutAdds, TransformVectorRegister& OutSubtracts) const
 	{
 #if !(UE_BUILD_SHIPPING || UE_BUILD_TEST) && WITH_EDITORONLY_DATA
 		// Make sure Rotation is normalized when we turn it into a matrix.
@@ -1404,7 +1402,7 @@ private:
 	 * also returning BIG_NUMBER causes sequential NaN issues by multiplying 
 	 * so we hardcode as 0
 	 */
-	[[nodiscard]] static FORCEINLINE_DEBUGGABLE TransformVectorRegister		GetSafeScaleReciprocal(const TransformVectorRegister& InScale, const ScalarRegister& Tolerance = ScalarRegister(GlobalVectorConstants::SmallNumber))
+	[[nodiscard]] static FORCEINLINE_DEBUGGABLE TransformVectorRegister	GetSafeScaleReciprocal(TransformVectorRegister InScale, const ScalarRegister& Tolerance = ScalarRegister(GlobalVectorConstants::SmallNumber))
 	{		
 		// SafeReciprocalScale.X = (InScale.X == 0) ? 0.f : 1/InScale.X; // same for YZW
 
@@ -1477,7 +1475,7 @@ private:
 	* @param	DesiredScale - there is no check on if the magnitude is correct here. It assumes that is correct. 
 	* @param	OutTransform the constructed transform 
 	*/
-	FORCEINLINE_DEBUGGABLE static void ConstructTransformFromMatrixWithDesiredScale(const TMatrix<T>& AMatrix, const TMatrix<T>& BMatrix, const TransformVectorRegister& DesiredScale, TTransform<T>& OutTransform);
+	FORCEINLINE_DEBUGGABLE static void ConstructTransformFromMatrixWithDesiredScale(const TMatrix<T>& AMatrix, const TMatrix<T>& BMatrix, TransformVectorRegister DesiredScale, TTransform<T>& OutTransform);
 	/**
 	* Create a new transform: OutTransform = Base * Relative(-1) using the matrix while keeping the scale that's given by Base and Relative
 	* Please note that this operation is a lot more expensive than normal GetRelativeTrnasform
@@ -1555,7 +1553,7 @@ FORCEINLINE_DEBUGGABLE void TTransform<T>::MultiplyUsingMatrixWithScale(TTransfo
 }
 
 template<typename T>
-FORCEINLINE_DEBUGGABLE void TTransform<T>::ConstructTransformFromMatrixWithDesiredScale(const UE::Math::TMatrix<T>& AMatrix, const UE::Math::TMatrix<T>& BMatrix, const TransformVectorRegister& DesiredScale, TTransform<T>& OutTransform)
+FORCEINLINE_DEBUGGABLE void TTransform<T>::ConstructTransformFromMatrixWithDesiredScale(const UE::Math::TMatrix<T>& AMatrix, const UE::Math::TMatrix<T>& BMatrix, TransformVectorRegister DesiredScale, TTransform<T>& OutTransform)
 {
 	// the goal of using M is to get the correct orientation
 	// but for translation, we still need scale
@@ -1925,7 +1923,7 @@ FORCEINLINE_DEBUGGABLE void TTransform<T>::operator*=(const TQuat<T>& Other)
 
 // x = 0, y = 1, z = 2
 template<typename T>
-FORCEINLINE_DEBUGGABLE TVector<T> TTransform<T>::GetScaledAxis( EAxis::Type InAxis ) const
+FORCEINLINE_DEBUGGABLE TVector<T> TTransform<T>::GetScaledAxis(EAxis::Type InAxis) const
 {
 	if ( InAxis == EAxis::X )
 	{
@@ -1941,7 +1939,7 @@ FORCEINLINE_DEBUGGABLE TVector<T> TTransform<T>::GetScaledAxis( EAxis::Type InAx
 
 // x = 0, y = 1, z = 2
 template<typename T>
-FORCEINLINE_DEBUGGABLE TVector<T> TTransform<T>::GetUnitAxis( EAxis::Type InAxis ) const
+FORCEINLINE_DEBUGGABLE TVector<T> TTransform<T>::GetUnitAxis(EAxis::Type InAxis) const
 {
 	if ( InAxis == EAxis::X )
 	{
diff --git a/Engine/Source/Runtime/Core/Public/Math/UnrealMathNeon.h b/Engine/Source/Runtime/Core/Public/Math/UnrealMathNeon.h
--- a/Engine/Source/Runtime/Core/Public/Math/UnrealMathNeon.h
+++ b/Engine/Source/Runtime/Core/Public/Math/UnrealMathNeon.h
@@ -147,18 +147,18 @@ struct alignas(16) VectorRegister4Double
 	};
 
 	FORCEINLINE VectorRegister4Double() = default;
-	FORCEINLINE VectorRegister4Double(const VectorRegister2Double& xy, const VectorRegister2Double& zw) : XY(xy), ZW(zw) {}
+	FORCEINLINE VectorRegister4Double(VectorRegister2Double xy, VectorRegister2Double zw) : XY(xy), ZW(zw) {}
 	FORCEINLINE constexpr VectorRegister4Double(VectorRegister2Double xy, VectorRegister2Double zw, VectorRegisterConstInit) : XY(xy), ZW(zw) {}
 
-	FORCEINLINE VectorRegister4Double(const VectorRegister4Float From)
+	FORCEINLINE VectorRegister4Double(VectorRegister4Float From)
 	{
 		XY = vcvt_f64_f32(*(float32x2_t*)&From);
 		ZW = vcvt_high_f64_f32(From);
 	}
 
-	VectorRegister4Double(const VectorRegister2Double& From) = delete;
+	VectorRegister4Double(VectorRegister2Double From) = delete;
 
-	FORCEINLINE VectorRegister4Double& operator=(const VectorRegister4Float& From)
+	FORCEINLINE VectorRegister4Double& operator=(VectorRegister4Float From)
 	{
 		*this = VectorRegister4Double(From);
 		return *this;
@@ -172,8 +172,8 @@ typedef VectorRegister4Double VectorRegister;
 // Forward declarations
 VectorRegister4Float VectorLoadAligned(const float* Ptr);
 VectorRegister4Double VectorLoadAligned(const double* Ptr);
-void VectorStoreAligned(const VectorRegister4Float& Vec, float* Ptr);
-void VectorStoreAligned(const VectorRegister4Double& Vec, double* Dst);
+void VectorStoreAligned(VectorRegister4Float Vec, float* Ptr);
+void VectorStoreAligned(VectorRegister4Double Vec, double* Dst);
 
 
 // Helper for conveniently aligning a float array for extraction from VectorRegister4Float
@@ -181,7 +181,7 @@ struct alignas(alignof(VectorRegister4Float)) AlignedFloat4
 {
 	float V[4];
 
-	FORCEINLINE AlignedFloat4(const VectorRegister4Float& Vec)
+	FORCEINLINE AlignedFloat4(VectorRegister4Float Vec)
 	{
 		VectorStoreAligned(Vec, V);
 	}
@@ -201,7 +201,7 @@ struct alignas(alignof(VectorRegister4Double)) AlignedDouble4
 {
 	double V[4];
 
-	FORCEINLINE AlignedDouble4(const VectorRegister4Double& Vec)
+	FORCEINLINE AlignedDouble4(VectorRegister4Double Vec)
 	{
 		VectorStoreAligned(Vec, V);
 	}
@@ -237,7 +237,7 @@ typedef VectorRegister4Int VectorRegisterInt;
  * @param W		4th uint32 component
  * @return		Bitwise equivalent vector with 4 floats
  */
-FORCEINLINE VectorRegister4Float MakeVectorRegister( uint32 X, uint32 Y, uint32 Z, uint32 W )
+FORCEINLINE VectorRegister4Float MakeVectorRegister(uint32 X, uint32 Y, uint32 Z, uint32 W)
 {
 	union U {
 		VectorRegister4Float V; uint32 F[4];
@@ -271,7 +271,7 @@ FORCEINLINE VectorRegister4Float MakeVectorRegisterFloatMask(uint32 X, uint32 Y,
  * @param W		4th float component
  * @return		Vector of the 4 floats
  */
-FORCEINLINE VectorRegister4Float MakeVectorRegister( float X, float Y, float Z, float W )
+FORCEINLINE VectorRegister4Float MakeVectorRegister(float X, float Y, float Z, float W)
 {
 	union U {
 		VectorRegister4Float V; float F[4];
@@ -317,7 +317,7 @@ FORCEINLINE VectorRegister4Double MakeVectorRegisterDouble(double X, double Y, d
 	return MakeVectorRegister(X, Y, Z, W);
 }
 
-FORCEINLINE VectorRegister4Double MakeVectorRegisterDouble(const VectorRegister2Double& XY, const VectorRegister2Double& ZW)
+FORCEINLINE VectorRegister4Double MakeVectorRegisterDouble(VectorRegister2Double XY, VectorRegister2Double ZW)
 {
 	return VectorRegister4Double(XY, ZW);
 }
@@ -401,13 +401,13 @@ FORCEINLINE VectorRegister4Int MakeVectorRegisterInt64(int64 X, int64 Y)
 }
 
 // Make double register from float register
-FORCEINLINE VectorRegister4Double MakeVectorRegisterDouble(const VectorRegister4Float& From)
+FORCEINLINE VectorRegister4Double MakeVectorRegisterDouble(VectorRegister4Float From)
 {
 	return VectorRegister4Double(From);
 }
 
 // Lossy conversion: double->float vector
-FORCEINLINE VectorRegister4Float MakeVectorRegisterFloatFromDouble(const VectorRegister4Double& Vec)
+FORCEINLINE VectorRegister4Float MakeVectorRegisterFloatFromDouble(VectorRegister4Double Vec)
 {
 	return vcombine_f32(vcvt_f32_f64(Vec.XY), vcvt_f32_f64(Vec.ZW));
 }
@@ -427,7 +427,7 @@ FORCEINLINE VectorRegister4Float VectorLoadNPlusOneUnalignedBytes(const void* Pt
  *	Constants:
  *============================================================================*/
 
-#include "Math/UnrealMathVectorConstants.h"
+#include "Math/UnrealMathVectorConstants.h.inl"
 
 
 /*=============================================================================
@@ -541,19 +541,19 @@ FORCEINLINE VectorRegister4Double VectorLoadFloat3_W1(const double* Ptr)
  * Sets a single component of a vector. Must be a define since ElementIndex needs to be a constant integer
  */
 template <int ElementIndex>
-FORCEINLINE VectorRegister4Float VectorSetComponentImpl(const VectorRegister4Float& Vec, float Scalar)
+FORCEINLINE VectorRegister4Float VectorSetComponentImpl(VectorRegister4Float Vec, float Scalar)
 {
 	return vsetq_lane_f32(Scalar, Vec, ElementIndex);
 }
 
 template <int ElementIndex>
-FORCEINLINE VectorRegister2Double VectorSetComponentImpl(const VectorRegister2Double& Vec, double Scalar)
+FORCEINLINE VectorRegister2Double VectorSetComponentImpl(VectorRegister2Double Vec, double Scalar)
 {
 	return vsetq_lane_f64(Scalar, Vec, ElementIndex);
 }
 
 template <int ElementIndex>
-FORCEINLINE VectorRegister4Double VectorSetComponentImpl(const VectorRegister4Double& Vec, double Scalar)
+FORCEINLINE VectorRegister4Double VectorSetComponentImpl(VectorRegister4Double Vec, double Scalar)
 {
 	VectorRegister4Double Result;
 	if constexpr (ElementIndex > 1)
@@ -578,7 +578,7 @@ FORCEINLINE VectorRegister4Double VectorSetComponentImpl(const VectorRegister4Do
  * @param Ptr	Aligned memory pointer to the 4 floats
  * @return		VectorRegister4Float(Ptr[0], Ptr[1], Ptr[2], Ptr[3])
  */
-FORCEINLINE VectorRegister4Float VectorLoadAligned( const float* Ptr )
+FORCEINLINE VectorRegister4Float VectorLoadAligned(const float* Ptr)
 {
 	return VectorLoad(Ptr);
 }
@@ -594,7 +594,7 @@ FORCEINLINE VectorRegister4Double VectorLoadAligned(const double* Ptr)
  * @param Ptr	Unaligned memory pointer to the float
  * @return		VectorRegister4Float(Ptr[0], Ptr[0], Ptr[0], Ptr[0])
  */
-FORCEINLINE VectorRegister4Float VectorLoadFloat1( const float *Ptr )
+FORCEINLINE VectorRegister4Float VectorLoadFloat1(const float *Ptr)
 {
 	return vdupq_n_f32(Ptr[0]);
 }
@@ -660,12 +660,12 @@ FORCEINLINE VectorRegister4Double VectorSetFloat1(double X)
  * @param Vec	Vector to store
  * @param Ptr	Aligned memory pointer
  */
-FORCEINLINE void VectorStoreAligned(const VectorRegister4Float& Vec, float* Ptr)
+FORCEINLINE void VectorStoreAligned(VectorRegister4Float Vec, float* Ptr)
 {
 	vst1q_f32(Ptr, Vec);
 }
 
-FORCEINLINE void VectorStoreAligned(const VectorRegister4Double& Vec, double* Ptr)
+FORCEINLINE void VectorStoreAligned(VectorRegister4Double Vec, double* Ptr)
 {
 	vst1q_f64_x2(Ptr, *(float64x2x2_t*)&Vec);
 }
@@ -686,7 +686,7 @@ FORCEINLINE void VectorStoreAligned(VectorRegister4Float Vec, FFloat16* Ptr)
 * @param Vec	Vector to store
 * @param Ptr	Aligned memory pointer
 */
-#define VectorStoreAlignedStreamed( Vec, Ptr )	VectorStoreAligned( Vec, Ptr )
+#define VectorStoreAlignedStreamed(Vec, Ptr)	VectorStoreAligned(Vec, Ptr)
 
 /**
  * Stores a vector to memory (aligned or unaligned).
@@ -694,12 +694,12 @@ FORCEINLINE void VectorStoreAligned(VectorRegister4Float Vec, FFloat16* Ptr)
  * @param Vec	Vector to store
  * @param Ptr	Memory pointer
  */
-FORCEINLINE void VectorStore(const VectorRegister4Float& Vec, float* Ptr)
+FORCEINLINE void VectorStore(VectorRegister4Float Vec, float* Ptr)
 {
 	vst1q_f32(Ptr, Vec);
 }
 
-FORCEINLINE void VectorStore(const VectorRegister4Double& Vec, double* Ptr)
+FORCEINLINE void VectorStore(VectorRegister4Double Vec, double* Ptr)
 {
 	vst1q_f64_x2(Ptr, *(float64x2x2_t*)&Vec);
 }
@@ -710,7 +710,7 @@ FORCEINLINE void VectorStore(const VectorRegister4Double& Vec, double* Ptr)
  * @param Vec	Vector to store
  * @param Ptr	Memory pointer
  */
-FORCEINLINE void VectorStore16(const VectorRegister4x4Float& Vec, float* Ptr)
+FORCEINLINE void VectorStore16(VectorRegister4x4Float Vec, float* Ptr)
 {
 	vst1q_f32_x4(Ptr, Vec);
 }
@@ -721,7 +721,7 @@ FORCEINLINE void VectorStore16(const VectorRegister4x4Float& Vec, float* Ptr)
  * @param Vec	Vector to store XYZ
  * @param Ptr	Unaligned memory pointer
  */
-FORCEINLINE void VectorStoreFloat3( const VectorRegister4Float& Vec, float* Ptr )
+FORCEINLINE void VectorStoreFloat3(VectorRegister4Float Vec, float* Ptr)
 {
 	vst1_f32(Ptr, *(float32x2_t*)&Vec);
 	vst1q_lane_f32(((float32_t*)Ptr) + 2, Vec, 2);
@@ -733,7 +733,7 @@ FORCEINLINE void VectorStoreFloat3( const VectorRegister4Float& Vec, float* Ptr
  * @param Vec	Vector to store XYZ
  * @param Ptr	Unaligned memory pointer
  */
-FORCEINLINE void VectorStoreFloat3(const VectorRegister4Double& Vec, double* Ptr)
+FORCEINLINE void VectorStoreFloat3(VectorRegister4Double Vec, double* Ptr)
 {
 	vst1q_f64(Ptr, Vec.XY);
 	vst1q_lane_f64(((float64_t*)Ptr) + 2, Vec.ZW, 0);
@@ -746,12 +746,12 @@ FORCEINLINE void VectorStoreFloat3(const VectorRegister4Double& Vec, double* Ptr
  * @param Vec	Vector to store X
  * @param Ptr	Unaligned memory pointer
  */
-FORCEINLINE void VectorStoreFloat1(VectorRegister4Float Vec, float* Ptr )
+FORCEINLINE void VectorStoreFloat1(VectorRegister4Float Vec, float* Ptr)
 {
 	vst1q_lane_f32( Ptr, Vec, 0 );
 }
 
-FORCEINLINE void VectorStoreFloat1(const VectorRegister4Double& Vec, double* Ptr)
+FORCEINLINE void VectorStoreFloat1(VectorRegister4Double Vec, double* Ptr)
 {
 	vst1q_lane_f64(Ptr, Vec.XY, 0);
 }
@@ -765,19 +765,19 @@ FORCEINLINE void VectorStoreFloat1(const VectorRegister4Double& Vec, double* Ptr
  * @return				VectorRegister4Float( Vec[ElementIndex], Vec[ElementIndex], Vec[ElementIndex], Vec[ElementIndex] )
  */
 template <int ElementIndex>
-FORCEINLINE VectorRegister4Float VectorReplicateImpl(const VectorRegister4Float& Vec)
+FORCEINLINE VectorRegister4Float VectorReplicateImpl(VectorRegister4Float Vec)
 {
 	return vdupq_n_f32(vgetq_lane_f32(Vec, ElementIndex));
 }
 
 template <int ElementIndex>
-FORCEINLINE VectorRegister2Double VectorReplicateImpl(const VectorRegister2Double& Vec)
+FORCEINLINE VectorRegister2Double VectorReplicateImpl(VectorRegister2Double Vec)
 {
 	return vdupq_n_f64(vgetq_lane_f64(Vec, ElementIndex));
 }
 
 template <int ElementIndex>
-FORCEINLINE VectorRegister4Double VectorReplicateImpl(const VectorRegister4Double& Vec)
+FORCEINLINE VectorRegister4Double VectorReplicateImpl(VectorRegister4Double Vec)
 {
 	VectorRegister4Double Result;
 	if constexpr (ElementIndex <= 1)
@@ -802,7 +802,7 @@ FORCEINLINE VectorRegister4Double VectorReplicateImpl(const VectorRegister4Doubl
  * @param Vec			Source vector
  * @return				VectorRegister4Float( abs(Vec.x), abs(Vec.y), abs(Vec.z), abs(Vec.w) )
  */
-FORCEINLINE VectorRegister4Float VectorAbs(VectorRegister4Float Vec )
+FORCEINLINE VectorRegister4Float VectorAbs(VectorRegister4Float Vec)
 {
 	return vabsq_f32( Vec );
 }
@@ -821,7 +821,7 @@ FORCEINLINE VectorRegister4Double VectorAbs(VectorRegister4Double Vec)
  * @param Vec			Source vector
  * @return				VectorRegister4Float( -Vec.x, -Vec.y, -Vec.z, -Vec.w )
  */
-FORCEINLINE VectorRegister4Float VectorNegate( VectorRegister4Float Vec )
+FORCEINLINE VectorRegister4Float VectorNegate(VectorRegister4Float Vec)
 {
 	return vnegq_f32( Vec );
 }
@@ -841,7 +841,7 @@ FORCEINLINE VectorRegister4Double VectorNegate(VectorRegister4Double Vec)
  * @param Vec2	2nd vector
  * @return		VectorRegister4Float( Vec1.x+Vec2.x, Vec1.y+Vec2.y, Vec1.z+Vec2.z, Vec1.w+Vec2.w )
  */
-FORCEINLINE VectorRegister4Float VectorAdd( VectorRegister4Float Vec1, VectorRegister4Float Vec2 )
+FORCEINLINE VectorRegister4Float VectorAdd(VectorRegister4Float Vec1, VectorRegister4Float Vec2)
 {
 	return vaddq_f32( Vec1, Vec2 );
 }
@@ -862,7 +862,7 @@ FORCEINLINE VectorRegister4Double VectorAdd(VectorRegister4Double Vec1, VectorRe
  * @param Vec2	2nd vector
  * @return		VectorRegister4Float( Vec1.x-Vec2.x, Vec1.y-Vec2.y, Vec1.z-Vec2.z, Vec1.w-Vec2.w )
  */
-FORCEINLINE VectorRegister4Float VectorSubtract( VectorRegister4Float Vec1, VectorRegister4Float Vec2 )
+FORCEINLINE VectorRegister4Float VectorSubtract(VectorRegister4Float Vec1, VectorRegister4Float Vec2)
 {
 	return vsubq_f32( Vec1, Vec2 );
 }
@@ -883,7 +883,7 @@ FORCEINLINE VectorRegister4Double VectorSubtract(VectorRegister4Double Vec1, Vec
  * @param Vec2	2nd vector
  * @return		VectorRegister4Float( Vec1.x*Vec2.x, Vec1.y*Vec2.y, Vec1.z*Vec2.z, Vec1.w*Vec2.w )
  */
-FORCEINLINE VectorRegister4Float VectorMultiply( VectorRegister4Float Vec1, VectorRegister4Float Vec2 ) 
+FORCEINLINE VectorRegister4Float VectorMultiply(VectorRegister4Float Vec1, VectorRegister4Float Vec2) 
 {
 	return vmulq_f32( Vec1, Vec2 );
 }
@@ -931,7 +931,7 @@ FORCEINLINE VectorRegister4Double VectorDivide(VectorRegister4Double Vec1, Vecto
  * @param Acc	3rd vector
  * @return		VectorRegister4Float( Vec1.x*Vec2.x + Acc.x, Vec1.y*Vec2.y + Acc.y, Vec1.z*Vec2.z + Acc.z, Vec1.w*Vec2.w + Acc.w )
  */
-FORCEINLINE VectorRegister4Float VectorMultiplyAdd( VectorRegister4Float Vec1, VectorRegister4Float Vec2, VectorRegister4Float Acc )
+FORCEINLINE VectorRegister4Float VectorMultiplyAdd(VectorRegister4Float Vec1, VectorRegister4Float Vec2, VectorRegister4Float Acc)
 {
 	return vfmaq_f32(Acc, Vec1, Vec2 );
 }
@@ -974,7 +974,7 @@ FORCEINLINE VectorRegister4Double VectorNegateMultiplyAdd(VectorRegister4Double
  * @param Vec2	2nd vector
  * @return		d = dot3(Vec1.xyz, Vec2.xyz), VectorRegister4Float( d, d, d, d )
  */
-FORCEINLINE VectorRegister4Float VectorDot3( const VectorRegister4Float& Vec1, const VectorRegister4Float& Vec2 )
+FORCEINLINE VectorRegister4Float VectorDot3(VectorRegister4Float Vec1, VectorRegister4Float Vec2)
 {
 	VectorRegister4Float Temp = VectorMultiply( Vec1, Vec2 );
 	Temp = vsetq_lane_f32( 0.0f, Temp, 3 );
@@ -983,7 +983,7 @@ FORCEINLINE VectorRegister4Float VectorDot3( const VectorRegister4Float& Vec1, c
 	return vdupq_lane_f32( sum, 0 );
 }
 
-FORCEINLINE VectorRegister4Double VectorDot3(const VectorRegister4Double& Vec1, const VectorRegister4Double& Vec2)
+FORCEINLINE VectorRegister4Double VectorDot3(VectorRegister4Double Vec1, VectorRegister4Double Vec2)
 {
 	VectorRegister2Double A, B;
 	A = vmulq_f64(Vec1.XY, Vec2.XY);
@@ -995,12 +995,12 @@ FORCEINLINE VectorRegister4Double VectorDot3(const VectorRegister4Double& Vec1,
 	return Temp;
 }
 
-FORCEINLINE float VectorDot3Scalar(const VectorRegister4Float& Vec1, const VectorRegister4Float& Vec2)
+FORCEINLINE float VectorDot3Scalar(VectorRegister4Float Vec1, VectorRegister4Float Vec2)
 {
 	return vgetq_lane_f32(VectorDot3(Vec1, Vec2), 0);
 }
 
-FORCEINLINE double VectorDot3Scalar(const VectorRegister4Double& Vec1, const VectorRegister4Double& Vec2)
+FORCEINLINE double VectorDot3Scalar(VectorRegister4Double Vec1, VectorRegister4Double Vec2)
 {
 	VectorRegister2Double A, B;
 	A = vmulq_f64(Vec1.XY, Vec2.XY);
@@ -1047,12 +1047,12 @@ FORCEINLINE VectorRegister4Double VectorDot4(VectorRegister4Double Vec1, VectorR
  * @return		VectorRegister4Float( Vec1.x == Vec2.x ? 0xFFFFFFFF : 0, same for yzw )
  */
 
-FORCEINLINE VectorRegister4Float VectorCompareEQ( const VectorRegister4Float& Vec1, const VectorRegister4Float& Vec2 )
+FORCEINLINE VectorRegister4Float VectorCompareEQ(VectorRegister4Float Vec1, VectorRegister4Float Vec2)
 {
 	return (VectorRegister4Float)vceqq_f32( Vec1, Vec2 );
 }
 
-FORCEINLINE VectorRegister4Double VectorCompareEQ(const VectorRegister4Double& Vec1, const VectorRegister4Double& Vec2)
+FORCEINLINE VectorRegister4Double VectorCompareEQ(VectorRegister4Double Vec1, VectorRegister4Double Vec2)
 {
 	VectorRegister4Double Result;
 	Result.XY = (VectorRegister2Double)vceqq_f64(Vec1.XY, Vec2.XY);
@@ -1070,12 +1070,12 @@ FORCEINLINE VectorRegister4Double VectorCompareEQ(const VectorRegister4Double& V
  * @return		VectorRegister4Float( Vec1.x != Vec2.x ? 0xFFFFFFFF : 0, same for yzw )
  */
 
-FORCEINLINE VectorRegister4Float VectorCompareNE( const VectorRegister4Float& Vec1, const VectorRegister4Float& Vec2 )
+FORCEINLINE VectorRegister4Float VectorCompareNE(VectorRegister4Float Vec1, VectorRegister4Float Vec2)
 {
 	return (VectorRegister4Float)vmvnq_u32( vceqq_f32( Vec1, Vec2 ) );
 }
 
-FORCEINLINE VectorRegister4Double VectorCompareNE(const VectorRegister4Double& Vec1, const VectorRegister4Double& Vec2)
+FORCEINLINE VectorRegister4Double VectorCompareNE(VectorRegister4Double Vec1, VectorRegister4Double Vec2)
 {
 	VectorRegister4Double Result;
 	Result.XY = (VectorRegister2Double)vmvnq_u32(vceqq_f64(Vec1.XY, Vec2.XY));
@@ -1091,12 +1091,12 @@ FORCEINLINE VectorRegister4Double VectorCompareNE(const VectorRegister4Double& V
  * @return		VectorRegister4Float( Vec1.x > Vec2.x ? 0xFFFFFFFF : 0, same for yzw )
  */
 
-FORCEINLINE VectorRegister4Float VectorCompareGT( const VectorRegister4Float& Vec1, const VectorRegister4Float& Vec2 )
+FORCEINLINE VectorRegister4Float VectorCompareGT(VectorRegister4Float Vec1, VectorRegister4Float Vec2)
 {
 	return (VectorRegister4Float)vcgtq_f32( Vec1, Vec2 );
 }
 
-FORCEINLINE VectorRegister4Double VectorCompareGT(const VectorRegister4Double& Vec1, const VectorRegister4Double& Vec2)
+FORCEINLINE VectorRegister4Double VectorCompareGT(VectorRegister4Double Vec1, VectorRegister4Double Vec2)
 {
 	VectorRegister4Double Result;
 	Result.XY = (VectorRegister2Double)vcgtq_f64(Vec1.XY, Vec2.XY);
@@ -1112,12 +1112,12 @@ FORCEINLINE VectorRegister4Double VectorCompareGT(const VectorRegister4Double& V
  * @return		VectorRegister4Float( Vec1.x >= Vec2.x ? 0xFFFFFFFF : 0, same for yzw )
  */
 
-FORCEINLINE VectorRegister4Float VectorCompareGE( const VectorRegister4Float& Vec1, const VectorRegister4Float& Vec2 )
+FORCEINLINE VectorRegister4Float VectorCompareGE(VectorRegister4Float Vec1, VectorRegister4Float Vec2)
 {
 	return (VectorRegister4Float)vcgeq_f32( Vec1, Vec2 );
 }
 
-FORCEINLINE VectorRegister4Double VectorCompareGE(const VectorRegister4Double& Vec1, const VectorRegister4Double& Vec2)
+FORCEINLINE VectorRegister4Double VectorCompareGE(VectorRegister4Double Vec1, VectorRegister4Double Vec2)
 {
 	VectorRegister4Double Result;
 	Result.XY = (VectorRegister2Double)vcgeq_f64(Vec1.XY, Vec2.XY);
@@ -1132,12 +1132,12 @@ FORCEINLINE VectorRegister4Double VectorCompareGE(const VectorRegister4Double& V
 * @param Vec2	2nd vector
 * @return		VectorRegister4Float( Vec1.x < Vec2.x ? 0xFFFFFFFF : 0, same for yzw )
 */
-FORCEINLINE VectorRegister4Float VectorCompareLT(const VectorRegister4Float& Vec1, const VectorRegister4Float& Vec2)
+FORCEINLINE VectorRegister4Float VectorCompareLT(VectorRegister4Float Vec1, VectorRegister4Float Vec2)
 {
 	return (VectorRegister4Float)vcltq_f32(Vec1, Vec2);
 }
 
-FORCEINLINE VectorRegister4Double VectorCompareLT(const VectorRegister4Double& Vec1, const VectorRegister4Double& Vec2)
+FORCEINLINE VectorRegister4Double VectorCompareLT(VectorRegister4Double Vec1, VectorRegister4Double Vec2)
 {
 	VectorRegister4Double Res;
 	Res.XY = (VectorRegister2Double)vcltq_f64(Vec1.XY, Vec2.XY);
@@ -1152,12 +1152,12 @@ FORCEINLINE VectorRegister4Double VectorCompareLT(const VectorRegister4Double& V
 * @param Vec2	2nd vector
 * @return		VectorRegister4Float( Vec1.x <= Vec2.x ? 0xFFFFFFFF : 0, same for yzw )
 */
-FORCEINLINE VectorRegister4Float VectorCompareLE(const VectorRegister4Float& Vec1, const VectorRegister4Float& Vec2)
+FORCEINLINE VectorRegister4Float VectorCompareLE(VectorRegister4Float Vec1, VectorRegister4Float Vec2)
 {
 	return (VectorRegister4Float)vcleq_f32(Vec1, Vec2);
 }
 
-FORCEINLINE VectorRegister4Double VectorCompareLE(const VectorRegister4Double& Vec1, const VectorRegister4Double& Vec2)
+FORCEINLINE VectorRegister4Double VectorCompareLE(VectorRegister4Double Vec1, VectorRegister4Double Vec2)
 {
 	VectorRegister4Double Res;
 	Res.XY = (VectorRegister2Double)vcleq_f64(Vec1.XY, Vec2.XY);
@@ -1175,12 +1175,12 @@ FORCEINLINE VectorRegister4Double VectorCompareLE(const VectorRegister4Double& V
  *
  */
 
-FORCEINLINE VectorRegister4Float VectorSelect(const VectorRegister4Float& Mask, const VectorRegister4Float& Vec1, const VectorRegister4Float& Vec2 )
+FORCEINLINE VectorRegister4Float VectorSelect(VectorRegister4Float Mask, VectorRegister4Float Vec1, VectorRegister4Float Vec2)
 {
 	return vbslq_f32((VectorRegister4Int)Mask, Vec1, Vec2);
 }
 
-FORCEINLINE VectorRegister4Double VectorSelect(const VectorRegister4Double& Mask, const VectorRegister4Double& Vec1, const VectorRegister4Double& Vec2)
+FORCEINLINE VectorRegister4Double VectorSelect(VectorRegister4Double Mask, VectorRegister4Double Vec1, VectorRegister4Double Vec2)
 {
 	VectorRegister4Double Result;
 	Result.XY = vbslq_f64((VectorRegister2Int64)Mask.XY, Vec1.XY, Vec2.XY);
@@ -1195,12 +1195,12 @@ FORCEINLINE VectorRegister4Double VectorSelect(const VectorRegister4Double& Mask
  * @param Vec2	2nd vector
  * @return		VectorRegister4Float( for each bit i: Vec1[i] | Vec2[i] )
  */
-FORCEINLINE VectorRegister4Float VectorBitwiseOr(const VectorRegister4Float& Vec1, const VectorRegister4Float& Vec2 )
+FORCEINLINE VectorRegister4Float VectorBitwiseOr(VectorRegister4Float Vec1, VectorRegister4Float Vec2)
 {
 	return (VectorRegister4Float)vorrq_u32( (VectorRegister4Int)Vec1, (VectorRegister4Int)Vec2 );
 }
 
-FORCEINLINE VectorRegister4Double VectorBitwiseOr(const VectorRegister4Double& Vec1, const VectorRegister4Double& Vec2)
+FORCEINLINE VectorRegister4Double VectorBitwiseOr(VectorRegister4Double Vec1, VectorRegister4Double Vec2)
 {
 	VectorRegister4Double Result;
 	Result.XY = (VectorRegister2Double)vorrq_u64((VectorRegister2Int64)Vec1.XY, (VectorRegister2Int64)Vec2.XY);
@@ -1215,12 +1215,12 @@ FORCEINLINE VectorRegister4Double VectorBitwiseOr(const VectorRegister4Double& V
  * @param Vec2	2nd vector
  * @return		VectorRegister4Float( for each bit i: Vec1[i] & Vec2[i] )
  */
-FORCEINLINE VectorRegister4Float VectorBitwiseAnd(const VectorRegister4Float& Vec1, const VectorRegister4Float& Vec2 )
+FORCEINLINE VectorRegister4Float VectorBitwiseAnd(VectorRegister4Float Vec1, VectorRegister4Float Vec2)
 {
 	return (VectorRegister4Float)vandq_u32( (VectorRegister4Int)Vec1, (VectorRegister4Int)Vec2 );
 }
 
-FORCEINLINE VectorRegister4Double VectorBitwiseAnd(const VectorRegister4Double& Vec1, const VectorRegister4Double& Vec2)
+FORCEINLINE VectorRegister4Double VectorBitwiseAnd(VectorRegister4Double Vec1, VectorRegister4Double Vec2)
 {
 	VectorRegister4Double Result;
 	Result.XY = (VectorRegister2Double)vandq_u64((VectorRegister2Int64)Vec1.XY, (VectorRegister2Int64)Vec2.XY);
@@ -1235,12 +1235,12 @@ FORCEINLINE VectorRegister4Double VectorBitwiseAnd(const VectorRegister4Double&
  * @param Vec2	2nd vector
  * @return		VectorRegister4Float( for each bit i: Vec1[i] ^ Vec2[i] )
  */
-FORCEINLINE VectorRegister4Float VectorBitwiseXor(const VectorRegister4Float& Vec1, const VectorRegister4Float& Vec2 )
+FORCEINLINE VectorRegister4Float VectorBitwiseXor(VectorRegister4Float Vec1, VectorRegister4Float Vec2)
 {
 	return (VectorRegister4Float)veorq_u32( (VectorRegister4Int)Vec1, (VectorRegister4Int)Vec2 );
 }
 
-FORCEINLINE VectorRegister4Double VectorBitwiseXor(const VectorRegister4Double& Vec1, const VectorRegister4Double& Vec2)
+FORCEINLINE VectorRegister4Double VectorBitwiseXor(VectorRegister4Double Vec1, VectorRegister4Double Vec2)
 {
 	VectorRegister4Double Result;
 	Result.XY = (VectorRegister2Double)veorq_u64((VectorRegister2Int64)Vec1.XY, (VectorRegister2Int64)Vec2.XY);
@@ -1509,12 +1509,12 @@ FORCEINLINE uint32 VectorMaskBits(VectorRegister4Double VecMask)
 * @param Vec2		Source vector2
 * @return			The combined vector
 */
-FORCEINLINE VectorRegister4Float VectorCombineHigh(const VectorRegister4Float& Vec1, const VectorRegister4Float& Vec2)
+FORCEINLINE VectorRegister4Float VectorCombineHigh(VectorRegister4Float Vec1, VectorRegister4Float Vec2)
 {
 	return vcombine_f32(vget_high_f32(Vec1), vget_high_f32(Vec2));
 }
 
-FORCEINLINE VectorRegister4Double VectorCombineHigh(const VectorRegister4Double& Vec1, const VectorRegister4Double& Vec2)
+FORCEINLINE VectorRegister4Double VectorCombineHigh(VectorRegister4Double Vec1, VectorRegister4Double Vec2)
 {
 	VectorRegister4Double Result;
 	Result.XY = Vec1.ZW;
@@ -1529,12 +1529,12 @@ FORCEINLINE VectorRegister4Double VectorCombineHigh(const VectorRegister4Double&
 * @param Vec2		Source vector2
 * @return			The combined vector
 */
-FORCEINLINE VectorRegister4Float VectorCombineLow(const VectorRegister4Float& Vec1, const VectorRegister4Float& Vec2)
+FORCEINLINE VectorRegister4Float VectorCombineLow(VectorRegister4Float Vec1, VectorRegister4Float Vec2)
 {
 	return vcombine_f32(vget_low_f32(Vec1), vget_low_f32(Vec2));
 }
 
-FORCEINLINE VectorRegister4Double VectorCombineLow(const VectorRegister4Double& Vec1, const VectorRegister4Double& Vec2)
+FORCEINLINE VectorRegister4Double VectorCombineLow(VectorRegister4Double Vec1, VectorRegister4Double Vec2)
 {
 	VectorRegister4Double Result;
 	Result.XY = Vec1.XY;
@@ -1551,14 +1551,14 @@ FORCEINLINE VectorRegister4Double VectorCombineLow(const VectorRegister4Double&
  * @param OutEvens [Even0, Even1, Even2, Even3]
  * @param OutOdds [Odd0, Odd1, Odd2, Odd3]
 */
-FORCEINLINE void VectorDeinterleave(VectorRegister4Float& OutEvens, VectorRegister4Float& OutOdds, const VectorRegister4Float& Lo, const VectorRegister4Float& Hi)
+FORCEINLINE void VectorDeinterleave(VectorRegister4Float& RESTRICT OutEvens, VectorRegister4Float& RESTRICT OutOdds, VectorRegister4Float Lo, VectorRegister4Float Hi)
 {
 	float32x4x2_t deinterleaved = vuzpq_f32(Lo, Hi);
 	OutEvens = deinterleaved.val[0];
 	OutOdds = deinterleaved.val[1];
 }
 
-FORCEINLINE void VectorDeinterleave(VectorRegister4Double& RESTRICT OutEvens, VectorRegister4Double& RESTRICT OutOdds, const VectorRegister4Double& Lo, const VectorRegister4Double& Hi)
+FORCEINLINE void VectorDeinterleave(VectorRegister4Double& RESTRICT OutEvens, VectorRegister4Double& RESTRICT OutOdds, VectorRegister4Double Lo, VectorRegister4Double Hi)
 {
 	OutEvens = VectorShuffle(Lo, Hi, 0, 2, 0, 2);
 	OutOdds = VectorShuffle(Lo, Hi, 1, 3, 1, 3);
@@ -1571,7 +1571,7 @@ FORCEINLINE void VectorDeinterleave(VectorRegister4Double& RESTRICT OutEvens, Ve
  * @param Vec2	2nd vector
  * @return		cross(Vec1.xyz, Vec2.xyz). W of the input should be 0, and will remain 0.
  */
-FORCEINLINE VectorRegister4Float VectorCross( const VectorRegister4Float& Vec1, const VectorRegister4Float& Vec2 )
+FORCEINLINE VectorRegister4Float VectorCross(VectorRegister4Float Vec1, VectorRegister4Float Vec2)
 {
 	VectorRegister4Float C = VectorMultiply(Vec1, VectorSwizzle(Vec2, 1, 2, 0, 3));
 	C = VectorNegateMultiplyAdd(VectorSwizzle(Vec1, 1, 2, 0, 3), Vec2, C);
@@ -1579,7 +1579,7 @@ FORCEINLINE VectorRegister4Float VectorCross( const VectorRegister4Float& Vec1,
 	return C;
 }
 
-FORCEINLINE VectorRegister4Double VectorCross(const VectorRegister4Double& Vec1, const VectorRegister4Double& Vec2)
+FORCEINLINE VectorRegister4Double VectorCross(VectorRegister4Double Vec1, VectorRegister4Double Vec2)
 {
 	VectorRegister4Double C = VectorMultiply(Vec1, VectorSwizzle(Vec2, 1, 2, 0, 3));
 	C = VectorNegateMultiplyAdd(VectorSwizzle(Vec1, 1, 2, 0, 3), Vec2, C);
@@ -1594,7 +1594,7 @@ FORCEINLINE VectorRegister4Double VectorCross(const VectorRegister4Double& Vec1,
  * @param Exponent	Exponent vector
  * @return			VectorRegister4Float( Base.x^Exponent.x, Base.y^Exponent.y, Base.z^Exponent.z, Base.w^Exponent.w )
  */
-FORCEINLINE VectorRegister4Float VectorPow( const VectorRegister4Float& Base, const VectorRegister4Float& Exponent )
+FORCEINLINE VectorRegister4Float VectorPow(VectorRegister4Float Base, VectorRegister4Float Exponent)
 {
 	//@TODO: Optimize this
 	union U { 
@@ -1606,7 +1606,7 @@ FORCEINLINE VectorRegister4Float VectorPow( const VectorRegister4Float& Base, co
 	return MakeVectorRegister( powf(B.F[0], E.F[0]), powf(B.F[1], E.F[1]), powf(B.F[2], E.F[2]), powf(B.F[3], E.F[3]) );
 }
 
-FORCEINLINE VectorRegister4Double VectorPow(const VectorRegister4Double& Base, const VectorRegister4Double& Exponent)
+FORCEINLINE VectorRegister4Double VectorPow(VectorRegister4Double Base, VectorRegister4Double Exponent)
 {
 	//@TODO: Optimize this
 	AlignedDouble4 Values(Base);
@@ -1625,12 +1625,12 @@ FORCEINLINE VectorRegister4Double VectorPow(const VectorRegister4Double& Base, c
  * @param Vec	1st vector
  * @return		VectorRegister4Float( (Estimate) 1.0f / Vec.x, (Estimate) 1.0f / Vec.y, (Estimate) 1.0f / Vec.z, (Estimate) 1.0f / Vec.w )
  */
-FORCEINLINE VectorRegister4Float VectorReciprocalEstimate(const VectorRegister4Float& Vec)
+FORCEINLINE VectorRegister4Float VectorReciprocalEstimate(VectorRegister4Float Vec)
 {
 	return vrecpeq_f32(Vec);
 }
 
-FORCEINLINE VectorRegister4Double VectorReciprocalEstimate(const VectorRegister4Double& Vec)
+FORCEINLINE VectorRegister4Double VectorReciprocalEstimate(VectorRegister4Double Vec)
 {
 	VectorRegister4Double Result;
 	Result.XY = vrecpeq_f64(Vec.XY);
@@ -1645,7 +1645,7 @@ FORCEINLINE VectorRegister4Double VectorReciprocalEstimate(const VectorRegister4
  * @param Vec	1st vector
  * @return		VectorRegister4Float( 1.0f / Vec.x, 1.0f / Vec.y, 1.0f / Vec.z, 1.0f / Vec.w )
  */
-FORCEINLINE VectorRegister4Float VectorReciprocal(const VectorRegister4Float& Vec)
+FORCEINLINE VectorRegister4Float VectorReciprocal(VectorRegister4Float Vec)
 {
 	// Perform two passes of Newton-Raphson iteration on the hardware estimate
 	// The built-in instruction (VRECPS) is not as accurate
@@ -1664,7 +1664,7 @@ FORCEINLINE VectorRegister4Float VectorReciprocal(const VectorRegister4Float& Ve
 	return VectorNegateMultiplyAdd(Vec, Squared, Double);
 }
 
-FORCEINLINE VectorRegister4Double VectorReciprocal(const VectorRegister4Double& Vec)
+FORCEINLINE VectorRegister4Double VectorReciprocal(VectorRegister4Double Vec)
 {
 	return VectorDivide(GlobalVectorConstants::DoubleOne, Vec);
 }
@@ -1676,12 +1676,12 @@ FORCEINLINE VectorRegister4Double VectorReciprocal(const VectorRegister4Double&
  * @param Vector	Vector
  * @return			VectorRegister4Float(sqrt(Vec.X), sqrt(Vec.Y), sqrt(Vec.Z), sqrt(Vec.W))
  */
-FORCEINLINE VectorRegister4Float VectorSqrt(const VectorRegister4Float& Vec)
+FORCEINLINE VectorRegister4Float VectorSqrt(VectorRegister4Float Vec)
 {
 	return vsqrtq_f32(Vec);
 }
 
-FORCEINLINE VectorRegister4Double VectorSqrt(const VectorRegister4Double& Vec)
+FORCEINLINE VectorRegister4Double VectorSqrt(VectorRegister4Double Vec)
 {
 	VectorRegister4Double Result;
 	Result.XY = vsqrtq_f64(Vec.XY);
@@ -1695,12 +1695,12 @@ FORCEINLINE VectorRegister4Double VectorSqrt(const VectorRegister4Double& Vec)
  * @param Vector	Vector 
  * @return			VectorRegister4Float(1/sqrt(t), 1/sqrt(t), 1/sqrt(t), 1/sqrt(t))
  */
-FORCEINLINE VectorRegister4Float VectorReciprocalSqrtEstimate(const VectorRegister4Float& Vec)
+FORCEINLINE VectorRegister4Float VectorReciprocalSqrtEstimate(VectorRegister4Float Vec)
 {
 	return vrsqrteq_f32(Vec);
 }
 
-FORCEINLINE VectorRegister4Double VectorReciprocalSqrtEstimate(const VectorRegister4Double& Vec)
+FORCEINLINE VectorRegister4Double VectorReciprocalSqrtEstimate(VectorRegister4Double Vec)
 {
 	VectorRegister4Double Result;
 	Result.XY = vrsqrteq_f64(Vec.XY);
@@ -1714,7 +1714,7 @@ FORCEINLINE VectorRegister4Double VectorReciprocalSqrtEstimate(const VectorRegis
  * @param Vector	Vector
  * @return			VectorRegister4Float(1/sqrt(Vec.X), 1/sqrt(Vec.Y), 1/sqrt(Vec.Z), 1/sqrt(Vec.W))
  */
-FORCEINLINE VectorRegister4Float VectorReciprocalSqrt(const VectorRegister4Float& Vec)
+FORCEINLINE VectorRegister4Float VectorReciprocalSqrt(VectorRegister4Float Vec)
 {
 	// Initial estimate
 	VectorRegister4Float RecipSqrt = VectorReciprocalSqrtEstimate(Vec);
@@ -1724,7 +1724,7 @@ FORCEINLINE VectorRegister4Float VectorReciprocalSqrt(const VectorRegister4Float
 	return VectorMultiply(vrsqrtsq_f32(Vec, VectorMultiply(RecipSqrt, RecipSqrt)), RecipSqrt);
 }
 
-FORCEINLINE VectorRegister4Double VectorReciprocalSqrt(const VectorRegister4Double& Vec)
+FORCEINLINE VectorRegister4Double VectorReciprocalSqrt(VectorRegister4Double Vec)
 {
 	// Initial estimate
 	VectorRegister4Double RecipSqrt = VectorReciprocalSqrtEstimate(Vec);
@@ -1746,12 +1746,12 @@ FORCEINLINE VectorRegister4Double VectorReciprocalSqrt(const VectorRegister4Doub
  * @param Vector	Vector
  * @return			VectorRegister4Float(rlen, rlen, rlen, rlen) when rlen = 1/sqrt(dot4(V))
  */
-FORCEINLINE VectorRegister4Float VectorReciprocalLen(const VectorRegister4Float& Vector)
+FORCEINLINE VectorRegister4Float VectorReciprocalLen(VectorRegister4Float Vector)
 {
 	return VectorReciprocalSqrt(VectorDot4(Vector, Vector));
 }
 
-FORCEINLINE VectorRegister4Double VectorReciprocalLen(const VectorRegister4Double& Vector)
+FORCEINLINE VectorRegister4Double VectorReciprocalLen(VectorRegister4Double Vector)
 {
 	return VectorReciprocalSqrt(VectorDot4(Vector, Vector));
 }
@@ -1762,12 +1762,12 @@ FORCEINLINE VectorRegister4Double VectorReciprocalLen(const VectorRegister4Doubl
  * @param Vector	Vector
  * @return			VectorRegister4Float(rlen, rlen, rlen, rlen) when rlen = 1/sqrt(dot4(V)) (estimate)
  */
-FORCEINLINE VectorRegister4Float VectorReciprocalLenEstimate(const VectorRegister4Float& Vector)
+FORCEINLINE VectorRegister4Float VectorReciprocalLenEstimate(VectorRegister4Float Vector)
 {
 	return VectorReciprocalSqrtEstimate(VectorDot4(Vector, Vector));
 }
 
-FORCEINLINE VectorRegister4Double VectorReciprocalLenEstimate(const VectorRegister4Double& Vector)
+FORCEINLINE VectorRegister4Double VectorReciprocalLenEstimate(VectorRegister4Double Vector)
 {
 	return VectorReciprocalSqrtEstimate(VectorDot4(Vector, Vector));
 }
@@ -1779,12 +1779,12 @@ FORCEINLINE VectorRegister4Double VectorReciprocalLenEstimate(const VectorRegist
 * @param Vector	VectorRegister4Float
 * @return		VectorRegister4Float(X, Y, Z, 0.0f)
 */
-FORCEINLINE VectorRegister4Float VectorSet_W0(const VectorRegister4Float& Vec)
+FORCEINLINE VectorRegister4Float VectorSet_W0(VectorRegister4Float Vec)
 {
 	return VectorSetComponent(Vec, 3, 0.0f);
 }
 
-FORCEINLINE VectorRegister4Double VectorSet_W0(const VectorRegister4Double& Vec)
+FORCEINLINE VectorRegister4Double VectorSet_W0(VectorRegister4Double Vec)
 {
 	return VectorSetComponent(Vec, 3, 0.0);
 }
@@ -1796,12 +1796,12 @@ FORCEINLINE VectorRegister4Double VectorSet_W0(const VectorRegister4Double& Vec)
 * @param Vector	VectorRegister4Float
 * @return		VectorRegister4Float(X, Y, Z, 1.0f)
 */
-FORCEINLINE VectorRegister4Float VectorSet_W1(const VectorRegister4Float& Vec)
+FORCEINLINE VectorRegister4Float VectorSet_W1(VectorRegister4Float Vec)
 {
 	return VectorSetComponent(Vec, 3, 1.0f);
 }
 
-FORCEINLINE VectorRegister4Double VectorSet_W1(const VectorRegister4Double& Vec)
+FORCEINLINE VectorRegister4Double VectorSet_W1(VectorRegister4Double Vec)
 {
 	return VectorSetComponent(Vec, 3, 1.0);
 }
@@ -1828,7 +1828,7 @@ FORCEINLINE double VectorGetComponentImpl(VectorRegister2Double Vec)
 }
 
 template <int ElementIndex>
-FORCEINLINE double VectorGetComponentImpl(const VectorRegister4Double& Vec)
+FORCEINLINE double VectorGetComponentImpl(VectorRegister4Double Vec)
 {
 	if constexpr (ElementIndex > 1)
 	{
@@ -1861,7 +1861,7 @@ FORCEINLINE double VectorGetComponentDynamic(VectorRegister4Double Vec, uint32 E
  * @param Matrix1	Pointer to the first matrix
  * @param Matrix2	Pointer to the second matrix
  */
-FORCEINLINE void VectorMatrixMultiply( FMatrix44f* Result, const FMatrix44f* Matrix1, const FMatrix44f* Matrix2 )
+FORCEINLINE void VectorMatrixMultiply(FMatrix44f* Result, const FMatrix44f* Matrix1, const FMatrix44f* Matrix2)
 {
 	float32x4x4_t A = vld1q_f32_x4((const float*)Matrix1);
 	float32x4x4_t B = vld1q_f32_x4((const float*)Matrix2);
@@ -1977,7 +1977,7 @@ FORCEINLINE bool VectorMatrixInverse(FMatrix44f* DstMatrix, const FMatrix44f* Sr
  * @param MatrixM		FMatrix pointer to the Matrix to apply transform
  * @return VectorRegister4Float = VecP*MatrixM
  */
-FORCEINLINE VectorRegister4Float VectorTransformVector(const VectorRegister4Float& VecP, const FMatrix44f* MatrixM )
+FORCEINLINE VectorRegister4Float VectorTransformVector(VectorRegister4Float VecP, const FMatrix44f* MatrixM )
 {
 	float32x4x4_t M = vld1q_f32_x4((const float*)MatrixM);
 	VectorRegister4Float Result;
@@ -1990,7 +1990,7 @@ FORCEINLINE VectorRegister4Float VectorTransformVector(const VectorRegister4Floa
 	return Result;
 }
 
-FORCEINLINE VectorRegister4Float VectorTransformVector(const VectorRegister4Float& VecP, const FMatrix44d* MatrixM)
+FORCEINLINE VectorRegister4Float VectorTransformVector(VectorRegister4Float VecP, const FMatrix44d* MatrixM)
 {
 	float64x2x4_t M1 = vld1q_f64_x4((const double*)MatrixM);
 	float64x2x4_t M2 = vld1q_f64_x4(((const double*)MatrixM) + 8);
@@ -2010,7 +2010,7 @@ FORCEINLINE VectorRegister4Float VectorTransformVector(const VectorRegister4Floa
 	return MakeVectorRegisterFloatFromDouble(Result);
 }
 
-FORCEINLINE VectorRegister4Double VectorTransformVector(const VectorRegister4Double& VecP, const FMatrix44d* MatrixM)
+FORCEINLINE VectorRegister4Double VectorTransformVector(VectorRegister4Double VecP, const FMatrix44d* MatrixM)
 {
 	float64x2x4_t M1 = vld1q_f64_x4((const double*)MatrixM);
 	float64x2x4_t M2 = vld1q_f64_x4(((const double*)MatrixM) + 8);
@@ -2037,7 +2037,7 @@ FORCEINLINE VectorRegister4Double VectorTransformVector(const VectorRegister4Dou
  * @param Vec2	2nd vector
  * @return		VectorRegister4Float( min(Vec1.x,Vec2.x), min(Vec1.y,Vec2.y), min(Vec1.z,Vec2.z), min(Vec1.w,Vec2.w) )
  */
-FORCEINLINE VectorRegister4Float VectorMin( VectorRegister4Float Vec1, VectorRegister4Float Vec2 )
+FORCEINLINE VectorRegister4Float VectorMin(VectorRegister4Float Vec1, VectorRegister4Float Vec2)
 {
 	return vminq_f32( Vec1, Vec2 );
 }
@@ -2057,7 +2057,7 @@ FORCEINLINE VectorRegister4Double VectorMin(VectorRegister4Double Vec1, VectorRe
  * @param Vec2	2nd vector
  * @return		VectorRegister4Float( max(Vec1.x,Vec2.x), max(Vec1.y,Vec2.y), max(Vec1.z,Vec2.z), max(Vec1.w,Vec2.w) )
  */
-FORCEINLINE VectorRegister4Float VectorMax( VectorRegister4Float Vec1, VectorRegister4Float Vec2 )
+FORCEINLINE VectorRegister4Float VectorMax(VectorRegister4Float Vec1, VectorRegister4Float Vec2)
 {
 	return vmaxq_f32( Vec1, Vec2 );
 }
@@ -2077,12 +2077,12 @@ FORCEINLINE VectorRegister4Double VectorMax(VectorRegister4Double Vec1, VectorRe
  * @param VecW		Source register for ___W (note: the fourth component is used, not the first)
  * @return			VectorRegister4Float(VecXYZ.x, VecXYZ.y, VecXYZ.z, VecW.w)
  */
-FORCEINLINE VectorRegister4Float VectorMergeVecXYZ_VecW(const VectorRegister4Float& VecXYZ, const VectorRegister4Float& VecW)
+FORCEINLINE VectorRegister4Float VectorMergeVecXYZ_VecW(VectorRegister4Float VecXYZ, VectorRegister4Float VecW)
 {
 	return vsetq_lane_f32(vgetq_lane_f32(VecW, 3), VecXYZ, 3);
 }
 
-FORCEINLINE VectorRegister4Double VectorMergeVecXYZ_VecW(const VectorRegister4Double& VecXYZ, const VectorRegister4Double& VecW)
+FORCEINLINE VectorRegister4Double VectorMergeVecXYZ_VecW(VectorRegister4Double VecXYZ, VectorRegister4Double VecW)
 {
 	VectorRegister4Double Res;
 	Res.XY = VecXYZ.XY;
@@ -2096,7 +2096,7 @@ FORCEINLINE VectorRegister4Double VectorMergeVecXYZ_VecW(const VectorRegister4Do
  * @param Ptr			Unaligned memory pointer to the 4 uint8s.
  * @return				VectorRegister4Float( float(Ptr[0]), float(Ptr[1]), float(Ptr[2]), float(Ptr[3]) )
  */
-FORCEINLINE VectorRegister4Float VectorLoadByte4( const void* Ptr )
+FORCEINLINE VectorRegister4Float VectorLoadByte4(const void* Ptr)
 {
 	uint8x8_t AsUInt8 = vreinterpret_u8_u32(vld1_dup_u32((const uint32*)Ptr));
 	uint16x8_t AsUInt16 = vmovl_u8(AsUInt8);
@@ -2124,7 +2124,7 @@ FORCEINLINE VectorRegister4Float VectorLoadSignedByte4(const void* Ptr)
  * @param Ptr			Unaligned memory pointer to the 4 uint8s.
  * @return				VectorRegister4Float( float(Ptr[3]), float(Ptr[2]), float(Ptr[1]), float(Ptr[0]) )
  */
-FORCEINLINE VectorRegister4Float VectorLoadByte4Reverse( const uint8* Ptr )
+FORCEINLINE VectorRegister4Float VectorLoadByte4Reverse(const uint8* Ptr)
 {
 	uint8x8_t AsUInt8 = vrev32_u8(vreinterpret_u8_u32(vld1_dup_u32((const uint32*)Ptr)));
 	uint16x8_t AsUInt16 = vmovl_u8(AsUInt8);
@@ -2138,7 +2138,7 @@ FORCEINLINE VectorRegister4Float VectorLoadByte4Reverse( const uint8* Ptr )
  * @param Vec			Vector containing 4 floats
  * @param Ptr			Unaligned memory pointer to store the 4 uint8s.
  */
-FORCEINLINE void VectorStoreByte4( VectorRegister4Float Vec, void* Ptr )
+FORCEINLINE void VectorStoreByte4(VectorRegister4Float Vec, void* Ptr)
 {
 	uint32x4_t AsUInt32 = vcvtq_u32_f32(Vec); // Saturates (clamps) to [0,2^32 - 1]
 	uint16x4_t AsUInt16 = vqmovn_u32(AsUInt32); // Saturates further to [0,2^16 - 1]
@@ -2167,7 +2167,7 @@ FORCEINLINE void VectorStoreSignedByte4(VectorRegister4Float Vec, void* Ptr)
  * @param Ptr			Memory pointer to store the 4 fp16's.
  */
 template <bool bAligned>
-FORCEINLINE void VectorStoreHalf4(VectorRegister4Float Vec, void* RESTRICT Ptr)
+FORCEINLINE void VectorStoreHalf4(VectorRegister4Float Vec, void* Ptr)
 {
 	float16x4_t f16x4 = vcvt_f16_f32(Vec);
 	vst1_u8((uint8_t*)Ptr, f16x4);
@@ -2198,7 +2198,7 @@ FORCEINLINE VectorRegister4Float VectorLoadURGB10A2N(void* Ptr)
 * @param Vec			Vector containing 4 FLOATs
 * @param Ptr			Unaligned memory pointer to store the packed RGB10A2(4 bytes).
 */
-FORCEINLINE void VectorStoreURGB10A2N(const VectorRegister4Float& Vec, void* Ptr)
+FORCEINLINE void VectorStoreURGB10A2N(VectorRegister4Float Vec, void* Ptr)
 {
 	union U { 
 		VectorRegister4Float V; float F[4]; 
@@ -2315,7 +2315,7 @@ FORCEINLINE void VectorSetControlRegister(uint32_t ControlStatus)
 * @param Quat2	Pointer to the second quaternion
 * @return Quat1 * Quat2
 */
-FORCEINLINE VectorRegister4Float VectorQuaternionMultiply2( const VectorRegister4Float& Quat1, const VectorRegister4Float& Quat2 )
+FORCEINLINE VectorRegister4Float VectorQuaternionMultiply2(VectorRegister4Float Quat1, VectorRegister4Float Quat2)
 {
 	VectorRegister4Float Result = VectorMultiply(VectorReplicate(Quat1, 3), Quat2);
 	Result = VectorMultiplyAdd( VectorMultiply(VectorReplicate(Quat1, 0), VectorSwizzle(Quat2, 3,2,1,0)), GlobalVectorConstants::QMULTI_SIGN_MASK0, Result);
@@ -2325,7 +2325,7 @@ FORCEINLINE VectorRegister4Float VectorQuaternionMultiply2( const VectorRegister
 	return Result;
 }
 
-FORCEINLINE VectorRegister4Double VectorQuaternionMultiply2(const VectorRegister4Double& Quat1, const VectorRegister4Double& Quat2)
+FORCEINLINE VectorRegister4Double VectorQuaternionMultiply2(VectorRegister4Double Quat1, VectorRegister4Double Quat2)
 {
 	VectorRegister4Double Result = VectorMultiply(VectorReplicate(Quat1, 3), Quat2);
 	Result = VectorMultiplyAdd(VectorMultiply(VectorReplicate(Quat1, 0), VectorSwizzle(Quat2, 3, 2, 1, 0)), GlobalVectorConstants::DOUBLE_QMULTI_SIGN_MASK0, Result);
@@ -2362,7 +2362,7 @@ FORCEINLINE void VectorQuaternionMultiply(VectorRegister4Double* RESTRICT Result
 * @param VCosAngles	VectorRegister4Float Pointer to where the Cos result should be stored
 * @param VAngles VectorRegister4Float Pointer to the input angles 
 */
-FORCEINLINE void VectorSinCos(  VectorRegister4Float* RESTRICT VSinAngles, VectorRegister4Float* RESTRICT VCosAngles, const VectorRegister4Float* RESTRICT VAngles )
+FORCEINLINE void VectorSinCos(VectorRegister4Float* RESTRICT VSinAngles, VectorRegister4Float* RESTRICT VCosAngles, const VectorRegister4Float* RESTRICT VAngles)
 {	
 	// Map to [-pi, pi]
 	// X = A - 2pi * round(A/2pi)
@@ -2416,7 +2416,7 @@ FORCEINLINE void VectorSinCos(  VectorRegister4Float* RESTRICT VSinAngles, Vecto
 }
 
 // Returns true if the vector contains a component that is either NAN or +/-infinite.
-inline bool VectorContainsNaNOrInfinite(const VectorRegister4Float& Vec)
+inline bool VectorContainsNaNOrInfinite(VectorRegister4Float Vec)
 {
 	// https://en.wikipedia.org/wiki/IEEE_754-1985
 	// Infinity is represented with all exponent bits set, with the correct sign bit.
@@ -2439,7 +2439,7 @@ inline bool VectorContainsNaNOrInfinite(const VectorRegister4Float& Vec)
 	return vgetq_lane_u32((uint32x4_t)vqtbx1q_u8(res, res, Table), 0) != 0;
 }
 
-inline bool VectorContainsNaNOrInfinite(const VectorRegister4Double& Vec)
+inline bool VectorContainsNaNOrInfinite(VectorRegister4Double Vec)
 {
 	// https://en.wikipedia.org/wiki/IEEE_754-1985
 	// Infinity is represented with all exponent bits set, with the correct sign bit.
@@ -2466,111 +2466,111 @@ inline bool VectorContainsNaNOrInfinite(const VectorRegister4Double& Vec)
 }
 
 //TODO: Vectorize
-FORCEINLINE VectorRegister4Float VectorExp(const VectorRegister4Float& X)
+FORCEINLINE VectorRegister4Float VectorExp(VectorRegister4Float X)
 {
 	AlignedFloat4 Val(X);
 	return MakeVectorRegister(FMath::Exp(Val[0]), FMath::Exp(Val[1]), FMath::Exp(Val[2]), FMath::Exp(Val[3]));
 }
 
-FORCEINLINE VectorRegister4Double VectorExp(const VectorRegister4Double& X)
+FORCEINLINE VectorRegister4Double VectorExp(VectorRegister4Double X)
 {
 	AlignedDouble4 Val(X);
 	return MakeVectorRegister(FMath::Exp(Val[0]), FMath::Exp(Val[1]), FMath::Exp(Val[2]), FMath::Exp(Val[3]));
 }
 
 //TODO: Vectorize
-FORCEINLINE VectorRegister4Float VectorExp2(const VectorRegister4Float& X)
+FORCEINLINE VectorRegister4Float VectorExp2(VectorRegister4Float X)
 {
 	AlignedFloat4 Val(X);
 	return MakeVectorRegister(FMath::Exp2(Val[0]), FMath::Exp2(Val[1]), FMath::Exp2(Val[2]), FMath::Exp2(Val[3]));
 }
 
-FORCEINLINE VectorRegister4Double VectorExp2(const VectorRegister4Double& X)
+FORCEINLINE VectorRegister4Double VectorExp2(VectorRegister4Double X)
 {
 	AlignedDouble4 Val(X);
 	return MakeVectorRegister(FMath::Exp2(Val[0]), FMath::Exp2(Val[1]), FMath::Exp2(Val[2]), FMath::Exp2(Val[3]));
 }
 
 //TODO: Vectorize
-FORCEINLINE VectorRegister4Float VectorLog(const VectorRegister4Float& X)
+FORCEINLINE VectorRegister4Float VectorLog(VectorRegister4Float X)
 {
 	AlignedFloat4 Val(X);
 	return MakeVectorRegister(FMath::Loge(Val[0]), FMath::Loge(Val[1]), FMath::Loge(Val[2]), FMath::Loge(Val[3]));
 }
 
-FORCEINLINE VectorRegister4Double VectorLog(const VectorRegister4Double& X)
+FORCEINLINE VectorRegister4Double VectorLog(VectorRegister4Double X)
 {
 	AlignedDouble4 Val(X);
 	return MakeVectorRegister(FMath::Loge(Val[0]), FMath::Loge(Val[1]), FMath::Loge(Val[2]), FMath::Loge(Val[3]));
 }
 
 //TODO: Vectorize
-FORCEINLINE VectorRegister4Float VectorLog2(const VectorRegister4Float& X)
+FORCEINLINE VectorRegister4Float VectorLog2(VectorRegister4Float X)
 {
 	AlignedFloat4 Val(X);
 	return MakeVectorRegister(FMath::Log2(Val[0]), FMath::Log2(Val[1]), FMath::Log2(Val[2]), FMath::Log2(Val[3]));
 }
 
-FORCEINLINE VectorRegister4Double VectorLog2(const VectorRegister4Double& X)
+FORCEINLINE VectorRegister4Double VectorLog2(VectorRegister4Double X)
 {
 	AlignedDouble4 Val(X);
 	return MakeVectorRegister(FMath::Log2(Val[0]), FMath::Log2(Val[1]), FMath::Log2(Val[2]), FMath::Log2(Val[3]));
 }
 
 //TODO: Vectorize
-FORCEINLINE VectorRegister4Float VectorTan(const VectorRegister4Float& X)
+FORCEINLINE VectorRegister4Float VectorTan(VectorRegister4Float X)
 {
 	AlignedFloat4 Val(X);
 	return MakeVectorRegister(FMath::Tan(Val[0]), FMath::Tan(Val[1]), FMath::Tan(Val[2]), FMath::Tan(Val[3]));
 }
 
-FORCEINLINE VectorRegister4Double VectorTan(const VectorRegister4Double& X)
+FORCEINLINE VectorRegister4Double VectorTan(VectorRegister4Double X)
 {
 	AlignedDouble4 Val(X);
 	return MakeVectorRegister(FMath::Tan(Val[0]), FMath::Tan(Val[1]), FMath::Tan(Val[2]), FMath::Tan(Val[3]));
 }
 
 //TODO: Vectorize
-FORCEINLINE VectorRegister4Float VectorASin(const VectorRegister4Float& X)
+FORCEINLINE VectorRegister4Float VectorASin(VectorRegister4Float X)
 {
 	AlignedFloat4 Val(X);
 	return MakeVectorRegister(FMath::Asin(Val[0]), FMath::Asin(Val[1]), FMath::Asin(Val[2]), FMath::Asin(Val[3]));
 }
 
-FORCEINLINE VectorRegister4Double VectorASin(const VectorRegister4Double& X)
+FORCEINLINE VectorRegister4Double VectorASin(VectorRegister4Double X)
 {
 	AlignedDouble4 Val(X);
 	return MakeVectorRegister(FMath::Asin(Val[0]), FMath::Asin(Val[1]), FMath::Asin(Val[2]), FMath::Asin(Val[3]));
 }
 
 //TODO: Vectorize
-FORCEINLINE VectorRegister4Float VectorACos(const VectorRegister4Float& X)
+FORCEINLINE VectorRegister4Float VectorACos(VectorRegister4Float X)
 {
 	AlignedFloat4 Val(X);
 	return MakeVectorRegister(FMath::Acos(Val[0]), FMath::Acos(Val[1]), FMath::Acos(Val[2]), FMath::Acos(Val[3]));
 }
 
-FORCEINLINE VectorRegister4Double VectorACos(const VectorRegister4Double& X)
+FORCEINLINE VectorRegister4Double VectorACos(VectorRegister4Double X)
 {
 	AlignedDouble4 Val(X);
 	return MakeVectorRegister(FMath::Acos(Val[0]), FMath::Acos(Val[1]), FMath::Acos(Val[2]), FMath::Acos(Val[3]));
 }
 
 //TODO: Vectorize
-FORCEINLINE VectorRegister4Float VectorATan(const VectorRegister4Float& X)
+FORCEINLINE VectorRegister4Float VectorATan(VectorRegister4Float X)
 {
 	AlignedFloat4 Val(X);
 	return MakeVectorRegister(FMath::Atan(Val[0]), FMath::Atan(Val[1]), FMath::Atan(Val[2]), FMath::Atan(Val[3]));
 }
 
-FORCEINLINE VectorRegister4Double VectorATan(const VectorRegister4Double& X)
+FORCEINLINE VectorRegister4Double VectorATan(VectorRegister4Double X)
 {
 	AlignedDouble4 Val(X);
 	return MakeVectorRegister(FMath::Atan(Val[0]), FMath::Atan(Val[1]), FMath::Atan(Val[2]), FMath::Atan(Val[3]));
 }
 
 //TODO: Vectorize
-FORCEINLINE VectorRegister4Float VectorATan2(const VectorRegister4Float& X, const VectorRegister4Float& Y)
+FORCEINLINE VectorRegister4Float VectorATan2(VectorRegister4Float X, VectorRegister4Float Y)
 {
 	AlignedFloat4 ValX(X);
 	AlignedFloat4 ValY(Y);
@@ -2581,7 +2581,7 @@ FORCEINLINE VectorRegister4Float VectorATan2(const VectorRegister4Float& X, cons
 							  FMath::Atan2(ValX[3], ValY[3]));
 }
 
-FORCEINLINE VectorRegister4Double VectorATan2(const VectorRegister4Double& X, const VectorRegister4Double& Y)
+FORCEINLINE VectorRegister4Double VectorATan2(VectorRegister4Double X, VectorRegister4Double Y)
 {
 	AlignedDouble4 ValX(X);
 	AlignedDouble4 ValY(Y);
@@ -2592,12 +2592,12 @@ FORCEINLINE VectorRegister4Double VectorATan2(const VectorRegister4Double& X, co
 							  FMath::Atan2(ValX[3], ValY[3]));
 }
 
-FORCEINLINE VectorRegister4Float VectorCeil(const VectorRegister4Float& X)
+FORCEINLINE VectorRegister4Float VectorCeil(VectorRegister4Float X)
 {
 	return vrndpq_f32(X);
 }
 
-FORCEINLINE VectorRegister4Double VectorCeil(const VectorRegister4Double& X)
+FORCEINLINE VectorRegister4Double VectorCeil(VectorRegister4Double X)
 {
 	VectorRegister4Double Result;
 	Result.XY = vrndpq_f64(X.XY);
@@ -2605,12 +2605,12 @@ FORCEINLINE VectorRegister4Double VectorCeil(const VectorRegister4Double& X)
 	return Result;
 }
 
-FORCEINLINE VectorRegister4Float VectorFloor(const VectorRegister4Float& X)
+FORCEINLINE VectorRegister4Float VectorFloor(VectorRegister4Float X)
 {
 	return vrndmq_f32(X);
 }
 
-FORCEINLINE VectorRegister4Double VectorFloor(const VectorRegister4Double& X)
+FORCEINLINE VectorRegister4Double VectorFloor(VectorRegister4Double X)
 {
 	VectorRegister4Double Result;
 	Result.XY = vrndmq_f64(X.XY);
@@ -2618,12 +2618,12 @@ FORCEINLINE VectorRegister4Double VectorFloor(const VectorRegister4Double& X)
 	return Result;
 }
 
-FORCEINLINE VectorRegister4Float VectorTruncate(const VectorRegister4Float& X)
+FORCEINLINE VectorRegister4Float VectorTruncate(VectorRegister4Float X)
 {
 	return vrndq_f32(X);
 }
 
-FORCEINLINE VectorRegister4Double VectorTruncate(const VectorRegister4Double& X)
+FORCEINLINE VectorRegister4Double VectorTruncate(VectorRegister4Double X)
 {
 	VectorRegister4Double Result;
 	Result.XY = vrndq_f64(X.XY);
@@ -2631,7 +2631,7 @@ FORCEINLINE VectorRegister4Double VectorTruncate(const VectorRegister4Double& X)
 	return Result;
 }
 
-FORCEINLINE VectorRegister4Float VectorMod(const VectorRegister4Float& X, const VectorRegister4Float& Y)
+FORCEINLINE VectorRegister4Float VectorMod(VectorRegister4Float X, VectorRegister4Float Y)
 {
 	// Check against invalid divisor
 	VectorRegister4Float InvalidDivisorMask = VectorCompareLE(VectorAbs(Y), GlobalVectorConstants::SmallNumber);
@@ -2648,7 +2648,7 @@ FORCEINLINE VectorRegister4Float VectorMod(const VectorRegister4Float& X, const
 	return Result;
 }
 
-FORCEINLINE VectorRegister4Double VectorMod(const VectorRegister4Double& X, const VectorRegister4Double& Y)
+FORCEINLINE VectorRegister4Double VectorMod(VectorRegister4Double X, VectorRegister4Double Y)
 {
 	// Check against invalid divisor
 	VectorRegister4Double InvalidDivisorMask = VectorCompareLE(VectorAbs(Y), GlobalVectorConstants::DoubleSmallNumber);
@@ -2665,25 +2665,25 @@ FORCEINLINE VectorRegister4Double VectorMod(const VectorRegister4Double& X, cons
 	return DoubleResult;
 }
 
-FORCEINLINE VectorRegister4Float VectorSign(const VectorRegister4Float& X)
+FORCEINLINE VectorRegister4Float VectorSign(VectorRegister4Float X)
 {
 	VectorRegister4Float Mask = VectorCompareGE(X, GlobalVectorConstants::FloatZero);
 	return VectorSelect(Mask, GlobalVectorConstants::FloatOne, GlobalVectorConstants::FloatMinusOne);
 }
 
-FORCEINLINE VectorRegister4Double VectorSign(const VectorRegister4Double& X)
+FORCEINLINE VectorRegister4Double VectorSign(VectorRegister4Double X)
 {
 	VectorRegister4Double Mask = VectorCompareGE(X, GlobalVectorConstants::DoubleZero);
 	return VectorSelect(Mask, GlobalVectorConstants::DoubleOne, GlobalVectorConstants::DoubleMinusOne);
 }
 
-FORCEINLINE VectorRegister4Float VectorStep(const VectorRegister4Float& X)
+FORCEINLINE VectorRegister4Float VectorStep(VectorRegister4Float X)
 {
 	VectorRegister4Float Mask = VectorCompareGE(X, GlobalVectorConstants::FloatZero);
 	return VectorSelect(Mask, GlobalVectorConstants::FloatOne, GlobalVectorConstants::FloatZero);
 }
 
-FORCEINLINE VectorRegister4Double VectorStep(const VectorRegister4Double& X)
+FORCEINLINE VectorRegister4Double VectorStep(VectorRegister4Double X)
 {
 	VectorRegister4Double Mask = VectorCompareGE(X, GlobalVectorConstants::DoubleZero);
 	return VectorSelect(Mask, GlobalVectorConstants::DoubleOne, GlobalVectorConstants::DoubleZero);
@@ -2698,7 +2698,7 @@ namespace VectorSinConstantsNEON
 	static const VectorRegister4Float B = MakeVectorRegisterConstant(b, b, b, b);
 }
 
-FORCEINLINE VectorRegister4Float VectorSin(const VectorRegister4Float& X)
+FORCEINLINE VectorRegister4Float VectorSin(VectorRegister4Float X)
 {
 	//Sine approximation using a squared parabola restrained to f(0) = 0, f(PI) = 0, f(PI/2) = 1.
 	//based on a good discussion here http://forum.devmaster.net/t/fast-and-accurate-sine-cosine/9648
@@ -2718,7 +2718,7 @@ FORCEINLINE VectorRegister4Float VectorSin(const VectorRegister4Float& X)
 	return VectorMultiply(Y, VectorAdd(VectorSinConstantsNEON::B, VectorAbs(Y)));
 }
 
-FORCEINLINE VectorRegister4Double VectorSin(const VectorRegister4Double& X)
+FORCEINLINE VectorRegister4Double VectorSin(VectorRegister4Double X)
 {
 	AlignedDouble4 Doubles(X);
 	Doubles[0] = FMath::Sin(Doubles[0]);
@@ -2728,12 +2728,12 @@ FORCEINLINE VectorRegister4Double VectorSin(const VectorRegister4Double& X)
 	return Doubles.ToVectorRegister();
 }
 
-FORCEINLINE VectorRegister4Float VectorCos(const VectorRegister4Float& X)
+FORCEINLINE VectorRegister4Float VectorCos(VectorRegister4Float X)
 {
 	return VectorSin(VectorAdd(X, GlobalVectorConstants::PiByTwo));
 }
 
-FORCEINLINE VectorRegister4Double VectorCos(const VectorRegister4Double& X)
+FORCEINLINE VectorRegister4Double VectorCos(VectorRegister4Double X)
 {
 	AlignedDouble4 Doubles(X);
 	Doubles[0] = FMath::Cos(Doubles[0]);
@@ -2781,7 +2781,7 @@ FORCEINLINE VectorRegister4Float VectorLoadSRGBA16N(const void* Ptr)
 * @param Vec			Vector containing 4 FLOATs
 * @param Ptr			Unaligned memory pointer to store the packed RGBA16(8 bytes).
 */
-FORCEINLINE void VectorStoreURGBA16N(const VectorRegister4Float& Vec, uint16* Out)
+FORCEINLINE void VectorStoreURGBA16N(VectorRegister4Float Vec, uint16* Out)
 {
 	VectorRegister4Float Tmp;
 	Tmp = VectorMax(Vec, VectorZeroFloat());
@@ -2816,7 +2816,7 @@ FORCEINLINE void VectorStoreURGBA16N(const VectorRegister4Float& Vec, uint16* Ou
 #define VectorIntCompareLE(A, B)	vcleq_s32(A,B)
 
 
-FORCEINLINE VectorRegister4Int VectorIntSelect(const VectorRegister4Int& Mask, const VectorRegister4Int& Vec1, const VectorRegister4Int& Vec2)
+FORCEINLINE VectorRegister4Int VectorIntSelect(VectorRegister4Int Mask, VectorRegister4Int Vec1, VectorRegister4Int Vec2)
 {
 	return vbslq_s32(Mask, Vec1, Vec2);
 }
@@ -2835,17 +2835,17 @@ FORCEINLINE VectorRegister4Int VectorIntSelect(const VectorRegister4Int& Mask, c
 
 #define VectorIntToFloat(A) vcvtq_f32_s32(A)
 
-FORCEINLINE VectorRegister4Int VectorFloatToInt(const VectorRegister4Float& A)
+FORCEINLINE VectorRegister4Int VectorFloatToInt(VectorRegister4Float A)
 {
 	return vcvtq_s32_f32(A);
 }
 
-FORCEINLINE VectorRegister4Int VectorFloatToInt(const VectorRegister4Double& A)
+FORCEINLINE VectorRegister4Int VectorFloatToInt(VectorRegister4Double A)
 {
 	return VectorFloatToInt(MakeVectorRegisterFloatFromDouble(A));
 }
 
-FORCEINLINE VectorRegister4Int VectorDoubleToInt(const VectorRegister4Double& Vec)
+FORCEINLINE VectorRegister4Int VectorDoubleToInt(VectorRegister4Double Vec)
 {
 	VectorRegister2Int64 A = vcvtq_s64_f64(Vec.XY);
 	VectorRegister2Int64 B = vcvtq_s64_f64(Vec.ZW);
@@ -2853,7 +2853,7 @@ FORCEINLINE VectorRegister4Int VectorDoubleToInt(const VectorRegister4Double& Ve
 	return vcombine_s32(vqmovn_s64(A), vqmovn_s64(B));
 }
 
-FORCEINLINE VectorRegister4Int VectorShuffleByte4(const VectorRegister4Int& Vec, const VectorRegister4Int& Mask)
+FORCEINLINE VectorRegister4Int VectorShuffleByte4(VectorRegister4Int Vec, VectorRegister4Int Mask)
 {
 	return vqtbl1q_u8(Vec, Mask);
 }
@@ -2900,7 +2900,7 @@ FORCEINLINE VectorRegister4Int VectorShuffleByte4(const VectorRegister4Int& Vec,
 * @param Ptr	Unaligned memory pointer to the 4 int32s
 * @return		VectorRegister4Int(*Ptr, *Ptr, *Ptr, *Ptr)
 */
-#define VectorIntLoad1( Ptr )	                    vld1q_dup_s32((int32*)(Ptr))
+#define VectorIntLoad1(Ptr)							vld1q_dup_s32((int32*)(Ptr))
 #define VectorIntLoad1_16(Ptr)                      vld1q_dup_s16((int16*)(Ptr))
 
 #define VectorIntSet1(F)                            (VectorRegister4Int)vdupq_n_s32(F)
@@ -2915,12 +2915,13 @@ FORCEINLINE VectorRegister4Int VectorShuffleByte4(const VectorRegister4Int& Vec,
 #define VectorShiftRightImmLogical(Vec, ImmAmt)     vshrq_n_u32(Vec, ImmAmt)
 #define VectorRound(Vec)							vrndnq_f32(Vec)
 
-FORCEINLINE VectorRegister4Int VectorRoundToIntHalfToEven(const VectorRegister4Float& Vec)
+FORCEINLINE VectorRegister4Int VectorRoundToIntHalfToEven(VectorRegister4Float Vec)
 {
 	return vcvtnq_s32_f32(Vec);
 }
 
-FORCEINLINE VectorRegister4i VectorIntExpandLow16To32(VectorRegister4i V) {
+FORCEINLINE VectorRegister4i VectorIntExpandLow16To32(VectorRegister4i V)
+{
 	return vmovl_u16(vget_low_u16(V));
 }
 
diff --git a/Engine/Source/Runtime/Core/Public/Math/UnrealMathSSE.h b/Engine/Source/Runtime/Core/Public/Math/UnrealMathSSE.h
--- a/Engine/Source/Runtime/Core/Public/Math/UnrealMathSSE.h
+++ b/Engine/Source/Runtime/Core/Public/Math/UnrealMathSSE.h
@@ -135,7 +135,7 @@ struct alignas(UE_SSE_DOUBLE_ALIGNMENT) VectorRegister4Double
 
 	FORCEINLINE VectorRegister4Double() = default;
 
-	FORCEINLINE VectorRegister4Double(const VectorRegister2Double& InXY, const VectorRegister2Double& InZW)
+	FORCEINLINE VectorRegister4Double(VectorRegister2Double InXY, VectorRegister2Double InZW)
 	{
 #if UE_PLATFORM_MATH_USE_AVX
 		XYZW = _mm256_setr_m128d(InXY, InZW);
@@ -151,7 +151,7 @@ struct alignas(UE_SSE_DOUBLE_ALIGNMENT) VectorRegister4Double
 	{}
 
 	// Construct from a vector of 4 floats
-	FORCEINLINE VectorRegister4Double(const VectorRegister4Float& FloatVector)
+	FORCEINLINE VectorRegister4Double(VectorRegister4Float FloatVector)
 	{
 #if !UE_PLATFORM_MATH_USE_AVX
 		XY = _mm_cvtps_pd(FloatVector);
@@ -162,7 +162,7 @@ struct alignas(UE_SSE_DOUBLE_ALIGNMENT) VectorRegister4Double
 	}
 
 	// Assign from a vector of 4 floats
-	FORCEINLINE VectorRegister4Double& operator=(const VectorRegister4Float& FloatVector)
+	FORCEINLINE VectorRegister4Double& operator=(VectorRegister4Float FloatVector)
 	{
 #if !UE_PLATFORM_MATH_USE_AVX
 		XY = _mm_cvtps_pd(FloatVector);
@@ -175,13 +175,13 @@ struct alignas(UE_SSE_DOUBLE_ALIGNMENT) VectorRegister4Double
 
 #if UE_PLATFORM_MATH_USE_AVX
 	// Convenience for things like 'Result = _mm256_add_pd(...)'
-	FORCEINLINE VectorRegister4Double(const __m256d& Register)
+	FORCEINLINE VectorRegister4Double(__m256d Register)
 	{
 		XYZW = Register;
 	}
 
 	// Convenience for things like 'Result = _mm256_add_pd(...)'
-	FORCEINLINE VectorRegister4Double& operator=(const __m256d& Register)
+	FORCEINLINE VectorRegister4Double& operator=(__m256d Register)
 	{
 		XYZW = Register;
 		return *this;
@@ -214,8 +214,8 @@ typedef VectorRegister4Int VectorRegisterInt;
 // Forward declarations
 VectorRegister4Float VectorLoadAligned(const float* Ptr);
 VectorRegister4Double VectorLoadAligned(const double* Ptr);
-void VectorStoreAligned(const VectorRegister4Float& Vec, float* Ptr);
-void VectorStoreAligned(const VectorRegister4Double& Vec, double* Dst);
+void VectorStoreAligned(VectorRegister4Float Vec, float* Ptr);
+void VectorStoreAligned(VectorRegister4Double Vec, double* Dst);
 
 
 // Helper for conveniently aligning a float array for extraction from VectorRegister4Float
@@ -223,7 +223,7 @@ struct alignas(UE_SSE_FLOAT_ALIGNMENT) AlignedFloat4
 {
 	float V[4];
 
-	FORCEINLINE AlignedFloat4(const VectorRegister4Float& Vec)
+	FORCEINLINE AlignedFloat4(VectorRegister4Float Vec)
 	{
 		VectorStoreAligned(Vec, V);
 	}
@@ -243,7 +243,7 @@ struct alignas(alignof(VectorRegister4Double)) AlignedDouble4
 {
 	double V[4];
 
-	FORCEINLINE AlignedDouble4(const VectorRegister4Double& Vec)
+	FORCEINLINE AlignedDouble4(VectorRegister4Double Vec)
 	{
 		VectorStoreAligned(Vec, V);
 	}
@@ -295,7 +295,7 @@ FORCEINLINE VectorRegister2Double MakeVectorRegister2DoubleMask(uint64 X, uint64
  * @param W		4th uint32 component
  * @return		Bitwise equivalent vector with 4 floats
  */
-FORCEINLINE VectorRegister4Float MakeVectorRegisterFloat( uint32 X, uint32 Y, uint32 Z, uint32 W )
+FORCEINLINE VectorRegister4Float MakeVectorRegisterFloat(uint32 X, uint32 Y, uint32 Z, uint32 W)
 {
  	union { VectorRegister4Float v; VectorRegister4Int i; } Tmp;
 	Tmp.i = _mm_setr_epi32( X, Y, Z, W );
@@ -359,19 +359,19 @@ FORCEINLINE VectorRegister4Double MakeVectorRegister(double X, double Y, double
 	return MakeVectorRegisterDouble(X, Y, Z, W);
 }
 
-FORCEINLINE VectorRegister4Double MakeVectorRegisterDouble(const VectorRegister2Double& XY, const VectorRegister2Double& ZW)
+FORCEINLINE VectorRegister4Double MakeVectorRegisterDouble(VectorRegister2Double XY, VectorRegister2Double ZW)
 {
 	return VectorRegister4Double(XY, ZW);
 }
 
 // Make double register from float register
-FORCEINLINE VectorRegister4Double MakeVectorRegisterDouble(const VectorRegister4Float& From)
+FORCEINLINE VectorRegister4Double MakeVectorRegisterDouble(VectorRegister4Float From)
 {
 	return VectorRegister4Double(From);
 }
 
 // Lossy conversion: double->float vector
-FORCEINLINE VectorRegister4Float MakeVectorRegisterFloatFromDouble(const VectorRegister4Double& Vec4d)
+FORCEINLINE VectorRegister4Float MakeVectorRegisterFloatFromDouble(VectorRegister4Double Vec4d)
 {
 #if !UE_PLATFORM_MATH_USE_AVX
 	return _mm_movelh_ps(_mm_cvtpd_ps(Vec4d.XY), _mm_cvtpd_ps(Vec4d.ZW));
@@ -394,7 +394,8 @@ FORCEINLINE VectorRegister4Int MakeVectorRegisterInt(int32 X, int32 Y, int32 Z,
 	return _mm_setr_epi32(X, Y, Z, W);
 }
 
-FORCEINLINE VectorRegister2Int64 MakeVectorRegisterInt64(int64 X, int64 Y) {
+FORCEINLINE VectorRegister2Int64 MakeVectorRegisterInt64(int64 X, int64 Y)
+{
 	return _mm_set_epi64x(Y, X);
 }
 
@@ -444,7 +445,7 @@ FORCEINLINE constexpr VectorRegister2Double MakeVectorRegister2DoubleConstant(do
  *	Constants:
  *============================================================================*/
 
-#include "Math/UnrealMathVectorConstants.h"
+#include "Math/UnrealMathVectorConstants.h.inl"
 
 /*=============================================================================
  *	Intrinsics:
@@ -496,16 +497,19 @@ FORCEINLINE VectorRegister4Double VectorOneDouble(void)
  */
 
 template <uint32 ComponentIndex>
-FORCEINLINE float VectorGetComponentImpl(const VectorRegister4Float& Vec)
+FORCEINLINE float VectorGetComponentImpl(VectorRegister4Float Vec)
 {
 	return (((float*)&(Vec))[ComponentIndex]);
 }
 
 // Specializations
-template <> FORCEINLINE float VectorGetComponentImpl<0>(const VectorRegister4Float& Vec) { return _mm_cvtss_f32(Vec); }
+template <> FORCEINLINE float VectorGetComponentImpl<0>(VectorRegister4Float Vec)
+{
+	return _mm_cvtss_f32(Vec); 
+}
 
 template <uint32 ComponentIndex>
-FORCEINLINE double VectorGetComponentImpl(const VectorRegister4Double& Vec)
+FORCEINLINE double VectorGetComponentImpl(VectorRegister4Double Vec)
 {
 #if !UE_PLATFORM_MATH_USE_AVX
 	return (((double*)&(Vec.XY))[ComponentIndex]);
@@ -517,20 +521,20 @@ FORCEINLINE double VectorGetComponentImpl(const VectorRegister4Double& Vec)
 // Specializations
 #if UE_PLATFORM_MATH_USE_AVX
 // Lower latency than `vmovsd`, required since MSVC doesn't optimize the above impl well compared to clang/gcc. The latter basically generates something like this below (checked in godbolt).
-template <> FORCEINLINE double VectorGetComponentImpl<0>(const VectorRegister4Double& Vec) { return _mm256_cvtsd_f64(Vec); }
-template <> FORCEINLINE double VectorGetComponentImpl<1>(const VectorRegister4Double& Vec) { return _mm256_cvtsd_f64(_mm256_permute_pd(Vec, 1)); }
-template <> FORCEINLINE double VectorGetComponentImpl<2>(const VectorRegister4Double& Vec) { return _mm_cvtsd_f64(_mm256_extractf128_pd(Vec, 1)); }
-template <> FORCEINLINE double VectorGetComponentImpl<3>(const VectorRegister4Double& Vec) { return _mm_cvtsd_f64(_mm_permute_pd(_mm256_extractf128_pd(Vec, 1), 1)); }
+template <> FORCEINLINE double VectorGetComponentImpl<0>(VectorRegister4Double Vec) { return _mm256_cvtsd_f64(Vec); }
+template <> FORCEINLINE double VectorGetComponentImpl<1>(VectorRegister4Double Vec) { return _mm256_cvtsd_f64(_mm256_permute_pd(Vec, 1)); }
+template <> FORCEINLINE double VectorGetComponentImpl<2>(VectorRegister4Double Vec) { return _mm_cvtsd_f64(_mm256_extractf128_pd(Vec, 1)); }
+template <> FORCEINLINE double VectorGetComponentImpl<3>(VectorRegister4Double Vec) { return _mm_cvtsd_f64(_mm_permute_pd(_mm256_extractf128_pd(Vec, 1), 1)); }
 #endif
 
 #define VectorGetComponent(Vec, ComponentIndex) VectorGetComponentImpl<ComponentIndex>(Vec)
 
-FORCEINLINE float VectorGetComponentDynamic(const VectorRegister4Float& Vec, uint32 ComponentIndex)
+FORCEINLINE float VectorGetComponentDynamic(VectorRegister4Float Vec, uint32 ComponentIndex)
 {
 	return (((float*)&(Vec))[ComponentIndex]);
 }
 
-FORCEINLINE double VectorGetComponentDynamic(const VectorRegister4Double& Vec, uint32 ComponentIndex)
+FORCEINLINE double VectorGetComponentDynamic(VectorRegister4Double Vec, uint32 ComponentIndex)
 {
 #if !UE_PLATFORM_MATH_USE_AVX
 	return (((double*)&(Vec.XY))[ComponentIndex]);
@@ -670,7 +674,7 @@ FORCEINLINE VectorRegister4Double VectorLoadDouble1(const double* Ptr)
 	return Result;
 }
 
-FORCEINLINE VectorRegister4i VectorLoad64Bits(const VectorRegister4i *Ptr)
+FORCEINLINE VectorRegister4i VectorLoad64Bits(const void* Ptr)
 {
 	return _mm_loadu_si64((__m128i *)Ptr);
 }
@@ -759,12 +763,12 @@ FORCEINLINE VectorRegister4Double VectorSetFloat1(double D)
  * @param Vec	Vector to store
  * @param Ptr	Memory pointer
  */
-FORCEINLINE void VectorStore(const VectorRegister4Float& Vec, float* Ptr)
+FORCEINLINE void VectorStore(VectorRegister4Float Vec, float* Ptr)
 {
 	_mm_storeu_ps(Ptr, Vec);
 }
 
-FORCEINLINE void VectorStore(const VectorRegister4Double& Vec, double* Dst)
+FORCEINLINE void VectorStore(VectorRegister4Double Vec, double* Dst)
 {
 #if !UE_PLATFORM_MATH_USE_AVX
 	_mm_storeu_pd(Dst, Vec.XY);
@@ -780,7 +784,7 @@ FORCEINLINE void VectorStore(const VectorRegister4Double& Vec, double* Dst)
  * @param Vec	Vector to store
  * @param Ptr	Memory pointer
  */
-FORCEINLINE void VectorStore16(const VectorRegister4x4Float& Vec, float* Ptr)
+FORCEINLINE void VectorStore16(VectorRegister4x4Float Vec, float* Ptr)
 {
 	VectorStore(Vec.val[0], Ptr);
 	VectorStore(Vec.val[1], Ptr + 4);
@@ -794,12 +798,12 @@ FORCEINLINE void VectorStore16(const VectorRegister4x4Float& Vec, float* Ptr)
  * @param Vec	Vector to store
  * @param Ptr	Aligned memory pointer
  */
-FORCEINLINE void VectorStoreAligned(const VectorRegister4Float& Vec, float* Dst)
+FORCEINLINE void VectorStoreAligned(VectorRegister4Float Vec, float* Dst)
 {
 	_mm_store_ps(Dst, Vec);
 }
 
-FORCEINLINE void VectorStoreAligned(const VectorRegister4Double& Vec, double* Dst)
+FORCEINLINE void VectorStoreAligned(VectorRegister4Double Vec, double* Dst)
 {
 #if !UE_PLATFORM_MATH_USE_AVX
 	_mm_store_pd(Dst, Vec.XY);
@@ -817,12 +821,12 @@ FORCEINLINE void VectorStoreAligned(const VectorRegister4Double& Vec, double* Ds
  * @param Vec	Vector to store
  * @param Ptr	Aligned memory pointer
  */
-FORCEINLINE void VectorStoreAlignedStreamed(const VectorRegister4Float& Vec, float* Dst)
+FORCEINLINE void VectorStoreAlignedStreamed(VectorRegister4Float Vec, float* Dst)
 {
 	_mm_stream_ps(Dst, Vec);
 }
 
-FORCEINLINE void VectorStoreAlignedStreamed(const VectorRegister4Double& Vec, double* Dst)
+FORCEINLINE void VectorStoreAlignedStreamed(VectorRegister4Double Vec, double* Dst)
 {
 #if !UE_PLATFORM_MATH_USE_AVX
 	_mm_stream_pd(Dst, Vec.XY);
@@ -840,14 +844,14 @@ FORCEINLINE void VectorStoreAlignedStreamed(const VectorRegister4Double& Vec, do
  * @param Vec	Vector to store XYZ
  * @param Ptr	Unaligned memory pointer
  */
-FORCEINLINE void VectorStoreFloat3( const VectorRegister4Float& Vec, float* Ptr )
+FORCEINLINE void VectorStoreFloat3(VectorRegister4Float Vec, float* Ptr)
 {
 	VectorRegister4Float Tmp = Vec;
 	_mm_storel_pi((__m64*)(Ptr), Tmp);
 	_mm_store_ss(&Ptr[2], _mm_movehl_ps(Tmp, Tmp));
 }
 
-FORCEINLINE void VectorStoreFloat3(const VectorRegister4Double& Vec, double* Dst)
+FORCEINLINE void VectorStoreFloat3(VectorRegister4Double Vec, double* Dst)
 {
 #if !UE_PLATFORM_MATH_USE_AVX
 	_mm_storeu_pd(Dst, Vec.XY);
@@ -864,12 +868,12 @@ FORCEINLINE void VectorStoreFloat3(const VectorRegister4Double& Vec, double* Dst
  * @param Vec	Vector to store X
  * @param Ptr	Unaligned memory pointer
  */
-FORCEINLINE void VectorStoreFloat1(const VectorRegister4Float& Vec, float* Ptr)
+FORCEINLINE void VectorStoreFloat1(VectorRegister4Float Vec, float* Ptr)
 {
 	_mm_store_ss(Ptr, Vec);
 }
 
-FORCEINLINE void VectorStoreFloat1(const VectorRegister4Double& Vec, double* Dst)
+FORCEINLINE void VectorStoreFloat1(VectorRegister4Double Vec, double* Dst)
 {
 	_mm_store_sd(Dst, Vec.XY);
 }
@@ -889,7 +893,7 @@ namespace SSEPermuteHelpers
 	// Templated swizzles required for double shuffles when using __m128d, since we have to break it down in to two separate operations.
 
 	template <int Index0, int Index1>
-	FORCEINLINE VectorRegister2Double SelectVectorSwizzle2(const VectorRegister4Double& Vec)
+	FORCEINLINE VectorRegister2Double SelectVectorSwizzle2(VectorRegister4Double Vec)
 	{
 		if constexpr (Index0 <= 1)
 		{
@@ -919,13 +923,13 @@ namespace SSEPermuteHelpers
 		}
 	}
 
-	template<> FORCEINLINE VectorRegister2Double SelectVectorSwizzle2<0, 1>(const VectorRegister4Double& Vec) { return Vec.GetXY(); }
-	template<> FORCEINLINE VectorRegister2Double SelectVectorSwizzle2<2, 3>(const VectorRegister4Double& Vec) { return Vec.GetZW(); }
+	template<> FORCEINLINE VectorRegister2Double SelectVectorSwizzle2<0, 1>(VectorRegister4Double Vec) { return Vec.GetXY(); }
+	template<> FORCEINLINE VectorRegister2Double SelectVectorSwizzle2<2, 3>(VectorRegister4Double Vec) { return Vec.GetZW(); }
 
 #if UE_PLATFORM_MATH_USE_SSE4_1
 	// blend can run on more ports than shuffle, so are preferable even if latency is claimed to be the same.
-	template<> FORCEINLINE VectorRegister2Double SelectVectorSwizzle2<0, 3>(const VectorRegister4Double& Vec) { return _mm_blend_pd(Vec.GetXY(), Vec.GetZW(), SHUFFLEMASK2(0, 1)); }
-	template<> FORCEINLINE VectorRegister2Double SelectVectorSwizzle2<2, 1>(const VectorRegister4Double& Vec) { return _mm_blend_pd(Vec.GetZW(), Vec.GetXY(), SHUFFLEMASK2(0, 1)); }
+	template<> FORCEINLINE VectorRegister2Double SelectVectorSwizzle2<0, 3>(VectorRegister4Double Vec) { return _mm_blend_pd(Vec.GetXY(), Vec.GetZW(), SHUFFLEMASK2(0, 1)); }
+	template<> FORCEINLINE VectorRegister2Double SelectVectorSwizzle2<2, 1>(VectorRegister4Double Vec) { return _mm_blend_pd(Vec.GetZW(), Vec.GetXY(), SHUFFLEMASK2(0, 1)); }
 #endif // UE_PLATFORM_MATH_USE_SSE4_1
 
 
@@ -935,20 +939,20 @@ namespace SSEPermuteHelpers
 	constexpr int PERMUTE_LANE_MASK(int A, int B) { return (A == 0 ? 0x00 : 0x01) | (B == 0 ? (0x02 << 4) : (0x03 << 4)); }
 
 	template<int Lane0, int Lane1>
-	FORCEINLINE VectorRegister4Double PermuteLanes(const VectorRegister4Double& Vec)
+	FORCEINLINE VectorRegister4Double PermuteLanes(VectorRegister4Double Vec)
 	{
 		static_assert(Lane0 >= 0 && Lane0 <= 1 && Lane1 >= 0 && Lane1 <= 1, "Invalid Index");
 		return _mm256_permute2f128_pd(Vec, Vec, PERMUTE_LANE_MASK(Lane0, Lane1));
 	}
 
 	// Identity
-	template<> FORCEINLINE VectorRegister4Double PermuteLanes<0, 1>(const VectorRegister4Double& Vec) { return Vec; }
+	template<> FORCEINLINE VectorRegister4Double PermuteLanes<0, 1>(VectorRegister4Double Vec) { return Vec; }
 #if !UE_PLATFORM_MATH_USE_AVX_2
 	// On AVX1, permute2f128 can be quite slow, so look for alternatives (extract + insert). On AVX2, permute2f128 is more efficient and should equal or beat (extract + insert).
 	// Sources: https://www.agner.org/optimize/instruction_tables.pdf, https://uops.info/table.html
-	template<> FORCEINLINE VectorRegister4Double PermuteLanes<0, 0>(const VectorRegister4Double& Vec) { return _mm256_insertf128_pd(Vec, Vec.GetXY(), 1); } // copy XY to lane 1
-	template<> FORCEINLINE VectorRegister4Double PermuteLanes<1, 0>(const VectorRegister4Double& Vec) { return _mm256_setr_m128d(Vec.GetZW(), Vec.GetXY()); } // swap XY and ZW
-	template<> FORCEINLINE VectorRegister4Double PermuteLanes<1, 1>(const VectorRegister4Double& Vec) { return _mm256_insertf128_pd(Vec, Vec.GetZW(), 0); } // copy ZW to lane 0
+	template<> FORCEINLINE VectorRegister4Double PermuteLanes<0, 0>(VectorRegister4Double Vec) { return _mm256_insertf128_pd(Vec, Vec.GetXY(), 1); } // copy XY to lane 1
+	template<> FORCEINLINE VectorRegister4Double PermuteLanes<1, 0>(VectorRegister4Double Vec) { return _mm256_setr_m128d(Vec.GetZW(), Vec.GetXY()); } // swap XY and ZW
+	template<> FORCEINLINE VectorRegister4Double PermuteLanes<1, 1>(VectorRegister4Double Vec) { return _mm256_insertf128_pd(Vec, Vec.GetZW(), 0); } // copy ZW to lane 0
 #endif // !AVX2
 
 	//
@@ -958,7 +962,7 @@ namespace SSEPermuteHelpers
 	constexpr int PERMUTE_MASK(int A, int B, int C, int D) { return ((A == 1 ? (1 << 0) : 0) | (B == 1 ? (1 << 1) : 0) | (C == 3 ? (1 << 2) : 0) | (D == 3 ? (1 << 3) : 0)); }
 
 	template <int Index0, int Index1, int Index2, int Index3>
-	FORCEINLINE VectorRegister4Double SelectVectorSwizzle(const VectorRegister4Double& Vec)
+	FORCEINLINE VectorRegister4Double SelectVectorSwizzle(VectorRegister4Double Vec)
 	{
 		if constexpr (InLane0(Index0, Index1) && InLane1(Index2, Index3))
 		{
@@ -1000,17 +1004,17 @@ namespace SSEPermuteHelpers
 	//
 	// Specializations
 	//
-	template<> FORCEINLINE VectorRegister4Double SelectVectorSwizzle<0, 0, 2, 2>(const VectorRegister4Double& Vec) { return _mm256_movedup_pd(Vec); } // special instruction exists for this.
-	template<> FORCEINLINE VectorRegister4Double SelectVectorSwizzle<0, 1, 2, 3>(const VectorRegister4Double& Vec) { return Vec; } // Identity
-	template<> FORCEINLINE VectorRegister4Double SelectVectorSwizzle<2, 3, 0, 1>(const VectorRegister4Double& Vec) { return PermuteLanes<1, 0>(Vec); }
-	template<> FORCEINLINE VectorRegister4Double SelectVectorSwizzle<0, 1, 0, 1>(const VectorRegister4Double& Vec) { return PermuteLanes<0, 0>(Vec); }
-	template<> FORCEINLINE VectorRegister4Double SelectVectorSwizzle<2, 3, 2, 3>(const VectorRegister4Double& Vec) { return PermuteLanes<1, 1>(Vec); }
+	template<> FORCEINLINE VectorRegister4Double SelectVectorSwizzle<0, 0, 2, 2>(VectorRegister4Double Vec) { return _mm256_movedup_pd(Vec); } // special instruction exists for this.
+	template<> FORCEINLINE VectorRegister4Double SelectVectorSwizzle<0, 1, 2, 3>(VectorRegister4Double Vec) { return Vec; } // Identity
+	template<> FORCEINLINE VectorRegister4Double SelectVectorSwizzle<2, 3, 0, 1>(VectorRegister4Double Vec) { return PermuteLanes<1, 0>(Vec); }
+	template<> FORCEINLINE VectorRegister4Double SelectVectorSwizzle<0, 1, 0, 1>(VectorRegister4Double Vec) { return PermuteLanes<0, 0>(Vec); }
+	template<> FORCEINLINE VectorRegister4Double SelectVectorSwizzle<2, 3, 2, 3>(VectorRegister4Double Vec) { return PermuteLanes<1, 1>(Vec); }
 
 #endif // AVX
 
 	// Double swizzle wrapper
 	template<int Index0, int Index1, int Index2, int Index3>
-	FORCEINLINE VectorRegister4Double VectorSwizzleTemplate(const VectorRegister4Double& Vec)
+	FORCEINLINE VectorRegister4Double VectorSwizzleTemplate(VectorRegister4Double Vec)
 	{
 		static_assert(Index0 >= 0 && Index0 <= 3 && Index1 >= 0 && Index1 <= 3 && Index2 >= 0 && Index2 <= 3 && Index3 >= 0 && Index3 <= 3, "Invalid Index");
 
@@ -1025,14 +1029,14 @@ namespace SSEPermuteHelpers
 	}
 
 	// Specializations
-	template<> FORCEINLINE VectorRegister4Double VectorSwizzleTemplate<0, 1, 2, 3>(const VectorRegister4Double& Vec) { return Vec; } // Identity
+	template<> FORCEINLINE VectorRegister4Double VectorSwizzleTemplate<0, 1, 2, 3>(VectorRegister4Double Vec) { return Vec; } // Identity
 
 	/////////////////////////////////////////////////////////////////////////////////////////////////////////////////////////////////
 	// Double replicate
 	/////////////////////////////////////////////////////////////////////////////////////////////////////////////////////////////////
 
 	template <int Index>
-	FORCEINLINE VectorRegister2Double VectorReplicateImpl2(const VectorRegister2Double& Vec)
+	FORCEINLINE VectorRegister2Double VectorReplicateImpl2(VectorRegister2Double Vec)
 	{
 		// Note: 2 doubles (VectorRegister2Double / m128d)
 		return _mm_shuffle_pd(Vec, Vec, SHUFFLEMASK2(Index, Index));
@@ -1040,7 +1044,7 @@ namespace SSEPermuteHelpers
 
 	// Double replicate (4 doubles)
 	template <int Index>
-	FORCEINLINE VectorRegister4Double VectorReplicateImpl4(const VectorRegister4Double& Vec)
+	FORCEINLINE VectorRegister4Double VectorReplicateImpl4(VectorRegister4Double Vec)
 	{
 		if constexpr (Index <= 1)
 		{
@@ -1058,7 +1062,7 @@ namespace SSEPermuteHelpers
 	// Double replicate wrapper
 	//
 	template<int Index>
-	FORCEINLINE VectorRegister4Double VectorReplicateTemplate(const VectorRegister4Double& Vec)
+	FORCEINLINE VectorRegister4Double VectorReplicateTemplate(VectorRegister4Double Vec)
 	{
 		static_assert(Index >= 0 && Index <= 3, "Invalid Index");
 
@@ -1079,20 +1083,20 @@ namespace SSEPermuteHelpers
 	// Lane shuffle helper
 	//
 	template<int Lane0, int Lane1>
-	FORCEINLINE VectorRegister4Double ShuffleLanes(const VectorRegister4Double& Vec1, const VectorRegister4Double& Vec2)
+	FORCEINLINE VectorRegister4Double ShuffleLanes(VectorRegister4Double Vec1, VectorRegister4Double Vec2)
 	{
 		static_assert(Lane0 >= 0 && Lane0 <= 1 && Lane1 >= 0 && Lane1 <= 1, "Invalid Index");
 		return _mm256_permute2f128_pd(Vec1, Vec2, PERMUTE_LANE_MASK(Lane0, Lane1));
 	}
 
 	// Lane shuffle helper specialization
-	template<> FORCEINLINE VectorRegister4Double ShuffleLanes<0, 1>(const VectorRegister4Double& Vec1, const VectorRegister4Double& Vec2) { return _mm256_blend_pd(Vec1, Vec2, 0b1100); }
+	template<> FORCEINLINE VectorRegister4Double ShuffleLanes<0, 1>(VectorRegister4Double Vec1, VectorRegister4Double Vec2) { return _mm256_blend_pd(Vec1, Vec2, 0b1100); }
 #if !UE_PLATFORM_MATH_USE_AVX_2
 	// On AVX1, permute2f128 can be quite slow, so look for alternatives (extract + insert). On AVX2, permute2f128 is more efficient and should equal or beat (extract + insert).
 	// Sources: https://www.agner.org/optimize/instruction_tables.pdf, https://uops.info/table.html
-	template<> FORCEINLINE VectorRegister4Double ShuffleLanes<0, 0>(const VectorRegister4Double& Vec1, const VectorRegister4Double& Vec2) { return _mm256_insertf128_pd(Vec1, Vec2.GetXY(), 1); } // copy XY to lane 1
-	template<> FORCEINLINE VectorRegister4Double ShuffleLanes<1, 0>(const VectorRegister4Double& Vec1, const VectorRegister4Double& Vec2) { return _mm256_setr_m128d(Vec1.GetZW(), Vec2.GetXY()); } // swap XY and ZW
-	template<> FORCEINLINE VectorRegister4Double ShuffleLanes<1, 1>(const VectorRegister4Double& Vec1, const VectorRegister4Double& Vec2) { return _mm256_insertf128_pd(Vec2, Vec1.GetZW(), 0); } // copy ZW to lane 0
+	template<> FORCEINLINE VectorRegister4Double ShuffleLanes<0, 0>(VectorRegister4Double Vec1, VectorRegister4Double Vec2) { return _mm256_insertf128_pd(Vec1, Vec2.GetXY(), 1); } // copy XY to lane 1
+	template<> FORCEINLINE VectorRegister4Double ShuffleLanes<1, 0>(VectorRegister4Double Vec1, VectorRegister4Double Vec2) { return _mm256_setr_m128d(Vec1.GetZW(), Vec2.GetXY()); } // swap XY and ZW
+	template<> FORCEINLINE VectorRegister4Double ShuffleLanes<1, 1>(VectorRegister4Double Vec1, VectorRegister4Double Vec2) { return _mm256_insertf128_pd(Vec2, Vec1.GetZW(), 0); } // copy ZW to lane 0
 #endif // !AVX2
 
 	//
@@ -1103,7 +1107,7 @@ namespace SSEPermuteHelpers
 	// then efficiently swizzles within 128-bit lanes using specializations for indices [0..1][0..1][2..3][2..3]
 	// 
 	template <int Index0, int Index1, int Index2, int Index3>
-	FORCEINLINE VectorRegister4Double SelectVectorShuffle(const VectorRegister4Double& Vec1, const VectorRegister4Double& Vec2)
+	FORCEINLINE VectorRegister4Double SelectVectorShuffle(VectorRegister4Double Vec1, VectorRegister4Double Vec2)
 	{
 		if constexpr (InLane0(Index0, Index1) && InLane1(Index2, Index3))
 		{
@@ -1152,16 +1156,16 @@ namespace SSEPermuteHelpers
 
 	// AVX Double Shuffle specializations 
 	// Shuffles of 128-bit pairs, ie combinations of [0,1][2,3].
-	template<> FORCEINLINE VectorRegister4Double SelectVectorShuffle<0, 1, 0, 1>(const VectorRegister4Double& Vec1, const VectorRegister4Double& Vec2) { return ShuffleLanes<0, 0>(Vec1, Vec2); }
-	template<> FORCEINLINE VectorRegister4Double SelectVectorShuffle<0, 1, 2, 3>(const VectorRegister4Double& Vec1, const VectorRegister4Double& Vec2) { return ShuffleLanes<0, 1>(Vec1, Vec2); }
-	template<> FORCEINLINE VectorRegister4Double SelectVectorShuffle<2, 3, 0, 1>(const VectorRegister4Double& Vec1, const VectorRegister4Double& Vec2) { return ShuffleLanes<1, 0>(Vec1, Vec2); }
-	template<> FORCEINLINE VectorRegister4Double SelectVectorShuffle<2, 3, 2, 3>(const VectorRegister4Double& Vec1, const VectorRegister4Double& Vec2) { return ShuffleLanes<1, 1>(Vec1, Vec2); }
+	template<> FORCEINLINE VectorRegister4Double SelectVectorShuffle<0, 1, 0, 1>(VectorRegister4Double Vec1, VectorRegister4Double Vec2) { return ShuffleLanes<0, 0>(Vec1, Vec2); }
+	template<> FORCEINLINE VectorRegister4Double SelectVectorShuffle<0, 1, 2, 3>(VectorRegister4Double Vec1, VectorRegister4Double Vec2) { return ShuffleLanes<0, 1>(Vec1, Vec2); }
+	template<> FORCEINLINE VectorRegister4Double SelectVectorShuffle<2, 3, 0, 1>(VectorRegister4Double Vec1, VectorRegister4Double Vec2) { return ShuffleLanes<1, 0>(Vec1, Vec2); }
+	template<> FORCEINLINE VectorRegister4Double SelectVectorShuffle<2, 3, 2, 3>(VectorRegister4Double Vec1, VectorRegister4Double Vec2) { return ShuffleLanes<1, 1>(Vec1, Vec2); }
 
 #else
 
 	// Non-AVX implementation
 	template<int Index0, int Index1, int Index2, int Index3>
-	FORCEINLINE VectorRegister4Double SelectVectorShuffle(const VectorRegister4Double& Vec1, const VectorRegister4Double& Vec2)
+	FORCEINLINE VectorRegister4Double SelectVectorShuffle(VectorRegister4Double Vec1, VectorRegister4Double Vec2)
 	{
 		return VectorRegister4Double(
 			SelectVectorSwizzle2<Index0, Index1>(Vec1),
@@ -1175,7 +1179,7 @@ namespace SSEPermuteHelpers
 	// Double shuffle wrapper
 	//
 	template<int Index0, int Index1, int Index2, int Index3>
-	FORCEINLINE VectorRegister4Double VectorShuffleTemplate(const VectorRegister4Double& Vec1, const VectorRegister4Double& Vec2)
+	FORCEINLINE VectorRegister4Double VectorShuffleTemplate(VectorRegister4Double Vec1, VectorRegister4Double Vec2)
 	{
 		static_assert(Index0 >= 0 && Index0 <= 3 && Index1 >= 0 && Index1 <= 3 && Index2 >= 0 && Index2 <= 3 && Index3 >= 0 && Index3 <= 3, "Invalid Index");
 		return SelectVectorShuffle<Index0, Index1, Index2, Index3>(Vec1, Vec2);
@@ -1186,7 +1190,7 @@ namespace SSEPermuteHelpers
 	/////////////////////////////////////////////////////////////////////////////////////////////////////////////////////////////////
 
 	template<int Index0, int Index1, int Index2, int Index3>
-	FORCEINLINE VectorRegister4Float VectorSwizzleTemplate(const VectorRegister4Float& Vec)
+	FORCEINLINE VectorRegister4Float VectorSwizzleTemplate(VectorRegister4Float Vec)
 	{
 		VectorRegister4Float Result = _mm_shuffle_ps(Vec, Vec, SHUFFLEMASK(Index0, Index1, Index2, Index3));
 		return Result;
@@ -1195,25 +1199,25 @@ namespace SSEPermuteHelpers
 	// Float Swizzle specializations.
 	// These can result in no-ops or simpler ops than shuffle which don't compete with the shuffle unit, or which can copy directly to the destination and avoid an intermediate mov.
 	// See: https://stackoverflow.com/questions/56238197/what-is-the-difference-between-mm-movehdup-ps-and-mm-shuffle-ps-in-this-case
-	template<> FORCEINLINE VectorRegister4Float VectorSwizzleTemplate<0, 1, 2, 3>(const VectorRegister4Float& Vec) { return Vec; }
-	template<> FORCEINLINE VectorRegister4Float VectorSwizzleTemplate<0, 1, 0, 1>(const VectorRegister4Float& Vec) { return _mm_movelh_ps(Vec, Vec); }
-	template<> FORCEINLINE VectorRegister4Float VectorSwizzleTemplate<2, 3, 2, 3>(const VectorRegister4Float& Vec) { return _mm_movehl_ps(Vec, Vec); }
-	template<> FORCEINLINE VectorRegister4Float VectorSwizzleTemplate<0, 0, 1, 1>(const VectorRegister4Float& Vec) { return _mm_unpacklo_ps(Vec, Vec); }
-	template<> FORCEINLINE VectorRegister4Float VectorSwizzleTemplate<2, 2, 3, 3>(const VectorRegister4Float& Vec) { return _mm_unpackhi_ps(Vec, Vec); }
+	template<> FORCEINLINE VectorRegister4Float VectorSwizzleTemplate<0, 1, 2, 3>(VectorRegister4Float Vec) { return Vec; }
+	template<> FORCEINLINE VectorRegister4Float VectorSwizzleTemplate<0, 1, 0, 1>(VectorRegister4Float Vec) { return _mm_movelh_ps(Vec, Vec); }
+	template<> FORCEINLINE VectorRegister4Float VectorSwizzleTemplate<2, 3, 2, 3>(VectorRegister4Float Vec) { return _mm_movehl_ps(Vec, Vec); }
+	template<> FORCEINLINE VectorRegister4Float VectorSwizzleTemplate<0, 0, 1, 1>(VectorRegister4Float Vec) { return _mm_unpacklo_ps(Vec, Vec); }
+	template<> FORCEINLINE VectorRegister4Float VectorSwizzleTemplate<2, 2, 3, 3>(VectorRegister4Float Vec) { return _mm_unpackhi_ps(Vec, Vec); }
 
 #if UE_PLATFORM_MATH_USE_SSE4_1
-	template<> FORCEINLINE VectorRegister4Float VectorSwizzleTemplate<0, 0, 2, 2>(const VectorRegister4Float& Vec) { return _mm_moveldup_ps(Vec); }
-	template<> FORCEINLINE VectorRegister4Float VectorSwizzleTemplate<1, 1, 3, 3>(const VectorRegister4Float& Vec) { return _mm_movehdup_ps(Vec); }
+	template<> FORCEINLINE VectorRegister4Float VectorSwizzleTemplate<0, 0, 2, 2>(VectorRegister4Float Vec) { return _mm_moveldup_ps(Vec); }
+	template<> FORCEINLINE VectorRegister4Float VectorSwizzleTemplate<1, 1, 3, 3>(VectorRegister4Float Vec) { return _mm_movehdup_ps(Vec); }
 #endif
 
 #if UE_PLATFORM_MATH_USE_AVX_2
-	template<> FORCEINLINE VectorRegister4Float VectorSwizzleTemplate<0, 0, 0, 0>(const VectorRegister4Float& Vec) { return _mm_broadcastss_ps(Vec); }
+	template<> FORCEINLINE VectorRegister4Float VectorSwizzleTemplate<0, 0, 0, 0>(VectorRegister4Float Vec) { return _mm_broadcastss_ps(Vec); }
 #endif
 
 	/////////////////////////////////////////////////////////////////////////////////////////////////////////////////////////////////
 	// Float replicate
 	template<int Index>
-	FORCEINLINE VectorRegister4Float VectorReplicateTemplate(const VectorRegister4Float& Vec)
+	FORCEINLINE VectorRegister4Float VectorReplicateTemplate(VectorRegister4Float Vec)
 	{
 		static_assert(Index >= 0 && Index <= 3, "Invalid Index");
 		return VectorSwizzleTemplate<Index, Index, Index, Index>(Vec);
@@ -1222,15 +1226,15 @@ namespace SSEPermuteHelpers
 	/////////////////////////////////////////////////////////////////////////////////////////////////////////////////////////////////
 	// Float shuffle
 	template<int Index0, int Index1, int Index2, int Index3>
-	FORCEINLINE VectorRegister4Float VectorShuffleTemplate(const VectorRegister4Float& Vec1, const VectorRegister4Float& Vec2)
+	FORCEINLINE VectorRegister4Float VectorShuffleTemplate(VectorRegister4Float Vec1, VectorRegister4Float Vec2)
 	{
 		static_assert(Index0 >= 0 && Index0 <= 3 && Index1 >= 0 && Index1 <= 3 && Index2 >= 0 && Index2 <= 3 && Index3 >= 0 && Index3 <= 3, "Invalid Index");
 		return _mm_shuffle_ps(Vec1, Vec2, SHUFFLEMASK(Index0, Index1, Index2, Index3));
 	}
 
 	// Float Shuffle specializations
-	template<> FORCEINLINE VectorRegister4Float VectorShuffleTemplate<0, 1, 0, 1>(const VectorRegister4Float& Vec1, const VectorRegister4Float& Vec2) { return _mm_movelh_ps(Vec1, Vec2); }
-	template<> FORCEINLINE VectorRegister4Float VectorShuffleTemplate<2, 3, 2, 3>(const VectorRegister4Float& Vec1, const VectorRegister4Float& Vec2) { return _mm_movehl_ps(Vec2, Vec1); } // Note: movehl copies first from the 2nd argument
+	template<> FORCEINLINE VectorRegister4Float VectorShuffleTemplate<0, 1, 0, 1>(VectorRegister4Float Vec1, VectorRegister4Float Vec2) { return _mm_movelh_ps(Vec1, Vec2); }
+	template<> FORCEINLINE VectorRegister4Float VectorShuffleTemplate<2, 3, 2, 3>(VectorRegister4Float Vec1, VectorRegister4Float Vec2) { return _mm_movehl_ps(Vec2, Vec1); } // Note: movehl copies first from the 2nd argument
 
 #undef OutOfLane
 #undef InSameLane
@@ -1246,7 +1250,7 @@ namespace SSEPermuteHelpers
  * @return				VectorRegister4Float( Vec[ElementIndex], Vec[ElementIndex], Vec[ElementIndex], Vec[ElementIndex] )
  */
 
-#define VectorReplicate( Vec, ElementIndex )	SSEPermuteHelpers::VectorReplicateTemplate<ElementIndex>(Vec)
+#define VectorReplicate(Vec, ElementIndex)	SSEPermuteHelpers::VectorReplicateTemplate<ElementIndex>(Vec)
 
  /**
   * Swizzles the 4 components of a vector and returns the result.
@@ -1258,7 +1262,7 @@ namespace SSEPermuteHelpers
   * @param W			Index for which component to use for W (literal 0-3)
   * @return			The swizzled vector
   */
-#define VectorSwizzle( Vec, X, Y, Z, W )		SSEPermuteHelpers::VectorSwizzleTemplate<X,Y,Z,W>( Vec )
+#define VectorSwizzle(Vec, X, Y, Z, W)		SSEPermuteHelpers::VectorSwizzleTemplate<X,Y,Z,W>(Vec)
 
   /**
    * Creates a vector through selecting two components from each vector via a shuffle mask.
@@ -1271,7 +1275,7 @@ namespace SSEPermuteHelpers
    * @param W			Index for which component of Vector2 to use for W (literal 0-3)
    * @return			The swizzled vector
    */
-#define VectorShuffle( Vec1, Vec2, X, Y, Z, W )		SSEPermuteHelpers::VectorShuffleTemplate<X,Y,Z,W>( Vec1, Vec2 )
+#define VectorShuffle(Vec1, Vec2, X, Y, Z, W)		SSEPermuteHelpers::VectorShuffleTemplate<X,Y,Z,W>(Vec1, Vec2)
 
 
 /**
@@ -1280,12 +1284,12 @@ namespace SSEPermuteHelpers
  * @param Vec			Source vector
  * @return				VectorRegister4Float( abs(Vec.x), abs(Vec.y), abs(Vec.z), abs(Vec.w) )
  */
-FORCEINLINE VectorRegister4Float VectorAbs(const VectorRegister4Float& Vec)
+FORCEINLINE VectorRegister4Float VectorAbs(VectorRegister4Float Vec)
 {
 	return _mm_and_ps(Vec, GlobalVectorConstants::SignMask());
 }
 
-FORCEINLINE VectorRegister4Double VectorAbs(const VectorRegister4Double& Vec)
+FORCEINLINE VectorRegister4Double VectorAbs(VectorRegister4Double Vec)
 {
 	VectorRegister4Double Result;
 #if !UE_PLATFORM_MATH_USE_AVX
@@ -1304,12 +1308,12 @@ FORCEINLINE VectorRegister4Double VectorAbs(const VectorRegister4Double& Vec)
  * @param Vec			Source vector
  * @return				VectorRegister4Float( -Vec.x, -Vec.y, -Vec.z, -Vec.w )
  */
-FORCEINLINE VectorRegister4Float VectorNegate(const VectorRegister4Float& Vec)
+FORCEINLINE VectorRegister4Float VectorNegate(VectorRegister4Float Vec)
 {
 	return _mm_sub_ps(_mm_setzero_ps(), Vec);
 }
 
-FORCEINLINE VectorRegister4Double VectorNegate(const VectorRegister4Double& Vec)
+FORCEINLINE VectorRegister4Double VectorNegate(VectorRegister4Double Vec)
 {
 	VectorRegister4Double Result;
 #if !UE_PLATFORM_MATH_USE_AVX
@@ -1329,12 +1333,12 @@ FORCEINLINE VectorRegister4Double VectorNegate(const VectorRegister4Double& Vec)
  * @return		VectorRegister4Float( Vec1.x+Vec2.x, Vec1.y+Vec2.y, Vec1.z+Vec2.z, Vec1.w+Vec2.w )
  */
 
-FORCEINLINE VectorRegister4Float VectorAdd( const VectorRegister4Float& Vec1, const VectorRegister4Float& Vec2 )
+FORCEINLINE VectorRegister4Float VectorAdd(VectorRegister4Float Vec1, VectorRegister4Float Vec2)
 {
 	return _mm_add_ps(Vec1, Vec2);
 }
 
-FORCEINLINE VectorRegister4Double VectorAdd(const VectorRegister4Double& Vec1, const VectorRegister4Double& Vec2)
+FORCEINLINE VectorRegister4Double VectorAdd(VectorRegister4Double Vec1, VectorRegister4Double Vec2)
 {
 	VectorRegister4Double Result;
 #if !UE_PLATFORM_MATH_USE_AVX
@@ -1353,12 +1357,12 @@ FORCEINLINE VectorRegister4Double VectorAdd(const VectorRegister4Double& Vec1, c
  * @param Vec2	2nd vector
  * @return		VectorRegister4Float( Vec1.x-Vec2.x, Vec1.y-Vec2.y, Vec1.z-Vec2.z, Vec1.w-Vec2.w )
  */
-FORCEINLINE VectorRegister4Float VectorSubtract( const VectorRegister4Float& Vec1, const VectorRegister4Float& Vec2 )
+FORCEINLINE VectorRegister4Float VectorSubtract(VectorRegister4Float Vec1, VectorRegister4Float Vec2)
 {
 	return _mm_sub_ps(Vec1, Vec2);
 }
 
-FORCEINLINE VectorRegister4Double VectorSubtract(const VectorRegister4Double& Vec1, const VectorRegister4Double& Vec2)
+FORCEINLINE VectorRegister4Double VectorSubtract(VectorRegister4Double Vec1, VectorRegister4Double Vec2)
 {
 	VectorRegister4Double Result;
 #if !UE_PLATFORM_MATH_USE_AVX
@@ -1377,12 +1381,12 @@ FORCEINLINE VectorRegister4Double VectorSubtract(const VectorRegister4Double& Ve
  * @param Vec2	2nd vector
  * @return		VectorRegister4Float( Vec1.x*Vec2.x, Vec1.y*Vec2.y, Vec1.z*Vec2.z, Vec1.w*Vec2.w )
  */
-FORCEINLINE VectorRegister4Float VectorMultiply( const VectorRegister4Float& Vec1, const VectorRegister4Float& Vec2 )
+FORCEINLINE VectorRegister4Float VectorMultiply(VectorRegister4Float Vec1, VectorRegister4Float Vec2)
 {
 	return _mm_mul_ps(Vec1, Vec2);
 }
 
-FORCEINLINE VectorRegister4Double VectorMultiply(const VectorRegister4Double& Vec1, const VectorRegister4Double& Vec2)
+FORCEINLINE VectorRegister4Double VectorMultiply(VectorRegister4Double Vec1, VectorRegister4Double Vec2)
 {
 	VectorRegister4Double Result;
 #if !UE_PLATFORM_MATH_USE_AVX
@@ -1404,7 +1408,7 @@ FORCEINLINE VectorRegister4Double VectorMultiply(const VectorRegister4Double& Ve
  * @param C		3rd vector
  * @return		VectorRegister4Float( A.x*B.x + C.x, A.y*B.y + C.y, A.z*B.z + C.z, A.w*B.w + C.w )
  */
-FORCEINLINE VectorRegister4Float VectorMultiplyAdd(const VectorRegister4Float& A, const VectorRegister4Float& B, const VectorRegister4Float& C)
+FORCEINLINE VectorRegister4Float VectorMultiplyAdd(VectorRegister4Float A, VectorRegister4Float B, VectorRegister4Float C)
 {
 #if UE_PLATFORM_MATH_USE_FMA3
 	return _mm_fmadd_ps(A, B, C);
@@ -1413,7 +1417,7 @@ FORCEINLINE VectorRegister4Float VectorMultiplyAdd(const VectorRegister4Float& A
 #endif
 }
 
-FORCEINLINE VectorRegister4Double VectorMultiplyAdd(const VectorRegister4Double& A, const VectorRegister4Double& B, const VectorRegister4Double& C)
+FORCEINLINE VectorRegister4Double VectorMultiplyAdd(VectorRegister4Double A, VectorRegister4Double B, VectorRegister4Double C)
 {
 #if UE_PLATFORM_MATH_USE_FMA3 && UE_PLATFORM_MATH_USE_AVX
 	return _mm256_fmadd_pd(A, B, C);
@@ -1435,7 +1439,7 @@ FORCEINLINE VectorRegister4Double VectorMultiplyAdd(const VectorRegister4Double&
  * @param C		3rd vector
  * @return		VectorRegister( C.x - A.x*B.x, C.y - A.y*B.y, C.z - A.z*B.z, C.w - A.w*B.w )
  */
-FORCEINLINE VectorRegister4Float VectorNegateMultiplyAdd(const VectorRegister4Float& A, const VectorRegister4Float& B, const VectorRegister4Float& C)
+FORCEINLINE VectorRegister4Float VectorNegateMultiplyAdd(VectorRegister4Float A, VectorRegister4Float B, VectorRegister4Float C)
 {
 #if UE_PLATFORM_MATH_USE_FMA3
 	return _mm_fnmadd_ps(A, B, C);
@@ -1444,7 +1448,7 @@ FORCEINLINE VectorRegister4Float VectorNegateMultiplyAdd(const VectorRegister4Fl
 #endif
 }
 
-FORCEINLINE VectorRegister4Double VectorNegateMultiplyAdd(const VectorRegister4Double& A, const VectorRegister4Double& B, const VectorRegister4Double& C)
+FORCEINLINE VectorRegister4Double VectorNegateMultiplyAdd(VectorRegister4Double A, VectorRegister4Double B, VectorRegister4Double C)
 {
 #if UE_PLATFORM_MATH_USE_FMA3 && UE_PLATFORM_MATH_USE_AVX
 	return _mm256_fnmadd_pd(A, B, C);
@@ -1466,12 +1470,12 @@ FORCEINLINE VectorRegister4Double VectorNegateMultiplyAdd(const VectorRegister4D
  * @param Vec2	2nd vector
  * @return		VectorRegister4Float( Vec1.x/Vec2.x, Vec1.y/Vec2.y, Vec1.z/Vec2.z, Vec1.w/Vec2.w )
  */
-FORCEINLINE VectorRegister4Float VectorDivide(const VectorRegister4Float& Vec1, const VectorRegister4Float& Vec2)
+FORCEINLINE VectorRegister4Float VectorDivide(VectorRegister4Float Vec1, VectorRegister4Float Vec2)
 {
 	return _mm_div_ps(Vec1, Vec2);
 }
 
-FORCEINLINE VectorRegister4Double VectorDivide(const VectorRegister4Double& Vec1, const VectorRegister4Double& Vec2)
+FORCEINLINE VectorRegister4Double VectorDivide(VectorRegister4Double Vec1, VectorRegister4Double Vec2)
 {
 	VectorRegister4Double Result;
 #if !UE_PLATFORM_MATH_USE_AVX
@@ -1486,7 +1490,7 @@ FORCEINLINE VectorRegister4Double VectorDivide(const VectorRegister4Double& Vec1
 namespace SSEVectorHelperFuncs
 {
 	// Computes VectorDot3 but only with the result in the first (X) element of a VectorRegister4Float
-	FORCEINLINE VectorRegister4Float InternalVectorDot3X(const VectorRegister4Float& Vec1, const VectorRegister4Float& Vec2)
+	FORCEINLINE VectorRegister4Float InternalVectorDot3X(VectorRegister4Float Vec1, VectorRegister4Float Vec2)
 	{
 		// (X, Y, Z, W)
 		VectorRegister4Float Prod = VectorMultiply(Vec1, Vec2);
@@ -1502,7 +1506,7 @@ namespace SSEVectorHelperFuncs
 	}
 
 	// Computes VectorDot3 but only with the result in the first (X) element of a VectorRegister4Double
-	FORCEINLINE VectorRegister4Double InternalVectorDot3X_Full(const VectorRegister4Double& Vec1, const VectorRegister4Double& Vec2)
+	FORCEINLINE VectorRegister4Double InternalVectorDot3X_Full(VectorRegister4Double Vec1, VectorRegister4Double Vec2)
 	{
 		// (X, Y, Z, W)
 		VectorRegister4Double Prod = VectorMultiply(Vec1, Vec2);
@@ -1519,7 +1523,7 @@ namespace SSEVectorHelperFuncs
 	}
 
 	// Computes VectorDot3 but only with the result in the first (X) element of a VectorRegister2Double (half of VectorRegister4Double)
-	FORCEINLINE VectorRegister2Double InternalVectorDot3X_Half(const VectorRegister4Double& Vec1, const VectorRegister4Double& Vec2)
+	FORCEINLINE VectorRegister2Double InternalVectorDot3X_Half(VectorRegister4Double Vec1, VectorRegister4Double Vec2)
 	{
 		VectorRegister2Double T, A;
 
@@ -1552,12 +1556,12 @@ namespace SSEVectorHelperFuncs
  * @param Vec2	2nd vector
  * @return		d = dot3(Vec1.xyz, Vec2.xyz)
  */
-FORCEINLINE float VectorDot3Scalar(const VectorRegister4Float& Vec1, const VectorRegister4Float& Vec2)
+FORCEINLINE float VectorDot3Scalar(VectorRegister4Float Vec1, VectorRegister4Float Vec2)
 {
 	return VectorGetComponent(SSEVectorHelperFuncs::InternalVectorDot3X(Vec1, Vec2), 0);
 }
 
-FORCEINLINE double VectorDot3Scalar(const VectorRegister4Double& Vec1, const VectorRegister4Double& Vec2)
+FORCEINLINE double VectorDot3Scalar(VectorRegister4Double Vec1, VectorRegister4Double Vec2)
 {
 #if UE_PLATFORM_MATH_USE_AVX
 	return VectorGetComponent(SSEVectorHelperFuncs::InternalVectorDot3X_Full(Vec1, Vec2), 0);
@@ -1575,12 +1579,12 @@ FORCEINLINE double VectorDot3Scalar(const VectorRegister4Double& Vec1, const Vec
  * @param Vec2	2nd vector
  * @return		d = dot3(Vec1.xyz, Vec2.xyz), VectorRegister4Float( d, d, d, d )
  */
-FORCEINLINE VectorRegister4Float VectorDot3(const VectorRegister4Float& Vec1, const VectorRegister4Float& Vec2)
+FORCEINLINE VectorRegister4Float VectorDot3(VectorRegister4Float Vec1, VectorRegister4Float Vec2)
 {
 	return VectorReplicate(SSEVectorHelperFuncs::InternalVectorDot3X(Vec1, Vec2), 0);
 }
 
-FORCEINLINE VectorRegister4Double VectorDot3(const VectorRegister4Double& Vec1, const VectorRegister4Double& Vec2)
+FORCEINLINE VectorRegister4Double VectorDot3(VectorRegister4Double Vec1, VectorRegister4Double Vec2)
 {
 #if UE_PLATFORM_MATH_USE_AVX
 	return VectorReplicate(SSEVectorHelperFuncs::InternalVectorDot3X_Full(Vec1, Vec2), 0);
@@ -1600,7 +1604,7 @@ FORCEINLINE VectorRegister4Double VectorDot3(const VectorRegister4Double& Vec1,
  * @param Vec2	2nd vector
  * @return		d = dot4(Vec1, Vec2), VectorRegister4Float( d, d, d, d )
  */
-FORCEINLINE VectorRegister4Float VectorDot4( const VectorRegister4Float& Vec1, const VectorRegister4Float& Vec2 )
+FORCEINLINE VectorRegister4Float VectorDot4(VectorRegister4Float Vec1, VectorRegister4Float Vec2)
 {
 	VectorRegister4Float R, T;
 	R = VectorMultiply(Vec1, Vec2);		// (XX, YY, ZZ, WW)
@@ -1610,7 +1614,7 @@ FORCEINLINE VectorRegister4Float VectorDot4( const VectorRegister4Float& Vec1, c
 	return VectorAdd(R, T);				// (XX + YY + ZZ + WW, YY + XX + WW + ZZ, ZZ + WW + XX + YY, WW + ZZ + YY + XX)
 }
 
-FORCEINLINE VectorRegister4Double VectorDot4(const VectorRegister4Double& Vec1, const VectorRegister4Double& Vec2)
+FORCEINLINE VectorRegister4Double VectorDot4(VectorRegister4Double Vec1, VectorRegister4Double Vec2)
 {
 #if UE_PLATFORM_MATH_USE_AVX
 	// AVX implementation uses fast permutes
@@ -1649,12 +1653,12 @@ FORCEINLINE VectorRegister4Double VectorDot4(const VectorRegister4Double& Vec1,
  * @param Vec2	2nd vector
  * @return		VectorRegister4Float( Vec1.x == Vec2.x ? 0xFFFFFFFF : 0, same for yzw )
  */
-FORCEINLINE VectorRegister4Float VectorCompareEQ(const VectorRegister4Float& Vec1, const VectorRegister4Float& Vec2)
+FORCEINLINE VectorRegister4Float VectorCompareEQ(VectorRegister4Float Vec1, VectorRegister4Float Vec2)
 {
 	return _mm_cmpeq_ps(Vec1, Vec2);
 }
 
-FORCEINLINE VectorRegister4Double VectorCompareEQ(const VectorRegister4Double& Vec1, const VectorRegister4Double& Vec2)
+FORCEINLINE VectorRegister4Double VectorCompareEQ(VectorRegister4Double Vec1, VectorRegister4Double Vec2)
 {
 	VectorRegister4Double Result;
 #if !UE_PLATFORM_MATH_USE_AVX
@@ -1673,12 +1677,12 @@ FORCEINLINE VectorRegister4Double VectorCompareEQ(const VectorRegister4Double& V
  * @param Vec2	2nd vector
  * @return		VectorRegister4Float( Vec1.x != Vec2.x ? 0xFFFFFFFF : 0, same for yzw )
  */
-FORCEINLINE VectorRegister4Float VectorCompareNE(const VectorRegister4Float& Vec1, const VectorRegister4Float& Vec2)
+FORCEINLINE VectorRegister4Float VectorCompareNE(VectorRegister4Float Vec1, VectorRegister4Float Vec2)
 {
 	return _mm_cmpneq_ps(Vec1, Vec2);
 }
 
-FORCEINLINE VectorRegister4Double VectorCompareNE(const VectorRegister4Double& Vec1, const VectorRegister4Double& Vec2)
+FORCEINLINE VectorRegister4Double VectorCompareNE(VectorRegister4Double Vec1, VectorRegister4Double Vec2)
 {
 	VectorRegister4Double Result;
 #if !UE_PLATFORM_MATH_USE_AVX
@@ -1699,12 +1703,12 @@ FORCEINLINE VectorRegister4Double VectorCompareNE(const VectorRegister4Double& V
  * @param Vec2	2nd vector
  * @return		VectorRegister4Float( Vec1.x > Vec2.x ? 0xFFFFFFFF : 0, same for yzw )
  */
-FORCEINLINE VectorRegister4Float VectorCompareGT(const VectorRegister4Float& Vec1, const VectorRegister4Float& Vec2)
+FORCEINLINE VectorRegister4Float VectorCompareGT(VectorRegister4Float Vec1, VectorRegister4Float Vec2)
 {
 	return _mm_cmpgt_ps(Vec1, Vec2);
 }
 
-FORCEINLINE VectorRegister4Double VectorCompareGT(const VectorRegister4Double& Vec1, const VectorRegister4Double& Vec2)
+FORCEINLINE VectorRegister4Double VectorCompareGT(VectorRegister4Double Vec1, VectorRegister4Double Vec2)
 {
 	VectorRegister4Double Result;
 #if !UE_PLATFORM_MATH_USE_AVX
@@ -1723,12 +1727,12 @@ FORCEINLINE VectorRegister4Double VectorCompareGT(const VectorRegister4Double& V
  * @param Vec2	2nd vector
  * @return		VectorRegister4Float( Vec1.x >= Vec2.x ? 0xFFFFFFFF : 0, same for yzw )
  */
-FORCEINLINE VectorRegister4Float VectorCompareGE(const VectorRegister4Float& Vec1, const VectorRegister4Float& Vec2)
+FORCEINLINE VectorRegister4Float VectorCompareGE(VectorRegister4Float Vec1, VectorRegister4Float Vec2)
 {
 	return _mm_cmpge_ps(Vec1, Vec2);
 }
 
-FORCEINLINE VectorRegister4Double VectorCompareGE(const VectorRegister4Double& Vec1, const VectorRegister4Double& Vec2)
+FORCEINLINE VectorRegister4Double VectorCompareGE(VectorRegister4Double Vec1, VectorRegister4Double Vec2)
 {
 	VectorRegister4Double Result;
 #if !UE_PLATFORM_MATH_USE_AVX
@@ -1747,12 +1751,12 @@ FORCEINLINE VectorRegister4Double VectorCompareGE(const VectorRegister4Double& V
  * @param Vec2	2nd vector
  * @return		VectorRegister4Float( Vec1.x < Vec2.x ? 0xFFFFFFFF : 0, same for yzw )
  */
-FORCEINLINE VectorRegister4Float VectorCompareLT(const VectorRegister4Float& Vec1, const VectorRegister4Float& Vec2)
+FORCEINLINE VectorRegister4Float VectorCompareLT(VectorRegister4Float Vec1, VectorRegister4Float Vec2)
 {
 	return _mm_cmplt_ps(Vec1, Vec2);
 }
 
-FORCEINLINE VectorRegister4Double VectorCompareLT(const VectorRegister4Double& Vec1, const VectorRegister4Double& Vec2)
+FORCEINLINE VectorRegister4Double VectorCompareLT(VectorRegister4Double Vec1, VectorRegister4Double Vec2)
 {
 	VectorRegister4Double Result;
 #if !UE_PLATFORM_MATH_USE_AVX
@@ -1771,12 +1775,12 @@ FORCEINLINE VectorRegister4Double VectorCompareLT(const VectorRegister4Double& V
  * @param Vec2	2nd vector
  * @return		VectorRegister4Float( Vec1.x <= Vec2.x ? 0xFFFFFFFF : 0, same for yzw )
  */
-FORCEINLINE VectorRegister4Float VectorCompareLE(const VectorRegister4Float& Vec1, const VectorRegister4Float& Vec2)
+FORCEINLINE VectorRegister4Float VectorCompareLE(VectorRegister4Float Vec1, VectorRegister4Float Vec2)
 {
 	return _mm_cmple_ps(Vec1, Vec2);
 }
 
-FORCEINLINE VectorRegister4Double VectorCompareLE(const VectorRegister4Double& Vec1, const VectorRegister4Double& Vec2)
+FORCEINLINE VectorRegister4Double VectorCompareLE(VectorRegister4Double Vec1, VectorRegister4Double Vec2)
 {
 	VectorRegister4Double Result;
 #if !UE_PLATFORM_MATH_USE_AVX
@@ -1798,19 +1802,19 @@ FORCEINLINE VectorRegister4Double VectorCompareLE(const VectorRegister4Double& V
  *
  */
 
-FORCEINLINE VectorRegister4Float VectorSelect(const VectorRegister4Float& Mask, const VectorRegister4Float& Vec1, const VectorRegister4Float& Vec2 )
+FORCEINLINE VectorRegister4Float VectorSelect(VectorRegister4Float Mask, VectorRegister4Float Vec1, VectorRegister4Float Vec2)
 {
 	// Can't (in general) use BLENDVPS despite our SSE4.1 minimum requirement since
 	// this is defined to be bitwise, not element-wise with MSB as toggle
 	return _mm_or_ps(_mm_and_ps(Mask, Vec1), _mm_andnot_ps(Mask, Vec2));
 }
 
-FORCEINLINE VectorRegister2Double VectorSelect(const VectorRegister2Double& Mask, const VectorRegister2Double& Vec1, const VectorRegister2Double& Vec2)
+FORCEINLINE VectorRegister2Double VectorSelect(VectorRegister2Double Mask, VectorRegister2Double Vec1, VectorRegister2Double Vec2)
 {
 	return _mm_or_pd(_mm_and_pd(Mask, Vec1), _mm_andnot_pd(Mask, Vec2));
 }
 
-FORCEINLINE VectorRegister4Double VectorSelect(const VectorRegister4Double& Mask, const VectorRegister4Double& Vec1, const VectorRegister4Double& Vec2)
+FORCEINLINE VectorRegister4Double VectorSelect(VectorRegister4Double Mask, VectorRegister4Double Vec1, VectorRegister4Double Vec2)
 {
 	VectorRegister4Double Result;
 #if !UE_PLATFORM_MATH_USE_AVX
@@ -1829,12 +1833,12 @@ FORCEINLINE VectorRegister4Double VectorSelect(const VectorRegister4Double& Mask
  * @param Vec2	2nd vector
  * @return		VectorRegister4Float( for each bit i: Vec1[i] | Vec2[i] )
  */
-FORCEINLINE VectorRegister4Float VectorBitwiseOr(const VectorRegister4Float& Vec1, const VectorRegister4Float& Vec2)
+FORCEINLINE VectorRegister4Float VectorBitwiseOr(VectorRegister4Float Vec1, VectorRegister4Float Vec2)
 {
 	return _mm_or_ps(Vec1, Vec2);
 }
 
-FORCEINLINE VectorRegister4Double VectorBitwiseOr(const VectorRegister4Double& Vec1, const VectorRegister4Double& Vec2)
+FORCEINLINE VectorRegister4Double VectorBitwiseOr(VectorRegister4Double Vec1, VectorRegister4Double Vec2)
 {
 	VectorRegister4Double Result;
 #if !UE_PLATFORM_MATH_USE_AVX
@@ -1853,12 +1857,12 @@ FORCEINLINE VectorRegister4Double VectorBitwiseOr(const VectorRegister4Double& V
  * @param Vec2	2nd vector
  * @return		VectorRegister4Float( for each bit i: Vec1[i] & Vec2[i] )
  */
-FORCEINLINE VectorRegister4Float VectorBitwiseAnd(const VectorRegister4Float& Vec1, const VectorRegister4Float& Vec2)
+FORCEINLINE VectorRegister4Float VectorBitwiseAnd(VectorRegister4Float Vec1, VectorRegister4Float Vec2)
 {
 	return _mm_and_ps(Vec1, Vec2);
 }
 
-FORCEINLINE VectorRegister4Double VectorBitwiseAnd(const VectorRegister4Double& Vec1, const VectorRegister4Double& Vec2)
+FORCEINLINE VectorRegister4Double VectorBitwiseAnd(VectorRegister4Double Vec1, VectorRegister4Double Vec2)
 {
 	VectorRegister4Double Result;
 #if !UE_PLATFORM_MATH_USE_AVX
@@ -1877,12 +1881,12 @@ FORCEINLINE VectorRegister4Double VectorBitwiseAnd(const VectorRegister4Double&
  * @param Vec2	2nd vector
  * @return		VectorRegister4Float( for each bit i: Vec1[i] ^ Vec2[i] )
  */
-FORCEINLINE VectorRegister4Float VectorBitwiseXor(const VectorRegister4Float& Vec1, const VectorRegister4Float& Vec2)
+FORCEINLINE VectorRegister4Float VectorBitwiseXor(VectorRegister4Float Vec1, VectorRegister4Float Vec2)
 {
 	return _mm_xor_ps(Vec1, Vec2);
 }
 
-FORCEINLINE VectorRegister4Double VectorBitwiseXor(const VectorRegister4Double& Vec1, const VectorRegister4Double& Vec2)
+FORCEINLINE VectorRegister4Double VectorBitwiseXor(VectorRegister4Double Vec1, VectorRegister4Double Vec2)
 {
 	VectorRegister4Double Result;
 #if !UE_PLATFORM_MATH_USE_AVX
@@ -1901,7 +1905,7 @@ FORCEINLINE VectorRegister4Double VectorBitwiseXor(const VectorRegister4Double&
  * @param Vec2	2nd vector
  * @return		cross(Vec1.xyz, Vec2.xyz). W of the input should be 0, and will remain 0.
  */
-FORCEINLINE VectorRegister4Float VectorCross( const VectorRegister4Float& Vec1, const VectorRegister4Float& Vec2 )
+FORCEINLINE VectorRegister4Float VectorCross(VectorRegister4Float Vec1, VectorRegister4Float Vec2)
 {
 	// YZX
 	VectorRegister4Float A = VectorSwizzle(Vec2, 1, 2, 0, 3);
@@ -1914,7 +1918,7 @@ FORCEINLINE VectorRegister4Float VectorCross( const VectorRegister4Float& Vec1,
 	return VectorSwizzle(A, 1, 2, 0, 3);
 }
 
-FORCEINLINE VectorRegister4Double VectorCross(const VectorRegister4Double& Vec1, const VectorRegister4Double& Vec2)
+FORCEINLINE VectorRegister4Double VectorCross(VectorRegister4Double Vec1, VectorRegister4Double Vec2)
 {
 	// YZX
 	VectorRegister4Double A = VectorSwizzle(Vec2, 1, 2, 0, 3);
@@ -1934,7 +1938,7 @@ FORCEINLINE VectorRegister4Double VectorCross(const VectorRegister4Double& Vec1,
  * @param Exponent	Exponent vector
  * @return			VectorRegister4Float( Base.x^Exponent.x, Base.y^Exponent.y, Base.z^Exponent.z, Base.w^Exponent.w )
  */
-FORCEINLINE VectorRegister4Float VectorPow( const VectorRegister4Float& Base, const VectorRegister4Float& Exponent )
+FORCEINLINE VectorRegister4Float VectorPow(VectorRegister4Float Base, VectorRegister4Float Exponent)
 {
 #if UE_PLATFORM_MATH_USE_SVML
 	return _mm_pow_ps(Base, Exponent);
@@ -1951,7 +1955,7 @@ FORCEINLINE VectorRegister4Float VectorPow( const VectorRegister4Float& Base, co
 */
 }
 
-FORCEINLINE VectorRegister4Double VectorPow(const VectorRegister4Double& Base, const VectorRegister4Double& Exponent)
+FORCEINLINE VectorRegister4Double VectorPow(VectorRegister4Double Base, VectorRegister4Double Exponent)
 {
 #if UE_PLATFORM_MATH_USE_SVML_AVX
 	return _mm256_pow_pd(Base, Exponent);
@@ -1975,12 +1979,12 @@ FORCEINLINE VectorRegister4Double VectorPow(const VectorRegister4Double& Base, c
  * @param Vector	Vector
  * @return			VectorRegister4Float(sqrt(Vec.X), sqrt(Vec.Y), sqrt(Vec.Z), sqrt(Vec.W))
  */
-FORCEINLINE VectorRegister4Float VectorSqrt(const VectorRegister4Float& Vec)
+FORCEINLINE VectorRegister4Float VectorSqrt(VectorRegister4Float Vec)
 {
 	return _mm_sqrt_ps(Vec);
 }
 
-FORCEINLINE VectorRegister4Double VectorSqrt(const VectorRegister4Double& Vec)
+FORCEINLINE VectorRegister4Double VectorSqrt(VectorRegister4Double Vec)
 {
 #if UE_PLATFORM_MATH_USE_AVX
 	return _mm256_sqrt_pd(Vec);
@@ -1995,7 +1999,7 @@ FORCEINLINE VectorRegister4Double VectorSqrt(const VectorRegister4Double& Vec)
  * @param Vector		Vector
  * @return			VectorRegister4Float(1/sqrt(t), 1/sqrt(t), 1/sqrt(t), 1/sqrt(t))
  */
-FORCEINLINE VectorRegister4Float VectorReciprocalSqrtEstimate(const VectorRegister4Float& Vec)
+FORCEINLINE VectorRegister4Float VectorReciprocalSqrtEstimate(VectorRegister4Float Vec)
 {
 	// Warning: Discrepancies between Intel and AMD hardware estimates make this diverge between platforms.
 	return _mm_rsqrt_ps(Vec);
@@ -2007,7 +2011,7 @@ FORCEINLINE VectorRegister4Float VectorReciprocalSqrtEstimate(const VectorRegist
  * @param Vector		Vector 
  * @return			VectorRegister4Float(1/sqrt(Vec.X), 1/sqrt(Vec.Y), 1/sqrt(Vec.Z), 1/sqrt(Vec.W))
  */
-FORCEINLINE VectorRegister4Float VectorReciprocalSqrt(const VectorRegister4Float& Vec)
+FORCEINLINE VectorRegister4Float VectorReciprocalSqrt(VectorRegister4Float Vec)
 {
 #if UE_PLATFORM_MATH_USE_SVML && 0 // NOTE: DISABLED
 	// TODO: this appears to deliver slightly different results on Intel vs AMD hardware,
@@ -2054,7 +2058,7 @@ FORCEINLINE VectorRegister4Float VectorReciprocalSqrt(const VectorRegister4Float
 	*/
 }
 
-FORCEINLINE VectorRegister4Double VectorReciprocalSqrt(const VectorRegister4Double& Vec)
+FORCEINLINE VectorRegister4Double VectorReciprocalSqrt(VectorRegister4Double Vec)
 {
 #if UE_PLATFORM_MATH_USE_AVX
 	return VectorDivide(GlobalVectorConstants::DoubleOne, _mm256_sqrt_pd(Vec));
@@ -2063,7 +2067,7 @@ FORCEINLINE VectorRegister4Double VectorReciprocalSqrt(const VectorRegister4Doub
 #endif
 }
 
-FORCEINLINE VectorRegister4Double VectorReciprocalSqrtEstimate(const VectorRegister4Double& Vec)
+FORCEINLINE VectorRegister4Double VectorReciprocalSqrtEstimate(VectorRegister4Double Vec)
 {
 #if UE_PLATFORM_MATH_USE_SVML_AVX
 	return _mm256_invsqrt_pd(Vec);
@@ -2081,12 +2085,12 @@ FORCEINLINE VectorRegister4Double VectorReciprocalSqrtEstimate(const VectorRegis
  * @param Vector	Vector
  * @return			VectorRegister4Float(rlen, rlen, rlen, rlen) when rlen = 1/sqrt(dot4(V))
  */
-FORCEINLINE VectorRegister4Float VectorReciprocalLen(const VectorRegister4Float& Vector)
+FORCEINLINE VectorRegister4Float VectorReciprocalLen(VectorRegister4Float Vector)
 {
 	return VectorReciprocalSqrt(VectorDot4(Vector, Vector));
 }
 
-FORCEINLINE VectorRegister4Double VectorReciprocalLen(const VectorRegister4Double& Vector)
+FORCEINLINE VectorRegister4Double VectorReciprocalLen(VectorRegister4Double Vector)
 {
 	return VectorReciprocalSqrt(VectorDot4(Vector, Vector));
 }
@@ -2097,12 +2101,12 @@ FORCEINLINE VectorRegister4Double VectorReciprocalLen(const VectorRegister4Doubl
  * @param Vector	Vector
  * @return			VectorRegister4Float(rlen, rlen, rlen, rlen) when rlen = 1/sqrt(dot4(V))
  */
-FORCEINLINE VectorRegister4Float VectorReciprocalLenEstimate(const VectorRegister4Float& Vector)
+FORCEINLINE VectorRegister4Float VectorReciprocalLenEstimate(VectorRegister4Float Vector)
 {
 	return VectorReciprocalSqrtEstimate(VectorDot4(Vector, Vector));
 }
 
-FORCEINLINE VectorRegister4Double VectorReciprocalLenEstimate(const VectorRegister4Double& Vector)
+FORCEINLINE VectorRegister4Double VectorReciprocalLenEstimate(VectorRegister4Double Vector)
 {
 	return VectorReciprocalSqrtEstimate(VectorDot4(Vector, Vector));
 }
@@ -2114,7 +2118,7 @@ FORCEINLINE VectorRegister4Double VectorReciprocalLenEstimate(const VectorRegist
  * @return		VectorRegister4Float( (Estimate) 1.0f / Vec.x, (Estimate) 1.0f / Vec.y, (Estimate) 1.0f / Vec.z, (Estimate) 1.0f / Vec.w )
  */
 
-FORCEINLINE VectorRegister4Float VectorReciprocalEstimate(const VectorRegister4Float& Vec)
+FORCEINLINE VectorRegister4Float VectorReciprocalEstimate(VectorRegister4Float Vec)
 {
 	// Warning: Discrepancies between Intel and AMD hardware estimates make this diverge between platforms.
 	return _mm_rcp_ps(Vec);
@@ -2126,7 +2130,7 @@ FORCEINLINE VectorRegister4Float VectorReciprocalEstimate(const VectorRegister4F
  * @param Vec	1st vector
  * @return		VectorRegister4Float( 1.0f / Vec.x, 1.0f / Vec.y, 1.0f / Vec.z, 1.0f / Vec.w )
  */
-FORCEINLINE VectorRegister4Float VectorReciprocal(const VectorRegister4Float& Vec)
+FORCEINLINE VectorRegister4Float VectorReciprocal(VectorRegister4Float Vec)
 {
 	return VectorDivide(GlobalVectorConstants::FloatOne, Vec);
 	/*
@@ -2162,12 +2166,12 @@ FORCEINLINE VectorRegister4Float VectorReciprocal(const VectorRegister4Float& Ve
 	*/
 }
 
-FORCEINLINE VectorRegister4Double VectorReciprocal(const VectorRegister4Double& Vec)
+FORCEINLINE VectorRegister4Double VectorReciprocal(VectorRegister4Double Vec)
 {
 	return VectorDivide(GlobalVectorConstants::DoubleOne, Vec);
 }
 
-FORCEINLINE VectorRegister4Double VectorReciprocalEstimate(const VectorRegister4Double& Vec)
+FORCEINLINE VectorRegister4Double VectorReciprocalEstimate(VectorRegister4Double Vec)
 {
 	// Not an estimate.
 	return VectorReciprocal(Vec);
@@ -2180,12 +2184,12 @@ FORCEINLINE VectorRegister4Double VectorReciprocalEstimate(const VectorRegister4
  * @param VecW		Source register for ___W (note: the fourth component is used, not the first)
  * @return			VectorRegister4Float(VecXYZ.x, VecXYZ.y, VecXYZ.z, VecW.w)
  */
-FORCEINLINE VectorRegister4Float VectorMergeVecXYZ_VecW(const VectorRegister4Float& VecXYZ, const VectorRegister4Float& VecW)
+FORCEINLINE VectorRegister4Float VectorMergeVecXYZ_VecW(VectorRegister4Float VecXYZ, VectorRegister4Float VecW)
 {
 	return _mm_blend_ps(VecXYZ, VecW, 0b1000);
 }
 
-FORCEINLINE VectorRegister4Double VectorMergeVecXYZ_VecW(const VectorRegister4Double& VecXYZ, const VectorRegister4Double& VecW)
+FORCEINLINE VectorRegister4Double VectorMergeVecXYZ_VecW(VectorRegister4Double VecXYZ, VectorRegister4Double VecW)
 {
 	return VectorRegister4Double(VecXYZ.XY, _mm_move_sd(VecW.ZW, VecXYZ.ZW));
 }
@@ -2196,12 +2200,12 @@ FORCEINLINE VectorRegister4Double VectorMergeVecXYZ_VecW(const VectorRegister4Do
 * @param Vector	VectorRegister4Float
 * @return		VectorRegister4Float(X, Y, Z, 0.0f)
 */
-FORCEINLINE VectorRegister4Float VectorSet_W0(const VectorRegister4Float& Vec)
+FORCEINLINE VectorRegister4Float VectorSet_W0(VectorRegister4Float Vec)
 {
 	return _mm_insert_ps(Vec, Vec, 0x08); // Copies lane 0 to lane 0 and zero-masks lane 3
 }
 
-FORCEINLINE VectorRegister4Double VectorSet_W0(const VectorRegister4Double& Vec)
+FORCEINLINE VectorRegister4Double VectorSet_W0(VectorRegister4Double Vec)
 {
 	VectorRegister4Double Result;
 #if !UE_PLATFORM_MATH_USE_AVX
@@ -2219,12 +2223,12 @@ FORCEINLINE VectorRegister4Double VectorSet_W0(const VectorRegister4Double& Vec)
 * @param Vector	VectorRegister4Float
 * @return		VectorRegister4Float(X, Y, Z, 1.0f)
 */
-FORCEINLINE VectorRegister4Float VectorSet_W1( const VectorRegister4Float& Vec)
+FORCEINLINE VectorRegister4Float VectorSet_W1(VectorRegister4Float Vec)
 {
 	return VectorMergeVecXYZ_VecW(Vec, VectorOneFloat());
 }
 
-FORCEINLINE VectorRegister4Double VectorSet_W1(const VectorRegister4Double& Vec)
+FORCEINLINE VectorRegister4Double VectorSet_W1(VectorRegister4Double Vec)
 {
 	VectorRegister4Double Result;
 #if !UE_PLATFORM_MATH_USE_AVX
@@ -2270,7 +2274,7 @@ FORCEINLINE bool VectorMatrixInverse(FMatrix44f* DstMatrix, const FMatrix44f* Sr
  * @param MatrixM		FMatrix pointer to the Matrix to apply transform
  * @return VectorRegister4Float = VecP*MatrixM
  */
-FORCEINLINE VectorRegister4Float VectorTransformVector(const VectorRegister4Float& VecP, const FMatrix44f* MatrixM)
+FORCEINLINE VectorRegister4Float VectorTransformVector(VectorRegister4Float VecP, const FMatrix44f* MatrixM)
 {
 	const VectorRegister4Float *M = (const VectorRegister4Float *) MatrixM;
 	VectorRegister4Float VTempX, VTempY, VTempZ, VTempW;
@@ -2289,7 +2293,7 @@ FORCEINLINE VectorRegister4Float VectorTransformVector(const VectorRegister4Floa
 	return VTempX;
 }
 
-FORCEINLINE VectorRegister4Float VectorTransformVector(const VectorRegister4Float& VecP, const FMatrix44d* MatrixM)
+FORCEINLINE VectorRegister4Float VectorTransformVector(VectorRegister4Float VecP, const FMatrix44d* MatrixM)
 {
 	// Warning: FMatrix44d alignment may not match VectorRegister4Double, so you can't just cast to VectorRegister4Double*.
 	typedef double Double4x4[4][4];
@@ -2319,7 +2323,7 @@ FORCEINLINE VectorRegister4Float VectorTransformVector(const VectorRegister4Floa
 	return MakeVectorRegisterFloatFromDouble(VTempX);
 }
 
-FORCEINLINE VectorRegister4Double VectorTransformVector(const VectorRegister4Double& VecP, const FMatrix44d* MatrixM)
+FORCEINLINE VectorRegister4Double VectorTransformVector(VectorRegister4Double VecP, const FMatrix44d* MatrixM)
 {
 	// Warning: FMatrix44d alignment may not match VectorRegister4Double, so you can't just cast to VectorRegister4Double*.
 	typedef double Double4x4[4][4];
@@ -2347,7 +2351,7 @@ FORCEINLINE VectorRegister4Double VectorTransformVector(const VectorRegister4Dou
 	return VTempX;
 }
 
-FORCEINLINE VectorRegister4Double VectorTransformVector(const VectorRegister4Double& VecP, const FMatrix44f* MatrixM)
+FORCEINLINE VectorRegister4Double VectorTransformVector(VectorRegister4Double VecP, const FMatrix44f* MatrixM)
 {
 	const VectorRegister4Float* M = (const VectorRegister4Float*) MatrixM;
 	VectorRegister4Double VTempX, VTempY, VTempZ, VTempW;
@@ -2373,12 +2377,12 @@ FORCEINLINE VectorRegister4Double VectorTransformVector(const VectorRegister4Dou
  * @param Vec2	2nd vector
  * @return		VectorRegister4Float( min(Vec1.x,Vec2.x), min(Vec1.y,Vec2.y), min(Vec1.z,Vec2.z), min(Vec1.w,Vec2.w) )
  */
-FORCEINLINE VectorRegister4Float VectorMin(const VectorRegister4Float& Vec1, const VectorRegister4Float& Vec2)
+FORCEINLINE VectorRegister4Float VectorMin(VectorRegister4Float Vec1, VectorRegister4Float Vec2)
 {
 	return _mm_min_ps(Vec1, Vec2);
 }
 
-FORCEINLINE VectorRegister4Double VectorMin(const VectorRegister4Double& Vec1, const VectorRegister4Double& Vec2)
+FORCEINLINE VectorRegister4Double VectorMin(VectorRegister4Double Vec1, VectorRegister4Double Vec2)
 {
 	VectorRegister4Double Result;
 #if !UE_PLATFORM_MATH_USE_AVX
@@ -2397,12 +2401,12 @@ FORCEINLINE VectorRegister4Double VectorMin(const VectorRegister4Double& Vec1, c
  * @param Vec2	2nd vector
  * @return		VectorRegister4Float( max(Vec1.x,Vec2.x), max(Vec1.y,Vec2.y), max(Vec1.z,Vec2.z), max(Vec1.w,Vec2.w) )
  */
-FORCEINLINE VectorRegister4Float VectorMax(const VectorRegister4Float& Vec1, const VectorRegister4Float& Vec2)
+FORCEINLINE VectorRegister4Float VectorMax(VectorRegister4Float Vec1, VectorRegister4Float Vec2)
 {
 	return _mm_max_ps(Vec1, Vec2);
 }
 
-FORCEINLINE VectorRegister4Double VectorMax(const VectorRegister4Double& Vec1, const VectorRegister4Double& Vec2)
+FORCEINLINE VectorRegister4Double VectorMax(VectorRegister4Double Vec1, VectorRegister4Double Vec2)
 {
 	VectorRegister4Double Result;
 #if !UE_PLATFORM_MATH_USE_AVX
@@ -2421,12 +2425,12 @@ FORCEINLINE VectorRegister4Double VectorMax(const VectorRegister4Double& Vec1, c
 * @param Vec2		Source vector2
 * @return			The combined vector
 */
-FORCEINLINE VectorRegister4Float VectorCombineHigh(const VectorRegister4Float& Vec1, const VectorRegister4Float& Vec2)
+FORCEINLINE VectorRegister4Float VectorCombineHigh(VectorRegister4Float Vec1, VectorRegister4Float Vec2)
 {
 	return VectorShuffle(Vec1, Vec2, 2, 3, 2, 3);
 }
 
-FORCEINLINE VectorRegister4Double VectorCombineHigh(const VectorRegister4Double& Vec1, const VectorRegister4Double& Vec2)
+FORCEINLINE VectorRegister4Double VectorCombineHigh(VectorRegister4Double Vec1, VectorRegister4Double Vec2)
 {
 	return VectorShuffle(Vec1, Vec2, 2, 3, 2, 3);
 }
@@ -2438,12 +2442,12 @@ FORCEINLINE VectorRegister4Double VectorCombineHigh(const VectorRegister4Double&
 * @param Vec2		Source vector2
 * @return			The combined vector
 */
-FORCEINLINE VectorRegister4Float VectorCombineLow(const VectorRegister4Float& Vec1, const VectorRegister4Float& Vec2)
+FORCEINLINE VectorRegister4Float VectorCombineLow(VectorRegister4Float Vec1, VectorRegister4Float Vec2)
 {
 	return VectorShuffle(Vec1, Vec2, 0, 1, 0, 1);
 }
 
-FORCEINLINE VectorRegister4Double VectorCombineLow(const VectorRegister4Double& Vec1, const VectorRegister4Double& Vec2)
+FORCEINLINE VectorRegister4Double VectorCombineLow(VectorRegister4Double Vec1, VectorRegister4Double Vec2)
 {
 	return VectorShuffle(Vec1, Vec2, 0, 1, 0, 1);
 }
@@ -2457,13 +2461,13 @@ FORCEINLINE VectorRegister4Double VectorCombineLow(const VectorRegister4Double&
  * @param OutEvens [Even0, Even1, Even2, Even3]
  * @param OutOdds [Odd0, Odd1, Odd2, Odd3]
 */
-FORCEINLINE void VectorDeinterleave(VectorRegister4Float& RESTRICT OutEvens, VectorRegister4Float& RESTRICT OutOdds, const VectorRegister4Float& RESTRICT Lo, const VectorRegister4Float& RESTRICT Hi)
+FORCEINLINE void VectorDeinterleave(VectorRegister4Float& RESTRICT OutEvens, VectorRegister4Float& RESTRICT OutOdds, VectorRegister4Float Lo, VectorRegister4Float Hi)
 {
 	OutEvens = VectorShuffle(Lo, Hi, 0, 2, 0, 2);
 	OutOdds = VectorShuffle(Lo, Hi, 1, 3, 1, 3);
 }
 
-FORCEINLINE void VectorDeinterleave(VectorRegister4Double& RESTRICT OutEvens, VectorRegister4Double& RESTRICT OutOdds, const VectorRegister4Double& RESTRICT Lo, const VectorRegister4Double& RESTRICT Hi)
+FORCEINLINE void VectorDeinterleave(VectorRegister4Double& RESTRICT OutEvens, VectorRegister4Double& RESTRICT OutOdds, VectorRegister4Double Lo, VectorRegister4Double Hi)
 {
 	OutEvens = VectorShuffle(Lo, Hi, 0, 2, 0, 2);
 	OutOdds = VectorShuffle(Lo, Hi, 1, 3, 1, 3);
@@ -2476,12 +2480,12 @@ FORCEINLINE void VectorDeinterleave(VectorRegister4Double& RESTRICT OutEvens, Ve
  * @param VecMask		Vector
  * @return				Bit 0 = sign(VecMask.x), Bit 1 = sign(VecMask.y), Bit 2 = sign(VecMask.z), Bit 3 = sign(VecMask.w)
  */
-FORCEINLINE int VectorMaskBits(const VectorRegister4Float& VecMask)
+FORCEINLINE int VectorMaskBits(VectorRegister4Float VecMask)
 {
 	return _mm_movemask_ps(VecMask);
 }
 
-FORCEINLINE int VectorMaskBits(const VectorRegister4Double& VecMask)
+FORCEINLINE int VectorMaskBits(VectorRegister4Double VecMask)
 {
 #if !UE_PLATFORM_MATH_USE_AVX
 	const int MaskXY = _mm_movemask_pd(VecMask.XY);
@@ -2500,7 +2504,7 @@ FORCEINLINE int VectorMaskBits(const VectorRegister4Double& VecMask)
  */
 // Looks complex but is really quite straightforward:
 // Load as 32-bit value, convert into 4x 32-bit ints, then convert to 4x floats
-#define VectorLoadByte4( Ptr )			_mm_cvtepi32_ps(_mm_cvtepu8_epi32(_mm_cvtsi32_si128(*(int32*)Ptr)))
+#define VectorLoadByte4(Ptr)			_mm_cvtepi32_ps(_mm_cvtepu8_epi32(_mm_cvtsi32_si128(*(int32*)Ptr)))
 
 /**
 * Loads 4 signed BYTEs from unaligned memory and converts them into 4 FLOATs.
@@ -2522,7 +2526,7 @@ FORCEINLINE VectorRegister4Float VectorLoadSignedByte4(const void* Ptr)
  * @param Ptr			Unaligned memory pointer to the 4 BYTEs.
  * @return				VectorRegister4Float( float(Ptr[3]), float(Ptr[2]), float(Ptr[1]), float(Ptr[0]) )
  */
-FORCEINLINE VectorRegister4Float VectorLoadByte4Reverse( void* Ptr )
+FORCEINLINE VectorRegister4Float VectorLoadByte4Reverse(void* Ptr)
 {
 	VectorRegister4Float Temp = VectorLoadByte4(Ptr);
 	return VectorSwizzle( Temp, 3, 2, 1, 0 );
@@ -2534,7 +2538,7 @@ FORCEINLINE VectorRegister4Float VectorLoadByte4Reverse( void* Ptr )
  * @param Vec			Vector containing 4 FLOATs
  * @param Ptr			Unaligned memory pointer to store the 4 BYTEs.
  */
-FORCEINLINE void VectorStoreByte4( const VectorRegister4Float& Vec, void* Ptr )
+FORCEINLINE void VectorStoreByte4(VectorRegister4Float Vec, void* Ptr)
 {
 	// Looks complex but is really quite straightforward:
 	// Convert 4x floats to 4x 32-bit ints, then pack into 4x 16-bit ints, then into 4x 8-bit unsigned ints, then store as a 32-bit value
@@ -2550,7 +2554,7 @@ FORCEINLINE void VectorStoreByte4( const VectorRegister4Float& Vec, void* Ptr )
 * @param Vec			Vector containing 4 FLOATs
 * @param Ptr			Unaligned memory pointer to store the 4 BYTEs.
 */
-FORCEINLINE void VectorStoreSignedByte4(const VectorRegister4Float& Vec, void* Ptr)
+FORCEINLINE void VectorStoreSignedByte4(VectorRegister4Float Vec, void* Ptr)
 {
 	// Looks complex but is really quite straightforward:
 	// Convert 4x floats to 4x 32-bit ints, then pack into 4x 16-bit ints, then into 4x 8-bit unsigned ints, then store as a 32-bit value
@@ -2586,7 +2590,7 @@ FORCEINLINE VectorRegister4Float VectorLoadURGB10A2N(void* Ptr)
 * @param Vec			Vector containing 4 FLOATs
 * @param Ptr			Unaligned memory pointer to store the packed RGBA16(8 bytes).
 */
-FORCEINLINE void VectorStoreURGB10A2N(const VectorRegister4Float& Vec, void* Ptr)
+FORCEINLINE void VectorStoreURGB10A2N(VectorRegister4Float Vec, void* Ptr)
 {
 	VectorRegister4Float Tmp;
 	Tmp = _mm_max_ps(Vec, MakeVectorRegisterFloat(0.0f, 0.0f, 0.0f, 0.0f));
@@ -2614,7 +2618,7 @@ FORCEINLINE void VectorStoreURGB10A2N(const VectorRegister4Float& Vec, void* Ptr
  * @param Ptr			Unaligned memory pointer to the RGBA16(8 bytes).
  * @return				VectorRegister4Float with 4 FLOATs loaded from Ptr.
  */
-#define VectorLoadURGBA16N( Ptr ) _mm_cvtepi32_ps(_mm_cvtepu16_epi32(_mm_loadl_epi64((const __m128i*)Ptr)))
+#define VectorLoadURGBA16N(Ptr) _mm_cvtepi32_ps(_mm_cvtepu16_epi32(_mm_loadl_epi64((const __m128i*)Ptr)))
 
 /**
  * Loads packed signed RGBA16(8 bytes) from unaligned memory and converts them into 4 FLOATs.
@@ -2634,7 +2638,7 @@ FORCEINLINE VectorRegister4Float VectorLoadSRGBA16N(const void* Ptr)
 * @param Vec			Vector containing 4 FLOATs
 * @param Ptr			Unaligned memory pointer to store the packed RGBA16(8 bytes).
 */
-FORCEINLINE void VectorStoreURGBA16N(const VectorRegister4Float& Vec, void* Ptr)
+FORCEINLINE void VectorStoreURGBA16N(VectorRegister4Float Vec, void* Ptr)
 {
 	VectorRegister4Float Tmp;
 	Tmp = _mm_max_ps(Vec, MakeVectorRegisterFloat(0.0f, 0.0f, 0.0f, 0.0f));
@@ -2652,12 +2656,12 @@ FORCEINLINE void VectorStoreURGBA16N(const VectorRegister4Float& Vec, void* Ptr)
  * @param Vec2			2nd source vector
  * @return				Non-zero integer if (Vec1.x > Vec2.x) || (Vec1.y > Vec2.y) || (Vec1.z > Vec2.z) || (Vec1.w > Vec2.w)
  */
-FORCEINLINE int VectorAnyGreaterThan(const VectorRegister4Float& Vec1, const VectorRegister4Float& Vec2)
+FORCEINLINE int VectorAnyGreaterThan(VectorRegister4Float Vec1, VectorRegister4Float Vec2)
 {
 	return VectorMaskBits(VectorCompareGT(Vec1, Vec2));
 }
 
-FORCEINLINE int VectorAnyGreaterThan(const VectorRegister4Double& Vec1, const VectorRegister4Double& Vec2)
+FORCEINLINE int VectorAnyGreaterThan(VectorRegister4Double Vec1, VectorRegister4Double Vec2)
 {
 	return VectorMaskBits(VectorCompareGT(Vec1, Vec2));
 }
@@ -2711,7 +2715,7 @@ FORCEINLINE int VectorAnyGreaterThan(const VectorRegister4Double& Vec1, const Ve
 * @param Quat2	Pointer to the second quaternion
 * @return Quat1 * Quat2
 */
-FORCEINLINE VectorRegister4Float VectorQuaternionMultiply2( const VectorRegister4Float& Quat1, const VectorRegister4Float& Quat2 )
+FORCEINLINE VectorRegister4Float VectorQuaternionMultiply2(VectorRegister4Float Quat1, VectorRegister4Float Quat2)
 {
 	VectorRegister4Float Result = VectorMultiply(VectorReplicate(Quat1, 3), Quat2);
 	Result = VectorMultiplyAdd( VectorMultiply(VectorReplicate(Quat1, 0), VectorSwizzle(Quat2, 3, 2, 1, 0)), GlobalVectorConstants::QMULTI_SIGN_MASK0, Result);
@@ -2721,7 +2725,7 @@ FORCEINLINE VectorRegister4Float VectorQuaternionMultiply2( const VectorRegister
 	return Result;
 }
 
-FORCEINLINE VectorRegister4Double VectorQuaternionMultiply2(const VectorRegister4Double& Quat1, const VectorRegister4Double& Quat2)
+FORCEINLINE VectorRegister4Double VectorQuaternionMultiply2(VectorRegister4Double Quat1, VectorRegister4Double Quat2)
 {
 	VectorRegister4Double Result = VectorMultiply(VectorReplicate(Quat1, 3), Quat2);
 	Result = VectorMultiplyAdd(VectorMultiply(VectorReplicate(Quat1, 0), VectorSwizzle(Quat2, 3, 2, 1, 0)), GlobalVectorConstants::DOUBLE_QMULTI_SIGN_MASK0, Result);
@@ -2752,7 +2756,7 @@ FORCEINLINE void VectorQuaternionMultiply(VectorRegister4Double* RESTRICT Result
 }
 
 // Returns true if the vector contains a component that is either NAN or +/-infinite.
-FORCEINLINE bool VectorContainsNaNOrInfinite(const VectorRegister4Float& Vec)
+FORCEINLINE bool VectorContainsNaNOrInfinite(VectorRegister4Float Vec)
 {
 	// https://en.wikipedia.org/wiki/IEEE_754-1985
 	// Infinity is represented with all exponent bits set, with the correct sign bit.
@@ -2771,7 +2775,7 @@ FORCEINLINE bool VectorContainsNaNOrInfinite(const VectorRegister4Float& Vec)
 	return !IsFinite;
 }
 
-FORCEINLINE bool VectorContainsNaNOrInfinite(const VectorRegister4Double& Vec)
+FORCEINLINE bool VectorContainsNaNOrInfinite(VectorRegister4Double Vec)
 {
 	// https://en.wikipedia.org/wiki/IEEE_754-1985
 	// Infinity is represented with all exponent bits set, with the correct sign bit.
@@ -2790,17 +2794,17 @@ FORCEINLINE bool VectorContainsNaNOrInfinite(const VectorRegister4Double& Vec)
 	return !IsFinite;
 }
 
-FORCEINLINE VectorRegister4Float VectorTruncate(const VectorRegister4Float& Vec)
+FORCEINLINE VectorRegister4Float VectorTruncate(VectorRegister4Float Vec)
 {
 	return _mm_round_ps(Vec, _MM_FROUND_TRUNC);
 }
 
-FORCEINLINE VectorRegister2Double TruncateVectorRegister2d(const VectorRegister2Double& V)
+FORCEINLINE VectorRegister2Double TruncateVectorRegister2d(VectorRegister2Double V)
 {
 	return _mm_round_pd(V, _MM_FROUND_TRUNC);
 }
 
-FORCEINLINE VectorRegister4Double VectorTruncate(const VectorRegister4Double& V)
+FORCEINLINE VectorRegister4Double VectorTruncate(VectorRegister4Double V)
 {
 	VectorRegister4Double Result;
 #if !UE_PLATFORM_MATH_USE_AVX
@@ -2812,22 +2816,22 @@ FORCEINLINE VectorRegister4Double VectorTruncate(const VectorRegister4Double& V)
 	return Result;
 }
 
-FORCEINLINE VectorRegister4Float VectorRound(const VectorRegister4Float &Vec)
+FORCEINLINE VectorRegister4Float VectorRound(VectorRegister4Float Vec)
 {
 	return _mm_round_ps(Vec, _MM_FROUND_TO_NEAREST_INT |_MM_FROUND_NO_EXC);
 }
 
-FORCEINLINE VectorRegister4Int VectorRoundToIntHalfToEven(const VectorRegister4Float& Vec)
+FORCEINLINE VectorRegister4Int VectorRoundToIntHalfToEven(VectorRegister4Float Vec)
 {
 	return _mm_cvtps_epi32(Vec);
 }
 
-FORCEINLINE VectorRegister4Float VectorCeil(const VectorRegister4Float& V)
+FORCEINLINE VectorRegister4Float VectorCeil(VectorRegister4Float V)
 {
 	return _mm_ceil_ps(V);
 }
 
-FORCEINLINE VectorRegister4Double VectorCeil(const VectorRegister4Double& V)
+FORCEINLINE VectorRegister4Double VectorCeil(VectorRegister4Double V)
 {
 #if UE_PLATFORM_MATH_USE_AVX
 	VectorRegister4Double Result;
@@ -2841,12 +2845,12 @@ FORCEINLINE VectorRegister4Double VectorCeil(const VectorRegister4Double& V)
 #endif
 }
 
-FORCEINLINE VectorRegister4Float VectorFloor(const VectorRegister4Float& V)
+FORCEINLINE VectorRegister4Float VectorFloor(VectorRegister4Float V)
 {
 	return _mm_floor_ps(V);
 }
 
-FORCEINLINE VectorRegister4Double VectorFloor(const VectorRegister4Double& V)
+FORCEINLINE VectorRegister4Double VectorFloor(VectorRegister4Double V)
 {
 #if UE_PLATFORM_MATH_USE_AVX
 	VectorRegister4Double Result;
@@ -2860,7 +2864,7 @@ FORCEINLINE VectorRegister4Double VectorFloor(const VectorRegister4Double& V)
 #endif
 }
 
-FORCEINLINE VectorRegister4Float VectorMod(const VectorRegister4Float& X, const VectorRegister4Float& Y)
+FORCEINLINE VectorRegister4Float VectorMod(VectorRegister4Float X, VectorRegister4Float Y)
 {
 	// Check against invalid divisor
 	VectorRegister4Float InvalidDivisorMask = VectorCompareLE(VectorAbs(Y), GlobalVectorConstants::SmallNumber);
@@ -2881,7 +2885,7 @@ FORCEINLINE VectorRegister4Float VectorMod(const VectorRegister4Float& X, const
 	return Result;
 }
 
-FORCEINLINE VectorRegister4Double VectorMod(const VectorRegister4Double& X, const VectorRegister4Double& Y)
+FORCEINLINE VectorRegister4Double VectorMod(VectorRegister4Double X, VectorRegister4Double Y)
 {
 	// Check against invalid divisor
 	VectorRegister4Double InvalidDivisorMask = VectorCompareLE(VectorAbs(Y), GlobalVectorConstants::DoubleSmallNumber);
@@ -2904,31 +2908,31 @@ FORCEINLINE VectorRegister4Double VectorMod(const VectorRegister4Double& X, cons
 	return DoubleResult;
 }
 
-FORCEINLINE VectorRegister4Float VectorSign(const VectorRegister4Float& X)
+FORCEINLINE VectorRegister4Float VectorSign(VectorRegister4Float X)
 {
 	VectorRegister4Float Mask = VectorCompareGE(X, GlobalVectorConstants::FloatZero);
 	return VectorSelect(Mask, GlobalVectorConstants::FloatOne, GlobalVectorConstants::FloatMinusOne);
 }
 
-FORCEINLINE VectorRegister4Double VectorSign(const VectorRegister4Double& X)
+FORCEINLINE VectorRegister4Double VectorSign(VectorRegister4Double X)
 {
 	VectorRegister4Double Mask = VectorCompareGE(X, (GlobalVectorConstants::DoubleZero));
 	return VectorSelect(Mask, GlobalVectorConstants::DoubleOne, GlobalVectorConstants::DoubleMinusOne);
 }
 
-FORCEINLINE VectorRegister4Float VectorStep(const VectorRegister4Float& X)
+FORCEINLINE VectorRegister4Float VectorStep(VectorRegister4Float X)
 {
 	VectorRegister4Float Mask = VectorCompareGE(X, GlobalVectorConstants::FloatZero);
 	return VectorSelect(Mask, GlobalVectorConstants::FloatOne, GlobalVectorConstants::FloatZero);
 }
 
-FORCEINLINE VectorRegister4Double VectorStep(const VectorRegister4Double& X)
+FORCEINLINE VectorRegister4Double VectorStep(VectorRegister4Double X)
 {
 	VectorRegister4Double Mask = VectorCompareGE(X, GlobalVectorConstants::DoubleZero);
 	return VectorSelect(Mask, GlobalVectorConstants::DoubleOne, GlobalVectorConstants::DoubleZero);
 }
 
-FORCEINLINE VectorRegister4Float VectorExp(const VectorRegister4Float& X)
+FORCEINLINE VectorRegister4Float VectorExp(VectorRegister4Float X)
 {
 #if UE_PLATFORM_MATH_USE_SVML
 	return _mm_exp_ps(X);
@@ -2937,7 +2941,7 @@ FORCEINLINE VectorRegister4Float VectorExp(const VectorRegister4Float& X)
 #endif
 }
 
-FORCEINLINE VectorRegister4Double VectorExp(const VectorRegister4Double& X)
+FORCEINLINE VectorRegister4Double VectorExp(VectorRegister4Double X)
 {
 #if UE_PLATFORM_MATH_USE_SVML_AVX
 	return _mm256_exp_pd(X);
@@ -2953,7 +2957,7 @@ FORCEINLINE VectorRegister4Double VectorExp(const VectorRegister4Double& X)
 #endif
 }
 
-FORCEINLINE VectorRegister4Float VectorExp2(const VectorRegister4Float& X)
+FORCEINLINE VectorRegister4Float VectorExp2(VectorRegister4Float X)
 {
 #if UE_PLATFORM_MATH_USE_SVML
 	return _mm_exp2_ps(X);
@@ -2967,7 +2971,7 @@ FORCEINLINE VectorRegister4Float VectorExp2(const VectorRegister4Float& X)
 #endif
 }
 
-FORCEINLINE VectorRegister4Double VectorExp2(const VectorRegister4Double& X)
+FORCEINLINE VectorRegister4Double VectorExp2(VectorRegister4Double X)
 {
 #if UE_PLATFORM_MATH_USE_SVML_AVX
 	return _mm256_exp2_pd(X);
@@ -2983,7 +2987,7 @@ FORCEINLINE VectorRegister4Double VectorExp2(const VectorRegister4Double& X)
 #endif
 }
 
-FORCEINLINE VectorRegister4Float VectorLog(const VectorRegister4Float& X)
+FORCEINLINE VectorRegister4Float VectorLog(VectorRegister4Float X)
 {
 #if UE_PLATFORM_MATH_USE_SVML
 	return _mm_log_ps(X);
@@ -2992,7 +2996,7 @@ FORCEINLINE VectorRegister4Float VectorLog(const VectorRegister4Float& X)
 #endif
 }
 
-FORCEINLINE VectorRegister4Double VectorLog(const VectorRegister4Double& X)
+FORCEINLINE VectorRegister4Double VectorLog(VectorRegister4Double X)
 {
 #if UE_PLATFORM_MATH_USE_SVML_AVX
 	return _mm256_log_pd(X);
@@ -3008,7 +3012,7 @@ FORCEINLINE VectorRegister4Double VectorLog(const VectorRegister4Double& X)
 #endif
 }
 
-FORCEINLINE VectorRegister4Float VectorLog2(const VectorRegister4Float& X)
+FORCEINLINE VectorRegister4Float VectorLog2(VectorRegister4Float X)
 {
 #if UE_PLATFORM_MATH_USE_SVML
 	return _mm_log2_ps(X);
@@ -3022,7 +3026,7 @@ FORCEINLINE VectorRegister4Float VectorLog2(const VectorRegister4Float& X)
 #endif
 }
 
-FORCEINLINE VectorRegister4Double VectorLog2(const VectorRegister4Double& X)
+FORCEINLINE VectorRegister4Double VectorLog2(VectorRegister4Double X)
 {
 #if UE_PLATFORM_MATH_USE_SVML_AVX
 	return _mm256_log2_pd(X);
@@ -3054,7 +3058,7 @@ namespace VectorSinConstantsSSE
 	static const VectorRegister4Float B = MakeVectorRegisterFloatConstant(b, b, b, b);
 }
 
-FORCEINLINE VectorRegister4Float VectorSin(const VectorRegister4Float& V)
+FORCEINLINE VectorRegister4Float VectorSin(VectorRegister4Float V)
 {
 #if UE_PLATFORM_MATH_USE_SVML
 	return _mm_sin_ps(V);
@@ -3063,7 +3067,7 @@ FORCEINLINE VectorRegister4Float VectorSin(const VectorRegister4Float& V)
 #endif
 }
 
-FORCEINLINE VectorRegister4Double VectorSin(const VectorRegister4Double& V)
+FORCEINLINE VectorRegister4Double VectorSin(VectorRegister4Double V)
 {
 #if UE_PLATFORM_MATH_USE_SVML_AVX
 	return _mm256_sin_pd(V);
@@ -3079,7 +3083,7 @@ FORCEINLINE VectorRegister4Double VectorSin(const VectorRegister4Double& V)
 #endif
 }
 
-FORCEINLINE VectorRegister4Float VectorCos(const VectorRegister4Float& V)
+FORCEINLINE VectorRegister4Float VectorCos(VectorRegister4Float V)
 {
 #if UE_PLATFORM_MATH_USE_SVML
 	return _mm_cos_ps(V);
@@ -3088,7 +3092,7 @@ FORCEINLINE VectorRegister4Float VectorCos(const VectorRegister4Float& V)
 #endif
 }
 
-FORCEINLINE VectorRegister4Double VectorCos(const VectorRegister4Double& V)
+FORCEINLINE VectorRegister4Double VectorCos(VectorRegister4Double V)
 {
 #if UE_PLATFORM_MATH_USE_SVML_AVX
 	return _mm256_cos_pd(V);
@@ -3181,7 +3185,7 @@ FORCEINLINE void VectorSinCos(VectorRegister4Double* RESTRICT VSinAngles, Vector
 }
 
 
-FORCEINLINE VectorRegister4Float VectorTan(const VectorRegister4Float& X)
+FORCEINLINE VectorRegister4Float VectorTan(VectorRegister4Float X)
 {
 #if UE_PLATFORM_MATH_USE_SVML
 	return _mm_tan_ps(X);
@@ -3196,7 +3200,7 @@ FORCEINLINE VectorRegister4Float VectorTan(const VectorRegister4Float& X)
 #endif
 }
 
-FORCEINLINE VectorRegister4Double VectorTan(const VectorRegister4Double& X)
+FORCEINLINE VectorRegister4Double VectorTan(VectorRegister4Double X)
 {
 #if UE_PLATFORM_MATH_USE_SVML_AVX
 	return _mm256_tan_pd(X);
@@ -3212,7 +3216,7 @@ FORCEINLINE VectorRegister4Double VectorTan(const VectorRegister4Double& X)
 #endif
 }
 
-FORCEINLINE VectorRegister4Float VectorASin(const VectorRegister4Float& X)
+FORCEINLINE VectorRegister4Float VectorASin(VectorRegister4Float X)
 {
 #if UE_PLATFORM_MATH_USE_SVML
 	return _mm_asin_ps(X);
@@ -3226,7 +3230,7 @@ FORCEINLINE VectorRegister4Float VectorASin(const VectorRegister4Float& X)
 #endif
 }
 
-FORCEINLINE VectorRegister4Double VectorASin(const VectorRegister4Double& X)
+FORCEINLINE VectorRegister4Double VectorASin(VectorRegister4Double X)
 {
 #if UE_PLATFORM_MATH_USE_SVML_AVX
 	return _mm256_asin_pd(X);
@@ -3242,7 +3246,7 @@ FORCEINLINE VectorRegister4Double VectorASin(const VectorRegister4Double& X)
 #endif
 }
 
-FORCEINLINE VectorRegister4Float VectorACos(const VectorRegister4Float& X)
+FORCEINLINE VectorRegister4Float VectorACos(VectorRegister4Float X)
 {
 #if UE_PLATFORM_MATH_USE_SVML
 	return _mm_acos_ps(X);
@@ -3256,7 +3260,7 @@ FORCEINLINE VectorRegister4Float VectorACos(const VectorRegister4Float& X)
 #endif
 }
 
-FORCEINLINE VectorRegister4Double VectorACos(const VectorRegister4Double& X)
+FORCEINLINE VectorRegister4Double VectorACos(VectorRegister4Double X)
 {
 #if UE_PLATFORM_MATH_USE_SVML_AVX
 	return _mm256_acos_pd(X);
@@ -3272,7 +3276,7 @@ FORCEINLINE VectorRegister4Double VectorACos(const VectorRegister4Double& X)
 #endif
 }
 
-FORCEINLINE VectorRegister4Float VectorATan(const VectorRegister4Float& X)
+FORCEINLINE VectorRegister4Float VectorATan(VectorRegister4Float X)
 {
 #if UE_PLATFORM_MATH_USE_SVML
 	return _mm_atan_ps(X);
@@ -3287,7 +3291,7 @@ FORCEINLINE VectorRegister4Float VectorATan(const VectorRegister4Float& X)
 #endif
 }
 
-FORCEINLINE VectorRegister4Double VectorATan(const VectorRegister4Double& X)
+FORCEINLINE VectorRegister4Double VectorATan(VectorRegister4Double X)
 {
 #if UE_PLATFORM_MATH_USE_SVML_AVX
 	return _mm256_atan_pd(X);
@@ -3303,7 +3307,7 @@ FORCEINLINE VectorRegister4Double VectorATan(const VectorRegister4Double& X)
 #endif
 }
 
-FORCEINLINE VectorRegister4Float VectorATan2(const VectorRegister4Float& Y, const VectorRegister4Float& X)
+FORCEINLINE VectorRegister4Float VectorATan2(VectorRegister4Float Y, VectorRegister4Float X)
 {
 #if UE_PLATFORM_MATH_USE_SVML
 	return _mm_atan2_ps(Y, X);
@@ -3319,7 +3323,7 @@ FORCEINLINE VectorRegister4Float VectorATan2(const VectorRegister4Float& Y, cons
 #endif
 }
 
-FORCEINLINE VectorRegister4Double VectorATan2(const VectorRegister4Double& Y, const VectorRegister4Double& X)
+FORCEINLINE VectorRegister4Double VectorATan2(VectorRegister4Double Y, VectorRegister4Double X)
 {
 #if UE_PLATFORM_MATH_USE_SVML_AVX
 	return _mm256_atan2_pd(Y, X);
@@ -3361,7 +3365,7 @@ FORCEINLINE VectorRegister4Double VectorATan2(const VectorRegister4Double& Y, co
 #define VectorIntCompareLE(A, B)	VectorIntNot(VectorIntCompareGT(A,B))
 
 
-FORCEINLINE VectorRegister4Int VectorIntSelect(const VectorRegister4Int& Mask, const VectorRegister4Int& Vec1, const VectorRegister4Int& Vec2)
+FORCEINLINE VectorRegister4Int VectorIntSelect(VectorRegister4Int Mask, VectorRegister4Int Vec1, VectorRegister4Int Vec2)
 {
 	// Can't use PBLENDVB in general because this is a bitwise select, not byte-lane-wise
 	return _mm_or_si128(_mm_and_si128(Mask, Vec1), _mm_andnot_si128(Mask, Vec2));
@@ -3371,56 +3375,56 @@ FORCEINLINE VectorRegister4Int VectorIntSelect(const VectorRegister4Int& Mask, c
 #define VectorIntAdd(A, B)	_mm_add_epi32(A, B)
 #define VectorIntSubtract(A, B)	_mm_sub_epi32(A, B)
 
-FORCEINLINE VectorRegister4Int VectorIntMultiply(const VectorRegister4Int& A, const VectorRegister4Int& B)
+FORCEINLINE VectorRegister4Int VectorIntMultiply(VectorRegister4Int A, VectorRegister4Int B)
 {
 	return _mm_mullo_epi32(A, B);
 }
 
-#define VectorIntNegate(A) VectorIntSubtract( GlobalVectorConstants::IntZero, A)
+#define VectorIntNegate(A) VectorIntSubtract(GlobalVectorConstants::IntZero, A)
 
-FORCEINLINE VectorRegister4Int VectorIntMin(const VectorRegister4Int& A, const VectorRegister4Int& B)
+FORCEINLINE VectorRegister4Int VectorIntMin(VectorRegister4Int A, VectorRegister4Int B)
 {
 	return _mm_min_epi32(A, B);
 }
 
-FORCEINLINE VectorRegister4Int VectorIntMax(const VectorRegister4Int& A, const VectorRegister4Int& B)
+FORCEINLINE VectorRegister4Int VectorIntMax(VectorRegister4Int A, VectorRegister4Int B)
 {
 	return _mm_max_epi32(A, B);
 }
 
-FORCEINLINE VectorRegister4Int VectorIntAbs(const VectorRegister4Int& A)
+FORCEINLINE VectorRegister4Int VectorIntAbs(VectorRegister4Int A)
 {
 	return _mm_abs_epi32(A);
 }
 
-FORCEINLINE VectorRegister4Int VectorIntClamp(const VectorRegister4Int& Vec1, const VectorRegister4Int& Vec2, const VectorRegister4Int& Vec3) 
+FORCEINLINE VectorRegister4Int VectorIntClamp(VectorRegister4Int Vec1, VectorRegister4Int Vec2, VectorRegister4Int Vec3) 
 {
 	return VectorIntMin(VectorIntMax(Vec1, Vec2), Vec3);
 }
 
-#define VectorIntSign(A) VectorIntSelect( VectorIntCompareGE(A, GlobalVectorConstants::IntZero), GlobalVectorConstants::IntOne, GlobalVectorConstants::IntMinusOne )
+#define VectorIntSign(A) VectorIntSelect(VectorIntCompareGE(A, GlobalVectorConstants::IntZero), GlobalVectorConstants::IntOne, GlobalVectorConstants::IntMinusOne)
 
 #define VectorIntToFloat(A) _mm_cvtepi32_ps(A)
 
-FORCEINLINE VectorRegister4Int VectorFloatToInt(const VectorRegister4Float& A)
+FORCEINLINE VectorRegister4Int VectorFloatToInt(VectorRegister4Float A)
 {
 	return _mm_cvttps_epi32(A);
 }
 
 // TODO: LWC: potential loss of data
-FORCEINLINE VectorRegister4Int VectorFloatToInt(const VectorRegister4Double& A)
+FORCEINLINE VectorRegister4Int VectorFloatToInt(VectorRegister4Double A)
 {
 	return VectorFloatToInt(MakeVectorRegisterFloatFromDouble(A));
 }
 
-FORCEINLINE VectorRegister4Int VectorDoubleToInt(const VectorRegister4Double& Vec)
+FORCEINLINE VectorRegister4Int VectorDoubleToInt(VectorRegister4Double Vec)
 {
 	VectorRegister4Int A = _mm_cvttpd_epi32(Vec.GetXY());
 	VectorRegister4Int B = _mm_cvttpd_epi32(Vec.GetZW());
 	return _mm_unpacklo_epi64(A, B);
 }
 
-FORCEINLINE VectorRegister4Int VectorShuffleByte4(const VectorRegister4Int& Vec, const VectorRegister4Int& Mask)
+FORCEINLINE VectorRegister4Int VectorShuffleByte4(VectorRegister4Int Vec, VectorRegister4Int Mask)
 {
 	return _mm_shuffle_epi8(Vec, Mask);
 }
diff --git a/Engine/Source/Runtime/Core/Public/Math/UnrealMathVectorCommon.h b/Engine/Source/Runtime/Core/Public/Math/UnrealMathVectorCommon.h
deleted file mode 100644
diff --git a/Engine/Source/Runtime/Core/Public/Math/UnrealMathVectorCommon.h.inl b/Engine/Source/Runtime/Core/Public/Math/UnrealMathVectorCommon.h.inl
new file mode 100644
--- /dev/null
+++ b/Engine/Source/Runtime/Core/Public/Math/UnrealMathVectorCommon.h.inl
@@ -0,0 +1,940 @@
+// Copyright Epic Games, Inc. All Rights Reserved.
+
+#pragma once
+
+// HEADER_UNIT_SKIP - Not included directly
+
+struct FGenericPlatformMath;
+struct FMath;
+
+// Overload to resolve compiler ambiguity for things like MakeVectorRegister(V.X, V.Y, V.Z, 0.f) when V is a double type.
+FORCEINLINE VectorRegister4Double MakeVectorRegister(double X, double Y, double Z, float W)
+{
+	return MakeVectorRegisterDouble(X, Y, Z, (double)W);
+}
+
+FORCEINLINE VectorRegister4Float VectorZero(void)
+{
+	return VectorZeroFloat();
+}
+
+FORCEINLINE VectorRegister4Float VectorOne(void)
+{
+	return VectorOneFloat();
+}
+
+/**
+ * Overloads for code passing &MyVector, &MyQuat, etc.
+ */
+
+///////////////////////////////////////////////////////////////////////////////////////////////////////////////////
+/// VectorLoad
+///////////////////////////////////////////////////////////////////////////////////////////////////////////////////
+
+template<typename T>
+FORCEINLINE TVectorRegisterType<T> VectorLoad(const UE::Math::TQuat<T>* Ptr)
+{
+	return VectorLoad((const T*)(Ptr));
+}
+
+template<typename T>
+FORCEINLINE TVectorRegisterType<T> VectorLoad(const UE::Math::TVector4<T>* Ptr)
+{
+	return VectorLoad((const T*)(Ptr));
+}
+
+FORCEINLINE VectorRegister4Float VectorLoad(const VectorRegister4Float* Ptr)
+{
+	return VectorLoad((const float*)(Ptr));
+}
+
+FORCEINLINE VectorRegister4Double VectorLoad(const VectorRegister4Double* Ptr)
+{
+	return VectorLoad((const double*)(Ptr));
+}
+
+
+///////////////////////////////////////////////////////////////////////////////////////////////////////////////////
+/// VectorLoadAligned
+///////////////////////////////////////////////////////////////////////////////////////////////////////////////////
+
+template<typename T>
+FORCEINLINE TVectorRegisterType<T> VectorLoadAligned(const UE::Math::TQuat<T>* Ptr)
+{
+	return VectorLoadAligned((const T*)(Ptr));
+}
+
+template<typename T>
+FORCEINLINE TVectorRegisterType<T> VectorLoadAligned(const UE::Math::TVector4<T>* Ptr)
+{
+	return VectorLoadAligned((const T*)(Ptr));
+}
+
+template<typename T>
+FORCEINLINE TVectorRegisterType<T> VectorLoadAligned(const UE::Math::TPlane<T>* Ptr)
+{
+	return VectorLoadAligned((const T*)(Ptr));
+}
+
+FORCEINLINE VectorRegister4Float VectorLoadAligned(const VectorRegister4Float* Ptr)
+{
+	//return VectorLoadAligned((const float*)(Ptr));
+	return *Ptr;
+}
+
+FORCEINLINE VectorRegister4Double VectorLoadAligned(const VectorRegister4Double* Ptr)
+{
+	//return VectorLoadAligned((const double*)(Ptr));
+	return *Ptr;
+}
+
+
+///////////////////////////////////////////////////////////////////////////////////////////////////////////////////
+/// VectorLoadFloat3
+///////////////////////////////////////////////////////////////////////////////////////////////////////////////////
+
+/**
+* Loads 3 floats from unaligned memory and sets W=0.
+*
+* @param Ptr	Unaligned memory pointer to the 3 floats
+* @return		VectorRegister4Float(Ptr[0], Ptr[1], Ptr[2], 0.0f)
+*/
+FORCEINLINE VectorRegister4Float VectorLoadFloat3(const float* Ptr)
+{
+	return MakeVectorRegister(Ptr[0], Ptr[1], Ptr[2], 0.0f);
+}
+
+template<typename T>
+FORCEINLINE TVectorRegisterType<T> VectorLoadFloat3(const UE::Math::TVector<T>* Ptr)
+{
+	return VectorLoadFloat3((const T*)(Ptr));
+}
+
+FORCEINLINE VectorRegister4Double VectorLoadDouble3(const double* Ptr)
+{
+	return VectorLoadFloat3((const double*)Ptr);
+}
+
+
+///////////////////////////////////////////////////////////////////////////////////////////////////////////////////
+/// VectorLoadFloat3_W0
+///////////////////////////////////////////////////////////////////////////////////////////////////////////////////
+
+/**
+ * Loads 3 floats from unaligned memory and sets W=0.
+ *
+ * @param Ptr	Unaligned memory pointer to the 3 floats
+ * @return		VectorRegister4Float(Ptr[0], Ptr[1], Ptr[2], 0.0f)
+ */
+FORCEINLINE VectorRegister4Float VectorLoadFloat3_W0(const float* Ptr)
+{
+	return VectorLoadFloat3(Ptr);
+}
+
+FORCEINLINE VectorRegister4Double VectorLoadFloat3_W0(const double* Ptr)
+{
+	return VectorLoadFloat3(Ptr);
+}
+
+template<typename T>
+FORCEINLINE TVectorRegisterType<T> VectorLoadFloat3_W0(const UE::Math::TVector<T>* Ptr)
+{
+	return VectorLoadFloat3_W0((const T*)(Ptr));
+}
+
+template<typename T>
+FORCEINLINE TVectorRegisterType<T> VectorLoadFloat3_W0(const UE::Math::TRotator<T>* Ptr)
+{
+	return VectorLoadFloat3_W0((const T*)(Ptr));
+}
+
+FORCEINLINE VectorRegister4Double VectorLoadDouble3_W0(const double* Ptr)
+{
+	return VectorLoadDouble3(Ptr);
+}
+
+
+///////////////////////////////////////////////////////////////////////////////////////////////////////////////////
+/// VectorLoadFloat3_W1
+///////////////////////////////////////////////////////////////////////////////////////////////////////////////////
+
+/**
+ * Loads 3 floats from unaligned memory and sets W=1.
+ *
+ * @param Ptr	Unaligned memory pointer to the 3 floats
+ * @return		VectorRegister4Float(Ptr[0], Ptr[1], Ptr[2], 1.0f)
+ */
+FORCEINLINE VectorRegister4Float VectorLoadFloat3_W1(const float* Ptr)
+{
+	return MakeVectorRegister(Ptr[0], Ptr[1], Ptr[2], 1.0f);
+}
+
+template<typename T>
+FORCEINLINE TVectorRegisterType<T> VectorLoadFloat3_W1(const UE::Math::TVector<T>* Ptr)
+{
+	return VectorLoadFloat3_W1((const T*)Ptr);
+}
+
+FORCEINLINE VectorRegister4Double VectorLoadDouble3_W1(const double* Ptr)
+{
+	return VectorLoadFloat3_W1(Ptr);
+}
+
+
+///////////////////////////////////////////////////////////////////////////////////////////////////////////////////
+/// VectorLoadFloat1
+///////////////////////////////////////////////////////////////////////////////////////////////////////////////////
+
+FORCEINLINE VectorRegister4Float VectorLoadFloat1(const VectorRegister4Float* Ptr)
+{
+	return VectorReplicate(*Ptr, 0);
+}
+
+FORCEINLINE VectorRegister4Double VectorLoadFloat1(const VectorRegister4Double* Ptr)
+{
+	return VectorReplicate(*Ptr, 0);
+}
+
+FORCEINLINE VectorRegister4Double VectorLoadFloat1(const double* Ptr)
+{
+	return VectorLoadDouble1(Ptr);
+}
+
+
+///////////////////////////////////////////////////////////////////////////////////////////////////////////////////
+/// VectorStoreAligned
+///////////////////////////////////////////////////////////////////////////////////////////////////////////////////
+
+// TODO: this is lossy, would rather not implement.
+#if SUPPORT_DOUBLE_TO_FLOAT_VECTOR_CONVERSION
+FORCEINLINE void VectorStoreAligned(VectorRegister4Double Vec, float* Dst)
+{
+	VectorStoreAligned(MakeVectorRegisterFloatFromDouble(Vec), Dst);
+}
+#endif
+
+template<typename T>
+FORCEINLINE void VectorStoreAligned(TVectorRegisterType<T> Vec, UE::Math::TVector4<T>* Dst)
+{
+	VectorStoreAligned(Vec, (T*)Dst);
+}
+
+// Specific overload to support promoting float->double vector and storing in TVector4<double>
+FORCEINLINE void VectorStoreAligned(VectorRegister4Float Vec, struct UE::Math::TVector4<double>* Dst)
+{
+	VectorRegister4Double DoubleVec(Vec);
+	VectorStoreAligned(DoubleVec, Dst);
+}
+
+template<typename T>
+FORCEINLINE void VectorStoreAligned(TVectorRegisterType<T> Vec, struct UE::Math::TQuat<T>* Dst)
+{
+	VectorStoreAligned(Vec, (T*)Dst);
+}
+
+// Specific overload to support promoting float->double vector and storing in TQuat<double>
+FORCEINLINE void VectorStoreAligned(VectorRegister4Float Vec, struct UE::Math::TQuat<double>* Dst)
+{
+	VectorRegister4Double DoubleVec(Vec);
+	VectorStoreAligned(DoubleVec, Dst);
+}
+
+FORCEINLINE void VectorStoreAligned(VectorRegister4Float Vec, VectorRegister4Float* Dst)
+{
+	*Dst = Vec;
+}
+
+FORCEINLINE void VectorStoreAligned(VectorRegister4Double Vec, VectorRegister4Double* Dst)
+{
+	*Dst = Vec;
+}
+
+
+///////////////////////////////////////////////////////////////////////////////////////////////////////////////////
+/// VectorStore
+///////////////////////////////////////////////////////////////////////////////////////////////////////////////////
+
+#if SUPPORT_DOUBLE_TO_FLOAT_VECTOR_CONVERSION
+FORCEINLINE void VectorStore(VectorRegister4Double Vec, float* Dst)
+{
+	VectorRegister4Float FloatVec = MakeVectorRegisterFloatFromDouble(Vec);
+	VectorStore(FloatVec, Dst);
+}
+#endif
+
+template<typename T>
+FORCEINLINE void VectorStore(VectorRegister4Float Vec, UE::Math::TVector4<T>* Dst)
+{
+	VectorStoreAligned(Vec, Dst);
+}
+
+template<typename T>
+FORCEINLINE void VectorStore(VectorRegister4Double Vec, UE::Math::TVector4<T>* Dst)
+{
+	VectorStoreAligned(Vec, Dst);
+}
+
+template<typename T>
+FORCEINLINE void VectorStore(TVectorRegisterType<T> Vec, struct UE::Math::TQuat<T>* Dst)
+{
+	VectorStoreAligned(Vec, Dst);
+}
+
+// Specific overload to support promoting float->double vector and storing in TQuat<double>
+FORCEINLINE void VectorStore(VectorRegister4Float Vec, struct UE::Math::TQuat<double>* Dst)
+{
+	VectorRegister4Double DoubleVec(Vec);
+	VectorStoreAligned(DoubleVec, Dst);
+}
+
+FORCEINLINE void VectorStore(VectorRegister4Float Vec, VectorRegister4Float* Dst)
+{
+	*Dst = Vec;
+}
+
+FORCEINLINE void VectorStore(VectorRegister4Double Vec, VectorRegister4Double* Dst)
+{
+	*Dst = Vec;
+}
+
+
+///////////////////////////////////////////////////////////////////////////////////////////////////////////////////
+/// VectorStoreFloat3
+///////////////////////////////////////////////////////////////////////////////////////////////////////////////////
+
+#if SUPPORT_DOUBLE_TO_FLOAT_VECTOR_CONVERSION
+FORCEINLINE void VectorStoreFloat3(VectorRegister4Double Vec, float* Dst)
+{
+	VectorRegister4Float FloatVec = MakeVectorRegisterFloatFromDouble(Vec);
+	VectorStoreFloat3(FloatVec, Dst);
+}
+#endif
+
+template<typename T>
+FORCEINLINE void VectorStoreFloat3(VectorRegister4Float Vec, UE::Math::TVector<T>* Dst)
+{
+	VectorStoreFloat3(Vec, (T*)Dst);
+}
+
+template<typename T>
+FORCEINLINE void VectorStoreFloat3(VectorRegister4Double Vec, UE::Math::TVector<T>* Dst)
+{
+	VectorStoreFloat3(Vec, (T*)Dst);
+}
+
+template<typename T>
+FORCEINLINE void VectorStoreFloat3(VectorRegister4Float Vec, UE::Math::TRotator<T>* Dst)
+{
+	VectorStoreFloat3(Vec, (T*)Dst);
+}
+
+template<typename T>
+FORCEINLINE void VectorStoreFloat3(VectorRegister4Double Vec, UE::Math::TRotator<T>* Dst)
+{
+	VectorStoreFloat3(Vec, (T*)Dst);
+}
+
+///////////////////////////////////////////////////////////////////////////////////////////////////////////////////
+/// VectorStoreFloat1
+///////////////////////////////////////////////////////////////////////////////////////////////////////////////////
+
+#if SUPPORT_DOUBLE_TO_FLOAT_VECTOR_CONVERSION
+FORCEINLINE void VectorStoreFloat1(VectorRegister4Double Vec, float* Dst)
+{
+	VectorRegister4Float FloatVec = MakeVectorRegisterFloatFromDouble(Vec);
+	VectorStoreFloat1(FloatVec, Dst);
+}
+#endif
+
+
+FORCEINLINE void VectorStoreFloat1(VectorRegister4Float Vec, int32* Dst)
+{
+	VectorStoreFloat1(Vec, (float*)Dst);
+}
+
+FORCEINLINE void VectorStoreFloat1(VectorRegister4Double Vec, int32* Dst)
+{
+	// This is used for bit masks, so don't do an actual convert, mask the lower bits instead
+	double X = VectorGetComponent(Vec, 0);
+	uint64 XMask64 = *((uint64*)(&X));
+	uint32 XMask32 = (uint32)(XMask64 & 0xFFFFFFFF);
+	VectorRegister4Float FloatVec = VectorLoadFloat1((float*)(&XMask32));
+	VectorStoreFloat1(FloatVec, Dst);
+}
+
+FORCEINLINE void VectorStoreFloat1(VectorRegister4Double Vec, int64* Dst)
+{
+	VectorStoreFloat1(Vec, (double*)Dst);
+}
+
+///////////////////////////////////////////////////////////////////////////////////////////////////////////////////
+// Wrappers for clarifying "accurate" versus "estimate" for some functions,
+// important for old functions which were ambiguous in their precision.
+
+// Returns accurate reciprocal square root.
+FORCEINLINE VectorRegister4Float VectorReciprocalSqrtAccurate(VectorRegister4Float Vec)
+{
+	return VectorReciprocalSqrt(Vec);
+}
+
+FORCEINLINE VectorRegister4Double VectorReciprocalSqrtAccurate(VectorRegister4Double Vec)
+{
+	return VectorReciprocalSqrt(Vec);
+}
+
+/**
+ * Computes the reciprocal of a vector (component-wise) and returns the result.
+ *
+ * @param Vec	1st vector
+ * @return		VectorRegister( 1.0f / Vec.x, 1.0f / Vec.y, 1.0f / Vec.z, 1.0f / Vec.w )
+ */
+FORCEINLINE VectorRegister4Float VectorReciprocalAccurate(VectorRegister4Float Vec)
+{
+	return VectorReciprocal(Vec);
+}
+
+FORCEINLINE VectorRegister4Double VectorReciprocalAccurate(VectorRegister4Double Vec)
+{
+	return VectorReciprocal(Vec);
+}
+
+///////////////////////////////////////////////////////////////////////////////////////////////////////////////////
+
+// Returns a normalized 4 vector = Vector / |Vector|.
+// There is no handling of zero length vectors, use VectorNormalizeSafe if this is a possible input.
+template<typename TVectorRegisterType>
+FORCEINLINE TVectorRegisterType VectorNormalize(TVectorRegisterType Vector)
+{
+	return VectorMultiply(Vector, VectorReciprocalLen(Vector));
+}
+
+template<typename TVectorRegisterType>
+FORCEINLINE TVectorRegisterType VectorNormalizeAccurate(TVectorRegisterType Vector)
+{
+	return VectorNormalize(Vector);
+}
+
+template<typename TVectorRegisterType>
+FORCEINLINE TVectorRegisterType VectorNormalizeEstimate(TVectorRegisterType Vector)
+{
+	return VectorMultiply(Vector, VectorReciprocalLenEstimate(Vector));
+}
+
+// Returns ((Vector dot Vector) >= 1e-8) ? (Vector / |Vector|) : DefaultValue
+// Uses accurate 1/sqrt, not the estimate
+FORCEINLINE VectorRegister4Float VectorNormalizeSafe(VectorRegister4Float Vector, VectorRegister4Float DefaultValue)
+{
+	const VectorRegister4Float SquareSum = VectorDot4(Vector, Vector);
+	const VectorRegister4Float NonZeroMask = VectorCompareGE(SquareSum, GlobalVectorConstants::SmallLengthThreshold);
+	const VectorRegister4Float InvLength = VectorReciprocalSqrtAccurate(SquareSum);
+	const VectorRegister4Float NormalizedVector = VectorMultiply(InvLength, Vector);
+	return VectorSelect(NonZeroMask, NormalizedVector, DefaultValue);
+}
+
+FORCEINLINE VectorRegister4Double VectorNormalizeSafe(VectorRegister4Double Vector, VectorRegister4Double DefaultValue)
+{
+	const VectorRegister4Double SquareSum = VectorDot4(Vector, Vector);
+	const VectorRegister4Double NonZeroMask = VectorCompareGE(SquareSum, GlobalVectorConstants::DoubleSmallLengthThreshold);
+	const VectorRegister4Double InvLength = VectorReciprocalSqrtAccurate(SquareSum);
+	const VectorRegister4Double NormalizedVector = VectorMultiply(InvLength, Vector);
+	return VectorSelect(NonZeroMask, NormalizedVector, DefaultValue);
+}
+
+/**
+ * Returns non-zero if any element in Vec1 is lesser than the corresponding element in Vec2, otherwise 0.
+ *
+ * @param Vec1			1st source vector
+ * @param Vec2			2nd source vector
+ * @return				Non-zero integer if (Vec1.x < Vec2.x) || (Vec1.y < Vec2.y) || (Vec1.z < Vec2.z) || (Vec1.w < Vec2.w)
+ */
+FORCEINLINE uint32 VectorAnyLesserThan(VectorRegister4Float Vec1, VectorRegister4Float Vec2)
+{
+	return VectorAnyGreaterThan( Vec2, Vec1 );
+}
+
+FORCEINLINE uint32 VectorAnyLesserThan(VectorRegister4Double Vec1, VectorRegister4Double Vec2)
+{
+	return VectorAnyGreaterThan(Vec2, Vec1);
+}
+
+/**
+ * Returns non-zero if all elements in Vec1 are greater than the corresponding elements in Vec2, otherwise 0.
+ *
+ * @param Vec1			1st source vector
+ * @param Vec2			2nd source vector
+ * @return				Non-zero integer if (Vec1.x > Vec2.x) && (Vec1.y > Vec2.y) && (Vec1.z > Vec2.z) && (Vec1.w > Vec2.w)
+ */
+FORCEINLINE uint32 VectorAllGreaterThan(VectorRegister4Float Vec1, VectorRegister4Float Vec2)
+{
+	return !VectorAnyGreaterThan( Vec2, Vec1 );
+}
+
+FORCEINLINE uint32 VectorAllGreaterThan(VectorRegister4Double Vec1, VectorRegister4Double Vec2)
+{
+	return !VectorAnyGreaterThan(Vec2, Vec1);
+}
+
+/**
+ * Returns non-zero if all elements in Vec1 are lesser than the corresponding elements in Vec2, otherwise 0.
+ *
+ * @param Vec1			1st source vector
+ * @param Vec2			2nd source vector
+ * @return				Non-zero integer if (Vec1.x < Vec2.x) && (Vec1.y < Vec2.y) && (Vec1.z < Vec2.z) && (Vec1.w < Vec2.w)
+ */
+FORCEINLINE uint32 VectorAllLesserThan(VectorRegister4Float Vec1, VectorRegister4Float Vec2)
+{
+	return !VectorAnyGreaterThan( Vec1, Vec2 );
+}
+
+FORCEINLINE uint32 VectorAllLesserThan(VectorRegister4Double Vec1, VectorRegister4Double Vec2)
+{
+	return !VectorAnyGreaterThan(Vec1, Vec2);
+}
+
+/** Clamps X to be between VecMin and VecMax, inclusive. */
+FORCEINLINE VectorRegister4Float VectorClamp(VectorRegister4Float X, VectorRegister4Float VecMin, VectorRegister4Float VecMax)
+{
+	return VectorMin(VectorMax(X, VecMin), VecMax);
+}
+
+FORCEINLINE VectorRegister4Double VectorClamp(VectorRegister4Double X, VectorRegister4Double VecMin, VectorRegister4Double VecMax)
+{
+	return VectorMin(VectorMax(X, VecMin), VecMax);
+}
+
+/*----------------------------------------------------------------------------
+	VectorRegister specialization of templates.
+----------------------------------------------------------------------------*/
+
+/** Returns the smaller of the two values (operates on each component individually) */
+template<> FORCEINLINE VectorRegister4Float FGenericPlatformMath::Min(VectorRegister4Float A, VectorRegister4Float B)
+{
+	return VectorMin(A, B);
+}
+
+template<> FORCEINLINE VectorRegister4Double FGenericPlatformMath::Min(VectorRegister4Double A, VectorRegister4Double B)
+{
+	return VectorMin(A, B);
+}
+
+/** Returns the larger of the two values (operates on each component individually) */
+template<> FORCEINLINE VectorRegister4Float FGenericPlatformMath::Max(VectorRegister4Float A, VectorRegister4Float B)
+{
+	return VectorMax(A, B);
+}
+
+template<> FORCEINLINE VectorRegister4Double FGenericPlatformMath::Max(VectorRegister4Double A, VectorRegister4Double B)
+{
+	return VectorMax(A, B);
+}
+
+/** Lerp between two vectors */
+FORCEINLINE VectorRegister4Float VectorLerp(VectorRegister4Float A, VectorRegister4Float B, VectorRegister4Float Alpha)
+{
+	VectorRegister4Float SubVec = VectorSubtract(GlobalVectorConstants::FloatOne, Alpha);
+	return VectorMultiplyAdd(B, Alpha, VectorMultiply(A, SubVec));
+}
+
+FORCEINLINE VectorRegister4Double VectorLerp(VectorRegister4Double A, VectorRegister4Double B, VectorRegister4Double Alpha)
+{
+	VectorRegister4Double SubVec = VectorSubtract(GlobalVectorConstants::DoubleOne, Alpha);
+	return VectorMultiplyAdd(B, Alpha, VectorMultiply(A, SubVec));
+}
+
+// TCustomLerp for FMath::Lerp<VectorRegister4Float>()
+template<>
+struct TCustomLerp<VectorRegister4Float>
+{
+	constexpr static bool Value = true;
+	
+	// Specialization of Lerp function that works with vector registers
+	static FORCEINLINE VectorRegister4Float Lerp(VectorRegister4Float A, VectorRegister4Float B, VectorRegister4Float Alpha)
+	{
+		return VectorLerp(A, B, Alpha);
+	}
+};
+
+// TCustomLerp for FMath::Lerp<VectorRegister4Double>()
+template<>
+struct TCustomLerp<VectorRegister4Double>
+{
+	constexpr static bool Value = true;
+
+	// Specialization of Lerp function that works with vector registers
+	static FORCEINLINE VectorRegister4Double Lerp(VectorRegister4Double A, VectorRegister4Double B, VectorRegister4Double Alpha)
+	{
+		return VectorLerp(A, B, Alpha);
+	}
+};
+
+// A and B are quaternions.  The result is A + (|A.B| >= 0 ? 1 : -1) * B
+FORCEINLINE VectorRegister4Float VectorAccumulateQuaternionShortestPath(VectorRegister4Float A, VectorRegister4Float B)
+{
+	// Blend rotation
+	//     To ensure the 'shortest route', we make sure the dot product between the both rotations is positive.
+	//     const float Bias = (|A.B| >= 0 ? 1 : -1)
+	//     return A + B * Bias;
+	const VectorRegister4Float Zero = VectorZeroFloat();
+	const VectorRegister4Float RotationDot = VectorDot4(A, B);
+	const VectorRegister4Float QuatRotationDirMask = VectorCompareGE(RotationDot, Zero);
+	const VectorRegister4Float NegativeB = VectorSubtract(Zero, B);
+	const VectorRegister4Float BiasTimesB = VectorSelect(QuatRotationDirMask, B, NegativeB);
+	return VectorAdd(A, BiasTimesB);
+}
+
+FORCEINLINE VectorRegister4Double VectorAccumulateQuaternionShortestPath(VectorRegister4Double A, VectorRegister4Double B)
+{
+	// Blend rotation
+	//     To ensure the 'shortest route', we make sure the dot product between the both rotations is positive.
+	//     const float Bias = (|A.B| >= 0 ? 1 : -1)
+	//     return A + B * Bias;
+	const VectorRegister4Double Zero = VectorZeroDouble();
+	const VectorRegister4Double RotationDot = VectorDot4(A, B);
+	const VectorRegister4Double QuatRotationDirMask = VectorCompareGE(RotationDot, Zero);
+	const VectorRegister4Double NegativeB = VectorSubtract(Zero, B);
+	const VectorRegister4Double BiasTimesB = VectorSelect(QuatRotationDirMask, B, NegativeB);
+	return VectorAdd(A, BiasTimesB);
+}
+
+// Normalize quaternion ( result = (Q.Q >= 1e-8) ? (Q / |Q|) : (0,0,0,1) )
+FORCEINLINE VectorRegister4Float VectorNormalizeQuaternion(VectorRegister4Float UnnormalizedQuat)
+{
+	return VectorNormalizeSafe(UnnormalizedQuat, GlobalVectorConstants::Float0001);
+}
+
+FORCEINLINE VectorRegister4Double VectorNormalizeQuaternion(VectorRegister4Double UnnormalizedQuat)
+{
+	return VectorNormalizeSafe(UnnormalizedQuat, GlobalVectorConstants::Double0001);
+}
+
+// VectorMod360: Essentially VectorMod(X, 360) but using faster computation that is still accurate given the known input constraints.
+FORCEINLINE VectorRegister4Float VectorMod360(VectorRegister4Float X)
+{
+	// R = X - (Y * (Trunc(X / Y))
+	VectorRegister4Float Temp = VectorTruncate(VectorDivide(X, GlobalVectorConstants::Float360));
+	VectorRegister4Float FloatResult = VectorNegateMultiplyAdd(GlobalVectorConstants::Float360, Temp, X);
+	return FloatResult;
+}
+
+FORCEINLINE VectorRegister4Double VectorMod360(VectorRegister4Double X)
+{
+	// R = X - (Y * (Trunc(X / Y))
+	VectorRegister4Double Temp = VectorTruncate(VectorDivide(X, GlobalVectorConstants::Double360));
+	VectorRegister4Double DoubleResult = VectorNegateMultiplyAdd(GlobalVectorConstants::Double360, Temp, X);
+	return DoubleResult;
+}
+
+// Normalize Rotator
+FORCEINLINE VectorRegister4Float VectorNormalizeRotator(VectorRegister4Float UnnormalizedRotator)
+{
+	// shift in the range [-360,360]
+	VectorRegister4Float V0	= VectorMod360(UnnormalizedRotator);
+	VectorRegister4Float V1	= VectorAdd( V0, GlobalVectorConstants::Float360 );
+	VectorRegister4Float V2	= VectorSelect(VectorCompareGE(V0, VectorZeroFloat()), V0, V1);
+
+	// shift to [-180,180]
+	VectorRegister4Float V3	= VectorSubtract( V2, GlobalVectorConstants::Float360 );
+	VectorRegister4Float V4	= VectorSelect(VectorCompareGT(V2, GlobalVectorConstants::Float180), V3, V2);
+
+	return  V4;
+}
+
+FORCEINLINE VectorRegister4Double VectorNormalizeRotator(VectorRegister4Double UnnormalizedRotator)
+{
+	// shift in the range [-360,360]
+	VectorRegister4Double V0 = VectorMod360(UnnormalizedRotator);
+	VectorRegister4Double V1 = VectorAdd(V0, GlobalVectorConstants::Double360);
+	VectorRegister4Double V2 = VectorSelect(VectorCompareGE(V0, VectorZeroDouble()), V0, V1);
+
+	// shift to [-180,180]
+	VectorRegister4Double V3 = VectorSubtract(V2, GlobalVectorConstants::Double360);
+	VectorRegister4Double V4 = VectorSelect(VectorCompareGT(V2, GlobalVectorConstants::Double180), V3, V2);
+
+	return  V4;
+}
+
+/** 
+ * Fast Linear Quaternion Interpolation for quaternions stored in VectorRegisters.
+ * Result is NOT normalized.
+ */
+FORCEINLINE VectorRegister4Float VectorLerpQuat(VectorRegister4Float A, VectorRegister4Float B, VectorRegister4Float Alpha)
+{
+	// Blend rotation
+	//     To ensure the 'shortest route', we make sure the dot product between the both rotations is positive.
+	//     const float Bias = (|A.B| >= 0 ? 1 : -1)
+	//     Rotation = (B * Alpha) + (A * (Bias * (1.f - Alpha)));
+	const VectorRegister4Float Zero = VectorZeroFloat();
+
+	const VectorRegister4Float OneMinusAlpha = VectorSubtract(VectorOneFloat(), Alpha);
+
+	const VectorRegister4Float RotationDot = VectorDot4(A, B);
+	const VectorRegister4Float QuatRotationDirMask = VectorCompareGE(RotationDot, Zero);
+	const VectorRegister4Float NegativeA = VectorSubtract(Zero, A);
+	const VectorRegister4Float BiasTimesA = VectorSelect(QuatRotationDirMask, A, NegativeA);
+	const VectorRegister4Float BTimesWeight = VectorMultiply(B, Alpha);
+	const VectorRegister4Float UnnormalizedResult = VectorMultiplyAdd(BiasTimesA, OneMinusAlpha, BTimesWeight);
+
+	return UnnormalizedResult;
+}
+
+FORCEINLINE VectorRegister4Double VectorLerpQuat(VectorRegister4Double A, VectorRegister4Double B, VectorRegister4Double Alpha)
+{
+	// Blend rotation
+	//     To ensure the 'shortest route', we make sure the dot product between the both rotations is positive.
+	//     const float Bias = (|A.B| >= 0 ? 1 : -1)
+	//     Rotation = (B * Alpha) + (A * (Bias * (1.f - Alpha)));
+	const VectorRegister4Double Zero = VectorZeroDouble();
+
+	const VectorRegister4Double OneMinusAlpha = VectorSubtract(VectorOneDouble(), Alpha);
+
+	const VectorRegister4Double RotationDot = VectorDot4(A, B);
+	const VectorRegister4Double QuatRotationDirMask = VectorCompareGE(RotationDot, Zero);
+	const VectorRegister4Double NegativeA = VectorSubtract(Zero, A);
+	const VectorRegister4Double BiasTimesA = VectorSelect(QuatRotationDirMask, A, NegativeA);
+	const VectorRegister4Double BTimesWeight = VectorMultiply(B, Alpha);
+	const VectorRegister4Double UnnormalizedResult = VectorMultiplyAdd(BiasTimesA, OneMinusAlpha, BTimesWeight);
+
+	return UnnormalizedResult;
+}
+
+/** 
+ * Bi-Linear Quaternion interpolation for quaternions stored in VectorRegisters.
+ * Result is NOT normalized.
+ */
+FORCEINLINE VectorRegister4Float VectorBiLerpQuat(VectorRegister4Float P00, VectorRegister4Float P10, VectorRegister4Float P01, VectorRegister4Float P11, VectorRegister4Float FracX, VectorRegister4Float FracY)
+{
+	return VectorLerpQuat(
+		VectorLerpQuat(P00, P10, FracX),
+		VectorLerpQuat(P01, P11, FracX),
+		FracY);
+}
+
+FORCEINLINE VectorRegister4Double VectorBiLerpQuat(VectorRegister4Double P00, VectorRegister4Double P10, VectorRegister4Double P01, VectorRegister4Double P11, VectorRegister4Double FracX, VectorRegister4Double FracY)
+{
+	return VectorLerpQuat(
+		VectorLerpQuat(P00, P10, FracX),
+		VectorLerpQuat(P01, P11, FracX),
+		FracY);
+}
+
+/** 
+ * Inverse quaternion ( -X, -Y, -Z, W) 
+ */
+FORCEINLINE VectorRegister4Float VectorQuaternionInverse(VectorRegister4Float NormalizedQuat)
+{
+	return VectorMultiply(GlobalVectorConstants::QINV_SIGN_MASK, NormalizedQuat);
+}
+
+FORCEINLINE VectorRegister4Double VectorQuaternionInverse(VectorRegister4Double NormalizedQuat)
+{
+	return VectorMultiply(GlobalVectorConstants::DOUBLE_QINV_SIGN_MASK, NormalizedQuat);
+}
+
+
+/**
+ * Rotate a vector using a unit Quaternion.
+ *
+ * @param Quat Unit Quaternion to use for rotation.
+ * @param VectorW0 Vector to rotate. W component must be zero.
+ * @return Vector after rotation by Quat.
+ */
+FORCEINLINE VectorRegister4Float VectorQuaternionRotateVector(VectorRegister4Float Quat, VectorRegister4Float VectorW0)
+{
+	// Q * V * Q.Inverse
+	//const VectorRegister InverseRotation = VectorQuaternionInverse(Quat);
+	//const VectorRegister Temp = VectorQuaternionMultiply2(Quat, VectorW0);
+	//const VectorRegister Rotated = VectorQuaternionMultiply2(Temp, InverseRotation);
+
+	// Equivalence of above can be shown to be:
+	// http://people.csail.mit.edu/bkph/articles/Quaternions.pdf
+	// V' = V + 2w(Q x V) + (2Q x (Q x V))
+	// refactor:
+	// V' = V + w(2(Q x V)) + (Q x (2(Q x V)))
+	// T = 2(Q x V);
+	// V' = V + w*(T) + (Q x T)
+
+	const VectorRegister4Float QW = VectorReplicate(Quat, 3);
+	VectorRegister4Float T = VectorCross(Quat, VectorW0);
+	T = VectorAdd(T, T);
+	const VectorRegister4Float VTemp0 = VectorMultiplyAdd(QW, T, VectorW0);
+	const VectorRegister4Float VTemp1 = VectorCross(Quat, T);
+	const VectorRegister4Float Rotated = VectorAdd(VTemp0, VTemp1);
+	return Rotated;
+}
+
+FORCEINLINE VectorRegister4Double VectorQuaternionRotateVector(VectorRegister4Double Quat, VectorRegister4Double VectorW0)
+{
+	const VectorRegister4Double QW = VectorReplicate(Quat, 3);
+	VectorRegister4Double T = VectorCross(Quat, VectorW0);
+	T = VectorAdd(T, T);
+	const VectorRegister4Double VTemp0 = VectorMultiplyAdd(QW, T, VectorW0);
+	const VectorRegister4Double VTemp1 = VectorCross(Quat, T);
+	const VectorRegister4Double Rotated = VectorAdd(VTemp0, VTemp1);
+	return Rotated;
+}
+
+/**
+ * Rotate a vector using the inverse of a unit Quaternion (rotation in the opposite direction).
+ *
+ * @param Quat Unit Quaternion to use for rotation.
+ * @param VectorW0 Vector to rotate. W component must be zero.
+ * @return Vector after rotation by the inverse of Quat.
+ */
+FORCEINLINE VectorRegister4Float VectorQuaternionInverseRotateVector(VectorRegister4Float Quat, VectorRegister4Float VectorW0)
+{
+	// Q.Inverse * V * Q
+	//const VectorRegister InverseRotation = VectorQuaternionInverse(Quat);
+	//const VectorRegister Temp = VectorQuaternionMultiply2(InverseRotation, VectorW0);
+	//const VectorRegister Rotated = VectorQuaternionMultiply2(Temp, Quat);
+
+	const VectorRegister4Float QInv = VectorQuaternionInverse(Quat);
+	return VectorQuaternionRotateVector(QInv, VectorW0);
+}
+
+FORCEINLINE VectorRegister4Double VectorQuaternionInverseRotateVector(VectorRegister4Double Quat, VectorRegister4Double VectorW0)
+{
+	const VectorRegister4Double QInv = VectorQuaternionInverse(Quat);
+	return VectorQuaternionRotateVector(QInv, VectorW0);
+}
+
+FORCEINLINE VectorRegister4Double VectorMultiply(VectorRegister4Double Vec1, VectorRegister4Float Vec2)
+{
+	return VectorMultiply(Vec1, VectorRegister4Double(Vec2));
+}
+
+FORCEINLINE VectorRegister4Double VectorMultiplyAdd(VectorRegister4Double Vec1, VectorRegister4Float Vec2, VectorRegister4Double Acc)
+{
+	return VectorMultiplyAdd(Vec1, VectorRegister4Double(Vec2), Acc);
+}
+
+FORCEINLINE VectorRegister4Double VectorMultiplyAdd(VectorRegister4Double Vec1, VectorRegister4Float Vec2, VectorRegister4Float Acc)
+{
+	return VectorMultiplyAdd(Vec1, VectorRegister4Double(Vec2), VectorRegister4Double(Acc));
+}
+
+FORCEINLINE VectorRegister4Double VectorSetDouble(double X, double Y, double Z, double W)
+{
+	return MakeVectorRegisterDouble(X, Y, Z, W);
+}
+
+FORCEINLINE VectorRegister4Double VectorSetDouble1(double D)
+{
+	return VectorSetFloat1(D);
+}
+
+FORCEINLINE void VectorStoreDouble3(VectorRegister4Double Vec, double* Ptr)
+{
+	VectorStoreFloat3(Vec, Ptr);
+}
+
+FORCEINLINE void VectorStoreDouble1(VectorRegister4Double Vec, double* Ptr)
+{
+	VectorStoreFloat1(Vec, Ptr);
+}
+
+/**
+ * Creates a vector out of three FLOATs and leaves W undefined.
+ *
+ * @param X		1st float component
+ * @param Y		2nd float component
+ * @param Z		3rd float component
+ * @return		VectorRegister4Float(X, Y, Z, undefined)
+ */
+FORCEINLINE VectorRegister4Float VectorSetFloat3(float X, float Y, float Z)
+{
+	return MakeVectorRegisterFloat(X, Y, Z, 0.0f);
+}
+
+FORCEINLINE VectorRegister4Double VectorSetFloat3(double X, double Y, double Z)
+{
+	return MakeVectorRegisterDouble(X, Y, Z, 0.0);
+}
+
+FORCEINLINE VectorRegister4Double VectorSetDouble3(double X, double Y, double Z)
+{
+	return VectorSetFloat3(X, Y, Z);
+}
+
+/**
+ * Creates a vector out of four FLOATs.
+ *
+ * @param X		1st float component
+ * @param Y		2nd float component
+ * @param Z		3rd float component
+ * @param W		4th float component
+ * @return		VectorRegister4Float(X, Y, Z, W)
+ */
+FORCEINLINE VectorRegister4Float VectorSet(float X, float Y, float Z, float W)
+{
+	return MakeVectorRegisterFloat(X, Y, Z, W);
+}
+
+FORCEINLINE VectorRegister4Double VectorSet(double X, double Y, double Z, double W)
+{
+	return MakeVectorRegisterDouble(X, Y, Z, W);
+}
+
+// Overload to resolve ambiguous mixes of ints and floats.
+FORCEINLINE VectorRegister4Float VectorSet(uint32 X, uint32 Y, uint32 Z, uint32 W)
+{
+	return VectorSet(float(X), float(Y), float(Z), float(W));
+}
+
+FORCEINLINE VectorRegister4Float VectorFractional(VectorRegister4Float Vec)
+{
+	return VectorSubtract(Vec, VectorTruncate(Vec));
+}
+
+FORCEINLINE VectorRegister4Double VectorFractional(VectorRegister4Double Vec)
+{
+	return VectorSubtract(Vec, VectorTruncate(Vec));
+}
+
+
+FORCEINLINE void VectorQuaternionMultiply(FQuat4f* RESTRICT Result, const FQuat4f* RESTRICT Quat1, const FQuat4f* RESTRICT Quat2)
+{
+	const VectorRegister4Float Q1 = VectorLoadAligned(Quat1);
+	const VectorRegister4Float Q2 = VectorLoadAligned(Quat2);
+	const VectorRegister4Float QResult = VectorQuaternionMultiply2(Q1, Q2);
+	VectorStoreAligned(QResult, Result);
+}
+
+FORCEINLINE void VectorQuaternionMultiply(FQuat4d* RESTRICT Result, const FQuat4d* RESTRICT Quat1, const FQuat4d* RESTRICT Quat2)
+{
+	// Warning: don't try to cast FQuat4d to VectorRegister4Double, they may not be of the same alignment.
+	// We can still use VectorLoadAligned/VectorStoreAligned, because those only enforce 16-byte alignment and fall back to unaligned otherwise.
+	const VectorRegister4Double Q1 = VectorLoadAligned(Quat1);
+	const VectorRegister4Double Q2 = VectorLoadAligned(Quat2);
+	const VectorRegister4Double QResult = VectorQuaternionMultiply2(Q1, Q2);
+	VectorStoreAligned(QResult, Result);
+}
+
+/**
+ * Counts the number of trailing zeros in the bit representation of the value,
+ * counting from least-significant bit to most.
+ *
+ * @param Value the value to determine the number of leading zeros for
+ * @return the number of zeros before the first "on" bit
+ */
+#if defined(_MSC_VER)
+#pragma intrinsic( _BitScanForward )
+FORCEINLINE uint32 appCountTrailingZeros(uint32 Value)
+{
+	if (Value == 0)
+	{
+		return 32;
+	}
+	unsigned long BitIndex;	// 0-based, where the LSB is 0 and MSB is 31
+	_BitScanForward(&BitIndex, Value);	// Scans from LSB to MSB
+	return BitIndex;
+}
+#else // !defined(_MSC_VER)
+FORCEINLINE uint32 appCountTrailingZeros(uint32 Value)
+{
+	if (Value == 0)
+	{
+		return 32;
+	}
+	return __builtin_ffs(Value) - 1;
+}
+#endif // _MSC_VER
diff --git a/Engine/Source/Runtime/Core/Public/Math/UnrealMathVectorConstants.h b/Engine/Source/Runtime/Core/Public/Math/UnrealMathVectorConstants.h
deleted file mode 100644
diff --git a/Engine/Source/Runtime/Core/Public/Math/UnrealMathVectorConstants.h.inl b/Engine/Source/Runtime/Core/Public/Math/UnrealMathVectorConstants.h.inl
new file mode 100644
--- /dev/null
+++ b/Engine/Source/Runtime/Core/Public/Math/UnrealMathVectorConstants.h.inl
@@ -0,0 +1,158 @@
+// Copyright Epic Games, Inc. All Rights Reserved.
+
+#pragma once
+
+// HEADER_UNIT_SKIP - Not included directly
+// IWYU pragma: private
+
+#include "Math/UnrealMathUtility.h"
+
+///////////////////////////////////////////////////////////////////////////////////////////////////////////////////
+// constexpr floating point vector constant creation functions that bypass SIMD intrinsic setters
+//
+// Added new functions instead of constexprifying MakeVectorRegisterXyz to avoid small risk of impacting codegen.
+// Long-term, we should only have one set of constexpr make functions without the Constant suffix.
+///////////////////////////////////////////////////////////////////////////////////////////////////////////////////
+
+FORCEINLINE constexpr VectorRegister4Double MakeVectorRegisterDoubleConstant(VectorRegister2Double XY, VectorRegister2Double ZW)
+{
+	return VectorRegister4Double(XY, ZW, VectorRegisterConstInit{});
+}
+
+FORCEINLINE constexpr VectorRegister4Double MakeVectorRegisterDoubleConstant(double X, double Y, double Z, double W)
+{
+	return MakeVectorRegisterDoubleConstant(MakeVectorRegister2DoubleConstant(X, Y),
+											MakeVectorRegister2DoubleConstant(Z, W));
+}
+
+FORCEINLINE constexpr VectorRegister4Float MakeVectorRegisterConstant(float X, float Y, float Z, float W)
+{
+	return MakeVectorRegisterFloatConstant(X, Y, Z, W);
+}
+
+FORCEINLINE constexpr VectorRegister4Double MakeVectorRegisterConstant(double X, double Y, double Z, double W)
+{
+	return MakeVectorRegisterDoubleConstant(X, Y, Z, W);
+}
+
+namespace GlobalVectorConstants
+{
+	inline constexpr VectorRegister4Float FloatOne = MakeVectorRegisterFloatConstant(1.f, 1.f, 1.f, 1.f);
+	inline constexpr VectorRegister4Float FloatZero = MakeVectorRegisterFloatConstant(0.f, 0.f, 0.f, 0.f);
+	inline constexpr VectorRegister4Float FloatMinusOne = MakeVectorRegisterFloatConstant(-1.f, -1.f, -1.f, -1.f);
+	inline constexpr VectorRegister4Float Float0001 = MakeVectorRegisterFloatConstant(0.f, 0.f, 0.f, 1.f);
+	inline constexpr VectorRegister4Float Float1000 = MakeVectorRegisterFloatConstant(1.f, 0.f, 0.f, 0.f);
+	inline constexpr VectorRegister4Float Float1110 = MakeVectorRegisterFloatConstant(1.f, 1.f, 1.f, 0.f);
+	inline constexpr VectorRegister4Float SmallLengthThreshold = MakeVectorRegisterFloatConstant(1.e-8f, 1.e-8f, 1.e-8f, 1.e-8f);
+	inline constexpr VectorRegister4Float FloatOneHundredth = MakeVectorRegisterFloatConstant(0.01f, 0.01f, 0.01f, 0.01f);
+	inline constexpr VectorRegister4Float Float111_Minus1 = MakeVectorRegisterFloatConstant( 1.f, 1.f, 1.f, -1.f );
+	inline constexpr VectorRegister4Float FloatMinus1_111= MakeVectorRegisterFloatConstant( -1.f, 1.f, 1.f, 1.f );
+	inline constexpr VectorRegister4Float FloatOneHalf = MakeVectorRegisterFloatConstant( 0.5f, 0.5f, 0.5f, 0.5f );
+	inline constexpr VectorRegister4Float FloatMinusOneHalf = MakeVectorRegisterFloatConstant( -0.5f, -0.5f, -0.5f, -0.5f );
+	inline constexpr VectorRegister4Float KindaSmallNumber = MakeVectorRegisterFloatConstant( UE_KINDA_SMALL_NUMBER, UE_KINDA_SMALL_NUMBER, UE_KINDA_SMALL_NUMBER, UE_KINDA_SMALL_NUMBER );
+	inline constexpr VectorRegister4Float SmallNumber = MakeVectorRegisterFloatConstant( UE_SMALL_NUMBER, UE_SMALL_NUMBER, UE_SMALL_NUMBER, UE_SMALL_NUMBER );
+	inline constexpr VectorRegister4Float ThreshQuatNormalized = MakeVectorRegisterFloatConstant( UE_THRESH_QUAT_NORMALIZED, UE_THRESH_QUAT_NORMALIZED, UE_THRESH_QUAT_NORMALIZED, UE_THRESH_QUAT_NORMALIZED );
+	inline constexpr VectorRegister4Float BigNumber = MakeVectorRegisterFloatConstant(UE_BIG_NUMBER, UE_BIG_NUMBER, UE_BIG_NUMBER, UE_BIG_NUMBER);
+
+	inline constexpr VectorRegister2Double DoubleOne2d = MakeVectorRegister2DoubleConstant(1.0, 1.0);
+	inline constexpr VectorRegister4Double DoubleOne = MakeVectorRegisterDoubleConstant(1.0, 1.0, 1.0, 1.0);
+	inline constexpr VectorRegister4Double DoubleZero = MakeVectorRegisterDoubleConstant(0.0, 0.0, 0.0, 0.0);
+	inline constexpr VectorRegister4Double DoubleMinusOne = MakeVectorRegisterDoubleConstant(-1.0, -1.0, -1.0, -1.0);
+	inline constexpr VectorRegister4Double Double0001 = MakeVectorRegisterDoubleConstant(0.0f, 0.0, 0.0, 1.0);
+	inline constexpr VectorRegister4Double Double1000 = MakeVectorRegisterDoubleConstant(1.0, 0.0, 0.0, 0.0);
+	inline constexpr VectorRegister4Double Double1110 = MakeVectorRegisterDoubleConstant(1.0, 1.0, 1.0, 0.0);
+	inline constexpr VectorRegister4Double DoubleSmallLengthThreshold = MakeVectorRegisterDoubleConstant(1.e-8, 1.e-8, 1.e-8, 1.e-8);
+	inline constexpr VectorRegister4Double DoubleOneHundredth = MakeVectorRegisterDoubleConstant(0.01, 0.01, 0.01, 0.01);
+	inline constexpr VectorRegister4Double Double111_Minus1 = MakeVectorRegisterDoubleConstant(1., 1., 1., -1.);
+	inline constexpr VectorRegister4Double DoubleMinus1_111 = MakeVectorRegisterDoubleConstant(-1., 1., 1., 1.);
+	inline constexpr VectorRegister4Double DoubleOneHalf = MakeVectorRegisterDoubleConstant(0.5, 0.5, 0.5, 0.5);
+	inline constexpr VectorRegister4Double DoubleMinusOneHalf = MakeVectorRegisterDoubleConstant(-0.5, -0.5, -0.5, -0.5);
+	inline constexpr VectorRegister4Double DoubleKindaSmallNumber = MakeVectorRegisterDoubleConstant(UE_DOUBLE_KINDA_SMALL_NUMBER, UE_DOUBLE_KINDA_SMALL_NUMBER, UE_DOUBLE_KINDA_SMALL_NUMBER, UE_DOUBLE_KINDA_SMALL_NUMBER);
+	inline constexpr VectorRegister4Double DoubleSmallNumber = MakeVectorRegisterDoubleConstant(UE_DOUBLE_SMALL_NUMBER, UE_DOUBLE_SMALL_NUMBER, UE_DOUBLE_SMALL_NUMBER, UE_DOUBLE_SMALL_NUMBER);
+	inline constexpr VectorRegister4Double DoubleThreshQuatNormalized = MakeVectorRegisterDoubleConstant(UE_DOUBLE_THRESH_QUAT_NORMALIZED, UE_DOUBLE_THRESH_QUAT_NORMALIZED, UE_DOUBLE_THRESH_QUAT_NORMALIZED, UE_DOUBLE_THRESH_QUAT_NORMALIZED);
+	inline constexpr VectorRegister4Double DoubleBigNumber = MakeVectorRegisterDoubleConstant(UE_DOUBLE_BIG_NUMBER, UE_DOUBLE_BIG_NUMBER, UE_DOUBLE_BIG_NUMBER, UE_DOUBLE_BIG_NUMBER);
+
+	inline constexpr VectorRegister Vector0001 = MakeVectorRegisterConstant(0.0, 0.0, 0.0, 1.0);
+	inline constexpr VectorRegister Vector1110 = MakeVectorRegisterConstant(1.0, 1.0, 1.0, 0.0);
+
+	inline constexpr VectorRegister4Int IntOne = MakeVectorRegisterIntConstant(1, 1, 1, 1);
+	inline constexpr VectorRegister4Int IntZero = MakeVectorRegisterIntConstant(0, 0, 0, 0);
+	inline constexpr VectorRegister4Int IntMinusOne = MakeVectorRegisterIntConstant(-1, -1, -1, -1);
+
+	/** This is to speed up Quaternion Inverse. Static variable to keep sign of inverse **/
+	inline constexpr VectorRegister4Float QINV_SIGN_MASK = MakeVectorRegisterFloatConstant( -1.f, -1.f, -1.f, 1.f );
+	inline constexpr VectorRegister4Double DOUBLE_QINV_SIGN_MASK = MakeVectorRegisterDoubleConstant(-1., -1., -1., 1.);
+
+	inline constexpr VectorRegister4Float QMULTI_SIGN_MASK0 = MakeVectorRegisterFloatConstant( 1.f, -1.f, 1.f, -1.f );
+	inline constexpr VectorRegister4Float QMULTI_SIGN_MASK1 = MakeVectorRegisterFloatConstant( 1.f, 1.f, -1.f, -1.f );
+	inline constexpr VectorRegister4Float QMULTI_SIGN_MASK2 = MakeVectorRegisterFloatConstant( -1.f, 1.f, 1.f, -1.f );
+	inline constexpr VectorRegister4Double DOUBLE_QMULTI_SIGN_MASK0 = MakeVectorRegisterDoubleConstant(1., -1., 1., -1.);
+	inline constexpr VectorRegister4Double DOUBLE_QMULTI_SIGN_MASK1 = MakeVectorRegisterDoubleConstant(1., 1., -1., -1.);
+	inline constexpr VectorRegister4Double DOUBLE_QMULTI_SIGN_MASK2 = MakeVectorRegisterDoubleConstant(-1., 1., 1., -1.);
+
+	inline constexpr VectorRegister4Float DEG_TO_RAD = MakeVectorRegisterConstant(UE_PI/(180.f), UE_PI/(180.f), UE_PI/(180.f), UE_PI/(180.f));
+	inline constexpr VectorRegister4Float DEG_TO_RAD_HALF = MakeVectorRegisterConstant((UE_PI/180.f)*0.5f, (UE_PI/180.f)*0.5f, (UE_PI/180.f)*0.5f, (UE_PI/180.f)*0.5f);
+	inline constexpr VectorRegister4Float RAD_TO_DEG = MakeVectorRegisterConstant((180.f)/UE_PI, (180.f)/UE_PI, (180.f)/UE_PI, (180.f)/UE_PI);
+	inline constexpr VectorRegister4Double DOUBLE_DEG_TO_RAD = MakeVectorRegisterConstant(UE_DOUBLE_PI/(180.), UE_DOUBLE_PI/(180.), UE_DOUBLE_PI/(180.), UE_DOUBLE_PI/(180.));
+	inline constexpr VectorRegister4Double DOUBLE_DEG_TO_RAD_HALF = MakeVectorRegisterConstant((UE_DOUBLE_PI/180.) * 0.5, (UE_DOUBLE_PI/180.) * 0.5, (UE_DOUBLE_PI/180.) * 0.5, (UE_DOUBLE_PI/180.) * 0.5);
+	inline constexpr VectorRegister4Double DOUBLE_RAD_TO_DEG = MakeVectorRegisterConstant((180.)/UE_DOUBLE_PI, (180.)/UE_DOUBLE_PI, (180.)/UE_DOUBLE_PI, (180.)/UE_DOUBLE_PI);
+
+	/** Bitmask to AND out the XYZ components in a vector */
+	inline VectorRegister4Float XYZMask()			{ return MakeVectorRegisterFloatMask( ~uint32(0), ~uint32(0), ~uint32(0), uint32(0)); }
+	inline VectorRegister4Double DoubleXYZMask()	{ return MakeVectorRegisterDoubleMask(~uint64(0), ~uint64(0), ~uint64(0), uint64(0)); }
+	
+	/** Bitmask to AND out the sign bit of each components in a vector */
+	inline VectorRegister4Float SignBit()			{ return MakeVectorRegisterFloatMask( 0x80000000,  0x80000000,  0x80000000,  0x80000000); }
+	inline VectorRegister4Float SignMask()			{ return MakeVectorRegisterFloatMask(~0x80000000, ~0x80000000, ~0x80000000, ~0x80000000); }
+	inline constexpr VectorRegister4Int IntSignBit =	MakeVectorRegisterIntConstant( 0x80000000,  0x80000000,  0x80000000,  0x80000000);
+	inline constexpr VectorRegister4Int IntSignMask =	MakeVectorRegisterIntConstant(~0x80000000, ~0x80000000, ~0x80000000, ~0x80000000);
+
+	inline constexpr VectorRegister4Int IntAllMask = MakeVectorRegisterIntConstant(			~uint32(0), ~uint32(0), ~uint32(0), ~uint32(0));
+	inline VectorRegister4Float AllMask()			{ return MakeVectorRegisterFloatMask(	~uint32(0), ~uint32(0), ~uint32(0), ~uint32(0)); }
+	inline VectorRegister4Double DoubleAllMask()	{ return MakeVectorRegisterDoubleMask(	~uint64(0), ~uint64(0), ~uint64(0), ~uint64(0)); }
+
+	inline VectorRegister4Double DoubleSignBit()	{ return MakeVectorRegisterDoubleMask(uint64(1) << 63, uint64(1) << 63, uint64(1) << 63, uint64(1) << 63); }
+	inline VectorRegister4Double DoubleSignMask()	{ return MakeVectorRegisterDoubleMask(~(uint64(1) << 63), ~(uint64(1) << 63), ~(uint64(1) << 63), ~(uint64(1) << 63)); }
+
+	/** Vector full of positive infinity */
+	inline VectorRegister4Float FloatInfinity()		{ return MakeVectorRegisterFloatMask((uint32)0x7F800000, (uint32)0x7F800000, (uint32)0x7F800000, (uint32)0x7F800000); }
+	inline VectorRegister4Double DoubleInfinity()	{ return MakeVectorRegisterDoubleMask((uint64)0x7FF0000000000000, (uint64)0x7FF0000000000000, (uint64)0x7FF0000000000000, (uint64)0x7FF0000000000000); }
+
+	inline constexpr VectorRegister4Float Pi = MakeVectorRegisterConstant(UE_PI, UE_PI, UE_PI, UE_PI);
+	inline constexpr VectorRegister4Float TwoPi = MakeVectorRegisterConstant(2.0f*UE_PI, 2.0f*UE_PI, 2.0f*UE_PI, 2.0f*UE_PI);
+	inline constexpr VectorRegister4Float PiByTwo = MakeVectorRegisterConstant(0.5f*UE_PI, 0.5f*UE_PI, 0.5f*UE_PI, 0.5f*UE_PI);
+	inline constexpr VectorRegister4Float PiByFour = MakeVectorRegisterConstant(0.25f*UE_PI, 0.25f*UE_PI, 0.25f*UE_PI, 0.25f*UE_PI);
+	inline constexpr VectorRegister4Float OneOverPi = MakeVectorRegisterConstant(1.0f / UE_PI, 1.0f / UE_PI, 1.0f / UE_PI, 1.0f / UE_PI);
+	inline constexpr VectorRegister4Float OneOverTwoPi = MakeVectorRegisterConstant(1.0f / (2.0f*UE_PI), 1.0f / (2.0f*UE_PI), 1.0f / (2.0f*UE_PI), 1.0f / (2.0f*UE_PI));
+
+	inline constexpr VectorRegister4Double DoublePi = MakeVectorRegisterDoubleConstant(UE_DOUBLE_PI, UE_DOUBLE_PI, UE_DOUBLE_PI, UE_DOUBLE_PI);
+	inline constexpr VectorRegister4Double DoubleTwoPi = MakeVectorRegisterDoubleConstant(2.0 * UE_DOUBLE_PI, 2.0 * UE_DOUBLE_PI, 2.0 * UE_DOUBLE_PI, 2.0 * UE_DOUBLE_PI);
+	inline constexpr VectorRegister4Double DoublePiByTwo = MakeVectorRegisterDoubleConstant(0.5 * UE_DOUBLE_PI, 0.5 * UE_DOUBLE_PI, 0.5 * UE_DOUBLE_PI, 0.5 * UE_DOUBLE_PI);
+	inline constexpr VectorRegister4Double DoublePiByFour = MakeVectorRegisterDoubleConstant(0.25 * UE_DOUBLE_PI, 0.25 * UE_DOUBLE_PI, 0.25 * UE_DOUBLE_PI, 0.25 * UE_DOUBLE_PI);
+	inline constexpr VectorRegister4Double DoubleOneOverPi = MakeVectorRegisterDoubleConstant(1.0 / UE_DOUBLE_PI, 1.0 / UE_DOUBLE_PI, 1.0 / UE_DOUBLE_PI, 1.0 / UE_DOUBLE_PI);
+	inline constexpr VectorRegister4Double DoubleOneOverTwoPi = MakeVectorRegisterDoubleConstant(1.0 / (2.0 * UE_DOUBLE_PI), 1.0 / (2.0 * UE_DOUBLE_PI), 1.0 / (2.0 * UE_DOUBLE_PI), 1.0 / (2.0 * UE_DOUBLE_PI));
+
+	inline constexpr VectorRegister4Float Float255 = MakeVectorRegisterConstant(255.0f, 255.0f, 255.0f, 255.0f);
+	inline constexpr VectorRegister4Float Float127 = MakeVectorRegisterConstant(127.0f, 127.0f, 127.0f, 127.0f);
+	inline constexpr VectorRegister4Float FloatNeg127 = MakeVectorRegisterConstant(-127.0f, -127.0f, -127.0f, -127.0f);
+	inline constexpr VectorRegister4Float Float360 = MakeVectorRegisterConstant(360.f, 360.f, 360.f, 360.f);
+	inline constexpr VectorRegister4Float Float180 = MakeVectorRegisterConstant(180.f, 180.f, 180.f, 180.f);
+
+	inline constexpr VectorRegister4Double Double255 = MakeVectorRegisterDoubleConstant(255.0, 255.0, 255.0, 255.0);
+	inline constexpr VectorRegister4Double Double127 = MakeVectorRegisterDoubleConstant(127.0, 127.0, 127.0, 127.0);
+	inline constexpr VectorRegister4Double DoubleNeg127 = MakeVectorRegisterDoubleConstant(-127.0, -127.0, -127.0, -127.0);
+	inline constexpr VectorRegister4Double Double360 = MakeVectorRegisterDoubleConstant(360., 360., 360., 360.);
+	inline constexpr VectorRegister4Double Double180 = MakeVectorRegisterDoubleConstant(180., 180., 180., 180.);
+
+	// All float numbers greater than or equal to this have no fractional value.
+	inline constexpr VectorRegister4Float FloatNonFractional = MakeVectorRegisterConstant(UE_FLOAT_NON_FRACTIONAL, UE_FLOAT_NON_FRACTIONAL, UE_FLOAT_NON_FRACTIONAL, UE_FLOAT_NON_FRACTIONAL);
+	inline constexpr VectorRegister4Double DoubleNonFractional = MakeVectorRegisterDoubleConstant(UE_DOUBLE_NON_FRACTIONAL, UE_DOUBLE_NON_FRACTIONAL, UE_DOUBLE_NON_FRACTIONAL, UE_DOUBLE_NON_FRACTIONAL);
+
+	inline constexpr VectorRegister4Float FloatTwo = MakeVectorRegisterConstant(2.0f, 2.0f, 2.0f, 2.0f);
+	inline constexpr uint32 AlmostTwoBits = 0x3fffffff;
+	inline VectorRegister4Float FloatAlmostTwo()	{ return MakeVectorRegisterFloatMask(AlmostTwoBits, AlmostTwoBits, AlmostTwoBits, AlmostTwoBits); }
+
+	inline constexpr VectorRegister4Double DoubleTwo = MakeVectorRegisterDoubleConstant(2.0, 2.0, 2.0, 2.0);
+	inline constexpr uint64 DoubleAlmostTwoBits = 0x3FFFFFFFFFFFFFFF;
+	inline VectorRegister4Double DoubleAlmostTwo()	{ return MakeVectorRegisterDoubleMask(DoubleAlmostTwoBits, DoubleAlmostTwoBits, DoubleAlmostTwoBits, DoubleAlmostTwoBits); }
+}
diff --git a/Engine/Source/Runtime/Core/Public/Math/VectorRegister.h b/Engine/Source/Runtime/Core/Public/Math/VectorRegister.h
--- a/Engine/Source/Runtime/Core/Public/Math/VectorRegister.h
+++ b/Engine/Source/Runtime/Core/Public/Math/VectorRegister.h
@@ -27,7 +27,7 @@ template<typename T>
 using TVectorRegisterType = std::conditional_t<std::is_same_v<T, float>, VectorRegister4Float, std::conditional_t<std::is_same_v<T, double>, VectorRegister4Double, void> >;
 
 // 'Cross-platform' vector intrinsics (built on the platform-specific ones defined above)
-#include "Math/UnrealMathVectorCommon.h"
+#include "Math/UnrealMathVectorCommon.h.inl"
 
 
 #if !defined(UE_SSE_DOUBLE_ALIGNMENT) || (UE_SSE_DOUBLE_ALIGNMENT <= 16)
diff --git a/Engine/Source/Runtime/Core/Tests/Math/UnrealMathTest.cpp b/Engine/Source/Runtime/Core/Tests/Math/UnrealMathTest.cpp
--- a/Engine/Source/Runtime/Core/Tests/Math/UnrealMathTest.cpp
+++ b/Engine/Source/Runtime/Core/Tests/Math/UnrealMathTest.cpp
@@ -1004,8 +1004,8 @@ MATHTEST_INLINE void TestVectorSinCos(VectorRegisterType& S, VectorRegisterType&
 }
 
 // Generic way to test functions taking a single FloatType param and returning a FloatType value, and comparing to a vectorized version of the same function.
-template<typename FloatType, typename VectorRegisterType, FloatType(*Func)(FloatType), VectorRegisterType(*VectorFunc)(const VectorRegisterType&)>
-FORCENOINLINE bool TestVectorFunction1Param(VectorRegisterType& ReferenceResult, VectorRegisterType& VectorResult, const VectorRegisterType& VParams, FloatType ErrorTolerance)
+template<typename FloatType, typename VectorRegisterType, FloatType(*Func)(FloatType), VectorRegisterType(*VectorFunc)(VectorRegisterType)>
+FORCENOINLINE bool TestVectorFunction1Param(VectorRegisterType& ReferenceResult, VectorRegisterType& VectorResult, VectorRegisterType VParams, FloatType ErrorTolerance)
 {
 	// Load params
 	FloatType FloatParams[4];
@@ -1027,8 +1027,8 @@ FORCENOINLINE bool TestVectorFunction1Param(VectorRegisterType& ReferenceResult,
 }
 
 // Generic way to test functions taking two FloatType params and returning a FloatType value, and comparing to a vectorized version of the same function.
-template<typename FloatType, typename VectorRegisterType, FloatType(*Func)(FloatType, FloatType), VectorRegisterType(*VectorFunc)(const VectorRegisterType&, const VectorRegisterType&)>
-FORCENOINLINE bool TestVectorFunction2Params(VectorRegisterType& ReferenceResult, VectorRegisterType& VectorResult, const VectorRegisterType& VParams0, const VectorRegisterType& VParams1, FloatType ErrorTolerance)
+template<typename FloatType, typename VectorRegisterType, FloatType(*Func)(FloatType, FloatType), VectorRegisterType(*VectorFunc)(VectorRegisterType, VectorRegisterType)>
+FORCENOINLINE bool TestVectorFunction2Params(VectorRegisterType& ReferenceResult, VectorRegisterType& VectorResult, VectorRegisterType VParams0, VectorRegisterType VParams1, FloatType ErrorTolerance)
 {
 	// Load params
 	FloatType FloatParams0[4];
diff --git a/Engine/Source/Runtime/Experimental/Chaos/Public/Chaos/AABBVectorized.h b/Engine/Source/Runtime/Experimental/Chaos/Public/Chaos/AABBVectorized.h
--- a/Engine/Source/Runtime/Experimental/Chaos/Public/Chaos/AABBVectorized.h
+++ b/Engine/Source/Runtime/Experimental/Chaos/Public/Chaos/AABBVectorized.h
@@ -5,7 +5,6 @@
 #include "Chaos/Core.h"
 #include "ChaosArchive.h"
 #include "Math/VectorRegister.h"
-#include "Math/UnrealMathVectorConstants.h"
 #include "Chaos/VectorUtility.h"
 
 
diff --git a/Engine/Source/Runtime/Experimental/Chaos/Public/Chaos/AABBVectorizedDouble.h b/Engine/Source/Runtime/Experimental/Chaos/Public/Chaos/AABBVectorizedDouble.h
--- a/Engine/Source/Runtime/Experimental/Chaos/Public/Chaos/AABBVectorizedDouble.h
+++ b/Engine/Source/Runtime/Experimental/Chaos/Public/Chaos/AABBVectorizedDouble.h
@@ -5,7 +5,6 @@
 #include "Chaos/Core.h"
 #include "ChaosArchive.h"
 #include "Math/VectorRegister.h"
-#include "Math/UnrealMathVectorConstants.h"
 #include "Chaos/VectorUtility.h"
 
 
diff --git a/Engine/Source/Runtime/Experimental/Chaos/Public/Chaos/Convex.h b/Engine/Source/Runtime/Experimental/Chaos/Public/Chaos/Convex.h
--- a/Engine/Source/Runtime/Experimental/Chaos/Public/Chaos/Convex.h
+++ b/Engine/Source/Runtime/Experimental/Chaos/Public/Chaos/Convex.h
@@ -5,9 +5,6 @@
 #include "Chaos/AABB.h"
 #include "Chaos/ConvexStructureData.h"
 #include "Chaos/MassProperties.h"
-
-#include "Math/UnrealMathVectorCommon.h"
-
 #include "CollisionConvexMesh.h"
 #include "ChaosArchive.h"
 #include "ChaosCheck.h"
diff --git a/Engine/Source/Runtime/SignalProcessing/Private/FloatArrayMath.cpp b/Engine/Source/Runtime/SignalProcessing/Private/FloatArrayMath.cpp
--- a/Engine/Source/Runtime/SignalProcessing/Private/FloatArrayMath.cpp
+++ b/Engine/Source/Runtime/SignalProcessing/Private/FloatArrayMath.cpp
@@ -2,7 +2,6 @@
 
 #include "DSP/FloatArrayMath.h"
 #include "CoreMinimal.h"
-#include "Math/UnrealMathVectorConstants.h"
 #include "SignalProcessingModule.h"
 #include "ProfilingDebugging/CsvProfiler.h"
 #include "DSP/Dsp.h"
