IntrusiveMutex: Fixed a compile error on older versions of Clang

#rnx

diff --git a/Engine/Source/Runtime/Core/Public/Async/IntrusiveMutex.h b/Engine/Source/Runtime/Core/Public/Async/IntrusiveMutex.h
--- a/Engine/Source/Runtime/Core/Public/Async/IntrusiveMutex.h
+++ b/Engine/Source/Runtime/Core/Public/Async/IntrusiveMutex.h
@@ -31,7 +31,7 @@ struct TIntrusiveMutexParams
 /**
  * An 2-bit intrusive mutex that is not fair and does not support recursive locking.
  */
-template <typename StateType, TIntrusiveMutexParams<StateType> Params = {}>
+template <typename StateType, TIntrusiveMutexParams<StateType> Params = TIntrusiveMutexParams<StateType>{}>
 class TIntrusiveMutex
 {
 	TIntrusiveMutex() = delete;
