- Bumpding key to flush out some stale DDC data

#rnx

diff --git a/Engine/Source/Runtime/Core/Private/UObject/FortniteMainBranchObjectVersion.cpp b/Engine/Source/Runtime/Core/Private/UObject/FortniteMainBranchObjectVersion.cpp
--- a/Engine/Source/Runtime/Core/Private/UObject/FortniteMainBranchObjectVersion.cpp
+++ b/Engine/Source/Runtime/Core/Private/UObject/FortniteMainBranchObjectVersion.cpp
@@ -11,7 +11,7 @@ TMap<FGuid, FGuid> FFortniteMainBranchObjectVersion::GetSystemGuids()
 	SystemGuids.Add(DevGuids.MATERIALSHADERMAP_DERIVEDDATA_VER, FGuid("04B8C817071A453C9EBEFBFB51D5F785"));
 	SystemGuids.Add(DevGuids.NANITE_DERIVEDDATA_VER, FGuid("FB5EAC6147EF4591B0D5EB7A85062AC9"));
 	SystemGuids.Add(DevGuids.NIAGARASHADERMAP_DERIVEDDATA_VER, FGuid("6360A977062842A29ED30E3A7ACB0E64"));
-	SystemGuids.Add(DevGuids.Niagara_LatestScriptCompileVersion, FGuid("081BBC4CD785874B8E77E9AEAEBD0537"));
+	SystemGuids.Add(DevGuids.Niagara_LatestScriptCompileVersion, FGuid("97F0FD7149D90D42B65D6557AF477D1A"));
 	SystemGuids.Add(DevGuids.SkeletalMeshDerivedDataVersion, FGuid("716424BCEC314E969ACFE36D2C2ECE49"));
 	SystemGuids.Add(DevGuids.STATICMESH_DERIVEDDATA_VER, FGuid("3ABF32149F9D4D838750368AB95573BE"));
 	SystemGuids.Add(DevGuids.MaterialTranslationDDCVersion, FGuid("AAAFC530186D4149A06F027D4A060F5E"));
