Fixed to GetConfigRuleVars FConditionVariable usage, as it could cause a hang

#android
#tests Tested that GetConfigRuleVars worked and the game ran

diff --git a/Engine/Source/Runtime/Core/Private/Android/AndroidPlatformMisc.cpp b/Engine/Source/Runtime/Core/Private/Android/AndroidPlatformMisc.cpp
--- a/Engine/Source/Runtime/Core/Private/Android/AndroidPlatformMisc.cpp
+++ b/Engine/Source/Runtime/Core/Private/Android/AndroidPlatformMisc.cpp
@@ -55,7 +55,7 @@
 #include "Misc/OutputDevice.h"
 #include "Logging/LogMacros.h"
 #include "Misc/OutputDeviceError.h"
-#include "Experimental/Async/ConditionVariable.h"
+#include "Async/EventCount.h"
 
 #include "VulkanCommon.h"
 #include "IVulkanDynamicRHI.h"
@@ -2384,18 +2384,23 @@ bool FAndroidMisc::IsExternalMemoryAndroidHardwareBufferExtensionLoaded()
 	return false;
 }
 
-static FCriticalSection ConfigRuleVarsCS;
-static UE::FConditionVariable AvailableConfigRuleVarsCV;
+static UE::FEventCount AvailableConfigRuleVarsEvent;
 static std::atomic_bool bConfigRulesReady = false;
 
 const TMap<FString, FString>& FAndroidMisc::GetConfigRuleVars()
 {
-	while(!bConfigRulesReady)
+	if (!bConfigRulesReady)
 	{
-		FPlatformMisc::LowLevelOutputDebugStringf(TEXT("thread waiting for configrules to be set"));
-		AvailableConfigRuleVarsCV.Wait(ConfigRuleVarsCS);
-		FPlatformMisc::LowLevelOutputDebugStringf(TEXT("done thread waiting for configrules to be set"));
+		UE::FEventCountToken Token = AvailableConfigRuleVarsEvent.PrepareWait();
+
+		if (!bConfigRulesReady)
+		{
+			FPlatformMisc::LowLevelOutputDebugStringf(TEXT("thread waiting for configrules to be set"));
+			AvailableConfigRuleVarsEvent.Wait(Token);
+			FPlatformMisc::LowLevelOutputDebugStringf(TEXT("done thread waiting for configrules to be set"));
+		}
 	}
+	
 	// ConfigRuleVars are read-only from here on.
 	return ConfigRulesVariables;
 }
@@ -2421,8 +2426,6 @@ bool FAndroidMisc::AllowThreadHeartBeat()
 
 JNI_METHOD void Java_com_epicgames_unreal_GameActivity_nativeSetConfigRulesVariables(JNIEnv* jenv, jobject thiz, jobjectArray KeyValuePairs)
 {
-	FScopeLock ConfigRuleVarsLock(&ConfigRuleVarsCS);
-
 	int32 Count = jenv->GetArrayLength(KeyValuePairs);
 	int32 Index = 0;
 	while (Index < Count)
@@ -2434,7 +2437,7 @@ JNI_METHOD void Java_com_epicgames_unreal_GameActivity_nativeSetConfigRulesVaria
 	}
 
 	bConfigRulesReady = true;
-	AvailableConfigRuleVarsCV.NotifyAll();
+	AvailableConfigRuleVarsEvent.Notify();
 }
 
 extern bool AndroidThunkCpp_HasMetaDataKey(const FString& Key);
