[AutoRTFM] Mark ParkingLot APIs as inaccessible from within a transaction.

It is only safe to access these APIs from the open; transactionally-safe locks always go into the open before taking the inner lock. Marking a function NOAUTORTFM will cause an immediate transactional failure if these functions are inadvertently called. All parking lot APIs publicly accessible from the header eventually funnel through these three functions.

diff --git a/Engine/Source/Runtime/Core/Private/Async/ParkingLot.cpp b/Engine/Source/Runtime/Core/Private/Async/ParkingLot.cpp
--- a/Engine/Source/Runtime/Core/Private/Async/ParkingLot.cpp
+++ b/Engine/Source/Runtime/Core/Private/Async/ParkingLot.cpp
@@ -5,6 +5,7 @@
 #include "Algo/Sort.h"
 #include "Async/UniqueLock.h"
 #include "Async/WordMutex.h"
+#include "AutoRTFM.h"
 #include "Containers/Array.h"
 #include "Containers/ArrayView.h"
 #include "HAL/PlatformManualResetEvent.h"
@@ -509,6 +510,7 @@ FWaitState WaitFor(const void* Address, const TFunctionRef<bool()>& CanWait, con
 	return WaitUntil(Address, CanWait, BeforeWait, FMonotonicTimePoint::Now() + WaitTime);
 }
 
+UE_AUTORTFM_NOAUTORTFM
 FWaitState WaitUntil(const void* Address, const TFunctionRef<bool()>& CanWait, const TFunctionRef<void()>& BeforeWait, FMonotonicTimePoint WaitTime)
 {
 	using namespace Private;
@@ -581,6 +583,7 @@ FWaitState WaitUntil(const void* Address, const TFunctionRef<bool()>& CanWait, c
 	return State;
 }
 
+UE_AUTORTFM_NOAUTORTFM
 void WakeOne(const void* Address, const TFunctionRef<uint64(FWakeState)>& OnWakeState)
 {
 	using namespace Private;
@@ -631,6 +634,7 @@ FWakeState WakeOne(const void* Address)
 	return OutState;
 }
 
+UE_AUTORTFM_NOAUTORTFM
 uint32 WakeMultiple(const void* const Address, const uint32 WakeCount)
 {
 	using namespace Private;
