Changed UE_TVARIANT_TRIVIAL_DESTRUCTOR_USING_CONCEPTS to PLATFORM_COMPILER_SUPPORTS_CONSTRAINED_DESTRUCTORS.

diff --git a/Engine/Source/Runtime/Core/Public/HAL/Platform.h b/Engine/Source/Runtime/Core/Public/HAL/Platform.h
--- a/Engine/Source/Runtime/Core/Public/HAL/Platform.h
+++ b/Engine/Source/Runtime/Core/Public/HAL/Platform.h
@@ -266,6 +266,16 @@
 #ifndef PLATFORM_COMPILER_SUPPORTS_BUILTIN_BITCAST
 	#error PLATFORM_COMPILER_SUPPORTS_BUILTIN_BITCAST should have been defined by now
 #endif
+#ifndef PLATFORM_COMPILER_SUPPORTS_CONSTRAINED_DESTRUCTORS
+	// Disable support for constrained destructors on MSVC until this internal compile error gets fixed:
+	//
+	// https://developercommunity.visualstudio.com/t/Internal-compile-error-when-returning-a/10650542
+	#if defined(__cpp_concepts) && __cpp_concepts >= 202002L && (defined(__clang__) || defined(__GNUC__) || !defined(_MSC_VER) || _MSC_VER >= 1942)
+		#define PLATFORM_COMPILER_SUPPORTS_CONSTRAINED_DESTRUCTORS 1
+	#else
+		#define PLATFORM_COMPILER_SUPPORTS_CONSTRAINED_DESTRUCTORS 0
+	#endif
+#endif
 #ifndef PLATFORM_TCHAR_IS_4_BYTES
 	#define PLATFORM_TCHAR_IS_4_BYTES			0
 #endif
diff --git a/Engine/Source/Runtime/Core/Public/Misc/TVariant.h b/Engine/Source/Runtime/Core/Public/Misc/TVariant.h
--- a/Engine/Source/Runtime/Core/Public/Misc/TVariant.h
+++ b/Engine/Source/Runtime/Core/Public/Misc/TVariant.h
@@ -40,13 +40,13 @@ struct FEmptyVariantState
  */
 template <typename T, typename... Ts>
 class TVariant final
-#if UE_TVARIANT_TRIVIAL_DESTRUCTOR_USING_CONCEPTS
+#if PLATFORM_COMPILER_SUPPORTS_CONSTRAINED_DESTRUCTORS
 	: private UE::Core::Private::TVariantStorage<T, Ts...>
 #else
 	: private std::conditional_t<!std::is_trivially_destructible_v<T> || (!std::is_trivially_destructible_v<Ts> || ...), UE::Core::Private::TDestructibleVariantStorage<T, Ts...> , UE::Core::Private::TVariantStorage<T, Ts...>>
 #endif
 {
-#if UE_TVARIANT_TRIVIAL_DESTRUCTOR_USING_CONCEPTS
+#if PLATFORM_COMPILER_SUPPORTS_CONSTRAINED_DESTRUCTORS
 	using Super = UE::Core::Private::TVariantStorage<T, Ts...>;
 #else
 	using Super = std::conditional_t<!std::is_trivially_destructible_v<T> || (!std::is_trivially_destructible_v<Ts> || ...), UE::Core::Private::TDestructibleVariantStorage<T, Ts...> , UE::Core::Private::TVariantStorage<T, Ts...>>;
@@ -118,7 +118,7 @@ public:
 		return *this;
 	}
 
-#if UE_TVARIANT_TRIVIAL_DESTRUCTOR_USING_CONCEPTS
+#if PLATFORM_COMPILER_SUPPORTS_CONSTRAINED_DESTRUCTORS
 	/** Destruct the underlying type (if appropriate) */
 	~TVariant()
 		requires(!std::is_trivially_destructible_v<T> || (!std::is_trivially_destructible_v<Ts> || ...))
@@ -254,7 +254,7 @@ public:
 	}
 
 private:
-#if UE_TVARIANT_TRIVIAL_DESTRUCTOR_USING_CONCEPTS
+#if PLATFORM_COMPILER_SUPPORTS_CONSTRAINED_DESTRUCTORS
 	/** Index into the template parameter pack for the type held. */
 	uint8 TypeIndex;
 #else
diff --git a/Engine/Source/Runtime/Core/Public/Misc/TVariantMeta.h b/Engine/Source/Runtime/Core/Public/Misc/TVariantMeta.h
--- a/Engine/Source/Runtime/Core/Public/Misc/TVariantMeta.h
+++ b/Engine/Source/Runtime/Core/Public/Misc/TVariantMeta.h
@@ -11,19 +11,6 @@
 
 #include "Misc/AssertionMacros.h"
 
-// Disable the concept checks on MSVC until this internal compile error gets fixed:
-//
-// https://developercommunity.visualstudio.com/t/Internal-compile-error-when-returning-a/10650542
-#ifndef _MSC_VER
-	#if defined(__cpp_concepts) && __cpp_concepts >= 202002L
-		#define UE_TVARIANT_TRIVIAL_DESTRUCTOR_USING_CONCEPTS 1
-	#endif
-#endif
-
-#if !defined(UE_TVARIANT_TRIVIAL_DESTRUCTOR_USING_CONCEPTS)
-	#define UE_TVARIANT_TRIVIAL_DESTRUCTOR_USING_CONCEPTS 0
-#endif
-
 template <typename T, typename... Ts>
 class TVariant;
 
@@ -170,7 +157,7 @@ namespace Private
 
 		TAlignedBytes<SizeofValue, AlignofValue> Storage;
 
-#if !UE_TVARIANT_TRIVIAL_DESTRUCTOR_USING_CONCEPTS
+#if !PLATFORM_COMPILER_SUPPORTS_CONSTRAINED_DESTRUCTORS
 		/** Index into the template parameter pack for the type held. */
 		uint8 TypeIndex;
 #endif
@@ -396,7 +383,7 @@ namespace Private
 		return Invokers[EncodedIndex](Forward<Func>(Callable), Forward<Variants>(Args)...);
 	}
 
-#if !UE_TVARIANT_TRIVIAL_DESTRUCTOR_USING_CONCEPTS
+#if !PLATFORM_COMPILER_SUPPORTS_CONSTRAINED_DESTRUCTORS
 	template <typename T, typename... Ts>
 	struct TDestructibleVariantStorage : TVariantStorage<T, Ts...>
 	{
