Use correct allocator in case callstack trace is initalized later

In cases where user initialized memalloc and module channels on startup, but enable callstacks at a later point, callstack module must use the correct allocator. With this change TraceMalloc is used when it exists, otherwise falling back to global GMalloc.

#rnx

diff --git a/Engine/Source/Runtime/Core/Private/ProfilingDebugging/MemoryTrace.cpp b/Engine/Source/Runtime/Core/Private/ProfilingDebugging/MemoryTrace.cpp
--- a/Engine/Source/Runtime/Core/Private/ProfilingDebugging/MemoryTrace.cpp
+++ b/Engine/Source/Runtime/Core/Private/ProfilingDebugging/MemoryTrace.cpp
@@ -773,6 +773,17 @@ void MemoryTrace_UnmarkAllocAsHeap(uint64 Address, HeapId Heap, uint32 ExternalC
 		<< HeapUnmarkAlloc.Heap(Heap);
 }
 
+/////////////////////////////////////////////////////////////////////////////
+FMalloc* MemoryTrace_GetAllocator()
+{
+	if (GTraceMalloc.IsConstructed())
+	{
+		return &GTraceMalloc;
+	}
+	return nullptr;
+}
+
+
 #else // UE_MEMORY_TRACE_ENABLED
 
 /////////////////////////////////////////////////////////////////////////////
diff --git a/Engine/Source/Runtime/Core/Private/ProfilingDebugging/TraceAuxiliary.cpp b/Engine/Source/Runtime/Core/Private/ProfilingDebugging/TraceAuxiliary.cpp
--- a/Engine/Source/Runtime/Core/Private/ProfilingDebugging/TraceAuxiliary.cpp
+++ b/Engine/Source/Runtime/Core/Private/ProfilingDebugging/TraceAuxiliary.cpp
@@ -1961,8 +1961,10 @@ void FTraceAuxiliary::Initialize(const TCHAR* CommandLine)
 	// Register end frame callbacks
 	GTraceAuxiliary.RegisterEndFrameCallbacks();
 
-	// Initialize callstack tracing with the regular malloc (it might have already been initialized by memory tracing).
-	CallstackTrace_Create(UE::Private::GMalloc);
+	// Initialize callstack tracing. If for any reason memory tracing has been initialized
+	// without callstack tracing we need to use the same allocator. Otherwise use regular malloc. 
+	FMalloc* TraceAllocator = MemoryTrace_GetAllocator();
+	CallstackTrace_Create(TraceAllocator ? TraceAllocator : UE::Private::GMalloc);
 	CallstackTrace_Initialize();
 
 	// By default use 1 msec for stack sampling interval.
diff --git a/Engine/Source/Runtime/Core/Public/ProfilingDebugging/MemoryTrace.h b/Engine/Source/Runtime/Core/Public/ProfilingDebugging/MemoryTrace.h
--- a/Engine/Source/Runtime/Core/Public/ProfilingDebugging/MemoryTrace.h
+++ b/Engine/Source/Runtime/Core/Public/ProfilingDebugging/MemoryTrace.h
@@ -96,6 +96,12 @@ CORE_API UE_TRACE_CHANNEL_EXTERN(MemAllocChannel);
 class FMalloc* MemoryTrace_Create(class FMalloc* InMalloc);
 void MemoryTrace_Initialize();
 
+/**
+ * Returns the untracked allocator used to initalize memory tracing. Only internally used.
+ * @return A valid allocator if exists, nullptr otherwise.
+ */
+class FMalloc* MemoryTrace_GetAllocator();
+
 /**
  * Register a new heap specification (name). Use the returned value when marking heaps.
  * @param ParentId Heap id of parent heap.
@@ -196,5 +202,6 @@ inline void MemoryTrace_ReallocFree(uint64 Address, HeapId RootHeap = EMemoryTra
 inline void MemoryTrace_ReallocAlloc(uint64 Address, uint64 NewSize, uint32 Alignment, HeapId RootHeap = EMemoryTraceRootHeap::SystemMemory, uint32 ExternalCallstackId = 0) {}
 inline void MemoryTrace_UpdateAlloc(uint64 Address, HeapId RootHeap = EMemoryTraceRootHeap::SystemMemory, uint32 ExternalCallstackId = 0) {}
 inline void MemoryTrace_SwapOp(uint64 PageAddress, EMemoryTraceSwapOperation SwapOperation, uint32 CompressedSize = 0, uint32 CallstackId = 0) {}
+inline class FMalloc* MemoryTrace_GetAllocator() { return nullptr; }
 
 #endif // UE_MEMORY_TRACE_ENABLED
