Fix an issue where resaving a package via the ResavePackages commandlet would remove the package's thumbnail.

#rnx

- The issue was introduced in CL 42374406 where a check was added to the thumbnail loading via FLinkerLoad to check if the thumbnail had valid image data after serialization. The code was only checking that uncompressed data exists after serialization but typically we will have compressed data which will be uncompressed later on first use. This was causing the commandlet to assume that each package being re-saved did not have existing thumbnail data so it was not saved back out.
-- The normal editor flow continued to work as thumbnails there are loaded via FAssetThumbnailPool instead.
- Moved the condition logic to a new method FObjectThumbnail::HasValidImageData so that it can be re-used and properly documented, the condition has been updated to check if either the uncompressed or compressed data exists.

diff --git a/Engine/Source/Runtime/Core/Public/Misc/ObjectThumbnail.h b/Engine/Source/Runtime/Core/Public/Misc/ObjectThumbnail.h
--- a/Engine/Source/Runtime/Core/Public/Misc/ObjectThumbnail.h
+++ b/Engine/Source/Runtime/Core/Public/Misc/ObjectThumbnail.h
@@ -124,13 +124,13 @@ public:
 	{
 		return bIsDirty;
 	}
-	
+
 	/** Marks the thumbnail as dirty. */
 	void MarkAsDirty()
 	{
 		bIsDirty = true;
 	}
-	
+
 	/** Access the image data in place (does not decompress). */
 	TArray< uint8 >& AccessImageData()
 	{
@@ -155,6 +155,17 @@ public:
 		return ImageWidth == 0 || ImageHeight == 0;
 	}
 
+	/**
+	 * Returns if the thumbnail actually has any valid image data or not.
+	 * Note that it is possible for ::IsEmpty to return true and this method
+	 * to return false if there was a problem during serialization or the 
+	 * thumbnail data has otherwise become corrupted.
+	 */
+	bool HasValidImageData() const
+	{
+		return !ImageData.IsEmpty() || !CompressedImageData.IsEmpty();
+	}
+
 	/** Returns thumbnail compressor used on current compressed image data. */
 	CORE_API FThumbnailCompressionInterface* GetCompressor() const;
 
diff --git a/Engine/Source/Runtime/CoreUObject/Private/UObject/LinkerLoad.cpp b/Engine/Source/Runtime/CoreUObject/Private/UObject/LinkerLoad.cpp
--- a/Engine/Source/Runtime/CoreUObject/Private/UObject/LinkerLoad.cpp
+++ b/Engine/Source/Runtime/CoreUObject/Private/UObject/LinkerLoad.cpp
@@ -2816,7 +2816,7 @@ FLinkerLoad::ELinkerStatus FLinkerLoad::SerializeThumbnails( bool bForceEnableIn
 				FObjectThumbnail LoadedThumbnail;
 				LoadedThumbnail.Serialize(DataStream.EnterElement());
 
-				if (LoadedThumbnail.AccessImageData().IsEmpty())
+				if (!LoadedThumbnail.HasValidImageData())
 				{
 					// If we failed to load the thumbnail, stop loading as it might be unsafe
 					// to continue reading the stream.
