Removed reference to unused TOptionalBase type.

diff --git a/Engine/Source/Runtime/Core/Public/Misc/IntrusiveUnsetOptionalState.h b/Engine/Source/Runtime/Core/Public/Misc/IntrusiveUnsetOptionalState.h
--- a/Engine/Source/Runtime/Core/Public/Misc/IntrusiveUnsetOptionalState.h
+++ b/Engine/Source/Runtime/Core/Public/Misc/IntrusiveUnsetOptionalState.h
@@ -60,12 +60,6 @@
 
 #include "Misc/OptionalFwd.h"
 
-namespace UE::Core::Private
-{
-	template <typename OptionalType, bool HasIntrusiveUnsetOptionalState>
-	struct TOptionalStorage;
-}
-
 struct FIntrusiveUnsetOptionalState
 {
 	// Defined in the global namespace for convenience and consistency
@@ -74,9 +68,6 @@ struct FIntrusiveUnsetOptionalState
 	template <typename>
 	friend struct TOptional;
 
-	template <typename, bool>
-	friend struct UE::Core::Private::TOptionalStorage;
-
 private:
 	explicit FIntrusiveUnsetOptionalState() = default;
 };
