[Scheduler]
- Remove scheduler TLS dependencies on other systems to avoid init and shutdown ordering problems

diff --git a/Engine/Source/Runtime/Core/Private/Async/Fundamental/Scheduler.cpp b/Engine/Source/Runtime/Core/Private/Async/Fundamental/Scheduler.cpp
--- a/Engine/Source/Runtime/Core/Private/Async/Fundamental/Scheduler.cpp
+++ b/Engine/Source/Runtime/Core/Private/Async/Fundamental/Scheduler.cpp
@@ -4,8 +4,10 @@
 #include "Async/Fundamental/Task.h"
 #include "Async/TaskTrace.h"
 #include "Async/UniqueLock.h"
+#include "Containers/ConsumeAllMpmcQueue.h"
 #include "Logging/LogMacros.h"
 #include "ProfilingDebugging/CsvProfiler.h"
+#include "HAL/MallocAnsi.h"
 #include "HAL/RunnableThread.h"
 #include "Misc/ScopeLock.h"
 #include "Misc/CommandLine.h"
@@ -22,12 +24,111 @@ namespace LowLevelTasks
 {
 	DEFINE_LOG_CATEGORY(LowLevelTasks);
 
-	UE::FMutex FSchedulerTls::Mutex;
-	TArray<FSchedulerTls::FTlsValues*>  FSchedulerTls::ThreadTlsValues;
 	thread_local FSchedulerTls::FTlsValuesHolder FSchedulerTls::TlsValuesHolder;
 	thread_local FTask* FTask::ActiveTask = nullptr;
 	thread_local bool Private::FOversubscriptionTls::bIsOversubscriptionAllowed = false;
 
+	// SchedulerTls are allocated very early at thread creation and we need an allocator
+	// that is OK with being called that early.
+	void* FSchedulerTls::FTlsValues::operator new(size_t Size)
+	{
+		return AnsiMalloc(Size, PLATFORM_CACHE_LINE_SIZE);
+	}
+
+	void FSchedulerTls::FTlsValues::operator delete(void* Ptr)
+	{
+		AnsiFree(Ptr);
+	}
+
+	struct FTlsValuesAllocator
+	{
+		typedef SIZE_T SizeType;
+		enum { NeedsElementType = false };
+
+		static void* Malloc(SIZE_T Count, uint32 Alignment = DEFAULT_ALIGNMENT)
+		{
+			return AnsiMalloc(Count, Alignment);
+		}
+
+		static void Free(void* Ptr)
+		{
+			AnsiFree(Ptr);
+		}
+	};
+
+	class FSchedulerTls::FImpl
+	{
+	public:
+		static UE::FMutex                 ThreadTlsValuesMutex;
+		static FSchedulerTls::FTlsValues* ThreadTlsValues;
+
+		static UE::TConsumeAllMpmcQueue<FSchedulerTls::FTlsValues*, FTlsValuesAllocator> PendingInsertTlsValues;
+		static UE::TConsumeAllMpmcQueue<FSchedulerTls::FTlsValues*, FTlsValuesAllocator> PendingDeleteTlsValues;
+
+		static void ProcessPendingTlsValuesNoLock()
+		{
+			PendingInsertTlsValues.ConsumeAllLifo(
+				[](FTlsValues* TlsValues)
+				{
+						TlsValues->LinkHead(ThreadTlsValues);
+				}
+			);
+
+			PendingDeleteTlsValues.ConsumeAllLifo(
+				[](FTlsValues* TlsValues)
+				{
+						TlsValues->Unlink();
+						delete TlsValues;
+				}
+			);
+		}
+	};
+
+	UE::FMutex                 FSchedulerTls::FImpl::ThreadTlsValuesMutex;
+	FSchedulerTls::FTlsValues* FSchedulerTls::FImpl::ThreadTlsValues = nullptr;
+
+	UE::TConsumeAllMpmcQueue<FSchedulerTls::FTlsValues*, FTlsValuesAllocator> FSchedulerTls::FImpl::PendingInsertTlsValues;
+	UE::TConsumeAllMpmcQueue<FSchedulerTls::FTlsValues*, FTlsValuesAllocator> FSchedulerTls::FImpl::PendingDeleteTlsValues;
+
+	FSchedulerTls::FTlsValuesHolder::FTlsValuesHolder()
+	{
+#if 0
+		while (!FPlatformMisc::IsDebuggerPresent())
+		{
+			FPlatformProcess::SleepNoStats(0.100f);
+		}
+#endif
+		TlsValues = new FTlsValues();
+
+		if (FImpl::ThreadTlsValuesMutex.TryLock())
+		{
+			FImpl::ProcessPendingTlsValuesNoLock();
+			TlsValues->LinkHead(FImpl::ThreadTlsValues);
+			FImpl::ThreadTlsValuesMutex.Unlock();
+		}
+		else
+		{
+			FImpl::PendingInsertTlsValues.ProduceItem(TlsValues);
+		}
+	}
+
+	FSchedulerTls::FTlsValuesHolder::~FTlsValuesHolder()
+	{
+		if (FImpl::ThreadTlsValuesMutex.TryLock())
+		{
+			FImpl::ProcessPendingTlsValuesNoLock();
+			TlsValues->Unlink();
+			FImpl::ThreadTlsValuesMutex.Unlock();
+			delete TlsValues;
+		}
+		else
+		{
+			FImpl::PendingDeleteTlsValues.ProduceItem(TlsValues);
+		}
+
+		TlsValues = nullptr;
+	}
+
 	FORCENOINLINE bool& Private::FOversubscriptionTls::GetIsOversubscriptionAllowedRef()
 	{
 		return bIsOversubscriptionAllowed;
@@ -381,7 +482,8 @@ namespace LowLevelTasks
 
 	bool FSchedulerTls::HasPendingWakeUp() const
 	{
-		UE::TUniqueLock ScopedLock(Mutex);
+		UE::TUniqueLock ScopedLock(FImpl::ThreadTlsValuesMutex);
+		FImpl::ProcessPendingTlsValuesNoLock();
 
 #if PLATFORM_SUPPORTS_ASYMMETRIC_FENCES
 		// Heavy barrier since bPendingWakeUp is only written to with a relaxed write, we need all cores to flush their store buffer to memory
@@ -391,10 +493,11 @@ namespace LowLevelTasks
 		constexpr std::memory_order MemoryOrder = std::memory_order_acquire;
 #endif
 
-		for (const FTlsValues* ThreadTls : ThreadTlsValues)
+		for (FTlsValues::TConstIterator It(FImpl::ThreadTlsValues); It; It.Next())
 		{
-			if (ThreadTls->ActiveScheduler != this &&
-				ThreadTls->bPendingWakeUp.load(MemoryOrder))
+			const FTlsValues& TlsValue = *It;
+			if (TlsValue.ActiveScheduler != this &&
+				TlsValue.bPendingWakeUp.load(MemoryOrder))
 			{
 				return true;
 			}
diff --git a/Engine/Source/Runtime/Core/Public/Async/Fundamental/Scheduler.h b/Engine/Source/Runtime/Core/Public/Async/Fundamental/Scheduler.h
--- a/Engine/Source/Runtime/Core/Public/Async/Fundamental/Scheduler.h
+++ b/Engine/Source/Runtime/Core/Public/Async/Fundamental/Scheduler.h
@@ -11,6 +11,7 @@
 #include "Async/UniqueLock.h"
 #include "Containers/Array.h"
 #include "Containers/ArrayView.h"
+#include "Containers/List.h"
 #include "Containers/ContainerAllocationPolicies.h"
 #include "HAL/Event.h"
 #include "HAL/PlatformAffinity.h"
@@ -39,6 +40,8 @@ namespace LowLevelTasks
 	class FSchedulerTls
 	{
 	protected:
+		class FImpl;
+
 		using FQueueRegistry	= Private::TLocalQueueRegistry<>;
 		using FLocalQueueType	= FQueueRegistry::TLocalQueue;
 
@@ -49,11 +52,7 @@ namespace LowLevelTasks
 			Foreground,
 		};
 
-		struct FTlsValues;
-		static UE::FMutex Mutex;
-		static TArray<FTlsValues*> ThreadTlsValues;
-
-		struct alignas(PLATFORM_CACHE_LINE_SIZE) FTlsValues
+		struct FTlsValues : public TIntrusiveLinkedList<FTlsValues>
 		{
 			FSchedulerTls* ActiveScheduler = nullptr;
 			FLocalQueueType* LocalQueue = nullptr;
@@ -75,26 +74,15 @@ namespace LowLevelTasks
 			{
 				bIsStandbyWorker = bInIsStandbyWorker;
 			}
+
+			static void* operator new(size_t Size);
+			static void operator delete(void* Ptr);
 		};
 
 		struct FTlsValuesHolder
 		{
-			FTlsValuesHolder()
-			{
-				TlsValues = new FTlsValues();
-
-				UE::TUniqueLock ScopedLock(Mutex);
-				ThreadTlsValues.Add(TlsValues);
-			}
-			~FTlsValuesHolder()
-			{
-				{
-					UE::TUniqueLock ScopedLock(Mutex);
-					ThreadTlsValues.Remove(TlsValues);
-				}
-				delete TlsValues;
-				TlsValues = nullptr;
-			}
+			CORE_API FTlsValuesHolder();
+			CORE_API ~FTlsValuesHolder();
 			
 			FTlsValues* TlsValues = nullptr;
 		};
