Added platform properties to detect guaranteed virtual memory alignment in compile time.
Constexpresified PlatformProperties

diff --git a/Engine/Source/Runtime/Core/Public/Android/AndroidPlatformProperties.h b/Engine/Source/Runtime/Core/Public/Android/AndroidPlatformProperties.h
--- a/Engine/Source/Runtime/Core/Public/Android/AndroidPlatformProperties.h
+++ b/Engine/Source/Runtime/Core/Public/Android/AndroidPlatformProperties.h
@@ -19,120 +19,121 @@
 struct FAndroidPlatformProperties
 	: public FGenericPlatformProperties
 {
-	static FORCEINLINE bool HasEditorOnlyData( )
+	static constexpr FORCEINLINE bool HasEditorOnlyData()
 	{
 		return false;
 	}
 
-	static FORCEINLINE const char* PlatformName()
+	static constexpr FORCEINLINE const char* PlatformName()
 	{
 		return "Android";
 	}
 
-	static FORCEINLINE const char* IniPlatformName()
+	static constexpr FORCEINLINE const char* IniPlatformName()
 	{
 		return "Android";
 	}
 
-	static FORCEINLINE const TCHAR* GetRuntimeSettingsClassName()
+	static constexpr FORCEINLINE const TCHAR* GetRuntimeSettingsClassName()
 	{
 		return TEXT("/Script/AndroidRuntimeSettings.AndroidRuntimeSettings");
 	}
 
-	static FORCEINLINE bool IsGameOnly( )
+	static constexpr FORCEINLINE bool IsGameOnly()
 	{
 		return true;
 	}
 
-	static FORCEINLINE bool RequiresCookedData( )
+	static constexpr FORCEINLINE bool RequiresCookedData()
 	{
 		return true;
 	}
 
-	static FORCEINLINE bool SupportsBuildTarget( EBuildTargetType TargetType )
+	static FORCEINLINE bool SupportsBuildTarget(EBuildTargetType TargetType)
 	{
 		return (TargetType == EBuildTargetType::Game);
 	}
 
-	static FORCEINLINE bool SupportsAutoSDK()
+	static constexpr FORCEINLINE bool SupportsAutoSDK()
 	{
 		return true;
 	}
 
-	static FORCEINLINE bool SupportsHighQualityLightmaps()
+	static constexpr FORCEINLINE bool SupportsHighQualityLightmaps()
 	{
 		return true; // always true because of Vulkan
 	}
 
-	static FORCEINLINE bool SupportsLowQualityLightmaps()
+	static constexpr FORCEINLINE bool SupportsLowQualityLightmaps()
 	{
 		return true;
 	}
 
-	static FORCEINLINE bool SupportsDistanceFieldShadows()
+	static constexpr FORCEINLINE bool SupportsDistanceFieldShadows()
 	{
 		return true;
 	}
 
-	static FORCEINLINE bool SupportsTextureStreaming()
+	static constexpr FORCEINLINE bool SupportsTextureStreaming()
 	{
 		return true;
 	}
 
-	static FORCEINLINE bool SupportsMinimize()
+	static constexpr FORCEINLINE bool SupportsMinimize()
 	{
 		return true;
 	}
 
-	static FORCEINLINE bool SupportsQuit() 
+	static constexpr FORCEINLINE bool SupportsQuit()
 	{
 		return true;
 	}
 
-	static FORCEINLINE bool HasFixedResolution()
+	static constexpr FORCEINLINE bool HasFixedResolution()
 	{
 		return true;
 	}
 
-	static FORCEINLINE bool AllowsFramerateSmoothing()
+	static constexpr FORCEINLINE bool AllowsFramerateSmoothing()
 	{
 		return true;
 	}
 
-	static FORCEINLINE bool AllowsCallStackDumpDuringAssert()
+	static constexpr FORCEINLINE bool AllowsCallStackDumpDuringAssert()
 	{
 		return true;
 	}
 
-	static FORCEINLINE bool SupportsAudioStreaming()
+	static constexpr FORCEINLINE bool SupportsAudioStreaming()
 	{
 		return true;
 	}
 
-	static FORCEINLINE bool SupportsMeshLODStreaming()
+	static constexpr FORCEINLINE bool SupportsMeshLODStreaming()
 	{
 		return true;
 	}
 
-	static FORCEINLINE bool SupportsMemoryMappedFiles()
+	static constexpr FORCEINLINE bool SupportsMemoryMappedFiles()
 	{
 		return true;
 	}
 
-	static FORCEINLINE bool SupportsMemoryMappedAudio()
+	static constexpr FORCEINLINE bool SupportsMemoryMappedAudio()
 	{
 		return true;
 	}
 
-	static FORCEINLINE bool SupportsMemoryMappedAnimation()
+	static constexpr FORCEINLINE bool SupportsMemoryMappedAnimation()
 	{
 		return true;
 	}
 
-	static FORCEINLINE int64 GetMemoryMappingAlignment()
+	static constexpr FORCEINLINE int64 GetMemoryMappingAlignment()
 	{
 		// Cook for largest page size available.
 		// Data cooked for 16kb page size will work for 4kb page size firmware.
+		// TODO: THIS IS WRONG!
 		return 16384;
 	}
 };
diff --git a/Engine/Source/Runtime/Core/Public/GenericPlatform/GenericPlatformProperties.h b/Engine/Source/Runtime/Core/Public/GenericPlatform/GenericPlatformProperties.h
--- a/Engine/Source/Runtime/Core/Public/GenericPlatform/GenericPlatformProperties.h
+++ b/Engine/Source/Runtime/Core/Public/GenericPlatform/GenericPlatformProperties.h
@@ -20,7 +20,7 @@ struct FGenericPlatformProperties
 	 *
 	 * @return The physics format name.
 	 */
-	static FORCEINLINE const char* GetPhysicsFormat()
+	static constexpr FORCEINLINE const char* GetPhysicsFormat()
 	{
 		return "Chaos";
 	}
@@ -30,7 +30,7 @@ struct FGenericPlatformProperties
 	 *
 	 * @return true if the platform has Editor-only data, false otherwise.
 	 */
-	static FORCEINLINE bool HasEditorOnlyData()
+	static constexpr FORCEINLINE bool HasEditorOnlyData()
 	{
 		return WITH_EDITORONLY_DATA;
 	}
@@ -49,7 +49,7 @@ struct FGenericPlatformProperties
 	 *
 	 * @return true if this is a game only platform, false otherwise.
 	 */
-	static FORCEINLINE bool IsGameOnly()
+	static constexpr FORCEINLINE bool IsGameOnly()
 	{
 		return UE_GAME;
 	}
@@ -59,7 +59,7 @@ struct FGenericPlatformProperties
 	 *
 	 * @return true if this is a server only platform, false otherwise.
 	 */
-	static FORCEINLINE bool IsServerOnly()
+	static constexpr FORCEINLINE bool IsServerOnly()
 	{
 		return UE_SERVER;
 	}
@@ -69,7 +69,7 @@ struct FGenericPlatformProperties
 	 *
 	 * @return true if this is a client only platform, false otherwise.
 	 */
-	static FORCEINLINE bool IsClientOnly()
+	static constexpr FORCEINLINE bool IsClientOnly()
 	{
 		return !WITH_SERVER_CODE;
 	}
@@ -77,7 +77,7 @@ struct FGenericPlatformProperties
 	/**
 	 *	Gets whether this was a monolithic build or not
 	 */
-	static FORCEINLINE bool IsMonolithicBuild()
+	static constexpr FORCEINLINE bool IsMonolithicBuild()
 	{
 		return IS_MONOLITHIC;
 	}
@@ -85,7 +85,7 @@ struct FGenericPlatformProperties
 	/**
 	 *	Gets whether this was a program or not
 	 */
-	static FORCEINLINE bool IsProgram()
+	static constexpr FORCEINLINE bool IsProgram()
 	{
 		return IS_PROGRAM;
 	}
@@ -95,7 +95,7 @@ struct FGenericPlatformProperties
 	 *
 	 * @return true if the platform is Little Endian, false otherwise.
 	 */
-	static FORCEINLINE bool IsLittleEndian()
+	static constexpr FORCEINLINE bool IsLittleEndian()
 	{
 		return true;
 	}
@@ -127,7 +127,7 @@ struct FGenericPlatformProperties
 	 *
 	 * @return true if cooked data is required, false otherwise.
 	 */
-	static FORCEINLINE bool RequiresCookedData()
+	static constexpr FORCEINLINE bool RequiresCookedData()
 	{
 		return !HasEditorOnlyData();
 	}
@@ -137,7 +137,7 @@ struct FGenericPlatformProperties
 	*
 	* @return true if packaged data is considered secure, false otherwise.
 	*/
-	static FORCEINLINE bool HasSecurePackageFormat()
+	static constexpr FORCEINLINE bool HasSecurePackageFormat()
 	{
 		return false;
 	}
@@ -147,7 +147,7 @@ struct FGenericPlatformProperties
 	 *
 	 * @return true if this platform requires user credentials, false otherwise.
 	 */
-	static FORCEINLINE bool RequiresUserCredentials()
+	static constexpr FORCEINLINE bool RequiresUserCredentials()
 	{
 		return false;
 	}
@@ -158,7 +158,7 @@ struct FGenericPlatformProperties
 	 * @param TargetType The build target to check.
 	 * @return true if the build target is supported, false otherwise.
 	 */
-	static FORCEINLINE bool SupportsBuildTarget( EBuildTargetType TargetType )
+	static FORCEINLINE bool SupportsBuildTarget(EBuildTargetType TargetType)
 	{
 		return true;
 	}
@@ -166,7 +166,7 @@ struct FGenericPlatformProperties
 	/**
 	 * Returns true if platform supports the AutoSDK system
 	 */
-	static FORCEINLINE bool SupportsAutoSDK()
+	static constexpr FORCEINLINE bool SupportsAutoSDK()
 	{
 		return false;
 	}
@@ -176,7 +176,7 @@ struct FGenericPlatformProperties
 	 *
 	 * @return true if gray scale sRGB texture formats are supported.
 	 */
-	static FORCEINLINE bool SupportsGrayscaleSRGB()
+	static constexpr FORCEINLINE bool SupportsGrayscaleSRGB()
 	{
 		return true;
 	}
@@ -186,7 +186,7 @@ struct FGenericPlatformProperties
 	 *
 	 * @return true if multiple instances are supported, false otherwise.
 	 */
-	static FORCEINLINE bool SupportsMultipleGameInstances()
+	static constexpr FORCEINLINE bool SupportsMultipleGameInstances()
 	{
 		return false;
 	}
@@ -196,7 +196,7 @@ struct FGenericPlatformProperties
 	 *
 	 * @return true if windowed mode is supported.
 	 */
-	static FORCEINLINE bool SupportsWindowedMode()
+	static constexpr FORCEINLINE bool SupportsWindowedMode()
 	{
 		return false;
 	}
@@ -204,7 +204,7 @@ struct FGenericPlatformProperties
 	/**
 	 * Whether this platform wants to allow framerate smoothing or not.
 	 */
-	static FORCEINLINE bool AllowsFramerateSmoothing()
+	static constexpr FORCEINLINE bool AllowsFramerateSmoothing()
 	{
 		return true;
 	}
@@ -212,69 +212,78 @@ struct FGenericPlatformProperties
 	/**
 	 * Whether this platform supports streaming audio
 	 */
-	static FORCEINLINE bool SupportsAudioStreaming()
+	static constexpr FORCEINLINE bool SupportsAudioStreaming()
 	{
 		return false;
 	}
 
-	static FORCEINLINE bool SupportsHighQualityLightmaps()
+	static constexpr FORCEINLINE bool SupportsHighQualityLightmaps()
 	{
 		return true;
 	}
 
-	static FORCEINLINE bool SupportsLowQualityLightmaps()
+	static constexpr FORCEINLINE bool SupportsLowQualityLightmaps()
 	{
 		return true;
 	}
 
-	static FORCEINLINE bool SupportsDistanceFieldShadows()
+	static constexpr FORCEINLINE bool SupportsDistanceFieldShadows()
 	{
 		return true;
 	}
 
-	static FORCEINLINE bool SupportsDistanceFieldAO()
+	static constexpr FORCEINLINE bool SupportsDistanceFieldAO()
 	{
 		return true;
 	}
 
-	static FORCEINLINE bool SupportsTextureStreaming()
+	static constexpr FORCEINLINE bool SupportsTextureStreaming()
 	{
 		return true;
 	}
 
-	static FORCEINLINE bool SupportsMeshLODStreaming()
+	static constexpr FORCEINLINE bool SupportsMeshLODStreaming()
 	{
 		return false;
 	}
 
-	static FORCEINLINE bool SupportsMemoryMappedFiles()
+	static constexpr FORCEINLINE bool SupportsMemoryMappedFiles()
 	{
 		return false;
 	}
-	static FORCEINLINE bool SupportsMemoryMappedAudio()
+
+	static constexpr FORCEINLINE bool SupportsMemoryMappedAudio()
 	{
 		return false;
 	}
-	static FORCEINLINE bool SupportsMemoryMappedAnimation()
+
+	static constexpr FORCEINLINE bool SupportsMemoryMappedAnimation()
 	{
 		return false;
 	}
-	static FORCEINLINE int64 GetMemoryMappingAlignment()
+
+	static constexpr FORCEINLINE int64 GetMemoryMappingAlignment()
 	{
 		return 0;
 	}
+
+	// Guaranteed virtual memory alignment on a given platform, regardless of a specific device
+	static constexpr FORCEINLINE int GetMaxSupportedVirtualMemoryAlignment()
+	{
+		return 4096;
+	}
 	
-	static FORCEINLINE bool SupportsRayTracing()
+	static constexpr FORCEINLINE bool SupportsRayTracing()
 	{
 		return false;
 	}
 
-	static FORCEINLINE bool SupportsLumenGI()
+	static constexpr FORCEINLINE bool SupportsLumenGI()
 	{
 		return true;
 	}
 
-	static FORCEINLINE bool SupportsHardwareLZDecompression()
+	static constexpr FORCEINLINE bool SupportsHardwareLZDecompression()
 	{
 		return false;
 	}
@@ -282,39 +291,38 @@ struct FGenericPlatformProperties
 	/**
 	 * Gets whether user settings should override the resolution or not
 	 */
-	static FORCEINLINE bool HasFixedResolution()
+	static constexpr FORCEINLINE bool HasFixedResolution()
 	{
 		return true;
 	}
 
-	static FORCEINLINE bool SupportsMinimize()
+	static constexpr FORCEINLINE bool SupportsMinimize()
 	{
 		return false;
 	}
 
 	// Whether the platform allows an application to quit to the OS
-	static FORCEINLINE bool SupportsQuit()
+	static constexpr FORCEINLINE bool SupportsQuit()
 	{
 		return false;
 	}
 
 	// Whether the platform allows the call stack to be dumped during an assert
-	static FORCEINLINE bool AllowsCallStackDumpDuringAssert()
+	static constexpr FORCEINLINE bool AllowsCallStackDumpDuringAssert()
 	{
 		return IsProgram();
 	}
 
 	// If this platform wants to replace Zlib with a platform-specific version, set the name of the compression format 
 	// plugin (matching its GetCompressionFormatName() function) in an override of this function
-	static FORCEINLINE const char* GetZlibReplacementFormat()
+	static constexpr FORCEINLINE const char* GetZlibReplacementFormat()
 	{
 		return nullptr;
 	}
 	 
 	// Whether the platform requires an original release version to make a patch
-	static FORCEINLINE bool RequiresOriginalReleaseVersionForPatch()
+	static constexpr FORCEINLINE bool RequiresOriginalReleaseVersionForPatch()
 	{
 		return false;
 	}
-
 };
diff --git a/Engine/Source/Runtime/Core/Public/IOS/IOSPlatformProperties.h b/Engine/Source/Runtime/Core/Public/IOS/IOSPlatformProperties.h
--- a/Engine/Source/Runtime/Core/Public/IOS/IOSPlatformProperties.h
+++ b/Engine/Source/Runtime/Core/Public/IOS/IOSPlatformProperties.h
@@ -18,89 +18,97 @@
 struct FIOSPlatformProperties
 	: public FGenericPlatformProperties
 {
-	static FORCEINLINE bool HasEditorOnlyData( )
+	static constexpr FORCEINLINE bool HasEditorOnlyData()
 	{
 		return false;
 	}
 
-	static FORCEINLINE const char* PlatformName( )
+	static constexpr FORCEINLINE const char* PlatformName()
 	{
 		return "IOS";
 	}
 
-	static FORCEINLINE const char* IniPlatformName( )
+	static constexpr FORCEINLINE const char* IniPlatformName()
 	{
 		return "IOS";
 	}
 
-	static FORCEINLINE const TCHAR* GetRuntimeSettingsClassName()
+	static constexpr FORCEINLINE const TCHAR* GetRuntimeSettingsClassName()
 	{
 		return TEXT("/Script/IOSRuntimeSettings.IOSRuntimeSettings");
 	}
 
-	static FORCEINLINE bool IsGameOnly()
+	static constexpr FORCEINLINE bool IsGameOnly()
 	{
 		return true;
 	}
 	
-	static FORCEINLINE bool RequiresCookedData( )
+	static constexpr FORCEINLINE bool RequiresCookedData()
 	{
 		return true;
 	}
     
-	static FORCEINLINE bool SupportsBuildTarget( EBuildTargetType TargetType )
+	static FORCEINLINE bool SupportsBuildTarget(EBuildTargetType TargetType)
 	{
 		return (TargetType == EBuildTargetType::Game);
 	}
 
-	static FORCEINLINE bool SupportsLowQualityLightmaps()
+	static constexpr FORCEINLINE bool SupportsLowQualityLightmaps()
 	{
 		return true;
 	}
 
-	static FORCEINLINE bool SupportsHighQualityLightmaps()
+	static constexpr FORCEINLINE bool SupportsHighQualityLightmaps()
 	{
 		return true;
 	}
 
-	static FORCEINLINE bool SupportsTextureStreaming()
+	static constexpr FORCEINLINE bool SupportsTextureStreaming()
 	{
 		return true;
 	}
 
-	static FORCEINLINE bool SupportsMemoryMappedFiles()
+	static constexpr FORCEINLINE bool SupportsMemoryMappedFiles()
 	{
 		return true;
 	}
-	static FORCEINLINE bool SupportsMemoryMappedAudio()
+
+	static constexpr FORCEINLINE bool SupportsMemoryMappedAudio()
 	{
 		return true;
 	}
-	static FORCEINLINE bool SupportsMemoryMappedAnimation()
+
+	static constexpr FORCEINLINE bool SupportsMemoryMappedAnimation()
 	{
 		return true;
 	}
-	static FORCEINLINE int64 GetMemoryMappingAlignment()
+
+	static constexpr FORCEINLINE int64 GetMemoryMappingAlignment()
+	{
+		return 16384;
+	}
+
+	static constexpr FORCEINLINE int GetMaxSupportedVirtualMemoryAlignment()
 	{
 		return 16384;
 	}
 
-	static FORCEINLINE bool HasFixedResolution()
+	static constexpr FORCEINLINE bool HasFixedResolution()
 	{
 		return true;
 	}
 
-	static FORCEINLINE bool AllowsFramerateSmoothing()
+	static constexpr FORCEINLINE bool AllowsFramerateSmoothing()
 	{
 		return true;
 	}
 
-	static FORCEINLINE bool SupportsAudioStreaming()
+	static constexpr FORCEINLINE bool SupportsAudioStreaming()
 	{
 		return true;
 	}
 	
-	static FORCEINLINE bool SupportsMeshLODStreaming()
+	static constexpr FORCEINLINE bool SupportsMeshLODStreaming()
 	{
 		return true;
 	}
@@ -109,12 +117,12 @@ struct FIOSPlatformProperties
 struct FTVOSPlatformProperties : public FIOSPlatformProperties
 {
 	// @todo breaking change here!
-	static FORCEINLINE const char* PlatformName()
+	static constexpr FORCEINLINE const char* PlatformName()
 	{
 		return "TVOS";
 	}
 
-	static FORCEINLINE const char* IniPlatformName()
+	static constexpr FORCEINLINE const char* IniPlatformName()
 	{
 		return "TVOS";
 	}
@@ -122,12 +130,12 @@ struct FTVOSPlatformProperties : public FIOSPlatformProperties
 
 struct FVisionOSPlatformProperties : public FIOSPlatformProperties
 {
-	static FORCEINLINE const char* PlatformName()
+	static constexpr FORCEINLINE const char* PlatformName()
 	{
 		return "IOS";
 	}
 
-	static FORCEINLINE const char* IniPlatformName()
+	static constexpr FORCEINLINE const char* IniPlatformName()
 	{
 		return "VisionOS";
 	}
diff --git a/Engine/Source/Runtime/Core/Public/Linux/LinuxPlatformProperties.h b/Engine/Source/Runtime/Core/Public/Linux/LinuxPlatformProperties.h
--- a/Engine/Source/Runtime/Core/Public/Linux/LinuxPlatformProperties.h
+++ b/Engine/Source/Runtime/Core/Public/Linux/LinuxPlatformProperties.h
@@ -20,42 +20,42 @@ template<bool HAS_EDITOR_DATA, bool IS_DEDICATED_SERVER, bool IS_CLIENT_ONLY, bo
 struct FLinuxPlatformProperties
 	: public FGenericPlatformProperties
 {
-	static FORCEINLINE bool HasEditorOnlyData( )
+	static constexpr FORCEINLINE bool HasEditorOnlyData()
 	{
 		return HAS_EDITOR_DATA;
 	}
 
-	static FORCEINLINE const char* IniPlatformName( )
+	static constexpr FORCEINLINE const char* IniPlatformName()
 	{
 		return IS_ARM64 ? "LinuxArm64" : "Linux";
 	}
 
-	static FORCEINLINE const TCHAR* GetRuntimeSettingsClassName()
+	static constexpr FORCEINLINE const TCHAR* GetRuntimeSettingsClassName()
 	{
 		return TEXT("/Script/LinuxTargetPlatform.LinuxTargetSettings");
 	}
 
-	static FORCEINLINE bool IsGameOnly( )
+	static constexpr FORCEINLINE bool IsGameOnly()
 	{
 		return UE_GAME;
 	}
 
-	static FORCEINLINE bool IsServerOnly( )
+	static constexpr FORCEINLINE bool IsServerOnly()
 	{
 		return IS_DEDICATED_SERVER;
 	}
 
-	static FORCEINLINE bool IsClientOnly()
+	static constexpr FORCEINLINE bool IsClientOnly()
 	{
 		return IS_CLIENT_ONLY;
 	}
 
-	static FORCEINLINE bool IsArm64()
+	static constexpr FORCEINLINE bool IsArm64()
 	{
 		return IS_ARM64;
 	}
 
-	static FORCEINLINE const char* PlatformName( )
+	static constexpr FORCEINLINE const char* PlatformName()
 	{
 		if (IS_DEDICATED_SERVER)
 		{
@@ -75,22 +75,22 @@ struct FLinuxPlatformProperties
 		return IS_ARM64 ? "LinuxArm64" : "Linux";
 	}
 
-	static FORCEINLINE bool RequiresCookedData( )
+	static constexpr FORCEINLINE bool RequiresCookedData()
 	{
 		return !HAS_EDITOR_DATA;
 	}
 
-	static FORCEINLINE bool HasSecurePackageFormat()
+	static constexpr FORCEINLINE bool HasSecurePackageFormat()
 	{
 		return IS_DEDICATED_SERVER;
 	}
 
-	static FORCEINLINE bool RequiresUserCredentials()
+	static constexpr FORCEINLINE bool RequiresUserCredentials()
 	{
 		return true;
 	}
 
-	static FORCEINLINE bool SupportsAutoSDK()
+	static constexpr FORCEINLINE bool SupportsAutoSDK()
 	{
 // linux cross-compiling / cross-building from windows supports AutoSDK.  But hosted linux doesn't yet.
 #if PLATFORM_WINDOWS
@@ -100,37 +100,37 @@ struct FLinuxPlatformProperties
 #endif
 	}
 
-	static FORCEINLINE bool SupportsMultipleGameInstances( )
+	static constexpr FORCEINLINE bool SupportsMultipleGameInstances()
 	{
 		return true;
 	}
 
-	static FORCEINLINE bool HasFixedResolution()
+	static constexpr FORCEINLINE bool HasFixedResolution()
 	{
 		return false;
 	}
 
-	static FORCEINLINE bool SupportsWindowedMode()
+	static constexpr FORCEINLINE bool SupportsWindowedMode()
 	{
 		return !IS_DEDICATED_SERVER;
 	}
 
-	static FORCEINLINE bool AllowsFramerateSmoothing()
+	static constexpr FORCEINLINE bool AllowsFramerateSmoothing()
 	{
 		return true;
 	}
 
-	static FORCEINLINE bool SupportsRayTracing()
+	static constexpr FORCEINLINE bool SupportsRayTracing()
 	{
 		return true;
 	}
 
-	static FORCEINLINE bool SupportsQuit()
+	static constexpr FORCEINLINE bool SupportsQuit()
 	{
 		return true;
 	}
 
-	static FORCEINLINE float GetVariantPriority()
+	static constexpr FORCEINLINE float GetVariantPriority()
 	{
 		if (IS_DEDICATED_SERVER)
 		{
@@ -150,15 +150,20 @@ struct FLinuxPlatformProperties
 		return 1.0f;
 	}
 
-	static FORCEINLINE bool AllowsCallStackDumpDuringAssert()
+	static constexpr FORCEINLINE bool AllowsCallStackDumpDuringAssert()
 	{
 		return true;
 	}
 
-	static FORCEINLINE bool SupportsAudioStreaming()
+	static constexpr FORCEINLINE bool SupportsAudioStreaming()
 	{
 		return !IsServerOnly();
 	}
+
+	static constexpr FORCEINLINE int64 GetMemoryMappingAlignment()
+	{
+		return 4096;
+	}
 };
 
 #ifdef PROPERTY_HEADER_SHOULD_DEFINE_TYPE
diff --git a/Engine/Source/Runtime/Core/Public/Windows/WindowsPlatformProperties.h b/Engine/Source/Runtime/Core/Public/Windows/WindowsPlatformProperties.h
--- a/Engine/Source/Runtime/Core/Public/Windows/WindowsPlatformProperties.h
+++ b/Engine/Source/Runtime/Core/Public/Windows/WindowsPlatformProperties.h
@@ -13,37 +13,37 @@ template<bool HAS_EDITOR_DATA, bool IS_DEDICATED_SERVER, bool IS_CLIENT_ONLY>
 struct FWindowsPlatformProperties
 	: public FGenericPlatformProperties
 {
-	static FORCEINLINE bool HasEditorOnlyData()
+	static constexpr FORCEINLINE bool HasEditorOnlyData()
 	{
 		return HAS_EDITOR_DATA;
 	}
 
-	static FORCEINLINE const char* IniPlatformName()
+	static constexpr FORCEINLINE const char* IniPlatformName()
 	{
 		return "Windows";
 	}
 
-	static FORCEINLINE const TCHAR* GetRuntimeSettingsClassName()
+	static constexpr FORCEINLINE const TCHAR* GetRuntimeSettingsClassName()
 	{
 		return TEXT("/Script/WindowsTargetPlatform.WindowsTargetSettings");
 	}
 
-	static FORCEINLINE bool IsGameOnly()
+	static constexpr FORCEINLINE bool IsGameOnly()
 	{
 		return UE_GAME;
 	}
 
-	static FORCEINLINE bool IsServerOnly()
+	static constexpr FORCEINLINE bool IsServerOnly()
 	{
 		return IS_DEDICATED_SERVER;
 	}
 
-	static FORCEINLINE bool IsClientOnly()
+	static constexpr FORCEINLINE bool IsClientOnly()
 	{
 		return IS_CLIENT_ONLY;
 	}
 
-	static FORCEINLINE const char* PlatformName()
+	static constexpr FORCEINLINE const char* PlatformName()
 	{
 		if (IS_DEDICATED_SERVER)
 		{
@@ -63,62 +63,62 @@ struct FWindowsPlatformProperties
 		return "Windows";
 	}
 
-	static FORCEINLINE bool RequiresCookedData()
+	static constexpr FORCEINLINE bool RequiresCookedData()
 	{
 		return !HAS_EDITOR_DATA;
 	}
 
-	static FORCEINLINE bool HasSecurePackageFormat()
+	static constexpr FORCEINLINE bool HasSecurePackageFormat()
 	{
 		return IS_DEDICATED_SERVER;
 	}
 
-	static FORCEINLINE bool SupportsMemoryMappedFiles()
+	static constexpr FORCEINLINE bool SupportsMemoryMappedFiles()
 	{
 		return true;
 	}
 
-	static FORCEINLINE bool SupportsAudioStreaming()
+	static constexpr FORCEINLINE bool SupportsAudioStreaming()
 	{
 		return !IsServerOnly();
 	}
 
-	static FORCEINLINE bool SupportsMeshLODStreaming()
+	static constexpr FORCEINLINE bool SupportsMeshLODStreaming()
 	{
 		return !IsServerOnly() && !HasEditorOnlyData();
 	}
 
-	static FORCEINLINE bool SupportsRayTracing()
+	static constexpr FORCEINLINE bool SupportsRayTracing()
 	{
 		return true;
 	}
 
-	static FORCEINLINE bool SupportsGrayscaleSRGB()
+	static constexpr FORCEINLINE bool SupportsGrayscaleSRGB()
 	{
 		return false; // Requires expand from G8 to RGBA
 	}
 
-	static FORCEINLINE bool SupportsMultipleGameInstances()
+	static constexpr FORCEINLINE bool SupportsMultipleGameInstances()
 	{
 		return true;
 	}
 
-	static FORCEINLINE bool SupportsWindowedMode()
+	static constexpr FORCEINLINE bool SupportsWindowedMode()
 	{
 		return true;
 	}
 	
-	static FORCEINLINE bool HasFixedResolution()
+	static constexpr FORCEINLINE bool HasFixedResolution()
 	{
 		return false;
 	}
 
-	static FORCEINLINE bool SupportsQuit()
+	static constexpr FORCEINLINE bool SupportsQuit()
 	{
 		return true;
 	}
 
-	static FORCEINLINE float GetVariantPriority()
+	static constexpr FORCEINLINE float GetVariantPriority()
 	{
 		if (IS_DEDICATED_SERVER)
 		{
@@ -137,6 +137,16 @@ struct FWindowsPlatformProperties
 
 		return 1.0f;
 	}
+
+	static constexpr FORCEINLINE int64 GetMemoryMappingAlignment()
+	{
+		return 4096;
+	}
+
+	static constexpr FORCEINLINE int GetMaxSupportedVirtualMemoryAlignment()
+	{
+		return 65536;
+	}
 };
 
 #ifdef PROPERTY_HEADER_SHOULD_DEFINE_TYPE
