Fix IO Unified Error logging

diff --git a/Engine/Source/Runtime/Core/Private/Experimental/IO/IoStatusError.cpp b/Engine/Source/Runtime/Core/Private/Experimental/IO/IoStatusError.cpp
--- a/Engine/Source/Runtime/Core/Private/Experimental/IO/IoStatusError.cpp
+++ b/Engine/Source/Runtime/Core/Private/Experimental/IO/IoStatusError.cpp
@@ -14,13 +14,8 @@ namespace UE::UnifiedError::IoStore
 		FString ErrorMessage;
 		mutable FUtf8String CachedErrorName;
 		static const FText GenericErrorFormatString;
-	public:
-		FIoStoreErrorDetails(const FStringView& InErrorMessage)
-		{
-			ErrorMessage = InErrorMessage;
-		}
 
-		virtual const FUtf8String& GetErrorName(const FError& Error) const final
+		const FUtf8String& GetErrorName(const FError& Error) const
 		{
 			if (CachedErrorName.IsEmpty() && (ErrorMessage[0] != TEXT('\0')))
 			{
@@ -28,15 +23,22 @@ namespace UE::UnifiedError::IoStore
 			}
 			return CachedErrorName;
 		}
-		virtual const FText GetErrorFormatString(const FError& Error) const final
+
+	public:
+		FIoStoreErrorDetails(const FStringView& InErrorMessage)
+		{
+			ErrorMessage = InErrorMessage;
+		}
+		
+		virtual const FText GetErrorFormatString(const FError& Error) const final override
 		{
 			return GenericErrorFormatString;
 		}
 
 		virtual void GetErrorProperties(const FError& Error, IErrorPropertyExtractor& OutProperties) const override
 		{
-			OutProperties.AddProperty(UTF8TEXTVIEW("ErrorCodeString"), CachedErrorName);
-			OutProperties.AddProperty(UTF8TEXTVIEW("ModuleIdString"), UE::UnifiedError::IoStore::StaticModuleId);
+			OutProperties.AddProperty(UTF8TEXTVIEW("ErrorCodeString"), GetErrorName(Error));
+			OutProperties.AddProperty(UTF8TEXTVIEW("ModuleIdString"), UE::UnifiedError::IoStore::StaticModuleName);
 			OutProperties.AddProperty(UTF8TEXTVIEW("ErrorCode"), Error.GetErrorCode());
 			OutProperties.AddProperty(UTF8TEXTVIEW("ModuleId"), Error.GetModuleId());
 			OutProperties.AddProperty(UTF8TEXTVIEW("IoStoreErrorMessage"), ErrorMessage);
diff --git a/Engine/Source/Runtime/Core/Private/Experimental/UnifiedError/UnifiedError.cpp b/Engine/Source/Runtime/Core/Private/Experimental/UnifiedError/UnifiedError.cpp
--- a/Engine/Source/Runtime/Core/Private/Experimental/UnifiedError/UnifiedError.cpp
+++ b/Engine/Source/Runtime/Core/Private/Experimental/UnifiedError/UnifiedError.cpp
@@ -368,8 +368,8 @@ FString LexToString(const UE::UnifiedError::FError& Error)
 {
 	
 	// TODO: daniel, switch to using structured logging output in the future instead of this
-	return FString::Printf(TEXT("%s:%s - %s"), 
-		UTF8_TO_TCHAR(*Error.GetModuleIdString()),
-		UTF8_TO_TCHAR(*Error.GetErrorCodeString()),
+	return FString::Printf(TEXT("%hs:%hs - %s"),
+		*Error.GetModuleIdString(),
+		*Error.GetErrorCodeString(),
 		*Error.GetErrorMessage().ToString());
 }
\ No newline at end of file
diff --git a/Engine/Source/Runtime/Engine/Private/AssetManager.cpp b/Engine/Source/Runtime/Engine/Private/AssetManager.cpp
--- a/Engine/Source/Runtime/Engine/Private/AssetManager.cpp
+++ b/Engine/Source/Runtime/Engine/Private/AssetManager.cpp
@@ -2124,7 +2124,7 @@ void UAssetManager::OnAssetStateChangeCompleted(TSharedPtr<FStreamableHandle> Bo
 				{
 					UE_LOG(LogAssetManager, Warning, TEXT("OnAssetStateChangeCompleted: Received error (%s) loading asset! Primary asset id:(%s)"), *LexToString(*BoundHandle->GetError()), *PrimaryAssetId.ToString());
 				}
-				UE_LOGFMT(LogAssetManager, Warning, "OnAssetStateChangeCompleted: Received error ({Error}) loading asset! Primary asset id:({PrimaryAssetId})", ("Error", BoundHandle->GetError()), ("PrimaryAssetId", *PrimaryAssetId.ToString()));
+				//UE_LOGFMT(LogAssetManager, Warning, "OnAssetStateChangeCompleted: Received error {Error} loading asset! Primary asset id:({PrimaryAssetId})", ("Error", BoundHandle->GetError()), ("PrimaryAssetId", *PrimaryAssetId.ToString()));
 			}
 			else if (BoundHandle->WasCanceled())
 			{
