Added a compiler define and CVar for managing how GenericPlatformMallocCrash handles allocation requests on non-crashing threads, allowing them to permanently suspend on allocs/frees, suspend on allocs and skip frees, or process as normal. Changed the new default to suspend on allocation and skip frees to allow threads to tear down, which may block crash processing on some platforms.

diff --git a/Engine/Source/Runtime/Core/Private/GenericPlatform/GenericPlatformMallocCrash.cpp b/Engine/Source/Runtime/Core/Private/GenericPlatform/GenericPlatformMallocCrash.cpp
--- a/Engine/Source/Runtime/Core/Private/GenericPlatform/GenericPlatformMallocCrash.cpp
+++ b/Engine/Source/Runtime/Core/Private/GenericPlatform/GenericPlatformMallocCrash.cpp
@@ -2,6 +2,7 @@
 
 #include "GenericPlatform/GenericPlatformMallocCrash.h"
 #include "Containers/ArrayView.h"
+#include "HAL/ConsoleManager.h"
 #include "HAL/PlatformProcess.h"
 #include "HAL/PlatformTLS.h"
 #include "Templates/AlignmentTemplates.h"
@@ -9,6 +10,27 @@
 #include "ProfilingDebugging/MemoryTrace.h"
 #include "Algo/BinarySearch.h"
 
+// Suspending other threads on the first allocator call following a crash may help with debugging by freezing
+// the program state and avoiding further manipulation of a bad program state. However, crash handling on some
+// platforms requires threads to be able to tear down TLS, which in many cases involves deallocating memory.
+// These modes allow threads that are not the one that signaled the crash to skip frees or proceed as normal.
+enum EMallocCrashNonCrashedThreadsMode
+{
+	SuspendAllRequests = 0, // Non-crashed threads suspend on allocation or free
+	SuspendAllocationsAndSkipFrees = 1, // Non-crashed threads suspend on allocation and do nothing on free
+	ProcessAllRequests = 2, // Non-crashed threads process allocations and frees as normal
+};
+
+// NOTE: Crashes prior to loading config/hotfixes will abide by the default mode, but this CVar may help recover crashes beyond that point
+static int32 GMallocCrashNonCrashedThreadsMode = SuspendAllocationsAndSkipFrees;
+static FAutoConsoleVariableRef CVarMallocCrashNonCrashedThreadsMode(
+	TEXT("MallocCrash.NonCrashedThreadsMode"),
+	GMallocCrashNonCrashedThreadsMode,
+	TEXT("How to process requests during a crash from threads other than the original crashed thread. ")
+	TEXT("0 = Non-crashed threads suspend on allocation or free. ")
+	TEXT("1 = Non-crashed threads suspend on allocation and do nothing on free. ")
+	TEXT("2 = Non-crashed threads process allocations and frees as normal."),
+	ECVF_Default);
 
 /** Describes a pool. */
 struct FPoolDesc
@@ -310,7 +332,7 @@ void* FGenericPlatformMallocCrash::Malloc( SIZE_T Size, uint32 Alignment )
 		FPlatformMisc::LowLevelOutputDebugString( TEXT( "Alignment > 16 is not supported\n" ) );
 	}
 
-	if( IsOnCrashedThread() )
+	if( CheckThreadForAllocation() )
 	{
 		FMallocCrashPool* Pool = ChoosePoolForSize( Size32 );
 		if( Pool )
@@ -349,7 +371,7 @@ void* FGenericPlatformMallocCrash::Malloc( SIZE_T Size, uint32 Alignment )
 
 void* FGenericPlatformMallocCrash::Realloc( void* Ptr, SIZE_T NewSize, uint32 Alignment )
 {
-	if( IsOnCrashedThread() )
+	if( CheckThreadForAllocation() )
 	{
 		void* Result = nullptr;
 		if( Ptr && NewSize )
@@ -410,7 +432,7 @@ void* FGenericPlatformMallocCrash::Realloc( void* Ptr, SIZE_T NewSize, uint32 Al
 
 void FGenericPlatformMallocCrash::Free( void* Ptr )
 {
-	if( IsOnCrashedThread() )
+	if( CheckThreadForFree() )
 	{
 		if( IsPtrInSmallPool(Ptr) )
 		{
@@ -456,22 +478,54 @@ void FGenericPlatformMallocCrash::PrintPoolsUsage()
 }
 
 PRAGMA_DISABLE_UNREACHABLE_CODE_WARNINGS
-bool FGenericPlatformMallocCrash::IsOnCrashedThread() const
+bool FGenericPlatformMallocCrash::CheckThreadForAllocation() const
 {
-	// Suspend threads other than the crashed one to prevent serious memory errors.
-	// Only the crashed thread can do anything meaningful from here anyway.
-	if( CrashedThreadId == FPlatformTLS::GetCurrentThreadId() )
+	// Crashing thread always processes allocations as normal, other threads abide by the current mode
+	if (CrashedThreadId == FPlatformTLS::GetCurrentThreadId())
 	{
 		return true;
 	}
 	else
 	{
-		FPlatformProcess::SleepInfinite();
-		return false;
+		switch (GMallocCrashNonCrashedThreadsMode)
+		{
+			case EMallocCrashNonCrashedThreadsMode::SuspendAllocationsAndSkipFrees:
+			case EMallocCrashNonCrashedThreadsMode::SuspendAllRequests:
+				FPlatformProcess::SleepInfinite();
+				return false;
+
+			case EMallocCrashNonCrashedThreadsMode::ProcessAllRequests:
+			default:
+				return true;
+		}
 	}
 }
-PRAGMA_RESTORE_UNREACHABLE_CODE_WARNINGS
 
+bool FGenericPlatformMallocCrash::CheckThreadForFree() const
+{
+	// Crashing thread always processes frees as normal, other threads abide by the current mode
+	if (CrashedThreadId == FPlatformTLS::GetCurrentThreadId())
+	{
+		return true;
+	}
+	else
+	{
+		switch (GMallocCrashNonCrashedThreadsMode)
+		{
+			case EMallocCrashNonCrashedThreadsMode::SuspendAllRequests:
+				FPlatformProcess::SleepInfinite();
+				return false;
+
+			case EMallocCrashNonCrashedThreadsMode::SuspendAllocationsAndSkipFrees:
+				return false;
+
+			case EMallocCrashNonCrashedThreadsMode::ProcessAllRequests:
+			default:
+				return true;
+		}
+	}
+}
+PRAGMA_RESTORE_UNREACHABLE_CODE_WARNINGS
 
 bool FGenericPlatformMallocCrash::IsPtrInLargePool( void* Ptr ) const
 {
diff --git a/Engine/Source/Runtime/Core/Public/GenericPlatform/GenericPlatformMallocCrash.h b/Engine/Source/Runtime/Core/Public/GenericPlatform/GenericPlatformMallocCrash.h
--- a/Engine/Source/Runtime/Core/Public/GenericPlatform/GenericPlatformMallocCrash.h
+++ b/Engine/Source/Runtime/Core/Public/GenericPlatform/GenericPlatformMallocCrash.h
@@ -108,7 +108,8 @@ public:
 
 protected:
 	/** Whether it is safe to call crash malloc's methods. */
-	bool IsOnCrashedThread() const;
+	bool CheckThreadForAllocation() const;
+	bool CheckThreadForFree() const;
 
 	bool IsPtrInLargePool(void* Ptr) const;
 	bool IsPtrInSmallPool(void* Ptr) const;
