Trivial amendment to last change: remove ProcessAllRequests mode since having a debugger attached makes it clear most threads run out of memory pretty quickly, which may interfere with the crash reporting.

#rnx

diff --git a/Engine/Source/Runtime/Core/Private/GenericPlatform/GenericPlatformMallocCrash.cpp b/Engine/Source/Runtime/Core/Private/GenericPlatform/GenericPlatformMallocCrash.cpp
--- a/Engine/Source/Runtime/Core/Private/GenericPlatform/GenericPlatformMallocCrash.cpp
+++ b/Engine/Source/Runtime/Core/Private/GenericPlatform/GenericPlatformMallocCrash.cpp
@@ -17,8 +17,7 @@
 enum EMallocCrashNonCrashedThreadsMode
 {
 	SuspendAllRequests = 0, // Non-crashed threads suspend on allocation or free
-	SuspendAllocationsAndSkipFrees = 1, // Non-crashed threads suspend on allocation and do nothing on free
-	ProcessAllRequests = 2, // Non-crashed threads process allocations and frees as normal
+	SuspendAllocationsAndSkipFrees = 1 // Non-crashed threads suspend on allocation and do nothing on free
 };
 
 // NOTE: Crashes prior to loading config/hotfixes will abide by the default mode, but this CVar may help recover crashes beyond that point
@@ -28,8 +27,7 @@ static FAutoConsoleVariableRef CVarMallocCrashNonCrashedThreadsMode(
 	GMallocCrashNonCrashedThreadsMode,
 	TEXT("How to process requests during a crash from threads other than the original crashed thread. ")
 	TEXT("0 = Non-crashed threads suspend on allocation or free. ")
-	TEXT("1 = Non-crashed threads suspend on allocation and do nothing on free. ")
-	TEXT("2 = Non-crashed threads process allocations and frees as normal."),
+	TEXT("1 = Non-crashed threads suspend on allocation and do nothing on free. "),
 	ECVF_Default);
 
 /** Describes a pool. */
@@ -494,9 +492,8 @@ bool FGenericPlatformMallocCrash::CheckThreadForAllocation() const
 				FPlatformProcess::SleepInfinite();
 				return false;
 
-			case EMallocCrashNonCrashedThreadsMode::ProcessAllRequests:
 			default:
-				return true;
+				return false;
 		}
 	}
 }
@@ -517,11 +514,8 @@ bool FGenericPlatformMallocCrash::CheckThreadForFree() const
 				return false;
 
 			case EMallocCrashNonCrashedThreadsMode::SuspendAllocationsAndSkipFrees:
-				return false;
-
-			case EMallocCrashNonCrashedThreadsMode::ProcessAllRequests:
 			default:
-				return true;
+				return false;
 		}
 	}
 }
