[AutoRTFM] Add a single-threaded mode to the concurrent linear array allocator.

Some users of this allocator (e.g. StateTree) do not actually use the allocator in a multi-threaded context, and simply prefer its performance characteristics. This CL allows single-threaded users of the allocator to opt out of using atomic variables for tracking allocations; this makes the allocator AutoRTFM-safe without incurring any extra overhead for dealing with atomics.

diff --git a/Engine/Source/Runtime/Core/Public/Experimental/ConcurrentLinearAllocator.h b/Engine/Source/Runtime/Core/Public/Experimental/ConcurrentLinearAllocator.h
--- a/Engine/Source/Runtime/Core/Public/Experimental/ConcurrentLinearAllocator.h
+++ b/Engine/Source/Runtime/Core/Public/Experimental/ConcurrentLinearAllocator.h
@@ -1,7 +1,10 @@
 // Copyright Epic Games, Inc. All Rights Reserved.
 
 #pragma once
+
 #include <atomic>
+#include <type_traits>
+
 #include "HAL/UnrealMemory.h"
 #include "HAL/LowLevelMemTracker.h"
 #include "HAL/MallocBinnedCommon.h"
@@ -10,7 +13,6 @@
 #include "Containers/LockFreeFixedSizeAllocator.h"
 #include "Misc/MemStack.h"
 
-
 #if PLATFORM_HAS_ASAN_INCLUDE
 #	include <sanitizer/asan_interface.h>
 #	if defined(__SANITIZE_ADDRESS__)
@@ -152,21 +154,37 @@ struct FDefaultBlockAllocationTag
 	using Allocator = TBlockAllocationLockFreeCache<BlockSize, FAlignedAllocator>;
 };
 
-// This concurrent fast linear Allocator can be used for temporary allocations that are usually produced and consumed on different threads and within the lifetime of a frame.
-// Although the lifetime of any individual Allocation is not hard tied to a frame (tracking is done using the FBlockHeader::NumAllocations atomic variable) the Application will eventually run OOM if Allocations are not cleaned up in a timely manner.
-// THere is a fastpath version of the allocator that skips Allocationheaders by aligning the BlockHeader with the Blocksize so that it easily can be found by AligingDOwn the Adress of the Allocation itself. 
-// The Allocator works by allocating a larger block in TLS which has a Header at the front which contains the Atomic and all allocations are than allocated from this block:
-//
-// ------------------------------------------------------------------------------------------------------------------------------------------------------------------------------------------------------------
-// | FBlockHeader(atomic counter etc.) | Alignment Waste | FAllocationHeader(size, optional) | Memory used for Allocation | Alignment Waste | FAllocationHeader(size, optional) | Memory used for Allocation | FreeSpace ... 
-// ------------------------------------------------------------------------------------------------------------------------------------------------------------------------------------------------------------
-template<typename BlockAllocationTag>
-class TConcurrentLinearAllocator
+/*************************************************************************************************************************
+ * This fast linear allocator can be used for temporary allocations, and is best suited for allocations that are 
+ * produced and consumed on different threads and within the lifetime of a frame. Although the lifetime of any 
+ * individual allocation is not hard-tied to a frame (tracking is done using the FBlockHeader::NumAllocations atomic 
+ * variable), the application will eventually run OOM if allocations are not cleaned up in a timely manner.
+ * 
+ * There is a fast-path version of the allocator that skips AllocationHeaders by aligning the BlockHeader with the 
+ * BlockSize, so that headers can easily be found by AligningDown the address of the Allocation itself. 
+ * The allocator works by allocating a larger block in TLS which has a Header at the front which contains the atomic,
+ * and all allocations are than allocated from this block:
+ *
+ * ------------------------------------------------------------------------------------------------------------------------------------------------------------------------------------------------------------
+ * | FBlockHeader(atomic counter etc.) | Alignment Waste | FAllocationHeader(size, optional) | Memory used for Allocation | Alignment Waste | FAllocationHeader(size, optional) | Memory used for Allocation | FreeSpace ... 
+ * ------------------------------------------------------------------------------------------------------------------------------------------------------------------------------------------------------------
+ *
+ * The allocator is most often used concurrently, but also supports single-threaded use cases, so it can be used for an
+ * array scratchpad.
+ */
+enum class ELinearAllocatorThreadPolicy
+{
+	ThreadSafe,
+	NotThreadSafe
+};
+
+template <typename BlockAllocationTag, ELinearAllocatorThreadPolicy ThreadPolicy>
+class TLinearAllocatorBase
 {
 	static constexpr bool SupportsFastPath = ((BlockAllocationTag::BlockSize <= (64 * 1024)) && (FPlatformProperties::GetMaxSupportedVirtualMemoryAlignment() >= 64 * 1024))
-		&& FMath::IsPowerOfTwo(BlockAllocationTag::BlockSize) //Aligndown only works with Pow2
+		&& FMath::IsPowerOfTwo(BlockAllocationTag::BlockSize) // AlignDown only works with Pow2
 		&& !IS_ASAN_ENABLED && !BlockAllocationTag::RequiresAccurateSize // Only enabled when not using ASAN or when the Size of an allocation is not required 
-		&& BlockAllocationTag::Allocator::SupportsAlignment; //The allocator needs to align the BlockHeader by Blocksize to find the BlockHeader
+		&& BlockAllocationTag::Allocator::SupportsAlignment; // The allocator needs to align the BlockHeader by BlockSize to find the BlockHeader
 
 	LLM(static constexpr ELLMAllocType LLMAllocationType = BlockAllocationTag::Allocator::UsesFMalloc ? ELLMAllocType::FMalloc : ELLMAllocType::None;)
 
@@ -202,7 +220,7 @@ class TConcurrentLinearAllocator
 
 	static FORCEINLINE FAllocationHeader* GetAllocationHeader(void* Pointer)
 	{
-		if(SupportsFastPath)
+		if (SupportsFastPath)
 		{
 			return nullptr;
 		}
@@ -220,7 +238,43 @@ class TConcurrentLinearAllocator
 		{
 		}
 
-		std::atomic_uint NumAllocations { UINT_MAX }; //this is shared between threads and tracks the number of live allocations (plus UINT_MAX) and we will fix this up when we close a block.
+		struct FPlainUInt
+		{
+			/** Returns the original number of allocations, before subtraction occurs. */
+			unsigned int FetchSub(unsigned int N, std::memory_order)
+			{
+				unsigned int OriginalValue = Value;
+				Value -= N;
+				return OriginalValue;
+			}
+
+			void Store(unsigned int N, std::memory_order)
+			{
+				Value = N;
+			}
+
+			unsigned int Value;
+		};
+
+		struct FAtomicUInt
+		{
+			/** Returns the original number of allocations, before subtraction occurs. */
+			unsigned int FetchSub(unsigned int N, std::memory_order Order)
+			{
+				return Value.fetch_sub(N, Order);
+			}
+
+			void Store(unsigned int N, std::memory_order Order)
+			{
+				return Value.store(N, Order);
+			}
+
+			std::atomic_uint Value;
+		};
+
+		using NumAllocationsType = std::conditional_t<ThreadPolicy == ELinearAllocatorThreadPolicy::ThreadSafe, FAtomicUInt, FPlainUInt>;
+
+		NumAllocationsType NumAllocations { .Value = UINT_MAX }; //this is shared between threads and tracks the number of live allocations (plus UINT_MAX) and we will fix this up when we close a block.
 		uint8 Padding[PLATFORM_CACHE_LINE_SIZE - sizeof(std::atomic_uint)]; //avoid false sharing
 		uintptr_t NextAllocationPtr; //this is the next address we are trying to allocate from.
 		unsigned int Num = 0; //this tracks the TLS local Number of allocation from a given Block
@@ -239,7 +293,7 @@ class TConcurrentLinearAllocator
 
 				// on the allocating side we only need to do a single atomic to reduce contention with the deletions
 				// this will leave the atomic in a state where it only counts the number of live allocations (before it was based of UINT_MAX)
-				if (Header->NumAllocations.fetch_sub(DeltaCount, std::memory_order_acq_rel) == DeltaCount)
+				if (Header->NumAllocations.FetchSub(DeltaCount, std::memory_order_acq_rel) == DeltaCount)
 				{
 					//if all allocations are already freed we can reuse the Block again
 					Header->~FBlockHeader();
@@ -340,7 +394,7 @@ public:
 
 					uintptr_t LargeAlignedOffset = Align(LargeHeader->NextAllocationPtr, Alignment);
 					LargeHeader->NextAllocationPtr = uintptr_t(LargeHeader) + HeaderSize + Size + Alignment;
-					LargeHeader->NumAllocations.store(1, std::memory_order_release);
+					LargeHeader->NumAllocations.Store(1, std::memory_order_release);
 
 					checkSlow(LargeAlignedOffset + Size <= LargeHeader->NextAllocationPtr);
 
@@ -382,7 +436,7 @@ public:
 
 					uintptr_t LargeAlignedOffset = Align(LargeHeader->NextAllocationPtr, Alignment);
 					LargeHeader->NextAllocationPtr = uintptr_t(LargeHeader) + HeaderSize + Size + Alignment;
-					LargeHeader->NumAllocations.store(1, std::memory_order_release);
+					LargeHeader->NumAllocations.Store(1, std::memory_order_release);
 
 					checkSlow(LargeAlignedOffset + Size <= LargeHeader->NextAllocationPtr);
 					FAllocationHeader* AllocationHeader = new (reinterpret_cast<FAllocationHeader*>(LargeAlignedOffset) - 1) FAllocationHeader(LargeHeader, Size);
@@ -402,7 +456,7 @@ public:
 
 		// on the allocating side we only need to do a single atomic to reduce contention with the deletions
 		// this will leave the atomic in a state where it only counts the number of live allocations (before it was based of UINT_MAX)
-		if (Header->NumAllocations.fetch_sub(DeltaCount, std::memory_order_acq_rel) == DeltaCount)
+		if (Header->NumAllocations.FetchSub(DeltaCount, std::memory_order_acq_rel) == DeltaCount)
 		{
 			//if all allocations are already freed we can reuse the Block again
 			Header->~FBlockHeader();
@@ -425,7 +479,7 @@ public:
 				FBlockHeader* Header = reinterpret_cast<FBlockHeader*>(AlignDown(Pointer, BlockAllocationTag::BlockSize));
 
 				// this deletes complete blocks when the last allocation is freed
-				if (Header->NumAllocations.fetch_sub(1, std::memory_order_acq_rel) == 1)
+				if (Header->NumAllocations.FetchSub(1, std::memory_order_acq_rel) == 1)
 				{
 					const uintptr_t NextAllocationPtr = Header->NextAllocationPtr;
 					Header->~FBlockHeader();
@@ -443,7 +497,7 @@ public:
 				ASAN_POISON_MEMORY_REGION( AllocationHeader, sizeof(FAllocationHeader) + AllocationHeader->GetAllocationSize() );
 
 				// this deletes complete blocks when the last allocation is freed
-				if (Header->NumAllocations.fetch_sub(1, std::memory_order_acq_rel) == 1)
+				if (Header->NumAllocations.FetchSub(1, std::memory_order_acq_rel) == 1)
 				{
 					const uintptr_t NextAllocationPtr = Header->NextAllocationPtr;
 					Header->~FBlockHeader();
@@ -492,7 +546,11 @@ public:
 	}
 };
 
-using FConcurrentLinearAllocator = TConcurrentLinearAllocator<FDefaultBlockAllocationTag>;
+template <typename T>
+using TConcurrentLinearAllocator = TLinearAllocatorBase<T, ELinearAllocatorThreadPolicy::ThreadSafe>;
+
+using FConcurrentLinearAllocator = TLinearAllocatorBase<FDefaultBlockAllocationTag, ELinearAllocatorThreadPolicy::ThreadSafe>;
+using FNonconcurrentLinearAllocator = TLinearAllocatorBase<FDefaultBlockAllocationTag, ELinearAllocatorThreadPolicy::NotThreadSafe>;
 
 template<typename ObjectType, typename BlockAllocationTag = FDefaultBlockAllocationTag>
 class TConcurrentLinearObject
@@ -547,8 +605,8 @@ namespace UE::Core::Private
 	[[noreturn]] CORE_API void OnInvalidConcurrentLinearArrayAllocatorNum(int32 NewNum, SIZE_T NumBytesPerElement);
 }
 
-template<typename BlockAllocationTag>
-class TConcurrentLinearArrayAllocator
+template <typename BlockAllocationTag, ELinearAllocatorThreadPolicy ThreadPolicy>
+class TLinearArrayAllocatorBase
 {
 public:
 	using SizeType = int32;
@@ -565,7 +623,7 @@ public:
 		{
 			if (Data)
 			{
-				TConcurrentLinearAllocator<BlockAllocationTag>::Free(Data);
+				TLinearAllocatorBase<BlockAllocationTag, ThreadPolicy>::Free(Data);
 			}
 		}
 		/**
@@ -579,7 +637,7 @@ public:
 
 			if (Data)
 			{
-				TConcurrentLinearAllocator<BlockAllocationTag>::Free(Data);
+				TLinearAllocatorBase<BlockAllocationTag, ThreadPolicy>::Free(Data);
 			}
 
 			Data = Other.Data;
@@ -600,7 +658,7 @@ public:
 			}
 
 			static_assert(alignof(ElementType) <= BlockAllocationTag::Allocator::MaxAlignment);
-			Data = (ElementType*)TConcurrentLinearAllocator<BlockAllocationTag>::Realloc(Data, NewMax * NumBytesPerElement, alignof(ElementType));
+			Data = (ElementType*)TLinearAllocatorBase<BlockAllocationTag, ThreadPolicy>::Realloc(Data, NewMax * NumBytesPerElement, alignof(ElementType));
 		}
 		SizeType CalculateSlackReserve(SizeType NewMax, SIZE_T NumBytesPerElement) const
 		{
@@ -640,6 +698,12 @@ public:
 	using ForAnyElementType = ForElementType<FScriptContainerElement>;
 };
 
+template <typename BlockAllocationTag>
+using TConcurrentLinearArrayAllocator = TLinearArrayAllocatorBase<BlockAllocationTag, ELinearAllocatorThreadPolicy::ThreadSafe>;
+
+template <typename BlockAllocationTag>
+using TNonconcurrentLinearArrayAllocator = TLinearArrayAllocatorBase<BlockAllocationTag, ELinearAllocatorThreadPolicy::NotThreadSafe>;
+
 template <typename BlockAllocationTag>
 struct TAllocatorTraits<TConcurrentLinearArrayAllocator<BlockAllocationTag>> : TAllocatorTraitsBase<TConcurrentLinearArrayAllocator<BlockAllocationTag>>
 {
@@ -660,6 +724,8 @@ using FConcurrentLinearBitArrayAllocator = TConcurrentLinearBitArrayAllocator<FD
 using FConcurrentLinearSparseArrayAllocator = TConcurrentLinearSparseArrayAllocator<FDefaultBlockAllocationTag>;
 using FConcurrentLinearSetAllocator = TConcurrentLinearSetAllocator<FDefaultBlockAllocationTag>;
 
+using FNonconcurrentLinearArrayAllocator = TNonconcurrentLinearArrayAllocator<FDefaultBlockAllocationTag>;
+
 //The BulkObjectAllocator can be used to atomically destroy all allocated Objects, it will properly call every destructor before deleting the memory as well
 template<typename BlockAllocationTag>
 class TConcurrentLinearBulkObjectAllocator
