[AutoRTFM] Convert NameToDisplayStringExemptionLock to a transactionally-safe lock.

This lock can be reached from the StateTree unit tests.

diff --git a/Engine/Source/Runtime/Core/Private/UObject/UnrealNames.cpp b/Engine/Source/Runtime/Core/Private/UObject/UnrealNames.cpp
--- a/Engine/Source/Runtime/Core/Private/UObject/UnrealNames.cpp
+++ b/Engine/Source/Runtime/Core/Private/UObject/UnrealNames.cpp
@@ -14,6 +14,7 @@
 #include "Misc/CString.h"
 #include "Misc/Crc.h"
 #include "Misc/StringBuilder.h"
+#include "Misc/TransactionallySafeRWLock.h"
 #include "Containers/UnrealString.h"
 #include "Containers/Utf8String.h"
 #include "UObject/NameTypes.h"
@@ -2617,7 +2618,7 @@ FString FNameEntrySerialized::GetPlainNameString() const
 -----------------------------------------------------------------------------*/
 
 static TArray<FString> NameToDisplayStringExemptions;
-static FRWLock NameToDisplayStringExemptionLock;
+static FTransactionallySafeRWLock NameToDisplayStringExemptionLock;
 
 void FName::Reserve(uint32 NumBytes, uint32 NumNames)
 {
@@ -2697,7 +2698,7 @@ FString FName::NameToDisplayString( const FString& InDisplayName, const bool bIs
 	// Characters that can follow the end of a word
 	constexpr FAsciiSet EndOfWord(" .,;:?!)");
 
-	FReadScopeLock Lock(NameToDisplayStringExemptionLock);
+	UE::TReadScopeLock Lock(NameToDisplayStringExemptionLock);
 
 	const FStringView DisplayNameStringView = InDisplayName;
 	const int32 DisplayNameLength = InDisplayName.Len();
@@ -2860,14 +2861,14 @@ FString FName::NameToDisplayString( const FString& InDisplayName, const bool bIs
 
 void FName::AddNameToDisplayStringExemption(const FString& InExemption)
 {
-	FWriteScopeLock Lock(NameToDisplayStringExemptionLock);
+	UE::TWriteScopeLock Lock(NameToDisplayStringExemptionLock);
 
 	NameToDisplayStringExemptions.Add(InExemption);
 }
 
 void FName::RemoveNameToDisplayStringExemption(const FString& InExemption)
 {
-	FWriteScopeLock Lock(NameToDisplayStringExemptionLock);
+	UE::TWriteScopeLock Lock(NameToDisplayStringExemptionLock);
 
 	NameToDisplayStringExemptions.Remove(InExemption);
 }
