Fixed Algo::RemoveIf calling the predicate too many times in some cases.  Added some tests.

diff --git a/Engine/Source/Runtime/Core/Private/Containers/Algo/AlgosTest.cpp b/Engine/Source/Runtime/Core/Private/Containers/Algo/AlgosTest.cpp
--- a/Engine/Source/Runtime/Core/Private/Containers/Algo/AlgosTest.cpp
+++ b/Engine/Source/Runtime/Core/Private/Containers/Algo/AlgosTest.cpp
@@ -6,18 +6,22 @@
 #include "Containers/Array.h"
 #include "Containers/UnrealString.h"
 #include "Misc/AutomationTest.h"
+#include "Algo/AllOf.h"
+#include "Algo/Contains.h"
 #include "Algo/Copy.h"
 #include "Algo/Heapify.h"
 #include "Algo/HeapSort.h"
 #include "Algo/Includes.h"
+#include "Algo/IndexOf.h"
 #include "Algo/IntroSort.h"
 #include "Algo/IsHeap.h"
 #include "Algo/IsSorted.h"
+#include "Algo/LevenshteinDistance.h"
 #include "Algo/Mismatch.h"
+#include "Algo/NoneOf.h"
+#include "Algo/RemoveIf.h"
 #include "Algo/Sort.h"
 #include "Algo/Transform.h"
-#include "Algo/IndexOf.h"
-#include "Algo/LevenshteinDistance.h"
 #include "Templates/Greater.h"
 
 namespace UE
@@ -806,6 +810,89 @@ public:
 		}
 	}
 
+	void TestRemoveIf()
+	{
+		struct FStringWrapper
+		{
+			FString Str;
+			bool bMovedFrom = false;
+			mutable bool bPredicateInvoked = false;
+
+			FStringWrapper(FString&& InStr)
+				: Str(MoveTemp(InStr))
+			{
+			}
+			~FStringWrapper() = default;
+
+			// These shouldn't be called
+			FStringWrapper(FStringWrapper&&) = delete;
+			FStringWrapper(const FStringWrapper&) = delete;
+			bool operator=(const FStringWrapper&) = delete;
+
+			FStringWrapper& operator=(FStringWrapper&& Rhs)
+			{
+				checkf(!Rhs.bMovedFrom, TEXT("Algo::RemoveIf - attempting to move a moved-from object"));
+				Str = MoveTemp(Rhs.Str);
+				bMovedFrom = false;
+				Rhs.bMovedFrom = true;
+				return *this;
+			}
+		};
+
+		auto Test = [](std::initializer_list<const ANSICHAR*> Input, std::initializer_list<const ANSICHAR*> ExpectedResult)
+		{
+			TArray<FStringWrapper> TestArray;
+			for (const ANSICHAR* InputStr : Input)
+			{
+				TestArray.Emplace(InputStr);
+			}
+
+			int32 NumCalls = 0;
+			auto ShouldRemove = [&NumCalls](const FStringWrapper& Wrapper)
+			{
+				++NumCalls;
+				checkf(!Wrapper.bPredicateInvoked, TEXT("Algo::RemoveIf - running predicate on the same object multiple times"));
+				Wrapper.bPredicateInvoked = true;
+				return Wrapper.Str.StartsWith(TEXT("Remove-"));
+			};
+
+			int32 Result = Algo::RemoveIf(TestArray, ShouldRemove);
+
+			checkf(NumCalls == Input.size(), TEXT("Algo::RemoveIf - expected predicate to be run %d times, got %d"), (int32)Input.size(), NumCalls);
+			checkf(Result == ExpectedResult.size(), TEXT("Algo::RemoveIf - expected %d elements to be kept, got %d"), (int32)ExpectedResult.size(), Result);
+			checkf(Algo::AllOf(TestArray, &FStringWrapper::bPredicateInvoked), TEXT("Algo::RemoveIf - expected predicate to be run over all elements but found one that was not"));
+
+			TestArray.RemoveAt(Result, TestArray.Num() - Result, EAllowShrinking::No);
+
+			checkf(Algo::NoneOf(TestArray, &FStringWrapper::bMovedFrom), TEXT("Algo::RemoveIf - found non-removed element in a moved-from state"));
+			checkf(Algo::AllOf(ExpectedResult, [&TestArray](const ANSICHAR* ExpectedStr){ return Algo::ContainsBy(TestArray, ExpectedStr, &FStringWrapper::Str); }), TEXT("Algo::RemoveIf - actual result does not match expected result"));
+		};
+
+		// Test empty range
+		Test({}, {});
+
+		// Test none removed
+		Test({ "1", "2", "3", "4" }, { "1", "2", "3", "4" });
+
+		// Test all removed
+		Test({ "Remove-1", "Remove-2", "Remove-3", "Remove-4" }, {});
+
+		// Test only the first removed
+		Test({ "Remove-1", "2", "3", "4" }, { "2", "3", "4" });
+
+		// Test all but the first removed
+		Test({ "1", "Remove-2", "Remove-3", "Remove-4" }, { "1" });
+
+		// Test only the last removed
+		Test({ "1", "2", "3", "Remove-4" }, { "1", "2", "3" });
+
+		// Test all but the last removed
+		Test({ "Remove-1", "Remove-2", "Remove-3", "4" }, { "4" });
+
+		// Test a mixture
+		Test({ "Remove-1", "2", "3", "4", "Remove-5", "Remove-6", "7", "Remove-8", "9" }, { "2", "3", "4", "7", "9" });
+	}
+
 private:
 	TArray<int> TestData;
 	TArray<int> TestData2;
@@ -830,6 +917,7 @@ bool FAlgosTest::RunTest(const FString& Parameters)
 	TestEditDistanceArray();
 	TestIncludes();
 	TestMismatch();
+	TestRemoveIf();
 	Cleanup();
 
 	return true;
diff --git a/Engine/Source/Runtime/Core/Public/Algo/RemoveIf.h b/Engine/Source/Runtime/Core/Public/Algo/RemoveIf.h
--- a/Engine/Source/Runtime/Core/Public/Algo/RemoveIf.h
+++ b/Engine/Source/Runtime/Core/Public/Algo/RemoveIf.h
@@ -46,23 +46,22 @@ namespace Algo
 			// Skip removed elements at the end
 			for (;;)
 			{
-				if (!Invoke(Pred, *(IterEnd - 1)))
-				{
-					break;
-				}
-
 				--IterEnd;
 
 				if (IterStart == IterEnd)
 				{
 					return UE_PTRDIFF_TO_INT32(IterStart - First);
 				}
+
+				if (!Invoke(Pred, *IterEnd))
+				{
+					break;
+				}
 			}
 
-			*IterStart = MoveTemp(*(IterEnd - 1));
+			*IterStart = MoveTemp(*IterEnd);
 
 			++IterStart;
-			--IterEnd;
 		}
 	}
 
