Added tests for Algo::StableRemoveIf.

diff --git a/Engine/Source/Runtime/Core/Private/Containers/Algo/AlgosTest.cpp b/Engine/Source/Runtime/Core/Private/Containers/Algo/AlgosTest.cpp
--- a/Engine/Source/Runtime/Core/Private/Containers/Algo/AlgosTest.cpp
+++ b/Engine/Source/Runtime/Core/Private/Containers/Algo/AlgosTest.cpp
@@ -7,6 +7,7 @@
 #include "Containers/UnrealString.h"
 #include "Misc/AutomationTest.h"
 #include "Algo/AllOf.h"
+#include "Algo/Compare.h"
 #include "Algo/Contains.h"
 #include "Algo/Copy.h"
 #include "Algo/Heapify.h"
@@ -841,31 +842,61 @@ public:
 
 		auto Test = [](std::initializer_list<const ANSICHAR*> Input, std::initializer_list<const ANSICHAR*> ExpectedResult)
 		{
-			TArray<FStringWrapper> TestArray;
-			for (const ANSICHAR* InputStr : Input)
 			{
-				TestArray.Emplace(InputStr);
+				TArray<FStringWrapper> TestArray;
+				for (const ANSICHAR* InputStr : Input)
+				{
+					TestArray.Emplace(InputStr);
+				}
+
+				int32 NumCalls = 0;
+				auto ShouldRemove = [&NumCalls](const FStringWrapper& Wrapper)
+				{
+					++NumCalls;
+					checkf(!Wrapper.bPredicateInvoked, TEXT("Algo::RemoveIf - running predicate on the same object multiple times"));
+					Wrapper.bPredicateInvoked = true;
+					return Wrapper.Str.StartsWith(TEXT("Remove-"));
+				};
+
+				int32 Result = Algo::RemoveIf(TestArray, ShouldRemove);
+
+				checkf(NumCalls == Input.size(), TEXT("Algo::RemoveIf - expected predicate to be run %d times, got %d"), (int32)Input.size(), NumCalls);
+				checkf(Result == ExpectedResult.size(), TEXT("Algo::RemoveIf - expected %d elements to be kept, got %d"), (int32)ExpectedResult.size(), Result);
+				checkf(Algo::AllOf(TestArray, &FStringWrapper::bPredicateInvoked), TEXT("Algo::RemoveIf - expected predicate to be run over all elements but found one that was not"));
+
+				TestArray.RemoveAt(Result, TestArray.Num() - Result, EAllowShrinking::No);
+
+				checkf(Algo::NoneOf(TestArray, &FStringWrapper::bMovedFrom), TEXT("Algo::RemoveIf - found non-removed element in a moved-from state"));
+				checkf(Algo::AllOf(ExpectedResult, [&TestArray](const ANSICHAR* ExpectedStr){ return Algo::ContainsBy(TestArray, ExpectedStr, &FStringWrapper::Str); }), TEXT("Algo::RemoveIf - actual result does not match expected result"));
 			}
 
-			int32 NumCalls = 0;
-			auto ShouldRemove = [&NumCalls](const FStringWrapper& Wrapper)
 			{
-				++NumCalls;
-				checkf(!Wrapper.bPredicateInvoked, TEXT("Algo::RemoveIf - running predicate on the same object multiple times"));
-				Wrapper.bPredicateInvoked = true;
-				return Wrapper.Str.StartsWith(TEXT("Remove-"));
-			};
+				TArray<FStringWrapper> TestArray;
+				for (const ANSICHAR* InputStr : Input)
+				{
+					TestArray.Emplace(InputStr);
+				}
 
-			int32 Result = Algo::RemoveIf(TestArray, ShouldRemove);
+				int32 NumCalls = 0;
+				auto ShouldRemove = [&NumCalls](const FStringWrapper& Wrapper)
+				{
+					++NumCalls;
+					checkf(!Wrapper.bPredicateInvoked, TEXT("Algo::StableRemoveIf - running predicate on the same object multiple times"));
+					Wrapper.bPredicateInvoked = true;
+					return Wrapper.Str.StartsWith(TEXT("Remove-"));
+				};
+
+				int32 Result = Algo::StableRemoveIf(TestArray, ShouldRemove);
 
-			checkf(NumCalls == Input.size(), TEXT("Algo::RemoveIf - expected predicate to be run %d times, got %d"), (int32)Input.size(), NumCalls);
-			checkf(Result == ExpectedResult.size(), TEXT("Algo::RemoveIf - expected %d elements to be kept, got %d"), (int32)ExpectedResult.size(), Result);
-			checkf(Algo::AllOf(TestArray, &FStringWrapper::bPredicateInvoked), TEXT("Algo::RemoveIf - expected predicate to be run over all elements but found one that was not"));
+				checkf(NumCalls == Input.size(), TEXT("Algo::StableRemoveIf - expected predicate to be run %d times, got %d"), (int32)Input.size(), NumCalls);
+				checkf(Result == ExpectedResult.size(), TEXT("Algo::StableRemoveIf - expected %d elements to be kept, got %d"), (int32)ExpectedResult.size(), Result);
+				checkf(Algo::AllOf(TestArray, &FStringWrapper::bPredicateInvoked), TEXT("Algo::StableRemoveIf - expected predicate to be run over all elements but found one that was not"));
 
-			TestArray.RemoveAt(Result, TestArray.Num() - Result, EAllowShrinking::No);
+				TestArray.RemoveAt(Result, TestArray.Num() - Result, EAllowShrinking::No);
 
-			checkf(Algo::NoneOf(TestArray, &FStringWrapper::bMovedFrom), TEXT("Algo::RemoveIf - found non-removed element in a moved-from state"));
-			checkf(Algo::AllOf(ExpectedResult, [&TestArray](const ANSICHAR* ExpectedStr){ return Algo::ContainsBy(TestArray, ExpectedStr, &FStringWrapper::Str); }), TEXT("Algo::RemoveIf - actual result does not match expected result"));
+				checkf(Algo::NoneOf(TestArray, &FStringWrapper::bMovedFrom), TEXT("Algo::StableRemoveIf - found non-removed element in a moved-from state"));
+				checkf(Algo::Compare(TestArray, ExpectedResult, [](const FStringWrapper& Lhs, const ANSICHAR* Rhs){ return Lhs.Str == Rhs; }), TEXT("Algo::StableRemoveIf - actual result does not match expected result"));
+			}
 		};
 
 		// Test empty range
