- Reapply commandline overrides after dynamic layers change and rebuild the branch

diff --git a/Engine/Source/Runtime/Core/Private/Misc/ConfigCacheIni.cpp b/Engine/Source/Runtime/Core/Private/Misc/ConfigCacheIni.cpp
--- a/Engine/Source/Runtime/Core/Private/Misc/ConfigCacheIni.cpp
+++ b/Engine/Source/Runtime/Core/Private/Misc/ConfigCacheIni.cpp
@@ -3261,7 +3261,8 @@ bool FConfigBranch::AddDynamicLayersToHierarchy(const TArray<FDynamicLayerInfo>&
 	if (AddedLayers.Num() > 0 || bForceFullDynamicLayerUpdate)
 	{
 		// if all were added at the end (or there's no replay), we can just apply them without rewinding
-		if (!bInsertedBeforeEnd && !bForceFullDynamicLayerUpdate)
+		// since commandline needs to trump the layers, we can't use this optimization if there are overrides
+		if (!bInsertedBeforeEnd && !bForceFullDynamicLayerUpdate && CommandLineOverrides.Num() == 0)
 		{
 			for (FConfigCommandStream* NewLayer : AddedLayers)
 			{
@@ -3279,6 +3280,7 @@ bool FConfigBranch::AddDynamicLayersToHierarchy(const TArray<FDynamicLayerInfo>&
 			{
 				FinalCombinedLayers.ApplyFile(*Node);
 			}
+			FinalCombinedLayers.ApplyFile(&CommandLineOverrides);
 			bool bOldSaveAllSections = InMemoryFile.bCanSaveAllSections;
 			InMemoryFile = FinalCombinedLayers;
 			InMemoryFile.bCanSaveAllSections = bOldSaveAllSections;
@@ -3362,7 +3364,9 @@ bool FConfigBranch::AddDynamicLayerStringToHierarchy(const FString& Filename, co
 		}
 	}
 
-	if (!bInsertedAtEnd)
+	// if we inserted at the end, and there are no commandline options that need to stomp on layers,
+	// then there's a chance of optimization
+	if (!bInsertedAtEnd || CommandLineOverrides.Num() > 0)
 	{
 		// rebuild
 		FinalCombinedLayers = CombinedStaticLayers;
@@ -3370,6 +3374,7 @@ bool FConfigBranch::AddDynamicLayerStringToHierarchy(const FString& Filename, co
 		{
 			FinalCombinedLayers.ApplyFile(*Node);
 		}
+		FinalCombinedLayers.ApplyFile(&CommandLineOverrides);
 		bool bOldSaveAllSections = InMemoryFile.bCanSaveAllSections;
 		InMemoryFile = FinalCombinedLayers;
 		InMemoryFile.bCanSaveAllSections = bOldSaveAllSections;
@@ -3434,6 +3439,7 @@ bool FConfigBranch::RemoveDynamicLayersFromHierarchy(const TArray<FString>& File
 	{
 		FinalCombinedLayers.ApplyFile(*Node);
 	}
+	FinalCombinedLayers.ApplyFile(&CommandLineOverrides);
 	bool bOldSaveAllSections = InMemoryFile.bCanSaveAllSections;
 	InMemoryFile = FinalCombinedLayers;
 	InMemoryFile.bCanSaveAllSections = bOldSaveAllSections;
diff --git a/Engine/Source/Runtime/Core/Private/Misc/ConfigContext.cpp b/Engine/Source/Runtime/Core/Private/Misc/ConfigContext.cpp
--- a/Engine/Source/Runtime/Core/Private/Misc/ConfigContext.cpp
+++ b/Engine/Source/Runtime/Core/Private/Misc/ConfigContext.cpp
@@ -1207,6 +1207,7 @@ bool FConfigContext::GenerateDestIniFile()
 		// process any commandline overrides
 		FConfigFile::OverrideFromCommandline(&Branch->CommandLineOverrides, BaseIniName);
 		// and push it into the current values
+		Branch->FinalCombinedLayers.ApplyFile(&Branch->CommandLineOverrides);
 		Branch->InMemoryFile.ApplyFile(&Branch->CommandLineOverrides);
 #endif
 
