Fixed allocator call for TExecutableTask to go via FMemory (and use LLM) and have an optimal alignment

diff --git a/Engine/Source/Runtime/Core/Public/Tasks/TaskPrivate.h b/Engine/Source/Runtime/Core/Public/Tasks/TaskPrivate.h
--- a/Engine/Source/Runtime/Core/Public/Tasks/TaskPrivate.h
+++ b/Engine/Source/Runtime/Core/Public/Tasks/TaskPrivate.h
@@ -894,13 +894,12 @@ protected:
 		};
 
 		inline constexpr int32 SmallTaskSize = 256;
-		inline constexpr int32 LargeTaskAlignment = 16; // Larger than this will result in very wasteful allocations with MallocBinned2/3
 		using FExecutableTaskAllocator = TLockFreeFixedSizeAllocator_TLSCache<SmallTaskSize, PLATFORM_CACHE_LINE_SIZE>;
 		CORE_API extern FExecutableTaskAllocator SmallTaskAllocator;
 
 		// a separate derived class to add "small task" allocation optimization to both base class specializations
 		template<typename TaskBodyType>
-		class TExecutableTask final : public TExecutableTaskBase<TaskBodyType>
+		class alignas(PLATFORM_CACHE_LINE_SIZE) TExecutableTask final : public TExecutableTaskBase<TaskBodyType>
 		{
 		public:
 			TExecutableTask(const TCHAR* InDebugName, TaskBodyType&& TaskBody, ETaskPriority InPriority, EExtendedTaskPriority InExtendedPriority, ETaskFlags Flags)
@@ -923,13 +922,13 @@ protected:
 				else
 				{
 					TASKGRAPH_VERBOSE_EVENT_SCOPE(TExecutableTask::LargeAlloc);
-					return UE::Private::GMalloc->Malloc(sizeof(TExecutableTask), LargeTaskAlignment);
+					return FMemory::Malloc(sizeof(TExecutableTask), alignof(TExecutableTask));
 				}
 			}
 
 			static void operator delete(void* Ptr, size_t Size)
 			{
-				Size <= SmallTaskSize ? SmallTaskAllocator.Free(Ptr) : UE::Private::GMalloc->Free(Ptr);
+				Size <= SmallTaskSize ? SmallTaskAllocator.Free(Ptr) : FMemory::Free(Ptr);
 			}
 		};
 
