Mutex: Use InitializeCriticalSectionAndSpinCount to avoid an extra function call during construction of FWindowsRecursiveMutex

diff --git a/Engine/Source/Runtime/Core/Public/Windows/WindowsPlatformMutex.h b/Engine/Source/Runtime/Core/Public/Windows/WindowsPlatformMutex.h
--- a/Engine/Source/Runtime/Core/Public/Windows/WindowsPlatformMutex.h
+++ b/Engine/Source/Runtime/Core/Public/Windows/WindowsPlatformMutex.h
@@ -25,9 +25,7 @@ public:
 	FORCEINLINE FWindowsRecursiveMutex()
 	{
 		CA_SUPPRESS(28125);
-		Windows::InitializeCriticalSection(&CriticalSection);
-		Windows::SetCriticalSectionSpinCount(&CriticalSection, 4000);
-		//Windows::InitializeCriticalSectionAndSpinCount(&CriticalSection, 4000);
+		Windows::InitializeCriticalSectionAndSpinCount(&CriticalSection, 4000);
 	}
 
 	FORCEINLINE ~FWindowsRecursiveMutex()
