Added a workaround to avoid a TLS initialization/destruction deadlock while the crash handler is running.

#rnx

diff --git a/Engine/Source/Runtime/Core/Private/Async/Fundamental/Scheduler.cpp b/Engine/Source/Runtime/Core/Private/Async/Fundamental/Scheduler.cpp
--- a/Engine/Source/Runtime/Core/Private/Async/Fundamental/Scheduler.cpp
+++ b/Engine/Source/Runtime/Core/Private/Async/Fundamental/Scheduler.cpp
@@ -8,6 +8,7 @@
 #include "Logging/LogMacros.h"
 #include "ProfilingDebugging/CsvProfiler.h"
 #include "HAL/MallocAnsi.h"
+#include "HAL/PlatformMallocCrash.h"
 #include "HAL/RunnableThread.h"
 #include "Misc/ScopeLock.h"
 #include "Misc/CommandLine.h"
@@ -92,41 +93,42 @@ namespace LowLevelTasks
 
 	FSchedulerTls::FTlsValuesHolder::FTlsValuesHolder()
 	{
-#if 0
-		while (!FPlatformMisc::IsDebuggerPresent())
+		// Avoid a deadlock on threads being spun up or down during a crash
+		if (!FPlatformMallocCrash::Get().IsActive())
 		{
-			FPlatformProcess::SleepNoStats(0.100f);
-		}
-#endif
-		TlsValues = new FTlsValues();
+			TlsValues = new FTlsValues();
 
-		if (FImpl::ThreadTlsValuesMutex.TryLock())
-		{
-			FImpl::ProcessPendingTlsValuesNoLock();
-			TlsValues->LinkHead(FImpl::ThreadTlsValues);
-			FImpl::ThreadTlsValuesMutex.Unlock();
-		}
-		else
-		{
-			FImpl::PendingInsertTlsValues.ProduceItem(TlsValues);
+			if (FImpl::ThreadTlsValuesMutex.TryLock())
+			{
+				FImpl::ProcessPendingTlsValuesNoLock();
+				TlsValues->LinkHead(FImpl::ThreadTlsValues);
+				FImpl::ThreadTlsValuesMutex.Unlock();
+			}
+			else
+			{
+				FImpl::PendingInsertTlsValues.ProduceItem(TlsValues);
+			}
 		}
 	}
 
 	FSchedulerTls::FTlsValuesHolder::~FTlsValuesHolder()
 	{
-		if (FImpl::ThreadTlsValuesMutex.TryLock())
+		if (TlsValues && !FPlatformMallocCrash::Get().IsActive())
 		{
-			FImpl::ProcessPendingTlsValuesNoLock();
-			TlsValues->Unlink();
-			FImpl::ThreadTlsValuesMutex.Unlock();
-			delete TlsValues;
-		}
-		else
-		{
-			FImpl::PendingDeleteTlsValues.ProduceItem(TlsValues);
-		}
+			if (FImpl::ThreadTlsValuesMutex.TryLock())
+			{
+				FImpl::ProcessPendingTlsValuesNoLock();
+				TlsValues->Unlink();
+				FImpl::ThreadTlsValuesMutex.Unlock();
+				delete TlsValues;
+			}
+			else
+			{
+				FImpl::PendingDeleteTlsValues.ProduceItem(TlsValues);
+			}
 
-		TlsValues = nullptr;
+			TlsValues = nullptr;
+		}
 	}
 
 	FORCENOINLINE bool& Private::FOversubscriptionTls::GetIsOversubscriptionAllowedRef()
