Rework on the UObjectClasses new TagSet for Insights to remove the necessity of adding `#if LLM_ALLOW_ASSETS_TAGS` around every usage of `ELLMTagSet::AssetClasses`.

Changing it to keep the low overhead allocation cost.

#rnx

diff --git a/Engine/Plugins/Experimental/EditorTelemetry/Source/Private/EditorTelemetry.cpp b/Engine/Plugins/Experimental/EditorTelemetry/Source/Private/EditorTelemetry.cpp
--- a/Engine/Plugins/Experimental/EditorTelemetry/Source/Private/EditorTelemetry.cpp
+++ b/Engine/Plugins/Experimental/EditorTelemetry/Source/Private/EditorTelemetry.cpp
@@ -336,7 +336,6 @@ void FEditorTelemetry::RecordEvent_MemoryLLM(const FString& Context, TArray<FAna
 		FLowLevelMemTracker::Get().GetTrackedTagsNamesWithAmount(LLMTrackedNoneMemory, ELLMTracker::Default, ELLMTagSet::None);
 		RecordLLMMemoryEvent(Context, TEXT("None"), LLMTrackedNoneMemory);
 
-#if LLM_ALLOW_ASSETS_TAGS
 		// AssetClasses TagSet
 		TMap<FName, uint64> LLMTrackedAssetClassesMemory;
 		FLowLevelMemTracker::Get().GetTrackedTagsNamesWithAmount(LLMTrackedAssetClassesMemory, ELLMTracker::Default, ELLMTagSet::AssetClasses);
@@ -346,7 +345,6 @@ void FEditorTelemetry::RecordEvent_MemoryLLM(const FString& Context, TArray<FAna
 		TMap<FName, uint64> LLMTrackedAssetMemory;
 		FLowLevelMemTracker::Get().GetTrackedTagsNamesWithAmount(LLMTrackedAssetMemory, ELLMTracker::Default, ELLMTagSet::Assets);
 		RecordLLMMemoryEvent(Context, TEXT("Assets"), LLMTrackedAssetMemory);
-#endif // LLM_ALLOW_ASSETS_TAGS
 	}
 #endif
 }
diff --git a/Engine/Plugins/Experimental/RuntimeTelemetry/Source/Private/RuntimeTelemetry.cpp b/Engine/Plugins/Experimental/RuntimeTelemetry/Source/Private/RuntimeTelemetry.cpp
--- a/Engine/Plugins/Experimental/RuntimeTelemetry/Source/Private/RuntimeTelemetry.cpp
+++ b/Engine/Plugins/Experimental/RuntimeTelemetry/Source/Private/RuntimeTelemetry.cpp
@@ -69,8 +69,6 @@ void FRuntimeTelemetry::RecordEvent_MemoryLLM(const FString& Context, TArray<FAn
 		FLowLevelMemTracker::Get().GetTrackedTagsNamesWithAmount(LLMTrackedNoneMemory, ELLMTracker::Default, ELLMTagSet::None);
 		RecordLLMMemoryEvent(Context, TEXT("None"), LLMTrackedNoneMemory);
 
-#if LLM_ALLOW_ASSETS_TAGS
-
 		// AssetClasses TagSet
 		TMap<FName, uint64> LLMTrackedAssetClassesMemory;
 		FLowLevelMemTracker::Get().GetTrackedTagsNamesWithAmount(LLMTrackedAssetClassesMemory, ELLMTracker::Default, ELLMTagSet::AssetClasses);
@@ -80,9 +78,6 @@ void FRuntimeTelemetry::RecordEvent_MemoryLLM(const FString& Context, TArray<FAn
 		TMap<FName, uint64> LLMTrackedAssetMemory;
 		FLowLevelMemTracker::Get().GetTrackedTagsNamesWithAmount(LLMTrackedAssetMemory, ELLMTracker::Default, ELLMTagSet::Assets);
 		RecordLLMMemoryEvent(Context, TEXT("Assets"), LLMTrackedAssetMemory);
-
-#endif // LLM_ALLOW_ASSETS_TAGS
-
 	}
 #endif
 }
diff --git a/Engine/Plugins/MemoryUsageQueries/Source/MemoryUsageQueries/Private/MemoryUsageInfoProviderLLM.cpp b/Engine/Plugins/MemoryUsageQueries/Source/MemoryUsageQueries/Private/MemoryUsageInfoProviderLLM.cpp
--- a/Engine/Plugins/MemoryUsageQueries/Source/MemoryUsageQueries/Private/MemoryUsageInfoProviderLLM.cpp
+++ b/Engine/Plugins/MemoryUsageQueries/Source/MemoryUsageQueries/Private/MemoryUsageInfoProviderLLM.cpp
@@ -8,34 +8,21 @@
 
 bool FMemoryUsageInfoProviderLLM::IsProviderAvailable() const
 {
-	// if !LLM_ALLOW_ASSETS_TAGS there will never be assets tags in any trackers
-	// meaning that we can return false in this case
-#if !LLM_ALLOW_ASSETS_TAGS
-
-	return false;
-
-#else // LLM_ALLOW_ASSETS_TAGS
-
 	return FLowLevelMemTracker::IsEnabled();
-
-#endif // LLM_ALLOW_ASSETS_TAGS
 }
 
 uint64 FMemoryUsageInfoProviderLLM::GetAssetMemoryUsage(FName Asset) const
 {
-#if LLM_ALLOW_ASSETS_TAGS
 	if (FLowLevelMemTracker::IsEnabled())
 	{
 		return FLowLevelMemTracker::Get().GetTagAmountForTracker(ELLMTracker::Default, Asset, ELLMTagSet::Assets, UE::LLM::ESizeParams::Default);
 	}
-#endif // LLM_ALLOW_ASSETS_TAGS
 
 	return 0U;
 }
 
 uint64 FMemoryUsageInfoProviderLLM::GetAssetsMemoryUsage(const TSet<FName>& Assets) const
 {
-#if LLM_ALLOW_ASSETS_TAGS
 	if (FLowLevelMemTracker::IsEnabled())
 	{
 		uint64 TotalSize = 0;
@@ -45,14 +32,12 @@ uint64 FMemoryUsageInfoProviderLLM::GetAssetsMemoryUsage(const TSet<FName>& Asse
 		}
 		return TotalSize;
 	}
-#endif // LLM_ALLOW_ASSETS_TAGS
 
 	return 0U;
 }
 
 uint64 FMemoryUsageInfoProviderLLM::GetAssetsMemoryUsageWithSize(const TSet<FName>& Assets, TMap<FName, uint64>& OutSizes) const
 {
-#if LLM_ALLOW_ASSETS_TAGS
 	if (FLowLevelMemTracker::IsEnabled())
 	{
 		uint64 TotalSize = 0;
@@ -64,22 +49,17 @@ uint64 FMemoryUsageInfoProviderLLM::GetAssetsMemoryUsageWithSize(const TSet<FNam
 		}
 		return TotalSize;
 	}
-#endif // LLM_ALLOW_ASSETS_TAGS
 
 	return 0U;
 }
 
 void FMemoryUsageInfoProviderLLM::GetAllAssetsWithSize(TMap<FName, uint64>& OutAssets) const
 {
-#if LLM_ALLOW_ASSETS_TAGS
-
 	OutAssets.Reset();
 	if (FLowLevelMemTracker::IsEnabled())
 	{
 		FLowLevelMemTracker::Get().GetTrackedTagsNamesWithAmount(OutAssets, ELLMTracker::Default, ELLMTagSet::Assets);
 	}
-	
-#endif // LLM_ALLOW_ASSETS_TAGS
 }
 
 void FMemoryUsageInfoProviderLLM::GetFilteredTagsWithSize(TMap<FName, uint64>& OutTags, ELLMTracker Tracker, ELLMTagSet TagSet, TArray<FLLMTagSetAllocationFilter>& Filters) const
@@ -107,12 +87,6 @@ IMemoryUsageInfoProvider* GetCurrentMemoryUsageInfoProvider()
 
 bool GetFilteredPackagesWithSize(TMap<FName, uint64>& OutPackagesWithSize, FName GroupName /* = NAME_None */, FString AssetSubstring /* = FString() */, FName ClassName /* = NAME_None */, FOutputDevice* ErrorOutput /*= GLog*/)
 {
-#if !LLM_ALLOW_ASSETS_TAGS
-
-	return false;
-
-#else // LLM_ALLOW_ASSETS_TAGS
-
 	TArray<FLLMTagSetAllocationFilter> Filters;
 
 	FName LongName = NAME_None;
@@ -146,18 +120,10 @@ bool GetFilteredPackagesWithSize(TMap<FName, uint64>& OutPackagesWithSize, FName
 	OutPackagesWithSize.ValueSort(TGreater<uint64>());
 
 	return true;
-	
-#endif // LLM_ALLOW_ASSETS_TAGS
 }
 
 bool GetFilteredPackagesCategorizedWithSize(TMap<FName, uint64>& OutPackagesWithSize, FName GroupName /*= NAME_None*/, FString AssetSubstring /*= FString()*/, FName ClassName /*= NAME_None*/, FName CategoryName /*= NAME_None*/, FOutputDevice* ErrorOutput /*= GLog*/)
 {
-#if !LLM_ALLOW_ASSETS_TAGS
-
-	return false;
-
-#else // LLM_ALLOW_ASSETS_TAGS
-
 	TArray<FLLMTagSetAllocationFilter> Filters;
 
 	FName LongName = NAME_None;
@@ -201,18 +167,10 @@ bool GetFilteredPackagesCategorizedWithSize(TMap<FName, uint64>& OutPackagesWith
 	OutPackagesWithSize.ValueSort(TGreater<uint64>());
 
 	return true;
-
-#endif // LLM_ALLOW_ASSETS_TAGS
 }
 
 bool GetFilteredClassesWithSize(TMap<FName, uint64>& OutClassesWithSize, FName GroupName /* = NAME_None */, FString AssetName /* = FString() */, FOutputDevice* ErrorOutput /* = GLog */)
 {
-#if !LLM_ALLOW_ASSETS_TAGS
-
-	return false;
-
-#else // LLM_ALLOW_ASSETS_TAGS
-
 	TArray<FLLMTagSetAllocationFilter> Filters;
 
 	FName LongName = NAME_None;
@@ -236,17 +194,10 @@ bool GetFilteredClassesWithSize(TMap<FName, uint64>& OutClassesWithSize, FName G
 	OutClassesWithSize.ValueSort(TGreater<uint64>());
 
 	return true;
-#endif // LLM_ALLOW_ASSETS_TAGS
 }
 
 bool GetFilteredGroupsWithSize(TMap<FName, uint64>& OutGroupsWithSize, FString AssetName /* = FString() */, FName ClassName /* = NAME_None */, FOutputDevice* ErrorOutput /* = GLog */)
 {
-#if !LLM_ALLOW_ASSETS_TAGS
-
-	return false;
-
-#else // LLM_ALLOW_ASSETS_TAGS
-
 	TArray<FLLMTagSetAllocationFilter> Filters;
 
 	FName LongName = NAME_None;
@@ -270,8 +221,6 @@ bool GetFilteredGroupsWithSize(TMap<FName, uint64>& OutGroupsWithSize, FString A
 	OutGroupsWithSize.ValueSort(TGreater<uint64>());
 
 	return true;
-
-#endif // LLM_ALLOW_ASSETS_TAGS
 }
 
 } // namespace MemoryUsageQueries
diff --git a/Engine/Source/Runtime/Core/Private/HAL/LowLevelMemTracker.cpp b/Engine/Source/Runtime/Core/Private/HAL/LowLevelMemTracker.cpp
--- a/Engine/Source/Runtime/Core/Private/HAL/LowLevelMemTracker.cpp
+++ b/Engine/Source/Runtime/Core/Private/HAL/LowLevelMemTracker.cpp
@@ -28,6 +28,13 @@
 #include "Templates/Atomic.h"
 #include "Trace/Trace.inl"
 
+static_assert(static_cast<uint32>(ELLMTagSet::Max) == 4, "You added a new TagSet, please update the #ifs bellow");
+#if LLM_ALLOW_ASSETS_TAGS
+#define LLM_TAGSET_MAX 3
+#else
+#define LLM_TAGSET_MAX 2
+#endif
+
 // Enable to allow memory tracking to be passed through to external tools without the extra overhead of LLM also tracking allocations.
 #define UE_ONLY_USE_PLATFORM_TRACKER 0
 
@@ -1454,6 +1461,32 @@ public:
 #endif
 		}
 
+		FORCEINLINE static int32 GetIndexFromTagSet(ELLMTagSet InTagSet)
+		{
+#if LLM_ALLOW_ASSETS_TAGS
+			switch (InTagSet)
+			{
+				case ELLMTagSet::None: return 0;
+				case ELLMTagSet::Assets: return 1;
+				case ELLMTagSet::AssetClasses: return 2;
+
+				default: 
+					LLMCheck(false);
+					return 0;
+			}
+#elif LLM_ALLOW_UOBJECTCLASSES_TAGS
+			switch (InTagSet)
+			{
+				case ELLMTagSet::None: return 0;
+				case ELLMTagSet::UObjectClasses: return 1; 
+
+				default: 
+					LLMCheck(false);
+					return 0;
+			}
+#endif
+		}
+
 		void SetTag(const FTagData* InTag, FLowLevelMemTracker& InLLMRef
 #if LLM_ALLOW_ASSETS_TAGS || LLM_ALLOW_UOBJECTCLASSES_TAGS
 		, ELLMTagSet InTagSet = ELLMTagSet::None
@@ -1464,7 +1497,8 @@ public:
 #if !LLM_ALLOW_ASSETS_TAGS && !LLM_ALLOW_UOBJECTCLASSES_TAGS
 			Tag = InTag->GetIndex();
 #else
-			Tag[static_cast<int32>(InTagSet)] = InTag ? InTag->GetIndex() : InvalidCompressedTagValue;
+			int32 TagSetIndex = GetIndexFromTagSet(InTagSet);
+			Tag[TagSetIndex] = InTag ? InTag->GetIndex() : InvalidCompressedTagValue;
 #endif //!LLM_ALLOW_ASSETS_TAGS
 #else
 			Tag = InTag->GetEnumTag();
@@ -1482,7 +1516,8 @@ public:
 #if !LLM_ALLOW_ASSETS_TAGS && !LLM_ALLOW_UOBJECTCLASSES_TAGS
 			return (*InLLMRef.TagDatas)[Tag];
 #else
-			return Tag[static_cast<int32>(InTagSet)] >= 0 ? (*InLLMRef.TagDatas)[Tag[static_cast<int32>(InTagSet)]] : nullptr;
+			int32 TagSetIndex = GetIndexFromTagSet(InTagSet);
+			return Tag[TagSetIndex] >= 0 ? (*InLLMRef.TagDatas)[Tag[TagSetIndex]] : nullptr;
 #endif
 #else
 			return InLLMRef.FindTagData(Tag);
@@ -1495,7 +1530,8 @@ public:
 #if !LLM_ALLOW_ASSETS_TAGS && !LLM_ALLOW_UOBJECTCLASSES_TAGS
 			return (InTagSet == ELLMTagSet::None) ? Tag : InvalidCompressedTagValue;
 #else
-			return Tag[static_cast<int32>(InTagSet)];
+			int32 TagSetIndex = GetIndexFromTagSet(InTagSet);
+			return Tag[TagSetIndex];
 #endif
 		}
 
@@ -1507,7 +1543,8 @@ public:
 				Tag = InTag;
 			}
 #else
-			Tag[static_cast<int32>(InTagSet)] = InTag;
+			int32 TagSetIndex = GetIndexFromTagSet(InTagSet);
+			Tag[TagSetIndex] = InTag;
 #endif
 		}
 #else
@@ -1522,13 +1559,11 @@ public:
 #if !LLM_ALLOW_ASSETS_TAGS && !LLM_ALLOW_UOBJECTCLASSES_TAGS
 		int32 Tag = 0;
 #else
-		int32 Tag[static_cast<int32>(ELLMTagSet::Max)];
-#if LLM_ALLOW_ASSETS_TAGS
-#define TAG_SET_MAX 3
-#else
-#define TAG_SET_MAX 2
-#endif
-		static_assert(static_cast<uint32>(ELLMTagSet::Max) <= TAG_SET_MAX, "Every new TagSet enum adds 4 bytes of LLM overhead memory for every active allocation. Do not add them without measuring the cost of the overhead in the cook and runtime of large projects.");
+		int32 Tag[LLM_TAGSET_MAX];
+
+		static_assert(LLM_TAGSET_MAX <= 3, 
+			"Incrementing the maximum amount of TagSet enums that can be defined at the same time adds 4 bytes of LLM overhead memory FOR EVERY ACTIVE ALLOCATION (e.g. 10 million * 4 == 40 million bytes). "
+			"Do not add them without measuring the cost of the overhead in the cook and runtime of large projects.");
 #endif // LLM_ALLOW_ASSETS_TAGS
 #else
 		ELLMTag Tag = ELLMTag::Untagged;
@@ -2349,8 +2384,7 @@ void FLowLevelMemTracker::ProcessCommandLine(const TCHAR* CmdLine)
 
 	// activate tag sets (we ignore None set, it's always on)
 	FString SetList;
-	static_assert((uint8)ELLMTagSet::Max == TAG_SET_MAX,
-		"You added a tagset, without updating FLowLevelMemTracker::ProcessCommandLine");
+	static_assert(static_cast<uint8>(ELLMTagSet::Max) == 4, "You added a tagset, without updating FLowLevelMemTracker::ProcessCommandLine");
 	if (FParse::Value(CmdLine, TEXT("LLMTAGSETS="), SetList, false /* bShouldStopOnSeparator */))
 	{
 		TArray<FString> Sets;
@@ -6367,18 +6401,13 @@ void FLLMTraceWriter::Publish(FLowLevelMemTracker& LLMRef, const FTrackerTagSize
 		bTagSetSpecSent = true;
 		FAnsiStringView TagSetNames[] = { 
 			"System", 
-#if LLM_ALLOW_ASSETS_TAGS
-			"Assets", "AssetClasses", 
-#elif LLM_ALLOW_UOBJECTCLASSES_TAGS
+			"Assets", 
+			"AssetClasses", 
 			"UObjectClasses" 
-#endif // LLM_ALLOW_UOBJECTCLASSES_TAGS
 		};
-#if LLM_ALLOW_ASSETS_TAGS
 		static_assert((int)ELLMTagSet::Assets == 1, "");
 		static_assert((int)ELLMTagSet::AssetClasses == 2, "");
-#elif LLM_ALLOW_UOBJECTCLASSES_TAGS
-		static_assert((int)ELLMTagSet::UObjectClasses == 1, "");
-#endif // LLM_ALLOW_UOBJECTCLASSES_TAGS
+		static_assert((int)ELLMTagSet::UObjectClasses == 3, "");
 		static_assert(UE_ARRAY_COUNT(TagSetNames) == (int)ELLMTagSet::Max, "");
 		for (uint8 TagSet = 0; TagSet < (uint8)ELLMTagSet::Max; ++TagSet)
 		{
diff --git a/Engine/Source/Runtime/Core/Private/UObject/UnrealNames.cpp b/Engine/Source/Runtime/Core/Private/UObject/UnrealNames.cpp
--- a/Engine/Source/Runtime/Core/Private/UObject/UnrealNames.cpp
+++ b/Engine/Source/Runtime/Core/Private/UObject/UnrealNames.cpp
@@ -454,10 +454,8 @@ private:
 		if (!Buffer)
 		{
 			LLM_SCOPE(ELLMTag::FName);
-#if LLM_ALLOW_ASSETS_TAGS
 			LLM_TAGSET_SCOPE_CLEAR(ELLMTagSet::Assets);
 			LLM_TAGSET_SCOPE_CLEAR(ELLMTagSet::AssetClasses);
-#endif // LLM_ALLOW_ASSETS_TAGS
 			UE_TRACE_METADATA_CLEAR_SCOPE();
 			constexpr uint32 BufferSizeBytes = NumWords * sizeof(WordType);
 			Buffer = (WordType*) FMemory::MallocZeroed(BufferSizeBytes, alignof(WordType));
@@ -772,10 +770,8 @@ private:
 	uint8* AllocBlock()
 	{
 		LLM_SCOPE(ELLMTag::FName);
-#if LLM_ALLOW_ASSETS_TAGS
 		LLM_TAGSET_SCOPE_CLEAR(ELLMTagSet::Assets);
 		LLM_TAGSET_SCOPE_CLEAR(ELLMTagSet::AssetClasses);
-#endif // LLM_ALLOW_ASSETS_TAGS
 		UE_TRACE_METADATA_CLEAR_SCOPE();
 
 #if USE_FNAME_MMAP
@@ -1179,10 +1175,8 @@ public:
 	void Initialize(FNameEntryAllocator& InEntries)
 	{
 		LLM_SCOPE(ELLMTag::FName);
-#if LLM_ALLOW_ASSETS_TAGS
 		LLM_TAGSET_SCOPE_CLEAR(ELLMTagSet::Assets);
 		LLM_TAGSET_SCOPE_CLEAR(ELLMTagSet::AssetClasses);
-#endif // LLM_ALLOW_ASSETS_TAGS
 		UE_TRACE_METADATA_CLEAR_SCOPE();
 		Entries = &InEntries;
 
@@ -1586,10 +1580,8 @@ private:
 	void Grow(const uint32 NewCapacity)
 	{
 		LLM_SCOPE(ELLMTag::FName);
-#if LLM_ALLOW_ASSETS_TAGS
 		LLM_TAGSET_SCOPE_CLEAR(ELLMTagSet::Assets);
 		LLM_TAGSET_SCOPE_CLEAR(ELLMTagSet::AssetClasses);
-#endif // LLM_ALLOW_ASSETS_TAGS
 		UE_TRACE_METADATA_CLEAR_SCOPE();
 		TArrayView<FNameSlot> OldSlots(Slots, Capacity());
 		const uint32 OldUsedSlots = UsedSlots;
diff --git a/Engine/Source/Runtime/Core/Public/HAL/LowLevelMemTracker.h b/Engine/Source/Runtime/Core/Public/HAL/LowLevelMemTracker.h
--- a/Engine/Source/Runtime/Core/Public/HAL/LowLevelMemTracker.h
+++ b/Engine/Source/Runtime/Core/Public/HAL/LowLevelMemTracker.h
@@ -131,12 +131,9 @@ enum class ELLMTracker : uint8
 enum class ELLMTagSet : uint8
 {
 	None,
-#if LLM_ALLOW_ASSETS_TAGS
 	Assets,
 	AssetClasses,
-#elif LLM_ALLOW_UOBJECTCLASSES_TAGS
 	UObjectClasses,
-#endif // LLM_ALLOW_UOBJECTCLASSES_TAGS
 	
 	Max, // note: see FLowLevelMemTracker::ShouldReduceThreads and IsAssetTagForAssets if you add any asset-style tagsets
 };
diff --git a/Engine/Source/Runtime/CoreUObject/Private/UObject/UObjectGlobals.cpp b/Engine/Source/Runtime/CoreUObject/Private/UObject/UObjectGlobals.cpp
--- a/Engine/Source/Runtime/CoreUObject/Private/UObject/UObjectGlobals.cpp
+++ b/Engine/Source/Runtime/CoreUObject/Private/UObject/UObjectGlobals.cpp
@@ -2794,10 +2794,8 @@ namespace NameReuse
 			return FName();
 		}
 
-#if LLM_ALLOW_ASSETS_TAGS
 		LLM_TAGSET_SCOPE_CLEAR(ELLMTagSet::Assets);
 		LLM_TAGSET_SCOPE_CLEAR(ELLMTagSet::AssetClasses);
-#endif // LLM_ALLOW_ASSETS_TAGS
 		UE_TRACE_METADATA_CLEAR_SCOPE();
 		LLM_SCOPE(ELLMTag::FName);
 
diff --git a/Engine/Source/Runtime/Engine/Private/ContentStreaming.cpp b/Engine/Source/Runtime/Engine/Private/ContentStreaming.cpp
--- a/Engine/Source/Runtime/Engine/Private/ContentStreaming.cpp
+++ b/Engine/Source/Runtime/Engine/Private/ContentStreaming.cpp
@@ -374,10 +374,8 @@ FStreamingManagerCollection& IStreamingManager::Get()
 	if (StreamingManagerCollection == nullptr)
 	{
 		LLM_SCOPE(ELLMTag::StreamingManager);
-#if LLM_ALLOW_ASSETS_TAGS
 		LLM_TAGSET_SCOPE_CLEAR(ELLMTagSet::Assets);
 		LLM_TAGSET_SCOPE_CLEAR(ELLMTagSet::AssetClasses);
-#endif // LLM_ALLOW_ASSETS_TAGS
 		UE_TRACE_METADATA_CLEAR_SCOPE();
 		// Since this is a lazily created static global variable we create it
 		// in the open.
diff --git a/Engine/Source/Runtime/Engine/Private/PSOPrecacheMaterial.cpp b/Engine/Source/Runtime/Engine/Private/PSOPrecacheMaterial.cpp
--- a/Engine/Source/Runtime/Engine/Private/PSOPrecacheMaterial.cpp
+++ b/Engine/Source/Runtime/Engine/Private/PSOPrecacheMaterial.cpp
@@ -188,10 +188,8 @@ public:
 	FMaterialPSOPrecacheRequestID PrecachePSOs(const FMaterialPSOPrecacheParams& Params, EPSOPrecachePriority Priority, FGraphEventArray& OutGraphEvents)
 	{
 		LLM_SCOPE(ELLMTag::PSO);
-#if LLM_ALLOW_ASSETS_TAGS
 		LLM_TAGSET_SCOPE_CLEAR(ELLMTagSet::Assets);
 		LLM_TAGSET_SCOPE_CLEAR(ELLMTagSet::AssetClasses);
-#endif // LLM_ALLOW_ASSETS_TAGS
 		UE_TRACE_METADATA_CLEAR_SCOPE();
 		FMaterialPSOPrecacheRequestID RequestID = INDEX_NONE;		
 
@@ -359,10 +357,8 @@ public:
 	void MarkCollectionComplete(const FMaterialPSOPrecacheParams& Params, const FPSOPrecacheDataArray& PrecacheData, const FPSOPrecacheRequestResultArray& PrecacheRequestResults, uint32 RequestLifecycleID)
 	{
 		LLM_SCOPE(ELLMTag::PSO);
-#if LLM_ALLOW_ASSETS_TAGS
 		LLM_TAGSET_SCOPE_CLEAR(ELLMTagSet::Assets);
 		LLM_TAGSET_SCOPE_CLEAR(ELLMTagSet::AssetClasses);
-#endif // LLM_ALLOW_ASSETS_TAGS
 		UE_TRACE_METADATA_CLEAR_SCOPE();
 		FRWScopeLock WriteLock(RWLock, SLT_Write);
 		
@@ -399,10 +395,8 @@ public:
 	void ReleasePrecacheData(FMaterialPSOPrecacheRequestID MaterialPSORequestID)
 	{
 		LLM_SCOPE(ELLMTag::PSO);
-#if LLM_ALLOW_ASSETS_TAGS
 		LLM_TAGSET_SCOPE_CLEAR(ELLMTagSet::Assets);
 		LLM_TAGSET_SCOPE_CLEAR(ELLMTagSet::AssetClasses);
-#endif // LLM_ALLOW_ASSETS_TAGS
 		UE_TRACE_METADATA_CLEAR_SCOPE();
 		check(MaterialPSORequestID != INDEX_NONE);
 
@@ -417,10 +411,8 @@ public:
 	void BoostPriority(EPSOPrecachePriority NewPri, FMaterialPSOPrecacheRequestID MaterialPSORequestID)
 	{
 		LLM_SCOPE(ELLMTag::PSO);
-#if LLM_ALLOW_ASSETS_TAGS
 		LLM_TAGSET_SCOPE_CLEAR(ELLMTagSet::Assets);
 		LLM_TAGSET_SCOPE_CLEAR(ELLMTagSet::AssetClasses);
-#endif // LLM_ALLOW_ASSETS_TAGS
 		UE_TRACE_METADATA_CLEAR_SCOPE();
 		check(MaterialPSORequestID != INDEX_NONE);
 
@@ -454,10 +446,8 @@ public:
 	void ClearMaterialPSORequests()
 	{
 		LLM_SCOPE(ELLMTag::PSO);
-#if LLM_ALLOW_ASSETS_TAGS
 		LLM_TAGSET_SCOPE_CLEAR(ELLMTagSet::Assets);
 		LLM_TAGSET_SCOPE_CLEAR(ELLMTagSet::AssetClasses);
-#endif // LLM_ALLOW_ASSETS_TAGS
 		UE_TRACE_METADATA_CLEAR_SCOPE();
 		check(IsInGameThread());
 
@@ -594,10 +584,8 @@ FMaterialPSORequestManager GMaterialPSORequestManager;
 void FMaterialPSOPrecacheCollectionTask::DoTask(ENamedThreads::Type CurrentThread, const FGraphEventRef& MyCompletionGraphEvent)
 {
 	LLM_SCOPE(ELLMTag::PSO);
-#if LLM_ALLOW_ASSETS_TAGS
 	LLM_TAGSET_SCOPE_CLEAR(ELLMTagSet::Assets);
 	LLM_TAGSET_SCOPE_CLEAR(ELLMTagSet::AssetClasses);
-#endif // LLM_ALLOW_ASSETS_TAGS
 	UE_TRACE_METADATA_CLEAR_SCOPE();
 	TRACE_CPUPROFILER_EVENT_SCOPE(FMaterialPSOPrecacheCollectionTask);
 
@@ -647,10 +635,8 @@ void FMaterialPSOPrecacheCollectionTask::DoTask(ENamedThreads::Type CurrentThrea
 void FShaderMapPreloadTask::DoTask(ENamedThreads::Type CurrentThread, const FGraphEventRef& MyCompletionGraphEvent)
 {
 	LLM_SCOPE(ELLMTag::PSO);
-#if LLM_ALLOW_ASSETS_TAGS
 	LLM_TAGSET_SCOPE_CLEAR(ELLMTagSet::Assets);
 	LLM_TAGSET_SCOPE_CLEAR(ELLMTagSet::AssetClasses);
-#endif // LLM_ALLOW_ASSETS_TAGS
 	UE_TRACE_METADATA_CLEAR_SCOPE();
 	TRACE_CPUPROFILER_EVENT_SCOPE(FShaderMapPreloadTask);
 
@@ -672,10 +658,8 @@ void FShaderMapPreloadTask::DoTask(ENamedThreads::Type CurrentThread, const FGra
 void FShaderPreloadCollectionTask::DoTask(ENamedThreads::Type CurrentThread, const FGraphEventRef& MyCompletionGraphEvent)
 {
 	LLM_SCOPE(ELLMTag::PSO);
-#if LLM_ALLOW_ASSETS_TAGS
 	LLM_TAGSET_SCOPE_CLEAR(ELLMTagSet::Assets);
 	LLM_TAGSET_SCOPE_CLEAR(ELLMTagSet::AssetClasses);
-#endif // LLM_ALLOW_ASSETS_TAGS
 	UE_TRACE_METADATA_CLEAR_SCOPE();
 	TRACE_CPUPROFILER_EVENT_SCOPE(FShaderPreloadCollectionTask);
 
diff --git a/Engine/Source/Runtime/Engine/Private/ShaderCompiler/ShaderCompiler.cpp b/Engine/Source/Runtime/Engine/Private/ShaderCompiler/ShaderCompiler.cpp
--- a/Engine/Source/Runtime/Engine/Private/ShaderCompiler/ShaderCompiler.cpp
+++ b/Engine/Source/Runtime/Engine/Private/ShaderCompiler/ShaderCompiler.cpp
@@ -5210,10 +5210,7 @@ static bool TryLoadCookedGlobalShaderMap(EShaderPlatform Platform, FScopedSlowTa
 
 void CompileGlobalShaderMap(EShaderPlatform Platform, const ITargetPlatform* TargetPlatform, bool bRefreshShaderMap)
 {
-	// as there is no assets tags, there is no need to register this
-#if LLM_ALLOW_ASSETS_TAGS
 	LLM_SCOPE_RENDER_RESOURCE(TEXT("GlobalShaderMap"));
-#endif // LLM_ALLOW_ASSETS_TAGS
 
 	// No global shaders needed on dedicated server or clients that use NullRHI. Note that cook commandlet needs to have them, even if it is not allowed to render otherwise.
 	if (FPlatformProperties::IsServerOnly() || (!IsRunningCommandlet() && !FApp::CanEverRender()))
diff --git a/Engine/Source/Runtime/Engine/Private/Streaming/StreamingManagerTexture.cpp b/Engine/Source/Runtime/Engine/Private/Streaming/StreamingManagerTexture.cpp
--- a/Engine/Source/Runtime/Engine/Private/Streaming/StreamingManagerTexture.cpp
+++ b/Engine/Source/Runtime/Engine/Private/Streaming/StreamingManagerTexture.cpp
@@ -315,10 +315,8 @@ void FRenderAssetStreamingManager::MarkMountedStateDirty(FIoFilenameHash Filenam
 	if (!bRecacheAllFiles && FilenameHash != INVALID_IO_FILENAME_HASH)
 	{
 		LLM_SCOPE(ELLMTag::StreamingManager);
-#if LLM_ALLOW_ASSETS_TAGS
 		LLM_TAGSET_SCOPE_CLEAR(ELLMTagSet::Assets);
 		LLM_TAGSET_SCOPE_CLEAR(ELLMTagSet::AssetClasses);
-#endif // LLM_ALLOW_ASSETS_TAGS
 		UE_TRACE_METADATA_CLEAR_SCOPE();
 		FScopeLock ScopeLock(&MountedStateDirtyFilesCS);
 		MountedStateDirtyFiles.Emplace(FilenameHash);
@@ -382,10 +380,8 @@ bool FRenderAssetStreamingManager::StreamOutRenderAssetData( int64 RequiredMemor
 {
 	FScopeLock ScopeLock(&CriticalSection);
 	LLM_SCOPE(ELLMTag::StreamingManager);
-#if LLM_ALLOW_ASSETS_TAGS
 	LLM_TAGSET_SCOPE_CLEAR(ELLMTagSet::Assets);
 	LLM_TAGSET_SCOPE_CLEAR(ELLMTagSet::AssetClasses);
-#endif // LLM_ALLOW_ASSETS_TAGS
 	UE_TRACE_METADATA_CLEAR_SCOPE();
 
 	const int64 MaxTempMemoryAllowed = static_cast<int64>(Settings.MaxTempMemoryAllowed) * 1024 * 1024;
@@ -879,10 +875,8 @@ void FRenderAssetStreamingManager::AddLevel( ULevel* Level )
 
 	FScopeLock ScopeLock(&CriticalSection);
 	LLM_SCOPE(ELLMTag::StreamingManager);
-#if LLM_ALLOW_ASSETS_TAGS
 	LLM_TAGSET_SCOPE_CLEAR(ELLMTagSet::Assets);
 	LLM_TAGSET_SCOPE_CLEAR(ELLMTagSet::AssetClasses);
-#endif // LLM_ALLOW_ASSETS_TAGS
 	UE_TRACE_METADATA_CLEAR_SCOPE();
 
 	check(Level);
@@ -1004,10 +998,8 @@ void FRenderAssetStreamingManager::AddStreamingRenderAsset(UStreamableRenderAsse
 {
 	FScopeLock ScopeLock(&CriticalSection);
 	LLM_SCOPE(ELLMTag::StreamingManager);
-#if LLM_ALLOW_ASSETS_TAGS
 	LLM_TAGSET_SCOPE_CLEAR(ELLMTagSet::Assets);
 	LLM_TAGSET_SCOPE_CLEAR(ELLMTagSet::AssetClasses);
-#endif // LLM_ALLOW_ASSETS_TAGS
 	UE_TRACE_METADATA_CLEAR_SCOPE();
 
 	STAT(GatheredStats.CallbacksCycles = -(int32)FPlatformTime::Cycles();)
@@ -1035,10 +1027,8 @@ void FRenderAssetStreamingManager::RemoveStreamingRenderAsset( UStreamableRender
 	
 	FScopeLock ScopeLock(&CriticalSection);
 	LLM_SCOPE(ELLMTag::StreamingManager);
-#if LLM_ALLOW_ASSETS_TAGS
 	LLM_TAGSET_SCOPE_CLEAR(ELLMTagSet::Assets);
 	LLM_TAGSET_SCOPE_CLEAR(ELLMTagSet::AssetClasses);
-#endif // LLM_ALLOW_ASSETS_TAGS
 	UE_TRACE_METADATA_CLEAR_SCOPE();
 	TArray<FStreamingRenderAsset>& StreamingRenderAssets = GetStreamingRenderAssetsAsyncSafe();
 
diff --git a/Engine/Source/Runtime/Experimental/IoStore/OnDemand/Private/OnDemandIoStore.cpp b/Engine/Source/Runtime/Experimental/IoStore/OnDemand/Private/OnDemandIoStore.cpp
--- a/Engine/Source/Runtime/Experimental/IoStore/OnDemand/Private/OnDemandIoStore.cpp
+++ b/Engine/Source/Runtime/Experimental/IoStore/OnDemand/Private/OnDemandIoStore.cpp
@@ -1364,10 +1364,8 @@ void FOnDemandIoStore::CreateContainersFromToc(
 {
 	static FName AssetClassName(TEXT("OnDemandIoStore"));
 	LLM_SCOPE_BYNAME(TEXT("FileSystem/OnDemandIoStore"));
-#if LLM_ALLOW_ASSETS_TAGS
 	LLM_TAGSET_SCOPE(FName(MountId), ELLMTagSet::Assets);
 	LLM_TAGSET_SCOPE(AssetClassName, ELLMTagSet::AssetClasses);
-#endif // LLM_ALLOW_ASSETS_TAGS
 	UE_TRACE_METADATA_SCOPE_ASSET_FNAME(FName(MountId), AssetClassName, FName(TocPath));
 	const FOnDemandTocHeader& Header = Toc.Header;
 	const FName CompressionFormat(Header.CompressionFormat);
@@ -1479,10 +1477,8 @@ void FOnDemandIoStore::CreateContainersFromToc(
 		}
 
 		// Do not count output container memory as being allocated for utocs
-#if LLM_ALLOW_ASSETS_TAGS
 		LLM_TAGSET_SCOPE_CLEAR(ELLMTagSet::Assets);
 		LLM_TAGSET_SCOPE_CLEAR(ELLMTagSet::AssetClasses);
-#endif // LLM_ALLOW_ASSETS_TAGS
 		UE_TRACE_METADATA_CLEAR_SCOPE();
 		Out.Add(MoveTemp(Container));
 	}
diff --git a/Engine/Source/Runtime/RHI/Private/RHI.cpp b/Engine/Source/Runtime/RHI/Private/RHI.cpp
--- a/Engine/Source/Runtime/RHI/Private/RHI.cpp
+++ b/Engine/Source/Runtime/RHI/Private/RHI.cpp
@@ -524,10 +524,8 @@ void FRHIResource::BeginTrackingResource(FRHIResource* InResource)
 	if (GRHITrackingResources)
 	{
 		LLM_SCOPE_BYNAME(TEXT("RHIMisc/ResourceTracking"));
-#if LLM_ALLOW_ASSETS_TAGS
 		LLM_TAGSET_SCOPE_CLEAR(ELLMTagSet::Assets);
 		LLM_TAGSET_SCOPE_CLEAR(ELLMTagSet::AssetClasses);
-#endif // LLM_ALLOW_ASSETS_TAGS
 		UE_TRACE_METADATA_CLEAR_SCOPE();
 
 		FScopeLock Lock(&GRHIResourceTrackingCriticalSection);
@@ -543,10 +541,8 @@ void FRHIResource::EndTrackingResource(FRHIResource* InResource)
 	if (InResource->bBeingTracked)
 	{
 		LLM_SCOPE_BYNAME(TEXT("RHIMisc/ResourceTracking"));
-#if LLM_ALLOW_ASSETS_TAGS
 		LLM_TAGSET_SCOPE_CLEAR(ELLMTagSet::Assets);
 		LLM_TAGSET_SCOPE_CLEAR(ELLMTagSet::AssetClasses);
-#endif // LLM_ALLOW_ASSETS_TAGS
 		UE_TRACE_METADATA_CLEAR_SCOPE();
 		FScopeLock Lock(&GRHIResourceTrackingCriticalSection);
 		GRHITrackedResources.Remove(InResource);
@@ -562,10 +558,8 @@ void FRHIResource::StartTrackingAllResources()
 void FRHIResource::StopTrackingAllResources()
 {
 	LLM_SCOPE_BYNAME(TEXT("RHIMisc/ResourceTracking"));
-#if LLM_ALLOW_ASSETS_TAGS
 	LLM_TAGSET_SCOPE_CLEAR(ELLMTagSet::Assets);
 	LLM_TAGSET_SCOPE_CLEAR(ELLMTagSet::AssetClasses);
-#endif // LLM_ALLOW_ASSETS_TAGS
 	UE_TRACE_METADATA_CLEAR_SCOPE();
 	FScopeLock Lock(&GRHIResourceTrackingCriticalSection);
 	for (FRHIResource* Resource : GRHITrackedResources)
diff --git a/Engine/Source/Runtime/RenderCore/Private/RenderGraphResourcePool.cpp b/Engine/Source/Runtime/RenderCore/Private/RenderGraphResourcePool.cpp
--- a/Engine/Source/Runtime/RenderCore/Private/RenderGraphResourcePool.cpp
+++ b/Engine/Source/Runtime/RenderCore/Private/RenderGraphResourcePool.cpp
@@ -182,10 +182,8 @@ FRDGPooledBuffer* FRDGBufferPool::CreateBuffer(FRHICommandListBase& RHICmdList,
 	
 	LLM_SCOPE_BYNAME(TEXT("RHIMisc/BufferPool"));
 	UE_TRACE_METADATA_CLEAR_SCOPE(); // Do not associate a pooled buffer with specific asset
-#if LLM_ALLOW_ASSETS_TAGS
 	LLM_TAGSET_SCOPE_CLEAR(ELLMTagSet::Assets);
 	LLM_TAGSET_SCOPE_CLEAR(ELLMTagSet::AssetClasses);
-#endif // LLM_ALLOW_ASSETS_TAGS
 
 	const FRHIBufferCreateDesc CreateDesc =
 		FRHIBufferCreateDesc::Create(InDebugName, NumBytes, Desc.BytesPerElement, Desc.Usage)
