Added CConvertibleTo concept mirroring std::convertible_to

diff --git a/Engine/Source/Runtime/Core/Public/Concepts/ConvertibleTo.h b/Engine/Source/Runtime/Core/Public/Concepts/ConvertibleTo.h
new file mode 100644
--- /dev/null
+++ b/Engine/Source/Runtime/Core/Public/Concepts/ConvertibleTo.h
@@ -0,0 +1,13 @@
+// Copyright Epic Games, Inc. All Rights Reserved.
+
+#pragma once
+
+#include <type_traits>
+
+/**
+ * Concept which describes convertibility from one type to another.
+ *
+ * We use this instead of std::convertible_to because <concepts> isn't a well supported header yet.
+ */
+template <typename From, typename To>
+concept CConvertibleTo = std::is_convertible_v<From, To> && requires { static_cast<To>(std::declval<From>()); };
