Add #ifndef guards around cheat CVars & console logging directories.

diff --git a/Engine/Source/Runtime/Core/Public/Misc/Build.h b/Engine/Source/Runtime/Core/Public/Misc/Build.h
--- a/Engine/Source/Runtime/Core/Public/Misc/Build.h
+++ b/Engine/Source/Runtime/Core/Public/Misc/Build.h
@@ -423,7 +423,9 @@
 	#define ALLOW_CHEAT_CVARS_IN_TEST 1
 #endif
 
-#define DISABLE_CHEAT_CVARS (UE_BUILD_SHIPPING || (UE_BUILD_TEST && !ALLOW_CHEAT_CVARS_IN_TEST))
+#ifndef DISABLE_CHEAT_CVARS
+	#define DISABLE_CHEAT_CVARS (UE_BUILD_SHIPPING || (UE_BUILD_TEST && !ALLOW_CHEAT_CVARS_IN_TEST))
+#endif
 
 // Controls the creation of a thread for detecting hangs (FThreadHeartBeat). This is subject to other criteria, USE_HANG_DETECTION
 #ifndef ALLOW_HANG_DETECTION
