Logging: Workaround for an issue where serializing the backlog to an output device can generate new logs

diff --git a/Engine/Source/Runtime/Core/Private/Misc/OutputDeviceRedirector.cpp b/Engine/Source/Runtime/Core/Private/Misc/OutputDeviceRedirector.cpp
--- a/Engine/Source/Runtime/Core/Private/Misc/OutputDeviceRedirector.cpp
+++ b/Engine/Source/Runtime/Core/Private/Misc/OutputDeviceRedirector.cpp
@@ -146,6 +146,7 @@ struct FOutputDeviceRedirectorState
 	/** A queue of lines logged before the editor added its output device. */
 	TArray<FBufferedLine> BacklogLines;
 	FRWLock BacklogLock;
+	std::atomic<uint32> SerializeBacklogCount = 0;
 
 	/** An optional dedicated primary thread for logging to buffered output devices. */
 	FThread Thread;
@@ -551,6 +552,12 @@ void FOutputDeviceRedirectorState::AddToBacklog(
 	const FName& Category,
 	const double Time)
 {
+	// Skip adding to the backlog if it happens while serializing the backlog.
+	// This is a workaround until the backlog storage is safe to write during iteration.
+	if (SerializeBacklogCount.load(std::memory_order_relaxed) > 0)
+	{
+		return;
+	}
 	FWriteScopeLock ScopeLock(BacklogLock);
 	BacklogLines.Emplace(Data, Category, Verbosity, Time);
 }
@@ -558,10 +565,12 @@ void FOutputDeviceRedirectorState::AddToBacklog(
 void FOutputDeviceRedirectorState::SerializeBacklog(TNotNull<FOutputDevice*> OutputDevice)
 {
 	FReadScopeLock ScopeLock(BacklogLock);
+	SerializeBacklogCount.fetch_add(1, std::memory_order_relaxed);
 	for (const FBufferedLine& BacklogLine : BacklogLines)
 	{
 		OutputDevice->Serialize(BacklogLine.Data.Get(), BacklogLine.Verbosity, BacklogLine.Category, BacklogLine.Time);
 	}
+	SerializeBacklogCount.fetch_sub(1, std::memory_order_relaxed);
 }
 
 } // UE::Private
