[Backout] - CL42506288
#rnx
Original CL Desc
-----------------------------------------------------------------
Added a workaround to avoid a TLS initialization/destruction deadlock while the crash handler is running.

#rnx

diff --git a/Engine/Source/Runtime/Core/Private/Async/Fundamental/Scheduler.cpp b/Engine/Source/Runtime/Core/Private/Async/Fundamental/Scheduler.cpp
--- a/Engine/Source/Runtime/Core/Private/Async/Fundamental/Scheduler.cpp
+++ b/Engine/Source/Runtime/Core/Private/Async/Fundamental/Scheduler.cpp
@@ -8,7 +8,6 @@
 #include "Logging/LogMacros.h"
 #include "ProfilingDebugging/CsvProfiler.h"
 #include "HAL/MallocAnsi.h"
-#include "HAL/PlatformMallocCrash.h"
 #include "HAL/RunnableThread.h"
 #include "Misc/ScopeLock.h"
 #include "Misc/CommandLine.h"
@@ -93,42 +92,41 @@ namespace LowLevelTasks
 
 	FSchedulerTls::FTlsValuesHolder::FTlsValuesHolder()
 	{
-		// Avoid a deadlock on threads being spun up or down during a crash
-		if (!FPlatformMallocCrash::Get().IsActive())
+#if 0
+		while (!FPlatformMisc::IsDebuggerPresent())
 		{
-			TlsValues = new FTlsValues();
+			FPlatformProcess::SleepNoStats(0.100f);
+		}
+#endif
+		TlsValues = new FTlsValues();
 
-			if (FImpl::ThreadTlsValuesMutex.TryLock())
-			{
-				FImpl::ProcessPendingTlsValuesNoLock();
-				TlsValues->LinkHead(FImpl::ThreadTlsValues);
-				FImpl::ThreadTlsValuesMutex.Unlock();
-			}
-			else
-			{
-				FImpl::PendingInsertTlsValues.ProduceItem(TlsValues);
-			}
+		if (FImpl::ThreadTlsValuesMutex.TryLock())
+		{
+			FImpl::ProcessPendingTlsValuesNoLock();
+			TlsValues->LinkHead(FImpl::ThreadTlsValues);
+			FImpl::ThreadTlsValuesMutex.Unlock();
+		}
+		else
+		{
+			FImpl::PendingInsertTlsValues.ProduceItem(TlsValues);
 		}
 	}
 
 	FSchedulerTls::FTlsValuesHolder::~FTlsValuesHolder()
 	{
-		if (TlsValues && !FPlatformMallocCrash::Get().IsActive())
+		if (FImpl::ThreadTlsValuesMutex.TryLock())
 		{
-			if (FImpl::ThreadTlsValuesMutex.TryLock())
-			{
-				FImpl::ProcessPendingTlsValuesNoLock();
-				TlsValues->Unlink();
-				FImpl::ThreadTlsValuesMutex.Unlock();
-				delete TlsValues;
-			}
-			else
-			{
-				FImpl::PendingDeleteTlsValues.ProduceItem(TlsValues);
-			}
-
-			TlsValues = nullptr;
+			FImpl::ProcessPendingTlsValuesNoLock();
+			TlsValues->Unlink();
+			FImpl::ThreadTlsValuesMutex.Unlock();
+			delete TlsValues;
 		}
+		else
+		{
+			FImpl::PendingDeleteTlsValues.ProduceItem(TlsValues);
+		}
+
+		TlsValues = nullptr;
 	}
 
 	FORCENOINLINE bool& Private::FOversubscriptionTls::GetIsOversubscriptionAllowedRef()
