Fixed AutomationTests

diff --git a/Engine/Source/Runtime/Core/Private/Misc/AutomationTest.cpp b/Engine/Source/Runtime/Core/Private/Misc/AutomationTest.cpp
--- a/Engine/Source/Runtime/Core/Private/Misc/AutomationTest.cpp
+++ b/Engine/Source/Runtime/Core/Private/Misc/AutomationTest.cpp
@@ -478,7 +478,7 @@ bool FAutomationTestFramework::RunSmokeTests()
 	//@todo - better determination of whether to run than requires cooked data
 	// Ensure there isn't another slow task in progress when trying to run unit tests
 	const bool bRequiresCookedData = FPlatformProperties::RequiresCookedData();
-	if ((!bRequiresCookedData && !GIsSlowTask && !GIsPlayInEditorWorld && !FPlatformProperties::IsProgram() && !IsRunningCommandlet()) || bForceSmokeTests)
+	if ((!bRequiresCookedData & !GIsSlowTask & !GIsPlayInEditorWorld & !FPlatformProperties::IsProgram() & !IsRunningCommandlet()) || bForceSmokeTests)
 	{
 		TArray<FAutomationTestInfo> TestInfo;
 
