[AutoRTFM] Replace use of FCriticalSection with FTransactionallySafeCriticalSection in FConfigCacheIni

This lock is used within a transaction, so we need to use the transactionally-safe variant.

diff --git a/Engine/Source/Runtime/Core/Private/Misc/ConfigCacheIni.cpp b/Engine/Source/Runtime/Core/Private/Misc/ConfigCacheIni.cpp
--- a/Engine/Source/Runtime/Core/Private/Misc/ConfigCacheIni.cpp
+++ b/Engine/Source/Runtime/Core/Private/Misc/ConfigCacheIni.cpp
@@ -3691,7 +3691,7 @@ void FConfigCacheIni::Tick(float DeltaSeconds)
 		OtherFileToCheckForUnload = 0;
 	}
 	{
-		FScopeLock Lock(&OtherFilesLock);
+		UE::TScopeLock Lock(OtherFilesLock);
 		BranchesToCheck[1] = OtherFiles.FindRef(OtherFileNames[OtherFileToCheckForUnload++]);
 		checkf(OtherFileNames.Num() == OtherFiles.Num(), TEXT("OtherFiles and OtherFileNames are out of sync! %d other files, %d other file names!"), OtherFileNames.Num(), OtherFiles.Num());
 	}
@@ -3742,7 +3742,7 @@ FConfigBranch* FConfigCacheIni::FindBranchWithNoReload(FName BaseIniName, const
 	}
 	if (Branch == nullptr)
 	{
-		FScopeLock Lock(&OtherFilesLock);
+		UE::TScopeLock Lock(OtherFilesLock);
 		Branch = OtherFiles.FindRef(Filename);
 		if (Branch == nullptr)
 		{
@@ -3799,7 +3799,7 @@ FConfigBranch& FConfigCacheIni::AddNewBranch(const FString& Filename)
 	}
 #endif
 	
-	FScopeLock Lock(&OtherFilesLock);
+	UE::TScopeLock Lock(OtherFilesLock);
 	if (OtherFiles.Find(Filename) == nullptr)
 	{
 		OtherFileNames.Add(Filename);
@@ -3812,7 +3812,7 @@ FConfigBranch& FConfigCacheIni::AddNewBranch(const FString& Filename)
 
 int32 FConfigCacheIni::Remove(const FString& Filename)
 {
-	FScopeLock Lock(&OtherFilesLock);
+	UE::TScopeLock Lock(OtherFilesLock);
 
 	OtherFileNames.Remove(Filename);
 	delete OtherFiles.FindRef(Filename);
@@ -3829,7 +3829,7 @@ FConfigFile* FConfigCacheIni::FindConfigFile( const FString& Filename )
 	}
 	else
 	{
-		FScopeLock Lock(&OtherFilesLock);
+		UE::TScopeLock Lock(OtherFilesLock);
 		Result = OtherFiles.FindRef(Filename);
 	}
 
@@ -3909,7 +3909,7 @@ FConfigFile* FConfigCacheIni::FindConfigFileWithBaseName(FName BaseName)
 	FConfigBranch* Result = KnownFiles.GetBranch(BaseName);
 	if (Result == nullptr)
 	{
-		FScopeLock Lock(&OtherFilesLock);
+		UE::TScopeLock Lock(OtherFilesLock);
 		for (TPair<FString,FConfigBranch*>& CurrentFilePair : OtherFiles)
 		{
 			if (CurrentFilePair.Value->IniName == BaseName)
@@ -3951,7 +3951,7 @@ FConfigFile& FConfigCacheIni::Add(const FString& Filename, const FConfigFile& Fi
 	}
 #endif
 
-	FScopeLock Lock(&OtherFilesLock);
+	UE::TScopeLock Lock(OtherFilesLock);
 
 	if (OtherFiles.Find(Filename) == nullptr)
 	{
@@ -3971,7 +3971,7 @@ bool FConfigCacheIni::ContainsConfigFile(const FConfigFile* ConfigFile) const
 	// they are testing. It is the caller's responsibility to not try to hold
 	// on to the ConfigFile pointer during writes to this ConfigCacheIni
 	{
-		FScopeLock Lock(&OtherFilesLock);
+		UE::TScopeLock Lock(OtherFilesLock);
 		for (const TPair<FString, FConfigBranch*>& CurrentFilePair : OtherFiles)
 		{
 			if (ConfigFile == &CurrentFilePair.Value->InMemoryFile)
@@ -3997,7 +3997,7 @@ TArray<FString> FConfigCacheIni::GetFilenames()
 {	
 	TArray<FString> Result;
 	{
-		FScopeLock Lock(&OtherFilesLock);
+		UE::TScopeLock Lock(OtherFilesLock);
 		Result = OtherFileNames;
 	}
 
@@ -4031,7 +4031,7 @@ void FConfigCacheIni::Flush(bool bRemoveFromCache, const FString& Filename )
 			}
 			else
 			{
-				FScopeLock Lock(&OtherFilesLock);
+				UE::TScopeLock Lock(OtherFilesLock);
 				// flush all files
 				for (TPair<FString, FConfigBranch*>& Pair : OtherFiles)
 				{
@@ -4060,7 +4060,7 @@ void FConfigCacheIni::Flush(bool bRemoveFromCache, const FString& Filename )
 		}
 		else
 		{
-			FScopeLock Lock(&OtherFilesLock);
+			UE::TScopeLock Lock(OtherFilesLock);
 			for (TPair<FString, FConfigBranch*>& It : OtherFiles)
 			{
 				delete It.Value;
@@ -4691,7 +4691,7 @@ void FConfigCacheIni::Dump(FOutputDevice& Ar, const TCHAR* BaseIniName)
 		}
 	}
 
-	FScopeLock Lock(&OtherFilesLock);
+	UE::TScopeLock Lock(OtherFilesLock);
 
 	// sort the non-known files for easier diffing
 	TArray<FString> Keys;
@@ -5280,7 +5280,7 @@ struct FConfigMemoryData
  */
 void FConfigCacheIni::ShowMemoryUsage( FOutputDevice& Ar )
 {
-	FScopeLock Lock(&OtherFilesLock);
+	UE::TScopeLock Lock(OtherFilesLock);
 	FConfigMemoryData ConfigCacheMemoryData;
 
 	for (TPair<FString, FConfigBranch*>& Pair : OtherFiles)
@@ -5343,7 +5343,7 @@ void FConfigCacheIni::ShowMemoryUsage( FOutputDevice& Ar )
 
 SIZE_T FConfigCacheIni::GetMaxMemoryUsage()
 {
-	FScopeLock Lock(&OtherFilesLock);
+	UE::TScopeLock Lock(OtherFilesLock);
 
 	// record the memory used by the FConfigCacheIni's TMap
 	FArchiveCountConfigMem MemAr;
diff --git a/Engine/Source/Runtime/Core/Public/Misc/ConfigCacheIni.h b/Engine/Source/Runtime/Core/Public/Misc/ConfigCacheIni.h
--- a/Engine/Source/Runtime/Core/Public/Misc/ConfigCacheIni.h
+++ b/Engine/Source/Runtime/Core/Public/Misc/ConfigCacheIni.h
@@ -1987,7 +1987,7 @@ private:
 	FKnownConfigFiles KnownFiles;
 
 	TMap<FString, FConfigBranch*> OtherFiles;
-	mutable FCriticalSection OtherFilesLock;
+	mutable FTransactionallySafeCriticalSection OtherFilesLock;
 	
 	/** Parallel array to OtherFiles */
 	TArray<FString> OtherFileNames;
