Improved crash report callstack on Apple platforms

#rnx

diff --git a/Engine/Source/Runtime/Core/Private/Apple/ApplePlatformCrashContext.cpp b/Engine/Source/Runtime/Core/Private/Apple/ApplePlatformCrashContext.cpp
--- a/Engine/Source/Runtime/Core/Private/Apple/ApplePlatformCrashContext.cpp
+++ b/Engine/Source/Runtime/Core/Private/Apple/ApplePlatformCrashContext.cpp
@@ -27,6 +27,7 @@ FApplePlatformCrashContext::FApplePlatformCrashContext(ECrashContextType InType,
 {
 	SignalDescription[ 0 ] = 0;
 	MinidumpCallstackInfo[ 0 ] = 0;
+	CallstackProperty[ 0 ] = 0;
 }
 
 FApplePlatformCrashContext::~FApplePlatformCrashContext()
@@ -89,11 +90,12 @@ int32 FApplePlatformCrashContext::ReportCrash() const
 		const_cast<FApplePlatformCrashContext*>(this)->CapturePortableCallStack(ErrorFrame, Context);
 		// Walk the stack and dump it to the allocated memory (ignore first 2 callstack lines as those are in stack walking code)
 		FPlatformStackWalk::StackWalkAndDump( StackBuffer, UE_ARRAY_COUNT(MinidumpCallstackInfo) - 1, ErrorFrame, Context );
-		// We will put this in xml, need to escape it
+		// Escape it and put into CallstackProperty, to be used as <Callstack> in crash report
 		FString EscapedStack;
 		AppendEscapedXMLString(EscapedStack, FString(MinidumpCallstackInfo));
-		FCString::Strncpy(GErrorHist, *EscapedStack, UE_ARRAY_COUNT(GErrorHist));
+		FCString::Strncpy(const_cast<TCHAR*>(CallstackProperty), *EscapedStack, UE_ARRAY_COUNT(CallstackProperty));
 		
+		FUTF8ToTCHAR_Convert::Convert(GErrorHist, UE_ARRAY_COUNT(GErrorHist) - 1, MinidumpCallstackInfo, FCStringAnsi::Strlen(MinidumpCallstackInfo));
 		CreateExceptionInfoString(Signal, Info);
 	}
 	
@@ -245,6 +247,5 @@ TCHAR* FApplePlatformCrashContext::ItoTCHAR(uint64 Val, uint64 Base, uint32 Len)
 
 const TCHAR* FApplePlatformCrashContext::GetCallstackProperty() const
 {
-	// FApplePlatformCrashContext::ReportCrash() already parsed callstack into GErrorHist at this point
-	return GErrorHist;
+	return CallstackProperty;
 }
diff --git a/Engine/Source/Runtime/Core/Public/Apple/ApplePlatformCrashContext.h b/Engine/Source/Runtime/Core/Public/Apple/ApplePlatformCrashContext.h
--- a/Engine/Source/Runtime/Core/Public/Apple/ApplePlatformCrashContext.h
+++ b/Engine/Source/Runtime/Core/Public/Apple/ApplePlatformCrashContext.h
@@ -33,7 +33,10 @@ struct CORE_API FApplePlatformCrashContext : public FGenericCrashContext
 	char SignalDescription[128];
 	
 	/** Memory reserved for minidump-style callstack info */
-	char MinidumpCallstackInfo[65536];
+	char MinidumpCallstackInfo[16384];
+	
+	/** To be returned by GetCallstackProperty() */
+	TCHAR CallstackProperty[16384];
 	
 	/**
 	 * Inits the crash context from data provided by a signal handler.
