[AutoRTFM] Use transactionally-safe locks inside AutomationTest.

These locks can be reached when running the StateTree unit tests.

diff --git a/Engine/Source/Runtime/Core/Private/Misc/AutomationTest.cpp b/Engine/Source/Runtime/Core/Private/Misc/AutomationTest.cpp
--- a/Engine/Source/Runtime/Core/Private/Misc/AutomationTest.cpp
+++ b/Engine/Source/Runtime/Core/Private/Misc/AutomationTest.cpp
@@ -97,7 +97,7 @@ namespace AutomationTest
 	{
 		ELogVerbosity::Type EffectiveVerbosity = LogVerbosity;
 
-		static FCriticalSection ActionCS;
+		static FTransactionallySafeCriticalSection ActionCS;
 		static FAutomationTestBase* LastTest = nullptr;
 
 		if (AutomationTest::bCaptureLogEvents == false)
@@ -106,7 +106,7 @@ namespace AutomationTest
 		}
 
 		{
-			FScopeLock Lock(&ActionCS);
+			UE::TScopeLock Lock(ActionCS);
 			if (CurrentTest != LastTest)
 			{
 				FAutomationTestBase::SuppressedLogCategories.Empty();
@@ -288,7 +288,7 @@ void FAutomationTestFramework::FAutomationTestMessageFilter::Serialize(const TCH
 			Verbosity = ELogVerbosity::Verbose;
 		}
 		{
-			FScopeLock CriticalSection(&ActionCS);
+			UE::TScopeLock CriticalSection(ActionCS);
 			LocalDestinationContext->Serialize(V, Verbosity, Category, Time);
 		}
 	}
@@ -313,7 +313,7 @@ void FAutomationTestFramework::FAutomationTestMessageFilter::SerializeRecord(con
 			}
 		}
 		{
-			FScopeLock CriticalSection(&ActionCS);
+			UE::TScopeLock CriticalSection(ActionCS);
 			LocalDestinationContext->SerializeRecord(LocalRecord);
 		}
 	}
@@ -1094,7 +1094,7 @@ bool FAutomationTestFramework::InternalStopTest(FAutomationTestExecutionInfo& Ou
 	bTestSuccessful = bTestSuccessful && !CurrentTest->HasAnyErrors() && CurrentTest->HasMetExpectedMessages();
 
 	{
-		FWriteScopeLock Lock(CurrentTest->ActionCS);
+		UE::TWriteScopeLock Lock(CurrentTest->ActionCS);
 		CurrentTest->ExpectedMessages.Empty();
 	}
 
@@ -1424,7 +1424,7 @@ void FAutomationTestBase::AddError(const FString& InError, int32 StackOffset)
 {
 	if( !IsExpectedMessage(InError, ELogVerbosity::Warning))
 	{
-		FWriteScopeLock Lock(ActionCS);
+		UE::TWriteScopeLock Lock(ActionCS);
 		ExecutionInfo.AddEvent(FAutomationEvent(EAutomationEventType::Error, InError), StackOffset + 1);
 	}
 }
@@ -1444,7 +1444,7 @@ void FAutomationTestBase::AddErrorS(const FString& InError, const FString& InFil
 {
 	if ( !IsExpectedMessage(InError, ELogVerbosity::Warning))
 	{
-		FWriteScopeLock Lock(ActionCS);
+		UE::TWriteScopeLock Lock(ActionCS);
 		//ExecutionInfo.AddEvent(FAutomationEvent(EAutomationEventType::Error, InError, ExecutionInfo.GetContext(), InFilename, InLineNumber));
 	}
 }
@@ -1454,7 +1454,7 @@ void FAutomationTestBase::AddWarningS(const FString& InWarning, const FString& I
 {
 	if ( !IsExpectedMessage(InWarning, ELogVerbosity::Warning))
 	{
-		FWriteScopeLock Lock(ActionCS);
+		UE::TWriteScopeLock Lock(ActionCS);
 		//ExecutionInfo.AddEvent(FAutomationEvent(EAutomationEventType::Warning, InWarning, ExecutionInfo.GetContext(), InFilename, InLineNumber));
 	}
 }
@@ -1464,7 +1464,7 @@ void FAutomationTestBase::AddWarning( const FString& InWarning, int32 StackOffse
 {
 	if ( !IsExpectedMessage(InWarning, ELogVerbosity::Warning))
 	{
-		FWriteScopeLock Lock(ActionCS);
+		UE::TWriteScopeLock Lock(ActionCS);
 		ExecutionInfo.AddEvent(FAutomationEvent(EAutomationEventType::Warning, InWarning), StackOffset + 1);
 	}
 }
@@ -1474,7 +1474,7 @@ void FAutomationTestBase::AddInfo( const FString& InLogItem, int32 StackOffset,
 {
 	if ( !IsExpectedMessage(InLogItem, ELogVerbosity::Display))
 	{
-		FWriteScopeLock Lock(ActionCS);
+		UE::TWriteScopeLock Lock(ActionCS);
 		ExecutionInfo.AddEvent(FAutomationEvent(EAutomationEventType::Info, InLogItem), StackOffset + 1, bCaptureStack);
 	}
 }
@@ -1482,21 +1482,21 @@ void FAutomationTestBase::AddInfo( const FString& InLogItem, int32 StackOffset,
 UE_AUTORTFM_ALWAYS_OPEN
 void FAutomationTestBase::AddAnalyticsItem(const FString& InAnalyticsItem)
 {
-	FWriteScopeLock Lock(ActionCS);
+	UE::TWriteScopeLock Lock(ActionCS);
 	ExecutionInfo.AnalyticsItems.Add(InAnalyticsItem);
 }
 
 UE_AUTORTFM_ALWAYS_OPEN
 void FAutomationTestBase::AddTelemetryData(const FString& DataPoint, double Measurement, const FString& Context)
 {
-	FWriteScopeLock Lock(ActionCS);
+	UE::TWriteScopeLock Lock(ActionCS);
 	ExecutionInfo.TelemetryItems.Add(FAutomationTelemetryData(DataPoint, Measurement, Context));
 }
 
 UE_AUTORTFM_ALWAYS_OPEN
 void FAutomationTestBase::AddTelemetryData(const TMap <FString, double>& ValuePairs, const FString& Context)
 {
-	FWriteScopeLock Lock(ActionCS);
+	UE::TWriteScopeLock Lock(ActionCS);
 	for (const TPair<FString, double>& Item : ValuePairs)
 	{
 		ExecutionInfo.TelemetryItems.Add(FAutomationTelemetryData(Item.Key, Item.Value, Context));
@@ -1523,7 +1523,7 @@ void FAutomationTestBase::AddEvent(const FAutomationEvent& InEvent, int32 StackO
 
 	if (!IsExpectedMessage(InEvent.Message, LogType))
 	{
-		FWriteScopeLock Lock(ActionCS);
+		UE::TWriteScopeLock Lock(ActionCS);
 		ExecutionInfo.AddEvent(InEvent, StackOffset + 1, bCaptureStack);
 	}
 }
@@ -1538,7 +1538,7 @@ bool FAutomationTestBase::HasMetExpectedMessages(ELogVerbosity::Type VerbosityTy
 	bool bHasMetAllExpectedMessages = true;
 	TArray<FAutomationExpectedMessage> ExpectedMessagesArray;
 	{
-		FReadScopeLock RLock(ActionCS);
+		UE::TReadScopeLock RLock(ActionCS);
 		ExpectedMessagesArray = ExpectedMessages.Array();
 	}
 	for (FAutomationExpectedMessage& ExpectedMessage : ExpectedMessagesArray)
@@ -1553,7 +1553,7 @@ bool FAutomationTestBase::HasMetExpectedMessages(ELogVerbosity::Type VerbosityTy
 
 		if (ExpectedMessage.ExpectedNumberOfOccurrences > 0 && (ExpectedMessage.ExpectedNumberOfOccurrences != ExpectedMessage.ActualNumberOfOccurrences))
 		{
-			FWriteScopeLock WLock(ActionCS);
+			UE::TWriteScopeLock WLock(ActionCS);
 			bHasMetAllExpectedMessages = false;
 
 			ExecutionInfo.AddEvent(FAutomationEvent(EAutomationEventType::Error,
@@ -1567,7 +1567,7 @@ bool FAutomationTestBase::HasMetExpectedMessages(ELogVerbosity::Type VerbosityTy
 		}
 		else if (ExpectedMessage.ExpectedNumberOfOccurrences == 0)
 		{
-			FWriteScopeLock WLock(ActionCS);
+			UE::TWriteScopeLock WLock(ActionCS);
 			if (ExpectedMessage.ActualNumberOfOccurrences == 0)
 			{
 				bHasMetAllExpectedMessages = false;
@@ -1644,7 +1644,7 @@ void FAutomationTestBase::AddExpectedMessage(
 	int32 Occurrences,
 	bool IsRegex)
 {
-	FWriteScopeLock Lock(ActionCS);
+	UE::TWriteScopeLock Lock(ActionCS);
 	ExpectedMessages.Add(FAutomationExpectedMessage(ExpectedPatternString, ExpectedVerbosity, CompareType, Occurrences, IsRegex));
 }
 
@@ -2415,7 +2415,7 @@ bool FAutomationTestBase::IsExpectedMessage(
 	const FString& Message,
 	const ELogVerbosity::Type& Verbosity)
 {
-	FReadScopeLock Lock(ActionCS);
+	UE::TReadScopeLock Lock(ActionCS);
 	for (FAutomationExpectedMessage& ExpectedMessage : ExpectedMessages)
 	{
 		// Maintains previous behavior: Adjust so that error and fatal messages are tested against when the input verbosity is "Warning"
diff --git a/Engine/Source/Runtime/Core/Public/Misc/AutomationTest.h b/Engine/Source/Runtime/Core/Public/Misc/AutomationTest.h
--- a/Engine/Source/Runtime/Core/Public/Misc/AutomationTest.h
+++ b/Engine/Source/Runtime/Core/Public/Misc/AutomationTest.h
@@ -43,6 +43,8 @@
 #include "Misc/SourceLocation.h"
 #include "Misc/TextFilterExpressionEvaluator.h"
 #include "Misc/Timespan.h"
+#include "Misc/TransactionallySafeCriticalSection.h"
+#include "Misc/TransactionallySafeRWLock.h"
 #include "Templates/Function.h"
 #include "Templates/SharedPointer.h"
 #include "Templates/UnrealTemplate.h"
@@ -1380,7 +1382,7 @@ private:
 	 private:
 		 std::atomic<FAutomationTestBase*> CurTest;
 		 std::atomic<FFeedbackContext*> DestinationContext;
-		 FCriticalSection ActionCS;
+		 FTransactionallySafeCriticalSection ActionCS;
 	 };
 
 	friend class FAutomationTestOutputDevice;
@@ -2640,7 +2642,7 @@ private:
 	TSet<FAutomationExpectedMessage> ExpectedMessages;
 
 	/** Critical section lock */
-	FRWLock ActionCS;
+	FTransactionallySafeRWLock ActionCS;
 };
 
 class FBDDAutomationTestBase : public FAutomationTestBase
@@ -2963,7 +2965,7 @@ private:
 					return true;
 				}
 
-				FScopeLock CriticalSection(&ActionCS);
+				UE::TScopeLock CriticalSection(ActionCS);
 				bDone = false;
 				StartedRunning = FDateTime::UtcNow();
 				Future = Async(Execution, [this]() {
@@ -2990,7 +2992,7 @@ private:
 
 		void Done()
 		{
-			FScopeLock CriticalSection(&ActionCS);
+			UE::TScopeLock CriticalSection(ActionCS);
 			if (Future.IsValid())
 			{
 				bDone = true;
@@ -2999,7 +3001,7 @@ private:
 
 		void Reset()
 		{
-			FScopeLock CriticalSection(&ActionCS);
+			UE::TScopeLock CriticalSection(ActionCS);
 			// Reset the done for the next potential run of this command
 			bDone = false;
 			Future.Reset();
@@ -3016,7 +3018,7 @@ private:
 		FThreadSafeBool bDone;
 		FDateTime StartedRunning;
 		TFuture<void> Future;
-		FCriticalSection ActionCS;
+		FTransactionallySafeCriticalSection ActionCS;
 	};
 
 	class FAsyncLatentCommand : public IAutomationLatentCommand
@@ -3044,7 +3046,7 @@ private:
 					return true;
 				}
 
-				FScopeLock CriticalSection(&ActionCS);
+				UE::TScopeLock CriticalSection(ActionCS);
 				bDone = false;
 				StartedRunning = FDateTime::UtcNow();
 				Future = Async(Execution, [this]() {
@@ -3072,7 +3074,7 @@ private:
 
 		void Done()
 		{
-			FScopeLock CriticalSection(&ActionCS);
+			UE::TScopeLock CriticalSection(ActionCS);
 			if (Future.IsValid())
 			{
 				bDone = true;
@@ -3081,7 +3083,7 @@ private:
 
 		void Reset()
 		{
-			FScopeLock CriticalSection(&ActionCS);
+			UE::TScopeLock CriticalSection(ActionCS);
 			// Reset the done for the next potential run of this command
 			bDone = false;
 			Future.Reset();
@@ -3098,7 +3100,7 @@ private:
 		FThreadSafeBool bDone;
 		FDateTime StartedRunning;
 		TFuture<void> Future;
-		FCriticalSection ActionCS;
+		FTransactionallySafeCriticalSection ActionCS;
 	};
 
 	struct FSpecIt
