Niagara - Disable the traversal cache while resolving issues.

diff --git a/Engine/Plugins/FX/Niagara/Source/NiagaraEditor/Private/NiagaraEditorModule.cpp b/Engine/Plugins/FX/Niagara/Source/NiagaraEditor/Private/NiagaraEditorModule.cpp
--- a/Engine/Plugins/FX/Niagara/Source/NiagaraEditor/Private/NiagaraEditorModule.cpp
+++ b/Engine/Plugins/FX/Niagara/Source/NiagaraEditor/Private/NiagaraEditorModule.cpp
@@ -232,7 +232,7 @@ static FAutoConsoleVariableRef CVarEnableCustomInlineDynamicInputFormats(
 	TEXT("If > 0 and experimental inline editors for dynamic input trees are enabled, custom formats which are defined on scripts will be applied.\n"),
 	ECVF_Default);
 
-int32 GbEnableTraversalCache = 1;
+int32 GbEnableTraversalCache = 0;
 static FAutoConsoleVariableRef CVarEnableTraversalCache(
 	TEXT("fx.Niagara.EnableTraversalCache"),
 	GbEnableTraversalCache,
diff --git a/Engine/Source/Runtime/Core/Private/UObject/FortniteReleaseBranchCustomObjectVersion.cpp b/Engine/Source/Runtime/Core/Private/UObject/FortniteReleaseBranchCustomObjectVersion.cpp
--- a/Engine/Source/Runtime/Core/Private/UObject/FortniteReleaseBranchCustomObjectVersion.cpp
+++ b/Engine/Source/Runtime/Core/Private/UObject/FortniteReleaseBranchCustomObjectVersion.cpp
@@ -6,7 +6,7 @@ TMap<FGuid, FGuid> FFortniteReleaseBranchCustomObjectVersion::GetSystemGuids()
 	TMap<FGuid, FGuid> SystemGuids;
 	const FDevSystemGuids& DevGuids = FDevSystemGuids::Get();
 
-	SystemGuids.Add(DevGuids.Niagara_LatestScriptCompileVersion, FGuid("D35A655CD3FDA4408724D77E1316E35B"));
+	SystemGuids.Add(DevGuids.Niagara_LatestScriptCompileVersion, FGuid("8BC8C557C87C4F629468F10EA4161F0A"));
 	SystemGuids.Add(DevGuids.SkeletalMeshDerivedDataVersion, FGuid("C96DEBAD45E3443D804017D4C29F7C4B"));
 	SystemGuids.Add(DevGuids.MaterialTranslationDDCVersion, FGuid("888A502A66F6401580B27451F211C46B"));
 	SystemGuids.Add(DevGuids.GLOBALSHADERMAP_DERIVEDDATA_VER, FGuid("0A0A1BF5D2A94105BB395EF0DA7CAEB9"));
