Fixed AutomationTests warning

diff --git a/Engine/Source/Runtime/Core/Private/Misc/AutomationTest.cpp b/Engine/Source/Runtime/Core/Private/Misc/AutomationTest.cpp
--- a/Engine/Source/Runtime/Core/Private/Misc/AutomationTest.cpp
+++ b/Engine/Source/Runtime/Core/Private/Misc/AutomationTest.cpp
@@ -529,7 +529,7 @@ bool FAutomationTestFramework::RunSmokeTests()
 			FAutomationTestFramework::DumpAutomationTestExecutionInfo( OutExecutionInfoMap );
 		}
 	}
-	else if( bRequiresCookedData || IsRunningCommandlet())
+	else if (bRequiresCookedData | IsRunningCommandlet())
 	{
 		UE_LOG( LogAutomationTest, Log, TEXT( "Skipping unit tests for the cooked build and commandlet." ) );
 	}
