Fixed AutomationTests warning

diff --git a/Engine/Source/Runtime/Core/Private/Misc/AutomationTest.cpp b/Engine/Source/Runtime/Core/Private/Misc/AutomationTest.cpp
--- a/Engine/Source/Runtime/Core/Private/Misc/AutomationTest.cpp
+++ b/Engine/Source/Runtime/Core/Private/Misc/AutomationTest.cpp
@@ -477,7 +477,7 @@ bool FAutomationTestFramework::RunSmokeTests()
 	// Skip running on cooked platforms like mobile
 	//@todo - better determination of whether to run than requires cooked data
 	// Ensure there isn't another slow task in progress when trying to run unit tests
-	const bool bRequiresCookedData = FPlatformProperties::RequiresCookedData();
+	constexpr bool bRequiresCookedData = FPlatformProperties::RequiresCookedData();
 	if ((!bRequiresCookedData & !GIsSlowTask & !GIsPlayInEditorWorld & !FPlatformProperties::IsProgram() & !IsRunningCommandlet()) || bForceSmokeTests)
 	{
 		TArray<FAutomationTestInfo> TestInfo;
@@ -529,7 +529,7 @@ bool FAutomationTestFramework::RunSmokeTests()
 			FAutomationTestFramework::DumpAutomationTestExecutionInfo( OutExecutionInfoMap );
 		}
 	}
-	else if (bRequiresCookedData | IsRunningCommandlet())
+	else if (IsRunningCommandlet() || bRequiresCookedData)
 	{
 		UE_LOG( LogAutomationTest, Log, TEXT( "Skipping unit tests for the cooked build and commandlet." ) );
 	}
