Added missing constexpressification to FMacPlatformProperties

diff --git a/Engine/Source/Runtime/Core/Public/Mac/MacPlatformProperties.h b/Engine/Source/Runtime/Core/Public/Mac/MacPlatformProperties.h
--- a/Engine/Source/Runtime/Core/Public/Mac/MacPlatformProperties.h
+++ b/Engine/Source/Runtime/Core/Public/Mac/MacPlatformProperties.h
@@ -19,37 +19,37 @@ template<bool HAS_EDITOR_DATA, bool IS_DEDICATED_SERVER, bool IS_CLIENT_ONLY>
 struct FMacPlatformProperties
 	: public FGenericPlatformProperties
 {
-	static FORCEINLINE bool HasEditorOnlyData( )
+	constexpr static FORCEINLINE bool HasEditorOnlyData()
 	{
 		return HAS_EDITOR_DATA;
 	}
 
-	static FORCEINLINE const char* IniPlatformName( )
+	constexpr static FORCEINLINE const char* IniPlatformName()
 	{
 		return "Mac";
 	}
 
-	static FORCEINLINE const TCHAR* GetRuntimeSettingsClassName()
+	constexpr static FORCEINLINE const TCHAR* GetRuntimeSettingsClassName()
 	{
 		return TEXT("/Script/MacTargetPlatform.MacTargetSettings");
 	}
 
-	static FORCEINLINE bool IsGameOnly( )
+	constexpr static FORCEINLINE bool IsGameOnly()
 	{
 		return UE_GAME;
 	}
 
-	static FORCEINLINE bool IsServerOnly( )
+	constexpr static FORCEINLINE bool IsServerOnly()
 	{
 		return IS_DEDICATED_SERVER;
 	}
 
-	static FORCEINLINE bool IsClientOnly()
+	constexpr static FORCEINLINE bool IsClientOnly()
 	{
 		return IS_CLIENT_ONLY;
 	}
 
-	static FORCEINLINE const char* PlatformName( )
+	constexpr static FORCEINLINE const char* PlatformName()
 	{
 		if (IS_DEDICATED_SERVER)
 		{
@@ -69,41 +69,42 @@ struct FMacPlatformProperties
 		return "Mac";
 	}
 
-	static FORCEINLINE bool RequiresCookedData( )
+	constexpr static FORCEINLINE bool RequiresCookedData()
 	{
 		return !HAS_EDITOR_DATA;
 	}
 
-	static FORCEINLINE bool HasSecurePackageFormat()
+	constexpr static FORCEINLINE bool HasSecurePackageFormat()
 	{
 		return IS_DEDICATED_SERVER;
 	}
 
-	static FORCEINLINE bool SupportsMultipleGameInstances( )
+	constexpr static FORCEINLINE bool SupportsMultipleGameInstances()
 	{
 		return false;
 	}
-	static FORCEINLINE bool SupportsWindowedMode( )
+	
+	constexpr static FORCEINLINE bool SupportsWindowedMode()
 	{
 		return true;
 	}
 
-	static FORCEINLINE bool AllowsFramerateSmoothing()
+	constexpr static FORCEINLINE bool AllowsFramerateSmoothing()
 	{
 		return true;
 	}
 
-	static FORCEINLINE bool HasFixedResolution()
+	constexpr static FORCEINLINE bool HasFixedResolution()
 	{
 		return false;
 	}
 
-	static FORCEINLINE bool SupportsQuit()
+	constexpr static FORCEINLINE bool SupportsQuit()
 	{
 		return true;
 	}
 
-	static FORCEINLINE float GetVariantPriority()
+	constexpr static FORCEINLINE float GetVariantPriority()
 	{
 		if (IS_DEDICATED_SERVER)
 		{
@@ -123,12 +124,12 @@ struct FMacPlatformProperties
 		return 1.0f;
 	}
 	
-	static FORCEINLINE bool SupportsAudioStreaming()
+	constexpr static FORCEINLINE bool SupportsAudioStreaming()
 	{
 		return !IsServerOnly();
 	}
 
-	static FORCEINLINE bool SupportsMeshLODStreaming()
+	constexpr static FORCEINLINE bool SupportsMeshLODStreaming()
 	{
 		return !IsServerOnly() && !HasEditorOnlyData();
 	}
