Added CSameAs and CDecaysTo concepts

diff --git a/Engine/Source/Runtime/Core/Public/Concepts/DecaysTo.h b/Engine/Source/Runtime/Core/Public/Concepts/DecaysTo.h
new file mode 100644
--- /dev/null
+++ b/Engine/Source/Runtime/Core/Public/Concepts/DecaysTo.h
@@ -0,0 +1,18 @@
+// Copyright Epic Games, Inc. All Rights Reserved.
+
+#pragma once
+
+#include "Concepts/SameAs.h"
+
+#include <type_traits>
+
+namespace UE
+{
+
+/**
+ * Concept which is satisfied if and only if T decays to the same type as U.
+ */
+template <typename T, typename U>
+concept CDecaysTo = CSameAs<std::decay_t<T>, U>;
+
+} // UE
diff --git a/Engine/Source/Runtime/Core/Public/Concepts/SameAs.h b/Engine/Source/Runtime/Core/Public/Concepts/SameAs.h
new file mode 100644
--- /dev/null
+++ b/Engine/Source/Runtime/Core/Public/Concepts/SameAs.h
@@ -0,0 +1,18 @@
+// Copyright Epic Games, Inc. All Rights Reserved.
+
+#pragma once
+
+#include <type_traits>
+
+namespace UE
+{
+
+/**
+ * Concept which is satisfied if and only if T and U are the same type.
+ *
+ * We use this instead of std::same_as because <concepts> isn't a well supported header yet.
+ */
+template <typename T, typename U>
+concept CSameAs = std::is_same_v<T, U> && std::is_same_v<U, T>;
+
+} // UE
