Moved CConvertibleTo into namespace UE.

diff --git a/Engine/Source/Runtime/Core/Public/Concepts/ConvertibleTo.h b/Engine/Source/Runtime/Core/Public/Concepts/ConvertibleTo.h
--- a/Engine/Source/Runtime/Core/Public/Concepts/ConvertibleTo.h
+++ b/Engine/Source/Runtime/Core/Public/Concepts/ConvertibleTo.h
@@ -4,6 +4,9 @@
 
 #include <type_traits>
 
+namespace UE
+{
+
 /**
  * Concept which describes convertibility from one type to another.
  *
@@ -11,3 +14,5 @@
  */
 template <typename From, typename To>
 concept CConvertibleTo = std::is_convertible_v<From, To> && requires { static_cast<To>(std::declval<From>()); };
+
+} // UE
