Resubmitting TLS fix with additional change to make FPlatformMallocCrash::IsActive a static function instead of requiring getting an instance, to allow it to be safely called during static/TLS initialization.

diff --git a/Engine/Source/Runtime/Core/Private/Async/Fundamental/Scheduler.cpp b/Engine/Source/Runtime/Core/Private/Async/Fundamental/Scheduler.cpp
--- a/Engine/Source/Runtime/Core/Private/Async/Fundamental/Scheduler.cpp
+++ b/Engine/Source/Runtime/Core/Private/Async/Fundamental/Scheduler.cpp
@@ -8,6 +8,7 @@
 #include "Logging/LogMacros.h"
 #include "ProfilingDebugging/CsvProfiler.h"
 #include "HAL/MallocAnsi.h"
+#include "HAL/PlatformMallocCrash.h"
 #include "HAL/RunnableThread.h"
 #include "Misc/ScopeLock.h"
 #include "Misc/CommandLine.h"
@@ -92,41 +93,42 @@ namespace LowLevelTasks
 
 	FSchedulerTls::FTlsValuesHolder::FTlsValuesHolder()
 	{
-#if 0
-		while (!FPlatformMisc::IsDebuggerPresent())
+		// Avoid a deadlock on threads being spun up or down during a crash
+		if (!FPlatformMallocCrash::IsActive())
 		{
-			FPlatformProcess::SleepNoStats(0.100f);
-		}
-#endif
-		TlsValues = new FTlsValues();
+			TlsValues = new FTlsValues();
 
-		if (FImpl::ThreadTlsValuesMutex.TryLock())
-		{
-			FImpl::ProcessPendingTlsValuesNoLock();
-			TlsValues->LinkHead(FImpl::ThreadTlsValues);
-			FImpl::ThreadTlsValuesMutex.Unlock();
-		}
-		else
-		{
-			FImpl::PendingInsertTlsValues.ProduceItem(TlsValues);
+			if (FImpl::ThreadTlsValuesMutex.TryLock())
+			{
+				FImpl::ProcessPendingTlsValuesNoLock();
+				TlsValues->LinkHead(FImpl::ThreadTlsValues);
+				FImpl::ThreadTlsValuesMutex.Unlock();
+			}
+			else
+			{
+				FImpl::PendingInsertTlsValues.ProduceItem(TlsValues);
+			}
 		}
 	}
 
 	FSchedulerTls::FTlsValuesHolder::~FTlsValuesHolder()
 	{
-		if (FImpl::ThreadTlsValuesMutex.TryLock())
+		if (TlsValues && !FPlatformMallocCrash::IsActive())
 		{
-			FImpl::ProcessPendingTlsValuesNoLock();
-			TlsValues->Unlink();
-			FImpl::ThreadTlsValuesMutex.Unlock();
-			delete TlsValues;
-		}
-		else
-		{
-			FImpl::PendingDeleteTlsValues.ProduceItem(TlsValues);
-		}
+			if (FImpl::ThreadTlsValuesMutex.TryLock())
+			{
+				FImpl::ProcessPendingTlsValuesNoLock();
+				TlsValues->Unlink();
+				FImpl::ThreadTlsValuesMutex.Unlock();
+				delete TlsValues;
+			}
+			else
+			{
+				FImpl::PendingDeleteTlsValues.ProduceItem(TlsValues);
+			}
 
-		TlsValues = nullptr;
+			TlsValues = nullptr;
+		}
 	}
 
 	FORCENOINLINE bool& Private::FOversubscriptionTls::GetIsOversubscriptionAllowedRef()
diff --git a/Engine/Source/Runtime/Core/Private/GenericPlatform/GenericPlatformMallocCrash.cpp b/Engine/Source/Runtime/Core/Private/GenericPlatform/GenericPlatformMallocCrash.cpp
--- a/Engine/Source/Runtime/Core/Private/GenericPlatform/GenericPlatformMallocCrash.cpp
+++ b/Engine/Source/Runtime/Core/Private/GenericPlatform/GenericPlatformMallocCrash.cpp
@@ -10,6 +10,9 @@
 #include "ProfilingDebugging/MemoryTrace.h"
 #include "Algo/BinarySearch.h"
 
+// Keep a separate value to avoid TLS initialization order issues from calling Get() too early
+static bool GIsMallocCrashActive = false;
+
 // Suspending other threads on the first allocator call following a crash may help with debugging by freezing
 // the program state and avoiding further manipulation of a bad program state. However, crash handling on some
 // platforms requires threads to be able to tear down TLS, which in many cases involves deallocating memory.
@@ -314,11 +317,12 @@ void FGenericPlatformMallocCrash::SetAsGMalloc()
 		*GFixedMallocLocationPtr = nullptr; // this disables any fast-path inline allocators
 	}
 	CrashedThreadId = FPlatformTLS::GetCurrentThreadId();
+	GIsMallocCrashActive = true;
 }
 
-bool FGenericPlatformMallocCrash::IsActive() const
+bool FGenericPlatformMallocCrash::IsActive()
 {
-	return UE::Private::GMalloc == this;
+	return GIsMallocCrashActive;
 }
 
 void* FGenericPlatformMallocCrash::Malloc( SIZE_T Size, uint32 Alignment )
@@ -696,11 +700,12 @@ void FGenericStackBasedMallocCrash::SetAsGMalloc()
 		*GFixedMallocLocationPtr = nullptr; // this disables any fast-path inline allocators
 	}
 	UE::Private::GMalloc = this;
+	GIsMallocCrashActive = true;
 }
 
-bool FGenericStackBasedMallocCrash::IsActive() const
+bool FGenericStackBasedMallocCrash::IsActive()
 {
-	return UE::Private::GMalloc == this;
+	return GIsMallocCrashActive;
 }
 
 void* FGenericStackBasedMallocCrash::Malloc(SIZE_T Size, uint32 Alignment)
diff --git a/Engine/Source/Runtime/Core/Private/Logging/LogMacros.cpp b/Engine/Source/Runtime/Core/Private/Logging/LogMacros.cpp
--- a/Engine/Source/Runtime/Core/Private/Logging/LogMacros.cpp
+++ b/Engine/Source/Runtime/Core/Private/Logging/LogMacros.cpp
@@ -99,7 +99,7 @@ void FMsg::LogV(const ANSICHAR* File, int32 Line, const FLogCategoryName& Catego
 
 #if CSV_PROFILER_STATS
 		// Only update the CSV stat if we're not crashing, otherwise things can get messy
-		if (LIKELY(!FPlatformMallocCrash::Get().IsActive()))
+		if (LIKELY(!FPlatformMallocCrash::IsActive()))
 		{
 			CSV_CUSTOM_STAT(FMsgLogf, FMsgLogfCount, 1, ECsvCustomStatOp::Accumulate);
 		}
diff --git a/Engine/Source/Runtime/Core/Public/GenericPlatform/GenericPlatformMallocCrash.h b/Engine/Source/Runtime/Core/Public/GenericPlatform/GenericPlatformMallocCrash.h
--- a/Engine/Source/Runtime/Core/Public/GenericPlatform/GenericPlatformMallocCrash.h
+++ b/Engine/Source/Runtime/Core/Public/GenericPlatform/GenericPlatformMallocCrash.h
@@ -77,7 +77,7 @@ public:
 	 */
 	void SetAsGMalloc();
 
-	bool IsActive() const;
+	static bool IsActive();
 
 	// FMalloc interface.
 	virtual void* Malloc( SIZE_T Size, uint32 Alignment ) override;
@@ -179,7 +179,7 @@ struct FGenericStackBasedMallocCrash : public FMalloc
 
 	void SetAsGMalloc();
 
-	bool IsActive() const;
+	static bool IsActive();
 
 	virtual void* Malloc(SIZE_T Size, uint32 Alignment) override;
 
