Make UE_NO_UNIQUE_ADDRESS make empty types zero size on MSVC with [[msvc::no_unique_address]]

diff --git a/Engine/Source/Runtime/Core/Public/HAL/Platform.h b/Engine/Source/Runtime/Core/Public/HAL/Platform.h
--- a/Engine/Source/Runtime/Core/Public/HAL/Platform.h
+++ b/Engine/Source/Runtime/Core/Public/HAL/Platform.h
@@ -761,7 +761,9 @@
 
 /* Use before a class data member declaration allow it to be overlapped with other non-static data members or base class subobjects of its class. */
 #if !defined(UE_NO_UNIQUE_ADDRESS) && defined(__has_cpp_attribute)
-	#if __has_cpp_attribute(no_unique_address)
+	#if __has_cpp_attribute(msvc::no_unique_address)
+		#define UE_NO_UNIQUE_ADDRESS [[msvc::no_unique_address]] 
+	#elif __has_cpp_attribute(no_unique_address)
 		#define UE_NO_UNIQUE_ADDRESS [[no_unique_address]]
 	#endif
 #endif
