Mark DestructItem(s) constexpr

diff --git a/Engine/Source/Runtime/Core/Public/Templates/MemoryOps.h b/Engine/Source/Runtime/Core/Public/Templates/MemoryOps.h
--- a/Engine/Source/Runtime/Core/Public/Templates/MemoryOps.h
+++ b/Engine/Source/Runtime/Core/Public/Templates/MemoryOps.h
@@ -62,7 +62,7 @@ FORCENOINLINE void DefaultConstructItems(void* Address, SizeType Count)
  * @note: This function is optimized for values of T, and so will not dynamically dispatch destructor calls if T's destructor is virtual.
  */
 template <typename ElementType>
-FORCEINLINE void DestructItem(ElementType* Element)
+FORCEINLINE constexpr void DestructItem(ElementType* Element)
 {
 	if constexpr (sizeof(ElementType) == 0)
 	{
@@ -90,7 +90,7 @@ template <
 	typename SizeType
 	UE_REQUIRES(sizeof(ElementType) > 0 && std::is_trivially_destructible_v<ElementType>) // the sizeof here should improve the error messages we get when we try to call this function with incomplete types
 >
-FORCEINLINE void DestructItems(ElementType* Element, SizeType Count)
+FORCEINLINE constexpr void DestructItems(ElementType* Element, SizeType Count)
 {
 }
 template <
@@ -98,7 +98,7 @@ template <
 	typename SizeType
 	UE_REQUIRES(sizeof(ElementType) > 0 && !std::is_trivially_destructible_v<ElementType>) // the sizeof here should improve the error messages we get when we try to call this function with incomplete types
 >
-FORCENOINLINE void DestructItems(ElementType* Element, SizeType Count)
+FORCENOINLINE constexpr void DestructItems(ElementType* Element, SizeType Count)
 {
 	while (Count)
 	{
