Delete stat checks
- The checks exist to ensure stat data is processed continually, however this relies on a magic number of frames to maintain (HistoryFrames * 2 + 5 == 31) which no longer holds when some of this data arrives from an async task at the end-of-pipe.
- The checks are firing in various scenarios in the editor due to hitching that delays the arrival of the EOP data, particularly since there is no explicit synchronization between the game thread and EOP tasks. This happens much less frequently in a client build.
- The stats system recovers once the hitching has finished and the data from the EOP task arrives.

diff --git a/Engine/Source/Runtime/Core/Private/Stats/StatsData.cpp b/Engine/Source/Runtime/Core/Private/Stats/StatsData.cpp
--- a/Engine/Source/Runtime/Core/Private/Stats/StatsData.cpp
+++ b/Engine/Source/Runtime/Core/Private/Stats/StatsData.cpp
@@ -1073,27 +1073,6 @@ void FStatsThreadState::AddToHistoryAndEmpty(FStatPacketArray& NewData)
 			It.RemoveCurrent();
 		}
 	}
-
-	check(History.Num() <= HistoryFrames * 2 + 5);
-	check(CondensedStackHistory.Num() <= HistoryFrames * 2 + 5);
-	check(GoodFrames.Num() <= HistoryFrames * 2 + 5);
-#if DO_CHECK
-	if (BadFrames.Num() > HistoryFrames * 2 + 5)
-	{
-		UE_LOG(LogStats, Display, TEXT("BadFrames.Num(): %d, History.Num(): %d, HistoryFrames: %d"), BadFrames.Num(), History.Num(), HistoryFrames);
-		UE_LOG(LogStats, Display, TEXT("CurrentGameFrame: %lld, CurrentRenderFrame: %lld, CurrentEndOfPipeFrame: %lld, LastFullFrameMetaAndNonFrame: %lld"), CurrentGameFrame, CurrentRenderFrame, CurrentEndOfPipeFrame, LastFullFrameMetaAndNonFrame);
-		for (int64 BadFrame : BadFrames)
-		{
-			UE_LOG(LogStats, Display, TEXT("Bad frame: %lld"), BadFrame);
-		}
-		for (auto It = History.CreateIterator(); It; ++It)
-		{
-			const int64 HistoryFrame = It.Key();
-			UE_LOG(LogStats, Display, TEXT("History frame: %lld"), HistoryFrame);
-		}
-	}
-#endif
-	check(BadFrames.Num() <= HistoryFrames * 2 + 5);
 }
 
 
