ManualResetEvent: Fix a missing test of MaybeWaitingFlag before parking the thread

diff --git a/Engine/Source/Runtime/Core/Public/Async/ManualResetEvent.h b/Engine/Source/Runtime/Core/Public/Async/ManualResetEvent.h
--- a/Engine/Source/Runtime/Core/Public/Async/ManualResetEvent.h
+++ b/Engine/Source/Runtime/Core/Public/Async/ManualResetEvent.h
@@ -68,7 +68,11 @@ public:
 			if ((CurrentState & MaybeWaitingFlag) ||
 				State.compare_exchange_weak(CurrentState, MaybeWaitingFlag, std::memory_order_acq_rel))
 			{
-				ParkingLot::WaitUntil(&State, [this] { return !IsNotified(); }, []{}, WaitTime);
+				ParkingLot::WaitUntil(&State, [this]
+				{
+					uint8 NewState = State.load(std::memory_order_relaxed);
+					return !(NewState & IsNotifiedFlag) && (NewState & MaybeWaitingFlag);
+				}, []{}, WaitTime);
 			}
 		}
 	}
