IntrusiveMutex: Replaced params template with params concept

Added GetWaitAddress() to the params as an optional customization point.
Added TryWakeWaitingThread() to return whether a thread was woken up.

diff --git a/Engine/Source/Runtime/Core/Public/Async/IntrusiveMutex.h b/Engine/Source/Runtime/Core/Public/Async/IntrusiveMutex.h
--- a/Engine/Source/Runtime/Core/Public/Async/IntrusiveMutex.h
+++ b/Engine/Source/Runtime/Core/Public/Async/IntrusiveMutex.h
@@ -4,42 +4,110 @@
 
 #include "Async/Fundamental/Oversubscription.h"
 #include "Async/ParkingLot.h"
+#include "Concepts/DecaysTo.h"
+#include "Concepts/Integral.h"
 #include "CoreTypes.h"
 #include "HAL/PlatformProcess.h"
 
 #include <atomic>
+#include <type_traits>
 
 namespace UE
 {
 
-template <typename StateType>
-struct TIntrusiveMutexParams
-{
-	/** Flag that is written into the state with a bitwise or when locking the mutex. */
-	StateType IsLockedFlag = StateType(1) << 0;
-
-	/** Mask that is tested against the state when testing if the mutex is locked. 0 uses IsLockedFlag. */
-	StateType IsLockedMask = 0;
-
-	/** Flag that is set, cleared, and tested in the state to indicate that a lock may be waiting. */
-	StateType MayHaveWaitingLockFlag = StateType(1) << 1;
+template <typename ParamsType>
+using TIntrusiveMutexStateType_T = std::decay_t<decltype(ParamsType::IsLockedFlag)>;
 
-	/** How many times to spin before parking the thread. */
-	int32 SpinLimit = 40;
+template <typename ParamsType, typename StateType = TIntrusiveMutexStateType_T<ParamsType>>
+concept CIntrusiveMutexParams = requires
+{
+	requires CIntegral<StateType>;
+
+	// Required: constexpr static StateType IsLockedFlag = ...;
+	// Flag that is set in the state when the mutex is locked.
+	{ ParamsType::IsLockedFlag } -> CDecaysTo<StateType>;
+
+	// Required: constexpr static StateType MayHaveWaitingLockFlag = ...;
+	// Flag that is set in the state when a thread may be waiting to lock the mutex.
+	{ ParamsType::MayHaveWaitingLockFlag } -> CDecaysTo<StateType>;
+
+	// Optional: constexpr static StateType IsLockedMask = ...;
+	// Mask tested against State to determine if the mutex is locked. Defaults to IsLockedFlag.
+	requires
+		(!requires { ParamsType::IsLockedMask; }) ||
+		requires { { ParamsType::IsLockedMask } -> CDecaysTo<StateType>; };
+
+	// Optional: constexpr static int32 SpinLimit = ...;
+	// Maximum number of times to spin before waiting. Defaults to 40.
+	requires
+		(!requires { ParamsType::SpinLimit; }) ||
+		requires { { ParamsType::SpinLimit } -> CDecaysTo<int32>; };
+
+	// Optional: static const void* GetWaitAddress(std::atomic<StateType>& State);
+	// Returns the address to pass to ParkingLot wait and wake functions. Defaults to &State.
+	requires
+		(!requires { ParamsType::GetWaitAddress; }) ||
+		requires(std::atomic<StateType>& State) { { ParamsType::GetWaitAddress(State) } -> CDecaysTo<const void*>; };
 };
 
 /**
- * An 2-bit intrusive mutex that is not fair and does not support recursive locking.
+ * A 2-bit intrusive mutex that is not fair and does not support recursive locking.
+ *
+ * All bits of the state referenced by IsLockedFlag, IsLockedMask, and MayHaveWaitingLockFlag
+ * must be initialized to 0 or to values that are consistent with the functions being called.
  */
-template <typename StateType, TIntrusiveMutexParams<StateType> Params = TIntrusiveMutexParams<StateType>{}>
+template <CIntrusiveMutexParams ParamsType>
 class TIntrusiveMutex
 {
 	TIntrusiveMutex() = delete;
 
-	constexpr static StateType IsLockedFlag = Params.IsLockedFlag;
-	constexpr static StateType IsLockedMask = Params.IsLockedMask ? Params.IsLockedMask : Params.IsLockedFlag;
-	constexpr static StateType MayHaveWaitingLockFlag = Params.MayHaveWaitingLockFlag;
-	constexpr static int32 SpinLimit = Params.SpinLimit;
+	using StateType = TIntrusiveMutexStateType_T<ParamsType>;
+
+	constexpr static StateType IsLockedFlag = ParamsType::IsLockedFlag;
+	constexpr static StateType MayHaveWaitingLockFlag = ParamsType::MayHaveWaitingLockFlag;
+
+	constexpr static StateType IsLockedMask = []
+	{
+		if constexpr (requires { ParamsType::IsLockedMask; })
+		{
+			return ParamsType::IsLockedMask;
+		}
+		else
+		{
+			return ParamsType::IsLockedFlag;
+		}
+	}();
+
+	constexpr static int32 SpinLimit = []
+	{
+		if constexpr (requires { ParamsType::SpinLimit; })
+		{
+			return ParamsType::SpinLimit;
+		}
+		else
+		{
+			return 40;
+		}
+	}();
+
+	static_assert(IsLockedFlag && (IsLockedFlag & (IsLockedFlag - 1)) == 0, "IsLockedFlag must be one bit.");
+	static_assert(MayHaveWaitingLockFlag && (MayHaveWaitingLockFlag & (MayHaveWaitingLockFlag - 1)) == 0, "MayHaveWaitingLockFlag must be one bit.");
+	static_assert(IsLockedFlag != MayHaveWaitingLockFlag, "IsLockedFlag and MayHaveWaitingLockFlag must be different bits.");
+	static_assert((IsLockedMask & IsLockedFlag) == IsLockedFlag, "IsLockedMask must contain IsLockedFlag.");
+	static_assert((IsLockedMask & MayHaveWaitingLockFlag) == 0, "IsLockedMask must not contain MayHaveWaitingLockFlag.");
+	static_assert(SpinLimit >= 0, "SpinLimit must be non-negative.");
+
+	FORCEINLINE static const void* GetWaitAddress(const std::atomic<StateType>& State)
+	{
+		if constexpr (requires { ParamsType::GetWaitAddress; })
+		{
+			return ParamsType::GetWaitAddress(State);
+		}
+		else
+		{
+			return &State;
+		}
+	}
 
 public:
 	FORCEINLINE static bool IsLocked(const std::atomic<StateType>& State)
@@ -104,7 +172,7 @@ public:
 			LowLevelTasks::Private::FOversubscriptionAllowedScope _(false);
 
 			// Wait if the state has not changed. Either way, loop back and try to acquire the lock after trying to wait.
-			ParkingLot::Wait(&State, [&State]
+			ParkingLot::Wait(GetWaitAddress(State), [&State]
 			{
 				const StateType NewState = State.load(std::memory_order_relaxed);
 				return (NewState & IsLockedMask) && (NewState & MayHaveWaitingLockFlag);
@@ -126,14 +194,29 @@ public:
 
 	FORCEINLINE static void WakeWaitingThread(std::atomic<StateType>& State)
 	{
-		ParkingLot::WakeOne(&State, [&State](ParkingLot::FWakeState WakeState) -> uint64
+		ParkingLot::WakeOne(GetWaitAddress(State), [&State](ParkingLot::FWakeState WakeState) -> uint64
+		{
+			if (!WakeState.bHasWaitingThreads)
+			{
+				State.fetch_and(~MayHaveWaitingLockFlag, std::memory_order_relaxed);
+			}
+			return 0;
+		});
+	}
+
+	[[nodiscard]] FORCEINLINE static bool TryWakeWaitingThread(std::atomic<StateType>& State)
+	{
+		bool bDidWake = false;
+		ParkingLot::WakeOne(GetWaitAddress(State), [&State, &bDidWake](ParkingLot::FWakeState WakeState) -> uint64
 		{
 			if (!WakeState.bHasWaitingThreads)
 			{
 				State.fetch_and(~MayHaveWaitingLockFlag, std::memory_order_relaxed);
 			}
+			bDidWake = WakeState.bDidWake;
 			return 0;
 		});
+		return bDidWake;
 	}
 
 private:
