[AutoRTFM] Add automation test flag `SupportsAutoRTFM`.

When this flag is enabled and the test harness supports AutoRTFM, the test is run three times: normally, within a committed transaction, and within an aborted transaction.

If the test harness does not support AutoRTFM, the test is still run once, non-transactionally.

This flag makes it much easier to port over existing test suites to AutoRTFM. It also simplifies many of our existing test cases! I did not do a comprehensive sweep over our existing tests, but picked a few simple examples to convert over to make sure everything worked as expected.

diff --git a/Engine/Source/Programs/AutoRTFMEngineTests/Private/AutomationTestRunner.cpp b/Engine/Source/Programs/AutoRTFMEngineTests/Private/AutomationTestRunner.cpp
--- a/Engine/Source/Programs/AutoRTFMEngineTests/Private/AutomationTestRunner.cpp
+++ b/Engine/Source/Programs/AutoRTFMEngineTests/Private/AutomationTestRunner.cpp
@@ -44,7 +44,7 @@ FAutomationTestRunner::RunTests(const TCHAR* TestFilter)
 		const FAutomationTestInfo& TestInfo = TestInfos[TestIt];
 		const FString& TestFullPath = TestInfo.GetFullTestPath();
 
-		if (!TestFullPath.Contains("AutoRTFM"))
+		if (!(TestInfo.GetTestFlags() & EAutomationTestFlags::SupportsAutoRTFM) && !TestFullPath.Contains("AutoRTFM"))
 		{
 			continue;
 		}
diff --git a/Engine/Source/Runtime/Core/Private/Misc/AutomationTest.cpp b/Engine/Source/Runtime/Core/Private/Misc/AutomationTest.cpp
--- a/Engine/Source/Runtime/Core/Private/Misc/AutomationTest.cpp
+++ b/Engine/Source/Runtime/Core/Private/Misc/AutomationTest.cpp
@@ -184,6 +184,7 @@ CORE_API const TMap<FString, EAutomationTestFlags>& EAutomationTestFlags_GetTest
 		{ TEXT("ServerContext"),          EAutomationTestFlags::ServerContext},
 		{ TEXT("CommandletContext"),      EAutomationTestFlags::CommandletContext},
 		{ TEXT("ProgramContext"),         EAutomationTestFlags::ProgramContext},
+		{ TEXT("SupportsAutoRTFM"),       EAutomationTestFlags::SupportsAutoRTFM},
 		{ TEXT("ApplicationContextMask"), EAutomationTestFlags_ApplicationContextMask},
 		{ TEXT("NonNullRHI"),             EAutomationTestFlags::NonNullRHI},
 		{ TEXT("RequiresUser"),           EAutomationTestFlags::RequiresUser},
@@ -1079,6 +1080,52 @@ void FAutomationTestFramework::InternalStartTest( const FString& InTestToRun, co
 			// Run the test!
 			bTestSuccessful = CurrentTest->RunTest(Parameters);
 		}
+
+#if UE_AUTORTFM
+		const bool bTestWithAutoRTFM = !!(CurrentTest->GetTestFlags() & EAutomationTestFlags::SupportsAutoRTFM);
+		if (bTestSuccessful && bTestWithAutoRTFM && AutoRTFM::ForTheRuntime::IsAutoRTFMRuntimeEnabled())
+		{
+			{
+				TRACE_CPUPROFILER_EVENT_SCOPE_TEXT(*FString::Printf(TEXT("AutomationTest %s - AutoRTFM Commit"), *CurrentTest->GetBeautifiedTestName()));
+				AutoRTFM::ETransactionResult Outcome = AutoRTFM::Transact([&]
+				{
+					if (!CurrentTest->RunTest(Parameters))
+					{
+						bTestSuccessful = false;
+					}
+				});
+				if (Outcome != AutoRTFM::ETransactionResult::Committed)
+				{
+					UE_LOG(LogAutomationTest, Warning, TEXT("%s %s: Test failed while committing an AutoRTFM transaction!"), *CurrentTest->GetBeautifiedTestName(), *Parameters);
+					bTestSuccessful = false;
+				}
+			}
+			{
+				TRACE_CPUPROFILER_EVENT_SCOPE_TEXT(*FString::Printf(TEXT("AutomationTest %s - AutoRTFM Abort"), *CurrentTest->GetBeautifiedTestName()));
+				AutoRTFM::ETransactionResult Outcome = AutoRTFM::Transact([&]
+				{
+					if (!CurrentTest->RunTest(Parameters))
+					{
+						UE_AUTORTFM_ONABORT(this)
+						{
+							bTestSuccessful = false;
+						};
+					}
+					AutoRTFM::AbortTransaction();
+				});
+				if (Outcome != AutoRTFM::ETransactionResult::AbortedByRequest)
+				{
+					UE_LOG(LogAutomationTest, Warning, TEXT("%s %s: Test failed while aborting an AutoRTFM transaction!"), *CurrentTest->GetBeautifiedTestName(), *Parameters);
+					bTestSuccessful = false;
+				}
+			}
+			if (!bTestSuccessful)
+			{
+				UE_LOG(LogAutomationTest, Warning, TEXT("%s %s: Test passes normally, but fails under AutoRTFM!"), *CurrentTest->GetBeautifiedTestName(), *Parameters);
+			}
+		}
+#endif
+
 	}
 }
 
diff --git a/Engine/Source/Runtime/Core/Public/Misc/AutomationTest.h b/Engine/Source/Runtime/Core/Public/Misc/AutomationTest.h
--- a/Engine/Source/Runtime/Core/Public/Misc/AutomationTest.h
+++ b/Engine/Source/Runtime/Core/Public/Misc/AutomationTest.h
@@ -105,10 +105,14 @@ enum class EAutomationTestFlags
 	// Test requires a user instigated session
 	RequiresUser = 0x00000200,
 
-	//~ One-off flag to allow for fast disabling of tests without commenting code out
-	// Temp disabled and never returns for a filter
+	//~ One-off flags
+	// Fast disable: turns off a test without commenting code out. Never returns for a filter.
 	Disabled = 0x00010000,
 
+	// When AutoRTFM is enabled, the test will be executed within a transaction that commits, 
+	// and within another transaction that aborts. (The test will also execute non-transactionally.)
+	SupportsAutoRTFM = 0x00020000,
+
 	//~ Priority of the test
 	// The highest priority possible. Showstopper/blocker.
 	CriticalPriority			= 0x00100000,
diff --git a/Engine/Source/Runtime/Core/Tests/Internationalization/TextTest.cpp b/Engine/Source/Runtime/Core/Tests/Internationalization/TextTest.cpp
--- a/Engine/Source/Runtime/Core/Tests/Internationalization/TextTest.cpp
+++ b/Engine/Source/Runtime/Core/Tests/Internationalization/TextTest.cpp
@@ -24,7 +24,7 @@
 
 #define LOCTEXT_NAMESPACE "Core.Tests.TextFormatTest"
 
-class  FTextTestClass {
+class FTextTestClass {
 public:
 		FText FormatWithoutArguments(const FText& Pattern)
 		{
@@ -855,7 +855,6 @@ TEST_CASE_NAMED(FTextRoundingTest, "System::Core::Misc::TextRounding", "[.][Edit
 	I18N.RestoreCultureState(OriginalCultureState);
 }
 
-
 TEST_CASE_NAMED(FTextPaddingTest, "System::Core::Misc::TextPadding", "[.][EditorContext][ClientContext][EngineFilter]")
 {
 	FInternationalization& I18N = FInternationalization::Get();
@@ -1156,13 +1155,13 @@ public:
 	}
 };
 
-TEST_CASE_NAMED(FTextNumericParsingTest, "System::Core::Misc::TextNumericParsing", "[.][EditorContext][ClientContext][EngineFilter]")
+TEST_CASE_NAMED(FTextNumericParsingTest, "System::Core::Misc::TextNumericParsing", "[.][EditorContext][ClientContext][EngineFilter][ServerContext][CommandletContext][SupportsAutoRTFM]")
 {
 	FTextNumericParsingTestClass Instance;
 	Instance.TextNumericParsingTest();
 }
 
-static void RunTestCasesForSystem_Core_Misc_TextStringification()
+TEST_CASE_NAMED(FTextStringificationTest, "System::Core::Misc::TextStringification", "[.][EditorContext][ClientContext][EngineFilter][ServerContext][CommandletContext][SupportsAutoRTFM]")
 {
 	FInternationalization& I18N = FInternationalization::Get();
 
@@ -1280,12 +1279,7 @@ static void RunTestCasesForSystem_Core_Misc_TextStringification()
 	I18N.RestoreCultureState(OriginalCultureState);
 }
 
-TEST_CASE_NAMED(FTextStringificationTest, "System::Core::Misc::TextStringification", "[.][EditorContext][ClientContext][EngineFilter]")
-{
-	RunTestCasesForSystem_Core_Misc_TextStringification();
-}
-
-static void RunTestCasesForSystem_Core_Misc_TextFormatArgModifiers()
+TEST_CASE_NAMED(FTextFormatArgModifierTest, "System::Core::Misc::TextFormatArgModifiers", "[.][EditorContext][ClientContext][EngineFilter][ServerContext][CommandletContext][SupportsAutoRTFM]")
 {
 	auto EnsureValidResult = [&](const FString& InResult, const FString& InExpected, const FString& InName, const FString& InDescription)
 	{
@@ -1406,11 +1400,6 @@ static void RunTestCasesForSystem_Core_Misc_TextFormatArgModifiers()
 	I18N.RestoreCultureState(OriginalCultureState);
 }
 
-TEST_CASE_NAMED(FTextFormatArgModifierTest, "System::Core::Misc::TextFormatArgModifiers", "[.][EditorContext][ClientContext][EngineFilter]")
-{
-	RunTestCasesForSystem_Core_Misc_TextFormatArgModifiers();
-}
-
 #if UE_ENABLE_ICU
 
 TEST_CASE_NAMED(FICUSanitizationTest, "System::Core::Misc::ICUSanitization", "[.][EditorContext][ClientContext][EngineFilter]")
@@ -1574,61 +1563,6 @@ TEST_CASE_NAMED(FICUTextTest, "System::Core::Misc::ICUText", "[.][EditorContext]
 
 #endif // #if UE_ENABLE_ICU
 
-#if WITH_DEV_AUTOMATION_TESTS
-
-IMPLEMENT_SIMPLE_AUTOMATION_TEST(FAutoRTFMTextTests, "AutoRTFM + FText", \
-								 EAutomationTestFlags::EngineFilter | EAutomationTestFlags::ClientContext | \
-								 EAutomationTestFlags::ServerContext | EAutomationTestFlags::CommandletContext)
-
-bool FAutoRTFMTextTests::RunTest(const FString& Parameters)
-{
-	if (!AutoRTFM::ForTheRuntime::IsAutoRTFMRuntimeEnabled())
-	{
-		ExecutionInfo.AddEvent(FAutomationEvent(EAutomationEventType::Info, TEXT("SKIPPED 'FAutoRTFMTextTests' test. AutoRTFM disabled.")));
-		return true;
-	}
-
-	// Verify that FText works properly when used in an AutoRTFM transaction, both committing and aborting.
-	AutoRTFM::Commit([]
-	{
-		FTextNumericParsingTestClass Instance;
-		Instance.TextNumericParsingTest();
-	});
-
-	AutoRTFM::Transact([]
-	{
-		FTextNumericParsingTestClass Instance;
-		Instance.TextNumericParsingTest();
-		AutoRTFM::AbortTransaction();
-	});
-
-	AutoRTFM::Commit([]
-	{
-		RunTestCasesForSystem_Core_Misc_TextFormatArgModifiers();
-	});
-
-	AutoRTFM::Transact([]
-	{
-		RunTestCasesForSystem_Core_Misc_TextFormatArgModifiers();
-		AutoRTFM::AbortTransaction();
-	});
-
-	AutoRTFM::Commit([]
-	{
-		RunTestCasesForSystem_Core_Misc_TextStringification();
-	});
-
-	AutoRTFM::Transact([]
-	{
-		RunTestCasesForSystem_Core_Misc_TextStringification();
-		AutoRTFM::AbortTransaction();
-	});
-
-	return true;
-}
-
-#endif // #if WITH_DEV_AUTOMATION_TESTS
-
 #undef LOCTEXT_NAMESPACE 
 
 #endif //WITH_TESTS
diff --git a/Engine/Source/Runtime/Engine/Private/Tests/AutoRTFM/AutoRTFMDataTableTests.cpp b/Engine/Source/Runtime/Engine/Private/Tests/AutoRTFM/AutoRTFMDataTableTests.cpp
--- a/Engine/Source/Runtime/Engine/Private/Tests/AutoRTFM/AutoRTFMDataTableTests.cpp
+++ b/Engine/Source/Runtime/Engine/Private/Tests/AutoRTFM/AutoRTFMDataTableTests.cpp
@@ -1,38 +1,17 @@
 // Copyright Epic Games, Inc. All Rights Reserved.
 
-#include "Misc/AutomationTest.h"
+#include "Tests/TestHarnessAdapter.h"
 #include "AutoRTFM.h"
 #include "Engine/DataTable.h"
 
 #if WITH_DEV_AUTOMATION_TESTS
 
-IMPLEMENT_SIMPLE_AUTOMATION_TEST(FAutoRTFMDataTableTests, "AutoRTFM + UDataTable", EAutomationTestFlags::EngineFilter | EAutomationTestFlags::ClientContext | EAutomationTestFlags::ServerContext | EAutomationTestFlags::CommandletContext)
-
-bool FAutoRTFMDataTableTests::RunTest(const FString & Parameters)
+TEST_CASE_NAMED(FAutoRTFMDataTableTests, "AutoRTFM.UDataTable", "[ClientContext][ServerContext][CommandletContext][EngineFilter][SupportsAutoRTFM]")
 {
-	if (!AutoRTFM::ForTheRuntime::IsAutoRTFMRuntimeEnabled())
-	{
-		ExecutionInfo.AddEvent(FAutomationEvent(EAutomationEventType::Info, TEXT("SKIPPED 'FAutoRTFMDataTableTests' test. AutoRTFM disabled.")));
-		return true;
-	}
-
 	UDataTable* Object = NewObject<UDataTable>();
 	Object->RowStruct = NewObject<UScriptStruct>();
 
-	AutoRTFM::ETransactionResult Result = AutoRTFM::Transact([&]
-	{
-		Object->EmptyTable();
-		AutoRTFM::AbortTransaction();
-	});
-	TestTrueExpr(AutoRTFM::ETransactionResult::AbortedByRequest == Result);
-
-	Result = AutoRTFM::Transact([&]
-	{
-		Object->EmptyTable();
-	});
-	TestTrueExpr(AutoRTFM::ETransactionResult::Committed == Result);
-
-	return true;
+	Object->EmptyTable();
 }
 
 #endif //WITH_DEV_AUTOMATION_TESTS
diff --git a/Engine/Source/Runtime/Engine/Private/Tests/AutoRTFM/AutoRTFMSlateTests.cpp b/Engine/Source/Runtime/Engine/Private/Tests/AutoRTFM/AutoRTFMSlateTests.cpp
--- a/Engine/Source/Runtime/Engine/Private/Tests/AutoRTFM/AutoRTFMSlateTests.cpp
+++ b/Engine/Source/Runtime/Engine/Private/Tests/AutoRTFM/AutoRTFMSlateTests.cpp
@@ -1,33 +1,20 @@
 // Copyright Epic Games, Inc. All Rights Reserved.
 
 #include "AutoRTFMTesting.h"
-#include "Misc/AutomationTest.h"
+#include "Tests/TestHarnessAdapter.h"
 #include "AutoRTFM.h"
 
 #include "Framework/Text/CharRangeList.h"
 
 #if WITH_DEV_AUTOMATION_TESTS
 
-IMPLEMENT_SIMPLE_AUTOMATION_TEST(FAutoRTFMSlateTests, "AutoRTFM + Slate", EAutomationTestFlags::EngineFilter | EAutomationTestFlags::ClientContext | EAutomationTestFlags::ServerContext | EAutomationTestFlags::CommandletContext)
-
-bool FAutoRTFMSlateTests::RunTest(const FString & Parameters)
+TEST_CASE_NAMED(FAutoRTFMSlateTests, "AutoRTFM.Slate", "[EngineFilter][ClientContext][ServerContext][CommandletContext][SupportsAutoRTFM]")
 {
-	if (!AutoRTFM::ForTheRuntime::IsAutoRTFMRuntimeEnabled())
-	{
-		ExecutionInfo.AddEvent(FAutomationEvent(EAutomationEventType::Info, TEXT("SKIPPED 'FAutoRTFMSlateTests' test. AutoRTFM disabled.")));
-		return true;
-	}
-
 	// Test for SOL-7842
-	AutoRTFM::Testing::Commit([&]
-	{
-		FCharRangeList CharRangeList;
-		TestTrueExpr(CharRangeList.IsEmpty());
-		CharRangeList.InitializeFromString(TEXT("a-zA-Z0-9._"));
-		TestFalseExpr(CharRangeList.IsEmpty());
-	});
-	
-	return true;
+	FCharRangeList CharRangeList;
+	TestTrueExpr(CharRangeList.IsEmpty());
+	CharRangeList.InitializeFromString(TEXT("a-zA-Z0-9._"));
+	TestFalseExpr(CharRangeList.IsEmpty());
 }
 
-#endif //WITH_DEV_AUTOMATION_TESTS
+#endif // WITH_DEV_AUTOMATION_TESTS
diff --git a/Engine/Source/Runtime/GeometryCore/Private/Tests/AutoRTFMConvexHullTests.cpp b/Engine/Source/Runtime/GeometryCore/Private/Tests/AutoRTFMConvexHullTests.cpp
--- a/Engine/Source/Runtime/GeometryCore/Private/Tests/AutoRTFMConvexHullTests.cpp
+++ b/Engine/Source/Runtime/GeometryCore/Private/Tests/AutoRTFMConvexHullTests.cpp
@@ -1,6 +1,6 @@
 // Copyright Epic Games, Inc. All Rights Reserved.
 
-#include "Misc/AutomationTest.h"
+#include "Tests/TestHarnessAdapter.h"
 #include "AutoRTFM.h"
 #include "CompGeom/ConvexHull3.h"
 #include "Math/UnrealMathUtility.h"
@@ -8,36 +8,23 @@
 
 #if WITH_DEV_AUTOMATION_TESTS
 
-IMPLEMENT_SIMPLE_AUTOMATION_TEST(FAutoRTFMConvexHullTests, "AutoRTFM + TConvexHull3", EAutomationTestFlags::EngineFilter | EAutomationTestFlags::ClientContext | EAutomationTestFlags::ServerContext | EAutomationTestFlags::CommandletContext)
-
-bool FAutoRTFMConvexHullTests::RunTest(const FString& Parameters)
+TEST_CASE_NAMED(FAutoRTFMConvexHullTests, "AutoRTFM.TConvexHull3", "[EngineFilter][ClientContext][ServerContext][CommandletContext][SupportsAutoRTFM]")
 {
-	if (!AutoRTFM::ForTheRuntime::IsAutoRTFMRuntimeEnabled())
-	{
-		ExecutionInfo.AddEvent(FAutomationEvent(EAutomationEventType::Info, TEXT("SKIPPED 'FAutoRTFMConvexHullTests' test. AutoRTFM disabled.")));
-		return true;
-	}
-
 	using Vec3 = UE::Math::TVector<float>;
 
-	AutoRTFM::Commit([this]
+	static const Vec3 Pyramid[] = 
 	{
-		static const Vec3 Pyramid[] = 
-		{
-			{ 0.f,  0.f,  0.f},
-			{20.f,  0.f,  0.f},
-			{ 0.f, 20.f,  0.f},
-			{20.f, 20.f,  0.f},
-			{10.f, 10.f,  3.f},
-		};
-
-		UE::Geometry::TConvexHull3<float> Hull;
-		double Volume = Hull.ComputeVolume(Pyramid);
+		{ 0.f,  0.f,  0.f},
+		{20.f,  0.f,  0.f},
+		{ 0.f, 20.f,  0.f},
+		{20.f, 20.f,  0.f},
+		{10.f, 10.f,  3.f},
+	};
 
-		TestTrueExpr(FMath::IsNearlyEqual(Volume, 400.));
-	});
+	UE::Geometry::TConvexHull3<float> Hull;
+	double Volume = Hull.ComputeVolume(Pyramid);
 
-	return true;
+	TestTrueExpr(FMath::IsNearlyEqual(Volume, 400.)); 
 }
 
 #endif // WITH_DEV_AUTOMATION_TESTS
