[TSAN/IoDispatcher]
- Suppress benign race on read in TickCsv without incurring a fetch_add/fetch_sub cost to maintain the value atomically since only one dispatcher thread exists

diff --git a/Engine/Source/Runtime/Core/Private/IO/IoDispatcher.cpp b/Engine/Source/Runtime/Core/Private/IO/IoDispatcher.cpp
--- a/Engine/Source/Runtime/Core/Private/IO/IoDispatcher.cpp
+++ b/Engine/Source/Runtime/Core/Private/IO/IoDispatcher.cpp
@@ -168,7 +168,7 @@ private:
 	}
 #endif
 
-	int64 PendingIoRequests = 0;
+	TSAN_ATOMIC(int64) PendingIoRequests = 0;
 #if COUNTERSTRACE_ENABLED
 	FCountersTrace::FCounterInt PendingIoRequestsCounter;
 #endif
