Added UE_IF_CONSTEVAL and UE_IF_NOT_CONSTEVAL

We do not have consteval if support everywhere yet and cannot fall back to std::is_constant_evaluated() because we target platforms that do not support that yet either.

diff --git a/Engine/Source/Runtime/Core/Public/HAL/Platform.h b/Engine/Source/Runtime/Core/Public/HAL/Platform.h
--- a/Engine/Source/Runtime/Core/Public/HAL/Platform.h
+++ b/Engine/Source/Runtime/Core/Public/HAL/Platform.h
@@ -759,6 +759,16 @@
 	#define UE_CONSTEVAL constexpr
 #endif
 
+/* Macro wrapper for the consteval if statement which isn't yet present on all compilers.
+   Can be used in place of std::is_constant_evaluated() which isn't yet available on all platforms. */
+#if defined(__cpp_if_consteval)
+	#define UE_IF_CONSTEVAL if consteval
+	#define UE_IF_NOT_CONSTEVAL if !consteval
+#else
+	#define UE_IF_CONSTEVAL if (__builtin_is_constant_evaluated())
+	#define UE_IF_NOT_CONSTEVAL if (!__builtin_is_constant_evaluated())
+#endif
+
 /* Use before a class data member declaration allow it to be overlapped with other non-static data members or base class subobjects of its class. */
 #if !defined(UE_NO_UNIQUE_ADDRESS) && defined(__has_cpp_attribute)
 	#if __has_cpp_attribute(msvc::no_unique_address)
