Fixed: Mark SceneDepth usage in post-process materials with a WorldPosition expressions

diff --git a/Engine/Source/Runtime/Core/Private/UObject/FortniteReleaseBranchCustomObjectVersion.cpp b/Engine/Source/Runtime/Core/Private/UObject/FortniteReleaseBranchCustomObjectVersion.cpp
--- a/Engine/Source/Runtime/Core/Private/UObject/FortniteReleaseBranchCustomObjectVersion.cpp
+++ b/Engine/Source/Runtime/Core/Private/UObject/FortniteReleaseBranchCustomObjectVersion.cpp
@@ -8,7 +8,7 @@ TMap<FGuid, FGuid> FFortniteReleaseBranchCustomObjectVersion::GetSystemGuids()
 
 	SystemGuids.Add(DevGuids.Niagara_LatestScriptCompileVersion, FGuid("8BC8C557C87C4F629468F10EA4161F0A"));
 	SystemGuids.Add(DevGuids.SkeletalMeshDerivedDataVersion, FGuid("C96DEBAD45E3443D804017D4C29F7C4B"));
-	SystemGuids.Add(DevGuids.MaterialTranslationDDCVersion, FGuid("888A502A66F6401580B27451F211C46B"));
+	SystemGuids.Add(DevGuids.MaterialTranslationDDCVersion, FGuid("3EF80806F2E945B1ADAFC49154058C00"));
 	SystemGuids.Add(DevGuids.GLOBALSHADERMAP_DERIVEDDATA_VER, FGuid("0A0A1BF5D2A94105BB395EF0DA7CAEB9"));
 	SystemGuids.Add(DevGuids.MATERIALSHADERMAP_DERIVEDDATA_VER, FGuid("C0E1F59386894C9D855B85178B2553CB"));
 	SystemGuids.Add(DevGuids.STATICMESH_DERIVEDDATA_VER, FGuid("0D7C30DE6E5F489A831CEBF1346725C4"));
diff --git a/Engine/Source/Runtime/Engine/Private/Materials/HLSLMaterialTranslator.cpp b/Engine/Source/Runtime/Engine/Private/Materials/HLSLMaterialTranslator.cpp
--- a/Engine/Source/Runtime/Engine/Private/Materials/HLSLMaterialTranslator.cpp
+++ b/Engine/Source/Runtime/Engine/Private/Materials/HLSLMaterialTranslator.cpp
@@ -6650,6 +6650,8 @@ int32 FHLSLMaterialTranslator::WorldPosition(EWorldPositionIncludedOffsets World
 	}
 	
 	bUsesVertexPosition = true;
+	// In PP material WorldPosition is reconstructed from SceneDepth
+	bUsesSceneDepth |= (ShaderFrequency == SF_Pixel && Material->GetMaterialDomain() == MD_PostProcess);
 
 	TArray<FStringFormatArg> FormatArgs =
 	{
