Modify TOptional to allow in-place construction in constexpr contexts which requires us not to use reinterpret_cast.
We can't use std::construct_at on some platforms so only support in-place construction from values, and not other modifications in a constexpr context.
No support for copying or moving other TOptionals.
Maintain trivial destruction for trivial types. This requires the use of a templated base class on compilers that don't support constrained destructors.
Added compile time tests for supported constexpr operations.

diff --git a/Engine/Source/Runtime/Core/Public/Misc/IntrusiveUnsetOptionalState.h b/Engine/Source/Runtime/Core/Public/Misc/IntrusiveUnsetOptionalState.h
--- a/Engine/Source/Runtime/Core/Public/Misc/IntrusiveUnsetOptionalState.h
+++ b/Engine/Source/Runtime/Core/Public/Misc/IntrusiveUnsetOptionalState.h
@@ -60,6 +60,11 @@
 
 #include "Misc/OptionalFwd.h"
 
+namespace UE::Core::Private
+{
+	struct FOptional;
+}
+
 struct FIntrusiveUnsetOptionalState
 {
 	// Defined in the global namespace for convenience and consistency
@@ -68,6 +73,8 @@ struct FIntrusiveUnsetOptionalState
 	template <typename>
 	friend struct TOptional;
 
+	friend struct UE::Core::Private::FOptional;
+
 private:
 	explicit FIntrusiveUnsetOptionalState() = default;
 };
diff --git a/Engine/Source/Runtime/Core/Public/Misc/Optional.h b/Engine/Source/Runtime/Core/Public/Misc/Optional.h
--- a/Engine/Source/Runtime/Core/Public/Misc/Optional.h
+++ b/Engine/Source/Runtime/Core/Public/Misc/Optional.h
@@ -9,6 +9,7 @@
 #include "Templates/MemoryOps.h"
 #include "Templates/UnrealTemplate.h"
 #include "Serialization/Archive.h"
+#include <type_traits>
 
 inline constexpr FNullOpt NullOpt{0};
 
@@ -21,13 +22,6 @@ namespace UE::Core::Private
 			b = std::decay_t<typename T::IntrusiveUnsetOptionalStateType>::bHasIntrusiveUnsetOptionalState
 		);
 	};
-
-	template <size_t N>
-	struct TNonIntrusiveOptionalStorage
-	{
-		uint8 Storage[N];
-		bool bIsSet;
-	};
 }
 
 template <typename T>
@@ -48,30 +42,149 @@ constexpr bool HasIntrusiveUnsetOptionalState()
 	}
 }
 
+namespace UE::Core::Private
+{
+	// Empty structure to replace bIsSet for intrusive optional types 
+	struct FEmpty {};
+
+	// Shared code for optionals to handle multiple implementations of IsSet depending on compiler support for constrained destructors
+	// This class can be a friend of FIntrusiveUnsetOptionalState rather than all possible base classes.
+	struct FOptional
+	{
+		template<typename Derived>
+		[[nodiscard]] FORCEINLINE static constexpr bool IsSet(Derived* This)
+		{
+			if constexpr (Derived::bUsingIntrusiveUnsetState)
+			{
+				return !(This->TypedValue == FIntrusiveUnsetOptionalState{});
+			}
+			else
+			{
+				return This->bIsSet;
+			}
+		}
+	};
+
+#if !PLATFORM_COMPILER_SUPPORTS_CONSTRAINED_DESTRUCTORS
+	// These base types are necessary for trivial destruction on compilers that don't yet support use of
+	// constraints to select an explicitly defaulted trivial destructor.
+	// Once such compilers are no longer supported, all (!PLATFORM_COMPILER_SUPPORTS_CONSTRAINED_DESTRUCTORS)
+	// code in this file can be deleted.
+	template <typename OptionalType, bool bIsTriviallyDestructible = std::is_trivially_destructible_v<OptionalType>>
+	struct TOptionalBase 
+	{
+		union
+		{
+			OptionalType TypedValue;
+		};
+		// Set flag takes up no space if bIsUsingIntrusiveUnsetState
+		static constexpr bool bUsingIntrusiveUnsetState = HasIntrusiveUnsetOptionalState<OptionalType>();
+		UE_NO_UNIQUE_ADDRESS std::conditional_t<bUsingIntrusiveUnsetState, FEmpty, bool> bIsSet = {};
+	
+		constexpr TOptionalBase()
+		{
+		}
+
+		template<typename... ArgTypes>
+		constexpr TOptionalBase(EInPlace, ArgTypes&&... InArgs)
+			: TypedValue(Forward<ArgTypes>(InArgs)...)
+		{
+		}
+
+		// Explicitly constexpr destructor for trivially destructible OptionalType
+		constexpr ~TOptionalBase() = default;
+
+		[[nodiscard]] constexpr bool IsSet() const
+		{
+			return FOptional::IsSet(this);
+		}
+
+		FORCEINLINE constexpr void DestroyValue()
+		{
+			// Type is known to be trivially destructible
+		}
+	};
+
+	template <typename OptionalType>
+	struct TOptionalBase<OptionalType, false> 
+	{
+		union
+		{
+			OptionalType TypedValue;
+		};
+		// Set flag takes up no space if bIsUsingIntrusiveUnsetState
+		static constexpr bool bUsingIntrusiveUnsetState = HasIntrusiveUnsetOptionalState<OptionalType>();
+		UE_NO_UNIQUE_ADDRESS std::conditional_t<bUsingIntrusiveUnsetState, FEmpty, bool> bIsSet = {};
+
+		constexpr TOptionalBase()
+		{
+		}
+
+		template<typename... ArgTypes>
+		constexpr TOptionalBase(EInPlace, ArgTypes&&... InArgs)
+			: TypedValue(Forward<ArgTypes>(InArgs)...)
+		{
+		}
+
+		constexpr ~TOptionalBase()
+		{
+			if (IsSet())
+			{
+				DestroyValue();
+			}
+		}
+
+		[[nodiscard]] constexpr bool IsSet() const
+		{
+			return FOptional::IsSet(this);
+		}
+
+		FORCEINLINE constexpr void DestroyValue()
+		{
+			DestructItem(std::addressof(this->TypedValue));
+		}
+	};
+#endif
+}
+
 /**
  * When we have an optional value IsSet() returns true, and GetValue() is meaningful.
  * Otherwise GetValue() is not meaningful.
  */
 template<typename OptionalType>
 struct TOptional
+#if !PLATFORM_COMPILER_SUPPORTS_CONSTRAINED_DESTRUCTORS
+    : private UE::Core::Private::TOptionalBase<OptionalType>
+#endif
 {
 private:
+#if !PLATFORM_COMPILER_SUPPORTS_CONSTRAINED_DESTRUCTORS
+    using Super = UE::Core::Private::TOptionalBase<OptionalType>;
+#endif
+
+	using FOptional = UE::Core::Private::FOptional;
+	friend FOptional;
 	static constexpr bool bUsingIntrusiveUnsetState = HasIntrusiveUnsetOptionalState<OptionalType>();
 
 public:
 	using ElementType = OptionalType;
 
 	/** Construct an OptionalType with a valid value. */
-	[[nodiscard]] TOptional(const OptionalType& InValue)
+	[[nodiscard]] constexpr TOptional(const OptionalType& InValue)
 		: TOptional(InPlace, InValue)
 	{
 	}
-	[[nodiscard]] TOptional(OptionalType&& InValue)
+	[[nodiscard]] constexpr TOptional(OptionalType&& InValue)
 		: TOptional(InPlace, MoveTempIfPossible(InValue))
 	{
 	}
 	template <typename... ArgTypes>
-	[[nodiscard]] explicit TOptional(EInPlace, ArgTypes&&... Args)
+	[[nodiscard]] explicit constexpr TOptional(EInPlace, ArgTypes&&... Args)
+#if PLATFORM_COMPILER_SUPPORTS_CONSTRAINED_DESTRUCTORS
+		: TypedValue(Forward<ArgTypes>(Args)...)
+#else
+		: Super(InPlace, Forward<ArgTypes>(Args)...)
+#endif
 	{
 		// If this fails to compile when trying to call TOptional(EInPlace, ...) with a non-public constructor,
 		// do not make TOptional a friend.
@@ -93,11 +206,10 @@ public:
 		//     // Won't compile if the caller doesn't have access to FMyType::FPrivateToken
 		//     TOptional<FMyType> Opt(InPlace, FMyType::FPrivateToken{}, 5, 3.14f, TEXT("Banana"));
 		//
-		::new((void*)&Value) OptionalType(Forward<ArgTypes>(Args)...);
 
 		if constexpr (!bUsingIntrusiveUnsetState)
 		{
-			Value.bIsSet = true;
+			this->bIsSet = true;
 		}
 		else
 		{
@@ -107,67 +219,68 @@ public:
 	}
 	
 	/** Construct an OptionalType with an invalid value. */
-	[[nodiscard]] TOptional(FNullOpt)
+	[[nodiscard]] constexpr TOptional(FNullOpt)
 		: TOptional()
 	{
 	}
 
+	/** Construct an OptionalType intrusively with no value; i.e. unset */
+	[[nodiscard]] constexpr TOptional() requires bUsingIntrusiveUnsetState
+#if PLATFORM_COMPILER_SUPPORTS_CONSTRAINED_DESTRUCTORS
+		: TypedValue(FIntrusiveUnsetOptionalState{})
+#else
+		: Super(InPlace, FIntrusiveUnsetOptionalState{})
+#endif
+	{
+	}
+
 	/** Construct an OptionalType with no value; i.e. unset */
-	[[nodiscard]] TOptional()
+	[[nodiscard]] constexpr TOptional() requires (!bUsingIntrusiveUnsetState)
 	{
-		if constexpr (bUsingIntrusiveUnsetState)
-		{
-			::new ((void*)&Value) OptionalType(FIntrusiveUnsetOptionalState{});
-		}
-		else
-		{
-			Value.bIsSet = false;
-		}
 	}
 
-	~TOptional()
+#if PLATFORM_COMPILER_SUPPORTS_CONSTRAINED_DESTRUCTORS
+	constexpr ~TOptional() requires std::is_trivially_destructible_v<OptionalType> = default;
+	constexpr ~TOptional() requires (!std::is_trivially_destructible_v<OptionalType>)
 	{
-		Reset();
+		// Destroy value but do not reconstruct an empty intrusive optional in its place
+		if (IsSet())
+		{
+			DestroyValue();
+		}
 	}
+#else
+	constexpr ~TOptional() = default; // Defer to base type to handle trivial/non-trivial destructor
+#endif
 
 	/** Copy/Move construction */
 	[[nodiscard]] TOptional(const TOptional& Other)
+		: TOptional()
 	{
 		bool bLocalIsSet = Other.IsSet();
 		if constexpr (!bUsingIntrusiveUnsetState)
 		{
-			Value.bIsSet = bLocalIsSet;
+			this->bIsSet = bLocalIsSet;
 		}
 		if (bLocalIsSet)
 		{
-			::new((void*)&Value) OptionalType(*(const OptionalType*)&Other.Value);
-		}
-		else
-		{
-			if constexpr (bUsingIntrusiveUnsetState)
-			{
-				::new ((void*)&Value) OptionalType(FIntrusiveUnsetOptionalState{});
-			}
+			::new((void*)std::addressof(this->TypedValue)) OptionalType(Other.TypedValue);
 		}
+		// Default constructor initialized unset intrusive optional if necessary
 	}
 	[[nodiscard]] TOptional(TOptional&& Other)
+		: TOptional()
 	{
 		bool bLocalIsSet = Other.IsSet();
 		if constexpr (!bUsingIntrusiveUnsetState)
 		{
-			Value.bIsSet = bLocalIsSet;
+			this->bIsSet = bLocalIsSet;
 		}
 		if (bLocalIsSet)
 		{
-			::new((void*)&Value) OptionalType(MoveTempIfPossible(*(OptionalType*)&Other.Value));
-		}
-		else
-		{
-			if constexpr (bUsingIntrusiveUnsetState)
-			{
-				::new ((void*)&Value) OptionalType(FIntrusiveUnsetOptionalState{});
-			}
+			::new((void*)std::addressof(this->TypedValue)) OptionalType(MoveTempIfPossible(Other.TypedValue));
 		}
+		// Default constructor initialized unset intrusive optional if necessary
 	}
 
 	TOptional& operator=(const TOptional& Other)
@@ -203,7 +316,7 @@ public:
 
 	TOptional& operator=(const OptionalType& InValue)
 	{
-		if (&InValue != (const OptionalType*)&Value)
+		if (std::addressof(InValue) != std::addressof(this->TypedValue))
 		{
 			Emplace(InValue);
 		}
@@ -211,7 +324,7 @@ public:
 	}
 	TOptional& operator=(OptionalType&& InValue)
 	{
-		if (&InValue != (const OptionalType*)&Value)
+		if (std::addressof(InValue) != std::addressof(this->TypedValue))
 		{
 			Emplace(MoveTempIfPossible(InValue));
 		}
@@ -225,11 +338,11 @@ public:
 			DestroyValue();
 			if constexpr (bUsingIntrusiveUnsetState)
 			{
-				::new((void*)&Value) OptionalType(FIntrusiveUnsetOptionalState{});
+				::new((void*)std::addressof(this->TypedValue)) OptionalType(FIntrusiveUnsetOptionalState{});
 			}
 			else
 			{
-				Value.bIsSet = false;
+				this->bIsSet = false;
 			}
 		}
 	}
@@ -272,11 +385,11 @@ public:
 		//     // Won't compile if the caller doesn't have access to FMyType::FPrivateToken
 		//     Opt.Emplace(FMyType::FPrivateToken{}, 5, 3.14f, TEXT("Banana"));
 		//
-		OptionalType* Result = ::new((void*)&Value) OptionalType(Forward<ArgsType>(Args)...);
+		OptionalType* Result = ::new((void*)std::addressof(this->TypedValue)) OptionalType(Forward<ArgsType>(Args)...);
 
 		if constexpr (!bUsingIntrusiveUnsetState)
 		{
-			Value.bIsSet = true;
+			this->bIsSet = true;
 		}
 		else
 		{
@@ -287,7 +400,7 @@ public:
 		return *Result;
 	}
 
-	[[nodiscard]] friend bool operator==(const TOptional& Lhs, const TOptional& Rhs)
+	[[nodiscard]] friend constexpr bool operator==(const TOptional& Lhs, const TOptional& Rhs)
 	{
 		bool bIsLhsSet = Lhs.IsSet();
 		bool bIsRhsSet = Rhs.IsSet();
@@ -301,10 +414,10 @@ public:
 			return true;
 		}
 
-		return (*(const OptionalType*)&Lhs.Value) == (*(const OptionalType*)&Rhs.Value);
+		return Lhs.TypedValue == Rhs.TypedValue;
 	}
 
-	[[nodiscard]] friend bool operator!=(const TOptional& Lhs, const TOptional& Rhs)
+	[[nodiscard]] friend constexpr bool operator!=(const TOptional& Lhs, const TOptional& Rhs)
 	{
 		return !(Lhs == Rhs);
 	}
@@ -340,63 +453,61 @@ public:
 	}
 
 	/** @return true when the value is meaningful; false if calling GetValue() is undefined. */
-	[[nodiscard]] bool IsSet() const
+#if PLATFORM_COMPILER_SUPPORTS_CONSTRAINED_DESTRUCTORS
+	[[nodiscard]] constexpr bool IsSet() const
 	{
-		if constexpr (bUsingIntrusiveUnsetState)
-		{
-			return !(*(const OptionalType*)&Value == FIntrusiveUnsetOptionalState{});
-		}
-		else
-		{
-			return Value.bIsSet;
-		}
+		return FOptional::IsSet(this);
 	}
-	[[nodiscard]] FORCEINLINE explicit operator bool() const
+#else
+	using Super::IsSet;
+#endif
+
+	[[nodiscard]] FORCEINLINE explicit constexpr operator bool() const
 	{
 		return IsSet();
 	}
 
 	/** @return The optional value; undefined when IsSet() returns false. */
-	[[nodiscard]] OptionalType& GetValue()
+	[[nodiscard]] constexpr OptionalType& GetValue()
 	{
 		checkf(IsSet(), TEXT("It is an error to call GetValue() on an unset TOptional. Please either check IsSet() or use Get(DefaultValue) instead."));
-		return *(OptionalType*)&Value;
+		return this->TypedValue;
 	}
-	[[nodiscard]] FORCEINLINE const OptionalType& GetValue() const
+	[[nodiscard]] FORCEINLINE constexpr const OptionalType& GetValue() const
 	{
 		return const_cast<TOptional*>(this)->GetValue();
 	}
 
-	[[nodiscard]] OptionalType* operator->()
+	[[nodiscard]] constexpr OptionalType* operator->()
 	{
-		return &GetValue();
+		return std::addressof(GetValue());
 	}
-	[[nodiscard]] FORCEINLINE const OptionalType* operator->() const
+	[[nodiscard]] FORCEINLINE constexpr const OptionalType* operator->() const
 	{
 		return const_cast<TOptional*>(this)->operator->();
 	}
 
-	[[nodiscard]] OptionalType& operator*()
+	[[nodiscard]] constexpr OptionalType& operator*()
 	{
 		return GetValue();
 	}
-	[[nodiscard]] FORCEINLINE const OptionalType& operator*() const
+	[[nodiscard]] FORCEINLINE constexpr const OptionalType& operator*() const
 	{
 		return const_cast<TOptional*>(this)->operator*();
 	}
 
 	/** @return The optional value when set; DefaultValue otherwise. */
-	[[nodiscard]] const OptionalType& Get(const OptionalType& DefaultValue UE_LIFETIMEBOUND) const UE_LIFETIMEBOUND
+	[[nodiscard]] constexpr const OptionalType& Get(const OptionalType& DefaultValue UE_LIFETIMEBOUND) const UE_LIFETIMEBOUND
 	{
-		return IsSet() ? *(const OptionalType*)&Value : DefaultValue;
+		return IsSet() ? this->TypedValue : DefaultValue;
 	}
 
 	/** @return A pointer to the optional value when set, nullptr otherwise. */
-	[[nodiscard]] OptionalType* GetPtrOrNull()
+	[[nodiscard]] constexpr OptionalType* GetPtrOrNull()
 	{
-		return IsSet() ? (OptionalType*)&Value : nullptr;
+		return IsSet() ? std::addressof(this->TypedValue) : nullptr;
 	}
-	[[nodiscard]] FORCEINLINE const OptionalType* GetPtrOrNull() const
+	[[nodiscard]] FORCEINLINE constexpr const OptionalType* GetPtrOrNull() const
 	{
 		return const_cast<TOptional*>(this)->GetPtrOrNull();
 	}
@@ -405,20 +516,24 @@ private:
 	/** 
 	 * Destroys the value, must only be called if the value is set, and callers must then mark the value unset or construct a new value in place. 
 	 */
-	FORCEINLINE void DestroyValue()
+#if PLATFORM_COMPILER_SUPPORTS_CONSTRAINED_DESTRUCTORS
+	FORCEINLINE constexpr void DestroyValue()
 	{
-		if constexpr (bUsingIntrusiveUnsetState)
-		{
-			DestructItem((OptionalType*)&Value);
-		}
-		else
-		{
-			DestructItem((OptionalType*)Value.Storage);
-		}
+		DestructItem(std::addressof(this->TypedValue));
 	}
+#else
+	using Super::DestroyValue;
+#endif
 
-	using ValueStorageType = std::conditional_t<bUsingIntrusiveUnsetState, uint8[sizeof(OptionalType)], UE::Core::Private::TNonIntrusiveOptionalStorage<sizeof(OptionalType)>>;
-	alignas(OptionalType) ValueStorageType Value;
+	// If we can have an explicitly conditional trivial destructor, we can declare our members inline and avoid the base class
+#if PLATFORM_COMPILER_SUPPORTS_CONSTRAINED_DESTRUCTORS
+	union
+	{
+		OptionalType TypedValue;
+	};
+	// Set flag takes up no space if bIsUsingIntrusiveUnsetState, otherwise is zero-initialized to false
+	UE_NO_UNIQUE_ADDRESS std::conditional_t<bUsingIntrusiveUnsetState, UE::Core::Private::FEmpty, bool> bIsSet = {};
+#endif
 };
 
 template<typename OptionalType>
diff --git a/Engine/Source/Runtime/Core/Tests/Misc/OptionalTest.cpp b/Engine/Source/Runtime/Core/Tests/Misc/OptionalTest.cpp
new file mode 100644
--- /dev/null
+++ b/Engine/Source/Runtime/Core/Tests/Misc/OptionalTest.cpp
@@ -0,0 +1,139 @@
+// Copyright Epic Games, Inc. All Rights Reserved.
+
+#include "HAL/Platform.h"
+#include <type_traits>
+#if WITH_TESTS
+
+#include "CoreTypes.h"
+#include "Containers/UnrealString.h"
+#include "Internationalization/Text.h"
+#include "Math/Vector.h"
+#include "Misc/LowLevelTestAdapter.h"
+#include "Misc/Optional.h"
+#include "Templates/AlignmentTemplates.h"
+#include "Tests/TestHarnessAdapter.h"
+
+// Tests for traits and constexpr use of TOptional
+namespace UE::Core::Private
+{
+	// TIsTOptional_V of optionals
+	static_assert(TIsTOptional_V<      TOptional<float>>, "Expected TIsTOptional_V<TOptional> to be true");
+	static_assert(TIsTOptional_V<const TOptional<float>>, "Expected TIsTOptional_V<const TOptional> to be true");
+
+	// TIsTOptional_V of non-optionals (references are not optionals)
+	static_assert(TIsTOptional_V<      TOptional<float>&>  == false, "Expected TIsTOptional_V<TOptional&> to be false");
+	static_assert(TIsTOptional_V<const TOptional<float>&>  == false, "Expected TIsTOptional_V<const TOptional&> to be false");
+	static_assert(TIsTOptional_V<      TOptional<float>&&> == false, "Expected TIsTOptional_V<TOptional&&> to be false");
+	static_assert(TIsTOptional_V<      bool>               == false, "Expected TIsTOptional_V<non-TOptional> to be false");
+
+	// Ensure that TOptional has a trivial destructor if the inner type does
+	static_assert( std::is_trivially_destructible_v<TOptional<int>>);
+	static_assert( std::is_trivially_destructible_v<TOptional<FVector>>);
+	static_assert(!std::is_trivially_destructible_v<FString>);
+	static_assert(!std::is_trivially_destructible_v<TOptional<FString>>);
+	static_assert(!std::is_trivially_destructible_v<TOptional<FText>>);
+
+    template<bool bIntrusive, bool bTrackDestruction>
+	struct TTestValue
+	{
+        static inline int32 DestructorCount = 0;
+		static constexpr inline bool bHasIntrusiveUnsetOptionalState = bIntrusive;
+		using IntrusiveUnsetOptionalStateType = TTestValue<true, bTrackDestruction>;
+
+		constexpr TTestValue(FIntrusiveUnsetOptionalState) requires bIntrusive
+		: Value(0)
+		{
+		}
+		constexpr TTestValue(int32 InValue)
+		: Value(InValue)
+		{
+            static_assert(!bIntrusive || sizeof(TOptional<TTestValue>) == sizeof(TTestValue), "Expected intrusive optional to be the same size as the inner type");
+            static_assert(bIntrusive || sizeof(TOptional<TTestValue>) > sizeof(TTestValue), "Expected non-intrusive optional to be larger than the inner type");
+            static_assert(bIntrusive || sizeof(TOptional<TTestValue>) <= Align(sizeof(TTestValue) + 1, alignof(TTestValue)),
+                "Expected non-intrusive optional to be larger than the inner type by at most the alignment of the inner type");
+		}
+        constexpr TTestValue(const TTestValue&) = default;
+        constexpr TTestValue& operator=(const TTestValue&) = default;
+        constexpr TTestValue(TTestValue&&) = default;
+        constexpr TTestValue& operator=(TTestValue&&) = default;
+
+        // Explicit destructor so not trivially destructible, but constexpr-destructible
+        constexpr ~TTestValue()
+        {
+            if constexpr (bTrackDestruction)
+            {
+                ++DestructorCount;
+            }
+        }
+
+        [[nodiscard]] constexpr bool operator==(FIntrusiveUnsetOptionalState) const requires bIntrusive
+        {
+            return Value == 0;
+        }
+
+        [[nodiscard]] constexpr bool operator==(const TTestValue&) const = default;
+
+		int32 Value;
+	};
+
+    template<bool bTrackDestruction = false>
+    using TIntrusive = TTestValue<true, bTrackDestruction>;
+    template<bool bTrackDestruction = false>
+	using TNonIntrusive = TTestValue<false, bTrackDestruction>;
+
+    static_assert(HasIntrusiveUnsetOptionalState<TIntrusive<>>());
+    static_assert(!HasIntrusiveUnsetOptionalState<TNonIntrusive<>>());
+
+    // Tests for constexpr use of optional
+    template<typename InnerType>
+    constexpr bool TestConstexprOptionalUse()
+    {
+        InnerType Value{32};
+        TOptional<InnerType> O1; // Default construction
+        TOptional<InnerType> O2{NullOpt}; // Explicit empty construction;
+        TOptional<InnerType> O3(Value); // Copying inner type
+        TOptional<InnerType> O4(Forward<InnerType>(InnerType(30))); // Moving inner type
+        TOptional<InnerType> O5(InPlace, 29); // In-place construction
+        if constexpr (false) // At present, copy and move construction are not constexpr-safe
+        {
+            TOptional<InnerType> O6(O3); // Copy construction
+            TOptional<InnerType> O7(MoveTemp(O4)); // Move construction
+        }
+        bool bTestResult = (O1 == O2) && (O4 != O5);
+        bTestResult = !O1.IsSet() && O3.IsSet() && bTestResult;
+        const TOptional<InnerType> C1(InPlace, 33);
+        const InnerType& I1 = C1.GetValue();
+        O5.GetValue().Value++;
+        O5->Value = C1->Value;
+        bTestResult = ((*O5) == (*C1)) && bTestResult;
+        const InnerType& I2 = O1.Get(Value);
+        *O4.GetPtrOrNull() = *C1.GetPtrOrNull();
+
+        return bTestResult;
+    }
+
+    static_assert(TestConstexprOptionalUse<TIntrusive<>>(), "Constexpr tests for TOptional<TIntrusive> failed");
+    static_assert(TestConstexprOptionalUse<TNonIntrusive<>>(), "Constexpr tests for TOptional<TNonIntrusive> failed");
+};
+
+TEST_CASE_NAMED(FTOptionalTest_NonIntrusive, "System::Core::TOptional::NonIntrusive", "[Core][TOptional][SmokeFilter]")
+{
+    using FNonIntrusiveDestructible = UE::Core::Private::TNonIntrusive<true>;
+    FNonIntrusiveDestructible::DestructorCount = 0;
+    { 
+        TOptional<FNonIntrusiveDestructible> Unset;
+    }
+    CHECK_MESSAGE("Unset optional is not destroyed", FNonIntrusiveDestructible::DestructorCount == 0);
+}
+
+TEST_CASE_NAMED(FTOptionalTest_Intrusive, "System::Core::TOptional::Intrusive", "[Core][TOptional][SmokeFilter]")
+{
+    using FIntrusiveDestructible = UE::Core::Private::TIntrusive<true>;
+    FIntrusiveDestructible::DestructorCount = 0;
+    { 
+        TOptional<FIntrusiveDestructible> Unset;
+    }
+    CHECK_MESSAGE("Unset optional is not destroyed", FIntrusiveDestructible::DestructorCount == 0);
+}
+
+#endif // WITH_TESTS
\ No newline at end of file
