Changed calls to FMemory in FAlignedAllocator to UE::Private::GMalloc so we'll bypass LLM

diff --git a/Engine/Source/Runtime/Core/Public/Experimental/ConcurrentLinearAllocator.h b/Engine/Source/Runtime/Core/Public/Experimental/ConcurrentLinearAllocator.h
--- a/Engine/Source/Runtime/Core/Public/Experimental/ConcurrentLinearAllocator.h
+++ b/Engine/Source/Runtime/Core/Public/Experimental/ConcurrentLinearAllocator.h
@@ -41,12 +41,18 @@ struct FAlignedAllocator
 
 	FORCEINLINE static void* Malloc(SIZE_T Size, uint32 Alignment)
 	{
-		return FMemory::Malloc(Size, Alignment);
+		if (UNLIKELY(!UE::Private::GMalloc))
+		{
+			// There is no public function to create GMalloc and this will do it for us.
+			FMemory::Free(FMemory::Malloc(0));
+			check(UE::Private::GMalloc);
+		}
+		return UE::Private::GMalloc->Malloc(Size, Alignment);
 	}
 
 	FORCEINLINE static void Free(void* Pointer, SIZE_T Size)
 	{
-		FMemory::Free(Pointer);
+		UE::Private::GMalloc->Free(Pointer);
 	}
 };
 
