Mutex: Use TIntrusiveMutex to implement parts of FMutex and FSharedMutex

diff --git a/Engine/Source/Runtime/Core/Private/Async/Mutex.cpp b/Engine/Source/Runtime/Core/Private/Async/Mutex.cpp
--- a/Engine/Source/Runtime/Core/Private/Async/Mutex.cpp
+++ b/Engine/Source/Runtime/Core/Private/Async/Mutex.cpp
@@ -2,69 +2,25 @@
 
 #include "Async/Mutex.h"
 
-#include "Async/Fundamental/Oversubscription.h"
-#include "Async/ParkingLot.h"
-#include "HAL/PlatformProcess.h"
+#include "Async/IntrusiveMutex.h"
 
 namespace UE
 {
 
-void FMutex::LockSlow()
+struct FMutex::FParams
 {
-	constexpr int32 SpinLimit = 40;
-	int32 SpinCount = 0;
-	for (uint8 CurrentState = State.load(std::memory_order_relaxed);;)
-	{
-		// Try to acquire the lock if it was unlocked, even if there are waiting threads.
-		// Acquiring the lock despite the waiting threads means that this lock is not FIFO and thus not fair.
-		if (LIKELY(!(CurrentState & IsLockedFlag)))
-		{
-			if (LIKELY(State.compare_exchange_weak(CurrentState, CurrentState | IsLockedFlag, std::memory_order_acquire, std::memory_order_relaxed)))
-			{
-				return;
-			}
-			continue;
-		}
+	constexpr static uint8 IsLockedFlag = FMutex::IsLockedFlag;
+	constexpr static uint8 MayHaveWaitingLockFlag = FMutex::MayHaveWaitingLockFlag;
+};
 
-		// Spin up to the spin limit while there are no waiting threads.
-		if (LIKELY(!(CurrentState & MayHaveWaitingLockFlag) && SpinCount < SpinLimit))
-		{
-			FPlatformProcess::YieldThread();
-			++SpinCount;
-			CurrentState = State.load(std::memory_order_relaxed);
-			continue;
-		}
-
-		// Store that there are waiting threads. Restart if the state has changed since it was loaded.
-		if (LIKELY(!(CurrentState & MayHaveWaitingLockFlag)))
-		{
-			if (UNLIKELY(!State.compare_exchange_weak(CurrentState, CurrentState | MayHaveWaitingLockFlag, std::memory_order_relaxed)))
-			{
-				continue;
-			}
-			CurrentState |= MayHaveWaitingLockFlag;
-		}
-
-		// Do not enter oversubscription during a wait on a mutex since the wait is generally too short
-		// for it to matter and it can worsen performance a lot for heavily contended locks.
-		LowLevelTasks::Private::FOversubscriptionAllowedScope _(false);
-
-		// Wait if the state has not changed. Either way, loop back and try to acquire the lock after trying to wait.
-		ParkingLot::Wait(&State, [this, CurrentState] { return State.load(std::memory_order_relaxed) == CurrentState; }, []{});
-		CurrentState = State.load(std::memory_order_relaxed);
-	}
+FORCENOINLINE void FMutex::LockSlow()
+{
+	TIntrusiveMutex<FParams>::LockLoop(State);
 }
 
 FORCENOINLINE void FMutex::WakeWaitingThread()
 {
-	ParkingLot::WakeOne(&State, [this](ParkingLot::FWakeState WakeState) -> uint64
-	{
-		if (!WakeState.bHasWaitingThreads)
-		{
-			State.fetch_and(~MayHaveWaitingLockFlag, std::memory_order_relaxed);
-		}
-		return 0;
-	});
+	TIntrusiveMutex<FParams>::WakeWaitingThread(State);
 }
 
 } // UE
diff --git a/Engine/Source/Runtime/Core/Private/Async/SharedMutex.cpp b/Engine/Source/Runtime/Core/Private/Async/SharedMutex.cpp
--- a/Engine/Source/Runtime/Core/Private/Async/SharedMutex.cpp
+++ b/Engine/Source/Runtime/Core/Private/Async/SharedMutex.cpp
@@ -3,16 +3,24 @@
 #include "Async/SharedMutex.h"
 
 #include "Async/Fundamental/Oversubscription.h"
+#include "Async/IntrusiveMutex.h"
 #include "Async/ParkingLot.h"
 #include "HAL/PlatformProcess.h"
 
 namespace UE
 {
 
-inline const void* FSharedMutex::GetLockAddress() const
+struct FSharedMutex::FParams
 {
-	return &State;
-}
+	constexpr static uint32 IsLockedFlag = FSharedMutex::IsLockedFlag;
+	constexpr static uint32 IsLockedMask = FSharedMutex::IsLockedFlag | FSharedMutex::SharedLockCountMask;
+	constexpr static uint32 MayHaveWaitingLockFlag = FSharedMutex::MayHaveWaitingLockFlag;
+
+	inline static const void* GetWaitAddress(const std::atomic<uint32>& State)
+	{
+		return &State;
+	}
+};
 
 inline const void* FSharedMutex::GetSharedLockAddress() const
 {
@@ -23,51 +31,7 @@ inline const void* FSharedMutex::GetSharedLockAddress() const
 
 void FSharedMutex::LockSlow()
 {
-	constexpr int32 SpinLimit = 40;
-	int32 SpinCount = 0;
-	for (uint32 CurrentState = State.load(std::memory_order_relaxed);;)
-	{
-		// Try to acquire the lock if it was unlocked, even if there are waiting threads.
-		// Acquiring the lock despite the waiting threads means that this lock is not FIFO and thus not fair.
-		if (LIKELY(!(CurrentState & (IsLockedFlag | SharedLockCountMask))))
-		{
-			if (LIKELY(State.compare_exchange_weak(CurrentState, CurrentState | IsLockedFlag, std::memory_order_acquire, std::memory_order_relaxed)))
-			{
-				return;
-			}
-			continue;
-		}
-
-		// Spin up to the spin limit while there are no waiting threads.
-		if (LIKELY(!(CurrentState & MayHaveWaitingLockFlag) && SpinCount < SpinLimit))
-		{
-			FPlatformProcess::YieldThread();
-			++SpinCount;
-			CurrentState = State.load(std::memory_order_relaxed);
-			continue;
-		}
-
-		// Store that there are waiting threads. Restart if the state has changed since it was loaded.
-		if (LIKELY(!(CurrentState & MayHaveWaitingLockFlag)))
-		{
-			if (UNLIKELY(!State.compare_exchange_weak(CurrentState, CurrentState | MayHaveWaitingLockFlag, std::memory_order_relaxed)))
-			{
-				continue;
-			}
-			CurrentState |= MayHaveWaitingLockFlag;
-		}
-
-		// Do not enter oversubscription during a wait on a mutex since the wait is generally too short
-		// for it to matter and it can worsen performance a lot for heavily contended locks.
-		LowLevelTasks::Private::FOversubscriptionAllowedScope _(false);
-
-		// Wait if the state has not changed. Either way, loop back and try to acquire the lock after trying to wait.
-		ParkingLot::Wait(GetLockAddress(), [this, CurrentState]
-		{
-			return State.load(std::memory_order_relaxed) == CurrentState;
-		}, []{});
-		CurrentState = State.load(std::memory_order_relaxed);
-	}
+	TIntrusiveMutex<FParams>::LockLoop(State);
 }
 
 void FSharedMutex::LockSharedSlow()
@@ -120,16 +84,7 @@ void FSharedMutex::LockSharedSlow()
 
 FORCENOINLINE void FSharedMutex::WakeWaitingThread()
 {
-	ParkingLot::WakeOne(GetLockAddress(), [this](ParkingLot::FWakeState WakeState) -> uint64
-	{
-		if (!WakeState.bDidWake)
-		{
-			// Keep the flag until no thread wakes, otherwise shared locks may win before
-			// an exclusive lock has a chance.
-			State.fetch_and(~MayHaveWaitingLockFlag, std::memory_order_relaxed);
-		}
-		return 0;
-	});
+	TIntrusiveMutex<FParams>::WakeWaitingThread(State);
 }
 
 FORCENOINLINE void FSharedMutex::WakeWaitingThreads(uint32 LastState)
@@ -137,19 +92,7 @@ FORCENOINLINE void FSharedMutex::WakeWaitingThreads(uint32 LastState)
 	if (LastState & MayHaveWaitingLockFlag)
 	{
 		// Wake one thread that is waiting to acquire an exclusive lock.
-		bool bDidWake = false;
-		ParkingLot::WakeOne(GetLockAddress(), [this, &bDidWake](ParkingLot::FWakeState WakeState) -> uint64
-		{
-			if (!WakeState.bDidWake)
-			{
-				// Keep the flag until no thread wakes, otherwise shared locks may win before
-				// an exclusive lock has a chance.
-				State.fetch_and(~MayHaveWaitingLockFlag, std::memory_order_relaxed);
-			}
-			bDidWake = WakeState.bDidWake;
-			return 0;
-		});
-		if (bDidWake)
+		if (TIntrusiveMutex<FParams>::TryWakeWaitingThread(State))
 		{
 			return;
 		}
diff --git a/Engine/Source/Runtime/Core/Public/Async/Mutex.h b/Engine/Source/Runtime/Core/Public/Async/Mutex.h
--- a/Engine/Source/Runtime/Core/Public/Async/Mutex.h
+++ b/Engine/Source/Runtime/Core/Public/Async/Mutex.h
@@ -13,7 +13,6 @@ namespace UE
 
 /**
  * A one-byte mutex that is not fair and does not support recursive locking.
- * Note: Changes to this class should also be ported to FExternalMutex.
  */
 class FMutex final
 {
@@ -66,6 +65,8 @@ private:
 	UE_API void LockSlow();
 	UE_API void WakeWaitingThread();
 
+	struct FParams;
+
 	static constexpr uint8 IsLockedFlag = 1 << 0;
 	static constexpr uint8 MayHaveWaitingLockFlag = 1 << 1;
 
diff --git a/Engine/Source/Runtime/Core/Public/Async/SharedMutex.h b/Engine/Source/Runtime/Core/Public/Async/SharedMutex.h
--- a/Engine/Source/Runtime/Core/Public/Async/SharedMutex.h
+++ b/Engine/Source/Runtime/Core/Public/Async/SharedMutex.h
@@ -111,9 +111,10 @@ private:
 	UE_API void WakeWaitingThread();
 	UE_API void WakeWaitingThreads(uint32 LastState);
 
-	const void* GetLockAddress() const;
 	const void* GetSharedLockAddress() const;
 
+	struct FParams;
+
 	static constexpr uint32 IsLockedFlag = 1 << 0;
 	static constexpr uint32 MayHaveWaitingLockFlag = 1 << 1;
 	static constexpr uint32 MayHaveWaitingSharedLockFlag = 1 << 2;
