Minor optimization in ConcurrentLinearAllocator to inline TLS Cleanup code for blocks that have InlineBlockAllocation set to true

diff --git a/Engine/Source/Runtime/Core/Public/Experimental/ConcurrentLinearAllocator.h b/Engine/Source/Runtime/Core/Public/Experimental/ConcurrentLinearAllocator.h
--- a/Engine/Source/Runtime/Core/Public/Experimental/ConcurrentLinearAllocator.h
+++ b/Engine/Source/Runtime/Core/Public/Experimental/ConcurrentLinearAllocator.h
@@ -324,9 +324,10 @@ class TLinearAllocatorBase
 			{
 				ASAN_POISON_MEMORY_REGION( Header + 1, BlockAllocationTag::BlockSize - sizeof(FBlockHeader) );
 			}
+
+			static thread_local FTLSCleanup Cleanup;
+			new (&Cleanup) FTLSCleanup{ Header };
 		}
-		static thread_local FTLSCleanup Cleanup;
-		new (&Cleanup) FTLSCleanup { Header };
 	}
 
 public:
@@ -369,8 +370,14 @@ public:
 				{
 					ASAN_POISON_MEMORY_REGION( Header + 1, BlockAllocationTag::BlockSize - sizeof(FBlockHeader) );
 				}
+
+				static thread_local FTLSCleanup Cleanup;
+				new (&Cleanup) FTLSCleanup{ Header };
+			}
+			else
+			{
+				AllocateBlock(Header);
 			}
-			AllocateBlock(Header);
 		}
 
 	AllocateNewItem:
