Allow RWAccessDetector failures to be triggered within a unit test.

Previously, anytime `GIsAutomationTesting` was true, these failures would be suppressed. We now report these failures correctly.

The unit tests which explicitly test `RWAccessDetector` error cases now have a separate flag which is private to everyone except the unit tests themselves.

Be aware that this CL may shake out latent race bugs in various systems that have automated tests. This CL did not introduce those bugs! It is just allowing us to catch those bugs in automated tests.

diff --git a/Engine/Source/Runtime/Core/Private/Misc/MTAccessDetector.cpp b/Engine/Source/Runtime/Core/Private/Misc/MTAccessDetector.cpp
--- a/Engine/Source/Runtime/Core/Private/Misc/MTAccessDetector.cpp
+++ b/Engine/Source/Runtime/Core/Private/Misc/MTAccessDetector.cpp
@@ -4,6 +4,8 @@
 
 #if ENABLE_MT_DETECTOR
 
+CORE_API bool UE::Private::FMTAccessDetectorOptions::bSuppressCheckFailure = false;
+
 FMRSWRecursiveAccessDetector::FReadersTls& FMRSWRecursiveAccessDetector::GetReadersTls()
 {
 	thread_local FReadersTls ReadersTls;
@@ -26,6 +28,25 @@ FMRSWRecursiveAccessDetector::FDestructionSentinelStackTls& FMRSWRecursiveAccess
 #include "Tasks/Task.h"
 #include "Tests/Benchmark.h"
 
+//----------------------------------------------------------------------//
+class UE::Private::FMTAccessDetectorSuppressCheckFailuresForUnitTests
+{
+public:
+	FMTAccessDetectorSuppressCheckFailuresForUnitTests()
+		: bOriginalValue(UE::Private::FMTAccessDetectorOptions::bSuppressCheckFailure)
+	{
+		UE::Private::FMTAccessDetectorOptions::bSuppressCheckFailure = true;
+	}
+
+	~FMTAccessDetectorSuppressCheckFailuresForUnitTests()
+	{
+		UE::Private::FMTAccessDetectorOptions::bSuppressCheckFailure = bOriginalValue;
+	}
+
+private:
+	bool bOriginalValue = false;
+};
+
 //----------------------------------------------------------------------//
 // FRWAccessDetector tests
 //----------------------------------------------------------------------//
@@ -57,6 +78,8 @@ bool FRWAccessDetector_ConcurentReadTest::RunTest(const FString& Parameters)
 IMPLEMENT_SIMPLE_AUTOMATION_TEST(FRWAccessDetector_ConcurentWriteTest, "System.Core.Misc.MTAccessDetector.ConcurrentWriteAccess", EAutomationTestFlags_ApplicationContextMask | EAutomationTestFlags::EngineFilter)
 bool FRWAccessDetector_ConcurentWriteTest::RunTest(const FString& Parameters)
 {
+	UE::Private::FMTAccessDetectorSuppressCheckFailuresForUnitTests SuppressCheckFailures;
+
 	bool bWriting1 = false;
 	bool bWriting2 = false;
 	UE_MT_DECLARE_RW_ACCESS_DETECTOR(MTAccessDetector);
@@ -82,6 +105,8 @@ bool FRWAccessDetector_ConcurentWriteTest::RunTest(const FString& Parameters)
 IMPLEMENT_SIMPLE_AUTOMATION_TEST(FRWAccessDetector_ConcurentReadWriteTest, "System.Core.Misc.MTAccessDetector.ConcurrentReadWriteAccess", EAutomationTestFlags_ApplicationContextMask | EAutomationTestFlags::EngineFilter)
 bool FRWAccessDetector_ConcurentReadWriteTest::RunTest(const FString& Parameters)
 {
+	UE::Private::FMTAccessDetectorSuppressCheckFailuresForUnitTests SuppressCheckFailures;
+
 	bool bReading = false;
 	bool bWriting = false;
 	UE_MT_DECLARE_RW_ACCESS_DETECTOR(MTAccessDetector);
@@ -107,6 +132,8 @@ bool FRWAccessDetector_ConcurentReadWriteTest::RunTest(const FString& Parameters
 IMPLEMENT_SIMPLE_AUTOMATION_TEST(FRWAccessDetector_WriteReentrance, "System.Core.Misc.MTAccessDetector.WriteReentrance", EAutomationTestFlags_ApplicationContextMask | EAutomationTestFlags::EngineFilter)
 bool FRWAccessDetector_WriteReentrance::RunTest(const FString& Parameters)
 {
+	UE::Private::FMTAccessDetectorSuppressCheckFailuresForUnitTests SuppressCheckFailures;
+
 	UE_MT_DECLARE_RW_ACCESS_DETECTOR(MTAccessDetector);
 
 	bool Success = false;
@@ -121,6 +148,8 @@ bool FRWAccessDetector_WriteReentrance::RunTest(const FString& Parameters)
 IMPLEMENT_SIMPLE_AUTOMATION_TEST(FRWAccessDetector_ReadReentrance, "System.Core.Misc.MTAccessDetector.ReadReentrance", EAutomationTestFlags_ApplicationContextMask | EAutomationTestFlags::EngineFilter)
 bool FRWAccessDetector_ReadReentrance::RunTest(const FString& Parameters)
 {
+	UE::Private::FMTAccessDetectorSuppressCheckFailuresForUnitTests SuppressCheckFailures;
+
 	UE_MT_DECLARE_RW_ACCESS_DETECTOR(MTAccessDetector);
 
 	bool Success = false;
@@ -163,6 +192,8 @@ bool FRWAccessDetector_RecursiveConcurentReadTest::RunTest(const FString& Parame
 IMPLEMENT_SIMPLE_AUTOMATION_TEST(FRWAccessDetector_RecursiveConcurentWriteTest, "System.Core.Misc.MTAccessDetector.RecursiveConcurrentWriteAccess", EAutomationTestFlags_ApplicationContextMask | EAutomationTestFlags::EngineFilter)
 bool FRWAccessDetector_RecursiveConcurentWriteTest::RunTest(const FString& Parameters)
 {
+	UE::Private::FMTAccessDetectorSuppressCheckFailuresForUnitTests SuppressCheckFailures;
+
 	bool bWriting1 = false;
 	bool bWriting2 = false;
 	UE_MT_DECLARE_RW_RECURSIVE_ACCESS_DETECTOR(MTAccessDetector);
@@ -188,6 +219,8 @@ bool FRWAccessDetector_RecursiveConcurentWriteTest::RunTest(const FString& Param
 IMPLEMENT_SIMPLE_AUTOMATION_TEST(FRWAccessDetector_RecursiveConcurentReadWriteTest, "System.Core.Misc.MTAccessDetector.RecursiveConcurrentReadWriteAccess", EAutomationTestFlags_ApplicationContextMask | EAutomationTestFlags::EngineFilter)
 bool FRWAccessDetector_RecursiveConcurentReadWriteTest::RunTest(const FString& Parameters)
 {
+	UE::Private::FMTAccessDetectorSuppressCheckFailuresForUnitTests SuppressCheckFailures;
+
 	bool bReading = false;
 	bool bWriting = false;
 	UE_MT_DECLARE_RW_RECURSIVE_ACCESS_DETECTOR(MTAccessDetector);
@@ -227,6 +260,8 @@ bool FRWAccessDetector_RecursiveWriteReentrance::RunTest(const FString& Paramete
 IMPLEMENT_SIMPLE_AUTOMATION_TEST(FRWAccessDetector_RecursiveReadReentrance, "System.Core.Misc.MTAccessDetector.RecursiveReadReentrance", EAutomationTestFlags_ApplicationContextMask | EAutomationTestFlags::EngineFilter)
 bool FRWAccessDetector_RecursiveReadReentrance::RunTest(const FString& Parameters)
 {
+	UE::Private::FMTAccessDetectorSuppressCheckFailuresForUnitTests SuppressCheckFailures;
+
 	UE_MT_DECLARE_RW_RECURSIVE_ACCESS_DETECTOR(MTAccessDetector);
 
 	bool Success = false;
@@ -269,6 +304,8 @@ bool FRWAccessDetector_FullyRecursiveConcurentReadTest::RunTest(const FString& P
 IMPLEMENT_SIMPLE_AUTOMATION_TEST(FRWAccessDetector_FullyRecursiveConcurentWriteTest, "System.Core.Misc.MTAccessDetector.FullyRecursiveConcurrentWriteAccess", EAutomationTestFlags_ApplicationContextMask | EAutomationTestFlags::EngineFilter)
 bool FRWAccessDetector_FullyRecursiveConcurentWriteTest::RunTest(const FString& Parameters)
 {
+	UE::Private::FMTAccessDetectorSuppressCheckFailuresForUnitTests SuppressCheckFailures;
+
 	bool bWriting1 = false;
 	bool bWriting2 = false;
 	UE_MT_DECLARE_RW_FULLY_RECURSIVE_ACCESS_DETECTOR(MTAccessDetector);
@@ -294,6 +331,8 @@ bool FRWAccessDetector_FullyRecursiveConcurentWriteTest::RunTest(const FString&
 IMPLEMENT_SIMPLE_AUTOMATION_TEST(FRWAccessDetector_FullyRecursiveConcurentReadWriteTest, "System.Core.Misc.MTAccessDetector.FullyRecursiveConcurrentReadWriteAccess", EAutomationTestFlags_ApplicationContextMask | EAutomationTestFlags::EngineFilter)
 bool FRWAccessDetector_FullyRecursiveConcurentReadWriteTest::RunTest(const FString& Parameters)
 {
+	UE::Private::FMTAccessDetectorSuppressCheckFailuresForUnitTests SuppressCheckFailures;
+
 	bool bReading = false;
 	bool bWriting = false;
 	UE_MT_DECLARE_RW_FULLY_RECURSIVE_ACCESS_DETECTOR(MTAccessDetector);
diff --git a/Engine/Source/Runtime/Core/Public/Misc/MTAccessDetector.h b/Engine/Source/Runtime/Core/Public/Misc/MTAccessDetector.h
--- a/Engine/Source/Runtime/Core/Public/Misc/MTAccessDetector.h
+++ b/Engine/Source/Runtime/Core/Public/Misc/MTAccessDetector.h
@@ -18,8 +18,25 @@
 #include "AutoRTFM.h"
 #include <atomic>
 
+struct FRWAccessDetector;
 
-extern CORE_API bool GIsAutomationTesting;
+namespace UE::Private
+{
+	class FMTAccessDetectorSuppressCheckFailuresForUnitTests;
+
+	class FMTAccessDetectorOptions
+	{
+	private:
+		/**
+		 * For testing purposes, we can suppress check-failures when access violations are detected. 
+		 * Instead, the relevant access methods will just return false.
+		 */
+		static CORE_API bool bSuppressCheckFailure;
+
+		friend struct ::FRWAccessDetector;
+		friend class ::UE::Private::FMTAccessDetectorSuppressCheckFailuresForUnitTests;
+	};
+}
 
 /**
  * Read write multithread access detector, will check on concurrent write/write and read/write access, but will not on concurrent read access.
@@ -70,7 +87,7 @@ public:
 	FORCEINLINE bool AcquireReadAccess() const
 	{
 		const bool ErrorDetected = (AtomicValue.fetch_add(1, std::memory_order_relaxed) & WriterBits) != 0;
-		checkf(!ErrorDetected || GIsAutomationTesting, TEXT("Acquiring a read access while there is already a write access"));
+		checkf(!ErrorDetected || UE::Private::FMTAccessDetectorOptions::bSuppressCheckFailure, TEXT("Acquiring a read access while there is already a write access"));
 		return !ErrorDetected;
 	}
 
@@ -81,7 +98,7 @@ public:
 	FORCEINLINE bool ReleaseReadAccess() const
 	{
 		const bool ErrorDetected = (AtomicValue.fetch_sub(1, std::memory_order_relaxed) & WriterBits) != 0;
-		checkf(!ErrorDetected || GIsAutomationTesting, TEXT("Another thread asked to have a write access during this read access"));
+		checkf(!ErrorDetected || UE::Private::FMTAccessDetectorOptions::bSuppressCheckFailure, TEXT("Another thread asked to have a write access during this read access"));
 		return !ErrorDetected;
 	}
 
@@ -92,7 +109,7 @@ public:
 	FORCEINLINE bool AcquireWriteAccess() const
 	{
 		const bool ErrorDetected = AtomicValue.fetch_add(WriterIncrementValue, std::memory_order_relaxed) != 0;
-		checkf(!ErrorDetected || GIsAutomationTesting, TEXT("Acquiring a write access while there are ongoing read or write access"));
+		checkf(!ErrorDetected || UE::Private::FMTAccessDetectorOptions::bSuppressCheckFailure, TEXT("Acquiring a write access while there are ongoing read or write access"));
 		return !ErrorDetected;
 	}
 
@@ -103,7 +120,7 @@ public:
 	FORCEINLINE bool ReleaseWriteAccess() const
 	{
 		const bool ErrorDetected = AtomicValue.fetch_sub(WriterIncrementValue, std::memory_order_relaxed) != WriterIncrementValue;
-		checkf(!ErrorDetected || GIsAutomationTesting, TEXT("Another thread asked to have a read or write access during this write access"));
+		checkf(!ErrorDetected || UE::Private::FMTAccessDetectorOptions::bSuppressCheckFailure, TEXT("Another thread asked to have a read or write access during this write access"));
 		return !ErrorDetected;
 	}
 
@@ -117,7 +134,7 @@ public:
 
 		if (UNLIKELY(bErrorDetected))
 		{
-			checkf(GIsAutomationTesting, TEXT("Upgrading a write access while there are ongoing read or write access"));
+			checkf(UE::Private::FMTAccessDetectorOptions::bSuppressCheckFailure, TEXT("Upgrading a write access while there are ongoing read or write access"));
 			return false;
 		}
 
