- Supporting NoReplay config branch rpelay mode with the fixed commandline application changes

diff --git a/Engine/Source/Runtime/Core/Private/Misc/ConfigCacheIni.cpp b/Engine/Source/Runtime/Core/Private/Misc/ConfigCacheIni.cpp
--- a/Engine/Source/Runtime/Core/Private/Misc/ConfigCacheIni.cpp
+++ b/Engine/Source/Runtime/Core/Private/Misc/ConfigCacheIni.cpp
@@ -3260,8 +3260,12 @@ bool FConfigBranch::AddDynamicLayersToHierarchy(const TArray<FDynamicLayerInfo>&
 	
 	if (AddedLayers.Num() > 0 || bForceFullDynamicLayerUpdate)
 	{
+		checkf(CommandLineOverrides.Num() == 0 || ReplayMethod != EBranchReplayMethod::NoReplay,
+			TEXT("The branch (%s) with commandline overrides must not have ReplayMethod of NoReplay!"), *IniName.ToString());
+
 		// if all were added at the end (or there's no replay), we can just apply them without rewinding
-		if (!bInsertedBeforeEnd && !bForceFullDynamicLayerUpdate)
+		// since commandline needs to trump the layers, we can't use this optimization if there are overrides
+		if (!bInsertedBeforeEnd && !bForceFullDynamicLayerUpdate && CommandLineOverrides.Num() == 0)
 		{
 			for (FConfigCommandStream* NewLayer : AddedLayers)
 			{
@@ -3279,6 +3283,7 @@ bool FConfigBranch::AddDynamicLayersToHierarchy(const TArray<FDynamicLayerInfo>&
 			{
 				FinalCombinedLayers.ApplyFile(*Node);
 			}
+			FinalCombinedLayers.ApplyFile(&CommandLineOverrides);
 			bool bOldSaveAllSections = InMemoryFile.bCanSaveAllSections;
 			InMemoryFile = FinalCombinedLayers;
 			InMemoryFile.bCanSaveAllSections = bOldSaveAllSections;
@@ -3362,7 +3367,9 @@ bool FConfigBranch::AddDynamicLayerStringToHierarchy(const FString& Filename, co
 		}
 	}
 
-	if (!bInsertedAtEnd)
+	// if we inserted at the end, and there are no commandline options that need to stomp on layers,
+	// then there's a chance of optimization
+	if (!bInsertedAtEnd || CommandLineOverrides.Num() > 0)
 	{
 		// rebuild
 		FinalCombinedLayers = CombinedStaticLayers;
@@ -3370,6 +3377,7 @@ bool FConfigBranch::AddDynamicLayerStringToHierarchy(const FString& Filename, co
 		{
 			FinalCombinedLayers.ApplyFile(*Node);
 		}
+		FinalCombinedLayers.ApplyFile(&CommandLineOverrides);
 		bool bOldSaveAllSections = InMemoryFile.bCanSaveAllSections;
 		InMemoryFile = FinalCombinedLayers;
 		InMemoryFile.bCanSaveAllSections = bOldSaveAllSections;
@@ -3434,6 +3442,7 @@ bool FConfigBranch::RemoveDynamicLayersFromHierarchy(const TArray<FString>& File
 	{
 		FinalCombinedLayers.ApplyFile(*Node);
 	}
+	FinalCombinedLayers.ApplyFile(&CommandLineOverrides);
 	bool bOldSaveAllSections = InMemoryFile.bCanSaveAllSections;
 	InMemoryFile = FinalCombinedLayers;
 	InMemoryFile.bCanSaveAllSections = bOldSaveAllSections;
diff --git a/Engine/Source/Runtime/Core/Private/Misc/ConfigContext.cpp b/Engine/Source/Runtime/Core/Private/Misc/ConfigContext.cpp
--- a/Engine/Source/Runtime/Core/Private/Misc/ConfigContext.cpp
+++ b/Engine/Source/Runtime/Core/Private/Misc/ConfigContext.cpp
@@ -1000,6 +1000,19 @@ bool FConfigContext::LoadIniFileHierarchy()
 	
 	// LogVariables(*BaseIniName, Platform);
 
+#if ALLOW_INI_OVERRIDE_FROM_COMMANDLINE
+	// if we have a commandline overrides, then we need to know now, because we can't 
+	// use NoReplay mode safely
+	FConfigFile::OverrideFromCommandline(&Branch->CommandLineOverrides, BaseIniName);
+	if (Branch->CommandLineOverrides.Num() > 0 && Branch->ReplayMethod == EBranchReplayMethod::NoReplay)
+	{
+		UE_LOG(LogConfig, Display, TEXT("Branch %s had commandline override(s), so override NoReplay mode. This will use a bit more memory, but is needed for plugin/GFP handling."), *BaseIniName);
+		Branch->ReplayMethod = EBranchReplayMethod::DynamicLayerReplay;
+	}
+#endif
+
+
+
 	bool bReadAnyFile = false;
 
 	TRACE_CPUPROFILER_EVENT_SCOPE(LoadIniFileHierarchy);
@@ -1204,9 +1217,8 @@ bool FConfigContext::GenerateDestIniFile()
 		Branch->InMemoryFile.ApplyFile(&Branch->SavedLayer);
 
 #if ALLOW_INI_OVERRIDE_FROM_COMMANDLINE
-		// process any commandline overrides
-		FConfigFile::OverrideFromCommandline(&Branch->CommandLineOverrides, BaseIniName);
 		// and push it into the current values
+		Branch->FinalCombinedLayers.ApplyFile(&Branch->CommandLineOverrides);
 		Branch->InMemoryFile.ApplyFile(&Branch->CommandLineOverrides);
 #endif
 
