Add support for `bSuppressCheckFailure` to the MRSW access detector.

This class originally had no way to squelch ensure-failures at all, so its access-violation tests are currently wrapped with `#if 0` and aren't even accessible from the Test Automation window at all due to an unrelated bug. These issues will be fixed in followup CLs.

diff --git a/Engine/Source/Runtime/Core/Public/Misc/MTAccessDetector.h b/Engine/Source/Runtime/Core/Public/Misc/MTAccessDetector.h
--- a/Engine/Source/Runtime/Core/Public/Misc/MTAccessDetector.h
+++ b/Engine/Source/Runtime/Core/Public/Misc/MTAccessDetector.h
@@ -19,6 +19,7 @@
 #include <atomic>
 
 struct FRWAccessDetector;
+class FMRSWRecursiveAccessDetector;
 
 namespace UE::Private
 {
@@ -34,6 +35,7 @@ namespace UE::Private
 		static CORE_API bool bSuppressCheckFailure;
 
 		friend struct ::FRWAccessDetector;
+		friend class ::FMRSWRecursiveAccessDetector;
 		friend class ::UE::Private::FMTAccessDetectorSuppressCheckFailuresForUnitTests;
 	};
 }
@@ -384,24 +386,30 @@ private:
 		return L.Value == R.Value;
 	}
 
-	FORCEINLINE static void CheckOtherThreadWriters(FState InState)
+	FORCEINLINE static bool CheckOtherThreadWriters(FState InState)
 	{
 		if (InState.WriterNum != 0)
 		{
 			uint32 CurrentThreadId = FPlatformTLS::GetCurrentThreadId();
-			ensureMsgf(InState.GetWriterThreadId() == CurrentThreadId,
-				TEXT("Data race detected! Acquiring read access on thread %u concurrently with %d writers on thread %u:\nCurrent thread %u callstack:\n%s\nWriter thread %u callstack:\n%s"),
-				CurrentThreadId, InState.WriterNum, InState.GetWriterThreadId(), 
-				CurrentThreadId, *GetCurrentThreadCallstack(),
-				InState.GetWriterThreadId(), *GetThreadCallstack(InState.GetWriterThreadId()));
+			if (InState.GetWriterThreadId() != CurrentThreadId)
+			{
+				ensureMsgf(UE::Private::FMTAccessDetectorOptions::bSuppressCheckFailure, 
+					TEXT("Data race detected! Acquiring read access on thread %u concurrently with %d writers on thread %u:\nCurrent thread %u callstack:\n%s\nWriter thread %u callstack:\n%s"),
+					CurrentThreadId, InState.WriterNum, InState.GetWriterThreadId(), 
+					CurrentThreadId, *GetCurrentThreadCallstack(),
+					InState.GetWriterThreadId(), *GetThreadCallstack(InState.GetWriterThreadId()));
+
+				return false;
+			}
 		}
+		return true;
 	}
 
 	static FString GetCurrentThreadCallstack()
 	{
 		const SIZE_T StackTraceSize = 65536;
 		ANSICHAR StackTrace[StackTraceSize] = { 0 };
-		UE_AUTORTFM_OPEN{ FPlatformStackWalk::StackWalkAndDump(StackTrace, StackTraceSize, 0); };
+		FPlatformStackWalk::StackWalkAndDump(StackTrace, StackTraceSize, 0); 
 		return ANSI_TO_TCHAR(StackTrace);
 	}
 
@@ -420,7 +428,7 @@ private:
 
 		const SIZE_T StackTraceSize = 65536;
 		ANSICHAR StackTrace[StackTraceSize] = { 0 };
-		UE_AUTORTFM_OPEN{ FPlatformStackWalk::ThreadStackWalkAndDump(StackTrace, StackTraceSize, 0, ThreadId); };
+		FPlatformStackWalk::ThreadStackWalkAndDump(StackTrace, StackTraceSize, 0, ThreadId); 
 		return ANSI_TO_TCHAR(StackTrace);
 	}
 
@@ -494,7 +502,7 @@ public:
 	FMRSWRecursiveAccessDetector() = default;
 
 	FORCEINLINE FMRSWRecursiveAccessDetector(const FMRSWRecursiveAccessDetector& Other)
-		// just default initialisation, the copy is not being accessed
+		// just default initialization, the copy is not being accessed
 	{
 		CheckOtherThreadWriters(Other.LoadState());
 	}
@@ -536,19 +544,22 @@ public:
 			}
 		}
 
-		ensureMsgf(LoadState().Value == ExpectedState.Value,
-			TEXT("Race detector destroyed while being accessed on another thread: %d readers, %d writers on thread %u:\nCurrent thread %u callstack:\n%s\nWriter thread %u callstack:\n%s"),
-			LoadState().ReaderNum - ExpectedState.ReaderNum, 
-			LoadState().WriterNum - ExpectedState.WriterNum, LoadState().GetWriterThreadId(), 
-			FPlatformTLS::GetCurrentThreadId(), *GetCurrentThreadCallstack(),
-			LoadState().GetWriterThreadId(), *GetThreadCallstack(LoadState().GetWriterThreadId()));
+		if (LoadState().Value != ExpectedState.Value)
+		{
+			ensureMsgf(UE::Private::FMTAccessDetectorOptions::bSuppressCheckFailure, 
+				TEXT("Race detector destroyed while being accessed on another thread: %d readers, %d writers on thread %u:\nCurrent thread %u callstack:\n%s\nWriter thread %u callstack:\n%s"),
+				LoadState().ReaderNum - ExpectedState.ReaderNum, 
+				LoadState().WriterNum - ExpectedState.WriterNum, LoadState().GetWriterThreadId(), 
+				FPlatformTLS::GetCurrentThreadId(), *GetCurrentThreadCallstack(),
+				LoadState().GetWriterThreadId(), *GetThreadCallstack(LoadState().GetWriterThreadId()));
+		}
 	}
 
-	FORCEINLINE void AcquireReadAccess() const
+	FORCEINLINE bool AcquireReadAccess() const
 	{
 		FState PrevState = IncrementReaderNum();
 
-		CheckOtherThreadWriters(PrevState);
+		bool bSuccess = CheckOtherThreadWriters(PrevState);
 
 		// register the reader in TLS
 		int32 ReaderIndex = PrevState.ReaderNum == 0 ? 
@@ -562,39 +573,44 @@ public:
 		{
 			++GetReadersTls()[ReaderIndex].Num;
 		}
+
+		return bSuccess;
 	}
 
 	// an overload that handles access detector destruction from inside a read access, must be used along with the corresponding
-	// overload of `ReleaseReadAcess`
-	FORCEINLINE void AcquireReadAccess(FDestructionSentinel& DestructionSentinel) const
+	// overload of `ReleaseReadAccess`
+	FORCEINLINE bool AcquireReadAccess(FDestructionSentinel& DestructionSentinel) const
 	{
 		DestructionSentinel.Accessor = this;
 		GetDestructionSentinelStackTls().Add(&DestructionSentinel);
 
-		AcquireReadAccess();
+		return AcquireReadAccess();
 	}
 
-	FORCEINLINE void ReleaseReadAccess() const
+	FORCEINLINE bool ReleaseReadAccess() const
 	{
 		RemoveReaderFromTls();
 		DecrementReaderNum();
 		// no need to check for writers
+		return true;
 	}
 
 	// an overload that handles access detector destruction from inside a read access, must be used along with the corresponding
-	// overload of `AcquireReadAcess`
-	FORCEINLINE void ReleaseReadAccess(FDestructionSentinel& DestructionSentinel) const
+	// overload of `AcquireReadAccess`
+	FORCEINLINE bool ReleaseReadAccess(FDestructionSentinel& DestructionSentinel) const
 	{
-		ReleaseReadAccess();
+		bool bSuccess = ReleaseReadAccess();
 
 		checkSlow(DestructionSentinel.Accessor == this);
 		checkfSlow(GetDestructionSentinelStackTls().Num() != 0, TEXT("An attempt to remove a not registered destruction sentinel"));
 		checkfSlow(GetDestructionSentinelStackTls().Last() == &DestructionSentinel, TEXT("Mismatched destruction sentinel: %p != %p"), GetDestructionSentinelStackTls().Last(), &DestructionSentinel);
 
 		GetDestructionSentinelStackTls().RemoveAtSwap(GetDestructionSentinelStackTls().Num() - 1);
+
+		return bSuccess;
 	}
 
-	FORCEINLINE void AcquireWriteAccess()
+	FORCEINLINE bool AcquireWriteAccess()
 	{
 		FState LocalState = LoadState();
 
@@ -602,54 +618,71 @@ public:
 		{	// check that all readers are on the current thread
 			int32 ReaderIndex = GetReadersTls().IndexOfByPredicate([this](FReaderNum ReaderNum) { return ReaderNum.Reader == this; });
 
-			ensureMsgf(ReaderIndex != INDEX_NONE, 
-				TEXT("Either a race detected (%u reader(s) on another thread(s) while acquiring write access on the current thread) or the access detector was trivially relocated:\nCurrent thread %u callstack:\n%s"), 
-				LocalState.ReaderNum, FPlatformTLS::GetCurrentThreadId(), *GetCurrentThreadCallstack());
-
 			if (ReaderIndex == INDEX_NONE)
 			{
-				return; // to avoid asserting on the next line inside TArray
+				ensureMsgf(UE::Private::FMTAccessDetectorOptions::bSuppressCheckFailure, 
+					TEXT("Either a race detected (%u reader(s) on another thread(s) while acquiring write access on the current thread) or the access detector was trivially relocated:\nCurrent thread %u callstack:\n%s"), 
+					LocalState.ReaderNum, FPlatformTLS::GetCurrentThreadId(), *GetCurrentThreadCallstack());
+
+				return false; // to avoid asserting on the next line inside TArray
 			}
 
-			ensureMsgf(GetReadersTls()[ReaderIndex].Num == LocalState.ReaderNum,
-				TEXT("Data race detected: %d reader(s) on another thread(s) while acquiring write access.\nCurrent thread %u callstack:\n%s"),
-				LocalState.ReaderNum - GetReadersTls()[ReaderIndex].Num, FPlatformTLS::GetCurrentThreadId(), *GetCurrentThreadCallstack());
+			if (GetReadersTls()[ReaderIndex].Num != LocalState.ReaderNum)
+			{
+				ensureMsgf(UE::Private::FMTAccessDetectorOptions::bSuppressCheckFailure,
+					TEXT("Data race detected: %d reader(s) on another thread(s) while acquiring write access.\nCurrent thread %u callstack:\n%s"),
+					LocalState.ReaderNum - GetReadersTls()[ReaderIndex].Num, FPlatformTLS::GetCurrentThreadId(), *GetCurrentThreadCallstack());
+				
+				return false;
+			}
 		}
 
 		uint32 CurrentThreadId = FPlatformTLS::GetCurrentThreadId();
 		if (LocalState.WriterNum != 0)
 		{
-			ensureMsgf(LocalState.GetWriterThreadId() == CurrentThreadId,
-				TEXT("Data race detected: acquiring write access on thread %u concurrently with %d writers on thread %u:\nCurrent thread %u callstack:\n%s\nWriter thread %u callstack:\n%s"), 
-				CurrentThreadId, LocalState.WriterNum, LocalState.GetWriterThreadId(), 
-				CurrentThreadId, *GetCurrentThreadCallstack(),
-				LocalState.GetWriterThreadId(), *GetThreadCallstack(LocalState.GetWriterThreadId()));
+			if (LocalState.GetWriterThreadId() != CurrentThreadId)
+			{
+				ensureMsgf(UE::Private::FMTAccessDetectorOptions::bSuppressCheckFailure,
+					TEXT("Data race detected: acquiring write access on thread %u concurrently with %d writers on thread %u:\nCurrent thread %u callstack:\n%s\nWriter thread %u callstack:\n%s"), 
+					CurrentThreadId, LocalState.WriterNum, LocalState.GetWriterThreadId(), 
+					CurrentThreadId, *GetCurrentThreadCallstack(),
+					LocalState.GetWriterThreadId(), *GetThreadCallstack(LocalState.GetWriterThreadId()));
+
+				return false;
+			}
 		}
 
 		FState NewState{ LocalState.ReaderNum, LocalState.WriterNum + 1u, CurrentThreadId + 1 };
 		FState PrevState = ExchangeState(NewState);
 
-		ensureMsgf(LocalState == PrevState,
-			TEXT("Data race detected: other thread(s) activity during acquiring write access on thread %u: %u -> %u readers, %u -> %u writers on thread %u -> %u:\nCurrent thread %u callstack:\n%s\nWriter thread %u callstack:\n%s"),
-			CurrentThreadId,
-			LocalState.ReaderNum, PrevState.ReaderNum,
-			LocalState.WriterNum, PrevState.WriterNum,
-			LocalState.GetWriterThreadId(), PrevState.GetWriterThreadId(), 
-			CurrentThreadId, *GetCurrentThreadCallstack(),
-			PrevState.GetWriterThreadId(), *GetThreadCallstack(PrevState.GetWriterThreadId()));
+		if (LocalState != PrevState)
+		{
+			ensureMsgf(UE::Private::FMTAccessDetectorOptions::bSuppressCheckFailure,
+				TEXT("Data race detected: other thread(s) activity during acquiring write access on thread %u: %u -> %u readers, %u -> %u writers on thread %u -> %u:\nCurrent thread %u callstack:\n%s\nWriter thread %u callstack:\n%s"),
+				CurrentThreadId,
+				LocalState.ReaderNum, PrevState.ReaderNum,
+				LocalState.WriterNum, PrevState.WriterNum,
+				LocalState.GetWriterThreadId(), PrevState.GetWriterThreadId(), 
+				CurrentThreadId, *GetCurrentThreadCallstack(),
+				PrevState.GetWriterThreadId(), *GetThreadCallstack(PrevState.GetWriterThreadId()));
+
+			return false;
+		}
+
+		return true;
 	}
 
 	// an overload that handles access detector destruction from inside a write access, must be used along with the corresponding
-	// overload of `ReleaseWriteAcess`
-	FORCEINLINE void AcquireWriteAccess(FDestructionSentinel& DestructionSentinel)
+	// overload of `ReleaseWriteAccess`
+	FORCEINLINE bool AcquireWriteAccess(FDestructionSentinel& DestructionSentinel)
 	{
 		DestructionSentinel.Accessor = this;
 		GetDestructionSentinelStackTls().Add(&DestructionSentinel);
 
-		AcquireWriteAccess();
+		return AcquireWriteAccess();
 	}
 
-	FORCEINLINE void ReleaseWriteAccess()
+	FORCEINLINE bool ReleaseWriteAccess()
 	{
 		FState LocalState = LoadState();
 
@@ -657,21 +690,28 @@ public:
 		FState NewState{ LocalState.ReaderNum, LocalState.WriterNum - 1u, WriterThreadId };
 		FState PrevState = ExchangeState(NewState);
 
-		ensureMsgf(LocalState == PrevState,
-			TEXT("Data race detected: other thread(s) activity during releasing write access on thread %d: %u -> %u readers, %u -> %u writers on thread %u -> %u\nCurrent thread %u callstack:\n%s\nWriter thread %u callstack:\n%s"),
-			FPlatformTLS::GetCurrentThreadId(),
-			LocalState.ReaderNum, PrevState.ReaderNum,
-			LocalState.WriterNum, PrevState.WriterNum, 
-			LocalState.GetWriterThreadId(), PrevState.GetWriterThreadId(),
-			FPlatformTLS::GetCurrentThreadId(), *GetCurrentThreadCallstack(),
-			PrevState.GetWriterThreadId(), *GetThreadCallstack(PrevState.GetWriterThreadId()));
+		if (LocalState != PrevState)
+		{
+			ensureMsgf(UE::Private::FMTAccessDetectorOptions::bSuppressCheckFailure,
+				TEXT("Data race detected: other thread(s) activity during releasing write access on thread %d: %u -> %u readers, %u -> %u writers on thread %u -> %u\nCurrent thread %u callstack:\n%s\nWriter thread %u callstack:\n%s"),
+				FPlatformTLS::GetCurrentThreadId(),
+				LocalState.ReaderNum, PrevState.ReaderNum,
+				LocalState.WriterNum, PrevState.WriterNum, 
+				LocalState.GetWriterThreadId(), PrevState.GetWriterThreadId(),
+				FPlatformTLS::GetCurrentThreadId(), *GetCurrentThreadCallstack(),
+				PrevState.GetWriterThreadId(), *GetThreadCallstack(PrevState.GetWriterThreadId()));
+			
+			return false;
+		}
+
+		return true;
 	}
 
 	// an overload that handles access detector destruction from inside a write access, must be used along with the corresponding
-	// overload of `AcquireWriteAcess`
-	FORCEINLINE void ReleaseWriteAccess(FDestructionSentinel& DestructionSentinel)
+	// overload of `AcquireWriteAccess`
+	FORCEINLINE bool ReleaseWriteAccess(FDestructionSentinel& DestructionSentinel)
 	{
-		ReleaseWriteAccess();
+		bool bSuccess = ReleaseWriteAccess();
 
 		FDestructionSentinelStackTls& DestructionSentinelStackTls = GetDestructionSentinelStackTls();
 		checkSlow(DestructionSentinel.Accessor == this);
@@ -679,6 +719,8 @@ public:
 		checkfSlow(GetDestructionSentinelStackTls().Last() == &DestructionSentinel, TEXT("Mismatched destruction sentinel: %p != %p"), GetDestructionSentinelStackTls().Last(), &DestructionSentinel);
 
 		GetDestructionSentinelStackTls().RemoveAtSwap(GetDestructionSentinelStackTls().Num() - 1);
+
+		return bSuccess;
 	}
 };
 
