On iOS, fixed the `RequestExit` from `FGenericPlatformMisc` that would call `abort()`, which would get the application stuck, causing the process to never end.
Added a specific implementation of the `RequestExit` to the `FApplePlatformMisc` with the `_Exit(0)` call.

On Android, changed the `FAndroidPlatformProperties::SupportsQuit()` to return `false` as Android, as well as iOS, do not support calling `RequestEngineExit`.

#rnx

diff --git a/Engine/Source/Runtime/Core/Private/Apple/ApplePlatformMisc.cpp b/Engine/Source/Runtime/Core/Private/Apple/ApplePlatformMisc.cpp
--- a/Engine/Source/Runtime/Core/Private/Apple/ApplePlatformMisc.cpp
+++ b/Engine/Source/Runtime/Core/Private/Apple/ApplePlatformMisc.cpp
@@ -15,6 +15,8 @@
 #include "FramePro/FrameProProfiler.h"
 #include "CoreGlobals.h"
 
+DEFINE_LOG_CATEGORY_STATIC(LogApplePlatformMisc, Log, All);
+
 #if !UE_BUILD_SHIPPING
 
 bool FApplePlatformMisc::IsDebuggerPresent()
@@ -354,6 +356,22 @@ bool FApplePlatformMisc::IsOSAtLeastVersion(const uint32 MacOSVersion[3], const
 	return true;
 }
 
+void FApplePlatformMisc::RequestExit( bool Force, const TCHAR* CallSite)
+{
+	UE_LOG(LogApplePlatformMisc, Log,  TEXT("FApplePlatformMisc::RequestExit(%i, %s)"), Force,
+		CallSite ? CallSite : TEXT("<NoCallSiteInfo>"));
+	if( Force )
+	{
+		// effectivelly exits the application
+		_Exit(0);
+	}
+	else
+	{
+		// Tell the platform specific code we want to exit cleanly from the main loop.
+		RequestEngineExit(TEXT("ApplePlatform RequestExit"));
+	}
+}
+
 #if APPLE_PROFILING_ENABLED
 
 void FApplePlatformMisc::BeginNamedEvent(const struct FColor& Color, const TCHAR* Text)
diff --git a/Engine/Source/Runtime/Core/Public/Android/AndroidPlatformProperties.h b/Engine/Source/Runtime/Core/Public/Android/AndroidPlatformProperties.h
--- a/Engine/Source/Runtime/Core/Public/Android/AndroidPlatformProperties.h
+++ b/Engine/Source/Runtime/Core/Public/Android/AndroidPlatformProperties.h
@@ -87,7 +87,8 @@ struct FAndroidPlatformProperties
 
 	static constexpr FORCEINLINE bool SupportsQuit()
 	{
-		return true;
+		// android does not support quitting to the OS like PC does
+		return false;
 	}
 
 	static constexpr FORCEINLINE bool HasFixedResolution()
diff --git a/Engine/Source/Runtime/Core/Public/Apple/ApplePlatformMisc.h b/Engine/Source/Runtime/Core/Public/Apple/ApplePlatformMisc.h
--- a/Engine/Source/Runtime/Core/Public/Apple/ApplePlatformMisc.h
+++ b/Engine/Source/Runtime/Core/Public/Apple/ApplePlatformMisc.h
@@ -90,6 +90,8 @@ struct CORE_API FApplePlatformMisc : public FGenericPlatformMisc
 
 	static bool IsOSAtLeastVersion(const uint32 MacOSVersion[3], const uint32 IOSVersion[3], const uint32 TVOSVersion[3]);
 
+	static void RequestExit( bool Force, const TCHAR* CallSite = nullptr);
+
 #if APPLE_PROFILING_ENABLED
 	/** Override to add Apple features when AppleProfiling is on, otherwise use version on FGenericPlatformMisc. */
 	static void BeginNamedEvent(const struct FColor& Color, const TCHAR* Text);
