Fix boolean logic in MTAccessDetector tests.

As originally written, these tests would pass if even a single check within the test passed. The tests now require that all checks pass.

diff --git a/Engine/Source/Runtime/Core/Private/Misc/MTAccessDetector.cpp b/Engine/Source/Runtime/Core/Private/Misc/MTAccessDetector.cpp
--- a/Engine/Source/Runtime/Core/Private/Misc/MTAccessDetector.cpp
+++ b/Engine/Source/Runtime/Core/Private/Misc/MTAccessDetector.cpp
@@ -57,22 +57,22 @@ bool FRWAccessDetector_ConcurentReadTest::RunTest(const FString& Parameters)
 	bool bReading2 = false;
 	UE_MT_DECLARE_RW_ACCESS_DETECTOR(MTAccessDetector);
 
-	bool Success = false;
+	bool bSuccess = true;
 	FThread AccessThread = FThread(TEXT("AccessThread"), [&]()
 	{
-		Success |= TestTrue(TEXT("Acquiring Read Access 1"), MTAccessDetector.AcquireReadAccess());
+		bSuccess &= TestTrue(TEXT("Acquiring Read Access 1"), MTAccessDetector.AcquireReadAccess());
 		bReading1 = true;
 		while (!bReading2) { FPlatformProcess::SleepNoStats(0.01f); }
-		Success |= TestTrue(TEXT("Releasing Read Access 1"), MTAccessDetector.ReleaseReadAccess());
+		bSuccess &= TestTrue(TEXT("Releasing Read Access 1"), MTAccessDetector.ReleaseReadAccess());
 		return true;
 	});
 	while (!bReading1) { FPlatformProcess::SleepNoStats(0.01f); }
-	Success |= TestTrue(TEXT("Acquiring Read Access 2"), MTAccessDetector.AcquireReadAccess());
+	bSuccess &= TestTrue(TEXT("Acquiring Read Access 2"), MTAccessDetector.AcquireReadAccess());
 	bReading2 = true;
 	AccessThread.Join();
-	Success |= TestTrue(TEXT("Releasing Read Access 2"), MTAccessDetector.ReleaseReadAccess());
+	bSuccess &= TestTrue(TEXT("Releasing Read Access 2"), MTAccessDetector.ReleaseReadAccess());
 
-	return Success;
+	return bSuccess;
 }
 
 IMPLEMENT_SIMPLE_AUTOMATION_TEST(FRWAccessDetector_ConcurentWriteTest, "System.Core.Misc.MTAccessDetector.ConcurrentWriteAccess", EAutomationTestFlags_ApplicationContextMask | EAutomationTestFlags::EngineFilter)
@@ -84,22 +84,22 @@ bool FRWAccessDetector_ConcurentWriteTest::RunTest(const FString& Parameters)
 	bool bWriting2 = false;
 	UE_MT_DECLARE_RW_ACCESS_DETECTOR(MTAccessDetector);
 
-	bool Success = false;
+	bool bSuccess = true;
 	FThread AccessThread = FThread(TEXT("AccessThread"), [&]()
 	{
-		Success |= TestTrue(TEXT("Acquiring Write Access"), MTAccessDetector.AcquireWriteAccess());
+		bSuccess &= TestTrue(TEXT("Acquiring Write Access"), MTAccessDetector.AcquireWriteAccess());
 		bWriting1 = true;
 		while (!bWriting2) { FPlatformProcess::SleepNoStats(0.01f); }
-		Success |= TestFalse(TEXT("Releasing Read Access"), MTAccessDetector.ReleaseWriteAccess());
+		bSuccess &= TestFalse(TEXT("Releasing Read Access"), MTAccessDetector.ReleaseWriteAccess());
 		return true;
 	});
 	while (!bWriting1) { FPlatformProcess::SleepNoStats(0.01f); }
-	Success |= TestFalse(TEXT("Acquiring Write Access"), MTAccessDetector.AcquireWriteAccess());
+	bSuccess &= TestFalse(TEXT("Acquiring Write Access"), MTAccessDetector.AcquireWriteAccess());
 	bWriting2 = true;
 	AccessThread.Join();
-	Success |= TestTrue(TEXT("Releasing Write Access"), MTAccessDetector.ReleaseWriteAccess());
+	bSuccess &= TestTrue(TEXT("Releasing Write Access"), MTAccessDetector.ReleaseWriteAccess());
 
-	return Success;
+	return bSuccess;
 }
 
 IMPLEMENT_SIMPLE_AUTOMATION_TEST(FRWAccessDetector_ConcurentReadWriteTest, "System.Core.Misc.MTAccessDetector.ConcurrentReadWriteAccess", EAutomationTestFlags_ApplicationContextMask | EAutomationTestFlags::EngineFilter)
@@ -111,22 +111,22 @@ bool FRWAccessDetector_ConcurentReadWriteTest::RunTest(const FString& Parameters
 	bool bWriting = false;
 	UE_MT_DECLARE_RW_ACCESS_DETECTOR(MTAccessDetector);
 
-	bool Success = false;
+	bool bSuccess = true;
 	FThread AccessThread = FThread(TEXT("AccessThread"), [&]()
 	{
-		Success |= TestTrue(TEXT("Acquiring Read Access"), MTAccessDetector.AcquireReadAccess());
+		bSuccess &= TestTrue(TEXT("Acquiring Read Access"), MTAccessDetector.AcquireReadAccess());
 		bReading = true;
 		while (!bWriting) { FPlatformProcess::SleepNoStats(0.01f); }
-		Success |= TestFalse(TEXT("Releasing Read Access"), MTAccessDetector.ReleaseReadAccess());
+		bSuccess &= TestFalse(TEXT("Releasing Read Access"), MTAccessDetector.ReleaseReadAccess());
 		return true;
 	});
 	while (!bReading) { FPlatformProcess::SleepNoStats(0.01f); }
-	Success |= TestFalse(TEXT("Acquiring Write Access"), MTAccessDetector.AcquireWriteAccess());
+	bSuccess &= TestFalse(TEXT("Acquiring Write Access"), MTAccessDetector.AcquireWriteAccess());
 	bWriting = true;
 	AccessThread.Join();
-	Success |= TestTrue(TEXT("Releasing Write Access"), MTAccessDetector.ReleaseWriteAccess());
+	bSuccess &= TestTrue(TEXT("Releasing Write Access"), MTAccessDetector.ReleaseWriteAccess());
 
-	return Success;
+	return bSuccess;
 }
 
 IMPLEMENT_SIMPLE_AUTOMATION_TEST(FRWAccessDetector_WriteReentrance, "System.Core.Misc.MTAccessDetector.WriteReentrance", EAutomationTestFlags_ApplicationContextMask | EAutomationTestFlags::EngineFilter)
@@ -136,13 +136,13 @@ bool FRWAccessDetector_WriteReentrance::RunTest(const FString& Parameters)
 
 	UE_MT_DECLARE_RW_ACCESS_DETECTOR(MTAccessDetector);
 
-	bool Success = false;
-	Success |= TestTrue(TEXT("Acquiring 1st Write Access"), MTAccessDetector.AcquireWriteAccess());
-	Success |= TestFalse(TEXT("Acquiring 2nd Write Access"), MTAccessDetector.AcquireWriteAccess());
-	Success |= TestFalse(TEXT("Releasing Write Access"), MTAccessDetector.ReleaseWriteAccess());
-	Success |= TestTrue(TEXT("Releasing Write Access"), MTAccessDetector.ReleaseWriteAccess());
+	bool bSuccess = true;
+	bSuccess &= TestTrue(TEXT("Acquiring 1st Write Access"), MTAccessDetector.AcquireWriteAccess());
+	bSuccess &= TestFalse(TEXT("Acquiring 2nd Write Access"), MTAccessDetector.AcquireWriteAccess());
+	bSuccess &= TestFalse(TEXT("Releasing Write Access"), MTAccessDetector.ReleaseWriteAccess());
+	bSuccess &= TestTrue(TEXT("Releasing Write Access"), MTAccessDetector.ReleaseWriteAccess());
 
-	return Success;
+	return bSuccess;
 }
 
 IMPLEMENT_SIMPLE_AUTOMATION_TEST(FRWAccessDetector_ReadReentrance, "System.Core.Misc.MTAccessDetector.ReadReentrance", EAutomationTestFlags_ApplicationContextMask | EAutomationTestFlags::EngineFilter)
@@ -152,13 +152,13 @@ bool FRWAccessDetector_ReadReentrance::RunTest(const FString& Parameters)
 
 	UE_MT_DECLARE_RW_ACCESS_DETECTOR(MTAccessDetector);
 
-	bool Success = false;
-	Success |= TestTrue(TEXT("Acquiring Write Access"), MTAccessDetector.AcquireWriteAccess());
-	Success |= TestFalse(TEXT("Acquiring Read Access"), MTAccessDetector.AcquireReadAccess());
-	Success |= TestFalse(TEXT("Releasing Read Access"), MTAccessDetector.ReleaseReadAccess());
-	Success |= TestTrue(TEXT("Releasing Write Access"), MTAccessDetector.ReleaseWriteAccess());
+	bool bSuccess = true;
+	bSuccess &= TestTrue(TEXT("Acquiring Write Access"), MTAccessDetector.AcquireWriteAccess());
+	bSuccess &= TestFalse(TEXT("Acquiring Read Access"), MTAccessDetector.AcquireReadAccess());
+	bSuccess &= TestFalse(TEXT("Releasing Read Access"), MTAccessDetector.ReleaseReadAccess());
+	bSuccess &= TestTrue(TEXT("Releasing Write Access"), MTAccessDetector.ReleaseWriteAccess());
 
-	return Success;
+	return bSuccess;
 }
 
 //----------------------------------------------------------------------//
@@ -171,22 +171,22 @@ bool FRWAccessDetector_RecursiveConcurentReadTest::RunTest(const FString& Parame
 	bool bReading2 = false;
 	UE_MT_DECLARE_RW_RECURSIVE_ACCESS_DETECTOR(MTAccessDetector);
 
-	bool Success = false;
+	bool bSuccess = true;
 	FThread AccessThread = FThread(TEXT("AccessThread"), [&]()
 	{
-		Success |= TestTrue(TEXT("Acquiring Read Access 1"), MTAccessDetector.AcquireReadAccess());
+		bSuccess &= TestTrue(TEXT("Acquiring Read Access 1"), MTAccessDetector.AcquireReadAccess());
 		bReading1 = true;
 		while (!bReading2) { FPlatformProcess::SleepNoStats(0.01f); }
-		Success |= TestTrue(TEXT("Releasing Read Access 1"), MTAccessDetector.ReleaseReadAccess());
+		bSuccess &= TestTrue(TEXT("Releasing Read Access 1"), MTAccessDetector.ReleaseReadAccess());
 		return true;
 	});
 	while (!bReading1) { FPlatformProcess::SleepNoStats(0.01f); }
-	Success |= TestTrue(TEXT("Acquiring Read Access 2"), MTAccessDetector.AcquireReadAccess());
+	bSuccess &= TestTrue(TEXT("Acquiring Read Access 2"), MTAccessDetector.AcquireReadAccess());
 	bReading2 = true;
 	AccessThread.Join();
-	Success |= TestTrue(TEXT("Releasing Read Access 2"), MTAccessDetector.ReleaseReadAccess());
+	bSuccess &= TestTrue(TEXT("Releasing Read Access 2"), MTAccessDetector.ReleaseReadAccess());
 
-	return Success;
+	return bSuccess;
 }
 
 IMPLEMENT_SIMPLE_AUTOMATION_TEST(FRWAccessDetector_RecursiveConcurentWriteTest, "System.Core.Misc.MTAccessDetector.RecursiveConcurrentWriteAccess", EAutomationTestFlags_ApplicationContextMask | EAutomationTestFlags::EngineFilter)
@@ -198,22 +198,22 @@ bool FRWAccessDetector_RecursiveConcurentWriteTest::RunTest(const FString& Param
 	bool bWriting2 = false;
 	UE_MT_DECLARE_RW_RECURSIVE_ACCESS_DETECTOR(MTAccessDetector);
 
-	bool Success = false;
+	bool bSuccess = true;
 	FThread AccessThread = FThread(TEXT("AccessThread"), [&]()
 	{
-		Success |= TestTrue(TEXT("Acquiring Write Access"), MTAccessDetector.AcquireWriteAccess());
+		bSuccess &= TestTrue(TEXT("Acquiring Write Access"), MTAccessDetector.AcquireWriteAccess());
 		bWriting1 = true;
 		while (!bWriting2) { FPlatformProcess::SleepNoStats(0.01f); }
-		Success |= TestFalse(TEXT("Releasing Read Access"), MTAccessDetector.ReleaseWriteAccess());
+		bSuccess &= TestFalse(TEXT("Releasing Read Access"), MTAccessDetector.ReleaseWriteAccess());
 		return true;
 	});
 	while (!bWriting1) { FPlatformProcess::SleepNoStats(0.01f); }
-	Success |= TestFalse(TEXT("Acquiring Write Access"), MTAccessDetector.AcquireWriteAccess());
+	bSuccess &= TestFalse(TEXT("Acquiring Write Access"), MTAccessDetector.AcquireWriteAccess());
 	bWriting2 = true;
 	AccessThread.Join();
-	Success |= TestTrue(TEXT("Releasing Write Access"), MTAccessDetector.ReleaseWriteAccess());
+	bSuccess &= TestTrue(TEXT("Releasing Write Access"), MTAccessDetector.ReleaseWriteAccess());
 
-	return Success;
+	return bSuccess;
 }
 
 IMPLEMENT_SIMPLE_AUTOMATION_TEST(FRWAccessDetector_RecursiveConcurentReadWriteTest, "System.Core.Misc.MTAccessDetector.RecursiveConcurrentReadWriteAccess", EAutomationTestFlags_ApplicationContextMask | EAutomationTestFlags::EngineFilter)
@@ -225,22 +225,22 @@ bool FRWAccessDetector_RecursiveConcurentReadWriteTest::RunTest(const FString& P
 	bool bWriting = false;
 	UE_MT_DECLARE_RW_RECURSIVE_ACCESS_DETECTOR(MTAccessDetector);
 
-	bool Success = false;
+	bool bSuccess = true;
 	FThread AccessThread = FThread(TEXT("AccessThread"), [&]()
 	{
-		Success |= TestTrue(TEXT("Acquiring Read Access"), MTAccessDetector.AcquireReadAccess());
+		bSuccess &= TestTrue(TEXT("Acquiring Read Access"), MTAccessDetector.AcquireReadAccess());
 		bReading = true;
 		while (!bWriting) { FPlatformProcess::SleepNoStats(0.01f); }
-		Success |= TestFalse(TEXT("Releasing Read Access"), MTAccessDetector.ReleaseReadAccess());
+		bSuccess &= TestFalse(TEXT("Releasing Read Access"), MTAccessDetector.ReleaseReadAccess());
 		return true;
 	});
 	while (!bReading) { FPlatformProcess::SleepNoStats(0.01f); }
-	Success |= TestFalse(TEXT("Acquiring Write Access"), MTAccessDetector.AcquireWriteAccess());
+	bSuccess &= TestFalse(TEXT("Acquiring Write Access"), MTAccessDetector.AcquireWriteAccess());
 	bWriting = true;
 	AccessThread.Join();
-	Success |= TestTrue(TEXT("Releasing Write Access"), MTAccessDetector.ReleaseWriteAccess());
+	bSuccess &= TestTrue(TEXT("Releasing Write Access"), MTAccessDetector.ReleaseWriteAccess());
 
-	return Success;
+	return bSuccess;
 }
 
 IMPLEMENT_SIMPLE_AUTOMATION_TEST(FRWAccessDetector_RecursiveWriteReentrance, "System.Core.Misc.MTAccessDetector.RecursiveWriteReentrance", EAutomationTestFlags_ApplicationContextMask | EAutomationTestFlags::EngineFilter)
@@ -248,13 +248,13 @@ bool FRWAccessDetector_RecursiveWriteReentrance::RunTest(const FString& Paramete
 {
 	UE_MT_DECLARE_RW_RECURSIVE_ACCESS_DETECTOR(MTAccessDetector);
 
-	bool Success = false;
-	Success |= TestTrue(TEXT("Acquiring 1st Write Access"), MTAccessDetector.AcquireWriteAccess());
-	Success |= TestTrue(TEXT("Acquiring 2nd Write Access"), MTAccessDetector.AcquireWriteAccess());
-	Success |= TestTrue(TEXT("Releasing Write Access"), MTAccessDetector.ReleaseWriteAccess());
-	Success |= TestTrue(TEXT("Releasing Write Access"), MTAccessDetector.ReleaseWriteAccess());
+	bool bSuccess = true;
+	bSuccess &= TestTrue(TEXT("Acquiring 1st Write Access"), MTAccessDetector.AcquireWriteAccess());
+	bSuccess &= TestTrue(TEXT("Acquiring 2nd Write Access"), MTAccessDetector.AcquireWriteAccess());
+	bSuccess &= TestTrue(TEXT("Releasing Write Access"), MTAccessDetector.ReleaseWriteAccess());
+	bSuccess &= TestTrue(TEXT("Releasing Write Access"), MTAccessDetector.ReleaseWriteAccess());
 
-	return Success;
+	return bSuccess;
 }
 
 IMPLEMENT_SIMPLE_AUTOMATION_TEST(FRWAccessDetector_RecursiveReadReentrance, "System.Core.Misc.MTAccessDetector.RecursiveReadReentrance", EAutomationTestFlags_ApplicationContextMask | EAutomationTestFlags::EngineFilter)
@@ -264,13 +264,13 @@ bool FRWAccessDetector_RecursiveReadReentrance::RunTest(const FString& Parameter
 
 	UE_MT_DECLARE_RW_RECURSIVE_ACCESS_DETECTOR(MTAccessDetector);
 
-	bool Success = false;
-	Success |= TestTrue(TEXT("Acquiring Write Access"), MTAccessDetector.AcquireWriteAccess());
-	Success |= TestFalse(TEXT("Acquiring Read Access"), MTAccessDetector.AcquireReadAccess());
-	Success |= TestFalse(TEXT("Releasing Read Access"), MTAccessDetector.ReleaseReadAccess());
-	Success |= TestTrue(TEXT("Releasing Write Access"), MTAccessDetector.ReleaseWriteAccess());
+	bool bSuccess = true;
+	bSuccess &= TestTrue(TEXT("Acquiring Write Access"), MTAccessDetector.AcquireWriteAccess());
+	bSuccess &= TestFalse(TEXT("Acquiring Read Access"), MTAccessDetector.AcquireReadAccess());
+	bSuccess &= TestFalse(TEXT("Releasing Read Access"), MTAccessDetector.ReleaseReadAccess());
+	bSuccess &= TestTrue(TEXT("Releasing Write Access"), MTAccessDetector.ReleaseWriteAccess());
 
-	return Success;
+	return bSuccess;
 }
 
 //----------------------------------------------------------------------//
@@ -283,22 +283,22 @@ bool FRWAccessDetector_FullyRecursiveConcurentReadTest::RunTest(const FString& P
 	bool bReading2 = false;
 	UE_MT_DECLARE_RW_FULLY_RECURSIVE_ACCESS_DETECTOR(MTAccessDetector);
 
-	bool Success = false;
+	bool bSuccess = true;
 	FThread AccessThread = FThread(TEXT("AccessThread"), [&]()
 	{
-		Success |= TestTrue(TEXT("Acquiring Read Access 1"), MTAccessDetector.AcquireReadAccess());
+		bSuccess &= TestTrue(TEXT("Acquiring Read Access 1"), MTAccessDetector.AcquireReadAccess());
 		bReading1 = true;
 		while (!bReading2) { FPlatformProcess::SleepNoStats(0.01f); }
-		Success |= TestTrue(TEXT("Releasing Read Access 1"), MTAccessDetector.ReleaseReadAccess());
+		bSuccess &= TestTrue(TEXT("Releasing Read Access 1"), MTAccessDetector.ReleaseReadAccess());
 		return true;
 	});
 	while (!bReading1) { FPlatformProcess::SleepNoStats(0.01f); }
-	Success |= TestTrue(TEXT("Acquiring Read Access 2"), MTAccessDetector.AcquireReadAccess());
+	bSuccess &= TestTrue(TEXT("Acquiring Read Access 2"), MTAccessDetector.AcquireReadAccess());
 	bReading2 = true;
 	AccessThread.Join();
-	Success |= TestTrue(TEXT("Releasing Read Access 2"), MTAccessDetector.ReleaseReadAccess());
+	bSuccess &= TestTrue(TEXT("Releasing Read Access 2"), MTAccessDetector.ReleaseReadAccess());
 
-	return Success;
+	return bSuccess;
 }
 
 IMPLEMENT_SIMPLE_AUTOMATION_TEST(FRWAccessDetector_FullyRecursiveConcurentWriteTest, "System.Core.Misc.MTAccessDetector.FullyRecursiveConcurrentWriteAccess", EAutomationTestFlags_ApplicationContextMask | EAutomationTestFlags::EngineFilter)
@@ -310,22 +310,22 @@ bool FRWAccessDetector_FullyRecursiveConcurentWriteTest::RunTest(const FString&
 	bool bWriting2 = false;
 	UE_MT_DECLARE_RW_FULLY_RECURSIVE_ACCESS_DETECTOR(MTAccessDetector);
 
-	bool Success = false;
+	bool bSuccess = true;
 	FThread AccessThread = FThread(TEXT("AccessThread"), [&]()
 	{
-		Success |= TestTrue(TEXT("Acquiring Write Access"), MTAccessDetector.AcquireWriteAccess());
+		bSuccess &= TestTrue(TEXT("Acquiring Write Access"), MTAccessDetector.AcquireWriteAccess());
 		bWriting1 = true;
 		while (!bWriting2) { FPlatformProcess::SleepNoStats(0.01f); }
-		Success |= TestFalse(TEXT("Releasing Read Access"), MTAccessDetector.ReleaseWriteAccess());
+		bSuccess &= TestFalse(TEXT("Releasing Read Access"), MTAccessDetector.ReleaseWriteAccess());
 		return true;
 	});
 	while (!bWriting1) { FPlatformProcess::SleepNoStats(0.01f); }
-	Success |= TestFalse(TEXT("Acquiring Write Access"), MTAccessDetector.AcquireWriteAccess());
+	bSuccess &= TestFalse(TEXT("Acquiring Write Access"), MTAccessDetector.AcquireWriteAccess());
 	bWriting2 = true;
 	AccessThread.Join();
-	Success |= TestTrue(TEXT("Releasing Write Access"), MTAccessDetector.ReleaseWriteAccess());
+	bSuccess &= TestTrue(TEXT("Releasing Write Access"), MTAccessDetector.ReleaseWriteAccess());
 
-	return Success;
+	return bSuccess;
 }
 
 IMPLEMENT_SIMPLE_AUTOMATION_TEST(FRWAccessDetector_FullyRecursiveConcurentReadWriteTest, "System.Core.Misc.MTAccessDetector.FullyRecursiveConcurrentReadWriteAccess", EAutomationTestFlags_ApplicationContextMask | EAutomationTestFlags::EngineFilter)
@@ -337,22 +337,22 @@ bool FRWAccessDetector_FullyRecursiveConcurentReadWriteTest::RunTest(const FStri
 	bool bWriting = false;
 	UE_MT_DECLARE_RW_FULLY_RECURSIVE_ACCESS_DETECTOR(MTAccessDetector);
 
-	bool Success = false;
+	bool bSuccess = true;
 	FThread AccessThread = FThread(TEXT("AccessThread"), [&]()
 	{
-		Success |= TestTrue(TEXT("Acquiring Read Access"), MTAccessDetector.AcquireReadAccess());
+		bSuccess &= TestTrue(TEXT("Acquiring Read Access"), MTAccessDetector.AcquireReadAccess());
 		bReading = true;
 		while (!bWriting) { FPlatformProcess::SleepNoStats(0.01f); }
-		Success |= TestFalse(TEXT("Releasing Read Access"), MTAccessDetector.ReleaseReadAccess());
+		bSuccess &= TestFalse(TEXT("Releasing Read Access"), MTAccessDetector.ReleaseReadAccess());
 		return true;
 	});
 	while (!bReading) { FPlatformProcess::SleepNoStats(0.01f); }
-	Success |= TestFalse(TEXT("Acquiring Write Access"), MTAccessDetector.AcquireWriteAccess());
+	bSuccess &= TestFalse(TEXT("Acquiring Write Access"), MTAccessDetector.AcquireWriteAccess());
 	bWriting = true;
 	AccessThread.Join();
-	Success |= TestTrue(TEXT("Releasing Write Access"), MTAccessDetector.ReleaseWriteAccess());
+	bSuccess &= TestTrue(TEXT("Releasing Write Access"), MTAccessDetector.ReleaseWriteAccess());
 
-	return Success;
+	return bSuccess;
 }
 
 IMPLEMENT_SIMPLE_AUTOMATION_TEST(FRWAccessDetector_FullyRecursiveWriteReentrance, "System.Core.Misc.MTAccessDetector.FullyRecursiveWriteReentrance", EAutomationTestFlags_ApplicationContextMask | EAutomationTestFlags::EngineFilter)
@@ -360,13 +360,13 @@ bool FRWAccessDetector_FullyRecursiveWriteReentrance::RunTest(const FString& Par
 {
 	UE_MT_DECLARE_RW_FULLY_RECURSIVE_ACCESS_DETECTOR(MTAccessDetector);
 
-	bool Success = false;
-	Success |= TestTrue(TEXT("Acquiring 1st Write Access"), MTAccessDetector.AcquireWriteAccess());
-	Success |= TestTrue(TEXT("Acquiring 2nd Write Access"), MTAccessDetector.AcquireWriteAccess());
-	Success |= TestTrue(TEXT("Releasing Write Access"), MTAccessDetector.ReleaseWriteAccess());
-	Success |= TestTrue(TEXT("Releasing Write Access"), MTAccessDetector.ReleaseWriteAccess());
+	bool bSuccess = true;
+	bSuccess &= TestTrue(TEXT("Acquiring 1st Write Access"), MTAccessDetector.AcquireWriteAccess());
+	bSuccess &= TestTrue(TEXT("Acquiring 2nd Write Access"), MTAccessDetector.AcquireWriteAccess());
+	bSuccess &= TestTrue(TEXT("Releasing Write Access"), MTAccessDetector.ReleaseWriteAccess());
+	bSuccess &= TestTrue(TEXT("Releasing Write Access"), MTAccessDetector.ReleaseWriteAccess());
 
-	return Success;
+	return bSuccess;
 }
 
 IMPLEMENT_SIMPLE_AUTOMATION_TEST(FRWAccessDetector_FullyRecursiveReadReentrance, "System.Core.Misc.MTAccessDetector.FullyRecursiveReadReentrance", EAutomationTestFlags_ApplicationContextMask | EAutomationTestFlags::EngineFilter)
@@ -374,13 +374,13 @@ bool FRWAccessDetector_FullyRecursiveReadReentrance::RunTest(const FString& Para
 {
 	UE_MT_DECLARE_RW_FULLY_RECURSIVE_ACCESS_DETECTOR(MTAccessDetector);
 
-	bool Success = false;
-	Success |= TestTrue(TEXT("Acquiring Write Access"), MTAccessDetector.AcquireWriteAccess());
-	Success |= TestTrue(TEXT("Acquiring Read Access"), MTAccessDetector.AcquireReadAccess());
-	Success |= TestTrue(TEXT("Releasing Read Access"), MTAccessDetector.ReleaseReadAccess());
-	Success |= TestTrue(TEXT("Releasing Write Access"), MTAccessDetector.ReleaseWriteAccess());
+	bool bSuccess = true;
+	bSuccess &= TestTrue(TEXT("Acquiring Write Access"), MTAccessDetector.AcquireWriteAccess());
+	bSuccess &= TestTrue(TEXT("Acquiring Read Access"), MTAccessDetector.AcquireReadAccess());
+	bSuccess &= TestTrue(TEXT("Releasing Read Access"), MTAccessDetector.ReleaseReadAccess());
+	bSuccess &= TestTrue(TEXT("Releasing Write Access"), MTAccessDetector.ReleaseWriteAccess());
 
-	return Success;
+	return bSuccess;
 }
 
 /////////////////////////////////////////////////////////////////////////
