Fix issue where command line overrides were not being applied or took effect when sampling gconfig variables.

diff --git a/Engine/Source/Runtime/Core/Private/Misc/ConfigContext.cpp b/Engine/Source/Runtime/Core/Private/Misc/ConfigContext.cpp
--- a/Engine/Source/Runtime/Core/Private/Misc/ConfigContext.cpp
+++ b/Engine/Source/Runtime/Core/Private/Misc/ConfigContext.cpp
@@ -1217,6 +1217,8 @@ bool FConfigContext::GenerateDestIniFile()
 		Branch->InMemoryFile.ApplyFile(&Branch->SavedLayer);
 
 #if ALLOW_INI_OVERRIDE_FROM_COMMANDLINE
+		// process any commandline overrides
+		FConfigFile::OverrideFromCommandline(&Branch->CommandLineOverrides, BaseIniName);
 		// and push it into the current values
 		Branch->FinalCombinedLayers.ApplyFile(&Branch->CommandLineOverrides);
 		Branch->InMemoryFile.ApplyFile(&Branch->CommandLineOverrides);
