Start worker thread on any new connection if not yet started.

During initialization only start TraceLog worker thread if tracing is requested on command line, otherwise it's started on demand. Until then TraceLog is only pumped using the EndFrame delegate. For the editor we always start a worker thread since frame pump is not reliable and start very late in the process startup. The exception is where worker thread is completely prohibited or opted out using "-notracethreading". For forked processes the configuration stays "Unset" until TraceAux is initialized in the forked process.

diff --git a/Engine/Source/Runtime/Core/Private/ProfilingDebugging/TraceAuxiliary.cpp b/Engine/Source/Runtime/Core/Private/ProfilingDebugging/TraceAuxiliary.cpp
--- a/Engine/Source/Runtime/Core/Private/ProfilingDebugging/TraceAuxiliary.cpp
+++ b/Engine/Source/Runtime/Core/Private/ProfilingDebugging/TraceAuxiliary.cpp
@@ -61,7 +61,6 @@ const FTraceAuxiliary::FChannelPreset GMemoryLightChannels(TEXT("Memory_Light"),
 
 #if UE_TRACE_ENABLED
 
-#include <atomic>
 #include "BuildSettings.h"
 #include "Containers/Array.h"
 #include "Containers/Map.h"
@@ -69,6 +68,8 @@ const FTraceAuxiliary::FChannelPreset GMemoryLightChannels(TEXT("Memory_Light"),
 #include "Containers/UnrealString.h"
 #include "HAL/FileManager.h"
 #include "HAL/IConsoleManager.h"
+#include "HAL/PlatformTime.h"
+#include "Logging/LogMacros.h"
 #include "Misc/App.h"
 #include "Misc/CString.h"
 #include "Misc/ConfigCacheIni.h"
@@ -81,9 +82,7 @@ const FTraceAuxiliary::FChannelPreset GMemoryLightChannels(TEXT("Memory_Light"),
 #include "ProfilingDebugging/CallstackTrace.h"
 #include "ProfilingDebugging/CpuProfilerTrace.h"
 #include "ProfilingDebugging/CsvProfiler.h"
-#include "ProfilingDebugging/MemoryTrace.h"
 #include "ProfilingDebugging/PlatformEvents.h"
-#include "Stats/Stats.h"
 #include "String/ParseTokens.h"
 #include "Templates/Invoke.h"
 #include "Templates/UnrealTemplate.h"
@@ -94,9 +93,19 @@ const FTraceAuxiliary::FChannelPreset GMemoryLightChannels(TEXT("Memory_Light"),
 const TCHAR* GTraceConfigSection = TEXT("Trace.Config");
 static UE::Trace::FInitializeDesc GInitializeDesc;
 
+enum class EWorkerThreadConfig : uint8 {
+	Unset,
+	OnDemand,		// Worker thread started when first trace is started
+	OnInit,			// Worker thread started on boot
+	Never,			// Worker threda never started (rely on end frame pump)
+};
+
+static EWorkerThreadConfig GWorkerThreadConfig = EWorkerThreadConfig::Unset;
+
 CSV_DEFINE_CATEGORY(Trace, true);
 
 #if UE_TRACEAUX_EMIT_STAT_COUNTERS
+#include "Stats/Stats.h"
 DECLARE_STATS_GROUP(TEXT("TraceLog"), STATGROUP_Trace, STATCAT_Advanced);
 DECLARE_MEMORY_STAT(TEXT("Total Memory Used"),             STAT_Trace_MemoryUsed,                 STATGROUP_Trace);
 DECLARE_MEMORY_STAT(TEXT("Block Pool Memory"),             STAT_Trace_Memory_BlockPool,           STATGROUP_Trace);
@@ -877,7 +886,9 @@ void FTraceAuxiliaryImpl::GetActiveChannelsString(FStringBuilderBase& String) co
 ////////////////////////////////////////////////////////////////////////////////
 void FTraceAuxiliaryImpl::StartWorkerThread()
 {
-	if (!bWorkerThreadStarted)
+	const bool bIsForkSafe = !FForkProcessHelper::IsForkedChildProcess() || FForkProcessHelper::IsForkedMultithreadInstance();
+	const bool bShouldStart = GWorkerThreadConfig == EWorkerThreadConfig::OnDemand && bIsForkSafe;
+	if (!bWorkerThreadStarted && bShouldStart)
 	{
 		TRACE_CPUPROFILER_EVENT_SCOPE(FTraceAux_StartWorkerThread);
 		UE::Trace::StartWorkerThread();
@@ -1003,6 +1014,7 @@ void FTraceAuxiliaryImpl::RegisterEndFrameCallbacks()
 void TraceAuxiliaryOnConnectionCallback()
 {
 	TRACE_CPUPROFILER_EVENT_SCOPE(FTraceAux_OnConnection);
+	GTraceAuxiliary.StartWorkerThread();
 	FTraceAuxiliary::OnConnection.Broadcast();
 }
 
@@ -1511,7 +1523,7 @@ UE_TRACE_EVENT_BEGIN(Diagnostics, Session2, NoSync|Important)
 UE_TRACE_EVENT_END()
 
 ////////////////////////////////////////////////////////////////////////////////
-static bool StartFromCommandlineArguments(const TCHAR* CommandLine, bool& bOutStartWorkerThread)
+static bool StartFromCommandlineArguments(const TCHAR* CommandLine)
 {
 #if UE_TRACE_ENABLED
 
@@ -1538,13 +1550,6 @@ static bool StartFromCommandlineArguments(const TCHAR* CommandLine, bool& bOutSt
 	FTraceAuxiliary::FOptions Opts;
 	Opts.bTruncateFile = FParse::Param(CommandLine, TEXT("tracefiletrunc"));
 
-	const bool bWorkerThreadAllowed = FGenericPlatformProcess::SupportsMultithreading() || FForkProcessHelper::IsForkedMultithreadInstance();
-
-	if (!bWorkerThreadAllowed || FParse::Param(FCommandLine::Get(), TEXT("notracethreading")))
-	{
-		Opts.bNoWorkerThread = true;
-	}
-
 	// Find if a connection type is specified
 	FString Parameter;
 	const TCHAR* Target = nullptr;
@@ -1594,11 +1599,6 @@ static bool StartFromCommandlineArguments(const TCHAR* CommandLine, bool& bOutSt
 		return false;
 	}
 
-	// Trace's worker thread should really only be started by Trace itself as
-	// order is important. At the very least it must be done after Trace is
-	// initialized. It isn't yet here so we defer it.
-	bOutStartWorkerThread = !Opts.bNoWorkerThread;
-	Opts.bNoWorkerThread = true;
 
 	// Finally start tracing to the requested connection
 	return FTraceAuxiliary::Start(Type, Target, *Channels, &Opts);
@@ -1630,15 +1630,6 @@ static bool TraceAuxiliaryStartShared(const TCHAR* Channels, const FTraceAuxilia
 		return false;
 	}
 
-	// Make sure the worker thread is started unless explicitly opt out.
-	if (!Options || !Options->bNoWorkerThread)
-	{
-		if (FGenericPlatformProcess::SupportsMultithreading() || FForkProcessHelper::IsForkedMultithreadInstance())
-		{
-			GTraceAuxiliary.StartWorkerThread();
-		}
-	}
-
 	if (Channels)
 	{
 		TRACE_CPUPROFILER_EVENT_SCOPE(FTraceAux_EnableChannels);
@@ -1889,8 +1880,7 @@ void FTraceAuxiliary::Initialize(const TCHAR* CommandLine)
 	// Attempt to send trace data somewhere from the command line. It perhaps
 	// seems odd to do this before initializing Trace, but it is done this way
 	// to support disabling the "important" cache without losing any events.
-	bool bShouldStartWorkerThread = false;
-	StartFromCommandlineArguments(CommandLine, bShouldStartWorkerThread);
+	const bool bWasConnected = StartFromCommandlineArguments(CommandLine);
 
 	// Emit empty stats once (to ensure all stats/counters start from zero).
 	{
@@ -1910,10 +1900,24 @@ void FTraceAuxiliary::Initialize(const TCHAR* CommandLine)
 #endif
 	TraceAuxiliarySetupInitFromConfig(Desc);
 
-	Desc.bUseWorkerThread = bShouldStartWorkerThread;
 	Desc.bUseImportantCache = (FParse::Param(CommandLine, TEXT("tracenocache")) == false);
 	Desc.OnConnectionFunc = &TraceAuxiliaryOnConnectionCallback;
 	Desc.OnUpdateFunc = &TraceAuxiliaryOnTraceUpdateCallback;
+	
+	// Determine if the worker thread is allowed 
+	if (!FGenericPlatformProcess::SupportsMultithreading() || FParse::Param(FCommandLine::Get(), TEXT("notracethreading")))
+	{
+		GWorkerThreadConfig = EWorkerThreadConfig::Never;
+		Desc.bUseWorkerThread = false;
+	}
+	else
+	{
+		// If allowed start worker thread if trace is started using command line argument,
+		// otherwise start on demand. For the editor we always want to start a worker thread
+		// since there is a lot of time until frames are reliably pumped.
+		Desc.bUseWorkerThread = bWasConnected || UE_EDITOR;
+		GWorkerThreadConfig = Desc.bUseWorkerThread ? EWorkerThreadConfig::OnInit : EWorkerThreadConfig::OnDemand;
+	}
 
 	FGuid SessionGuid;
 	if (!FParse::Value(CommandLine, TEXT("-tracesessionguid="), SessionGuid))
diff --git a/Engine/Source/Runtime/Core/Public/ProfilingDebugging/TraceAuxiliary.h b/Engine/Source/Runtime/Core/Public/ProfilingDebugging/TraceAuxiliary.h
--- a/Engine/Source/Runtime/Core/Public/ProfilingDebugging/TraceAuxiliary.h
+++ b/Engine/Source/Runtime/Core/Public/ProfilingDebugging/TraceAuxiliary.h
@@ -10,6 +10,7 @@
 #include "Logging/LogCategory.h"
 #include "Logging/LogMacros.h"
 #include "Misc/Build.h"
+#include "Misc/CoreMiscDefines.h"
 
 #if !defined(UE_TRACE_SERVER_CONTROLS_ENABLED)
 #	define UE_TRACE_SERVER_CONTROLS_ENABLED (PLATFORM_DESKTOP && !UE_BUILD_SHIPPING && !IS_PROGRAM)
@@ -117,6 +118,7 @@ public:
 	struct FOptions
 	{
 		/** When set, trace will not start a worker thread, instead it is updated from end frame delegate. */
+		UE_DEPRECATED(5.7, "Use -notracethread command line instead")
 		bool bNoWorkerThread = false;
 		/** When set, the target file will be truncated if it already exists. */
 		bool bTruncateFile = false;
