Fixed potential deadlock on boot due to logging occurring in CachedOSVeryLargePageAllocator::Refresh. Removed all logging from that function.

diff --git a/Engine/Source/Runtime/Core/Private/HAL/Allocators/CachedOSVeryLargePageAllocator.cpp b/Engine/Source/Runtime/Core/Private/HAL/Allocators/CachedOSVeryLargePageAllocator.cpp
--- a/Engine/Source/Runtime/Core/Private/HAL/Allocators/CachedOSVeryLargePageAllocator.cpp
+++ b/Engine/Source/Runtime/Core/Private/HAL/Allocators/CachedOSVeryLargePageAllocator.cpp
@@ -113,60 +113,42 @@ void FCachedOSVeryLargePageAllocator::Refresh()
 	{
 		return;
 	}
-
-	// Function is not thread safe 
-	UE_LOG(LogMemory, Log, TEXT("LargePageAllocator - Refresh"));
 	
+	// Shrink Empty back store
+	for (int i = 0; i < FMemory::AllocationHints::Max; i++)
 	{
-		// Shrink Empty back store
-		uint64 StartTime = FPlatformTime::Cycles64();
-
-		for (int i = 0; i < FMemory::AllocationHints::Max; i++)
-		{
-			FMemory::AllocationHints AllocationHint = FMemory::AllocationHints(i);
-			ShrinkEmptyBackStore(0, AllocationHint);
-		}
-
-		uint64 LenghtCycles64 = (FPlatformTime::Cycles64() - StartTime);
-		UE_LOG(LogMemory, Log, TEXT("LargePageAllocator - Shrink empty backstore duration: %.1f ms"), float(FPlatformTime::ToMilliseconds64(LenghtCycles64)));
+		FMemory::AllocationHints AllocationHint = FMemory::AllocationHints(i);
+		ShrinkEmptyBackStore(0, AllocationHint);
 	}
 
+	// Preallocate
+	for (int32 i = 0; i < FMemory::AllocationHints::Max; i++)
 	{
-		// Preallocate
-		uint64 StartTime = FPlatformTime::Cycles64();
-
-		for (int32 i = 0; i < FMemory::AllocationHints::Max; i++)
-		{
-			FMemory::AllocationHints AllocationHint = FMemory::AllocationHints(i);
+		FMemory::AllocationHints AllocationHint = FMemory::AllocationHints(i);
 
-			int32 LargePageCount = CommittedLargePagesCount[AllocationHint] + EmptyBackStoreCount[AllocationHint];
+		int32 LargePageCount = CommittedLargePagesCount[AllocationHint] + EmptyBackStoreCount[AllocationHint];
 
-			if (LargePageCount < GMaxEmptyBackStoreCount[AllocationHint] && GPreAllocatePools)
+		if (LargePageCount < GMaxEmptyBackStoreCount[AllocationHint] && GPreAllocatePools)
+		{
+			// Preallocate large pages
+			for (; LargePageCount < GMaxEmptyBackStoreCount[AllocationHint]; LargePageCount++)
 			{
-				// Preallocate large pages
-				for (; LargePageCount < GMaxEmptyBackStoreCount[AllocationHint]; LargePageCount++)
+				bool bCommitFailure = false;
+				FLargePage* LargePage = AllocNewLargePage(AllocationHint, nullptr, bCommitFailure);
+				if (!LargePage)
 				{
-					bool bCommitFailure = false;
-					FLargePage* LargePage = AllocNewLargePage(AllocationHint, nullptr, bCommitFailure);
-					if (!LargePage)
+					if (bCommitFailure && GDisablePageCachingOnOOM)
 					{
-						if (bCommitFailure && GDisablePageCachingOnOOM)
-						{
-							// If we failed to preallocate due to a commit failure then disable preallocation in future
-							UE_LOG(LogMemory, Warning, TEXT("LargePageAllocator - Preallocation permanently disabled due to failure to commit pages"));
-							GPreAllocatePools = false;
-						}
-						break;
+						// If we failed to preallocate due to a commit failure then disable preallocation in future
+						GPreAllocatePools = false;
 					}
-					LargePage->LinkHead(EmptyButAvailableLargePagesHead[AllocationHint]);
-					EmptyBackStoreCount[AllocationHint] += 1;
-					ImmediatelyFreeable += SizeOfLargePage;
+					break;
 				}
+				LargePage->LinkHead(EmptyButAvailableLargePagesHead[AllocationHint]);
+				EmptyBackStoreCount[AllocationHint] += 1;
+				ImmediatelyFreeable += SizeOfLargePage;
 			}
 		}
-
-		uint64 LenghtCycles64 = (FPlatformTime::Cycles64() - StartTime);
-		UE_LOG(LogMemory, Log, TEXT("LargePageAllocator - Preallocate memory duration: %.1f ms"), float(FPlatformTime::ToMilliseconds64(LenghtCycles64)));
 	}
 }
 
