ParkingLot: Optimized callback code generation using a private TLambdaAsFunctionWithContext

FMutex::LockSlow() from 77 to 50 instructions on Windows.
FEventCount::Wait() from 30 to 15 instructions on Windows.

diff --git a/Engine/Source/Runtime/Core/Private/Async/ParkingLot.cpp b/Engine/Source/Runtime/Core/Private/Async/ParkingLot.cpp
--- a/Engine/Source/Runtime/Core/Private/Async/ParkingLot.cpp
+++ b/Engine/Source/Runtime/Core/Private/Async/ParkingLot.cpp
@@ -500,18 +500,18 @@ FThreadLocalData::~FThreadLocalData()
 
 ///////////////////////////////////////////////////////////////////////////////////////////////////////////////////////
 
-FWaitState Wait(const void* Address, const TFunctionRef<bool()>& CanWait, const TFunctionRef<void()>& BeforeWait)
+FWaitState Wait(const void* Address, bool (*CanWait)(void*), void* CanWaitContext, void (*BeforeWait)(void*), void* BeforeWaitContext)
 {
-	return WaitUntil(Address, CanWait, BeforeWait, FMonotonicTimePoint::Infinity());
+	return WaitUntil(Address, CanWait, CanWaitContext, BeforeWait, BeforeWaitContext, FMonotonicTimePoint::Infinity());
 }
 
-FWaitState WaitFor(const void* Address, const TFunctionRef<bool()>& CanWait, const TFunctionRef<void()>& BeforeWait, FMonotonicTimeSpan WaitTime)
+FWaitState WaitFor(const void* Address, bool (*CanWait)(void*), void* CanWaitContext, void (*BeforeWait)(void*), void* BeforeWaitContext, FMonotonicTimeSpan WaitTime)
 {
-	return WaitUntil(Address, CanWait, BeforeWait, FMonotonicTimePoint::Now() + WaitTime);
+	return WaitUntil(Address, CanWait, CanWaitContext, BeforeWait, BeforeWaitContext, FMonotonicTimePoint::Now() + WaitTime);
 }
 
 UE_AUTORTFM_NOAUTORTFM
-FWaitState WaitUntil(const void* Address, const TFunctionRef<bool()>& CanWait, const TFunctionRef<void()>& BeforeWait, FMonotonicTimePoint WaitTime)
+FWaitState WaitUntil(const void* Address, bool (*CanWait)(void*), void* CanWaitContext, void (*BeforeWait)(void*), void* BeforeWaitContext, FMonotonicTimePoint WaitTime)
 {
 	using namespace Private;
 
@@ -528,7 +528,7 @@ FWaitState WaitUntil(const void* Address, const TFunctionRef<bool()>& CanWait, c
 	{
 		TDynamicUniqueLock<FWordMutex> BucketLock;
 		FBucket& Bucket = FTable::FindOrCreateBucket(Address, BucketLock);
-		State.bDidWait = CanWait();
+		State.bDidWait = !CanWait || CanWait(CanWaitContext);
 		if (!State.bDidWait)
 		{
 			return State;
@@ -539,7 +539,10 @@ FWaitState WaitUntil(const void* Address, const TFunctionRef<bool()>& CanWait, c
 	}
 
 	// BeforeWait must be invoked after the bucket is unlocked.
-	BeforeWait();
+	if (BeforeWait)
+	{
+		BeforeWait(BeforeWaitContext);
+	}
 
 	// Wait until the timeout or until the thread has been dequeued.
 	Self->Event.WaitUntil(WaitTime);
@@ -584,7 +587,7 @@ FWaitState WaitUntil(const void* Address, const TFunctionRef<bool()>& CanWait, c
 }
 
 UE_AUTORTFM_NOAUTORTFM
-void WakeOne(const void* Address, const TFunctionRef<uint64(FWakeState)>& OnWakeState)
+void WakeOne(const void* Address, uint64 (*OnWakeState)(void*, FWakeState), void* OnWakeStateContext)
 {
 	using namespace Private;
 
@@ -606,7 +609,7 @@ void WakeOne(const void* Address, const TFunctionRef<uint64(FWakeState)>& OnWake
 		FWakeState WakeState;
 		WakeState.bDidWake = !!WakeThread;
 		WakeState.bHasWaitingThreads = !Bucket.IsEmpty();
-		WakeToken = OnWakeState(WakeState);
+		WakeToken = OnWakeState ? OnWakeState(OnWakeStateContext, WakeState) : 0;
 	}
 
 	if (WakeThread)
@@ -626,7 +629,7 @@ namespace UE::ParkingLot
 FWakeState WakeOne(const void* Address)
 {
 	FWakeState OutState;
-	Private::WakeOne(Address, [&OutState](FWakeState State) -> uint64
+	WakeOne(Address, [&OutState](FWakeState State) -> uint64
 	{
 		OutState = State;
 		return 0;
diff --git a/Engine/Source/Runtime/Core/Public/Async/ParkingLot.h b/Engine/Source/Runtime/Core/Public/Async/ParkingLot.h
--- a/Engine/Source/Runtime/Core/Public/Async/ParkingLot.h
+++ b/Engine/Source/Runtime/Core/Public/Async/ParkingLot.h
@@ -4,7 +4,9 @@
 
 #include "CoreTypes.h"
 #include "Misc/MonotonicTime.h"
-#include "Templates/Function.h"
+#include "Templates/Invoke.h"
+
+#include <type_traits>
 
 #define UE_API CORE_API
 
@@ -34,11 +36,66 @@ struct FWakeState final
 
 namespace Private
 {
-UE_API FWaitState Wait(const void* Address, const TFunctionRef<bool()>& CanWait, const TFunctionRef<void()>& BeforeWait);
-UE_API FWaitState WaitFor(const void* Address, const TFunctionRef<bool()>& CanWait, const TFunctionRef<void()>& BeforeWait, FMonotonicTimeSpan WaitTime);
-UE_API FWaitState WaitUntil(const void* Address, const TFunctionRef<bool()>& CanWait, const TFunctionRef<void()>& BeforeWait, FMonotonicTimePoint WaitTime);
-UE_API void WakeOne(const void* Address, const TFunctionRef<uint64(FWakeState)>& OnWakeState);
-}
+
+template <typename>
+class TLambdaAsFunctionWithContext;
+
+template <typename ReturnType, typename... ArgTypes>
+class TLambdaAsFunctionWithContext<ReturnType (ArgTypes...)>
+{
+public:
+	using FunctionType = ReturnType (void*, ArgTypes...);
+
+	inline constexpr TLambdaAsFunctionWithContext(decltype(nullptr))
+	{
+	}
+
+	template <typename LambdaType>
+	inline TLambdaAsFunctionWithContext(LambdaType&& Lambda UE_LIFETIMEBOUND)
+		requires (!std::is_same_v<std::decay_t<LambdaType>, TLambdaAsFunctionWithContext>) &&
+	std::is_invocable_r_v<ReturnType, std::decay_t<LambdaType>, ArgTypes...>
+	{
+		Function = &Call<std::decay_t<LambdaType>>;
+		Context = &Lambda;
+	}
+
+	inline explicit operator bool() const
+	{
+		return !!Function;
+	}
+
+	inline ReturnType operator()(ArgTypes... Args) const
+	{
+		return Function(Context, Forward<ArgTypes>(Args)...);
+	}
+
+	inline FunctionType* GetFunction() const
+	{
+		return Function;
+	}
+
+	inline void* GetContext() const
+	{
+		return Context;
+	}
+
+private:
+	template <typename LambdaType>
+	static ReturnType Call(void* Lambda, ArgTypes... Args)
+	{
+		return Invoke(*(LambdaType*)Lambda, Forward<ArgTypes>(Args)...);
+	}
+
+	FunctionType* Function = nullptr;
+	void* Context = nullptr;
+};
+
+UE_API FWaitState Wait(const void* Address, bool (*CanWait)(void*), void* CanWaitContext, void (*BeforeWait)(void*), void* BeforeWaitContext);
+UE_API FWaitState WaitFor(const void* Address, bool (*CanWait)(void*), void* CanWaitContext, void (*BeforeWait)(void*), void* BeforeWaitContext, FMonotonicTimeSpan WaitTime);
+UE_API FWaitState WaitUntil(const void* Address, bool (*CanWait)(void*), void* CanWaitContext, void (*BeforeWait)(void*), void* BeforeWaitContext, FMonotonicTimePoint WaitTime);
+UE_API void WakeOne(const void* Address, uint64 (*OnWakeState)(void*, FWakeState), void* OnWakeStateContext);
+
+} // Private
 
 /**
  * Queue the calling thread to wait if CanWait returns true. BeforeWait is only called if CanWait returns true.
@@ -47,10 +104,9 @@ UE_API void WakeOne(const void* Address, const TFunctionRef<uint64(FWakeState)>&
  * @param CanWait      Function called while the queue is locked. A return of false cancels the wait.
  * @param BeforeWait   Function called after the queue is unlocked and before the thread waits.
  */
-template <typename CanWaitType, typename BeforeWaitType>
-inline FWaitState Wait(const void* Address, CanWaitType&& CanWait, BeforeWaitType&& BeforeWait)
+inline FWaitState Wait(const void* Address, Private::TLambdaAsFunctionWithContext<bool()> CanWait, Private::TLambdaAsFunctionWithContext<void()> BeforeWait)
 {
-	return Private::Wait(Address, Forward<CanWaitType>(CanWait), Forward<BeforeWaitType>(BeforeWait));
+	return Private::Wait(Address, CanWait.GetFunction(), CanWait.GetContext(), BeforeWait.GetFunction(), BeforeWait.GetContext());
 }
 
 /**
@@ -61,10 +117,9 @@ inline FWaitState Wait(const void* Address, CanWaitType&& CanWait, BeforeWaitTyp
  * @param BeforeWait   Function called after the queue is unlocked and before the thread waits.
  * @param WaitTime     Relative time after which waiting is automatically canceled and the thread wakes.
  */
-template <typename CanWaitType, typename BeforeWaitType>
-inline FWaitState WaitFor(const void* Address, CanWaitType&& CanWait, BeforeWaitType&& BeforeWait, FMonotonicTimeSpan WaitTime)
+inline FWaitState WaitFor(const void* Address, Private::TLambdaAsFunctionWithContext<bool()> CanWait, Private::TLambdaAsFunctionWithContext<void()> BeforeWait, FMonotonicTimeSpan WaitTime)
 {
-	return Private::WaitFor(Address, Forward<CanWaitType>(CanWait), Forward<BeforeWaitType>(BeforeWait), WaitTime);
+	return Private::WaitFor(Address, CanWait.GetFunction(), CanWait.GetContext(), BeforeWait.GetFunction(), BeforeWait.GetContext(), WaitTime);
 }
 
 /**
@@ -75,10 +130,9 @@ inline FWaitState WaitFor(const void* Address, CanWaitType&& CanWait, BeforeWait
  * @param BeforeWait   Function called after the queue is unlocked and before the thread waits.
  * @param WaitTime     Absolute time after which waiting is automatically canceled and the thread wakes.
  */
-template <typename CanWaitType, typename BeforeWaitType>
-inline FWaitState WaitUntil(const void* Address, CanWaitType&& CanWait, BeforeWaitType&& BeforeWait, FMonotonicTimePoint WaitTime)
+inline FWaitState WaitUntil(const void* Address, Private::TLambdaAsFunctionWithContext<bool()> CanWait, Private::TLambdaAsFunctionWithContext<void()> BeforeWait, FMonotonicTimePoint WaitTime)
 {
-	return Private::WaitUntil(Address, Forward<CanWaitType>(CanWait), Forward<BeforeWaitType>(BeforeWait), WaitTime);
+	return Private::WaitUntil(Address, CanWait.GetFunction(), CanWait.GetContext(), BeforeWait.GetFunction(), BeforeWait.GetContext(), WaitTime);
 }
 
 /**
@@ -87,10 +141,9 @@ inline FWaitState WaitUntil(const void* Address, CanWaitType&& CanWait, BeforeWa
  * @param Address       Address to use as the key for the queue. Must match the address used in Wait.
  * @param OnWakeState   Function called while the queue is locked. Receives the wake state. Returns WakeToken.
  */
-template <typename OnWakeStateType>
-inline void WakeOne(const void* Address, OnWakeStateType&& OnWakeState)
+inline void WakeOne(const void* Address, Private::TLambdaAsFunctionWithContext<uint64(FWakeState)> OnWakeState)
 {
-	return Private::WakeOne(Address, Forward<OnWakeStateType>(OnWakeState));
+	return Private::WakeOne(Address, OnWakeState.GetFunction(), OnWakeState.GetContext());
 }
 
 /**
@@ -121,3 +174,7 @@ UE_API uint32 WakeAll(const void* Address);
 } // UE::ParkingLot
 
 #undef UE_API
+
+#if UE_ENABLE_INCLUDE_ORDER_DEPRECATED_IN_5_7
+#include "Templates/Function.h"
+#endif
