ParkingLot: Replaced empty callbacks with nullptr

diff --git a/Engine/Source/Runtime/Core/Private/Async/RecursiveMutex.cpp b/Engine/Source/Runtime/Core/Private/Async/RecursiveMutex.cpp
--- a/Engine/Source/Runtime/Core/Private/Async/RecursiveMutex.cpp
+++ b/Engine/Source/Runtime/Core/Private/Async/RecursiveMutex.cpp
@@ -107,7 +107,7 @@ FORCENOINLINE void FRecursiveMutex::LockSlow(uint32 CurrentState, const uint32 C
 		LowLevelTasks::Private::FOversubscriptionAllowedScope _(false);
 
 		// Wait if the state has not changed. Either way, loop back and try to acquire the lock after trying to wait.
-		ParkingLot::Wait(&State, [this, CurrentState] { return State.load(std::memory_order_relaxed) == CurrentState; }, []{});
+		ParkingLot::Wait(&State, [this, CurrentState] { return State.load(std::memory_order_relaxed) == CurrentState; }, nullptr);
 		CurrentState = State.load(std::memory_order_relaxed);
 	}
 }
diff --git a/Engine/Source/Runtime/Core/Private/Async/SharedMutex.cpp b/Engine/Source/Runtime/Core/Private/Async/SharedMutex.cpp
--- a/Engine/Source/Runtime/Core/Private/Async/SharedMutex.cpp
+++ b/Engine/Source/Runtime/Core/Private/Async/SharedMutex.cpp
@@ -77,7 +77,7 @@ void FSharedMutex::LockSharedSlow()
 		ParkingLot::Wait(GetSharedLockAddress(), [this, CurrentState]
 		{
 			return State.load(std::memory_order_relaxed) == CurrentState;
-		}, []{});
+		}, nullptr);
 		CurrentState = State.load(std::memory_order_relaxed);
 	}
 }
diff --git a/Engine/Source/Runtime/Core/Private/Async/SharedRecursiveMutex.cpp b/Engine/Source/Runtime/Core/Private/Async/SharedRecursiveMutex.cpp
--- a/Engine/Source/Runtime/Core/Private/Async/SharedRecursiveMutex.cpp
+++ b/Engine/Source/Runtime/Core/Private/Async/SharedRecursiveMutex.cpp
@@ -179,7 +179,7 @@ FORCENOINLINE void FSharedRecursiveMutex::LockSlow(uint32 CurrentState, const ui
 		ParkingLot::Wait(GetLockAddress(), [this, CurrentState]
 		{
 			return State.load(std::memory_order_relaxed) == CurrentState;
-		}, []{});
+		}, nullptr);
 		CurrentState = State.load(std::memory_order_relaxed);
 	}
 }
@@ -307,7 +307,7 @@ FORCENOINLINE void FSharedRecursiveMutex::LockSharedSlow(Core::Private::FSharedR
 		ParkingLot::Wait(GetSharedLockAddress(), [this, CurrentState]
 		{
 			return State.load(std::memory_order_relaxed) == CurrentState;
-		}, []{});
+		}, nullptr);
 		CurrentState = State.load(std::memory_order_relaxed);
 	}
 }
diff --git a/Engine/Source/Runtime/Core/Public/Async/EventCount.h b/Engine/Source/Runtime/Core/Public/Async/EventCount.h
--- a/Engine/Source/Runtime/Core/Public/Async/EventCount.h
+++ b/Engine/Source/Runtime/Core/Public/Async/EventCount.h
@@ -100,7 +100,7 @@ public:
 	{
 		if ((Count.load(std::memory_order_acquire) & ~CounterType(1)) == Compare.Value)
 		{
-			ParkingLot::Wait(&Count, [this, Compare] { return (Count.load(std::memory_order_acquire) & ~CounterType(1)) == Compare.Value; }, []{});
+			ParkingLot::Wait(&Count, [this, Compare] { return (Count.load(std::memory_order_acquire) & ~CounterType(1)) == Compare.Value; }, nullptr);
 		}
 	}
 
@@ -115,7 +115,7 @@ public:
 	{
 		if ((Count.load(std::memory_order_acquire) & ~CounterType(1)) == Compare.Value)
 		{
-			const ParkingLot::FWaitState WaitState = ParkingLot::WaitFor(&Count, [this, Compare] { return (Count.load(std::memory_order_acquire) & ~CounterType(1)) == Compare.Value; }, [] {}, WaitTime);
+			const ParkingLot::FWaitState WaitState = ParkingLot::WaitFor(&Count, [this, Compare] { return (Count.load(std::memory_order_acquire) & ~CounterType(1)) == Compare.Value; }, nullptr, WaitTime);
 
 			// Make sure we return true if we did wait but also if the wait was skipped because the value actually changed before we had the opportunity to wait.
 			return WaitState.bDidWake || !WaitState.bDidWait;
@@ -134,7 +134,7 @@ public:
 	{
 		if ((Count.load(std::memory_order_acquire) & ~CounterType(1)) == Compare.Value)
 		{
-			const ParkingLot::FWaitState WaitState = ParkingLot::WaitUntil(&Count, [this, Compare] { return (Count.load(std::memory_order_acquire) & ~CounterType(1)) == Compare.Value; }, [] {}, WaitTime);
+			const ParkingLot::FWaitState WaitState = ParkingLot::WaitUntil(&Count, [this, Compare] { return (Count.load(std::memory_order_acquire) & ~CounterType(1)) == Compare.Value; }, nullptr, WaitTime);
 
 			// Make sure we return true if we did wait but also if the wait was skipped because the value actually changed before we had the opportunity to wait.
 			return WaitState.bDidWake || !WaitState.bDidWait;
diff --git a/Engine/Source/Runtime/Core/Public/Async/IntrusiveMutex.h b/Engine/Source/Runtime/Core/Public/Async/IntrusiveMutex.h
--- a/Engine/Source/Runtime/Core/Public/Async/IntrusiveMutex.h
+++ b/Engine/Source/Runtime/Core/Public/Async/IntrusiveMutex.h
@@ -176,7 +176,7 @@ public:
 			{
 				const StateType NewState = State.load(std::memory_order_relaxed);
 				return (NewState & IsLockedMask) && (NewState & MayHaveWaitingLockFlag);
-			}, []{});
+			}, nullptr);
 			CurrentState = State.load(std::memory_order_relaxed);
 		}
 	}
diff --git a/Engine/Source/Runtime/Core/Public/Async/ManualResetEvent.h b/Engine/Source/Runtime/Core/Public/Async/ManualResetEvent.h
--- a/Engine/Source/Runtime/Core/Public/Async/ManualResetEvent.h
+++ b/Engine/Source/Runtime/Core/Public/Async/ManualResetEvent.h
@@ -72,7 +72,7 @@ public:
 				{
 					uint8 NewState = State.load(std::memory_order_relaxed);
 					return !(NewState & IsNotifiedFlag) && (NewState & MaybeWaitingFlag);
-				}, []{}, WaitTime);
+				}, nullptr, WaitTime);
 			}
 		}
 	}
