MemoryWriter: Fixed IndexSize to be a bit count

diff --git a/Engine/Source/Runtime/Core/Public/Serialization/MemoryWriter.h b/Engine/Source/Runtime/Core/Public/Serialization/MemoryWriter.h
--- a/Engine/Source/Runtime/Core/Public/Serialization/MemoryWriter.h
+++ b/Engine/Source/Runtime/Core/Public/Serialization/MemoryWriter.h
@@ -20,7 +20,8 @@ template <typename ArrayAllocatorType>
 class TMemoryWriterBase : public FMemoryArchive
 {
 	using IndexSizeType = typename ArrayAllocatorType::SizeType;
-	static constexpr int32 IndexSize = sizeof(IndexSizeType);
+	static constexpr int32 IndexSize = sizeof(IndexSizeType) * 8;
+
 public:
 	TMemoryWriterBase( TArray<uint8, ArrayAllocatorType>& InBytes, bool bIsPersistent = false, bool bSetOffset = false, const FName InArchiveName = NAME_None )
 	: FMemoryArchive()
