Add support for GetLastError for Linux/Apple platforms

Should return errno at least for these vs 0

diff --git a/Engine/Source/Runtime/Core/Private/Apple/ApplePlatformMisc.cpp b/Engine/Source/Runtime/Core/Private/Apple/ApplePlatformMisc.cpp
--- a/Engine/Source/Runtime/Core/Private/Apple/ApplePlatformMisc.cpp
+++ b/Engine/Source/Runtime/Core/Private/Apple/ApplePlatformMisc.cpp
@@ -99,6 +99,11 @@ const TCHAR* FApplePlatformMisc::GetSystemErrorMessage(TCHAR* OutBuffer, int32 B
 	return OutBuffer;
 }
 
+uint32 FApplePlatformMisc::GetLastError()
+{
+	return (uint32)errno;
+}
+
 FString FApplePlatformMisc::GetDefaultLocale()
 {
 	CFLocaleRef Locale = CFLocaleCopyCurrent();
diff --git a/Engine/Source/Runtime/Core/Private/Unix/UnixPlatformMisc.cpp b/Engine/Source/Runtime/Core/Private/Unix/UnixPlatformMisc.cpp
--- a/Engine/Source/Runtime/Core/Private/Unix/UnixPlatformMisc.cpp
+++ b/Engine/Source/Runtime/Core/Private/Unix/UnixPlatformMisc.cpp
@@ -606,6 +606,11 @@ const TCHAR* FUnixPlatformMisc::GetSystemErrorMessage(TCHAR* OutBuffer, int32 Bu
 	return OutBuffer;
 }
 
+uint32 FUnixPlatformMisc::GetLastError()
+{
+	return (uint32)errno;
+}
+
 CORE_API TFunction<EAppReturnType::Type(EAppMsgType::Type MsgType, const TCHAR* Text, const TCHAR* Caption)> MessageBoxExtCallback;
 
 EAppReturnType::Type FUnixPlatformMisc::MessageBoxExt(EAppMsgType::Type MsgType, const TCHAR* Text, const TCHAR* Caption)
diff --git a/Engine/Source/Runtime/Core/Public/Apple/ApplePlatformMisc.h b/Engine/Source/Runtime/Core/Public/Apple/ApplePlatformMisc.h
--- a/Engine/Source/Runtime/Core/Public/Apple/ApplePlatformMisc.h
+++ b/Engine/Source/Runtime/Core/Public/Apple/ApplePlatformMisc.h
@@ -72,6 +72,7 @@ struct CORE_API FApplePlatformMisc : public FGenericPlatformMisc
 
 	static void LocalPrint(const TCHAR* Message);
 	static const TCHAR* GetSystemErrorMessage(TCHAR* OutBuffer, int32 BufferCount, int32 Error);
+	static uint32 GetLastError();
 	static int32 NumberOfCores();
 	
 	/** @return Whether filehandles can be opened on one thread and read/written on another thread */
diff --git a/Engine/Source/Runtime/Core/Public/Unix/UnixPlatformMisc.h b/Engine/Source/Runtime/Core/Public/Unix/UnixPlatformMisc.h
--- a/Engine/Source/Runtime/Core/Public/Unix/UnixPlatformMisc.h
+++ b/Engine/Source/Runtime/Core/Public/Unix/UnixPlatformMisc.h
@@ -45,6 +45,7 @@ struct FUnixPlatformMisc : public FGenericPlatformMisc
 	static CORE_API void RequestExit(bool Force, const TCHAR* CallSite = nullptr);
 	static CORE_API void RequestExitWithStatus(bool Force, uint8 ReturnCode, const TCHAR* CallSite = nullptr);
 	static CORE_API const TCHAR* GetSystemErrorMessage(TCHAR* OutBuffer, int32 BufferCount, int32 Error);
+	static CORE_API uint32 GetLastError();
 
 	static CORE_API void NormalizePath(FString& InPath);
 	static CORE_API void NormalizePath(FStringBuilderBase& InPath);
