Fixed Android `RequestEngineExit` issue related to how `UAC` was queuing background work and locking forever the ability to shut down the engine.

#rnx

diff --git a/Engine/Source/Runtime/Core/Private/Android/AndroidPlatformMisc.cpp b/Engine/Source/Runtime/Core/Private/Android/AndroidPlatformMisc.cpp
--- a/Engine/Source/Runtime/Core/Private/Android/AndroidPlatformMisc.cpp
+++ b/Engine/Source/Runtime/Core/Private/Android/AndroidPlatformMisc.cpp
@@ -287,7 +287,8 @@ void FAndroidMisc::RequestExit( bool Force, const TCHAR* CallSite)
 #if USE_ANDROID_JNI
 		AndroidThunkCpp_ForceQuit();
 #else
-		exit(1);
+		// for Android we should not exit with 1 when forcing exit
+		exit(0);
 #endif
 	}
 	else
diff --git a/Engine/Source/Runtime/Core/Public/Android/AndroidPlatformProperties.h b/Engine/Source/Runtime/Core/Public/Android/AndroidPlatformProperties.h
--- a/Engine/Source/Runtime/Core/Public/Android/AndroidPlatformProperties.h
+++ b/Engine/Source/Runtime/Core/Public/Android/AndroidPlatformProperties.h
@@ -87,8 +87,7 @@ struct FAndroidPlatformProperties
 
 	static constexpr FORCEINLINE bool SupportsQuit()
 	{
-		// android does not support quitting to the OS like PC does
-		return false;
+		return true;
 	}
 
 	static constexpr FORCEINLINE bool HasFixedResolution()
