Fixed iOS `RequestEngineExit`.
iOS will now gracefully request engine exit on `BootTest` and other automation jobs.

#rnx

diff --git a/Engine/Source/Runtime/ApplicationCore/Private/IOS/IOSPlatformFramePacer.cpp b/Engine/Source/Runtime/ApplicationCore/Private/IOS/IOSPlatformFramePacer.cpp
--- a/Engine/Source/Runtime/ApplicationCore/Private/IOS/IOSPlatformFramePacer.cpp
+++ b/Engine/Source/Runtime/ApplicationCore/Private/IOS/IOSPlatformFramePacer.cpp
@@ -62,6 +62,13 @@ namespace IOSDisplayConstants
 	// during shutdown, this can cause crashes (only non-backgrounding apps do this)
 	if (IsEngineExitRequested())
 	{
+		for( auto& NextEvent : ListeningEvents )
+		{
+			NextEvent->Trigger();
+		}
+		
+		ListeningEvents.Empty();
+
 		return;
 	};
 
diff --git a/Engine/Source/Runtime/Core/Private/Apple/ApplePlatformMisc.cpp b/Engine/Source/Runtime/Core/Private/Apple/ApplePlatformMisc.cpp
--- a/Engine/Source/Runtime/Core/Private/Apple/ApplePlatformMisc.cpp
+++ b/Engine/Source/Runtime/Core/Private/Apple/ApplePlatformMisc.cpp
@@ -361,22 +361,6 @@ bool FApplePlatformMisc::IsOSAtLeastVersion(const uint32 MacOSVersion[3], const
 	return true;
 }
 
-void FApplePlatformMisc::RequestExit( bool Force, const TCHAR* CallSite)
-{
-	UE_LOG(LogApplePlatformMisc, Log,  TEXT("FApplePlatformMisc::RequestExit(%i, %s)"), Force,
-		CallSite ? CallSite : TEXT("<NoCallSiteInfo>"));
-	if( Force )
-	{
-		// effectivelly exits the application
-		_Exit(0);
-	}
-	else
-	{
-		// Tell the platform specific code we want to exit cleanly from the main loop.
-		RequestEngineExit(TEXT("ApplePlatform RequestExit"));
-	}
-}
-
 #if APPLE_PROFILING_ENABLED
 
 void FApplePlatformMisc::BeginNamedEvent(const struct FColor& Color, const TCHAR* Text)
diff --git a/Engine/Source/Runtime/Core/Public/Apple/ApplePlatformMisc.h b/Engine/Source/Runtime/Core/Public/Apple/ApplePlatformMisc.h
--- a/Engine/Source/Runtime/Core/Public/Apple/ApplePlatformMisc.h
+++ b/Engine/Source/Runtime/Core/Public/Apple/ApplePlatformMisc.h
@@ -91,8 +91,6 @@ struct CORE_API FApplePlatformMisc : public FGenericPlatformMisc
 
 	static bool IsOSAtLeastVersion(const uint32 MacOSVersion[3], const uint32 IOSVersion[3], const uint32 TVOSVersion[3]);
 
-	static void RequestExit( bool Force, const TCHAR* CallSite = nullptr);
-
 #if APPLE_PROFILING_ENABLED
 	/** Override to add Apple features when AppleProfiling is on, otherwise use version on FGenericPlatformMisc. */
 	static void BeginNamedEvent(const struct FColor& Color, const TCHAR* Text);
diff --git a/Engine/Source/Runtime/Core/Public/IOS/IOSPlatformProperties.h b/Engine/Source/Runtime/Core/Public/IOS/IOSPlatformProperties.h
--- a/Engine/Source/Runtime/Core/Public/IOS/IOSPlatformProperties.h
+++ b/Engine/Source/Runtime/Core/Public/IOS/IOSPlatformProperties.h
@@ -112,6 +112,11 @@ struct FIOSPlatformProperties
 	{
 		return true;
 	}
+
+	static constexpr FORCEINLINE bool SupportsQuit()
+	{
+		return true;
+	}
 };
 
 struct FTVOSPlatformProperties : public FIOSPlatformProperties
