Automate Revert All Files option works correctly test case

#rnx

diff --git a/Engine/Source/Runtime/Core/Public/Misc/AutomationTest.h b/Engine/Source/Runtime/Core/Public/Misc/AutomationTest.h
--- a/Engine/Source/Runtime/Core/Public/Misc/AutomationTest.h
+++ b/Engine/Source/Runtime/Core/Public/Misc/AutomationTest.h
@@ -3865,6 +3865,26 @@ class CommandName : public IAutomationLatentCommand \
 	ParamType1 ParamName1; \
 }
 
+#define DEFINE_EXPORTED_LATENT_AUTOMATION_COMMAND_FOUR_PARAMETER(EXPORT_API, CommandName,ParamType0,ParamName0,ParamType1,ParamName1,ParamType2,ParamName2,ParamType3,ParamName3)	\
+class CommandName : public IAutomationLatentCommand \
+	{ \
+	public: \
+	CommandName(ParamType0 InputParam0, ParamType1 InputParam1, ParamType2 InputParam2, ParamType3 InputParam3) \
+	: ParamName0(InputParam0) \
+	, ParamName1(InputParam1) \
+	, ParamName2(InputParam2) \
+	, ParamName3(InputParam3) \
+		{} \
+		virtual ~CommandName() \
+		{} \
+		EXPORT_API virtual bool Update() override; \
+	private: \
+	ParamType0 ParamName0; \
+	ParamType1 ParamName1; \
+	ParamType2 ParamName2; \
+	ParamType3 ParamName3; \
+}
+
 #define DEFINE_ENGINE_LATENT_AUTOMATION_COMMAND(CommandName)	\
 	DEFINE_EXPORTED_LATENT_AUTOMATION_COMMAND(ENGINE_API, CommandName)
 
