Moved CIntegral into the UE:: namespace to match all new C++20 concepts.

diff --git a/Engine/Source/Runtime/Core/Public/Async/IntrusiveMutex.h b/Engine/Source/Runtime/Core/Public/Async/IntrusiveMutex.h
--- a/Engine/Source/Runtime/Core/Public/Async/IntrusiveMutex.h
+++ b/Engine/Source/Runtime/Core/Public/Async/IntrusiveMutex.h
@@ -21,7 +21,7 @@ using TIntrusiveMutexStateType_T = std::decay_t<decltype(ParamsType::IsLockedFla
 template <typename ParamsType, typename StateType = TIntrusiveMutexStateType_T<ParamsType>>
 concept CIntrusiveMutexParams = requires
 {
-	requires CIntegral<StateType>;
+	requires UE::CIntegral<StateType>;
 
 	// Required: constexpr static StateType IsLockedFlag = ...;
 	// Flag that is set in the state when the mutex is locked.
diff --git a/Engine/Source/Runtime/Core/Public/Clang/ClangPlatformMath.h b/Engine/Source/Runtime/Core/Public/Clang/ClangPlatformMath.h
--- a/Engine/Source/Runtime/Core/Public/Clang/ClangPlatformMath.h
+++ b/Engine/Source/Runtime/Core/Public/Clang/ClangPlatformMath.h
@@ -129,7 +129,7 @@ struct FClangPlatformMath : public FGenericPlatformMath
 	 * If there was overflow, it returns false, and OutResult may or may not be written.
 	 * If there wasn't overflow, it returns true, and the result of the addition is written to OutResult.
 	 */
-	template <CIntegral IntType>
+	template <UE::CIntegral IntType>
 	static FORCEINLINE bool AddAndCheckForOverflow(IntType A, IntType B, IntType& OutResult)
 	{
 		return !__builtin_add_overflow(A, B, &OutResult);
@@ -140,7 +140,7 @@ struct FClangPlatformMath : public FGenericPlatformMath
 	 * If there was overflow, it returns false, and OutResult may or may not be written.
 	 * If there wasn't overflow, it returns true, and the result of the subtraction is written to OutResult.
 	 */
-	template <CIntegral IntType>
+	template <UE::CIntegral IntType>
 	static FORCEINLINE bool SubtractAndCheckForOverflow(IntType A, IntType B, IntType& OutResult)
 	{
 		return !__builtin_sub_overflow(A, B, &OutResult);
@@ -151,7 +151,7 @@ struct FClangPlatformMath : public FGenericPlatformMath
 	 * If there was overflow, it returns false, and OutResult may or may not be written.
 	 * If there wasn't overflow, it returns true, and the result of the multiplication is written to OutResult.
 	 */
-	template <CIntegral IntType>
+	template <UE::CIntegral IntType>
 	static FORCEINLINE bool MultiplyAndCheckForOverflow(IntType A, IntType B, IntType& OutResult)
 	{
 		return !__builtin_mul_overflow(A, B, &OutResult);
diff --git a/Engine/Source/Runtime/Core/Public/Concepts/Integral.h b/Engine/Source/Runtime/Core/Public/Concepts/Integral.h
--- a/Engine/Source/Runtime/Core/Public/Concepts/Integral.h
+++ b/Engine/Source/Runtime/Core/Public/Concepts/Integral.h
@@ -4,8 +4,11 @@
 
 #include <type_traits>
 
-/**
- * Concept which describes an integral type.  We use this instead of std::integral because <concepts> isn't a well supported header yet.
- */
-template <typename T>
-concept CIntegral = std::is_integral_v<T>;
+namespace UE
+{
+	/**
+	 * Concept which describes an integral type.  We use this instead of std::integral because <concepts> isn't a well supported header yet.
+	 */
+	template <typename T>
+	concept CIntegral = std::is_integral_v<T>;
+}
diff --git a/Engine/Source/Runtime/Core/Public/Containers/Array.h b/Engine/Source/Runtime/Core/Public/Containers/Array.h
--- a/Engine/Source/Runtime/Core/Public/Containers/Array.h
+++ b/Engine/Source/Runtime/Core/Public/Containers/Array.h
@@ -2111,7 +2111,7 @@ public:
 	 * @param Count (Optional) Number of elements to remove. Default is 1.
 	 * @param AllowShrinking (Optional) Tells if this call can shrink array if suitable after remove. Default is yes.
 	 */
-	template <CIntegral CountType>
+	template <UE::CIntegral CountType>
 	FORCEINLINE void RemoveAt(SizeType Index, CountType Count, EAllowShrinking AllowShrinking = EAllowShrinking::Default)
 	{
 		static_assert(!std::is_same_v<CountType, bool>, "TArray::RemoveAt: unexpected bool passed as the Count argument");
@@ -2217,7 +2217,7 @@ public:
 	 * @param AllowShrinking (Optional) Tells if this call can shrink array if
 	 *                        suitable after remove. Default is yes.
 	 */
-	template <CIntegral CountType>
+	template <UE::CIntegral CountType>
 	FORCEINLINE void RemoveAtSwap(SizeType Index, CountType Count, EAllowShrinking AllowShrinking = EAllowShrinking::Default)
 	{
 		static_assert(!std::is_same_v<CountType, bool>, "TArray::RemoveAtSwap: unexpected bool passed as the Count argument");
diff --git a/Engine/Source/Runtime/Core/Public/GenericPlatform/GenericPlatformMath.h b/Engine/Source/Runtime/Core/Public/GenericPlatform/GenericPlatformMath.h
--- a/Engine/Source/Runtime/Core/Public/GenericPlatform/GenericPlatformMath.h
+++ b/Engine/Source/Runtime/Core/Public/GenericPlatform/GenericPlatformMath.h
@@ -1182,7 +1182,7 @@ struct FGenericPlatformMath
 	 * If there was overflow, it returns false, and OutResult may or may not be written.
 	 * If there wasn't overflow, it returns true, and the result of the addition is written to OutResult.
 	 */
-	template <CIntegral IntType>
+	template <UE::CIntegral IntType>
 	[[nodiscard]] static FORCEINLINE bool AddAndCheckForOverflow(IntType A, IntType B, IntType& OutResult)
 	{
 		// This follows Hacker's Delight, Chapter 2-12
@@ -1229,7 +1229,7 @@ struct FGenericPlatformMath
 	 * If there was overflow, it returns false, and OutResult may or may not be written.
 	 * If there wasn't overflow, it returns true, and the result of the subtraction is written to OutResult.
 	 */
-	template <CIntegral IntType>
+	template <UE::CIntegral IntType>
 	[[nodiscard]] static FORCEINLINE bool SubtractAndCheckForOverflow(IntType A, IntType B, IntType& OutResult)
 	{
 		// This follows Hacker's Delight, Chapter 2-12
@@ -1270,7 +1270,7 @@ struct FGenericPlatformMath
 	 * If there was overflow, it returns false, and OutResult may or may not be written.
 	 * If there wasn't overflow, it returns true, and the result of the multiplication is written to OutResult.
 	 */
-	template <CIntegral IntType>
+	template <UE::CIntegral IntType>
 	[[nodiscard]] static FORCEINLINE bool MultiplyAndCheckForOverflow(IntType A, IntType B, IntType& OutResult)
 	{
 		// Handle the case where the second factor is 0 specially (why will become clear in a minute).
diff --git a/Engine/Source/Runtime/Core/Public/Internationalization/FastDecimalFormat.h b/Engine/Source/Runtime/Core/Public/Internationalization/FastDecimalFormat.h
--- a/Engine/Source/Runtime/Core/Public/Internationalization/FastDecimalFormat.h
+++ b/Engine/Source/Runtime/Core/Public/Internationalization/FastDecimalFormat.h
@@ -78,7 +78,7 @@ struct FDecimalNumberIntegralLimits
 	uint64 NumericLimitMax;
 	bool bIsNumericSigned;
 
-	template <CIntegral IntegralType>
+	template <UE::CIntegral IntegralType>
 	static FDecimalNumberIntegralLimits FromNumericLimits()
 	{
 		return FDecimalNumberIntegralLimits(TNumericLimits<IntegralType>::Lowest(), TNumericLimits<IntegralType>::Max(), std::is_signed_v<IntegralType>);
diff --git a/Engine/Source/Runtime/Core/Public/Math/GuardedInt.h b/Engine/Source/Runtime/Core/Public/Math/GuardedInt.h
--- a/Engine/Source/Runtime/Core/Public/Math/GuardedInt.h
+++ b/Engine/Source/Runtime/Core/Public/Math/GuardedInt.h
@@ -51,7 +51,7 @@
  * As-is, the current approach is not the fastest, but it's not presently intended to be used in contexts
  * where speed of arithmetic operations is a major concern.
  */
-template<CIntegral IntType>
+template<UE::CIntegral IntType>
 class TGuardedInt
 {
 private:
@@ -70,7 +70,7 @@ public:
 	[[nodiscard]] TGuardedInt() = default;
 
 	/** Construct a TGuardedInt from a regular signed integer value. If it's out of range, it results in an invalid value. */
-	template<CIntegral ValueType>
+	template<UE::CIntegral ValueType>
 	[[nodiscard]] explicit TGuardedInt(ValueType InValue)
 		: Value((IntType)InValue), bIsValid(false)
 	{
diff --git a/Engine/Source/Runtime/Core/Public/Math/UnrealMathUtility.h b/Engine/Source/Runtime/Core/Public/Math/UnrealMathUtility.h
--- a/Engine/Source/Runtime/Core/Public/Math/UnrealMathUtility.h
+++ b/Engine/Source/Runtime/Core/Public/Math/UnrealMathUtility.h
@@ -535,7 +535,7 @@ public:
 	 * Converts an integral type to a nearest less or equal integer.
 	 * Unlike std::floor, it returns an IntegralType.
 	 */
-	template <CIntegral IntegralType>
+	template <UE::CIntegral IntegralType>
 	[[nodiscard]] static constexpr FORCEINLINE IntegralType Floor(IntegralType I)
 	{
 		return I;
@@ -657,7 +657,7 @@ public:
 	}
 
 	/** Wraps X to be between Min and Max, exclusive. Will never return Max. */
-	template <CIntegral T>
+	template <UE::CIntegral T>
 	[[nodiscard]] static constexpr FORCEINLINE T WrapExclusive(const T X, const T Min, const T Max)
 	{
 		// Use unsigned type allow for large ranges which don't overflow on subtraction.
diff --git a/Engine/Source/Runtime/SlateCore/Public/Containers/ObservableArray.h b/Engine/Source/Runtime/SlateCore/Public/Containers/ObservableArray.h
--- a/Engine/Source/Runtime/SlateCore/Public/Containers/ObservableArray.h
+++ b/Engine/Source/Runtime/SlateCore/Public/Containers/ObservableArray.h
@@ -341,7 +341,7 @@ public:
 		Array.RemoveAt(Index, AllowShrinking);
 		ArrayChangedDelegate.Broadcast(ObservableArrayChangedArgsType::MakeRemoveAction({ &RemovedElement, 1 }, Index));
 	}
-	template <CIntegral CountType>
+	template <UE::CIntegral CountType>
 	void RemoveAt(SizeType Index, CountType NumToRemove, EAllowShrinking AllowShrinking = EAllowShrinking::Default)
 	{
 		static_assert(!std::is_same_v<CountType, bool>, "TObservableArray::RemoveAt: unexpected bool passed as the Count argument");
@@ -386,7 +386,7 @@ public:
 		Array.RemoveAtSwap(Index, AllowShrinking);
 		ArrayChangedDelegate.Broadcast(ObservableArrayChangedArgsType::MakeRemoveSwapAction({ &RemovedElement, 1 }, Index, PreviousNum - SwapAmount));
 	}
-	template <CIntegral CountType>
+	template <UE::CIntegral CountType>
 	void RemoveAtSwap(SizeType Index, CountType NumToRemove, EAllowShrinking AllowShrinking = EAllowShrinking::Default)
 	{
 		static_assert(!std::is_same_v<CountType, bool>, "TObservableArray::RemoveAtSwap: unexpected bool passed as the Count argument");
