Removed the dependency on UnrealTemplate.h from Invoke.h

diff --git a/Engine/Source/Runtime/Core/Public/Algo/Accumulate.h b/Engine/Source/Runtime/Core/Public/Algo/Accumulate.h
--- a/Engine/Source/Runtime/Core/Public/Algo/Accumulate.h
+++ b/Engine/Source/Runtime/Core/Public/Algo/Accumulate.h
@@ -4,6 +4,7 @@
 
 #include "CoreTypes.h"
 #include "Templates/Invoke.h"
+#include "Templates/UnrealTemplate.h"
 
 // TPlus<T> specifically takes const T& and returns T.
 // TPlus<> (empty angle brackets) is late-binding, taking whatever is passed and returning the correct result type for (A+B)
diff --git a/Engine/Source/Runtime/Core/Public/Templates/Invoke.h b/Engine/Source/Runtime/Core/Public/Templates/Invoke.h
--- a/Engine/Source/Runtime/Core/Public/Templates/Invoke.h
+++ b/Engine/Source/Runtime/Core/Public/Templates/Invoke.h
@@ -4,7 +4,6 @@
 
 #include "CoreTypes.h"
 #include "Traits/MemberFunctionPtrOuter.h"
-#include "Templates/UnrealTemplate.h" // For Forward<>()
 #include <type_traits>
 
 
@@ -42,16 +41,16 @@ namespace UE::Core::Private
  */
 template <typename FuncType, typename... ArgTypes>
 FORCEINLINE auto Invoke(FuncType&& Func, ArgTypes&&... Args)
-	-> decltype(Forward<FuncType>(Func)(Forward<ArgTypes>(Args)...))
+	-> decltype(((FuncType&&)Func)((ArgTypes&&)Args...))
 {
-	return Forward<FuncType>(Func)(Forward<ArgTypes>(Args)...);
+	return ((FuncType&&)Func)((ArgTypes&&)Args...);
 }
 
 template <typename ReturnType, typename ObjType, typename TargetType>
 FORCEINLINE auto Invoke(ReturnType ObjType::*pdm, TargetType&& Target)
-	-> decltype(UE::Core::Private::DereferenceIfNecessary<ObjType>(Forward<TargetType>(Target), &Target).*pdm)
+	-> decltype(UE::Core::Private::DereferenceIfNecessary<ObjType>((TargetType&&)Target, &Target).*pdm)
 {
-	return UE::Core::Private::DereferenceIfNecessary<ObjType>(Forward<TargetType>(Target), &Target).*pdm;
+	return UE::Core::Private::DereferenceIfNecessary<ObjType>((TargetType&&)Target, &Target).*pdm;
 }
 
 template <
@@ -61,9 +60,9 @@ template <
 	typename    ObjType = TMemberFunctionPtrOuter_T<PtrMemFunType>
 >
 FORCEINLINE auto Invoke(PtrMemFunType PtrMemFun, TargetType&& Target, ArgTypes&&... Args)
-	-> decltype((UE::Core::Private::DereferenceIfNecessary<ObjType>(Forward<TargetType>(Target), &Target).*PtrMemFun)(Forward<ArgTypes>(Args)...))
+	-> decltype((UE::Core::Private::DereferenceIfNecessary<ObjType>((TargetType&&)Target, &Target).*PtrMemFun)((ArgTypes&&)Args...))
 {
-	return (UE::Core::Private::DereferenceIfNecessary<ObjType>(Forward<TargetType>(Target), &Target).*PtrMemFun)(Forward<ArgTypes>(Args)...);
+	return (UE::Core::Private::DereferenceIfNecessary<ObjType>((TargetType&&)Target, &Target).*PtrMemFun)((ArgTypes&&)Args...);
 }
 
 
@@ -84,7 +83,7 @@ FORCEINLINE auto Invoke(PtrMemFunType PtrMemFun, TargetType&& Target, ArgTypes&&
 #define UE_PROJECTION(FuncName) \
 	[](auto&&... Args) -> decltype(auto) \
 	{ \
-		return FuncName(Forward<decltype(Args)>(Args)...); \
+		return FuncName((decltype(Args)&&)Args...); \
 	}
 
 /**
@@ -105,7 +104,7 @@ FORCEINLINE auto Invoke(PtrMemFunType PtrMemFun, TargetType&& Target, ArgTypes&&
 #define UE_PROJECTION_MEMBER(Type, FuncName) \
 	[](auto&& Obj, auto&&... Args) -> decltype(auto) \
 	{ \
-		return UE::Core::Private::DereferenceIfNecessary<Type>(Forward<decltype(Obj)>(Obj), &Obj).FuncName(Forward<decltype(Args)>(Args)...); \
+		return UE::Core::Private::DereferenceIfNecessary<Type>((decltype(Obj)&&)Obj, &Obj).FuncName((decltype(Args)&&)Args...); \
 	}
 
 namespace UE::Core::Private
@@ -116,9 +115,9 @@ namespace UE::Core::Private
 	};
 
 	template <typename FuncType, typename... ArgTypes>
-	struct TInvokeResult_Impl<decltype((void)Invoke(DeclVal<FuncType>(), DeclVal<ArgTypes>()...)), FuncType, ArgTypes...>
+	struct TInvokeResult_Impl<decltype((void)Invoke(std::declval<FuncType>(), std::declval<ArgTypes>()...)), FuncType, ArgTypes...>
 	{
-		using Type = decltype(Invoke(DeclVal<FuncType>(), DeclVal<ArgTypes>()...));
+		using Type = decltype(Invoke(std::declval<FuncType>(), std::declval<ArgTypes>()...));
 	};
 }
 
@@ -133,3 +132,7 @@ struct TInvokeResult : UE::Core::Private::TInvokeResult_Impl<void, FuncType, Arg
 
 template <typename FuncType, typename... ArgTypes>
 using TInvokeResult_T = typename TInvokeResult<FuncType, ArgTypes...>::Type;
+
+#if UE_ENABLE_INCLUDE_ORDER_DEPRECATED_IN_5_7
+#include "Templates/UnrealTemplate.h"
+#endif
