StringBuilder: Fixed overallocation for unused null terminators

diff --git a/Engine/Source/Runtime/Core/Private/Misc/StringBuilder.cpp b/Engine/Source/Runtime/Core/Private/Misc/StringBuilder.cpp
--- a/Engine/Source/Runtime/Core/Private/Misc/StringBuilder.cpp
+++ b/Engine/Source/Runtime/Core/Private/Misc/StringBuilder.cpp
@@ -24,7 +24,7 @@ TStringBuilderBase<C>::~TStringBuilderBase()
 {
 	if (bIsDynamic)
 	{
-		FreeBuffer(Base, static_cast<SIZE_T>(End - Base));
+		FreeBuffer(Base);
 	}
 }
 
@@ -41,7 +41,7 @@ void TStringBuilderBase<C>::Extend(SIZE_T ExtraCapacity)
 
 	if (bIsDynamic)
 	{
-		FreeBuffer(Base, OldCapacity);
+		FreeBuffer(Base);
 	}
 
 	Base		= NewBase;
@@ -51,13 +51,13 @@ void TStringBuilderBase<C>::Extend(SIZE_T ExtraCapacity)
 }
 
 template <typename C>
-void* TStringBuilderBase<C>::AllocBuffer(SIZE_T CharCount)
+inline void* TStringBuilderBase<C>::AllocBuffer(SIZE_T CharCount)
 {
-	return FMemory::Malloc(CharCount * sizeof(C));
+	return FMemory::Malloc(CharCount * sizeof(C), alignof(C));
 }
 
 template <typename C>
-void TStringBuilderBase<C>::FreeBuffer(void* Buffer, SIZE_T CharCount)
+inline void TStringBuilderBase<C>::FreeBuffer(void* Buffer)
 {
 	FMemory::Free(Buffer);
 }
diff --git a/Engine/Source/Runtime/Core/Public/Misc/StringBuilder.h b/Engine/Source/Runtime/Core/Public/Misc/StringBuilder.h
--- a/Engine/Source/Runtime/Core/Public/Misc/StringBuilder.h
+++ b/Engine/Source/Runtime/Core/Public/Misc/StringBuilder.h
@@ -123,7 +123,7 @@ public:
 	UE_DEPRECATED(5.6, "Call the non-const operator*().")
 	inline const CharType* ToString() const UE_LIFETIMEBOUND
 	{
-		EnsureNulTerminated();
+		const_cast<TStringBuilderBase*>(this)->EnsureNulTerminated();
 		return Base;
 	}
 
@@ -138,7 +138,7 @@ public:
 	UE_DEPRECATED(5.6, "Call the non-const operator*().")
 	inline const CharType* operator*() const UE_LIFETIMEBOUND
 	{
-		EnsureNulTerminated();
+		const_cast<TStringBuilderBase*>(this)->EnsureNulTerminated();
 		return Base;
 	}
 
@@ -426,17 +426,17 @@ protected:
 		End		= Base + InCapacity;
 	}
 
-	inline void EnsureNulTerminated() const
+	inline void EnsureNulTerminated()
 	{
+		EnsureAdditionalCapacity(1);
 		*CurPos = CharType(0);
 	}
 
 	inline void EnsureAdditionalCapacity(int32 RequiredAdditionalCapacity)
 	{
-		// precondition: we know the current buffer has enough capacity
-		// for the existing string including NUL terminator
+		// precondition: we know the current buffer has enough capacity for the existing string
 
-		if ((CurPos + RequiredAdditionalCapacity) < End)
+		if (LIKELY((CurPos + RequiredAdditionalCapacity) <= End))
 		{
 			return;
 		}
@@ -444,9 +444,10 @@ protected:
 		Extend(RequiredAdditionalCapacity);
 	}
 
-	CORE_API void	Extend(SIZE_T ExtraCapacity);
-	CORE_API void*	AllocBuffer(SIZE_T CharCount);
-	CORE_API void	FreeBuffer(void* Buffer, SIZE_T CharCount);
+	CORE_API void Extend(SIZE_T ExtraCapacity);
+
+	void*	AllocBuffer(SIZE_T CharCount);
+	void	FreeBuffer(void* Buffer);
 
 	static inline CharType EmptyBuffer[1]{};
 
diff --git a/Engine/Source/Runtime/Core/Tests/Misc/StringBuilderTest.cpp b/Engine/Source/Runtime/Core/Tests/Misc/StringBuilderTest.cpp
--- a/Engine/Source/Runtime/Core/Tests/Misc/StringBuilderTest.cpp
+++ b/Engine/Source/Runtime/Core/Tests/Misc/StringBuilderTest.cpp
@@ -342,6 +342,24 @@ TEST_CASE("Core::String::StringBuilder", "[Core][String][Smoke]")
 		Builder.JoinQuoted(TSet<TAppendable<int32>>{{1}, {2}, {3}}, TAppendable{", "}, TAppendable{'"'});
 		CHECK(Builder == TEXTVIEW(R"("1", "2", "3")"));
 	}
+
+	SECTION("Appendf")
+	{
+		TStringBuilder<64> Builder;
+		const TCHAR* HexDigits = TEXT("0123456789abcdef");
+		Builder.Appendf(TEXT("%s%s%s%s%s%s%s%s"), HexDigits, HexDigits, HexDigits, HexDigits, HexDigits, HexDigits, HexDigits, HexDigits);
+		CHECK(Builder == TEXTVIEW("0123456789abcdef0123456789abcdef0123456789abcdef0123456789abcdef0123456789abcdef0123456789abcdef0123456789abcdef0123456789abcdef"));
+	}
+
+	SECTION("NullTerminator")
+	{
+		// Validate that space for a null terminator is only allocated when null termination is requested.
+		TStringBuilder<16> Builder;
+		const TCHAR* ExpectedData = Builder.GetData();
+		Builder.Append(TEXTVIEW("0123456789abcdef"));
+		CHECK(Builder.GetData() == ExpectedData);
+		CHECK_FALSE(*Builder == ExpectedData);
+	}
 }
 
 #endif
\ No newline at end of file
