Shrink the platform I/O dispatcher allocators from default to 1024.

#rnx

diff --git a/Engine/Source/Runtime/Core/Internal/IO/PlatformIoDispatcherBase.h b/Engine/Source/Runtime/Core/Internal/IO/PlatformIoDispatcherBase.h
--- a/Engine/Source/Runtime/Core/Internal/IO/PlatformIoDispatcherBase.h
+++ b/Engine/Source/Runtime/Core/Internal/IO/PlatformIoDispatcherBase.h
@@ -118,7 +118,7 @@ struct FIoEncodedBlockRequest
 	uint8							FileBlockCount = 0;
 };
 
-using FIoEncodedBlockRequestAllocator	= TSingleThreadedSlabAllocator<FIoEncodedBlockRequest>;
+using FIoEncodedBlockRequestAllocator	= TSingleThreadedSlabAllocator<FIoEncodedBlockRequest, 1024>;
 using FIoEncodedBlockRequestList		= TIntrusiveList<FIoEncodedBlockRequest>;
 
 ////////////////////////////////////////////////////////////////////////////////
@@ -130,7 +130,7 @@ struct FIoFileBlockLink
 	FIoFileBlockLink*	NextLink = nullptr;
 };
 static_assert(sizeof(FIoFileBlockLink) == 32);
-using FIoFileBlockLinkAllocator = TSingleThreadedSlabAllocator<FIoFileBlockLink>;
+using FIoFileBlockLinkAllocator = TSingleThreadedSlabAllocator<FIoFileBlockLink, 128>;
 
 ////////////////////////////////////////////////////////////////////////////////
 struct FIoPlatformReadRequest
@@ -206,7 +206,7 @@ struct FIoPlatformReadRequest
 	EType						Type = EType::ScatterGather;
 };
 
-using FIoPlatformReadRequestAllocator	= TSingleThreadedSlabAllocator<FIoPlatformReadRequest>;
+using FIoPlatformReadRequestAllocator	= TSingleThreadedSlabAllocator<FIoPlatformReadRequest, 1024>;
 using FIoPlatformReadRequestList		= TIntrusiveList<FIoPlatformReadRequest>;
 
 ////////////////////////////////////////////////////////////////////////////////
@@ -241,7 +241,7 @@ struct FIoFileBlockRequest
 	EQueueStatus					QueueStatus = EQueueStatus::None;
 };
 
-using FIoFileBlockRequestAllocator	= TSingleThreadedSlabAllocator<FIoFileBlockRequest>;
+using FIoFileBlockRequestAllocator	= TSingleThreadedSlabAllocator<FIoFileBlockRequest, 1024>;
 using FIoFileBlockRequestList		= TIntrusiveList<FIoFileBlockRequest>;
 
 ////////////////////////////////////////////////////////////////////////////////
