IntrusiveMutex: Removed unnecessary UE:: prefix

#rnx

diff --git a/Engine/Source/Runtime/Core/Public/Async/IntrusiveMutex.h b/Engine/Source/Runtime/Core/Public/Async/IntrusiveMutex.h
--- a/Engine/Source/Runtime/Core/Public/Async/IntrusiveMutex.h
+++ b/Engine/Source/Runtime/Core/Public/Async/IntrusiveMutex.h
@@ -21,7 +21,7 @@ using TIntrusiveMutexStateType_T = std::decay_t<decltype(ParamsType::IsLockedFla
 template <typename ParamsType, typename StateType = TIntrusiveMutexStateType_T<ParamsType>>
 concept CIntrusiveMutexParams = requires
 {
-	requires UE::CIntegral<StateType>;
+	requires CIntegral<StateType>;
 
 	// Required: constexpr static StateType IsLockedFlag = ...;
 	// Flag that is set in the state when the mutex is locked.
