Added intrusive unset optional support to TNotNull.

diff --git a/Engine/Source/Runtime/Core/Public/Misc/IntrusiveUnsetOptionalState.h b/Engine/Source/Runtime/Core/Public/Misc/IntrusiveUnsetOptionalState.h
--- a/Engine/Source/Runtime/Core/Public/Misc/IntrusiveUnsetOptionalState.h
+++ b/Engine/Source/Runtime/Core/Public/Misc/IntrusiveUnsetOptionalState.h
@@ -59,6 +59,8 @@
  */
 
 #include "Misc/OptionalFwd.h"
+#include "Concepts/DecaysTo.h"
+#include "Concepts/SameAs.h"
 
 namespace UE::Core::Private
 {
@@ -78,3 +80,18 @@ struct FIntrusiveUnsetOptionalState
 private:
 	explicit FIntrusiveUnsetOptionalState() = default;
 };
+
+template <typename T>
+constexpr bool HasIntrusiveUnsetOptionalState()
+{
+	// Ensure the type has a nested bHasIntrusiveUnsetOptionalState, and as derived types are not guaranteed to have
+	// an intrusive state even if the base does, ensure IntrusiveUnsetOptionalStateType matches the type in the optional.
+	if constexpr (requires{ { T::bHasIntrusiveUnsetOptionalState } -> UE::CDecaysTo<bool>; requires UE::CSameAs<const typename T::IntrusiveUnsetOptionalStateType, const T>; })
+	{
+		return T::bHasIntrusiveUnsetOptionalState;
+	}
+	else
+	{
+		return false;
+	}
+}
diff --git a/Engine/Source/Runtime/Core/Public/Misc/NotNull.h b/Engine/Source/Runtime/Core/Public/Misc/NotNull.h
--- a/Engine/Source/Runtime/Core/Public/Misc/NotNull.h
+++ b/Engine/Source/Runtime/Core/Public/Misc/NotNull.h
@@ -3,6 +3,7 @@
 #pragma once
 
 #include "CoreTypes.h"
+#include "Misc/IntrusiveUnsetOptionalState.h"
 #include "Templates/Requires.h"
 #include "Templates/TypeHash.h"
 #include "Traits/IsImplicitlyConstructible.h"
@@ -183,6 +184,26 @@
 #endif
 		}
 
+		/////////////////////////////////////////////////
+		// Start - intrusive TOptional<TNotNull> state //
+		/////////////////////////////////////////////////
+		static constexpr bool bHasIntrusiveUnsetOptionalState = HasIntrusiveUnsetOptionalState<T>();
+		using IntrusiveUnsetOptionalStateType = TNotNull;
+
+		UE_NODEBUG [[nodiscard]] TNotNull(FIntrusiveUnsetOptionalState Tag)
+			requires(bHasIntrusiveUnsetOptionalState)
+			: Val(Tag)
+		{
+		}
+		UE_NODEBUG [[nodiscard]] bool operator==(FIntrusiveUnsetOptionalState Tag) const
+			requires(bHasIntrusiveUnsetOptionalState)
+		{
+			return Val == Tag;
+		}
+		///////////////////////////////////////////////
+		// End - intrusive TOptional<TNotNull> state //
+		///////////////////////////////////////////////
+
 		// Allow conversions
 		template <
 			typename OtherType
diff --git a/Engine/Source/Runtime/Core/Public/Misc/Optional.h b/Engine/Source/Runtime/Core/Public/Misc/Optional.h
--- a/Engine/Source/Runtime/Core/Public/Misc/Optional.h
+++ b/Engine/Source/Runtime/Core/Public/Misc/Optional.h
@@ -13,35 +13,6 @@
 
 inline constexpr FNullOpt NullOpt{0};
 
-namespace UE::Core::Private
-{
-	struct CIntrusiveUnsettable
-	{
-		template <typename T>
-		auto Requires(bool& b) -> decltype(
-			b = std::decay_t<typename T::IntrusiveUnsetOptionalStateType>::bHasIntrusiveUnsetOptionalState
-		);
-	};
-}
-
-template <typename T>
-constexpr bool HasIntrusiveUnsetOptionalState()
-{
-	if constexpr (!TModels<UE::Core::Private::CIntrusiveUnsettable, T>::Value)
-	{
-		return false;
-	}
-	// Derived types are not guaranteed to have an intrusive state, so ensure IntrusiveUnsetOptionalStateType matches the type in the optional
-	else if constexpr (!std::is_same_v<const typename T::IntrusiveUnsetOptionalStateType, const T>)
-	{
-		return false;
-	}
-	else
-	{
-		return T::bHasIntrusiveUnsetOptionalState;
-	}
-}
-
 namespace UE::Core::Private
 {
 	// Empty structure to replace bIsSet for intrusive optional types 
diff --git a/Engine/Source/Runtime/Core/Tests/Misc/OptionalTest.cpp b/Engine/Source/Runtime/Core/Tests/Misc/OptionalTest.cpp
--- a/Engine/Source/Runtime/Core/Tests/Misc/OptionalTest.cpp
+++ b/Engine/Source/Runtime/Core/Tests/Misc/OptionalTest.cpp
@@ -9,6 +9,7 @@
 #include "Internationalization/Text.h"
 #include "Math/Vector.h"
 #include "Misc/LowLevelTestAdapter.h"
+#include "Misc/NotNull.h"
 #include "Misc/Optional.h"
 #include "Templates/AlignmentTemplates.h"
 #include "Tests/TestHarnessAdapter.h"
@@ -33,10 +34,10 @@ namespace UE::Core::Private
 	static_assert(!std::is_trivially_destructible_v<TOptional<FString>>);
 	static_assert(!std::is_trivially_destructible_v<TOptional<FText>>);
 
-    template<bool bIntrusive, bool bTrackDestruction>
+	template<bool bIntrusive, bool bTrackDestruction>
 	struct TTestValue
 	{
-        static inline int32 DestructorCount = 0;
+		static inline int32 DestructorCount = 0;
 		static constexpr inline bool bHasIntrusiveUnsetOptionalState = bIntrusive;
 		using IntrusiveUnsetOptionalStateType = TTestValue<true, bTrackDestruction>;
 
@@ -47,93 +48,100 @@ namespace UE::Core::Private
 		constexpr TTestValue(int32 InValue)
 		: Value(InValue)
 		{
-            static_assert(!bIntrusive || sizeof(TOptional<TTestValue>) == sizeof(TTestValue), "Expected intrusive optional to be the same size as the inner type");
-            static_assert(bIntrusive || sizeof(TOptional<TTestValue>) > sizeof(TTestValue), "Expected non-intrusive optional to be larger than the inner type");
-            static_assert(bIntrusive || sizeof(TOptional<TTestValue>) <= Align(sizeof(TTestValue) + 1, alignof(TTestValue)),
-                "Expected non-intrusive optional to be larger than the inner type by at most the alignment of the inner type");
+			static_assert(!bIntrusive || sizeof(TOptional<TTestValue>) == sizeof(TTestValue), "Expected intrusive optional to be the same size as the inner type");
+			static_assert(bIntrusive || sizeof(TOptional<TTestValue>) > sizeof(TTestValue), "Expected non-intrusive optional to be larger than the inner type");
+			static_assert(bIntrusive || sizeof(TOptional<TTestValue>) <= Align(sizeof(TTestValue) + 1, alignof(TTestValue)),
+				"Expected non-intrusive optional to be larger than the inner type by at most the alignment of the inner type");
 		}
-        constexpr TTestValue(const TTestValue&) = default;
-        constexpr TTestValue& operator=(const TTestValue&) = default;
-        constexpr TTestValue(TTestValue&&) = default;
-        constexpr TTestValue& operator=(TTestValue&&) = default;
-
-        // Explicit destructor so not trivially destructible, but constexpr-destructible
-        constexpr ~TTestValue()
-        {
-            if constexpr (bTrackDestruction)
-            {
-                ++DestructorCount;
-            }
-        }
-
-        [[nodiscard]] constexpr bool operator==(FIntrusiveUnsetOptionalState) const requires bIntrusive
-        {
-            return Value == 0;
-        }
-
-        [[nodiscard]] constexpr bool operator==(const TTestValue&) const = default;
+		constexpr TTestValue(const TTestValue&) = default;
+		constexpr TTestValue& operator=(const TTestValue&) = default;
+		constexpr TTestValue(TTestValue&&) = default;
+		constexpr TTestValue& operator=(TTestValue&&) = default;
+
+		// Explicit destructor so not trivially destructible, but constexpr-destructible
+		constexpr ~TTestValue()
+		{
+			if constexpr (bTrackDestruction)
+			{
+				++DestructorCount;
+			}
+		}
+
+		[[nodiscard]] constexpr bool operator==(FIntrusiveUnsetOptionalState) const requires bIntrusive
+		{
+			return Value == 0;
+		}
+
+		[[nodiscard]] constexpr bool operator==(const TTestValue&) const = default;
 
 		int32 Value;
 	};
 
-    template<bool bTrackDestruction = false>
-    using TIntrusive = TTestValue<true, bTrackDestruction>;
-    template<bool bTrackDestruction = false>
+	template<bool bTrackDestruction = false>
+	using TIntrusive = TTestValue<true, bTrackDestruction>;
+	template<bool bTrackDestruction = false>
 	using TNonIntrusive = TTestValue<false, bTrackDestruction>;
+	struct FDerivedIntrusive : TIntrusive<> {};
+
+	static_assert(HasIntrusiveUnsetOptionalState<TIntrusive<>>());
+	static_assert(!HasIntrusiveUnsetOptionalState<FDerivedIntrusive>()); // a class which derives from another with an intrusive unset state should not be assumed to have an intrusive unset state itself.
+	static_assert(!HasIntrusiveUnsetOptionalState<TNonIntrusive<>>());
+
+	// Ensure TNotNull honors the intrusive unset state of the thing it's wrapping.
+	static_assert(HasIntrusiveUnsetOptionalState<TNotNull<TIntrusive<>>>());
+	static_assert(!HasIntrusiveUnsetOptionalState<TNotNull<FDerivedIntrusive>>());
+	static_assert(!HasIntrusiveUnsetOptionalState<TNotNull<TNonIntrusive<>>>());
 
-    static_assert(HasIntrusiveUnsetOptionalState<TIntrusive<>>());
-    static_assert(!HasIntrusiveUnsetOptionalState<TNonIntrusive<>>());
-
-    // Tests for constexpr use of optional
-    template<typename InnerType>
-    constexpr bool TestConstexprOptionalUse()
-    {
-        InnerType Value{32};
-        TOptional<InnerType> O1; // Default construction
-        TOptional<InnerType> O2{NullOpt}; // Explicit empty construction;
-        TOptional<InnerType> O3(Value); // Copying inner type
-        TOptional<InnerType> O4(Forward<InnerType>(InnerType(30))); // Moving inner type
-        TOptional<InnerType> O5(InPlace, 29); // In-place construction
-        if constexpr (false) // At present, copy and move construction are not constexpr-safe
-        {
-            TOptional<InnerType> O6(O3); // Copy construction
-            TOptional<InnerType> O7(MoveTemp(O4)); // Move construction
-        }
-        bool bTestResult = (O1 == O2) && (O4 != O5);
-        bTestResult = !O1.IsSet() && O3.IsSet() && bTestResult;
-        const TOptional<InnerType> C1(InPlace, 33);
-        const InnerType& I1 = C1.GetValue();
-        O5.GetValue().Value++;
-        O5->Value = C1->Value;
-        bTestResult = ((*O5) == (*C1)) && bTestResult;
-        const InnerType& I2 = O1.Get(Value);
-        *O4.GetPtrOrNull() = *C1.GetPtrOrNull();
-
-        return bTestResult;
-    }
-
-    static_assert(TestConstexprOptionalUse<TIntrusive<>>(), "Constexpr tests for TOptional<TIntrusive> failed");
-    static_assert(TestConstexprOptionalUse<TNonIntrusive<>>(), "Constexpr tests for TOptional<TNonIntrusive> failed");
+	// Tests for constexpr use of optional
+	template<typename InnerType>
+	constexpr bool TestConstexprOptionalUse()
+	{
+		InnerType Value{32};
+		TOptional<InnerType> O1; // Default construction
+		TOptional<InnerType> O2{NullOpt}; // Explicit empty construction;
+		TOptional<InnerType> O3(Value); // Copying inner type
+		TOptional<InnerType> O4(Forward<InnerType>(InnerType(30))); // Moving inner type
+		TOptional<InnerType> O5(InPlace, 29); // In-place construction
+		if constexpr (false) // At present, copy and move construction are not constexpr-safe
+		{
+			TOptional<InnerType> O6(O3); // Copy construction
+			TOptional<InnerType> O7(MoveTemp(O4)); // Move construction
+		}
+		bool bTestResult = (O1 == O2) && (O4 != O5);
+		bTestResult = !O1.IsSet() && O3.IsSet() && bTestResult;
+		const TOptional<InnerType> C1(InPlace, 33);
+		const InnerType& I1 = C1.GetValue();
+		O5.GetValue().Value++;
+		O5->Value = C1->Value;
+		bTestResult = ((*O5) == (*C1)) && bTestResult;
+		const InnerType& I2 = O1.Get(Value);
+		*O4.GetPtrOrNull() = *C1.GetPtrOrNull();
+
+		return bTestResult;
+	}
+
+	static_assert(TestConstexprOptionalUse<TIntrusive<>>(), "Constexpr tests for TOptional<TIntrusive> failed");
+	static_assert(TestConstexprOptionalUse<TNonIntrusive<>>(), "Constexpr tests for TOptional<TNonIntrusive> failed");
 };
 
 TEST_CASE_NAMED(FTOptionalTest_NonIntrusive, "System::Core::TOptional::NonIntrusive", "[Core][TOptional][SmokeFilter]")
 {
-    using FNonIntrusiveDestructible = UE::Core::Private::TNonIntrusive<true>;
-    FNonIntrusiveDestructible::DestructorCount = 0;
-    { 
-        TOptional<FNonIntrusiveDestructible> Unset;
-    }
-    CHECK_MESSAGE("Unset optional is not destroyed", FNonIntrusiveDestructible::DestructorCount == 0);
+	using FNonIntrusiveDestructible = UE::Core::Private::TNonIntrusive<true>;
+	FNonIntrusiveDestructible::DestructorCount = 0;
+	{ 
+		TOptional<FNonIntrusiveDestructible> Unset;
+	}
+	CHECK_MESSAGE("Unset optional is not destroyed", FNonIntrusiveDestructible::DestructorCount == 0);
 }
 
 TEST_CASE_NAMED(FTOptionalTest_Intrusive, "System::Core::TOptional::Intrusive", "[Core][TOptional][SmokeFilter]")
 {
-    using FIntrusiveDestructible = UE::Core::Private::TIntrusive<true>;
-    FIntrusiveDestructible::DestructorCount = 0;
-    { 
-        TOptional<FIntrusiveDestructible> Unset;
-    }
-    CHECK_MESSAGE("Unset optional is not destroyed", FIntrusiveDestructible::DestructorCount == 0);
+	using FIntrusiveDestructible = UE::Core::Private::TIntrusive<true>;
+	FIntrusiveDestructible::DestructorCount = 0;
+	{ 
+		TOptional<FIntrusiveDestructible> Unset;
+	}
+	CHECK_MESSAGE("Unset optional is not destroyed", FIntrusiveDestructible::DestructorCount == 0);
 }
 
 #endif // WITH_TESTS
\ No newline at end of file
