Removed UE_ASSUME usage from TNotNull, as it would cause different codegen between UE_ENABLE_NOTNULL_WRAPPER states, and gives warnings under Clang about side effects when used on smart pointers.

diff --git a/Engine/Source/Runtime/Core/Public/Misc/NotNull.h b/Engine/Source/Runtime/Core/Public/Misc/NotNull.h
--- a/Engine/Source/Runtime/Core/Public/Misc/NotNull.h
+++ b/Engine/Source/Runtime/Core/Public/Misc/NotNull.h
@@ -100,15 +100,11 @@
 			// Allow implicit conversion to the wrapped value, for compatibility with non-TNotNull APIs
 			operator T() const&
 			{
-				T Result = ((Sub*)this)->Val;
-				UE_ASSUME(Result);
-				return Result;
+				return ((Sub*)this)->Val;
 			}
 			operator T() &&
 			{
-				T Result = (T&&)((Sub*)this)->Val;
-				UE_ASSUME(Result);
-				return Result;
+				return (T&&)((Sub*)this)->Val;
 			}
 
 			// Allow dereferencing, if it's supported
@@ -128,9 +124,7 @@
 			operator T*() const
 		UE_NOTNULL_FUNCTION_NON_NULL_RETURN_END
 			{
-				T* Result = ((Sub*)this)->Val;
-				UE_ASSUME(Result);
-				return Result;
+				return ((Sub*)this)->Val;
 			}
 
 			// Allow dereferencing, if it's supported
@@ -241,7 +235,6 @@
 		template <typename... ArgTypes>
 		auto operator()(ArgTypes&&... Args) const -> decltype(this->Val((ArgTypes&&)Args...))
 		{
-			UE_ASSUME(Val);
 			return Val((ArgTypes&&)Args...);
 		}
 
