Fixed a bug when ConcurrentLinearAllocator can allocate a non 64kb aligned block when one is requested

diff --git a/Engine/Source/Runtime/Core/Private/Misc/MemStack.cpp b/Engine/Source/Runtime/Core/Private/Misc/MemStack.cpp
--- a/Engine/Source/Runtime/Core/Private/Misc/MemStack.cpp
+++ b/Engine/Source/Runtime/Core/Private/Misc/MemStack.cpp
@@ -119,7 +119,7 @@ int32 FrameIndex()
 	return int32(CurrentFrame % FRAMES_TO_WAIT);
 }
 
-void *FPageAllocator::Alloc()
+void *FPageAllocator::Alloc(int32 /*Alignment*/)
 {
 	check(GMemStackProtectionLatched); // attempt to allocate a page before we know if we are doing protected mode or not
 	void *Result;
@@ -233,13 +233,13 @@ uint64 FPageAllocator::BytesFree()
 
 #define MEMSTACK_PUGATORY_COMPILED_IN (0)
 
-void *FPageAllocator::Alloc()
+void *FPageAllocator::Alloc(int32 Alignment)
 {
 #if USING_ADDRESS_SANITISER
-	return FMemory::Malloc(FPageAllocator::PageSize);
+	return FMemory::Malloc(FPageAllocator::PageSize, Alignment);
 #else
 	LLM_SCOPE_BY_BOOTSTRAP_TAG(MemStack);
-	void *Result = TheAllocator.Allocate();
+	void *Result = TheAllocator.Allocate(Alignment);
 	STAT(UpdateStats());
 	return Result;
 #endif
diff --git a/Engine/Source/Runtime/Core/Public/Containers/LockFreeFixedSizeAllocator.h b/Engine/Source/Runtime/Core/Public/Containers/LockFreeFixedSizeAllocator.h
--- a/Engine/Source/Runtime/Core/Public/Containers/LockFreeFixedSizeAllocator.h
+++ b/Engine/Source/Runtime/Core/Public/Containers/LockFreeFixedSizeAllocator.h
@@ -217,7 +217,7 @@ public:
 	 * @return Pointer to the allocated memory.
 	 * @see Free
 	 */
-	void* Allocate()
+	void* Allocate(int32 Alignment = MIN_ALIGNMENT)
 	{
 		void* Memory = nullptr;
 
@@ -227,14 +227,18 @@ public:
 		UE_AUTORTFM_OPEN
 		{
 			NumUsed.Increment();
-			Memory = FreeList.Pop();
-			if (Memory)
+			if (Alignment <= 4096)
 			{
-				NumFree.Decrement();
+				// Pop from a free list only if Alignment is not large than a memory page size
+				Memory = FreeList.Pop();
+				if (Memory)
+				{
+					NumFree.Decrement();
+				}
 			}
-			else
+			if (!Memory)
 			{
-				Memory = FMemory::Malloc(SIZE);
+				Memory = FMemory::Malloc(SIZE, Alignment);
 			}
 		};
 
diff --git a/Engine/Source/Runtime/Core/Public/Experimental/ConcurrentLinearAllocator.h b/Engine/Source/Runtime/Core/Public/Experimental/ConcurrentLinearAllocator.h
--- a/Engine/Source/Runtime/Core/Public/Experimental/ConcurrentLinearAllocator.h
+++ b/Engine/Source/Runtime/Core/Public/Experimental/ConcurrentLinearAllocator.h
@@ -128,7 +128,7 @@ public:
 	{
 		if (Size == BlockSize)
 		{
-			return FPageAllocator::Get().Alloc();
+			return FPageAllocator::Get().Alloc(Alignment);
 		}
 		else
 		{
diff --git a/Engine/Source/Runtime/Core/Public/Misc/MemStack.h b/Engine/Source/Runtime/Core/Public/Misc/MemStack.h
--- a/Engine/Source/Runtime/Core/Public/Misc/MemStack.h
+++ b/Engine/Source/Runtime/Core/Public/Misc/MemStack.h
@@ -49,7 +49,7 @@ public:
 	static CORE_API FPageAllocator& Get();
 
 	CORE_API ~FPageAllocator();
-	CORE_API void* Alloc();
+	CORE_API void* Alloc(int32 Alignment = MIN_ALIGNMENT);
 	CORE_API void Free(void* Mem);
 	CORE_API void* AllocSmall();
 	CORE_API void FreeSmall(void* Mem);
