Added TFunctionWithContext to adapt a callable into a function pointer and context

diff --git a/Engine/Source/Runtime/Core/Public/Async/ParkingLot.h b/Engine/Source/Runtime/Core/Public/Async/ParkingLot.h
--- a/Engine/Source/Runtime/Core/Public/Async/ParkingLot.h
+++ b/Engine/Source/Runtime/Core/Public/Async/ParkingLot.h
@@ -4,7 +4,7 @@
 
 #include "CoreTypes.h"
 #include "Misc/MonotonicTime.h"
-#include "Templates/Invoke.h"
+#include "Templates/FunctionWithContext.h"
 
 #include <type_traits>
 
@@ -36,65 +36,10 @@ struct FWakeState final
 
 namespace Private
 {
-
-template <typename>
-class TLambdaAsFunctionWithContext;
-
-template <typename ReturnType, typename... ArgTypes>
-class TLambdaAsFunctionWithContext<ReturnType (ArgTypes...)>
-{
-public:
-	using FunctionType = ReturnType (void*, ArgTypes...);
-
-	inline constexpr TLambdaAsFunctionWithContext(decltype(nullptr))
-	{
-	}
-
-	template <typename LambdaType>
-	inline TLambdaAsFunctionWithContext(LambdaType&& Lambda UE_LIFETIMEBOUND)
-		requires (!std::is_same_v<std::decay_t<LambdaType>, TLambdaAsFunctionWithContext>) &&
-	std::is_invocable_r_v<ReturnType, std::decay_t<LambdaType>, ArgTypes...>
-	{
-		Function = &Call<std::decay_t<LambdaType>>;
-		Context = &Lambda;
-	}
-
-	inline explicit operator bool() const
-	{
-		return !!Function;
-	}
-
-	inline ReturnType operator()(ArgTypes... Args) const
-	{
-		return Function(Context, Forward<ArgTypes>(Args)...);
-	}
-
-	inline FunctionType* GetFunction() const
-	{
-		return Function;
-	}
-
-	inline void* GetContext() const
-	{
-		return Context;
-	}
-
-private:
-	template <typename LambdaType>
-	static ReturnType Call(void* Lambda, ArgTypes... Args)
-	{
-		return Invoke(*(LambdaType*)Lambda, Forward<ArgTypes>(Args)...);
-	}
-
-	FunctionType* Function = nullptr;
-	void* Context = nullptr;
-};
-
 UE_API FWaitState Wait(const void* Address, bool (*CanWait)(void*), void* CanWaitContext, void (*BeforeWait)(void*), void* BeforeWaitContext);
 UE_API FWaitState WaitFor(const void* Address, bool (*CanWait)(void*), void* CanWaitContext, void (*BeforeWait)(void*), void* BeforeWaitContext, FMonotonicTimeSpan WaitTime);
 UE_API FWaitState WaitUntil(const void* Address, bool (*CanWait)(void*), void* CanWaitContext, void (*BeforeWait)(void*), void* BeforeWaitContext, FMonotonicTimePoint WaitTime);
 UE_API void WakeOne(const void* Address, uint64 (*OnWakeState)(void*, FWakeState), void* OnWakeStateContext);
-
 } // Private
 
 /**
@@ -104,7 +49,7 @@ UE_API void WakeOne(const void* Address, uint64 (*OnWakeState)(void*, FWakeState
  * @param CanWait      Function called while the queue is locked. A return of false cancels the wait.
  * @param BeforeWait   Function called after the queue is unlocked and before the thread waits.
  */
-inline FWaitState Wait(const void* Address, Private::TLambdaAsFunctionWithContext<bool()> CanWait, Private::TLambdaAsFunctionWithContext<void()> BeforeWait)
+inline FWaitState Wait(const void* Address, TFunctionWithContext<bool()> CanWait, TFunctionWithContext<void()> BeforeWait)
 {
 	return Private::Wait(Address, CanWait.GetFunction(), CanWait.GetContext(), BeforeWait.GetFunction(), BeforeWait.GetContext());
 }
@@ -117,7 +62,7 @@ inline FWaitState Wait(const void* Address, Private::TLambdaAsFunctionWithContex
  * @param BeforeWait   Function called after the queue is unlocked and before the thread waits.
  * @param WaitTime     Relative time after which waiting is automatically canceled and the thread wakes.
  */
-inline FWaitState WaitFor(const void* Address, Private::TLambdaAsFunctionWithContext<bool()> CanWait, Private::TLambdaAsFunctionWithContext<void()> BeforeWait, FMonotonicTimeSpan WaitTime)
+inline FWaitState WaitFor(const void* Address, TFunctionWithContext<bool()> CanWait, TFunctionWithContext<void()> BeforeWait, FMonotonicTimeSpan WaitTime)
 {
 	return Private::WaitFor(Address, CanWait.GetFunction(), CanWait.GetContext(), BeforeWait.GetFunction(), BeforeWait.GetContext(), WaitTime);
 }
@@ -130,7 +75,7 @@ inline FWaitState WaitFor(const void* Address, Private::TLambdaAsFunctionWithCon
  * @param BeforeWait   Function called after the queue is unlocked and before the thread waits.
  * @param WaitTime     Absolute time after which waiting is automatically canceled and the thread wakes.
  */
-inline FWaitState WaitUntil(const void* Address, Private::TLambdaAsFunctionWithContext<bool()> CanWait, Private::TLambdaAsFunctionWithContext<void()> BeforeWait, FMonotonicTimePoint WaitTime)
+inline FWaitState WaitUntil(const void* Address, TFunctionWithContext<bool()> CanWait, TFunctionWithContext<void()> BeforeWait, FMonotonicTimePoint WaitTime)
 {
 	return Private::WaitUntil(Address, CanWait.GetFunction(), CanWait.GetContext(), BeforeWait.GetFunction(), BeforeWait.GetContext(), WaitTime);
 }
@@ -141,7 +86,7 @@ inline FWaitState WaitUntil(const void* Address, Private::TLambdaAsFunctionWithC
  * @param Address       Address to use as the key for the queue. Must match the address used in Wait.
  * @param OnWakeState   Function called while the queue is locked. Receives the wake state. Returns WakeToken.
  */
-inline void WakeOne(const void* Address, Private::TLambdaAsFunctionWithContext<uint64(FWakeState)> OnWakeState)
+inline void WakeOne(const void* Address, TFunctionWithContext<uint64(FWakeState)> OnWakeState)
 {
 	return Private::WakeOne(Address, OnWakeState.GetFunction(), OnWakeState.GetContext());
 }
diff --git a/Engine/Source/Runtime/Core/Public/Templates/FunctionWithContext.h b/Engine/Source/Runtime/Core/Public/Templates/FunctionWithContext.h
new file mode 100644
--- /dev/null
+++ b/Engine/Source/Runtime/Core/Public/Templates/FunctionWithContext.h
@@ -0,0 +1,111 @@
+// Copyright Epic Games, Inc. All Rights Reserved.
+
+#pragma once
+
+#include "HAL/Platform.h"
+#include "Templates/Invoke.h"
+#include <type_traits>
+
+namespace UE
+{
+
+/**
+ * Type that adapts a callable into a function pointer with a context pointer.
+ *
+ * This does not take ownership of the callable. If constructed with a lambda,
+ * it is only valid until the lambda goes out of scope. Use TFunction to take
+ * ownership of the lambda.
+ *
+ * This behaves like a nullable TFunctionRef with the addition of accessors to
+ * pass on the function and context pointers to implementation functions. This
+ * tends to generate more efficient code than when passing a TFunctionRef or a
+ * TFunctionWithContext either by value or by reference.
+ *
+ * A function taking a string and float and returning int32 usage might be:
+ *
+ * void ParseLines(FStringView View, void (*Visitor)(void* Context, FStringView Line), void* Context);
+ * inline void ParseLines(FStringView View, TFunctionWithContext<void (FStringView Line)> Visitor)
+ * {
+ *     ParseLines(View, Visitor.GetFunction(), Visitor.GetContext());
+ * }
+ *
+ * The example ParseLines can be called as:
+ *
+ * ParseLines(Input, [](FStringView Line) { PrintLine(Line); });
+ */
+template <typename FunctionType>
+class TFunctionWithContext;
+
+template <typename ReturnType, typename... ArgTypes>
+class TFunctionWithContext<ReturnType (ArgTypes...)>
+{
+public:
+	using FunctionType = ReturnType (void*, ArgTypes...);
+
+	/** Construct from a lambda or a function pointer. */
+	template <typename LambdaType>
+	inline TFunctionWithContext(LambdaType&& Lambda UE_LIFETIMEBOUND)
+		requires (!std::is_same_v<std::decay_t<LambdaType>, TFunctionWithContext>) &&
+			std::is_invocable_r_v<ReturnType, std::decay_t<LambdaType>, ArgTypes...>
+		: Function(&Call<std::decay_t<LambdaType>>)
+		, Context(&Lambda)
+	{
+	}
+
+	/** Assign from a lambda or a function pointer. */
+	template <typename LambdaType>
+	inline TFunctionWithContext& operator=(LambdaType&& Lambda UE_LIFETIMEBOUND)
+		requires (!std::is_same_v<std::decay_t<LambdaType>, TFunctionWithContext>) &&
+			std::is_invocable_r_v<ReturnType, std::decay_t<LambdaType>, ArgTypes...>
+	{
+		Function = &Call<std::decay_t<LambdaType>>;
+		Context = &Lambda;
+		return *this;
+	}
+
+	/** Construct from a function pointer and context. Function and context may both be null. */
+	inline explicit TFunctionWithContext(FunctionType* InFunction, void* InContext)
+		: Function(InFunction)
+		, Context(InContext)
+	{
+	}
+
+	/** Construct a null function with null context. */
+	inline constexpr TFunctionWithContext(decltype(nullptr))
+	{
+	}
+
+	/** Returns true if the function is non-null. */
+	inline explicit operator bool() const
+	{
+		return !!Function;
+	}
+
+	/** Calls the function with the stored context and provided arguments. Function must be non-null. */
+	inline ReturnType operator()(ArgTypes... Args) const
+	{
+		return Function(Context, (ArgTypes&&)Args...);
+	}
+
+	inline FunctionType* GetFunction() const
+	{
+		return Function;
+	}
+
+	inline void* GetContext() const
+	{
+		return Context;
+	}
+
+private:
+	template <typename LambdaType>
+	static ReturnType Call(void* Lambda, ArgTypes... Args)
+	{
+		return Invoke(*(LambdaType*)Lambda, (ArgTypes&&)Args...);
+	}
+
+	FunctionType* Function = nullptr;
+	void* Context = nullptr;
+};
+
+} // UE
