[TSAN]
- Fix a few benign races on values that are going to be eventually consistent anyway

diff --git a/Engine/Source/Runtime/Core/Public/Clang/ClangPlatformCodeAnalysis.h b/Engine/Source/Runtime/Core/Public/Clang/ClangPlatformCodeAnalysis.h
--- a/Engine/Source/Runtime/Core/Public/Clang/ClangPlatformCodeAnalysis.h
+++ b/Engine/Source/Runtime/Core/Public/Clang/ClangPlatformCodeAnalysis.h
@@ -172,6 +172,16 @@
 			__attribute__((always_inline)) T operator--() { return Value.fetch_sub(1, std::memory_order_relaxed) - 1; }
 			__attribute__((always_inline)) T operator--(int) { return Value.fetch_sub(1, std::memory_order_relaxed); }
 
+			__attribute__((always_inline)) TTSANSafeValue(const TTSANSafeValue& Other)
+			{
+				Value.store((T)Other, std::memory_order_relaxed);
+			}
+
+			__attribute__((always_inline)) TTSANSafeValue& operator=(const TTSANSafeValue& Other)
+			{
+				Value.store((T)Other, std::memory_order_relaxed);
+			}
+
 			template <typename OtherType>
 			__attribute__((always_inline)) T operator+=(OtherType InValue) { return Value.fetch_add(InValue, std::memory_order_relaxed) + InValue; }
 
diff --git a/Engine/Source/Runtime/Core/Public/IO/IoDispatcherBackend.h b/Engine/Source/Runtime/Core/Public/IO/IoDispatcherBackend.h
--- a/Engine/Source/Runtime/Core/Public/IO/IoDispatcherBackend.h
+++ b/Engine/Source/Runtime/Core/Public/IO/IoDispatcherBackend.h
@@ -31,7 +31,7 @@ public:
 	/** Read options. */
 	FIoReadOptions Options;
 	/** I/O dispatcher priority (EIoDispatcherPriority). */
-	int32 Priority = 0;
+	TSAN_ATOMIC(int32) Priority = 0;
 
 	FIoRequestImpl(class FIoRequestAllocator& InAllocator)
 		: Allocator(InAllocator)
diff --git a/Engine/Source/Runtime/Core/Public/Misc/ConfigCacheIni.h b/Engine/Source/Runtime/Core/Public/Misc/ConfigCacheIni.h
--- a/Engine/Source/Runtime/Core/Public/Misc/ConfigCacheIni.h
+++ b/Engine/Source/Runtime/Core/Public/Misc/ConfigCacheIni.h
@@ -1203,7 +1203,7 @@ private:
 	friend class FConfigCacheIni;
 	
 	// when we last Found the branch to pull data from it, this is used to unload after it's been unused for some time
-	double InactiveTimer;
+	TSAN_ATOMIC(double) InactiveTimer;
 };
 
 
diff --git a/Engine/Source/Runtime/CoreUObject/Private/Serialization/AsyncLoading2.cpp b/Engine/Source/Runtime/CoreUObject/Private/Serialization/AsyncLoading2.cpp
--- a/Engine/Source/Runtime/CoreUObject/Private/Serialization/AsyncLoading2.cpp
+++ b/Engine/Source/Runtime/CoreUObject/Private/Serialization/AsyncLoading2.cpp
@@ -732,7 +732,7 @@ struct FAsyncPackageDesc2
 	// A unique request id for each external call to LoadPackage
 	int32 RequestID;
 	// Package priority
-	int32 Priority;
+	TSAN_ATOMIC(int32) Priority;
 	/** The flags that should be applied to the package */
 	EPackageFlags PackageFlags;
 #if WITH_EDITOR
diff --git a/Engine/Source/Runtime/Engine/Private/Actor.cpp b/Engine/Source/Runtime/Engine/Private/Actor.cpp
--- a/Engine/Source/Runtime/Engine/Private/Actor.cpp
+++ b/Engine/Source/Runtime/Engine/Private/Actor.cpp
@@ -124,7 +124,7 @@ FOnProcessEvent AActor::ProcessEventDelegate;
 #if (CSV_PROFILER_STATS && !UE_BUILD_SHIPPING)
 
 /** Count of total actors created */
-int32 CSVActorTotalCount = 0;
+TSAN_ATOMIC(int32) CSVActorTotalCount = 0;
 /** Map of Actor class names to count */
 TMap<FName, int32> CSVActorClassNameToCountMap;
 /** Critical section to control access to map */
diff --git a/Engine/Source/Runtime/Engine/Private/LevelTick.cpp b/Engine/Source/Runtime/Engine/Private/LevelTick.cpp
--- a/Engine/Source/Runtime/Engine/Private/LevelTick.cpp
+++ b/Engine/Source/Runtime/Engine/Private/LevelTick.cpp
@@ -1402,7 +1402,7 @@ static TAutoConsoleVariable<int32> CVarRecordActorCountsToCSVThreshold(
 
 extern FCriticalSection CSVActorClassNameToCountMapLock;
 extern TMap<FName, int32> CSVActorClassNameToCountMap;
-extern int32 CSVActorTotalCount;
+extern TSAN_ATOMIC(int32) CSVActorTotalCount;
 
 /** Add additional stats to CSV profile for tick and actor counts */
 static void RecordWorldCountsToCSV(UWorld* World, bool bDoingActorTicks)
