ParseLines: Use FWideStringView explicitly now that there is a FUtf8StringView overload

#rnx

diff --git a/Engine/Source/Runtime/Core/Private/String/ParseLines.cpp b/Engine/Source/Runtime/Core/Private/String/ParseLines.cpp
--- a/Engine/Source/Runtime/Core/Private/String/ParseLines.cpp
+++ b/Engine/Source/Runtime/Core/Private/String/ParseLines.cpp
@@ -8,7 +8,7 @@ namespace UE::String
 {
 
 template <typename CharType>
-void ParseLinesImpl(
+inline static void ParseLinesImpl(
 	const TStringView<CharType> View,
 	const TFunctionRef<void (TStringView<CharType>)> Visitor,
 	const EParseLinesOptions Options)
@@ -52,8 +52,8 @@ void ParseLinesImpl(
 }
 
 void ParseLines(
-	FStringView View,
-	TFunctionRef<void (FStringView)> Visitor,
+	FWideStringView View,
+	TFunctionRef<void (FWideStringView)> Visitor,
 	EParseLinesOptions Options)
 {
 	return ParseLinesImpl(View, Visitor, Options);
@@ -66,4 +66,5 @@ void ParseLines(
 {
 	return ParseLinesImpl(View, Visitor, Options);
 }
+
 } // UE::String
diff --git a/Engine/Source/Runtime/Core/Public/String/ParseLines.h b/Engine/Source/Runtime/Core/Public/String/ParseLines.h
--- a/Engine/Source/Runtime/Core/Public/String/ParseLines.h
+++ b/Engine/Source/Runtime/Core/Public/String/ParseLines.h
@@ -35,8 +35,8 @@ ENUM_CLASS_FLAGS(EParseLinesOptions);
  * @param Options   Flags to modify the default behavior.
  */
 CORE_API void ParseLines(
-	FStringView View,
-	TFunctionRef<void (FStringView)> Visitor,
+	FWideStringView View,
+	TFunctionRef<void (FWideStringView)> Visitor,
 	EParseLinesOptions Options = EParseLinesOptions::None);
 CORE_API void ParseLines(
 	FUtf8StringView View,
