StringView: Fix CopyString for copying an empty range

diff --git a/Engine/Source/Runtime/Core/Public/Containers/StringView.h b/Engine/Source/Runtime/Core/Public/Containers/StringView.h
--- a/Engine/Source/Runtime/Core/Public/Containers/StringView.h
+++ b/Engine/Source/Runtime/Core/Public/Containers/StringView.h
@@ -646,7 +646,10 @@ template <typename CharType>
 inline int32 TStringView<CharType>::CopyString(CharType* Dest, int32 CharCount, int32 Position) const
 {
 	const int32 CopyCount = FMath::Min(Size - Position, CharCount);
-	FMemory::Memcpy(Dest, DataPtr + Position, CopyCount * sizeof(CharType));
+	if (CopyCount > 0)
+	{
+		FMemory::Memcpy(Dest, DataPtr + Position, CopyCount * sizeof(CharType));
+	}
 	return CopyCount;
 }
 
