Add new rigid particle control flag to serialization method for archives

diff --git a/Engine/Source/Runtime/Core/Public/UObject/FortniteMainBranchObjectVersions.inl b/Engine/Source/Runtime/Core/Public/UObject/FortniteMainBranchObjectVersions.inl
--- a/Engine/Source/Runtime/Core/Public/UObject/FortniteMainBranchObjectVersions.inl
+++ b/Engine/Source/Runtime/Core/Public/UObject/FortniteMainBranchObjectVersions.inl
@@ -637,5 +637,8 @@ FORTNITE_MAIN_VERSION(PCGChangedSurfaceSamplerDefaultGridCreationMode, 207)
 // Media plate overlay composite technique replacement with holdout composite
 FORTNITE_MAIN_VERSION(MediaPlateOverlayTechniqueRemoval, 208)
 
+// Added particle flag to allow/disallow partial island sleeping in the island the particle is in
+FORTNITE_MAIN_VERSION(PerParticleFlagToAllowPartialIslandSleepInConnectedIsland, 209)
+
 // -----<new versions can be added above this line>-------------------------------------------------
 #undef FORTNITE_MAIN_VERSION
\ No newline at end of file
diff --git a/Engine/Source/Runtime/Experimental/ChaosVisualDebugger/Private/DataWrappers/ChaosVDParticleDataWrapper.cpp b/Engine/Source/Runtime/Experimental/ChaosVisualDebugger/Private/DataWrappers/ChaosVDParticleDataWrapper.cpp
--- a/Engine/Source/Runtime/Experimental/ChaosVisualDebugger/Private/DataWrappers/ChaosVDParticleDataWrapper.cpp
+++ b/Engine/Source/Runtime/Experimental/ChaosVisualDebugger/Private/DataWrappers/ChaosVDParticleDataWrapper.cpp
@@ -43,6 +43,11 @@ bool FChaosVDFRigidParticleControlFlags::Serialize(FArchive& Ar)
 		Ar << bGyroscopicTorqueEnabled;
 	}
 
+	if (Ar.CustomVer(FFortniteMainBranchObjectVersion::GUID) >= FFortniteMainBranchObjectVersion::PerParticleFlagToAllowPartialIslandSleepInConnectedIsland)
+	{
+		Ar << bPartialIslandSleepAllowed;
+	}
+
 	return !Ar.IsError();
 }
 
