Added move support to FScopedLoadingState's constructor.

diff --git a/Engine/Source/Runtime/Core/Public/CoreGlobals.h b/Engine/Source/Runtime/Core/Public/CoreGlobals.h
--- a/Engine/Source/Runtime/Core/Public/CoreGlobals.h
+++ b/Engine/Source/Runtime/Core/Public/CoreGlobals.h
@@ -773,7 +773,12 @@ CORE_API void NotifyLoadingStateChanged(bool bState, const TCHAR *Message);
 struct FScopedLoadingState
 {
 	FString Message;
-	FScopedLoadingState(const TCHAR* InMessage)
+	FScopedLoadingState(FString&& InMessage)
+		: Message(MoveTemp(InMessage))
+	{
+		NotifyLoadingStateChanged(true, *Message);
+	}
+	FScopedLoadingState(const FString& InMessage)
 		: Message(InMessage)
 	{
 		NotifyLoadingStateChanged(true, *Message);
@@ -789,7 +794,8 @@ FORCEINLINE void NotifyLoadingStateChanged(bool bState, const TCHAR *Message)
 }
 struct FScopedLoadingState
 {
-	FORCEINLINE FScopedLoadingState(const TCHAR* InMessage)
+	template <typename T>
+	FORCEINLINE FScopedLoadingState(T&& InMessage)
 	{
 	}
 };
