Add a `MallocZeroed` override to `FMallocPoisonProxy` to avoid the double memset.

diff --git a/Engine/Source/Runtime/Core/Public/HAL/MallocPoisonProxy.h b/Engine/Source/Runtime/Core/Public/HAL/MallocPoisonProxy.h
--- a/Engine/Source/Runtime/Core/Public/HAL/MallocPoisonProxy.h
+++ b/Engine/Source/Runtime/Core/Public/HAL/MallocPoisonProxy.h
@@ -42,7 +42,7 @@ public:
 
 	virtual void* Malloc(SIZE_T Size, uint32 Alignment) override
 	{
-		void* Result=UsedMalloc->Malloc(Size, Alignment);
+		void* Result = UsedMalloc->Malloc(Size, Alignment);
 		if (LIKELY(Result != nullptr && Size > 0))
 		{
 			FMemory::Memset(Result, UE_DEBUG_FILL_NEW, Size);
@@ -50,6 +50,16 @@ public:
 		return Result;
 	}
 
+	virtual void* MallocZeroed(SIZE_T Size, uint32 Alignment) override
+	{
+		void* Result = UsedMalloc->Malloc(Size, Alignment);
+		if (LIKELY(Result != nullptr && Size > 0))
+		{
+			FMemory::Memset(Result, 0, Size);
+		}
+		return Result;
+	}
+
 	virtual void* Realloc(void* Ptr, SIZE_T NewSize, uint32 Alignment) override
 	{
 		// NOTE: case when Realloc returns completely new pointer is not handled properly (we would like to have previous memory poisoned completely).
