Enable the platform I/O dispatcher with compile flag.

#rnx

diff --git a/Engine/Source/Runtime/Core/Internal/IO/IoDispatcherConfig.cpp b/Engine/Source/Runtime/Core/Internal/IO/IoDispatcherConfig.cpp
--- a/Engine/Source/Runtime/Core/Internal/IO/IoDispatcherConfig.cpp
+++ b/Engine/Source/Runtime/Core/Internal/IO/IoDispatcherConfig.cpp
@@ -5,6 +5,10 @@
 #include "HAL/IConsoleManager.h"
 #include "Misc/Parse.h"
 
+#ifndef UE_PLATFORM_IO_DISPATCHER_ENABLED
+#define UE_PLATFORM_IO_DISPATCHER_ENABLED 0
+#endif
+
 int32 GIoDispatcherBufferSizeKB = 256;
 static FAutoConsoleVariableRef CVar_IoDispatcherBufferSizeKB(
 	TEXT("s.IoDispatcherBufferSizeKB"),
@@ -113,13 +117,15 @@ static FAutoConsoleVariableRef CVar_IoDispatcherMaxResolveBatchSize(
 
 bool IsPlatformIoDispatcherEnabled()
 {
-#if !UE_BUILD_SHIPPING
 	static struct FIoDispatcherConfig
 	{
-		bool bPlatformIoDispatcherEnabled;
+		bool bPlatformIoDispatcherEnabled = false;
 		FIoDispatcherConfig(const TCHAR* CmdLine)
 		{
-			bPlatformIoDispatcherEnabled = FParse::Param(CmdLine, TEXT("PlatformIo"));
+#if UE_PLATFORM_IO_DISPATCHER_ENABLED
+			bPlatformIoDispatcherEnabled = true;
+#endif
+			bPlatformIoDispatcherEnabled |= FParse::Param(CmdLine, TEXT("PlatformIo"));
 			if (!bPlatformIoDispatcherEnabled)
 			{
 				FString Line;
@@ -129,11 +135,13 @@ bool IsPlatformIoDispatcherEnabled()
 					bPlatformIoDispatcherEnabled = Line.Contains(TEXT("s.PlatformIo 1"), ESearchCase::IgnoreCase);
 				}
 			}
+
+			if (FParse::Param(CmdLine, TEXT("NoPlatformIo")))
+			{
+				bPlatformIoDispatcherEnabled = false;
+			}
 		}
 	} IoDispatcherConfig(FCommandLine::Get());
 
 	return IoDispatcherConfig.bPlatformIoDispatcherEnabled;
-#else
-	return false;
-#endif
 }
