Fix crash when parsing platform I/O dispatcher command line args.

#rnx

diff --git a/Engine/Source/Runtime/Core/Internal/IO/IoDispatcherConfig.cpp b/Engine/Source/Runtime/Core/Internal/IO/IoDispatcherConfig.cpp
--- a/Engine/Source/Runtime/Core/Internal/IO/IoDispatcherConfig.cpp
+++ b/Engine/Source/Runtime/Core/Internal/IO/IoDispatcherConfig.cpp
@@ -125,6 +125,12 @@ bool IsPlatformIoDispatcherEnabled()
 #if UE_PLATFORM_IO_DISPATCHER_ENABLED
 			bPlatformIoDispatcherEnabled = true;
 #endif
+			if (FParse::Param(CmdLine, TEXT("NoPlatformIo")))
+			{
+				bPlatformIoDispatcherEnabled = false;
+				return;
+			}
+
 			bPlatformIoDispatcherEnabled |= FParse::Param(CmdLine, TEXT("PlatformIo"));
 			if (!bPlatformIoDispatcherEnabled)
 			{
@@ -135,11 +141,6 @@ bool IsPlatformIoDispatcherEnabled()
 					bPlatformIoDispatcherEnabled = Line.Contains(TEXT("s.PlatformIo 1"), ESearchCase::IgnoreCase);
 				}
 			}
-
-			if (FParse::Param(CmdLine, TEXT("NoPlatformIo")))
-			{
-				bPlatformIoDispatcherEnabled = false;
-			}
 		}
 	} IoDispatcherConfig(FCommandLine::Get());
 
