[Scheduler]
- Fix potential assert during taskgraph shutdown

diff --git a/Engine/Source/Runtime/Core/Private/Async/Fundamental/WaitingQueue.cpp b/Engine/Source/Runtime/Core/Private/Async/Fundamental/WaitingQueue.cpp
--- a/Engine/Source/Runtime/Core/Private/Async/Fundamental/WaitingQueue.cpp
+++ b/Engine/Source/Runtime/Core/Private/Async/Fundamental/WaitingQueue.cpp
@@ -149,7 +149,7 @@ void FWaitingQueue::FinishShutdown()
 	using namespace WaitingQueueImpl;
 
 	check((State & (StackMask | WaiterMask)) == StackMask);
-	check((StandbyState & (StackMask | WaiterMask)) == StackMask);
+	check((StandbyState & StackMask) == StackMask);
 }
 
 void FWaitingQueue::PrepareWait(FWaitEvent* Node)
