[RaceDetector]
- Fix crash and deadlock when running a game due to reentrancy in the heap allocator
- Add support for in-place reallocs to avoid breaking assumptions the API provides

diff --git a/Engine/Source/Runtime/Core/Private/Sanitizer/RaceDetectorWindows.cpp b/Engine/Source/Runtime/Core/Private/Sanitizer/RaceDetectorWindows.cpp
--- a/Engine/Source/Runtime/Core/Private/Sanitizer/RaceDetectorWindows.cpp
+++ b/Engine/Source/Runtime/Core/Private/Sanitizer/RaceDetectorWindows.cpp
@@ -564,7 +564,7 @@ INSTRUMENTATION_FUNCTION_ATTRIBUTES NTSTATUS STDAPICALLTYPE DetouredRtlFreeHeap(
 	SIZE_T Size = 0;
 	if (lpMem)
 	{
-		Size = HeapSize(hHeap, 0, lpMem);
+		Size = HeapSize(hHeap, dwFlags & HEAP_NO_SERIALIZE, lpMem);
 	}
 	
 	// Free the memory range before freeing the memory to avoid
@@ -590,14 +590,43 @@ typedef NTSYSAPI PVOID (NTAPI RTLREALLOCATEHEAP)(
 RTLREALLOCATEHEAP* TrueRtlReallocateHeap = nullptr;
 INSTRUMENTATION_FUNCTION_ATTRIBUTES PVOID NTAPI DetouredRtlReallocateHeap(PVOID hHeap, ULONG dwFlags, PVOID lpMem, ULONG dwBytes)
 {
+	// Some calls coming to RtlReallocateHeap are actually well inside the deeper functions of RtlAllocateHeap itself.
+	// It appears that we can cause recursion and/or corruption if we try to instrument those so just ignore them
+	// since these flags are not documented, it's very unlikely that a call using this flag would come from
+	// user facing code that we want to instrument.
+	if (dwFlags & 0x800000)
+	{
+		return TrueRtlReallocateHeap(hHeap, dwFlags, lpMem, dwBytes);
+	}
+
 	START_WINAPI_INSTRUMENTATION(return TrueRtlReallocateHeap(hHeap, dwFlags, lpMem, dwBytes));
 
+	if (dwFlags & HEAP_REALLOC_IN_PLACE_ONLY)
+	{
+		// If we reduce the size, technically another thread could come-in and claim the part
+        	// of the allocation that we got rid of, so free it before as it's very unlikely
+		// that RtlReallocateHeap would fail when trying to reduce the size of an alloc in place.
+		// Worst case, it would be possible to miss a race condition in that zone if the reallocate fails
+		// and the application still uses that memory after, but that is very unlikely anyway.
+		if (lpMem)
+		{
+			const SIZE_T OldSize = HeapSize(hHeap, dwFlags & HEAP_NO_SERIALIZE, lpMem);
+
+			if (OldSize > dwBytes)
+			{
+				FreeMemoryRange((void*)((UPTRINT)lpMem + dwBytes), OldSize - dwBytes);
+			}
+		}
+
+		return TrueRtlReallocateHeap(hHeap, dwFlags, lpMem, dwBytes);
+	}
+
 	// We have to always allocate new blocs in order to invalidate the old memory range before it can get reused.
 	void* NewPtr = TrueRtlAllocateHeap(hHeap, dwFlags, dwBytes);
 
 	if (lpMem)
 	{
-		const SIZE_T OldSize = HeapSize(hHeap, 0, lpMem);
+		const SIZE_T OldSize = HeapSize(hHeap, dwFlags & HEAP_NO_SERIALIZE, lpMem);
 		if (OldSize > 0)
 		{
 			if (NewPtr)
