Added DisabledCVar key to PreviewPlatform to be able to disable a preview platform in the case a cvar is enabled

#rnx

diff --git a/Engine/Source/Runtime/Core/Private/Misc/DataDrivenPlatformInfoRegistry.cpp b/Engine/Source/Runtime/Core/Private/Misc/DataDrivenPlatformInfoRegistry.cpp
--- a/Engine/Source/Runtime/Core/Private/Misc/DataDrivenPlatformInfoRegistry.cpp
+++ b/Engine/Source/Runtime/Core/Private/Misc/DataDrivenPlatformInfoRegistry.cpp
@@ -267,22 +267,31 @@ static void ParsePreviewPlatforms(const FConfigFile& IniFile)
 				const FString& SectionName = Section.Key;
 				FName PreviewPlatformName = *(SectionName.Mid(16) + TEXT("_Preview"));
 
-				// Early-out if enabled cvar is specified and not set
-				TArray<FString> Tokens;
-				GetSectionString(Section.Value, FName("EnabledCVar")).ParseIntoArray(Tokens, TEXT(":"));
-				if (Tokens.Num() == 5)
-				{
-					// now load a local version of the ini hierarchy
-					FConfigFile LocalIni;
-					FConfigCacheIni::LoadLocalIniFile(LocalIni, *Tokens[1], true, *Tokens[2]);
-
-					// and get the enabled cvar's value
-					bool bEnabled = false;
-					LocalIni.GetBool(*Tokens[3], *Tokens[4], bEnabled);
-					if (!bEnabled)
+				// Early-out if enabled / disabled cvar is specified and not set (or the other way around)
+				auto ShouldEarlyOut = [&](const FName& KeyName, bool bValueToEarlyOut) -> bool
 					{
-						continue;
-					}
+						TArray<FString> Tokens;
+						GetSectionString(Section.Value, KeyName).ParseIntoArray(Tokens, TEXT(":"));
+						if (Tokens.Num() == 5)
+						{
+							// now load a local version of the ini hierarchy
+							FConfigFile LocalIni;
+							FConfigCacheIni::LoadLocalIniFile(LocalIni, *Tokens[1], true, *Tokens[2]);
+
+							// and get the enabled cvar's value
+							bool bEnabled = false;
+							LocalIni.GetBool(*Tokens[3], *Tokens[4], bEnabled);
+							if (bEnabled == bValueToEarlyOut)
+							{
+								return true;
+							}
+						}
+						return false;
+					};
+					
+				if (ShouldEarlyOut(FName("EnabledCVar"), false) || ShouldEarlyOut(FName("DisabledCVar"), true))
+				{
+					continue;
 				}
 
 				FName PlatformName = *GetSectionString(Section.Value, FName("PlatformName"));
