[ConsumeAllMpmcQueue]
- Reduce destructor cost by just verifying the Head is empty instead of calling exchange(0) on it

diff --git a/Engine/Source/Runtime/Core/Public/Containers/ConsumeAllMpmcQueue.h b/Engine/Source/Runtime/Core/Public/Containers/ConsumeAllMpmcQueue.h
--- a/Engine/Source/Runtime/Core/Public/Containers/ConsumeAllMpmcQueue.h
+++ b/Engine/Source/Runtime/Core/Public/Containers/ConsumeAllMpmcQueue.h
@@ -41,7 +41,10 @@ namespace UE
 		~TConsumeAllMpmcQueue()
 		{
 			static_assert(std::is_trivially_destructible_v<FNode>);
-			ConsumeAllLifo([](T&&){});
+			if (Head.load(std::memory_order_acquire) != nullptr)
+			{
+				ConsumeAllLifo([](T&&) {});
+			}
 		}
 
 		//Push an Item to the Queue and 
