Fix PVS warnings

#rnx

diff --git a/Engine/Source/Runtime/Core/Private/Modules/ModuleManifest.cpp b/Engine/Source/Runtime/Core/Private/Modules/ModuleManifest.cpp
--- a/Engine/Source/Runtime/Core/Private/Modules/ModuleManifest.cpp
+++ b/Engine/Source/Runtime/Core/Private/Modules/ModuleManifest.cpp
@@ -70,7 +70,12 @@ bool FModuleManifest::TryRead(const FString& FileName, FModuleManifest& OutManif
 		return FSimpleParse::MatchChar(Ptr, TEXT(','));
 	};
 
-	if (!Whitespace() || !ObjectStart() || !Whitespace() || ObjectEnd())
+	if (!Whitespace() || !ObjectStart())
+	{
+		return false;
+	}
+
+	if (!Whitespace() || ObjectEnd())
 	{
 		return false;
 	}
@@ -83,7 +88,12 @@ bool FModuleManifest::TryRead(const FString& FileName, FModuleManifest& OutManif
 			return false;
 		}
 
-		if (!Whitespace() || !Colon() || !Whitespace())
+		if (!Whitespace() || !Colon())
+		{
+			return false;
+		}
+
+		if (!Whitespace())
 		{
 			return false;
 		}
@@ -99,7 +109,12 @@ bool FModuleManifest::TryRead(const FString& FileName, FModuleManifest& OutManif
 		// Modules is an array of ModuleName:ModulePath pairs
 		else if (Field.ToView() == TEXTVIEW("Modules"))
 		{
-			if (!Whitespace() || !ObjectStart() || !Whitespace())
+			if (!Whitespace() || !ObjectStart())
+			{
+				return false;
+			}
+
+			if (!Whitespace())
 			{
 				return false;
 			}
@@ -133,7 +148,12 @@ bool FModuleManifest::TryRead(const FString& FileName, FModuleManifest& OutManif
 		// LibraryDependencies is an array of LibraryName keys to arrays of DependencyName values
 		else if (Field.ToView() == TEXTVIEW("LibraryDependencies"))
 		{
-			if (!Whitespace() || !ObjectStart() || !Whitespace())
+			if (!Whitespace() || !ObjectStart())
+			{
+				return false;
+			}
+
+			if (!Whitespace())
 			{
 				return false;
 			}
@@ -143,7 +163,17 @@ bool FModuleManifest::TryRead(const FString& FileName, FModuleManifest& OutManif
 				for (;;)
 				{
 					TStringBuilder<64> LibraryName;
-					if (!FSimpleParse::ParseString(Ptr, LibraryName) || !Whitespace() || !Colon() || !Whitespace() || !ArrayStart() || !Whitespace())
+					if (!FSimpleParse::ParseString(Ptr, LibraryName) || !Whitespace() || !Colon())
+					{
+						return false;
+					}
+
+					if (!Whitespace() || !ArrayStart())
+					{
+						return false;
+					}
+
+					if (!Whitespace())
 					{
 						return false;
 					}
