Don't clamp I/O dispatcher priority until assigning to prio queues.

#rnx

diff --git a/Engine/Source/Runtime/Core/Internal/IO/PlatformIoDispatcher.h b/Engine/Source/Runtime/Core/Internal/IO/PlatformIoDispatcher.h
--- a/Engine/Source/Runtime/Core/Internal/IO/PlatformIoDispatcher.h
+++ b/Engine/Source/Runtime/Core/Internal/IO/PlatformIoDispatcher.h
@@ -21,36 +21,6 @@ CORE_API DECLARE_LOG_CATEGORY_EXTERN(LogPlatformIoDispatcher, Log, All);
 namespace UE
 {
 
-/** Priority levels when reading encoded I/O store chunks from file. */
-enum class EIoFileReadPriority : uint8
-{
-	Min,
-	Low,
-	Medium,
-	High,
-	Max,
-	Count
-};
-
-/** Narrows an I/O dispatcher priorty to file read priority. */
-static inline EIoFileReadPriority IoFileReadPriorityFromDispatcherPriority(int32 Priority)
-{
-	if (Priority < IoDispatcherPriority_Low)
-	{
-		return EIoFileReadPriority::Min;
-	}
-	if (Priority < IoDispatcherPriority_Medium)
-	{
-		return EIoFileReadPriority::Low;
-	}
-	if (Priority < IoDispatcherPriority_High)
-	{
-		return EIoFileReadPriority::Medium;
-	}
-
-	return Priority > IoDispatcherPriority_High ? EIoFileReadPriority::Max : EIoFileReadPriority::High;
-}
-
 /** Handle to a file on disk. */
 struct FIoFileHandle
 {
@@ -154,7 +124,7 @@ public:
 		FIoBuffer& InDestination,
 		uint64 InDestinationSize,
 		void* InUserData,
-		EIoFileReadPriority InPriority)
+		int32 InPriority)
 			: FileHandle(InFileHandle)
 			, Destination(InDestination)
 			, DestinationSize(InDestinationSize)
@@ -207,7 +177,7 @@ public:
 	FIoBuffer&				Destination;
 	uint64					DestinationSize;
 	void*					UserData;
-	EIoFileReadPriority		Priority;
+	int32					Priority;
 	FScatterParamsArray		Params;
 };
 
@@ -256,7 +226,7 @@ public:
 	/** Tick the dispatcher. Only used when multithreading is disabled. */ 
 	virtual bool								Tick() = 0;
 	/** Update priority for an inflight I/O request. */ 
-	virtual void								UpdatePriority(FIoFileReadRequest Request, EIoFileReadPriority NewPriority) = 0;
+	virtual void								UpdatePriority(FIoFileReadRequest Request, int32 NewPriority) = 0;
 	/** Cancel inflight I/O request. */ 
 	virtual void								CancelRequest(FIoFileReadRequest Request) = 0;
 	/** Cancel all inflight I/O request for the specified file. */ 
@@ -276,33 +246,3 @@ public:
 };
 
 } // namespace UE
-
-////////////////////////////////////////////////////////////////////////////////
-static inline bool operator>(UE::EIoFileReadPriority LHS, UE::EIoFileReadPriority RHS)
-{
-	return static_cast<uint8>(LHS) > static_cast<uint8>(RHS);
-}
-
-static inline bool operator<(UE::EIoFileReadPriority LHS, UE::EIoFileReadPriority RHS)
-{
-	return static_cast<uint8>(LHS) < static_cast<uint8>(RHS);
-}
-
-static inline bool operator<=(UE::EIoFileReadPriority LHS, UE::EIoFileReadPriority RHS)
-{
-	return static_cast<uint8>(LHS) <= static_cast<uint8>(RHS);
-}
-
-static inline const TCHAR* LexToString(UE::EIoFileReadPriority Priority)
-{
-	static const TCHAR* StringTable[] =
-	{
-		TEXT("Min"),
-		TEXT("Low"),
-		TEXT("Medium"),
-		TEXT("High"),
-		TEXT("Max"),
-	};
-
-	return StringTable[static_cast<uint8>(Priority)];
-};
diff --git a/Engine/Source/Runtime/Core/Internal/IO/PlatformIoDispatcherBase.cpp b/Engine/Source/Runtime/Core/Internal/IO/PlatformIoDispatcherBase.cpp
--- a/Engine/Source/Runtime/Core/Internal/IO/PlatformIoDispatcherBase.cpp
+++ b/Engine/Source/Runtime/Core/Internal/IO/PlatformIoDispatcherBase.cpp
@@ -48,8 +48,9 @@ void FIoFileBlockRequestQueue::Enqueue(FIoFileBlockRequestList&& FileBlockReques
 		Request->QueueStatus = FIoFileBlockRequest::EQueueStatus::Enqueued;
 		if (bSortByOffset)
 		{
-			FPrioQueue& PrioQueue	= PrioQueues[static_cast<int32>(Request->Priority)];
-			const int32 Idx			= Algo::UpperBoundBy(PrioQueue.ByOffset, ToFileOffset(Request), ToFileOffset, FileOffsetLess);
+			const EIoFileReadPriority Priority	= IoFileReadPriorityFromDispatcherPriority(Request->Priority);
+			FPrioQueue& PrioQueue				= PrioQueues[static_cast<int32>(Priority)];
+			const int32 Idx						= Algo::UpperBoundBy(PrioQueue.ByOffset, ToFileOffset(Request), ToFileOffset, FileOffsetLess);
 
 			PrioQueue.ByOffset.Insert(Request, Idx);
 			PrioQueue.BySequence.AddTail(Request);
@@ -122,8 +123,10 @@ void FIoFileBlockRequestQueue::Reprioritize()
 
 			for (int32 RequestIdx = PrioQueue.ByOffset.Num() - 1; RequestIdx >= 0; --RequestIdx)
 			{
-				FIoFileBlockRequest* Request = PrioQueue.ByOffset[RequestIdx];
-				if (Request->Priority != QueuePriority)
+				FIoFileBlockRequest* Request		= PrioQueue.ByOffset[RequestIdx];
+				const EIoFileReadPriority Priority	= IoFileReadPriorityFromDispatcherPriority(Request->Priority);
+
+				if (Priority != QueuePriority)
 				{
 					RequestsToReprioritize.Add(Request);
 					PrioQueue.ByOffset.RemoveAt(RequestIdx, EAllowShrinking::No);
@@ -134,8 +137,9 @@ void FIoFileBlockRequestQueue::Reprioritize()
 
 		for (FIoFileBlockRequest* Request : RequestsToReprioritize)
 		{
-			FPrioQueue& PrioQueue	= PrioQueues[static_cast<int32>(Request->Priority)];
-			const int32 Idx			= Algo::UpperBoundBy(PrioQueue.ByOffset, ToFileOffset(Request), ToFileOffset, FileOffsetLess);
+			const EIoFileReadPriority Priority	= IoFileReadPriorityFromDispatcherPriority(Request->Priority);
+			FPrioQueue& PrioQueue				= PrioQueues[static_cast<int32>(Priority)];
+			const int32 Idx						= Algo::UpperBoundBy(PrioQueue.ByOffset, ToFileOffset(Request), ToFileOffset, FileOffsetLess);
 
 			PrioQueue.ByOffset.Insert(Request, Idx);
 			PrioQueue.BySequence.AddTail(Request);
@@ -177,14 +181,15 @@ FIoFileBlockRequest* FIoFileBlockRequestQueue::GetByOffset(FPrioQueue& PrioQueue
 		if (GIoDispatcherRequestLatencyCircuitBreakerMS > 0)
 		{
 			// If our oldest request has been unserviced for too long, grab that instead of the next sequential read
-			const uint64 ThresholdCycles = uint64((GIoDispatcherRequestLatencyCircuitBreakerMS * 1000.0) / FPlatformTime::GetSecondsPerCycle64());
+			const uint64 ThresholdCycles = FPlatformTime::SecondsToCycles64(GIoDispatcherRequestLatencyCircuitBreakerMS * 1000.0);
 			bHeadRequestTooOld = (FPlatformTime::Cycles64() - PrioQueue.BySequence.PeekHead()->TimeCreated) >= ThresholdCycles;
 		}
 
+		const EIoFileReadPriority LastFileOffsetPriority = IoFileReadPriorityFromDispatcherPriority(LastFileOffset.Priority);
 		const bool bChooseByOffset = 
 				LastFileOffset.FileHandle.IsValid()
 			&&	!bHeadRequestTooOld 
-			&&  (GIoDispatcherMaintainSortingOnPriorityChange || LastFileOffset.Priority == QueuePriority);
+			&&  (GIoDispatcherMaintainSortingOnPriorityChange || LastFileOffsetPriority == QueuePriority);
 		if (bChooseByOffset)
 		{
 			RequestIndex = Algo::LowerBoundBy(PrioQueue.ByOffset, LastFileOffset, ToFileOffset, FileOffsetLess);
@@ -1039,7 +1044,7 @@ void FPlatformIoDispatcherBase::CancelAllRequests(FIoFileHandle FileHandle)
 	}
 }
 
-void FPlatformIoDispatcherBase::UpdatePriority(FIoFileReadRequest Request, EIoFileReadPriority NewPriority)
+void FPlatformIoDispatcherBase::UpdatePriority(FIoFileReadRequest Request, int32 NewPriority)
 {
 	if (Request.IsValid() == false)
 	{
@@ -1057,7 +1062,7 @@ void FPlatformIoDispatcherBase::UpdatePriority(FIoFileReadRequest Request, EIoFi
 
 		for (FIoFileBlockRequest* FileBlockRequest : FileBlockRequests)
 		{
-			if (NewPriority > FileBlockRequest->Priority)
+			if (NewPriority != FileBlockRequest->Priority)
 			{
 				FileBlockRequest->Priority	= NewPriority;
 				bReprioritize				= true;
diff --git a/Engine/Source/Runtime/Core/Internal/IO/PlatformIoDispatcherBase.h b/Engine/Source/Runtime/Core/Internal/IO/PlatformIoDispatcherBase.h
--- a/Engine/Source/Runtime/Core/Internal/IO/PlatformIoDispatcherBase.h
+++ b/Engine/Source/Runtime/Core/Internal/IO/PlatformIoDispatcherBase.h
@@ -11,6 +11,7 @@
 #include "IO/IoBuffer.h"
 #include "IO/IoChunkEncoding.h"
 #include "IO/IoContainers.h"
+#include "IO/IoDispatcherPriority.h"
 #include "IO/PlatformIoDispatcher.h"
 #include "IO/PlatformIoDispatcherStats.h"
 #include "IO/IoStatus.h"
@@ -23,6 +24,36 @@
 namespace UE
 {
 
+////////////////////////////////////////////////////////////////////////////////
+enum class EIoFileReadPriority : uint8
+{
+	Min,
+	Low,
+	Medium,
+	High,
+	Max,
+	Count
+};
+
+////////////////////////////////////////////////////////////////////////////////
+static inline EIoFileReadPriority IoFileReadPriorityFromDispatcherPriority(int32 Priority)
+{
+	if (Priority < IoDispatcherPriority_Low)
+	{
+		return EIoFileReadPriority::Min;
+	}
+	if (Priority < IoDispatcherPriority_Medium)
+	{
+		return EIoFileReadPriority::Low;
+	}
+	if (Priority < IoDispatcherPriority_High)
+	{
+		return EIoFileReadPriority::Medium;
+	}
+
+	return Priority > IoDispatcherPriority_High ? EIoFileReadPriority::Max : EIoFileReadPriority::High;
+}
+
 ////////////////////////////////////////////////////////////////////////////////
 template<typename T, SIZE_T Size = 8>
 using FTempArray = TArray<T, TInlineAllocator<Size>>;
@@ -237,7 +268,7 @@ struct FIoFileBlockRequest
 	uint32							SeqNo = NextSeqNo++;
 	FIoBufferHandle					BufferHandle;
 	EIoErrorCode					ErrorCode = EIoErrorCode::Unknown;
-	EIoFileReadPriority				Priority = EIoFileReadPriority::Low;
+	int32							Priority = 0;
 	EQueueStatus					QueueStatus = EQueueStatus::None;
 };
 
@@ -257,9 +288,9 @@ class FIoFileBlockRequestQueue
 {
 	struct FFileOffset
 	{
-		FIoFileHandle		FileHandle;
-		uint64				FileOffset = MAX_uint64;
-		EIoFileReadPriority	Priority = EIoFileReadPriority::Min;
+		FIoFileHandle	FileHandle;
+		uint64			FileOffset = MAX_uint64;
+		int32			Priority = 0;
 	};
 
 	struct FPrioQueue
@@ -294,7 +325,7 @@ class FIoFileBlockRequestQueue
 		{
 			return LHS.SeqNo < RHS.SeqNo;
 		}
-		return static_cast<uint8>(LHS.Priority) > static_cast<uint8>(RHS.Priority);
+		return LHS.Priority > RHS.Priority;
 	}
 
 public:
@@ -485,7 +516,7 @@ public:
 	virtual FIoFileReadRequest			ScatterGather(FIoScatterGatherRequestParams&& Params, FIoFileReadRequestCompleted&& OnCompleted) override;
 	virtual FIoFileReadRequest			ReadDirect(FIoDirectReadRequestParams&& Params, FIoFileReadRequestCompleted&& OnCompleted) override;
 	virtual bool						Tick() override { ensure(!bMultithreaded); return false; }
-	virtual void						UpdatePriority(FIoFileReadRequest Request, EIoFileReadPriority NewPriority) override;
+	virtual void						UpdatePriority(FIoFileReadRequest Request, int32 NewPriority) override;
 	virtual void						CancelRequest(FIoFileReadRequest Request) override;
 	virtual void						CancelAllRequests(FIoFileHandle FileHandle) override;
 	virtual void						DeleteRequest(FIoFileReadRequest Request) override;
diff --git a/Engine/Source/Runtime/Core/Tests/IO/PlatformIoDispatcherTests.cpp b/Engine/Source/Runtime/Core/Tests/IO/PlatformIoDispatcherTests.cpp
--- a/Engine/Source/Runtime/Core/Tests/IO/PlatformIoDispatcherTests.cpp
+++ b/Engine/Source/Runtime/Core/Tests/IO/PlatformIoDispatcherTests.cpp
@@ -122,11 +122,11 @@ TEST_CASE("Core::IO::Platform::FileReadQueue", "[Core][IO]")
 
 		// Act
 		Queue.Enqueue(MoveTemp(ToEnqueue));
-		Enqueued[0].Priority = EIoFileReadPriority::Min;
-		Enqueued[1].Priority = EIoFileReadPriority::Low;
-		Enqueued[2].Priority = EIoFileReadPriority::Medium;
-		Enqueued[3].Priority = EIoFileReadPriority::High;
-		Enqueued[4].Priority = EIoFileReadPriority::Max;
+		Enqueued[0].Priority = IoDispatcherPriority_Min;
+		Enqueued[1].Priority = IoDispatcherPriority_Low;
+		Enqueued[2].Priority = IoDispatcherPriority_Medium;
+		Enqueued[3].Priority = IoDispatcherPriority_High;
+		Enqueued[4].Priority = IoDispatcherPriority_Max;
 		Queue.Reprioritize();
 
 		TArray<FIoFileBlockRequest*> Dequeued;
diff --git a/Engine/Source/Runtime/PakFile/Private/FileIoDispatcherBackend.cpp b/Engine/Source/Runtime/PakFile/Private/FileIoDispatcherBackend.cpp
--- a/Engine/Source/Runtime/PakFile/Private/FileIoDispatcherBackend.cpp
+++ b/Engine/Source/Runtime/PakFile/Private/FileIoDispatcherBackend.cpp
@@ -626,7 +626,6 @@ struct FResolvedRequest
 	FIoFileHandle		FileHandle;
 	uint64				Offset = MAX_uint64;
 	uint64				Size = MAX_uint64;
-	EIoFileReadPriority Priority = EIoFileReadPriority::Medium;
 };
 
 using FRequestAllocator = TSingleThreadedSlabAllocator<FResolvedRequest>;
@@ -733,7 +732,6 @@ void FFileIoDispatcherBackend::ResolveIoRequests(FIoRequestList Requests, FIoReq
 				ResolvedRequest->ChunkInfo			= ChunkInfo;
 				ResolvedRequest->Offset				= ChunkInfo.Offset() + DispatcherRequest->Options.GetOffset();
 				ResolvedRequest->Size				= ResolvedSize; 
-				ResolvedRequest->Priority			= IoFileReadPriorityFromDispatcherPriority(DispatcherRequest->Priority);
 				check(DispatcherRequest->BackendData != nullptr);
 
 				ResolvedRequests.AddTail(DispatcherRequest);
@@ -812,7 +810,7 @@ void FFileIoDispatcherBackend::ResolveIoRequests(FIoRequestList Requests, FIoReq
 			}
 		}
 
-		FIoScatterGatherRequestParams ScatterGather(Partition.FileHandle, Dst, ResolvedRequest.Size, DispatcherRequest, ResolvedRequest.Priority);
+		FIoScatterGatherRequestParams ScatterGather(Partition.FileHandle, Dst, ResolvedRequest.Size, DispatcherRequest, DispatcherRequest->Priority);
 
 		// Scatter offsets
 		uint64 RequestRemainingBytes	= ResolvedRequest.Size;
@@ -948,11 +946,7 @@ void FFileIoDispatcherBackend::UpdatePriorityForIoRequest(FIoRequestImpl* Dispat
 {
 	if (FResolvedRequest* ResolvedRequest = FResolvedRequest::TryGet(DispatcherRequest))
 	{
-		const EIoFileReadPriority NewPriority = IoFileReadPriorityFromDispatcherPriority(DispatcherRequest->Priority);
-		if (static_cast<uint8>(NewPriority) > static_cast<uint8>(ResolvedRequest->Priority))
-		{
-			FPlatformIoDispatcher::Get().UpdatePriority(ResolvedRequest->PlatformRequest, NewPriority);
-		}
+		FPlatformIoDispatcher::Get().UpdatePriority(ResolvedRequest->PlatformRequest, DispatcherRequest->Priority);
 	}
 }
 
