Logging: Added FOutputDeviceFence to support lifetime management of dynamic log records

diff --git a/Engine/Source/Runtime/Core/Private/Misc/OutputDeviceRedirector.cpp b/Engine/Source/Runtime/Core/Private/Misc/OutputDeviceRedirector.cpp
--- a/Engine/Source/Runtime/Core/Private/Misc/OutputDeviceRedirector.cpp
+++ b/Engine/Source/Runtime/Core/Private/Misc/OutputDeviceRedirector.cpp
@@ -67,7 +67,7 @@ struct FOutputDeviceLine
 	FOutputDeviceLine(const FOutputDeviceLine&) = delete;
 	FOutputDeviceLine& operator=(const FOutputDeviceLine&) = delete;
 
-	FORCEINLINE FOutputDeviceLine(FOutputDeviceLine&& Other)
+	inline FOutputDeviceLine(FOutputDeviceLine&& Other)
 		: Time(Other.Time)
 		, Data(Other.Data)
 		, Category(Other.Category)
@@ -76,7 +76,7 @@ struct FOutputDeviceLine
 		Other.Data = nullptr;
 	}
 
-	FORCEINLINE FOutputDeviceLine(const TCHAR* const InData, const FName& InCategory, const ELogVerbosity::Type InVerbosity, const double InTime)
+	inline FOutputDeviceLine(const TCHAR* const InData, const FName& InCategory, const ELogVerbosity::Type InVerbosity, const double InTime)
 		: Time(InTime)
 		, Data(CopyData(InData))
 		, Category(InCategory)
@@ -84,7 +84,7 @@ struct FOutputDeviceLine
 	{
 	}
 
-	FORCEINLINE ~FOutputDeviceLine()
+	inline ~FOutputDeviceLine()
 	{
 		if (Data)
 		{
@@ -93,7 +93,7 @@ struct FOutputDeviceLine
 	}
 
 private:
-	FORCEINLINE static const TCHAR* CopyData(const TCHAR* const InData)
+	inline static const TCHAR* CopyData(const TCHAR* const InData)
 	{
 		const int32 Len = FCString::Strlen(InData) + 1;
 		void* const Dest = FOutputDeviceLinearAllocator::Malloc(sizeof(TCHAR) * Len, alignof(TCHAR));
@@ -103,18 +103,23 @@ private:
 
 struct FOutputDeviceItem
 {
-	TVariant<FOutputDeviceLine, FLogRecord> Value;
+	TVariant<FOutputDeviceLine, FLogRecord, uint64> Value;
 
-	FORCEINLINE FOutputDeviceItem(const TCHAR* const Data, const FName& Category, const ELogVerbosity::Type Verbosity, const double Time)
+	inline FOutputDeviceItem(const TCHAR* const Data, const FName& Category, const ELogVerbosity::Type Verbosity, const double Time)
 		: Value(TInPlaceType<FOutputDeviceLine>(), Data, Category, Verbosity, Time)
 	{
 	}
 
-	FORCEINLINE explicit FOutputDeviceItem(const FLogRecord& Record)
+	inline explicit FOutputDeviceItem(const FLogRecord& Record)
 		: Value(TInPlaceType<FLogRecord>(), Record)
 	{
 	}
 
+	inline explicit FOutputDeviceItem(uint64 FenceIndex)
+		: Value(TInPlaceType<uint64>(), FenceIndex)
+	{
+	}
+
 	FOutputDeviceItem(FOutputDeviceItem&&) = default;
 	FOutputDeviceItem& operator=(FOutputDeviceItem&&) = default;
 
@@ -160,6 +165,15 @@ struct FOutputDeviceRedirectorState
 	/** An event to wake the dedicated primary thread to process buffered items. */
 	FEventCount ThreadWakeEvent;
 
+	/** An event to wake threads waiting on fences. */
+	FEventCount FenceCompleteEvent;
+
+	/** The index of the last created fence. */
+	std::atomic<uint64> FenceCreateIndex = 0;
+
+	/** The index of the last completed fence. Never decreases even when fences complete out of order. */
+	std::atomic<uint64> FenceCompleteIndex = 0;
+
 	/** The ID of the thread holding the primary lock. */
 	std::atomic<uint32> LockedThreadId = MAX_uint32;
 
@@ -212,8 +226,12 @@ struct FOutputDeviceRedirectorState
 
 	void ThreadLoop();
 
+	void ConditionalFlushBufferedItems(EOutputDeviceRedirectorFlushOptions Options);
 	void FlushBufferedItems();
 
+	uint64 CreateFence();
+	void WaitFence(uint64 Index);
+
 	template <typename OutputDevicesType, typename FunctionType, typename... ArgTypes>
 	FORCEINLINE void BroadcastTo(
 		const uint32 ThreadId,
@@ -281,7 +299,7 @@ private:
 		// A write lock has set the LSB. Cancel this read lock and wait for the write.
 		State.OutputDevicesLockState.fetch_sub(2, std::memory_order_relaxed);
 		// This read lock will wait until the write lock exits.
-		UE::TReadScopeLock ScopeLock(State.OutputDevicesLock);
+		TReadScopeLock ScopeLock(State.OutputDevicesLock);
 		// This is relaxed because locking OutputDevicesLock has acquire semantics.
 		uint32 LockState = State.OutputDevicesLockState.fetch_add(2, std::memory_order_relaxed);
 		check((LockState & 1) == 0);
@@ -413,7 +431,7 @@ void FOutputDeviceRedirectorState::AddOutputDevice(FOutputDevice* OutputDevice)
 			Flags.Add(OutputDevice->CanBeUsedOnPanicThread());
 		}
 	};
-	UE::Private::FOutputDevicesWriteScopeLock ScopeLock(*this);
+	FOutputDevicesWriteScopeLock ScopeLock(*this);
 	if (OutputDevice->CanBeUsedOnMultipleThreads())
 	{
 		AddTo(UnbufferedOutputDevices, UnbufferedOutputDevicesCanBeUsedOnPanicThread);
@@ -435,7 +453,7 @@ void FOutputDeviceRedirectorState::RemoveOutputDevice(FOutputDevice* OutputDevic
 			Flags.RemoveAt(Index);
 		}
 	};
-	UE::Private::FOutputDevicesWriteScopeLock ScopeLock(*this);
+	FOutputDevicesWriteScopeLock ScopeLock(*this);
 	RemoveFrom(BufferedOutputDevices, BufferedOutputDevicesCanBeUsedOnPanicThread);
 	RemoveFrom(UnbufferedOutputDevices, UnbufferedOutputDevicesCanBeUsedOnPanicThread);
 }
@@ -491,11 +509,29 @@ void FOutputDeviceRedirectorState::ThreadLoop()
 }
 
 UE_AUTORTFM_NOAUTORTFM
-void FOutputDeviceRedirectorState::FlushBufferedItems()
+void FOutputDeviceRedirectorState::ConditionalFlushBufferedItems(EOutputDeviceRedirectorFlushOptions Options)
 {
-	using namespace UE;
-	using namespace UE::Private;
+	if (FReadScopeLock Lock(ThreadLock); bThreadStarted.load(std::memory_order_relaxed))
+	{
+		if (!EnumHasAnyFlags(Options, EOutputDeviceRedirectorFlushOptions::Async))
+		{
+			TRACE_CPUPROFILER_EVENT_SCOPE(FOutputDeviceRedirector::FlushThreadedLogs);
+			FEventCountToken Token = ThreadIdleEvent.PrepareWait();
+			ThreadWakeEvent.Notify();
+			ThreadIdleEvent.Wait(Token);
+		}
+		return;
+	}
 
+	if (FOutputDevicesPrimaryScopeLock Lock(*this); Lock.IsLocked())
+	{
+		FlushBufferedItems();
+	}
+}
+
+UE_AUTORTFM_NOAUTORTFM
+void FOutputDeviceRedirectorState::FlushBufferedItems()
+{
 	if (BufferedItems.IsEmpty())
 	{
 		return;
@@ -509,9 +545,10 @@ void FOutputDeviceRedirectorState::FlushBufferedItems()
 	{
 		if (!bIsPanicThread && HasPanicThread())
 		{
-			// Enqueuing during Deplete is okay because the items passed to the Deplete consumer
-			// are a moved list that is no longer associated with the queue; Deplete will return
-			// even though new items are in the queue.
+			// Re-queue items to be consumed by the panic thread.
+			// ProduceItem within ConsumeAllFifo is okay because the items passed to the consumer
+			// are from a moved list that is no longer associated with the queue.
+			// ConsumeAllFifo will return even though new items are in the queue.
 			BufferedItems.ProduceItem(MoveTemp(Item));
 		}
 		else
@@ -532,11 +569,55 @@ void FOutputDeviceRedirectorState::FlushBufferedItems()
 					BroadcastTo(ThreadId, BufferedOutputDevices, BufferedOutputDevicesCanBeUsedOnPanicThread,
 						UE_PROJECTION_MEMBER(FOutputDevice, SerializeRecord), Record);
 				}
+				else if constexpr (std::is_same_v<ValueType, uint64>)
+				{
+					const uint64 FenceIndex = Value;
+					for (uint64 LastFenceIndex = FenceCompleteIndex.load(std::memory_order_relaxed);;)
+					{
+						if (FenceIndex <= LastFenceIndex || FenceCompleteIndex.compare_exchange_weak(LastFenceIndex, FenceIndex, std::memory_order_release))
+						{
+							break;
+						}
+					}
+					FenceCompleteEvent.Notify();
+				}
 			}, Item.Value);
 		}
 	});
 }
 
+UE_AUTORTFM_NOAUTORTFM
+uint64 FOutputDeviceRedirectorState::CreateFence()
+{
+	uint64 Index = FenceCreateIndex.fetch_add(1, std::memory_order_relaxed) + 1;
+	if (BufferedItems.ProduceItem(Index) == EConsumeAllMpmcQueueResult::WasEmpty)
+	{
+		ThreadWakeEvent.Notify();
+	}
+	return Index;
+}
+
+UE_AUTORTFM_NOAUTORTFM
+void FOutputDeviceRedirectorState::WaitFence(uint64 Index)
+{
+	for (;;)
+	{
+		if (Index <= FenceCompleteIndex.load(std::memory_order_acquire))
+		{
+			return;
+		}
+
+		ConditionalFlushBufferedItems(EOutputDeviceRedirectorFlushOptions::Async);
+
+		FEventCountToken Token = FenceCompleteEvent.PrepareWait();
+		if (Index <= FenceCompleteIndex.load(std::memory_order_acquire))
+		{
+			return;
+		}
+		FenceCompleteEvent.Wait(Token);
+	}
+}
+
 UE_AUTORTFM_NOAUTORTFM
 void FOutputDeviceRedirectorState::AddToBacklog(const FLogRecord& Record)
 {
@@ -630,22 +711,7 @@ bool FOutputDeviceRedirector::IsRedirectingTo(FOutputDevice* OutputDevice)
 UE_AUTORTFM_ALWAYS_OPEN
 void FOutputDeviceRedirector::FlushThreadedLogs(EOutputDeviceRedirectorFlushOptions Options)
 {
-	if (FReadScopeLock ThreadLock(State->ThreadLock); State->bThreadStarted.load(std::memory_order_relaxed))
-	{
-		if (!EnumHasAnyFlags(Options, EOutputDeviceRedirectorFlushOptions::Async))
-		{
-			TRACE_CPUPROFILER_EVENT_SCOPE(FOutputDeviceRedirector::FlushThreadedLogs);
-			UE::FEventCountToken Token = State->ThreadIdleEvent.PrepareWait();
-			State->ThreadWakeEvent.Notify();
-			State->ThreadIdleEvent.Wait(Token);
-		}
-		return;
-	}
-
-	if (UE::Private::FOutputDevicesPrimaryScopeLock Lock(*State); Lock.IsLocked())
-	{
-		State->FlushBufferedItems();
-	}
+	State->ConditionalFlushBufferedItems(Options);
 }
 
 UE_AUTORTFM_ALWAYS_OPEN
@@ -812,6 +878,15 @@ void FOutputDeviceRedirector::RedirectLog(const FLazyName& Category, ELogVerbosi
 	Serialize(Data, Verbosity, Category, -1.0);
 }
 
+UE_AUTORTFM_NOAUTORTFM
+FOutputDeviceFence FOutputDeviceRedirector::CreateFence()
+{
+	FOutputDeviceFence Fence;
+	Fence.State = State.Get();
+	Fence.Index = State->CreateFence();
+	return Fence;
+}
+
 UE_AUTORTFM_ALWAYS_OPEN
 void FOutputDeviceRedirector::Flush()
 {
@@ -914,3 +989,13 @@ CORE_API FOutputDeviceRedirector* GetGlobalLogSingleton()
 {
 	return GWasGLogDestroyed ? nullptr : FOutputDeviceRedirector::Get();
 }
+
+void FOutputDeviceFence::Wait() const
+{
+	State->WaitFence(Index);
+}
+
+bool FOutputDeviceFence::Poll() const
+{
+	return Index <= State->FenceCompleteIndex.load(std::memory_order_relaxed);
+}
diff --git a/Engine/Source/Runtime/Core/Public/Misc/OutputDeviceRedirector.h b/Engine/Source/Runtime/Core/Public/Misc/OutputDeviceRedirector.h
--- a/Engine/Source/Runtime/Core/Public/Misc/OutputDeviceRedirector.h
+++ b/Engine/Source/Runtime/Core/Public/Misc/OutputDeviceRedirector.h
@@ -17,6 +17,7 @@
 FOutputDeviceRedirector.
 -----------------------------------------------------------------------------*/
 
+class FOutputDeviceFence;
 namespace UE::Private { struct FOutputDeviceRedirectorState; }
 
 /** The type of lines buffered by secondary threads. */
@@ -121,7 +122,7 @@ public:
 	/**
 	 * Serializes the log record via all current output devices.
 	 *
-	 * The format string pointed to by the record must remain valid indefinitely.
+	 * All pointers in the record must remain valid until completion of a fence created after this call.
 	 */
 	CORE_API void SerializeRecord(const UE::FLogRecord& Record) final;
 
@@ -140,11 +141,18 @@ public:
 	 * @param Event  Event name used for suppression purposes.
 	 */
 	CORE_API void Serialize(const TCHAR* Data, ELogVerbosity::Type Verbosity, const FName& Category) final;
-	
+
 	/** Same as Serialize(). */
 	CORE_API void RedirectLog(const FName& Category, ELogVerbosity::Type Verbosity, const TCHAR* Data);
 	CORE_API void RedirectLog(const FLazyName& Category, ELogVerbosity::Type Verbosity, const TCHAR* Data);
 
+	/**
+	 * Creates a fence that tracks the status of async operations started prior to creation of the fence.
+	 *
+	 * The fence is considered complete when every log prior to the fence has been broadcast to every output device.
+	 */
+	CORE_API FOutputDeviceFence CreateFence();
+
 	/** Passes on the flush request to all current output devices. */
 	CORE_API void Flush() final;
 
@@ -178,3 +186,19 @@ public:
 private:
 	TPimplPtr<UE::Private::FOutputDeviceRedirectorState> State;
 };
+
+/** Tracks pending operations on output devices. */
+class FOutputDeviceFence
+{
+public:
+	CORE_API void Wait() const;
+	CORE_API bool Poll() const;
+
+private:
+	FOutputDeviceFence() = default;
+
+	UE::Private::FOutputDeviceRedirectorState* State = nullptr;
+	uint64 Index = 0;
+
+	friend FOutputDeviceRedirector;
+};
diff --git a/Engine/Source/Runtime/Core/Tests/Misc/OutputDeviceRedirectorTest.cpp b/Engine/Source/Runtime/Core/Tests/Misc/OutputDeviceRedirectorTest.cpp
new file mode 100644
--- /dev/null
+++ b/Engine/Source/Runtime/Core/Tests/Misc/OutputDeviceRedirectorTest.cpp
@@ -0,0 +1,37 @@
+// Copyright Epic Games, Inc. All Rights Reserved.
+
+#include "Misc/OutputDeviceRedirector.h"
+
+#if WITH_TESTS
+
+#include "Tests/TestHarnessAdapter.h"
+
+TEST_CASE_NAMED(FOutputDeviceRedirectorTest, "Core::Logging::OutputDeviceRedirector", "[Core][EngineFilter]")
+{
+	bool bHasThread = GLog->TryStartDedicatedPrimaryThread();
+
+	// Test a fence with a dedicated logging thread when available.
+	if (bHasThread)
+	{
+		FOutputDeviceFence Fence = GLog->CreateFence();
+		Fence.Wait();
+	}
+
+	if (bHasThread)
+	{
+		GLog->SetCurrentThreadAsPrimaryThread();
+	}
+
+	// Test a fence without a dedicated logging thread.
+	{
+		FOutputDeviceFence Fence = GLog->CreateFence();
+		Fence.Wait();
+	}
+
+	if (bHasThread)
+	{
+		GLog->TryStartDedicatedPrimaryThread();
+	}
+}
+
+#endif // WITH_TESTS
