Update Unreal.natvis for TDeque, supporting visualisation of elements when the container is set up to use a different allocator such as TInlineAllocator or TFixedAllocator. Added a couple of simple test cases to use TFixedAllocator for validate testing the natvis changes.

#tests Passed FoundationTests.

diff --git a/Engine/Source/Runtime/Core/Tests/Containers/DequeTest.cpp b/Engine/Source/Runtime/Core/Tests/Containers/DequeTest.cpp
--- a/Engine/Source/Runtime/Core/Tests/Containers/DequeTest.cpp
+++ b/Engine/Source/Runtime/Core/Tests/Containers/DequeTest.cpp
@@ -154,6 +154,19 @@ TEST_CASE_NAMED(FDequeResetTest, "System::Core::Containers::TDeque::Reset", "[Sm
 		CHECK(Deque.IsEmpty());
 		CHECK(Deque.Max() == Test::DefaultCapacity);
 	}
+
+	SECTION("Fixed allocator")
+	{
+		TDeque<int32Token, TFixedAllocator<Test::DefaultCapacity>> Deque;
+		Deque.Reset();
+		CHECK(Deque.Max() == Test::DefaultCapacity);
+		Deque.EmplaceLast(0);
+		CHECK(Deque.Num() == 1);
+		CHECK(Deque.Max() == Test::DefaultCapacity);
+		Deque.Reset();
+		CHECK(Deque.IsEmpty());
+		CHECK(Deque.Max() == Test::DefaultCapacity);
+	}
 }
 
 TEST_CASE_NAMED(FDequeEmptyTest, "System::Core::Containers::TDeque::Empty", "[SmokeFilter][Core][Containers][Deque]")
@@ -179,6 +192,15 @@ TEST_CASE_NAMED(FDequeEmptyTest, "System::Core::Containers::TDeque::Empty", "[Sm
 		CHECK(!Deque.Num());
 		CHECK(Deque.IsEmpty());
 	}
+
+	SECTION("Fixed allocator")
+	{
+		TDeque<int32Token, TFixedAllocator<Test::DefaultCapacity>> Deque;
+		Deque.Empty();
+		CHECK(Deque.Max() == Test::DefaultCapacity);
+		CHECK(!Deque.Num());
+		CHECK(Deque.IsEmpty());
+	}
 }
 
 TEST_CASE_NAMED(FDequeEmptyAfterSingleElementEmplaceLastTest, "System::Core::Containers::TDeque::Empty after single element EmplaceLast", "[SmokeFilter][Core][Containers][Deque]")
