Added [[nodiscard]] to relevant TValueOrError functions.

diff --git a/Engine/Source/Runtime/Core/Public/Templates/ValueOrError.h b/Engine/Source/Runtime/Core/Public/Templates/ValueOrError.h
--- a/Engine/Source/Runtime/Core/Public/Templates/ValueOrError.h
+++ b/Engine/Source/Runtime/Core/Public/Templates/ValueOrError.h
@@ -11,25 +11,25 @@
 template <typename... ArgTypes>
 struct TValueOrError_ValueProxy
 {
-	TValueOrError_ValueProxy(ArgTypes&&... InArgs) : Args(Forward<ArgTypes>(InArgs)...) {}
+	[[nodiscard]] TValueOrError_ValueProxy(ArgTypes&&... InArgs) : Args(Forward<ArgTypes>(InArgs)...) {}
 	TTuple<ArgTypes&&...> Args;
 };
 
 template <typename... ArgTypes>
 struct TValueOrError_ErrorProxy
 {
-	TValueOrError_ErrorProxy(ArgTypes&&... InArgs) : Args(Forward<ArgTypes>(InArgs)...) {}
+	[[nodiscard]] TValueOrError_ErrorProxy(ArgTypes&&... InArgs) : Args(Forward<ArgTypes>(InArgs)...) {}
 	TTuple<ArgTypes&&...> Args;
 };
 
 template <typename... ArgTypes>
-FORCEINLINE TValueOrError_ValueProxy<ArgTypes...> MakeValue(ArgTypes&&... Args)
+[[nodiscard]] FORCEINLINE TValueOrError_ValueProxy<ArgTypes...> MakeValue(ArgTypes&&... Args)
 {
 	return TValueOrError_ValueProxy<ArgTypes...>(Forward<ArgTypes>(Args)...);
 }
 
 template <typename... ArgTypes>
-FORCEINLINE TValueOrError_ErrorProxy<ArgTypes...> MakeError(ArgTypes&&... Args)
+[[nodiscard]] FORCEINLINE TValueOrError_ErrorProxy<ArgTypes...> MakeError(ArgTypes&&... Args)
 {
 	return TValueOrError_ErrorProxy<ArgTypes...>(Forward<ArgTypes>(Args)...);
 }
@@ -61,13 +61,13 @@ class TValueOrError
 	};
 
 	template <typename... ArgTypes, uint32... ArgIndices>
-	inline TValueOrError(TValueOrError_ValueProxy<ArgTypes...>&& Proxy, TIntegerSequence<uint32, ArgIndices...>)
+	[[nodiscard]] inline TValueOrError(TValueOrError_ValueProxy<ArgTypes...>&& Proxy, TIntegerSequence<uint32, ArgIndices...>)
 		: Variant(TInPlaceType<ValueType>(), MoveTemp(Proxy.Args).template Get<ArgIndices>()...)
 	{
 	}
 
 	template <typename... ArgTypes, uint32... ArgIndices>
-	inline TValueOrError(TValueOrError_ErrorProxy<ArgTypes...>&& Proxy, TIntegerSequence<uint32, ArgIndices...>)
+	[[nodiscard]] inline TValueOrError(TValueOrError_ErrorProxy<ArgTypes...>&& Proxy, TIntegerSequence<uint32, ArgIndices...>)
 		: Variant(TInPlaceType<FWrapErrorType>(), MoveTemp(Proxy.Args).template Get<ArgIndices>()...)
 	{
 	}
@@ -75,31 +75,31 @@ class TValueOrError
 public:
 	/** Construct the value from a proxy from MakeValue. */
 	template <typename... ArgTypes>
-	inline TValueOrError(TValueOrError_ValueProxy<ArgTypes...>&& Proxy)
+	[[nodiscard]] inline TValueOrError(TValueOrError_ValueProxy<ArgTypes...>&& Proxy)
 		: TValueOrError(MoveTemp(Proxy), TMakeIntegerSequence<uint32, sizeof...(ArgTypes)>())
 	{
 	}
 
 	/** Construct the error from a proxy from MakeError. */
 	template <typename... ArgTypes>
-	inline TValueOrError(TValueOrError_ErrorProxy<ArgTypes...>&& Proxy)
+	[[nodiscard]] inline TValueOrError(TValueOrError_ErrorProxy<ArgTypes...>&& Proxy)
 		: TValueOrError(MoveTemp(Proxy), TMakeIntegerSequence<uint32, sizeof...(ArgTypes)>())
 	{
 	}
 
 	/** Check whether a value is set. Prefer HasValue and HasError to this. !IsValid() does *not* imply HasError(). */
-	inline bool IsValid() const { return Variant.template IsType<ValueType>(); }
+	[[nodiscard]] inline bool IsValid() const { return Variant.template IsType<ValueType>(); }
 
 	/** Whether the error is set. An error does imply no value. No error does *not* imply that a value is set. */
-	inline bool HasError() const { return Variant.template IsType<FWrapErrorType>(); }
+	[[nodiscard]] inline bool HasError() const { return Variant.template IsType<FWrapErrorType>(); }
 
 	/** Access the error. Asserts if this does not have an error. */
-	inline       ErrorType& GetError()       &  { return Variant.template Get<FWrapErrorType>().Error; }
-	inline const ErrorType& GetError() const &  { return Variant.template Get<FWrapErrorType>().Error; }
-	inline       ErrorType  GetError()       && { return MoveTemp(Variant.template Get<FWrapErrorType>().Error); }
+	[[nodiscard]] inline       ErrorType& GetError()       &  { return Variant.template Get<FWrapErrorType>().Error; }
+	[[nodiscard]] inline const ErrorType& GetError() const &  { return Variant.template Get<FWrapErrorType>().Error; }
+	[[nodiscard]] inline       ErrorType  GetError()       && { return MoveTemp(Variant.template Get<FWrapErrorType>().Error); }
 
 	/** Access the error if it is set. */
-	inline       ErrorType* TryGetError()
+	[[nodiscard]] inline ErrorType* TryGetError()
 	{
 		if (FWrapErrorType* Wrap = Variant.template TryGet<FWrapErrorType>())
 		{
@@ -107,7 +107,7 @@ public:
 		}
 		return nullptr;
 	}
-	inline const ErrorType* TryGetError() const
+	[[nodiscard]] inline const ErrorType* TryGetError() const
 	{
 		if (const FWrapErrorType* Wrap = Variant.template TryGet<FWrapErrorType>())
 		{
@@ -117,7 +117,7 @@ public:
 	}
 
 	/** Steal the error. Asserts if this does not have an error. This causes the error to be unset. */
-	inline ErrorType StealError()
+	[[nodiscard]] inline ErrorType StealError()
 	{
 		ErrorType Temp = MoveTemp(GetError());
 		Variant.template Emplace<FEmptyType>();
@@ -125,19 +125,19 @@ public:
 	}
 
 	/** Whether the value is set. A value does imply no error. No value does *not* imply that an error is set. */
-	inline bool HasValue() const { return Variant.template IsType<ValueType>(); }
+	[[nodiscard]] inline bool HasValue() const { return Variant.template IsType<ValueType>(); }
 
 	/** Access the value. Asserts if this does not have a value. */
-	inline       ValueType& GetValue()       &  { return Variant.template Get<ValueType>(); }
-	inline const ValueType& GetValue() const &  { return Variant.template Get<ValueType>(); }
-	inline       ValueType  GetValue()       && { return MoveTemp(Variant.template Get<ValueType>()); }
+	[[nodiscard]] inline       ValueType& GetValue()       &  { return Variant.template Get<ValueType>(); }
+	[[nodiscard]] inline const ValueType& GetValue() const &  { return Variant.template Get<ValueType>(); }
+	[[nodiscard]] inline       ValueType  GetValue()       && { return MoveTemp(Variant.template Get<ValueType>()); }
 
 	/** Access the value if it is set. */
-	inline       ValueType* TryGetValue()       { return Variant.template TryGet<ValueType>(); }
-	inline const ValueType* TryGetValue() const { return Variant.template TryGet<ValueType>(); }
+	[[nodiscard]] inline       ValueType* TryGetValue()       { return Variant.template TryGet<ValueType>(); }
+	[[nodiscard]] inline const ValueType* TryGetValue() const { return Variant.template TryGet<ValueType>(); }
 
 	/** Steal the value. Asserts if this does not have a value. This causes the value to be unset. */
-	inline ValueType StealValue()
+	[[nodiscard]] inline ValueType StealValue()
 	{
 		ValueType Temp = MoveTemp(GetValue());
 		Variant.template Emplace<FEmptyType>();
@@ -152,7 +152,7 @@ template <typename ValueType>
 class TValueOrError<ValueType, void>
 {
 	template <typename... ArgTypes, uint32... ArgIndices>
-	inline TValueOrError(TValueOrError_ValueProxy<ArgTypes...>&& Proxy, TIntegerSequence<uint32, ArgIndices...>)
+	[[nodiscard]] inline TValueOrError(TValueOrError_ValueProxy<ArgTypes...>&& Proxy, TIntegerSequence<uint32, ArgIndices...>)
 		: Value(InPlace, MoveTemp(Proxy.Args).template Get<ArgIndices>()...)
 	{
 	}
@@ -160,33 +160,33 @@ class TValueOrError<ValueType, void>
 public:
 	/** Construct the value from a proxy from MakeValue. */
 	template <typename... ArgTypes>
-	inline TValueOrError(TValueOrError_ValueProxy<ArgTypes...>&& Proxy)
+	[[nodiscard]] inline TValueOrError(TValueOrError_ValueProxy<ArgTypes...>&& Proxy)
 		: TValueOrError(MoveTemp(Proxy), TMakeIntegerSequence<uint32, sizeof...(ArgTypes)>())
 	{
 	}
 
 	/** Construct the error from a proxy from MakeError. */
-	inline TValueOrError(TValueOrError_ErrorProxy<>&& Proxy)
+	[[nodiscard]] inline TValueOrError(TValueOrError_ErrorProxy<>&& Proxy)
 	{
 	}
 
 	/** Whether the error is set. An error does imply no value. */
-	inline bool HasError() const { return !Value.IsSet(); }
+	[[nodiscard]] inline bool HasError() const { return !Value.IsSet(); }
 
 	/** Whether the value is set. A value does imply no error. */
-	inline bool HasValue() const { return Value.IsSet(); }
+	[[nodiscard]] inline bool HasValue() const { return Value.IsSet(); }
 
 	/** Access the value. Asserts if this does not have a value. */
-	inline       ValueType& GetValue()       &  { return Value.GetValue(); }
-	inline const ValueType& GetValue() const &  { return Value.GetValue(); }
-	inline       ValueType  GetValue()       && { return MoveTemp(Value.GetValue()); }
+	[[nodiscard]] inline       ValueType& GetValue()       &  { return Value.GetValue(); }
+	[[nodiscard]] inline const ValueType& GetValue() const &  { return Value.GetValue(); }
+	[[nodiscard]] inline       ValueType  GetValue()       && { return MoveTemp(Value.GetValue()); }
 
 	/** Access the value if it is set. */
-	inline       ValueType* TryGetValue()       { return Value.GetPtrOrNull(); }
-	inline const ValueType* TryGetValue() const { return Value.GetPtrOrNull(); }
+	[[nodiscard]] inline       ValueType* TryGetValue()       { return Value.GetPtrOrNull(); }
+	[[nodiscard]] inline const ValueType* TryGetValue() const { return Value.GetPtrOrNull(); }
 
 	/** Steal the value. Asserts if this does not have a value. This causes the value to be unset. */
-	inline ValueType StealValue()
+	[[nodiscard]] inline ValueType StealValue()
 	{
 		ValueType Temp = MoveTemp(GetValue());
 		Value.Reset();
@@ -201,41 +201,41 @@ template <typename ErrorType>
 class TValueOrError<void, ErrorType>
 {
 	template <typename... ArgTypes, uint32... ArgIndices>
-	inline TValueOrError(TValueOrError_ErrorProxy<ArgTypes...>&& Proxy, TIntegerSequence<uint32, ArgIndices...>)
+	[[nodiscard]] inline TValueOrError(TValueOrError_ErrorProxy<ArgTypes...>&& Proxy, TIntegerSequence<uint32, ArgIndices...>)
 		: Error(InPlace, MoveTemp(Proxy.Args).template Get<ArgIndices>()...)
 	{
 	}
 
 public:
 	/** Construct the value from a proxy from MakeValue. */
-	inline TValueOrError(TValueOrError_ValueProxy<>&& Proxy)
+	[[nodiscard]] inline TValueOrError(TValueOrError_ValueProxy<>&& Proxy)
 	{
 	}
 
 	/** Construct the error from a proxy from MakeError. */
 	template <typename... ArgTypes>
-	inline TValueOrError(TValueOrError_ErrorProxy<ArgTypes...>&& Proxy)
+	[[nodiscard]] inline TValueOrError(TValueOrError_ErrorProxy<ArgTypes...>&& Proxy)
 		: TValueOrError(MoveTemp(Proxy), TMakeIntegerSequence<uint32, sizeof...(ArgTypes)>())
 	{
 	}
 
 	/** Whether the value is set. A value does imply no error. */
-	inline bool HasValue() const { return !Error.IsSet(); }
+	[[nodiscard]] inline bool HasValue() const { return !Error.IsSet(); }
 
 	/** Whether the error is set. An error does imply no value. */
-	inline bool HasError() const { return Error.IsSet(); }
+	[[nodiscard]] inline bool HasError() const { return Error.IsSet(); }
 
 	/** Access the error. Asserts if this does not have an error. */
-	inline       ErrorType& GetError()       &  { return Error.GetValue(); }
-	inline const ErrorType& GetError() const &  { return Error.GetValue(); }
-	inline       ErrorType  GetError()       && { return MoveTemp(Error.GetValue()); }
+	[[nodiscard]] inline       ErrorType& GetError()       &  { return Error.GetValue(); }
+	[[nodiscard]] inline const ErrorType& GetError() const &  { return Error.GetValue(); }
+	[[nodiscard]] inline       ErrorType  GetError()       && { return MoveTemp(Error.GetValue()); }
 
 	/** Access the error if it is set. */
-	inline       ErrorType* TryGetError()       { return Error.GetPtrOrNull(); }
-	inline const ErrorType* TryGetError() const { return Error.GetPtrOrNull(); }
+	[[nodiscard]] inline       ErrorType* TryGetError()       { return Error.GetPtrOrNull(); }
+	[[nodiscard]] inline const ErrorType* TryGetError() const { return Error.GetPtrOrNull(); }
 
 	/** Steal the error. Asserts if this does not have an error. This causes the error to be unset. */
-	inline ErrorType StealError()
+	[[nodiscard]] inline ErrorType StealError()
 	{
 		ErrorType Temp = MoveTemp(GetError());
 		Error.Reset();
@@ -251,22 +251,22 @@ class TValueOrError<void, void>
 {
 public:
 	/** Construct the value from a proxy from MakeValue. */
-	inline TValueOrError(TValueOrError_ValueProxy<>&& Proxy)
+	[[nodiscard]] inline TValueOrError(TValueOrError_ValueProxy<>&& Proxy)
 		: bValue(true)
 	{
 	}
 
 	/** Construct the error from a proxy from MakeError. */
-	inline TValueOrError(TValueOrError_ErrorProxy<>&& Proxy)
+	[[nodiscard]] inline TValueOrError(TValueOrError_ErrorProxy<>&& Proxy)
 		: bValue(false)
 	{
 	}
 
 	/** Whether the error is set. An error does imply no value. */
-	inline bool HasError() const { return !bValue; }
+	[[nodiscard]] inline bool HasError() const { return !bValue; }
 
 	/** Whether the value is set. A value does imply no error. */
-	inline bool HasValue() const { return bValue; }
+	[[nodiscard]] inline bool HasValue() const { return bValue; }
 
 private:
 	bool bValue;
diff --git a/Engine/Source/Runtime/Core/Tests/Templates/ValueOrErrorTest.cpp b/Engine/Source/Runtime/Core/Tests/Templates/ValueOrErrorTest.cpp
--- a/Engine/Source/Runtime/Core/Tests/Templates/ValueOrErrorTest.cpp
+++ b/Engine/Source/Runtime/Core/Tests/Templates/ValueOrErrorTest.cpp
@@ -141,7 +141,7 @@ TEST_CASE_NAMED(FTemplatesValueOrErrorTest, "System::Core::Templates::ValueOrErr
 		CHECK(ValueCount == 1);
 		CHECK(ErrorCount == 0);
 		FTestType UnsetValueOrError = MakeValue();
-		UnsetValueOrError.StealValue();
+		(void)UnsetValueOrError.StealValue();
 		ValueOrError = MoveTemp(UnsetValueOrError);
 		CHECK(ValueCount == 0);
 		CHECK(ErrorCount == 0);
