Fixed return types for rvalue overrides of TValueOrError's GetValue() and GetError() functions.
Added UE_LIFETIMEBOUND to relevant functions - intended to catch any change of behavior caused by the above change, but an unrelated problem was found and fixed.

diff --git a/Engine/Source/Runtime/Core/Public/Templates/ValueOrError.h b/Engine/Source/Runtime/Core/Public/Templates/ValueOrError.h
--- a/Engine/Source/Runtime/Core/Public/Templates/ValueOrError.h
+++ b/Engine/Source/Runtime/Core/Public/Templates/ValueOrError.h
@@ -94,12 +94,12 @@ public:
 	[[nodiscard]] inline bool HasError() const { return Variant.template IsType<FWrapErrorType>(); }
 
 	/** Access the error. Asserts if this does not have an error. */
-	[[nodiscard]] inline       ErrorType& GetError()       &  { return Variant.template Get<FWrapErrorType>().Error; }
-	[[nodiscard]] inline const ErrorType& GetError() const &  { return Variant.template Get<FWrapErrorType>().Error; }
-	[[nodiscard]] inline       ErrorType  GetError()       && { return MoveTemp(Variant.template Get<FWrapErrorType>().Error); }
+	[[nodiscard]] inline       ErrorType&  GetError()       &  UE_LIFETIMEBOUND { return Variant.template Get<FWrapErrorType>().Error; }
+	[[nodiscard]] inline const ErrorType&  GetError() const &  UE_LIFETIMEBOUND { return Variant.template Get<FWrapErrorType>().Error; }
+	[[nodiscard]] inline       ErrorType&& GetError()       && UE_LIFETIMEBOUND { return MoveTemp(Variant.template Get<FWrapErrorType>().Error); }
 
 	/** Access the error if it is set. */
-	[[nodiscard]] inline ErrorType* TryGetError()
+	[[nodiscard]] inline ErrorType* TryGetError() UE_LIFETIMEBOUND
 	{
 		if (FWrapErrorType* Wrap = Variant.template TryGet<FWrapErrorType>())
 		{
@@ -107,7 +107,7 @@ public:
 		}
 		return nullptr;
 	}
-	[[nodiscard]] inline const ErrorType* TryGetError() const
+	[[nodiscard]] inline const ErrorType* TryGetError() const UE_LIFETIMEBOUND
 	{
 		if (const FWrapErrorType* Wrap = Variant.template TryGet<FWrapErrorType>())
 		{
@@ -128,13 +128,13 @@ public:
 	[[nodiscard]] inline bool HasValue() const { return Variant.template IsType<ValueType>(); }
 
 	/** Access the value. Asserts if this does not have a value. */
-	[[nodiscard]] inline       ValueType& GetValue()       &  { return Variant.template Get<ValueType>(); }
-	[[nodiscard]] inline const ValueType& GetValue() const &  { return Variant.template Get<ValueType>(); }
-	[[nodiscard]] inline       ValueType  GetValue()       && { return MoveTemp(Variant.template Get<ValueType>()); }
+	[[nodiscard]] inline       ValueType&  GetValue()       &  UE_LIFETIMEBOUND { return Variant.template Get<ValueType>(); }
+	[[nodiscard]] inline const ValueType&  GetValue() const &  UE_LIFETIMEBOUND { return Variant.template Get<ValueType>(); }
+	[[nodiscard]] inline       ValueType&& GetValue()       && UE_LIFETIMEBOUND { return MoveTemp(Variant.template Get<ValueType>()); }
 
 	/** Access the value if it is set. */
-	[[nodiscard]] inline       ValueType* TryGetValue()       { return Variant.template TryGet<ValueType>(); }
-	[[nodiscard]] inline const ValueType* TryGetValue() const { return Variant.template TryGet<ValueType>(); }
+	[[nodiscard]] inline       ValueType* TryGetValue()       UE_LIFETIMEBOUND { return Variant.template TryGet<ValueType>(); }
+	[[nodiscard]] inline const ValueType* TryGetValue() const UE_LIFETIMEBOUND { return Variant.template TryGet<ValueType>(); }
 
 	/** Steal the value. Asserts if this does not have a value. This causes the value to be unset. */
 	[[nodiscard]] inline ValueType StealValue()
@@ -177,13 +177,13 @@ public:
 	[[nodiscard]] inline bool HasValue() const { return Value.IsSet(); }
 
 	/** Access the value. Asserts if this does not have a value. */
-	[[nodiscard]] inline       ValueType& GetValue()       &  { return Value.GetValue(); }
-	[[nodiscard]] inline const ValueType& GetValue() const &  { return Value.GetValue(); }
-	[[nodiscard]] inline       ValueType  GetValue()       && { return MoveTemp(Value.GetValue()); }
+	[[nodiscard]] inline       ValueType&  GetValue()       &  UE_LIFETIMEBOUND { return Value.GetValue(); }
+	[[nodiscard]] inline const ValueType&  GetValue() const &  UE_LIFETIMEBOUND { return Value.GetValue(); }
+	[[nodiscard]] inline       ValueType&& GetValue()       && UE_LIFETIMEBOUND { return MoveTemp(Value.GetValue()); }
 
 	/** Access the value if it is set. */
-	[[nodiscard]] inline       ValueType* TryGetValue()       { return Value.GetPtrOrNull(); }
-	[[nodiscard]] inline const ValueType* TryGetValue() const { return Value.GetPtrOrNull(); }
+	[[nodiscard]] inline       ValueType* TryGetValue()       UE_LIFETIMEBOUND { return Value.GetPtrOrNull(); }
+	[[nodiscard]] inline const ValueType* TryGetValue() const UE_LIFETIMEBOUND { return Value.GetPtrOrNull(); }
 
 	/** Steal the value. Asserts if this does not have a value. This causes the value to be unset. */
 	[[nodiscard]] inline ValueType StealValue()
@@ -226,13 +226,13 @@ public:
 	[[nodiscard]] inline bool HasError() const { return Error.IsSet(); }
 
 	/** Access the error. Asserts if this does not have an error. */
-	[[nodiscard]] inline       ErrorType& GetError()       &  { return Error.GetValue(); }
-	[[nodiscard]] inline const ErrorType& GetError() const &  { return Error.GetValue(); }
-	[[nodiscard]] inline       ErrorType  GetError()       && { return MoveTemp(Error.GetValue()); }
+	[[nodiscard]] inline       ErrorType&  GetError()       &  UE_LIFETIMEBOUND { return Error.GetValue(); }
+	[[nodiscard]] inline const ErrorType&  GetError() const &  UE_LIFETIMEBOUND { return Error.GetValue(); }
+	[[nodiscard]] inline       ErrorType&& GetError()       && UE_LIFETIMEBOUND { return MoveTemp(Error.GetValue()); }
 
 	/** Access the error if it is set. */
-	[[nodiscard]] inline       ErrorType* TryGetError()       { return Error.GetPtrOrNull(); }
-	[[nodiscard]] inline const ErrorType* TryGetError() const { return Error.GetPtrOrNull(); }
+	[[nodiscard]] inline       ErrorType* TryGetError()       UE_LIFETIMEBOUND { return Error.GetPtrOrNull(); }
+	[[nodiscard]] inline const ErrorType* TryGetError() const UE_LIFETIMEBOUND { return Error.GetPtrOrNull(); }
 
 	/** Steal the error. Asserts if this does not have an error. This causes the error to be unset. */
 	[[nodiscard]] inline ErrorType StealError()
diff --git a/Engine/Source/Runtime/Experimental/Dataflow/Engine/Public/Dataflow/DataflowInstance.h b/Engine/Source/Runtime/Experimental/Dataflow/Engine/Public/Dataflow/DataflowInstance.h
--- a/Engine/Source/Runtime/Experimental/Dataflow/Engine/Public/Dataflow/DataflowInstance.h
+++ b/Engine/Source/Runtime/Experimental/Dataflow/Engine/Public/Dataflow/DataflowInstance.h
@@ -21,7 +21,8 @@ namespace UE::Dataflow::InstanceUtils::Private
 	{
 		if (const FPropertyBagPropertyDesc* Desc = Variables.FindPropertyDescByName(VariableName))
 		{
-			if (FPropertyBagArrayRef* ArrayRef = Variables.GetMutableArrayRef(*Desc).TryGetValue())
+			TValueOrError<FPropertyBagArrayRef, EPropertyBagResult> MutableArrayRef = Variables.GetMutableArrayRef(*Desc);
+			if (FPropertyBagArrayRef* ArrayRef = MutableArrayRef.TryGetValue())
 			{
 				bool bSuccess = true;
 				ArrayRef->EmptyValues();
