Avoid hang detection when loading new localization content

diff --git a/Engine/Source/Runtime/Core/Private/Internationalization/TextLocalizationManager.cpp b/Engine/Source/Runtime/Core/Private/Internationalization/TextLocalizationManager.cpp
--- a/Engine/Source/Runtime/Core/Private/Internationalization/TextLocalizationManager.cpp
+++ b/Engine/Source/Runtime/Core/Private/Internationalization/TextLocalizationManager.cpp
@@ -1375,6 +1375,7 @@ void FTextLocalizationManager::OnCultureChanged()
 
 	if (!TextLocalizationManager::AsyncLoadLocalizationDataOnLanguageChange)
 	{
+		FSlowHeartBeatScope SuspendHeartBeat;
 		WaitForAsyncTasks();
 	}
 }
