Fix trivial LOCTEXT warning

diff --git a/Engine/Source/Runtime/Core/Private/Misc/AxisDisplayInfo.cpp b/Engine/Source/Runtime/Core/Private/Misc/AxisDisplayInfo.cpp
--- a/Engine/Source/Runtime/Core/Private/Misc/AxisDisplayInfo.cpp
+++ b/Engine/Source/Runtime/Core/Private/Misc/AxisDisplayInfo.cpp
@@ -175,7 +175,7 @@ FText AxisDisplayInfo::GetRotationAxisToolTip(EAxisList::Type Axis)
 		IAxisDisplayInfo& DisplayInfo = IModularFeatures::Get().GetModularFeature<IAxisDisplayInfo>(IAxisDisplayInfo::GetModularFeatureName());
 		return DisplayInfo.GetRotationAxisToolTip(Axis);
 	}
-	return LOCTEXT("UnsupportedRotationAxisToolTip", "");
+	return LOCTEXT("UnsupportedRotationAxisToolTip", "?");
 }
 
 FIntVector4 AxisDisplayInfo::GetTransformAxisSwizzle()
