[Scheduler]
- Make thread creation thread-safe and remove mutex

diff --git a/Engine/Source/Runtime/Core/Private/Async/Fundamental/Scheduler.cpp b/Engine/Source/Runtime/Core/Private/Async/Fundamental/Scheduler.cpp
--- a/Engine/Source/Runtime/Core/Private/Async/Fundamental/Scheduler.cpp
+++ b/Engine/Source/Runtime/Core/Private/Async/Fundamental/Scheduler.cpp
@@ -261,7 +261,7 @@ namespace LowLevelTasks
 		if (OldActiveWorkers == 0 && bSupportsMultithreading && ActiveWorkers.compare_exchange_strong(OldActiveWorkers, NumForegroundWorkers + NumBackgroundWorkers, std::memory_order_relaxed))
 		{
 			UE::TUniqueLock Lock(WorkerThreadsCS);
-			check(!WorkerThreads.Num());
+			check(!WorkerThreads);
 			check(!WorkerLocalQueues.Num());
 			check(!WorkerEvents.Num());
 			check(NextWorkerId == 0);
@@ -281,13 +281,13 @@ namespace LowLevelTasks
 			FSchedulerTls::GetTlsValuesRef().LocalQueue = GameThreadLocalQueue.Get();
 
 			WorkerEvents.SetNum(MaxWorkers);
-			WorkerLocalQueues.Reserve(MaxWorkers);
-			WorkerThreads.SetNum(MaxWorkers);
+			WorkerLocalQueues.SetNum(MaxWorkers);
+			WorkerThreads.Reset(new std::atomic<FThread*>[MaxWorkers]);
 
 			auto CreateThread =
 				[this, IsForkable](Private::ELocalQueueType LocalQueueType, const TCHAR* ThreadGroup, const TCHAR* Prefix, std::atomic<int32>& CreationIndex, int32 NumWorkers, int32 NumMaxWorkers, EThreadPriority Priority, uint64 Affinity)
 				{
-					UE::TUniqueLock Lock(WorkerCreationCS);
+					TRACE_CPUPROFILER_EVENT_SCOPE(CreateWorkerThread);
 					LLM_SCOPE_BYNAME(TEXT("EngineMisc/WorkerThreads"));
 
 					// Thread creation can end up waiting, we don't want to recursively oversubscribe if that happens.
@@ -308,7 +308,7 @@ namespace LowLevelTasks
 
 					uint32 WorkerId = NextWorkerId++;
 					UE::Trace::ThreadGroupBegin(ThreadGroup);
-					WorkerLocalQueues.Emplace(QueueRegistry, LocalQueueType);
+					WorkerLocalQueues[WorkerId].Init(QueueRegistry, LocalQueueType);
 					WorkerEvents[WorkerId].bIsStandby = bIsStandbyWorker;
 					WorkerThreads[WorkerId] =
 						CreateWorker(
@@ -444,7 +444,7 @@ namespace LowLevelTasks
 
 			// We wait on threads to exit, once we're done with that
 			// it means no more threads can possibly get created.
-			for (std::atomic<FThread*>& ThreadEntry : WorkerThreads)
+			for (std::atomic<FThread*>& ThreadEntry : TArrayView<std::atomic<FThread*>>(WorkerThreads.Get(), WorkerLocalQueues.Num()))
 			{
 				if (FThread* Thread = ThreadEntry.exchange(nullptr))
 				{
diff --git a/Engine/Source/Runtime/Core/Public/Async/Fundamental/LocalQueue.h b/Engine/Source/Runtime/Core/Public/Async/Fundamental/LocalQueue.h
--- a/Engine/Source/Runtime/Core/Public/Async/Fundamental/LocalQueue.h
+++ b/Engine/Source/Runtime/Core/Public/Async/Fundamental/LocalQueue.h
@@ -178,28 +178,48 @@ public:
 		friend class TLocalQueueRegistry;
 
 	public:
-		TLocalQueue(TLocalQueueRegistry& InRegistry, ELocalQueueType InQueueType) : Registry(&InRegistry), QueueType(InQueueType)
+		TLocalQueue() = default;
+			
+		TLocalQueue(TLocalQueueRegistry& InRegistry, ELocalQueueType InQueueType)
 		{
-			// Local queues are never unregistered, everything is shutdown at once.
-			Registry->AddLocalQueue(this);
-			for (int32 PriorityIndex = 0; PriorityIndex < int32(ETaskPriority::Count); ++PriorityIndex)
+			Init(InRegistry, InQueueType);
+		}
+
+		void Init(TLocalQueueRegistry& InRegistry, ELocalQueueType InQueueType)
+		{
+			if (bIsInitialized.exchange(true, std::memory_order_relaxed))
 			{
-				DequeueHazards[PriorityIndex] = Registry->OverflowQueues[PriorityIndex].getHeadHazard();
+				checkf(false, TEXT("Trying to initialize local queue more than once"));
+			}
+			else
+			{
+				Registry = &InRegistry;
+				QueueType = InQueueType;
+
+				// Local queues are never unregistered, everything is shutdown at once.
+				Registry->AddLocalQueue(this);
+				for (int32 PriorityIndex = 0; PriorityIndex < int32(ETaskPriority::Count); ++PriorityIndex)
+				{
+					DequeueHazards[PriorityIndex] = Registry->OverflowQueues[PriorityIndex].getHeadHazard();
+				}
 			}
 		}
 
 		~TLocalQueue()
 		{
-			for (int32 PriorityIndex = 0; PriorityIndex < int32(ETaskPriority::Count); PriorityIndex++)
+			if (bIsInitialized.exchange(false, std::memory_order_relaxed))
 			{
-				while (true)
+				for (int32 PriorityIndex = 0; PriorityIndex < int32(ETaskPriority::Count); PriorityIndex++)
 				{
-					FTask* Item;
-					if (!LocalQueues[PriorityIndex].Get(Item))
+					while (true)
 					{
-						break;
+						FTask* Item;
+						if (!LocalQueues[PriorityIndex].Get(Item))
+						{
+							break;
+						}
+						Registry->OverflowQueues[PriorityIndex].enqueue(Item);
 					}
-					Registry->OverflowQueues[PriorityIndex].enqueue(Item);
 				}
 			}
 		}
@@ -274,10 +294,11 @@ public:
 		static constexpr uint32    InvalidIndex = ~0u;
 		FLocalQueueType            LocalQueues[uint32(ETaskPriority::Count)];
 		DequeueHazard              DequeueHazards[uint32(ETaskPriority::Count)];
-		TLocalQueueRegistry*       Registry;
+		TLocalQueueRegistry*       Registry = nullptr;
 		uint32                     CachedRandomIndex = InvalidIndex;
 		uint32                     CachedPriorityIndex = 0;
 		ELocalQueueType            QueueType;
+		std::atomic<bool>          bIsInitialized = false;
 	};
 
 	TLocalQueueRegistry()
diff --git a/Engine/Source/Runtime/Core/Public/Async/Fundamental/Scheduler.h b/Engine/Source/Runtime/Core/Public/Async/Fundamental/Scheduler.h
--- a/Engine/Source/Runtime/Core/Public/Async/Fundamental/Scheduler.h
+++ b/Engine/Source/Runtime/Core/Public/Async/Fundamental/Scheduler.h
@@ -158,9 +158,8 @@ namespace LowLevelTasks
 	private:
 		Private::FWaitingQueue                         WaitingQueue[2] = { { WorkerEvents, OversubscriptionLimitReachedEvent }, { WorkerEvents, OversubscriptionLimitReachedEvent } };
 		FSchedulerTls::FQueueRegistry                  QueueRegistry;
-		UE::FPlatformRecursiveMutex                    WorkerCreationCS;
 		UE::FPlatformRecursiveMutex                    WorkerThreadsCS;
-		TArray<std::atomic<FThread*>>                  WorkerThreads;
+		TUniquePtr<std::atomic<FThread*>[]>            WorkerThreads;
 		TAlignedArray<FSchedulerTls::FLocalQueueType>  WorkerLocalQueues;
 		TAlignedArray<Private::FWaitEvent>             WorkerEvents;
 		TUniquePtr<FSchedulerTls::FLocalQueueType>     GameThreadLocalQueue;
