[Insights] Enabled memory tracing for AnsiMalloc.

diff --git a/Engine/Source/Runtime/Core/Private/HAL/MallocAnsi.cpp b/Engine/Source/Runtime/Core/Private/HAL/MallocAnsi.cpp
--- a/Engine/Source/Runtime/Core/Private/HAL/MallocAnsi.cpp
+++ b/Engine/Source/Runtime/Core/Private/HAL/MallocAnsi.cpp
@@ -7,6 +7,7 @@
 #include "HAL/MallocAnsi.h"
 #include "Misc/AssertionMacros.h"
 #include "Math/UnrealMathUtility.h"
+#include "ProfilingDebugging/MemoryTrace.h"
 #include "Templates/AlignmentTemplates.h"
 #include "AutoRTFM.h"
 
@@ -46,6 +47,7 @@ void* AnsiMalloc(SIZE_T Size, uint32 Alignment)
 		*((SIZE_T*)((uint8*)Result - sizeof(void*) - sizeof(SIZE_T))) = Size;
 	}
 #endif
+	MemoryTrace_Alloc(uint64(Result), Size, Alignment, EMemoryTraceRootHeap::SystemMemory);
 
 	return Result;
 }
@@ -66,6 +68,7 @@ void* AnsiRealloc(void* Ptr, SIZE_T NewSize, uint32 Alignment)
 	void* Result;
 
 #if MALLOC_ANSI_USES__ALIGNED_MALLOC
+	MemoryTrace_ReallocFree(uint64(Ptr), EMemoryTraceRootHeap::SystemMemory);
 	if (Ptr && NewSize)
 	{
 		Result = _aligned_realloc(Ptr, NewSize, Alignment);
@@ -79,7 +82,9 @@ void* AnsiRealloc(void* Ptr, SIZE_T NewSize, uint32 Alignment)
 		_aligned_free(Ptr);
 		Result = nullptr;
 	}
+	MemoryTrace_ReallocAlloc(uint64(Result), NewSize, Alignment, EMemoryTraceRootHeap::SystemMemory);
 #elif PLATFORM_USE_ANSI_POSIX_MALLOC
+	MemoryTrace_ReallocFree(uint64(Ptr), EMemoryTraceRootHeap::SystemMemory);
 	if (Ptr && NewSize)
 	{
 		SIZE_T UsableSize = malloc_usable_size(Ptr);
@@ -105,8 +110,11 @@ void* AnsiRealloc(void* Ptr, SIZE_T NewSize, uint32 Alignment)
 		free(Ptr);
 		Result = nullptr;
 	}
+	MemoryTrace_ReallocAlloc(uint64(Result), NewSize, Alignment, EMemoryTraceRootHeap::SystemMemory);
 #elif PLATFORM_USE_ANSI_MEMALIGN
+	MemoryTrace_ReallocFree(uint64(Ptr), EMemoryTraceRootHeap::SystemMemory);
 	Result = reallocalign(Ptr, NewSize, Alignment);
+	MemoryTrace_ReallocAlloc(uint64(Result), NewSize, Alignment, EMemoryTraceRootHeap::SystemMemory);
 #else
 	if (Ptr && NewSize)
 	{
@@ -122,6 +130,7 @@ void* AnsiRealloc(void* Ptr, SIZE_T NewSize, uint32 Alignment)
 	}
 	else
 	{
+		MemoryTrace_Free(uint64(Ptr), EMemoryTraceRootHeap::SystemMemory);
 		free(*((void**)((uint8*)Ptr - sizeof(void*))));
 		Result = nullptr;
 	}
@@ -132,6 +141,7 @@ void* AnsiRealloc(void* Ptr, SIZE_T NewSize, uint32 Alignment)
 
 void AnsiFree(void* Ptr)
 {
+	MemoryTrace_Free(uint64(Ptr), EMemoryTraceRootHeap::SystemMemory);
 #if MALLOC_ANSI_USES__ALIGNED_MALLOC
 	_aligned_free(Ptr);
 #elif PLATFORM_USE_ANSI_POSIX_MALLOC || PLATFORM_USE_ANSI_MEMALIGN
