[TSAN]
- Fix TSAN_ATOMIC assignment operator

diff --git a/Engine/Source/Runtime/Core/Public/Clang/ClangPlatformCodeAnalysis.h b/Engine/Source/Runtime/Core/Public/Clang/ClangPlatformCodeAnalysis.h
--- a/Engine/Source/Runtime/Core/Public/Clang/ClangPlatformCodeAnalysis.h
+++ b/Engine/Source/Runtime/Core/Public/Clang/ClangPlatformCodeAnalysis.h
@@ -180,6 +180,7 @@
 			__attribute__((always_inline)) TTSANSafeValue& operator=(const TTSANSafeValue& Other)
 			{
 				Value.store((T)Other, std::memory_order_relaxed);
+				return *this;
 			}
 
 			template <typename OtherType>
