[RaceDetector]
- Fix potential false positive when using TSAN_BEFORE/TSAN_AFTER since the address is used for synchronization only, it is not touched
- The reason we were instrumenting the access before is that we stored a bit to know if sync objects were present for an address, its not the case anymore since clock bank shadow was implemented.

diff --git a/Engine/Source/Runtime/Core/Private/Sanitizer/RaceDetector.cpp b/Engine/Source/Runtime/Core/Private/Sanitizer/RaceDetector.cpp
--- a/Engine/Source/Runtime/Core/Private/Sanitizer/RaceDetector.cpp
+++ b/Engine/Source/Runtime/Core/Private/Sanitizer/RaceDetector.cpp
@@ -1592,9 +1592,6 @@ extern "C" {
 		BEGIN_HANDLE_INSTRUMENTATION(return);
 
 		AddCallstackFrame(Context, _ReturnAddress());
-		bool bHasAddedCallstackFrame = true;
-		InstrumentAccess(Context, addr, sizeof(UPTRINT), ACCESS_TYPE_ATOMIC_READ, FAtomicMemoryOrder::MEMORY_ORDER_RELEASE, _ReturnAddress(), TEXT("AnnotateHappensBefore"), bHasAddedCallstackFrame);
-
 		FSyncObjectRef Atomic = GetSyncObject(Context, addr);
 		Atomic->SyncRelease(Context, []() {}, _ReturnAddress(), addr, TEXT("AnnotateHappensBefore"));
 		Context.IncrementClock();
@@ -1608,9 +1605,6 @@ extern "C" {
 		BEGIN_HANDLE_INSTRUMENTATION(return);
 
 		AddCallstackFrame(Context, _ReturnAddress());
-		bool bHasAddedCallstackFrame = true;
-		InstrumentAccess(Context, addr, sizeof(UPTRINT), ACCESS_TYPE_ATOMIC_READ, FAtomicMemoryOrder::MEMORY_ORDER_ACQUIRE, _ReturnAddress(), TEXT("AnnotateHappensBefore"), bHasAddedCallstackFrame);
-
 		FSyncObjectRef Atomic = GetSyncObject(Context, addr);
 		Atomic->SyncAcquire(Context, []() {}, _ReturnAddress(), addr, TEXT("AnnotateHappensAfter"));
 		Context.IncrementClock();
