[Coroutine]
- Remove coroutine experimentation

diff --git a/Engine/Source/Programs/UnrealBuildTool/Configuration/UEBuildTarget.cs b/Engine/Source/Programs/UnrealBuildTool/Configuration/UEBuildTarget.cs
--- a/Engine/Source/Programs/UnrealBuildTool/Configuration/UEBuildTarget.cs
+++ b/Engine/Source/Programs/UnrealBuildTool/Configuration/UEBuildTarget.cs
@@ -5872,13 +5872,7 @@ namespace UnrealBuildTool
 			// Whether C++20 coroutines are enabled
 			if (Rules.bEnableCppCoroutinesForEvaluation)
 			{
-				Log.TraceInformationOnce($"NOTE: C++ coroutine support is considered experimental and should be used for evaluation purposes only ({nameof(TargetRules.bEnableCppCoroutinesForEvaluation)})");
-				GlobalCompileEnvironment.Definitions.Add("WITH_CPP_COROUTINES=1");
-				GlobalCompileEnvironment.bEnableCoroutines = true;
-			}
-			else
-			{
-				GlobalCompileEnvironment.Definitions.Add("WITH_CPP_COROUTINES=0");
+				Log.TraceInformationOnce($"NOTE: C++ coroutine experimental support is no longer available");
 			}
 
 			if (Rules.bEnableProcessPriorityControl)
diff --git a/Engine/Source/Runtime/Core/Private/Async/Fundamental/UnitTests.cpp b/Engine/Source/Runtime/Core/Private/Async/Fundamental/UnitTests.cpp
--- a/Engine/Source/Runtime/Core/Private/Async/Fundamental/UnitTests.cpp
+++ b/Engine/Source/Runtime/Core/Private/Async/Fundamental/UnitTests.cpp
@@ -6,10 +6,6 @@
 #include "Async/Fundamental/Scheduler.h"
 #include "Async/ManualResetEvent.h"
 #include "Experimental/Async/AwaitableTask.h"
-#include "Experimental/Coroutine/CoroEvent.h"
-#include "Experimental/Coroutine/CoroParallelFor.h"
-#include "Experimental/Coroutine/CoroSpinLock.h"
-#include "Experimental/Coroutine/CoroTimeout.h"
 
 #include <atomic>
 
@@ -719,84 +715,6 @@ namespace Tasks2Tests
 		}
 		return true;
 	}
-
-	IMPLEMENT_SIMPLE_AUTOMATION_TEST(FTasksCoroutineTests, "System.Core.Coroutine.UnitTests", EAutomationTestFlags_ApplicationContextMask | EAutomationTestFlags::EngineFilter);
-	bool FTasksCoroutineTests::RunTest(const FString& Parameters)
-	{
-#if WITH_CPP_COROUTINES
-		{
-			FCoroEvent Event;
-			auto InnerFrame = [&]() -> CORO_FRAME(void)
-			{
-				CO_AWAIT Event;
-			};
-
-			auto WorkLambda = [&]() -> CORO_TASK(int)
-			{
-				CORO_INVOKE(InnerFrame());
-				CO_RETURN_TASK(42);
-			};
-
-			LAUNCHED_TASK(int) Task = WorkLambda().Launch(TEXT("CoroUnitTest"), ETaskPriority::BackgroundLow);
-			FPlatformProcess::Sleep(0.001f);
-			Event.Trigger();
-			int Value = Task.SpinWait();
-			verify(Value == 42);
-		}
-		{
-			TCoroLocal<int> CLS(0);
-			auto WorkLambda = [&]() -> CORO_TASK(void)
-			{
-				for(int i = 0; i < 1000; i++)
-				{
-					*CLS += 1;
-					int intermediate_value = *CLS;
-					verify(intermediate_value == i+1);
-				}
-				int final_value = *CLS;
-				verify(final_value == 1000);
-				CO_RETURN_TASK();
-			};
-
-			TArray<LAUNCHED_TASK(void)> Tasks;
-			for(int i = 0; i < 100; i++)
-			{
-				Tasks.Emplace(WorkLambda().Launch(TEXT("CoroLocalStateTest")));
-			}
-
-			for(int i = 0; i < 100; i++)
-			{
-				Tasks[i].SpinWait();
-			}
-		}
-		{
-			FCoroSpinLock SpinLock;
-			int Mutable = 0;
-			auto WorkLambda = [&](int32) -> CORO_FRAME(void)
-			{
-				auto LockScope = CO_AWAIT SpinLock.Lock();
-				Mutable++;
-				LockScope.Release();
-				CO_RETURN;
-			};
-
-			CoroParallelFor(TEXT("CoroParallelForTest"), 100, WorkLambda, EParallelForFlags::None);
-			verify(Mutable == 100);
-		}
-		{	
-			auto WorkLambda = [&](int32) -> CORO_FRAME(void)
-			{
-				FCoroTimeoutAwaitable Timeout(FTimespan::FromMilliseconds(1), ECoroTimeoutFlags::Suspend_Worker);
-				FPlatformProcess::Sleep(0.001f);
-				CO_AWAIT Timeout;
-				CO_RETURN;
-			};
-
-			CoroParallelFor(TEXT("CoroTimeoutForTest"), 100, WorkLambda, EParallelForFlags::None);
-		}
-#endif
-		return true;
-	}
 }
 
 #endif
diff --git a/Engine/Source/Runtime/Core/Private/Experimental/Coroutine/Coroutine.cpp b/Engine/Source/Runtime/Core/Private/Experimental/Coroutine/Coroutine.cpp
deleted file mode 100644
diff --git a/Engine/Source/Runtime/Core/Private/HAL/ThreadingBase.cpp b/Engine/Source/Runtime/Core/Private/HAL/ThreadingBase.cpp
--- a/Engine/Source/Runtime/Core/Private/HAL/ThreadingBase.cpp
+++ b/Engine/Source/Runtime/Core/Private/HAL/ThreadingBase.cpp
@@ -15,7 +15,6 @@
 #include "ProfilingDebugging/MiscTrace.h"
 #include "Async/Fundamental/Scheduler.h"
 #include "Tasks/Pipe.h"
-#include "Experimental/Coroutine/Coroutine.h"
 #include "AutoRTFM.h"
 
 #if PLATFORM_WINDOWS
@@ -186,11 +185,7 @@ CORE_API bool IsInGameThread()
 PRAGMA_DISABLE_DEPRECATION_WARNINGS
 		bool newValue = FTaskTagScope::IsCurrentTag(ETaskTag::EGameThread) || FTaskTagScope::IsRunningDuringStaticInit();
 #if !UE_BUILD_SHIPPING && !UE_BUILD_TEST
-		if (
-#if WITH_CPP_COROUTINES
-			!CoroTask_Detail::FCoroLocalState::IsCoroLaunchedTask() &&
-#endif
-			!UE::Tasks::Private::IsThreadRetractingTask())
+		if (!UE::Tasks::Private::IsThreadRetractingTask())
 		{
 			const uint32 CurrentThreadId = FPlatformTLS::GetCurrentThreadId();
 			bool oldValue = CurrentThreadId == GGameThreadId;
@@ -221,11 +216,7 @@ PRAGMA_DISABLE_DEPRECATION_WARNINGS
 	// If this explicitly is a slate thread, not just the main thread running slate
 	bool newValue = FTaskTagScope::IsCurrentTag(ETaskTag::ESlateThread);
 #if !UE_BUILD_SHIPPING && !UE_BUILD_TEST
-	if (
-#if WITH_CPP_COROUTINES
-		!CoroTask_Detail::FCoroLocalState::IsCoroLaunchedTask() &&
-#endif
-		!UE::Tasks::Private::IsThreadRetractingTask())
+	if (!UE::Tasks::Private::IsThreadRetractingTask())
 	{
 		bool oldValue = GSlateLoadingThreadId != 0 && FPlatformTLS::GetCurrentThreadId() == GSlateLoadingThreadId;
 		ensureMsgf(oldValue == newValue, TEXT("oldValue(%i) newValue(%i) If this check fails make sure that there is a FTaskTagScope(ETaskTag::ESlateThread) as as deep as possible on the current callstack, you can see the current value in ActiveNamedThreads(%x)"), oldValue, newValue, FTaskTagScope::GetCurrentTag());
@@ -260,9 +251,6 @@ CORE_API bool IsInActualRenderingThread()
 PRAGMA_DISABLE_DEPRECATION_WARNINGS
 	bool newValue = FTaskTagScope::IsCurrentTag(ETaskTag::ERenderingThread);
 #if !UE_BUILD_SHIPPING && !UE_BUILD_TEST
-#if WITH_CPP_COROUTINES
-	if (!CoroTask_Detail::FCoroLocalState::IsCoroLaunchedTask())
-#endif
 	{
 		bool oldValue = FPlatformTLS::GetCurrentThreadId() == GRenderThreadId;
 		ensureMsgf(oldValue == newValue, TEXT("oldValue(%i) newValue(%i) If this check fails make sure that there is a FTaskTagScope(ETaskTag::ERenderingThread) as deep as possible on the current callstack, you can see the current value in ActiveNamedThreads(%x)"), oldValue, newValue, FTaskTagScope::GetCurrentTag());
@@ -281,11 +269,7 @@ PRAGMA_DISABLE_DEPRECATION_WARNINGS
 		: FTaskTagScope::IsCurrentTag(ETaskTag::ERenderingThread);
 
 #if !UE_BUILD_SHIPPING && !UE_BUILD_TEST
-	if (
-#if WITH_CPP_COROUTINES
-		!CoroTask_Detail::FCoroLocalState::IsCoroLaunchedTask() && 
-#endif
-		!UE::Tasks::Private::IsThreadRetractingTask())
+	if (!UE::Tasks::Private::IsThreadRetractingTask())
 	{
 		const uint32 CurrentThreadId = FPlatformTLS::GetCurrentThreadId();
 		bool oldValue = (GRenderThreadId == 0) ? (CurrentThreadId == GGameThreadId) || FTaskTagScope::IsRunningDuringStaticInit() : (CurrentThreadId == GRenderThreadId);
@@ -322,11 +306,7 @@ PRAGMA_DISABLE_DEPRECATION_WARNINGS
 	}
 
 #if !UE_BUILD_SHIPPING && !UE_BUILD_TEST
-	if (
-#if WITH_CPP_COROUTINES
-		!CoroTask_Detail::FCoroLocalState::IsCoroLaunchedTask() &&
-#endif
-		!UE::Tasks::Private::IsThreadRetractingTask())
+	if (!UE::Tasks::Private::IsThreadRetractingTask())
 	{
 		const uint32 CurrentThreadId = FPlatformTLS::GetCurrentThreadId();
 		bool oldValue = (GRenderThreadId == 0) ?  true : CurrentThreadId != GGameThreadId;
diff --git a/Engine/Source/Runtime/Core/Public/Experimental/Coroutine/CoroEvent.h b/Engine/Source/Runtime/Core/Public/Experimental/Coroutine/CoroEvent.h
deleted file mode 100644
diff --git a/Engine/Source/Runtime/Core/Public/Experimental/Coroutine/CoroLocalVariable.h b/Engine/Source/Runtime/Core/Public/Experimental/Coroutine/CoroLocalVariable.h
deleted file mode 100644
diff --git a/Engine/Source/Runtime/Core/Public/Experimental/Coroutine/CoroParallelFor.h b/Engine/Source/Runtime/Core/Public/Experimental/Coroutine/CoroParallelFor.h
deleted file mode 100644
diff --git a/Engine/Source/Runtime/Core/Public/Experimental/Coroutine/CoroSpinLock.h b/Engine/Source/Runtime/Core/Public/Experimental/Coroutine/CoroSpinLock.h
deleted file mode 100644
diff --git a/Engine/Source/Runtime/Core/Public/Experimental/Coroutine/CoroTimeout.h b/Engine/Source/Runtime/Core/Public/Experimental/Coroutine/CoroTimeout.h
deleted file mode 100644
diff --git a/Engine/Source/Runtime/Core/Public/Experimental/Coroutine/Coroutine.h b/Engine/Source/Runtime/Core/Public/Experimental/Coroutine/Coroutine.h
deleted file mode 100644
diff --git a/Engine/Source/Runtime/Core/Public/Experimental/Coroutine/CoroutineHandle.h b/Engine/Source/Runtime/Core/Public/Experimental/Coroutine/CoroutineHandle.h
deleted file mode 100644
