[TSAN]
- Fix benign warning on CSV

diff --git a/Engine/Source/Runtime/Core/Private/ProfilingDebugging/CsvProfiler.cpp b/Engine/Source/Runtime/Core/Private/ProfilingDebugging/CsvProfiler.cpp
--- a/Engine/Source/Runtime/Core/Private/ProfilingDebugging/CsvProfiler.cpp
+++ b/Engine/Source/Runtime/Core/Private/ProfilingDebugging/CsvProfiler.cpp
@@ -266,9 +266,9 @@ static FAutoConsoleVariableRef CVarTrackWaitsRT(TEXT("csv.trackWaitsRT"), GCsvTr
 // Categories
 //
 static const uint32 CSV_MAX_CATEGORY_COUNT = 2048;
-static bool GCsvCategoriesEnabled[CSV_MAX_CATEGORY_COUNT];
+static TSAN_ATOMIC(bool) GCsvCategoriesEnabled[CSV_MAX_CATEGORY_COUNT];
 
-static bool GCsvProfilerIsCapturing = false;
+static TSAN_ATOMIC(bool) GCsvProfilerIsCapturing = false;
 static bool GCsvProfilerIsCapturingRT = false; // Renderthread version of the above
 
 static bool GCsvProfilerIsWritingFile = false;
