Update Memstack safety check code to compile and work properly when enabled (and fix incorrect name)

diff --git a/Engine/Source/Runtime/Core/Private/Misc/MemStack.cpp b/Engine/Source/Runtime/Core/Private/Misc/MemStack.cpp
--- a/Engine/Source/Runtime/Core/Private/Misc/MemStack.cpp
+++ b/Engine/Source/Runtime/Core/Private/Misc/MemStack.cpp
@@ -43,21 +43,21 @@ FORCENOINLINE FPageAllocator& FPageAllocator::Construct()
 	return *Instance;
 }
 
-#define USE_MEMSTACK_PUGATORY (0)
+#define USE_MEMSTACK_PURGATORY (0)
 
-#if PLATFORM_WINDOWS && USE_MEMSTACK_PUGATORY
+#if PLATFORM_WINDOWS && USE_MEMSTACK_PURGATORY
 
-#define MEMSTACK_PUGATORY_COMPILED_IN (1)
+#define MEMSTACK_PURGATORY_COMPILED_IN (1)
 
 #include "Windows/AllowWindowsPlatformTypes.h"
 #include <Psapi.h>
 
-TLockFreePointerListUnordered<void> TheAllocatorReady;
-TLockFreePointerListUnordered<void> TheSmallAllocatorReady;
+TLockFreePointerListUnordered<void, PLATFORM_CACHE_LINE_SIZE> TheAllocatorReady;
+TLockFreePointerListUnordered<void, PLATFORM_CACHE_LINE_SIZE> TheSmallAllocatorReady;
 
 #define FRAMES_TO_WAIT (4)
-TLockFreePointerListUnordered<void> TheAllocatorPurgatory[FRAMES_TO_WAIT];
-TLockFreePointerListUnordered<void> TheSmallAllocatorPurgatory[FRAMES_TO_WAIT];
+TLockFreePointerListUnordered<void, PLATFORM_CACHE_LINE_SIZE> TheAllocatorPurgatory[FRAMES_TO_WAIT];
+TLockFreePointerListUnordered<void, PLATFORM_CACHE_LINE_SIZE> TheSmallAllocatorPurgatory[FRAMES_TO_WAIT];
 
 FThreadSafeCounter BigBlocksOutstanding;
 FThreadSafeCounter SmallBlocksOutstanding;
@@ -218,20 +218,20 @@ void FPageAllocator::FreeSmall(void *Mem)
 uint64 FPageAllocator::BytesUsed()
 {
 	return uint64(BigBlocksInUse.GetValue()) * PageSize + uint64(SmallBlocksInUse.GetValue()) * SmallPageSize +
-		uint64(TheAllocator.GetNumUsed().GetValue()) * PageSize + uint64(TheSmallAllocator.GetNumUsed().GetValue()) * SmallPageSize;
+		uint64(TheAllocator.GetNumUsed()) * PageSize;
 }
 
 uint64 FPageAllocator::BytesFree()
 {
 	return uint64(BigBlocksOutstanding.GetValue() - BigBlocksInUse.GetValue()) * PageSize + uint64(SmallBlocksOutstanding.GetValue() - SmallBlocksInUse.GetValue()) * SmallPageSize +
-		uint64(TheAllocator.GetNumFree().GetValue()) * PageSize + uint64(TheSmallAllocator.GetNumFree().GetValue()) * SmallPageSize;
+		uint64(TheAllocator.GetNumFree()) * PageSize;
 }
 
 #include "Windows/HideWindowsPlatformTypes.h"
 
 #else
 
-#define MEMSTACK_PUGATORY_COMPILED_IN (0)
+#define MEMSTACK_PURGATORY_COMPILED_IN (0)
 
 void *FPageAllocator::Alloc(int32 Alignment)
 {
@@ -284,7 +284,7 @@ void FPageAllocator::LatchProtectedMode()
 		Get().TheAllocator.Trim();
 	});
 
-#if MEMSTACK_PUGATORY_COMPILED_IN
+#if MEMSTACK_PURGATORY_COMPILED_IN
 	const bool bNoProtection = FParse::Param(FCommandLine::Get(), TEXT("NoProtectMemStack"));
 	GMemStackProtection = 1;
 	if (bNoProtection)
