[MallocCrash]
- Fix benign race with other threads when we're setting GMalloc

diff --git a/Engine/Source/Runtime/Core/Private/GenericPlatform/GenericPlatformMallocCrash.cpp b/Engine/Source/Runtime/Core/Private/GenericPlatform/GenericPlatformMallocCrash.cpp
--- a/Engine/Source/Runtime/Core/Private/GenericPlatform/GenericPlatformMallocCrash.cpp
+++ b/Engine/Source/Runtime/Core/Private/GenericPlatform/GenericPlatformMallocCrash.cpp
@@ -311,7 +311,7 @@ FGenericPlatformMallocCrash& FGenericPlatformMallocCrash::Get( FMalloc* MainMall
 void FGenericPlatformMallocCrash::SetAsGMalloc()
 {
 	InternalLock.Lock();
-	UE::Private::GMalloc = this;
+	FPlatformAtomics::InterlockedExchangePtr((void**)&UE::Private::GMalloc, this);
 	if (PLATFORM_USES_FIXED_GMalloc_CLASS && GFixedMallocLocationPtr)
 	{
 		*GFixedMallocLocationPtr = nullptr; // this disables any fast-path inline allocators
