Add early return in Schedule when there is no available work

#rnx

diff --git a/Engine/Source/Runtime/Core/Private/Misc/QueuedThreadPoolWrapper.cpp b/Engine/Source/Runtime/Core/Private/Misc/QueuedThreadPoolWrapper.cpp
--- a/Engine/Source/Runtime/Core/Private/Misc/QueuedThreadPoolWrapper.cpp
+++ b/Engine/Source/Runtime/Core/Private/Misc/QueuedThreadPoolWrapper.cpp
@@ -248,18 +248,27 @@ void FQueuedThreadPoolWrapper::Schedule(FScheduledWork* Work)
 	TGuardValue<bool> SchedulingScope(bIsScheduling, true);
 
 	// If a higher priority task comes in, try to retract lower priority ones if possible to make room
-	EQueuedWorkPriority NextWorkPriority;
+	EQueuedWorkPriority NextWorkPriority = EQueuedWorkPriority::Count;
 	IQueuedWork* NextWork = QueuedWork.Peek(&NextWorkPriority);
-	if (NextWork)
+		 
+	// When NextWork is nullptr here, it means that QueuedWork.Peek did not 
+	// find any work at any priority. That means NextWorkPriority is not 
+	// valid, there is no work to try schedule, so early out.
+	if (NextWork == nullptr)
+	{
+		return;
+	}
+
+	// Continue retracting more work until nothing can be retracted anymore or we can finally squeeze the higher priority task in
+	while (!CanSchedule(NextWorkPriority) && TryRetractWorkNoLock(NextWorkPriority))
 	{
-		// Continue retracting more work until nothing can be retracted anymore or we can finally squeeze the higher priority task in
-		while (!CanSchedule(NextWorkPriority) && TryRetractWorkNoLock(NextWorkPriority))
-		{
-		}
 	}
 
 	// Schedule as many tasks we can fit
-	while (CanSchedule(NextWorkPriority))
+	// Note: The reason for checking that there is work queued is that CanSchedule calls GetMaxConcurrency 
+	// 	which is virtual and can (and does at times) do work that is only required if there is more work to 
+	// 	schedule. For example, adjusting max concurrency based on currently available memory.
+	while (QueuedWork.Num() > 0 && CanSchedule(NextWorkPriority))
 	{
 		EQueuedWorkPriority WorkPriority;
 		IQueuedWork* InnerWork = QueuedWork.Dequeue(&WorkPriority);
