Fix a runtime breakage caused by CL 43032251.

#rnx

diff --git a/Engine/Source/Runtime/Core/Private/Misc/StringBuilder.cpp b/Engine/Source/Runtime/Core/Private/Misc/StringBuilder.cpp
--- a/Engine/Source/Runtime/Core/Private/Misc/StringBuilder.cpp
+++ b/Engine/Source/Runtime/Core/Private/Misc/StringBuilder.cpp
@@ -70,22 +70,24 @@ TStringBuilderBase<C>& TStringBuilderBase<C>::AppendV(const C* Fmt, va_list Args
 {
 	for (;;)
 	{
-		va_list ArgPack;
-		va_copy(ArgPack, Args);
 		const int32 RemainingSize = (int32)(End - CurPos);
-		const int32 Result = TCString<C>::GetVarArgs(CurPos, RemainingSize, Fmt, ArgPack);
-		va_end(ArgPack);
 
-		if (Result >= 0 && Result < RemainingSize)
+		if (0 != RemainingSize)
 		{
-			CurPos += Result;
-			return *this;
-		}
-		else
-		{
-			// Total size will be rounded up to the next power of two. Start with at least 64.
-			Extend(64);
+			va_list ArgPack;
+			va_copy(ArgPack, Args);
+			const int32 Result = TCString<C>::GetVarArgs(CurPos, RemainingSize, Fmt, ArgPack);
+			va_end(ArgPack);
+
+			if (Result >= 0 && Result < RemainingSize)
+			{
+				CurPos += Result;
+				return *this;
+			}
 		}
+
+		// Total size will be rounded up to the next power of two. Start with at least 64.
+		Extend(64);
 	}
 }
 
