[MIR] Bump "MaterialTranslationDDCVersion" GUID to fix mismatch in shadermaps after "NumVirtualTextureFeedbackRequests" was added.

#rnx

diff --git a/Engine/Source/Runtime/Core/Private/UObject/FortniteMainBranchObjectVersion.cpp b/Engine/Source/Runtime/Core/Private/UObject/FortniteMainBranchObjectVersion.cpp
--- a/Engine/Source/Runtime/Core/Private/UObject/FortniteMainBranchObjectVersion.cpp
+++ b/Engine/Source/Runtime/Core/Private/UObject/FortniteMainBranchObjectVersion.cpp
@@ -14,6 +14,6 @@ TMap<FGuid, FGuid> FFortniteMainBranchObjectVersion::GetSystemGuids()
 	SystemGuids.Add(DevGuids.Niagara_LatestScriptCompileVersion, FGuid("97F0FD7149D90D42B65D6557AF477D1A"));
 	SystemGuids.Add(DevGuids.SkeletalMeshDerivedDataVersion, FGuid("716424BCEC314E969ACFE36D2C2ECE49"));
 	SystemGuids.Add(DevGuids.STATICMESH_DERIVEDDATA_VER, FGuid("3ABF32149F9D4D838750368AB95573BE"));
-	SystemGuids.Add(DevGuids.MaterialTranslationDDCVersion, FGuid("2F6A9DD0B2674CF5B7564F4A14CC40E4"));
+	SystemGuids.Add(DevGuids.MaterialTranslationDDCVersion, FGuid("A3E8C83AF1484A74A986F0AA9181F798"));
 	return SystemGuids;
 }
