Added FGenericPlatformManualResetEvent::Poll()

Moved the implementation back into the header now that Oversubscription.h is available as a low-impact header to access oversubscription functionality.
Made bWait non-atomic because access to it needs to be guarded by the mutex anyway.

#rnx

diff --git a/Engine/Source/Runtime/Core/Private/GenericPlatform/GenericPlatformManualResetEvent.cpp b/Engine/Source/Runtime/Core/Private/GenericPlatform/GenericPlatformManualResetEvent.cpp
deleted file mode 100644
diff --git a/Engine/Source/Runtime/Core/Private/GenericPlatform/GenericPlatformManualResetEvent.h b/Engine/Source/Runtime/Core/Private/GenericPlatform/GenericPlatformManualResetEvent.h
--- a/Engine/Source/Runtime/Core/Private/GenericPlatform/GenericPlatformManualResetEvent.h
+++ b/Engine/Source/Runtime/Core/Private/GenericPlatform/GenericPlatformManualResetEvent.h
@@ -2,6 +2,7 @@
 
 #pragma once
 
+#include "Async/Fundamental/Oversubscription.h"
 #include "HAL/PlatformMath.h"
 #include "Misc/MonotonicTime.h"
 
@@ -34,7 +35,19 @@ public:
 	 */
 	void Reset()
 	{
-		bWait.store(true, std::memory_order_release);
+		std::unique_lock SelfLock(Lock);
+		bWait = true;
+	}
+
+	/**
+	 * Polls whether the event is in the notified state.
+	 *
+	 * @return True if notified, otherwise false.
+	 */
+	bool Poll()
+	{
+		std::unique_lock SelfLock(Lock);
+		return !bWait;
 	}
 
 	/**
@@ -44,7 +57,37 @@ public:
 	 */
 	void Wait()
 	{
-		WaitUntil(FMonotonicTimePoint::Infinity());
+		LowLevelTasks::FOversubscriptionScope _;
+		std::unique_lock SelfLock(Lock);
+		Condition.wait(SelfLock, [this] { return !bWait; });
+	}
+
+	/**
+	 * Waits for the wait time for Notify() to be called.
+	 *
+	 * Notify() may be called prior to WaitFor(), and this will return immediately in that case.
+	 *
+	 * @param WaitTime   Relative time after which waiting is canceled and the thread wakes.
+	 * @return True if Notify() was called before the wait time elapsed, otherwise false.
+	 */
+	bool WaitFor(FMonotonicTimeSpan WaitTime)
+	{
+		if (WaitTime.IsInfinity())
+		{
+			Wait();
+			return true;
+		}
+
+		if (WaitTime <= FMonotonicTimeSpan::Zero())
+		{
+			std::unique_lock SelfLock(Lock);
+			return !bWait;
+		}
+
+		LowLevelTasks::FOversubscriptionScope _;
+		std::unique_lock SelfLock(Lock);
+		const int64 WaitMs = FPlatformMath::CeilToInt64(WaitTime.ToMilliseconds());
+		return Condition.wait_for(SelfLock, std::chrono::milliseconds(WaitMs), [this] { return !bWait; });
 	}
 
 	/**
@@ -55,7 +98,10 @@ public:
 	 * @param WaitTime   Absolute time after which waiting is canceled and the thread wakes.
 	 * @return True if Notify() was called before the wait time elapsed, otherwise false.
 	 */
-	bool WaitUntil(FMonotonicTimePoint WaitTime);
+	bool WaitUntil(FMonotonicTimePoint WaitTime)
+	{
+		return WaitFor(WaitTime - FMonotonicTimePoint::Now());
+	}
 
 	/**
 	 * Notifies the waiting thread.
@@ -65,17 +111,17 @@ public:
 	 */
 	void Notify()
 	{
-		// We need this lock to ensure wait_for(SelfLock, std::chrono::milliseconds(WaitMs) does not return until Condition.notify_one() is fully called.
-		// Otherwise it's possible to destroy the Event right after Notify() call and while some other thread is still waiting on it
+		// The lock must be held by Notify() until it is finished accessing its members, otherwise a waiting thread
+		// may destroy the event after its wait but before Notify() has finished.
 		std::unique_lock SelfLock(Lock);
-		bWait.store(false, std::memory_order_relaxed);
+		bWait = false;
 		Condition.notify_one();
 	}
 
 private:
 	std::mutex Lock;
 	std::condition_variable Condition;
-	std::atomic<bool> bWait = true;
+	bool bWait = true;
 };
 
 } // UE::HAL::Private
