Optimized FMicrosoftManualResetEvent

diff --git a/Engine/Source/Runtime/Core/Private/Microsoft/MicrosoftPlatformManualResetEvent.cpp b/Engine/Source/Runtime/Core/Private/Microsoft/MicrosoftPlatformManualResetEvent.cpp
--- a/Engine/Source/Runtime/Core/Private/Microsoft/MicrosoftPlatformManualResetEvent.cpp
+++ b/Engine/Source/Runtime/Core/Private/Microsoft/MicrosoftPlatformManualResetEvent.cpp
@@ -2,57 +2,69 @@
 
 #include "Microsoft/MicrosoftPlatformManualResetEvent.h"
 
+#include "Async/Fundamental/Oversubscription.h"
 #include "HAL/PlatformMath.h"
-#include "Misc/MonotonicTime.h"
-#include "Async/Fundamental/Scheduler.h"
 #include "Microsoft/WindowsHWrapper.h"
 
 namespace UE::HAL::Private
 {
 
-void FMicrosoftPlatformManualResetEvent::Wait()
+FORCENOINLINE void FMicrosoftManualResetEvent::WaitSlow()
 {
-	WaitUntil(FMonotonicTimePoint::Infinity());
-}
+	LowLevelTasks::FOversubscriptionScope _;
 
-bool FMicrosoftPlatformManualResetEvent::WaitUntil(FMonotonicTimePoint WaitTime)
-{
-	bool bLocalWait = true;
-	if (WaitTime.IsInfinity())
+	for (;;)
 	{
-		// Let the scheduler know one of its thread might be waiting.
-		LowLevelTasks::FOversubscriptionScope _;
-
-		for (;;)
+		bool bLocalWait = bWait.load(std::memory_order_acquire);
+		if (!bLocalWait)
 		{
-			if (WaitOnAddress(&bWait, &bLocalWait, sizeof(bool), INFINITE) && !bWait.load(std::memory_order_acquire))
-			{
-				return true;
-			}
+			return;
 		}
+		WaitOnAddress(&bWait, &bLocalWait, sizeof(bool), INFINITE);
+	}
+}
+
+FORCENOINLINE bool FMicrosoftManualResetEvent::WaitForSlow(FMonotonicTimeSpan WaitTime)
+{
+	bool bLocalWait = bWait.load(std::memory_order_acquire);
+	if (!bLocalWait || WaitTime <= FMonotonicTimeSpan::Zero())
+	{
+		return !bLocalWait;
 	}
-	else
+
+	LowLevelTasks::FOversubscriptionScope _;
+
+	const DWORD WaitMs = WaitTime.IsInfinity() ? INFINITE : DWORD(FPlatformMath::CeilToInt64(WaitTime.ToMilliseconds()));
+	if (LIKELY(WaitOnAddress(&bWait, &bLocalWait, sizeof(bool), WaitMs) || GetLastError() == ERROR_TIMEOUT))
 	{
-		// Let the scheduler know one of its thread might be waiting.
-		LowLevelTasks::FOversubscriptionScope _(WaitTime - FMonotonicTimePoint::Now() > FMonotonicTimeSpan::Zero());
+		return !bWait.load(std::memory_order_acquire);
+	}
+
+	// Handle a spurious wake by waiting until the wait time has elapsed one more time because WaitUntilSlow
+	// handles spurious wakes in a loop and avoids exceeding the originally requested wake time by more than
+	// the typical variation due to scheduling imprecision.
+	return WaitUntilSlow(FMonotonicTimePoint::Now() + WaitTime);
+}
+
+FORCENOINLINE bool FMicrosoftManualResetEvent::WaitUntilSlow(FMonotonicTimePoint WaitTime)
+{
+	FMonotonicTimeSpan WaitSpan = WaitTime - FMonotonicTimePoint::Now();
+	LowLevelTasks::FOversubscriptionScope _(WaitSpan > FMonotonicTimeSpan::Zero());
 
-		for (;;)
+	for (;;)
+	{
+		bool bLocalWait = bWait.load(std::memory_order_acquire);
+		if (!bLocalWait || WaitSpan <= FMonotonicTimeSpan::Zero())
 		{
-			FMonotonicTimeSpan WaitSpan = WaitTime - FMonotonicTimePoint::Now();
-			if (WaitSpan <= FMonotonicTimeSpan::Zero())
-			{
-				return false;
-			}
-			const DWORD WaitMs = DWORD(FPlatformMath::CeilToInt64(WaitSpan.ToMilliseconds()));
-			if (WaitOnAddress(&bWait, &bLocalWait, sizeof(bool), WaitMs) && !bWait.load(std::memory_order_acquire))
-			{
-				return true;
-			}
+			return !bLocalWait;
 		}
+		const DWORD WaitMs = WaitTime.IsInfinity() ? INFINITE : DWORD(FPlatformMath::CeilToInt64(WaitSpan.ToMilliseconds()));
+		WaitOnAddress(&bWait, &bLocalWait, sizeof(bool), WaitMs);
+		WaitSpan = WaitTime - FMonotonicTimePoint::Now();
 	}
 }
 
-void FMicrosoftPlatformManualResetEvent::Notify()
+void FMicrosoftManualResetEvent::Notify()
 {
 	bWait.store(false, std::memory_order_release);
 	WakeByAddressSingle((void*)&bWait);
diff --git a/Engine/Source/Runtime/Core/Private/Microsoft/MicrosoftPlatformManualResetEvent.h b/Engine/Source/Runtime/Core/Private/Microsoft/MicrosoftPlatformManualResetEvent.h
--- a/Engine/Source/Runtime/Core/Private/Microsoft/MicrosoftPlatformManualResetEvent.h
+++ b/Engine/Source/Runtime/Core/Private/Microsoft/MicrosoftPlatformManualResetEvent.h
@@ -3,34 +3,57 @@
 #pragma once
 
 #include "CoreTypes.h"
+#include "Misc/MonotonicTime.h"
 
 #include <atomic>
 
-#define UE_API CORE_API
-
-namespace UE { struct FMonotonicTimePoint; }
-
 namespace UE::HAL::Private
 {
 
 /** @see FGenericPlatformManualResetEvent */
-class FMicrosoftPlatformManualResetEvent
+class FMicrosoftManualResetEvent
 {
 public:
-	FMicrosoftPlatformManualResetEvent() = default;
-	FMicrosoftPlatformManualResetEvent(const FMicrosoftPlatformManualResetEvent&) = delete;
-	FMicrosoftPlatformManualResetEvent& operator=(const FMicrosoftPlatformManualResetEvent&) = delete;
-
-	void Reset() { bWait.store(true, std::memory_order_relaxed); }
-
-	UE_API void Wait();
-	UE_API bool WaitUntil(FMonotonicTimePoint WaitTime);
-	UE_API void Notify();
+	FMicrosoftManualResetEvent() = default;
+	FMicrosoftManualResetEvent(const FMicrosoftManualResetEvent&) = delete;
+	FMicrosoftManualResetEvent& operator=(const FMicrosoftManualResetEvent&) = delete;
+
+	inline void Reset()
+	{
+		bWait.store(true, std::memory_order_relaxed);
+	}
+
+	inline bool Poll()
+	{
+		return !bWait.load(std::memory_order_acquire);
+	}
+
+	inline void Wait()
+	{
+		if (bWait.load(std::memory_order_acquire))
+		{
+			WaitSlow();
+		}
+	}
+
+	inline bool WaitFor(FMonotonicTimeSpan WaitTime)
+	{
+		return !bWait.load(std::memory_order_acquire) || WaitForSlow(WaitTime);
+	}
+
+	inline bool WaitUntil(FMonotonicTimePoint WaitTime)
+	{
+		return !bWait.load(std::memory_order_acquire) || WaitUntilSlow(WaitTime);
+	}
+
+	void Notify();
 
 private:
+	void WaitSlow();
+	bool WaitForSlow(FMonotonicTimeSpan WaitTime);
+	bool WaitUntilSlow(FMonotonicTimePoint WaitTime);
+
 	std::atomic<bool> bWait = true;
 };
 
 } // UE::HAL::Private
-
-#undef UE_API
diff --git a/Engine/Source/Runtime/Core/Private/Windows/WindowsPlatformManualResetEvent.h b/Engine/Source/Runtime/Core/Private/Windows/WindowsPlatformManualResetEvent.h
--- a/Engine/Source/Runtime/Core/Private/Windows/WindowsPlatformManualResetEvent.h
+++ b/Engine/Source/Runtime/Core/Private/Windows/WindowsPlatformManualResetEvent.h
@@ -4,4 +4,4 @@
 
 #include "Microsoft/MicrosoftPlatformManualResetEvent.h"
 
-namespace UE { using FPlatformManualResetEvent = UE::HAL::Private::FMicrosoftPlatformManualResetEvent; }
+namespace UE { using FPlatformManualResetEvent = UE::HAL::Private::FMicrosoftManualResetEvent; }
