[Core]
* Moved FTaskPriorityCVar ctor to source file since it was calling functions that was not inlined and can just as well be next to those functions

diff --git a/Engine/Source/Runtime/Core/Private/Async/TaskGraph.cpp b/Engine/Source/Runtime/Core/Private/Async/TaskGraph.cpp
--- a/Engine/Source/Runtime/Core/Private/Async/TaskGraph.cpp
+++ b/Engine/Source/Runtime/Core/Private/Async/TaskGraph.cpp
@@ -437,6 +437,18 @@ static ENamedThreads::Type ThreadPriorityFromChar(TCHAR InChar)
 	return ENamedThreads::BackgroundThreadPriority;
 }
 
+FAutoConsoleTaskPriority::FAutoConsoleTaskPriority(const TCHAR* Name, const TCHAR* Help, ENamedThreads::Type DefaultThreadPriority, ENamedThreads::Type DefaultTaskPriority, ENamedThreads::Type DefaultTaskPriorityIfForcedToNormalThreadPriority)
+	: RawSetting(ConfigStringFromPriorities(DefaultThreadPriority, DefaultTaskPriority, DefaultTaskPriorityIfForcedToNormalThreadPriority))
+	, FullHelpText(CreateFullHelpText(Name, Help))
+	, Variable(Name, RawSetting, *FullHelpText, FConsoleVariableDelegate::CreateRaw(this, &FAutoConsoleTaskPriority::OnSettingChanged), ECVF_Default)
+	, ThreadPriority(DefaultThreadPriority)
+	, TaskPriority(DefaultTaskPriority)
+	, TaskPriorityIfForcedToNormalThreadPriority(DefaultTaskPriorityIfForcedToNormalThreadPriority)
+{
+	// if you are asking for a hi or background thread priority, you must provide a separate task priority to use if those threads are not available.
+	check(TaskPriorityIfForcedToNormalThreadPriority != ENamedThreads::UnusedAnchor || ThreadPriority == ENamedThreads::NormalThreadPriority);
+}
+
 FString FAutoConsoleTaskPriority::CreateFullHelpText(const TCHAR* Name, const TCHAR* OriginalHelp)
 {
 	return FString::Printf(
diff --git a/Engine/Source/Runtime/Core/Private/Tasks/TaskPrivate.cpp b/Engine/Source/Runtime/Core/Private/Tasks/TaskPrivate.cpp
--- a/Engine/Source/Runtime/Core/Private/Tasks/TaskPrivate.cpp
+++ b/Engine/Source/Runtime/Core/Private/Tasks/TaskPrivate.cpp
@@ -388,6 +388,15 @@ namespace UE::Tasks
 		return false;
 	}
 
+	FTaskPriorityCVar::FTaskPriorityCVar(const TCHAR* Name, const TCHAR* Help, ETaskPriority DefaultPriority, EExtendedTaskPriority DefaultExtendedPriority)
+		: RawSetting(ConfigStringFromPriorities(DefaultPriority, DefaultExtendedPriority))
+		, FullHelpText(CreateFullHelpText(Name, Help))
+		, Variable(Name, RawSetting, *FullHelpText, FConsoleVariableDelegate::CreateRaw(this, &FTaskPriorityCVar::OnSettingChanged), ECVF_Default)
+		, Priority(DefaultPriority)
+		, ExtendedPriority(DefaultExtendedPriority)
+	{
+	}
+
 	FString FTaskPriorityCVar::CreateFullHelpText(const TCHAR* Name, const TCHAR* OriginalHelp)
 	{
 		TStringBuilder<1024> TaskPriorities;
diff --git a/Engine/Source/Runtime/Core/Public/Async/TaskGraphInterfaces.h b/Engine/Source/Runtime/Core/Public/Async/TaskGraphInterfaces.h
--- a/Engine/Source/Runtime/Core/Public/Async/TaskGraphInterfaces.h
+++ b/Engine/Source/Runtime/Core/Public/Async/TaskGraphInterfaces.h
@@ -218,22 +218,12 @@ class FAutoConsoleTaskPriority
 	ENamedThreads::Type TaskPriority;
 	ENamedThreads::Type TaskPriorityIfForcedToNormalThreadPriority;
 
-	static CORE_API FString CreateFullHelpText(const TCHAR* Name, const TCHAR* OriginalHelp);
-	static CORE_API FString ConfigStringFromPriorities(ENamedThreads::Type InThreadPriority, ENamedThreads::Type InTaskPriority, ENamedThreads::Type InTaskPriorityBackup);
-	CORE_API void OnSettingChanged(IConsoleVariable* Variable);
+	static FString CreateFullHelpText(const TCHAR* Name, const TCHAR* OriginalHelp);
+	static FString ConfigStringFromPriorities(ENamedThreads::Type InThreadPriority, ENamedThreads::Type InTaskPriority, ENamedThreads::Type InTaskPriorityBackup);
+	void OnSettingChanged(IConsoleVariable* Variable);
 
 public:
-	FAutoConsoleTaskPriority(const TCHAR* Name, const TCHAR* Help, ENamedThreads::Type DefaultThreadPriority, ENamedThreads::Type DefaultTaskPriority, ENamedThreads::Type DefaultTaskPriorityIfForcedToNormalThreadPriority = ENamedThreads::UnusedAnchor)
-		: RawSetting(ConfigStringFromPriorities(DefaultThreadPriority, DefaultTaskPriority, DefaultTaskPriorityIfForcedToNormalThreadPriority))
-		, FullHelpText(CreateFullHelpText(Name, Help))
-		, Variable(Name, RawSetting, *FullHelpText, FConsoleVariableDelegate::CreateRaw(this, &FAutoConsoleTaskPriority::OnSettingChanged), ECVF_Default)
-		, ThreadPriority(DefaultThreadPriority)
-		, TaskPriority(DefaultTaskPriority)
-		, TaskPriorityIfForcedToNormalThreadPriority(DefaultTaskPriorityIfForcedToNormalThreadPriority)
-	{
-		// if you are asking for a hi or background thread priority, you must provide a separate task priority to use if those threads are not available.
-		check(TaskPriorityIfForcedToNormalThreadPriority != ENamedThreads::UnusedAnchor || ThreadPriority == ENamedThreads::NormalThreadPriority);
-	}
+	CORE_API FAutoConsoleTaskPriority(const TCHAR* Name, const TCHAR* Help, ENamedThreads::Type DefaultThreadPriority, ENamedThreads::Type DefaultTaskPriority, ENamedThreads::Type DefaultTaskPriorityIfForcedToNormalThreadPriority = ENamedThreads::UnusedAnchor);
 
 	FORCEINLINE ENamedThreads::Type Get(ENamedThreads::Type Thread = ENamedThreads::AnyThread)
 	{
diff --git a/Engine/Source/Runtime/Core/Public/Tasks/Task.h b/Engine/Source/Runtime/Core/Public/Tasks/Task.h
--- a/Engine/Source/Runtime/Core/Public/Tasks/Task.h
+++ b/Engine/Source/Runtime/Core/Public/Tasks/Task.h
@@ -522,14 +522,7 @@ namespace UE::Tasks
 	class FTaskPriorityCVar
 	{
 	public:
-		FTaskPriorityCVar(const TCHAR* Name, const TCHAR* Help, ETaskPriority DefaultPriority, EExtendedTaskPriority DefaultExtendedPriority)
-			: RawSetting(ConfigStringFromPriorities(DefaultPriority, DefaultExtendedPriority))
-			, FullHelpText(CreateFullHelpText(Name, Help))
-			, Variable(Name, RawSetting, *FullHelpText, FConsoleVariableDelegate::CreateRaw(this, &FTaskPriorityCVar::OnSettingChanged), ECVF_Default)
-			, Priority(DefaultPriority)
-			, ExtendedPriority(DefaultExtendedPriority)
-		{
-		}
+		CORE_API FTaskPriorityCVar(const TCHAR* Name, const TCHAR* Help, ETaskPriority DefaultPriority, EExtendedTaskPriority DefaultExtendedPriority);
 
 		ETaskPriority GetTaskPriority() const
 		{
