[Log] Remove warning when log category verbosity is misconfigured.

This is to prevent a deadlock on certain platforms when this warning is emitted before GConfig has been initialized.
It was introduced in CL 41611908. We'll opt in for an ensure condition once it's safe to do so.

#rnx

diff --git a/Engine/Source/Runtime/Core/Private/Logging/LogCategory.cpp b/Engine/Source/Runtime/Core/Private/Logging/LogCategory.cpp
--- a/Engine/Source/Runtime/Core/Private/Logging/LogCategory.cpp
+++ b/Engine/Source/Runtime/Core/Private/Logging/LogCategory.cpp
@@ -43,6 +43,8 @@ void FLogCategoryBase::SetVerbosity(ELogVerbosity::Type InNewVerbosity)
 	const ELogVerbosity::Type OldVerbosity = Verbosity;
 	const ELogVerbosity::Type NewVerbosity = (ELogVerbosity::Type)(InNewVerbosity & ELogVerbosity::VerbosityMask);
 	// regularize the verbosity to be at most whatever we were compiled with
+	//@todo-lh - Turn back to the following ensure condition once it's safe to do so:
+	// if (ensureMsgf(NewVerbosity <= CompileTimeVerbosity, TEXT("Log category '%s' tried to set verbosity (%s) higher than its compile-time verbosity level (%s)"), *GetCategoryName().ToString(), ToString(NewVerbosity), ToString(CompileTimeVerbosity)))
 	if (NewVerbosity <= CompileTimeVerbosity)
 	{
 		Verbosity = NewVerbosity;
@@ -50,7 +52,6 @@ void FLogCategoryBase::SetVerbosity(ELogVerbosity::Type InNewVerbosity)
 	else
 	{
 		Verbosity = CompileTimeVerbosity;
-		UE_LOG(LogCore, Warning, TEXT("Log category '%s' tried to set verbosity (%s) higher than its compile-time verbosity level (%s)"), *GetCategoryName().ToString(), ToString(NewVerbosity), ToString(CompileTimeVerbosity));
 	}
 	DebugBreakOnLog = !!(InNewVerbosity & ELogVerbosity::BreakOnLog);
 	checkSlow(!(Verbosity & ELogVerbosity::BreakOnLog)); // this bit is factored out of this variable, always
