[CoreGlobals]
- Add IsInParallelLoadingThread() in preparation of going wide in the loader

diff --git a/Engine/Source/Runtime/Core/Private/HAL/ThreadingBase.cpp b/Engine/Source/Runtime/Core/Private/HAL/ThreadingBase.cpp
--- a/Engine/Source/Runtime/Core/Private/HAL/ThreadingBase.cpp
+++ b/Engine/Source/Runtime/Core/Private/HAL/ThreadingBase.cpp
@@ -205,6 +205,11 @@ CORE_API bool IsInParallelGameThread()
 	return FTaskTagScope::IsCurrentTag(ETaskTag::EParallelGameThread);
 }
 
+CORE_API bool IsInParallelLoadingThread()
+{
+	return FTaskTagScope::IsCurrentTag(ETaskTag::EParallelLoadingThread);
+}
+
 CORE_API bool IsInWorkerThread()
 {
 	return LowLevelTasks::FScheduler::Get().IsWorkerThread();
diff --git a/Engine/Source/Runtime/Core/Public/CoreGlobals.h b/Engine/Source/Runtime/Core/Public/CoreGlobals.h
--- a/Engine/Source/Runtime/Core/Public/CoreGlobals.h
+++ b/Engine/Source/Runtime/Core/Public/CoreGlobals.h
@@ -645,6 +645,7 @@ enum class ETaskTag : int32
 	EParallelRenderingThread	= ERenderingThread | EParallelThread,
 	EParallelGameThread			= EGameThread | EParallelThread,
 	EParallelRhiThread			= ERhiThread | EParallelThread,
+	EParallelLoadingThread		= EAsyncLoadingThread | EParallelThread,
 };
 
 ENUM_CLASS_FLAGS(ETaskTag);
@@ -723,6 +724,9 @@ extern CORE_API bool IsInGameThread();
 /** @return True if called from the game thread in a parallel for. */
 extern CORE_API bool IsInParallelGameThread();
 
+/** @return True if called from a parallel loading thread task. */
+extern CORE_API bool IsInParallelLoadingThread();
+
 /** @return True if called from a worker thread. */
 extern CORE_API bool IsInWorkerThread();
 
diff --git a/Engine/Source/Runtime/Engine/Private/StaticMesh.cpp b/Engine/Source/Runtime/Engine/Private/StaticMesh.cpp
--- a/Engine/Source/Runtime/Engine/Private/StaticMesh.cpp
+++ b/Engine/Source/Runtime/Engine/Private/StaticMesh.cpp
@@ -6951,7 +6951,7 @@ void UStaticMesh::Serialize(FArchive& Ar)
 
 #if WITH_EDITOR
 	// Don't try to finish compilation while still in the loader as we can't possibly have our mesh being compiled if it has not been loaded yet.
-	if (!IsInAsyncLoadingThread())
+	if (!IsInAsyncLoadingThread() && !IsInParallelLoadingThread())
 	{
 		// Since UPROPERTY are accessed directly by offset during serialization instead of using accessors, 
 		// the protection put in place to automatically finish compilation if a locked property is accessed will not work. 
diff --git a/Engine/Source/Runtime/Engine/Private/StaticMeshSourceData.cpp b/Engine/Source/Runtime/Engine/Private/StaticMeshSourceData.cpp
--- a/Engine/Source/Runtime/Engine/Private/StaticMeshSourceData.cpp
+++ b/Engine/Source/Runtime/Engine/Private/StaticMeshSourceData.cpp
@@ -394,7 +394,7 @@ void FStaticMeshSourceModel::SerializeBulkData(FArchive& Ar, UObject* Owner)
 		if (StaticMeshDescriptionBulkData == nullptr)
 		{
 			// For now, we assume that serialization will only occur on the game thread
-			ensure(IsInGameThread() || IsInAsyncLoadingThread());
+			ensure(IsInGameThread() || IsInAsyncLoadingThread() || IsInParallelLoadingThread());
 			StaticMeshDescriptionBulkData = NewObject<UStaticMeshDescriptionBulkData>(Owner, NAME_None, RF_Transactional);
 		}
 	}
diff --git a/Engine/Source/Runtime/Slate/Public/Framework/Application/SlateApplication.h b/Engine/Source/Runtime/Slate/Public/Framework/Application/SlateApplication.h
--- a/Engine/Source/Runtime/Slate/Public/Framework/Application/SlateApplication.h
+++ b/Engine/Source/Runtime/Slate/Public/Framework/Application/SlateApplication.h
@@ -319,7 +319,7 @@ public:
 	static FSlateApplication& Get()
 	{
 		check( CurrentApplication.IsValid() );
-		check( IsInGameThread() || IsInSlateThread() || IsInAsyncLoadingThread() );
+		check( IsInGameThread() || IsInSlateThread() || IsInAsyncLoadingThread() || IsInParallelLoadingThread() );
 		return *CurrentApplication;
 	}
 
