Workaround for code generation bug on Windows Arm64
- fixes invalid TimeManagementTests failures

#rnx

diff --git a/Engine/Source/Runtime/Core/Public/Misc/FrameTime.h b/Engine/Source/Runtime/Core/Public/Misc/FrameTime.h
--- a/Engine/Source/Runtime/Core/Public/Misc/FrameTime.h
+++ b/Engine/Source/Runtime/Core/Public/Misc/FrameTime.h
@@ -267,7 +267,7 @@ inline FFrameTime::FFrameTime(FFrameNumber InFrameNumber, float InSubFrame)
 {
 	// Hack to ensure that SubFrames are in a sensible range of precision to work around
 	// problems with FloorToXYZ returning the wrong thing for very small negative numbers
-	SubFrame = FMath::Clamp(SubFrame + 0.5f - 0.5f, 0.f, MaxSubframe);
+	SubFrame = FMath::Clamp(InSubFrame + 0.5f - 0.5f, 0.f, MaxSubframe);
 	checkSlow(InSubFrame >= 0.f && InSubFrame < 1.f);
 }
 
