[Core]
* Added dllstorage to UE_DEFINE_THREAD_SINGLETON_TLS because header units does not seem to like linking otherwise.

diff --git a/Engine/Source/Runtime/Core/Private/Misc/CoreMisc.cpp b/Engine/Source/Runtime/Core/Private/Misc/CoreMisc.cpp
--- a/Engine/Source/Runtime/Core/Private/Misc/CoreMisc.cpp
+++ b/Engine/Source/Runtime/Core/Private/Misc/CoreMisc.cpp
@@ -361,7 +361,7 @@ FBoolConfigValueHelper::FBoolConfigValueHelper(const TCHAR* Section, const TCHAR
 FScriptExceptionHandler
 ----------------------------------------------------------------------------*/
 
-UE_DEFINE_THREAD_SINGLETON_TLS(FScriptExceptionHandler)
+UE_DEFINE_THREAD_SINGLETON_TLS(FScriptExceptionHandler, CORE_API)
 
 FScriptExceptionHandlerFunc FScriptExceptionHandler::DefaultExceptionHandler = &FScriptExceptionHandler::LoggingExceptionHandler;
 
diff --git a/Engine/Source/Runtime/Core/Private/Misc/MemStack.cpp b/Engine/Source/Runtime/Core/Private/Misc/MemStack.cpp
--- a/Engine/Source/Runtime/Core/Private/Misc/MemStack.cpp
+++ b/Engine/Source/Runtime/Core/Private/Misc/MemStack.cpp
@@ -409,7 +409,7 @@ bool FMemStackBase::ContainsPointer(const void* Pointer) const
 	return false;
 }
 
-UE_DEFINE_THREAD_SINGLETON_TLS(FMemStack)
+UE_DEFINE_THREAD_SINGLETON_TLS(FMemStack, CORE_API)
 
 
 FORCENOINLINE void UE::Core::Private::OnInvalidMemStackAllocatorNum(int32 NewNum, SIZE_T NumBytesPerElement)
diff --git a/Engine/Source/Runtime/Core/Private/Stats/ThreadIdleStats.cpp b/Engine/Source/Runtime/Core/Private/Stats/ThreadIdleStats.cpp
--- a/Engine/Source/Runtime/Core/Private/Stats/ThreadIdleStats.cpp
+++ b/Engine/Source/Runtime/Core/Private/Stats/ThreadIdleStats.cpp
@@ -8,7 +8,7 @@ UE_TRACE_CHANNEL_DEFINE(ThreadIdleScopeChannel);
 TRACE_CPUPROFILER_EVENT_DECLARE(ThreadIdleScopeTraceEventId);
 #endif
 
-UE_DEFINE_THREAD_SINGLETON_TLS(UE::Stats::FThreadIdleStats)
+UE_DEFINE_THREAD_SINGLETON_TLS(UE::Stats::FThreadIdleStats, CORE_API)
 
 namespace UE::Stats
 {
diff --git a/Engine/Source/Runtime/Core/Public/HAL/ThreadSingleton.h b/Engine/Source/Runtime/Core/Public/HAL/ThreadSingleton.h
--- a/Engine/Source/Runtime/Core/Public/HAL/ThreadSingleton.h
+++ b/Engine/Source/Runtime/Core/Public/HAL/ThreadSingleton.h
@@ -33,7 +33,7 @@ public:
 };
 
 #define UE_DECLARE_THREAD_SINGLETON_TLS(Type, Api) template<> Api uint32& TThreadSingleton<Type>::GetTlsSlot();
-#define UE_DEFINE_THREAD_SINGLETON_TLS(Type) template<> uint32& TThreadSingleton<Type>::GetTlsSlot() { static uint32 TlsSlot = FPlatformTLS::InvalidTlsSlot; return TlsSlot; }
+#define UE_DEFINE_THREAD_SINGLETON_TLS(Type, Api) template<> Api uint32& TThreadSingleton<Type>::GetTlsSlot() { static uint32 TlsSlot = FPlatformTLS::InvalidTlsSlot; return TlsSlot; }
 
 /**
  * This a special version of singleton. It means that there is created only one instance for each thread.
diff --git a/Engine/Source/Runtime/CoreUObject/Private/UObject/SoftObjectPath.cpp b/Engine/Source/Runtime/CoreUObject/Private/UObject/SoftObjectPath.cpp
--- a/Engine/Source/Runtime/CoreUObject/Private/UObject/SoftObjectPath.cpp
+++ b/Engine/Source/Runtime/CoreUObject/Private/UObject/SoftObjectPath.cpp
@@ -1051,7 +1051,7 @@ FSoftClassPath FSoftClassPath::GetOrCreateIDForClass(const UClass *InClass)
 	return FSoftClassPath(InClass);
 }
 
-UE_DEFINE_THREAD_SINGLETON_TLS(FSoftObjectPathThreadContext)
+UE_DEFINE_THREAD_SINGLETON_TLS(FSoftObjectPathThreadContext, COREUOBJECT_API)
 
 bool FSoftObjectPathThreadContext::GetSerializationOptions(FName& OutPackageName, FName& OutPropertyName, ESoftObjectPathCollectType& OutCollectType, ESoftObjectPathSerializeType& OutSerializeType, FArchive* Archive) const
 {
diff --git a/Engine/Source/Runtime/CoreUObject/Private/UObject/UObjectThreadContext.cpp b/Engine/Source/Runtime/CoreUObject/Private/UObject/UObjectThreadContext.cpp
--- a/Engine/Source/Runtime/CoreUObject/Private/UObject/UObjectThreadContext.cpp
+++ b/Engine/Source/Runtime/CoreUObject/Private/UObject/UObjectThreadContext.cpp
@@ -13,7 +13,7 @@
 
 DEFINE_LOG_CATEGORY(LogUObjectThreadContext);
 
-UE_DEFINE_THREAD_SINGLETON_TLS(FUObjectThreadContext)
+UE_DEFINE_THREAD_SINGLETON_TLS(FUObjectThreadContext, COREUOBJECT_API)
 
 FUObjectThreadContext::FUObjectThreadContext()
 : IsRoutingPostLoad(false)
diff --git a/Engine/Source/Runtime/Experimental/Chaos/Private/ChaosDebugDraw/ChaosDDContext.cpp b/Engine/Source/Runtime/Experimental/Chaos/Private/ChaosDebugDraw/ChaosDDContext.cpp
--- a/Engine/Source/Runtime/Experimental/Chaos/Private/ChaosDebugDraw/ChaosDDContext.cpp
+++ b/Engine/Source/Runtime/Experimental/Chaos/Private/ChaosDebugDraw/ChaosDDContext.cpp
@@ -11,7 +11,7 @@ DEFINE_LOG_CATEGORY(LogChaosDD);
 
 #if CHAOS_DEBUG_DRAW
 
-UE_DEFINE_THREAD_SINGLETON_TLS(ChaosDD::Private::FChaosDDContext)
+UE_DEFINE_THREAD_SINGLETON_TLS(ChaosDD::Private::FChaosDDContext, CHAOS_API)
 
 namespace ChaosDD::Private
 {
diff --git a/Engine/Source/Runtime/Experimental/Chaos/Private/ChaosVisualDebugger/ChaosVDContextProvider.cpp b/Engine/Source/Runtime/Experimental/Chaos/Private/ChaosVisualDebugger/ChaosVDContextProvider.cpp
--- a/Engine/Source/Runtime/Experimental/Chaos/Private/ChaosVisualDebugger/ChaosVDContextProvider.cpp
+++ b/Engine/Source/Runtime/Experimental/Chaos/Private/ChaosVisualDebugger/ChaosVDContextProvider.cpp
@@ -6,7 +6,7 @@
 
 #if WITH_CHAOS_VISUAL_DEBUGGER
 
-UE_DEFINE_THREAD_SINGLETON_TLS(FChaosVDThreadContext)
+UE_DEFINE_THREAD_SINGLETON_TLS(FChaosVDThreadContext, CHAOS_API)
 
 bool FChaosVDThreadContext::GetCurrentContext(FChaosVDContext& OutContext)
 {
diff --git a/Engine/Source/Runtime/Experimental/Chaos/Private/Framework/Threading.cpp b/Engine/Source/Runtime/Experimental/Chaos/Private/Framework/Threading.cpp
--- a/Engine/Source/Runtime/Experimental/Chaos/Private/Framework/Threading.cpp
+++ b/Engine/Source/Runtime/Experimental/Chaos/Private/Framework/Threading.cpp
@@ -4,7 +4,7 @@
 #include "ChaosCheck.h"
 
 #if PHYSICS_THREAD_CONTEXT
-UE_DEFINE_THREAD_SINGLETON_TLS(Chaos::FPhysicsThreadContext)
+UE_DEFINE_THREAD_SINGLETON_TLS(Chaos::FPhysicsThreadContext, CHAOS_API)
 #endif
 
 namespace Chaos
