- Added some calls to FConfigBranch::SafeReload() when adding/removing dynamic layers

diff --git a/Engine/Source/Runtime/Core/Private/Misc/ConfigCacheIni.cpp b/Engine/Source/Runtime/Core/Private/Misc/ConfigCacheIni.cpp
--- a/Engine/Source/Runtime/Core/Private/Misc/ConfigCacheIni.cpp
+++ b/Engine/Source/Runtime/Core/Private/Misc/ConfigCacheIni.cpp
@@ -3152,6 +3152,11 @@ bool FConfigBranch::AddDynamicLayersToHierarchy(const TArray<FString>& Filenames
 
 bool FConfigBranch::AddDynamicLayersToHierarchy(const TArray<FDynamicLayerInfo>& Layers, FConfigModificationTracker* ModificationTracker, TSet<FString>* GlobalConfigFileCache, TSet<FString>* PluginConfigFileCache, bool bForceFullDynamicLayerUpdate)
 {
+	if (bIsSafeUnloaded)
+	{
+		SafeReload();
+	}
+
 	static bool bDumpIniLoadInfo = FParse::Param(FCommandLine::Get(), TEXT("dumpiniloads"));
 
 	bool bFoundAFile = false;
@@ -3301,6 +3306,11 @@ bool FConfigBranch::AddDynamicLayersToHierarchy(const TArray<FDynamicLayerInfo>&
 
 bool FConfigBranch::AddDynamicLayerStringToHierarchy(const FString& Filename, const FString& Contents, FName Tag, DynamicLayerPriority Priority, FConfigModificationTracker* ModificationTracker)
 {
+	if (bIsSafeUnloaded)
+	{
+		SafeReload();
+	}
+
 	bool bInsertedAtEnd = false;
 
 	// calculate a patch so we don't lose in-memory changes
@@ -3405,6 +3415,11 @@ bool FConfigBranch::RemoveDynamicLayerFromHierarchy(const FString& Filename, FCo
 
 bool FConfigBranch::RemoveDynamicLayersFromHierarchy(const TArray<FString>& Filenames, FConfigModificationTracker* ModificationTracker)
 {
+	if (bIsSafeUnloaded)
+	{
+		SafeReload();
+	}
+
 	if (ReplayMethod == EBranchReplayMethod::NoReplay)
 	{
 		UE_LOG(LogConfig, Warning, TEXT("Attempted to remove dynamic layer(s) from branch %s, but it is using NoReplay mode, so this cannot work. Skipping."), *IniName.ToString());
