Enabled TSharedPtr and TUniquePtr to be usable as keys in sorted containers, and TUniquePtr to be usable as keys in hashed containers.

diff --git a/Engine/Source/Runtime/Core/Public/Templates/SharedPointer.h b/Engine/Source/Runtime/Core/Public/Templates/SharedPointer.h
--- a/Engine/Source/Runtime/Core/Public/Templates/SharedPointer.h
+++ b/Engine/Source/Runtime/Core/Public/Templates/SharedPointer.h
@@ -1820,6 +1820,54 @@ template< class ObjectTypeA, class ObjectTypeB, ESPMode Mode >
 }
 
 
+/**
+ * Global less than operator for TSharedRef
+ *
+ * @return  True if InSharedRefA is less than InSharedRefB
+ */
+template< class ObjectTypeA, class ObjectTypeB, ESPMode Mode >
+[[nodiscard]] FORCEINLINE bool operator<( TSharedRef< ObjectTypeA, Mode > const& InSharedRefA, TSharedRef< ObjectTypeB, Mode > const& InSharedRefB )
+{
+	return &( InSharedRefA.Get() ) < &( InSharedRefB.Get() );
+}
+
+
+/**
+ * Global greater than operator for TSharedRef
+ *
+ * @return  True if InSharedRefA is greater than InSharedRefB
+ */
+template< class ObjectTypeA, class ObjectTypeB, ESPMode Mode >
+[[nodiscard]] FORCEINLINE bool operator>( TSharedRef< ObjectTypeA, Mode > const& InSharedRefA, TSharedRef< ObjectTypeB, Mode > const& InSharedRefB )
+{
+	return &( InSharedRefA.Get() ) > &( InSharedRefB.Get() );
+}
+
+
+/**
+ * Global less than or equal to operator for TSharedRef
+ *
+ * @return  True if InSharedRefA is less than or equal to InSharedRefB
+ */
+template< class ObjectTypeA, class ObjectTypeB, ESPMode Mode >
+[[nodiscard]] FORCEINLINE bool operator<=( TSharedRef< ObjectTypeA, Mode > const& InSharedRefA, TSharedRef< ObjectTypeB, Mode > const& InSharedRefB )
+{
+	return &( InSharedRefA.Get() ) <= &( InSharedRefB.Get() );
+}
+
+
+/**
+ * Global greater than or equal to operator for TSharedRef
+ *
+ * @return  True if InSharedRefA is greater than or equal to InSharedRefB
+ */
+template< class ObjectTypeA, class ObjectTypeB, ESPMode Mode >
+[[nodiscard]] FORCEINLINE bool operator>=( TSharedRef< ObjectTypeA, Mode > const& InSharedRefA, TSharedRef< ObjectTypeB, Mode > const& InSharedRefB )
+{
+	return &( InSharedRefA.Get() ) >= &( InSharedRefB.Get() );
+}
+
+
 /**
  * Global inequality operator for TSharedRef
  *
@@ -1844,6 +1892,54 @@ template< class ObjectTypeA, class ObjectTypeB, ESPMode Mode >
 }
 
 
+/**
+ * Global less than operator for TSharedPtr
+ *
+ * @return  True if InSharedPtrA is less than InSharedPtrB
+ */
+template< class ObjectTypeA, class ObjectTypeB, ESPMode Mode >
+[[nodiscard]] FORCEINLINE bool operator<( TSharedPtr< ObjectTypeA, Mode > const& InSharedPtrA, TSharedPtr< ObjectTypeB, Mode > const& InSharedPtrB )
+{
+	return InSharedPtrA.Get() < InSharedPtrB.Get();
+}
+
+
+/**
+ * Global greater than operator for TSharedPtr
+ *
+ * @return  True if InSharedPtrA is greater than InSharedPtrB
+ */
+template< class ObjectTypeA, class ObjectTypeB, ESPMode Mode >
+[[nodiscard]] FORCEINLINE bool operator>( TSharedPtr< ObjectTypeA, Mode > const& InSharedPtrA, TSharedPtr< ObjectTypeB, Mode > const& InSharedPtrB )
+{
+	return InSharedPtrA.Get() > InSharedPtrB.Get();
+}
+
+
+/**
+ * Global less than or equal to operator for TSharedPtr
+ *
+ * @return  True if InSharedPtrA is less than or equal to InSharedPtrB
+ */
+template< class ObjectTypeA, class ObjectTypeB, ESPMode Mode >
+[[nodiscard]] FORCEINLINE bool operator<=( TSharedPtr< ObjectTypeA, Mode > const& InSharedPtrA, TSharedPtr< ObjectTypeB, Mode > const& InSharedPtrB )
+{
+	return InSharedPtrA.Get() <= InSharedPtrB.Get();
+}
+
+
+/**
+ * Global greater than or equal to operator for TSharedPtr
+ *
+ * @return  True if InSharedPtrA is greater than or equal to InSharedPtrB
+ */
+template< class ObjectTypeA, class ObjectTypeB, ESPMode Mode >
+[[nodiscard]] FORCEINLINE bool operator>=( TSharedPtr< ObjectTypeA, Mode > const& InSharedPtrA, TSharedPtr< ObjectTypeB, Mode > const& InSharedPtrB )
+{
+	return InSharedPtrA.Get() >= InSharedPtrB.Get();
+}
+
+
 /**
  * Global equality operator for TSharedPtr
  *
@@ -1964,6 +2060,54 @@ template< class ObjectTypeA, class ObjectTypeB, ESPMode Mode >
 }
 
 
+/**
+ * Global less than operator for TWeakPtr
+ *
+ * @return  True if InWeakPtrA is less than InWeakPtrB
+ */
+template< class ObjectTypeA, class ObjectTypeB, ESPMode Mode >
+[[nodiscard]] FORCEINLINE bool operator<( TWeakPtr< ObjectTypeA, Mode > const& InWeakPtrA, TWeakPtr< ObjectTypeB, Mode > const& InWeakPtrB )
+{
+	return InWeakPtrA.Pin().Get() < InWeakPtrB.Pin().Get();
+}
+
+
+/**
+ * Global greater than operator for TWeakPtr
+ *
+ * @return  True if InWeakPtrA is greater than InWeakPtrB
+ */
+template< class ObjectTypeA, class ObjectTypeB, ESPMode Mode >
+[[nodiscard]] FORCEINLINE bool operator>( TWeakPtr< ObjectTypeA, Mode > const& InWeakPtrA, TWeakPtr< ObjectTypeB, Mode > const& InWeakPtrB )
+{
+	return InWeakPtrA.Pin().Get() > InWeakPtrB.Pin().Get();
+}
+
+
+/**
+ * Global less than or equal to operator for TWeakPtr
+ *
+ * @return  True if InWeakPtrA is less than or equal to InWeakPtrB
+ */
+template< class ObjectTypeA, class ObjectTypeB, ESPMode Mode >
+[[nodiscard]] FORCEINLINE bool operator<=( TWeakPtr< ObjectTypeA, Mode > const& InWeakPtrA, TWeakPtr< ObjectTypeB, Mode > const& InWeakPtrB )
+{
+	return InWeakPtrA.Pin().Get() <= InWeakPtrB.Pin().Get();
+}
+
+
+/**
+ * Global greater than or equal to operator for TWeakPtr
+ *
+ * @return  True if InWeakPtrA is greater than or equal to InWeakPtrB
+ */
+template< class ObjectTypeA, class ObjectTypeB, ESPMode Mode >
+[[nodiscard]] FORCEINLINE bool operator>=( TWeakPtr< ObjectTypeA, Mode > const& InWeakPtrA, TWeakPtr< ObjectTypeB, Mode > const& InWeakPtrB )
+{
+	return InWeakPtrA.Pin().Get() >= InWeakPtrB.Pin().Get();
+}
+
+
 /**
  * Global equality operator for TWeakPtr
  *
diff --git a/Engine/Source/Runtime/Core/Public/Templates/UniquePtr.h b/Engine/Source/Runtime/Core/Public/Templates/UniquePtr.h
--- a/Engine/Source/Runtime/Core/Public/Templates/UniquePtr.h
+++ b/Engine/Source/Runtime/Core/Public/Templates/UniquePtr.h
@@ -393,6 +393,62 @@ public:
 		return Get() == Rhs.Get();
 	}
 
+	/**
+	 * Less-than comparison operator
+	 *
+	 * @param Lhs The first TUniquePtr to compare.
+	 * @param Rhs The second TUniquePtr to compare.
+	 *
+	 * @return true if the inner pointer of the left TUniquePtr is less than the inner pointer of the right, false otherwise.
+	 */
+	template <typename RhsT>
+	[[nodiscard]] FORCEINLINE bool operator<(const TUniquePtr<RhsT>& Rhs) const
+	{
+		return Get() < Rhs.Get();
+	}
+
+	/**
+	 * Greater-than comparison operator
+	 *
+	 * @param Lhs The first TUniquePtr to compare.
+	 * @param Rhs The second TUniquePtr to compare.
+	 *
+	 * @return true if the inner pointer of the left TUniquePtr is greater than the inner pointer of the right, false otherwise.
+	 */
+	template <typename RhsT>
+	[[nodiscard]] FORCEINLINE bool operator>(const TUniquePtr<RhsT>& Rhs) const
+	{
+		return Get() > Rhs.Get();
+	}
+
+	/**
+	 * Less-than-or-equals comparison operator
+	 *
+	 * @param Lhs The first TUniquePtr to compare.
+	 * @param Rhs The second TUniquePtr to compare.
+	 *
+	 * @return true if the inner pointer of the left TUniquePtr is less than or equal to the inner pointer of the right, false otherwise.
+	 */
+	template <typename RhsT>
+	[[nodiscard]] FORCEINLINE bool operator<=(const TUniquePtr<RhsT>& Rhs) const
+	{
+		return Get() <= Rhs.Get();
+	}
+
+	/**
+	 * Greater-than-or-equals comparison operator
+	 *
+	 * @param Lhs The first TUniquePtr to compare.
+	 * @param Rhs The second TUniquePtr to compare.
+	 *
+	 * @return true if the inner pointer of the left TUniquePtr is greater than or equal to the inner pointer of the right, false otherwise.
+	 */
+	template <typename RhsT>
+	[[nodiscard]] FORCEINLINE bool operator>=(const TUniquePtr<RhsT>& Rhs) const
+	{
+		return Get() >= Rhs.Get();
+	}
+
 	/**
 	 * Equality comparison operator against nullptr.
 	 *
@@ -432,6 +488,11 @@ public:
 		return IsValid();
 	}
 #endif
+
+	friend uint32 GetTypeHash(const TUniquePtr& UniquePtr)
+	{
+		return GetTypeHash(UniquePtr.Get());
+	}
 };
 
 template <typename T, typename Deleter>
@@ -711,6 +772,62 @@ public:
 		return Get() == Rhs.Get();
 	}
 
+	/**
+	 * Less-than comparison operator
+	 *
+	 * @param Lhs The first TUniquePtr to compare.
+	 * @param Rhs The second TUniquePtr to compare.
+	 *
+	 * @return true if the inner pointer of the left TUniquePtr is less than the inner pointer of the right, false otherwise.
+	 */
+	template <typename RhsT>
+	[[nodiscard]] FORCEINLINE bool operator<(const TUniquePtr<RhsT>& Rhs) const
+	{
+		return Get() < Rhs.Get();
+	}
+
+	/**
+	 * Greater-than comparison operator
+	 *
+	 * @param Lhs The first TUniquePtr to compare.
+	 * @param Rhs The second TUniquePtr to compare.
+	 *
+	 * @return true if the inner pointer of the left TUniquePtr is greater than the inner pointer of the right, false otherwise.
+	 */
+	template <typename RhsT>
+	[[nodiscard]] FORCEINLINE bool operator>(const TUniquePtr<RhsT>& Rhs) const
+	{
+		return Get() > Rhs.Get();
+	}
+
+	/**
+	 * Less-than-or-equals comparison operator
+	 *
+	 * @param Lhs The first TUniquePtr to compare.
+	 * @param Rhs The second TUniquePtr to compare.
+	 *
+	 * @return true if the inner pointer of the left TUniquePtr is less than or equal to the inner pointer of the right, false otherwise.
+	 */
+	template <typename RhsT>
+	[[nodiscard]] FORCEINLINE bool operator<=(const TUniquePtr<RhsT>& Rhs) const
+	{
+		return Get() <= Rhs.Get();
+	}
+
+	/**
+	 * Greater-than-or-equals comparison operator
+	 *
+	 * @param Lhs The first TUniquePtr to compare.
+	 * @param Rhs The second TUniquePtr to compare.
+	 *
+	 * @return true if the inner pointer of the left TUniquePtr is greater than or equal to the inner pointer of the right, false otherwise.
+	 */
+	template <typename RhsT>
+	[[nodiscard]] FORCEINLINE bool operator>=(const TUniquePtr<RhsT>& Rhs) const
+	{
+		return Get() >= Rhs.Get();
+	}
+
 	/**
 	 * Equality comparison operator against nullptr.
 	 *
@@ -750,6 +867,11 @@ public:
 		return IsValid();
 	}
 #endif
+
+	[[nodiscard]] friend uint32 GetTypeHash(const TUniquePtr& UniquePtr)
+	{
+		return GetTypeHash(UniquePtr.Get());
+	}
 };
 
 #if !PLATFORM_COMPILER_HAS_GENERATED_COMPARISON_OPERATORS
