PR #13304: fix: trim StorageDevice's serial number for comparison.

#rnx

- The storage device serial numbers can be both prefixed and postfixed with whitespace. We now trim the whitespace to keep the serial numbers consistent.
- Cleaned up some coding standards (StorageDevice->FStorageDevice and removing use of auto)

diff --git a/Engine/Source/Runtime/Core/Private/Windows/WindowsPlatformMisc.cpp b/Engine/Source/Runtime/Core/Private/Windows/WindowsPlatformMisc.cpp
--- a/Engine/Source/Runtime/Core/Private/Windows/WindowsPlatformMisc.cpp
+++ b/Engine/Source/Runtime/Core/Private/Windows/WindowsPlatformMisc.cpp
@@ -420,20 +420,22 @@ namespace
 		return false;
 	}
 
-	struct StorageDevice
+	struct FStorageDevice
 	{
 		FString SerialNumber;
 		WIDECHAR Drive;
 		FPlatformDriveStats Stats;
 
-		StorageDevice(FString&& SerialNumber, WIDECHAR DriveLetter)
-			: SerialNumber(MoveTemp(SerialNumber))
+		FStorageDevice(FString&& InSerialNumber, WIDECHAR DriveLetter)
+			: SerialNumber(MoveTemp(InSerialNumber))
 			, Drive(DriveLetter)
 			, Stats{ DriveLetter, 0, 0, EStorageDeviceType ::Unknown}
-		{}
+		{
+			SerialNumber.TrimStartAndEndInline();
+		}
 	};
 
-	TArray<StorageDevice> StorageDevices;
+	TArray<FStorageDevice> StorageDevices;
 
 	static void LogStorageInformationWarning(HRESULT HRes,  const TCHAR* message)
 	{
@@ -657,8 +659,11 @@ namespace
 			StorageWbemObject->Get(L"BusType", 0, &BusType, NULL, NULL);
 			StorageWbemObject->Get(L"SpindleSpeed", 0, &SpindleSpeed, NULL, NULL);
 
-			FString Serial(SerialNumber.bstrVal);
-			for (auto& StorageDevice : StorageDevices)
+			FString OriginalSerial(SerialNumber.bstrVal);
+			// Use a stringview as ::TrimStartAndEndInline would need to move data around if there is whitespace at the start of the string
+			FStringView Serial = FStringView(OriginalSerial).TrimStartAndEnd();
+
+			for (FStorageDevice& StorageDevice : StorageDevices)
 			{
 				if (StorageDevice.SerialNumber == Serial)
 				{
