Crash Context: Change MIINDUMP_USER_STREAM BufferSize arguments from allocated size to length of the FString we are providing. BufferSize is a misnomer and is interpreted as data size rather than as the maximum writable room in the buffer. This change removes unnecessary garbage data from the end of the string in the minidump.
#rnx

diff --git a/Engine/Source/Runtime/Core/Private/Windows/WindowsPlatformCrashContext.cpp b/Engine/Source/Runtime/Core/Private/Windows/WindowsPlatformCrashContext.cpp
--- a/Engine/Source/Runtime/Core/Private/Windows/WindowsPlatformCrashContext.cpp
+++ b/Engine/Source/Runtime/Core/Private/Windows/WindowsPlatformCrashContext.cpp
@@ -147,7 +147,7 @@ namespace {
 		// CrashContext.runtime-xml is now a part of the minidump file.
 		MINIDUMP_USER_STREAM CrashContextStream = { 0 };
 		CrashContextStream.Type = UE4_MINIDUMP_CRASHCONTEXT;
-		CrashContextStream.BufferSize = (ULONG)InContext.GetBuffer().GetAllocatedSize();
+		CrashContextStream.BufferSize = (ULONG)(InContext.GetBuffer().Len() * sizeof(InContext.GetBuffer()[0]));
 		CrashContextStream.Buffer = (void*)*InContext.GetBuffer();
 
 		MINIDUMP_USER_STREAM_INFORMATION CrashContextStreamInformation = { 0 };
