Proper key/value .natvis visualization of TMap, TMultiMap and TSortedMap, so that the keys are in the Index column and the values are in the Value column.
Added 'unit tests' for natvis.

diff --git a/Engine/Source/Runtime/Core/Private/Containers/ContainersTest.cpp b/Engine/Source/Runtime/Core/Private/Containers/ContainersTest.cpp
--- a/Engine/Source/Runtime/Core/Private/Containers/ContainersTest.cpp
+++ b/Engine/Source/Runtime/Core/Private/Containers/ContainersTest.cpp
@@ -1,16 +1,17 @@
 // Copyright Epic Games, Inc. All Rights Reserved.
 
 #include "CoreTypes.h"
-#include "Misc/AssertionMacros.h"
 #include "Concepts/EqualityComparable.h"
+#include "Containers/ArrayView.h"
 #include "Containers/Array.h"
 #include "Containers/Map.h"
 #include "Containers/Set.h"
 #include "Containers/SortedMap.h"
-#include "Containers/ArrayView.h"
+#include "Math/RandomStream.h"
+#include "Misc/AssertionMacros.h"
 #include "Misc/AutomationTest.h"
 #include "Stats/StatsMisc.h"
-#include "Math/RandomStream.h"
+#include "Templates/UniquePtr.h"
 
 #if WITH_DEV_AUTOMATION_TESTS
 
@@ -1209,4 +1210,328 @@ namespace ArrayViewTests
 	}
 };
 
+UE_DISABLE_OPTIMIZATION_SHIP
+
+FORCENOINLINE void PlaceBreakpointHere()
+{
+	// This function exists to support the debugger-inspecting workflow in the tests below.
+
+	// Do some memory allocation here to ensure the compiler doesn't optimize away the function.
+	void* Temp = FMemory::Malloc(16);
+	FMemory::Free(Temp);
+}
+
+IMPLEMENT_SIMPLE_AUTOMATION_TEST(FContainersNatvisTest, "System.Core.Containers.Natvis", EAutomationTestFlags_ApplicationContextMask | EAutomationTestFlags::SmokeFilter)
+bool FContainersNatvisTest::RunTest(const FString& Parameters)
+{
+	// This results of this test are intended to be inspected visually in a platform's debugger, to ensure that the .natvis is correct.
+	// A breakpoint should be placed on the last line of each block, the block variables observed in the debugger and compared with the
+	// textual description.
+
+	// TArray
+	{
+		// Name       | Value | Type
+		// -----------+-------+---------------------------------------
+		// > EmptyArr | Empty | TArray<int,TSizedDefaultAllocator<32>>
+		TArray<int32> EmptyArr;
+
+		// Name       | Value         | Type
+		// -----------+---------------+---------------------------------------
+		// > SlackArr | Empty, Max=16 | TArray<int,TSizedDefaultAllocator<32>>
+		TArray<int32> SlackArr;
+		SlackArr.Reserve(16);
+
+		// Name      | Value        | Type
+		// ----------+--------------+---------------------------------------
+		// > ArrVals | Num=5, Max=5 | TArray<int,TSizedDefaultAllocator<32>>
+		//     [0]   | 4            | int
+		//     [1]   | 1            | int
+		//     [2]   | 2            | int
+		//     [3]   | 5            | int
+		//     [4]   | 3            | int
+		TArray<int32> ArrVals = { 4, 1, 2, 5, 3 };
+
+		// Name           | Value        | Type
+		// ---------------+--------------+---------------------------------------
+		// > SlackArrVals | Num=3, Max=5 | TArray<int,TSizedDefaultAllocator<32>>
+		//     [0]        | 4            | int
+		//     [1]        | 2            | int
+		//     [2]        | 3            | int
+		TArray<int32> SlackArrVals = { 4, 1, 2, 5, 3 };
+		SlackArrVals.RemoveAt(3);
+		SlackArrVals.RemoveAt(1);
+
+		// Name      | Value        | Type
+		// ----------+--------------+-------------------------------------------
+		// > ArrStrs | Num=5, Max=5 | TArray<FString,TSizedDefaultAllocator<32>>
+		//   > [0]   | L"d1"        | FString
+		//       [0] | 100 'd'      | wchar_t
+		//       [1] | 49 '1'       | wchar_t
+		//   > [1]   | L"a2"        | FString
+		//       [0] | 97 'a'       | wchar_t
+		//       [1] | 50 '2'       | wchar_t
+		//   > [2]   | L"b3"        | FString
+		//       [0] | 98 'b'       | wchar_t
+		//       [1] | 51 '3'       | wchar_t
+		//   > [3]   | L"e4"        | FString
+		//       [0] | 101 'e'      | wchar_t
+		//       [1] | 52 '4'       | wchar_t
+		//   > [4]   | L"c5"        | FString
+		//       [0] | 99 'c'       | wchar_t
+		//       [1] | 53 '5'       | wchar_t
+		TArray<FString> ArrStrs = { "d1", "a2", "b3", "e4", "c5" };
+
+		// Name            | Value                | Type
+		// ----------------+----------------------+-------------------------------------------------------------------------------
+		// > ArrUniqueStrs | Num=5, Max=5         | TArray<TUniquePtr<FString,TDefaultDelete<FString>>,TSizedDefaultAllocator<32>>
+		//   > [0]         | Ptr=0x01234567 L"d1" | TUniquePtr<FString,TDefaultDelete<FString>>
+		//       [0]       | 100 'd'              | wchar_t
+		//       [1]       | 49 '1'               | wchar_t
+		//   > [1]         | Ptr=0x01234567 L"a2" | TUniquePtr<FString,TDefaultDelete<FString>>
+		//       [0]       | 97 'a'               | wchar_t
+		//       [1]       | 50 '2'               | wchar_t
+		//   > [2]         | Ptr=0x01234567 L"b3" | TUniquePtr<FString,TDefaultDelete<FString>>
+		//       [0]       | 98 'b'               | wchar_t
+		//       [1]       | 51 '3'               | wchar_t
+		//   > [3]         | Ptr=0x01234567 L"e4" | TUniquePtr<FString,TDefaultDelete<FString>>
+		//       [0]       | 101 'e'              | wchar_t
+		//       [1]       | 52 '4'               | wchar_t
+		//   > [4]         | Ptr=0x01234567 L"c5" | TUniquePtr<FString,TDefaultDelete<FString>>
+		//       [0]       | 99 'c'               | wchar_t
+		//       [1]       | 53 '5'               | wchar_t
+		TArray<TUniquePtr<FString>> ArrUniqueStrs;
+		ArrUniqueStrs.Reserve(5);
+		ArrUniqueStrs.Add(MakeUnique<FString>("d1"));
+		ArrUniqueStrs.Add(MakeUnique<FString>("a2"));
+		ArrUniqueStrs.Add(MakeUnique<FString>("b3"));
+		ArrUniqueStrs.Add(MakeUnique<FString>("e4"));
+		ArrUniqueStrs.Add(MakeUnique<FString>("c5"));
+
+		PlaceBreakpointHere();
+	}
+
+	// TSet
+	{
+		// Name       | Value | Type
+		// -----------+-------+------------------------------------------------------
+		// > EmptySet | Empty | TSet<int,DefaultKeyFuncs<int,0>,FDefaultSetAllocator>
+		TSet<int32> EmptySet;
+
+		// Name       | Value         | Type
+		// -----------+---------------+------------------------------------------------------
+		// > SlackSet | Empty, Max=16 | TSet<int,DefaultKeyFuncs<int,0>,FDefaultSetAllocator>
+		TSet<int32> SlackSet;
+		SlackSet.Reserve(16);
+
+		// Name      | Value        | Type
+		// ----------+--------------+------------------------------------------------------
+		// > SetVals | Num=5, Max=5 | TSet<int,DefaultKeyFuncs<int,0>,FDefaultSetAllocator>
+		//     [0]   | 4            | int
+		//     [1]   | 1            | int
+		//     [2]   | 2            | int
+		//     [3]   | 5            | int
+		//     [4]   | 3            | int
+		TSet<int32> SetVals = { 4, 1, 2, 5, 3 };
+
+		// Name           | Value         | Type
+		// ---------------+---------------+------------------------------------------------------
+		// > SlackSetVals | Num=5, Max=16 | TSet<int,DefaultKeyFuncs<int,0>,FDefaultSetAllocator>
+		//     [0]        | 4             | int
+		//     [1]        | 1             | int
+		//     [2]        | 2             | int
+		//     [3]        | 5             | int
+		//     [4]        | 3             | int
+		TSet<int32> SlackSetVals = { 4, 1, 2, 5, 3 };
+		SlackSetVals.Reserve(16);
+
+		// Name            | Value        | Type
+		// ----------------+--------------+------------------------------------------------------
+		// > SparseSetVals | Num=3, Max=5 | TSet<int,DefaultKeyFuncs<int,0>,FDefaultSetAllocator>
+		//     [0]         | 4            | int
+		//     [1]         | 2            | int
+		//     [2]         | 3            | int
+		TSet<int32> SparseSetVals = { 4, 1, 2, 5, 3 };
+		SparseSetVals.Remove(1);
+		SparseSetVals.Remove(5);
+
+		// Name      | Value        | Type
+		// ----------+--------------+--------------------------------------------------------------
+		// > SetStrs | Num=5, Max=5 | TSet<FString,DefaultKeyFuncs<FString,0>,FDefaultSetAllocator>
+		//   > [0]   | L"d1"        | FString
+		//       [0] | 100 'd'      | wchar_t
+		//       [1] | 49 '1'       | wchar_t
+		//   > [1]   | L"a2"        | FString
+		//       [0] | 97 'a'       | wchar_t
+		//       [1] | 50 '2'       | wchar_t
+		//   > [2]   | L"b3"        | FString
+		//       [0] | 98 'b'       | wchar_t
+		//       [1] | 51 '3'       | wchar_t
+		//   > [3]   | L"e4"        | FString
+		//       [0] | 101 'e'      | wchar_t
+		//       [1] | 52 '4'       | wchar_t
+		//   > [4]   | L"c5"        | FString
+		//       [0] | 99 'c'       | wchar_t
+		//       [1] | 53 '5'       | wchar_t
+		TSet<FString> SetStrs = { "d1", "a2", "b3", "e4", "c5" };
+
+		// Name            | Value                | Type
+		// ----------------+----------------------+--------------------------------------------------------------------------------------------------------------------------------------
+		// > SetUniqueStrs | Num=5, Max=5         | TSet<TUniquePtr<FString,TDefaultDelete<FString>>,DefaultKeyFuncs<TUniquePtr<FString,TDefaultDelete<FString>>,0>,FDefaultSetAllocator>
+		//   > [0]         | Ptr=0x01234567 L"d1" | TUniquePtr<FString,TDefaultDelete<FString>>
+		//       [0]       | 100 'd'              | wchar_t
+		//       [1]       | 49 '1'               | wchar_t
+		//   > [1]         | Ptr=0x01234567 L"a2" | TUniquePtr<FString,TDefaultDelete<FString>>
+		//       [0]       | 97 'a'               | wchar_t
+		//       [1]       | 50 '2'               | wchar_t
+		//   > [2]         | Ptr=0x01234567 L"b3" | TUniquePtr<FString,TDefaultDelete<FString>>
+		//       [0]       | 98 'b'               | wchar_t
+		//       [1]       | 51 '3'               | wchar_t
+		//   > [3]         | Ptr=0x01234567 L"e4" | TUniquePtr<FString,TDefaultDelete<FString>>
+		//       [0]       | 101 'e'              | wchar_t
+		//       [1]       | 52 '4'               | wchar_t
+		//   > [4]         | Ptr=0x01234567 L"c5" | TUniquePtr<FString,TDefaultDelete<FString>>
+		//       [0]       | 99 'c'               | wchar_t
+		//       [1]       | 53 '5'               | wchar_t
+		TSet<TUniquePtr<FString>> SetUniqueStrs;
+		SetUniqueStrs.Reserve(5);
+		SetUniqueStrs.Add(MakeUnique<FString>("d1"));
+		SetUniqueStrs.Add(MakeUnique<FString>("a2"));
+		SetUniqueStrs.Add(MakeUnique<FString>("b3"));
+		SetUniqueStrs.Add(MakeUnique<FString>("e4"));
+		SetUniqueStrs.Add(MakeUnique<FString>("c5"));
+
+		PlaceBreakpointHere();
+	}
+
+	// TMap
+	{
+		// Name       | Value | Type
+		// -----------+-------+----------------------------------------------------------------------------------
+		// > EmptyMap | Empty | TMap<FString,int,FDefaultSetAllocator,TDefaultMapHashableKeyFuncs<FString,int,0>>
+		TMap<FString, int32> EmptyMap;
+
+		// Name       | Value         | Type
+		// -----------+---------------+------------------------------------------------------
+		// > SlackMap | Empty, Max=16 | TMap<FString,int,FDefaultSetAllocator,TDefaultMapHashableKeyFuncs<FString,int,0>>
+		TMap<FString, int32> SlackMap;
+		SlackMap.Reserve(16);
+
+		// Name       | Value        | Type
+		// -----------+--------------+------------------------------------------------------
+		// > MapVals  | Num=5, Max=5 | TMap<FString,int,FDefaultSetAllocator,TDefaultMapHashableKeyFuncs<FString,int,0>>
+		//     [L"d"] | 1            | int
+		//     [L"a"] | 2            | int
+		//     [L"b"] | 3            | int
+		//     [L"e"] | 4            | int
+		//     [L"c"] | 5            | int
+		TMap<FString, int32> MapVals = { { "d", 1 }, { "a", 2 }, { "b", 3 }, { "e", 4 }, { "c", 5 } };
+
+		// Name           | Value         | Type
+		// ---------------+---------------+------------------------------------------------------
+		// > SlackMapVals | Num=5, Max=16 | TMap<FString,int,FDefaultSetAllocator,TDefaultMapHashableKeyFuncs<FString,int,0>>
+		//     [L"d"]     | 1             | int
+		//     [L"a"]     | 2             | int
+		//     [L"b"]     | 3             | int
+		//     [L"e"]     | 4             | int
+		//     [L"c"]     | 5             | int
+		TMap<FString, int32> SlackMapVals = { { "d", 1 }, { "a", 2 }, { "b", 3 }, { "e", 4 }, { "c", 5 } };
+		SlackMapVals.Reserve(16);
+
+		// Name            | Value        | Type
+		// ----------------+--------------+------------------------------------------------------
+		// > SparseMapVals | Num=3, Max=5 | TMap<FString,int,FDefaultSetAllocator,TDefaultMapHashableKeyFuncs<FString,int,0>>
+		//     [L"d"]      | 1            | int
+		//     [L"b"]      | 3            | int
+		//     [L"c"]      | 5            | int
+		TMap<FString, int32> SparseMapVals = { { "d", 1 }, { "a", 2 }, { "b", 3 }, { "e", 4 }, { "c", 5 } };
+		SparseMapVals.Remove("a");
+		SparseMapVals.Remove("e");
+
+		// Name                      | Value        | Type
+		// --------------------------+--------------+----------------------------------------------------------------------------------------------------------------------------------------------------------
+		// > MapUniqueVals           | Num=5, Max=5 | TMap<TUniquePtr<FString,TDefaultDelete<FString>>,int,FDefaultSetAllocator,TDefaultMapHashableKeyFuncs<TUniquePtr<FString,TDefaultDelete<FString>>,int,0>>
+		//     [Ptr=0x12345678 L"d"] | 1            | int
+		//     [Ptr=0x12345678 L"a"] | 2            | int
+		//     [Ptr=0x12345678 L"b"] | 3            | int
+		//     [Ptr=0x12345678 L"e"] | 4            | int
+		//     [Ptr=0x12345678 L"c"] | 5            | int
+		TMap<TUniquePtr<FString>, int32> MapUniqueVals;
+		MapUniqueVals.Reserve(5);
+		MapUniqueVals.Emplace(MakeUnique<FString>("d"), 1);
+		MapUniqueVals.Emplace(MakeUnique<FString>("a"), 2);
+		MapUniqueVals.Emplace(MakeUnique<FString>("b"), 3);
+		MapUniqueVals.Emplace(MakeUnique<FString>("e"), 4);
+		MapUniqueVals.Emplace(MakeUnique<FString>("c"), 5);
+
+		PlaceBreakpointHere();
+	}
+
+	// TSortedMap
+	{
+		// Name        | Value | Type
+		// ------------+-------+--------------------------------------------------------------------------
+		// > EmptySMap | Empty | TSortedMap<FString,int,TSizedDefaultAllocator<32>,TLess<const FString &>>
+		TSortedMap<FString, int32> EmptySMap;
+
+		// Name        | Value         | Type
+		// ------------+---------------+------------------------------------------------------
+		// > SlackSMap | Empty, Max=16 | TSortedMap<FString,int,TSizedDefaultAllocator<32>,TLess<const FString &>>
+		TSortedMap<FString, int32> SlackSMap;
+		SlackSMap.Reserve(16);
+
+		// Name       | Value        | Type
+		// -----------+--------------+------------------------------------------------------
+		// > SMapVals | Num=5, Max=5 | TSortedMap<FString,int,TSizedDefaultAllocator<32>,TLess<const FString &>>
+		//     [L"a"] | 2            | int
+		//     [L"b"] | 3            | int
+		//     [L"c"] | 5            | int
+		//     [L"d"] | 1            | int
+		//     [L"e"] | 4            | int
+		TSortedMap<FString, int32> SMapVals = { { "d", 1 }, { "a", 2 }, { "b", 3 }, { "e", 4 }, { "c", 5 } };
+
+		// Name            | Value         | Type
+		// ----------------+---------------+------------------------------------------------------
+		// > SlackSMapVals | Num=5, Max=16 | TSortedMap<FString,int,TSizedDefaultAllocator<32>,TLess<const FString &>>
+		//     [L"a"]      | 2             | int
+		//     [L"b"]      | 3             | int
+		//     [L"c"]      | 5             | int
+		//     [L"d"]      | 1             | int
+		//     [L"e"]      | 4             | int
+		TSortedMap<FString, int32> SlackSMapVals = { { "d", 1 }, { "a", 2 }, { "b", 3 }, { "e", 4 }, { "c", 5 } };
+		SlackSMapVals.Reserve(16);
+
+		// Name             | Value        | Type
+		// -----------------+--------------+------------------------------------------------------
+		// > SparseSMapVals | Num=3, Max=5 | TSortedMap<FString,int,TSizedDefaultAllocator<32>,TLess<const FString &>>
+		//     [L"b"]       | 3            | int
+		//     [L"c"]       | 5            | int
+		//     [L"d"]       | 1            | int
+		TSortedMap<FString, int32> SparseSMapVals = { { "d", 1 }, { "a", 2 }, { "b", 3 }, { "e", 4 }, { "c", 5 } };
+		SparseSMapVals.Remove("a");
+		SparseSMapVals.Remove("e");
+
+		// Name                      | Value        | Type
+		// --------------------------+--------------+--------------------------------------------------------------------------------------------------------------------------------------------------
+		// > SMapUniqueVals          | Num=5, Max=5 | TSortedMap<TUniquePtr<FString,TDefaultDelete<FString>>,int,TSizedDefaultAllocator<32>,TLess<const TUniquePtr<FString,TDefaultDelete<FString>> &>>
+		//     [Ptr=0x12345678 L"a"] | 2            | int
+		//     [Ptr=0x12345678 L"b"] | 3            | int
+		//     [Ptr=0x12345678 L"c"] | 5            | int
+		//     [Ptr=0x12345678 L"d"] | 1            | int
+		//     [Ptr=0x12345678 L"e"] | 4            | int
+		TSortedMap<TUniquePtr<FString>, int32> SMapUniqueVals;
+		SMapUniqueVals.Reserve(5);
+		SMapUniqueVals.Emplace(MakeUnique<FString>("d"), 1);
+		SMapUniqueVals.Emplace(MakeUnique<FString>("a"), 2);
+		SMapUniqueVals.Emplace(MakeUnique<FString>("b"), 3);
+		SMapUniqueVals.Emplace(MakeUnique<FString>("e"), 4);
+		SMapUniqueVals.Emplace(MakeUnique<FString>("c"), 5);
+
+		PlaceBreakpointHere();
+	}
+
+	return true;
+}
+
+UE_ENABLE_OPTIMIZATION_SHIP
+
 #endif // WITH_DEV_AUTOMATION_TESTS
