Added std::initializer_list support to TSparseArray.

diff --git a/Engine/Source/Runtime/Core/Public/Containers/SparseArray.h b/Engine/Source/Runtime/Core/Public/Containers/SparseArray.h
--- a/Engine/Source/Runtime/Core/Public/Containers/SparseArray.h
+++ b/Engine/Source/Runtime/Core/Public/Containers/SparseArray.h
@@ -17,6 +17,7 @@
 #include "Serialization/MemoryImageWriter.h"
 #include "Containers/UnrealString.h"
 #include "Misc/IntrusiveUnsetOptionalState.h"
+#include <initializer_list>
 
 
 #if UE_BUILD_SHIPPING || UE_BUILD_TEST
@@ -214,18 +215,6 @@ public:
 		return Result;
 	}
 
-	/** 
-	 * Add an element at the lowest free index, instead of the last freed index. 
-	 * This requires a search which can be accelerated with LowestFreeIndexSearchStart.
-	 */
-	UE_DEPRECATED(4.26, "AddAtLowestFreeIndex API is deprecated; please use EmplaceAtLowestFreeIndex instead.")
-	int32 AddAtLowestFreeIndex(const ElementType& Element, int32& LowestFreeIndexSearchStart)
-	{
-		FSparseArrayAllocationInfo Allocation = AddUninitializedAtLowestFreeIndex(LowestFreeIndexSearchStart);
-		new(Allocation) ElementType(Element);
-		return Allocation.Index;
-	}
-
 	/**
 	 * Constructs a new item at the last freed index of the array.
 	 *
@@ -800,6 +789,13 @@ public:
 		*this = InCopy;
 	}
 
+	/** Initializer list constructor. */
+	[[nodiscard]] TSparseArray(std::initializer_list<ElementType> InList)
+		: TSparseArray()
+	{
+		*this = InList;
+	}
+
 	/** Move assignment operator. */
 	TSparseArray& operator=(TSparseArray&& InCopy)
 	{
@@ -860,6 +856,18 @@ public:
 		return *this;
 	}
 
+	/** Initializer list assignment operator. */
+	TSparseArray& operator=(std::initializer_list<ElementType> InList)
+	{
+		this->Reset();
+		this->Reserve((uint32)InList.size());
+		for (const ElementType& Element : InList)
+		{
+			Add(Element);
+		}
+		return *this;
+	}
+
 private:
 	template <typename SparseArrayType>
 	FORCEINLINE static void Move(SparseArrayType& ToArray, SparseArrayType& FromArray)
