Substrate - fixed a case where blend mode was not set to TranslucentColoredTransmittance required for thin translucent when  hidden conversion is enabled.

#tests PC editor
#rnx

diff --git a/Engine/Source/Runtime/Core/Private/UObject/FortniteMainBranchObjectVersion.cpp b/Engine/Source/Runtime/Core/Private/UObject/FortniteMainBranchObjectVersion.cpp
--- a/Engine/Source/Runtime/Core/Private/UObject/FortniteMainBranchObjectVersion.cpp
+++ b/Engine/Source/Runtime/Core/Private/UObject/FortniteMainBranchObjectVersion.cpp
@@ -14,6 +14,6 @@ TMap<FGuid, FGuid> FFortniteMainBranchObjectVersion::GetSystemGuids()
 	SystemGuids.Add(DevGuids.Niagara_LatestScriptCompileVersion, FGuid("97F0FD7149D90D42B65D6557AF477D1A"));
 	SystemGuids.Add(DevGuids.SkeletalMeshDerivedDataVersion, FGuid("716424BCEC314E969ACFE36D2C2ECE49"));
 	SystemGuids.Add(DevGuids.STATICMESH_DERIVEDDATA_VER, FGuid("3ABF32149F9D4D838750368AB95573BE"));
-	SystemGuids.Add(DevGuids.MaterialTranslationDDCVersion, FGuid("A3E8C83AF1484A74A986F0AA9181F798"));
+	SystemGuids.Add(DevGuids.MaterialTranslationDDCVersion, FGuid("E541A459F3374E628240BBB309910753"));
 	return SystemGuids;
 }
diff --git a/Engine/Source/Runtime/Engine/Private/Materials/Material.cpp b/Engine/Source/Runtime/Engine/Private/Materials/Material.cpp
--- a/Engine/Source/Runtime/Engine/Private/Materials/Material.cpp
+++ b/Engine/Source/Runtime/Engine/Private/Materials/Material.cpp
@@ -7289,6 +7289,11 @@ EBlendMode UMaterial::GetBlendMode() const
 	}
 	else
 	{
+		if (GetShadingModels() == MSM_ThinTranslucent && !HasSubstrateFrontMaterialConnected() && Substrate::IsHiddenMaterialAssetConversionEnabled())
+		{
+			// Set the correct blend mode when using Substrate hidden conversion for thin translucent materials. Only when the FrontMaterial is not pluggedin.
+			return BLEND_TranslucentColoredTransmittance;
+		}
 		return BlendMode;
 	}
 }
