Printing the name of a checkpoint when it hangs.

#rnx

diff --git a/Engine/Source/Runtime/Core/Private/HAL/ThreadHeartBeat.cpp b/Engine/Source/Runtime/Core/Private/HAL/ThreadHeartBeat.cpp
--- a/Engine/Source/Runtime/Core/Private/HAL/ThreadHeartBeat.cpp
+++ b/Engine/Source/Runtime/Core/Private/HAL/ThreadHeartBeat.cpp
@@ -698,7 +698,7 @@ uint32 FThreadHeartBeat::CheckCheckpointHeartBeat(double& OutHangDuration)
 				{
 					if (HeartBeatInfo.HangDuration > 0.0)
 					{
-						UE_LOG(LogCore, Warning, TEXT("Failed to reach checkpoint within alotted time of %.2f. Triggering hang detector."), HeartBeatInfo.HangDuration);
+						UE_LOG(LogCore, Warning, TEXT("Failed to reach checkpoint: '%s' within alotted time of %.2f. Triggering hang detector."), *HeartBeatInfo.HeartBeatName.ToString(), HeartBeatInfo.HangDuration);
 
 						HeartBeatInfo.LastHangTime = CurrentTime;
 						OutHangDuration = HeartBeatInfo.HangDuration;
