Fix UPS cooking indeterminism for TBox[2] structs on stack

Make ForceInit constructor zero out padding

#udn 17968507
[RN] Fix UPS cooking indeterminism for TBox[2] structs on stack

diff --git a/Engine/Source/Runtime/Core/Public/Math/Box.h b/Engine/Source/Runtime/Core/Public/Math/Box.h
--- a/Engine/Source/Runtime/Core/Public/Math/Box.h
+++ b/Engine/Source/Runtime/Core/Public/Math/Box.h
@@ -33,7 +33,6 @@ struct TBoxConstInit {};
 template<typename T>
 struct TBox
 {
-public:
 	using FReal = T;
 
 	/** Holds the box's minimum point. */
@@ -45,6 +44,9 @@ public:
 	/** Holds a flag indicating whether this box is valid. */
 	uint8 IsValid;
 
+private:
+	uint8 Pad[sizeof(T) - sizeof(IsValid)];
+
 public:
 
 	/**
@@ -79,11 +81,12 @@ public:
 		: Min(0, TVectorConstInit{})
 		, Max(0, TVectorConstInit{})
 		, IsValid(0)
+		, Pad() // For deterministic cooking with UPS serialize-as-zero for immutable/atomic noexport types
 	{
 	}
 
 	[[nodiscard]] constexpr TBox(EForceInit, TBoxConstInit)
-		: Min(0, TVectorConstInit{}), Max(0, TVectorConstInit{}), IsValid(0)
+		: TBox(ForceInit)
 	{
 	}
 
diff --git a/Engine/Source/Runtime/Core/Public/Math/Box2D.h b/Engine/Source/Runtime/Core/Public/Math/Box2D.h
--- a/Engine/Source/Runtime/Core/Public/Math/Box2D.h
+++ b/Engine/Source/Runtime/Core/Public/Math/Box2D.h
@@ -323,8 +323,8 @@ public:
 	 */
 	void Init()
 	{
-		Min = Max = TVector2<T>::ZeroVector;
-		bIsValid = false;
+		// Zero padding for deterministic cooking with UPS serialize-as-zero for immutable/atomic noexport types
+		FMemory::Memset(this, 0, sizeof(TBox2));
 	}
 
 	/**
