Fixed managed storage handle attempting to resolve an unmanaged file more than once after the manager becomes active.

diff --git a/Engine/Source/Runtime/Core/Public/HAL/IPlatformFileManagedStorageWrapper.h b/Engine/Source/Runtime/Core/Public/HAL/IPlatformFileManagedStorageWrapper.h
--- a/Engine/Source/Runtime/Core/Public/HAL/IPlatformFileManagedStorageWrapper.h
+++ b/Engine/Source/Runtime/Core/Public/HAL/IPlatformFileManagedStorageWrapper.h
@@ -54,6 +54,10 @@ struct FPersistentManagedFile
 	bool IsValid() const { return Category >= 0; }
 	explicit operator bool() const { return Category >= 0; }
 
+	// Helpers to avoid unnecessarily attempting to manage multiple times, use empty string as sentinel
+	bool HasFailedTryToManage() const { return FullFilename.IsEmpty(); }
+	void SetHasFailedTryToManage() { FullFilename.Empty(); }
+
 	void Clear()
 	{
 		Category = INDEX_NONE;
@@ -441,6 +445,18 @@ public:
 		return OutFile;
 	}
 
+	bool TryManageFile(FPersistentManagedFile& InOutFile)
+	{
+		check(!InOutFile.IsValid());
+		TryManageFileInternal(InOutFile);
+		if (InOutFile.IsValid())
+		{
+			return true;
+		}
+		InOutFile.SetHasFailedTryToManage();
+		return false;
+	}
+
 private:
 	void TryManageFileInternal(FPersistentManagedFile& OutFile)
 	{
@@ -650,20 +666,19 @@ private:
 class FManagedStorageFileWriteHandle : public IFileHandle
 {
 private:
-	static bool IsReady()
-	{
-		// FManagedStoragePlatformFile will just pass through to LowerLevel until FPersistentStorageManager is ready
-		return FPersistentStorageManager::IsReady();
-	}
-
 	bool TryManageFile()
 	{
-		if (!File && IsReady())
+		if (File.IsValid())
 		{
-			File = FPersistentStorageManager::Get().TryManageFile(File.FullFilename);
+			return true;
 		}
 
-		return File.IsValid();
+		if (!File.HasFailedTryToManage() && FPersistentStorageManager::IsReady())
+		{
+			return FPersistentStorageManager::Get().TryManageFile(File);
+		}
+
+		return false;
 	}
 
 public:
