CompactBinary: Fix memory overwrite in AppendBase64String; 0 was being written into the end of stringbuilder without first allocating space for the null terminator.
#rnx

diff --git a/Engine/Source/Runtime/Core/Private/Serialization/CompactBinarySerialization.cpp b/Engine/Source/Runtime/Core/Private/Serialization/CompactBinarySerialization.cpp
--- a/Engine/Source/Runtime/Core/Private/Serialization/CompactBinarySerialization.cpp
+++ b/Engine/Source/Runtime/Core/Private/Serialization/CompactBinarySerialization.cpp
@@ -648,10 +648,14 @@ private:
 
 	void AppendBase64String(FMemoryView Value)
 	{
-		Builder << '"';
 		checkf(Value.GetSize() <= 512 * 1024 * 1024,
 			TEXT("Encoding 512 MiB or larger is not supported. Size: " UINT64_FMT), Value.GetSize());
 		const uint32 EncodedSize = FBase64::GetEncodedDataSize(uint32(Value.GetSize()));
+		// Encode writes a zero-terminator after its text, so we need to have allocated space for it to write to by
+		// reserving at least one more than the EncodedSize. We handle that by reserving space for the trailing quote
+		// symbol we write after it; Encode will write 0 to it and we then overwrite it with a quote.
+		Builder.Reserve(Builder.Len() + 2 + EncodedSize);
+		Builder << '"';
 		const int32 EncodedIndex = Builder.AddUninitialized(int32(EncodedSize));
 		FBase64::Encode(static_cast<const uint8*>(Value.GetData()), uint32(Value.GetSize()),
 			reinterpret_cast<CharType*>(Builder.GetData() + EncodedIndex));
