Renamed Container test functions to make it clear they're for maps only, and added some missing basic TSortedMap tests.

diff --git a/Engine/Source/Runtime/Core/Private/Containers/ContainersTest.cpp b/Engine/Source/Runtime/Core/Private/Containers/ContainersTest.cpp
--- a/Engine/Source/Runtime/Core/Private/Containers/ContainersTest.cpp
+++ b/Engine/Source/Runtime/Core/Private/Containers/ContainersTest.cpp
@@ -2,8 +2,8 @@
 
 #include "CoreTypes.h"
 #include "Concepts/EqualityComparable.h"
-#include "Containers/ArrayView.h"
 #include "Containers/Array.h"
+#include "Containers/ArrayView.h"
 #include "Containers/Map.h"
 #include "Containers/Set.h"
 #include "Containers/SortedMap.h"
@@ -315,23 +315,23 @@ namespace
 		return FString::Printf(TEXT("TestString%d"), Input);
 	}
 
-	template <typename ContainerType, typename KeyType>
-	void RunContainerTests()
+	template <typename MapType, typename KeyType>
+	void RunMapTests()
 	{
-		ContainerType Cont;
+		MapType Map;
 
 		int32 MaxNum = 0;
 		SIZE_T MaxAllocatedSize = 0;
-		const SIZE_T InitialAllocatedSize = Cont.GetAllocatedSize();
+		const SIZE_T InitialAllocatedSize = Map.GetAllocatedSize();
 
 		ContainerTestStats.Reset();
 
-		auto CheckContainer = [&Cont]()
+		auto CheckMap = [&Map]()
 		{
-			CheckContainerNum(Cont);
-			CheckContainerEnds(Cont);
-			CheckContainerElements(Cont);
-			CheckContainerSelfEquality(Cont);
+			CheckContainerNum(Map);
+			CheckContainerEnds(Map);
+			CheckContainerElements(Map);
+			CheckContainerSelfEquality(Map);
 		};
 
 		// Test Add and Remove
@@ -340,88 +340,88 @@ namespace
 		{
 			for (int32 N = 0; N != Count; ++N)
 			{
-				Cont.Add(GenerateTestKey<KeyType>(N), FContainerTestValueType(TEXT("New Value")));
-				CheckContainer();
+				Map.Add(GenerateTestKey<KeyType>(N), FContainerTestValueType(TEXT("New Value")));
+				CheckMap();
 			}
-			MaxNum = Cont.Num();
-			MaxAllocatedSize = Cont.GetAllocatedSize();
+			MaxNum = Map.Num();
+			MaxAllocatedSize = Map.GetAllocatedSize();
 
 			for (int32 N = 0; N != Count; ++N)
 			{
-				Cont.Remove(GenerateTestKey<KeyType>(N));
-				CheckContainer();
+				Map.Remove(GenerateTestKey<KeyType>(N));
+				CheckMap();
 			}
 
-			check(Cont.IsEmpty());
+			check(Map.IsEmpty());
 
 			for (int32 N = 0; N != Count; ++N)
 			{
-				Cont.Add(GenerateTestKey<KeyType>((Count - 1) - N), FContainerTestValueType(TEXT("New Value")));
-				CheckContainer();
+				Map.Add(GenerateTestKey<KeyType>((Count - 1) - N), FContainerTestValueType(TEXT("New Value")));
+				CheckMap();
 			}
 
 			for (int32 N = 0; N != Count; ++N)
 			{
-				Cont.Remove(GenerateTestKey<KeyType>(N));
-				CheckContainer();
+				Map.Remove(GenerateTestKey<KeyType>(N));
+				CheckMap();
 			}
 
-			check(Cont.IsEmpty());
+			check(Map.IsEmpty());
 		}
 
 		// Test Empty and Shrink 
 		{
 			// Test releasing memory allocations
-			Cont.Empty();
-			CheckContainer();
-			check(Cont.GetAllocatedSize() == InitialAllocatedSize);
+			Map.Empty();
+			CheckMap();
+			check(Map.GetAllocatedSize() == InitialAllocatedSize);
 
 			// Test integrity after re-growing container to MaxNum elements again
 			for (int32 N = 0; N < MaxNum; ++N)
 			{
-				Cont.Add(GenerateTestKey<KeyType>(N), FContainerTestValueType(TEXT("New Value")));
+				Map.Add(GenerateTestKey<KeyType>(N), FContainerTestValueType(TEXT("New Value")));
 			}
-			CheckContainer();
-			check(Cont.GetAllocatedSize() == MaxAllocatedSize);
+			CheckMap();
+			check(Map.GetAllocatedSize() == MaxAllocatedSize);
 
 			// Test data integrity while removing and shrinking continously
 			{
-				SIZE_T PrevAllocatedSize = Cont.GetAllocatedSize();
+				SIZE_T PrevAllocatedSize = Map.GetAllocatedSize();
 				for (int32 N = MaxNum - 1; N >= MaxNum / 4; --N)
 				{
-					Cont.Remove(GenerateTestKey<KeyType>(N));
-					Cont.Shrink();
-					CheckContainer();
-					check(Cont.GetAllocatedSize() <= PrevAllocatedSize);
-					PrevAllocatedSize = Cont.GetAllocatedSize();
+					Map.Remove(GenerateTestKey<KeyType>(N));
+					Map.Shrink();
+					CheckMap();
+					check(Map.GetAllocatedSize() <= PrevAllocatedSize);
+					PrevAllocatedSize = Map.GetAllocatedSize();
 				}
 			}
 
 			// Test removing and releasing remaining elements
-			Cont.Empty();
-			check(Cont.IsEmpty());
-			check(Cont.GetAllocatedSize() == InitialAllocatedSize);
+			Map.Empty();
+			check(Map.IsEmpty());
+			check(Map.GetAllocatedSize() == InitialAllocatedSize);
 		}
 
 		// Test key iterators
 		{
-			static_assert(std::is_same_v<decltype((Cont.CreateKeyIterator     (DeclVal<KeyType&>())->Key)),         KeyType&>);
-			static_assert(std::is_same_v<decltype((Cont.CreateKeyIterator     (DeclVal<KeyType&>())->Value)),       FContainerTestValueType&>);
-			static_assert(std::is_same_v<decltype((Cont.CreateConstKeyIterator(DeclVal<KeyType&>())->Key)),   const KeyType&>);
-			static_assert(std::is_same_v<decltype((Cont.CreateConstKeyIterator(DeclVal<KeyType&>())->Value)), const FContainerTestValueType&>);
+			static_assert(std::is_same_v<decltype((Map.CreateKeyIterator     (DeclVal<KeyType&>())->Key)),         KeyType&>);
+			static_assert(std::is_same_v<decltype((Map.CreateKeyIterator     (DeclVal<KeyType&>())->Value)),       FContainerTestValueType&>);
+			static_assert(std::is_same_v<decltype((Map.CreateConstKeyIterator(DeclVal<KeyType&>())->Key)),   const KeyType&>);
+			static_assert(std::is_same_v<decltype((Map.CreateConstKeyIterator(DeclVal<KeyType&>())->Value)), const FContainerTestValueType&>);
 
 			const TCHAR* RegularValue  = TEXT("Regular");
 			const TCHAR* ReplacedValue = TEXT("Replaced");
 
 			for (int32 Count = 0; Count < MAX_TEST_OBJECTS - 1; Count += MAX_TEST_OBJECTS_STEP)
 			{
-				Cont.Empty();
-				CheckContainer();
+				Map.Empty();
+				CheckMap();
 
 				for (int32 N = 0; N != Count; ++N)
 				{
-					Cont.Add(GenerateTestKey<KeyType>(N), FContainerTestValueType(RegularValue));
-					CheckContainer();
+					Map.Add(GenerateTestKey<KeyType>(N), FContainerTestValueType(RegularValue));
+					CheckMap();
 				}
 
 				// Iterate over all possible keys, and some before/after the range [0, Count) that won't exist
@@ -431,21 +431,21 @@ namespace
 
 					// Check that at most one key is found by the const key iterator
 					const KeyType* FoundConstKey = nullptr;
-					for (auto It = Cont.CreateConstKeyIterator(Key); It; ++It)
+					for (auto It = Map.CreateConstKeyIterator(Key); It; ++It)
 					{
 						check(!FoundConstKey);
 						check(It->Key == Key);
-						check(Cont.Get(It.GetId()) == *It);
+						check(Map.Get(It.GetId()) == *It);
 						FoundConstKey = &It->Key;
 					}
 
 					// Check that at most one key is found by the key iterator, and that we can mutate the value via one
 					KeyType* FoundKey = nullptr;
-					for (auto It = Cont.CreateKeyIterator(Key); It; ++It)
+					for (auto It = Map.CreateKeyIterator(Key); It; ++It)
 					{
 						check(!FoundKey);
 						check(It->Key == Key);
-						check(Cont.Get(It.GetId()) == *It);
+						check(Map.Get(It.GetId()) == *It);
 						FoundKey = &It->Key;
 						It->Value.Str = ReplacedValue;
 					}
@@ -455,7 +455,7 @@ namespace
 					if (FoundConstKey)
 					{
 						check(KeyValue >= 0 && KeyValue < Count);
-						check(Cont[Key].Str == ReplacedValue);
+						check(Map[Key].Str == ReplacedValue);
 					}
 					else
 					{
@@ -468,10 +468,10 @@ namespace
 
 	// Test container element address consistency when using SortFreeList 
 	// (see TSparseArray::SortFreeList for comments)
-	template <typename ContainerType, typename KeyType>
-	void RunContainerConsistencyTests()
+	template <typename MapType, typename KeyType>
+	void RunMapConsistencyTests()
 	{
-		ContainerType Cont;
+		MapType Map;
 
 		{
 			// Add 3 elements, then remove 2 in the same order they were added
@@ -479,29 +479,29 @@ namespace
 			const KeyType Key1 = GenerateTestKey<KeyType>(1);
 			const KeyType Key2 = GenerateTestKey<KeyType>(2);
 			const FContainerTestValueType Value = FContainerTestValueType(TEXT("New Value"));
-			Cont.Add(Key0, Value);
-			Cont.Add(Key1, Value);
-			Cont.Add(Key2, Value);
-			const FContainerTestValueType* ValuePtr0 = Cont.Find(Key0);
-			const FContainerTestValueType* ValuePtr1 = Cont.Find(Key1);
-			const FContainerTestValueType* ValuePtr2 = Cont.Find(Key2);
+			Map.Add(Key0, Value);
+			Map.Add(Key1, Value);
+			Map.Add(Key2, Value);
+			const FContainerTestValueType* ValuePtr0 = Map.Find(Key0);
+			const FContainerTestValueType* ValuePtr1 = Map.Find(Key1);
+			const FContainerTestValueType* ValuePtr2 = Map.Find(Key2);
 			check(ValuePtr0 != nullptr);
 			check(ValuePtr1 != nullptr);
 			check(ValuePtr2 != nullptr);
-			Cont.Remove(Key1);
-			Cont.Remove(Key2);
+			Map.Remove(Key1);
+			Map.Remove(Key2);
 
 			// Re-add the 2 elements in the same order. Without the call to SortFreeList() 
 			// the elements would end up in a different locations/order compared to the 
 			// original insertions. SortFreeList() should ensure that re-adding the elements 
 			// gives use the same container layout as long as we perform the same operations 
 			// in the same order as before.
-			Cont.SortFreeList();
-			Cont.Add(Key1, Value);
-			Cont.Add(Key2, Value);
-			const FContainerTestValueType* NewValuePtr0 = Cont.Find(Key0);
-			const FContainerTestValueType* NewValuePtr1 = Cont.Find(Key1);
-			const FContainerTestValueType* NewValuePtr2 = Cont.Find(Key2);
+			Map.SortFreeList();
+			Map.Add(Key1, Value);
+			Map.Add(Key2, Value);
+			const FContainerTestValueType* NewValuePtr0 = Map.Find(Key0);
+			const FContainerTestValueType* NewValuePtr1 = Map.Find(Key1);
+			const FContainerTestValueType* NewValuePtr2 = Map.Find(Key2);
 
 			check(ValuePtr0 == NewValuePtr0);
 			check(ValuePtr1 == NewValuePtr1);
@@ -516,12 +516,12 @@ namespace
 		CheckContainerSelfEquality(Cont);
 	}
 
-	template <typename ContainerType, typename KeyType>
-	void RunPerformanceTest(const FString& Description, int32 NumObjects, int32 NumOperations)
+	template <typename MapType, typename KeyType>
+	void RunMapPerformanceTest(const FString& Description, int32 NumObjects, int32 NumOperations)
 	{
 		ContainerTestStats.Reset();
 		
-		ContainerType Cont;
+		MapType Map;
 		FRandomStream RandomStream(RANDOM_SEED);
 
 		// Prep keys, not part of performance test
@@ -547,7 +547,7 @@ namespace
 		// Add elements in stably randomized order
 		for (int32 i = 0; i < NumObjects; i++)
 		{
-			Cont.Add(KeyArray[i], FString(TEXT("New Value")));
+			Map.Add(KeyArray[i], FString(TEXT("New Value")));
 		}
 		
 		// Now do searches
@@ -555,17 +555,17 @@ namespace
 		{
 			KeyType& Key = KeyArray[RandomStream.RandRange(0, NumObjects - 1)];
 
-			FString* FoundValue = Cont.Find(Key);
+			FString* FoundValue = Map.Find(Key);
 			check(FoundValue);
 		}
 	}
 
-	template <typename ContainerType, typename KeyType>
+	template <typename SetType, typename KeyType>
 	void RunSetPerformanceTest(const FString& Description, int32 NumObjects, int32 NumOperations)
 	{
 		ContainerTestStats.Reset();
 
-		ContainerType Cont;
+		SetType Set;
 		FRandomStream RandomStream(RANDOM_SEED);
 
 		// Prep keys, not part of performance test
@@ -591,7 +591,7 @@ namespace
 		// Add elements in stably randomized order
 		for (int32 i = 0; i < NumObjects; i++)
 		{
-			Cont.Add(KeyArray[i]);
+			Set.Add(KeyArray[i]);
 		}
 
 		// Now do searches
@@ -599,17 +599,11 @@ namespace
 		{
 			KeyType& Key = KeyArray[RandomStream.RandRange(0, NumObjects - 1)];
 
-			bool FoundValue = Cont.Contains(Key);
+			bool FoundValue = Set.Contains(Key);
 			check(FoundValue);
 		}
 	}
 
-	template <typename ContainerType, typename ElemType, typename GetRefFuncType>
-	void TestGetRef(ContainerType& Container, ElemType&& Elem, GetRefFuncType&& GetRefFunc)
-	{
-		auto& Ref = GetRefFunc(Container, Forward<ElemType>(Elem));
-	}
-
 	template<typename TValueType>
 	struct FCaseSensitiveLookupKeyFuncs : BaseKeyFuncs<TValueType, FString>
 	{
@@ -665,9 +659,13 @@ namespace
 IMPLEMENT_SIMPLE_AUTOMATION_TEST(FContainersSmokeTest, "System.Core.Containers.Smoke", EAutomationTestFlags_ApplicationContextMask | EAutomationTestFlags::SmokeFilter)
 bool FContainersSmokeTest::RunTest( const FString& Parameters )
 {
-	RunContainerTests<TMap<int32, FContainerTestValueType>, int32>();
-	RunContainerTests<TMap<int32, FContainerTestValueType, TInlineSetAllocator<32>>, int32>();
-	RunContainerTests<TMap<int32, FContainerTestValueType, TFixedSetAllocator<64>>, int32>();
+	RunMapTests<TMap<int32, FContainerTestValueType>, int32>();
+	RunMapTests<TMap<int32, FContainerTestValueType, TInlineSetAllocator<32>>, int32>();
+	RunMapTests<TMap<int32, FContainerTestValueType, TFixedSetAllocator<64>>, int32>();
+
+	RunMapTests<TSortedMap<int32, FContainerTestValueType>, int32>();
+	RunMapTests<TSortedMap<int32, FContainerTestValueType, TInlineAllocator<32>>, int32>();
+	RunMapTests<TSortedMap<int32, FContainerTestValueType, TFixedAllocator<64>>, int32>();
 
 	return true;
 }
@@ -675,20 +673,22 @@ bool FContainersSmokeTest::RunTest( const FString& Parameters )
 IMPLEMENT_SIMPLE_AUTOMATION_TEST(FContainersFullTest, "System.Core.Containers.Full", EAutomationTestFlags_ApplicationContextMask | EAutomationTestFlags::EngineFilter)
 bool FContainersFullTest::RunTest(const FString& Parameters)
 {
-	RunContainerTests<TMap<int32, FContainerTestValueType>, int32>();
-	RunContainerTests<TMap<FName, FContainerTestValueType>, FName>();
-	RunContainerTests<TMap<FString, FContainerTestValueType>, FString>();
-	RunContainerTests<TMap<int32, FContainerTestValueType, TInlineSetAllocator<32>>, int32>();
-	RunContainerTests<TMap<int32, FContainerTestValueType, TInlineSetAllocator<64>>, int32>();
-	RunContainerTests<TMap<int32, FContainerTestValueType, TFixedSetAllocator<64>>, int32>();
-	RunContainerTests<TMap<FString, FContainerTestValueType, FDefaultSetAllocator, FCaseSensitiveLookupKeyFuncs<FContainerTestValueType>>, FString>();
-
-	RunContainerConsistencyTests<TMap<int32, FContainerTestValueType>, int32>();
-
-	RunContainerTests<TSortedMap<int32, FContainerTestValueType>, int32>();
-	RunContainerTests<TSortedMap<FName, FContainerTestValueType, FDefaultAllocator, FNameLexicalLess>, FName>();
-	RunContainerTests<TSortedMap<FString, FContainerTestValueType>, FString>();
-	RunContainerTests<TSortedMap<FString, FContainerTestValueType, TInlineAllocator<64>>, FString>();
+	RunMapTests<TMap<int32, FContainerTestValueType>, int32>();
+	RunMapTests<TMap<FName, FContainerTestValueType>, FName>();
+	RunMapTests<TMap<FString, FContainerTestValueType>, FString>();
+	RunMapTests<TMap<int32, FContainerTestValueType, TInlineSetAllocator<32>>, int32>();
+	RunMapTests<TMap<int32, FContainerTestValueType, TInlineSetAllocator<64>>, int32>();
+	RunMapTests<TMap<int32, FContainerTestValueType, TFixedSetAllocator<64>>, int32>();
+	RunMapTests<TMap<FString, FContainerTestValueType, FDefaultSetAllocator, FCaseSensitiveLookupKeyFuncs<FContainerTestValueType>>, FString>();
+
+	RunMapConsistencyTests<TMap<int32, FContainerTestValueType>, int32>();
+	RunMapConsistencyTests<TMap<int32, FContainerTestValueType, TInlineSetAllocator<32>>, int32>();
+	RunMapConsistencyTests<TMap<int32, FContainerTestValueType, TFixedSetAllocator<64>>, int32>();
+
+	RunMapTests<TSortedMap<int32, FContainerTestValueType>, int32>();
+	RunMapTests<TSortedMap<FName, FContainerTestValueType, FDefaultAllocator, FNameLexicalLess>, FName>();
+	RunMapTests<TSortedMap<FString, FContainerTestValueType>, FString>();
+	RunMapTests<TSortedMap<FString, FContainerTestValueType, TInlineAllocator<64>>, FString>();
 
 	RunEmptyContainerSelfEqualityTest<TArray<int32>>();
 
@@ -716,41 +716,41 @@ bool FContainersFullTest::RunTest(const FString& Parameters)
 IMPLEMENT_SIMPLE_AUTOMATION_TEST(FContainerPerformanceTest, "System.Core.Containers.Performance", EAutomationTestFlags_ApplicationContextMask | EAutomationTestFlags::EngineFilter)
 bool FContainerPerformanceTest::RunTest(const FString& Parameters)
 {
-	RunPerformanceTest<TMap<int32, FString>, int32>(TEXT("TMap int32"), 1, 1000000);
-	RunPerformanceTest<TMap<int32, FString>, int32>(TEXT("TMap int32"), 10, 1000000);
-	RunPerformanceTest<TMap<int32, FString>, int32>(TEXT("TMap int32"), 100, 1000000);
-	RunPerformanceTest<TMap<int32, FString>, int32>(TEXT("TMap int32"), 1000, 1000000);
-	RunPerformanceTest<TMap<int32, FString>, int32>(TEXT("TMap int32"), 10000, 1000000);
-
-	RunPerformanceTest<TMap<FName, FString>, FName>(TEXT("TMap FName"), 1, 1000000);
-	RunPerformanceTest<TMap<FName, FString>, FName>(TEXT("TMap FName"), 10, 1000000);
-	RunPerformanceTest<TMap<FName, FString>, FName>(TEXT("TMap FName"), 100, 1000000);
-	RunPerformanceTest<TMap<FName, FString>, FName>(TEXT("TMap FName"), 1000, 1000000);
-	RunPerformanceTest<TMap<FName, FString>, FName>(TEXT("TMap FName"), 10000, 1000000);
-
-	RunPerformanceTest<TMap<FString, FString>, FString>(TEXT("TMap FString"), 1, 1000000);
-	RunPerformanceTest<TMap<FString, FString>, FString>(TEXT("TMap FString"), 10, 1000000);
-	RunPerformanceTest<TMap<FString, FString>, FString>(TEXT("TMap FString"), 100, 1000000);
-	RunPerformanceTest<TMap<FString, FString>, FString>(TEXT("TMap FString"), 1000, 1000000);
-	RunPerformanceTest<TMap<FString, FString>, FString>(TEXT("TMap FString"), 10000, 1000000);
-
-	RunPerformanceTest<TSortedMap<int32, FString>, int32>(TEXT("TSortedMap int32"), 1, 1000000);
-	RunPerformanceTest<TSortedMap<int32, FString>, int32>(TEXT("TSortedMap int32"), 10, 1000000);
-	RunPerformanceTest<TSortedMap<int32, FString>, int32>(TEXT("TSortedMap int32"), 100, 1000000);
-	RunPerformanceTest<TSortedMap<int32, FString>, int32>(TEXT("TSortedMap int32"), 1000, 1000000);
-	RunPerformanceTest<TSortedMap<int32, FString>, int32>(TEXT("TSortedMap int32"), 10000, 1000000);
-
-	RunPerformanceTest<TSortedMap<FName, FString, FDefaultAllocator, FNameLexicalLess>, FName>(TEXT("TSortedMap FName"), 1, 1000000);
-	RunPerformanceTest<TSortedMap<FName, FString, FDefaultAllocator, FNameLexicalLess>, FName>(TEXT("TSortedMap FName"), 10, 1000000);
-	RunPerformanceTest<TSortedMap<FName, FString, FDefaultAllocator, FNameLexicalLess>, FName>(TEXT("TSortedMap FName"), 100, 1000000);
-	RunPerformanceTest<TSortedMap<FName, FString, FDefaultAllocator, FNameLexicalLess>, FName>(TEXT("TSortedMap FName"), 1000, 1000000);
-	RunPerformanceTest<TSortedMap<FName, FString, FDefaultAllocator, FNameLexicalLess>, FName>(TEXT("TSortedMap FName"), 10000, 1000000);
-
-	RunPerformanceTest<TSortedMap<FString, FString>, FString>(TEXT("TSortedMap FString"), 1, 1000000);
-	RunPerformanceTest<TSortedMap<FString, FString>, FString>(TEXT("TSortedMap FString"), 10, 1000000);
-	RunPerformanceTest<TSortedMap<FString, FString>, FString>(TEXT("TSortedMap FString"), 100, 1000000);
-	RunPerformanceTest<TSortedMap<FString, FString>, FString>(TEXT("TSortedMap FString"), 1000, 1000000);
-	RunPerformanceTest<TSortedMap<FString, FString>, FString>(TEXT("TSortedMap FString"), 10000, 1000000);
+	RunMapPerformanceTest<TMap<int32, FString>, int32>(TEXT("TMap int32"), 1, 1000000);
+	RunMapPerformanceTest<TMap<int32, FString>, int32>(TEXT("TMap int32"), 10, 1000000);
+	RunMapPerformanceTest<TMap<int32, FString>, int32>(TEXT("TMap int32"), 100, 1000000);
+	RunMapPerformanceTest<TMap<int32, FString>, int32>(TEXT("TMap int32"), 1000, 1000000);
+	RunMapPerformanceTest<TMap<int32, FString>, int32>(TEXT("TMap int32"), 10000, 1000000);
+
+	RunMapPerformanceTest<TMap<FName, FString>, FName>(TEXT("TMap FName"), 1, 1000000);
+	RunMapPerformanceTest<TMap<FName, FString>, FName>(TEXT("TMap FName"), 10, 1000000);
+	RunMapPerformanceTest<TMap<FName, FString>, FName>(TEXT("TMap FName"), 100, 1000000);
+	RunMapPerformanceTest<TMap<FName, FString>, FName>(TEXT("TMap FName"), 1000, 1000000);
+	RunMapPerformanceTest<TMap<FName, FString>, FName>(TEXT("TMap FName"), 10000, 1000000);
+
+	RunMapPerformanceTest<TMap<FString, FString>, FString>(TEXT("TMap FString"), 1, 1000000);
+	RunMapPerformanceTest<TMap<FString, FString>, FString>(TEXT("TMap FString"), 10, 1000000);
+	RunMapPerformanceTest<TMap<FString, FString>, FString>(TEXT("TMap FString"), 100, 1000000);
+	RunMapPerformanceTest<TMap<FString, FString>, FString>(TEXT("TMap FString"), 1000, 1000000);
+	RunMapPerformanceTest<TMap<FString, FString>, FString>(TEXT("TMap FString"), 10000, 1000000);
+
+	RunMapPerformanceTest<TSortedMap<int32, FString>, int32>(TEXT("TSortedMap int32"), 1, 1000000);
+	RunMapPerformanceTest<TSortedMap<int32, FString>, int32>(TEXT("TSortedMap int32"), 10, 1000000);
+	RunMapPerformanceTest<TSortedMap<int32, FString>, int32>(TEXT("TSortedMap int32"), 100, 1000000);
+	RunMapPerformanceTest<TSortedMap<int32, FString>, int32>(TEXT("TSortedMap int32"), 1000, 1000000);
+	RunMapPerformanceTest<TSortedMap<int32, FString>, int32>(TEXT("TSortedMap int32"), 10000, 1000000);
+
+	RunMapPerformanceTest<TSortedMap<FName, FString, FDefaultAllocator, FNameLexicalLess>, FName>(TEXT("TSortedMap FName"), 1, 1000000);
+	RunMapPerformanceTest<TSortedMap<FName, FString, FDefaultAllocator, FNameLexicalLess>, FName>(TEXT("TSortedMap FName"), 10, 1000000);
+	RunMapPerformanceTest<TSortedMap<FName, FString, FDefaultAllocator, FNameLexicalLess>, FName>(TEXT("TSortedMap FName"), 100, 1000000);
+	RunMapPerformanceTest<TSortedMap<FName, FString, FDefaultAllocator, FNameLexicalLess>, FName>(TEXT("TSortedMap FName"), 1000, 1000000);
+	RunMapPerformanceTest<TSortedMap<FName, FString, FDefaultAllocator, FNameLexicalLess>, FName>(TEXT("TSortedMap FName"), 10000, 1000000);
+
+	RunMapPerformanceTest<TSortedMap<FString, FString>, FString>(TEXT("TSortedMap FString"), 1, 1000000);
+	RunMapPerformanceTest<TSortedMap<FString, FString>, FString>(TEXT("TSortedMap FString"), 10, 1000000);
+	RunMapPerformanceTest<TSortedMap<FString, FString>, FString>(TEXT("TSortedMap FString"), 100, 1000000);
+	RunMapPerformanceTest<TSortedMap<FString, FString>, FString>(TEXT("TSortedMap FString"), 1000, 1000000);
+	RunMapPerformanceTest<TSortedMap<FString, FString>, FString>(TEXT("TSortedMap FString"), 10000, 1000000);
 
 	RunSetPerformanceTest<TSet<FName>, FName>(TEXT("TSet FName"), 1, 1000000);
 	RunSetPerformanceTest<TSet<FName>, FName>(TEXT("TSet FName"), 10, 1000000);
