TChunkedArray:
- Added non-POD types support for CopyToLinearArray by copying elements into destination array instead of using memcpy.
- Added MoveToLinearArray to move elements into destination array.

diff --git a/Engine/Source/Runtime/Core/Public/Containers/ChunkedArray.h b/Engine/Source/Runtime/Core/Public/Containers/ChunkedArray.h
--- a/Engine/Source/Runtime/Core/Public/Containers/ChunkedArray.h
+++ b/Engine/Source/Runtime/Core/Public/Containers/ChunkedArray.h
@@ -254,7 +254,6 @@ public:
 	template<typename OtherAllocator>
 	void CopyToLinearArray(TArray<ElementType, OtherAllocator>& DestinationArray)
 	{
-		static_assert(TIsPODType<ElementType>::Value, "CopyToLinearArray does not support a constructor / destructor on the element class.");
 		if (NumElements > 0)
 		{
 			int32 OriginalNumElements = DestinationArray.Num();
@@ -265,12 +264,33 @@ public:
 			{
 				const int32 NumElementsInCurrentChunk = FMath::Min<int32>(NumElements - ChunkIndex * NumElementsPerChunk, NumElementsPerChunk);
 				check(NumElementsInCurrentChunk > 0);
-				FMemory::Memcpy(CopyDestPtr, &Chunks[ChunkIndex].Elements[0], NumElementsInCurrentChunk * sizeof(ElementType)); //-V598
+				ConstructItems<ElementType>(CopyDestPtr, &Chunks[ChunkIndex].Elements[0], NumElementsInCurrentChunk);
 				CopyDestPtr += NumElementsInCurrentChunk;
 			}
 		}
 	}
 
+	template<typename OtherAllocator>
+	void MoveToLinearArray(TArray<ElementType, OtherAllocator>& DestinationArray)
+	{
+		if (NumElements > 0)
+		{
+			int32 OriginalNumElements = DestinationArray.Num();
+			DestinationArray.AddUninitialized(NumElements);
+			InElementType* CopyDestPtr = &DestinationArray[OriginalNumElements];
+
+			for (int32 ChunkIndex = 0; ChunkIndex < Chunks.Num(); ChunkIndex++)
+			{
+				const int32 NumElementsInCurrentChunk = FMath::Min<int32>(NumElements - ChunkIndex * NumElementsPerChunk, NumElementsPerChunk);
+				check(NumElementsInCurrentChunk > 0);
+				MoveConstructItems<ElementType>(CopyDestPtr, &Chunks[ChunkIndex].Elements[0], NumElementsInCurrentChunk);
+				CopyDestPtr += NumElementsInCurrentChunk;
+			}
+
+			Empty();
+		}
+	}
+
 	void Empty( int32 Slack=0 ) 
 	{
 		// Compute the number of chunks needed.
