Added FPlatformManualResetEvent for Linux

diff --git a/Engine/Source/Runtime/Core/Private/Linux/LinuxPlatformManualResetEvent.h b/Engine/Source/Runtime/Core/Private/Linux/LinuxPlatformManualResetEvent.h
--- a/Engine/Source/Runtime/Core/Private/Linux/LinuxPlatformManualResetEvent.h
+++ b/Engine/Source/Runtime/Core/Private/Linux/LinuxPlatformManualResetEvent.h
@@ -4,4 +4,4 @@
 
 #include "Unix/UnixPlatformManualResetEvent.h"
 
-namespace UE { using FPlatformManualResetEvent = UE::HAL::Private::FUnixPlatformManualResetEvent; }
+namespace UE { using FPlatformManualResetEvent = UE::HAL::Private::FUnixManualResetEvent; }
diff --git a/Engine/Source/Runtime/Core/Private/Unix/UnixPlatformManualResetEvent.cpp b/Engine/Source/Runtime/Core/Private/Unix/UnixPlatformManualResetEvent.cpp
new file mode 100644
--- /dev/null
+++ b/Engine/Source/Runtime/Core/Private/Unix/UnixPlatformManualResetEvent.cpp
@@ -0,0 +1,104 @@
+// Copyright Epic Games, Inc. All Rights Reserved.
+
+#include "Unix/UnixPlatformManualResetEvent.h"
+
+#include "Async/Fundamental/Oversubscription.h"
+#include "HAL/PlatformMath.h"
+#include "HAL/PlatformTime.h"
+
+#include <linux/futex.h>
+#include <sys/syscall.h>
+#include <time.h>
+#include <unistd.h>
+
+namespace UE::HAL::Private
+{
+
+inline static const struct timespec* SecondsToTimeSpec(double Seconds, struct timespec& OutTime)
+{
+	if (FPlatformMath::IsFinite(Seconds))
+	{
+		const double TotalSec = Seconds >= 0.0 ? Seconds : 0.0;
+		int32 Sec = int32(TotalSec);
+		int32 NSec = int32(FPlatformMath::Abs((TotalSec - Sec) * 1e9));
+		OutTime.tv_sec = Sec;
+		OutTime.tv_nsec = NSec;
+		return &OutTime;
+	}
+	return nullptr;
+}
+
+FORCENOINLINE void FUnixManualResetEvent::WaitSlow()
+{
+	LowLevelTasks::FOversubscriptionScope _;
+
+	for (;;)
+	{
+		if (State.load(std::memory_order_acquire))
+		{
+			return;
+		}
+		syscall(SYS_futex, &State, FUTEX_WAIT_PRIVATE, 0, nullptr, nullptr, 0);
+	}
+}
+
+FORCENOINLINE bool FUnixManualResetEvent::WaitForSlow(FMonotonicTimeSpan WaitTime)
+{
+	if (State.load(std::memory_order_acquire))
+	{
+		return true;
+	}
+
+	LowLevelTasks::FOversubscriptionScope _(WaitTime > FMonotonicTimeSpan::Zero());
+	struct timespec WaitTimeSpec;
+	const struct timespec* WaitTimeSpecPtr = SecondsToTimeSpec(WaitTime.ToSeconds(), WaitTimeSpec);
+
+	// FUTEX_WAIT takes a relative wait time.
+	const int ClockOption = FPlatformTime::GetClockSource() == CLOCK_REALTIME ? FUTEX_CLOCK_REALTIME : 0;
+	if (syscall(SYS_futex, &State, FUTEX_WAIT_PRIVATE | ClockOption, 0, WaitTimeSpecPtr, nullptr, 0) == -1 && errno == ETIMEDOUT)
+	{
+		return false;
+	}
+
+	if (LIKELY(State.load(std::memory_order_acquire)))
+	{
+		return true;
+	}
+
+	// Handle a spurious wake by waiting until the wait time has elapsed one more time because WaitUntilSlow
+	// handles spurious wakes in a loop and avoids exceeding the originally requested wake time by more than
+	// the typical variation due to scheduling imprecision.
+	return WaitUntilSlow(FMonotonicTimePoint::Now() + WaitTime);
+}
+
+FORCENOINLINE bool FUnixManualResetEvent::WaitUntilSlow(FMonotonicTimePoint WaitTime)
+{
+	LowLevelTasks::FOversubscriptionScope _;
+	struct timespec WaitTimeSpec;
+	const struct timespec* WaitTimeSpecPtr = SecondsToTimeSpec(WaitTime.ToSeconds(), WaitTimeSpec);
+
+	for (;;)
+	{
+		if (State.load(std::memory_order_acquire))
+		{
+			return true;
+		}
+
+		// FUTEX_WAIT_BITSET takes an absolute wait time.
+		const int ClockOption = FPlatformTime::GetClockSource() == CLOCK_REALTIME ? FUTEX_CLOCK_REALTIME : 0;
+		if (syscall(SYS_futex, &State, FUTEX_WAIT_BITSET_PRIVATE | ClockOption, 0, WaitTimeSpecPtr, nullptr, FUTEX_BITSET_MATCH_ANY) == -1 && errno == ETIMEDOUT)
+		{
+			return false;
+		}
+	}
+}
+
+void FUnixManualResetEvent::Notify()
+{
+	if (State.exchange(1, std::memory_order_release) == 0)
+	{
+		syscall(SYS_futex, &State, FUTEX_WAKE_PRIVATE, INT_MAX, nullptr, nullptr, 0);
+	}
+}
+
+} // UE::HAL::Private
diff --git a/Engine/Source/Runtime/Core/Private/Unix/UnixPlatformManualResetEvent.h b/Engine/Source/Runtime/Core/Private/Unix/UnixPlatformManualResetEvent.h
--- a/Engine/Source/Runtime/Core/Private/Unix/UnixPlatformManualResetEvent.h
+++ b/Engine/Source/Runtime/Core/Private/Unix/UnixPlatformManualResetEvent.h
@@ -2,6 +2,59 @@
 
 #pragma once
 
-#include "GenericPlatform/GenericPlatformManualResetEvent.h"
+#include "CoreTypes.h"
+#include "Misc/MonotonicTime.h"
 
-namespace UE::HAL::Private { using FUnixPlatformManualResetEvent = FGenericPlatformManualResetEvent; }
+#include <atomic>
+
+namespace UE::HAL::Private
+{
+
+/** @see FGenericPlatformManualResetEvent */
+class FUnixManualResetEvent
+{
+public:
+	FUnixManualResetEvent() = default;
+	FUnixManualResetEvent(const FUnixManualResetEvent&) = delete;
+	FUnixManualResetEvent& operator=(const FUnixManualResetEvent&) = delete;
+
+	inline void Reset()
+	{
+		State.store(0, std::memory_order_relaxed);
+	}
+
+	inline bool Poll()
+	{
+		return !!State.load(std::memory_order_acquire);
+	}
+
+	inline void Wait()
+	{
+		if (!State.load(std::memory_order_acquire))
+		{
+			WaitSlow();
+		}
+	}
+
+	inline bool WaitFor(FMonotonicTimeSpan WaitTime)
+	{
+		return State.load(std::memory_order_acquire) || WaitForSlow(WaitTime);
+	}
+
+	inline bool WaitUntil(FMonotonicTimePoint WaitTime)
+	{
+		return State.load(std::memory_order_acquire) || WaitUntilSlow(WaitTime);
+	}
+
+	void Notify();
+
+private:
+	void WaitSlow();
+	bool WaitForSlow(FMonotonicTimeSpan WaitTime);
+	bool WaitUntilSlow(FMonotonicTimePoint WaitTime);
+
+	// 0 = Reset, 1 = Notified
+	std::atomic<uint32> State = 0;
+};
+
+} // UE::HAL::Private
diff --git a/Engine/Source/Runtime/Core/Public/Unix/UnixPlatformTime.h b/Engine/Source/Runtime/Core/Public/Unix/UnixPlatformTime.h
--- a/Engine/Source/Runtime/Core/Public/Unix/UnixPlatformTime.h
+++ b/Engine/Source/Runtime/Core/Public/Unix/UnixPlatformTime.h
@@ -62,6 +62,11 @@ struct FUnixTime : public FGenericPlatformTime
 	 */
 	static CORE_API void PrintCalibrationLog();
 
+	UE_INTERNAL static int GetClockSource()
+	{
+		return ClockSource;
+	}
+
 private:
 
 	/** Clock source to use */
