[AutoRTFM] Fix compile error when bUseLoggingInShipping = false.

Also, removed an unbalanced `va_end` from each logging function after confirming that the callers were already doing a `va_end`.

diff --git a/Engine/Source/Runtime/Core/Private/AutoRTFM/AutoRTFMUE.cpp b/Engine/Source/Runtime/Core/Private/AutoRTFM/AutoRTFMUE.cpp
--- a/Engine/Source/Runtime/Core/Private/AutoRTFM/AutoRTFMUE.cpp
+++ b/Engine/Source/Runtime/Core/Private/AutoRTFM/AutoRTFMUE.cpp
@@ -277,6 +277,7 @@ TStringConversion<TStringConvert<UTF8CHAR, TCHAR>, 128> FmtToTChar(const char* F
 UE_AUTORTFM_ALWAYS_OPEN_NO_MEMORY_VALIDATION
 void AutoRTFMLog(const char* File, int Line, void* ProgramCounter, autortfm_log_severity Severity, const char* Format, va_list Args)
 {
+#if !NO_LOGGING
 	if (!IsSeverityActive(Severity))
 	{
 		return;
@@ -304,13 +305,13 @@ void AutoRTFMLog(const char* File, int Line, void* ProgramCounter, autortfm_log_
 			UE::Logging::Private::BasicFatalLogWithProgramCounter(LogAutoRTFM, &Record, ProgramCounter, FmtToTChar(Format, Args).Get());
 			break;
 	}
-
-	va_end(Args);
+#endif
 }
 
 UE_AUTORTFM_ALWAYS_OPEN_NO_MEMORY_VALIDATION
 void AutoRTFMLogWithCallstack(void* ProgramCounter, autortfm_log_severity Severity, const char* Format, va_list Args)
 {
+#if !NO_LOGGING
 	if (!IsSeverityActive(Severity))
 	{
 		return;
@@ -339,8 +340,7 @@ void AutoRTFMLogWithCallstack(void* ProgramCounter, autortfm_log_severity Severi
 			UE::Logging::Private::BasicFatalLogWithProgramCounter(LogAutoRTFM, &Record, ProgramCounter, FmtToTChar(Format, Args).Get());
 			break;
 	}
-
-	va_end(Args);
+#endif
 }
 
 UE_AUTORTFM_ALWAYS_OPEN_NO_MEMORY_VALIDATION
