Added Poll() to remaining implementations of FPlatformManualResetEvent

#rnx

diff --git a/Engine/Source/Runtime/Core/Private/Apple/ApplePlatformManualResetEvent.h b/Engine/Source/Runtime/Core/Private/Apple/ApplePlatformManualResetEvent.h
--- a/Engine/Source/Runtime/Core/Private/Apple/ApplePlatformManualResetEvent.h
+++ b/Engine/Source/Runtime/Core/Private/Apple/ApplePlatformManualResetEvent.h
@@ -2,6 +2,7 @@
 
 #pragma once
 
-#include "GenericPlatform/GenericPlatformManualResetEvent.h"
+// TODO: Use os_sync_wait_on_address/os_sync_wake_by_address_all in macOS 14.4+, iOS/iPadOS/tvOS 17.4+, visionOS 1.1+.
 
+#include "GenericPlatform/GenericPlatformManualResetEvent.h"
 namespace UE::HAL::Private { using FApplePlatformManualResetEvent = FGenericPlatformManualResetEvent; }
