FHierarchicalLogArchive: make FIndentScope movable to it can be stored in a stack + also made indentation look better.

#okfogrithub public

diff --git a/Engine/Source/Runtime/Core/Private/Misc/HierarchicalLogArchive.cpp b/Engine/Source/Runtime/Core/Private/Misc/HierarchicalLogArchive.cpp
--- a/Engine/Source/Runtime/Core/Private/Misc/HierarchicalLogArchive.cpp
+++ b/Engine/Source/Runtime/Core/Private/Misc/HierarchicalLogArchive.cpp
@@ -21,7 +21,16 @@ void FHierarchicalLogArchive::WriteLine(const FString& InLine, bool bIndent)
 			Builder << " |  ";
 		}
 
-		Builder << " |- ";
+		Builder << " |-";
+
+		if (bIndent)
+		{
+			Builder << "-";
+		}
+		else
+		{
+			Builder << " ";
+		}
 	}
 	
 	if (bIndent)
diff --git a/Engine/Source/Runtime/Core/Public/Misc/HierarchicalLogArchive.h b/Engine/Source/Runtime/Core/Public/Misc/HierarchicalLogArchive.h
--- a/Engine/Source/Runtime/Core/Public/Misc/HierarchicalLogArchive.h
+++ b/Engine/Source/Runtime/Core/Public/Misc/HierarchicalLogArchive.h
@@ -17,8 +17,32 @@ public:
 
 	struct FIndentScope
 	{
-		FIndentScope(FHierarchicalLogArchive* InAr) : Ar(InAr) { Ar->Indentation++; }
-		~FIndentScope() { check(Ar->Indentation); Ar->Indentation--; }
+		explicit FIndentScope(FHierarchicalLogArchive* InAr = nullptr)
+			: Ar(InAr)
+		{
+			if (Ar)
+			{
+				Ar->Indentation++;
+			}
+		}
+
+		FIndentScope(const FIndentScope& InOther) = delete;
+
+		explicit FIndentScope(FIndentScope&& InOther)
+			: Ar(InOther.Ar)
+		{
+			InOther.Ar = nullptr;
+		}
+
+		~FIndentScope()
+		{
+			if (Ar)
+			{
+				check(Ar->Indentation);
+				Ar->Indentation--;
+			}
+		}
+
 		FHierarchicalLogArchive* Ar;
 	};
 
