Fixed unreachable code compilation error when stat group is compiled out with DECLARE_STATS_GROUP_MAYBE_COMPILED_OUT.

diff --git a/Engine/Source/Runtime/Core/Public/Stats/LightweightStats.h b/Engine/Source/Runtime/Core/Public/Stats/LightweightStats.h
--- a/Engine/Source/Runtime/Core/Public/Stats/LightweightStats.h
+++ b/Engine/Source/Runtime/Core/Public/Stats/LightweightStats.h
@@ -310,23 +310,25 @@ namespace UE::Stats::Private
 		{
 			return false;
 		}
-
-		if constexpr (EnumHasAnyFlags(StatDataType::GetFlags(), EStatFlags::Verbose))
+		else
 		{
-			if (!GShouldEmitVerboseNamedEvents)
+			if constexpr (EnumHasAnyFlags(StatDataType::GetFlags(), EStatFlags::Verbose))
 			{
-				return false;
+				if (!GShouldEmitVerboseNamedEvents)
+				{
+					return false;
+				}
 			}
-		}
 
 #if UE_STATS_ALLOW_PER_THREAD_IGNORELIST
-		const uint32 GroupHash = GroupNameHash != 0 ? GroupNameHash : GroupDataType::GetNameHash();
-		if (UE::Stats::IsStatOrGroupIgnoredOnCurrentThread(StatNameHash, GroupHash))
-		{
-			return false;
-		}
+			const uint32 GroupHash = GroupNameHash != 0 ? GroupNameHash : GroupDataType::GetNameHash();
+			if (UE::Stats::IsStatOrGroupIgnoredOnCurrentThread(StatNameHash, GroupHash))
+			{
+				return false;
+			}
 #endif // UE_STATS_ALLOW_PER_THREAD_IGNORELIST
-		return true;
+			return true;
+		}
 	}
 
 	/**
