Removed unused FContainerTestKeyType.

diff --git a/Engine/Source/Runtime/Core/Private/Containers/ContainersTest.cpp b/Engine/Source/Runtime/Core/Private/Containers/ContainersTest.cpp
--- a/Engine/Source/Runtime/Core/Private/Containers/ContainersTest.cpp
+++ b/Engine/Source/Runtime/Core/Private/Containers/ContainersTest.cpp
@@ -149,43 +149,6 @@ namespace
 		int32                    Id;
 	};
 
-	struct FContainerTestKeyType : FContainerTestType
-	{
-		FContainerTestKeyType()
-			: FContainerTestType(TEXT("<default key>"), EContainerTestType::Key)
-		{
-		}
-
-		explicit FContainerTestKeyType(const TCHAR* InStr)
-			: FContainerTestType(InStr, EContainerTestType::Key)
-		{
-		}
-
-		FContainerTestKeyType(FContainerTestKeyType&& Other)
-			: FContainerTestType(MoveTemp(Other))
-		{
-		}
-
-		FContainerTestKeyType(const FContainerTestKeyType& Other)
-			: FContainerTestType(Other)
-		{
-		}
-
-		FContainerTestKeyType& operator=(FContainerTestKeyType&& Other)
-		{
-			(FContainerTestType&)*this = MoveTemp(Other);
-			
-			return *( static_cast<FContainerTestKeyType*>( this ) );
-		}
-
-		FContainerTestKeyType& operator=(const FContainerTestKeyType& Other)
-		{
-			(FContainerTestType&)*this = Other;
-			
-			return *( static_cast<FContainerTestKeyType*>( this ) );
-		}
-	};
-
 	struct FContainerTestValueType : FContainerTestType
 	{
 		FContainerTestValueType()
