Collapsed FContainerTestValueType into a single type with no base.

diff --git a/Engine/Source/Runtime/Core/Private/Containers/ContainersTest.cpp b/Engine/Source/Runtime/Core/Private/Containers/ContainersTest.cpp
--- a/Engine/Source/Runtime/Core/Private/Containers/ContainersTest.cpp
+++ b/Engine/Source/Runtime/Core/Private/Containers/ContainersTest.cpp
@@ -83,107 +83,82 @@ namespace
 		}
 	} ContainerTestStats;
 
-	struct FContainerTestType
+	struct FContainerTestValueType
 	{
-		FContainerTestType(const TCHAR* InStr, EContainerTestType::Type InType)
-			: Str (InStr)
+	private:
+		FContainerTestValueType(const TCHAR* InStr, EContainerTestType::Type InType)
+			: Str(InStr)
 			, Type(InType)
-			, Id  (ContainerTestStats.Add())
+			, Id(ContainerTestStats.Add())
 		{
 		}
 
-		FContainerTestType(const FContainerTestType& Other)
-			: Str (Other.Str)
-			, Type(Other.Type)
-			, Id  (ContainerTestStats.Add())
+	public:
+		FContainerTestValueType()
+			: FContainerTestValueType(TEXT("<default value>"), EContainerTestType::Value)
+		{
+		}
+
+		explicit FContainerTestValueType(const TCHAR* InStr)
+			: FContainerTestValueType(InStr, EContainerTestType::Value)
 		{
 		}
 
-		FContainerTestType(FContainerTestType&& Other)
-			: Str (Other.Str)
+		FContainerTestValueType(FContainerTestValueType&& Other)
+			: Str(Other.Str)
 			, Type(Other.Type)
-			, Id  (ContainerTestStats.Add())
+			, Id(ContainerTestStats.Add())
 		{
-			Other.Str  = NULL;
+			Other.Str = nullptr;
 			Other.Type = EContainerTestType::MovedFrom;
 		}
 
-		FContainerTestType& operator=(FContainerTestType&& Other)
+		FContainerTestValueType(const FContainerTestValueType& Other)
+			: Str(Other.Str)
+			, Type(Other.Type)
+			, Id(ContainerTestStats.Add())
 		{
-			Str  = Other.Str;
+		}
+
+		FContainerTestValueType& operator=(FContainerTestValueType&& Other)
+		{
+			Str = Other.Str;
 			Type = Other.Type;
 
-			Other.Str  = NULL;
+			Other.Str = nullptr;
 			Other.Type = EContainerTestType::MovedFrom;
 			
-			return *( static_cast<FContainerTestType*>( this ) );
+			return *this;
 		}
 
-		FContainerTestType& operator=(const FContainerTestType& Other)
+		FContainerTestValueType& operator=(const FContainerTestValueType& Other)
 		{
-			Str  = Other.Str;
+			Str = Other.Str;
 			Type = Other.Type;
 			
-			return *( static_cast<FContainerTestType*>( this ) );
+			return *this;
 		}
 
-		~FContainerTestType()
+		~FContainerTestValueType()
 		{
 			ContainerTestStats.Remove(Id);
 		}
 
-		friend bool operator==(const FContainerTestType& Lhs, const FContainerTestType& Rhs)
+		const TCHAR* Str;
+		EContainerTestType::Type Type;
+		int32 Id;
+
+		friend bool operator==(const FContainerTestValueType& Lhs, const FContainerTestValueType& Rhs)
 		{
 			return Lhs.Type == Rhs.Type &&
 			       Lhs.Id   == Rhs.Id   &&
 			       !FCString::Strcmp(Lhs.Str, Rhs.Str);
 		}
 
-		friend bool operator!=(const FContainerTestType& Lhs, const FContainerTestType& Rhs)
+		friend bool operator!=(const FContainerTestValueType& Lhs, const FContainerTestValueType& Rhs)
 		{
 			return !(Lhs == Rhs);
 		}
-
-		const TCHAR*             Str;
-		EContainerTestType::Type Type;
-		int32                    Id;
-	};
-
-	struct FContainerTestValueType : FContainerTestType
-	{
-		FContainerTestValueType()
-			: FContainerTestType(TEXT("<default value>"), EContainerTestType::Value)
-		{
-		}
-
-		explicit FContainerTestValueType(const TCHAR* InStr)
-			: FContainerTestType(InStr, EContainerTestType::Value)
-		{
-		}
-
-		FContainerTestValueType(FContainerTestValueType&& Other)
-			: FContainerTestType(MoveTemp(Other))
-		{
-		}
-
-		FContainerTestValueType(const FContainerTestValueType& Other)
-			: FContainerTestType(Other)
-		{
-		}
-
-		FContainerTestValueType& operator=(FContainerTestValueType&& Other)
-		{
-			(FContainerTestType&)*this = MoveTemp(Other);
-			
-			return *( static_cast<FContainerTestValueType*>( this ) );
-		}
-
-		FContainerTestValueType& operator=(const FContainerTestValueType& Other)
-		{
-			(FContainerTestType&)*this = Other;
-			
-			return *( static_cast<FContainerTestValueType*>( this ) );
-		}
 	};
 
 	template <typename Container>
