Added checks to .Remove() calls to ensure elements were actually removed.

diff --git a/Engine/Source/Runtime/Core/Private/Containers/ContainersTest.cpp b/Engine/Source/Runtime/Core/Private/Containers/ContainersTest.cpp
--- a/Engine/Source/Runtime/Core/Private/Containers/ContainersTest.cpp
+++ b/Engine/Source/Runtime/Core/Private/Containers/ContainersTest.cpp
@@ -286,7 +286,7 @@ namespace
 
 			for (int32 N = 0; N != Count; ++N)
 			{
-				Map.Remove(GenerateTestKey<KeyType>(N));
+				check(Map.Remove(GenerateTestKey<KeyType>(N)) == 1);
 				CheckMap();
 			}
 
@@ -300,7 +300,7 @@ namespace
 
 			for (int32 N = 0; N != Count; ++N)
 			{
-				Map.Remove(GenerateTestKey<KeyType>(N));
+				check(Map.Remove(GenerateTestKey<KeyType>(N)) == 1);
 				CheckMap();
 			}
 
@@ -327,7 +327,7 @@ namespace
 				SIZE_T PrevAllocatedSize = Map.GetAllocatedSize();
 				for (int32 N = MaxNum - 1; N >= MaxNum / 4; --N)
 				{
-					Map.Remove(GenerateTestKey<KeyType>(N));
+					check(Map.Remove(GenerateTestKey<KeyType>(N)) == 1);
 					Map.Shrink();
 					CheckMap();
 					check(Map.GetAllocatedSize() <= PrevAllocatedSize);
@@ -426,8 +426,8 @@ namespace
 			check(ValuePtr0 != nullptr);
 			check(ValuePtr1 != nullptr);
 			check(ValuePtr2 != nullptr);
-			Map.Remove(Key1);
-			Map.Remove(Key2);
+			check(Map.Remove(Key1) == 1);
+			check(Map.Remove(Key2) == 1);
 
 			// Re-add the 2 elements in the same order. Without the call to SortFreeList() 
 			// the elements would end up in a different locations/order compared to the 
@@ -890,7 +890,7 @@ bool FContainersTSetTest::RunTest(const FString& Parameters)
 			{
 				if (I != IndexToKeep)
 				{
-					Set.Remove(LexToString(I));
+					check(Set.Remove(LexToString(I)) == 1);
 				}
 			}
 
@@ -902,7 +902,7 @@ bool FContainersTSetTest::RunTest(const FString& Parameters)
 			}
 
 			Found = nullptr;
-			Set.Remove(LexToString(IndexToKeep));
+			check(Set.Remove(LexToString(IndexToKeep)) == 1);
 
 			FString* Found2 = Set.FindArbitraryElement();
 			TestNull(TEXT("FindArbitraryElement returns null on an emptied set"), Found2);
@@ -936,7 +936,7 @@ bool FContainersTMapTest::RunTest(const FString& Parameters)
 			{
 				if (I != IndexToKeep)
 				{
-					Map.Remove(I);
+					check(Map.Remove(I) == 1);
 				}
 			}
 
@@ -950,7 +950,7 @@ bool FContainersTMapTest::RunTest(const FString& Parameters)
 			}
 
 			Found = nullptr;
-			Map.Remove(IndexToKeep);
+			check(Map.Remove(IndexToKeep) == 1);
 
 			TPair<int32, FString>* Found2 = Map.FindArbitraryElement();
 			TestNull(TEXT("FindArbitraryElement returns null on an emptied map"), Found2);
@@ -984,7 +984,7 @@ bool FContainersTSortedMapTest::RunTest(const FString& Parameters)
 			{
 				if (I != IndexToKeep)
 				{
-					Map.Remove(I);
+					check(Map.Remove(I) == 1);
 				}
 			}
 
@@ -998,7 +998,7 @@ bool FContainersTSortedMapTest::RunTest(const FString& Parameters)
 			}
 
 			Found = nullptr;
-			Map.Remove(IndexToKeep);
+			check(Map.Remove(IndexToKeep) == 1);
 
 			TPair<int32, FString>* Found2 = Map.FindArbitraryElement();
 			TestNull(TEXT("FindArbitraryElement returns null on an emptied map"), Found2);
