Simplified the moved-from state of FContainerTestValueType.

diff --git a/Engine/Source/Runtime/Core/Private/Containers/ContainersTest.cpp b/Engine/Source/Runtime/Core/Private/Containers/ContainersTest.cpp
--- a/Engine/Source/Runtime/Core/Private/Containers/ContainersTest.cpp
+++ b/Engine/Source/Runtime/Core/Private/Containers/ContainersTest.cpp
@@ -21,16 +21,6 @@
 
 namespace
 {
-	namespace EContainerTestType
-	{
-		enum Type
-		{
-			MovedFrom,
-			Key,
-			Value
-		};
-	}
-
 	struct FContainerTestStats
 	{
 		int32  NextId;
@@ -85,48 +75,38 @@ namespace
 
 	struct FContainerTestValueType
 	{
-	private:
-		FContainerTestValueType(const TCHAR* InStr, EContainerTestType::Type InType)
-			: Str(InStr)
-			, Type(InType)
-			, Id(ContainerTestStats.Add())
-		{
-		}
-
-	public:
 		FContainerTestValueType()
-			: FContainerTestValueType(TEXT("<default value>"), EContainerTestType::Value)
+			: FContainerTestValueType(TEXT("<default value>"))
 		{
 		}
 
 		explicit FContainerTestValueType(const TCHAR* InStr)
-			: FContainerTestValueType(InStr, EContainerTestType::Value)
+			: Str(InStr)
+			, Id(ContainerTestStats.Add())
 		{
 		}
-
 		FContainerTestValueType(FContainerTestValueType&& Other)
 			: Str(Other.Str)
-			, Type(Other.Type)
 			, Id(ContainerTestStats.Add())
 		{
 			Other.Str = nullptr;
-			Other.Type = EContainerTestType::MovedFrom;
+			Other.bMovedFrom = true;
 		}
 
 		FContainerTestValueType(const FContainerTestValueType& Other)
 			: Str(Other.Str)
-			, Type(Other.Type)
 			, Id(ContainerTestStats.Add())
+			, bMovedFrom(Other.bMovedFrom)
 		{
 		}
 
 		FContainerTestValueType& operator=(FContainerTestValueType&& Other)
 		{
 			Str = Other.Str;
-			Type = Other.Type;
+			bMovedFrom = Other.bMovedFrom;
 
 			Other.Str = nullptr;
-			Other.Type = EContainerTestType::MovedFrom;
+			Other.bMovedFrom = true;
 			
 			return *this;
 		}
@@ -134,7 +114,7 @@ namespace
 		FContainerTestValueType& operator=(const FContainerTestValueType& Other)
 		{
 			Str = Other.Str;
-			Type = Other.Type;
+			bMovedFrom = Other.bMovedFrom;
 			
 			return *this;
 		}
@@ -145,14 +125,12 @@ namespace
 		}
 
 		const TCHAR* Str;
-		EContainerTestType::Type Type;
 		int32 Id;
+		bool bMovedFrom = false;
 
 		friend bool operator==(const FContainerTestValueType& Lhs, const FContainerTestValueType& Rhs)
 		{
-			return Lhs.Type == Rhs.Type &&
-			       Lhs.Id   == Rhs.Id   &&
-			       !FCString::Strcmp(Lhs.Str, Rhs.Str);
+			return Lhs.bMovedFrom == Rhs.bMovedFrom && Lhs.Id   == Rhs.Id && !FCString::Strcmp(Lhs.Str, Rhs.Str);
 		}
 
 		friend bool operator!=(const FContainerTestValueType& Lhs, const FContainerTestValueType& Rhs)
