Added some basic TSet tests, based on the existing TMap tests.

diff --git a/Engine/Source/Runtime/Core/Private/Containers/ContainersTest.cpp b/Engine/Source/Runtime/Core/Private/Containers/ContainersTest.cpp
--- a/Engine/Source/Runtime/Core/Private/Containers/ContainersTest.cpp
+++ b/Engine/Source/Runtime/Core/Private/Containers/ContainersTest.cpp
@@ -382,6 +382,136 @@ namespace
 		}
 	}
 
+	template <typename SetType, typename KeyType>
+	void RunSetTests()
+	{
+		SetType Set;
+
+		int32 MaxNum = 0;
+		SIZE_T MaxAllocatedSize = 0;
+		const SIZE_T InitialAllocatedSize = Set.GetAllocatedSize();
+
+		ContainerTestStats.Reset();
+
+		auto CheckSet = [&Set]()
+		{
+			CheckContainerNum(Set);
+			CheckContainerEnds(Set);
+			CheckContainerElements(Set);
+			CheckContainerSelfEquality(Set);
+		};
+
+		// Test Add and Remove
+		// Subtract one to account for temporaries that will be created during an Add
+		for (int32 Count = 0; Count < MAX_TEST_OBJECTS - 1; Count += MAX_TEST_OBJECTS_STEP)
+		{
+			for (int32 N = 0; N != Count; ++N)
+			{
+				Set.Add(GenerateTestKey<KeyType>(N));
+				CheckSet();
+			}
+			MaxNum = Set.Num();
+			MaxAllocatedSize = Set.GetAllocatedSize();
+
+			for (int32 N = 0; N != Count; ++N)
+			{
+				auto* ElementPtr = Set.FindArbitraryElement();
+				check(ElementPtr);
+				check(Set.Remove(*ElementPtr) == 1);
+				CheckSet();
+			}
+
+			check(Set.IsEmpty());
+
+			for (int32 N = 0; N != Count; ++N)
+			{
+				Set.Add(GenerateTestKey<KeyType>((Count - 1) - N));
+				CheckSet();
+			}
+
+			for (int32 N = 0; N != Count; ++N)
+			{
+				auto* ElementPtr = Set.FindArbitraryElement();
+				check(ElementPtr);
+				check(Set.Remove(*ElementPtr) == 1);
+				CheckSet();
+			}
+
+			check(Set.IsEmpty());
+		}
+
+		// Test Empty and Shrink 
+		{
+			// Test releasing memory allocations
+			Set.Empty();
+			CheckSet();
+			check(Set.GetAllocatedSize() == InitialAllocatedSize);
+
+			// Test integrity after re-growing container to MaxNum elements again
+			for (int32 N = 0; N < MaxNum; ++N)
+			{
+				Set.Add(GenerateTestKey<KeyType>(N));
+			}
+			CheckSet();
+			check(Set.GetAllocatedSize() == MaxAllocatedSize);
+
+			// Test data integrity while removing and shrinking continously
+			{
+				SIZE_T PrevAllocatedSize = Set.GetAllocatedSize();
+				for (int32 N = MaxNum - 1; N >= MaxNum / 4; --N)
+				{
+					auto* ElementPtr = Set.FindArbitraryElement();
+					check(ElementPtr);
+					Set.Remove(*ElementPtr);
+					Set.Shrink();
+					CheckSet();
+					check(Set.GetAllocatedSize() <= PrevAllocatedSize);
+					PrevAllocatedSize = Set.GetAllocatedSize();
+				}
+			}
+
+			// Test removing and releasing remaining elements
+			Set.Empty();
+			check(Set.IsEmpty());
+			check(Set.GetAllocatedSize() == InitialAllocatedSize);
+		}
+
+		// Test iterators
+		{
+			static_assert(std::is_same_v<decltype(*Set.CreateIterator     ()),         KeyType&>);
+			static_assert(std::is_same_v<decltype(*Set.CreateConstIterator()),   const KeyType&>);
+
+			for (int32 Count = 0; Count < MAX_TEST_OBJECTS - 1; Count += MAX_TEST_OBJECTS_STEP)
+			{
+				Set.Empty();
+				CheckSet();
+
+				for (int32 N = 0; N != Count; ++N)
+				{
+					Set.Add(GenerateTestKey<KeyType>(N));
+					CheckSet();
+				}
+
+				// Check that the const iterator finds all the values
+				TBitArray<> Found;
+				Found.Init(false, Count);
+				for (auto It = Set.CreateConstIterator(); It; ++It)
+				{
+					Found[*It] = true;
+				}
+				check(Found.Find(false) == INDEX_NONE);
+
+				// Check that the non-const iterator finds all the values
+				Found.Init(false, Count);
+				for (auto It = Set.CreateConstIterator(); It; ++It)
+				{
+					Found[*It] = true;
+				}
+				check(Found.Find(false) == INDEX_NONE);
+			}
+		}
+	}
+
 	// Test container element address consistency when using SortFreeList 
 	// (see TSparseArray::SortFreeList for comments)
 	template <typename MapType, typename KeyType>
@@ -583,6 +713,10 @@ bool FContainersSmokeTest::RunTest( const FString& Parameters )
 	RunMapTests<TSortedMap<int32, FContainerTestValueType, TInlineAllocator<32>>, int32>();
 	RunMapTests<TSortedMap<int32, FContainerTestValueType, TFixedAllocator<64>>, int32>();
 
+	RunSetTests<TSet<int32>, int32>();
+	RunSetTests<TSet<int32, DefaultKeyFuncs<int32>, TInlineSetAllocator<32>>, int32>();
+	RunSetTests<TSet<int32, DefaultKeyFuncs<int32>, TFixedSetAllocator<64>>, int32>();
+
 	return true;
 }
 
@@ -608,6 +742,10 @@ bool FContainersFullTest::RunTest(const FString& Parameters)
 
 	RunEmptyContainerSelfEqualityTest<TArray<int32>>();
 
+	RunSetTests<TSet<int32>, int32>();
+	RunSetTests<TSet<int32, DefaultKeyFuncs<int32>, TInlineSetAllocator<32>>, int32>();
+	RunSetTests<TSet<int32, DefaultKeyFuncs<int32>, TFixedSetAllocator<64>>, int32>();
+
 	// Verify use of FName index sorter with SortedMap
 
 	TSortedMap<FName, int32, FDefaultAllocator, FNameFastLess> NameMap;
