Added download manager error analytics

diff --git a/Engine/Source/Runtime/Core/Private/Experimental/IO/IoStatusError.cpp b/Engine/Source/Runtime/Core/Private/Experimental/IO/IoStatusError.cpp
--- a/Engine/Source/Runtime/Core/Private/Experimental/IO/IoStatusError.cpp
+++ b/Engine/Source/Runtime/Core/Private/Experimental/IO/IoStatusError.cpp
@@ -2,24 +2,30 @@
 
 #include "Experimental/IO/IoStatusError.h"
 #include "IO/IoStatus.h"
+#include "Containers/AnsiString.h"
 #include "Internationalization/Internationalization.h"
 #include "Internationalization/Text.h"
 
 
 namespace UE::UnifiedError::IoStore
 {
-	class FIoStoreErrorDetails : public UE::UnifiedError::FRefCountedErrorDetails
+	class FIoStoreErrorDetails : public UE::UnifiedError::FManditoryErrorDetails, public FRefCountBase
 	{
 	private:
 		FString ErrorMessage;
 		mutable FUtf8String CachedErrorName;
 		static const FText GenericErrorFormatString;
 	public:
+
 		FIoStoreErrorDetails(const FStringView& InErrorMessage)
 		{
 			ErrorMessage = InErrorMessage;
 		}
 
+		virtual FReturnedRefCountValue AddRef() const final { return FRefCountBase::AddRef(); }
+		virtual uint32 Release() const final { return FRefCountBase::Release(); }
+		virtual uint32 GetRefCount() const final { return FRefCountBase::GetRefCount(); }
+
 		const FUtf8String& GetErrorName(const FError& Error) const
 		{
 			if (CachedErrorName.IsEmpty() && (ErrorMessage[0] != TEXT('\0')))
@@ -29,6 +35,16 @@ namespace UE::UnifiedError::IoStore
 			return CachedErrorName;
 		}
 
+		FAnsiString GetErrorCodeString(const UE::UnifiedError::FError& Error) const override
+		{
+			return TCHAR_TO_ANSI(GetIoErrorText(EIoErrorCode(Error.GetErrorCode())));
+		}
+
+		FAnsiString GetModuleIdString(const UE::UnifiedError::FError& Error) const override
+		{
+			return FAnsiString(UE::UnifiedError::IoStore::StaticModuleName);
+		}
+
 		/*
 
 		virtual UE::UnifiedError::FAnyString GetErrorCodeString(const FError& Error) const final
diff --git a/Engine/Source/Runtime/Core/Private/Experimental/UnifiedError/UnifiedError.cpp b/Engine/Source/Runtime/Core/Private/Experimental/UnifiedError/UnifiedError.cpp
--- a/Engine/Source/Runtime/Core/Private/Experimental/UnifiedError/UnifiedError.cpp
+++ b/Engine/Source/Runtime/Core/Private/Experimental/UnifiedError/UnifiedError.cpp
@@ -1,6 +1,7 @@
 // Copyright Epic Games, Inc. All Rights Reserved.
 
 #include "Experimental/UnifiedError/UnifiedError.h"
+#include "Containers/AnsiString.h"
 #include "Internationalization/Text.h"
 #include "Internationalization/Internationalization.h"
 #include "Logging/StructuredLog.h"
@@ -10,8 +11,6 @@
 namespace UE::UnifiedError
 {
 
-
-
 // FError functions
 
 void FError::PushErrorDetails(TRefCountPtr<IErrorDetails> InErrorDetails)
@@ -44,19 +43,63 @@ FText FError::GetFormatErrorText() const
 {
 	return ErrorDetails->GetErrorFormatString(*this);
 }
-/*
 
+const FManditoryErrorDetails* FError::GetManditoryErrorDetails() const
+{
+
+
+	TRefCountPtr<const IErrorDetails> CurrentIt = ErrorDetails;
+	while (CurrentIt != nullptr)
+	{
+		TRefCountPtr<const IErrorDetails> Next = CurrentIt->GetInnerErrorDetails();
+		if (Next == nullptr)
+		{
+			break;
+		}
+		CurrentIt = Next;
+	}
+
+	if (!ensure(CurrentIt.GetReference()))
+	{
+		return nullptr;
+	}
+	return (const FManditoryErrorDetails*)(CurrentIt.GetReference());
+
+
+	/*const TRefCountPtr<IErrorDetails> It = ErrorDetails;
+	while (const auto Next = It->GetInnerErrorDetails())
+	{
+		It = Next;
+	}
 
-UE::UnifiedError::FAnyString FError::GetErrorCodeString() const
+	
+	if(!ensure(It.GetReference()))
+	{
+		return nullptr;
+	}
+	return (const FManditoryErrorDetails*)(It.GetReference());*/
+}
+
+FAnsiString FError::GetModuleIdAndErrorCodeString() const
 {
-	return ErrorDetails->GetErrorCodeString(*this);
+	return FAnsiString::Printf("%s.%s", *GetModuleIdString(), *GetErrorCodeString());
 }
 
-UE::UnifiedError::FAnyString FError::GetModuleIdString() const
+
+FAnsiString FError::GetErrorCodeString() const
 {
-	return ErrorDetails->GetModuleIdString(*this);
+	const FManditoryErrorDetails* ManditoryErrorDetails = GetManditoryErrorDetails();
+	check(ManditoryErrorDetails);
+	return ManditoryErrorDetails->GetErrorCodeString(*this);
 }
-*/
+
+FAnsiString FError::GetModuleIdString() const
+{
+	const FManditoryErrorDetails* ManditoryErrorDetails = GetManditoryErrorDetails();
+	check(ManditoryErrorDetails);
+	return ManditoryErrorDetails->GetModuleIdString(*this);
+}
+
 
 
 #define USE_STRUCTURED_LOG_FOR_FERRORMESSAGE 1
@@ -279,6 +322,16 @@ void FError::AppendFormatString(FText&& InFormatString)
 }
 
 
+FAnsiString FStaticErrorDetails::GetErrorCodeString(const FError& Error) const
+{
+	return FAnsiString(ErrorName);
+}
+
+FAnsiString FStaticErrorDetails::GetModuleIdString(const FError& Error) const
+{
+	return FAnsiString(ModuleName);
+}
+
 // Structured logging FError integration
 class FCbWriterErrorPropertyExtractor : public IErrorPropertyExtractor
 {
diff --git a/Engine/Source/Runtime/Core/Public/Experimental/UnifiedError/UnifiedError.h b/Engine/Source/Runtime/Core/Public/Experimental/UnifiedError/UnifiedError.h
--- a/Engine/Source/Runtime/Core/Public/Experimental/UnifiedError/UnifiedError.h
+++ b/Engine/Source/Runtime/Core/Public/Experimental/UnifiedError/UnifiedError.h
@@ -274,6 +274,19 @@ namespace UE::UnifiedError
 	};
 
 
+	/// <summary>
+	/// FManditoryErrorDetails 
+	/// every FError needs to be initialized with one of these, it contains core information about the error including errorcodestring and module name
+	/// </summary>
+	class FManditoryErrorDetails : public IErrorDetails
+	{
+		DECLARE_FERROR_DETAILS_ABSTRACT(UnifiedError, FManditoryErrorDetails);
+	private:
+	public:
+		virtual FAnsiString GetErrorCodeString(const FError& Error) const = 0;
+		virtual FAnsiString GetModuleIdString(const FError& Error) const = 0;
+	};
+
 	class FErrorDetailsRegistry
 	{
 	private:
@@ -376,10 +389,10 @@ namespace UE::UnifiedError
 			ErrorDetail = MoveTemp(InErrorDetail);
 			bAppendContextToMessage = bInAppendContextToMessage;
 		}
-		TErrorDetails(T&& InErrorDetail, TRefCountPtr<const IErrorDetails> InInnerErrorDetails) : FDynamicErrorDetails(InInnerErrorDetails)
+		/*TErrorDetails(T&& InErrorDetail, TRefCountPtr<const IErrorDetails> InInnerErrorDetails) : FDynamicErrorDetails(InInnerErrorDetails)
 		{
 			ErrorDetail = MoveTemp(InErrorDetail);
-		}
+		}*/
 
 		friend IErrorDetails* Create();
 		inline static IErrorDetails* Create()
@@ -465,7 +478,7 @@ namespace UE::UnifiedError
 		int32 ModuleId;
 		TRefCountPtr<const IErrorDetails> ErrorDetails;
 	public:
-		FError(int32 InModuleId, int32 InErrorCode, const IErrorDetails* InErrorDetails)
+		FError(int32 InModuleId, int32 InErrorCode, const FManditoryErrorDetails* InErrorDetails)
 		{
 			ErrorCode = InErrorCode;
 			ModuleId = InModuleId;
@@ -535,6 +548,16 @@ namespace UE::UnifiedError
 			return nullptr;*/
 		}
 
+		CORE_API FAnsiString GetErrorCodeString() const;
+		CORE_API FAnsiString GetModuleIdString() const;
+
+
+		/// <summary>
+		/// GetFullErrorCodeString, Return the combined module id and error code
+		/// </summary>
+		/// <returns></returns>
+		CORE_API FAnsiString GetModuleIdAndErrorCodeString() const;
+
 		CORE_API void AppendFormatString(FText&& InFormatString);
 
 		CORE_API void SerializeDetailsForLog(FCbWriter& Writer) const;
@@ -560,6 +583,7 @@ namespace UE::UnifiedError
 		}
 	private:
 
+		const FManditoryErrorDetails* GetManditoryErrorDetails() const;
 
 		CORE_API void PushErrorDetails(TRefCountPtr<IErrorDetails> InErrorDetails);
 		template<typename DetailType>
@@ -629,13 +653,14 @@ namespace UE::UnifiedError
 	};
 
 
+
 	/// <summary>
 	/// FStaticErrorDetails; static error details and members are statically allocated
 	///  Every error which uses DEFINE_ERROR will have FStaticErrorDetails generated for it
 	///  Can not rely on it to be available for every error as some Error conversion functions will not use pregenerated errors or error codesF
 	///  Use FError::GetErrorDetails to discover FStaticErrorDetails
 	/// </summary>
-	class FStaticErrorDetails : public IErrorDetails
+	class FStaticErrorDetails : public FManditoryErrorDetails
 	{
 		DECLARE_FERROR_DETAILS_ABSTRACT(UnifiedError, FStaticErrorDetails);
 	private:
@@ -654,14 +679,14 @@ namespace UE::UnifiedError
 		///  See also: FError::GetErrorDetails
 		/// </summary>
 		/// <returns></returns>
-		//virtual UE::UnifiedError::FAnyString GetErrorCodeString(const FError& Error) const override;
+		virtual FAnsiString GetErrorCodeString(const FError& Error) const override;
 
 		/// <summary>
 		/// GetModuleIdString; accessor for ModuleName.
 		///  Can be called directly on FStaticErrorDetails object
 		///  See also: FError::GetErrorDetails
 		/// </summary>
-		///virtual UE::UnifiedError::FAnyString GetModuleIdString(const FError& Error) const override;
+		virtual FAnsiString GetModuleIdString(const FError& Error) const override;
 
 		// IErrorDetails implementation
 
@@ -827,7 +852,7 @@ CORE_API constexpr uint32 GetTypeIdHash(const TCHAR* TypeName)
 #define DECLARE_ERROR(ErrorName, ErrorCode, ModuleName, FormatString) \
 	DECLARE_ERROR_INTERNAL(ErrorName, ErrorCode, ModuleName, FormatString) \
 	namespace UE::UnifiedError { namespace ModuleName { namespace ErrorName {\
-			FORCEINLINE FError MakeError() { return FError(UE::UnifiedError::ModuleName::StaticModuleId, GetErrorCodeId(), TRefCountPtr<const IErrorDetails>(GetStaticErrorDetails().GetReference())); }\
+			FORCEINLINE FError MakeError() { return FError(UE::UnifiedError::ModuleName::StaticModuleId, GetErrorCodeId(), TRefCountPtr<const UE::UnifiedError::FManditoryErrorDetails>(GetStaticErrorDetails().GetReference())); }\
 			FORCEINLINE const FError& GetStaticError() \
 			{\
 				static FError StaticError = MakeError();\
@@ -840,44 +865,33 @@ CORE_API constexpr uint32 GetTypeIdHash(const TCHAR* TypeName)
 
 #define DECLARE_ERROR_ONEPARAM(ErrorName, ErrorCode, ModuleName, FormatString, ParamOneType, ParamOneName, ParamOneDefault) \
 	DECLARE_ERROR_INTERNAL(ErrorName, ErrorCode, ModuleName, FormatString) \
-	namespace UE::UnifiedError { namespace ModuleName { namespace ErrorName { \
-			class FDetails : public FDynamicErrorDetails \
-			{ \
-				DECLARE_FERROR_DETAILS(ModuleName::ErrorName, FDetails)\
-			private:\
-				ParamOneType ParamOneName; \
-				FDetails() : ParamOneName(ParamOneDefault) {}\
-			public: \
-				FDetails(const ParamOneType& In##ParamOneName, TRefCountPtr<const IErrorDetails> InErrorDetails) : FDynamicErrorDetails(InErrorDetails) \
-				{ \
-					ParamOneName = In##ParamOneName; \
-				} \
-				virtual void GetErrorProperties(const FError& Error, IErrorPropertyExtractor& OutProperties) const override \
-				{ \
-					OutProperties.AddProperty(UTF8TEXTVIEW(#ParamOneName), ParamOneName); \
-					FDynamicErrorDetails::GetErrorProperties(Error, OutProperties); \
-				} \
-				virtual void SerializeToCb(FCbWriter& Writer, const FError& Error) const override \
-				{\
-					Writer.BeginObject();\
-					SerializeBasicToCb(Writer, UTF8TEXTVIEW(#ParamOneName), ParamOneName); \
-					Writer.AddString(UTF8TEXTVIEW("$format"), GetInnerErrorDetails()->GetErrorFormatString(Error).ToString());\
-					Writer.EndObject();\
-				}\
-				virtual const FText GetErrorFormatString(const FError& Error) const override\
-				{\
-					FUtf8String ReplacedFormatString = FUtf8String::Printf("{%s}", StaticGetErrorDetailsTypeNameForLog().GetData());\
-					return FText::FromString(*ReplacedFormatString);\
-				}\
-				const ParamOneType& Get##ParamOneName() const \
-				{ \
-					return ParamOneName; \
-				} \
-			}; \
+	namespace UE::UnifiedError { namespace ModuleName { \
+			struct F##ErrorName \
+			{\
+				ParamOneType ParamOneName;\
+			};\
+	}}\
+	DECLARE_ERRORSTRUCT_FEATURES(ModuleName, F##ErrorName);\
+	inline void SerializeForLog(FCbWriter& Writer, const UE::UnifiedError::ModuleName::F##ErrorName& Context)\
+	{\
+		Writer.BeginObject();\
+		Writer.AddString(ANSITEXTVIEW("$type"), TErrorStructFeatures<UE::UnifiedError::ModuleName::F##ErrorName>::GetErrorContextTypeNameAsString());\
+		Writer.AddString(ANSITEXTVIEW("$format"), FormatString.ToString());\
+		UE::UnifiedError::SerializeBasicToCb(Writer, ANSITEXTVIEW(#ParamOneName), Context.ParamOneName);\
+		Writer.EndObject();\
+	}\
+	namespace UE::UnifiedError { namespace ModuleName { \
+	inline void GatherPropertiesForError(const UE::UnifiedError::FError& Error, const UE::UnifiedError::ModuleName::F##ErrorName& Context, UE::UnifiedError::IErrorPropertyExtractor& PropertyExtractor) \
+	{ \
+		PropertyExtractor.AddProperty(ANSITEXTVIEW(#ParamOneName), Context.ParamOneName);\
+	}\
+	namespace ErrorName { \
 			FORCEINLINE FError MakeError(ParamOneType ParamOneName = ParamOneDefault) \
 			{ \
-				FDetails* NewDetails = new FDetails(ParamOneName, TRefCountPtr<const IErrorDetails>(GetStaticErrorDetails().GetReference())); \
-				return FError(UE::UnifiedError::ModuleName::StaticModuleId, GetErrorCodeId(), NewDetails); \
+				FError Error = FError(UE::UnifiedError::ModuleName::StaticModuleId, UE::UnifiedError::ModuleName::ErrorName::GetErrorCodeId(), UE::UnifiedError::ModuleName::ErrorName::GetStaticErrorDetails().GetReference()); \
+				F##ErrorName Context = {ParamOneName}; \
+				Error.PushErrorContext(MoveTemp(Context)); \
+				return Error; \
 			} \
 			FORCEINLINE const FError& GetStaticError() \
 			{\
diff --git a/Engine/Source/Runtime/Engine/Classes/Engine/Experimental/StreamableManagerError.h b/Engine/Source/Runtime/Engine/Classes/Engine/Experimental/StreamableManagerError.h
--- a/Engine/Source/Runtime/Engine/Classes/Engine/Experimental/StreamableManagerError.h
+++ b/Engine/Source/Runtime/Engine/Classes/Engine/Experimental/StreamableManagerError.h
@@ -8,6 +8,50 @@
 DECLARE_ERROR_MODULE(StreamableManager, 0x10);
 
 
+//DECLARE_ERROR_ONEPARAM(PackageLoadFailed, 1, StreamableManager, NSLOCTEXT("StreamableManager", "PackageLoadError", "Failed to load package {PackageName}"), FString, PackageName, TEXT("Unknown"));
+
+/*
+
+
+	DECLARE_ERROR_INTERNAL(PackageLoadFailed, 1, StreamableManager, NSLOCTEXT("StreamableManager", "PackageLoadError", "Failed to load package {PackageName}")) \
+	namespace UE::UnifiedError { namespace StreamableManager { \
+			struct FPackageLoadFailed \
+			{\
+				FString PackageName;\
+			};\
+	}}\
+	DECLARE_ERRORSTRUCT_FEATURES(StreamableManager, FPackageLoadFailed);\
+	inline void SerializeForLog(FCbWriter& Writer, const UE::UnifiedError::StreamableManager::FPackageLoadFailed& Context)\
+	{\
+		Writer.BeginObject();\
+		Writer.AddString(ANSITEXTVIEW("$type"), TErrorStructFeatures<UE::UnifiedError::StreamableManager::FPackageLoadFailed>::GetErrorContextTypeNameAsString());\
+		Writer.AddString(ANSITEXTVIEW("$format"),  NSLOCTEXT("StreamableManager", "PackageLoadError", "Failed to load package {PackageName}").ToString());\
+		Writer.AddString(ANSITEXTVIEW("PackageName"), Context.PackageName);\
+		Writer.EndObject();\
+	}\
+	namespace UE::UnifiedError { namespace StreamableManager { namespace PackageLoadError { \
+			inline const int32 ErrorCodeId = UE::UnifiedError::FErrorRegistry::Get().RegisterErrorCode(TEXTVIEW("PackageLoadFailed"), UE::UnifiedError::StreamableManager::StaticModuleId, 1); \
+			FORCEINLINE int32 GetErrorCodeId() \
+			{\
+				return ErrorCodeId; \
+			}\
+			FORCEINLINE FError MakeError(FString PackageName = TEXT("Unknown")) \
+			{ \
+				FError Error = FError(UE::UnifiedError::StreamableManager::StaticModuleId, 1, UE::UnifiedError::StreamableManager::PackageLoadFailed::GetStaticErrorDetails().GetReference()); \
+				FPackageLoadFailed Context = { PackageName };\
+				Error.PushErrorContext(Context); \
+				return Error; \
+			} \
+			FORCEINLINE const FError& GetStaticError() \
+			{\
+				static FError StaticError = MakeError();\
+				return StaticError;\
+			}\
+	}}}
+
+*/
+
+
 DECLARE_ERROR_ONEPARAM(PackageLoadFailed, 1, StreamableManager, NSLOCTEXT("StreamableManager", "PackageLoadError", "Failed to load package {PackageName}"), FString, PackageName, TEXT("Unknown"));
 DECLARE_ERROR_ONEPARAM(PackageLoadCanceled, 2, StreamableManager, NSLOCTEXT("StreamableManager", "PackageLoadCancelled", "Async load canceled {PackageName}"), FString, PackageName, TEXT("Unknown"));
 DECLARE_ERROR(DownloadError, 3, StreamableManager, NSLOCTEXT("StreamableManager", "DownloadError", "Failed to download"));
@@ -18,7 +62,47 @@ DECLARE_ERROR_ONEPARAM(SyncLoadIncomplete, 7, StreamableManager, NSLOCTEXT("Stre
 
 DECLARE_ERROR(AsyncLoadFailed, 8, StreamableManager, NSLOCTEXT("StreamableManager", "AsyncLoadFailed", "Async load failed"));
 DECLARE_ERROR(AsyncLoadCancelled, 9, StreamableManager, NSLOCTEXT("StreamableManager", "AsyncLoadCancelled", "Async load cancelled"));
-DECLARE_ERROR_ONEPARAM(AsyncLoadUnknownError, 10, StreamableManager, NSLOCTEXT("StreamableManager", "AsyncLoadUnknownError", "Unknown async loading error {AsyncLoadingErrorId}."), int32, AsyncLoadingErrorId, -1);
+//DECLARE_ERROR_ONEPARAM(AsyncLoadUnknownError, 10, StreamableManager, NSLOCTEXT("StreamableManager", "AsyncLoadUnknownError", "Unknown async loading error {AsyncLoadingErrorId}."), int32, AsyncLoadingErrorId, -1);
+
+	DECLARE_ERROR_INTERNAL(AsyncLoadUnknownError, 10, StreamableManager, NSLOCTEXT("StreamableManager", "AsyncLoadUnknownError", "Unknown async loading error {AsyncLoadingErrorId}.")) \
+	namespace UE::UnifiedError { namespace StreamableManager { \
+			struct FAsyncLoadUnknownError \
+			{\
+				int32 AsyncLoadingErrorId;\
+			};\
+	}}\
+	DECLARE_ERRORSTRUCT_FEATURES(StreamableManager, FAsyncLoadUnknownError);\
+	inline void SerializeForLog(FCbWriter& Writer, const UE::UnifiedError::StreamableManager::FAsyncLoadUnknownError& Context)\
+	{\
+		Writer.BeginObject();\
+		Writer.AddString(ANSITEXTVIEW("$type"), TErrorStructFeatures<UE::UnifiedError::StreamableManager::FAsyncLoadUnknownError>::GetErrorContextTypeNameAsString());\
+		Writer.AddString(ANSITEXTVIEW("$format"), NSLOCTEXT("StreamableManager", "AsyncLoadUnknownError", "Unknown async loading error {AsyncLoadingErrorId}.").ToString());\
+		UE::UnifiedError::SerializeBasicToCb(Writer, ANSITEXTVIEW("AsyncLoadingErrorId"), Context.AsyncLoadingErrorId); \
+		Writer.EndObject(); \
+	}\
+	
+		namespace UE::UnifiedError {
+		namespace StreamableManager {
+			
+			inline void GatherPropertiesForError(const UE::UnifiedError::FError& Error, const UE::UnifiedError::StreamableManager::FAsyncLoadUnknownError& Context, UE::UnifiedError::IErrorPropertyExtractor& PropertyExtractor) \
+			{ \
+				PropertyExtractor.AddProperty(ANSITEXTVIEW("AsyncLoadingErrorId"), Context.AsyncLoadingErrorId); \
+			}\
+	namespace AsyncLoadUnknownError { \
+		
+			FORCEINLINE FError MakeError(int32 AsyncLoadingErrorId = -1) \
+			{ \
+				FError Error = FError(UE::UnifiedError::StreamableManager::StaticModuleId, UE::UnifiedError::StreamableManager::AsyncLoadUnknownError::GetErrorCodeId(), UE::UnifiedError::StreamableManager::AsyncLoadUnknownError::GetStaticErrorDetails().GetReference()); \
+				UE::UnifiedError::StreamableManager::FAsyncLoadUnknownError Context = {AsyncLoadingErrorId}; \
+				Error.PushErrorContext(MoveTemp(Context)); \
+				return Error; \
+			} \
+			FORCEINLINE const FError& GetStaticError() \
+			{\
+				static FError StaticError = MakeError();\
+				return StaticError;\
+			}\
+	}}}
 
 
 DECLARE_ERROR(UnknownError, 11, StreamableManager, NSLOCTEXT("StreamableManager", "UnknownError", "Unknown error occured while streaming asset"));
