Fix for missing functions declarations in vkcontent worker.

diff --git a/Engine/Source/Runtime/Core/Private/Experimental/UnifiedError/UnifiedError.cpp b/Engine/Source/Runtime/Core/Private/Experimental/UnifiedError/UnifiedError.cpp
--- a/Engine/Source/Runtime/Core/Private/Experimental/UnifiedError/UnifiedError.cpp
+++ b/Engine/Source/Runtime/Core/Private/Experimental/UnifiedError/UnifiedError.cpp
@@ -322,16 +322,6 @@ void FError::AppendFormatString(FText&& InFormatString)
 }
 
 
-FAnsiString FStaticErrorDetails::GetErrorCodeString(const FError& Error) const
-{
-	return FAnsiString(ErrorName);
-}
-
-FAnsiString FStaticErrorDetails::GetModuleIdString(const FError& Error) const
-{
-	return FAnsiString(ModuleName);
-}
-
 // Structured logging FError integration
 class FCbWriterErrorPropertyExtractor : public IErrorPropertyExtractor
 {
@@ -486,7 +476,6 @@ FStaticErrorDetails::FStaticErrorDetails(const FAnsiStringView& InErrorName, con
 }
 
 
-
 /*
 
 UE::UnifiedError::FAnyString FStaticErrorDetails::GetErrorCodeString(const FError& Error) const
diff --git a/Engine/Source/Runtime/Core/Public/Experimental/UnifiedError/UnifiedError.h b/Engine/Source/Runtime/Core/Public/Experimental/UnifiedError/UnifiedError.h
--- a/Engine/Source/Runtime/Core/Public/Experimental/UnifiedError/UnifiedError.h
+++ b/Engine/Source/Runtime/Core/Public/Experimental/UnifiedError/UnifiedError.h
@@ -2,6 +2,7 @@
 #pragma once
 
 #include "CoreTypes.h"
+#include "Containers/AnsiString.h"
 #include "Containers/Utf8String.h"
 #include "Containers/StringView.h"
 #include "Internationalization/Internationalization.h"
@@ -679,14 +680,21 @@ namespace UE::UnifiedError
 		///  See also: FError::GetErrorDetails
 		/// </summary>
 		/// <returns></returns>
-		virtual FAnsiString GetErrorCodeString(const FError& Error) const override;
+		virtual FAnsiString GetErrorCodeString(const FError& Error) const override
+		{
+			return FAnsiString(ErrorName);
+		}
+
 
 		/// <summary>
 		/// GetModuleIdString; accessor for ModuleName.
 		///  Can be called directly on FStaticErrorDetails object
 		///  See also: FError::GetErrorDetails
 		/// </summary>
-		virtual FAnsiString GetModuleIdString(const FError& Error) const override;
+		virtual FAnsiString GetModuleIdString(const FError& Error) const override
+		{
+			return FAnsiString(ModuleName);
+		}
 
 		// IErrorDetails implementation
 
