Removed deprecated Persistent Auxiliary allocator API

diff --git a/Engine/Source/Runtime/Core/Private/HAL/UnrealMemory.cpp b/Engine/Source/Runtime/Core/Private/HAL/UnrealMemory.cpp
--- a/Engine/Source/Runtime/Core/Private/HAL/UnrealMemory.cpp
+++ b/Engine/Source/Runtime/Core/Private/HAL/UnrealMemory.cpp
@@ -23,7 +23,6 @@
 #include "HAL/MallocDoubleFreeFinder.h"
 #include "HAL/MallocFrameProfiler.h"
 #include "HAL/MallocStomp2.h"
-#include "Memory/LinearAllocator.h"
 
 #if MALLOC_GT_HOOKS
 
@@ -672,21 +671,6 @@ void FUseSystemMallocForNew::operator delete[](void* Ptr)
 	FMemory::SystemFree(Ptr);
 }
 
-void* FMemory::MallocPersistentAuxiliary(SIZE_T InSize, uint32 InAlignment)
-{
-	return GetPersistentLinearAllocator().Allocate(InSize, InAlignment);
-}
-
-bool FMemory::IsPersistentAuxiliaryActive()
-{
-	return GPersistentLinearAllocatorExtends.Size != 0;
-}
-
-SIZE_T FMemory::GetUsedPersistentAuxiliary()
-{
-	return GetPersistentLinearAllocator().GetAllocatedMemorySize();
-}
-
 
 #if !INLINE_FMEMORY_OPERATION && !PLATFORM_USES_FIXED_GMalloc_CLASS
 #	include "HAL/FMemory.inl"
diff --git a/Engine/Source/Runtime/Core/Public/HAL/UnrealMemory.h b/Engine/Source/Runtime/Core/Public/HAL/UnrealMemory.h
--- a/Engine/Source/Runtime/Core/Public/HAL/UnrealMemory.h
+++ b/Engine/Source/Runtime/Core/Public/HAL/UnrealMemory.h
@@ -272,30 +272,6 @@ struct FMemory
 	*/
 	static CORE_API void ExplicitInit(FMalloc& Allocator);
 
-	/**
-	* Functions to handle special memory given to the title from the platform
-	* This memory is allocated like a stack, it's never really freed
-	*/
-	UE_DEPRECATED(5.5, "Persistent Auxiliary allocator is obsolete and is replaced by a GetPersistentLinearAllocator()")
-	static inline void RegisterPersistentAuxiliary(void* /*InMemory*/, SIZE_T /*InSize*/) {}
-
-	UE_DEPRECATED(5.5, "Persistent Auxiliary allocator is obsolete and is replaced by a GetPersistentLinearAllocator()")
-	static CORE_API void* MallocPersistentAuxiliary(SIZE_T InSize, uint32 InAlignment = 0);
-
-	UE_DEPRECATED(5.5, "Persistent Auxiliary allocator is obsolete and is replaced by a GetPersistentLinearAllocator()")
-	static inline void FreePersistentAuxiliary(void* InPtr) {}
-
-	UE_DEPRECATED(5.5, "Persistent Auxiliary allocator is obsolete and is replaced by a GetPersistentLinearAllocator()")
-	static CORE_API bool IsPersistentAuxiliaryActive();
-
-	UE_DEPRECATED(5.5, "Persistent Auxiliary allocator is obsolete and is replaced by a GetPersistentLinearAllocator()")
-	static inline void DisablePersistentAuxiliary() {}
-
-	UE_DEPRECATED(5.5, "Persistent Auxiliary allocator is obsolete and is replaced by a GetPersistentLinearAllocator()")
-	static inline void EnablePersistentAuxiliary() {}
-
-	UE_DEPRECATED(5.5, "Persistent Auxiliary allocator is obsolete and is replaced by a GetPersistentLinearAllocator()")
-	static CORE_API SIZE_T GetUsedPersistentAuxiliary();
 private:
 	static CORE_API void GCreateMalloc();
 	// These versions are called either at startup or in the event of a crash
