Added TSortedSet and natvis.
Added tests for TSet and TSortedSet in general - forked from the existing map tests.

diff --git a/Engine/Source/Runtime/Core/Private/Containers/ContainersTest.cpp b/Engine/Source/Runtime/Core/Private/Containers/ContainersTest.cpp
--- a/Engine/Source/Runtime/Core/Private/Containers/ContainersTest.cpp
+++ b/Engine/Source/Runtime/Core/Private/Containers/ContainersTest.cpp
@@ -7,13 +7,14 @@
 #include "Containers/Map.h"
 #include "Containers/Set.h"
 #include "Containers/SortedMap.h"
+#include "Containers/SortedSet.h"
 #include "Math/RandomStream.h"
 #include "Misc/AssertionMacros.h"
 #include "Misc/AutomationTest.h"
 #include "Stats/StatsMisc.h"
 #include "Templates/UniquePtr.h"
 
-#if WITH_DEV_AUTOMATION_TESTS
+#if 1//WITH_DEV_AUTOMATION_TESTS
 
 #define MAX_TEST_OBJECTS      65
 #define MAX_TEST_OBJECTS_STEP 1
@@ -212,6 +213,20 @@ namespace
 		check(!CIter);
 	}
 
+	template <typename Container>
+	void CheckContainerSorted(Container& Cont)
+	{
+		using ElementType = typename Container::ElementType;
+
+		const ElementType* PrevElement = nullptr;
+		for (auto Iter = Cont.CreateConstIterator(); Iter; ++Iter)
+		{
+			const ElementType* CurrElement = &*Iter;
+			check(!PrevElement || *PrevElement < *CurrElement);
+			PrevElement = CurrElement;
+		}
+	}
+
 	template <typename KeyType>
 	KeyType GenerateTestKey(int32 Input)
 	{
@@ -382,7 +397,7 @@ namespace
 		}
 	}
 
-	template <typename SetType, typename KeyType>
+	template <typename SetType, typename KeyType, bool bSorted = false>
 	void RunSetTests()
 	{
 		SetType Set;
@@ -399,6 +414,10 @@ namespace
 			CheckContainerEnds(Set);
 			CheckContainerElements(Set);
 			CheckContainerSelfEquality(Set);
+			if constexpr (bSorted)
+			{
+				CheckContainerSorted(Set);
+			}
 		};
 
 		// Test Add and Remove
@@ -717,6 +736,10 @@ bool FContainersSmokeTest::RunTest( const FString& Parameters )
 	RunSetTests<TSet<int32, DefaultKeyFuncs<int32>, TInlineSetAllocator<32>>, int32>();
 	RunSetTests<TSet<int32, DefaultKeyFuncs<int32>, TFixedSetAllocator<64>>, int32>();
 
+	RunSetTests<TSortedSet<int32>, int32, true>();
+	RunSetTests<TSortedSet<int32, TInlineAllocator<32>>, int32, true>();
+	RunSetTests<TSortedSet<int32, TFixedAllocator<64>>, int32, true>();
+
 	return true;
 }
 
@@ -746,6 +769,10 @@ bool FContainersFullTest::RunTest(const FString& Parameters)
 	RunSetTests<TSet<int32, DefaultKeyFuncs<int32>, TInlineSetAllocator<32>>, int32>();
 	RunSetTests<TSet<int32, DefaultKeyFuncs<int32>, TFixedSetAllocator<64>>, int32>();
 
+	RunSetTests<TSortedSet<int32>, int32, true>();
+	RunSetTests<TSortedSet<int32, TInlineAllocator<32>>, int32, true>();
+	RunSetTests<TSortedSet<int32, TFixedAllocator<64>>, int32, true>();
+
 	// Verify use of FName index sorter with SortedMap
 
 	TSortedMap<FName, int32, FDefaultAllocator, FNameFastLess> NameMap;
@@ -1457,6 +1484,99 @@ bool FContainersNatvisTest::RunTest(const FString& Parameters)
 		PlaceBreakpointHere();
 	}
 
+	// TSortedSet
+	{
+		// Name        | Value | Type
+		// ------------+-------+------------------------------------------------------
+		// > EmptySSet | Empty | TSortedSet<int,TSizedDefaultAllocator<32>,TLess<int>>
+		TSortedSet<int32> EmptySSet;
+
+		// Name        | Value         | Type
+		// ------------+---------------+------------------------------------------------------
+		// > SlackSSet | Empty, Max=16 | TSortedSet<int,TSizedDefaultAllocator<32>,TLess<int>>
+		TSortedSet<int32> SlackSSet;
+		SlackSSet.Reserve(16);
+
+		// Name       | Value        | Type
+		// -----------+--------------+------------------------------------------------------
+		// > SSetVals | Num=5, Max=5 | TSortedSet<int,TSizedDefaultAllocator<32>,TLess<int>>
+		//     [0]    | 1            | int
+		//     [1]    | 2            | int
+		//     [2]    | 3            | int
+		//     [3]    | 4            | int
+		//     [4]    | 5            | int
+		TSortedSet<int32> SSetVals = { 4, 1, 2, 5, 3 };
+
+		// Name           | Value         | Type
+		// ---------------+---------------+------------------------------------------------------
+		// > SlackSetVals | Num=5, Max=16 | TSortedSet<int,TSizedDefaultAllocator<32>,TLess<int>>
+		//     [0]        | 1             | int
+		//     [1]        | 2             | int
+		//     [2]        | 3             | int
+		//     [3]        | 4             | int
+		//     [4]        | 5             | int
+		TSortedSet<int32> SlackSSetVals = { 4, 1, 2, 5, 3 };
+		SlackSSetVals.Reserve(16);
+
+		// Name             | Value        | Type
+		// -----------------+--------------+------------------------------------------------------
+		// > SparseSSetVals | Num=3, Max=5 | TSortedSet<int,TSizedDefaultAllocator<32>,TLess<int>>
+		//     [0]          | 2            | int
+		//     [1]          | 3            | int
+		//     [2]          | 4            | int
+		TSortedSet<int32> SparseSSetVals = { 4, 1, 2, 5, 3 };
+		SparseSSetVals.Remove(1);
+		SparseSSetVals.Remove(5);
+
+		// Name       | Value        | Type
+		// -----------+--------------+--------------------------------------------------------------
+		// > SSetStrs | Num=5, Max=5 | TSortedSet<FString,TSizedDefaultAllocator<32>,TLess<FString>>
+		//   > [0]    | L"a2"        | FString
+		//       [0]  | 97 'a'       | wchar_t
+		//       [1]  | 50 '2'       | wchar_t
+		//   > [1]    | L"b3"        | FString
+		//       [0]  | 98 'b'       | wchar_t
+		//       [1]  | 51 '3'       | wchar_t
+		//   > [2]    | L"c5"        | FString
+		//       [0]  | 99 'c'       | wchar_t
+		//       [1]  | 53 '5'       | wchar_t
+		//   > [3]    | L"d1"        | FString
+		//       [0]  | 100 'd'      | wchar_t
+		//       [1]  | 49 '1'       | wchar_t
+		//   > [4]    | L"e4"        | FString
+		//       [0]  | 101 'e'      | wchar_t
+		//       [1]  | 52 '4'       | wchar_t
+		TSortedSet<FString> SSetStrs = { "d1", "a2", "b3", "e4", "c5" };
+
+		// Name             | Value                | Type
+		// -----------------+----------------------+--------------------------------------------------------------------------------------------------------------------------------------
+		// > SSetUniqueStrs | Num=5, Max=5         | TSortedSet<TUniquePtr<FString,TDefaultDelete<FString>>,TSizedDefaultAllocator<32>,TLess<TUniquePtr<FString,TDefaultDelete<FString>>>>
+		//   > [0]          | Ptr=0x01234567 L"a2" | TUniquePtr<FString,TDefaultDelete<FString>>
+		//       [0]        | 97 'a'               | wchar_t
+		//       [1]        | 50 '2'               | wchar_t
+		//   > [1]          | Ptr=0x01234567 L"b3" | TUniquePtr<FString,TDefaultDelete<FString>>
+		//       [0]        | 98 'b'               | wchar_t
+		//       [1]        | 51 '3'               | wchar_t
+		//   > [2]          | Ptr=0x01234567 L"c5" | TUniquePtr<FString,TDefaultDelete<FString>>
+		//       [0]        | 99 'c'               | wchar_t
+		//       [1]        | 53 '5'               | wchar_t
+		//   > [3]          | Ptr=0x01234567 L"d1" | TUniquePtr<FString,TDefaultDelete<FString>>
+		//       [0]        | 100 'd'              | wchar_t
+		//       [1]        | 49 '1'               | wchar_t
+		//   > [4]          | Ptr=0x01234567 L"e4" | TUniquePtr<FString,TDefaultDelete<FString>>
+		//       [0]        | 101 'e'              | wchar_t
+		//       [1]        | 52 '4'               | wchar_t
+		TSortedSet<TUniquePtr<FString>> SSetUniqueStrs;
+		SSetUniqueStrs.Reserve(5);
+		SSetUniqueStrs.Add(MakeUnique<FString>("d1"));
+		SSetUniqueStrs.Add(MakeUnique<FString>("a2"));
+		SSetUniqueStrs.Add(MakeUnique<FString>("b3"));
+		SSetUniqueStrs.Add(MakeUnique<FString>("e4"));
+		SSetUniqueStrs.Add(MakeUnique<FString>("c5"));
+
+		PlaceBreakpointHere();
+	}
+
 	// TMap
 	{
 		// Name       | Value | Type
diff --git a/Engine/Source/Runtime/Core/Public/Containers/ContainersFwd.h b/Engine/Source/Runtime/Core/Public/Containers/ContainersFwd.h
--- a/Engine/Source/Runtime/Core/Public/Containers/ContainersFwd.h
+++ b/Engine/Source/Runtime/Core/Public/Containers/ContainersFwd.h
@@ -70,6 +70,7 @@ template<typename KeyType, typename ValueType, typename SetAllocator = FDefaultS
 template <typename T = void > struct TLess;
 template <typename> struct TTypeTraits;
 template<typename InKeyType, typename InValueType, typename ArrayAllocator = FDefaultAllocator, typename SortPredicate = TLess<typename TTypeTraits<InKeyType>::ConstPointerType> > class TSortedMap;
+template <typename InElementType, typename ArrayAllocator = FDefaultAllocator, typename SortPredicate = TLess<InElementType>> class TSortedSet;
 template<typename ElementType,bool bInAllowDuplicateKeys = false> struct DefaultKeyFuncs;
 template<typename InElementType, typename KeyFuncs = DefaultKeyFuncs<InElementType>, typename Allocator = FDefaultSetAllocator> class TSet;
 template<typename InElementType, typename InSizeType = int32> class TStridedView;
diff --git a/Engine/Source/Runtime/Core/Public/Containers/SortedSet.h b/Engine/Source/Runtime/Core/Public/Containers/SortedSet.h
new file mode 100644
--- /dev/null
+++ b/Engine/Source/Runtime/Core/Public/Containers/SortedSet.h
@@ -0,0 +1,603 @@
+// Copyright Epic Games, Inc. All Rights Reserved.
+
+#pragma once
+
+#include "Algo/BinarySearch.h"
+#include "Algo/Sort.h"
+#include "Containers/SetUtilities.h"
+#include "UObject/NameTypes.h"
+
+template <typename ElementType, typename ArrayAllocator, typename SortPredicate>
+FArchive& operator<<(FArchive& Ar, TSortedSet<ElementType, ArrayAllocator, SortPredicate>& Set);
+
+/** 
+ * A set of values, implemented as a sorted TArray of elements.
+ *
+ * It has a mostly identical interface to TSet and is designed as a drop in replacement. Keys must be unique.
+ * It uses half as much memory as TSet, but adding and removing elements is O(n), and finding is O(Log n). In practice
+ * it is faster than TSet for low element counts, and slower as n increases, This set is always kept sorted by key
+ * so cannot be sorted manually.
+ */
+template <typename InElementType, typename ArrayAllocator /*= FDefaultAllocator*/, typename SortPredicate /*= TLess<InElementType>*/>
+class TSortedSet
+{
+	template <typename OtherElementType, typename OtherArrayAllocator, typename OtherSortPredicate>
+	friend class TSortedSet;
+
+	friend FArchive& operator<< <InElementType, ArrayAllocator, SortPredicate>(FArchive& Ar, TSortedSet& Set);
+
+public:
+	using ElementType = InElementType;
+	using KeyInitType = typename TTypeTraits<ElementType>::ConstPointerType;
+
+	[[nodiscard]] TSortedSet() = default;
+	[[nodiscard]] TSortedSet(TSortedSet&&) = default;
+	[[nodiscard]] TSortedSet(const TSortedSet&) = default;
+	TSortedSet& operator=(TSortedSet&&) = default;
+	TSortedSet& operator=(const TSortedSet&) = default;
+
+	/** Constructor for moving elements from a TSortedSet with a different ArrayAllocator. */
+	template<typename OtherArrayAllocator>
+	[[nodiscard]] TSortedSet(TSortedSet<InElementType, OtherArrayAllocator, SortPredicate>&& Other)
+		: Elements(MoveTemp(Other.Elements))
+	{
+	}
+
+	/** Constructor for copying elements from a TSortedSet with a different ArrayAllocator. */
+	template<typename OtherArrayAllocator>
+	[[nodiscard]] TSortedSet(const TSortedSet<InElementType, OtherArrayAllocator, SortPredicate>& Other)
+		: Elements(Other.Elements)
+	{
+	}
+
+	/** Constructor which gets its elements from a native initializer list */
+	[[nodiscard]] TSortedSet(std::initializer_list<ElementType> InitList)
+	{
+		this->Reserve((int32)InitList.size());
+		for (const ElementType& Element : InitList)
+		{
+			this->Add(Element);
+		}
+	}
+
+	///////////////////////////////////////////////////
+	// Start - intrusive TOptional<TSortedSet> state //
+	///////////////////////////////////////////////////
+	constexpr static bool bHasIntrusiveUnsetOptionalState = true;
+	using IntrusiveUnsetOptionalStateType = TSortedSet;
+
+	[[nodiscard]] explicit TSortedSet(FIntrusiveUnsetOptionalState Tag)
+		: Elements(Tag)
+	{
+	}
+	[[nodiscard]] bool operator==(FIntrusiveUnsetOptionalState Tag) const
+	{
+		return Elements == Tag;
+	}
+	/////////////////////////////////////////////////
+	// End - intrusive TOptional<TSortedSet> state //
+	/////////////////////////////////////////////////
+
+	/** Assignment operator for moving elements from a TSortedSet with a different ArrayAllocator. */
+	template<typename OtherArrayAllocator>
+	TSortedSet& operator=(TSortedSet<ElementType, OtherArrayAllocator, SortPredicate>&& Other)
+	{
+		Elements = MoveTemp(Other.Elements);
+		return *this;
+	}
+
+	/** Assignment operator for copying elements from a TSortedSet with a different ArrayAllocator. */
+	template<typename OtherArrayAllocator>
+	TSortedSet& operator=(const TSortedSet<ElementType, OtherArrayAllocator, SortPredicate>& Other)
+	{
+		Elements = Other.Elements;
+		return *this;
+	}
+
+	/** Assignment operator which gets its elements from a native initializer list */
+	TSortedSet& operator=(std::initializer_list<ElementType> InitList)
+	{
+		this->Empty((int32)InitList.size());
+		for (const ElementType& Element : InitList)
+		{
+			this->Add(Element);
+		}
+		return *this;
+	}
+
+	/** Equality operator. This is efficient because elements are always sorted. */
+	FORCEINLINE bool operator==(const TSortedSet& Other) const
+	{
+		return Elements == Other.Elements;
+	}
+
+	/** Inequality operator. This is efficient because elements are always sorted. */
+	FORCEINLINE bool operator!=(const TSortedSet& Other) const
+	{
+		return Elements != Other.Elements;
+	}
+
+	/**
+	 * Removes all elements from the set, potentially leaving space allocated for an expected number of elements about to be added.
+	 *
+	 * @param ExpectedNumElements The number of elements about to be added to the set.
+	 */
+	FORCEINLINE void Empty(int32 ExpectedNumElements = 0)
+	{
+		Elements.Empty(ExpectedNumElements);
+	}
+
+	/** Efficiently empties out the set but preserves all allocations and capacities. */
+	FORCEINLINE void Reset()
+	{
+		Elements.Reset();
+	}
+
+	/** Shrinks the set to avoid slack. */
+	FORCEINLINE void Shrink()
+	{
+		Elements.Shrink();
+	}
+
+	/** Preallocates enough memory to contain Number elements. */
+	FORCEINLINE void Reserve(int32 Number)
+	{
+		Elements.Reserve(Number);
+	}
+
+	/**
+	 * Returns true if the set is empty and contains no elements. 
+	 *
+	 * @returns True if the set is empty.
+	 * @see Num
+	 */
+	[[nodiscard]] bool IsEmpty() const
+	{
+		return Elements.IsEmpty();
+	}
+
+	/** @return The number of elements in the set. */
+	[[nodiscard]] FORCEINLINE int32 Num() const
+	{
+		return Elements.Num();
+	}
+
+	/** @return The max valid index of the elements in the inner array. */
+	[[nodiscard]] FORCEINLINE int32 GetMaxIndex() const
+	{
+		return Elements.Num() - 1;
+	}
+
+	/** 
+	 * Helper function to return the amount of memory allocated by this container.
+	 * Only returns the size of allocations made directly by the container, not the elements themselves.
+	 *
+	 * @return number of bytes allocated by this container.
+	 */
+	[[nodiscard]] FORCEINLINE SIZE_T GetAllocatedSize() const
+	{
+		return Elements.GetAllocatedSize();
+	}
+
+	/** Tracks the container's memory use through an archive. */
+	FORCEINLINE void CountBytes(FArchive& Ar) const
+	{
+		Elements.CountBytes(Ar);
+	}
+
+	/**
+	 * Adds an element to the set.
+	 *
+	 * @param	InValue	The value to add.
+	 * @param	bIsAlreadyInSetPtr	[out]	Optional pointer to bool that will be set depending on whether element is already in set
+	 * @return A handle to the value in the set (only valid until the next change to the set).
+	 */
+	FORCEINLINE FSetElementId Add(const ElementType& InValue, bool* bIsAlreadyInSetPtr = nullptr)
+	{
+		return Emplace(InValue, bIsAlreadyInSetPtr);
+	}
+	FORCEINLINE FSetElementId Add(ElementType&& InValue, bool* bIsAlreadyInSetPtr = nullptr)
+	{
+		return Emplace(MoveTemp(InValue), bIsAlreadyInSetPtr);
+	}
+
+	/**
+	 * Adds an element to the set.
+	 *
+	 * @param	InInitArg	The argument to be forwarded to the set element's constructor.
+	 * @param	bIsAlreadyInSetPtr	[out]	Optional pointer to bool that will be set depending on whether element is already in set
+	 * @return A handle to the value in the set (only valid until the next change to the set).
+	 */
+	template <typename InitArgType = ElementType>
+	FSetElementId Emplace(InitArgType&& InInitArg, bool* bIsAlreadyInSetPtr = nullptr)
+	{
+		ElementType* DataPtr = AllocateMemoryForEmplace(InInitArg, bIsAlreadyInSetPtr);
+	
+		::new((void*)DataPtr) ElementType(Forward<InitArgType>(InInitArg));
+
+		return FSetElementId::FromInteger((int32)(DataPtr - Elements.GetData()));
+	}
+
+	/**
+	 * Removes all elements from the set matching the specified key.
+	 *
+	 * @param	InKey	The key to match elements against.
+	 * @return	The number of elements removed.
+	 */
+	FORCEINLINE int32 Remove(KeyInitType InKey)
+	{
+		int32 RemoveIndex = FindIndex(InKey);
+		if (RemoveIndex == INDEX_NONE)
+		{
+			return 0;
+		}
+
+		Elements.RemoveAt(RemoveIndex);
+		return 1;
+	}
+
+	/**
+	 * Returns the value associated with a specified key.
+	 *
+	 * @param Key The key to search for.
+	 * @return A pointer to the value associated with the specified key, or nullptr if the key isn't contained in this set.  The pointer (only valid until the next change to any key in the set).
+	 */
+	[[nodiscard]] FORCEINLINE ElementType* Find(KeyInitType Key)
+	{
+		int32 FoundIndex = FindIndex(Key);
+		if (FoundIndex != INDEX_NONE)
+		{
+			return &Elements[FoundIndex];
+		}
+
+		return nullptr;
+	}
+	[[nodiscard]] FORCEINLINE const ElementType* Find(KeyInitType Key) const
+	{
+		return const_cast<TSortedSet*>(this)->Find(Key);
+	}
+
+	/**
+	 * Adds an element to the set if not already present and returns a reference to the added or existing element.
+	 *
+	 * @param	InValue	The value to add.
+	 * @param	bIsAlreadyInSetPtr	[out]	Optional pointer to bool that will be set depending on whether element is already in set
+	 * @return A handle to the value in the set (only valid until the next change to the set).
+	 */
+	FORCEINLINE ElementType& FindOrAdd(const ElementType& InValue, bool* bIsAlreadyInSetPtr = nullptr)
+	{
+		return FindOrAddImpl(InValue, bIsAlreadyInSetPtr);
+	}
+	FORCEINLINE ElementType& FindOrAdd(ElementType&& InValue, bool* bIsAlreadyInSetPtr = nullptr)
+	{
+		return FindOrAddImpl(MoveTemp(InValue), bIsAlreadyInSetPtr);
+	}
+
+	/**
+	 * Finds any element in the set and returns a pointer to it.
+	 * Callers should not depend on particular patterns in the behavior of this function.
+	 * @return A pointer to an arbitrary element, or nullptr if the container is empty.
+	 */
+	[[nodiscard]] ElementType* FindArbitraryElement()
+	{
+		// The goal of this function is to be fast, and so the implementation may be improved at any time even if it gives different results.
+
+		return Elements.IsEmpty() ? nullptr : &Elements.Last();
+	}
+	[[nodiscard]] const ElementType* FindArbitraryElement() const
+	{
+		return const_cast<TSortedSet*>(this)->FindArbitraryElement();
+	}
+
+	/**
+	 * Checks if set contains the specified key.
+	 *
+	 * @param Key The key to check for.
+	 * @return true if the set contains the key.
+	 */
+	[[nodiscard]] FORCEINLINE bool Contains(KeyInitType Key) const
+	{
+		return FindIndex(Key) != INDEX_NONE;
+	}
+
+	/**
+	 * Returns a TArray of the elements.
+	 */
+	TArray<ElementType> Array() const&
+	{
+		return Elements;
+	}
+	TArray<ElementType> Array() &&
+	{
+		return MoveTemp(Elements);
+	}
+
+	/**
+	 * Describes the set's contents through an output device.
+	 *
+	 * @param Ar The output device to describe the set's contents through.
+	 */
+	void Dump(FOutputDevice& Ar)
+	{
+		Ar.Logf(TEXT("TSortedSet: %i elements"), Elements.Num());
+	}
+
+	/**
+	 * Adds all items from another set into our set (if any elements are in both, the value from the other set wins).
+	 * If the source is an rvalue container, the elements will be moved and the container will be left empty.
+	 *
+	 * @param Other The range of elements to add.
+	 */
+	template <typename OtherArrayAllocator, typename OtherSortPredicate>
+	void Append(TSortedSet<ElementType, OtherArrayAllocator, OtherSortPredicate>&& Other)
+	{
+		this->Reserve(this->Num() + Other.Num());
+		for (ElementType& Element : Other)
+		{
+			this->Add(MoveTemp(Element));
+		}
+
+		Other.Reset();
+	}
+	template <typename OtherArrayAllocator, typename OtherSortPredicate>
+	void Append(const TSortedSet<ElementType, OtherArrayAllocator, OtherSortPredicate>& Other)
+	{
+		this->Reserve(this->Num() + Other.Num());
+		for (const ElementType& Element : Other)
+		{
+			this->Add(Element);
+		}
+	}
+	template <typename OtherArrayAllocator>
+	void Append(TArray<ElementType, OtherArrayAllocator>&& Other)
+	{
+		this->Reserve(this->Num() + Other.Num());
+		for (ElementType& Element : Other)
+		{
+			this->Add(MoveTemp(Element));
+		}
+		Other.Reset();
+	}
+	void Append(TArrayView<const ElementType> Other)
+	{
+		this->Reserve(this->Num() + Other.Num());
+		for (const ElementType& Element : Other)
+		{
+			this->Add(Element);
+		}
+	}
+	template <typename OtherArrayAllocator>
+	void Append(std::initializer_list<ElementType> Other)
+	{
+		this->Reserve(this->Num() + (int32)Other.size());
+		for (const ElementType& Element : Other)
+		{
+			this->Add(Element);
+		}
+	}
+
+	/**
+	 * Checks whether an element id is valid.
+	 * @param Id - The element id to check.
+	 * @return true if the element identifier refers to a valid element in this set.
+	 */
+	[[nodiscard]] FORCEINLINE bool IsValidId(FSetElementId Id) const
+	{
+		return Elements.IsValidIndex(Id.AsInteger());
+	}
+
+	[[nodiscard]] FORCEINLINE ElementType& operator[](FSetElementId Id)
+	{
+		return Elements[Id.AsInteger()];
+	}
+	[[nodiscard]] FORCEINLINE const ElementType& operator[](FSetElementId Id) const
+	{
+		return (*const_cast<TSortedSet*>(this))[Id];
+	}
+	[[nodiscard]] FORCEINLINE ElementType& Get(FSetElementId Id)
+	{
+		return (*this)[Id];
+	}
+	[[nodiscard]] FORCEINLINE const ElementType& Get(FSetElementId Id) const
+	{
+		return (*this)[Id];
+	}
+
+private:
+	using ElementArrayType = TArray<ElementType, ArrayAllocator>;
+
+	/** Implementation of find and add */
+	template <typename InitArgType>
+	[[nodiscard]] FORCEINLINE ElementType& FindOrAddImpl(InitArgType&& InInitArg, bool* bIsAlreadyInSetPtr)
+	{
+		int32 InsertIndex = Algo::LowerBound(Elements, InInitArg, SortPredicate());
+
+		// Since we returned lower bound we already know InInitArg <= InsertIndex key. So if InInitArg is not < InsertIndex key, they must be equal
+		bool bAlreadyExists = InsertIndex < Elements.Num() && !SortPredicate()(InInitArg, Elements[InsertIndex]);
+
+		ElementType* DataPtr = nullptr;
+		if (!bAlreadyExists)
+		{
+			Elements.EmplaceAt(InsertIndex, Forward<InitArgType>(InInitArg));
+		}
+
+		if (bIsAlreadyInSetPtr)
+		{
+			*bIsAlreadyInSetPtr = bAlreadyExists;
+		}
+
+		return *(Elements.GetData() + InsertIndex);
+	}
+
+	/** Find index of key */
+	[[nodiscard]] FORCEINLINE int32 FindIndex(KeyInitType Key) const
+	{
+		return Algo::BinarySearch(Elements, Key, SortPredicate());
+	}
+
+	/** Allocates raw memory for emplacing */
+	template <typename InitArgType>
+	FORCEINLINE ElementType* AllocateMemoryForEmplace(InitArgType&& InInitArg, bool* bIsAlreadyInSetPtr = nullptr)
+	{
+		int32 InsertIndex = Algo::LowerBound(Elements, InInitArg, SortPredicate());
+
+		// Since we returned lower bound we already know InInitArg <= InsertIndex key. So if InInitArg is not < InsertIndex key, they must be equal
+		bool bAlreadyExists = InsertIndex < Elements.Num() && !SortPredicate()(InInitArg, Elements[InsertIndex]);
+
+		ElementType* DataPtr = nullptr;
+		if (bAlreadyExists)
+		{
+			// Replacing element, delete old one
+			DataPtr = Elements.GetData() + InsertIndex;
+			DestructItem(DataPtr);
+		}
+		else
+		{
+			// Adding new one, this may reallocate Elements
+			Elements.InsertUninitialized(InsertIndex);
+			DataPtr = Elements.GetData() + InsertIndex;
+		}
+
+		if (bIsAlreadyInSetPtr)
+		{
+			*bIsAlreadyInSetPtr = bAlreadyExists;
+		}
+
+		return DataPtr;
+	}
+
+	/** The base of TSortedSet iterators */
+	template<bool bConst>
+	class TBaseIterator
+	{
+	private:
+		using ArrayIteratorType = std::conditional_t<bConst, typename ElementArrayType::TConstIterator, typename ElementArrayType::TIterator>;
+		using ItElementType = std::conditional_t<bConst, const ElementType, ElementType>;
+
+	protected:
+		[[nodiscard]] FORCEINLINE TBaseIterator(const ArrayIteratorType& InArrayIt)
+			: ArrayIt(InArrayIt)
+		{
+		}
+
+	public:
+		FORCEINLINE TBaseIterator& operator++()
+		{
+			++ArrayIt;
+			return *this;
+		}
+
+		/** conversion to "bool" returning true if the iterator is valid. */
+		[[nodiscard]] FORCEINLINE explicit operator bool() const
+		{
+			return !!ArrayIt; 
+		}
+
+		[[nodiscard]] FORCEINLINE bool operator==(const TBaseIterator& Rhs) const
+		{
+			return ArrayIt == Rhs.ArrayIt;
+		}
+		[[nodiscard]] FORCEINLINE bool operator!=(const TBaseIterator& Rhs) const
+		{
+			return !(*this == Rhs);
+		}
+
+		[[nodiscard]] FORCEINLINE FSetElementId GetId() const
+		{
+			return FSetElementId::FromInteger(ArrayIt.GetIndex());
+		}
+
+		FORCEINLINE ItElementType& operator*() const
+		{
+			return *ArrayIt;
+		}
+		FORCEINLINE ItElementType* operator->() const
+		{
+			return &*ArrayIt;
+		}
+
+	protected:
+		ArrayIteratorType ArrayIt;
+	};
+
+	/** An array of the elements in the set */
+	ElementArrayType Elements;
+
+public:
+	/** Set iterator */
+	class TIterator : public TBaseIterator<false>
+	{
+	public:
+		[[nodiscard]] FORCEINLINE explicit TIterator(TSortedSet& InSet)
+			: TBaseIterator<false>(InSet.Elements.CreateIterator())
+		{
+		}
+
+		/** Removes the current element from the set */
+		FORCEINLINE void RemoveCurrent()
+		{
+			TBaseIterator<false>::ArrayIt.RemoveCurrent();
+		}
+	};
+
+	/** Const set iterator */
+	class TConstIterator : public TBaseIterator<true>
+	{
+	public:
+		[[nodiscard]] FORCEINLINE explicit TConstIterator(const TSortedSet& InSet)
+			: TBaseIterator<true>(InSet.Elements.CreateConstIterator())
+		{
+		}
+	};
+
+	/** Creates an iterator over all the elements in this set */
+	[[nodiscard]] FORCEINLINE TIterator CreateIterator()
+	{
+		return TIterator(*this);
+	}
+
+	/** Creates a const iterator over all the elements in this set */
+	[[nodiscard]] FORCEINLINE TConstIterator CreateConstIterator() const
+	{
+		return TConstIterator(*this);
+	}
+
+public:
+	/**
+	 * DO NOT USE DIRECTLY
+	 * STL-like iterators to enable range-based for loop support.
+	 */
+	[[nodiscard]] FORCEINLINE auto begin()
+	{
+		return Elements.begin();
+	}
+	[[nodiscard]] FORCEINLINE auto begin() const
+	{
+		return Elements.begin();
+	}
+	[[nodiscard]] FORCEINLINE auto end()
+	{
+		return Elements.end();
+	}
+	[[nodiscard]] FORCEINLINE auto end() const
+	{
+		return Elements.end();
+	}
+};
+
+DECLARE_TEMPLATE_INTRINSIC_TYPE_LAYOUT((template <typename ElementType, typename ArrayAllocator, typename SortPredicate>), (TSortedSet<ElementType, ArrayAllocator, SortPredicate>));
+
+template <typename ElementType, typename ArrayAllocator, typename SortPredicate>
+FORCEINLINE FArchive& operator<<(FArchive& Ar, TSortedSet<ElementType, ArrayAllocator, SortPredicate>& Set)
+{
+	Ar << Set.Elements;
+
+	if (Ar.IsLoading())
+	{
+		// We need to re-sort, in case the predicate is not consistent with what it was before
+		Algo::Sort(Set.Elements, SortPredicate());
+	}
+
+	return Ar;
+}
+
