ParkingLot: Call Wait/WaitFor/WaitUntil to park threads based on how ParkingLot was invoked

#rnx

diff --git a/Engine/Source/Runtime/Core/Private/Async/ParkingLot.cpp b/Engine/Source/Runtime/Core/Private/Async/ParkingLot.cpp
--- a/Engine/Source/Runtime/Core/Private/Async/ParkingLot.cpp
+++ b/Engine/Source/Runtime/Core/Private/Async/ParkingLot.cpp
@@ -500,23 +500,17 @@ FThreadLocalData::~FThreadLocalData()
 
 ///////////////////////////////////////////////////////////////////////////////////////////////////////////////////////
 
-FWaitState Wait(const void* Address, bool (*CanWait)(void*), void* CanWaitContext, void (*BeforeWait)(void*), void* BeforeWaitContext)
-{
-	return WaitUntil(Address, CanWait, CanWaitContext, BeforeWait, BeforeWaitContext, FMonotonicTimePoint::Infinity());
-}
-
-FWaitState WaitFor(const void* Address, bool (*CanWait)(void*), void* CanWaitContext, void (*BeforeWait)(void*), void* BeforeWaitContext, FMonotonicTimeSpan WaitTime)
-{
-	return WaitUntil(Address, CanWait, CanWaitContext, BeforeWait, BeforeWaitContext, FMonotonicTimePoint::Now() + WaitTime);
-}
-
-UE_AUTORTFM_NOAUTORTFM
-FWaitState WaitUntil(const void* Address, bool (*CanWait)(void*), void* CanWaitContext, void (*BeforeWait)(void*), void* BeforeWaitContext, FMonotonicTimePoint WaitTime)
+FWaitState Wait(
+	const void* Address,
+	bool (*CanWait)(void*),
+	void* CanWaitContext,
+	void (*BeforeWait)(void*),
+	void* BeforeWaitContext,
+	void (*WaitOnEvent)(FPlatformManualResetEvent&, void*),
+	void* WaitOnEventContext)
 {
 	using namespace Private;
 
-	check(!WaitTime.IsNaN());
-
 	TRefCountPtr<FThread> Self = FThreadLocalData::Get();
 
 	checkfSlow(!Self->WaitAddress.load(std::memory_order_relaxed), TEXT("WaitAddress must be null. This can happen if Wait is called by BeforeWait."));
@@ -545,7 +539,7 @@ FWaitState WaitUntil(const void* Address, bool (*CanWait)(void*), void* CanWaitC
 	}
 
 	// Wait until the timeout or until the thread has been dequeued.
-	Self->Event.WaitUntil(WaitTime);
+	WaitOnEvent(Self->Event, WaitOnEventContext);
 
 	// WaitAddress is reset when the thread is dequeued.
 	if (!Self->WaitAddress.load(std::memory_order_relaxed))
@@ -586,6 +580,29 @@ FWaitState WaitUntil(const void* Address, bool (*CanWait)(void*), void* CanWaitC
 	return State;
 }
 
+UE_AUTORTFM_NOAUTORTFM
+FWaitState Wait(const void* Address, bool (*CanWait)(void*), void* CanWaitContext, void (*BeforeWait)(void*), void* BeforeWaitContext)
+{
+	return Wait(Address, CanWait, CanWaitContext, BeforeWait, BeforeWaitContext,
+		[](FPlatformManualResetEvent& Event, void*) { Event.Wait(); }, nullptr);
+}
+
+UE_AUTORTFM_NOAUTORTFM
+FWaitState WaitFor(const void* Address, bool (*CanWait)(void*), void* CanWaitContext, void (*BeforeWait)(void*), void* BeforeWaitContext, FMonotonicTimeSpan WaitTime)
+{
+	check(!WaitTime.IsNaN());
+	return Wait(Address, CanWait, CanWaitContext, BeforeWait, BeforeWaitContext,
+		[](FPlatformManualResetEvent& Event, void* Context) { Event.WaitFor(*(FMonotonicTimeSpan*)Context); }, &WaitTime);
+}
+
+UE_AUTORTFM_NOAUTORTFM
+FWaitState WaitUntil(const void* Address, bool (*CanWait)(void*), void* CanWaitContext, void (*BeforeWait)(void*), void* BeforeWaitContext, FMonotonicTimePoint WaitTime)
+{
+	check(!WaitTime.IsNaN());
+	return Wait(Address, CanWait, CanWaitContext, BeforeWait, BeforeWaitContext,
+		[](FPlatformManualResetEvent& Event, void* Context) { Event.WaitUntil(*(FMonotonicTimePoint*)Context); }, &WaitTime);
+}
+
 UE_AUTORTFM_NOAUTORTFM
 void WakeOne(const void* Address, uint64 (*OnWakeState)(void*, FWakeState), void* OnWakeStateContext)
 {
diff --git a/Engine/Source/Runtime/Core/Tests/Async/ParkingLotTest.cpp b/Engine/Source/Runtime/Core/Tests/Async/ParkingLotTest.cpp
--- a/Engine/Source/Runtime/Core/Tests/Async/ParkingLotTest.cpp
+++ b/Engine/Source/Runtime/Core/Tests/Async/ParkingLotTest.cpp
@@ -110,6 +110,22 @@ TEST_CASE("Core::Async::ParkingLot", "[Core][Async]")
 		CHECK(State.WakeToken == 0);
 	}
 
+	SECTION("WaitUntil with Timeout")
+	{
+		int32 Value = 0;
+		int32 CanWaitCount = 0;
+		int32 BeforeWaitCount = 0;
+		ParkingLot::FWaitState State = ParkingLot::WaitUntil(&Value,
+			[&CanWaitCount] { ++CanWaitCount; return true; },
+			[&BeforeWaitCount] { ++BeforeWaitCount; },
+			FMonotonicTimePoint::Now() + FMonotonicTimeSpan::FromMilliseconds(1.0));
+		CHECK(CanWaitCount == 1);
+		CHECK(BeforeWaitCount == 1);
+		CHECK(State.bDidWait);
+		CHECK_FALSE(State.bDidWake);
+		CHECK(State.WakeToken == 0);
+	}
+
 	SECTION("WaitFor in ParallelFor")
 	{
 		// This exercises the code that grows the parking lot hash table.
