Fix up CPU Utilization on Unix
- Get the usage from the active process instead of the global load
- This seems to line up roughly with Windows is showing on Horde (~20%)

diff --git a/Engine/Source/Runtime/Core/Private/Unix/UnixPlatformProcess.cpp b/Engine/Source/Runtime/Core/Private/Unix/UnixPlatformProcess.cpp
--- a/Engine/Source/Runtime/Core/Private/Unix/UnixPlatformProcess.cpp
+++ b/Engine/Source/Runtime/Core/Private/Unix/UnixPlatformProcess.cpp
@@ -34,9 +34,8 @@
 #include "HAL/PlatformTime.h"
 #include "Misc/DateTime.h"
 
-
-
 DECLARE_LOG_CATEGORY_CLASS(LogFork, Log, All);
+DECLARE_LOG_CATEGORY_CLASS(LogUtilization, Verbose, All);
 
 namespace PlatformProcessLimits
 {
@@ -1799,94 +1798,142 @@ FString FUnixPlatformProcess::GetCurrentWorkingDirectory()
 	return UTF8_TO_TCHAR(CurrentDir);
 }
 
-static UE::Profiling::FCPUState CurrentCPUState;
-
 bool FUnixPlatformProcess::GetPerFrameProcessorUsage(uint32 ProcessId, float& ProcessUsageFraction, float& IdleUsageFraction)
 {
- 	(void)(ProcessId);
-	int64_t UserTime, NiceTime, SystemTime, SoftIRQTime, IRQTime, IdleTime, WallTime, IOWaitTime;
-	int32		Index = 0;
-	ANSICHAR	Buffer[500];
+	static bool bProcStatPrintHasHappened = false;
+	static bool bProcStatParsePrintHappened = false;
+	static bool bProcPIDPrintHappened = false;
+	static bool bProcPIDParsePrintHappened = false;
+	static uint64 PreviousTotalTicks = 0;
+	static uint64 PreviousUserTimeTicks = 0;
+	static uint64 PreviousSystemTimeTicks = 0;
+	static uint64 PreviousChildUserTimeTicks = 0;
+	static uint64 PreviousChildSystemTimeTicks = 0;
+
+	uint64_t UserTime, NiceTime, SystemTime, SoftIRQTime, IRQTime, IdleTime, IOWaitTime;
+	char CPUName[8];
+	char PIDString[20];
+	char Buffer[500];
+	
+	FILE* ProcStatFile = fopen("/proc/stat", "r");
 
-	CurrentCPUState.CoreCount = FMath::Min( 1, UE::Profiling::FCPUState::MaxSupportedCores);
-	FILE* FileHandle = fopen("/proc/stat", "r");
-	if (FileHandle)
+	if (ProcStatFile == nullptr)
 	{
-		CurrentCPUState.ActivatedCoreCount = 0;
-		for (size_t n = 0; n < CurrentCPUState.CoreCount; n++)
-		{
-			CurrentCPUState.Status[n] = 0;
-			CurrentCPUState.PreviousUsage[n] = CurrentCPUState.CurrentUsage[n];
-		}
+		UE_CLOG(bProcStatPrintHasHappened, LogUtilization, Verbose, TEXT("Failed to open /proc/stat."));
+		bProcStatPrintHasHappened = true;
+		return false;
+	}
+
+	// First line of the file is the total CPU usage so no need to tally everything together
+	if (fgets(Buffer, sizeof(Buffer), ProcStatFile) == nullptr)
+	{
+		UE_CLOG(bProcStatParsePrintHappened, LogUtilization, Verbose, TEXT("Failed to parse the first line of /proc/stat."));
+		bProcStatParsePrintHappened = true;
+		fclose(ProcStatFile);
+		return false;
+	}
+
+	fclose(ProcStatFile);
 
-		while (fgets(Buffer, 100, FileHandle)) {
 #if PLATFORM_64BITS
-			sscanf(Buffer, "%5s %8lu %8lu %8lu %8lu %8lu %8lu %8lu", CurrentCPUState.Name,
-				&UserTime, &NiceTime, &SystemTime, &IdleTime, &IOWaitTime, &IRQTime, &SoftIRQTime);
+	sscanf(Buffer, "%5s %8lu %8lu %8lu %8lu %8lu %8lu %8lu", CPUName,
+		&UserTime, &NiceTime, &SystemTime, &IdleTime, &IOWaitTime, &IRQTime, &SoftIRQTime);
 #else
-			sscanf(Buffer, "%5s %8llu %8llu %8llu %8llu %8llu %8llu %8llu", CurrentCPUState.Name,
-				&UserTime, &NiceTime, &SystemTime, &IdleTime, &IOWaitTime, &IRQTime, &SoftIRQTime);
+	sscanf(Buffer, "%5s %8llu %8llu %8llu %8llu %8llu %8llu %8llu", CurrentCPUState.Name,
+		&UserTime, &NiceTime, &SystemTime, &IdleTime, &IOWaitTime, &IRQTime, &SoftIRQTime);
 #endif
 
-			if (0 == strncmp(CurrentCPUState.Name, "cpu", 3))
-			{
-				Index = CurrentCPUState.Name[3] - '0';
-				if (Index >= 0 && Index < CurrentCPUState.CoreCount)
-				{
-					if(CurrentCPUState.Name[5] != '\0')
-					{
-						Index = atol(&CurrentCPUState.Name[3]);
-					}
-					CurrentCPUState.CurrentUsage[Index].IdleTime = IdleTime;
-					CurrentCPUState.CurrentUsage[Index].NiceTime = NiceTime;
-					CurrentCPUState.CurrentUsage[Index].SystemTime = SystemTime;
-					CurrentCPUState.CurrentUsage[Index].SoftIRQTime = SoftIRQTime;
-					CurrentCPUState.CurrentUsage[Index].IRQTime = IRQTime;
-					CurrentCPUState.CurrentUsage[Index].IOWaitTime = IOWaitTime;
-					CurrentCPUState.CurrentUsage[Index].UserTime = UserTime;
-					CurrentCPUState.CurrentUsage[Index].TotalTime = UserTime + NiceTime + SystemTime + SoftIRQTime + IRQTime + IdleTime + IOWaitTime;
-					CurrentCPUState.Status[Index] = 1;
-					CurrentCPUState.ActivatedCoreCount++;
-				}
-				if (Index == CurrentCPUState.CoreCount - 1)
-					break;
-			}
-		}
-		fclose(FileHandle);
+	uint64 TotalTicks = UserTime + NiceTime + SystemTime + SoftIRQTime + IRQTime + IdleTime + IOWaitTime;
 
-		CurrentCPUState.AverageUtilization = 0.0;
-		CurrentCPUState.AverageIdleTime = 0.0f;
-		for (size_t n = 0; n < CurrentCPUState.CoreCount; n++)
-		{
-			if (CurrentCPUState.CurrentUsage[n].TotalTime <= CurrentCPUState.PreviousUsage[n].TotalTime)
-			{
-				continue;
-			}
-			
-			WallTime = CurrentCPUState.CurrentUsage[n].TotalTime - CurrentCPUState.PreviousUsage[n].TotalTime;
-			IdleTime = CurrentCPUState.CurrentUsage[n].IdleTime - CurrentCPUState.PreviousUsage[n].IdleTime;
+	sprintf(PIDString, "/proc/%d/stat", ProcessId);
+	FILE * ProcPIDStatFile = fopen(PIDString, "r");
 
-			if (!WallTime || WallTime <= IdleTime)
-			{
-				continue;
-			}
-			CurrentCPUState.Utilization[n] = ((double)WallTime - (double)IdleTime) * 100.0 / (double)WallTime;
-			CurrentCPUState.IdleTime[n] = (double)IdleTime * 100.0 / (double)WallTime;
-			CurrentCPUState.AverageUtilization += CurrentCPUState.Utilization[n];
-			CurrentCPUState.AverageIdleTime += CurrentCPUState.IdleTime[n];
-		}
-		CurrentCPUState.AverageUtilization /= (double)CurrentCPUState.CoreCount;
-		CurrentCPUState.AverageIdleTime /= (double)CurrentCPUState.CoreCount;
-	}
- 	else
- 	{
- 		UE_LOG(LogHAL, Warning, TEXT("Could not open /proc/stat, CPU usage stats will be missing"));
-		FMemory::Memzero(CurrentCPUState);
+	if (ProcPIDStatFile == nullptr)
+	{
+		UE_CLOG(bProcPIDPrintHappened, LogUtilization, Verbose, TEXT("Could not open /proc/%d/stat to read PID usage."), ProcessId);
+		bProcPIDPrintHappened = true;
 		return false;
 	}
-	ProcessUsageFraction = (float)CurrentCPUState.AverageUtilization;
- 	IdleUsageFraction = (float)CurrentCPUState.AverageIdleTime;
+	
+	if (fgets(Buffer, sizeof(Buffer), ProcPIDStatFile) == nullptr)
+	{
+		UE_CLOG(bProcPIDParsePrintHappened, LogUtilization, Verbose, TEXT("Could not get the first line of the PID stats file."));
+		bProcPIDParsePrintHappened = true;
+		fclose(ProcPIDStatFile);
+		return false;
+	}
+	
+	fclose(ProcPIDStatFile);
+
+	// We only need the fields around the total time spent in user and system for both this process and its children
+	// These are from man proc stat
+	constexpr uint8 UserTimeField = 14;
+	constexpr uint8 SystemTimeField = 15;
+	constexpr uint8 ChildUserTimeField = 16;
+	constexpr uint8 ChildSystemTimeField = 17;
+	uint64 UserTimeTicks = 0;
+	uint64 SystemTimeTicks = 0;
+	uint64 ChildUserTimeTicks = 0;
+	uint64 ChildSystemTimeTicks = 0;
+
+	char* Remainder = nullptr;
+	char* Token =  strtok_r(Buffer, " ", &Remainder);
+	uint32 NumTokens = 1;
+
+	while (Token != nullptr)
+	{
+		Token = strtok_r(nullptr, " ", &Remainder);
+		NumTokens++;
+		if (NumTokens == UserTimeField)
+		{
+			UserTimeTicks = atol(Token);
+		}
+		if (NumTokens == SystemTimeField)
+		{
+			SystemTimeTicks = atol(Token);
+		}
+		if (NumTokens == ChildUserTimeField)
+		{
+			ChildUserTimeTicks = atol(Token);
+		}
+		if (NumTokens == ChildSystemTimeField)
+		{
+			ChildSystemTimeTicks = atol(Token);
+			break;
+		}
+	}
+
+	// First sample period, there will be no valid data here
+	// TODO: Catch a large delta 
+	if (PreviousTotalTicks == 0)
+	{
+		PreviousTotalTicks = TotalTicks;
+		PreviousUserTimeTicks = UserTimeTicks;
+		PreviousSystemTimeTicks = SystemTimeTicks;
+		PreviousChildUserTimeTicks = ChildUserTimeTicks;
+		PreviousChildSystemTimeTicks = ChildSystemTimeTicks;
+		return false;	
+	}
+
+	// Quick math
+	uint32 TotalTicksDelta = TotalTicks - PreviousTotalTicks;
+	uint32 UserTimeDelta = UserTimeTicks - PreviousUserTimeTicks;
+	uint32 SystemTimeDelta = SystemTimeTicks - PreviousSystemTimeTicks;
+	uint32 ChildUserTimeDelta = ChildUserTimeTicks - PreviousChildUserTimeTicks;
+	uint32 ChildSystemTimeDelta = ChildSystemTimeTicks - PreviousChildSystemTimeTicks;
+
+	uint32 TotalProcessTicksDelta = UserTimeDelta + SystemTimeDelta + ChildUserTimeDelta + ChildSystemTimeDelta;
 
+	ProcessUsageFraction = (float)TotalProcessTicksDelta/(float)TotalTicksDelta * 100.0f;
+	// Not sure what the correct logic would be here because we do not have an idle time for this process
+	IdleUsageFraction = 100.0f - ProcessUsageFraction;
+
+	PreviousTotalTicks = TotalTicks;
+	PreviousUserTimeTicks = UserTimeTicks;
+	PreviousSystemTimeTicks = SystemTimeTicks;
+	PreviousChildUserTimeTicks = ChildUserTimeTicks;
+	PreviousChildSystemTimeTicks = ChildSystemTimeTicks;
+	
  	return true;
 }
 
@@ -2330,4 +2377,4 @@ void FUnixPlatformProcess::SetThreadPriority(EThreadPriority NewPriority)
 	int32 NiceValue = FUnixPlatformProcess::TranslateThreadPriority(NewPriority);
 
 	FUnixPlatformProcess::SetThreadNiceValue(ThreadId, NiceValue);
-}
+}
\ No newline at end of file
