Add comments to Crc and CityHash to prefer other hashes for new code.

diff --git a/Engine/Source/Runtime/Core/Public/Hash/CityHash.h b/Engine/Source/Runtime/Core/Public/Hash/CityHash.h
--- a/Engine/Source/Runtime/Core/Public/Hash/CityHash.h
+++ b/Engine/Source/Runtime/Core/Public/Hash/CityHash.h
@@ -72,18 +72,22 @@ struct Uint128_64
 };
 
 // Hash function for a byte array.
+// NOTE: provided for compatibility, newer code should prefer XXHash.
 CORE_API uint64 CityHash64(const char *buf, uint32 len);
 
 // Hash function for a byte array.  For convenience, a 64-bit seed is also
 // hashed into the result.
+// NOTE: provided for compatibility, newer code should prefer XXHash.
 CORE_API uint64 CityHash64WithSeed(const char *buf, uint32 len, uint64 seed);
 
 // Hash function for a byte array.  For convenience, two seeds are also
 // hashed into the result.
+// NOTE: provided for compatibility, newer code should prefer XXHash.
 CORE_API uint64 CityHash64WithSeeds(const char *buf, uint32 len,
 	uint64 seed0, uint64 seed1);
 
 // Hash function for a byte array.  Most useful in 32-bit binaries.
+// NOTE: provided for compatibility, newer code should prefer XXHash.
 CORE_API uint32 CityHash32(const char *buf, uint32 len);
 
 // Hash 128 input bits down to 64 bits of output.
@@ -99,4 +103,4 @@ inline uint64 CityHash128to64(const Uint128_64& x) {
 	return b;
 }
 
-#endif  // CITY_HASH_H_
\ No newline at end of file
+#endif  // CITY_HASH_H_
diff --git a/Engine/Source/Runtime/Core/Public/Misc/Crc.h b/Engine/Source/Runtime/Core/Public/Misc/Crc.h
--- a/Engine/Source/Runtime/Core/Public/Misc/Crc.h
+++ b/Engine/Source/Runtime/Core/Public/Misc/Crc.h
@@ -18,14 +18,14 @@ struct TTuple;
  **/
 struct FCrc
 {
-	/** lookup table with precalculated CRC values - slicing by 8 implementation */
+	/** Lookup table with precalculated CRC values - slicing by 8 implementation */
 	static CORE_API uint32 CRCTablesSB8[8][256];
 
-	/** initializes the CRC lookup table. Must be called before any of the
+	/** Initializes the CRC lookup table. Must be called before any of the
 		CRC functions are used. */
 	static CORE_API void Init();
 
-	/** generates CRC hash of the memory area */
+	/** Generates CRC hash of the memory area. Prefer Hash/xxhash.h for new code, CRCs are not very good and fairly slow. */
 	typedef uint32 (*MemCrc32Functor)( const void* Data, int32 Length, uint32 CRC );
 	static CORE_API MemCrc32Functor MemCrc32Func;
 	[[nodiscard]] static FORCEINLINE uint32 MemCrc32(const void* Data, int32 Length, uint32 CRC = 0)
@@ -33,14 +33,14 @@ struct FCrc
 		return MemCrc32Func(Data, Length, CRC);
 	}
 
-	/** generates CRC hash of the element */
+	/** Generates CRC hash of the element. Prefer Hash/xxhash.h for new code, CRCs are not very good and fairly slow. */
 	template <typename T>
 	[[nodiscard]] static uint32 TypeCrc32( const T& Data, uint32 CRC=0 )
 	{
 		return MemCrc32(&Data, sizeof(T), CRC);
 	}
 
-	/** String CRC. */
+	/** String CRC. Prefer Hash/Fnv.h for new code, CRCs are not very good and fairly slow. */
 	template <typename CharType>
 	[[nodiscard]] static uint32 StrCrc32(const CharType* Data, uint32 CRC = 0)
 	{
@@ -79,7 +79,7 @@ struct FCrc
 		}
 	}
 
-	/** String CRC. */
+	/** String CRC. Prefer Hash/Fnv.h for new code, CRCs are not very good and fairly slow. */
 	template <typename CharType>
 	[[nodiscard]] static uint32 StrCrc32Len(const CharType* Data, int32 Length, uint32 CRC = 0)
 	{
