Remove deprecated FLinearColor methods, deprecate unused DXT block structs

diff --git a/Engine/Source/Runtime/Core/Private/Math/Color.cpp b/Engine/Source/Runtime/Core/Private/Math/Color.cpp
--- a/Engine/Source/Runtime/Core/Private/Math/Color.cpp
+++ b/Engine/Source/Runtime/Core/Private/Math/Color.cpp
@@ -369,19 +369,6 @@ FLinearColor FColor::FromRGBE() const
 	}
 }
 
-/**
- * Old version of MakeFromHSV8, formula is not correct
- */
-FLinearColor FLinearColor::FGetHSV( uint8 H, uint8 S, uint8 V )
-{
-	float Brightness = V * 1.4f / 255.f;
-	Brightness *= 0.7f/(0.01f + FMath::Sqrt(Brightness));
-	Brightness  = FMath::Clamp(Brightness,0.f,1.f);
-	const FVector3f Hue = (H<86) ? FVector3f((85-H)/85.f,(H-0)/85.f,0) : (H<171) ? FVector3f(0,(170-H)/85.f,(H-85)/85.f) : FVector3f((H-170)/85.f,0,(255-H)/84.f);
-	const FVector3f ColorVector = (Hue + S/255.f * (FVector3f(1,1,1) - Hue)) * Brightness;
-	return FLinearColor(ColorVector.X, ColorVector.Y, ColorVector.Z,1);
-}
-
 /**
  * Converts byte hue-saturation-brightness to floating point red-green-blue.
  */
diff --git a/Engine/Source/Runtime/Core/Public/Math/Color.h b/Engine/Source/Runtime/Core/Public/Math/Color.h
--- a/Engine/Source/Runtime/Core/Public/Math/Color.h
+++ b/Engine/Source/Runtime/Core/Public/Math/Color.h
@@ -294,9 +294,6 @@ struct FLinearColor
 		return NewCopy;
 	}
 
-	UE_DEPRECATED(4.22, "FGetHSV doesn't perform a valid HSV conversion, use MakeFromHSV8 instead")
-	[[nodiscard]] static CORE_API FLinearColor FGetHSV(uint8 H,uint8 S,uint8 V);
-
 	/**
 	 * Converts byte hue-saturation-brightness to floating point red-green-blue.
 	 */
@@ -366,12 +363,6 @@ struct FLinearColor
 	*/
 	[[nodiscard]] FORCEINLINE FColor QuantizeFloor() const;
 
-	/** backwards compatible Quantize function name, does QuantizeFloor.
-	* @todo deprecate me
-	*/
-	UE_DEPRECATED(5.0, "Most callers of Quantize should have been calling QuantizeRound; to match old behavior use QuantizeFloor")
-	[[nodiscard]] FORCEINLINE FColor Quantize() const;
-
 	/** Quantizes the linear color and returns the result as a FColor with optional sRGB conversion. 
 	* Clamps in [0,1] range before conversion.
 	* ToFColor(false) is QuantizeRound
@@ -816,11 +807,6 @@ FORCEINLINE FColor FLinearColor::QuantizeFloor() const
 	);
 }
 
-FORCEINLINE FColor FLinearColor::Quantize() const
-{
-	return QuantizeFloor();
-}
-
 FORCEINLINE FColor FLinearColor::ToFColor(const bool bSRGB) const
 {
 	if ( bSRGB )
@@ -848,7 +834,7 @@ template <> struct TIsPODType<FLinearColor> { enum { Value = true }; };
 /**
  * Helper struct for a 16 bit 565 color of a DXT1/3/5 block.
  */
-struct FDXTColor565
+struct UE_DEPRECATED(5.7, "FDXTColor565 is deprecated, manipulate DXT blocks directly as required") FDXTColor565
 {
 	/** Blue component, 5 bit. */
 	uint16 B:5;
@@ -864,8 +850,9 @@ struct FDXTColor565
 /**
  * Helper struct for a 16 bit 565 color of a DXT1/3/5 block.
  */
-struct FDXTColor16
+struct UE_DEPRECATED(5.7, "FDXTColor16 is deprecated, manipulate DXT blocks directly as required") FDXTColor16
 {
+PRAGMA_DISABLE_DEPRECATION_WARNINGS
 	union 
 	{
 		/** 565 Color */
@@ -873,20 +860,23 @@ struct FDXTColor16
 		/** 16 bit entity representation for easy access. */
 		uint16 Value;
 	};
+PRAGMA_ENABLE_DEPRECATION_WARNINGS
 };
 
 
 /**
  * Structure encompassing single DXT1 block.
  */
-struct FDXT1
+struct UE_DEPRECATED(5.7, "FDXT1 is deprecated, manipulate DXT blocks directly as required") FDXT1
 {
+PRAGMA_DISABLE_DEPRECATION_WARNINGS
 	/** Color 0/1 */
 	union
 	{
 		FDXTColor16 Color[2];
 		uint32 Colors;
 	};
+PRAGMA_ENABLE_DEPRECATION_WARNINGS
 	/** Indices controlling how to blend colors. */
 	uint32 Indices;
 };
@@ -895,17 +885,22 @@ struct FDXT1
 /**
  * Structure encompassing single DXT5 block
  */
-struct FDXT5
+struct UE_DEPRECATED(5.7, "FDXT5 is deprecated, manipulate DXT blocks directly as required") FDXT5
 {
 	/** Alpha component of DXT5 */
 	uint8	Alpha[8];
 	/** DXT1 color component. */
+PRAGMA_DISABLE_DEPRECATION_WARNINGS
 	FDXT1	DXT1;
+PRAGMA_ENABLE_DEPRECATION_WARNINGS
 };
 
+PRAGMA_DISABLE_DEPRECATION_WARNINGS
 
 // Make DXT helpers act like PODs
 template <> struct TIsPODType<FDXT1> { enum { Value = true }; };
 template <> struct TIsPODType<FDXT5> { enum { Value = true }; };
 template <> struct TIsPODType<FDXTColor16> { enum { Value = true }; };
 template <> struct TIsPODType<FDXTColor565> { enum { Value = true }; };
+
+PRAGMA_ENABLE_DEPRECATION_WARNINGS
