Remove Vector2D deprecated comparison operators.

diff --git a/Engine/Source/Runtime/Core/Public/Math/Vector2D.h b/Engine/Source/Runtime/Core/Public/Math/Vector2D.h
--- a/Engine/Source/Runtime/Core/Public/Math/Vector2D.h
+++ b/Engine/Source/Runtime/Core/Public/Math/Vector2D.h
@@ -242,54 +242,6 @@ public:
 	*/
 	[[nodiscard]] bool operator!=(const TVector2<T>& V) const;
 
-	/**
-	* Deprecated comparison operator. Use ComponentwiseAllLessThan.
-	*
-	* @param Other The vector to compare against.
-	* @return true if this is the smaller vector, otherwise false.
-	*/
-	UE_DEPRECATED(5.1, "TVector2 comparison operators are deprecated. Use ComponentwiseAllLessThan. For componentwise min/max/abs/clamp, use TVector2::{Min,Max,GetAbs,Clamp}, FMath::{Min,Max,Abs,Clamp} compute something different.")
-	[[nodiscard]] bool operator<(const TVector2<T>& Other) const
-	{
-		return ComponentwiseAllLessThan(Other);
-	}
-
-	/**
-	* Deprecated comparison operator. Use ComponentwiseAllGreaterThan.
-	*
-	* @param Other The vector to compare against.
-	* @return true if this is the larger vector, otherwise false.
-	*/
-	UE_DEPRECATED(5.1, "TVector2 comparison operators are deprecated. Use ComponentwiseAllGreaterThan. For componentwise min/max/abs/clamp, use TVector2::{Min,Max,GetAbs,Clamp}, FMath::{Min,Max,Abs,Clamp} compute something different.")
-	[[nodiscard]] bool operator>(const TVector2<T>& Other) const
-	{
-		return ComponentwiseAllGreaterThan(Other);
-	}
-
-	/**
-	* Deprecated comparison operator. Use ComponentwiseAllLessOrEqual.
-	*
-	* @param Other The vector to compare against.
-	* @return true if this vector is less than or equal to the other vector, otherwise false.
-	*/
-	UE_DEPRECATED(5.1, "TVector2 comparison operators are deprecated. Use ComponentwiseAllLessOrEqual. For componentwise min/max/abs/clamp, use TVector2::{Min,Max,GetAbs,Clamp}, FMath::{Min,Max,Abs,Clamp} compute something different.")
-	[[nodiscard]] bool operator<=(const TVector2<T>& Other) const
-	{
-		return ComponentwiseAllLessOrEqual(Other);
-	}
-
-	/**
-	* Deprecated comparison operator. Use ComponentwiseAllGreaterOrEqual.
-	*
-	* @param Other The vector to compare against.
-	* @return true if this vector is greater than or equal to the other vector, otherwise false.
-	*/
-	UE_DEPRECATED(5.1, "TVector2 comparison operators are deprecated. Use ComponentwiseAllGreaterOrEqual. For componentwise min/max/abs/clamp, use TVector2::{Min,Max,GetAbs,Clamp}, FMath::{Min,Max,Abs,Clamp} compute something different.")
-	[[nodiscard]] bool operator>=(const TVector2<T>& Other) const
-	{
-		return ComponentwiseAllGreaterOrEqual(Other);
-	}
-
 	/**
 	* Checks whether both components of this vector are less than another.
 	*
