Fix analytics for iad cosmetics.

diff --git a/Engine/Source/Runtime/Core/Private/IO/IoStatus.cpp b/Engine/Source/Runtime/Core/Private/IO/IoStatus.cpp
--- a/Engine/Source/Runtime/Core/Private/IO/IoStatus.cpp
+++ b/Engine/Source/Runtime/Core/Private/IO/IoStatus.cpp
@@ -22,6 +22,10 @@ extern const TCHAR* const* GetIoErrorText_ErrorCodeText;
 
 const TCHAR* GetIoErrorText(EIoErrorCode ErrorCode)
 {
+	if (ErrorCode >= EIoErrorCode::Last)
+	{
+		return GetIoErrorText_ErrorCodeText[static_cast<uint32>(EIoErrorCode::Unknown)];
+	}
 	return GetIoErrorText_ErrorCodeText[static_cast<uint32>(ErrorCode)];
 }
 
