Added Max() functions to sets and maps - and TSparseArray which is needed for sets - and .

diff --git a/Engine/Source/Runtime/Core/Public/Containers/Map.h b/Engine/Source/Runtime/Core/Public/Containers/Map.h
--- a/Engine/Source/Runtime/Core/Public/Containers/Map.h
+++ b/Engine/Source/Runtime/Core/Public/Containers/Map.h
@@ -290,6 +290,12 @@ public:
 		return Pairs.Num();
 	}
 
+	/** @return The number of elements the map can hold before reallocation. */
+	[[nodiscard]] FORCEINLINE int32 Max() const
+	{
+		return Pairs.Max();
+	}
+
 	/** @return The max valid index of the elements in the sparse storage. */
 	[[nodiscard]] FORCEINLINE int32 GetMaxIndex() const
 	{
@@ -1922,6 +1928,12 @@ public:
 		return Pairs.Num();
 	}
 
+	/** @return The number of elements the map can hold before reallocation. */
+	[[nodiscard]] int32 Max() const
+	{
+		return Pairs.Max();
+	}
+
 	[[nodiscard]] int32 NumUnchecked() const
 	{
 		return Pairs.NumUnchecked();
diff --git a/Engine/Source/Runtime/Core/Public/Containers/Set.h b/Engine/Source/Runtime/Core/Public/Containers/Set.h
--- a/Engine/Source/Runtime/Core/Public/Containers/Set.h
+++ b/Engine/Source/Runtime/Core/Public/Containers/Set.h
@@ -448,6 +448,12 @@ public:
 		return Elements.Num();
 	}
 
+	/** @return The number of elements the set can hold before reallocation. */
+	[[nodiscard]] FORCEINLINE int32 Max() const
+	{
+		return Elements.Max();
+	}
+
 	[[nodiscard]] FORCEINLINE int32 GetMaxIndex() const
 	{
 		return Elements.GetMaxIndex();
@@ -1811,6 +1817,12 @@ public:
 		return Elements.Num();
 	}
 
+	/** @return The number of elements the set can hold before reallocation. */
+	[[nodiscard]] int32 Max() const
+	{
+		return Elements.Max();
+	}
+
 	[[nodiscard]] int32 NumUnchecked() const
 	{
 		return Elements.NumUnchecked();
diff --git a/Engine/Source/Runtime/Core/Public/Containers/SortedMap.h b/Engine/Source/Runtime/Core/Public/Containers/SortedMap.h
--- a/Engine/Source/Runtime/Core/Public/Containers/SortedMap.h
+++ b/Engine/Source/Runtime/Core/Public/Containers/SortedMap.h
@@ -161,6 +161,12 @@ public:
 		return Pairs.Num();
 	}
 
+	/** @return The number of elements the map can hold before reallocation. */
+	[[nodiscard]] FORCEINLINE int32 Max() const
+	{
+		return Pairs.Max();
+	}
+
 	/** 
 	 * Helper function to return the amount of memory allocated by this container.
 	 * Only returns the size of allocations made directly by the container, not the elements themselves.
diff --git a/Engine/Source/Runtime/Core/Public/Containers/SortedSet.h b/Engine/Source/Runtime/Core/Public/Containers/SortedSet.h
--- a/Engine/Source/Runtime/Core/Public/Containers/SortedSet.h
+++ b/Engine/Source/Runtime/Core/Public/Containers/SortedSet.h
@@ -162,6 +162,12 @@ public:
 		return Elements.Num();
 	}
 
+	/** @return The number of elements the set can hold before reallocation. */
+	[[nodiscard]] FORCEINLINE int32 Max() const
+	{
+		return Elements.Max();
+	}
+
 	/** @return The max valid index of the elements in the inner array. */
 	[[nodiscard]] FORCEINLINE int32 GetMaxIndex() const
 	{
diff --git a/Engine/Source/Runtime/Core/Public/Containers/SparseArray.h b/Engine/Source/Runtime/Core/Public/Containers/SparseArray.h
--- a/Engine/Source/Runtime/Core/Public/Containers/SparseArray.h
+++ b/Engine/Source/Runtime/Core/Public/Containers/SparseArray.h
@@ -932,6 +932,12 @@ public:
 		return Data.Num() - NumFreeIndices;
 	}
 
+	/** @return The number of elements the array can hold before reallocation. */
+	[[nodiscard]] int32 Max() const
+	{
+		return Data.Max();
+	}
+
 	/**
 	 * Checks that the specified address is not part of an element within the container.  Used for implementations
 	 * to check that reference arguments aren't going to be invalidated by possible reallocation.
@@ -1434,6 +1440,12 @@ public:
 		return Data.Num() - NumFreeIndices;
 	}
 
+	/** @return The number of elements the array can hold before reallocation. */
+	[[nodiscard]] int32 Max() const
+	{
+		return Data.Max();
+	}
+
 	[[nodiscard]] int32 NumUnchecked() const
 	{
 		return Data.NumUnchecked() - NumFreeIndices;
