PR #13336: Add trace counter to track FCoreDelegates::OnEndFrame timing

#rnx

diff --git a/Engine/Source/Runtime/Core/Private/ProfilingDebugging/TraceAuxiliary.cpp b/Engine/Source/Runtime/Core/Private/ProfilingDebugging/TraceAuxiliary.cpp
--- a/Engine/Source/Runtime/Core/Private/ProfilingDebugging/TraceAuxiliary.cpp
+++ b/Engine/Source/Runtime/Core/Private/ProfilingDebugging/TraceAuxiliary.cpp
@@ -1004,7 +1004,11 @@ void FTraceAuxiliaryImpl::RegisterEndFrameCallbacks()
 {
 	// Always register end frame updates. This path is short circuited if a worker thread exists.
 	// If the worker thread is disabled, pump the update from end frame.
-	FCoreDelegates::OnEndFrame.AddStatic(UE::Trace::Update);
+	FCoreDelegates::OnEndFrame.AddLambda([]()
+		{
+			TRACE_CPUPROFILER_EVENT_SCOPE(FTraceAux_Trace_Update);
+			UE::Trace::Update();
+		});
 
 	// Update stats every frame.
 	FCoreDelegates::OnEndFrame.AddStatic(TraceAuxiliaryOnEndFrameCallback);
diff --git a/Engine/Source/Runtime/Launch/Private/LaunchEngineLoop.cpp b/Engine/Source/Runtime/Launch/Private/LaunchEngineLoop.cpp
--- a/Engine/Source/Runtime/Launch/Private/LaunchEngineLoop.cpp
+++ b/Engine/Source/Runtime/Launch/Private/LaunchEngineLoop.cpp
@@ -5895,7 +5895,10 @@ void FEngineLoop::Tick()
 		}
 #endif
 
-		FCoreDelegates::OnEndFrame.Broadcast();
+		{
+			QUICK_SCOPE_CYCLE_COUNTER(STAT_FEngineLoop_CoreDelegates_EndFrame);
+			FCoreDelegates::OnEndFrame.Broadcast();
+		}
 
 		// Increment global frame counter. Once for each engine tick.
 		GFrameCounter++;
