Change thread priorities for GT/RT/RHIT from Normal to AboveNormal on Windows

diff --git a/Engine/Source/Runtime/Core/Public/Windows/WindowsPlatformAffinity.h b/Engine/Source/Runtime/Core/Public/Windows/WindowsPlatformAffinity.h
--- a/Engine/Source/Runtime/Core/Public/Windows/WindowsPlatformAffinity.h
+++ b/Engine/Source/Runtime/Core/Public/Windows/WindowsPlatformAffinity.h
@@ -3,4 +3,24 @@
 #pragma once
 
 #include "GenericPlatform/GenericPlatformAffinity.h"
-typedef FGenericPlatformAffinity FPlatformAffinity;
+
+class FWindowsPlatformAffinity : public FGenericPlatformAffinity
+{
+public:
+	static EThreadPriority GetRenderingThreadPriority()
+	{
+		return TPri_AboveNormal;
+	}
+
+	static EThreadPriority GetRHIThreadPriority()
+	{
+		return TPri_AboveNormal;
+	}
+
+	static EThreadPriority GetGameThreadPriority()
+	{
+		return TPri_AboveNormal;
+	}
+};
+
+typedef FWindowsPlatformAffinity FPlatformAffinity;
