[PlatformFile]

Defensive changes to how we read the change journal:

- Adjusts Change Journal reading code to be a bit more defensive to possible implementation details that might allow for our read buffer to always be fully populated while enumerating records. The previous code would read until the read buffer isn't at least half populated, however it's unclear if DeviceIoControl might not always write as much as it can. Instead we use the known journal range at the time of scanning and only scan that.

- Aligns the read buffer to an 8 byte boundary as MSDN mentions this as a requirement (however it doesn't seem like it's been an issue thus far)

#rnx

diff --git a/Engine/Source/Runtime/Core/Private/Windows/WindowsPlatformFileJournal.cpp b/Engine/Source/Runtime/Core/Private/Windows/WindowsPlatformFileJournal.cpp
--- a/Engine/Source/Runtime/Core/Private/Windows/WindowsPlatformFileJournal.cpp
+++ b/Engine/Source/Runtime/Core/Private/Windows/WindowsPlatformFileJournal.cpp
@@ -438,12 +438,9 @@ EFileJournalResult FWindowsPlatformFile::FileJournalReadModified(const TCHAR* Vo
 	EFileJournalResult Result;
 	OutNextJournalEntry = StartingJournalEntry;
 
-	uint8 ReadBuffer[4096];
+	// Records are expected to be aligned to a 64-bit boundary. The USN taking the first 64-bits.
+	alignas(sizeof(uint64)) uint8 ReadBuffer[4096];
 	DWORD ReadBufferSize = UE_ARRAY_COUNT(ReadBuffer);
-	// For detection of when we have reached the end of the journal, we need buffer size large enough to read more
-	// records at once than we expect can be written in the time it takes us to call it twice in a row.
-	// sizeof(USN_RECORD_V3) == 80
-	static_assert(sizeof(ReadBuffer) > 10 * sizeof(USN_RECORD_V3), "ReadBuffer too small");
 
 	HANDLE JournalWindowsHandle = NULL;
 	TStringBuilder<16> NormalizedVolumeName;
@@ -479,8 +476,9 @@ EFileJournalResult FWindowsPlatformFile::FileJournalReadModified(const TCHAR* Vo
 	// it returns no results and reports the input StartUsn as the next USN to read. We want this to be an error
 	// condition similar to wrapping the journal. StartUSn == NextUSN is valid, but StartUSN > NextUSN is not.
 	// Return a wrap code for that case.
-	USN CurrentJournalUSN = static_cast<USN>(OutNextJournalEntry);
-	if (CurrentJournalUSN > JournalDescriptor.NextUsn)
+	const USN LastJournalUSN = JournalDescriptor.NextUsn;
+	USN CurrentJournalUSN = static_cast<USN>(StartingJournalEntry);
+	if (CurrentJournalUSN > LastJournalUSN)
 	{
 		if (OutError)
 		{
@@ -521,14 +519,13 @@ EFileJournalResult FWindowsPlatformFile::FileJournalReadModified(const TCHAR* Vo
 	bool bShowFilenames = bTryShowFilenames;
 	bool bTriedAndFailedShowFilenames = false;
 	bool bFirstCall = true;
-	for (;;)
+	do
 	{
 		// FSCTL_READ_UNPRIVILEGED_USN_JOURNAL can read without elevation but does not show filenames.
 		// FCSTL_READ_USN_JOURNAL requires elevation and shows filenames
 		DWORD dwIoControlCode = bShowFilenames ? FSCTL_READ_USN_JOURNAL : FSCTL_READ_UNPRIVILEGED_USN_JOURNAL;
 		bool bWasFirstCall = bFirstCall;
 		bFirstCall = false;
-		CurrentJournalUSN = static_cast<USN>(OutNextJournalEntry);
 		ReadJournalRequest.StartUsn = CurrentJournalUSN;
 
 		BOOL ReadJournalResult = ::DeviceIoControl(
@@ -574,15 +571,12 @@ EFileJournalResult FWindowsPlatformFile::FileJournalReadModified(const TCHAR* Vo
 				}
 			}
 		}
-		if (ReadBufferResultsSize < sizeof(USN))
+		if (ReadBufferResultsSize <= sizeof(USN))
 		{
-			if (OutError)
-			{
-				*OutError = FString::Printf(
-					TEXT("FSCTL_READ_UNPRIVILEGED_USN_JOURNAL returned an unexpected value; ReadBufferResultsSize == %u, and the minimum expected is %zu."),
-					(uint32)ReadBufferResultsSize, sizeof(USN));
-			}
-			return EFileJournalResult::JournalInternalError;
+			// There are no more USN entries or Journal Records left
+			// Note, we check for <=. While theoretically we can read a USN, if there are no bytes for 
+			// the records we must be out of data to read, so return without error.
+			return EFileJournalResult::Success;
 		}
 		if (bTriedAndFailedShowFilenames)
 		{
@@ -624,14 +618,7 @@ EFileJournalResult FWindowsPlatformFile::FileJournalReadModified(const TCHAR* Vo
 
 			RecordBytes += Record->RecordLength;
 		}
-		static_assert(sizeof(uint64) >= sizeof(USN), "We are storing USNs as uint64");
-		OutNextJournalEntry = static_cast<uint64>(NextJournalUSN);
 
-		// Keep going until we reach the end, which we detect by Journal did not have enough results to fill the buffer
-		if (RecordBytesEnd <= ReadBuffer + ReadBufferSize / 2)
-		{
-			break;
-		}
 		// Prevent a possible journal bug from causing an infinite loop
 		if (NextJournalUSN <= CurrentJournalUSN)
 		{
@@ -644,7 +631,11 @@ EFileJournalResult FWindowsPlatformFile::FileJournalReadModified(const TCHAR* Vo
 			}
 			return EFileJournalResult::JournalInternalError;
 		}
-	}
+
+		static_assert(sizeof(uint64) >= sizeof(USN), "We are storing USNs as uint64");
+		OutNextJournalEntry = static_cast<uint64>(NextJournalUSN);
+		CurrentJournalUSN = NextJournalUSN;
+	} while (CurrentJournalUSN < LastJournalUSN);
 	return EFileJournalResult::Success;
 }
 
