Bump material translation DDC key to fix inconsistency in DDC records.
This was missing from some previous changelists.
#rnx

diff --git a/Engine/Source/Runtime/Core/Private/UObject/FortniteMainBranchObjectVersion.cpp b/Engine/Source/Runtime/Core/Private/UObject/FortniteMainBranchObjectVersion.cpp
--- a/Engine/Source/Runtime/Core/Private/UObject/FortniteMainBranchObjectVersion.cpp
+++ b/Engine/Source/Runtime/Core/Private/UObject/FortniteMainBranchObjectVersion.cpp
@@ -14,6 +14,6 @@ TMap<FGuid, FGuid> FFortniteMainBranchObjectVersion::GetSystemGuids()
 	SystemGuids.Add(DevGuids.Niagara_LatestScriptCompileVersion, FGuid("97F0FD7149D90D42B65D6557AF477D1A"));
 	SystemGuids.Add(DevGuids.SkeletalMeshDerivedDataVersion, FGuid("C998F2F3B2884A02B7E290B8C25835F0"));
 	SystemGuids.Add(DevGuids.STATICMESH_DERIVEDDATA_VER, FGuid("3ABF32149F9D4D838750368AB95573BE"));
-	SystemGuids.Add(DevGuids.MaterialTranslationDDCVersion, FGuid("E541A459F3374E628240BBB309910753"));
+	SystemGuids.Add(DevGuids.MaterialTranslationDDCVersion, FGuid("9074F4446D3843F1ABD42CA551384212"));
 	return SystemGuids;
 }
