Fix some minor issues with the UnixPlatformProcess calculation
- For the one time prints, negate the checks so they only fire once
- Everything is in percentage (0.0, 100.0) so the old code in LaunchEngineLoop.cpp that set things to 0, 1 was incorrect

diff --git a/Engine/Source/Runtime/Core/Private/Unix/UnixPlatformProcess.cpp b/Engine/Source/Runtime/Core/Private/Unix/UnixPlatformProcess.cpp
--- a/Engine/Source/Runtime/Core/Private/Unix/UnixPlatformProcess.cpp
+++ b/Engine/Source/Runtime/Core/Private/Unix/UnixPlatformProcess.cpp
@@ -1819,7 +1819,7 @@ bool FUnixPlatformProcess::GetPerFrameProcessorUsage(uint32 ProcessId, float& Pr
 
 	if (ProcStatFile == nullptr)
 	{
-		UE_CLOG(bProcStatPrintHasHappened, LogUtilization, Verbose, TEXT("Failed to open /proc/stat."));
+		UE_CLOG(!bProcStatPrintHasHappened, LogUtilization, Verbose, TEXT("Failed to open /proc/stat."));
 		bProcStatPrintHasHappened = true;
 		return false;
 	}
@@ -1827,7 +1827,7 @@ bool FUnixPlatformProcess::GetPerFrameProcessorUsage(uint32 ProcessId, float& Pr
 	// First line of the file is the total CPU usage so no need to tally everything together
 	if (fgets(Buffer, sizeof(Buffer), ProcStatFile) == nullptr)
 	{
-		UE_CLOG(bProcStatParsePrintHappened, LogUtilization, Verbose, TEXT("Failed to parse the first line of /proc/stat."));
+		UE_CLOG(!bProcStatParsePrintHappened, LogUtilization, Verbose, TEXT("Failed to parse the first line of /proc/stat."));
 		bProcStatParsePrintHappened = true;
 		fclose(ProcStatFile);
 		return false;
@@ -1850,14 +1850,14 @@ bool FUnixPlatformProcess::GetPerFrameProcessorUsage(uint32 ProcessId, float& Pr
 
 	if (ProcPIDStatFile == nullptr)
 	{
-		UE_CLOG(bProcPIDPrintHappened, LogUtilization, Verbose, TEXT("Could not open /proc/%d/stat to read PID usage."), ProcessId);
+		UE_CLOG(!bProcPIDPrintHappened, LogUtilization, Verbose, TEXT("Could not open /proc/%d/stat to read PID usage."), ProcessId);
 		bProcPIDPrintHappened = true;
 		return false;
 	}
 	
 	if (fgets(Buffer, sizeof(Buffer), ProcPIDStatFile) == nullptr)
 	{
-		UE_CLOG(bProcPIDParsePrintHappened, LogUtilization, Verbose, TEXT("Could not get the first line of the PID stats file."));
+		UE_CLOG(!bProcPIDParsePrintHappened, LogUtilization, Verbose, TEXT("Could not get the first line of the PID stats file."));
 		bProcPIDParsePrintHappened = true;
 		fclose(ProcPIDStatFile);
 		return false;
diff --git a/Engine/Source/Runtime/Launch/Private/LaunchEngineLoop.cpp b/Engine/Source/Runtime/Launch/Private/LaunchEngineLoop.cpp
--- a/Engine/Source/Runtime/Launch/Private/LaunchEngineLoop.cpp
+++ b/Engine/Source/Runtime/Launch/Private/LaunchEngineLoop.cpp
@@ -1344,12 +1344,6 @@ static void UpdateCoreCsvStats_BeginFrame()
 		    float ProcessUsagePercent = 0.f, IdleUsagePercent = 0.f;
 			if (FPlatformProcess::GetPerFrameProcessorUsage(ProcessId, ProcessUsagePercent, IdleUsagePercent))
 			{
-				// This occasionally happens on Linux and I am not entirely sure why.
-				// If we have zero idle time and zero utilization assume we have full idle time
-				if (ProcessUsagePercent == 0.0f && IdleUsagePercent == 0.0f)
-				{
-					IdleUsagePercent = 1.0f;
-				}
 				CSV_CUSTOM_STAT_GLOBAL(CPUUsage_Process, ProcessUsagePercent, ECsvCustomStatOp::Set);
 				CSV_CUSTOM_STAT_GLOBAL(CPUUsage_Idle, IdleUsagePercent, ECsvCustomStatOp::Set);
 			}
