Replace some usages of FORCEINLINE with inline in Core modules.

diff --git a/Engine/Source/Runtime/Core/Public/Misc/ByteSwap.h b/Engine/Source/Runtime/Core/Public/Misc/ByteSwap.h
--- a/Engine/Source/Runtime/Core/Public/Misc/ByteSwap.h
+++ b/Engine/Source/Runtime/Core/Public/Misc/ByteSwap.h
@@ -22,7 +22,7 @@ namespace Internal
 		return (BYTESWAP_ORDER32_unsigned(Value));
 	}
 
-	static FORCEINLINE uint64 ByteSwapGeneric64(uint64 Value)
+	static inline uint64 ByteSwapGeneric64(uint64 Value)
 	{
 		Value = ((Value << 8) & 0xFF00FF00FF00FF00ULL ) | ((Value >> 8) & 0x00FF00FF00FF00FFULL);
 		Value = ((Value << 16) & 0xFFFF0000FFFF0000ULL ) | ((Value >> 16) & 0x0000FFFF0000FFFFULL);
@@ -47,7 +47,7 @@ namespace Internal
 	#endif
 #endif
 
-FORCEINLINE uint16 BYTESWAP_ORDER16(uint16 Val)
+inline uint16 BYTESWAP_ORDER16(uint16 Val)
 {
 #if defined(UE_BYTESWAP_INTRINSIC_PRIVATE_16)
 	return UE_BYTESWAP_INTRINSIC_PRIVATE_16(Val);
@@ -61,7 +61,7 @@ FORCEINLINE int16 BYTESWAP_ORDER16(int16 Val)
 	return static_cast<int16>(BYTESWAP_ORDER16(static_cast<uint16>(Val)));
 }
 
-FORCEINLINE uint32 BYTESWAP_ORDER32(uint32 Val)
+inline uint32 BYTESWAP_ORDER32(uint32 Val)
 {
 #if defined(UE_BYTESWAP_INTRINSIC_PRIVATE_32)
 	return UE_BYTESWAP_INTRINSIC_PRIVATE_32(Val);
@@ -75,7 +75,7 @@ FORCEINLINE int32 BYTESWAP_ORDER32(int32 val)
 	return static_cast<int32>(BYTESWAP_ORDER32(static_cast<uint32>(val)));
 }
 
-FORCEINLINE uint64 BYTESWAP_ORDER64(uint64 Value)
+inline uint64 BYTESWAP_ORDER64(uint64 Value)
 {
 #if defined(UE_BYTESWAP_INTRINSIC_PRIVATE_64)
 	return UE_BYTESWAP_INTRINSIC_PRIVATE_64(Value);
@@ -89,13 +89,13 @@ FORCEINLINE int64 BYTESWAP_ORDER64(int64 Value)
 	return static_cast<int64>(BYTESWAP_ORDER64(static_cast<uint64>(Value)));
 }
 
-FORCEINLINE float BYTESWAP_ORDERF(float val)
+inline float BYTESWAP_ORDERF(float val)
 {
 	uint32 uval = BYTESWAP_ORDER32(*reinterpret_cast<const uint32*>(&val));
 	return *reinterpret_cast<const float*>(&uval);
 }
 
-FORCEINLINE double BYTESWAP_ORDERD(double val)
+inline double BYTESWAP_ORDERD(double val)
 {
 	uint64 uval = BYTESWAP_ORDER64(*reinterpret_cast<const uint64*>(&val));
 	return *reinterpret_cast<const double*>(&uval);
@@ -112,7 +112,7 @@ FORCEINLINE void BYTESWAP_ORDER_TCHARARRAY(TCHAR* str)
 
 #else
 
-FORCEINLINE void BYTESWAP_ORDER_TCHARARRAY(TCHAR* str)
+inline void BYTESWAP_ORDER_TCHARARRAY(TCHAR* str)
 {
 	static_assert(sizeof(TCHAR) == sizeof(uint16), "Assuming TCHAR is 2 bytes wide.");
 	for (TCHAR* c = str; *c; ++c)
diff --git a/Engine/Source/Runtime/Core/Public/Misc/CString.h b/Engine/Source/Runtime/Core/Public/Misc/CString.h
--- a/Engine/Source/Runtime/Core/Public/Misc/CString.h
+++ b/Engine/Source/Runtime/Core/Public/Misc/CString.h
@@ -980,7 +980,7 @@ typename TCString<T>::CharType* TCString<T>::Strcpy(CharType* Dest, const CharTy
 	return FPlatformString::Strcpy(Dest, Src);
 }
 
-template <typename T> FORCEINLINE
+template <typename T> inline
 typename TCString<T>::CharType* TCString<T>::Strcpy(CharType* Dest, SIZE_T DestCount, const CharType* Src)
 {
 	PRAGMA_DISABLE_DEPRECATION_WARNINGS
@@ -988,7 +988,7 @@ typename TCString<T>::CharType* TCString<T>::Strcpy(CharType* Dest, SIZE_T DestC
 	PRAGMA_ENABLE_DEPRECATION_WARNINGS
 }
 
-template <typename T> FORCEINLINE
+template <typename T> inline
 typename TCString<T>::CharType* TCString<T>::Strncpy( CharType* Dest, const CharType* Src, SIZE_T MaxLen )
 {
 	FPlatformString::Strncpy(Dest, Src, MaxLen);
@@ -1001,7 +1001,7 @@ typename TCString<T>::CharType* TCString<T>::Strcat(CharType* Dest, const CharTy
 	return FPlatformString::Strcat(Dest, Src);
 }
 
-template <typename T> FORCEINLINE
+template <typename T> inline
 typename TCString<T>::CharType* TCString<T>::Strcat(CharType* Dest, SIZE_T DestCount, const CharType* Src)
 {
 	PRAGMA_DISABLE_DEPRECATION_WARNINGS
@@ -1045,7 +1045,7 @@ namespace UE::Core::Private
 }
 
 template <typename T>
-UE_NODEBUG FORCEINLINE int32 TCString<T>::Strlen( const CharType* String ) 
+UE_NODEBUG inline int32 TCString<T>::Strlen( const CharType* String ) 
 {
 	if constexpr (std::is_same_v<T, UTF32CHAR>)
 	{
@@ -1105,7 +1105,7 @@ const typename TCString<T>::CharType* TCString<T>::Strrstr( const CharType* Stri
 	return Strrstr((CharType*)String, Find);
 }
 
-template <typename T> FORCEINLINE
+template <typename T> inline
 typename TCString<T>::CharType* TCString<T>::Strrstr( CharType* String, const CharType* Find )
 {
 	if (*Find == (CharType)0)
@@ -1127,7 +1127,7 @@ typename TCString<T>::CharType* TCString<T>::Strrstr( CharType* String, const Ch
 	}
 }
 
-template <typename T> FORCEINLINE
+template <typename T> inline
 int32 TCString<T>::Strspn( const CharType* String, const CharType* Mask )
 {
 	const CharType* StringIt = String;
@@ -1150,7 +1150,7 @@ int32 TCString<T>::Strspn( const CharType* String, const CharType* Mask )
 	return UE_PTRDIFF_TO_INT32(StringIt - String);
 }
 
-template <typename T> FORCEINLINE
+template <typename T> inline
 int32 TCString<T>::Strcspn( const CharType* String, const CharType* Mask )
 {
 	const CharType* StringIt = String;
diff --git a/Engine/Source/Runtime/Core/Public/Misc/ConfigCacheIni.h b/Engine/Source/Runtime/Core/Public/Misc/ConfigCacheIni.h
--- a/Engine/Source/Runtime/Core/Public/Misc/ConfigCacheIni.h
+++ b/Engine/Source/Runtime/Core/Public/Misc/ConfigCacheIni.h
@@ -283,7 +283,7 @@ public:
 	}
 
 	// Returns the ini setting with any macros expanded out
-	FORCEINLINE FString GetValue() const
+	inline FString GetValue() const
 	{
 #if UE_WITH_CONFIG_TRACKING
 		UE::ConfigAccessTracking::Private::OnConfigValueRead(SectionAccess, ValueName, *this);
@@ -619,97 +619,97 @@ public:
 	// looks for a section by name, and creates an empty one if it can't be found
 	CORE_API const FConfigSection* FindOrAddConfigSection(const FString& Name);
 
-	FORCEINLINE const FConfigSection* FindSection(const FString& SectionName) const
+	inline const FConfigSection* FindSection(const FString& SectionName) const
 	{
 		UE::TReadScopeLock ScopeLock(ConfigFileMapLock);
 		return FConfigFileMap::Find(SectionName);
 	}
 
-	FORCEINLINE int32 Num() const
+	inline int32 Num() const
 	{
 		UE::TReadScopeLock ScopeLock(ConfigFileMapLock);
 		return FConfigFileMap::Num();
 	}
 
-	FORCEINLINE bool IsEmpty() const
+	inline bool IsEmpty() const
 	{
 		UE::TReadScopeLock ScopeLock(ConfigFileMapLock);
 		return FConfigFileMap::IsEmpty();
 	}
 
-	FORCEINLINE void Empty(int32 ExpectedNumElements = 0)
+	inline void Empty(int32 ExpectedNumElements = 0)
 	{
 		UE::TWriteScopeLock ScopeLock(ConfigFileMapLock);
 		FConfigFileMap::Empty(ExpectedNumElements);
 	}
 
-	FORCEINLINE bool Contains(const FString& SectionName) const
+	inline bool Contains(const FString& SectionName) const
 	{
 		UE::TReadScopeLock ScopeLock(ConfigFileMapLock);
 		return FConfigFileMap::Contains(SectionName);
 	}
 
-	FORCEINLINE int32 GetKeys(TArray<FString>& Keys) const
+	inline int32 GetKeys(TArray<FString>& Keys) const
 	{
 		UE::TReadScopeLock ScopeLock(ConfigFileMapLock);
 		return FConfigFileMap::GetKeys(Keys);
 	}
 
-	FORCEINLINE int32 GetKeys(TSet<FString>& Keys) const
+	inline int32 GetKeys(TSet<FString>& Keys) const
 	{
 		UE::TReadScopeLock ScopeLock(ConfigFileMapLock);
 		return FConfigFileMap::GetKeys(Keys);
 	}
 
-	FORCEINLINE int32 Remove(KeyConstPointerType InKey)
+	inline int32 Remove(KeyConstPointerType InKey)
 	{
 		UE::TWriteScopeLock ScopeLock(ConfigFileMapLock);
 		return FConfigFileMap::Remove(InKey);
 	}
 
-	FORCEINLINE ValueType& Add(const KeyType& InKey, const ValueType& InValue)
+	inline ValueType& Add(const KeyType& InKey, const ValueType& InValue)
 	{
 		UE::TWriteScopeLock ScopeLock(ConfigFileMapLock);
 		return FConfigFileMap::Add(InKey, InValue);
 	}
 
-	FORCEINLINE ValueType& Add(const KeyType&  InKey, ValueType&& InValue)
+	inline ValueType& Add(const KeyType&  InKey, ValueType&& InValue)
 	{
 		UE::TWriteScopeLock ScopeLock(ConfigFileMapLock);
 		return FConfigFileMap::Add(InKey, MoveTempIfPossible(InValue));
 	}
 
-	FORCEINLINE ValueType& Add(KeyType&& InKey, const ValueType&  InValue)
+	inline ValueType& Add(KeyType&& InKey, const ValueType&  InValue)
 	{
 		UE::TWriteScopeLock ScopeLock(ConfigFileMapLock);
 		return FConfigFileMap::Add(MoveTempIfPossible(InKey), InValue);
 	}
 
-	FORCEINLINE ValueType& Add(KeyType&& InKey,	ValueType&& InValue)
+	inline ValueType& Add(KeyType&& InKey,	ValueType&& InValue)
 	{
 		UE::TWriteScopeLock ScopeLock(ConfigFileMapLock);
 		return FConfigFileMap::Add(MoveTempIfPossible(InKey), MoveTempIfPossible(InValue));
 	}
 	
-	FORCEINLINE void Append(TMap<FString, FConfigSection> Other)
+	inline void Append(TMap<FString, FConfigSection> Other)
 	{
 		UE::TWriteScopeLock ScopeLock(ConfigFileMapLock);
 		FConfigFileMap::Append(MoveTemp(Other));
 	}
 
-	FORCEINLINE void Reset()
+	inline void Reset()
 	{
 		UE::TWriteScopeLock ScopeLock(ConfigFileMapLock);
 		FConfigFileMap::Reset();
 	}
 
-	FORCEINLINE TRangedForConstIterator begin() const
+	inline TRangedForConstIterator begin() const
 	{
 		UE::TReadScopeLock ScopeLock(ConfigFileMapLock); 
 		return TRangedForConstIterator(Pairs.begin());
 	}
 
-	FORCEINLINE TRangedForConstIterator end() const
+	inline TRangedForConstIterator end() const
 	{
 		UE::TReadScopeLock ScopeLock(ConfigFileMapLock);
 		return TRangedForConstIterator(Pairs.end());
@@ -786,7 +786,7 @@ private:
 
 	FConfigSection* FindOrAddSectionInternal(const FString& SectionName);
 
-	FORCEINLINE FConfigSection* FindInternal(const FString& SectionName)
+	inline FConfigSection* FindInternal(const FString& SectionName)
 	{
 		UE::TReadScopeLock ScopeLock(ConfigFileMapLock);
 		return FConfigFileMap::Find(SectionName);
diff --git a/Engine/Source/Runtime/Core/Public/Misc/CoreMisc.h b/Engine/Source/Runtime/Core/Public/Misc/CoreMisc.h
--- a/Engine/Source/Runtime/Core/Public/Misc/CoreMisc.h
+++ b/Engine/Source/Runtime/Core/Public/Misc/CoreMisc.h
@@ -149,7 +149,7 @@ CORE_API int32 StaticClientOnlyCheck(); // Internal, use IsRunningClientOnly()
  * An editor build can be launched with -server to set this to true, but it will be false during single process PlayInEditor mode.
  * This function should not be used for gameplay or networking purposes, check for NM_DedicatedServer instead.
  */
-FORCEINLINE bool IsRunningDedicatedServer()
+inline bool IsRunningDedicatedServer()
 {
 	if (FPlatformProperties::IsServerOnly())
 	{
@@ -174,7 +174,7 @@ FORCEINLINE bool IsRunningDedicatedServer()
  * An editor build can be launched with -game to set this to true, but it will be false during single process PlayInEditor mode.
  * This function should not be used for gameplay or networking purposes, check the NetMode instead.
  */
-FORCEINLINE bool IsRunningGame()
+inline bool IsRunningGame()
 {
 	if (FPlatformProperties::IsGameOnly())
 	{
@@ -199,7 +199,7 @@ FORCEINLINE bool IsRunningGame()
  * An editor build can be launched with -clientonly to set this to true, but it will be false during single process PlayInEditor mode.
  * This function should not be used for gameplay or networking purposes, check for NM_Client instead.
  */
-FORCEINLINE bool IsRunningClientOnly()
+inline bool IsRunningClientOnly()
 {
 	if (FPlatformProperties::IsClientOnly())
 	{
diff --git a/Engine/Source/Runtime/Core/Public/Misc/CoreMiscDefines.h b/Engine/Source/Runtime/Core/Public/Misc/CoreMiscDefines.h
--- a/Engine/Source/Runtime/Core/Public/Misc/CoreMiscDefines.h
+++ b/Engine/Source/Runtime/Core/Public/Misc/CoreMiscDefines.h
@@ -454,7 +454,7 @@ struct FPlatformUserId
 	}
 
 	/** Explicit function to create from an internal id */
-	FORCEINLINE static FPlatformUserId CreateFromInternalId(int32 InInternalId)
+	inline static FPlatformUserId CreateFromInternalId(int32 InInternalId)
 	{
 		FPlatformUserId IdToReturn;
 		IdToReturn.InternalId = InInternalId;
@@ -496,7 +496,7 @@ inline constexpr FPlatformUserId PLATFORMUSERID_NONE;
 struct FInputDeviceId
 {
 	/** Explicit function to create from an internal id */
-	FORCEINLINE static FInputDeviceId CreateFromInternalId(int32 InInternalId)
+	inline static FInputDeviceId CreateFromInternalId(int32 InInternalId)
 	{
 		FInputDeviceId IdToReturn;
 		IdToReturn.InternalId = InInternalId;
diff --git a/Engine/Source/Runtime/Core/Public/Misc/EnumRange.h b/Engine/Source/Runtime/Core/Public/Misc/EnumRange.h
--- a/Engine/Source/Runtime/Core/Public/Misc/EnumRange.h
+++ b/Engine/Source/Runtime/Core/Public/Misc/EnumRange.h
@@ -109,7 +109,7 @@ namespace NEnumRangePrivate
 		{
 		}
 
-		FORCEINLINE TEnumContiguousIterator& operator++()
+		inline TEnumContiguousIterator& operator++()
 		{
 			++Value;
 			return *this;
@@ -137,7 +137,7 @@ namespace NEnumRangePrivate
 		{
 		}
 
-		FORCEINLINE TEnumValueArrayIterator& operator++()
+		inline TEnumValueArrayIterator& operator++()
 		{
 			++Ptr;
 			return *this;
@@ -162,7 +162,7 @@ namespace NEnumRangePrivate
 	{
 	    using IntType = typename TEnumContiguousIterator<EnumType>::IntType;
 
-		FORCEINLINE explicit TEnumRange_Impl()
+		inline explicit TEnumRange_Impl()
 			: BeginValue(TEnumRangeTraits<EnumType>::Begin)
 			, EndValue(TEnumRangeTraits<EnumType>::End)
 		{
@@ -172,7 +172,7 @@ namespace NEnumRangePrivate
 		 * Note that this is a closed range, so it's not currently possible (or valid) to iterate over zero elements.
 		 * By specifying the same first and last value, it will iterate over that single value.
 		 */
-		FORCEINLINE explicit TEnumRange_Impl(const EnumType InFirstValue, const EnumType InLastValue)
+		inline explicit TEnumRange_Impl(const EnumType InFirstValue, const EnumType InLastValue)
 			: BeginValue(static_cast<IntType>(InFirstValue))
 			, EndValue(static_cast<IntType>(InLastValue) + 1)
 		{
@@ -215,14 +215,14 @@ namespace UE::EnumFlags::Private
 		{
 		}
 
-		FORCEINLINE TIterator& operator++()
+		inline TIterator& operator++()
 		{
 			const IntType PoppedBit = Flags & (~Flags + 1);
 			Flags ^= PoppedBit;
 			return *this;
 		}
 
-		FORCEINLINE EnumType operator*() const
+		inline EnumType operator*() const
 		{
 			const IntType Result = Flags & (~Flags + 1); 
 			return (EnumType)Result;
diff --git a/Engine/Source/Runtime/Core/Public/Misc/FrameTime.h b/Engine/Source/Runtime/Core/Public/Misc/FrameTime.h
--- a/Engine/Source/Runtime/Core/Public/Misc/FrameTime.h
+++ b/Engine/Source/Runtime/Core/Public/Misc/FrameTime.h
@@ -153,7 +153,7 @@ public:
 	}
 
 
-	friend FORCEINLINE_DEBUGGABLE FFrameTime& operator+=(FFrameTime& LHS, FFrameTime RHS)
+	friend inline FFrameTime& operator+=(FFrameTime& LHS, FFrameTime RHS)
 	{
 		float       NewSubFrame    = LHS.SubFrame + RHS.SubFrame;
 		const int64 NewFrameNumber = int64(LHS.FrameNumber.Value) + int64(RHS.FrameNumber.Value) + FMath::FloorToInt(NewSubFrame);
@@ -165,7 +165,7 @@ public:
 	}
 
 
-	friend FORCEINLINE_DEBUGGABLE FFrameTime operator+(FFrameTime A, FFrameTime B)
+	friend inline FFrameTime operator+(FFrameTime A, FFrameTime B)
 	{
 		const float        NewSubFrame    = A.SubFrame + B.SubFrame;
 		const int64        NewFrameNumber = int64(A.FrameNumber.Value) + int64(B.FrameNumber.Value) + FMath::FloorToInt(NewSubFrame);
@@ -174,7 +174,7 @@ public:
 	}
 
 
-	friend FORCEINLINE_DEBUGGABLE FFrameTime& operator-=(FFrameTime& LHS, FFrameTime RHS)
+	friend inline FFrameTime& operator-=(FFrameTime& LHS, FFrameTime RHS)
 	{
 		// Ensure SubFrame is always between 0 and 1
 		// Note that the difference between frame -1.5 and 1.5 is 2, not 3, since sub frame positions are always positive
@@ -189,7 +189,7 @@ public:
 	}
 
 
-	friend FORCEINLINE_DEBUGGABLE FFrameTime operator-(FFrameTime A, FFrameTime B)
+	friend inline FFrameTime operator-(FFrameTime A, FFrameTime B)
 	{
 		// Ensure SubFrame is always between 0 and 1
 		// Note that the difference between frame -1.5 and 1.5 is 2, not 3, since sub frame positions are always positive
@@ -201,7 +201,7 @@ public:
 	}
 
 
-	friend FORCEINLINE_DEBUGGABLE FFrameTime operator%(FFrameTime A, FFrameTime B)
+	friend inline FFrameTime operator%(FFrameTime A, FFrameTime B)
 	{
 		check(B.FrameNumber.Value != 0 || B.GetSubFrame() != 0.f);
 
@@ -221,7 +221,7 @@ public:
 	}
 
 
-	friend FORCEINLINE_DEBUGGABLE FFrameTime operator-(FFrameTime A)
+	friend inline FFrameTime operator-(FFrameTime A)
 	{
 		return A.GetSubFrame() == 0.f
 			? FFrameTime(-A.FrameNumber)
@@ -310,7 +310,7 @@ FORCEINLINE_DEBUGGABLE double FFrameTime::AsDecimal() const
 	return double(FrameNumber.Value) + SubFrame;
 }
 
-FORCEINLINE_DEBUGGABLE FFrameTime FFrameTime::FromDecimal(double InDecimalFrame)
+inline FFrameTime FFrameTime::FromDecimal(double InDecimalFrame)
 {
 	int32 NewFrame = static_cast<int32>(FMath::Clamp(FMath::FloorToDouble(InDecimalFrame), (double)TNumericLimits<int32>::Min(), (double)TNumericLimits<int32>::Max()));
 
diff --git a/Engine/Source/Runtime/Core/Public/Misc/HashBuilder.h b/Engine/Source/Runtime/Core/Public/Misc/HashBuilder.h
--- a/Engine/Source/Runtime/Core/Public/Misc/HashBuilder.h
+++ b/Engine/Source/Runtime/Core/Public/Misc/HashBuilder.h
@@ -21,7 +21,7 @@ public:
 		: Hash(~InHash)
 	{}
 
-	FORCEINLINE FHashBuilder& AppendRaw(const void* Data, int64 Num)
+	inline FHashBuilder& AppendRaw(const void* Data, int64 Num)
 	{
 		Hash = HashCombineFast(Hash, GetTypeHash(CityHash64(static_cast<const char*>(Data),  static_cast<int32>(Num))));
 		return *this;
@@ -34,7 +34,7 @@ public:
 	}
 
 	template <typename T>
-	FORCEINLINE FHashBuilder& Append(const T& InData)
+	inline FHashBuilder& Append(const T& InData)
 	{
 		if constexpr (TModels_V<CGetTypeHashable, T>)
 		{
@@ -48,7 +48,7 @@ public:
 	}
 
 	template <typename T>
-	FORCEINLINE FHashBuilder& Append(const TArray<T>& InArray)
+	inline FHashBuilder& Append(const TArray<T>& InArray)
 	{
 		for (const T& Value: InArray)
 		{
@@ -58,7 +58,7 @@ public:
 	}
 
 	template <typename T>
-	FORCEINLINE FHashBuilder& Append(const TSet<T>& InSet)
+	inline FHashBuilder& Append(const TSet<T>& InSet)
 	{
 		for (const T& Value: InSet)
 		{
@@ -68,7 +68,7 @@ public:
 	}
 
 	template <typename T, typename U>
-	FORCEINLINE FHashBuilder& Append(const TMap<T, U>& InMap)
+	inline FHashBuilder& Append(const TMap<T, U>& InMap)
 	{
 		for (const TPair<T, U>& Value: InMap)
 		{
diff --git a/Engine/Source/Runtime/Core/Public/Misc/IEngineCrypto.h b/Engine/Source/Runtime/Core/Public/Misc/IEngineCrypto.h
--- a/Engine/Source/Runtime/Core/Public/Misc/IEngineCrypto.h
+++ b/Engine/Source/Runtime/Core/Public/Misc/IEngineCrypto.h
@@ -14,7 +14,7 @@ struct IEngineCrypto : public IModularFeature
 	/**
 	* Get the name of this modular feature
 	*/
-	static FORCEINLINE FName GetFeatureName()
+	static inline FName GetFeatureName()
 	{
 		static const FName Name(TEXT("EngineCryptoFeature"));
 		return Name;
diff --git a/Engine/Source/Runtime/Core/Public/Misc/InlineValue.h b/Engine/Source/Runtime/Core/Public/Misc/InlineValue.h
--- a/Engine/Source/Runtime/Core/Public/Misc/InlineValue.h
+++ b/Engine/Source/Runtime/Core/Public/Misc/InlineValue.h
@@ -127,7 +127,7 @@ public:
 	 * Emplace a new type (deriving from BaseType) into this inline value
 	 */
 	template <typename T, typename... ArgsType>
-	FORCEINLINE void Emplace(ArgsType&&... Args)
+	inline void Emplace(ArgsType&&... Args)
 	{
 		static_assert(TPointerIsConvertibleFromTo<T, BaseType>::Value, "T must derive from BaseType.");
 		Reset();
@@ -146,7 +146,7 @@ public:
 	 * Access the wrapped object's base type
 	 * @return A reference to the object. Will assert where IsValid() is false.
 	 */
-	FORCEINLINE BaseType& GetValue() const
+	inline BaseType& GetValue() const
 	{
 		checkf(bIsValid, TEXT("It is an error to call GetValue() on an invalid TInlineValue. Please either check IsValid() or use Get(DefaultValue) instead."));
 		return bInline ? (BaseType&)Data : **((BaseType**)&Data);
diff --git a/Engine/Source/Runtime/Core/Public/Misc/MTAccessDetector.h b/Engine/Source/Runtime/Core/Public/Misc/MTAccessDetector.h
--- a/Engine/Source/Runtime/Core/Public/Misc/MTAccessDetector.h
+++ b/Engine/Source/Runtime/Core/Public/Misc/MTAccessDetector.h
@@ -86,7 +86,7 @@ public:
 	 * Acquires read access, will check if there are any writers
 	 * @return true if no errors were detected
 	 */
-	FORCEINLINE bool AcquireReadAccess() const
+	inline bool AcquireReadAccess() const
 	{
 		const bool ErrorDetected = (AtomicValue.fetch_add(1, std::memory_order_relaxed) & WriterBits) != 0;
 		checkf(!ErrorDetected || UE::Private::FMTAccessDetectorOptions::bSuppressCheckFailure, TEXT("Acquiring a read access while there is already a write access"));
@@ -97,7 +97,7 @@ public:
 	 * Releases read access, will check if there are any writers
 	 * @return true if no errors were detected
 	 */
-	FORCEINLINE bool ReleaseReadAccess() const
+	inline bool ReleaseReadAccess() const
 	{
 		const bool ErrorDetected = (AtomicValue.fetch_sub(1, std::memory_order_relaxed) & WriterBits) != 0;
 		checkf(!ErrorDetected || UE::Private::FMTAccessDetectorOptions::bSuppressCheckFailure, TEXT("Another thread asked to have a write access during this read access"));
@@ -108,7 +108,7 @@ public:
 	 * Acquires write access, will check if there are readers or other writers
 	 * @return true if no errors were detected
 	 */
-	FORCEINLINE bool AcquireWriteAccess() const
+	inline bool AcquireWriteAccess() const
 	{
 		const bool ErrorDetected = AtomicValue.fetch_add(WriterIncrementValue, std::memory_order_relaxed) != 0;
 		checkf(!ErrorDetected || UE::Private::FMTAccessDetectorOptions::bSuppressCheckFailure, TEXT("Acquiring a write access while there are ongoing read or write access"));
@@ -119,7 +119,7 @@ public:
 	 * Releases write access, will check if there are readers or other writers
 	 * @return true if no errors were detected
 	 */
-	FORCEINLINE bool ReleaseWriteAccess() const
+	inline bool ReleaseWriteAccess() const
 	{
 		const bool ErrorDetected = AtomicValue.fetch_sub(WriterIncrementValue, std::memory_order_relaxed) != WriterIncrementValue;
 		checkf(!ErrorDetected || UE::Private::FMTAccessDetectorOptions::bSuppressCheckFailure, TEXT("Another thread asked to have a read or write access during this write access"));
@@ -130,7 +130,7 @@ public:
 	 * Upgrades a previously taken `AcquireReadAccess` to `AcquireWriteAccess`, and will check if there are other readers or writers
 	 * @return true if no errors were detected
 	 */
-	FORCEINLINE bool UpgradeReadAccessToWriteAccess() const
+	inline bool UpgradeReadAccessToWriteAccess() const
 	{
 		const bool bErrorDetected = AtomicValue.fetch_add(WriterIncrementValue, std::memory_order_relaxed) != 1;
 
@@ -166,7 +166,7 @@ public:
 	 * Acquires write access, will check if there are readers or other writers
 	 * @return true if no errors were detected
 	 */
-	FORCEINLINE bool AcquireWriteAccess() const
+	inline bool AcquireWriteAccess() const
 	{
 		uint32 CurThreadID = FPlatformTLS::GetCurrentThreadId();
 
@@ -189,7 +189,7 @@ public:
 	 * Releases write access, will check if there are readers or other writers
 	 * @return true if no errors were detected
 	 */
-	FORCEINLINE bool ReleaseWriteAccess() const
+	inline bool ReleaseWriteAccess() const
 	{
 		uint32 CurThreadID = FPlatformTLS::GetCurrentThreadId();
 		if (WriterThreadID == CurThreadID)
@@ -227,7 +227,7 @@ public:
 	 * Acquires read access, will check if there are any writers
 	 * @return true if no errors were detected
 	 */
-	FORCEINLINE bool AcquireReadAccess() const
+	inline bool AcquireReadAccess() const
 	{
 		uint32 CurThreadID = FPlatformTLS::GetCurrentThreadId();
 		if (WriterThreadID == CurThreadID)
@@ -241,7 +241,7 @@ public:
 	 * Releases read access, will check if there are any writers
 	 * @return true if no errors were detected
 	 */
-	FORCEINLINE bool ReleaseReadAccess() const
+	inline bool ReleaseReadAccess() const
 	{
 		uint32 CurThreadID = FPlatformTLS::GetCurrentThreadId();
 		if (WriterThreadID == CurThreadID)
@@ -261,7 +261,7 @@ struct TScopedReaderAccessDetector : public FBaseScopedAccessDetector
 {
 public:
 
-	FORCEINLINE TScopedReaderAccessDetector(RWAccessDetector& InAccessDetector)
+	inline TScopedReaderAccessDetector(RWAccessDetector& InAccessDetector)
 	: AccessDetector(InAccessDetector)
 	{
 		AccessDetector.AcquireReadAccess();
@@ -285,7 +285,7 @@ template<typename RWAccessDetector>
 struct TScopedWriterDetector : public FBaseScopedAccessDetector
 {
 public:
-	FORCEINLINE TScopedWriterDetector(RWAccessDetector& InAccessDetector)
+	inline TScopedWriterDetector(RWAccessDetector& InAccessDetector)
 		: AccessDetector(InAccessDetector)
 	{
 		AccessDetector.AcquireWriteAccess();
@@ -369,13 +369,13 @@ private:
 		return FState(State.exchange(NewState.Value, std::memory_order_relaxed));
 	}
 
-	FORCEINLINE FState IncrementReaderNum() const
+	inline FState IncrementReaderNum() const
 	{
 		constexpr FState OneReader{ 1, 0, 0 };
 		return FState(State.fetch_add(OneReader.Value, std::memory_order_relaxed));
 	}
 
-	FORCEINLINE FState DecrementReaderNum() const
+	inline FState DecrementReaderNum() const
 	{
 		constexpr FState OneReader{ 1, 0, 0 };
 		return FState(State.fetch_sub(OneReader.Value, std::memory_order_relaxed));
@@ -386,7 +386,7 @@ private:
 		return L.Value == R.Value;
 	}
 
-	FORCEINLINE static bool CheckOtherThreadWriters(FState InState)
+	inline static bool CheckOtherThreadWriters(FState InState)
 	{
 		if (InState.WriterNum != 0)
 		{
@@ -507,7 +507,7 @@ public:
 		CheckOtherThreadWriters(Other.LoadState());
 	}
 
-	FORCEINLINE FMRSWRecursiveAccessDetector& operator=(const FMRSWRecursiveAccessDetector& Other)
+	inline FMRSWRecursiveAccessDetector& operator=(const FMRSWRecursiveAccessDetector& Other)
 		// do not alter the state, it can be accessed
 	{
 		CheckOtherThreadWriters(Other.LoadState());
@@ -516,7 +516,7 @@ public:
 
 	// use copy construction/assignment ^^ for move semantics too
 
-	FORCEINLINE ~FMRSWRecursiveAccessDetector()
+	inline ~FMRSWRecursiveAccessDetector()
 	{
 		// search for all destruction sentinels for this instance and remove them from the stack, while building an expected correct state of
 		// the access detector.
@@ -555,7 +555,7 @@ public:
 		}
 	}
 
-	FORCEINLINE bool AcquireReadAccess() const
+	inline bool AcquireReadAccess() const
 	{
 		FState PrevState = IncrementReaderNum();
 
@@ -579,7 +579,7 @@ public:
 
 	// an overload that handles access detector destruction from inside a read access, must be used along with the corresponding
 	// overload of `ReleaseReadAccess`
-	FORCEINLINE bool AcquireReadAccess(FDestructionSentinel& DestructionSentinel) const
+	inline bool AcquireReadAccess(FDestructionSentinel& DestructionSentinel) const
 	{
 		DestructionSentinel.Accessor = this;
 		GetDestructionSentinelStackTls().Add(&DestructionSentinel);
@@ -587,7 +587,7 @@ public:
 		return AcquireReadAccess();
 	}
 
-	FORCEINLINE bool ReleaseReadAccess() const
+	inline bool ReleaseReadAccess() const
 	{
 		RemoveReaderFromTls();
 		DecrementReaderNum();
@@ -597,7 +597,7 @@ public:
 
 	// an overload that handles access detector destruction from inside a read access, must be used along with the corresponding
 	// overload of `AcquireReadAccess`
-	FORCEINLINE bool ReleaseReadAccess(FDestructionSentinel& DestructionSentinel) const
+	inline bool ReleaseReadAccess(FDestructionSentinel& DestructionSentinel) const
 	{
 		bool bSuccess = ReleaseReadAccess();
 
@@ -610,7 +610,7 @@ public:
 		return bSuccess;
 	}
 
-	FORCEINLINE bool AcquireWriteAccess()
+	inline bool AcquireWriteAccess()
 	{
 		FState LocalState = LoadState();
 		bool bSuccess = true;
@@ -675,7 +675,7 @@ public:
 
 	// an overload that handles access detector destruction from inside a write access, must be used along with the corresponding
 	// overload of `ReleaseWriteAccess`
-	FORCEINLINE bool AcquireWriteAccess(FDestructionSentinel& DestructionSentinel)
+	inline bool AcquireWriteAccess(FDestructionSentinel& DestructionSentinel)
 	{
 		DestructionSentinel.Accessor = this;
 		GetDestructionSentinelStackTls().Add(&DestructionSentinel);
@@ -683,7 +683,7 @@ public:
 		return AcquireWriteAccess();
 	}
 
-	FORCEINLINE bool ReleaseWriteAccess()
+	inline bool ReleaseWriteAccess()
 	{
 		FState LocalState = LoadState();
 
@@ -710,7 +710,7 @@ public:
 
 	// an overload that handles access detector destruction from inside a write access, must be used along with the corresponding
 	// overload of `AcquireWriteAccess`
-	FORCEINLINE bool ReleaseWriteAccess(FDestructionSentinel& DestructionSentinel)
+	inline bool ReleaseWriteAccess(FDestructionSentinel& DestructionSentinel)
 	{
 		bool bSuccess = ReleaseWriteAccess();
 
diff --git a/Engine/Source/Runtime/Core/Public/Misc/MTTransactionallySafeAccessDetector.h b/Engine/Source/Runtime/Core/Public/Misc/MTTransactionallySafeAccessDetector.h
--- a/Engine/Source/Runtime/Core/Public/Misc/MTTransactionallySafeAccessDetector.h
+++ b/Engine/Source/Runtime/Core/Public/Misc/MTTransactionallySafeAccessDetector.h
@@ -45,7 +45,7 @@ private:
 	 * Acquires read access, will check if there are any writers
 	 * @return true if no errors were detected
 	 */
-	FORCEINLINE bool AcquireReadAccess() const
+	inline bool AcquireReadAccess() const
 	{
 		if (AutoRTFM::IsTransactional() || AutoRTFM::IsCommittingOrAborting())
 		{
@@ -89,7 +89,7 @@ private:
 	 * Releases read access, will check if there are any writers
 	 * @return true if no errors were detected
 	 */
-	FORCEINLINE bool ReleaseReadAccess() const
+	inline bool ReleaseReadAccess() const
 	{
 		if (AutoRTFM::IsTransactional() || AutoRTFM::IsCommittingOrAborting())
 		{
@@ -110,7 +110,7 @@ private:
 	 * Acquires write access, will check if there are readers or other writers
 	 * @return true if no errors were detected
 	 */
-	FORCEINLINE bool AcquireWriteAccess() const
+	inline bool AcquireWriteAccess() const
 	{
 		if (AutoRTFM::IsTransactional() || AutoRTFM::IsCommittingOrAborting())
 		{
@@ -163,7 +163,7 @@ private:
 	 * Releases write access, will check if there are readers or other writers
 	 * @return true if no errors were detected
 	 */
-	FORCEINLINE bool ReleaseWriteAccess() const
+	inline bool ReleaseWriteAccess() const
 	{
 		if (AutoRTFM::IsTransactional() || AutoRTFM::IsCommittingOrAborting())
 		{
diff --git a/Engine/Source/Runtime/Core/Public/Misc/MemStack.h b/Engine/Source/Runtime/Core/Public/Misc/MemStack.h
--- a/Engine/Source/Runtime/Core/Public/Misc/MemStack.h
+++ b/Engine/Source/Runtime/Core/Public/Misc/MemStack.h
@@ -122,14 +122,14 @@ public:
 		return (uint8*)Alloc(AllocSize, FMath::Max(AllocSize >= 16 ? (size_t)16 : (size_t)8, Alignment));
 	}
 
-	FORCEINLINE bool CanFitInPage(size_t AllocSize, size_t Alignment) const
+	inline bool CanFitInPage(size_t AllocSize, size_t Alignment) const
 	{
 		const uint8* Result = Align(Top, Alignment);
 		const uint8* NewTop = Result + AllocSize;
 		return NewTop <= End;
 	}
 
-	FORCEINLINE void* Alloc(size_t AllocSize, size_t Alignment)
+	inline void* Alloc(size_t AllocSize, size_t Alignment)
 	{
 		// Debug checks.
 		checkSlow(AllocSize>=0);
@@ -168,7 +168,7 @@ public:
 		return TopChunk == nullptr;
 	}
 
-	FORCEINLINE void Flush()
+	inline void Flush()
 	{
 		check(!NumMarks);
 		FreeChunks(nullptr);
@@ -410,7 +410,7 @@ public:
 		 * Assumes that the allocator is currently empty, i.e. memory may be allocated but any existing elements have already been destructed (if necessary).
 		 * @param Other - The allocator to move the state from.  This allocator should be left in a valid empty state.
 		 */
-		FORCEINLINE void MoveToEmpty(ForElementType& Other)
+		inline void MoveToEmpty(ForElementType& Other)
 		{
 			checkSlow(this != &Other);
 
diff --git a/Engine/Source/Runtime/Core/Public/Misc/MemoryReadStream.h b/Engine/Source/Runtime/Core/Public/Misc/MemoryReadStream.h
--- a/Engine/Source/Runtime/Core/Public/Misc/MemoryReadStream.h
+++ b/Engine/Source/Runtime/Core/Public/Misc/MemoryReadStream.h
@@ -24,7 +24,7 @@ public:
 
 	FORCEINLINE uint32 AddRef() const { return uint32(NumRefs.Increment()); }
 
-	FORCEINLINE uint32 Release() const
+	inline uint32 Release() const
 	{
 		const int32 Refs = NumRefs.Decrement();
 		if (Refs == 0)
diff --git a/Engine/Source/Runtime/Core/Public/Misc/Optional.h b/Engine/Source/Runtime/Core/Public/Misc/Optional.h
--- a/Engine/Source/Runtime/Core/Public/Misc/Optional.h
+++ b/Engine/Source/Runtime/Core/Public/Misc/Optional.h
@@ -23,7 +23,7 @@ namespace UE::Core::Private
 	struct FOptional
 	{
 		template<typename Derived>
-		[[nodiscard]] FORCEINLINE static constexpr bool IsSet(Derived* This)
+		[[nodiscard]] inline static constexpr bool IsSet(Derived* This)
 		{
 			if constexpr (Derived::bUsingIntrusiveUnsetState)
 			{
diff --git a/Engine/Source/Runtime/Core/Public/Misc/OutputDevice.h b/Engine/Source/Runtime/Core/Public/Misc/OutputDevice.h
--- a/Engine/Source/Runtime/Core/Public/Misc/OutputDevice.h
+++ b/Engine/Source/Runtime/Core/Public/Misc/OutputDevice.h
@@ -238,7 +238,7 @@ public:
 	}
 
 	template <typename FmtType, typename... Types>
-	FORCEINLINE void Logf(const FmtType& Fmt, Types... Args)
+	inline void Logf(const FmtType& Fmt, Types... Args)
 	{
 		static_assert(TIsArrayOrRefOfTypeByPredicate<FmtType, TIsCharEncodingCompatibleWithTCHAR>::Value, "Formatting string must be a TCHAR array.");
 		static_assert((TIsValidVariadicFunctionArg<Types>::Value && ...), "Invalid argument(s) passed to FOutputDevice::Logf");
@@ -247,7 +247,7 @@ public:
 	}
 
 	template <typename FmtType, typename... Types>
-	FORCEINLINE void Logf(ELogVerbosity::Type Verbosity, const FmtType& Fmt, Types... Args)
+	inline void Logf(ELogVerbosity::Type Verbosity, const FmtType& Fmt, Types... Args)
 	{
 		static_assert(TIsArrayOrRefOfTypeByPredicate<FmtType, TIsCharEncodingCompatibleWithTCHAR>::Value, "Formatting string must be a TCHAR array.");
 		static_assert((TIsValidVariadicFunctionArg<Types>::Value && ...), "Invalid argument(s) passed to FOutputDevice::Logf");
@@ -256,7 +256,7 @@ public:
 	}
 
 	template <typename FmtType, typename... Types>
-	FORCEINLINE void CategorizedLogf(const FName& Category, ELogVerbosity::Type Verbosity, const FmtType& Fmt, Types... Args)
+	inline void CategorizedLogf(const FName& Category, ELogVerbosity::Type Verbosity, const FmtType& Fmt, Types... Args)
 	{
 		static_assert(TIsArrayOrRefOfTypeByPredicate<FmtType, TIsCharEncodingCompatibleWithTCHAR>::Value, "Formatting string must be a TCHAR array.");
 		static_assert((TIsValidVariadicFunctionArg<Types>::Value && ...), "Invalid argument(s) passed to FOutputDevice::CategorizedLogf");
diff --git a/Engine/Source/Runtime/Core/Public/Misc/Parse.h b/Engine/Source/Runtime/Core/Public/Misc/Parse.h
--- a/Engine/Source/Runtime/Core/Public/Misc/Parse.h
+++ b/Engine/Source/Runtime/Core/Public/Misc/Parse.h
@@ -149,7 +149,7 @@ struct FParse
 	//
 	// Parse a hex digit.
 	//
-	static FORCEINLINE int32 HexDigit(TCHAR c)
+	static inline int32 HexDigit(TCHAR c)
 	{
 		int32 Result = 0;
 
diff --git a/Engine/Source/Runtime/Core/Public/Misc/Paths.h b/Engine/Source/Runtime/Core/Public/Misc/Paths.h
--- a/Engine/Source/Runtime/Core/Public/Misc/Paths.h
+++ b/Engine/Source/Runtime/Core/Public/Misc/Paths.h
@@ -721,7 +721,7 @@ private:
 	struct FStaticData;
 
 	template <typename... PathTypes>
-	FORCEINLINE static FString CombineImpl(UE::Core::Private::TToStringType_T<std::decay_t<PathTypes>>... InPaths)
+	inline static FString CombineImpl(UE::Core::Private::TToStringType_T<std::decay_t<PathTypes>>... InPaths)
 	{
 		const FStringView Paths[] = { FStringView(InPaths)... };
 		return CombineInternal(Paths, UE_ARRAY_COUNT(Paths));
diff --git a/Engine/Source/Runtime/Core/Public/Misc/ScopedSlowTask.h b/Engine/Source/Runtime/Core/Public/Misc/ScopedSlowTask.h
--- a/Engine/Source/Runtime/Core/Public/Misc/ScopedSlowTask.h
+++ b/Engine/Source/Runtime/Core/Public/Misc/ScopedSlowTask.h
@@ -38,7 +38,7 @@ struct FScopedSlowTask : FSlowTask
 	 * @param		InDefaultMessage		A message to display to the user to describe the purpose of the scope
 	 * @param		bInEnabled				When false, this scope will have no effect. Allows for proper scoped objects that are conditionally disabled.
 	 */
-	FORCEINLINE FScopedSlowTask(float InAmountOfWork, const FText& InDefaultMessage = FText(), bool bInEnabled = true, FFeedbackContext& InContext = *GWarn)
+	inline FScopedSlowTask(float InAmountOfWork, const FText& InDefaultMessage = FText(), bool bInEnabled = true, FFeedbackContext& InContext = *GWarn)
 		: FSlowTask(InAmountOfWork, InDefaultMessage, bInEnabled, InContext)
 	{
 		Initialize();
diff --git a/Engine/Source/Runtime/Core/Public/Misc/StringOutputDevice.h b/Engine/Source/Runtime/Core/Public/Misc/StringOutputDevice.h
--- a/Engine/Source/Runtime/Core/Public/Misc/StringOutputDevice.h
+++ b/Engine/Source/Runtime/Core/Public/Misc/StringOutputDevice.h
@@ -117,14 +117,14 @@ public:
 	FStringOutputDeviceCountLines(const FStringOutputDeviceCountLines&) = default;
 	FStringOutputDeviceCountLines& operator=(const FStringOutputDeviceCountLines&) = default;
 
-	FORCEINLINE FStringOutputDeviceCountLines(FStringOutputDeviceCountLines&& Other)
+	inline FStringOutputDeviceCountLines(FStringOutputDeviceCountLines&& Other)
 		: Super    ((Super&&)Other)
 		, LineCount(Other.LineCount)
 	{
 		Other.LineCount = 0;
 	}
 
-	FORCEINLINE FStringOutputDeviceCountLines& operator=(FStringOutputDeviceCountLines&& Other)
+	inline FStringOutputDeviceCountLines& operator=(FStringOutputDeviceCountLines&& Other)
 	{
 		if (this != &Other)
 		{
diff --git a/Engine/Source/Runtime/Core/Public/Misc/TextFilterUtils.h b/Engine/Source/Runtime/Core/Public/Misc/TextFilterUtils.h
--- a/Engine/Source/Runtime/Core/Public/Misc/TextFilterUtils.h
+++ b/Engine/Source/Runtime/Core/Public/Misc/TextFilterUtils.h
@@ -174,7 +174,7 @@ namespace TextFilterUtils
 			Init(Name.GetDisplayNameEntry(), Name.GetNumber());
 		}
 
-		FORCEINLINE void Init(const FNameEntry* const NameEntry, int32 NumberInternal)
+		inline void Init(const FNameEntry* const NameEntry, int32 NumberInternal)
 		{
 			if (NameEntry == nullptr)
 			{
@@ -212,13 +212,13 @@ namespace TextFilterUtils
 			return bIsWide;
 		}
 
-		FORCEINLINE ANSICHAR* GetAnsiNamePtr()
+		inline ANSICHAR* GetAnsiNamePtr()
 		{
 			checkSlow(!IsWide());
 			return AnsiName;
 		}
 
-		FORCEINLINE WIDECHAR* GetWideNamePtr()
+		inline WIDECHAR* GetWideNamePtr()
 		{
 			checkSlow(IsWide());
 			return WideName;
diff --git a/Engine/Source/Runtime/Core/Public/Misc/TimeGuard.h b/Engine/Source/Runtime/Core/Public/Misc/TimeGuard.h
--- a/Engine/Source/Runtime/Core/Public/Misc/TimeGuard.h
+++ b/Engine/Source/Runtime/Core/Public/Misc/TimeGuard.h
@@ -49,7 +49,7 @@ public:
 public:
 
 
-	FORCEINLINE FTimeGuard(FTimerNameDelegate InNameDelegate, const float InTargetMS = 0.0)
+	inline FTimeGuard(FTimerNameDelegate InNameDelegate, const float InTargetMS = 0.0)
 		: Name(nullptr), ObjectName(NAME_None)
 	{
 		NameDelegate = (bEnabled && IsInGameThread()) ? InNameDelegate : nullptr;
@@ -61,7 +61,7 @@ public:
 		}
 	}
 
-	FORCEINLINE FTimeGuard(TCHAR const* InName, FName InObjectName = NAME_None, const float InTargetMS = 0.0)
+	inline FTimeGuard(TCHAR const* InName, FName InObjectName = NAME_None, const float InTargetMS = 0.0)
 		: Name(nullptr), ObjectName(NAME_None)
 	{
 		Name = (bEnabled && IsInGameThread()) ? InName : nullptr;
@@ -77,7 +77,7 @@ public:
 	/**
 	* Updates the stat with the time spent
 	*/
-	FORCEINLINE ~FTimeGuard()
+	inline ~FTimeGuard()
 	{
 		if (Name)
 		{
diff --git a/Engine/Source/Runtime/Core/Public/Misc/Variant.h b/Engine/Source/Runtime/Core/Public/Misc/Variant.h
--- a/Engine/Source/Runtime/Core/Public/Misc/Variant.h
+++ b/Engine/Source/Runtime/Core/Public/Misc/Variant.h
@@ -400,7 +400,7 @@ private:
  * @see GetBytes
  */
 template<>
-FORCEINLINE TArray<uint8> FVariant::GetValue<TArray<uint8> >() const
+inline TArray<uint8> FVariant::GetValue<TArray<uint8> >() const
 {
 	check(Type == EVariantTypes::ByteArray);
 
diff --git a/Engine/Source/Runtime/Core/Public/PixelFormat.h b/Engine/Source/Runtime/Core/Public/PixelFormat.h
--- a/Engine/Source/Runtime/Core/Public/PixelFormat.h
+++ b/Engine/Source/Runtime/Core/Public/PixelFormat.h
@@ -282,7 +282,7 @@ inline bool IsETCBlockCompressedPixelFormat(EPixelFormat PixelFormat)
 	return false;
 }
 
-FORCEINLINE bool IsASTCBlockCompressedTextureFormat(EPixelFormat PixelFormat)
+inline bool IsASTCBlockCompressedTextureFormat(EPixelFormat PixelFormat)
 {
 	PRAGMA_DISABLE_SWITCH_UNHANDLED_ENUM_CASE_WARNINGS;
 	switch (PixelFormat)
@@ -334,13 +334,13 @@ FORCEINLINE bool IsAnyBlockCompressedPixelFormat(EPixelFormat PixelFormat)
 
 // unclear what IsHDR is supposed to mean
 //	see also IsFloatFormat
-FORCEINLINE bool IsHDR(EPixelFormat PixelFormat)
+inline bool IsHDR(EPixelFormat PixelFormat)
 {
 	return PixelFormat == PF_FloatRGBA || PixelFormat == PF_BC6H || PixelFormat == PF_R16F || PixelFormat == PF_R32_FLOAT || PixelFormat == PF_A32B32G32R32F
 		|| PixelFormat == PF_ASTC_4x4_HDR || PixelFormat == PF_ASTC_6x6_HDR || PixelFormat == PF_ASTC_8x8_HDR || PixelFormat == PF_ASTC_10x10_HDR || PixelFormat == PF_ASTC_12x12_HDR;
 }
 
-FORCEINLINE bool IsInteger(EPixelFormat PixelFormat)
+inline bool IsInteger(EPixelFormat PixelFormat)
 {
 	PRAGMA_DISABLE_SWITCH_UNHANDLED_ENUM_CASE_WARNINGS;
 	switch (PixelFormat)
@@ -390,7 +390,7 @@ inline bool IsFloatFormat(EPixelFormat Format)
 	return false;
 }
 
-FORCEINLINE bool IsDepthOrStencilFormat(EPixelFormat Format)
+inline bool IsDepthOrStencilFormat(EPixelFormat Format)
 {
 	PRAGMA_DISABLE_SWITCH_UNHANDLED_ENUM_CASE_WARNINGS;
 	switch (Format)
@@ -406,7 +406,7 @@ FORCEINLINE bool IsDepthOrStencilFormat(EPixelFormat Format)
 	return false;
 }
 
-FORCEINLINE bool IsStencilFormat(EPixelFormat Format)
+inline bool IsStencilFormat(EPixelFormat Format)
 {
 	PRAGMA_DISABLE_SWITCH_UNHANDLED_ENUM_CASE_WARNINGS;
 	switch (Format)
@@ -419,7 +419,7 @@ FORCEINLINE bool IsStencilFormat(EPixelFormat Format)
 	return false;
 }
 
-FORCEINLINE bool IsDXTCBlockCompressedTextureFormat(EPixelFormat PixelFormat)
+inline bool IsDXTCBlockCompressedTextureFormat(EPixelFormat PixelFormat)
 {
 	PRAGMA_DISABLE_SWITCH_UNHANDLED_ENUM_CASE_WARNINGS;
 	switch (PixelFormat)
diff --git a/Engine/Source/Runtime/Core/Public/ProfilingDebugging/CsvProfiler.h b/Engine/Source/Runtime/Core/Public/ProfilingDebugging/CsvProfiler.h
--- a/Engine/Source/Runtime/Core/Public/ProfilingDebugging/CsvProfiler.h
+++ b/Engine/Source/Runtime/Core/Public/ProfilingDebugging/CsvProfiler.h
@@ -407,7 +407,7 @@ public:
 	static CORE_API int32 GetCategoryIndex(const FString& Name);
 
 	template <typename FmtType, typename... Types>
-	FORCEINLINE static void RecordEventf(int32 CategoryIndex, const FmtType& Fmt, Types... Args)
+	inline static void RecordEventf(int32 CategoryIndex, const FmtType& Fmt, Types... Args)
 	{
 		UE_AUTORTFM_OPEN
 		{
diff --git a/Engine/Source/Runtime/Core/Public/ProfilingDebugging/ExternalProfiler.h b/Engine/Source/Runtime/Core/Public/ProfilingDebugging/ExternalProfiler.h
--- a/Engine/Source/Runtime/Core/Public/ProfilingDebugging/ExternalProfiler.h
+++ b/Engine/Source/Runtime/Core/Public/ProfilingDebugging/ExternalProfiler.h
@@ -206,7 +206,7 @@ class FExternalProfilerTrace
 {
 public:
 	/** Starts a scoped event specific to the profiler. */
-	FORCEINLINE static void StartScopedEvent(const struct FColor& Color, const TCHAR* Text)
+	inline static void StartScopedEvent(const struct FColor& Color, const TCHAR* Text)
 	{
 		if (FActiveExternalProfilerBase::IsActiveProfilerRecording())
 		{
@@ -218,7 +218,7 @@ public:
 	}
 
 	/** Starts a scoped event specific to the profiler. */
-	FORCEINLINE static void StartScopedEvent(const struct FColor& Color, const ANSICHAR* Text)
+	inline static void StartScopedEvent(const struct FColor& Color, const ANSICHAR* Text)
 	{
 		if (FActiveExternalProfilerBase::IsActiveProfilerRecording())
 		{
@@ -230,7 +230,7 @@ public:
 	}
 
 	/** Ends a scoped event specific to the profiler. */
-	FORCEINLINE static void EndScopedEvent()
+	inline static void EndScopedEvent()
 	{
 		if (FActiveExternalProfilerBase::IsActiveProfilerRecording())
 		{
diff --git a/Engine/Source/Runtime/Core/Public/ProfilingDebugging/Histogram.h b/Engine/Source/Runtime/Core/Public/ProfilingDebugging/Histogram.h
--- a/Engine/Source/Runtime/Core/Public/ProfilingDebugging/Histogram.h
+++ b/Engine/Source/Runtime/Core/Public/ProfilingDebugging/Histogram.h
@@ -145,7 +145,7 @@ struct FHistogram
 		return MaximalMeasurement;
 	}
 
-	FORCEINLINE FHistogram operator-(const FHistogram& Other) const
+	inline FHistogram operator-(const FHistogram& Other) const
 	{
 		// Logic below expects the bins to be of equal size
 		check(GetNumBins() == Other.GetNumBins());
@@ -158,7 +158,7 @@ struct FHistogram
 		return Tmp;
 	}
 
-	FORCEINLINE FHistogram operator+(const FHistogram& Other) const
+	inline FHistogram operator+(const FHistogram& Other) const
 	{
 		// Logic below expects the bins to be of equal size
 		check(GetNumBins() == Other.GetNumBins());
@@ -171,7 +171,7 @@ struct FHistogram
 		return Tmp;
 	}
 
-	FORCEINLINE FHistogram& operator+=(const FHistogram& Other)
+	inline FHistogram& operator+=(const FHistogram& Other)
 	{
 		// Logic below expects the bins to be of equal size
 		check(GetNumBins() == Other.GetNumBins());
@@ -245,7 +245,7 @@ protected:
 		{
 			return FBin(MinValue, UpperBound, Sum + Other.Sum, Count + Other.Count);
 		}
-		FORCEINLINE FBin& operator+=(const FBin& Other)
+		inline FBin& operator+=(const FBin& Other)
 		{
 			Sum += Other.Sum;
 			Count += Other.Count;
diff --git a/Engine/Source/Runtime/Core/Public/Serialization/Archive.h b/Engine/Source/Runtime/Core/Public/Serialization/Archive.h
--- a/Engine/Source/Runtime/Core/Public/Serialization/Archive.h
+++ b/Engine/Source/Runtime/Core/Public/Serialization/Archive.h
@@ -166,7 +166,7 @@ public:
 	}
 
 	/** Returns true if data larger than 1 byte should be swapped to deal with endian mismatches. */
-	FORCEINLINE bool IsByteSwapping()
+	inline bool IsByteSwapping()
 	{
 	#if PLATFORM_LITTLE_ENDIAN
 		bool SwapBytes = ArForceByteSwapping;
@@ -251,7 +251,7 @@ public:
 	}
 
 	/** Returns true if this archive is transacting, which is used to keep track of changes to objects for things like the editor undo system. */
-	FORCEINLINE bool IsTransacting() const
+	inline bool IsTransacting() const
 	{
 		if (FPlatformProperties::HasEditorOnlyData())
 		{
@@ -388,7 +388,7 @@ public:
 	}
 
 	/** Returns true if this archive contains data required to be gathered for localization. */
-	FORCEINLINE bool RequiresLocalizationGather() const
+	inline bool RequiresLocalizationGather() const
 	{
 PRAGMA_DISABLE_DEPRECATION_WARNINGS
 		return ArRequiresLocalizationGather;
@@ -489,7 +489,7 @@ PRAGMA_ENABLE_DEPRECATION_WARNINGS
 	}
 
 	/** Returns the editor-only debug serialization flags. */
-	FORCEINLINE uint32 GetDebugSerializationFlags() const
+	inline uint32 GetDebugSerializationFlags() const
 	{
 #if WITH_EDITOR
 		return ArDebugSerializationFlags;
@@ -811,7 +811,7 @@ public:
 		{
 			Reset();
 		}
-		FORCEINLINE void Reset()
+		inline void Reset()
 		{
 			StartFastPathLoadBuffer = nullptr;
 			EndFastPathLoadBuffer = nullptr;
@@ -1142,7 +1142,7 @@ public:
 	 *
 	 * @param InObject	The value to assign to the pointer
 	 */
-	FORCEINLINE TCheckedObjPtr& operator = (T* InObject)
+	inline TCheckedObjPtr& operator = (T* InObject)
 	{
 		Object = InObject;
 
@@ -1288,7 +1288,7 @@ public:
 	 * @param Value The value to serialize.
 	 * @return This instance.
 	 */
-	template<class T> FORCEINLINE FArchive& operator<<(TCheckedObjPtr<T>& Value)
+	template<class T> inline FArchive& operator<<(TCheckedObjPtr<T>& Value)
 	{
 		Value.bError = false;
 
@@ -1384,7 +1384,7 @@ public:
 	 * @param Ar The archive to serialize from or to.
 	 * @param Value The value to serialize.
 	 */
-	FORCEINLINE friend FArchive& operator<<(FArchive& Ar, ANSICHAR& Value)
+	inline friend FArchive& operator<<(FArchive& Ar, ANSICHAR& Value)
 	{
 #if DEVIRTUALIZE_FLinkerLoad_Serialize
 		if (!Ar.FastPathLoad<sizeof(Value)>(&Value))
@@ -1401,7 +1401,7 @@ public:
 	 * @param Ar The archive to serialize from or to.
 	 * @param Value The value to serialize.
 	 */
-	FORCEINLINE friend FArchive& operator<<(FArchive& Ar, WIDECHAR& Value)
+	inline friend FArchive& operator<<(FArchive& Ar, WIDECHAR& Value)
 	{
 #if DEVIRTUALIZE_FLinkerLoad_Serialize
 		if (!Ar.FastPathLoad<sizeof(Value)>(&Value))
@@ -1418,7 +1418,7 @@ public:
 	 * @param Ar The archive to serialize from or to.
 	 * @param Value The value to serialize.
 	 */
-	FORCEINLINE friend FArchive& operator<<(FArchive& Ar, UTF32CHAR& Value)
+	inline friend FArchive& operator<<(FArchive& Ar, UTF32CHAR& Value)
 	{
 #if DEVIRTUALIZE_FLinkerLoad_Serialize
 		if (!Ar.FastPathLoad<sizeof(Value)>(&Value))
@@ -1435,7 +1435,7 @@ public:
 	 * @param Ar The archive to serialize from or to.
 	 * @param Value The value to serialize.
 	 */
-	FORCEINLINE friend FArchive& operator<<(FArchive& Ar, uint8& Value)
+	inline friend FArchive& operator<<(FArchive& Ar, uint8& Value)
 	{
 #if DEVIRTUALIZE_FLinkerLoad_Serialize
 		if (!Ar.FastPathLoad<sizeof(Value)>(&Value))
@@ -1452,7 +1452,7 @@ public:
 	 * @param Ar The archive to serialize from or to.
 	 * @param Value The value to serialize.
 	 */
-	FORCEINLINE friend FArchive& operator<<(FArchive& Ar, int8& Value)
+	inline friend FArchive& operator<<(FArchive& Ar, int8& Value)
 	{
 #if DEVIRTUALIZE_FLinkerLoad_Serialize
 		if (!Ar.FastPathLoad<sizeof(Value)>(&Value))
@@ -1469,7 +1469,7 @@ public:
 	 * @param Ar The archive to serialize from or to.
 	 * @param Value The value to serialize.
 	 */
-	FORCEINLINE friend FArchive& operator<<(FArchive& Ar, uint16& Value)
+	inline friend FArchive& operator<<(FArchive& Ar, uint16& Value)
 	{
 #if DEVIRTUALIZE_FLinkerLoad_Serialize
 		if (!Ar.FastPathLoad<sizeof(Value)>(&Value))
@@ -1486,7 +1486,7 @@ public:
 	 * @param Ar The archive to serialize from or to.
 	 * @param Value The value to serialize.
 	 */
-	FORCEINLINE friend FArchive& operator<<(FArchive& Ar, int16& Value)
+	inline friend FArchive& operator<<(FArchive& Ar, int16& Value)
 	{
 #if DEVIRTUALIZE_FLinkerLoad_Serialize
 		if (!Ar.FastPathLoad<sizeof(Value)>(&Value))
@@ -1503,7 +1503,7 @@ public:
 	 * @param Ar The archive to serialize from or to.
 	 * @param Value The value to serialize.
 	 */
-	FORCEINLINE friend FArchive& operator<<(FArchive& Ar, uint32& Value)
+	inline friend FArchive& operator<<(FArchive& Ar, uint32& Value)
 	{
 #if DEVIRTUALIZE_FLinkerLoad_Serialize
 		if (!Ar.FastPathLoad<sizeof(Value)>(&Value))
@@ -1524,13 +1524,13 @@ public:
 protected:
 	CORE_API virtual void SerializeBool( bool& D );
 public:
-	FORCEINLINE friend FArchive& operator<<(FArchive& Ar, bool& D)
+	inline friend FArchive& operator<<(FArchive& Ar, bool& D)
 	{
 		Ar.SerializeBool(D);
 		return Ar;
 	}
 #else
-	FORCEINLINE friend FArchive& operator<<( FArchive& Ar, bool& D )
+	inline friend FArchive& operator<<( FArchive& Ar, bool& D )
 	{
 		// Serialize bool as if it were UBOOL (legacy, 32 bit int).
 #if DEVIRTUALIZE_FLinkerLoad_Serialize
@@ -1577,7 +1577,7 @@ public:
 	 * @param Ar The archive to serialize from or to.
 	 * @param Value The value to serialize.
 	 */
-	FORCEINLINE friend FArchive& operator<<(FArchive& Ar, int32& Value)
+	inline friend FArchive& operator<<(FArchive& Ar, int32& Value)
 	{
 #if DEVIRTUALIZE_FLinkerLoad_Serialize
 		if (!Ar.FastPathLoad<sizeof(Value)>(&Value))
@@ -1595,7 +1595,7 @@ public:
 	 * @param Ar The archive to serialize from or to.
 	 * @param Value The value to serialize.
 	 */
-	FORCEINLINE friend FArchive& operator<<(FArchive& Ar, long& Value)
+	inline friend FArchive& operator<<(FArchive& Ar, long& Value)
 	{
 #if DEVIRTUALIZE_FLinkerLoad_Serialize
 		if (!Ar.FastPathLoad<sizeof(Value)>(&Value))
@@ -1613,7 +1613,7 @@ public:
 	 * @param Ar The archive to serialize from or to.
 	 * @param Value The value to serialize.
 	 */
-	FORCEINLINE friend FArchive& operator<<( FArchive& Ar, float& Value)
+	inline friend FArchive& operator<<( FArchive& Ar, float& Value)
 	{
 #if DEVIRTUALIZE_FLinkerLoad_Serialize
 		if (!Ar.FastPathLoad<sizeof(Value)>(&Value))
@@ -1631,7 +1631,7 @@ public:
 	 * @param Ar The archive to serialize from or to.
 	 * @param Value The value to serialize.
 	 */
-	FORCEINLINE friend FArchive& operator<<(FArchive& Ar, double& Value)
+	inline friend FArchive& operator<<(FArchive& Ar, double& Value)
 	{
 #if DEVIRTUALIZE_FLinkerLoad_Serialize
 		if (!Ar.FastPathLoad<sizeof(Value)>(&Value))
@@ -1649,7 +1649,7 @@ public:
 	 * @param Ar The archive to serialize from or to.
 	 * @param Value The value to serialize.
 	 */
-	FORCEINLINE friend FArchive& operator<<(FArchive &Ar, uint64& Value)
+	inline friend FArchive& operator<<(FArchive &Ar, uint64& Value)
 	{
 #if DEVIRTUALIZE_FLinkerLoad_Serialize
 		if (!Ar.FastPathLoad<sizeof(Value)>(&Value))
@@ -1666,7 +1666,7 @@ public:
 	 * @param Ar The archive to serialize from or to.
 	 * @param Value The value to serialize.
 	 */
-	/*FORCEINLINE*/friend FArchive& operator<<(FArchive& Ar, int64& Value)
+	/*inline*/friend FArchive& operator<<(FArchive& Ar, int64& Value)
 	{
 #if DEVIRTUALIZE_FLinkerLoad_Serialize
 		if (!Ar.FastPathLoad<sizeof(Value)>(&Value))
@@ -1902,7 +1902,7 @@ public:
 	CORE_API void ByteSwap(void* V, int32 Length);
 
 	/** Serialize data of Length bytes, taking into account byte swapping if needed. */
-	FORCEINLINE FArchive& ByteOrderSerialize(void* V, int32 Length)
+	inline FArchive& ByteOrderSerialize(void* V, int32 Length)
 	{
 		if (!IsByteSwapping()) // Most likely case (hot path)
 		{
@@ -2103,7 +2103,7 @@ public:
 public:
 #if DEVIRTUALIZE_FLinkerLoad_Serialize
 	template<SIZE_T Size>
-	FORCEINLINE bool FastPathLoad(void* InDest)
+	inline bool FastPathLoad(void* InDest)
 	{
 		const uint8* RESTRICT Src = ActiveFPLB->StartFastPathLoadBuffer;
 		if (Src + Size <= ActiveFPLB->EndFastPathLoadBuffer)
@@ -2390,7 +2390,7 @@ template<class T> T Arctor(FArchive& Ar)
 * @param Value The value to serialize.
 */
 template<class TEnum>
-FORCEINLINE FArchive& operator<<(FArchive& Ar, TEnumAsByte<TEnum>& Value)
+inline FArchive& operator<<(FArchive& Ar, TEnumAsByte<TEnum>& Value)
 {
 #if DEVIRTUALIZE_FLinkerLoad_Serialize
 	if (!Ar.FastPathLoad<sizeof(Value)>(&Value))
diff --git a/Engine/Source/Runtime/Core/Public/Serialization/ArchiveProxy.h b/Engine/Source/Runtime/Core/Public/Serialization/ArchiveProxy.h
--- a/Engine/Source/Runtime/Core/Public/Serialization/ArchiveProxy.h
+++ b/Engine/Source/Runtime/Core/Public/Serialization/ArchiveProxy.h
@@ -274,7 +274,7 @@ public:
 	}
 #endif
 
-	FORCEINLINE void SetSerializedProperty(FProperty* InProperty) override
+	inline void SetSerializedProperty(FProperty* InProperty) override
 	{
 		FArchive::SetSerializedProperty(InProperty);
 		InnerArchive.SetSerializedProperty(InProperty);
@@ -287,13 +287,13 @@ public:
 	}
 
 	/** Pushes editor-only marker to the stack of currently serialized properties */
-	virtual FORCEINLINE void PushSerializedProperty(class FProperty* InProperty, const bool bIsEditorOnlyProperty)
+	virtual inline void PushSerializedProperty(class FProperty* InProperty, const bool bIsEditorOnlyProperty)
 	{
 		FArchive::PushSerializedProperty(InProperty, bIsEditorOnlyProperty);
 		InnerArchive.PushSerializedProperty(InProperty, bIsEditorOnlyProperty);
 	}
 	/** Pops editor-only marker from the stack of currently serialized properties */
-	virtual FORCEINLINE void PopSerializedProperty(class FProperty* InProperty, const bool bIsEditorOnlyProperty)
+	virtual inline void PopSerializedProperty(class FProperty* InProperty, const bool bIsEditorOnlyProperty)
 	{
 		FArchive::PopSerializedProperty(InProperty, bIsEditorOnlyProperty);
 		InnerArchive.PopSerializedProperty(InProperty, bIsEditorOnlyProperty);
diff --git a/Engine/Source/Runtime/Core/Public/Serialization/BitReader.h b/Engine/Source/Runtime/Core/Public/Serialization/BitReader.h
--- a/Engine/Source/Runtime/Core/Public/Serialization/BitReader.h
+++ b/Engine/Source/Runtime/Core/Public/Serialization/BitReader.h
@@ -43,7 +43,7 @@ public:
 
 // Disable false positive buffer overrun warning during pgoprofile linking step
 PRAGMA_DISABLE_BUFFER_OVERRUN_WARNING
-	FORCEINLINE_DEBUGGABLE void SerializeBits( void* Dest, int64 LengthBits )
+	inline void SerializeBits( void* Dest, int64 LengthBits )
 	{
 		//@TODO: FLOATPRECISION: This function/class pretends it is 64-bit aware, e.g., in the type of LengthBits and the Pos member, but it is not as appBitsCpy is only 32 bits, the inner Buffer is a 32 bit TArray, etc...
 		if ( IsError() || Pos+LengthBits > Num)
@@ -78,7 +78,7 @@ CORE_API PRAGMA_ENABLE_BUFFER_OVERRUN_WARNING
 	virtual void SerializeBitsWithOffset( void* Dest, int32 DestBit, int64 LengthBits ) override;
 
 	// OutValue < ValueMax
-	FORCEINLINE_DEBUGGABLE void SerializeInt(uint32& OutValue, uint32 ValueMax)
+	inline void SerializeInt(uint32& OutValue, uint32 ValueMax)
 	{
 		if (!IsError())
 		{
@@ -109,7 +109,7 @@ CORE_API PRAGMA_ENABLE_BUFFER_OVERRUN_WARNING
 
 	CORE_API virtual void SerializeIntPacked(uint32& Value) override;
 
-	FORCEINLINE_DEBUGGABLE uint32 ReadInt(uint32 Max)
+	inline uint32 ReadInt(uint32 Max)
 	{
 		uint32 Value = 0;
 
@@ -118,7 +118,7 @@ CORE_API PRAGMA_ENABLE_BUFFER_OVERRUN_WARNING
 		return Value;
 	}
 
-	FORCEINLINE_DEBUGGABLE uint8 ReadBit()
+	inline uint8 ReadBit()
 	{
 		uint8 Bit=0;
 		//SerializeBits( &Bit, 1 );
@@ -160,7 +160,7 @@ CORE_API PRAGMA_ENABLE_BUFFER_OVERRUN_WARNING
 		return Buffer;
 	}
 
-	FORCEINLINE_DEBUGGABLE uint8* GetDataPosChecked()
+	inline uint8* GetDataPosChecked()
 	{
 		check(Pos % 8 == 0);
 		return &Buffer[(int32)(Pos >> 3)];
@@ -190,7 +190,7 @@ CORE_API PRAGMA_ENABLE_BUFFER_OVERRUN_WARNING
 	{
 		return Pos;
 	}
-	FORCEINLINE_DEBUGGABLE void EatByteAlign()
+	inline void EatByteAlign()
 	{
 		int64 PrePos = Pos;
 
@@ -203,7 +203,7 @@ CORE_API PRAGMA_ENABLE_BUFFER_OVERRUN_WARNING
 			SetOverflowed(Pos - PrePos);
 		}
 	}
-	FORCEINLINE_DEBUGGABLE void Skip(int32 BitCount)
+	inline void Skip(int32 BitCount)
 	{
 		if (Pos + BitCount > Num)
 		{
diff --git a/Engine/Source/Runtime/Core/Public/Serialization/BitWriter.h b/Engine/Source/Runtime/Core/Public/Serialization/BitWriter.h
--- a/Engine/Source/Runtime/Core/Public/Serialization/BitWriter.h
+++ b/Engine/Source/Runtime/Core/Public/Serialization/BitWriter.h
@@ -66,7 +66,7 @@ public:
 	/**
 	 * Returns a pointer to our internal buffer.
 	 */
-	FORCEINLINE uint8* GetData(void)
+	inline uint8* GetData(void)
 	{
 #if !UE_BUILD_SHIPPING
 		// If this happens, your code has insufficient IsError() checks.
@@ -76,7 +76,7 @@ public:
 		return Buffer.GetData();
 	}
 
-	FORCEINLINE const uint8* GetData(void) const
+	inline const uint8* GetData(void) const
 	{
 #if !UE_BUILD_SHIPPING
 		// If this happens, your code has insufficient IsError() checks.
@@ -86,7 +86,7 @@ public:
 		return Buffer.GetData();
 	}
 
-	FORCEINLINE const TArray<uint8>* GetBuffer(void) const
+	inline const TArray<uint8>* GetBuffer(void) const
 	{
 #if !UE_BUILD_SHIPPING
 		// If this happens, your code has insufficient IsError() checks.
@@ -101,7 +101,7 @@ public:
 	 * Must be not greater than max bits allowed
 	 * Useful when writing bytes to buffer manually via pointer from GetData()
 	 */
-	FORCEINLINE void SetNumBits(int64 NumBits)
+	inline void SetNumBits(int64 NumBits)
 	{
 		check(NumBits <= Max);
 		Num = NumBits;
@@ -146,7 +146,7 @@ public:
 		bAllowOverflow = bInAllow;
 	}
 
-	FORCEINLINE bool AllowAppend(int64 LengthBits)
+	inline bool AllowAppend(int64 LengthBits)
 	{
 		//@TODO: FLOATPRECISION: This class pretends it is 64-bit aware, e.g., in the type of LengthBits and the Num/Max members, but it is not as the inner Buffer is a 32 bit TArray, etc...
 
@@ -224,7 +224,7 @@ public:
 		return Num;
 	}
 
-	FORCEINLINE_DEBUGGABLE void Init( FBitWriter& Writer)
+	inline void Init( FBitWriter& Writer)
 	{
 		Num = Writer.Num;
 		Overflowed = Writer.IsError();
diff --git a/Engine/Source/Runtime/Core/Public/Serialization/CustomVersion.h b/Engine/Source/Runtime/Core/Public/Serialization/CustomVersion.h
--- a/Engine/Source/Runtime/Core/Public/Serialization/CustomVersion.h
+++ b/Engine/Source/Runtime/Core/Public/Serialization/CustomVersion.h
@@ -59,7 +59,7 @@ struct FCustomVersion
 	}
 
 	/** Helper constructor. */
-	FORCEINLINE FCustomVersion(FGuid InKey, int32 InVersion, FName InFriendlyName, CustomVersionValidatorFunc InValidatorFunc = nullptr)
+	inline FCustomVersion(FGuid InKey, int32 InVersion, FName InFriendlyName, CustomVersionValidatorFunc InValidatorFunc = nullptr)
 	: Key           (InKey)
 	, Version       (InVersion)
 	, ReferenceCount(1)
diff --git a/Engine/Source/Runtime/Core/Public/Serialization/MemoryImage.h b/Engine/Source/Runtime/Core/Public/Serialization/MemoryImage.h
--- a/Engine/Source/Runtime/Core/Public/Serialization/MemoryImage.h
+++ b/Engine/Source/Runtime/Core/Public/Serialization/MemoryImage.h
@@ -764,7 +764,7 @@ public:
 		typename CharType,
 		typename = typename TEnableIf<TIsCharType<CharType>::Value>::Type // This TEnableIf is to ensure we don't instantiate this constructor for non-char types, like id* in Obj-C
 	>
-		FORCEINLINE FMemoryImageString(const CharType* Src)
+		inline FMemoryImageString(const CharType* Src)
 	{
 		if (Src && *Src)
 		{
diff --git a/Engine/Source/Runtime/Core/Public/Serialization/MemoryLayout.h b/Engine/Source/Runtime/Core/Public/Serialization/MemoryLayout.h
--- a/Engine/Source/Runtime/Core/Public/Serialization/MemoryLayout.h
+++ b/Engine/Source/Runtime/Core/Public/Serialization/MemoryLayout.h
@@ -249,7 +249,7 @@ namespace Freeze
 	// Any attempt to cleanup/destroy frozen data must take care to safely allow DestroyObject to be called multiple times on the same frozen object
 	// Generally speaking, it's *NOT* safe to invoke the built-in destructor multiple times
 	template<typename T>
-	FORCEINLINE void DestroyObject(T* Object, const FPointerTableBase* PtrTable, bool bIsFrozen)
+	inline void DestroyObject(T* Object, const FPointerTableBase* PtrTable, bool bIsFrozen)
 	{
 		// Only call destructor for non-frozen objects
 		if (!bIsFrozen)
@@ -267,7 +267,7 @@ namespace Freeze
 	}
 
 	template<typename T>
-	FORCEINLINE uint32 IntrinsicUnfrozenCopy(const FMemoryUnfreezeContent& Context, const T& Object, void* OutDst)
+	inline uint32 IntrinsicUnfrozenCopy(const FMemoryUnfreezeContent& Context, const T& Object, void* OutDst)
 	{
 		new(OutDst) T(Object);
 		return sizeof(T);
diff --git a/Engine/Source/Runtime/Core/Public/Serialization/StructuredArchiveChildReader.h b/Engine/Source/Runtime/Core/Public/Serialization/StructuredArchiveChildReader.h
--- a/Engine/Source/Runtime/Core/Public/Serialization/StructuredArchiveChildReader.h
+++ b/Engine/Source/Runtime/Core/Public/Serialization/StructuredArchiveChildReader.h
@@ -51,7 +51,7 @@ private:
 
 #if !WITH_TEXT_ARCHIVE_SUPPORT
 
-FORCEINLINE FStructuredArchiveChildReader::FStructuredArchiveChildReader(FStructuredArchiveSlot InSlot)
+inline FStructuredArchiveChildReader::FStructuredArchiveChildReader(FStructuredArchiveSlot InSlot)
 	: OwnedFormatter(nullptr)
 	, StructuredArchive(nullptr)
 {
diff --git a/Engine/Source/Runtime/Core/Public/Serialization/StructuredArchiveSlotBase.h b/Engine/Source/Runtime/Core/Public/Serialization/StructuredArchiveSlotBase.h
--- a/Engine/Source/Runtime/Core/Public/Serialization/StructuredArchiveSlotBase.h
+++ b/Engine/Source/Runtime/Core/Public/Serialization/StructuredArchiveSlotBase.h
@@ -52,7 +52,7 @@ namespace UE::StructuredArchive::Private
 		int32 Depth;
 		FElementId ElementId;
 
-		FORCEINLINE explicit FSlotPosition(int32 InDepth, FElementId InElementId)
+		inline explicit FSlotPosition(int32 InDepth, FElementId InElementId)
 			: Depth(InDepth)
 			, ElementId(InElementId)
 		{
@@ -74,7 +74,7 @@ namespace UE::StructuredArchive::Private
 
 	public:
 #if WITH_TEXT_ARCHIVE_SUPPORT
-		FORCEINLINE explicit FSlotBase(EPrivateToken, FStructuredArchive& InStructuredArchive, int32 InDepth, FElementId InElementId)
+		inline explicit FSlotBase(EPrivateToken, FStructuredArchive& InStructuredArchive, int32 InDepth, FElementId InElementId)
 			: FSlotPosition(InDepth, InElementId)
 			, StructuredArchive(InStructuredArchive)
 		{
diff --git a/Engine/Source/Runtime/Core/Public/Serialization/StructuredArchiveSlots.h b/Engine/Source/Runtime/Core/Public/Serialization/StructuredArchiveSlots.h
--- a/Engine/Source/Runtime/Core/Public/Serialization/StructuredArchiveSlots.h
+++ b/Engine/Source/Runtime/Core/Public/Serialization/StructuredArchiveSlots.h
@@ -85,7 +85,7 @@ public:
 	UE_API void operator << (FObjectPtr& Value);
 
 	template <typename T>
-	FORCEINLINE void operator<<(TEnumAsByte<T>& Value)
+	inline void operator<<(TEnumAsByte<T>& Value)
 	{
 		uint8 Tmp = (uint8)Value.GetValue();
 		*this << Tmp;
@@ -108,7 +108,7 @@ public:
 	}
 
 	template <typename T>
-	FORCEINLINE void operator<<(UE::StructuredArchive::Private::TOptionalNamedAttribute<T> Item)
+	inline void operator<<(UE::StructuredArchive::Private::TOptionalNamedAttribute<T> Item)
 	{
 		if (TOptional<FStructuredArchiveSlot> Attribute = TryEnterAttribute(Item.Name, Item.Value != Item.Default))
 		{
@@ -151,7 +151,7 @@ public:
 
 	UE_API TOptional<FStructuredArchiveSlot> TryEnterField(FArchiveFieldName Name, bool bEnterForSaving);
 
-	template<typename T> FORCEINLINE FStructuredArchiveRecord& operator<<(UE::StructuredArchive::Private::TNamedValue<T> Item)
+	template<typename T> inline FStructuredArchiveRecord& operator<<(UE::StructuredArchive::Private::TNamedValue<T> Item)
 	{
 		EnterField(Item.Name) << Item.Value;
 		return *this;
@@ -173,7 +173,7 @@ class FStructuredArchiveArray final : public UE::StructuredArchive::Private::FSl
 public:
 	UE_API FStructuredArchiveSlot EnterElement();
 
-	template<typename T> FORCEINLINE FStructuredArchiveArray& operator<<(T& Item)
+	template<typename T> inline FStructuredArchiveArray& operator<<(T& Item)
 	{
 		EnterElement() << Item;
 		return *this;
@@ -194,7 +194,7 @@ class FStructuredArchiveStream final : public UE::StructuredArchive::Private::FS
 public:
 	UE_API FStructuredArchiveSlot EnterElement();
 
-	template<typename T> FORCEINLINE FStructuredArchiveStream& operator<<(T& Item)
+	template<typename T> inline FStructuredArchiveStream& operator<<(T& Item)
 	{
 		EnterElement() << Item;
 		return *this;
@@ -224,7 +224,7 @@ private:
 };
 
 template <typename T>
-FORCEINLINE_DEBUGGABLE void operator<<(FStructuredArchiveSlot Slot, TArray<T>& InArray)
+inline void operator<<(FStructuredArchiveSlot Slot, TArray<T>& InArray)
 {
 	int32 NumElements = InArray.Num();
 	FStructuredArchiveArray Array = Slot.EnterArray(NumElements);
@@ -254,7 +254,7 @@ FORCEINLINE_DEBUGGABLE void operator<<(FStructuredArchiveSlot Slot, TArray<uint8
 		return FStructuredArchiveRecord(FStructuredArchiveRecord::EPrivateToken{}, StructuredArchive);
 	}
 
-	FORCEINLINE FStructuredArchiveArray FStructuredArchiveSlot::EnterArray(int32& Num)
+	inline FStructuredArchiveArray FStructuredArchiveSlot::EnterArray(int32& Num)
 	{
 		UE::StructuredArchive::Private::GetFormatterImpl(StructuredArchive).EnterArray(Num);
 		return FStructuredArchiveArray(FStructuredArchiveArray::EPrivateToken{}, StructuredArchive);
@@ -265,19 +265,19 @@ FORCEINLINE_DEBUGGABLE void operator<<(FStructuredArchiveSlot Slot, TArray<uint8
 		return FStructuredArchiveStream(FStructuredArchiveStream::EPrivateToken{}, StructuredArchive);
 	}
 
-	FORCEINLINE FStructuredArchiveMap FStructuredArchiveSlot::EnterMap(int32& Num)
+	inline FStructuredArchiveMap FStructuredArchiveSlot::EnterMap(int32& Num)
 	{
 		UE::StructuredArchive::Private::GetFormatterImpl(StructuredArchive).EnterMap(Num);
 		return FStructuredArchiveMap(FStructuredArchiveMap::EPrivateToken{}, StructuredArchive);
 	}
 
-	FORCEINLINE FStructuredArchiveSlot FStructuredArchiveSlot::EnterAttribute(FArchiveFieldName FieldName)
+	inline FStructuredArchiveSlot FStructuredArchiveSlot::EnterAttribute(FArchiveFieldName FieldName)
 	{
 		UE::StructuredArchive::Private::GetFormatterImpl(StructuredArchive).EnterAttribute(FieldName);
 		return FStructuredArchiveSlot(FStructuredArchiveSlot::EPrivateToken{}, StructuredArchive);
 	}
 
-	FORCEINLINE TOptional<FStructuredArchiveSlot> FStructuredArchiveSlot::TryEnterAttribute(FArchiveFieldName FieldName, bool bEnterWhenWriting)
+	inline TOptional<FStructuredArchiveSlot> FStructuredArchiveSlot::TryEnterAttribute(FArchiveFieldName FieldName, bool bEnterWhenWriting)
 	{
 		if (UE::StructuredArchive::Private::GetFormatterImpl(StructuredArchive).TryEnterAttribute(FieldName, bEnterWhenWriting))
 		{
@@ -438,7 +438,7 @@ FORCEINLINE_DEBUGGABLE void operator<<(FStructuredArchiveSlot Slot, TArray<uint8
 		return EnterField(Name).EnterMap(Num);
 	}
 
-	FORCEINLINE TOptional<FStructuredArchiveSlot> FStructuredArchiveRecord::TryEnterField(FArchiveFieldName Name, bool bEnterWhenWriting)
+	inline TOptional<FStructuredArchiveSlot> FStructuredArchiveRecord::TryEnterField(FArchiveFieldName Name, bool bEnterWhenWriting)
 	{
 		if (UE::StructuredArchive::Private::GetFormatterImpl(StructuredArchive).TryEnterField(Name, bEnterWhenWriting))
 		{
@@ -466,7 +466,7 @@ FORCEINLINE_DEBUGGABLE void operator<<(FStructuredArchiveSlot Slot, TArray<uint8
 
 	//////////// FStructuredArchiveMap ////////////
 
-	FORCEINLINE FStructuredArchiveSlot FStructuredArchiveMap::EnterElement(FString& Name)
+	inline FStructuredArchiveSlot FStructuredArchiveMap::EnterElement(FString& Name)
 	{
 		UE::StructuredArchive::Private::GetFormatterImpl(StructuredArchive).EnterMapElement(Name);
 		return FStructuredArchiveSlot(FStructuredArchiveSlot::EPrivateToken{}, StructuredArchive);
diff --git a/Engine/Source/Runtime/Core/Public/Stats/HitchTrackingStatScope.h b/Engine/Source/Runtime/Core/Public/Stats/HitchTrackingStatScope.h
--- a/Engine/Source/Runtime/Core/Public/Stats/HitchTrackingStatScope.h
+++ b/Engine/Source/Runtime/Core/Public/Stats/HitchTrackingStatScope.h
@@ -27,7 +27,7 @@ namespace UE::Stats
 			StatString = GHitchDetected ? nullptr : InStat;
 		}
 
-		FORCEINLINE ~FHitchTrackingStatScope()
+		inline ~FHitchTrackingStatScope()
 		{
 			if (GHitchDetected && StatString)
 			{
diff --git a/Engine/Source/Runtime/Core/Public/Stats/StatsData.h b/Engine/Source/Runtime/Core/Public/Stats/StatsData.h
--- a/Engine/Source/Runtime/Core/Public/Stats/StatsData.h
+++ b/Engine/Source/Runtime/Core/Public/Stats/StatsData.h
@@ -148,7 +148,7 @@ struct EStatDisplayMode
 template< typename T >
 struct FStatNameComparer
 {
-	FORCEINLINE_STATS bool operator()( const T& A, const T& B ) const
+	inline bool operator()( const T& A, const T& B ) const
 	{
 		check(0);
 		return 0;
@@ -159,7 +159,7 @@ struct FStatNameComparer
 template< typename T >
 struct FStatDurationComparer
 {
-	FORCEINLINE_STATS bool operator()( const T& A, const T& B ) const
+	inline bool operator()( const T& A, const T& B ) const
 	{
 		check(0);
 		return 0;
@@ -170,7 +170,7 @@ struct FStatDurationComparer
 template< typename T >
 struct FStatCallCountComparer
 {
-	FORCEINLINE_STATS bool operator()( const T& A, const T& B ) const
+	inline bool operator()( const T& A, const T& B ) const
 	{
 		check(0);
 		return 0;
@@ -191,7 +191,7 @@ struct FStatNameComparer<FStatMessage>
 template<>
 struct FStatDurationComparer<FStatMessage>
 {
-	FORCEINLINE_STATS bool operator()( const FStatMessage& A, const FStatMessage& B ) const
+	inline bool operator()( const FStatMessage& A, const FStatMessage& B ) const
 	{
 		const uint32 DurationA = FromPackedCallCountDuration_Duration( A.GetValue_int64() );
 		const uint32 DurationB = FromPackedCallCountDuration_Duration( B.GetValue_int64() );
@@ -203,7 +203,7 @@ struct FStatDurationComparer<FStatMessage>
 template<>
 struct FStatCallCountComparer<FStatMessage>
 {
-	FORCEINLINE_STATS bool operator()( const FStatMessage& A, const FStatMessage& B ) const
+	inline bool operator()( const FStatMessage& A, const FStatMessage& B ) const
 	{
 		const uint32 CallCountA = FromPackedCallCountDuration_CallCount( A.GetValue_int64() );
 		const uint32 CallCountB = FromPackedCallCountDuration_CallCount( B.GetValue_int64() );
@@ -353,7 +353,7 @@ struct FStatNameComparer<FRawStatStackNode>
 template<>
 struct FStatDurationComparer<FRawStatStackNode>
 {
-	FORCEINLINE_STATS bool operator()( const FRawStatStackNode& A, const FRawStatStackNode& B ) const
+	inline bool operator()( const FRawStatStackNode& A, const FRawStatStackNode& B ) const
 	{
 		const uint32 DurationA = FromPackedCallCountDuration_Duration( A.Meta.GetValue_int64() );
 		const uint32 DurationB = FromPackedCallCountDuration_Duration( B.Meta.GetValue_int64() );
@@ -365,7 +365,7 @@ struct FStatDurationComparer<FRawStatStackNode>
 template<>
 struct FStatCallCountComparer<FRawStatStackNode>
 {
-	FORCEINLINE_STATS bool operator()( const FRawStatStackNode& A, const FRawStatStackNode& B ) const
+	inline bool operator()( const FRawStatStackNode& A, const FRawStatStackNode& B ) const
 	{
 		const uint32 CallCountA = FromPackedCallCountDuration_CallCount( A.Meta.GetValue_int64() );
 		const uint32 CallCountB = FromPackedCallCountDuration_CallCount( B.Meta.GetValue_int64() );
diff --git a/Engine/Source/Runtime/Core/Public/Stats/StatsFile.h b/Engine/Source/Runtime/Core/Public/Stats/StatsFile.h
--- a/Engine/Source/Runtime/Core/Public/Stats/StatsFile.h
+++ b/Engine/Source/Runtime/Core/Public/Stats/StatsFile.h
@@ -368,7 +368,7 @@ protected:
 	CORE_API void WriteCondensedMessages( FArchive& Ar, int64 TargetFrame );
 
 	/** Sends an FName, and the string it represents if we have not sent that string before. **/
-	FORCEINLINE_STATS void WriteFName( FArchive& Ar, FStatNameAndInfo NameAndInfo )
+	inline void WriteFName( FArchive& Ar, FStatNameAndInfo NameAndInfo )
 	{
 		FName RawName = NameAndInfo.GetRawName();
 		bool bSendFName = !FNamesSent.Contains( RawName.GetComparisonIndex() );
@@ -391,7 +391,7 @@ protected:
 	}
 
 	/** Write a stat message. **/
-	FORCEINLINE_STATS void WriteMessage( FArchive& Ar, FStatMessage const& Item )
+	inline void WriteMessage( FArchive& Ar, FStatMessage const& Item )
 	{
 		WriteFName( Ar, Item.NameAndInfo );
 		switch (Item.NameAndInfo.GetField<EStatDataType>())
@@ -687,7 +687,7 @@ public:
 	}
 
 	/** Read and translate or create an FName. **/
-	FORCEINLINE_STATS FStatNameAndInfo ReadFName( FArchive& Ar, bool bHasFNameMap )
+	inline FStatNameAndInfo ReadFName( FArchive& Ar, bool bHasFNameMap )
 	{
 		// If we read the whole FNames translation map, we don't want to add the FName again.
 		// This is a bit tricky, even if we have the FName translation map, we still need to read the FString.
@@ -751,7 +751,7 @@ public:
 	}
 
 	/** Read a stat message. **/
-	FORCEINLINE_STATS FStatMessage ReadMessage( FArchive& Ar, bool bHasFNameMap = false )
+	inline FStatMessage ReadMessage( FArchive& Ar, bool bHasFNameMap = false )
 	{
 		FStatMessage Result( ReadFName( Ar, bHasFNameMap ) );
 		Result.Clear();
diff --git a/Engine/Source/Runtime/Core/Public/Stats/StatsSystemTypes.h b/Engine/Source/Runtime/Core/Public/Stats/StatsSystemTypes.h
--- a/Engine/Source/Runtime/Core/Public/Stats/StatsSystemTypes.h
+++ b/Engine/Source/Runtime/Core/Public/Stats/StatsSystemTypes.h
@@ -368,7 +368,7 @@ public:
 	/**
 	 * Build from a raw FName
 	 */
-	FORCEINLINE_STATS FStatNameAndInfo(FName Other, bool bAlreadyHasMeta)
+	inline FStatNameAndInfo(FName Other, bool bAlreadyHasMeta)
 		: Index(Other.GetComparisonIndex())
 		, Number(Other.GetNumber())
 	{
@@ -384,7 +384,7 @@ public:
 	/**
 	 * Build with stat metadata
 	 */
-	FORCEINLINE_STATS FStatNameAndInfo(FName InStatName, char const* InGroup, char const* InCategory, TCHAR const* InDescription, EStatDataType::Type InStatType, bool bShouldClearEveryFrame, bool bCycleStat, bool bSortByName, FPlatformMemory::EMemoryCounterRegion MemoryRegion = FPlatformMemory::MCR_Invalid)
+	inline FStatNameAndInfo(FName InStatName, char const* InGroup, char const* InCategory, TCHAR const* InDescription, EStatDataType::Type InStatType, bool bShouldClearEveryFrame, bool bCycleStat, bool bSortByName, FPlatformMemory::EMemoryCounterRegion MemoryRegion = FPlatformMemory::MCR_Invalid)
 	{
 		FName LongName = ToLongName(InStatName, InGroup, InCategory, InDescription, bSortByName);
 		Index = LongName.GetComparisonIndex();
@@ -422,7 +422,7 @@ public:
 	/**
 	* Internal use, used by the serializer
 	 */
-	FORCEINLINE_STATS int32 GetRawNumber() const
+	inline int32 GetRawNumber() const
 	{
 		CheckInvariants();
 		return Number;
@@ -431,7 +431,7 @@ public:
 	/**
 	 * Internal use by FStatsThreadState to force an update to the long name
 	 */
-	FORCEINLINE_STATS void SetRawName(FName RawName)
+	inline void SetRawName(FName RawName)
 	{
 		// ok, you can't have numbered stat FNames too large
 		checkStats(!(RawName.GetNumber() >> EStatAllFields::StartShift));
@@ -446,7 +446,7 @@ public:
 	 * The encoded FName with the correct, original Number
 	 * The original number usually is 0
 	 */
-	FORCEINLINE_STATS FName GetRawName() const
+	inline FName GetRawName() const
 	{
 		CheckInvariants();
 		return FName(Index, Index, Number & ((1 << EStatAllFields::StartShift) - 1));
@@ -456,7 +456,7 @@ public:
 	 * The encoded FName with the encoded, new Number
 	 * The number contains all encoded metadata
 	 */
-	FORCEINLINE_STATS FName GetEncodedName() const
+	inline FName GetEncodedName() const
 	{
 		CheckInvariants();
 		return FName(Index, Index, Number);
@@ -465,7 +465,7 @@ public:
 	/**
 	 * Expensive! Extracts the shortname if this is a long name or just returns the name
 	 */
-	FORCEINLINE_STATS FName GetShortName() const
+	inline FName GetShortName() const
 	{
 		CheckInvariants();
 		return GetShortNameFrom(GetRawName());
@@ -474,7 +474,7 @@ public:
 	/**
 	 * Expensive! Extracts the group name if this is a long name or just returns none
 	 */
-	FORCEINLINE_STATS FName GetGroupName() const
+	inline FName GetGroupName() const
 	{
 		CheckInvariants();
 		return GetGroupNameFrom(GetRawName());
@@ -483,7 +483,7 @@ public:
 	/**
 	 * Expensive! Extracts the group category if this is a long name or just returns none
 	 */
-	FORCEINLINE_STATS FName GetGroupCategory() const
+	inline FName GetGroupCategory() const
 	{
 		CheckInvariants();
 		return GetGroupCategoryFrom(GetRawName());
@@ -492,7 +492,7 @@ public:
 	/**
 	 * Expensive! Extracts the description if this is a long name or just returns the empty string
 	 */
-	FORCEINLINE_STATS FString GetDescription() const
+	inline FString GetDescription() const
 	{
 		CheckInvariants();
 		return GetDescriptionFrom(GetRawName());
@@ -500,7 +500,7 @@ public:
 	/**
 	 * Expensive! Extracts the sort by name flag
 	 */
-	FORCEINLINE_STATS bool GetSortByName() const
+	inline bool GetSortByName() const
 	{
 		CheckInvariants();
 		return GetSortByNameFrom(GetRawName());
@@ -510,7 +510,7 @@ public:
 	/**
 	 * Makes sure this object is in good shape
 	 */
-	FORCEINLINE_STATS void CheckInvariants() const
+	inline void CheckInvariants() const
 	{
 		checkStats((Number & (EStatMetaFlags::DummyAlwaysOne << (EStatAllFields::StartShift + EStatMetaFlags::Shift)))
 			&& Index);
@@ -666,7 +666,7 @@ struct FStatMessage
 	/**
 	* Clock operation
 	*/
-	FORCEINLINE_STATS FStatMessage(FName InStatName, EStatOperation::Type InStatOperation)
+	inline FStatMessage(FName InStatName, EStatOperation::Type InStatOperation)
 		: NameAndInfo(InStatName, true)
 	{
 		NameAndInfo.SetField<EStatOperation>(InStatOperation);
@@ -687,7 +687,7 @@ struct FStatMessage
 	/**
 	* int64 operation
 	*/
-	FORCEINLINE_STATS FStatMessage(FName InStatName, EStatOperation::Type InStatOperation, int64 Value, bool bIsCycle)
+	inline FStatMessage(FName InStatName, EStatOperation::Type InStatOperation, int64 Value, bool bIsCycle)
 		: NameAndInfo(InStatName, true)
 	{
 		NameAndInfo.SetField<EStatOperation>(InStatOperation);
@@ -699,7 +699,7 @@ struct FStatMessage
 	/**
 	* double operation
 	*/
-	FORCEINLINE_STATS FStatMessage(FName InStatName, EStatOperation::Type InStatOperation, double Value, bool)
+	inline FStatMessage(FName InStatName, EStatOperation::Type InStatOperation, double Value, bool)
 		: NameAndInfo(InStatName, true)
 	{
 		NameAndInfo.SetField<EStatOperation>(InStatOperation);
@@ -711,7 +711,7 @@ struct FStatMessage
 	/**
 	* name operation
 	*/
-	FORCEINLINE_STATS FStatMessage(FName InStatName, EStatOperation::Type InStatOperation, FName Value, bool)
+	inline FStatMessage(FName InStatName, EStatOperation::Type InStatOperation, FName Value, bool)
 		: NameAndInfo(InStatName, true)
 	{
 		NameAndInfo.SetField<EStatOperation>(InStatOperation);
@@ -723,7 +723,7 @@ struct FStatMessage
 	/**
 	 * Ptr operation
 	 */
-	FORCEINLINE_STATS FStatMessage(FName InStatName, EStatOperation::Type InStatOperation, uint64 Value, bool)
+	inline FStatMessage(FName InStatName, EStatOperation::Type InStatOperation, uint64 Value, bool)
 		: NameAndInfo(InStatName, true)
 	{
 		NameAndInfo.SetField<EStatOperation>(InStatOperation);
@@ -735,7 +735,7 @@ struct FStatMessage
 	/**
 	* Clear any data type
 	*/
-	FORCEINLINE_STATS void Clear()
+	inline void Clear()
 	{
 		static_assert(sizeof(uint64) == DATA_SIZE, "Bad clear.");
 		*(int64*)&StatData = 0;
@@ -744,31 +744,31 @@ struct FStatMessage
 	/**
 	* Payload retrieval and setting methods
 	*/
-	FORCEINLINE_STATS int64& GetValue_int64()
+	inline int64& GetValue_int64()
 	{
 		static_assert(sizeof(int64) <= DATA_SIZE && alignof(int64) <= DATA_ALIGN, "Bad data for stat message.");
 		checkStats(NameAndInfo.GetField<EStatDataType>() == EStatDataType::ST_int64);
 		return *(int64*)&StatData;
 	}
-	FORCEINLINE_STATS int64 GetValue_int64() const
+	inline int64 GetValue_int64() const
 	{
 		checkStats(NameAndInfo.GetField<EStatDataType>() == EStatDataType::ST_int64);
 		return *(int64 const*)&StatData;
 	}
 
-	FORCEINLINE_STATS uint64& GetValue_Ptr()
+	inline uint64& GetValue_Ptr()
 	{
 		static_assert(sizeof(uint64) <= DATA_SIZE && alignof(uint64) <= DATA_ALIGN, "Bad data for stat message.");
 		checkStats(NameAndInfo.GetField<EStatDataType>() == EStatDataType::ST_Ptr);
 		return *(uint64*)&StatData;
 	}
-	FORCEINLINE_STATS uint64 GetValue_Ptr() const
+	inline uint64 GetValue_Ptr() const
 	{
 		checkStats(NameAndInfo.GetField<EStatDataType>() == EStatDataType::ST_Ptr);
 		return *(uint64 const*)&StatData;
 	}
 
-	FORCEINLINE_STATS int64 GetValue_Duration() const
+	inline int64 GetValue_Duration() const
 	{
 		checkStats(NameAndInfo.GetFlag(EStatMetaFlags::IsCycle) && NameAndInfo.GetField<EStatDataType>() == EStatDataType::ST_int64);
 		if (NameAndInfo.GetFlag(EStatMetaFlags::IsPackedCCAndDuration))
@@ -778,39 +778,39 @@ struct FStatMessage
 		return *(int64 const*)&StatData;
 	}
 
-	FORCEINLINE_STATS uint32 GetValue_CallCount() const
+	inline uint32 GetValue_CallCount() const
 	{
 		checkStats(NameAndInfo.GetFlag(EStatMetaFlags::IsPackedCCAndDuration) && NameAndInfo.GetFlag(EStatMetaFlags::IsCycle) && NameAndInfo.GetField<EStatDataType>() == EStatDataType::ST_int64);
 		return FromPackedCallCountDuration_CallCount(*(int64 const*)&StatData);
 	}
 
-	FORCEINLINE_STATS double& GetValue_double()
+	inline double& GetValue_double()
 	{
 		static_assert(sizeof(double) <= DATA_SIZE && alignof(double) <= DATA_ALIGN, "Bad data for stat message.");
 		checkStats(NameAndInfo.GetField<EStatDataType>() == EStatDataType::ST_double);
 		return *(double*)&StatData;
 	}
 
-	FORCEINLINE_STATS double GetValue_double() const
+	inline double GetValue_double() const
 	{
 		checkStats(NameAndInfo.GetField<EStatDataType>() == EStatDataType::ST_double);
 		return *(double const*)&StatData;
 	}
 
-	FORCEINLINE_STATS FMinimalName& GetValue_FMinimalName()
+	inline FMinimalName& GetValue_FMinimalName()
 	{
 		static_assert(sizeof(FMinimalName) <= DATA_SIZE && alignof(FMinimalName) <= DATA_ALIGN, "Bad data for stat message.");
 		checkStats(NameAndInfo.GetField<EStatDataType>() == EStatDataType::ST_FName);
 		return *(FMinimalName*)&StatData;
 	}
 
-	FORCEINLINE_STATS FMinimalName GetValue_FMinimalName() const
+	inline FMinimalName GetValue_FMinimalName() const
 	{
 		checkStats(NameAndInfo.GetField<EStatDataType>() == EStatDataType::ST_FName);
 		return *(FMinimalName const*)&StatData;
 	}
 
-	FORCEINLINE_STATS FName GetValue_FName() const
+	inline FName GetValue_FName() const
 	{
 		checkStats(NameAndInfo.GetField<EStatDataType>() == EStatDataType::ST_FName);
 		return MinimalNameToName(*(FMinimalName const*)&StatData);
@@ -857,7 +857,7 @@ struct TStatMessage
 	/**
 	* Copy constructor from FStatMessage
 	*/
-	explicit FORCEINLINE_STATS TStatMessage(const FStatMessage& Other)
+	explicit inline TStatMessage(const FStatMessage& Other)
 		: NameAndInfo(Other.NameAndInfo)
 	{
 		// Reset data type and clear all fields.
@@ -920,7 +920,7 @@ struct TStatMessage
 	/**
 	* Clear any data type
 	*/
-	FORCEINLINE_STATS void Clear()
+	inline void Clear()
 	{
 		static_assert(sizeof(uint64) == DATA_SIZE / EnumCount, "Bad clear.");
 
@@ -933,7 +933,7 @@ struct TStatMessage
 	/**
 	* Payload retrieval and setting methods
 	*/
-	FORCEINLINE_STATS int64& GetValue_int64(typename TEnum::Type Index)
+	inline int64& GetValue_int64(typename TEnum::Type Index)
 	{
 		static_assert(sizeof(int64) <= DATA_SIZE && alignof(int64) <= DATA_ALIGN, "Bad data for stat message.");
 		checkStats(NameAndInfo.GetField<EStatDataType>() == EStatDataType::ST_int64);
@@ -941,7 +941,7 @@ struct TStatMessage
 		int64& Value = *((int64*)&StatData + (uint32)Index);
 		return Value;
 	}
-	FORCEINLINE_STATS int64 GetValue_int64(typename TEnum::Type Index) const
+	inline int64 GetValue_int64(typename TEnum::Type Index) const
 	{
 		checkStats(NameAndInfo.GetField<EStatDataType>() == EStatDataType::ST_int64);
 		checkStats(Index < EnumCount);
@@ -949,7 +949,7 @@ struct TStatMessage
 		return Value;
 	}
 
-	FORCEINLINE_STATS int64 GetValue_Duration(typename TEnum::Type Index) const
+	inline int64 GetValue_Duration(typename TEnum::Type Index) const
 	{
 		checkStats(NameAndInfo.GetFlag(EStatMetaFlags::IsCycle) && NameAndInfo.GetField<EStatDataType>() == EStatDataType::ST_int64);
 		checkStats(Index < EnumCount);
@@ -962,7 +962,7 @@ struct TStatMessage
 		return Value;
 	}
 
-	FORCEINLINE_STATS uint32 GetValue_CallCount(typename TEnum::Type Index) const
+	inline uint32 GetValue_CallCount(typename TEnum::Type Index) const
 	{
 		checkStats(NameAndInfo.GetFlag(EStatMetaFlags::IsPackedCCAndDuration) && NameAndInfo.GetFlag(EStatMetaFlags::IsCycle) && NameAndInfo.GetField<EStatDataType>() == EStatDataType::ST_int64);
 		checkStats(Index < EnumCount);
@@ -970,7 +970,7 @@ struct TStatMessage
 		return Value;
 	}
 
-	FORCEINLINE_STATS double& GetValue_double(typename TEnum::Type Index)
+	inline double& GetValue_double(typename TEnum::Type Index)
 	{
 		static_assert(sizeof(double) <= DATA_SIZE && alignof(double) <= DATA_ALIGN, "Bad data for stat message.");
 		checkStats(Index < EnumCount);
@@ -979,7 +979,7 @@ struct TStatMessage
 		return Value;
 	}
 
-	FORCEINLINE_STATS double GetValue_double(typename TEnum::Type Index) const
+	inline double GetValue_double(typename TEnum::Type Index) const
 	{
 		checkStats(NameAndInfo.GetField<EStatDataType>() == EStatDataType::ST_double);
 		checkStats(Index < EnumCount);
@@ -1204,7 +1204,7 @@ protected:
 
 public:
 	/** Checks the TLS for a thread packet and if it isn't found, it makes a new one. **/
-	static FORCEINLINE_STATS FThreadStats* GetThreadStats()
+	static inline FThreadStats* GetThreadStats()
 	{
 		FThreadStats* Stats = (FThreadStats*)FPlatformTLS::GetTlsValue(TlsSlot);
 		if (!Stats)
@@ -1225,7 +1225,7 @@ public:
 	CORE_API bool DetectAndUpdateCurrentGameFrame();
 
 	/** Maintains the explicit flush. */
-	FORCEINLINE_STATS void UpdateExplicitFlush()
+	inline void UpdateExplicitFlush()
 	{
 		if (Packet.ThreadType != EThreadType::Other && bSawExplicitFlush)
 		{
@@ -1247,7 +1247,7 @@ public:
 	/** Checks the command line whether we want to enable collecting startup stats. */
 	static CORE_API void CheckForCollectingStartupStats();
 
-	FORCEINLINE_STATS void AddStatMessage(const FStatMessage& StatMessage)
+	inline void AddStatMessage(const FStatMessage& StatMessage)
 	{
 		LLM_SCOPE(ELLMTag::Stats);
 		FStatMessageLock MessageLock(MemoryMessageScope);
@@ -1257,7 +1257,7 @@ public:
 protected:
 	/** Any non-clock operation with an ordinary payload. **/
 	template <typename TValue>
-	FORCEINLINE_STATS void AddMessageInner(FName InStatName, EStatOperation::Type InStatOperation, TValue Value, bool bIsCycle = false)
+	inline void AddMessageInner(FName InStatName, EStatOperation::Type InStatOperation, TValue Value, bool bIsCycle = false)
 	{
 		AddStatMessage(FStatMessage(InStatName, InStatOperation, Value, bIsCycle));
 
@@ -1282,7 +1282,7 @@ protected:
 	}
 
 	/** Clock operation. **/
-	FORCEINLINE_STATS void AddMessageInner(FName InStatName, EStatOperation::Type InStatOperation)
+	inline void AddMessageInner(FName InStatName, EStatOperation::Type InStatOperation)
 	{
 		checkStats((InStatOperation == EStatOperation::CycleScopeStart || InStatOperation == EStatOperation::CycleScopeEnd));
 
@@ -1331,7 +1331,7 @@ public:
 	}
 
 	/** Clock operation. **/
-	static FORCEINLINE_STATS void AddMessage(FName InStatName, EStatOperation::Type InStatOperation)
+	static inline void AddMessage(FName InStatName, EStatOperation::Type InStatOperation)
 	{
 		if (!InStatName.IsNone() && WillEverCollectData() && IsThreadingReady())
 		{
@@ -1341,7 +1341,7 @@ public:
 
 	/** Any non-clock operation with an ordinary payload. **/
 	template<typename TValue>
-	static FORCEINLINE_STATS void AddMessage(FName InStatName, EStatOperation::Type InStatOperation, TValue Value, bool bIsCycle = false)
+	static inline void AddMessage(FName InStatName, EStatOperation::Type InStatOperation, TValue Value, bool bIsCycle = false)
 	{
 		if (!InStatName.IsNone() && WillEverCollectData() && IsThreadingReady())
 		{
@@ -1389,28 +1389,28 @@ public:
 	}
 
 	/** Indicate that you would like the system to begin collecting data, if it isn't already collecting data. Think reference count. **/
-	static FORCEINLINE_STATS void PrimaryEnableAdd(int32 Value = 1)
+	static inline void PrimaryEnableAdd(int32 Value = 1)
 	{
 		PrimaryEnableCounter.Add(Value);
 		CheckEnable();
 	}
 
 	/** Indicate that you no longer need stat data, if nobody else needs stat data, then no stat data will be collected. Think reference count. **/
-	static FORCEINLINE_STATS void PrimaryEnableSubtract(int32 Value = 1)
+	static inline void PrimaryEnableSubtract(int32 Value = 1)
 	{
 		PrimaryEnableCounter.Subtract(Value);
 		CheckEnable();
 	}
 
 	/** Indicate that you no longer need stat data, forever. **/
-	static FORCEINLINE_STATS void PrimaryDisableForever()
+	static inline void PrimaryDisableForever()
 	{
 		bPrimaryDisableForever = true;
 		CheckEnable();
 	}
 
 	/** This is called before we start to change something that will invalidate. **/
-	static FORCEINLINE_STATS void PrimaryDisableChangeTagLockAdd(int32 Value = 1)
+	static inline void PrimaryDisableChangeTagLockAdd(int32 Value = 1)
 	{
 		PrimaryDisableChangeTagLock.Add(Value);
 		FPlatformMisc::MemoryBarrier();
@@ -1418,7 +1418,7 @@ public:
 	}
 
 	/** Indicate that you no longer need stat data, if nobody else needs stat data, then no stat data will be collected. Think reference count. **/
-	static FORCEINLINE_STATS void PrimaryDisableChangeTagLockSubtract(int32 Value = 1)
+	static inline void PrimaryDisableChangeTagLockSubtract(int32 Value = 1)
 	{
 		FPlatformMisc::MemoryBarrier();
 		PrimaryEnableUpdateNumber.Increment();
@@ -1427,7 +1427,7 @@ public:
 	}
 
 	/** Everytime primary enable changes, this number increases. This is used to determine full frames. **/
-	static FORCEINLINE_STATS int32 PrimaryDisableChangeTag()
+	static inline int32 PrimaryDisableChangeTag()
 	{
 		if (PrimaryDisableChangeTagLock.GetValue())
 		{
@@ -1438,7 +1438,7 @@ public:
 	}
 
 	/** Call this if something disrupts data gathering. For example when the render thread is killed, data is abandoned.**/
-	static FORCEINLINE_STATS void FrameDataIsIncomplete()
+	static inline void FrameDataIsIncomplete()
 	{
 		FPlatformMisc::MemoryBarrier();
 		PrimaryEnableUpdateNumber.Increment();
@@ -1446,14 +1446,14 @@ public:
 	}
 
 	/** Enables the raw stats mode. */
-	static FORCEINLINE_STATS void EnableRawStats() TSAN_SAFE
+	static inline void EnableRawStats() TSAN_SAFE
 	{
 		bIsRawStatsActive = true;
 		FPlatformMisc::MemoryBarrier();
 	}
 
 	/** Disables the raw stats mode. */
-	static FORCEINLINE_STATS void DisableRawStats() TSAN_SAFE
+	static inline void DisableRawStats() TSAN_SAFE
 	{
 		bIsRawStatsActive = false;
 		FPlatformMisc::MemoryBarrier();
@@ -1524,7 +1524,7 @@ public:
 	 * Pushes the specified stat onto the hierarchy for this thread. Starts
 	 * the timing of the cycles used
 	 */
-	FORCEINLINE_STATS void Start(TStatId InStatId, EStatFlags InStatFlags, bool bAlways = false, UE::FSourceLocation SourceLocation = UE::FSourceLocation::Current())
+	inline void Start(TStatId InStatId, EStatFlags InStatFlags, bool bAlways = false, UE::FSourceLocation SourceLocation = UE::FSourceLocation::Current())
 	{
 		FMinimalName StatMinimalName = InStatId.GetMinimalName(EMemoryOrder::Relaxed);
 		if (StatMinimalName.IsNone())
@@ -1567,7 +1567,7 @@ public:
 		Start(InStatId, EStatFlags::None, bAlways, SourceLocation);
 	}
 
-	FORCEINLINE_STATS void StartTrace(const FName Name)
+	inline void StartTrace(const FName Name)
 	{
 #if CPUPROFILERTRACE_ENABLED
 		if (UE_TRACE_CHANNELEXPR_IS_ENABLED(CpuChannel))
@@ -1578,7 +1578,7 @@ public:
 #endif
 	}
 
-	FORCEINLINE_STATS void StartTrace(const FName Name, const TCHAR* Desc)
+	inline void StartTrace(const FName Name, const TCHAR* Desc)
 	{
 #if CPUPROFILERTRACE_ENABLED
 		if (UE_TRACE_CHANNELEXPR_IS_ENABLED(CpuChannel))
@@ -1592,7 +1592,7 @@ public:
 	/**
 	 * Stops the capturing and stores the result
 	 */
-	FORCEINLINE_STATS void Stop()
+	inline void Stop()
 	{
 		if (EmittedEvent & NamedEvent)
 		{
@@ -1617,7 +1617,7 @@ public:
 	/**
 	 * Stops the capturing and stores the result and resets the stat id.
 	 */
-	FORCEINLINE_STATS void StopAndResetStatId()
+	inline void StopAndResetStatId()
 	{
 		Stop();
 		StatId = NAME_None;
@@ -1636,7 +1636,7 @@ public:
 	 * Pushes the specified stat onto the hierarchy for this thread. Starts
 	 * the timing of the cycles used
 	 */
-	FORCEINLINE_STATS FScopeCycleCounter(TStatId StatId, EStatFlags StatFlags, bool bAlways = false, UE::FSourceLocation SourceLocation = UE::FSourceLocation::Current())
+	inline FScopeCycleCounter(TStatId StatId, EStatFlags StatFlags, bool bAlways = false, UE::FSourceLocation SourceLocation = UE::FSourceLocation::Current())
 	{
 		UE_AUTORTFM_OPEN
 		{
@@ -1653,7 +1653,7 @@ public:
 	/**
 	 * Updates the stat with the time spent
 	 */
-	FORCEINLINE_STATS ~FScopeCycleCounter()
+	inline ~FScopeCycleCounter()
 	{
 		AutoRTFM::PopOnAbortHandler(this);
 		UE_AUTORTFM_OPEN
@@ -1787,7 +1787,7 @@ protected:
 template<class TStatData, bool TCompiledIn>
 struct FThreadSafeStaticStatInner : public FThreadSafeStaticStatBase
 {
-	FORCEINLINE_STATS TStatId GetStatId() const
+	inline TStatId GetStatId() const
 	{
 		const TStatIdData* LocalHighPerformanceEnable = HighPerformanceEnable.Load(EMemoryOrder::Relaxed);
 		if (!LocalHighPerformanceEnable)
diff --git a/Engine/Source/Runtime/Core/Public/Templates/AlignmentTemplates.h b/Engine/Source/Runtime/Core/Public/Templates/AlignmentTemplates.h
--- a/Engine/Source/Runtime/Core/Public/Templates/AlignmentTemplates.h
+++ b/Engine/Source/Runtime/Core/Public/Templates/AlignmentTemplates.h
@@ -15,7 +15,7 @@
  * @return The value aligned up to the specified alignment.
  */
 template <typename T>
-FORCEINLINE constexpr T Align(T Val, uint64 Alignment)
+inline constexpr T Align(T Val, uint64 Alignment)
 {
 	static_assert(TIsIntegral<T>::Value || TIsPointer<T>::Value, "Align expects an integer or pointer type");
 
@@ -31,7 +31,7 @@ FORCEINLINE constexpr T Align(T Val, uint64 Alignment)
  * @return The value aligned down to the specified alignment.
  */
 template <typename T>
-FORCEINLINE constexpr T AlignDown(T Val, uint64 Alignment)
+inline constexpr T AlignDown(T Val, uint64 Alignment)
 {
 	static_assert(TIsIntegral<T>::Value || TIsPointer<T>::Value, "AlignDown expects an integer or pointer type");
 
@@ -47,7 +47,7 @@ FORCEINLINE constexpr T AlignDown(T Val, uint64 Alignment)
  * @return true if the pointer is aligned to the specified alignment, false otherwise.
  */
 template <typename T>
-FORCEINLINE constexpr bool IsAligned(T Val, uint64 Alignment)
+inline constexpr bool IsAligned(T Val, uint64 Alignment)
 {
 	static_assert(TIsIntegral<T>::Value || TIsPointer<T>::Value, "IsAligned expects an integer or pointer type");
 
@@ -63,7 +63,7 @@ FORCEINLINE constexpr bool IsAligned(T Val, uint64 Alignment)
  * @return The value aligned up to the specified alignment.
  */
 template <typename T>
-FORCEINLINE constexpr T AlignArbitrary(T Val, uint64 Alignment)
+inline constexpr T AlignArbitrary(T Val, uint64 Alignment)
 {
 	static_assert(TIsIntegral<T>::Value || TIsPointer<T>::Value, "AlignArbitrary expects an integer or pointer type");
 
