Replace some usages of FORCEINLINE with inline in Core modules.

diff --git a/Engine/Source/Runtime/Core/Public/Templates/Atomic.h b/Engine/Source/Runtime/Core/Public/Templates/Atomic.h
--- a/Engine/Source/Runtime/Core/Public/Templates/Atomic.h
+++ b/Engine/Source/Runtime/Core/Public/Templates/Atomic.h
@@ -57,7 +57,7 @@ namespace UE::Core::Private::Atomic
 	template <>           struct TIsVoidPointer<const volatile void*> { static constexpr bool Value = true; };
 
 	template <typename T>
-	FORCEINLINE T LoadRelaxed(const volatile T* Element)
+	inline T LoadRelaxed(const volatile T* Element)
 	{
 		#ifdef __clang__
 			TUnderlyingIntegerType_T<T> Result;
@@ -69,14 +69,14 @@ namespace UE::Core::Private::Atomic
 	}
 
 	template <typename T>
-	FORCEINLINE T Load(const volatile T* Element)
+	inline T Load(const volatile T* Element)
 	{
 		auto Result = FPlatformAtomics::AtomicRead((volatile TUnderlyingIntegerType_T<T>*)Element);
 		return *(const T*)&Result;
 	}
 
 	template <typename T>
-	FORCEINLINE void StoreRelaxed(const volatile T* Element, T Value)
+	inline void StoreRelaxed(const volatile T* Element, T Value)
 	{
 		#ifdef __clang__
 			__atomic_store((volatile TUnderlyingIntegerType_T<T>*)Element, (TUnderlyingIntegerType_T<T>*)&Value, __ATOMIC_RELAXED);
@@ -92,91 +92,91 @@ namespace UE::Core::Private::Atomic
 	}
 
 	template <typename T>
-	FORCEINLINE T Exchange(volatile T* Element, T Value)
+	inline T Exchange(volatile T* Element, T Value)
 	{
 		auto Result = FPlatformAtomics::InterlockedExchange((volatile TUnderlyingIntegerType_T<T>*)Element, *(const TUnderlyingIntegerType_T<T>*)&Value);
 		return *(const T*)&Result;
 	}
 
 	template <typename T>
-	FORCEINLINE T IncrementExchange(volatile T* Element)
+	inline T IncrementExchange(volatile T* Element)
 	{
 		auto Result = FPlatformAtomics::InterlockedIncrement((volatile TUnderlyingIntegerType_T<T>*)Element) - 1;
 		return *(const T*)&Result;
 	}
 
 	template <typename T>
-	FORCEINLINE T* IncrementExchange(T* volatile* Element)
+	inline T* IncrementExchange(T* volatile* Element)
 	{
 		auto Result = FPlatformAtomics::InterlockedAdd((volatile TUnderlyingIntegerType_T<T*>*)Element, sizeof(T));
 		return *(T* const*)&Result;
 	}
 
 	template <typename T, typename DiffType>
-	FORCEINLINE T AddExchange(volatile T* Element, DiffType Diff)
+	inline T AddExchange(volatile T* Element, DiffType Diff)
 	{
 		auto Result = FPlatformAtomics::InterlockedAdd((volatile TUnderlyingIntegerType_T<T>*)Element, (TUnderlyingIntegerType_T<T>)Diff);
 		return *(const T*)&Result;
 	}
 
 	template <typename T, typename DiffType>
-	FORCEINLINE T* AddExchange(T* volatile* Element, DiffType Diff)
+	inline T* AddExchange(T* volatile* Element, DiffType Diff)
 	{
 		auto Result = FPlatformAtomics::InterlockedAdd((volatile TUnderlyingIntegerType_T<T*>*)Element, (TUnderlyingIntegerType_T<T*>)(Diff * sizeof(T)));
 		return *(T* const*)&Result;
 	}
 
 	template <typename T>
-	FORCEINLINE T DecrementExchange(volatile T* Element)
+	inline T DecrementExchange(volatile T* Element)
 	{
 		auto Result = FPlatformAtomics::InterlockedDecrement((volatile TUnderlyingIntegerType_T<T>*)Element) + 1;
 		return *(const T*)&Result;
 	}
 
 	template <typename T>
-	FORCEINLINE T* DecrementExchange(T* volatile* Element)
+	inline T* DecrementExchange(T* volatile* Element)
 	{
 		auto Result = FPlatformAtomics::InterlockedAdd((volatile TUnderlyingIntegerType_T<T*>*)Element, -(TUnderlyingIntegerType_T<T*>)sizeof(T));
 		return *(T* const*)&Result;
 	}
 
 	template <typename T, typename DiffType>
-	FORCEINLINE T SubExchange(volatile T* Element, DiffType Diff)
+	inline T SubExchange(volatile T* Element, DiffType Diff)
 	{
 		auto Result = FPlatformAtomics::InterlockedAdd((volatile TUnderlyingIntegerType_T<T>*)Element, -(TUnderlyingIntegerType_T<T>)Diff);
 		return *(const T*)&Result;
 	}
 
 	template <typename T, typename DiffType>
-	FORCEINLINE T* SubExchange(T* volatile* Element, DiffType Diff)
+	inline T* SubExchange(T* volatile* Element, DiffType Diff)
 	{
 		auto Result = FPlatformAtomics::InterlockedAdd((volatile TUnderlyingIntegerType_T<T*>*)Element, -(TUnderlyingIntegerType_T<T*>)(Diff * sizeof(T)));
 		return *(T* const*)&Result;
 	}
 
 	template <typename T>
-	FORCEINLINE T CompareExchange(volatile T* Element, T ExpectedValue, T NewValue)
+	inline T CompareExchange(volatile T* Element, T ExpectedValue, T NewValue)
 	{
 		auto Result = FPlatformAtomics::InterlockedCompareExchange((volatile TUnderlyingIntegerType_T<T>*)Element, *(const TUnderlyingIntegerType_T<T>*)&NewValue, *(const TUnderlyingIntegerType_T<T>*)&ExpectedValue);
 		return *(const T*)&Result;
 	}
 
 	template <typename T>
-	FORCEINLINE T AndExchange(volatile T* Element, T AndValue)
+	inline T AndExchange(volatile T* Element, T AndValue)
 	{
 		auto Result = FPlatformAtomics::InterlockedAnd((volatile TUnderlyingIntegerType_T<T>*)Element, (TUnderlyingIntegerType_T<T>)AndValue);
 		return *(const T*)&Result;
 	}
 
 	template <typename T>
-	FORCEINLINE T OrExchange(volatile T* Element, T OrValue)
+	inline T OrExchange(volatile T* Element, T OrValue)
 	{
 		auto Result = FPlatformAtomics::InterlockedOr((volatile TUnderlyingIntegerType_T<T>*)Element, (TUnderlyingIntegerType_T<T>)OrValue);
 		return *(const T*)&Result;
 	}
 
 	template <typename T>
-	FORCEINLINE T XorExchange(volatile T* Element, T XorValue)
+	inline T XorExchange(volatile T* Element, T XorValue)
 	{
 		auto Result = FPlatformAtomics::InterlockedXor((volatile TUnderlyingIntegerType_T<T>*)Element, (TUnderlyingIntegerType_T<T>)XorValue);
 		return *(const T*)&Result;
@@ -575,7 +575,7 @@ public:
 	 *
 	 * @return A copy of Value.
 	 */
-	FORCEINLINE T operator=(T Value)
+	inline T operator=(T Value)
 	{
 		this->Exchange(Value);
 		return Value;
diff --git a/Engine/Source/Runtime/Core/Public/Templates/Invoke.h b/Engine/Source/Runtime/Core/Public/Templates/Invoke.h
--- a/Engine/Source/Runtime/Core/Public/Templates/Invoke.h
+++ b/Engine/Source/Runtime/Core/Public/Templates/Invoke.h
@@ -47,7 +47,7 @@ FORCEINLINE auto Invoke(FuncType&& Func, ArgTypes&&... Args)
 }
 
 template <typename ReturnType, typename ObjType, typename TargetType>
-FORCEINLINE auto Invoke(ReturnType ObjType::*pdm, TargetType&& Target)
+inline auto Invoke(ReturnType ObjType::*pdm, TargetType&& Target)
 	-> decltype(UE::Core::Private::DereferenceIfNecessary<ObjType>((TargetType&&)Target, &Target).*pdm)
 {
 	return UE::Core::Private::DereferenceIfNecessary<ObjType>((TargetType&&)Target, &Target).*pdm;
@@ -59,7 +59,7 @@ template <
 	typename... ArgTypes,
 	typename    ObjType = TMemberFunctionPtrOuter_T<PtrMemFunType>
 >
-FORCEINLINE auto Invoke(PtrMemFunType PtrMemFun, TargetType&& Target, ArgTypes&&... Args)
+inline auto Invoke(PtrMemFunType PtrMemFun, TargetType&& Target, ArgTypes&&... Args)
 	-> decltype((UE::Core::Private::DereferenceIfNecessary<ObjType>((TargetType&&)Target, &Target).*PtrMemFun)((ArgTypes&&)Args...))
 {
 	return (UE::Core::Private::DereferenceIfNecessary<ObjType>((TargetType&&)Target, &Target).*PtrMemFun)((ArgTypes&&)Args...);
diff --git a/Engine/Source/Runtime/Core/Public/Templates/NonNullPointer.h b/Engine/Source/Runtime/Core/Public/Templates/NonNullPointer.h
--- a/Engine/Source/Runtime/Core/Public/Templates/NonNullPointer.h
+++ b/Engine/Source/Runtime/Core/Public/Templates/NonNullPointer.h
@@ -43,7 +43,7 @@ public:
 	/**
 	 * Constructs a non-null pointer from the provided pointer. Must not be nullptr.
 	 */
-	FORCEINLINE TNonNullPtr(ObjectType* InObject)
+	inline TNonNullPtr(ObjectType* InObject)
 		: Object(InObject)
 	{
 		ensureMsgf(InObject, TEXT("Tried to initialize TNonNullPtr with a null pointer!"));
@@ -78,7 +78,7 @@ public:
 	/**
 	 * Assignment operator taking a nullptr - not allowed.
 	 */
-	FORCEINLINE TNonNullPtr& operator=(TYPE_OF_NULLPTR)
+	inline TNonNullPtr& operator=(TYPE_OF_NULLPTR)
 	{
 		// Essentially static_assert(false), but this way prevents GCC/Clang from crying wolf by merely inspecting the function body
 		static_assert(sizeof(ObjectType) == 0, "Tried to assign a null pointer to a TNonNullPtr!");
@@ -88,7 +88,7 @@ public:
 	/**
 	 * Assignment operator taking a pointer
 	 */
-	FORCEINLINE TNonNullPtr& operator=(ObjectType* InObject)
+	inline TNonNullPtr& operator=(ObjectType* InObject)
 	{
 		ensureMsgf(InObject, TEXT("Tried to assign a null pointer to a TNonNullPtr!"));
 		Object = InObject;
@@ -103,7 +103,7 @@ public:
 		typename OtherObjectType
 		UE_REQUIRES(std::is_convertible_v<OtherObjectType, ObjectType*>)
 	>
-	FORCEINLINE TNonNullPtr& operator=(TNotNull<OtherObjectType> InObject)
+	inline TNonNullPtr& operator=(TNotNull<OtherObjectType> InObject)
 	{
 		Object = InObject;
 		return *this;
@@ -117,7 +117,7 @@ public:
 		typename OtherObjectType
 		UE_REQUIRES(std::is_convertible_v<OtherObjectType*, ObjectType*>)
 	>
-	FORCEINLINE TNonNullPtr& operator=(const TNonNullPtr<OtherObjectType>& Other)
+	inline TNonNullPtr& operator=(const TNonNullPtr<OtherObjectType>& Other)
 	{
 		Object = Other.Get();
 		return *this;
@@ -178,7 +178,7 @@ public:
 	/**
 	 * Returns the internal pointer
 	 */
-	FORCEINLINE operator ObjectType*() const
+	inline operator ObjectType*() const
 	{
 		ensureMsgf(Object, TEXT("Tried to access null pointer!"));
 		return Object;
@@ -187,7 +187,7 @@ public:
 	/**
 	 * Returns the internal pointer
 	 */
-	FORCEINLINE ObjectType* Get() const
+	inline ObjectType* Get() const
 	{
 		ensureMsgf(Object, TEXT("Tried to access null pointer!"));
 		return Object;
@@ -196,7 +196,7 @@ public:
 	/**
 	 * Dereference operator returns a reference to the object this pointer points to
 	 */
-	FORCEINLINE ObjectType& operator*() const
+	inline ObjectType& operator*() const
 	{
 		ensureMsgf(Object, TEXT("Tried to access null pointer!"));
 		return *Object;
@@ -205,7 +205,7 @@ public:
 	/**
 	 * Arrow operator returns a pointer to this pointer's object
 	 */
-	FORCEINLINE ObjectType* operator->() const
+	inline ObjectType* operator->() const
 	{
 		ensureMsgf(Object, TEXT("Tried to access null pointer!"));
 		return Object;
diff --git a/Engine/Source/Runtime/Core/Public/Templates/PimplPtr.h b/Engine/Source/Runtime/Core/Public/Templates/PimplPtr.h
--- a/Engine/Source/Runtime/Core/Public/Templates/PimplPtr.h
+++ b/Engine/Source/Runtime/Core/Public/Templates/PimplPtr.h
@@ -118,7 +118,7 @@ namespace UE::Core::Private::PimplPtr
 		alignas(RequiredAlignment) T Val;
 	};
 
-	FORCEINLINE void CallDeleter(void* Ptr)
+	inline void CallDeleter(void* Ptr)
 	{
 		void* ThunkedPtr = (char*)Ptr - RequiredAlignment;
 
@@ -129,7 +129,7 @@ namespace UE::Core::Private::PimplPtr
 		(*(void(**)(void*) /*noexcept*/)ThunkedPtr)(ThunkedPtr);
 	}
 
-	FORCEINLINE void* CallCopier(void* Ptr)
+	inline void* CallCopier(void* Ptr)
 	{
 		void* BasePtr = (char*)Ptr - RequiredAlignment;
 		void* ThunkedPtr = (char*)BasePtr + sizeof(FDeleteFunc);
@@ -289,7 +289,7 @@ public:
 	{
 	}
 
-	FORCEINLINE TPimplPtr& operator=(TYPE_OF_NULLPTR A)
+	inline TPimplPtr& operator=(TYPE_OF_NULLPTR A)
 	{
 		Super::operator = (A);
 		return *this;
@@ -316,7 +316,7 @@ template <typename T, EPimplPtrMode Mode> FORCEINLINE bool operator!=(TYPE_OF_NU
  * DeepCopy Usage: TPimplPtr<FMyType, EPimplPtrMode::DeepCopy> MyPtr = MakePimpl<FMyType, EPimplPtrMode::DeepCopy>(...arguments...);
  */
 template <typename T, EPimplPtrMode Mode = EPimplPtrMode::NoCopy, typename... ArgTypes>
-FORCEINLINE TPimplPtr<T, Mode> MakePimpl(ArgTypes&&... Args)
+inline TPimplPtr<T, Mode> MakePimpl(ArgTypes&&... Args)
 {
 	using FHeapType = UE::Core::Private::PimplPtr::TPimplHeapObjectImpl<T>;
 	using FHeapConstructType = typename std::conditional<Mode == EPimplPtrMode::NoCopy, typename FHeapType::ENoCopyType,
diff --git a/Engine/Source/Runtime/Core/Public/Templates/RefCounting.h b/Engine/Source/Runtime/Core/Public/Templates/RefCounting.h
--- a/Engine/Source/Runtime/Core/Public/Templates/RefCounting.h
+++ b/Engine/Source/Runtime/Core/Public/Templates/RefCounting.h
@@ -485,7 +485,7 @@ public:
 		}
 	}
 
-	FORCEINLINE TRefCountPtr(TRefCountPtr&& Move)
+	inline TRefCountPtr(TRefCountPtr&& Move)
 	{
 		Reference = Move.Reference;
 		Move.Reference = nullptr;
@@ -575,7 +575,7 @@ public:
 		return Reference;
 	}
 
-	FORCEINLINE ReferencedType** GetInitReference()
+	inline ReferencedType** GetInitReference()
 	{
 		*this = nullptr;
 		return &Reference;
@@ -612,7 +612,7 @@ public:
 		return Result;
 	}
 
-	FORCEINLINE void Swap(TRefCountPtr& InPtr) // this does not change the reference count, and so is faster
+	inline void Swap(TRefCountPtr& InPtr) // this does not change the reference count, and so is faster
 	{
 		ReferencedType* OldReference = Reference;
 		Reference = InPtr.Reference;
@@ -676,7 +676,7 @@ template <
 	typename... TArgs
 	UE_REQUIRES(!std::is_array_v<T>)
 >
-[[nodiscard]] FORCEINLINE TRefCountPtr<T> MakeRefCount(TArgs&&... Args)
+[[nodiscard]] inline TRefCountPtr<T> MakeRefCount(TArgs&&... Args)
 {
 	T* NewUnwrappedObject = new T(Forward<TArgs>(Args)...);
 
diff --git a/Engine/Source/Runtime/Core/Public/Templates/ResolveTypeAmbiguity.h b/Engine/Source/Runtime/Core/Public/Templates/ResolveTypeAmbiguity.h
--- a/Engine/Source/Runtime/Core/Public/Templates/ResolveTypeAmbiguity.h
+++ b/Engine/Source/Runtime/Core/Public/Templates/ResolveTypeAmbiguity.h
@@ -46,7 +46,7 @@
 			std::is_signed_v<Arg2> && std::is_integral_v<Arg2> \
 		) \
 	> \
-	static OptionalMarkup FORCEINLINE auto Func(Arg1 X, Arg2 Y) -> decltype(X * Y) \
+	static OptionalMarkup inline auto Func(Arg1 X, Arg2 Y) -> decltype(X * Y) \
 	{ \
 		using ArgType = decltype(X * Y); \
 		return Func((ArgType)X, (ArgType)Y); \
@@ -72,7 +72,7 @@
 			std::is_signed_v<Arg3> && std::is_integral_v<Arg3> \
 		) \
 	> \
-	static OptionalMarkup FORCEINLINE auto Func(Arg1 X, Arg2 Y, Arg3 Z) -> decltype(X * Y * Z) \
+	static OptionalMarkup inline auto Func(Arg1 X, Arg2 Y, Arg3 Z) -> decltype(X * Y * Z) \
 	{ \
 		using ArgType = decltype(X * Y * Z); \
 		return Func((ArgType)X, (ArgType)Y, (ArgType)Z); \
@@ -89,7 +89,7 @@
 		typename Arg2 \
 		UE_REQUIRES((std::is_floating_point_v<Arg1> || std::is_floating_point_v<Arg2>) && !std::is_same_v<Arg1, Arg2>) \
 	> \
-	static FORCEINLINE auto Func(Arg1 X, Arg2 Y) -> decltype(X * Y) \
+	static inline auto Func(Arg1 X, Arg2 Y) -> decltype(X * Y) \
 	{ \
 		using ArgType = decltype(X * Y); \
 		return Func((ArgType)X, (ArgType)Y); \
@@ -106,7 +106,7 @@
 			(!std::is_same_v<Arg1, Arg2> || !std::is_same_v<Arg2, Arg3> || !std::is_same_v<Arg1, Arg3>) \
 		) \
 	> \
-	static FORCEINLINE auto Func(Arg1 X, Arg2 Y, Arg3 Z) -> decltype(X * Y * Z) \
+	static inline auto Func(Arg1 X, Arg2 Y, Arg3 Z) -> decltype(X * Y * Z) \
 	{ \
 		using ArgType = decltype(X * Y * Z); \
 		return Func((ArgType)X, (ArgType)Y, (ArgType)Z); \
@@ -135,7 +135,7 @@
 			!std::is_same_v<Arg1, Arg2> \
 		) \
 	> \
-	static FORCEINLINE ReturnType Func(Arg1 X, Arg2 Y) \
+	static inline ReturnType Func(Arg1 X, Arg2 Y) \
 	{ \
 		using ArgType = decltype(X * Y); \
 		return Func((ArgType)X, (ArgType)Y); \
@@ -152,7 +152,7 @@
 			(!std::is_same_v<Arg1, Arg2> || !std::is_same_v<Arg1, Arg3> || !std::is_same_v<Arg2, Arg3>) \
 		) \
 	> \
-	static FORCEINLINE ReturnType Func(Arg1 X, Arg2 Y, Arg3 Z) \
+	static inline ReturnType Func(Arg1 X, Arg2 Y, Arg3 Z) \
 	{ \
 		using ArgType = decltype(X * Y * Z); \
 		return Func((ArgType)X, (ArgType)Y, (ArgType)Z); \
diff --git a/Engine/Source/Runtime/Core/Public/Templates/SharedPointer.h b/Engine/Source/Runtime/Core/Public/Templates/SharedPointer.h
--- a/Engine/Source/Runtime/Core/Public/Templates/SharedPointer.h
+++ b/Engine/Source/Runtime/Core/Public/Templates/SharedPointer.h
@@ -168,7 +168,7 @@ public:
 		typename OtherType
 		UE_REQUIRES(std::is_convertible_v<OtherType*, ObjectType*>)
 	>
-	FORCEINLINE explicit TSharedRef( OtherType* InObject )
+	inline explicit TSharedRef( OtherType* InObject )
 		: Object( InObject )
 		, SharedReferenceCount( SharedPointerInternals::NewDefaultReferenceController< Mode >( InObject ) )
 	{
@@ -186,7 +186,7 @@ public:
 		typename DeleterType
 		UE_REQUIRES(std::is_convertible_v<OtherType*, ObjectType*>)
 	>
-	FORCEINLINE TSharedRef( OtherType* InObject, DeleterType&& InDeleter )
+	inline TSharedRef( OtherType* InObject, DeleterType&& InDeleter )
 		: Object( InObject )
 		, SharedReferenceCount( SharedPointerInternals::NewCustomReferenceController< Mode >( InObject, Forward< DeleterType >( InDeleter ) ) )
 	{
@@ -217,7 +217,7 @@ public:
 		typename OtherType
 		UE_REQUIRES(std::is_convertible_v<OtherType*, ObjectType*>)
 	>
-	FORCEINLINE TSharedRef( SharedPointerInternals::TRawPtrProxy< OtherType > const& InRawPtrProxy )
+	inline TSharedRef( SharedPointerInternals::TRawPtrProxy< OtherType > const& InRawPtrProxy )
 		: Object( InRawPtrProxy.Object )
 		, SharedReferenceCount( SharedPointerInternals::NewDefaultReferenceController< Mode >( InRawPtrProxy.Object ) )
 	{
@@ -242,7 +242,7 @@ public:
 		typename DeleterType
 		UE_REQUIRES(std::is_convertible_v<OtherType*, ObjectType*>)
 	>
-	FORCEINLINE TSharedRef( SharedPointerInternals::TRawPtrProxyWithDeleter< OtherType, DeleterType > const& InRawPtrProxy )
+	inline TSharedRef( SharedPointerInternals::TRawPtrProxyWithDeleter< OtherType, DeleterType > const& InRawPtrProxy )
 		: Object( InRawPtrProxy.Object )
 		, SharedReferenceCount( SharedPointerInternals::NewCustomReferenceController< Mode >( InRawPtrProxy.Object, InRawPtrProxy.Deleter ) )
 	{
@@ -267,7 +267,7 @@ public:
 		typename DeleterType
 		UE_REQUIRES(std::is_convertible_v<OtherType*, ObjectType*>)
 	>
-	FORCEINLINE TSharedRef( SharedPointerInternals::TRawPtrProxyWithDeleter< OtherType, DeleterType >&& InRawPtrProxy )
+	inline TSharedRef( SharedPointerInternals::TRawPtrProxyWithDeleter< OtherType, DeleterType >&& InRawPtrProxy )
 		: Object( InRawPtrProxy.Object )
 		, SharedReferenceCount( SharedPointerInternals::NewCustomReferenceController< Mode >( InRawPtrProxy.Object, MoveTemp( InRawPtrProxy.Deleter ) ) )
 	{
@@ -290,7 +290,7 @@ public:
 		typename OtherType
 		UE_REQUIRES(std::is_convertible_v<OtherType*, ObjectType*>)
 	>
-	FORCEINLINE TSharedRef( TSharedRef< OtherType, Mode > const& InSharedRef )
+	inline TSharedRef( TSharedRef< OtherType, Mode > const& InSharedRef )
 		: Object( InSharedRef.Object )
 		, SharedReferenceCount( InSharedRef.SharedReferenceCount )
 	{
@@ -305,7 +305,7 @@ public:
 	 * @param  InSharedRef  The shared reference whose object we should create an additional reference to
 	 */
 	template <typename OtherType>
-	FORCEINLINE TSharedRef( TSharedRef< OtherType, Mode > const& InSharedRef, SharedPointerInternals::FStaticCastTag )
+	inline TSharedRef( TSharedRef< OtherType, Mode > const& InSharedRef, SharedPointerInternals::FStaticCastTag )
 		: Object( static_cast< ObjectType* >( InSharedRef.Object ) )
 		, SharedReferenceCount( InSharedRef.SharedReferenceCount )
 	{
@@ -320,7 +320,7 @@ public:
 	 * @param  InSharedRef  The shared reference whose object we should create an additional reference to
 	 */
 	template <typename OtherType>
-	FORCEINLINE TSharedRef( TSharedRef< OtherType, Mode > const& InSharedRef, SharedPointerInternals::FConstCastTag )
+	inline TSharedRef( TSharedRef< OtherType, Mode > const& InSharedRef, SharedPointerInternals::FConstCastTag )
 		: Object( const_cast< ObjectType* >( InSharedRef.Object ) )
 		, SharedReferenceCount( InSharedRef.SharedReferenceCount )
 	{
@@ -334,7 +334,7 @@ public:
 	 * @param  InObject  The object pointer to use (instead of the incoming shared pointer's object)
 	 */
 	template <typename OtherType>
-	FORCEINLINE TSharedRef( TSharedRef< OtherType, Mode > const& OtherSharedRef, ObjectType* InObject )
+	inline TSharedRef( TSharedRef< OtherType, Mode > const& OtherSharedRef, ObjectType* InObject )
 		: Object( InObject )
 		, SharedReferenceCount( OtherSharedRef.SharedReferenceCount )
 	{
@@ -343,12 +343,12 @@ public:
 		check( InObject != nullptr );
 	}
 
-	FORCEINLINE TSharedRef( TSharedRef const& InSharedRef )
+	inline TSharedRef( TSharedRef const& InSharedRef )
 		: Object( InSharedRef.Object )
 		, SharedReferenceCount( InSharedRef.SharedReferenceCount )
 	{ }
 
-	FORCEINLINE TSharedRef( TSharedRef&& InSharedRef )
+	inline TSharedRef( TSharedRef&& InSharedRef )
 		: Object( InSharedRef.Object )
 		, SharedReferenceCount( InSharedRef.SharedReferenceCount )
 	{
@@ -364,14 +364,14 @@ public:
 	 *
 	 * @param  InSharedRef  Shared reference to replace with
 	 */
-	FORCEINLINE TSharedRef& operator=( TSharedRef const& InSharedRef )
+	inline TSharedRef& operator=( TSharedRef const& InSharedRef )
 	{
 		TSharedRef Temp = InSharedRef;
 		::Swap(Temp, *this);
 		return *this;
 	}
 
-	FORCEINLINE TSharedRef& operator=( TSharedRef&& InSharedRef )
+	inline TSharedRef& operator=( TSharedRef&& InSharedRef )
 	{
 		::Swap(*this, InSharedRef);
 		return *this;
@@ -389,7 +389,7 @@ public:
 		typename OtherType
 		UE_REQUIRES(std::is_convertible_v<OtherType*, ObjectType*>)
 	>
-	FORCEINLINE TSharedRef& operator=( SharedPointerInternals::TRawPtrProxy< OtherType > const& InRawPtrProxy )
+	inline TSharedRef& operator=( SharedPointerInternals::TRawPtrProxy< OtherType > const& InRawPtrProxy )
 	{
 		// If the following assert goes off, it means a TSharedRef was initialized from a nullptr object pointer.
 		// Shared references must never be nullptr, so either pass a valid object or consider using TSharedPtr instead.
@@ -412,7 +412,7 @@ public:
 		typename DeleterType
 		UE_REQUIRES(std::is_convertible_v<OtherType*, ObjectType*>)
 	>
-	FORCEINLINE TSharedRef& operator=( SharedPointerInternals::TRawPtrProxyWithDeleter< OtherType, DeleterType > const& InRawPtrProxy )
+	inline TSharedRef& operator=( SharedPointerInternals::TRawPtrProxyWithDeleter< OtherType, DeleterType > const& InRawPtrProxy )
 	{
 		// If the following assert goes off, it means a TSharedRef was initialized from a nullptr object pointer.
 		// Shared references must never be nullptr, so either pass a valid object or consider using TSharedPtr instead.
@@ -435,7 +435,7 @@ public:
 		typename DeleterType
 		UE_REQUIRES(std::is_convertible_v<OtherType*, ObjectType*>)
 	>
-	FORCEINLINE TSharedRef& operator=( SharedPointerInternals::TRawPtrProxyWithDeleter< OtherType, DeleterType >&& InRawPtrProxy )
+	inline TSharedRef& operator=( SharedPointerInternals::TRawPtrProxyWithDeleter< OtherType, DeleterType >&& InRawPtrProxy )
 	{
 		// If the following assert goes off, it means a TSharedRef was initialized from a nullptr object pointer.
 		// Shared references must never be nullptr, so either pass a valid object or consider using TSharedPtr instead.
@@ -470,7 +470,7 @@ public:
 	 *
 	 * @return  The object owned by this shared reference
 	 */
-	[[nodiscard]] FORCEINLINE ObjectType& Get() const
+	[[nodiscard]] inline ObjectType& Get() const
 	{
 		// Should never be nullptr as TSharedRef is never nullable
 		checkSlow( IsValid() );
@@ -482,7 +482,7 @@ public:
 	 *
 	 * @return  Reference to the object
 	 */
-	[[nodiscard]] FORCEINLINE ObjectType& operator*() const
+	[[nodiscard]] inline ObjectType& operator*() const
 	{
 		// Should never be nullptr as TSharedRef is never nullable
 		checkSlow( IsValid() );
@@ -494,7 +494,7 @@ public:
 	 *
 	 * @return  Returns a pointer to the object referenced by this shared reference
 	 */
-	[[nodiscard]] FORCEINLINE ObjectType* operator->() const
+	[[nodiscard]] inline ObjectType* operator->() const
 	{
 		// Should never be nullptr as TSharedRef is never nullable
 		checkSlow( IsValid() );
@@ -569,7 +569,7 @@ private:
 		typename OtherType
 		UE_REQUIRES(std::is_convertible_v<OtherType*, ObjectType*>)
 	>
-	FORCEINLINE explicit TSharedRef( TSharedPtr< OtherType, Mode > const& InSharedPtr )
+	inline explicit TSharedRef( TSharedPtr< OtherType, Mode > const& InSharedPtr )
 		: Object( InSharedPtr.Object )
 		, SharedReferenceCount( InSharedPtr.SharedReferenceCount )
 	{
@@ -582,7 +582,7 @@ private:
 		typename OtherType
 		UE_REQUIRES(std::is_convertible_v<OtherType*, ObjectType*>)
 	>
-	FORCEINLINE explicit TSharedRef( TSharedPtr< OtherType, Mode >&& InSharedPtr )
+	inline explicit TSharedRef( TSharedPtr< OtherType, Mode >&& InSharedPtr )
 		: Object( InSharedPtr.Object )
 		, SharedReferenceCount( MoveTemp(InSharedPtr.SharedReferenceCount) )
 	{
@@ -632,7 +632,7 @@ private:
 		#pragma warning(pop)
 	#endif
 
-	FORCEINLINE explicit TSharedRef(ObjectType* InObject, SharedPointerInternals::TReferenceControllerBase<Mode>* InSharedReferenceCount)
+	inline explicit TSharedRef(ObjectType* InObject, SharedPointerInternals::TReferenceControllerBase<Mode>* InSharedReferenceCount)
 		: Object(InObject)
 		, SharedReferenceCount(InSharedReferenceCount)
 	{
@@ -671,7 +671,7 @@ public:
 	 * Constructs an empty shared pointer
 	 */
 	// NOTE: FNullTag parameter is an Unreal extension to standard shared_ptr behavior
-	FORCEINLINE TSharedPtr( SharedPointerInternals::FNullTag* = nullptr )
+	inline TSharedPtr( SharedPointerInternals::FNullTag* = nullptr )
 		: Object( nullptr )
 		, SharedReferenceCount()
 	{
@@ -687,7 +687,7 @@ public:
 		typename OtherType
 		UE_REQUIRES(std::is_convertible_v<OtherType*, ObjectType*>)
 	>
-	FORCEINLINE explicit TSharedPtr( OtherType* InObject )
+	inline explicit TSharedPtr( OtherType* InObject )
 		: Object( InObject )
 		, SharedReferenceCount( SharedPointerInternals::NewDefaultReferenceController< Mode >( InObject ) )
 	{
@@ -708,7 +708,7 @@ public:
 		typename DeleterType
 		UE_REQUIRES(std::is_convertible_v<OtherType*, ObjectType*>)
 	>
-	FORCEINLINE TSharedPtr( OtherType* InObject, DeleterType&& InDeleter )
+	inline TSharedPtr( OtherType* InObject, DeleterType&& InDeleter )
 		: Object( InObject )
 		, SharedReferenceCount( SharedPointerInternals::NewCustomReferenceController< Mode >( InObject, Forward< DeleterType >( InDeleter ) ) )
 	{
@@ -727,7 +727,7 @@ public:
 		typename OtherType
 		UE_REQUIRES(std::is_convertible_v<OtherType*, ObjectType*>)
 	>
-	FORCEINLINE TSharedPtr( SharedPointerInternals::TRawPtrProxy< OtherType > const& InRawPtrProxy )
+	inline TSharedPtr( SharedPointerInternals::TRawPtrProxy< OtherType > const& InRawPtrProxy )
 		: Object( InRawPtrProxy.Object )
 		, SharedReferenceCount( SharedPointerInternals::NewDefaultReferenceController< Mode >( InRawPtrProxy.Object ) )
 	{
@@ -747,7 +747,7 @@ public:
 		typename DeleterType
 		UE_REQUIRES(std::is_convertible_v<OtherType*, ObjectType*>)
 	>
-	FORCEINLINE TSharedPtr( SharedPointerInternals::TRawPtrProxyWithDeleter< OtherType, DeleterType > const& InRawPtrProxy )
+	inline TSharedPtr( SharedPointerInternals::TRawPtrProxyWithDeleter< OtherType, DeleterType > const& InRawPtrProxy )
 		: Object( InRawPtrProxy.Object )
 		, SharedReferenceCount( SharedPointerInternals::NewCustomReferenceController< Mode >( InRawPtrProxy.Object, InRawPtrProxy.Deleter ) )
 	{
@@ -767,7 +767,7 @@ public:
 		typename DeleterType
 		UE_REQUIRES(std::is_convertible_v<OtherType*, ObjectType*>)
 	>
-	FORCEINLINE TSharedPtr( SharedPointerInternals::TRawPtrProxyWithDeleter< OtherType, DeleterType >&& InRawPtrProxy )
+	inline TSharedPtr( SharedPointerInternals::TRawPtrProxyWithDeleter< OtherType, DeleterType >&& InRawPtrProxy )
 		: Object( InRawPtrProxy.Object )
 		, SharedReferenceCount( SharedPointerInternals::NewCustomReferenceController< Mode >( InRawPtrProxy.Object, MoveTemp( InRawPtrProxy.Deleter ) ) )
 	{
@@ -786,19 +786,19 @@ public:
 		typename OtherType
 		UE_REQUIRES(std::is_convertible_v<OtherType*, ObjectType*>)
 	>
-	FORCEINLINE TSharedPtr( TSharedPtr< OtherType, Mode > const& InSharedPtr )
+	inline TSharedPtr( TSharedPtr< OtherType, Mode > const& InSharedPtr )
 		: Object( InSharedPtr.Object )
 		, SharedReferenceCount( InSharedPtr.SharedReferenceCount )
 	{
 	}
 
-	FORCEINLINE TSharedPtr( TSharedPtr const& InSharedPtr )
+	inline TSharedPtr( TSharedPtr const& InSharedPtr )
 		: Object( InSharedPtr.Object )
 		, SharedReferenceCount( InSharedPtr.SharedReferenceCount )
 	{
 	}
 
-	FORCEINLINE TSharedPtr( TSharedPtr&& InSharedPtr )
+	inline TSharedPtr( TSharedPtr&& InSharedPtr )
 		: Object( InSharedPtr.Object )
 		, SharedReferenceCount( MoveTemp(InSharedPtr.SharedReferenceCount) )
 	{
@@ -816,7 +816,7 @@ public:
 		typename OtherType
 		UE_REQUIRES(std::is_convertible_v<OtherType*, ObjectType*>)
 	>
-	FORCEINLINE TSharedPtr( TSharedRef< OtherType, Mode > const& InSharedRef )
+	inline TSharedPtr( TSharedRef< OtherType, Mode > const& InSharedRef )
 		: Object( InSharedRef.Object )
 		, SharedReferenceCount( InSharedRef.SharedReferenceCount )
 	{
@@ -833,7 +833,7 @@ public:
 	 * @param  InSharedPtr  The shared pointer whose object we should create an additional reference to
 	 */
 	template <typename OtherType>
-	FORCEINLINE TSharedPtr( TSharedPtr< OtherType, Mode > const& InSharedPtr, SharedPointerInternals::FStaticCastTag )
+	inline TSharedPtr( TSharedPtr< OtherType, Mode > const& InSharedPtr, SharedPointerInternals::FStaticCastTag )
 		: Object( static_cast< ObjectType* >( InSharedPtr.Object ) )
 		, SharedReferenceCount( InSharedPtr.SharedReferenceCount )
 	{
@@ -848,7 +848,7 @@ public:
 	 * @param  InSharedPtr  The shared pointer whose object we should create an additional reference to
 	 */
 	template <typename OtherType>
-	FORCEINLINE TSharedPtr( TSharedPtr< OtherType, Mode > const& InSharedPtr, SharedPointerInternals::FConstCastTag )
+	inline TSharedPtr( TSharedPtr< OtherType, Mode > const& InSharedPtr, SharedPointerInternals::FConstCastTag )
 		: Object( const_cast< ObjectType* >( InSharedPtr.Object ) )
 		, SharedReferenceCount( InSharedPtr.SharedReferenceCount )
 	{
@@ -862,7 +862,7 @@ public:
 	 * @param  InObject  The object pointer to use (instead of the incoming shared pointer's object)
 	 */
 	template <typename OtherType>
-	FORCEINLINE TSharedPtr( TSharedPtr< OtherType, Mode > const& OtherSharedPtr, ObjectType* InObject )
+	inline TSharedPtr( TSharedPtr< OtherType, Mode > const& OtherSharedPtr, ObjectType* InObject )
 		: Object( InObject )
 		, SharedReferenceCount( OtherSharedPtr.SharedReferenceCount )
 	{
@@ -876,7 +876,7 @@ public:
 	 * @param  InObject  The object pointer to use (instead of the incoming shared pointer's object)
 	 */
 	template <typename OtherType>
-	FORCEINLINE TSharedPtr( TSharedPtr< OtherType, Mode >&& OtherSharedPtr, ObjectType* InObject )
+	inline TSharedPtr( TSharedPtr< OtherType, Mode >&& OtherSharedPtr, ObjectType* InObject )
 		: Object( InObject )
 		, SharedReferenceCount( MoveTemp(OtherSharedPtr.SharedReferenceCount) )
 	{
@@ -891,7 +891,7 @@ public:
 	 * @param  InObject  The object pointer to use (instead of the incoming shared pointer's object)
 	 */
 	template <typename OtherType>
-	FORCEINLINE TSharedPtr( TSharedRef< OtherType, Mode > const& OtherSharedRef, ObjectType* InObject )
+	inline TSharedPtr( TSharedRef< OtherType, Mode > const& OtherSharedRef, ObjectType* InObject )
 		: Object( InObject )
 		, SharedReferenceCount( OtherSharedRef.SharedReferenceCount )
 	{
@@ -902,7 +902,7 @@ public:
 	 * referenced and will be deleted if there are no other referencers.
 	 */
 	// NOTE: The following is an Unreal extension to standard shared_ptr behavior
-	FORCEINLINE TSharedPtr& operator=( SharedPointerInternals::FNullTag* )
+	inline TSharedPtr& operator=( SharedPointerInternals::FNullTag* )
 	{
 		Reset();
 		return *this;
@@ -915,14 +915,14 @@ public:
 	 *
 	 * @param  InSharedPtr  Shared pointer to replace with
 	 */
-	FORCEINLINE TSharedPtr& operator=( TSharedPtr const& InSharedPtr )
+	inline TSharedPtr& operator=( TSharedPtr const& InSharedPtr )
 	{
 		TSharedPtr Temp = InSharedPtr;
 		::Swap(Temp, *this);
 		return *this;
 	}
 
-	FORCEINLINE TSharedPtr& operator=( TSharedPtr&& InSharedPtr )
+	inline TSharedPtr& operator=( TSharedPtr&& InSharedPtr )
 	{
 		if (this != &InSharedPtr)
 		{
@@ -945,7 +945,7 @@ public:
 		typename OtherType
 		UE_REQUIRES(std::is_convertible_v<OtherType*, ObjectType*>)
 	>
-	FORCEINLINE TSharedPtr& operator=( SharedPointerInternals::TRawPtrProxy< OtherType > const& InRawPtrProxy )
+	inline TSharedPtr& operator=( SharedPointerInternals::TRawPtrProxy< OtherType > const& InRawPtrProxy )
 	{
 		*this = TSharedPtr< ObjectType, Mode >( InRawPtrProxy );
 		return *this;
@@ -964,7 +964,7 @@ public:
 		typename DeleterType
 		UE_REQUIRES(std::is_convertible_v<OtherType*, ObjectType*>)
 	>
-	FORCEINLINE TSharedPtr& operator=( SharedPointerInternals::TRawPtrProxyWithDeleter< OtherType, DeleterType > const& InRawPtrProxy )
+	inline TSharedPtr& operator=( SharedPointerInternals::TRawPtrProxyWithDeleter< OtherType, DeleterType > const& InRawPtrProxy )
 	{
 		*this = TSharedPtr< ObjectType, Mode >( InRawPtrProxy );
 		return *this;
@@ -983,7 +983,7 @@ public:
 		typename DeleterType
 		UE_REQUIRES(std::is_convertible_v<OtherType*, ObjectType*>)
 	>
-	FORCEINLINE TSharedPtr& operator=( SharedPointerInternals::TRawPtrProxyWithDeleter< OtherType, DeleterType >&& InRawPtrProxy )
+	inline TSharedPtr& operator=( SharedPointerInternals::TRawPtrProxyWithDeleter< OtherType, DeleterType >&& InRawPtrProxy )
 	{
 		*this = TSharedPtr< ObjectType, Mode >( MoveTemp( InRawPtrProxy ) );
 		return *this;
@@ -995,7 +995,7 @@ public:
 	 * @return  Reference to the object
 	 */
 	// NOTE: The following is an Unreal extension to standard shared_ptr behavior
-	[[nodiscard]] FORCEINLINE TSharedRef< ObjectType, Mode > ToSharedRef() const&
+	[[nodiscard]] inline TSharedRef< ObjectType, Mode > ToSharedRef() const&
 	{
 		// If this assert goes off, it means a shared reference was created from a shared pointer that was nullptr.
 		// Shared references are never allowed to be null.  Consider using TSharedPtr instead.
@@ -1009,7 +1009,7 @@ public:
 	 * @return  Reference to the object
 	 */
 	 // NOTE: The following is an Unreal extension to standard shared_ptr behavior
-	[[nodiscard]] FORCEINLINE TSharedRef< ObjectType, Mode > ToSharedRef() &&
+	[[nodiscard]] inline TSharedRef< ObjectType, Mode > ToSharedRef() &&
 	{
 		// If this assert goes off, it means a shared reference was created from a shared pointer that was nullptr.
 		// Shared references are never allowed to be null.  Consider using TSharedPtr instead.
@@ -1066,7 +1066,7 @@ public:
 		typename DummyObjectType = ObjectType
 		UE_REQUIRES(UE_REQUIRES_EXPR(*(DummyObjectType*)nullptr)) // this construct means that operator* is only considered for overload resolution if T is dereferenceable
 	>
-	[[nodiscard]] FORCEINLINE DummyObjectType& operator*() const
+	[[nodiscard]] inline DummyObjectType& operator*() const
 	{
 		check( IsValid() );
 		return *Object;
@@ -1077,7 +1077,7 @@ public:
 	 *
 	 * @return  Returns a pointer to the object referenced by this shared pointer
 	 */
-	[[nodiscard]] FORCEINLINE ObjectType* operator->() const
+	[[nodiscard]] inline ObjectType* operator->() const
 	{
 		check( IsValid() );
 		return Object;
@@ -1147,7 +1147,7 @@ private:
 		typename OtherType
 		UE_REQUIRES(std::is_convertible_v<OtherType*, ObjectType*>)
 	>
-	FORCEINLINE explicit TSharedPtr( TWeakPtr< OtherType, Mode > const& InWeakPtr )
+	inline explicit TSharedPtr( TWeakPtr< OtherType, Mode > const& InWeakPtr )
 		: Object( nullptr )
 		, SharedReferenceCount( InWeakPtr.WeakReferenceCount )
 	{
@@ -1172,7 +1172,7 @@ private:
 		typename OtherType
 		UE_REQUIRES(std::is_convertible_v<OtherType*, ObjectType*>)
 	>
-	FORCEINLINE explicit TSharedPtr(TWeakPtr< OtherType, Mode >&& InWeakPtr)
+	inline explicit TSharedPtr(TWeakPtr< OtherType, Mode >&& InWeakPtr)
 		: Object(nullptr)
 		, SharedReferenceCount( MoveTemp( InWeakPtr.WeakReferenceCount ) )
 	{
@@ -1246,7 +1246,7 @@ public:
 
 	/** Constructs an empty TWeakPtr */
 	// NOTE: FNullTag parameter is an Unreal extension to standard shared_ptr behavior
-	FORCEINLINE TWeakPtr( SharedPointerInternals::FNullTag* = nullptr )
+	inline TWeakPtr( SharedPointerInternals::FNullTag* = nullptr )
 		: Object( nullptr )
 		, WeakReferenceCount()
 	{
@@ -1262,7 +1262,7 @@ public:
 		typename OtherType
 		UE_REQUIRES(std::is_convertible_v<OtherType*, ObjectType*>)
 	>
-	FORCEINLINE TWeakPtr( TSharedRef< OtherType, Mode > const& InSharedRef )
+	inline TWeakPtr( TSharedRef< OtherType, Mode > const& InSharedRef )
 		: Object( InSharedRef.Object )
 		, WeakReferenceCount( InSharedRef.SharedReferenceCount )
 	{
@@ -1277,7 +1277,7 @@ public:
 		typename OtherType
 		UE_REQUIRES(std::is_convertible_v<OtherType*, ObjectType*>)
 	>
-	FORCEINLINE TWeakPtr( TSharedPtr< OtherType, Mode > const& InSharedPtr )
+	inline TWeakPtr( TSharedPtr< OtherType, Mode > const& InSharedPtr )
 		: Object( InSharedPtr.Object )
 		, WeakReferenceCount( InSharedPtr.SharedReferenceCount )
 	{
@@ -1292,7 +1292,7 @@ public:
 	 * @param  InWeakPtr  The weak pointer whose object we should create an additional reference to
 	 */
 	template <typename OtherType>
-	FORCEINLINE TWeakPtr(TWeakPtr< OtherType, Mode > const& InWeakPtr, SharedPointerInternals::FStaticCastTag)
+	inline TWeakPtr(TWeakPtr< OtherType, Mode > const& InWeakPtr, SharedPointerInternals::FStaticCastTag)
 		: Object(static_cast<ObjectType*>(InWeakPtr.Object))
 		, WeakReferenceCount(InWeakPtr.WeakReferenceCount)
 	{
@@ -1307,7 +1307,7 @@ public:
 	 * @param  InWeakPtr  The weak pointer whose object we should create an additional reference to
 	 */
 	template <typename OtherType>
-	FORCEINLINE TWeakPtr(TWeakPtr< OtherType, Mode > const& InWeakPtr, SharedPointerInternals::FConstCastTag)
+	inline TWeakPtr(TWeakPtr< OtherType, Mode > const& InWeakPtr, SharedPointerInternals::FConstCastTag)
 		: Object(const_cast<ObjectType*>(InWeakPtr.Object))
 		, WeakReferenceCount(InWeakPtr.WeakReferenceCount)
 	{
@@ -1323,7 +1323,7 @@ public:
 		typename OtherType
 		UE_REQUIRES(std::is_convertible_v<OtherType*, ObjectType*>)
 	>
-	FORCEINLINE TWeakPtr( TWeakPtr< OtherType, Mode > const& InWeakPtr )
+	inline TWeakPtr( TWeakPtr< OtherType, Mode > const& InWeakPtr )
 		: Object( InWeakPtr.Object )
 		, WeakReferenceCount( InWeakPtr.WeakReferenceCount )
 	{
@@ -1333,20 +1333,20 @@ public:
 		typename OtherType
 		UE_REQUIRES(std::is_convertible_v<OtherType*, ObjectType*>)
 	>
-	FORCEINLINE TWeakPtr( TWeakPtr< OtherType, Mode >&& InWeakPtr )
+	inline TWeakPtr( TWeakPtr< OtherType, Mode >&& InWeakPtr )
 		: Object( InWeakPtr.Object )
 		, WeakReferenceCount( MoveTemp(InWeakPtr.WeakReferenceCount) )
 	{
 		InWeakPtr.Object = nullptr;
 	}
 
-	FORCEINLINE TWeakPtr( TWeakPtr const& InWeakPtr )
+	inline TWeakPtr( TWeakPtr const& InWeakPtr )
 		: Object( InWeakPtr.Object )
 		, WeakReferenceCount( InWeakPtr.WeakReferenceCount )
 	{
 	}
 
-	FORCEINLINE TWeakPtr( TWeakPtr&& InWeakPtr )
+	inline TWeakPtr( TWeakPtr&& InWeakPtr )
 		: Object( InWeakPtr.Object )
 		, WeakReferenceCount( MoveTemp(InWeakPtr.WeakReferenceCount) )
 	{
@@ -1357,7 +1357,7 @@ public:
 	 * Assignment to a nullptr pointer.  Clears this weak pointer's reference.
 	 */
 	// NOTE: The following is an Unreal extension to standard shared_ptr behavior
-	FORCEINLINE TWeakPtr& operator=( SharedPointerInternals::FNullTag* )
+	inline TWeakPtr& operator=( SharedPointerInternals::FNullTag* )
 	{
 		Reset();
 		return *this;
@@ -1368,14 +1368,14 @@ public:
 	 *
 	 * @param  InWeakPtr  The weak pointer for the object to assign
 	 */
-	FORCEINLINE TWeakPtr& operator=( TWeakPtr const& InWeakPtr )
+	inline TWeakPtr& operator=( TWeakPtr const& InWeakPtr )
 	{
 		TWeakPtr Temp = InWeakPtr;
 		::Swap(Temp, *this);
 		return *this;
 	}
 
-	FORCEINLINE TWeakPtr& operator=( TWeakPtr&& InWeakPtr )
+	inline TWeakPtr& operator=( TWeakPtr&& InWeakPtr )
 	{
 		if (this != &InWeakPtr)
 		{
@@ -1396,7 +1396,7 @@ public:
 		typename OtherType
 		UE_REQUIRES(std::is_convertible_v<OtherType*, ObjectType*>)
 	>
-	FORCEINLINE TWeakPtr& operator=( TWeakPtr<OtherType, Mode> const& InWeakPtr )
+	inline TWeakPtr& operator=( TWeakPtr<OtherType, Mode> const& InWeakPtr )
 	{
 		Object = InWeakPtr.Pin().Get();
 		WeakReferenceCount = InWeakPtr.WeakReferenceCount;
@@ -1407,7 +1407,7 @@ public:
 		typename OtherType
 		UE_REQUIRES(std::is_convertible_v<OtherType*, ObjectType*>)
 	>
-	FORCEINLINE TWeakPtr& operator=( TWeakPtr<OtherType, Mode>&& InWeakPtr )
+	inline TWeakPtr& operator=( TWeakPtr<OtherType, Mode>&& InWeakPtr )
 	{
 		Object             = InWeakPtr.Object;
 		InWeakPtr.Object   = nullptr;
@@ -1425,7 +1425,7 @@ public:
 		typename OtherType
 		UE_REQUIRES(std::is_convertible_v<OtherType*, ObjectType*>)
 	>
-	FORCEINLINE TWeakPtr& operator=( TSharedRef< OtherType, Mode > const& InSharedRef )
+	inline TWeakPtr& operator=( TSharedRef< OtherType, Mode > const& InSharedRef )
 	{
 		Object = InSharedRef.Object;
 		WeakReferenceCount = InSharedRef.SharedReferenceCount;
@@ -1441,7 +1441,7 @@ public:
 		typename OtherType
 		UE_REQUIRES(std::is_convertible_v<OtherType*, ObjectType*>)
 	>
-	FORCEINLINE TWeakPtr& operator=( TSharedPtr< OtherType, Mode > const& InSharedPtr )
+	inline TWeakPtr& operator=( TSharedPtr< OtherType, Mode > const& InSharedPtr )
 	{
 		Object = InSharedPtr.Object;
 		WeakReferenceCount = InSharedPtr.SharedReferenceCount;
@@ -1735,7 +1735,7 @@ public:		// @todo: Ideally this would be private, but template sharing problems
 	 * Note that until this function is called, calls to AsShared() will result in an empty pointer.
 	 */
 	template< class SharedPtrType, ESPMode SharedPtrMode, class OtherType >
-	FORCEINLINE void UpdateWeakReferenceInternal( TSharedPtr< SharedPtrType, SharedPtrMode > const* InSharedPtr, OtherType* InObject ) const
+	inline void UpdateWeakReferenceInternal( TSharedPtr< SharedPtrType, SharedPtrMode > const* InSharedPtr, OtherType* InObject ) const
 	{
 		static_assert(SharedPtrMode == Mode, "You cannot use a TSharedPtr of one mode with a type which inherits TSharedFromThis of another mode.");
 
@@ -1751,7 +1751,7 @@ public:		// @todo: Ideally this would be private, but template sharing problems
 	 * Note that until this function is called, calls to AsShared() will result in an empty pointer.
 	 */
 	template< class SharedRefType, ESPMode SharedPtrMode, class OtherType >
-	FORCEINLINE void UpdateWeakReferenceInternal( TSharedRef< SharedRefType, SharedPtrMode > const* InSharedRef, OtherType* InObject ) const
+	inline void UpdateWeakReferenceInternal( TSharedRef< SharedRefType, SharedPtrMode > const* InSharedRef, OtherType* InObject ) const
 	{
 		static_assert(SharedPtrMode == Mode, "You cannot use a TSharedPtr of one mode with a type which inherits TSharedFromThis of another mode.");
 
@@ -2331,7 +2331,7 @@ template< class CastToType, class CastFromType, ESPMode Mode >
  */
 // NOTE: The following is an Unreal extension to standard shared_ptr behavior
 template< class ObjectType >
-[[nodiscard]] FORCEINLINE SharedPointerInternals::TRawPtrProxy< ObjectType > MakeShareable( ObjectType* InObject )
+[[nodiscard]] inline SharedPointerInternals::TRawPtrProxy< ObjectType > MakeShareable( ObjectType* InObject )
 {
 	///////////////////////////////////////////////////////////////////////////////////////////////
 	// If you get an 'ambiguous call' compile error in this function, it means you have multiple //
@@ -2355,7 +2355,7 @@ template< class ObjectType >
  */
 // NOTE: The following is an Unreal extension to standard shared_ptr behavior
 template< class ObjectType, class DeleterType >
-[[nodiscard]] FORCEINLINE SharedPointerInternals::TRawPtrProxyWithDeleter< ObjectType, DeleterType > MakeShareable( ObjectType* InObject, DeleterType&& InDeleter )
+[[nodiscard]] inline SharedPointerInternals::TRawPtrProxyWithDeleter< ObjectType, DeleterType > MakeShareable( ObjectType* InObject, DeleterType&& InDeleter )
 {
 	///////////////////////////////////////////////////////////////////////////////////////////////
 	// If you get an 'ambiguous call' compile error in this function, it means you have multiple //
@@ -2393,7 +2393,7 @@ template< class ObjectType, class DeleterType >
  *     TSharedPtr<FMyType> Val = MakeShared<FMyType>(FMyType::FPrivateToken{}, 5, 3.14f, TEXT("Banana"));
  */
 template <typename InObjectType, ESPMode InMode = ESPMode::ThreadSafe, typename... InArgTypes>
-[[nodiscard]] FORCEINLINE TSharedRef<InObjectType, InMode> MakeShared(InArgTypes&&... Args)
+[[nodiscard]] inline TSharedRef<InObjectType, InMode> MakeShared(InArgTypes&&... Args)
 {
 	SharedPointerInternals::TIntrusiveReferenceController<InObjectType, InMode>* Controller = SharedPointerInternals::NewIntrusiveReferenceController<InMode, InObjectType>(Forward<InArgTypes>(Args)...);
 	return UE::Core::Private::MakeSharedRef<InObjectType, InMode>(Controller->GetObjectPtr(), (SharedPointerInternals::TReferenceControllerBase<InMode>*)Controller);
@@ -2405,7 +2405,7 @@ template <typename InObjectType, ESPMode InMode = ESPMode::ThreadSafe, typename.
  * @param  PointerArray  The pointer array to prune invalid pointers out of
  */
 template <class Type>
-FORCEINLINE void CleanupPointerArray(TArray< TWeakPtr<Type> >& PointerArray)
+inline void CleanupPointerArray(TArray< TWeakPtr<Type> >& PointerArray)
 {
 	TArray< TWeakPtr<Type> > NewArray;
 	for (int32 i = 0; i < PointerArray.Num(); ++i)
@@ -2424,7 +2424,7 @@ FORCEINLINE void CleanupPointerArray(TArray< TWeakPtr<Type> >& PointerArray)
  * @param  PointerMap  The pointer map to prune invalid pointers out of
  */
 template <class KeyType, class ValueType>
-FORCEINLINE void CleanupPointerMap(TMap< TWeakPtr<KeyType>, ValueType >& PointerMap)
+inline void CleanupPointerMap(TMap< TWeakPtr<KeyType>, ValueType >& PointerMap)
 {
 	TMap< TWeakPtr<KeyType>, ValueType > NewMap;
 	for (typename TMap< TWeakPtr<KeyType>, ValueType >::TConstIterator Op(PointerMap); Op; ++Op)
diff --git a/Engine/Source/Runtime/Core/Public/Templates/SharedPointerInternals.h b/Engine/Source/Runtime/Core/Public/Templates/SharedPointerInternals.h
--- a/Engine/Source/Runtime/Core/Public/Templates/SharedPointerInternals.h
+++ b/Engine/Source/Runtime/Core/Public/Templates/SharedPointerInternals.h
@@ -70,7 +70,7 @@ namespace SharedPointerInternals
 		}
 
 		/** Returns the shared reference count */
-		FORCEINLINE int32 GetSharedReferenceCount() const
+		inline int32 GetSharedReferenceCount() const
 		{
 			if constexpr (Mode == ESPMode::ThreadSafe)
 			{
@@ -103,7 +103,7 @@ namespace SharedPointerInternals
 		}
 
 		/** Adds a shared reference to this counter */
-		FORCEINLINE void AddSharedReference()
+		inline void AddSharedReference()
 		{
 			if constexpr (Mode == ESPMode::ThreadSafe)
 			{
@@ -209,7 +209,7 @@ namespace SharedPointerInternals
 		}
 
 		/** Releases a shared reference to this counter */
-		FORCEINLINE void ReleaseSharedReference()
+		inline void ReleaseSharedReference()
 		{
 			if constexpr (Mode == ESPMode::ThreadSafe)
 			{
@@ -249,7 +249,7 @@ namespace SharedPointerInternals
 		}
 
 		/** Adds a weak reference to this counter */
-		FORCEINLINE void AddWeakReference()
+		inline void AddWeakReference()
 		{
 			if constexpr (Mode == ESPMode::ThreadSafe)
 			{
@@ -500,14 +500,14 @@ namespace SharedPointerInternals
 		DeleterType Deleter;
 
 		/** Construct implicitly from an object and a custom deleter */
-		FORCEINLINE TRawPtrProxyWithDeleter( ObjectType* InObject, const DeleterType& InDeleter )
+		inline TRawPtrProxyWithDeleter( ObjectType* InObject, const DeleterType& InDeleter )
 			: Object ( InObject )
 			, Deleter( InObject, InDeleter )
 		{
 		}
 
 		/** Construct implicitly from an object and a custom deleter */
-		FORCEINLINE TRawPtrProxyWithDeleter( ObjectType* InObject, DeleterType&& InDeleter )
+		inline TRawPtrProxyWithDeleter( ObjectType* InObject, DeleterType&& InDeleter )
 			: Object ( InObject )
 			, Deleter( MoveTemp( InDeleter ) )
 		{
@@ -535,7 +535,7 @@ namespace SharedPointerInternals
 		{ }
 		
 		/** Copy constructor creates a new reference to the existing object */
-		FORCEINLINE FSharedReferencer( FSharedReferencer const& InSharedReference )
+		inline FSharedReferencer( FSharedReferencer const& InSharedReference )
 			: ReferenceController( InSharedReference.ReferenceController )
 		{
 			// If the incoming reference had an object associated with it, then go ahead and increment the
@@ -547,7 +547,7 @@ namespace SharedPointerInternals
 		}
 
 		/** Move constructor creates no new references */
-		FORCEINLINE FSharedReferencer( FSharedReferencer&& InSharedReference )
+		inline FSharedReferencer( FSharedReferencer&& InSharedReference )
 			: ReferenceController( InSharedReference.ReferenceController )
 		{
 			InSharedReference.ReferenceController = nullptr;
@@ -597,7 +597,7 @@ namespace SharedPointerInternals
 		}
 
 		/** Destructor. */
-		FORCEINLINE ~FSharedReferencer()
+		inline ~FSharedReferencer()
 		{
 			if( ReferenceController != nullptr )
 			{
@@ -715,7 +715,7 @@ namespace SharedPointerInternals
 		{ }
 
 		/** Construct a weak referencer object from another weak referencer */
-		FORCEINLINE FWeakReferencer( FWeakReferencer const& InWeakRefCountPointer )
+		inline FWeakReferencer( FWeakReferencer const& InWeakRefCountPointer )
 			: ReferenceController( InWeakRefCountPointer.ReferenceController )
 		{
 			// If the weak referencer has a valid controller, then go ahead and add a weak reference to it!
@@ -726,14 +726,14 @@ namespace SharedPointerInternals
 		}
 
 		/** Construct a weak referencer object from an rvalue weak referencer */
-		FORCEINLINE FWeakReferencer( FWeakReferencer&& InWeakRefCountPointer )
+		inline FWeakReferencer( FWeakReferencer&& InWeakRefCountPointer )
 			: ReferenceController( InWeakRefCountPointer.ReferenceController )
 		{
 			InWeakRefCountPointer.ReferenceController = nullptr;
 		}
 
 		/** Construct a weak referencer object from a shared referencer object */
-		FORCEINLINE FWeakReferencer( FSharedReferencer< Mode > const& InSharedRefCountPointer )
+		inline FWeakReferencer( FSharedReferencer< Mode > const& InSharedRefCountPointer )
 			: ReferenceController( InSharedRefCountPointer.ReferenceController )
 		{
 			// If the shared referencer had a valid controller, then go ahead and add a weak reference to it!
@@ -744,7 +744,7 @@ namespace SharedPointerInternals
 		}
 
 		/** Destructor. */
-		FORCEINLINE ~FWeakReferencer()
+		inline ~FWeakReferencer()
 		{
 			if( ReferenceController != nullptr )
 			{
@@ -756,7 +756,7 @@ namespace SharedPointerInternals
 		
 		/** Assignment operator from a weak referencer object.  If this counter was previously referencing an
 		    object, that reference will be released. */
-		FORCEINLINE FWeakReferencer& operator=( FWeakReferencer const& InWeakReference )
+		inline FWeakReferencer& operator=( FWeakReferencer const& InWeakReference )
 		{
 			AssignReferenceController( InWeakReference.ReferenceController );
 
@@ -765,7 +765,7 @@ namespace SharedPointerInternals
 
 		/** Assignment operator from an rvalue weak referencer object.  If this counter was previously referencing an
 		    object, that reference will be released. */
-		FORCEINLINE FWeakReferencer& operator=( FWeakReferencer&& InWeakReference )
+		inline FWeakReferencer& operator=( FWeakReferencer&& InWeakReference )
 		{
 			auto OldReferenceController         = ReferenceController;
 			ReferenceController                 = InWeakReference.ReferenceController;
@@ -780,7 +780,7 @@ namespace SharedPointerInternals
 
 		/** Assignment operator from a shared reference counter.  If this counter was previously referencing an
 		   object, that reference will be released. */
-		FORCEINLINE FWeakReferencer& operator=( FSharedReferencer< Mode > const& InSharedReference )
+		inline FWeakReferencer& operator=( FSharedReferencer< Mode > const& InSharedReference )
 		{
 			AssignReferenceController( InSharedReference.ReferenceController );
 
@@ -855,7 +855,7 @@ namespace SharedPointerInternals
 {
 	/** Templated helper function (const) that creates a shared pointer from an object instance */
 	template< class SharedPtrType, class ObjectType >
-	FORCEINLINE void EnableSharedFromThis( SharedPtrType* InSharedPtrOrRef, ObjectType const* InObject )
+	inline void EnableSharedFromThis( SharedPtrType* InSharedPtrOrRef, ObjectType const* InObject )
 	{
 		///////////////////////////////////////////////////////////////////////////////////////////////
 		// If you get an 'ambiguous call' compile error in this function, it means you have multiple //
diff --git a/Engine/Source/Runtime/Core/Public/Templates/Sorting.h b/Engine/Source/Runtime/Core/Public/Templates/Sorting.h
--- a/Engine/Source/Runtime/Core/Public/Templates/Sorting.h
+++ b/Engine/Source/Runtime/Core/Public/Templates/Sorting.h
@@ -527,7 +527,7 @@ void RadixSort32( ValueType* RESTRICT Dst, ValueType* RESTRICT Src, CountType Nu
 // http://codercorner.com/RadixSortRevisited.htm
 struct FRadixSortKeyFloat
 {
-	FORCEINLINE uint32 operator()( float Value ) const
+	inline uint32 operator()( float Value ) const
 	{
 		union { float f; uint32 i; } v;
 		v.f = Value;
diff --git a/Engine/Source/Runtime/Core/Public/Templates/Tuple.h b/Engine/Source/Runtime/Core/Public/Templates/Tuple.h
--- a/Engine/Source/Runtime/Core/Public/Templates/Tuple.h
+++ b/Engine/Source/Runtime/Core/Public/Templates/Tuple.h
@@ -123,7 +123,7 @@ namespace UE::Core::Private::Tuple
 	struct TTupleElementGetterByIndex
 	{
 		template <typename DeducedType, typename TupleType>
-		static FORCEINLINE decltype(auto) GetImpl(const volatile TTupleBaseElement<DeducedType, Index, TupleSize>&, TupleType&& Tuple)
+		static inline decltype(auto) GetImpl(const volatile TTupleBaseElement<DeducedType, Index, TupleSize>&, TupleType&& Tuple)
 		{
 			// Brackets are important here - we want a reference type to be returned, not object type
 			decltype(auto) Result = (ForwardAsBase<TupleType, TTupleBaseElement<DeducedType, Index, TupleSize>>(Tuple).Value);
@@ -144,7 +144,7 @@ namespace UE::Core::Private::Tuple
 	struct TTupleElementGetterByIndex<0, 2>
 	{
 		template <typename TupleType>
-		static FORCEINLINE decltype(auto) Get(TupleType&& Tuple)
+		static inline decltype(auto) Get(TupleType&& Tuple)
 		{
 			// Brackets are important here - we want a reference type to be returned, not object type
 			decltype(auto) Result = (Forward<TupleType>(Tuple).Key);
@@ -157,7 +157,7 @@ namespace UE::Core::Private::Tuple
 	struct TTupleElementGetterByIndex<1, 2>
 	{
 		template <typename TupleType>
-		static FORCEINLINE decltype(auto) Get(TupleType&& Tuple)
+		static inline decltype(auto) Get(TupleType&& Tuple)
 		{
 			// Brackets are important here - we want a reference type to be returned, not object type
 			decltype(auto) Result = (Forward<TupleType>(Tuple).Value);
@@ -171,7 +171,7 @@ namespace UE::Core::Private::Tuple
 	struct TTupleElementGetterByIndex<0, 2>
 	{
 		template <typename TupleType>
-		static FORCEINLINE decltype(auto) Get(TupleType&& Tuple)
+		static inline decltype(auto) Get(TupleType&& Tuple)
 		{
 			// Brackets are important here - we want a reference type to be returned, not object type
 			decltype(auto) Result = (ForwardAsBase<TupleType, TTupleBaseElement<decltype(Tuple.Key), 0, 2>>(Tuple).Key);
@@ -203,7 +203,7 @@ namespace UE::Core::Private::Tuple
 	struct TTupleElementGetterByType<Type, 2>
 	{
 		template <typename TupleType>
-		static FORCEINLINE decltype(auto) Get(TupleType&& Tuple)
+		static inline decltype(auto) Get(TupleType&& Tuple)
 		{
 			if constexpr (std::is_same_v<Type, decltype(Tuple.Key)>)
 			{
@@ -334,14 +334,14 @@ namespace UE::Core::Private::Tuple
 		template <typename FuncType, typename... ArgTypes> decltype(auto) ApplyBefore(FuncType&& Func, ArgTypes&&... Args)       volatile&& { return ::Invoke(Func, static_cast<      volatile TTupleBase&&>(*this).template Get<Indices>()..., Forward<ArgTypes>(Args)...); }
 		template <typename FuncType, typename... ArgTypes> decltype(auto) ApplyBefore(FuncType&& Func, ArgTypes&&... Args) const volatile&& { return ::Invoke(Func, static_cast<const volatile TTupleBase&&>(*this).template Get<Indices>()..., Forward<ArgTypes>(Args)...); }
 
-		FORCEINLINE static void Serialize(FArchive& Ar, TTupleBase& Tuple)
+		inline static void Serialize(FArchive& Ar, TTupleBase& Tuple)
 		{
 			// This should be implemented with a fold expression when our compilers support it
 			int Temp[] = { 0, (Ar << Tuple.template Get<Indices>(), 0)... };
 			(void)Temp;
 		}
 
-		FORCEINLINE static void SerializeStructured(FStructuredArchive::FSlot Slot, TTupleBase& Tuple)
+		inline static void SerializeStructured(FStructuredArchive::FSlot Slot, TTupleBase& Tuple)
 		{
 			// This should be implemented with a fold expression when our compilers support it
 			FStructuredArchive::FStream Stream = Slot.EnterStream();
@@ -454,13 +454,13 @@ namespace UE::Core::Private::Tuple
 		template <typename FuncType, typename... ArgTypes> decltype(auto) ApplyBefore(FuncType&& Func, ArgTypes&&... Args)       volatile&& { return ::Invoke(Func, static_cast<      volatile TTupleBase&&>(*this).template Get<0>(), static_cast<      volatile TTupleBase&&>(*this).template Get<1>(), Forward<ArgTypes>(Args)...); }
 		template <typename FuncType, typename... ArgTypes> decltype(auto) ApplyBefore(FuncType&& Func, ArgTypes&&... Args) const volatile&& { return ::Invoke(Func, static_cast<const volatile TTupleBase&&>(*this).template Get<0>(), static_cast<const volatile TTupleBase&&>(*this).template Get<1>(), Forward<ArgTypes>(Args)...); }
 
-		FORCEINLINE static void Serialize(FArchive& Ar, TTupleBase& Tuple)
+		inline static void Serialize(FArchive& Ar, TTupleBase& Tuple)
 		{
 			Ar << Tuple.Key;
 			Ar << Tuple.Value;
 		}
 
-		FORCEINLINE static void SerializeStructured(FStructuredArchive::FSlot Slot, TTupleBase& Tuple)
+		inline static void SerializeStructured(FStructuredArchive::FSlot Slot, TTupleBase& Tuple)
 		{
 			FStructuredArchive::FStream Stream = Slot.EnterStream();
 			Stream.EnterElement() << Tuple.Key;
@@ -994,7 +994,7 @@ decltype(auto) get(TupleType&& val)
 ///////////////////////////////////////
 
 template <typename... Types>
-FORCEINLINE FArchive& operator<<(FArchive& Ar, TTuple<Types...>& Tuple)
+inline FArchive& operator<<(FArchive& Ar, TTuple<Types...>& Tuple)
 {
 	TTuple<Types...>::Serialize(Ar, Tuple);
 	return Ar;
diff --git a/Engine/Source/Runtime/Core/Public/Templates/TypeHash.h b/Engine/Source/Runtime/Core/Public/Templates/TypeHash.h
--- a/Engine/Source/Runtime/Core/Public/Templates/TypeHash.h
+++ b/Engine/Source/Runtime/Core/Public/Templates/TypeHash.h
@@ -14,7 +14,7 @@ namespace UE
 {
 	namespace Private
 	{
-		FORCEINLINE uint32 MurmurFinalize32(uint32 Hash)
+		inline uint32 MurmurFinalize32(uint32 Hash)
 		{
 			Hash ^= Hash >> 16;
 			Hash *= 0x85ebca6b;
@@ -54,7 +54,7 @@ template <
 	typename... Types
 	UE_REQUIRES((std::is_convertible_v<Types, uint32> && ...))
 >
-[[nodiscard]] FORCEINLINE constexpr uint32 HashCombine(uint32 A, Types... Vals)
+[[nodiscard]] inline constexpr uint32 HashCombine(uint32 A, Types... Vals)
 {
 	if constexpr (sizeof...(Types) > 0)
 	{
@@ -79,7 +79,7 @@ template <
 	typename... Types
 	UE_REQUIRES((std::is_convertible_v<Types, uint32> && ...))
 >
-[[nodiscard]] FORCEINLINE constexpr uint32 HashCombineFast(uint32 A, Types... Vals)
+[[nodiscard]] inline constexpr uint32 HashCombineFast(uint32 A, Types... Vals)
 {
 	if constexpr (sizeof...(Types) > 0)
 	{
diff --git a/Engine/Source/Runtime/Core/Public/Templates/UniquePtr.h b/Engine/Source/Runtime/Core/Public/Templates/UniquePtr.h
--- a/Engine/Source/Runtime/Core/Public/Templates/UniquePtr.h
+++ b/Engine/Source/Runtime/Core/Public/Templates/UniquePtr.h
@@ -120,7 +120,7 @@ public:
 	/**
 	 * Default constructor - initializes the TUniquePtr to null.
 	 */
-	FORCEINLINE TUniquePtr()
+	inline TUniquePtr()
 		: Deleter()
 		, Ptr    (nullptr)
 	{
@@ -135,7 +135,7 @@ public:
 		typename U
 		UE_REQUIRES(std::is_convertible_v<U*, T*>)
 	>
-	explicit FORCEINLINE TUniquePtr(U* InPtr)
+	explicit inline TUniquePtr(U* InPtr)
 		: Deleter()
 		, Ptr    (InPtr)
 	{
@@ -150,7 +150,7 @@ public:
 		typename U
 		UE_REQUIRES(std::is_convertible_v<U*, T*>)
 	>
-	explicit FORCEINLINE TUniquePtr(U* InPtr, Deleter&& InDeleter)
+	explicit inline TUniquePtr(U* InPtr, Deleter&& InDeleter)
 		: Deleter(MoveTemp(InDeleter))
 		, Ptr    (InPtr)
 	{
@@ -165,7 +165,7 @@ public:
 		typename U
 		UE_REQUIRES(std::is_convertible_v<U*, T*>)
 	>
-	explicit FORCEINLINE TUniquePtr(U* InPtr, const Deleter& InDeleter)
+	explicit inline TUniquePtr(U* InPtr, const Deleter& InDeleter)
 		: Deleter(InDeleter)
 		, Ptr    (InPtr)
 	{
@@ -174,7 +174,7 @@ public:
 	/**
 	 * nullptr constructor - initializes the TUniquePtr to null.
 	 */
-	FORCEINLINE TUniquePtr(TYPE_OF_NULLPTR)
+	inline TUniquePtr(TYPE_OF_NULLPTR)
 		: Deleter()
 		, Ptr    (nullptr)
 	{
@@ -183,7 +183,7 @@ public:
 	/**
 	 * Move constructor
 	 */
-	FORCEINLINE TUniquePtr(TUniquePtr&& Other)
+	inline TUniquePtr(TUniquePtr&& Other)
 		: Deleter(MoveTemp(Other.GetDeleter()))
 		, Ptr    (Other.Ptr)
 	{
@@ -201,7 +201,7 @@ public:
 			std::is_convertible_v<OtherT*, T*>
 		)
 	>
-	FORCEINLINE TUniquePtr(TUniquePtr<OtherT, OtherDeleter>&& Other)
+	inline TUniquePtr(TUniquePtr<OtherT, OtherDeleter>&& Other)
 		: Deleter(MoveTemp(Other.GetDeleter()))
 		, Ptr    (Other.Ptr)
 	{
@@ -211,7 +211,7 @@ public:
 	/**
 	 * Move assignment operator
 	 */
-	FORCEINLINE TUniquePtr& operator=(TUniquePtr&& Other)
+	inline TUniquePtr& operator=(TUniquePtr&& Other)
 	{
 		if (this != &Other)
 		{
@@ -238,7 +238,7 @@ public:
 			std::is_convertible_v<OtherT*, T*>
 		)
 	>
-	FORCEINLINE TUniquePtr& operator=(TUniquePtr<OtherT, OtherDeleter>&& Other)
+	inline TUniquePtr& operator=(TUniquePtr<OtherT, OtherDeleter>&& Other)
 	{
 		// We delete last, because we don't want odd side effects if the destructor of T relies on the state of this or Other
 		T* OldPtr = Ptr;
@@ -254,7 +254,7 @@ public:
 	/**
 	 * Nullptr assignment operator
 	 */
-	FORCEINLINE TUniquePtr& operator=(TYPE_OF_NULLPTR)
+	inline TUniquePtr& operator=(TYPE_OF_NULLPTR)
 	{
 		// We delete last, because we don't want odd side effects if the destructor of T relies on the state of this
 		T* OldPtr = Ptr;
@@ -331,7 +331,7 @@ public:
 	 *
 	 * @return The pointer to the object that was owned by the TUniquePtr, or nullptr if no object was being owned.
 	 */
-	[[nodiscard]] FORCEINLINE T* Release()
+	[[nodiscard]] inline T* Release()
 	{
 		T* Result = Ptr;
 		Ptr = nullptr;
@@ -343,7 +343,7 @@ public:
 	 *
 	 * @param InPtr A pointer to the object to take ownership of.
 	 */
-	FORCEINLINE void Reset(T* InPtr = nullptr)
+	inline void Reset(T* InPtr = nullptr)
 	{
 		if (Ptr != InPtr)
 		{
@@ -511,7 +511,7 @@ public:
 	/**
 	 * Default constructor - initializes the TUniquePtr to null.
 	 */
-	FORCEINLINE TUniquePtr()
+	inline TUniquePtr()
 		: Deleter()
 		, Ptr    (nullptr)
 	{
@@ -526,7 +526,7 @@ public:
 		typename U
 		UE_REQUIRES(std::is_convertible_v<U(*)[], T(*)[]>)
 	>
-	explicit FORCEINLINE TUniquePtr(U* InPtr)
+	explicit inline TUniquePtr(U* InPtr)
 		: Deleter()
 		, Ptr    (InPtr)
 	{
@@ -541,7 +541,7 @@ public:
 		typename U
 		UE_REQUIRES(std::is_convertible_v<U(*)[], T(*)[]>)
 	>
-	explicit FORCEINLINE TUniquePtr(U* InPtr, Deleter&& InDeleter)
+	explicit inline TUniquePtr(U* InPtr, Deleter&& InDeleter)
 		: Deleter(MoveTemp(InDeleter))
 		, Ptr    (InPtr)
 	{
@@ -556,7 +556,7 @@ public:
 		typename U
 		UE_REQUIRES(std::is_convertible_v<U(*)[], T(*)[]>)
 	>
-	explicit FORCEINLINE TUniquePtr(U* InPtr, const Deleter& InDeleter)
+	explicit inline TUniquePtr(U* InPtr, const Deleter& InDeleter)
 		: Deleter(InDeleter)
 		, Ptr    (InPtr)
 	{
@@ -565,7 +565,7 @@ public:
 	/**
 	 * nullptr constructor - initializes the TUniquePtr to null.
 	 */
-	FORCEINLINE TUniquePtr(TYPE_OF_NULLPTR)
+	inline TUniquePtr(TYPE_OF_NULLPTR)
 		: Deleter()
 		, Ptr    (nullptr)
 	{
@@ -574,7 +574,7 @@ public:
 	/**
 	 * Move constructor
 	 */
-	FORCEINLINE TUniquePtr(TUniquePtr&& Other)
+	inline TUniquePtr(TUniquePtr&& Other)
 		: Deleter(MoveTemp(Other.GetDeleter()))
 		, Ptr    (Other.Ptr)
 	{
@@ -589,7 +589,7 @@ public:
 		typename OtherDeleter
 		UE_REQUIRES(std::is_convertible_v<OtherT(*)[], T(*)[]>)
 	>
-	FORCEINLINE TUniquePtr(TUniquePtr<OtherT, OtherDeleter>&& Other)
+	inline TUniquePtr(TUniquePtr<OtherT, OtherDeleter>&& Other)
 		: Deleter(MoveTemp(Other.GetDeleter()))
 		, Ptr    (Other.Ptr)
 	{
@@ -599,7 +599,7 @@ public:
 	/**
 	 * Move assignment operator
 	 */
-	FORCEINLINE TUniquePtr& operator=(TUniquePtr&& Other)
+	inline TUniquePtr& operator=(TUniquePtr&& Other)
 	{
 		if (this != &Other)
 		{
@@ -623,7 +623,7 @@ public:
 		typename OtherDeleter
 		UE_REQUIRES(std::is_convertible_v<OtherT(*)[], T(*)[]>)
 	>
-	FORCEINLINE TUniquePtr& operator=(TUniquePtr<OtherT, OtherDeleter>&& Other)
+	inline TUniquePtr& operator=(TUniquePtr<OtherT, OtherDeleter>&& Other)
 	{
 		// We delete last, because we don't want odd side effects if the destructor of T relies on the state of this or Other
 		T* OldPtr = Ptr;
@@ -639,7 +639,7 @@ public:
 	/**
 	 * Nullptr assignment operator
 	 */
-	FORCEINLINE TUniquePtr& operator=(TYPE_OF_NULLPTR)
+	inline TUniquePtr& operator=(TYPE_OF_NULLPTR)
 	{
 		// We delete last, because we don't want odd side effects if the destructor of T relies on the state of this
 		T* OldPtr = Ptr;
@@ -702,7 +702,7 @@ public:
 	 *
 	 * @return The pointer to the array that was owned by the TUniquePtr, or nullptr if no array was being owned.
 	 */
-	[[nodiscard]] FORCEINLINE T* Release()
+	[[nodiscard]] inline T* Release()
 	{
 		T* Result = Ptr;
 		Ptr = nullptr;
@@ -718,7 +718,7 @@ public:
 		typename U
 		UE_REQUIRES(std::is_convertible_v<U(*)[], T(*)[]>)
 	>
-	FORCEINLINE void Reset(U* InPtr)
+	inline void Reset(U* InPtr)
 	{
 		// We delete last, because we don't want odd side effects if the destructor of T relies on the state of this
 		T* OldPtr = Ptr;
@@ -726,7 +726,7 @@ public:
 		GetDeleter()(OldPtr);
 	}
 
-	FORCEINLINE void Reset(TYPE_OF_NULLPTR InPtr = nullptr)
+	inline void Reset(TYPE_OF_NULLPTR InPtr = nullptr)
 	{
 		// We delete last, because we don't want odd side effects if the destructor of T relies on the state of this
 		T* OldPtr = Ptr;
@@ -950,7 +950,7 @@ template <
 	typename T
 	UE_REQUIRES(TIsUnboundedArray<T>::Value)
 >
-[[nodiscard]] FORCEINLINE TUniquePtr<T> MakeUnique(SIZE_T Size)
+[[nodiscard]] inline TUniquePtr<T> MakeUnique(SIZE_T Size)
 {
 	using ElementType = std::remove_extent_t<T>;
 	return TUniquePtr<T>(new ElementType[Size]());
@@ -968,7 +968,7 @@ template <
 	typename T
 	UE_REQUIRES(TIsUnboundedArray<T>::Value)
 >
-[[nodiscard]] FORCEINLINE TUniquePtr<T> MakeUniqueForOverwrite(SIZE_T Size)
+[[nodiscard]] inline TUniquePtr<T> MakeUniqueForOverwrite(SIZE_T Size)
 {
 	using ElementType = std::remove_extent_t<T>;
 	return TUniquePtr<T>(new ElementType[Size]);
diff --git a/Engine/Source/Runtime/Core/Public/Templates/UnrealTemplate.h b/Engine/Source/Runtime/Core/Public/Templates/UnrealTemplate.h
--- a/Engine/Source/Runtime/Core/Public/Templates/UnrealTemplate.h
+++ b/Engine/Source/Runtime/Core/Public/Templates/UnrealTemplate.h
@@ -49,7 +49,7 @@ inline bool XOR(bool A, bool B)
 
 /** This is used to provide type specific behavior for a copy which cannot change the value of B. */
 template<typename T>
-FORCEINLINE void Move(T& A,typename TMoveSupportTraits<T>::Copy B)
+inline void Move(T& A,typename TMoveSupportTraits<T>::Copy B)
 {
 	// Destruct the previous value of A.
 	A.~T();
@@ -60,7 +60,7 @@ FORCEINLINE void Move(T& A,typename TMoveSupportTraits<T>::Copy B)
 
 /** This is used to provide type specific behavior for a move which may change the value of B. */
 template<typename T>
-FORCEINLINE void Move(T& A,typename TMoveSupportTraits<T>::Move B)
+inline void Move(T& A,typename TMoveSupportTraits<T>::Move B)
 {
 	// Destruct the previous value of A.
 	A.~T();
@@ -709,7 +709,7 @@ template <
 	typename T
 	UE_REQUIRES(std::is_same_v<T, uint32>)
 >
-FORCEINLINE T ReverseBits( T Bits )
+inline T ReverseBits( T Bits )
 {
 	Bits = ( Bits << 16) | ( Bits >> 16);
 	Bits = ( (Bits & 0x00ff00ff) << 8 ) | ( (Bits & 0xff00ff00) >> 8 );
@@ -726,28 +726,28 @@ template <typename T>
 FORCEINLINE T BitMask( uint32 Count );
 
 template <>
-FORCEINLINE uint64 BitMask<uint64>( uint32 Count )
+inline uint64 BitMask<uint64>( uint32 Count )
 {
 	checkSlow(Count <= 64);
 	return (uint64(Count < 64) << Count) - 1;
 }
 
 template <>
-FORCEINLINE uint32 BitMask<uint32>( uint32 Count )
+inline uint32 BitMask<uint32>( uint32 Count )
 {
 	checkSlow(Count <= 32);
 	return uint32(uint64(1) << Count) - 1;
 }
 
 template <>
-FORCEINLINE uint16 BitMask<uint16>( uint32 Count )
+inline uint16 BitMask<uint16>( uint32 Count )
 {
 	checkSlow(Count <= 16);
 	return uint16((uint32(1) << Count) - 1);
 }
 
 template <>
-FORCEINLINE uint8 BitMask<uint8>( uint32 Count )
+inline uint8 BitMask<uint8>( uint32 Count )
 {
 	checkSlow(Count <= 8);
 	return uint8((uint32(1) << Count) - 1);
diff --git a/Engine/Source/Runtime/Core/Public/Templates/UnrealTypeTraits.h b/Engine/Source/Runtime/Core/Public/Templates/UnrealTypeTraits.h
--- a/Engine/Source/Runtime/Core/Public/Templates/UnrealTypeTraits.h
+++ b/Engine/Source/Runtime/Core/Public/Templates/UnrealTypeTraits.h
@@ -82,7 +82,7 @@ template<typename T>
 struct TFormatSpecifier
 {
 	template <typename CharType = TCHAR>
-	FORCEINLINE static constexpr decltype(auto) GetFormatSpecifier()
+	inline static constexpr decltype(auto) GetFormatSpecifier()
 	{
 		// Force the template instantiation to be dependent upon T so the compiler cannot automatically decide that this template can never be instantiated.
 		// If the static_assert below were a constant 0 or something not dependent on T, compilers are free to detect this and fail to compile the template.
@@ -189,7 +189,7 @@ struct TIsWeakPointerType
 template<typename T> 
 struct TNameOf
 { 
-	FORCEINLINE static TCHAR const* GetName()
+	inline static TCHAR const* GetName()
 	{
 		check(0); // request for the name of a type we do not know about
 		return TEXT("Unknown");
diff --git a/Engine/Source/Runtime/Core/Public/UObject/ScriptDelegates.h b/Engine/Source/Runtime/Core/Public/UObject/ScriptDelegates.h
--- a/Engine/Source/Runtime/Core/Public/UObject/ScriptDelegates.h
+++ b/Engine/Source/Runtime/Core/Public/UObject/ScriptDelegates.h
@@ -288,7 +288,7 @@ public:
 	}
 
 	/** Comparison operators */
-	FORCEINLINE bool operator==( const TScriptDelegate& Other ) const
+	inline bool operator==( const TScriptDelegate& Other ) const
 	{
 		WeakPtrType OtherObject;
 		FName OtherFunctionName;
@@ -313,7 +313,7 @@ public:
 		UE_REQUIRES(UE::Core::Private::BackwardCompatibilityCheck<InThreadSafetyMode, OtherThreadSafetyMode>())
 	>
 	/* UE_DEPRECATED(5.3, "Deprecated - remove after TScriptDelegateTraits<FWeakObjectPtr> is removed") */
-	FORCEINLINE bool operator==(const TScriptDelegate<OtherThreadSafetyMode>& Other) const
+	inline bool operator==(const TScriptDelegate<OtherThreadSafetyMode>& Other) const
 	{
 		WeakPtrType OtherObject;
 		FName OtherFunctionName;
diff --git a/Engine/Source/Runtime/Core/Public/UObject/StrongObjectPtrTemplates.h b/Engine/Source/Runtime/Core/Public/UObject/StrongObjectPtrTemplates.h
--- a/Engine/Source/Runtime/Core/Public/UObject/StrongObjectPtrTemplates.h
+++ b/Engine/Source/Runtime/Core/Public/UObject/StrongObjectPtrTemplates.h
@@ -54,7 +54,7 @@ public:
 		static_assert(TPointerIsConvertibleFromTo<ObjectType, const volatile UObject>::Value, "TStrongObjectPtr can only be constructed with UObject types");
 	}
 
-	[[nodiscard]] FORCEINLINE_DEBUGGABLE explicit TStrongObjectPtr(ObjectType* InObject)
+	[[nodiscard]] inline explicit TStrongObjectPtr(ObjectType* InObject)
 	{
 		static_assert(TPointerIsConvertibleFromTo<ObjectType, const volatile UObject>::Value, "TStrongObjectPtr can only be constructed with UObject types");
 		Reset(InObject);
@@ -75,7 +75,7 @@ public:
 		Reset(InOther.Get());
 	}
 
-	FORCEINLINE_DEBUGGABLE TStrongObjectPtr& operator=(const TStrongObjectPtr& InOther)
+	inline TStrongObjectPtr& operator=(const TStrongObjectPtr& InOther)
 	{
 		Reset(InOther.Get());
 		return *this;
@@ -86,19 +86,19 @@ public:
 		typename OtherReferencerNameProvider
 		UE_REQUIRES(std::is_convertible_v<OtherObjectType*, ObjectType*>)
 	>
-	FORCEINLINE_DEBUGGABLE TStrongObjectPtr& operator=(const TStrongObjectPtr<OtherObjectType, OtherReferencerNameProvider>& InOther)
+	inline TStrongObjectPtr& operator=(const TStrongObjectPtr<OtherObjectType, OtherReferencerNameProvider>& InOther)
 	{
 		Reset(InOther.Get());
 		return *this;
 	}
 
-	[[nodiscard]] FORCEINLINE_DEBUGGABLE ObjectType& operator*() const
+	[[nodiscard]] inline ObjectType& operator*() const
 	{
 		check(IsValid());
 		return *(ObjectType*)Get();
 	}
 
-	[[nodiscard]] FORCEINLINE_DEBUGGABLE ObjectType* operator->() const
+	[[nodiscard]] inline ObjectType* operator->() const
 	{
 		check(IsValid());
 		return (ObjectType*)Get();
@@ -119,7 +119,7 @@ public:
 		return (ObjectType*)Object;
 	}
 
-	FORCEINLINE_DEBUGGABLE void Reset()
+	inline void Reset()
 	{
 		if (Object)
 		{
@@ -134,14 +134,14 @@ private:
 	template<class T, class TWeakObjectPtrBase> friend struct TWeakObjectPtr;
 
 	// Attach an object without incrementing its ref-count.
-	FORCEINLINE_DEBUGGABLE void Attach(ObjectType* InNewObject)
+	inline void Attach(ObjectType* InNewObject)
 	{
 		Reset();
 		Object = InNewObject;
 	}
 
 	// Detach the current owned object without decrementing its ref-count.
-	FORCEINLINE_DEBUGGABLE ObjectType* Detach()
+	inline ObjectType* Detach()
 	{
 		ObjectType* DetachedObject = Get();
 		Object = nullptr;
@@ -149,7 +149,7 @@ private:
 	}
 
 public:
-	FORCEINLINE_DEBUGGABLE void Reset(ObjectType* InNewObject)
+	inline void Reset(ObjectType* InNewObject)
 	{
 		if (InNewObject)
 		{
diff --git a/Engine/Source/Runtime/VectorVM/Public/VectorVM.h b/Engine/Source/Runtime/VectorVM/Public/VectorVM.h
--- a/Engine/Source/Runtime/VectorVM/Public/VectorVM.h
+++ b/Engine/Source/Runtime/VectorVM/Public/VectorVM.h
@@ -274,7 +274,7 @@ struct FDataSetMeta
 
 	FDataSetMeta() = default;
 
-	VM_FORCEINLINE void Reset()
+	inline void Reset()
 	{
 		InputRegisters = FInputRegisterView();
 		OutputRegisters = FOutputRegisterView();
@@ -288,7 +288,7 @@ struct FDataSetMeta
 		IDAcquireTag = INDEX_NONE;
 	}
 
-	VM_FORCEINLINE void Init(FInputRegisterView InInputRegisters, FOutputRegisterView InOutputRegisters, int32 InInstanceOffset, TArray<int32>* InIDTable, TArray<int32>* InFreeIDTable, int32* InNumFreeIDs, int32 *InNumSpawnedIDs, int32* InMaxUsedID, int32 InIDAcquireTag, TArray<int32>* InSpawnedIDsTable)
+	inline void Init(FInputRegisterView InInputRegisters, FOutputRegisterView InOutputRegisters, int32 InInstanceOffset, TArray<int32>* InIDTable, TArray<int32>* InFreeIDTable, int32* InNumFreeIDs, int32 *InNumSpawnedIDs, int32* InMaxUsedID, int32 InIDAcquireTag, TArray<int32>* InSpawnedIDsTable)
 	{
 		InputRegisters = InInputRegisters;
 		OutputRegisters = InOutputRegisters;
@@ -335,7 +335,7 @@ public:
 
 	VM_FORCEINLINE int32                                  GetStartInstance() const { return StartInstance; }
 	VM_FORCEINLINE int32                                  GetNumInstances() const { return NumInstances; }
-	VM_FORCEINLINE int32* GetRandCounters()
+	inline int32* GetRandCounters()
 	{
 		if (!*RandCounters)
 		{
@@ -343,11 +343,11 @@ public:
 		}
 		return *RandCounters;
 	}
-	VM_FORCEINLINE FRandomStream& GetRandStream() { return *RandStream; }
-	VM_FORCEINLINE void* GetUserPtrTable(int32 UserPtrIdx) { check(UserPtrIdx < NumUserPtrs);  return UserPtrTable[UserPtrIdx]; }
-	template<uint32 InstancesPerOp> VM_FORCEINLINE int32  GetNumLoops() const { static_assert(InstancesPerOp == 4); return NumLoops; };
+	inline FRandomStream& GetRandStream() { return *RandStream; }
+	inline void* GetUserPtrTable(int32 UserPtrIdx) { check(UserPtrIdx < NumUserPtrs);  return UserPtrTable[UserPtrIdx]; }
+	template<uint32 InstancesPerOp> inline int32  GetNumLoops() const { static_assert(InstancesPerOp == 4); return NumLoops; };
 
-	VM_FORCEINLINE float* GetNextRegister(int32* OutAdvanceOffset)
+	inline float* GetNextRegister(int32* OutAdvanceOffset)
 	{
 		check(RegReadCount < NumRegisters);
 		*OutAdvanceOffset = RegInc[RegReadCount];
@@ -439,7 +439,7 @@ namespace VectorVM
 	public:
 		FExternalFuncInputHandler() = default;
 
-		VM_FORCEINLINE FExternalFuncInputHandler(FVectorVMExternalFunctionContext& Context)
+		inline FExternalFuncInputHandler(FVectorVMExternalFunctionContext& Context)
 		{
 			Init(Context);
 		}
@@ -452,20 +452,20 @@ namespace VectorVM
 			StartPtr = InputPtr;
 		}
 
-		VM_FORCEINLINE bool IsConstant() const { return !IsRegister(); }
-		VM_FORCEINLINE bool IsRegister() const { return AdvanceOffset > 0; }
-		VM_FORCEINLINE void Reset(){ InputPtr = StartPtr; }
+		inline bool IsConstant() const { return !IsRegister(); }
+		inline bool IsRegister() const { return AdvanceOffset > 0; }
+		inline void Reset(){ InputPtr = StartPtr; }
 
-		VM_FORCEINLINE const T Get() { return *InputPtr; }
-		VM_FORCEINLINE const T* GetDest() { return InputPtr; }
-		VM_FORCEINLINE void Advance(int32 Count=1) { InputPtr += AdvanceOffset * Count; }
-		VM_FORCEINLINE const T GetAndAdvance()
+		inline const T Get() { return *InputPtr; }
+		inline const T* GetDest() { return InputPtr; }
+		inline void Advance(int32 Count=1) { InputPtr += AdvanceOffset * Count; }
+		inline const T GetAndAdvance()
 		{
 			const T* Ret = InputPtr;
 			InputPtr += AdvanceOffset;
 			return *Ret;
 		}
-		VM_FORCEINLINE const T* GetDestAndAdvance()
+		inline const T* GetDestAndAdvance()
 		{
 			const T* Ret = InputPtr;
 			InputPtr += AdvanceOffset;
@@ -488,18 +488,18 @@ namespace VectorVM
 			Register += Context.PerInstanceFnInstanceIdx * AdvanceOffset;
 		}
 
-		VM_FORCEINLINE bool IsValid() const { return AdvanceOffset > 0; }
-		VM_FORCEINLINE const T Get() { return *Register; }
-		VM_FORCEINLINE T* GetDest() { return Register; }
-		VM_FORCEINLINE void Advance() { Register += AdvanceOffset; }
-		VM_FORCEINLINE void Advance(int32 Count) { Register += AdvanceOffset * Count; }
-		VM_FORCEINLINE const T GetAndAdvance()
+		inline bool IsValid() const { return AdvanceOffset > 0; }
+		inline const T Get() { return *Register; }
+		inline T* GetDest() { return Register; }
+		inline void Advance() { Register += AdvanceOffset; }
+		inline void Advance(int32 Count) { Register += AdvanceOffset * Count; }
+		inline const T GetAndAdvance()
 		{
 			T* Ret = Register;
 			Register += AdvanceOffset;
 			return *Ret;
 		}
-		VM_FORCEINLINE T* GetDestAndAdvance()
+		inline T* GetDestAndAdvance()
 		{
 			T* Ret = Register;
 			Register += AdvanceOffset;
