Fix linter warnings
- Have sscanf return the number arguments we expected and then check it against
- Change the loop criteria to make it resistant to a null dereference
- Protect bad divisions when attempting to get the process utilization

diff --git a/Engine/Source/Runtime/Core/Private/Unix/UnixPlatformProcess.cpp b/Engine/Source/Runtime/Core/Private/Unix/UnixPlatformProcess.cpp
--- a/Engine/Source/Runtime/Core/Private/Unix/UnixPlatformProcess.cpp
+++ b/Engine/Source/Runtime/Core/Private/Unix/UnixPlatformProcess.cpp
@@ -1801,6 +1801,7 @@ FString FUnixPlatformProcess::GetCurrentWorkingDirectory()
 bool FUnixPlatformProcess::GetPerFrameProcessorUsage(uint32 ProcessId, float& ProcessUsageFraction, float& IdleUsageFraction)
 {
 	static bool bProcStatPrintHasHappened = false;
+	static bool bProcStatGetLinePrintHappened = false;
 	static bool bProcStatParsePrintHappened = false;
 	static bool bProcPIDPrintHappened = false;
 	static bool bProcPIDParsePrintHappened = false;
@@ -1827,7 +1828,7 @@ bool FUnixPlatformProcess::GetPerFrameProcessorUsage(uint32 ProcessId, float& Pr
 	// First line of the file is the total CPU usage so no need to tally everything together
 	if (fgets(Buffer, sizeof(Buffer), ProcStatFile) == nullptr)
 	{
-		UE_CLOG(!bProcStatParsePrintHappened, LogUtilization, Verbose, TEXT("Failed to parse the first line of /proc/stat."));
+		UE_CLOG(!bProcStatParsePrintHappened, LogUtilization, Verbose, TEXT("Failed to get the first line of /proc/stat."));
 		bProcStatParsePrintHappened = true;
 		fclose(ProcStatFile);
 		return false;
@@ -1836,8 +1837,14 @@ bool FUnixPlatformProcess::GetPerFrameProcessorUsage(uint32 ProcessId, float& Pr
 	fclose(ProcStatFile);
 
 #if PLATFORM_64BITS
-	sscanf(Buffer, "%5s %8lu %8lu %8lu %8lu %8lu %8lu %8lu", CPUName,
+	int Result = sscanf(Buffer, "%5s %8lu %8lu %8lu %8lu %8lu %8lu %8lu", CPUName,
 		&UserTime, &NiceTime, &SystemTime, &IdleTime, &IOWaitTime, &IRQTime, &SoftIRQTime);
+	if (Result != 8)
+	{
+		UE_CLOG(!bProcStatGetLinePrintHappened, LogUtilization, Verbose, TEXT("Failed to parse the first line of /proc/stat."));
+		bProcStatGetLinePrintHappened = true;
+		return false;
+	}
 #else
 	sscanf(Buffer, "%5s %8llu %8llu %8llu %8llu %8llu %8llu %8llu", CurrentCPUState.Name,
 		&UserTime, &NiceTime, &SystemTime, &IdleTime, &IOWaitTime, &IRQTime, &SoftIRQTime);
@@ -1882,8 +1889,6 @@ bool FUnixPlatformProcess::GetPerFrameProcessorUsage(uint32 ProcessId, float& Pr
 
 	while (Token != nullptr)
 	{
-		Token = strtok_r(nullptr, " ", &Remainder);
-		NumTokens++;
 		if (NumTokens == UserTimeField)
 		{
 			UserTimeTicks = atol(Token);
@@ -1901,6 +1906,9 @@ bool FUnixPlatformProcess::GetPerFrameProcessorUsage(uint32 ProcessId, float& Pr
 			ChildSystemTimeTicks = atol(Token);
 			break;
 		}
+
+		Token = strtok_r(nullptr, " ", &Remainder);
+		NumTokens++;
 	}
 
 	// First sample period, there will be no valid data here
@@ -1923,7 +1931,11 @@ bool FUnixPlatformProcess::GetPerFrameProcessorUsage(uint32 ProcessId, float& Pr
 	uint32 ChildSystemTimeDelta = ChildSystemTimeTicks - PreviousChildSystemTimeTicks;
 
 	uint32 TotalProcessTicksDelta = UserTimeDelta + SystemTimeDelta + ChildUserTimeDelta + ChildSystemTimeDelta;
-
+	if (TotalTicksDelta == 0 || TotalProcessTicksDelta > TotalTicksDelta)
+	{
+		return false;
+	}
+	
 	ProcessUsageFraction = (float)TotalProcessTicksDelta/(float)TotalTicksDelta * 100.0f;
 	// Not sure what the correct logic would be here because we do not have an idle time for this process
 	IdleUsageFraction = 100.0f - ProcessUsageFraction;
