Replace some usages of FORCEINLINE with inline in Platform modules.

For context about this change, see this doc: https://docs.google.com/document/d/1LjpT7NpzTecfXdFNyv5KM6sCu67827BRljNi1iaWVSE/edit?tab=t.0#heading=h.fvtz8gjukz7x

Reviewers were automatically determined. Please comment if anyone else should be added.

#rnx

diff --git a/Engine/Source/Runtime/ApplicationCore/Public/Android/AndroidApplication.h b/Engine/Source/Runtime/ApplicationCore/Public/Android/AndroidApplication.h
--- a/Engine/Source/Runtime/ApplicationCore/Public/Android/AndroidApplication.h
+++ b/Engine/Source/Runtime/ApplicationCore/Public/Android/AndroidApplication.h
@@ -32,35 +32,35 @@ public:
 
 #if USE_ANDROID_JNI
 	// Returns the java environment
-	static FORCEINLINE void InitializeJavaEnv(JavaVM* VM, jint Version, jobject GlobalThis)
+	static inline void InitializeJavaEnv(JavaVM* VM, jint Version, jobject GlobalThis)
 	{
 		AndroidJavaEnv::InitializeJavaEnv(VM, Version, GlobalThis);
     }
-	static FORCEINLINE jobject GetGameActivityThis()
+	static inline jobject GetGameActivityThis()
 	{
 		return AndroidJavaEnv::GetGameActivityThis();
 	} 
-	static FORCEINLINE jobject GetClassLoader()
+	static inline jobject GetClassLoader()
 	{
 		return AndroidJavaEnv::GetClassLoader();
 	} 
-	static FORCEINLINE JNIEnv* GetJavaEnv(bool bRequireGlobalThis = true)
+	static inline JNIEnv* GetJavaEnv(bool bRequireGlobalThis = true)
 	{
 		return AndroidJavaEnv::GetJavaEnv(bRequireGlobalThis);
 	}
-	static FORCEINLINE jclass FindJavaClass(const char* name)
+	static inline jclass FindJavaClass(const char* name)
 	{
 		return AndroidJavaEnv::FindJavaClass(name);
 	}
-	static FORCEINLINE jclass FindJavaClassGlobalRef(const char* name)
+	static inline jclass FindJavaClassGlobalRef(const char* name)
 	{
 		return AndroidJavaEnv::FindJavaClassGlobalRef(name);
 	}
-	static FORCEINLINE void DetachJavaEnv()
+	static inline void DetachJavaEnv()
 	{
 		AndroidJavaEnv::DetachJavaEnv();
 	}
-	static FORCEINLINE bool CheckJavaException()
+	static inline bool CheckJavaException()
 	{
 		return AndroidJavaEnv::CheckJavaException();
 	}
diff --git a/Engine/Source/Runtime/ApplicationCore/Public/GenericPlatform/GenericApplication.h b/Engine/Source/Runtime/ApplicationCore/Public/GenericPlatform/GenericApplication.h
--- a/Engine/Source/Runtime/ApplicationCore/Public/GenericPlatform/GenericApplication.h
+++ b/Engine/Source/Runtime/ApplicationCore/Public/GenericPlatform/GenericApplication.h
@@ -45,7 +45,7 @@ namespace EModifierKey
 	/** Cmd key (Control key on Mac, Win key on Windows) */
 	inline const Type Command = 1 << 3;
 
-	FORCEINLINE EModifierKey::Type FromBools(const bool bControl, const bool bAlt, const bool bShift, const bool bCommand)
+	inline EModifierKey::Type FromBools(const bool bControl, const bool bAlt, const bool bShift, const bool bCommand)
 	{
 		EModifierKey::Type ModifierMask = EModifierKey::None;
 		if (bControl)	ModifierMask |= EModifierKey::Control;
diff --git a/Engine/Source/Runtime/ApplicationCore/Public/GenericPlatform/GenericPlatformApplicationMisc.h b/Engine/Source/Runtime/ApplicationCore/Public/GenericPlatform/GenericPlatformApplicationMisc.h
--- a/Engine/Source/Runtime/ApplicationCore/Public/GenericPlatform/GenericPlatformApplicationMisc.h
+++ b/Engine/Source/Runtime/ApplicationCore/Public/GenericPlatform/GenericPlatformApplicationMisc.h
@@ -130,7 +130,7 @@ struct FGenericPlatformApplicationMisc
 	 *	Pumps Windows messages.
 	 *	@param bFromMainLoop if true, this is from the main loop, otherwise we are spinning waiting for the render thread
 	 */
-	FORCEINLINE static void PumpMessages(bool bFromMainLoop)
+	inline static void PumpMessages(bool bFromMainLoop)
 	{
 	}
 
diff --git a/Engine/Source/Runtime/ApplicationCore/Public/GenericPlatform/GenericPlatformSplash.h b/Engine/Source/Runtime/ApplicationCore/Public/GenericPlatform/GenericPlatformSplash.h
--- a/Engine/Source/Runtime/ApplicationCore/Public/GenericPlatform/GenericPlatformSplash.h
+++ b/Engine/Source/Runtime/ApplicationCore/Public/GenericPlatform/GenericPlatformSplash.h
@@ -40,10 +40,10 @@ class FText;
 struct FGenericPlatformSplash
 {
 	/** Show the splash screen. */
-	FORCEINLINE static void Show() { }
+	inline static void Show() { }
 
 	/** Hide the splash screen. */
-	FORCEINLINE static void Hide() { }
+	inline static void Hide() { }
 
 	/**
 	 * Sets a custom splash image to display
@@ -58,7 +58,7 @@ struct FGenericPlatformSplash
 	 *
 	 * @param InType Progress value in percent.
 	 */
-	FORCEINLINE static void SetProgress(int ProgressPercent)
+	inline static void SetProgress(int ProgressPercent)
 	{
 	
 	}
@@ -71,14 +71,14 @@ struct FGenericPlatformSplash
 	 */
 	UE_DEPRECATED(5.6, "Use the overload of FPlatformSlash::SetSplashText that takes a FText")
 	static APPLICATIONCORE_API void SetSplashText( const SplashTextType::Type InType, const TCHAR* InText );
-	FORCEINLINE static void SetSplashText(const SplashTextType::Type InType, const FText& InText)
+	inline static void SetSplashText(const SplashTextType::Type InType, const FText& InText)
 	{
 	}
 
 	/**
 	 * Return whether the splash screen is being shown or not
 	 */
-	FORCEINLINE static bool IsShown()
+	inline static bool IsShown()
 	{
 		return true;
 	}
diff --git a/Engine/Source/Runtime/ApplicationCore/Public/GenericPlatform/GenericWindow.h b/Engine/Source/Runtime/ApplicationCore/Public/GenericPlatform/GenericWindow.h
--- a/Engine/Source/Runtime/ApplicationCore/Public/GenericPlatform/GenericWindow.h
+++ b/Engine/Source/Runtime/ApplicationCore/Public/GenericPlatform/GenericWindow.h
@@ -46,7 +46,7 @@ namespace EWindowMode
 	}
 }
 
-FORCEINLINE const TCHAR* LexToString(EWindowMode::Type WindowMode)
+inline const TCHAR* LexToString(EWindowMode::Type WindowMode)
 {
 	switch (WindowMode)
 	{
diff --git a/Engine/Source/Runtime/ApplicationCore/Public/IOS/IOSView.h b/Engine/Source/Runtime/ApplicationCore/Public/IOS/IOSView.h
--- a/Engine/Source/Runtime/ApplicationCore/Public/IOS/IOSView.h
+++ b/Engine/Source/Runtime/ApplicationCore/Public/IOS/IOSView.h
@@ -25,7 +25,7 @@ struct FKeyboardConfig
 	UITextAutocapitalizationType AutocapitalizationType;
 	BOOL bSecureTextEntry;
 	
-	FORCEINLINE FKeyboardConfig() :
+	inline FKeyboardConfig() :
 		KeyboardType(UIKeyboardTypeDefault),
 		AutocorrectionType(UITextAutocorrectionTypeNo),
 		AutocapitalizationType(UITextAutocapitalizationTypeNone),
diff --git a/Engine/Source/Runtime/ApplicationCore/Public/Linux/LinuxApplication.h b/Engine/Source/Runtime/ApplicationCore/Public/Linux/LinuxApplication.h
--- a/Engine/Source/Runtime/ApplicationCore/Public/Linux/LinuxApplication.h
+++ b/Engine/Source/Runtime/ApplicationCore/Public/Linux/LinuxApplication.h
@@ -103,7 +103,7 @@ public:
  	virtual TSharedPtr<FGenericWindow> GetWindowUnderCursor() override;
 
 	/** Returns true if this application is foreground */
-	FORCEINLINE bool IsForeground()
+	inline bool IsForeground()
 	{
 		// if there are no windows, consider ourselves foreground so servers and commandlets aren't impacted
 		return (Windows.Num() > 0) ? bActivateApp : true;
diff --git a/Engine/Source/Runtime/AssetRegistry/Public/AssetRegistry/AssetDataMap.h b/Engine/Source/Runtime/AssetRegistry/Public/AssetRegistry/AssetDataMap.h
--- a/Engine/Source/Runtime/AssetRegistry/Public/AssetRegistry/AssetDataMap.h
+++ b/Engine/Source/Runtime/AssetRegistry/Public/AssetRegistry/AssetDataMap.h
@@ -88,17 +88,17 @@ struct FCachedAssetKeyFuncs
 
 	enum { bAllowDuplicateKeys = false };
 
-	static FORCEINLINE KeyInitType GetSetKey(const FAssetData* Element)
+	static inline KeyInitType GetSetKey(const FAssetData* Element)
 	{
 		return FCachedAssetKey(*Element);
 	}
 
-	static FORCEINLINE bool Matches(KeyInitType A, KeyInitType B)
+	static inline bool Matches(KeyInitType A, KeyInitType B)
 	{
 		return A == B;
 	}
 
-	static FORCEINLINE uint32 GetKeyHash(KeyInitType Key)
+	static inline uint32 GetKeyHash(KeyInitType Key)
 	{
 		return GetTypeHash(Key);
 	}
@@ -400,7 +400,7 @@ private:
 namespace UE::AssetRegistry::Private
 {
 
-FORCEINLINE uint32 HashCombineQuick(uint32 A, uint32 B)
+inline uint32 HashCombineQuick(uint32 A, uint32 B)
 {
 	return A ^ (B + 0x9e3779b9 + (A << 6) + (A >> 2));
 }
diff --git a/Engine/Source/Runtime/Core/Public/Android/AndroidPlatformString.h b/Engine/Source/Runtime/Core/Public/Android/AndroidPlatformString.h
--- a/Engine/Source/Runtime/Core/Public/Android/AndroidPlatformString.h
+++ b/Engine/Source/Runtime/Core/Public/Android/AndroidPlatformString.h
@@ -55,7 +55,7 @@ struct FAndroidPlatformString : public FGenericPlatformString
 	/** 
 	 * Widechar implementation 
 	 **/
-	static FORCEINLINE WIDECHAR* Strcpy(WIDECHAR* Dest, const WIDECHAR* Src)
+	static inline WIDECHAR* Strcpy(WIDECHAR* Dest, const WIDECHAR* Src)
 	{
 		if (!Dest || !Src)
 		{
@@ -82,7 +82,7 @@ struct FAndroidPlatformString : public FGenericPlatformString
 		return Strcpy(Dest, Src);
 	}
 
-	static FORCEINLINE WIDECHAR* Strncpy(WIDECHAR* Dest, const WIDECHAR* Src, SIZE_T MaxLen)
+	static inline WIDECHAR* Strncpy(WIDECHAR* Dest, const WIDECHAR* Src, SIZE_T MaxLen)
 	{
 		if (!Dest || !Src)
 		{
@@ -106,7 +106,7 @@ struct FAndroidPlatformString : public FGenericPlatformString
 		return Dest;
 	}
 
-	static FORCEINLINE WIDECHAR* Strcat(WIDECHAR* Dest, const WIDECHAR* Src)
+	static inline WIDECHAR* Strcat(WIDECHAR* Dest, const WIDECHAR* Src)
 	{
 		if (!Dest || !Src)
 		{
@@ -133,7 +133,7 @@ struct FAndroidPlatformString : public FGenericPlatformString
 		return Strcat(Dest, Src);
 	}
 
-	static FORCEINLINE WIDECHAR* Strncat(WIDECHAR* Dest, const WIDECHAR* Src, SIZE_T SrcLen)
+	static inline WIDECHAR* Strncat(WIDECHAR* Dest, const WIDECHAR* Src, SIZE_T SrcLen)
 	{
 		if (!Dest || !Src)
 		{
@@ -154,7 +154,7 @@ struct FAndroidPlatformString : public FGenericPlatformString
 		return Dest;
 	}
 
-	static FORCEINLINE int32 Strcmp( const WIDECHAR* String1, const WIDECHAR* String2 )
+	static inline int32 Strcmp( const WIDECHAR* String1, const WIDECHAR* String2 )
 	{
 		for (; *String1 || *String2; String1++, String2++)
 		{
@@ -166,7 +166,7 @@ struct FAndroidPlatformString : public FGenericPlatformString
 		return 0;
 	}
 
-	static FORCEINLINE int32 Strlen( const WIDECHAR* String )
+	static inline int32 Strlen( const WIDECHAR* String )
 	{
 		if (!String)
 			return 0;
@@ -180,7 +180,7 @@ struct FAndroidPlatformString : public FGenericPlatformString
 		return Len;
 	}
 
-	static FORCEINLINE int32 Strnlen( const WIDECHAR* String, SIZE_T StringSize )
+	static inline int32 Strnlen( const WIDECHAR* String, SIZE_T StringSize )
 	{
 		int Len = 0;
 		while (StringSize-- > 0 && *String++)
@@ -190,7 +190,7 @@ struct FAndroidPlatformString : public FGenericPlatformString
 		return Len;
 	}
 
-	static FORCEINLINE void CopyWideToAnsi(ANSICHAR* Dest, const WIDECHAR* Src)
+	static inline void CopyWideToAnsi(ANSICHAR* Dest, const WIDECHAR* Src)
 	{
 		if (!Src || !Dest)
 			return;
@@ -213,7 +213,7 @@ struct FAndroidPlatformString : public FGenericPlatformString
 		Dest[Pos] = 0;
 	}
 
-	static FORCEINLINE void CopyAnsiToWide(WIDECHAR* Dest, const ANSICHAR* Src)
+	static inline void CopyAnsiToWide(WIDECHAR* Dest, const ANSICHAR* Src)
 	{
 		if (!Src || !Dest)
 			return;
@@ -228,7 +228,7 @@ struct FAndroidPlatformString : public FGenericPlatformString
 		Dest[Pos] = 0;
 	}
 
-	static FORCEINLINE const WIDECHAR* Strstr( const WIDECHAR* String, const WIDECHAR* Find)
+	static inline const WIDECHAR* Strstr( const WIDECHAR* String, const WIDECHAR* Find)
 	{
 		WIDECHAR FindChar = *Find;
 
@@ -263,7 +263,7 @@ struct FAndroidPlatformString : public FGenericPlatformString
 		return nullptr;
 	}
 
-	static FORCEINLINE const WIDECHAR* Strchr( const WIDECHAR* String, WIDECHAR C)
+	static inline const WIDECHAR* Strchr( const WIDECHAR* String, WIDECHAR C)
 	{
 		if (!String)
 			return NULL;
@@ -287,7 +287,7 @@ struct FAndroidPlatformString : public FGenericPlatformString
 		return NULL;
 	}
 
-	static FORCEINLINE const WIDECHAR* Strrchr( const WIDECHAR* String, WIDECHAR C)
+	static inline const WIDECHAR* Strrchr( const WIDECHAR* String, WIDECHAR C)
 	{
 		if (!String)
 			return NULL;
@@ -313,7 +313,7 @@ struct FAndroidPlatformString : public FGenericPlatformString
 		return Last;
 	}
 
-	static FORCEINLINE int32 Atoi(const WIDECHAR* String)
+	static inline int32 Atoi(const WIDECHAR* String)
 	{
 		int StringLen = Strlen(String);
 		ANSICHAR* AnsiString = (ANSICHAR*)FMemory_Alloca(StringLen+1);
@@ -322,7 +322,7 @@ struct FAndroidPlatformString : public FGenericPlatformString
 		return atoi(AnsiString);
 	}
 
-	static FORCEINLINE int64 Atoi64(const WIDECHAR* String)
+	static inline int64 Atoi64(const WIDECHAR* String)
 	{
 		int StringLen = Strlen(String);
 		ANSICHAR* AnsiString = (ANSICHAR*)FMemory_Alloca(StringLen+1);
@@ -331,7 +331,7 @@ struct FAndroidPlatformString : public FGenericPlatformString
 		return strtoll(AnsiString, NULL, 10);
 	}
 
-	static FORCEINLINE float Atof(const WIDECHAR* String)
+	static inline float Atof(const WIDECHAR* String)
 	{
 		int StringLen = Strlen(String);
 		ANSICHAR* AnsiString = (ANSICHAR*)FMemory_Alloca(StringLen+1);
@@ -340,7 +340,7 @@ struct FAndroidPlatformString : public FGenericPlatformString
 		return (float)atof(AnsiString);
 	}
 
-	static FORCEINLINE double Atod(const WIDECHAR* String)
+	static inline double Atod(const WIDECHAR* String)
 	{
 		int StringLen = Strlen(String);
 		ANSICHAR* AnsiString = (ANSICHAR*)FMemory_Alloca(StringLen+1);
@@ -349,7 +349,7 @@ struct FAndroidPlatformString : public FGenericPlatformString
 		return atof(AnsiString);
 	}
 
-	static FORCEINLINE int32 Strtoi( const WIDECHAR* Start, WIDECHAR** End, int32 Base ) 
+	static inline int32 Strtoi( const WIDECHAR* Start, WIDECHAR** End, int32 Base ) 
 	{
 		int StartLen = Strlen(Start);
 		ANSICHAR* AnsiStart = (ANSICHAR*)FMemory_Alloca(StartLen+1);
@@ -374,7 +374,7 @@ struct FAndroidPlatformString : public FGenericPlatformString
 		return Res;
 	}
 
-	static FORCEINLINE int64 Strtoi64( const WIDECHAR* Start, WIDECHAR** End, int32 Base ) 
+	static inline int64 Strtoi64( const WIDECHAR* Start, WIDECHAR** End, int32 Base ) 
 	{
 		int StartLen = Strlen(Start);
 		ANSICHAR* AnsiStart = (ANSICHAR*)FMemory_Alloca(StartLen+1);
@@ -399,7 +399,7 @@ struct FAndroidPlatformString : public FGenericPlatformString
 		return Res;
 	}
 
-	static FORCEINLINE uint64 Strtoui64( const WIDECHAR* Start, WIDECHAR** End, int32 Base ) 
+	static inline uint64 Strtoui64( const WIDECHAR* Start, WIDECHAR** End, int32 Base ) 
 	{
 		int StartLen = Strlen(Start);
 		ANSICHAR* AnsiStart = (ANSICHAR*)FMemory_Alloca(StartLen+1);
@@ -424,7 +424,7 @@ struct FAndroidPlatformString : public FGenericPlatformString
 		return Res;
 	}
 
-	static FORCEINLINE WIDECHAR* Strtok(WIDECHAR* StrToken, const WIDECHAR* Delim, WIDECHAR** Context)
+	static inline WIDECHAR* Strtok(WIDECHAR* StrToken, const WIDECHAR* Delim, WIDECHAR** Context)
 	{
 		int StrTokenLen = Strlen(StrToken);
 		int DelimLen = Strlen(Delim);
@@ -443,7 +443,7 @@ struct FAndroidPlatformString : public FGenericPlatformString
 		return StrToken + (Pos - AnsiStrToken);
 	}
 
-	static FORCEINLINE int32 GetVarArgs( WIDECHAR* Dest, SIZE_T DestSize, const WIDECHAR*& Fmt, va_list ArgPtr )
+	static inline int32 GetVarArgs( WIDECHAR* Dest, SIZE_T DestSize, const WIDECHAR*& Fmt, va_list ArgPtr )
 	{
 #if PLATFORM_USE_LS_SPEC_FOR_WIDECHAR
 		// fix up the Fmt string, as fast as possible, without using an FString
@@ -517,7 +517,7 @@ struct FAndroidPlatformString : public FGenericPlatformString
 		return Strcpy(Dest, Src);
 	}
 
-	static FORCEINLINE ANSICHAR* Strncpy(ANSICHAR* Dest, const ANSICHAR* Src, SIZE_T MaxLen)
+	static inline ANSICHAR* Strncpy(ANSICHAR* Dest, const ANSICHAR* Src, SIZE_T MaxLen)
 	{
 		::strncpy(Dest, Src, MaxLen);
 		Dest[MaxLen-1]=0;
@@ -616,7 +616,7 @@ struct FAndroidPlatformString : public FGenericPlatformString
 		return strtok(StrToken, Delim);
 	}
 
-	static FORCEINLINE int32 GetVarArgs( ANSICHAR* Dest, SIZE_T DestSize, const ANSICHAR*& Fmt, va_list ArgPtr )
+	static inline int32 GetVarArgs( ANSICHAR* Dest, SIZE_T DestSize, const ANSICHAR*& Fmt, va_list ArgPtr )
 	{
 		int32 Result = vsnprintf(Dest, DestSize, Fmt, ArgPtr);
 		return (Result != -1 && Result < (int32)DestSize) ? Result : -1;
@@ -632,7 +632,7 @@ struct FAndroidPlatformString : public FGenericPlatformString
 	// This version was deprecated because Strcpy taking a size field does not match the CRT strcpy, and on some platforms the size field was being ignored.
 	// If the memzeroing causes a performance problem, request a new function StrcpyTruncate.
 	UE_DEPRECATED(5.6, "Use Strncpy instead. Note that Strncpy has a behavior difference from Strcpy: it memzeroes the entire DestCount-sized buffer after the end of string.")
-	static FORCEINLINE UTF8CHAR* Strcpy(UTF8CHAR* Dest, SIZE_T DestCount, const UTF8CHAR* Src)
+	static inline UTF8CHAR* Strcpy(UTF8CHAR* Dest, SIZE_T DestCount, const UTF8CHAR* Src)
 	{
 		PRAGMA_DISABLE_DEPRECATION_WARNINGS
 		return (UTF8CHAR*)Strcpy((ANSICHAR*)Dest, DestCount, (const ANSICHAR*)Src);
@@ -651,7 +651,7 @@ struct FAndroidPlatformString : public FGenericPlatformString
 
 	// This version was deprecated because Strcat taking a size field does not match the CRT strcat, and on some platforms the size field was being ignored.
 	UE_DEPRECATED(5.6, "Use Strncat instead. !!NOTE THAT STRNCAT takes SrcLen rather than DestCount. You must call Strncat(Dest, Src, DestCount - Strlen(Dest) - 1).")
-	static FORCEINLINE UTF8CHAR* Strcat(UTF8CHAR* Dest, SIZE_T DestCount, const UTF8CHAR* Src)
+	static inline UTF8CHAR* Strcat(UTF8CHAR* Dest, SIZE_T DestCount, const UTF8CHAR* Src)
 	{
 		PRAGMA_DISABLE_DEPRECATION_WARNINGS
 		return (UTF8CHAR*)Strcat((ANSICHAR*)Dest, DestCount, (const ANSICHAR*)Src);
@@ -747,7 +747,7 @@ struct FAndroidPlatformString : public FGenericPlatformString
 	 * UCS2 implementation 
 	 **/
 
-	static FORCEINLINE int32 Strlen( const UCS2CHAR* String )
+	static inline int32 Strlen( const UCS2CHAR* String )
 	{
 		int32 Result = 0;
 		while (*String++)
@@ -758,7 +758,7 @@ struct FAndroidPlatformString : public FGenericPlatformString
 		return Result;
 	}
 
-	static FORCEINLINE int32 Strnlen( const UCS2CHAR* String, SIZE_T StringSize )
+	static inline int32 Strnlen( const UCS2CHAR* String, SIZE_T StringSize )
 	{
 		int32 Result = 0;
 		while (StringSize-- > 0 && *String++)
diff --git a/Engine/Source/Runtime/Core/Public/Android/AndroidPlatformTLS.h b/Engine/Source/Runtime/Core/Public/Android/AndroidPlatformTLS.h
--- a/Engine/Source/Runtime/Core/Public/Android/AndroidPlatformTLS.h
+++ b/Engine/Source/Runtime/Core/Public/Android/AndroidPlatformTLS.h
@@ -17,7 +17,7 @@ struct FAndroidTLS : public FGenericPlatformTLS
 	/**
 	 * Returns the currently executing thread's id
 	 */
-	static FORCEINLINE uint32 GetCurrentThreadId(void)
+	static inline uint32 GetCurrentThreadId(void)
 	{
 		static_assert(sizeof(pid_t) <= sizeof(uint32), "pid_t is larger than uint32, reconsider implementation of GetCurrentThreadId()");
 		return static_cast<uint32>(gettid());
@@ -27,7 +27,7 @@ struct FAndroidTLS : public FGenericPlatformTLS
 	/**
 	 * Allocates a thread local store slot
 	 */
-	static FORCEINLINE uint32 AllocTlsSlot(void)
+	static inline uint32 AllocTlsSlot(void)
 	{
 		// allocate a per-thread mem slot
 		pthread_key_t Key = 0;
diff --git a/Engine/Source/Runtime/Core/Public/Android/AndroidPlatformTime.h b/Engine/Source/Runtime/Core/Public/Android/AndroidPlatformTime.h
--- a/Engine/Source/Runtime/Core/Public/Android/AndroidPlatformTime.h
+++ b/Engine/Source/Runtime/Core/Public/Android/AndroidPlatformTime.h
@@ -17,21 +17,21 @@
 struct FAndroidTime : public FGenericPlatformTime
 {
 	// android uses BSD time code from GenericPlatformTime
-	static FORCEINLINE double Seconds()
+	static inline double Seconds()
 	{
 		struct timespec ts;
 		clock_gettime(CLOCK_MONOTONIC, &ts);
 		return ((double) ts.tv_sec) + (((double) ts.tv_nsec) / 1000000000.0);
 	}
 
-	static FORCEINLINE uint32 Cycles()
+	static inline uint32 Cycles()
 	{
 		struct timespec ts;
 		clock_gettime(CLOCK_MONOTONIC, &ts);
 		return (uint32) ((((uint64)ts.tv_sec) * 1000000ULL) + (((uint64)ts.tv_nsec) / 1000ULL));
 	}
 
-	static FORCEINLINE uint64 Cycles64()
+	static inline uint64 Cycles64()
 	{
 		struct timespec ts;
 		clock_gettime(CLOCK_MONOTONIC, &ts);
diff --git a/Engine/Source/Runtime/Core/Public/Apple/ApplePlatformTLS.h b/Engine/Source/Runtime/Core/Public/Apple/ApplePlatformTLS.h
--- a/Engine/Source/Runtime/Core/Public/Apple/ApplePlatformTLS.h
+++ b/Engine/Source/Runtime/Core/Public/Apple/ApplePlatformTLS.h
@@ -29,7 +29,7 @@ struct CORE_API FApplePlatformTLS : public FGenericPlatformTLS
 	/**
 	 * Allocates a thread local store slot
 	 */
-	static FORCEINLINE uint32 AllocTlsSlot(void)
+	static inline uint32 AllocTlsSlot(void)
 	{
 		// allocate a per-thread mem slot
 		pthread_key_t SlotKey = 0;
diff --git a/Engine/Source/Runtime/Core/Public/Apple/ApplePlatformTime.h b/Engine/Source/Runtime/Core/Public/Apple/ApplePlatformTime.h
--- a/Engine/Source/Runtime/Core/Public/Apple/ApplePlatformTime.h
+++ b/Engine/Source/Runtime/Core/Public/Apple/ApplePlatformTime.h
@@ -29,20 +29,20 @@ struct CORE_API FApplePlatformTime : public FGenericPlatformTime
 {
 	static double InitTiming();
 
-	static FORCEINLINE double Seconds()
+	static inline double Seconds()
 	{
 		uint64 Cycles = mach_absolute_time();
 		// Add big number to make bugs apparent where return value is being passed to float
 		return Cycles * GetSecondsPerCycle() + 16777216.0;
 	}
 
-	static FORCEINLINE uint32 Cycles()
+	static inline uint32 Cycles()
 	{
 		uint64 Cycles = mach_absolute_time();
 		return Cycles;
 	}
 
-	static FORCEINLINE uint64 Cycles64()
+	static inline uint64 Cycles64()
 	{
 		uint64 Cycles = mach_absolute_time();
 		return Cycles;
diff --git a/Engine/Source/Runtime/Core/Public/GenericPlatform/GenericPlatformMath.h b/Engine/Source/Runtime/Core/Public/GenericPlatform/GenericPlatformMath.h
--- a/Engine/Source/Runtime/Core/Public/GenericPlatform/GenericPlatformMath.h
+++ b/Engine/Source/Runtime/Core/Public/GenericPlatform/GenericPlatformMath.h
@@ -32,7 +32,7 @@ struct FGenericPlatformMath
 {
 	// load half (F16) to float
 	//https://gist.github.com/rygorous/2156668
-	[[nodiscard]] static constexpr FORCEINLINE float LoadHalf(const uint16* Ptr)
+	[[nodiscard]] static constexpr inline float LoadHalf(const uint16* Ptr)
 	{
 		uint16 FP16 = *Ptr;
 		constexpr uint32 shifted_exp = 0x7c00 << 13;			// exponent mask after shift
@@ -66,7 +66,7 @@ struct FGenericPlatformMath
 	// values too large for F16 are stored as +-Inf
 	// https://gist.github.com/rygorous/2156668
 	// float_to_half_fast3_rtne
-	static constexpr FORCEINLINE void StoreHalf(uint16* Ptr, float Value)
+	static constexpr inline void StoreHalf(uint16* Ptr, float Value)
 	{
 		union FP32T
 		{
@@ -123,7 +123,7 @@ struct FGenericPlatformMath
 		*Ptr = FP16;
 	}
 
-	static constexpr FORCEINLINE void VectorLoadHalf(float* RESTRICT Dst, const uint16* RESTRICT Src)
+	static constexpr inline void VectorLoadHalf(float* RESTRICT Dst, const uint16* RESTRICT Src)
 	{
 		Dst[0] = LoadHalf(&Src[0]);
 		Dst[1] = LoadHalf(&Src[1]);
@@ -131,7 +131,7 @@ struct FGenericPlatformMath
 		Dst[3] = LoadHalf(&Src[3]);
 	}
 
-	static constexpr FORCEINLINE void VectorStoreHalf(uint16* RESTRICT Dst, const float* RESTRICT Src)
+	static constexpr inline void VectorStoreHalf(uint16* RESTRICT Dst, const float* RESTRICT Src)
 	{
 		StoreHalf(&Dst[0], Src[0]);
 		StoreHalf(&Dst[1], Src[1]);
@@ -139,13 +139,13 @@ struct FGenericPlatformMath
 		StoreHalf(&Dst[3], Src[3]);
 	}
 
-	static constexpr FORCEINLINE void WideVectorLoadHalf(float* RESTRICT Dst, const uint16* RESTRICT Src)
+	static constexpr inline void WideVectorLoadHalf(float* RESTRICT Dst, const uint16* RESTRICT Src)
 	{
 		VectorLoadHalf(Dst, Src);
 		VectorLoadHalf(Dst + 4, Src + 4);
 	}
 
-	static constexpr FORCEINLINE void WideVectorStoreHalf(uint16* RESTRICT Dst, const float* RESTRICT Src)
+	static constexpr inline void WideVectorStoreHalf(uint16* RESTRICT Dst, const float* RESTRICT Src)
 	{
 		VectorStoreHalf(Dst, Src);
 		VectorStoreHalf(Dst + 4, Src + 4);
@@ -255,19 +255,19 @@ struct FGenericPlatformMath
 	 * @param F		Floating point value to convert
 	 * @return		An integer less or equal to 'F'.
 	 */
-	[[nodiscard]] static FORCEINLINE int32 FloorToInt32(float F)
+	[[nodiscard]] static inline int32 FloorToInt32(float F)
 	{
 		int32 I = TruncToInt32(F);
 		I -= ((float)I > F);
 		return I;
 	}
-	[[nodiscard]] static FORCEINLINE int32 FloorToInt32(double F)
+	[[nodiscard]] static inline int32 FloorToInt32(double F)
 	{
 		int32 I = TruncToInt32(F);
 		I -= ((double)I > F);
 		return I;
 	}
-	[[nodiscard]] static FORCEINLINE int64 FloorToInt64(double F)
+	[[nodiscard]] static inline int64 FloorToInt64(double F)
 	{
 		int64 I = TruncToInt64(F);
 		I -= ((double)I > F);
@@ -354,19 +354,19 @@ struct FGenericPlatformMath
 	* @param F		Floating point value to convert
 	* @return		An integer greater or equal to 'F'.
 	*/
-	[[nodiscard]] static FORCEINLINE int32 CeilToInt32(float F)
+	[[nodiscard]] static inline int32 CeilToInt32(float F)
 	{
 		int32 I = TruncToInt32(F);
 		I += ((float)I < F);
 		return I;
 	}
-	[[nodiscard]] static FORCEINLINE int32 CeilToInt32(double F)
+	[[nodiscard]] static inline int32 CeilToInt32(double F)
 	{
 		int32 I = TruncToInt32(F);
 		I += ((double)I < F);
 		return I;
 	}
-	[[nodiscard]] static FORCEINLINE int64 CeilToInt64(double F)
+	[[nodiscard]] static inline int64 CeilToInt64(double F)
 	{
 		int64 I = TruncToInt64(F);
 		I += ((double)I < F);
@@ -633,7 +633,7 @@ struct FGenericPlatformMath
 	static FORCEINLINE void RandInit(int32 Seed) { srand( Seed ); }
 
 	/** Returns a random float between 0 and 1, inclusive. */
-	[[nodiscard]] static FORCEINLINE float FRand() 
+	[[nodiscard]] static inline float FRand() 
 	{ 
 		// FP32 mantissa can only represent 24 bits before losing precision
 		constexpr int32 RandMax = 0x00ffffff < RAND_MAX ? 0x00ffffff : RAND_MAX;
@@ -656,7 +656,7 @@ struct FGenericPlatformMath
 	 * @param Value		The value to compute the log of
 	 * @return			Log2 of Value. 0 if Value is 0.
 	 */	
-	[[nodiscard]] static constexpr FORCEINLINE uint32 FloorLog2(uint32 Value) 
+	[[nodiscard]] static constexpr inline uint32 FloorLog2(uint32 Value) 
 	{
 		uint32 pos = 0;
 		if (Value >= 1<<16) { Value >>= 16; pos += 16; }
@@ -680,7 +680,7 @@ struct FGenericPlatformMath
 	 * @param Value		The value to compute the log of
 	 * @return			Log2 of Value. 0 if Value is 0.
 	 */	
-	[[nodiscard]] static constexpr FORCEINLINE uint64 FloorLog2_64(uint64 Value) 
+	[[nodiscard]] static constexpr inline uint64 FloorLog2_64(uint64 Value) 
 	{
 		uint64 pos = 0;
 		if (Value >= 1ull<<32) { Value >>= 32; pos += 32; }
@@ -705,7 +705,7 @@ struct FGenericPlatformMath
 	 *
 	 * @return the number of zeros before the first "on" bit
 	 */
-	[[nodiscard]] static constexpr FORCEINLINE uint8 CountLeadingZeros8(uint8 Value)
+	[[nodiscard]] static constexpr inline uint8 CountLeadingZeros8(uint8 Value)
 	{
 		if (Value == 0) return 8;
 		return uint8(7 - FloorLog2(uint32(Value)));
@@ -718,7 +718,7 @@ struct FGenericPlatformMath
 	 *
 	 * @return the number of zeros before the first "on" bit
 	 */
-	[[nodiscard]] static constexpr FORCEINLINE uint32 CountLeadingZeros(uint32 Value)
+	[[nodiscard]] static constexpr inline uint32 CountLeadingZeros(uint32 Value)
 	{
 		if (Value == 0) return 32;
 		return 31 - FloorLog2(Value);
@@ -731,7 +731,7 @@ struct FGenericPlatformMath
 	 *
 	 * @return the number of zeros before the first "on" bit
 	 */
-	[[nodiscard]] static constexpr FORCEINLINE uint64 CountLeadingZeros64(uint64 Value)
+	[[nodiscard]] static constexpr inline uint64 CountLeadingZeros64(uint64 Value)
 	{
 		if (Value == 0) return 64;
 		return 63 - FloorLog2_64(Value);
@@ -744,7 +744,7 @@ struct FGenericPlatformMath
 	 *
 	 * @return the number of zeros after the last "on" bit
 	 */
-	[[nodiscard]] static constexpr FORCEINLINE uint32 CountTrailingZerosConstExpr(uint32 Value)
+	[[nodiscard]] static constexpr inline uint32 CountTrailingZerosConstExpr(uint32 Value)
 	{
 		if (Value == 0)
 		{
@@ -771,7 +771,7 @@ struct FGenericPlatformMath
 	 *
 	 * @return the number of zeros after the last "on" bit
 	 */
-	[[nodiscard]] static constexpr FORCEINLINE uint64 CountTrailingZeros64ConstExpr(uint64 Value)
+	[[nodiscard]] static constexpr inline uint64 CountTrailingZeros64ConstExpr(uint64 Value)
 	{
 		if (Value == 0)
 		{
@@ -795,14 +795,14 @@ struct FGenericPlatformMath
 	 * Returns smallest N such that (1<<N)>=Arg.
 	 * Note: CeilLogTwo(0)=0 
 	 */
-	[[nodiscard]] static constexpr FORCEINLINE uint32 CeilLogTwo( uint32 Arg )
+	[[nodiscard]] static constexpr inline uint32 CeilLogTwo( uint32 Arg )
 	{
 		// if Arg is 0, change it to 1 so that we return 0
 		Arg = Arg ? Arg : 1;
 		return 32 - CountLeadingZeros(Arg - 1);
 	}
 
-	[[nodiscard]] static constexpr FORCEINLINE uint64 CeilLogTwo64( uint64 Arg )
+	[[nodiscard]] static constexpr inline uint64 CeilLogTwo64( uint64 Arg )
 	{
 		// if Arg is 0, change it to 1 so that we return 0
 		Arg = Arg ? Arg : 1;
@@ -813,7 +813,7 @@ struct FGenericPlatformMath
 	 * Returns the smallest N such that (1<<N)>=Arg. This is a less efficient version of CeilLogTwo, but written in a
 	 * way that can be evaluated at compile-time.
 	 */
-	[[nodiscard]] static constexpr FORCEINLINE uint8 ConstExprCeilLogTwo(SIZE_T Arg)
+	[[nodiscard]] static constexpr inline uint8 ConstExprCeilLogTwo(SIZE_T Arg)
 	{
 		if (Arg <= 1)
 		{
@@ -839,7 +839,7 @@ struct FGenericPlatformMath
 	}
 
 	/** Spreads bits to every other. */
-	[[nodiscard]] static constexpr FORCEINLINE uint32 MortonCode2( uint32 x )
+	[[nodiscard]] static constexpr inline uint32 MortonCode2( uint32 x )
 	{
 		x &= 0x0000ffff;
 		x = (x ^ (x << 8)) & 0x00ff00ff;
@@ -849,7 +849,7 @@ struct FGenericPlatformMath
 		return x;
 	}
 
-	[[nodiscard]] static constexpr FORCEINLINE uint64 MortonCode2_64( uint64 x )
+	[[nodiscard]] static constexpr inline uint64 MortonCode2_64( uint64 x )
 	{
 		x &= 0x00000000ffffffff;
 		x = (x ^ (x << 16)) & 0x0000ffff0000ffff;
@@ -861,7 +861,7 @@ struct FGenericPlatformMath
 	}
 
 	/** Reverses MortonCode2. Compacts every other bit to the right. */
-	[[nodiscard]] static constexpr FORCEINLINE uint32 ReverseMortonCode2( uint32 x )
+	[[nodiscard]] static constexpr inline uint32 ReverseMortonCode2( uint32 x )
 	{
 		x &= 0x55555555;
 		x = (x ^ (x >> 1)) & 0x33333333;
@@ -871,7 +871,7 @@ struct FGenericPlatformMath
 		return x;
 	}
 
-	[[nodiscard]] static constexpr FORCEINLINE uint64 ReverseMortonCode2_64( uint64 x )
+	[[nodiscard]] static constexpr inline uint64 ReverseMortonCode2_64( uint64 x )
 	{
 		x &= 0x5555555555555555;
 		x = (x ^ (x >> 1)) & 0x3333333333333333;
@@ -883,7 +883,7 @@ struct FGenericPlatformMath
 	}
 
 	/** Spreads bits to every 3rd. */
-	[[nodiscard]] static constexpr FORCEINLINE uint32 MortonCode3( uint32 x )
+	[[nodiscard]] static constexpr inline uint32 MortonCode3( uint32 x )
 	{
 		x &= 0x000003ff;
 		x = (x ^ (x << 16)) & 0xff0000ff;
@@ -894,7 +894,7 @@ struct FGenericPlatformMath
 	}
 
 	/** Reverses MortonCode3. Compacts every 3rd bit to the right. */
-	[[nodiscard]] static constexpr FORCEINLINE uint32 ReverseMortonCode3( uint32 x )
+	[[nodiscard]] static constexpr inline uint32 ReverseMortonCode3( uint32 x )
 	{
 		x &= 0x09249249;
 		x = (x ^ (x >>  2)) & 0x030c30c3;
@@ -969,7 +969,7 @@ struct FGenericPlatformMath
 		typename T, typename... OtherTypes
 		UE_REQUIRES((std::is_same_v<T, OtherTypes> && ...))
 	>
-	[[nodiscard]] static constexpr FORCEINLINE T Max(T A, OtherTypes... Others)
+	[[nodiscard]] static constexpr inline T Max(T A, OtherTypes... Others)
 	{
 		if constexpr (sizeof...(OtherTypes) > 0)
 		{
@@ -991,7 +991,7 @@ struct FGenericPlatformMath
 		typename T, typename... OtherTypes
 		UE_REQUIRES((std::is_same_v<T, OtherTypes> && ...))
 	>
-	[[nodiscard]] static constexpr FORCEINLINE T Min(T A, OtherTypes... Others)
+	[[nodiscard]] static constexpr inline T Min(T A, OtherTypes... Others)
 	{
 		if constexpr (sizeof...(OtherTypes) > 0)
 		{
@@ -1014,7 +1014,7 @@ struct FGenericPlatformMath
 	* @return	The min value found in the array or default value if the array was empty
 	*/
 	template< class T >
-	[[nodiscard]] static FORCEINLINE T Min(const TArray<T>& Values)
+	[[nodiscard]] static inline T Min(const TArray<T>& Values)
 	{
 		if (Values.Num() == 0)
 		{
@@ -1040,7 +1040,7 @@ struct FGenericPlatformMath
 	* @return	The min value found in the array or default value if the array was empty
 	*/
 	template< class T >
-	static FORCEINLINE T Min(const TArray<T>& Values, int32* MinIndex)
+	static inline T Min(const TArray<T>& Values, int32* MinIndex)
 	{
 		if (Values.Num() == 0)
 		{
@@ -1076,7 +1076,7 @@ struct FGenericPlatformMath
 * @return	The max value found in the array or default value if the array was empty
 */
 	template< class T >
-	[[nodiscard]] static FORCEINLINE T Max(const TArray<T>& Values)
+	[[nodiscard]] static inline T Max(const TArray<T>& Values)
 	{
 		if (Values.Num() == 0)
 		{
@@ -1103,7 +1103,7 @@ struct FGenericPlatformMath
 	* @return	The max value found in the array or default value if the array was empty
 	*/
 	template< class T >
-	static FORCEINLINE T Max(const TArray<T>& Values, int32* MaxIndex)
+	static inline T Max(const TArray<T>& Values, int32* MaxIndex)
 	{
 		if (Values.Num() == 0)
 		{
@@ -1137,7 +1137,7 @@ struct FGenericPlatformMath
 	 * Sort input A & B values into OutMin & OutMax
 	 */
 	template< class T >
-	static constexpr FORCEINLINE void GetMinMax(const T& A, const T& B, T& OutMin, T& OutMax)
+	static constexpr inline void GetMinMax(const T& A, const T& B, T& OutMin, T& OutMax)
 	{
 		if (A < B)
 		{
@@ -1155,7 +1155,7 @@ struct FGenericPlatformMath
 	 * Sort input values to be Min and Max
 	 */
 	template< class T >
-	static constexpr FORCEINLINE void GetMinMax(T& Min, T& Max)
+	static constexpr inline void GetMinMax(T& Min, T& Max)
 	{
 		if (Max < Min)
 		{
@@ -1163,7 +1163,7 @@ struct FGenericPlatformMath
 		}
 	}
 
-	[[nodiscard]] static constexpr FORCEINLINE int32 CountBits(uint64 Bits)
+	[[nodiscard]] static constexpr inline int32 CountBits(uint64 Bits)
 	{
 		// https://en.wikipedia.org/wiki/Hamming_weight
 		Bits -= (Bits >> 1) & 0x5555555555555555ull;
@@ -1183,7 +1183,7 @@ struct FGenericPlatformMath
 	 * If there wasn't overflow, it returns true, and the result of the addition is written to OutResult.
 	 */
 	template <UE::CIntegral IntType>
-	[[nodiscard]] static FORCEINLINE bool AddAndCheckForOverflow(IntType A, IntType B, IntType& OutResult)
+	[[nodiscard]] static inline bool AddAndCheckForOverflow(IntType A, IntType B, IntType& OutResult)
 	{
 		// This follows Hacker's Delight, Chapter 2-12
 		// Signed->unsigned conversion and unsigned addition have defined behavior always
@@ -1230,7 +1230,7 @@ struct FGenericPlatformMath
 	 * If there wasn't overflow, it returns true, and the result of the subtraction is written to OutResult.
 	 */
 	template <UE::CIntegral IntType>
-	[[nodiscard]] static FORCEINLINE bool SubtractAndCheckForOverflow(IntType A, IntType B, IntType& OutResult)
+	[[nodiscard]] static inline bool SubtractAndCheckForOverflow(IntType A, IntType B, IntType& OutResult)
 	{
 		// This follows Hacker's Delight, Chapter 2-12
 		// Signed->unsigned conversion and unsigned addition have defined behavior always
@@ -1271,7 +1271,7 @@ struct FGenericPlatformMath
 	 * If there wasn't overflow, it returns true, and the result of the multiplication is written to OutResult.
 	 */
 	template <UE::CIntegral IntType>
-	[[nodiscard]] static FORCEINLINE bool MultiplyAndCheckForOverflow(IntType A, IntType B, IntType& OutResult)
+	[[nodiscard]] static inline bool MultiplyAndCheckForOverflow(IntType A, IntType B, IntType& OutResult)
 	{
 		// Handle the case where the second factor is 0 specially (why will become clear in a minute).
 		if (B == 0)
diff --git a/Engine/Source/Runtime/Core/Public/GenericPlatform/GenericPlatformMemory.h b/Engine/Source/Runtime/Core/Public/GenericPlatform/GenericPlatformMemory.h
--- a/Engine/Source/Runtime/Core/Public/GenericPlatform/GenericPlatformMemory.h
+++ b/Engine/Source/Runtime/Core/Public/GenericPlatform/GenericPlatformMemory.h
@@ -606,7 +606,7 @@ struct FGenericPlatformMemory
 	}
 
 	/** On some platforms memcpy can be distributed over multiple threads for throughput. */
-	static FORCEINLINE void* ParallelMemcpy(void* Dest, const void* Src, SIZE_T Count, EMemcpyCachePolicy Policy = EMemcpyCachePolicy::StoreCached)
+	static inline void* ParallelMemcpy(void* Dest, const void* Src, SIZE_T Count, EMemcpyCachePolicy Policy = EMemcpyCachePolicy::StoreCached)
 	{
 		(void)Policy;
 		return memcpy( Dest, Src, Count );
@@ -614,7 +614,7 @@ struct FGenericPlatformMemory
 
 private:
 	template <typename T>
-	static FORCEINLINE void Valswap(T& A, T& B)
+	static inline void Valswap(T& A, T& B)
 	{
 		// Usually such an implementation would use move semantics, but
 		// we're only ever going to call it on fundamental types and MoveTemp
@@ -690,7 +690,7 @@ public:
 	* @return Value at Ptr
 	*/
 	template <typename T>
-	static FORCEINLINE T ReadUnaligned(const void* Ptr)
+	static inline T ReadUnaligned(const void* Ptr)
 	{
 		T AlignedT;
 		memcpy((void*)&AlignedT, Ptr, sizeof(T));
diff --git a/Engine/Source/Runtime/Core/Public/GenericPlatform/GenericPlatformProcess.h b/Engine/Source/Runtime/Core/Public/GenericPlatform/GenericPlatformProcess.h
--- a/Engine/Source/Runtime/Core/Public/GenericPlatform/GenericPlatformProcess.h
+++ b/Engine/Source/Runtime/Core/Public/GenericPlatform/GenericPlatformProcess.h
@@ -852,7 +852,7 @@ struct FGenericPlatformProcess
 	 * Tells the processor to pause for implementation-specific amount of time. Is used for spin-loops to improve the speed at 
 	 * which the code detects the release of the lock and power-consumption.
 	 */
-	static FORCEINLINE void Yield()
+	static inline void Yield()
 	{
 #if PLATFORM_USE_SSE2_FOR_THREAD_YIELD
 		_mm_pause();
@@ -872,7 +872,7 @@ struct FGenericPlatformProcess
 	* Tells the processor to pause for at least the amount of cycles given. Is used for spin-loops to improve the speed at 
 	* which the code detects the release of the lock and power-consumption.
 	*/
-	static FORCEINLINE void YieldCycles(uint64 Cycles)
+	static inline void YieldCycles(uint64 Cycles)
 	{
 #if PLATFORM_CPU_X86_FAMILY
 		auto ReadCycleCounter = []()
diff --git a/Engine/Source/Runtime/Core/Public/GenericPlatform/GenericPlatformString.h b/Engine/Source/Runtime/Core/Public/GenericPlatform/GenericPlatformString.h
--- a/Engine/Source/Runtime/Core/Public/GenericPlatform/GenericPlatformString.h
+++ b/Engine/Source/Runtime/Core/Public/GenericPlatform/GenericPlatformString.h
@@ -135,7 +135,7 @@ struct FGenericPlatformString : public FGenericPlatformStricmp
 	 * @return          A pointer to one past the last-written element.
 	 */
 	template <typename SourceEncoding, typename DestEncoding>
-	static FORCEINLINE DestEncoding* Convert(DestEncoding* Dest, int32 DestSize, const SourceEncoding* Src)
+	static inline DestEncoding* Convert(DestEncoding* Dest, int32 DestSize, const SourceEncoding* Src)
 	{
 		if constexpr (TIsCharEncodingSimplyConvertibleTo_V<SourceEncoding, DestEncoding>)
 		{
@@ -229,7 +229,7 @@ struct FGenericPlatformString : public FGenericPlatformStricmp
 	 * @return          A pointer to one past the last-written element.
 	 */
 	template <typename SourceEncoding, typename DestEncoding>
-	static FORCEINLINE DestEncoding* Convert(DestEncoding* Dest, int32 DestSize, const SourceEncoding* Src, int32 SrcSize)
+	static inline DestEncoding* Convert(DestEncoding* Dest, int32 DestSize, const SourceEncoding* Src, int32 SrcSize)
 	{
 		if constexpr (TIsCharEncodingCompatibleWith_V<SourceEncoding, DestEncoding>)
 		{
diff --git a/Engine/Source/Runtime/Core/Public/GenericPlatform/GenericPlatformTime.h b/Engine/Source/Runtime/Core/Public/GenericPlatform/GenericPlatformTime.h
--- a/Engine/Source/Runtime/Core/Public/GenericPlatform/GenericPlatformTime.h
+++ b/Engine/Source/Runtime/Core/Public/GenericPlatform/GenericPlatformTime.h
@@ -52,20 +52,20 @@ struct FGenericPlatformTime
 	 */
 	static CORE_API double InitTiming();
 
-	static FORCEINLINE double Seconds()
+	static inline double Seconds()
 	{
 		struct timeval tv;
 		gettimeofday( &tv, NULL );
 		return ((double) tv.tv_sec) + (((double) tv.tv_usec) / 1000000.0);
 	}
 
-	static FORCEINLINE uint32 Cycles()
+	static inline uint32 Cycles()
 	{
 		struct timeval tv;
 		gettimeofday( &tv, NULL );
 		return (uint32) ((((uint64)tv.tv_sec) * 1000000ULL) + (((uint64)tv.tv_usec)));
 	}
-	static FORCEINLINE uint64 Cycles64()
+	static inline uint64 Cycles64()
 	{
 		struct timeval tv;
 		gettimeofday( &tv, NULL );
diff --git a/Engine/Source/Runtime/Core/Public/GenericPlatform/GenericWidePlatformString.h b/Engine/Source/Runtime/Core/Public/GenericPlatform/GenericWidePlatformString.h
--- a/Engine/Source/Runtime/Core/Public/GenericPlatform/GenericWidePlatformString.h
+++ b/Engine/Source/Runtime/Core/Public/GenericPlatform/GenericWidePlatformString.h
@@ -212,7 +212,7 @@ public:
 	/**
 	 * Ansi implementation
 	 **/
-	static FORCEINLINE ANSICHAR* Strcpy(ANSICHAR* Dest, const ANSICHAR* Src)
+	static inline ANSICHAR* Strcpy(ANSICHAR* Dest, const ANSICHAR* Src)
 	{
 // Skip suggestions about using strcpy_s instead.
 PRAGMA_DISABLE_DEPRECATION_WARNINGS
@@ -228,7 +228,7 @@ PRAGMA_ENABLE_DEPRECATION_WARNINGS
 		return Strcpy(Dest, Src);
 	}
 
-	static FORCEINLINE ANSICHAR* Strncpy(ANSICHAR* Dest, const ANSICHAR* Src, SIZE_T MaxLen)
+	static inline ANSICHAR* Strncpy(ANSICHAR* Dest, const ANSICHAR* Src, SIZE_T MaxLen)
 	{
 // Skip suggestions about using strncpy_s instead.
 PRAGMA_DISABLE_DEPRECATION_WARNINGS
@@ -238,7 +238,7 @@ PRAGMA_ENABLE_DEPRECATION_WARNINGS
 		return Dest;
 	}
 
-	static FORCEINLINE ANSICHAR* Strcat(ANSICHAR* Dest, const ANSICHAR* Src)
+	static inline ANSICHAR* Strcat(ANSICHAR* Dest, const ANSICHAR* Src)
 	{
 // Skip suggestions about using strcat_s instead.
 PRAGMA_DISABLE_DEPRECATION_WARNINGS
@@ -258,7 +258,7 @@ PRAGMA_ENABLE_DEPRECATION_WARNINGS
 	 * If the length of the string in Src is less than SrcLen, only the content up to the terminating
 	 * null-character is copied.
 	 */
-	static FORCEINLINE ANSICHAR* Strncat(ANSICHAR* Dest, const ANSICHAR* Src, SIZE_T SrcLen)
+	static inline ANSICHAR* Strncat(ANSICHAR* Dest, const ANSICHAR* Src, SIZE_T SrcLen)
 	{
 // Skip suggestions about using strcat_s instead.
 PRAGMA_DISABLE_DEPRECATION_WARNINGS
@@ -336,7 +336,7 @@ PRAGMA_ENABLE_DEPRECATION_WARNINGS
 		return strtoull(Start, End, Base);
 	}
 
-	static FORCEINLINE ANSICHAR* Strtok(ANSICHAR* StrToken, const ANSICHAR* Delim, ANSICHAR** Context)
+	static inline ANSICHAR* Strtok(ANSICHAR* StrToken, const ANSICHAR* Delim, ANSICHAR** Context)
 	{
 // Skip suggestions about using strtok_s instead.
 PRAGMA_DISABLE_DEPRECATION_WARNINGS
@@ -506,7 +506,7 @@ PRAGMA_ENABLE_DEPRECATION_WARNINGS
 	 * UCS2 implementation
 	 **/
 
-	static FORCEINLINE int32 Strlen( const UCS2CHAR* String )
+	static inline int32 Strlen( const UCS2CHAR* String )
 	{
 		int32 Result = 0;
 		while (*String++)
@@ -517,7 +517,7 @@ PRAGMA_ENABLE_DEPRECATION_WARNINGS
 		return Result;
 	}
 
-	static FORCEINLINE int32 Strnlen( const UCS2CHAR* String, SIZE_T StringSize )
+	static inline int32 Strnlen( const UCS2CHAR* String, SIZE_T StringSize )
 	{
 		int32 Result = 0;
 		while (StringSize-- > 0 && *String++)
diff --git a/Engine/Source/Runtime/Core/Public/GenericPlatform/StandardPlatformString.h b/Engine/Source/Runtime/Core/Public/GenericPlatform/StandardPlatformString.h
--- a/Engine/Source/Runtime/Core/Public/GenericPlatform/StandardPlatformString.h
+++ b/Engine/Source/Runtime/Core/Public/GenericPlatform/StandardPlatformString.h
@@ -40,7 +40,7 @@ public:
 	/**
 	 * Unicode implementation
 	 **/
-	static FORCEINLINE WIDECHAR* Strcpy(WIDECHAR* Dest, const WIDECHAR* Src)
+	static inline WIDECHAR* Strcpy(WIDECHAR* Dest, const WIDECHAR* Src)
 	{
 // Skip suggestions about using wcscpy_s instead.
 PRAGMA_DISABLE_DEPRECATION_WARNINGS
@@ -51,7 +51,7 @@ PRAGMA_ENABLE_DEPRECATION_WARNINGS
 	// This version was deprecated because Strcpy taking a size field does not match the CRT strcpy, and on some platforms the size field was being ignored.
 	// If the memzeroing causes a performance problem, request a new function StrcpyTruncate.
 	UE_DEPRECATED(5.6, "Use Strncpy instead. Note that Strncpy has a behavior difference from Strcpy: it memzeroes the entire DestCount-sized buffer after the end of string.")
-	static FORCEINLINE WIDECHAR* Strcpy(WIDECHAR* Dest, SIZE_T DestCount, const WIDECHAR* Src)
+	static inline WIDECHAR* Strcpy(WIDECHAR* Dest, SIZE_T DestCount, const WIDECHAR* Src)
 	{
 		// Skip suggestions about using wcscpy_s instead.
 PRAGMA_DISABLE_DEPRECATION_WARNINGS
@@ -59,7 +59,7 @@ PRAGMA_DISABLE_DEPRECATION_WARNINGS
 PRAGMA_ENABLE_DEPRECATION_WARNINGS
 	}
 
-	static FORCEINLINE WIDECHAR* Strncpy(WIDECHAR* Dest, const WIDECHAR* Src, SIZE_T MaxLen)
+	static inline WIDECHAR* Strncpy(WIDECHAR* Dest, const WIDECHAR* Src, SIZE_T MaxLen)
 	{
 // Skip suggestions about using wcsncpy_s instead.
 PRAGMA_DISABLE_DEPRECATION_WARNINGS
@@ -69,7 +69,7 @@ PRAGMA_ENABLE_DEPRECATION_WARNINGS
 		return Dest;
 	}
 
-	static FORCEINLINE WIDECHAR* Strcat(WIDECHAR* Dest, const WIDECHAR* Src)
+	static inline WIDECHAR* Strcat(WIDECHAR* Dest, const WIDECHAR* Src)
 	{
 // Skip suggestions about using wcscat_s instead.
 PRAGMA_DISABLE_DEPRECATION_WARNINGS
@@ -79,7 +79,7 @@ PRAGMA_ENABLE_DEPRECATION_WARNINGS
 
 	// This version was deprecated because Strcat taking a size field does not match the CRT strcat, and on some platforms the size field was being ignored.
 	UE_DEPRECATED(5.6, "Use Strncat instead. !!NOTE THAT STRNCAT takes SrcLen rather than DestCount. You must call Strncat(Dest, Src, DestCount - Strlen(Dest) - 1).")
-	static FORCEINLINE WIDECHAR* Strcat(WIDECHAR* Dest, SIZE_T DestCount, const WIDECHAR* Src)
+	static inline WIDECHAR* Strcat(WIDECHAR* Dest, SIZE_T DestCount, const WIDECHAR* Src)
 	{
 		// Skip suggestions about using wcscat_s instead.
 PRAGMA_DISABLE_DEPRECATION_WARNINGS
@@ -87,7 +87,7 @@ PRAGMA_DISABLE_DEPRECATION_WARNINGS
 PRAGMA_ENABLE_DEPRECATION_WARNINGS
 	}
 
-	static FORCEINLINE WIDECHAR* Strncat(WIDECHAR* Dest, const WIDECHAR* Src, SIZE_T SrcLen)
+	static inline WIDECHAR* Strncat(WIDECHAR* Dest, const WIDECHAR* Src, SIZE_T SrcLen)
 	{
 		if (!SrcLen || !*Src)
 		{
@@ -174,7 +174,7 @@ PRAGMA_ENABLE_DEPRECATION_WARNINGS
 		return wcstoull( Start, End, Base );
 	}
 
-	static FORCEINLINE WIDECHAR* Strtok(WIDECHAR* StrToken, const WIDECHAR* Delim, WIDECHAR** Context)
+	static inline WIDECHAR* Strtok(WIDECHAR* StrToken, const WIDECHAR* Delim, WIDECHAR** Context)
 	{
 // Skip suggestions about using wcstok_s instead.
 PRAGMA_DISABLE_DEPRECATION_WARNINGS
@@ -249,7 +249,7 @@ PRAGMA_ENABLE_DEPRECATION_WARNINGS
 	/**
 	 * Ansi implementation
 	 **/
-	static FORCEINLINE ANSICHAR* Strcpy(ANSICHAR* Dest, const ANSICHAR* Src)
+	static inline ANSICHAR* Strcpy(ANSICHAR* Dest, const ANSICHAR* Src)
 	{
 // Skip suggestions about using strcpy_s instead.
 PRAGMA_DISABLE_DEPRECATION_WARNINGS
@@ -260,7 +260,7 @@ PRAGMA_ENABLE_DEPRECATION_WARNINGS
 	// This version was deprecated because Strcpy taking a size field does not match the CRT strcpy, and on some platforms the size field was being ignored.
 	// If the memzeroing causes a performance problem, request a new function StrcpyTruncate.
 	UE_DEPRECATED(5.6, "Use Strncpy instead. Note that Strncpy has a behavior difference from Strcpy: it memzeroes the entire DestCount-sized buffer after the end of string.")
-	static FORCEINLINE ANSICHAR* Strcpy(ANSICHAR* Dest, SIZE_T DestCount, const ANSICHAR* Src)
+	static inline ANSICHAR* Strcpy(ANSICHAR* Dest, SIZE_T DestCount, const ANSICHAR* Src)
 	{
 		// Skip suggestions about using strcpy_s instead.
 PRAGMA_DISABLE_DEPRECATION_WARNINGS
@@ -268,7 +268,7 @@ PRAGMA_DISABLE_DEPRECATION_WARNINGS
 PRAGMA_ENABLE_DEPRECATION_WARNINGS
 	}
 
-	static FORCEINLINE ANSICHAR* Strncpy(ANSICHAR* Dest, const ANSICHAR* Src, SIZE_T MaxLen)
+	static inline ANSICHAR* Strncpy(ANSICHAR* Dest, const ANSICHAR* Src, SIZE_T MaxLen)
 	{
 // Skip suggestions about using strncpy_s instead.
 PRAGMA_DISABLE_DEPRECATION_WARNINGS
@@ -278,7 +278,7 @@ PRAGMA_ENABLE_DEPRECATION_WARNINGS
 		return Dest;
 	}
 
-	static FORCEINLINE ANSICHAR* Strcat(ANSICHAR* Dest, const ANSICHAR* Src)
+	static inline ANSICHAR* Strcat(ANSICHAR* Dest, const ANSICHAR* Src)
 	{
 // Skip suggestions about using strcat_s instead.
 PRAGMA_DISABLE_DEPRECATION_WARNINGS
@@ -288,7 +288,7 @@ PRAGMA_ENABLE_DEPRECATION_WARNINGS
 
 	// This version was deprecated because Strcat taking a size field does not match the CRT strcat, and on some platforms the size field was being ignored.
 	UE_DEPRECATED(5.6, "Use Strncat instead. !!NOTE THAT STRNCAT takes SrcLen rather than DestSize. You must call Strncat(Dest, Src, DestCount - Strlen(Dest) - 1).")
-	static FORCEINLINE ANSICHAR* Strcat(ANSICHAR* Dest, SIZE_T DestCount, const ANSICHAR* Src)
+	static inline ANSICHAR* Strcat(ANSICHAR* Dest, SIZE_T DestCount, const ANSICHAR* Src)
 	{
 // Skip suggestions about using strcat_s instead.
 PRAGMA_DISABLE_DEPRECATION_WARNINGS
@@ -371,7 +371,7 @@ PRAGMA_ENABLE_DEPRECATION_WARNINGS
 		return strtoull(Start, End, Base);
 	}
 
-	static FORCEINLINE ANSICHAR* Strtok(ANSICHAR* StrToken, const ANSICHAR* Delim, ANSICHAR** Context)
+	static inline ANSICHAR* Strtok(ANSICHAR* StrToken, const ANSICHAR* Delim, ANSICHAR** Context)
 	{
 // Skip suggestions about using strtok_s instead.
 PRAGMA_DISABLE_DEPRECATION_WARNINGS
@@ -396,7 +396,7 @@ PRAGMA_ENABLE_DEPRECATION_WARNINGS
 	// This version was deprecated because Strcpy taking a size field does not match the CRT strcpy, and on some platforms the size field was being ignored.
 	// If the memzeroing causes a performance problem, request a new function StrcpyTruncate.
 	UE_DEPRECATED(5.6, "Use Strncpy instead. Note that Strncpy has a behavior difference from Strcpy: it memzeroes the entire DestCount-sized buffer after the end of string.")
-	static FORCEINLINE UTF8CHAR* Strcpy(UTF8CHAR* Dest, SIZE_T DestCount, const UTF8CHAR* Src)
+	static inline UTF8CHAR* Strcpy(UTF8CHAR* Dest, SIZE_T DestCount, const UTF8CHAR* Src)
 	{
 PRAGMA_DISABLE_DEPRECATION_WARNINGS
 		return (UTF8CHAR*)Strcpy((ANSICHAR*)Dest, DestCount, (const ANSICHAR*)Src);
@@ -415,7 +415,7 @@ PRAGMA_ENABLE_DEPRECATION_WARNINGS
 
 	// This version was deprecated because Strcat taking a size field does not match the CRT strcat, and on some platforms the size field was being ignored.
 	UE_DEPRECATED(5.6, "Use Strncat instead. !!NOTE THAT STRNCAT takes SrcLen rather than DestSize. You must call Strncat(Dest, Src, DestCount - Strlen(Dest) - 1).")
-	static FORCEINLINE UTF8CHAR* Strcat(UTF8CHAR* Dest, SIZE_T DestCount, const UTF8CHAR* Src)
+	static inline UTF8CHAR* Strcat(UTF8CHAR* Dest, SIZE_T DestCount, const UTF8CHAR* Src)
 	{
 PRAGMA_DISABLE_DEPRECATION_WARNINGS
 		return (UTF8CHAR*)Strcat((ANSICHAR*)Dest, DestCount, (const ANSICHAR*)Src);
@@ -511,7 +511,7 @@ PRAGMA_ENABLE_DEPRECATION_WARNINGS
 	 * UCS2 implementation
 	 **/
 
-	static FORCEINLINE int32 Strlen( const UCS2CHAR* String )
+	static inline int32 Strlen( const UCS2CHAR* String )
 	{
 		int32 Result = 0;
 		while (*String++)
@@ -522,7 +522,7 @@ PRAGMA_ENABLE_DEPRECATION_WARNINGS
 		return Result;
 	}
 
-	static FORCEINLINE int32 Strnlen( const UCS2CHAR* String, SIZE_T StringSize )
+	static inline int32 Strnlen( const UCS2CHAR* String, SIZE_T StringSize )
 	{
 		int32 Result = 0;
 		while (StringSize-- > 0 && *String++)
diff --git a/Engine/Source/Runtime/Core/Public/Linux/LinuxPlatformProperties.h b/Engine/Source/Runtime/Core/Public/Linux/LinuxPlatformProperties.h
--- a/Engine/Source/Runtime/Core/Public/Linux/LinuxPlatformProperties.h
+++ b/Engine/Source/Runtime/Core/Public/Linux/LinuxPlatformProperties.h
@@ -55,7 +55,7 @@ struct FLinuxPlatformProperties
 		return IS_ARM64;
 	}
 
-	static constexpr FORCEINLINE const char* PlatformName()
+	static constexpr inline const char* PlatformName()
 	{
 		if (IS_DEDICATED_SERVER)
 		{
@@ -90,7 +90,7 @@ struct FLinuxPlatformProperties
 		return true;
 	}
 
-	static constexpr FORCEINLINE bool SupportsAutoSDK()
+	static constexpr inline bool SupportsAutoSDK()
 	{
 // linux cross-compiling / cross-building from windows supports AutoSDK.  But hosted linux doesn't yet.
 #if PLATFORM_WINDOWS
@@ -130,7 +130,7 @@ struct FLinuxPlatformProperties
 		return true;
 	}
 
-	static constexpr FORCEINLINE float GetVariantPriority()
+	static constexpr inline float GetVariantPriority()
 	{
 		if (IS_DEDICATED_SERVER)
 		{
diff --git a/Engine/Source/Runtime/Core/Public/Mac/MacPlatformMisc.h b/Engine/Source/Runtime/Core/Public/Mac/MacPlatformMisc.h
--- a/Engine/Source/Runtime/Core/Public/Mac/MacPlatformMisc.h
+++ b/Engine/Source/Runtime/Core/Public/Mac/MacPlatformMisc.h
@@ -235,7 +235,7 @@ struct CORE_API FMacPlatformMisc : public FApplePlatformMisc
 
 	static CGDisplayModeRef GetSupportedDisplayMode(CGDirectDisplayID DisplayID, uint32 Width, uint32 Height);
 
-	FORCEINLINE static void ChooseHDRDeviceAndColorGamut(uint32 DeviceId, uint32 DisplayNitLevel, EDisplayOutputFormat& OutputDevice, EDisplayColorGamut& ColorGamut)
+	inline static void ChooseHDRDeviceAndColorGamut(uint32 DeviceId, uint32 DisplayNitLevel, EDisplayOutputFormat& OutputDevice, EDisplayColorGamut& ColorGamut)
 	{
 		// ScRGB, 1000 or 2000 nits, DCI-P3
 		OutputDevice = DisplayNitLevel == 1000 ? EDisplayOutputFormat::HDR_ACES_1000nit_ScRGB : EDisplayOutputFormat::HDR_ACES_2000nit_ScRGB;
diff --git a/Engine/Source/Runtime/Core/Public/Mac/MacPlatformProcess.h b/Engine/Source/Runtime/Core/Public/Mac/MacPlatformProcess.h
--- a/Engine/Source/Runtime/Core/Public/Mac/MacPlatformProcess.h
+++ b/Engine/Source/Runtime/Core/Public/Mac/MacPlatformProcess.h
@@ -13,7 +13,7 @@
 struct FProcState
 {
 	/** Default constructor. */
-	FORCEINLINE FProcState()
+	inline FProcState()
 		:	ProcessId(0)
 		,	bIsRunning(false)
 		,	bHasBeenWaitedFor(false)
@@ -61,7 +61,7 @@ protected:  // the below is not a public API!
 	// FProcState should not be copyable since it breeds problems (e.g. one copy could have wait()ed, but another would not know it)
 
 	/** Copy constructor - should not be publicly accessible */
-	FORCEINLINE FProcState(const FProcState& Other)
+	inline FProcState(const FProcState& Other)
 		:	ProcessId(Other.ProcessId)
 		,	bIsRunning(Other.bIsRunning)  // assume it is
 		,	bHasBeenWaitedFor(Other.bHasBeenWaitedFor)
@@ -72,7 +72,7 @@ protected:  // the below is not a public API!
 	}
 
 	/** Assignment operator - should not be publicly accessible */
-	FORCEINLINE FProcState& operator=(const FProcState& Other)
+	inline FProcState& operator=(const FProcState& Other)
 	{
 		checkf(false, TEXT("FProcState should not be copied"));
 		return *this;
@@ -131,7 +131,7 @@ struct FProcHandle
 	}
 
 	/** Resets the handle to invalid */
-	FORCEINLINE void Reset()
+	inline void Reset()
 	{
 		ProcInfo = nullptr;
 		OpenedPid = -1;
diff --git a/Engine/Source/Runtime/Core/Public/Mac/MacPlatformProperties.h b/Engine/Source/Runtime/Core/Public/Mac/MacPlatformProperties.h
--- a/Engine/Source/Runtime/Core/Public/Mac/MacPlatformProperties.h
+++ b/Engine/Source/Runtime/Core/Public/Mac/MacPlatformProperties.h
@@ -49,7 +49,7 @@ struct FMacPlatformProperties
 		return IS_CLIENT_ONLY;
 	}
 
-	constexpr static FORCEINLINE const char* PlatformName()
+	constexpr static inline const char* PlatformName()
 	{
 		if (IS_DEDICATED_SERVER)
 		{
@@ -104,7 +104,7 @@ struct FMacPlatformProperties
 		return true;
 	}
 
-	constexpr static FORCEINLINE float GetVariantPriority()
+	constexpr static inline float GetVariantPriority()
 	{
 		if (IS_DEDICATED_SERVER)
 		{
diff --git a/Engine/Source/Runtime/Core/Public/Unix/UnixPlatformProcess.h b/Engine/Source/Runtime/Core/Public/Unix/UnixPlatformProcess.h
--- a/Engine/Source/Runtime/Core/Public/Unix/UnixPlatformProcess.h
+++ b/Engine/Source/Runtime/Core/Public/Unix/UnixPlatformProcess.h
@@ -20,7 +20,7 @@ class Error;
 struct FProcState
 {
 	/** Default constructor. */
-	FORCEINLINE FProcState()
+	inline FProcState()
 		:	ProcessId(0)
 		,	bIsRunning(false)
 		,	bHasBeenWaitedFor(false)
@@ -68,7 +68,7 @@ protected:  // the below is not a public API!
 	// FProcState should not be copyable since it breeds problems (e.g. one copy could have wait()ed, but another would not know it)
 
 	/** Copy constructor - should not be publicly accessible */
-	FORCEINLINE FProcState(const FProcState& Other)
+	inline FProcState(const FProcState& Other)
 		:	ProcessId(Other.ProcessId)
 		,	bIsRunning(Other.bIsRunning)  // assume it is
 		,	bHasBeenWaitedFor(Other.bHasBeenWaitedFor)
@@ -79,7 +79,7 @@ protected:  // the below is not a public API!
 	}
 
 	/** Assignment operator - should not be publicly accessible */
-	FORCEINLINE FProcState& operator=(const FProcState& Other)
+	inline FProcState& operator=(const FProcState& Other)
 	{
 		checkf(false, TEXT("FProcState should not be copied"));
 		return *this;
@@ -140,7 +140,7 @@ struct FProcHandle
 	}
 
 	/** Resets the handle to invalid */
-	FORCEINLINE void Reset()
+	inline void Reset()
 	{
 		ProcInfo = nullptr;
 		OpenedPid = -1;
diff --git a/Engine/Source/Runtime/Core/Public/Unix/UnixPlatformString.h b/Engine/Source/Runtime/Core/Public/Unix/UnixPlatformString.h
--- a/Engine/Source/Runtime/Core/Public/Unix/UnixPlatformString.h
+++ b/Engine/Source/Runtime/Core/Public/Unix/UnixPlatformString.h
@@ -20,7 +20,7 @@ struct FUnixPlatformString : public
 #endif
 {
 	template<typename CHAR>
-	static FORCEINLINE int32 Strlen(const CHAR* String)
+	static inline int32 Strlen(const CHAR* String)
 	{
 		if(!String)
 			return 0;
diff --git a/Engine/Source/Runtime/Core/Public/Unix/UnixPlatformTLS.h b/Engine/Source/Runtime/Core/Public/Unix/UnixPlatformTLS.h
--- a/Engine/Source/Runtime/Core/Public/Unix/UnixPlatformTLS.h
+++ b/Engine/Source/Runtime/Core/Public/Unix/UnixPlatformTLS.h
@@ -23,7 +23,7 @@ struct FUnixTLS : public FGenericPlatformTLS
 	/**
 	 * Returns the currently executing thread's id
 	 */
-	static FORCEINLINE uint32 GetCurrentThreadId(void)
+	static inline uint32 GetCurrentThreadId(void)
 	{
 		// note: cannot use pthread_self() without updating the rest of API to opaque (or at least 64-bit) thread handles
 #if defined(_GNU_SOURCE)
diff --git a/Engine/Source/Runtime/Core/Public/Unix/UnixPlatformTime.h b/Engine/Source/Runtime/Core/Public/Unix/UnixPlatformTime.h
--- a/Engine/Source/Runtime/Core/Public/Unix/UnixPlatformTime.h
+++ b/Engine/Source/Runtime/Core/Public/Unix/UnixPlatformTime.h
@@ -18,7 +18,7 @@ struct FUnixTime : public FGenericPlatformTime
 {
 	static CORE_API double InitTiming();
 
-	static FORCEINLINE double Seconds()
+	static inline double Seconds()
 	{
 		if (UNLIKELY(ClockSource < 0))
 		{
@@ -30,7 +30,7 @@ struct FUnixTime : public FGenericPlatformTime
 		return static_cast<double>(ts.tv_sec) + static_cast<double>(ts.tv_nsec) / 1e9;
 	}
 
-	static FORCEINLINE uint32 Cycles()
+	static inline uint32 Cycles()
 	{
 		if (UNLIKELY(ClockSource < 0))
 		{
@@ -42,7 +42,7 @@ struct FUnixTime : public FGenericPlatformTime
 		return static_cast<uint32>(static_cast<uint64>(ts.tv_sec) * (uint64)1e6 + static_cast<uint64>(ts.tv_nsec) / 1000ULL);
 	}
 
-	static FORCEINLINE uint64 Cycles64()
+	static inline uint64 Cycles64()
 	{
 		struct timespec ts;
 		clock_gettime(CLOCK_MONOTONIC, &ts);
diff --git a/Engine/Source/Runtime/Core/Public/Windows/WindowsPlatformMisc.h b/Engine/Source/Runtime/Core/Public/Windows/WindowsPlatformMisc.h
--- a/Engine/Source/Runtime/Core/Public/Windows/WindowsPlatformMisc.h
+++ b/Engine/Source/Runtime/Core/Public/Windows/WindowsPlatformMisc.h
@@ -119,7 +119,7 @@ struct FWindowsPlatformMisc
 	static CORE_API EProcessDiagnosticFlags GetProcessDiagnostics();
 #endif
 
-	FORCEINLINE static void MemoryBarrier() 
+	inline static void MemoryBarrier() 
 	{
 #if PLATFORM_CPU_X86_FAMILY
 		_mm_sfence();
@@ -343,7 +343,7 @@ struct FWindowsPlatformMisc
 	 */
 	static CORE_API bool IsUsingBatterySaver();
 
-	FORCEINLINE static void ChooseHDRDeviceAndColorGamut(uint32 DeviceId, uint32 DisplayNitLevel, EDisplayOutputFormat& OutputDevice, EDisplayColorGamut& ColorGamut)
+	inline static void ChooseHDRDeviceAndColorGamut(uint32 DeviceId, uint32 DisplayNitLevel, EDisplayOutputFormat& OutputDevice, EDisplayColorGamut& ColorGamut)
 	{
 		// needs to match GRHIHDRDisplayOutputFormat chosen in FD3D12DynamicRHI::Init
 #if WITH_EDITOR
diff --git a/Engine/Source/Runtime/Core/Public/Windows/WindowsPlatformMutex.h b/Engine/Source/Runtime/Core/Public/Windows/WindowsPlatformMutex.h
--- a/Engine/Source/Runtime/Core/Public/Windows/WindowsPlatformMutex.h
+++ b/Engine/Source/Runtime/Core/Public/Windows/WindowsPlatformMutex.h
@@ -22,7 +22,7 @@ public:
 	FWindowsRecursiveMutex(const FWindowsRecursiveMutex&) = delete;
 	FWindowsRecursiveMutex& operator=(const FWindowsRecursiveMutex&) = delete;
 
-	FORCEINLINE FWindowsRecursiveMutex()
+	inline FWindowsRecursiveMutex()
 	{
 		CA_SUPPRESS(28125);
 		Windows::InitializeCriticalSectionAndSpinCount(&CriticalSection, 4000);
diff --git a/Engine/Source/Runtime/Core/Public/Windows/WindowsPlatformProperties.h b/Engine/Source/Runtime/Core/Public/Windows/WindowsPlatformProperties.h
--- a/Engine/Source/Runtime/Core/Public/Windows/WindowsPlatformProperties.h
+++ b/Engine/Source/Runtime/Core/Public/Windows/WindowsPlatformProperties.h
@@ -43,7 +43,7 @@ struct FWindowsPlatformProperties
 		return IS_CLIENT_ONLY;
 	}
 
-	static constexpr FORCEINLINE const char* PlatformName()
+	static constexpr inline const char* PlatformName()
 	{
 		if (IS_DEDICATED_SERVER)
 		{
@@ -118,7 +118,7 @@ struct FWindowsPlatformProperties
 		return true;
 	}
 
-	static constexpr FORCEINLINE float GetVariantPriority()
+	static constexpr inline float GetVariantPriority()
 	{
 		if (IS_DEDICATED_SERVER)
 		{
diff --git a/Engine/Source/Runtime/Core/Public/Windows/WindowsPlatformTLS.h b/Engine/Source/Runtime/Core/Public/Windows/WindowsPlatformTLS.h
--- a/Engine/Source/Runtime/Core/Public/Windows/WindowsPlatformTLS.h
+++ b/Engine/Source/Runtime/Core/Public/Windows/WindowsPlatformTLS.h
@@ -79,7 +79,7 @@ struct FWindowsPlatformTLS
 	 * @param SlotIndex the TLS index to store it in.
 	 * @param Value the value to store in the slot.
 	 */
-	static FORCEINLINE void SetTlsValue(uint32 SlotIndex, void* Value)
+	static inline void SetTlsValue(uint32 SlotIndex, void* Value)
 	{
 #if WITH_EDITOR
 		using namespace WindowsPlatformTLS_Private;
@@ -103,7 +103,7 @@ struct FWindowsPlatformTLS
 	 * @param SlotIndex The index of the slot to read.
 	 * @return The value stored in the slot.
 	 */
-	static FORCEINLINE void* GetTlsValue(uint32 SlotIndex)
+	static inline void* GetTlsValue(uint32 SlotIndex)
 	{
 #if WITH_EDITOR
 		using namespace WindowsPlatformTLS_Private;
diff --git a/Engine/Source/Runtime/Core/Public/Windows/WindowsPlatformTime.h b/Engine/Source/Runtime/Core/Public/Windows/WindowsPlatformTime.h
--- a/Engine/Source/Runtime/Core/Public/Windows/WindowsPlatformTime.h
+++ b/Engine/Source/Runtime/Core/Public/Windows/WindowsPlatformTime.h
@@ -20,7 +20,7 @@ struct FWindowsPlatformTime
 {
 	static CORE_API double InitTiming();
 
-	static FORCEINLINE double Seconds()
+	static inline double Seconds()
 	{
 		Windows::LARGE_INTEGER Cycles;
 		Windows::QueryPerformanceCounter(&Cycles);
@@ -29,14 +29,14 @@ struct FWindowsPlatformTime
 		return (double)Cycles.QuadPart * GetSecondsPerCycle() + 16777216.0;
 	}
 
-	static FORCEINLINE uint32 Cycles()
+	static inline uint32 Cycles()
 	{
 		Windows::LARGE_INTEGER Cycles;
 		Windows::QueryPerformanceCounter(&Cycles);
 		return (uint32)Cycles.QuadPart;
 	}
 
-	static FORCEINLINE uint64 Cycles64()
+	static inline uint64 Cycles64()
 	{
 		Windows::LARGE_INTEGER Cycles;
 		QueryPerformanceCounter(&Cycles);
