TBounds fixes + micro-opts. Also add tests.

TBounds was abusing VectorAnyGreaterThan; the return value happens to be movemask on SSE, but that's not guaranteed,
and not true on NEON! Also, each of these should use VectorMaskBits only once, since it's multiple instructions including
a memory access on NEON.

Only adding float tests since TBounds code is same for float/double (templated) and we already have individual
unit tests that the component operations on VectorFloat/VectorDouble do the right thing.

diff --git a/Engine/Source/Runtime/Core/Public/Math/Bounds.h b/Engine/Source/Runtime/Core/Public/Math/Bounds.h
--- a/Engine/Source/Runtime/Core/Public/Math/Bounds.h
+++ b/Engine/Source/Runtime/Core/Public/Math/Bounds.h
@@ -45,18 +45,18 @@ struct TBounds
 
 	[[nodiscard]] inline bool Intersect( const TBounds<T>& Other ) const
 	{
-		int Separated;
-		Separated  = VectorAnyGreaterThan( VectorLoadAligned( &Min ), VectorLoadAligned( &Other.Max ) );
-		Separated |= VectorAnyGreaterThan( VectorLoadAligned( &Other.Min ), VectorLoadAligned( &Max ) );
-		Separated &= 0b111;
-		return Separated == 0;
+		auto MinVsOtherMax = VectorCompareGT(VectorLoadAligned(&Min), VectorLoadAligned(&Other.Max));
+		auto OtherMinVsMax = VectorCompareGT(VectorLoadAligned(&Other.Min), VectorLoadAligned(&Max));
+		int SeparatedMask = VectorMaskBits(VectorBitwiseOr(MinVsOtherMax, OtherMinVsMax));
+		return (SeparatedMask & 7) == 0;
 	}
 
 	[[nodiscard]] inline bool Contains( const TBounds<T>& Other ) const
 	{
-		int MaskMin = VectorMaskBits( VectorCompareLE( VectorLoadAligned( &Min ), VectorLoadAligned( &Other.Min ) ) );
-		int MaskMax = VectorMaskBits( VectorCompareGE( VectorLoadAligned( &Max ), VectorLoadAligned( &Other.Max ) ) );
-		return ( MaskMin & MaskMax & 7 ) == 7;
+		auto MinVsOtherMin = VectorCompareLE(VectorLoadAligned(&Min), VectorLoadAligned(&Other.Min));
+		auto MaxVsOtherMax = VectorCompareGE(VectorLoadAligned(&Max), VectorLoadAligned(&Other.Max));
+		int ContainedMask = VectorMaskBits(VectorBitwiseAnd(MinVsOtherMin, MaxVsOtherMax));
+		return (ContainedMask & 7) == 7;
 	}
 
 	[[nodiscard]] inline T DistSqr( const TVector<T>& Point ) const
@@ -120,4 +120,4 @@ struct TBounds
 };
 
 using FBounds3f = TBounds< float >;
-using FBounds3d = TBounds< double >;
\ No newline at end of file
+using FBounds3d = TBounds< double >;
diff --git a/Engine/Source/Runtime/Core/Tests/Math/UnrealMathTest.cpp b/Engine/Source/Runtime/Core/Tests/Math/UnrealMathTest.cpp
--- a/Engine/Source/Runtime/Core/Tests/Math/UnrealMathTest.cpp
+++ b/Engine/Source/Runtime/Core/Tests/Math/UnrealMathTest.cpp
@@ -20,6 +20,7 @@
 #include "Math/QuatRotationTranslationMatrix.h"
 #include "Math/Color.h"
 #include "Math/Box.h"
+#include "Math/Bounds.h"
 #include "Async/ParallelFor.h"
 #include "Misc/ScopeLock.h"
 #include <limits>
@@ -1106,6 +1107,18 @@ FORCENOINLINE void LogTest<double>(const TCHAR* TestName, bool bHasPassed)
 	}
 }
 
+// Report an error if bComparison is not equal to bExpected.
+FORCENOINLINE void LogBoundsTest(bool bExpected, const TCHAR* TestName, const FBounds3f& A, const FBounds3f& B, bool bComparison)
+{
+	const bool bHasPassed = (bComparison == bExpected);
+	if (bHasPassed == false)
+	{
+		UE_LOG(LogUnrealMathTest, Log, TEXT("%s: %s"), bHasPassed ? TEXT("PASSED") : TEXT("FAILED"), TestName);
+		UE_LOG(LogUnrealMathTest, Log, TEXT("A=(Min:%s, Max:%s) B=(Min:%s, Max:%s)"), *A.Min.ToString(), *A.Max.ToString(), *B.Min.ToString(), *B.Max.ToString());
+		CheckPassing(false);
+	}
+}
+
 /**
  * Set the contents of the scratch memory
  *
@@ -3904,6 +3917,32 @@ TEST_CASE_NAMED(FVectorRegisterAbstractionTest, "System::Core::Math::Vector Regi
 		LogTest<float>(TEXT("QuaterionMatrixConversion Q1"), TestMatricesEqual(M0, M1, 0.000001f));
 	}
 
+	// Bounds
+	{
+		FBounds3f B0, B1, B2;
+
+		// B0 is the unit cube
+		B0 += FVector3f(-1.0f, -1.0f, -1.0f);
+		B0 += FVector3f( 1.0f,  1.0f,  1.0f);
+		// B1 is fully contained within
+		B1 += FVector3f(-0.50f, 0.00f, -0.75f);
+		B1 += FVector3f( 0.75f, 0.25f,  1.00f);
+		// B2 intersects with B0 but not with B1, with the separation occuring only on the Y axis
+		B2 += FVector3f(-0.75f, -1.25f, -0.60f);
+		B2 += FVector3f( 0.50f, -0.25f,  1.00f);
+
+		LogBoundsTest(true, TEXT("BoundsIntersect(B0, B1)"), B0, B1, B0.Intersect(B1));
+		LogBoundsTest(true, TEXT("BoundsIntersect(B0, B2)"), B0, B2, B0.Intersect(B2));
+		LogBoundsTest(false, TEXT("BoundsIntersect(B1, B2)"), B1, B2, B1.Intersect(B2));
+
+		LogBoundsTest(true, TEXT("BoundsContains(B0, B1)"), B0, B1, B0.Contains(B1));
+		LogBoundsTest(false, TEXT("BoundsContains(B0, B2)"), B0, B2, B0.Contains(B2));
+		LogBoundsTest(false, TEXT("BoundsContains(B1, B0)"), B1, B0, B1.Contains(B0));
+		LogBoundsTest(false, TEXT("BoundsContains(B1, B2)"), B1, B2, B1.Contains(B2));
+		LogBoundsTest(false, TEXT("BoundsContains(B2, B0)"), B2, B0, B2.Contains(B0));
+		LogBoundsTest(false, TEXT("BoundsContains(B2, B1)"), B2, B1, B2.Contains(B1));
+	}
+
 	if (!GPassing)
 	{
 		UE_LOG(LogUnrealMathTest, Fatal, TEXT("VectorIntrinsics <float> Failed."));
