De-inline FFloat16::SetTruncate

diff --git a/Engine/Source/Runtime/Core/Private/Math/Float16.cpp b/Engine/Source/Runtime/Core/Private/Math/Float16.cpp
new file mode 100644
--- /dev/null
+++ b/Engine/Source/Runtime/Core/Private/Math/Float16.cpp
@@ -0,0 +1,70 @@
+// Copyright Epic Games, Inc. All Rights Reserved.
+
+#include "Math/Float16.h"
+
+void FFloat16::SetTruncate(float FP32Value)
+{
+
+	union
+	{
+		struct
+		{
+#if PLATFORM_LITTLE_ENDIAN
+			uint16	Mantissa : 10;
+			uint16	Exponent : 5;
+			uint16	Sign : 1;
+#else
+			uint16	Sign : 1;
+			uint16	Exponent : 5;
+			uint16	Mantissa : 10;			
+#endif
+		} Components;
+
+		uint16	Encoded;
+	} FP16;
+
+
+	FFloat32 FP32(FP32Value);
+
+	// Copy sign-bit
+	FP16.Components.Sign = FP32.Components.Sign;
+
+	// Check for zero, denormal or too small value.
+	if (FP32.Components.Exponent <= 112)			// Too small exponent? (0+127-15)
+	{
+		// Set to 0.
+		FP16.Components.Exponent = 0;
+		FP16.Components.Mantissa = 0;
+
+         // Exponent unbias the single, then bias the halfp
+         const int32 NewExp = FP32.Components.Exponent - 127 + 15;
+ 
+         if ( (14 - NewExp) <= 24 ) // Mantissa might be non-zero
+         {
+             uint32 Mantissa = FP32.Components.Mantissa | 0x800000; // Hidden 1 bit
+             FP16.Components.Mantissa = (uint16)(Mantissa >> (14 - NewExp));
+			 // Check for rounding
+             if ( (Mantissa >> (13 - NewExp)) & 1 ) //-V1051
+			 {
+                 FP16.Encoded++; // Round, might overflow into exp bit, but this is OK
+			 }
+         }
+	}
+	// Check for INF or NaN, or too high value
+	else if (FP32.Components.Exponent >= 143)		// Too large exponent? (31+127-15)
+	{
+		// Set to 65504.0 (max value)
+		FP16.Components.Exponent = 30;
+		FP16.Components.Mantissa = 1023;
+	}
+	// Handle normal number.
+	else
+	{
+		FP16.Components.Exponent = uint16(int32(FP32.Components.Exponent) - 127 + 15);
+		FP16.Components.Mantissa = uint16(FP32.Components.Mantissa >> 13);
+	}
+
+	Encoded = FP16.Encoded;
+}
+
+
diff --git a/Engine/Source/Runtime/Core/Public/Math/Float16.h b/Engine/Source/Runtime/Core/Public/Math/Float16.h
--- a/Engine/Source/Runtime/Core/Public/Math/Float16.h
+++ b/Engine/Source/Runtime/Core/Public/Math/Float16.h
@@ -71,7 +71,7 @@ public:
 	/** Convert from Fp32 to Fp16, truncating low bits. 
 	(backward-compatible conversion; was used by Set() previously)
 	Clamps values out of range to [-MaxF16Float,MaxF16Float] */
-	void SetTruncate(float FP32Value);
+	void CORE_API SetTruncate(float FP32Value);
 
 	/** Set to 0.0 **/
 	void SetZero()
@@ -172,72 +172,6 @@ FORCEINLINE float FFloat16::GetFloat() const
 }
 
 
-// NOTE: SetTruncate() on values out of F16 max range store them as +-Inf
-inline void FFloat16::SetTruncate(float FP32Value)
-{
-
-	union
-	{
-		struct
-		{
-#if PLATFORM_LITTLE_ENDIAN
-			uint16	Mantissa : 10;
-			uint16	Exponent : 5;
-			uint16	Sign : 1;
-#else
-			uint16	Sign : 1;
-			uint16	Exponent : 5;
-			uint16	Mantissa : 10;			
-#endif
-		} Components;
-
-		uint16	Encoded;
-	} FP16;
-
-
-	FFloat32 FP32(FP32Value);
-
-	// Copy sign-bit
-	FP16.Components.Sign = FP32.Components.Sign;
-
-	// Check for zero, denormal or too small value.
-	if (FP32.Components.Exponent <= 112)			// Too small exponent? (0+127-15)
-	{
-		// Set to 0.
-		FP16.Components.Exponent = 0;
-		FP16.Components.Mantissa = 0;
-
-         // Exponent unbias the single, then bias the halfp
-         const int32 NewExp = FP32.Components.Exponent - 127 + 15;
- 
-         if ( (14 - NewExp) <= 24 ) // Mantissa might be non-zero
-         {
-             uint32 Mantissa = FP32.Components.Mantissa | 0x800000; // Hidden 1 bit
-             FP16.Components.Mantissa = (uint16)(Mantissa >> (14 - NewExp));
-			 // Check for rounding
-             if ( (Mantissa >> (13 - NewExp)) & 1 ) //-V1051
-			 {
-                 FP16.Encoded++; // Round, might overflow into exp bit, but this is OK
-			 }
-         }
-	}
-	// Check for INF or NaN, or too high value
-	else if (FP32.Components.Exponent >= 143)		// Too large exponent? (31+127-15)
-	{
-		// Set to 65504.0 (max value)
-		FP16.Components.Exponent = 30;
-		FP16.Components.Mantissa = 1023;
-	}
-	// Handle normal number.
-	else
-	{
-		FP16.Components.Exponent = uint16(int32(FP32.Components.Exponent) - 127 + 15);
-		FP16.Components.Mantissa = uint16(FP32.Components.Mantissa >> 13);
-	}
-
-	Encoded = FP16.Encoded;
-}
-
 /** Return float clamp in [0,MaxF16Float] , no negatives or infinites or nans returned **/
 inline FFloat16 FFloat16::GetClampedNonNegativeAndFinite() const
 {
@@ -282,4 +216,4 @@ inline FFloat16 FFloat16::GetClampedFinite() const
 	}
 	
 	return ReturnValue;
-}
\ No newline at end of file
+}
