Introduced FTextChar class to determine whether the specified code point is a letter character.

Added next functions: GetCodepoint, IsAlpha and IsWhitespace. Made FText::IsWhitepace as an alias for FTextChar::IsWhitespace.

#tests. Tested using test automation framework. Branch: Main

diff --git a/Engine/Source/Runtime/Core/Private/Internationalization/ICUText.cpp b/Engine/Source/Runtime/Core/Private/Internationalization/ICUText.cpp
--- a/Engine/Source/Runtime/Core/Private/Internationalization/ICUText.cpp
+++ b/Engine/Source/Runtime/Core/Private/Internationalization/ICUText.cpp
@@ -6,6 +6,7 @@
 #include "Logging/LogMacros.h"
 #include "Templates/SharedPointer.h"
 #include "Internationalization/Text.h"
+#include "Internationalization/TextChar.h"
 #include "Internationalization/TextChronoFormatter.h"
 #include "Internationalization/TextTransformer.h"
 #include "Internationalization/Internationalization.h"
@@ -95,12 +96,24 @@ FString FTextTransformer::ToUpper(const FString& InStr)
 	return ICUUtilities::ConvertString(ICUUtilities::ConvertString(InStr).toUpper());
 }
 
-bool FText::IsWhitespace(const TCHAR Char)
+bool FTextChar::IsAlpha(const UTF32CHAR Codepoint)
 {
-	// TCHAR should either be UTF-16 or UTF-32, so we should be fine to cast it to a UChar32 for the whitespace
-	// check, since whitespace is never a pair of UTF-16 characters
-	const UChar32 ICUChar = static_cast<UChar32>(Char);
-	return u_isWhitespace(ICUChar) != 0;
+	return u_isalpha(Codepoint) != 0;
+}
+
+bool FTextChar::IsWhitespace(const UTF32CHAR Codepoint)
+{
+	return u_isWhitespace(Codepoint) != 0;
+}
+
+UTF32CHAR FTextChar::ToLower(const UTF32CHAR Codepoint)
+{
+	return u_tolower(Codepoint);
+}
+
+UTF32CHAR FTextChar::ToUpper(const UTF32CHAR Codepoint)
+{
+	return u_toupper(Codepoint);
 }
 
 int32 FTextComparison::CompareTo( const FString& A, const FString& B, const ETextComparisonLevel::Type ComparisonLevel )
diff --git a/Engine/Source/Runtime/Core/Private/Internationalization/LegacyText.cpp b/Engine/Source/Runtime/Core/Private/Internationalization/LegacyText.cpp
--- a/Engine/Source/Runtime/Core/Private/Internationalization/LegacyText.cpp
+++ b/Engine/Source/Runtime/Core/Private/Internationalization/LegacyText.cpp
@@ -4,6 +4,7 @@
 #include "Misc/AssertionMacros.h"
 #include "Misc/DateTime.h"
 #include "Internationalization/Text.h"
+#include "Internationalization/TextChar.h"
 #include "Internationalization/TextChronoFormatter.h"
 #include "Internationalization/TextTransformer.h"
 #include "Internationalization/Internationalization.h"
@@ -45,9 +46,24 @@ FString FTextTransformer::ToUpper(const FString& InStr)
 	return InStr.ToUpper();
 }
 
-bool FText::IsWhitespace(const TCHAR Char)
+bool FTextChar::IsAlpha(const UTF32CHAR Codepoint)
 {
-	return FChar::IsWhitespace(Char);
+	return FChar::IsAlpha(static_cast<TCHAR>(Codepoint));
+}
+
+bool FTextChar::IsWhitespace(const UTF32CHAR Codepoint)
+{
+	return FChar::IsWhitespace(static_cast<TCHAR>(Codepoint));
+}
+
+UTF32CHAR FTextChar::ToLower(const UTF32CHAR Codepoint)
+{
+	return FChar::ToLower(static_cast<TCHAR>(Codepoint));
+}
+
+UTF32CHAR FTextChar::ToUpper(const UTF32CHAR Codepoint)
+{
+	return FChar::ToUpper(static_cast<TCHAR>(Codepoint));
 }
 
 int32 FTextComparison::CompareTo( const FString& A, const FString& B, const ETextComparisonLevel::Type ComparisonLevel )
diff --git a/Engine/Source/Runtime/Core/Private/Internationalization/Text.cpp b/Engine/Source/Runtime/Core/Private/Internationalization/Text.cpp
--- a/Engine/Source/Runtime/Core/Private/Internationalization/Text.cpp
+++ b/Engine/Source/Runtime/Core/Private/Internationalization/Text.cpp
@@ -15,8 +15,9 @@
 #include "Internationalization/TextHistory.h"
 #include "Misc/Guid.h"
 #include "Internationalization/TextCache.h"
-#include "Internationalization/TextFormatter.h"
+#include "Internationalization/TextChar.h"
 #include "Internationalization/TextChronoFormatter.h"
+#include "Internationalization/TextFormatter.h"
 #include "Internationalization/TextTransformer.h"
 #include "Internationalization/TextNamespaceUtil.h"
 #include "Internationalization/FastDecimalFormat.h"
@@ -453,6 +454,13 @@ FText FText::TrimPrecedingAndTrailing( const FText& InText )
 	return InText.IsCultureInvariant() ? FText::AsCultureInvariant(MoveTemp(TrimmedString)) : FText::FromString(MoveTemp(TrimmedString));
 }
 
+bool FText::IsWhitespace(const TCHAR Char)
+{
+	// TCHAR should either be UTF-16 or UTF-32, so we should be fine to cast it to a UTF32CHAR for the whitespace
+	// check, since whitespace is never a pair of UTF-16 characters
+	return FTextChar::IsWhitespace(static_cast<UTF32CHAR>(Char));
+}
+
 void FText::GetFormatPatternParameters(const FTextFormat& Fmt, TArray<FString>& ParameterNames)
 {
 	return Fmt.GetFormatArgumentNames(ParameterNames);
@@ -2114,6 +2122,40 @@ void FTextBuilder::BuildAndAppendLine(FText&& Data)
 	}
 }
 
+UTF32CHAR FTextChar::GetCodepoint(const TCHAR* Buffer, int32* OutNumTCHARInCodepoint)
+{
+	static_assert(sizeof(TCHAR) == 2 || sizeof(TCHAR) == 4, "Expected TCHAR to be UTF-16 or UTF-32");
+
+	if (!Buffer)
+	{
+		return 0;
+	}
+
+    UTF32CHAR Codepoint = 0;
+
+#if !PLATFORM_TCHAR_IS_4_BYTES
+	const TCHAR NextChar = *(Buffer + 1);
+    if (NextChar && StringConv::IsHighSurrogate(*Buffer) && StringConv::IsLowSurrogate(NextChar))
+    {
+        Codepoint = StringConv::EncodeSurrogate(*Buffer, NextChar);
+        if (OutNumTCHARInCodepoint)
+        {
+            *OutNumTCHARInCodepoint = 2;
+        }
+    }
+    else
+#endif  // !PLATFORM_TCHAR_IS_4_BYTES
+    {
+        Codepoint = *Buffer;
+        if (OutNumTCHARInCodepoint)
+        {
+            *OutNumTCHARInCodepoint = 1;
+        }
+    }
+
+    return Codepoint;
+}
+
 bool LexTryParseString(ETextGender& OutValue, const TCHAR* Buffer)
 {
 #define ENUM_CASE_FROM_STRING(Enum) if (FCString::Stricmp(Buffer, TEXT(#Enum)) == 0) { OutValue = ETextGender::Enum; return true; }
diff --git a/Engine/Source/Runtime/Core/Public/Internationalization/TextChar.h b/Engine/Source/Runtime/Core/Public/Internationalization/TextChar.h
new file mode 100644
--- /dev/null
+++ b/Engine/Source/Runtime/Core/Public/Internationalization/TextChar.h
@@ -0,0 +1,39 @@
+// Copyright Epic Games, Inc. All Rights Reserved.
+
+#pragma once
+
+#include "CoreTypes.h"
+
+/**
+ * Utility for performing character checks.
+ * The implementation can be found in LegacyText.cpp and ICUText.cpp.
+ */
+class FTextChar
+{
+public:
+	/**
+	* Return UTF32CHAR codepoint. Buffer must be a null terminated string and at least 1 character long
+	* EX: FString Name;
+	*     if (Name.Len() > 1)
+	*	  	  const UTF32CHAR FirstLetter = FTextChar::GetCodepoint(*Name);
+	*/
+	[[nodiscard]] static CORE_API UTF32CHAR GetCodepoint( const TCHAR* Buffer, int32* OutNumTCHARInCodepoint = nullptr );
+
+	/**
+	* Determine whether the specified code point is a letter character.
+	*/
+	[[nodiscard]] static CORE_API bool IsAlpha( const UTF32CHAR Codepoint );
+
+	/**
+	* Determine whether the Char code is a whitespace character.
+	* TCHAR should either be UTF-16 or UTF-32, so we should be fine to cast it to a UChar32 for the whitespace
+	* check, since whitespace is never a pair of UTF-16 characters
+	* EX: TChar Char;
+	*     IsWhitespace(static_cast<UTF32CHAR>(Char));
+	*/
+	[[nodiscard]] static CORE_API bool IsWhitespace( const UTF32CHAR Codepoint );
+
+	[[nodiscard]] static CORE_API UTF32CHAR ToLower( const UTF32CHAR Codepoint );
+
+	[[nodiscard]] static CORE_API UTF32CHAR ToUpper( const UTF32CHAR Codepoint );
+};
diff --git a/Engine/Source/Runtime/Core/Tests/Internationalization/TextTest.cpp b/Engine/Source/Runtime/Core/Tests/Internationalization/TextTest.cpp
--- a/Engine/Source/Runtime/Core/Tests/Internationalization/TextTest.cpp
+++ b/Engine/Source/Runtime/Core/Tests/Internationalization/TextTest.cpp
@@ -9,6 +9,8 @@
 #include "CoreGlobals.h"
 #include "AutoRTFM.h"
 #include "Internationalization/Text.h"
+#include "Internationalization/TextChar.h"
+#include "Internationalization/TextTransformer.h"
 #include "Internationalization/Culture.h"
 #include "Internationalization/FastDecimalFormat.h"
 #include "Internationalization/StringTableRegistry.h"
@@ -1561,6 +1563,46 @@ TEST_CASE_NAMED(FICUTextTest, "System::Core::Misc::ICUText", "[.][EditorContext]
 	}
 }
 
+TEST_CASE_NAMED(FICUTextUtilities, "System::Core::Misc::ICUTextUtilities", "[.][EditorContext][ClientContext][EngineFilter]")
+{
+	FInternationalization& I18N = FInternationalization::Get();
+
+	FInternationalization::FCultureStateSnapshot OriginalCultureState;
+	I18N.BackupCultureState(OriginalCultureState);
+
+	if (I18N.SetCurrentCulture("fr"))
+	{
+		if (FTextTransformer::ToLower(FString(TEXT("ÔTE"))) != FString(TEXT("ôte")))
+		{
+			FAIL_CHECK(FString::Printf(TEXT("FTextTransformer::ToLower doesn't work correctly")));
+		}
+		
+		if (FTextChar::IsWhitespace(TEXT('\u0212')) == true)
+		{
+			FAIL_CHECK(FString::Printf(TEXT("FTextChar::IsWhitespace can't detect whitespace")));
+		}
+
+		if (FTextChar::IsAlpha(TEXT('\u0212')) == false)
+		{
+			FAIL_CHECK(FString::Printf(TEXT("FTextChar::IsAlpha can't detect fr letter")));
+		}
+
+		if (FTextChar::IsAlpha(TEXT('7')) == true || FTextChar::IsAlpha(TEXT('A')) == false)
+		{
+			FAIL_CHECK(FString::Printf(TEXT("FTextChar::IsAlpha can't detect not letter")));
+		}
+		
+		TCHAR EmojiChars[] = { 0xD83D, 0xDE00, 0x0000 };
+		if (FTextChar::IsAlpha(FTextChar::GetCodepoint(EmojiChars)) == true)
+		{
+			FAIL_CHECK(FString::Printf(TEXT("FTextChar::IsAlpha can't detect surrogate pair")));
+		}
+	}
+
+	// Restore original culture
+	I18N.RestoreCultureState(OriginalCultureState);
+}
+
 #endif // #if UE_ENABLE_ICU
 
 #undef LOCTEXT_NAMESPACE 
