Move MicrosoftPlatformString.h from GenericPlatform to Microsoft

Also replace a few FORCEINLINEs with inline on some larger functions.

#rnx

diff --git a/Engine/Source/Runtime/Core/Public/GenericPlatform/MicrosoftPlatformString.h b/Engine/Source/Runtime/Core/Public/GenericPlatform/MicrosoftPlatformString.h
--- a/Engine/Source/Runtime/Core/Public/GenericPlatform/MicrosoftPlatformString.h
+++ b/Engine/Source/Runtime/Core/Public/GenericPlatform/MicrosoftPlatformString.h
@@ -2,455 +2,10 @@
 
 #pragma once
 
-// HEADER_UNIT_SKIP - Included through other header
+#include "Misc/CoreMiscDefines.h"
 
-#include "Misc/Char.h"
-#if PLATFORM_USE_GENERIC_STRING_IMPLEMENTATION
-	#include "GenericPlatform/GenericWidePlatformString.h"
-#else
-	#include "GenericPlatform/GenericPlatformString.h"
-#endif
-#include "GenericPlatform/GenericPlatformStricmp.h"
-#include <stdarg.h>
-#include <stdio.h>
-#include <stdlib.h>
-#include <string.h>
-#include <tchar.h>
-
-/**
-* Microsoft specific implementation 
-**/
-
-#pragma warning(push)
-#pragma warning(disable : 4996) // 'function' was declared deprecated  (needed for the secure string functions)
-#pragma warning(disable : 4995) // 'function' was declared deprecated  (needed for the secure string functions)
-
-struct FMicrosoftPlatformString :
-#if PLATFORM_USE_GENERIC_STRING_IMPLEMENTATION
-	public FGenericWidePlatformString
-#else
-	public FGenericPlatformString
-#endif
-{
-#if PLATFORM_USE_GENERIC_STRING_IMPLEMENTATION
-	using Super = FGenericWidePlatformString;
-#else
-	using Super = FGenericPlatformString;
-#endif
-
-	using FGenericPlatformString::Stricmp;
-	using FGenericPlatformString::Strncmp;
-	using FGenericPlatformString::Strnicmp;
-
-#if !PLATFORM_USE_GENERIC_STRING_IMPLEMENTATION
-	template <typename CharType>
-	static CharType* Strupr(CharType* Dest, SIZE_T DestCount)
-	{
-		for (CharType* Char = Dest; *Char && DestCount > 0; Char++, DestCount--)
-		{
-			*Char = TChar<CharType>::ToUpper(*Char);
-		}
-		return Dest;
-	}
-#endif
-
-	/**
-	 * Wide character implementation 
-	 **/
-	static FORCEINLINE WIDECHAR* Strcpy(WIDECHAR* Dest, const WIDECHAR* Src)
-	{
-		return (WIDECHAR*)_tcscpy(Dest, Src);
-	}
-
-	// This version was deprecated because Strcpy taking a size field does not match the CRT strcpy, and on some platforms the size field was being ignored.
-	// If the memzeroing causes a performance problem, request a new function StrcpyTruncate.
-	UE_DEPRECATED(5.6, "Use Strncpy instead. Note that Strncpy has a behavior difference from Strcpy: it memzeroes the entire DestCount-sized buffer after the end of string.")
-	static FORCEINLINE WIDECHAR* Strcpy(WIDECHAR* Dest, SIZE_T DestCount, const WIDECHAR* Src)
-	{
-		return (WIDECHAR*)_tcscpy(Dest, Src);
-	}
-
-	static FORCEINLINE WIDECHAR* Strncpy(WIDECHAR* Dest, const WIDECHAR* Src, SIZE_T MaxLen)
-	{
-		_tcsncpy(Dest, Src, MaxLen-1);
-		Dest[MaxLen-1] = 0;
-		return Dest;
-	}
-
-	static FORCEINLINE WIDECHAR* Strcat(WIDECHAR* Dest, const WIDECHAR* Src)
-	{
-		return (WIDECHAR*)_tcscat(Dest, Src);
-	}
-
-	// This version was deprecated because Strcat taking a size field does not match the CRT strcat, and on some platforms the size field was being ignored.
-	UE_DEPRECATED(5.6, "Use Strncat instead. !!NOTE THAT STRNCAT takes SrcLen rather than DestCount. You must call Strncat(Dest, Src, DestCount - Strlen(Dest) - 1).")
-	static FORCEINLINE WIDECHAR* Strcat(WIDECHAR* Dest, SIZE_T DestCount, const WIDECHAR* Src)
-	{
-		return (WIDECHAR*)_tcscat(Dest, Src);
-	}
-
-	static FORCEINLINE WIDECHAR* Strncat(WIDECHAR* Dest, const WIDECHAR* Src, SIZE_T SrcLen)
-	{
-		// The Microsoft library has _tcscat_s but not a strncat equivalent. strncat takes SrcLen
-		// but tcscat_s takes size of the destination string buffer (which must also contain null terminator).
-		// We therefore provide our own implementation.
-		if (!SrcLen || !*Src)
-		{
-			return Dest;
-		}
-		WIDECHAR* NewDest = Dest + Strlen(Dest);
-		SIZE_T AppendedCount = 0;
-		do
-		{
-			*NewDest++ = *Src++;
-			++AppendedCount;
-		} while (AppendedCount < SrcLen && *Src);
-		*NewDest = 0;
-		return Dest;
-	}
-
-	static FORCEINLINE int32 Strcmp( const WIDECHAR* String1, const WIDECHAR* String2 )
-	{
-		return _tcscmp(String1, String2);
-	}
-
-	static FORCEINLINE int32 Strncmp( const WIDECHAR* String1, const WIDECHAR* String2, SIZE_T Count )
-	{
-		return _tcsncmp( String1, String2, Count );
-	}
-
-	static FORCEINLINE int32 Strlen( const WIDECHAR* String )
-	{
-		return _tcslen( String );
-	}
-
-	static FORCEINLINE int32 Strnlen( const WIDECHAR* String, SIZE_T StringSize )
-	{
-		return _tcsnlen( String, StringSize );
-	}
+UE_DEPRECATED_HEADER(5.7, "Use Microsoft/MicrosoftPlatformString.h instead.") 
 
-	static FORCEINLINE const WIDECHAR* Strstr( const WIDECHAR* String, const WIDECHAR* Find)
-	{
-		return _tcsstr( String, Find );
-	}
-
-	static FORCEINLINE const WIDECHAR* Strchr( const WIDECHAR* String, WIDECHAR C)
-	{
-		return _tcschr( String, C ); 
-	}
-
-	static FORCEINLINE const WIDECHAR* Strrchr( const WIDECHAR* String, WIDECHAR C)
-	{
-		return _tcsrchr( String, C ); 
-	}
-
-	static FORCEINLINE int32 Atoi(const WIDECHAR* String)
-	{
-		return _tstoi( String ); 
-	}
-
-	static FORCEINLINE int64 Atoi64(const WIDECHAR* String)
-	{
-		return _tstoi64( String ); 
-	}
-
-	static FORCEINLINE float Atof(const WIDECHAR* String)
-	{
-		return (float)_tstof( String );
-	}
-
-	static FORCEINLINE double Atod(const WIDECHAR* String)
-	{
-		return _tcstod( String, NULL ); 
-	}
-
-	static FORCEINLINE int32 Strtoi( const WIDECHAR* Start, WIDECHAR** End, int32 Base ) 
-	{
-		return _tcstoul( Start, End, Base );
-	}
-
-	static FORCEINLINE int64 Strtoi64( const WIDECHAR* Start, WIDECHAR** End, int32 Base ) 
-	{
-		return _tcstoi64( Start, End, Base ); 
-	}
-
-	static FORCEINLINE uint64 Strtoui64( const WIDECHAR* Start, WIDECHAR** End, int32 Base ) 
-	{
-		return _tcstoui64( Start, End, Base ); 
-	}
-
-	static FORCEINLINE WIDECHAR* Strtok(WIDECHAR* StrToken, const WIDECHAR* Delim, WIDECHAR** Context)
-	{
-		return _tcstok_s(StrToken, Delim, Context);
-	}
-
-// Allow fallback to FGenericWidePlatformString::GetVarArgs when PLATFORM_USE_GENERIC_STRING_IMPLEMENTATION is set.
-#if PLATFORM_USE_GENERIC_STRING_IMPLEMENTATION
-	using Super::GetVarArgs;
-#else
-	static FORCEINLINE int32 GetVarArgs( WIDECHAR* Dest, SIZE_T DestSize, const WIDECHAR*& Fmt, va_list ArgPtr )
-	{
-		int32 Result = vswprintf(Dest, DestSize, Fmt, ArgPtr);
-		return Result;
-	}
+#if PLATFORM_MICROSOFT
+#include "Microsoft/MicrosoftPlatformString.h"
 #endif
-
-	/** 
-	 * Ansi implementation 
-	 **/
-	static FORCEINLINE ANSICHAR* Strcpy(ANSICHAR* Dest, const ANSICHAR* Src)
-	{
-		return (ANSICHAR*)strcpy(Dest, Src);
-	}
-
-	// This version was deprecated because Strcpy taking a size field does not match the CRT strcpy, and on some platforms the size field was being ignored.
-	// If the memzeroing causes a performance problem, request a new function StrcpyTruncate.
-	UE_DEPRECATED(5.6, "Use Strncpy instead. Note that Strncpy has a behavior difference from Strcpy: it memzeroes the entire DestCount-sized buffer after the end of string.")
-	static FORCEINLINE ANSICHAR* Strcpy(ANSICHAR* Dest, SIZE_T DestCount, const ANSICHAR* Src)
-	{
-		return (ANSICHAR*)strcpy(Dest, Src);
-	}
-
-	static FORCEINLINE ANSICHAR* Strncpy(ANSICHAR* Dest, const ANSICHAR* Src, SIZE_T MaxLen)
-	{
-		strncpy(Dest, Src, MaxLen);
-		Dest[MaxLen-1] = 0;
-		return Dest;
-	}
-
-	static FORCEINLINE ANSICHAR* Strcat(ANSICHAR* Dest, const ANSICHAR* Src)
-	{
-		return (ANSICHAR*)strcat( Dest, Src );
-	}
-
-	// This version was deprecated because Strcat taking a size field does not match the CRT strcat, and on some platforms the size field was being ignored.
-	UE_DEPRECATED(5.6, "Use Strncat instead. !!NOTE THAT STRNCAT takes SrcLen rather than DestSize. You must call Strncat(Dest, Src, DestCount - Strlen(Dest) - 1).")
-	static FORCEINLINE ANSICHAR* Strcat(ANSICHAR* Dest, SIZE_T DestCount, const ANSICHAR* Src)
-	{
-		return (ANSICHAR*)strcat(Dest, Src);
-	}
-
-	static FORCEINLINE ANSICHAR* Strncat(ANSICHAR* Dest, const ANSICHAR* Src, SIZE_T SrcLen)
-	{
-		return (ANSICHAR*)strncat(Dest, Src, SrcLen);
-	}
-
-	static FORCEINLINE int32 Strcmp( const ANSICHAR* String1, const ANSICHAR* String2 )
-	{
-		return strcmp(String1, String2);
-	}
-
-	static FORCEINLINE int32 Strncmp( const ANSICHAR* String1, const ANSICHAR* String2, SIZE_T Count )
-	{
-		return strncmp( String1, String2, Count );
-	}
-	
-	static FORCEINLINE int32 Strlen( const ANSICHAR* String )
-	{
-		return strlen( String ); 
-	}
-
-	static FORCEINLINE int32 Strnlen( const ANSICHAR* String, SIZE_T StringSize )
-	{
-		return strnlen( String, StringSize );
-	}
-
-	static FORCEINLINE const ANSICHAR* Strstr( const ANSICHAR* String, const ANSICHAR* Find)
-	{
-		return strstr(String, Find);
-	}
-
-	static FORCEINLINE const ANSICHAR* Strchr( const ANSICHAR* String, ANSICHAR C)
-	{
-		return strchr(String, C);
-	}
-
-	static FORCEINLINE const ANSICHAR* Strrchr( const ANSICHAR* String, ANSICHAR C)
-	{
-		return strrchr(String, C);
-	}
-
-	static FORCEINLINE int32 Atoi(const ANSICHAR* String)
-	{
-		return atoi( String ); 
-	}
-
-	static FORCEINLINE int64 Atoi64(const ANSICHAR* String)
-	{
-		return _strtoi64( String, NULL, 10 ); 
-	}
-
-	static FORCEINLINE float Atof(const ANSICHAR* String)
-	{
-		return (float)atof( String ); 
-	}
-
-	static FORCEINLINE double Atod(const ANSICHAR* String)
-	{
-		return atof( String ); 
-	}
-
-	static FORCEINLINE int32 Strtoi( const ANSICHAR* Start, ANSICHAR** End, int32 Base ) 
-	{
-		return strtol( Start, End, Base ); 
-	}
-
-	static FORCEINLINE int64 Strtoi64( const ANSICHAR* Start, ANSICHAR** End, int32 Base ) 
-	{
-		return _strtoi64( Start, End, Base ); 
-	}
-
-	static FORCEINLINE uint64 Strtoui64( const ANSICHAR* Start, ANSICHAR** End, int32 Base ) 
-	{
-		return _strtoui64( Start, End, Base ); 
-	}
-
-	static FORCEINLINE ANSICHAR* Strtok(ANSICHAR* StrToken, const ANSICHAR* Delim, ANSICHAR** Context)
-	{
-		return strtok_s(StrToken, Delim, Context);
-	}
-
-	static FORCEINLINE int32 GetVarArgs( ANSICHAR* Dest, SIZE_T DestSize, const ANSICHAR*& Fmt, va_list ArgPtr )
-	{
-		int32 Result = vsnprintf( Dest, DestSize, Fmt, ArgPtr );
-		return (Result != -1 && Result < (int32)DestSize) ? Result : -1;
-	}
-
-	/**
-	 * UCS2CHAR implementation - this is identical to WIDECHAR for Windows platforms
-	 */
-
-	static FORCEINLINE int32 Strlen( const UCS2CHAR* String )
-	{
-		return _tcslen( (const WIDECHAR*)String );
-	}
-
-	static FORCEINLINE int32 Strnlen( const UCS2CHAR* String, SIZE_T StringSize )
-	{
-		return _tcsnlen( (const WIDECHAR*)String, StringSize );
-	}
-
-	/**
-	 * UTF8CHAR implementation.
-	 */
-	static FORCEINLINE UTF8CHAR* Strcpy(UTF8CHAR* Dest, const UTF8CHAR* Src)
-	{
-		return (UTF8CHAR*)Strcpy((ANSICHAR*)Dest, (const ANSICHAR*)Src);
-	}
-
-	// This version was deprecated because Strcpy taking a size field does not match the CRT strcpy, and on some platforms the size field was being ignored.
-	// If the memzeroing causes a performance problem, request a new function StrcpyTruncate.
-	UE_DEPRECATED(5.6, "Use Strncpy instead. Note that Strncpy has a behavior difference from Strcpy: it memzeroes the entire DestCount-sized buffer after the end of string.")
-	static FORCEINLINE UTF8CHAR* Strcpy(UTF8CHAR* Dest, SIZE_T DestCount, const UTF8CHAR* Src)
-	{
-		PRAGMA_DISABLE_DEPRECATION_WARNINGS
-		return (UTF8CHAR*)Strcpy((ANSICHAR*)Dest, DestCount, (const ANSICHAR*)Src);
-		PRAGMA_ENABLE_DEPRECATION_WARNINGS
-	}
-
-	static FORCEINLINE UTF8CHAR* Strncpy(UTF8CHAR* Dest, const UTF8CHAR* Src, SIZE_T MaxLen)
-	{
-		return (UTF8CHAR*)Strncpy((ANSICHAR*)Dest, (const ANSICHAR*)Src, MaxLen);
-	}
-
-	static FORCEINLINE UTF8CHAR* Strcat(UTF8CHAR* Dest, const UTF8CHAR* Src)
-	{
-		return (UTF8CHAR*)Strcat((ANSICHAR*)Dest, (const ANSICHAR*)Src);
-	}
-
-	// This version was deprecated because Strcat taking a size field does not match the CRT strcat, and on some platforms the size field was being ignored.
-	UE_DEPRECATED(5.6, "Use Strncat instead. !!NOTE THAT STRNCAT takes SrcLen rather than DestSize. You must call Strncat(Dest, Src, DestCount - Strlen(Dest) - 1).")
-	static FORCEINLINE UTF8CHAR* Strcat(UTF8CHAR* Dest, SIZE_T DestCount, const UTF8CHAR* Src)
-	{
-		PRAGMA_DISABLE_DEPRECATION_WARNINGS
-		return (UTF8CHAR*)Strcat((ANSICHAR*)Dest, DestCount, (const ANSICHAR*)Src);
-		PRAGMA_ENABLE_DEPRECATION_WARNINGS
-	}
-
-	static FORCEINLINE UTF8CHAR* Strncat(UTF8CHAR* Dest, const UTF8CHAR* Src, SIZE_T SrcLen)
-	{
-		return (UTF8CHAR*)Strncat((ANSICHAR*)Dest, (const ANSICHAR*)Src, SrcLen);
-	}
-
-	static FORCEINLINE int32 Strcmp(const UTF8CHAR* String1, const UTF8CHAR* String2)
-	{
-		return Strcmp((const ANSICHAR*)String1, (const ANSICHAR*)String2);
-	}
-
-	static FORCEINLINE int32 Strncmp(const UTF8CHAR* String1, const UTF8CHAR* String2, SIZE_T Count)
-	{
-		return Strncmp((const ANSICHAR*)String1, (const ANSICHAR*)String2, Count);
-	}
-	
-	static FORCEINLINE int32 Strlen(const UTF8CHAR* String)
-	{
-		return Strlen((const ANSICHAR*)String); 
-	}
-
-	static FORCEINLINE int32 Strnlen(const UTF8CHAR* String, SIZE_T StringSize)
-	{
-		return Strnlen((const ANSICHAR*)String, StringSize);
-	}
-
-	static FORCEINLINE const UTF8CHAR* Strstr(const UTF8CHAR* String, const UTF8CHAR* Find)
-	{
-		return (const UTF8CHAR*)Strstr((const ANSICHAR*)String, (const ANSICHAR*)Find);
-	}
-
-	static FORCEINLINE const UTF8CHAR* Strchr(const UTF8CHAR* String, UTF8CHAR C)
-	{
-		return (const UTF8CHAR*)Strchr((const ANSICHAR*)String, (ANSICHAR)C);
-	}
-
-	static FORCEINLINE const UTF8CHAR* Strrchr(const UTF8CHAR* String, UTF8CHAR C)
-	{
-		return (const UTF8CHAR*)Strrchr((const ANSICHAR*)String, (ANSICHAR)C);
-	}
-
-	static FORCEINLINE int32 Atoi(const UTF8CHAR* String)
-	{
-		return Atoi((const ANSICHAR*)String);
-	}
-
-	static FORCEINLINE int64 Atoi64(const UTF8CHAR* String)
-	{
-		return Atoi64((const ANSICHAR*)String);
-	}
-
-	static FORCEINLINE float Atof(const UTF8CHAR* String)
-	{
-		return Atof((const ANSICHAR*)String);
-	}
-
-	static FORCEINLINE double Atod(const UTF8CHAR* String)
-	{
-		return Atod((const ANSICHAR*)String);
-	}
-
-	static FORCEINLINE int32 Strtoi(const UTF8CHAR* Start, UTF8CHAR** End, int32 Base)
-	{
-		return Strtoi((const ANSICHAR*)Start, (ANSICHAR**)End, Base);
-	}
-
-	static FORCEINLINE int64 Strtoi64(const UTF8CHAR* Start, UTF8CHAR** End, int32 Base)
-	{
-		return Strtoi64((const ANSICHAR*)Start, (ANSICHAR**)End, Base);
-	}
-
-	static FORCEINLINE uint64 Strtoui64(const UTF8CHAR* Start, UTF8CHAR** End, int32 Base)
-	{
-		return Strtoui64((const ANSICHAR*)Start, (ANSICHAR**)End, Base);
-	}
-
-	static FORCEINLINE UTF8CHAR* Strtok(UTF8CHAR* StrToken, const UTF8CHAR* Delim, UTF8CHAR** Context)
-	{
-		return (UTF8CHAR*)Strtok((ANSICHAR*)StrToken, (const ANSICHAR*)Delim, (ANSICHAR**)Context);
-	}
-
-	static FORCEINLINE int32 GetVarArgs(UTF8CHAR* Dest, SIZE_T DestSize, const UTF8CHAR*& Fmt, va_list ArgPtr)
-	{
-		return GetVarArgs((ANSICHAR*)Dest, DestSize, *(const ANSICHAR**)&Fmt, ArgPtr);
-	}
-};
-
-#pragma warning(pop) // 'function' was was declared deprecated  (needed for the secure string functions)
diff --git a/Engine/Source/Runtime/Core/Public/Microsoft/MicrosoftPlatformString.h b/Engine/Source/Runtime/Core/Public/Microsoft/MicrosoftPlatformString.h
new file mode 100644
--- /dev/null
+++ b/Engine/Source/Runtime/Core/Public/Microsoft/MicrosoftPlatformString.h
@@ -0,0 +1,456 @@
+// Copyright Epic Games, Inc. All Rights Reserved.
+
+#pragma once
+
+// HEADER_UNIT_SKIP - Included through other header
+
+#include "Misc/Char.h"
+#if PLATFORM_USE_GENERIC_STRING_IMPLEMENTATION
+	#include "GenericPlatform/GenericWidePlatformString.h"
+#else
+	#include "GenericPlatform/GenericPlatformString.h"
+#endif
+#include "GenericPlatform/GenericPlatformStricmp.h"
+#include <stdarg.h>
+#include <stdio.h>
+#include <stdlib.h>
+#include <string.h>
+#include <tchar.h>
+
+/**
+* Microsoft specific implementation 
+**/
+
+#pragma warning(push)
+#pragma warning(disable : 4996) // 'function' was declared deprecated  (needed for the secure string functions)
+#pragma warning(disable : 4995) // 'function' was declared deprecated  (needed for the secure string functions)
+
+struct FMicrosoftPlatformString :
+#if PLATFORM_USE_GENERIC_STRING_IMPLEMENTATION
+	public FGenericWidePlatformString
+#else
+	public FGenericPlatformString
+#endif
+{
+#if PLATFORM_USE_GENERIC_STRING_IMPLEMENTATION
+	using Super = FGenericWidePlatformString;
+#else
+	using Super = FGenericPlatformString;
+#endif
+
+	using FGenericPlatformString::Stricmp;
+	using FGenericPlatformString::Strncmp;
+	using FGenericPlatformString::Strnicmp;
+
+#if !PLATFORM_USE_GENERIC_STRING_IMPLEMENTATION
+	template <typename CharType>
+	static CharType* Strupr(CharType* Dest, SIZE_T DestCount)
+	{
+		for (CharType* Char = Dest; *Char && DestCount > 0; Char++, DestCount--)
+		{
+			*Char = TChar<CharType>::ToUpper(*Char);
+		}
+		return Dest;
+	}
+#endif
+
+	/**
+	 * Wide character implementation 
+	 **/
+	static FORCEINLINE WIDECHAR* Strcpy(WIDECHAR* Dest, const WIDECHAR* Src)
+	{
+		return (WIDECHAR*)_tcscpy(Dest, Src);
+	}
+
+	// This version was deprecated because Strcpy taking a size field does not match the CRT strcpy, and on some platforms the size field was being ignored.
+	// If the memzeroing causes a performance problem, request a new function StrcpyTruncate.
+	UE_DEPRECATED(5.6, "Use Strncpy instead. Note that Strncpy has a behavior difference from Strcpy: it memzeroes the entire DestCount-sized buffer after the end of string.")
+	static FORCEINLINE WIDECHAR* Strcpy(WIDECHAR* Dest, SIZE_T DestCount, const WIDECHAR* Src)
+	{
+		return (WIDECHAR*)_tcscpy(Dest, Src);
+	}
+
+	static inline WIDECHAR* Strncpy(WIDECHAR* Dest, const WIDECHAR* Src, SIZE_T MaxLen)
+	{
+		_tcsncpy(Dest, Src, MaxLen-1);
+		Dest[MaxLen-1] = 0;
+		return Dest;
+	}
+
+	static FORCEINLINE WIDECHAR* Strcat(WIDECHAR* Dest, const WIDECHAR* Src)
+	{
+		return (WIDECHAR*)_tcscat(Dest, Src);
+	}
+
+	// This version was deprecated because Strcat taking a size field does not match the CRT strcat, and on some platforms the size field was being ignored.
+	UE_DEPRECATED(5.6, "Use Strncat instead. !!NOTE THAT STRNCAT takes SrcLen rather than DestCount. You must call Strncat(Dest, Src, DestCount - Strlen(Dest) - 1).")
+	static FORCEINLINE WIDECHAR* Strcat(WIDECHAR* Dest, SIZE_T DestCount, const WIDECHAR* Src)
+	{
+		return (WIDECHAR*)_tcscat(Dest, Src);
+	}
+
+	static inline WIDECHAR* Strncat(WIDECHAR* Dest, const WIDECHAR* Src, SIZE_T SrcLen)
+	{
+		// The Microsoft library has _tcscat_s but not a strncat equivalent. strncat takes SrcLen
+		// but tcscat_s takes size of the destination string buffer (which must also contain null terminator).
+		// We therefore provide our own implementation.
+		if (!SrcLen || !*Src)
+		{
+			return Dest;
+		}
+		WIDECHAR* NewDest = Dest + Strlen(Dest);
+		SIZE_T AppendedCount = 0;
+		do
+		{
+			*NewDest++ = *Src++;
+			++AppendedCount;
+		} while (AppendedCount < SrcLen && *Src);
+		*NewDest = 0;
+		return Dest;
+	}
+
+	static FORCEINLINE int32 Strcmp( const WIDECHAR* String1, const WIDECHAR* String2 )
+	{
+		return _tcscmp(String1, String2);
+	}
+
+	static FORCEINLINE int32 Strncmp( const WIDECHAR* String1, const WIDECHAR* String2, SIZE_T Count )
+	{
+		return _tcsncmp( String1, String2, Count );
+	}
+
+	static FORCEINLINE int32 Strlen( const WIDECHAR* String )
+	{
+		return _tcslen( String );
+	}
+
+	static FORCEINLINE int32 Strnlen( const WIDECHAR* String, SIZE_T StringSize )
+	{
+		return _tcsnlen( String, StringSize );
+	}
+
+	static FORCEINLINE const WIDECHAR* Strstr( const WIDECHAR* String, const WIDECHAR* Find)
+	{
+		return _tcsstr( String, Find );
+	}
+
+	static FORCEINLINE const WIDECHAR* Strchr( const WIDECHAR* String, WIDECHAR C)
+	{
+		return _tcschr( String, C ); 
+	}
+
+	static FORCEINLINE const WIDECHAR* Strrchr( const WIDECHAR* String, WIDECHAR C)
+	{
+		return _tcsrchr( String, C ); 
+	}
+
+	static FORCEINLINE int32 Atoi(const WIDECHAR* String)
+	{
+		return _tstoi( String ); 
+	}
+
+	static FORCEINLINE int64 Atoi64(const WIDECHAR* String)
+	{
+		return _tstoi64( String ); 
+	}
+
+	static FORCEINLINE float Atof(const WIDECHAR* String)
+	{
+		return (float)_tstof( String );
+	}
+
+	static FORCEINLINE double Atod(const WIDECHAR* String)
+	{
+		return _tcstod( String, NULL ); 
+	}
+
+	static FORCEINLINE int32 Strtoi( const WIDECHAR* Start, WIDECHAR** End, int32 Base ) 
+	{
+		return _tcstoul( Start, End, Base );
+	}
+
+	static FORCEINLINE int64 Strtoi64( const WIDECHAR* Start, WIDECHAR** End, int32 Base ) 
+	{
+		return _tcstoi64( Start, End, Base ); 
+	}
+
+	static FORCEINLINE uint64 Strtoui64( const WIDECHAR* Start, WIDECHAR** End, int32 Base ) 
+	{
+		return _tcstoui64( Start, End, Base ); 
+	}
+
+	static FORCEINLINE WIDECHAR* Strtok(WIDECHAR* StrToken, const WIDECHAR* Delim, WIDECHAR** Context)
+	{
+		return _tcstok_s(StrToken, Delim, Context);
+	}
+
+// Allow fallback to FGenericWidePlatformString::GetVarArgs when PLATFORM_USE_GENERIC_STRING_IMPLEMENTATION is set.
+#if PLATFORM_USE_GENERIC_STRING_IMPLEMENTATION
+	using Super::GetVarArgs;
+#else
+	static inline int32 GetVarArgs( WIDECHAR* Dest, SIZE_T DestSize, const WIDECHAR*& Fmt, va_list ArgPtr )
+	{
+		int32 Result = vswprintf(Dest, DestSize, Fmt, ArgPtr);
+		return Result;
+	}
+#endif
+
+	/** 
+	 * Ansi implementation 
+	 **/
+	static FORCEINLINE ANSICHAR* Strcpy(ANSICHAR* Dest, const ANSICHAR* Src)
+	{
+		return (ANSICHAR*)strcpy(Dest, Src);
+	}
+
+	// This version was deprecated because Strcpy taking a size field does not match the CRT strcpy, and on some platforms the size field was being ignored.
+	// If the memzeroing causes a performance problem, request a new function StrcpyTruncate.
+	UE_DEPRECATED(5.6, "Use Strncpy instead. Note that Strncpy has a behavior difference from Strcpy: it memzeroes the entire DestCount-sized buffer after the end of string.")
+	static FORCEINLINE ANSICHAR* Strcpy(ANSICHAR* Dest, SIZE_T DestCount, const ANSICHAR* Src)
+	{
+		return (ANSICHAR*)strcpy(Dest, Src);
+	}
+
+	static inline ANSICHAR* Strncpy(ANSICHAR* Dest, const ANSICHAR* Src, SIZE_T MaxLen)
+	{
+		strncpy(Dest, Src, MaxLen);
+		Dest[MaxLen-1] = 0;
+		return Dest;
+	}
+
+	static FORCEINLINE ANSICHAR* Strcat(ANSICHAR* Dest, const ANSICHAR* Src)
+	{
+		return (ANSICHAR*)strcat( Dest, Src );
+	}
+
+	// This version was deprecated because Strcat taking a size field does not match the CRT strcat, and on some platforms the size field was being ignored.
+	UE_DEPRECATED(5.6, "Use Strncat instead. !!NOTE THAT STRNCAT takes SrcLen rather than DestSize. You must call Strncat(Dest, Src, DestCount - Strlen(Dest) - 1).")
+	static FORCEINLINE ANSICHAR* Strcat(ANSICHAR* Dest, SIZE_T DestCount, const ANSICHAR* Src)
+	{
+		return (ANSICHAR*)strcat(Dest, Src);
+	}
+
+	static FORCEINLINE ANSICHAR* Strncat(ANSICHAR* Dest, const ANSICHAR* Src, SIZE_T SrcLen)
+	{
+		return (ANSICHAR*)strncat(Dest, Src, SrcLen);
+	}
+
+	static FORCEINLINE int32 Strcmp( const ANSICHAR* String1, const ANSICHAR* String2 )
+	{
+		return strcmp(String1, String2);
+	}
+
+	static FORCEINLINE int32 Strncmp( const ANSICHAR* String1, const ANSICHAR* String2, SIZE_T Count )
+	{
+		return strncmp( String1, String2, Count );
+	}
+	
+	static FORCEINLINE int32 Strlen( const ANSICHAR* String )
+	{
+		return strlen( String ); 
+	}
+
+	static FORCEINLINE int32 Strnlen( const ANSICHAR* String, SIZE_T StringSize )
+	{
+		return strnlen( String, StringSize );
+	}
+
+	static FORCEINLINE const ANSICHAR* Strstr( const ANSICHAR* String, const ANSICHAR* Find)
+	{
+		return strstr(String, Find);
+	}
+
+	static FORCEINLINE const ANSICHAR* Strchr( const ANSICHAR* String, ANSICHAR C)
+	{
+		return strchr(String, C);
+	}
+
+	static FORCEINLINE const ANSICHAR* Strrchr( const ANSICHAR* String, ANSICHAR C)
+	{
+		return strrchr(String, C);
+	}
+
+	static FORCEINLINE int32 Atoi(const ANSICHAR* String)
+	{
+		return atoi( String ); 
+	}
+
+	static FORCEINLINE int64 Atoi64(const ANSICHAR* String)
+	{
+		return _strtoi64( String, NULL, 10 ); 
+	}
+
+	static FORCEINLINE float Atof(const ANSICHAR* String)
+	{
+		return (float)atof( String ); 
+	}
+
+	static FORCEINLINE double Atod(const ANSICHAR* String)
+	{
+		return atof( String ); 
+	}
+
+	static FORCEINLINE int32 Strtoi( const ANSICHAR* Start, ANSICHAR** End, int32 Base ) 
+	{
+		return strtol( Start, End, Base ); 
+	}
+
+	static FORCEINLINE int64 Strtoi64( const ANSICHAR* Start, ANSICHAR** End, int32 Base ) 
+	{
+		return _strtoi64( Start, End, Base ); 
+	}
+
+	static FORCEINLINE uint64 Strtoui64( const ANSICHAR* Start, ANSICHAR** End, int32 Base ) 
+	{
+		return _strtoui64( Start, End, Base ); 
+	}
+
+	static FORCEINLINE ANSICHAR* Strtok(ANSICHAR* StrToken, const ANSICHAR* Delim, ANSICHAR** Context)
+	{
+		return strtok_s(StrToken, Delim, Context);
+	}
+
+	static inline int32 GetVarArgs( ANSICHAR* Dest, SIZE_T DestSize, const ANSICHAR*& Fmt, va_list ArgPtr )
+	{
+		int32 Result = vsnprintf( Dest, DestSize, Fmt, ArgPtr );
+		return (Result != -1 && Result < (int32)DestSize) ? Result : -1;
+	}
+
+	/**
+	 * UCS2CHAR implementation - this is identical to WIDECHAR for Windows platforms
+	 */
+
+	static FORCEINLINE int32 Strlen( const UCS2CHAR* String )
+	{
+		return _tcslen( (const WIDECHAR*)String );
+	}
+
+	static FORCEINLINE int32 Strnlen( const UCS2CHAR* String, SIZE_T StringSize )
+	{
+		return _tcsnlen( (const WIDECHAR*)String, StringSize );
+	}
+
+	/**
+	 * UTF8CHAR implementation.
+	 */
+	static FORCEINLINE UTF8CHAR* Strcpy(UTF8CHAR* Dest, const UTF8CHAR* Src)
+	{
+		return (UTF8CHAR*)Strcpy((ANSICHAR*)Dest, (const ANSICHAR*)Src);
+	}
+
+	// This version was deprecated because Strcpy taking a size field does not match the CRT strcpy, and on some platforms the size field was being ignored.
+	// If the memzeroing causes a performance problem, request a new function StrcpyTruncate.
+	UE_DEPRECATED(5.6, "Use Strncpy instead. Note that Strncpy has a behavior difference from Strcpy: it memzeroes the entire DestCount-sized buffer after the end of string.")
+	static inline UTF8CHAR* Strcpy(UTF8CHAR* Dest, SIZE_T DestCount, const UTF8CHAR* Src)
+	{
+		PRAGMA_DISABLE_DEPRECATION_WARNINGS
+		return (UTF8CHAR*)Strcpy((ANSICHAR*)Dest, DestCount, (const ANSICHAR*)Src);
+		PRAGMA_ENABLE_DEPRECATION_WARNINGS
+	}
+
+	static FORCEINLINE UTF8CHAR* Strncpy(UTF8CHAR* Dest, const UTF8CHAR* Src, SIZE_T MaxLen)
+	{
+		return (UTF8CHAR*)Strncpy((ANSICHAR*)Dest, (const ANSICHAR*)Src, MaxLen);
+	}
+
+	static FORCEINLINE UTF8CHAR* Strcat(UTF8CHAR* Dest, const UTF8CHAR* Src)
+	{
+		return (UTF8CHAR*)Strcat((ANSICHAR*)Dest, (const ANSICHAR*)Src);
+	}
+
+	// This version was deprecated because Strcat taking a size field does not match the CRT strcat, and on some platforms the size field was being ignored.
+	UE_DEPRECATED(5.6, "Use Strncat instead. !!NOTE THAT STRNCAT takes SrcLen rather than DestSize. You must call Strncat(Dest, Src, DestCount - Strlen(Dest) - 1).")
+	static inline UTF8CHAR* Strcat(UTF8CHAR* Dest, SIZE_T DestCount, const UTF8CHAR* Src)
+	{
+		PRAGMA_DISABLE_DEPRECATION_WARNINGS
+		return (UTF8CHAR*)Strcat((ANSICHAR*)Dest, DestCount, (const ANSICHAR*)Src);
+		PRAGMA_ENABLE_DEPRECATION_WARNINGS
+	}
+
+	static FORCEINLINE UTF8CHAR* Strncat(UTF8CHAR* Dest, const UTF8CHAR* Src, SIZE_T SrcLen)
+	{
+		return (UTF8CHAR*)Strncat((ANSICHAR*)Dest, (const ANSICHAR*)Src, SrcLen);
+	}
+
+	static FORCEINLINE int32 Strcmp(const UTF8CHAR* String1, const UTF8CHAR* String2)
+	{
+		return Strcmp((const ANSICHAR*)String1, (const ANSICHAR*)String2);
+	}
+
+	static FORCEINLINE int32 Strncmp(const UTF8CHAR* String1, const UTF8CHAR* String2, SIZE_T Count)
+	{
+		return Strncmp((const ANSICHAR*)String1, (const ANSICHAR*)String2, Count);
+	}
+	
+	static FORCEINLINE int32 Strlen(const UTF8CHAR* String)
+	{
+		return Strlen((const ANSICHAR*)String); 
+	}
+
+	static FORCEINLINE int32 Strnlen(const UTF8CHAR* String, SIZE_T StringSize)
+	{
+		return Strnlen((const ANSICHAR*)String, StringSize);
+	}
+
+	static FORCEINLINE const UTF8CHAR* Strstr(const UTF8CHAR* String, const UTF8CHAR* Find)
+	{
+		return (const UTF8CHAR*)Strstr((const ANSICHAR*)String, (const ANSICHAR*)Find);
+	}
+
+	static FORCEINLINE const UTF8CHAR* Strchr(const UTF8CHAR* String, UTF8CHAR C)
+	{
+		return (const UTF8CHAR*)Strchr((const ANSICHAR*)String, (ANSICHAR)C);
+	}
+
+	static FORCEINLINE const UTF8CHAR* Strrchr(const UTF8CHAR* String, UTF8CHAR C)
+	{
+		return (const UTF8CHAR*)Strrchr((const ANSICHAR*)String, (ANSICHAR)C);
+	}
+
+	static FORCEINLINE int32 Atoi(const UTF8CHAR* String)
+	{
+		return Atoi((const ANSICHAR*)String);
+	}
+
+	static FORCEINLINE int64 Atoi64(const UTF8CHAR* String)
+	{
+		return Atoi64((const ANSICHAR*)String);
+	}
+
+	static FORCEINLINE float Atof(const UTF8CHAR* String)
+	{
+		return Atof((const ANSICHAR*)String);
+	}
+
+	static FORCEINLINE double Atod(const UTF8CHAR* String)
+	{
+		return Atod((const ANSICHAR*)String);
+	}
+
+	static FORCEINLINE int32 Strtoi(const UTF8CHAR* Start, UTF8CHAR** End, int32 Base)
+	{
+		return Strtoi((const ANSICHAR*)Start, (ANSICHAR**)End, Base);
+	}
+
+	static FORCEINLINE int64 Strtoi64(const UTF8CHAR* Start, UTF8CHAR** End, int32 Base)
+	{
+		return Strtoi64((const ANSICHAR*)Start, (ANSICHAR**)End, Base);
+	}
+
+	static FORCEINLINE uint64 Strtoui64(const UTF8CHAR* Start, UTF8CHAR** End, int32 Base)
+	{
+		return Strtoui64((const ANSICHAR*)Start, (ANSICHAR**)End, Base);
+	}
+
+	static FORCEINLINE UTF8CHAR* Strtok(UTF8CHAR* StrToken, const UTF8CHAR* Delim, UTF8CHAR** Context)
+	{
+		return (UTF8CHAR*)Strtok((ANSICHAR*)StrToken, (const ANSICHAR*)Delim, (ANSICHAR**)Context);
+	}
+
+	static FORCEINLINE int32 GetVarArgs(UTF8CHAR* Dest, SIZE_T DestSize, const UTF8CHAR*& Fmt, va_list ArgPtr)
+	{
+		return GetVarArgs((ANSICHAR*)Dest, DestSize, *(const ANSICHAR**)&Fmt, ArgPtr);
+	}
+};
+
+#pragma warning(pop) // 'function' was was declared deprecated  (needed for the secure string functions)
diff --git a/Engine/Source/Runtime/Core/Public/Windows/WindowsPlatformString.h b/Engine/Source/Runtime/Core/Public/Windows/WindowsPlatformString.h
--- a/Engine/Source/Runtime/Core/Public/Windows/WindowsPlatformString.h
+++ b/Engine/Source/Runtime/Core/Public/Windows/WindowsPlatformString.h
@@ -3,7 +3,7 @@
 #pragma once
 
 #include "CoreTypes.h"
-#include "GenericPlatform/MicrosoftPlatformString.h"
+#include "Microsoft/MicrosoftPlatformString.h"
 
 
 /**
