Fix format string for single param FErrors

diff --git a/Engine/Source/Runtime/Core/Public/Experimental/UnifiedError/UnifiedError.h b/Engine/Source/Runtime/Core/Public/Experimental/UnifiedError/UnifiedError.h
--- a/Engine/Source/Runtime/Core/Public/Experimental/UnifiedError/UnifiedError.h
+++ b/Engine/Source/Runtime/Core/Public/Experimental/UnifiedError/UnifiedError.h
@@ -872,7 +872,7 @@ CORE_API constexpr uint32 GetTypeIdHash(const TCHAR* TypeName)
 	// TODO: Daniel, $format below should be $formatns and $format
 
 #define DECLARE_ERROR_ONEPARAM(ErrorName, ErrorCode, ModuleName, FormatString, ParamOneType, ParamOneName, ParamOneDefault) \
-	DECLARE_ERROR_INTERNAL(ErrorName, ErrorCode, ModuleName, FormatString) \
+	DECLARE_ERROR_INTERNAL(ErrorName, ErrorCode, ModuleName, FText::FromString(UTF8TEXT("{"#ModuleName"_F"#ErrorName"}")) ) \
 	namespace UE::UnifiedError { namespace ModuleName { \
 			struct F##ErrorName \
 			{\
@@ -890,7 +890,7 @@ CORE_API constexpr uint32 GetTypeIdHash(const TCHAR* TypeName)
 	}\
 	namespace UE::UnifiedError { namespace ModuleName { \
 	inline void GatherPropertiesForError(const UE::UnifiedError::FError& Error, const UE::UnifiedError::ModuleName::F##ErrorName& Context, UE::UnifiedError::IErrorPropertyExtractor& PropertyExtractor) \
-			{ \
+	{ \
 		PropertyExtractor.AddProperty(ANSITEXTVIEW(#ParamOneName), Context.ParamOneName);\
 	}\
 	namespace ErrorName { \
