[Insights] Fixed CPU scope for FTraceAuxiliary::Relay.

#rnx

diff --git a/Engine/Source/Runtime/Core/Private/ProfilingDebugging/TraceAuxiliary.cpp b/Engine/Source/Runtime/Core/Private/ProfilingDebugging/TraceAuxiliary.cpp
--- a/Engine/Source/Runtime/Core/Private/ProfilingDebugging/TraceAuxiliary.cpp
+++ b/Engine/Source/Runtime/Core/Private/ProfilingDebugging/TraceAuxiliary.cpp
@@ -1627,7 +1627,6 @@ static bool TraceAuxiliaryStartShared(const TCHAR* Channels, const FTraceAuxilia
 		return false;
 	}
 
-
 	if (GTraceAuxiliary.IsConnected())
 	{
 		UE_LOG_REF(LogCategory, Error, TEXT("Unable to start trace, already tracing to %s"), *GTraceAuxiliary.GetDest());
@@ -1683,7 +1682,7 @@ bool FTraceAuxiliary::Start(EConnectionType Type, const TCHAR* Target, const TCH
 bool FTraceAuxiliary::Relay(UPTRINT Handle, UE::Trace::IoWriteFunc WriteFunc, UE::Trace::IoCloseFunc CloseFunc, const TCHAR* Channels, const FOptions* Options)
 {
 #if UE_TRACE_ENABLED
-	TRACE_CPUPROFILER_EVENT_SCOPE(FTraceAux_Start);
+	TRACE_CPUPROFILER_EVENT_SCOPE(FTraceAux_Relay);
 
 	if (!TraceAuxiliaryStartShared(Channels, Options, LogTrace))
 	{
@@ -1907,8 +1906,8 @@ void FTraceAuxiliary::Initialize(const TCHAR* CommandLine)
 	Desc.bUseImportantCache = (FParse::Param(CommandLine, TEXT("tracenocache")) == false);
 	Desc.OnConnectionFunc = &TraceAuxiliaryOnConnectionCallback;
 	Desc.OnUpdateFunc = &TraceAuxiliaryOnTraceUpdateCallback;
-	
-	// Determine if the worker thread is allowed 
+
+	// Determine if the worker thread is allowed
 	if (!FGenericPlatformProcess::SupportsMultithreading() || FParse::Param(FCommandLine::Get(), TEXT("notracethreading")))
 	{
 		GWorkerThreadConfig = EWorkerThreadConfig::Never;
@@ -1970,7 +1969,7 @@ void FTraceAuxiliary::Initialize(const TCHAR* CommandLine)
 	GTraceAuxiliary.RegisterEndFrameCallbacks();
 
 	// Initialize callstack tracing. If for any reason memory tracing has been initialized
-	// without callstack tracing we need to use the same allocator. Otherwise use regular malloc. 
+	// without callstack tracing we need to use the same allocator. Otherwise use regular malloc.
 	FMalloc* TraceAllocator = MemoryTrace_GetAllocator();
 	CallstackTrace_Create(TraceAllocator ? TraceAllocator : UE::Private::GMalloc);
 	CallstackTrace_Initialize();
