Change default trace file name to minimize collision risk.

diff --git a/Engine/Source/Runtime/Core/Private/ProfilingDebugging/TraceAuxiliary.cpp b/Engine/Source/Runtime/Core/Private/ProfilingDebugging/TraceAuxiliary.cpp
--- a/Engine/Source/Runtime/Core/Private/ProfilingDebugging/TraceAuxiliary.cpp
+++ b/Engine/Source/Runtime/Core/Private/ProfilingDebugging/TraceAuxiliary.cpp
@@ -684,7 +684,11 @@ bool FTraceAuxiliaryImpl::FinalizeFilePath(const TCHAR* InPath, FString& OutPath
 	const FStringView Path(InPath);
 
 	// Default file name functor
-	auto GetDefaultName = [] { return FDateTime::Now().ToString(TEXT("%Y%m%d_%H%M%S.utrace")); };
+	auto GetDefaultName = []
+		{
+			const FDateTime Now = FDateTime::Now();
+			return Now.ToString(TEXT("%Y%m%d_%H%M%S_")).Appendf(TEXT("%06" INT64_X_FMT ".utrace"), (Now.GetTicks() % 10000000LL));
+		};
 
 	if (Path.IsEmpty())
 	{
