On Linux, adds KSM process profit to:
Extended platform memory stats
CSV profiler
Perf counters

diff --git a/Engine/Source/Programs/TestPAL/Private/Main.cpp b/Engine/Source/Programs/TestPAL/Private/Main.cpp
--- a/Engine/Source/Programs/TestPAL/Private/Main.cpp
+++ b/Engine/Source/Programs/TestPAL/Private/Main.cpp
@@ -320,6 +320,12 @@ int32 SysInfoTest(const TCHAR* CommandLine)
 		ToMB(StatsEx.Shared_Clean), ToMB(StatsEx.Shared_Dirty));
 	UE_LOG(LogTestPAL, Display, TEXT("Private_Clean:%.2fMB Private_Dirty:%.2fMB"),
 		ToMB(StatsEx.Private_Clean), ToMB(StatsEx.Private_Dirty));
+
+	if (StatsEx.bIs_KSM_Enabled)
+	{
+		UE_LOG(LogTestPAL, Display, TEXT("KSM_Process_Profit:%.2fMB"),
+			ToMB(StatsEx.KSM_Process_Profit));
+	}
 #endif // PLATFORM_LINUX
 
 	FEngineLoop::AppPreExit();
diff --git a/Engine/Source/Runtime/Core/Private/HAL/MemoryMisc.cpp b/Engine/Source/Runtime/Core/Private/HAL/MemoryMisc.cpp
--- a/Engine/Source/Runtime/Core/Private/HAL/MemoryMisc.cpp
+++ b/Engine/Source/Runtime/Core/Private/HAL/MemoryMisc.cpp
@@ -33,12 +33,26 @@ namespace MemoryMiscInternal
 #if PLATFORM_UNIX
 	void LogSharedMemoryDiff(const TCHAR* Context, const FExtendedPlatformMemoryStats& StartStats, const FExtendedPlatformMemoryStats& EndStats)
 	{
-		UE_LOG(LogMemory, Log, TEXT("SharedMemoryTracker[%s] Unique: %.02fMB (%+.02fMB), Shared: %.02fMB (%+.02fMB)"),
-			Context,
-			BytesToMB(EndStats.Private_Clean + EndStats.Private_Dirty),
-			MemDeltaInMB(EndStats.Private_Clean + EndStats.Private_Dirty, StartStats.Private_Clean + StartStats.Private_Dirty),
-			BytesToMB(EndStats.Shared_Clean + EndStats.Shared_Dirty),
-			MemDeltaInMB(EndStats.Shared_Clean + EndStats.Shared_Dirty, StartStats.Shared_Clean + StartStats.Shared_Dirty));
+		if (EndStats.bIs_KSM_Enabled)
+		{
+			UE_LOG(LogMemory, Log, TEXT("SharedMemoryTracker[%s] Unique: %.02fMB (%+.02fMB), Shared: %.02fMB (%+.02fMB), KSM Profit: %.02fMB (%+.02fMB)"),
+				   Context,
+				   BytesToMB(EndStats.Private_Clean + EndStats.Private_Dirty),
+				   MemDeltaInMB(EndStats.Private_Clean + EndStats.Private_Dirty, StartStats.Private_Clean + StartStats.Private_Dirty),
+				   BytesToMB(EndStats.Shared_Clean + EndStats.Shared_Dirty),
+				   MemDeltaInMB(EndStats.Shared_Clean + EndStats.Shared_Dirty, StartStats.Shared_Clean + StartStats.Shared_Dirty),
+				   BytesToMB(EndStats.KSM_Process_Profit),
+				   MemDeltaInMB(EndStats.KSM_Process_Profit, StartStats.KSM_Process_Profit));
+		}
+		else
+		{
+			UE_LOG(LogMemory, Log, TEXT("SharedMemoryTracker[%s] Unique: %.02fMB (%+.02fMB), Shared: %.02fMB (%+.02fMB)"),
+				   Context,
+				   BytesToMB(EndStats.Private_Clean + EndStats.Private_Dirty),
+				   MemDeltaInMB(EndStats.Private_Clean + EndStats.Private_Dirty, StartStats.Private_Clean + StartStats.Private_Dirty),
+				   BytesToMB(EndStats.Shared_Clean + EndStats.Shared_Dirty),
+				   MemDeltaInMB(EndStats.Shared_Clean + EndStats.Shared_Dirty, StartStats.Shared_Clean + StartStats.Shared_Dirty));
+		}
 	}
 #endif //PLATFORM_UNIX
 }
diff --git a/Engine/Source/Runtime/Core/Private/Unix/UnixPlatformMemory.cpp b/Engine/Source/Runtime/Core/Private/Unix/UnixPlatformMemory.cpp
--- a/Engine/Source/Runtime/Core/Private/Unix/UnixPlatformMemory.cpp
+++ b/Engine/Source/Runtime/Core/Private/Unix/UnixPlatformMemory.cpp
@@ -1200,6 +1200,22 @@ FExtendedPlatformMemoryStats FUnixPlatformMemory::GetExtendedStats()
 		}
 	}
 
+	{
+		QUICK_SCOPE_CYCLE_COUNTER(STAT_GetExtendedStats_KsmStat);
+		
+		if (access("/sys/kernel/mm/ksm/run", F_OK) == 0)
+		{
+			MemoryStats.bIs_KSM_Enabled = true;
+
+			FProcField Field = { "ksm_process_profit", (uint64 *)&MemoryStats.KSM_Process_Profit };
+			if (ReadProcFields("/proc/self/ksm_stat", &Field, 1))
+			{
+				// ReadProcFields expects KiB so has already multiplied by 1024
+				MemoryStats.KSM_Process_Profit /= 1024;
+			}
+		}
+	}
+
 	return MemoryStats;
 }
 
diff --git a/Engine/Source/Runtime/Core/Public/Unix/UnixPlatformMemory.h b/Engine/Source/Runtime/Core/Public/Unix/UnixPlatformMemory.h
--- a/Engine/Source/Runtime/Core/Public/Unix/UnixPlatformMemory.h
+++ b/Engine/Source/Runtime/Core/Public/Unix/UnixPlatformMemory.h
@@ -40,6 +40,12 @@ struct FExtendedPlatformMemoryStats
 
 	/** Private memory used */
 	SIZE_T Private_Dirty;
+
+	/** How much private memory is being saved by KSM, minus memory used by KSM to do so */
+	SIZE_T KSM_Process_Profit;
+
+	/** */
+	bool bIs_KSM_Enabled;
 };
 
 /**
