Replaced TSortedSet's O(N^2) Append with an O(N) (at best) or O(N log N) (at worst).
Added set appending tests.

diff --git a/Engine/Source/Runtime/Core/Private/Containers/ContainersTest.cpp b/Engine/Source/Runtime/Core/Private/Containers/ContainersTest.cpp
--- a/Engine/Source/Runtime/Core/Private/Containers/ContainersTest.cpp
+++ b/Engine/Source/Runtime/Core/Private/Containers/ContainersTest.cpp
@@ -14,7 +14,7 @@
 #include "Stats/StatsMisc.h"
 #include "Templates/UniquePtr.h"
 
-#if 1//WITH_DEV_AUTOMATION_TESTS
+#if WITH_DEV_AUTOMATION_TESTS
 
 #define MAX_TEST_OBJECTS      65
 #define MAX_TEST_OBJECTS_STEP 1
@@ -408,7 +408,7 @@ namespace
 
 		ContainerTestStats.Reset();
 
-		auto CheckSet = [&Set]()
+		auto CheckSet = [](SetType& Set)
 		{
 			CheckContainerNum(Set);
 			CheckContainerEnds(Set);
@@ -427,7 +427,7 @@ namespace
 			for (int32 N = 0; N != Count; ++N)
 			{
 				Set.Add(GenerateTestKey<KeyType>(N));
-				CheckSet();
+				CheckSet(Set);
 			}
 			MaxNum = Set.Num();
 			MaxAllocatedSize = Set.GetAllocatedSize();
@@ -437,7 +437,7 @@ namespace
 				auto* ElementPtr = Set.FindArbitraryElement();
 				check(ElementPtr);
 				check(Set.Remove(*ElementPtr) == 1);
-				CheckSet();
+				CheckSet(Set);
 			}
 
 			check(Set.IsEmpty());
@@ -445,7 +445,7 @@ namespace
 			for (int32 N = 0; N != Count; ++N)
 			{
 				Set.Add(GenerateTestKey<KeyType>((Count - 1) - N));
-				CheckSet();
+				CheckSet(Set);
 			}
 
 			for (int32 N = 0; N != Count; ++N)
@@ -453,7 +453,7 @@ namespace
 				auto* ElementPtr = Set.FindArbitraryElement();
 				check(ElementPtr);
 				check(Set.Remove(*ElementPtr) == 1);
-				CheckSet();
+				CheckSet(Set);
 			}
 
 			check(Set.IsEmpty());
@@ -463,7 +463,7 @@ namespace
 		{
 			// Test releasing memory allocations
 			Set.Empty();
-			CheckSet();
+			CheckSet(Set);
 			check(Set.GetAllocatedSize() == InitialAllocatedSize);
 
 			// Test integrity after re-growing container to MaxNum elements again
@@ -471,7 +471,7 @@ namespace
 			{
 				Set.Add(GenerateTestKey<KeyType>(N));
 			}
-			CheckSet();
+			CheckSet(Set);
 			check(Set.GetAllocatedSize() == MaxAllocatedSize);
 
 			// Test data integrity while removing and shrinking continously
@@ -483,7 +483,7 @@ namespace
 					check(ElementPtr);
 					Set.Remove(*ElementPtr);
 					Set.Shrink();
-					CheckSet();
+					CheckSet(Set);
 					check(Set.GetAllocatedSize() <= PrevAllocatedSize);
 					PrevAllocatedSize = Set.GetAllocatedSize();
 				}
@@ -503,12 +503,12 @@ namespace
 			for (int32 Count = 0; Count < MAX_TEST_OBJECTS - 1; Count += MAX_TEST_OBJECTS_STEP)
 			{
 				Set.Empty();
-				CheckSet();
+				CheckSet(Set);
 
 				for (int32 N = 0; N != Count; ++N)
 				{
 					Set.Add(GenerateTestKey<KeyType>(N));
-					CheckSet();
+					CheckSet(Set);
 				}
 
 				// Check that the const iterator finds all the values
@@ -529,6 +529,122 @@ namespace
 				check(Found.Find(false) == INDEX_NONE);
 			}
 		}
+
+		// Test bulk insertion
+		{
+			auto TestAppend = [&CheckSet](std::initializer_list<KeyType> InitSet, std::initializer_list<KeyType> Append, std::initializer_list<KeyType> ExpectedResult)
+			{
+				auto CheckSetEquals = [&CheckSet](SetType& Set, std::initializer_list<KeyType> ExpectedResult)
+				{
+					// Check set integrity
+					CheckSet(Set);
+
+					// Check 
+					if constexpr (bSorted)
+					{
+						TArray<KeyType> ExpectedResultArray = ExpectedResult;
+						ExpectedResultArray.Sort();
+						check(Set.Array() == ExpectedResultArray);
+					}
+					else
+					{
+						// Check every expected result value is there and no more
+						for (const KeyType& Key : ExpectedResult)
+						{
+							check(Set.Remove(Key) == 1);
+						}
+						check(Set.IsEmpty());
+					}
+				};
+
+				SetType DestSet;
+				TArray<KeyType> SrcArray = Append;
+
+				// Check append by copy
+				{
+					DestSet = InitSet;
+
+					// Copy elements to set
+					DestSet.Append(SrcArray);
+
+					// Check that the source is unchanged
+					check(SrcArray == TArray<int>(Append));
+
+					// Check that the set contains the expected values (this may mutate the set)
+					CheckSetEquals(DestSet, ExpectedResult);
+				}
+
+				// Check append by move
+				{
+					DestSet = InitSet;
+
+					// Copy elements to set
+					DestSet.Append(SrcArray);
+
+					// Check that the source is unchanged
+					check(SrcArray == TArray<int>(Append));
+
+					// Check that the set contains the expected values (this may mutate the set)
+					CheckSetEquals(DestSet, ExpectedResult);
+				}
+			};
+
+			// Copy unique ordered elements to empty set
+			TestAppend(
+				{},
+				{ 1, 2, 3, 4, 5 },
+				{ 1, 2, 3, 4, 5 }
+			);
+
+			// Copy unique unordered elements to empty set
+			TestAppend(
+				{},
+				{ 3, 5, 1, 2, 4 },
+				{ 1, 2, 3, 4, 5 }
+			);
+
+			// Copy duplicate ordered elements to empty set
+			TestAppend(
+				{},
+				{ 1, 1, 2, 3, 3, 3, 4, 5, 5 },
+				{ 1, 2, 3, 4, 5 }
+			);
+
+			// Copy duplicate unordered elements to empty set
+			TestAppend(
+				{},
+				{ 3, 5, 1, 1, 4, 2, 5, 4, 3, 1 },
+				{ 1, 2, 3, 4, 5 }
+			);
+
+			// Copy unique ordered elements to populated set
+			TestAppend(
+				{ 1, 5, 7 },
+				{ 1, 2, 3, 4, 5 },
+				{ 1, 2, 3, 4, 5, 7 }
+			);
+
+			// Copy unique unordered elements to populated set
+			TestAppend(
+				{ 1, 5, 7 },
+				{ 3, 5, 1, 2, 4 },
+				{ 1, 2, 3, 4, 5, 7 }
+			);
+
+			// Copy duplicate ordered elements to populated set
+			TestAppend(
+				{ 1, 5, 7 },
+				{ 1, 1, 2, 3, 3, 3, 4, 5, 5 },
+				{ 1, 2, 3, 4, 5, 7 }
+			);
+
+			// Copy duplicate unordered elements to populated set
+			TestAppend(
+				{ 1, 5, 7 },
+				{ 3, 5, 1, 1, 4, 2, 5, 4, 3, 1 },
+				{ 1, 2, 3, 4, 5, 7 }
+			);
+		}
 	}
 
 	// Test container element address consistency when using SortFreeList 
diff --git a/Engine/Source/Runtime/Core/Public/Containers/SortedSet.h b/Engine/Source/Runtime/Core/Public/Containers/SortedSet.h
--- a/Engine/Source/Runtime/Core/Public/Containers/SortedSet.h
+++ b/Engine/Source/Runtime/Core/Public/Containers/SortedSet.h
@@ -3,7 +3,9 @@
 #pragma once
 
 #include "Algo/BinarySearch.h"
+#include "Algo/IndexOf.h"
 #include "Algo/Sort.h"
+#include "Containers/ArrayView.h"
 #include "Containers/SetUtilities.h"
 #include "UObject/NameTypes.h"
 
@@ -14,9 +16,9 @@ FArchive& operator<<(FArchive& Ar, TSortedSet<ElementType, ArrayAllocator, SortP
  * A set of values, implemented as a sorted TArray of elements.
  *
  * It has a mostly identical interface to TSet and is designed as a drop in replacement. Keys must be unique.
- * It uses half as much memory as TSet, but adding and removing elements is O(n), and finding is O(Log n). In practice
- * it is faster than TSet for low element counts, and slower as n increases, This set is always kept sorted by key
- * so cannot be sorted manually.
+ * It uses half as much memory as TSet and has a smaller static footprint, but adding and removing elements is O(n),
+ * and finding is O(Log n). In practice it is faster than TSet for low element counts, and slower as n increases.
+ * This set is always kept sorted by key so cannot be sorted manually.
  */
 template <typename InElementType, typename ArrayAllocator /*= FDefaultAllocator*/, typename SortPredicate /*= TLess<InElementType>*/>
 class TSortedSet
@@ -54,10 +56,7 @@ public:
 	[[nodiscard]] TSortedSet(std::initializer_list<ElementType> InitList)
 	{
 		this->Reserve((int32)InitList.size());
-		for (const ElementType& Element : InitList)
-		{
-			this->Add(Element);
-		}
+		this->Append(InitList);
 	}
 
 	///////////////////////////////////////////////////
@@ -98,10 +97,7 @@ public:
 	TSortedSet& operator=(std::initializer_list<ElementType> InitList)
 	{
 		this->Empty((int32)InitList.size());
-		for (const ElementType& Element : InitList)
-		{
-			this->Add(Element);
-		}
+		this->Append(InitList);
 		return *this;
 	}
 
@@ -311,11 +307,11 @@ public:
 	 */
 	TArray<ElementType> Array() const&
 	{
-		return Elements;
+		return TArray<ElementType>(Elements);
 	}
 	TArray<ElementType> Array() &&
 	{
-		return MoveTemp(Elements);
+		return TArray<ElementType>(MoveTemp(Elements));
 	}
 
 	/**
@@ -328,6 +324,117 @@ public:
 		Ar.Logf(TEXT("TSortedSet: %i elements"), Elements.Num());
 	}
 
+private:
+	template <bool bMove>
+	void AppendImpl(std::conditional_t<bMove, ElementType, const ElementType>* Ptr, int32 Count)
+	{
+		// Insert new unsorted elements at the end of sorted elements, then sort them all together into the right place at the end.
+		//
+		// Needs to handle duplicates in both the existing set and within new elements.
+
+		if (Count == 0)
+		{
+			return;
+		}
+
+		// Reserve space for the new elements - we may not need this many if there are duplicates
+		Elements.Reserve(Elements.Num() + Count);
+
+		// Get a view over the sorted elements - even if Elements grows in the loop below, this won't be invalidated due to the Reserve above
+		TArrayView<ElementType> SortedView = Elements;
+
+		for (;;)
+		{
+			// Find existing element in sorted elements
+			int32 InsertIndex = 0;
+			if (!SortedView.IsEmpty())
+			{
+				// If the new element is strictly 'greater' than the last of the sorted elements, then insert at the end, otherwise binary search for the insertion point
+				if (SortPredicate()(SortedView.Last(), *Ptr))
+				{
+					InsertIndex = SortedView.Num();
+				}
+				else
+				{
+					InsertIndex = Algo::LowerBound(SortedView, *Ptr, SortPredicate());
+				}
+			}
+
+			// If the element at the insertion point is equivalent to ('not less than') the new element, then we can insert it in-place without breaking the sort order
+			if (InsertIndex < SortedView.Num() && !SortPredicate()(*Ptr, Elements[InsertIndex]))
+			{
+				// Replace duplicate element by constructing it over the previous one - this will still leave the sorted elements in order
+				ElementType* OldElement = Elements.GetData() + InsertIndex;
+				DestructItem(OldElement);
+				if constexpr (bMove)
+				{
+					::new ((void*)OldElement) ElementType(MoveTemp(*Ptr));
+				}
+				else
+				{
+					::new ((void*)OldElement) ElementType(*Ptr);
+				}
+			}
+			else
+			{
+				// Find existing element in unsorted elements
+				ElementType* ExistingElement = nullptr;
+				int32 NumUnsortedElements = Elements.Num() - SortedView.Num();
+				if (NumUnsortedElements > 0)
+				{
+					ExistingElement = MakeArrayView(Elements.GetData() + SortedView.Num(), NumUnsortedElements).FindByPredicate([Ptr](const ElementType& Val) { return !SortPredicate()(*Ptr, Val) && !SortPredicate()(Val, *Ptr); });
+				}
+
+				if (ExistingElement)
+				{
+					// Replace duplicate element by constructing it over the previous one.
+					DestructItem(ExistingElement);
+					if constexpr (bMove)
+					{
+						::new ((void*)ExistingElement) ElementType(MoveTemp(*Ptr));
+					}
+					else
+					{
+						::new ((void*)ExistingElement) ElementType(*Ptr);
+					}
+				}
+				else
+				{
+					// Add the element to the end of the container
+					if constexpr (bMove)
+					{
+						Elements.Add(MoveTemp(*Ptr));
+					}
+					else
+					{
+						Elements.Add(*Ptr);
+					}
+
+					// If the new element is precisely 'greater than' all existing sorted elements and we haven't added any unsorted elements yet, then just expand
+					// the sorted view to include the new element.
+					if (InsertIndex == SortedView.Num() && Elements.Num() == SortedView.Num() + 1)
+					{
+						SortedView = Elements;
+					}
+				}
+			}
+
+			--Count;
+			if (Count == 0)
+			{
+				// If we still have unsorted elements, sort them into the rest of the container
+				if (SortedView.Num() != Elements.Num())
+				{
+					Algo::Sort(Elements, SortPredicate());
+				}
+				return;
+			}
+
+			++Ptr;
+		}
+	}
+
+public:
 	/**
 	 * Adds all items from another set into our set (if any elements are in both, the value from the other set wins).
 	 * If the source is an rvalue container, the elements will be moved and the container will be left empty.
@@ -337,49 +444,28 @@ public:
 	template <typename OtherArrayAllocator, typename OtherSortPredicate>
 	void Append(TSortedSet<ElementType, OtherArrayAllocator, OtherSortPredicate>&& Other)
 	{
-		this->Reserve(this->Num() + Other.Num());
-		for (ElementType& Element : Other)
-		{
-			this->Add(MoveTemp(Element));
-		}
-
+		this->AppendImpl<true>(Other.Elements.GetData(), Other.Elements.Num());
 		Other.Reset();
 	}
 	template <typename OtherArrayAllocator, typename OtherSortPredicate>
 	void Append(const TSortedSet<ElementType, OtherArrayAllocator, OtherSortPredicate>& Other)
 	{
-		this->Reserve(this->Num() + Other.Num());
-		for (const ElementType& Element : Other)
-		{
-			this->Add(Element);
-		}
+		this->AppendImpl<false>(Other.Elements.GetData(), Other.Elements.Num());
 	}
 	template <typename OtherArrayAllocator>
 	void Append(TArray<ElementType, OtherArrayAllocator>&& Other)
 	{
-		this->Reserve(this->Num() + Other.Num());
-		for (ElementType& Element : Other)
-		{
-			this->Add(MoveTemp(Element));
-		}
+		this->AppendImpl<true>(Other.GetData(), Other.Num());
 		Other.Reset();
 	}
 	void Append(TArrayView<const ElementType> Other)
 	{
-		this->Reserve(this->Num() + Other.Num());
-		for (const ElementType& Element : Other)
-		{
-			this->Add(Element);
-		}
+		this->AppendImpl<false>(Other.GetData(), Other.Num());
 	}
 	template <typename OtherArrayAllocator>
 	void Append(std::initializer_list<ElementType> Other)
 	{
-		this->Reserve(this->Num() + (int32)Other.size());
-		for (const ElementType& Element : Other)
-		{
-			this->Add(Element);
-		}
+		this->AppendImpl<false>(Other.GetData(), (int32)Other.Num());
 	}
 
 	/**
