[RaceDetector]
- Fix assert when clocks are wrapping
- Make the toggling on and off ref-counted so that many threads can decide to activate the race detector, last thread to turn it off will deactivate it
- Add a grace period to reduce toggling speed for short lived scopes
- Add a scope that can be used for activation

diff --git a/Engine/Source/Runtime/Core/Internal/Sanitizer/RaceDetectorTypes.h b/Engine/Source/Runtime/Core/Internal/Sanitizer/RaceDetectorTypes.h
--- a/Engine/Source/Runtime/Core/Internal/Sanitizer/RaceDetectorTypes.h
+++ b/Engine/Source/Runtime/Core/Internal/Sanitizer/RaceDetectorTypes.h
@@ -565,23 +565,15 @@ namespace UE::Sanitizer::RaceDetector {
 			// Make sure we sync with the light fence.
 			Platform::AsymmetricThreadFenceHeavy();
 
-			OutClockRange.First = MAX_uint32;
-			OutClockRange.Last  = 0;
+			OutClockRange.First = Head ? Head->StartClock : 0;
 
 			for (FHistoryChunk* Chunk = Head; Chunk; Chunk = Chunk->Next)
 			{
-				if (Chunk->StartClock < OutClockRange.First)
-				{
-					OutClockRange.First = Chunk->StartClock;
-				}
-
-				if (Chunk->EndClock > OutClockRange.Last)
-				{
-					OutClockRange.Last = Chunk->EndClock;
-				}
+				OutClockRange.Last = Chunk->EndClock;
 
 				// Do not bother searching a chunk that is outside the range we're looking for.
-				if (InAccess.Clock >= Chunk->StartClock && InAccess.Clock <= Chunk->EndClock)
+				// Do the range comparison in a way that handles clock wrapping
+				if ((InAccess.Clock - Chunk->StartClock) <= (Chunk->EndClock - Chunk->StartClock))
 				{
 					TArray<void*, TInlineAllocator<1024>> Stack;
 
@@ -623,16 +615,6 @@ namespace UE::Sanitizer::RaceDetector {
 				}
 			}
 
-			// We should never end up empty handed unless the clock is from
-			// a long time ago and outside the history range.
-			// We have no choice but to include the equality comparison here
-			// since we have no idea if the chunk that got recycled had some
-			// entry for this clock.
-			check(InAccess.Clock <= OutClockRange.First);
-			if (OutClockRange.First == MAX_uint32)
-			{
-				OutClockRange.First = OutClockRange.Last;
-			}
 			return false;
 		}
 
@@ -694,6 +676,8 @@ namespace UE::Sanitizer::RaceDetector {
 		void* ThreadArgs = nullptr;
 		// Avoid using UniquePtr because it's instrumented and each access has a cost.
 		FAccessHistory* AccessHistory = nullptr;
+		// Wether or not to always report race for this thread.
+		bool bAlwaysReport = false;
 
 		// This can be bumped again if we ever face a need for deeper callstacks since
 		// this is a virtual allocation anyway so it's not going to take physical memory
@@ -721,11 +705,11 @@ namespace UE::Sanitizer::RaceDetector {
 		
 		INSTRUMENTATION_FUNCTION_ATTRIBUTES void IncrementClock()
 		{
-			CurrentClock()++; 
+			CurrentClock()++;
 			
 			if (UNLIKELY(DetailedLogDepth || bDetailedLogGlobal))
 			{
-				FPlatformMisc::LowLevelOutputDebugStringf(TEXT("[%d] Thread is now at clock %d\n"), ThreadId, CurrentClock());
+				FPlatformMisc::LowLevelOutputDebugStringf(TEXT("[%d] Thread is now at clock %u\n"), ThreadId, CurrentClock());
 			}
 		}
 
diff --git a/Engine/Source/Runtime/Core/Private/Sanitizer/RaceDetector.cpp b/Engine/Source/Runtime/Core/Private/Sanitizer/RaceDetector.cpp
--- a/Engine/Source/Runtime/Core/Private/Sanitizer/RaceDetector.cpp
+++ b/Engine/Source/Runtime/Core/Private/Sanitizer/RaceDetector.cpp
@@ -17,6 +17,7 @@
 #include "Misc/ScopeLock.h"
 #include "Misc/Paths.h"
 #include "Async/Mutex.h"
+#include "Async/RecursiveMutex.h"
 #include "Misc/Parse.h"
 #include "Math/NumericLimits.h"
 #include "Misc/CommandLine.h"
@@ -47,7 +48,8 @@ UE_TRACE_CHANNEL_DEFINE(RaceDetectorChannel);
 UE_DISABLE_OPTIMIZATION
 #endif
 
-static uint32 GRaceDetectorThreadFilter = 0;
+static bool GRaceDetectorFilterOtherThreads = false;
+static bool GRaceDetectorOnlyActiveScopes = false;
 static bool GRaceDetectorFilterDuplicates = true;
 static FAutoConsoleVariableRef CVarReportSameRaceOnce(
 	TEXT("r.RaceDetector.FilterDuplicates"),
@@ -85,7 +87,7 @@ static FAutoConsoleVariableRef CVarIgnoreRaceWithoutSecondStack(
 	TEXT("This generally happens when races are far enough that maybe it is not a concern after all.\n"),
 	ECVF_Default);
 
-static int32 GRaceDetectorMaxMemoryUsage = 16;
+static int32 GRaceDetectorMaxMemoryUsage = 64;
 static FAutoConsoleVariableRef CVarRaceDetectorMaxMemoryUsage(
 	TEXT("r.RaceDetector.MaxMemoryUsage"),
 	GRaceDetectorMaxMemoryUsage,
@@ -134,10 +136,14 @@ namespace UE::Sanitizer::RaceDetector {
 	volatile FContextId CurrentContextId = 1;
 	volatile int64 HistoryChunkCount = 0;
 
+	static UE::FRecursiveMutex Mutex;
+
 	// Since the shadow space can be unmapped in a single shot loosing all our pointers,
 	// we need to keep a list of all the banks we've allocated so we can free them.
 	volatile FSyncObjectBank* SyncObjectBankHead = nullptr;
 
+	volatile FSyncObjectBank* FreeObjectBankHead = nullptr;
+
 	FRWSpinLock RaceReportsLock;
 	TArray<FString> RaceReports;
 
@@ -150,7 +156,9 @@ namespace UE::Sanitizer::RaceDetector {
 	TSafeSet<uint64> RaceHashes;
 	FRWSpinLock RaceHashesLock;
 	UE::FTimeout RaceDetectorTimeout = UE::FTimeout::Never();
-	
+	UE::FTimeout GracePeriodTimeout = UE::FTimeout::Never();
+	volatile uint32 GracePeriodEpoch = 0;
+
 	volatile bool bIsResettingShadow = false;
 	volatile bool bIsDebuggerPresent = false;
 	volatile bool bDetailedLogGlobal = false;
@@ -160,10 +168,11 @@ namespace UE::Sanitizer::RaceDetector {
 	// are unoptimized and we end up with a function call to an instrumented function
 	// going directly back to the ShouldInstrument call. std::atomics that are tested
 	// after InstrumentationDepth should work but are not optimal.
-	volatile bool bRaceDetectorActive = false;
+	volatile uint64 RaceDetectorActive = 0;
+
 	INSTRUMENTATION_FUNCTION_ATTRIBUTES bool ShouldInstrument(FContext& Context)
 	{
-		return bRaceDetectorActive && !Context.WinInstrumentationDepth && !Context.InstrumentationDepth && !bIsResettingShadow;
+		return !!RaceDetectorActive && !Context.WinInstrumentationDepth && !Context.InstrumentationDepth && !bIsResettingShadow;
 	}
 
 	FORCEINLINE INSTRUMENTATION_FUNCTION_ATTRIBUTES void* GetTlsValueFast(uint32 Index)
@@ -382,6 +391,39 @@ namespace UE::Sanitizer::RaceDetector {
 		}
 	}
 
+	INSTRUMENTATION_FUNCTION_ATTRIBUTES void PushFreeObjectBank(FSyncObjectBank* SyncObjectBank)
+	{
+		while (true)
+		{
+			FSyncObjectBank* LocalSyncBankHead = (FSyncObjectBank*)FreeObjectBankHead;
+			SyncObjectBank->Next = LocalSyncBankHead;
+
+			if (FPlatformAtomics::InterlockedCompareExchangePointer((void**)&FreeObjectBankHead, SyncObjectBank, LocalSyncBankHead) == LocalSyncBankHead)
+			{
+				FPlatformAtomics::InterlockedDecrement(&FSyncObjectBank::ObjectCount);
+				break;
+			}
+		}
+	}
+
+	INSTRUMENTATION_FUNCTION_ATTRIBUTES FSyncObjectBank* PopFreeObjectBank()
+	{
+		while (true)
+		{
+			FSyncObjectBank* LocalSyncBankHead = (FSyncObjectBank*)FreeObjectBankHead;
+			if (LocalSyncBankHead == nullptr)
+			{
+				return nullptr;
+			}
+			if (FPlatformAtomics::InterlockedCompareExchangePointer((void**)&FreeObjectBankHead, LocalSyncBankHead->Next, LocalSyncBankHead) == LocalSyncBankHead)
+			{
+				FPlatformAtomics::InterlockedIncrement(&FSyncObjectBank::ObjectCount);
+				LocalSyncBankHead->Next = nullptr;
+				return LocalSyncBankHead;
+			}
+		}
+	}
+
 	INSTRUMENTATION_FUNCTION_ATTRIBUTES FSyncObjectRef GetSyncObject(FContext& Context, void* SyncAddr)
 	{
 		FShadowClockBankSlot* ClockBankSlot = GetShadowClockBank((UPTRINT)SyncAddr);
@@ -395,7 +437,11 @@ namespace UE::Sanitizer::RaceDetector {
 			FSyncObjectBank* LocalBank = ClockBankSlot->SyncObjectBank;
 			if (LocalBank == nullptr)
 			{
-				FSyncObjectBank* NewObjectBank = new FSyncObjectBank();
+				FSyncObjectBank* NewObjectBank = PopFreeObjectBank();
+				if (NewObjectBank == nullptr)
+				{
+					NewObjectBank = new FSyncObjectBank();
+				}
 
 				// When publishing in shadow memory, we need to make sure the refcount is not 0, since another thread could try to use
 				// the newly published bank from the shadow, perform an AddRef and then release it and we'd end up with
@@ -441,7 +487,7 @@ namespace UE::Sanitizer::RaceDetector {
 
 	INSTRUMENTATION_FUNCTION_ATTRIBUTES void FreeMemoryRange(void* Ptr, uint64 Size)
 	{
-		if (!bRaceDetectorActive || bIsResettingShadow)
+		if (!RaceDetectorActive || bIsResettingShadow)
 		{
 			return;
 		}
@@ -562,13 +608,11 @@ namespace UE::Sanitizer::RaceDetector {
 				}
 			}
 
-			if (GRaceDetectorThreadFilter)
+			if (GRaceDetectorFilterOtherThreads &&
+				Context.bAlwaysReport != true &&
+				RacingContext->bAlwaysReport != true)
 			{
-				if (Context.ThreadId != GRaceDetectorThreadFilter &&
-					RacingContext->ThreadId != GRaceDetectorThreadFilter)
-				{
-					return;
-				}
+				return;
 			}
 
 			void* LowLimit = 0; void *HighLimit = 0;
@@ -845,7 +889,7 @@ namespace UE::Sanitizer::RaceDetector {
 			// We can find false positives when the race detector is shutting down or during shadow resets
 			// because some function will stop instrumenting before others, etc...
 			// Just ignore any race we have found if anything seems incoherent.
-			if (bWonReportingRace && bRaceDetectorActive && Epoch == GlobalEpoch && !bIsResettingShadow)
+			if (bWonReportingRace && RaceDetectorActive && Epoch == GlobalEpoch && !bIsResettingShadow)
 			{
 				ReportRace(Context, Ptr, CurrentAccess, PreviousAccess);
 			}
@@ -1026,11 +1070,14 @@ namespace UE::Sanitizer::RaceDetector {
 			Lock.Unlock();
 		};
 
+		// Need to take the mutex during reset shadow since we're unhooking and rehooking instrumentation, which is not thread-safe.
+		UE::TUniqueLock ScopeLock(Mutex);
+
 		FPlatformMisc::LowLevelOutputDebugStringf(TEXT("[%d] Reset Shadow Memory Started\n"), Platform::GetCurrentThreadId());
 
 		FPlatformAtomics::InterlockedExchange((volatile int8*)&bIsResettingShadow, (int8)1);
 
-		if (bRaceDetectorActive)
+		if (RaceDetectorActive)
 		{
 			UnhookInstrumentation();
 		}
@@ -1060,7 +1107,7 @@ namespace UE::Sanitizer::RaceDetector {
 				}
 				else
 				{
-					ToDelete->Release();
+					PushFreeObjectBank(ToDelete);
 				}
 			}
 
@@ -1068,7 +1115,7 @@ namespace UE::Sanitizer::RaceDetector {
 
 		FPlatformMisc::LowLevelOutputDebugStringf(TEXT("[%d] Reset Shadow Memory Ended\n"), Platform::GetCurrentThreadId());
 		
-		if (bRaceDetectorActive)
+		if (RaceDetectorActive)
 		{
 			GlobalEpoch++;
 
@@ -1106,6 +1153,14 @@ namespace UE::Sanitizer::RaceDetector {
 		}
 	}
 
+	INSTRUMENTATION_FUNCTION_ATTRIBUTES void CheckGracePeriod()
+	{
+		if (GlobalEpoch == GracePeriodEpoch && GracePeriodTimeout.IsExpired())
+		{
+			ToggleRaceDetection(false);
+		}
+	}
+
 	INSTRUMENTATION_FUNCTION_ATTRIBUTES void SanitizerThreadRun(volatile bool& bContinue)
 	{
 		FInstrumentationScope Scope;
@@ -1125,7 +1180,9 @@ namespace UE::Sanitizer::RaceDetector {
 
 			HandleReports();
 
-			if (bRaceDetectorActive)
+			CheckGracePeriod();
+
+			if (RaceDetectorActive)
 			{
 				if (RaceDetectorTimeout.IsExpired())
 				{
@@ -1193,8 +1250,7 @@ namespace UE::Sanitizer::RaceDetector {
 		INSTRUMENTATION_FUNCTION_ATTRIBUTES virtual void* Realloc(void* InPtr, SIZE_T NewSize, uint32 Alignment) override
 		{
 			FInstrumentationScope Scope;
-
-			if (!bRaceDetectorActive)
+			if (!RaceDetectorActive)
 			{
 				return InnerMalloc->Realloc(InPtr, NewSize, Alignment);
 			}
@@ -1465,7 +1521,9 @@ namespace UE::Sanitizer::RaceDetector {
 			return true;
 		}
 		
-		ToggleRaceDetection(false);
+		UnhookInstrumentation();
+
+		ResetShadow();
 
 		UnhookAlwaysOnInstrumentation();
 
@@ -1500,44 +1558,80 @@ namespace UE::Sanitizer::RaceDetector {
 	
 	INSTRUMENTATION_FUNCTION_ATTRIBUTES void ToggleFilterOtherThreads(bool bEnable)
 	{
-		GRaceDetectorThreadFilter = Platform::GetCurrentThreadId();
+		GRaceDetectorFilterOtherThreads = bEnable;
 	}
 
 	INSTRUMENTATION_FUNCTION_ATTRIBUTES bool IsActive()
 	{
-		return bRaceDetectorActive;
+		return RaceDetectorActive != 0;
 	}
 
-	INSTRUMENTATION_FUNCTION_ATTRIBUTES void ToggleRaceDetection(bool bEnable)
+	INSTRUMENTATION_FUNCTION_ATTRIBUTES void ToggleRaceDetection(bool bEnable, float GracePeriodSeconds)
 	{
-		if (bRaceDetectorActive == bEnable)
-		{
-			return;
-		}
+		int32 Result = 0;
 
 		// Just make sure we're not initializing the stack from an access that comes from this function
 		{
-			FInstrumentationScope InstrumentationScope;
+			// Need to take the mutex during reset shadow since we're unhooking and rehooking instrumentation, which is not thread-safe.
+			UE::TUniqueLock ScopeLock(Mutex);
+
+			FContext& Context = EnsureCurrentContext();
 
-			bRaceDetectorActive = bEnable;
-			if (!bEnable)
+			FInstrumentationScope InstrumentationScope;
+			if (bEnable)
 			{
-				UnhookInstrumentation();
+				Context.bAlwaysReport = true;
 
-				ResetShadow();
+				RaceDetectorTimeout = UE::FTimeout::Never();
+				if (RaceDetectorActive++ == 0)
+				{
+					GracePeriodTimeout = UE::FTimeout::Never();
+
+					GlobalEpoch++;
+
+					Result = 1;
+					HookInstrumentation();
+				}
 			}
 			else
 			{
-				GlobalEpoch++;
+				Context.bAlwaysReport = false;
 
-				HookInstrumentation();
+				if (GracePeriodSeconds > 0.0f)
+				{
+					// There is already a grace period active, decrement the count to avoid accumulating
+					if (!GracePeriodTimeout.WillNeverExpire())
+					{
+						if (RaceDetectorActive > 1)
+						{
+							--RaceDetectorActive;
+						}
+					}
+					GracePeriodTimeout = UE::FTimeout(GracePeriodSeconds);
+					GracePeriodEpoch = GlobalEpoch;
+				}
+				else
+				{
+					if (RaceDetectorActive > 0 && --RaceDetectorActive == 0)
+					{
+						RaceDetectorTimeout = UE::FTimeout::Never();
+						GracePeriodTimeout = UE::FTimeout::Never();
+
+						Result = 2;
+						UnhookInstrumentation();
+						ResetShadow();
+					}
+				}
 			}
 		}
 
 		// This needs to be properly instrumented because it does synchronize with another thread so if the
 		// instrumentation doesn't see the atomics being used from this thread, a future race will be
 		// detected when the stack memory used as synchronization during logging starts being reused.
-		UE_LOG(LogRaceDetector, Log, TEXT("Race detector has been toggled %s"), bEnable ? TEXT("on") : TEXT("off"));
+		if (Result)
+		{
+			UE_LOG(LogRaceDetector, Log, TEXT("Race detector has been toggled %s"),	Result == 1 ? TEXT("on") : TEXT("off"))
+		}
 	}
 
 	INSTRUMENTATION_FUNCTION_ATTRIBUTES void ToggleThreadDetailedLog(bool bEnabled)
diff --git a/Engine/Source/Runtime/Core/Public/Sanitizer/RaceDetector.h b/Engine/Source/Runtime/Core/Public/Sanitizer/RaceDetector.h
--- a/Engine/Source/Runtime/Core/Public/Sanitizer/RaceDetector.h
+++ b/Engine/Source/Runtime/Core/Public/Sanitizer/RaceDetector.h
@@ -33,7 +33,9 @@ namespace UE::Sanitizer::RaceDetector
 	CORE_API void ToggleFilterOtherThreads(bool bEnable);
 
 	// Toggle race detection until a call to the same function that toggles it back off.
-	CORE_API void ToggleRaceDetection(bool bEnable);
+	// If multiple threads activate race detection at the same time, it will be turned off when all thread toggle it back to off.
+	// The grace period is how much time to wait until race detection is really turned off. Can be used when toggling is expected to happen too often or too fast.
+	CORE_API void ToggleRaceDetection(bool bEnable, float GracePeriodSeconds = 0.0f);
 	// Toggle filter on duplicate races. Does not apply when a callback is set.
 	CORE_API void ToggleFilterDuplicateRaces(bool bEnable);
 	// Activate race detection for certain time after which it will be automatically turned off.
@@ -49,6 +51,31 @@ namespace UE::Sanitizer::RaceDetector
 	// Returns whether the race detector is currently active or not.
 	CORE_API bool IsActive();
 	
+	class FRaceDetectorScope
+	{
+		bool bConditionalActivation;
+		float GracePeriodSeconds;
+	public:
+		FRaceDetectorScope(bool bInConditionalActivation = true, float InGracePeriodSeconds = 0.0f)
+			: bConditionalActivation(bInConditionalActivation)
+			, GracePeriodSeconds(InGracePeriodSeconds)
+		{
+			if (bConditionalActivation)
+			{
+				ToggleRaceDetection(true);
+			}
+		}
+
+		~FRaceDetectorScope()
+		{
+			if (bConditionalActivation)
+			{
+				ToggleRaceDetection(false, GracePeriodSeconds);
+			}
+		}
+	};
+
+
 } // UE::Sanitizer::RaceDetector
 
 #endif
\ No newline at end of file
