Bumped Nanite DDC version to flush out potential stale data after multiple reports of streamer crashes with invalid or inconsistent data

diff --git a/Engine/Source/Runtime/Core/Private/UObject/FortniteMainBranchObjectVersion.cpp b/Engine/Source/Runtime/Core/Private/UObject/FortniteMainBranchObjectVersion.cpp
--- a/Engine/Source/Runtime/Core/Private/UObject/FortniteMainBranchObjectVersion.cpp
+++ b/Engine/Source/Runtime/Core/Private/UObject/FortniteMainBranchObjectVersion.cpp
@@ -9,7 +9,7 @@ TMap<FGuid, FGuid> FFortniteMainBranchObjectVersion::GetSystemGuids()
 	SystemGuids.Add(DevGuids.GLOBALSHADERMAP_DERIVEDDATA_VER, FGuid("125CF21C881744418BD4A614B3B125FF"));
 	SystemGuids.Add(DevGuids.LANDSCAPE_MOBILE_COOK_VERSION, FGuid("32D02EF867C74B71A0D4E0FA41392732"));
 	SystemGuids.Add(DevGuids.MATERIALSHADERMAP_DERIVEDDATA_VER, FGuid("04B8C817071A453C9EBEFBFB51D5F785"));
-	SystemGuids.Add(DevGuids.NANITE_DERIVEDDATA_VER, FGuid("FB5EAC6147EF4591B0D5EB7A85062AC9"));
+	SystemGuids.Add(DevGuids.NANITE_DERIVEDDATA_VER, FGuid("0092F35BA3894D8DBA29E99C9C8866A2"));
 	SystemGuids.Add(DevGuids.NIAGARASHADERMAP_DERIVEDDATA_VER, FGuid("6360A977062842A29ED30E3A7ACB0E64"));
 	SystemGuids.Add(DevGuids.Niagara_LatestScriptCompileVersion, FGuid("97F0FD7149D90D42B65D6557AF477D1A"));
 	SystemGuids.Add(DevGuids.SkeletalMeshDerivedDataVersion, FGuid("C998F2F3B2884A02B7E290B8C25835F0"));
