Deprecated TSortedMap::GetMaxIndex, fixed TSortedSet::GetMaxIndex and correctly commented TSet::GetMaxIndex and TMap::GetMaxIndex.

diff --git a/Engine/Source/Runtime/Core/Public/Containers/Map.h b/Engine/Source/Runtime/Core/Public/Containers/Map.h
--- a/Engine/Source/Runtime/Core/Public/Containers/Map.h
+++ b/Engine/Source/Runtime/Core/Public/Containers/Map.h
@@ -296,7 +296,7 @@ public:
 		return Pairs.Max();
 	}
 
-	/** @return The max valid index of the elements in the sparse storage. */
+	/** @return The non-inclusive maximum index of elements in the map. */
 	[[nodiscard]] FORCEINLINE int32 GetMaxIndex() const
 	{
 		return Pairs.GetMaxIndex();
diff --git a/Engine/Source/Runtime/Core/Public/Containers/Set.h b/Engine/Source/Runtime/Core/Public/Containers/Set.h
--- a/Engine/Source/Runtime/Core/Public/Containers/Set.h
+++ b/Engine/Source/Runtime/Core/Public/Containers/Set.h
@@ -454,6 +454,7 @@ public:
 		return Elements.Max();
 	}
 
+	/** @return The non-inclusive maximum index of elements in the set. */
 	[[nodiscard]] FORCEINLINE int32 GetMaxIndex() const
 	{
 		return Elements.GetMaxIndex();
diff --git a/Engine/Source/Runtime/Core/Public/Containers/SortedMap.h b/Engine/Source/Runtime/Core/Public/Containers/SortedMap.h
--- a/Engine/Source/Runtime/Core/Public/Containers/SortedMap.h
+++ b/Engine/Source/Runtime/Core/Public/Containers/SortedMap.h
@@ -523,9 +523,14 @@ public:
 
 	// Interface functions to match TMap/TSet
 
-	/** @return The max valid index of the elements in the sparse storage. */
+	/** @return The non-inclusive maximum index of elements in the map (was: the max valid index of the elements in the sparse storage). */
+	UE_DEPRECATED(5.7 until 5.9, "GetMaxIndex() should be replaced with Num() - 1.")
 	[[nodiscard]] FORCEINLINE int32 GetMaxIndex() const
 	{
+		// We don't want to deprecate the function entirely, because it should mirror the the existing TMap API, but it's been
+		// returning the wrong value and we want to fix that.
+		//
+		// The deprecation macro lists "5.7 until 5.9", and from 5.9 this function should be de-deprecated and the return value should become "return Pairs.Num();".
 		return Pairs.Num() - 1;
 	}
 
diff --git a/Engine/Source/Runtime/Core/Public/Containers/SortedSet.h b/Engine/Source/Runtime/Core/Public/Containers/SortedSet.h
--- a/Engine/Source/Runtime/Core/Public/Containers/SortedSet.h
+++ b/Engine/Source/Runtime/Core/Public/Containers/SortedSet.h
@@ -164,10 +164,10 @@ public:
 		return Elements.Max();
 	}
 
-	/** @return The max valid index of the elements in the inner array. */
+	/** @return The non-inclusive maximum index of elements in the set. */
 	[[nodiscard]] FORCEINLINE int32 GetMaxIndex() const
 	{
-		return Elements.Num() - 1;
+		return Elements.Num();
 	}
 
 	/** 
diff --git a/Engine/Source/Runtime/Core/Public/Containers/SparseArray.h b/Engine/Source/Runtime/Core/Public/Containers/SparseArray.h
--- a/Engine/Source/Runtime/Core/Public/Containers/SparseArray.h
+++ b/Engine/Source/Runtime/Core/Public/Containers/SparseArray.h
@@ -919,6 +919,7 @@ public:
 	{
 		return AllocationFlags[Index];
 	}
+	/** @return The non-inclusive maximum index of elements in the sparse array. */
 	[[nodiscard]] int32 GetMaxIndex() const
 	{
 		return Data.Num();
@@ -1451,6 +1452,7 @@ public:
 		return Data.NumUnchecked() - NumFreeIndices;
 	}
 
+	/** @return The non-inclusive maximum index of elements in the sparse array. */
 	[[nodiscard]] int32 GetMaxIndex() const
 	{
 		return Data.Num();
