Revert removal of some FORCEINLINE usage from Atomic.h

Must be force-inlined so they're not instrumented by RaceDetector, otherwise it causes a recursive loop since instrumentation uses atomic methods.

#rnx

diff --git a/Engine/Source/Runtime/Core/Public/Templates/Atomic.h b/Engine/Source/Runtime/Core/Public/Templates/Atomic.h
--- a/Engine/Source/Runtime/Core/Public/Templates/Atomic.h
+++ b/Engine/Source/Runtime/Core/Public/Templates/Atomic.h
@@ -57,7 +57,7 @@ namespace UE::Core::Private::Atomic
 	template <>           struct TIsVoidPointer<const volatile void*> { static constexpr bool Value = true; };
 
 	template <typename T>
-	inline T LoadRelaxed(const volatile T* Element)
+	FORCEINLINE T LoadRelaxed(const volatile T* Element)
 	{
 		#ifdef __clang__
 			TUnderlyingIntegerType_T<T> Result;
@@ -69,14 +69,14 @@ namespace UE::Core::Private::Atomic
 	}
 
 	template <typename T>
-	inline T Load(const volatile T* Element)
+	FORCEINLINE T Load(const volatile T* Element)
 	{
 		auto Result = FPlatformAtomics::AtomicRead((volatile TUnderlyingIntegerType_T<T>*)Element);
 		return *(const T*)&Result;
 	}
 
 	template <typename T>
-	inline void StoreRelaxed(const volatile T* Element, T Value)
+	FORCEINLINE void StoreRelaxed(const volatile T* Element, T Value)
 	{
 		#ifdef __clang__
 			__atomic_store((volatile TUnderlyingIntegerType_T<T>*)Element, (TUnderlyingIntegerType_T<T>*)&Value, __ATOMIC_RELAXED);
@@ -92,91 +92,91 @@ namespace UE::Core::Private::Atomic
 	}
 
 	template <typename T>
-	inline T Exchange(volatile T* Element, T Value)
+	FORCEINLINE T Exchange(volatile T* Element, T Value)
 	{
 		auto Result = FPlatformAtomics::InterlockedExchange((volatile TUnderlyingIntegerType_T<T>*)Element, *(const TUnderlyingIntegerType_T<T>*)&Value);
 		return *(const T*)&Result;
 	}
 
 	template <typename T>
-	inline T IncrementExchange(volatile T* Element)
+	FORCEINLINE T IncrementExchange(volatile T* Element)
 	{
 		auto Result = FPlatformAtomics::InterlockedIncrement((volatile TUnderlyingIntegerType_T<T>*)Element) - 1;
 		return *(const T*)&Result;
 	}
 
 	template <typename T>
-	inline T* IncrementExchange(T* volatile* Element)
+	FORCEINLINE T* IncrementExchange(T* volatile* Element)
 	{
 		auto Result = FPlatformAtomics::InterlockedAdd((volatile TUnderlyingIntegerType_T<T*>*)Element, sizeof(T));
 		return *(T* const*)&Result;
 	}
 
 	template <typename T, typename DiffType>
-	inline T AddExchange(volatile T* Element, DiffType Diff)
+	FORCEINLINE T AddExchange(volatile T* Element, DiffType Diff)
 	{
 		auto Result = FPlatformAtomics::InterlockedAdd((volatile TUnderlyingIntegerType_T<T>*)Element, (TUnderlyingIntegerType_T<T>)Diff);
 		return *(const T*)&Result;
 	}
 
 	template <typename T, typename DiffType>
-	inline T* AddExchange(T* volatile* Element, DiffType Diff)
+	FORCEINLINE T* AddExchange(T* volatile* Element, DiffType Diff)
 	{
 		auto Result = FPlatformAtomics::InterlockedAdd((volatile TUnderlyingIntegerType_T<T*>*)Element, (TUnderlyingIntegerType_T<T*>)(Diff * sizeof(T)));
 		return *(T* const*)&Result;
 	}
 
 	template <typename T>
-	inline T DecrementExchange(volatile T* Element)
+	FORCEINLINE T DecrementExchange(volatile T* Element)
 	{
 		auto Result = FPlatformAtomics::InterlockedDecrement((volatile TUnderlyingIntegerType_T<T>*)Element) + 1;
 		return *(const T*)&Result;
 	}
 
 	template <typename T>
-	inline T* DecrementExchange(T* volatile* Element)
+	FORCEINLINE T* DecrementExchange(T* volatile* Element)
 	{
 		auto Result = FPlatformAtomics::InterlockedAdd((volatile TUnderlyingIntegerType_T<T*>*)Element, -(TUnderlyingIntegerType_T<T*>)sizeof(T));
 		return *(T* const*)&Result;
 	}
 
 	template <typename T, typename DiffType>
-	inline T SubExchange(volatile T* Element, DiffType Diff)
+	FORCEINLINE T SubExchange(volatile T* Element, DiffType Diff)
 	{
 		auto Result = FPlatformAtomics::InterlockedAdd((volatile TUnderlyingIntegerType_T<T>*)Element, -(TUnderlyingIntegerType_T<T>)Diff);
 		return *(const T*)&Result;
 	}
 
 	template <typename T, typename DiffType>
-	inline T* SubExchange(T* volatile* Element, DiffType Diff)
+	FORCEINLINE T* SubExchange(T* volatile* Element, DiffType Diff)
 	{
 		auto Result = FPlatformAtomics::InterlockedAdd((volatile TUnderlyingIntegerType_T<T*>*)Element, -(TUnderlyingIntegerType_T<T*>)(Diff * sizeof(T)));
 		return *(T* const*)&Result;
 	}
 
 	template <typename T>
-	inline T CompareExchange(volatile T* Element, T ExpectedValue, T NewValue)
+	FORCEINLINE T CompareExchange(volatile T* Element, T ExpectedValue, T NewValue)
 	{
 		auto Result = FPlatformAtomics::InterlockedCompareExchange((volatile TUnderlyingIntegerType_T<T>*)Element, *(const TUnderlyingIntegerType_T<T>*)&NewValue, *(const TUnderlyingIntegerType_T<T>*)&ExpectedValue);
 		return *(const T*)&Result;
 	}
 
 	template <typename T>
-	inline T AndExchange(volatile T* Element, T AndValue)
+	FORCEINLINE T AndExchange(volatile T* Element, T AndValue)
 	{
 		auto Result = FPlatformAtomics::InterlockedAnd((volatile TUnderlyingIntegerType_T<T>*)Element, (TUnderlyingIntegerType_T<T>)AndValue);
 		return *(const T*)&Result;
 	}
 
 	template <typename T>
-	inline T OrExchange(volatile T* Element, T OrValue)
+	FORCEINLINE T OrExchange(volatile T* Element, T OrValue)
 	{
 		auto Result = FPlatformAtomics::InterlockedOr((volatile TUnderlyingIntegerType_T<T>*)Element, (TUnderlyingIntegerType_T<T>)OrValue);
 		return *(const T*)&Result;
 	}
 
 	template <typename T>
-	inline T XorExchange(volatile T* Element, T XorValue)
+	FORCEINLINE T XorExchange(volatile T* Element, T XorValue)
 	{
 		auto Result = FPlatformAtomics::InterlockedXor((volatile TUnderlyingIntegerType_T<T>*)Element, (TUnderlyingIntegerType_T<T>)XorValue);
 		return *(const T*)&Result;
@@ -575,7 +575,7 @@ public:
 	 *
 	 * @return A copy of Value.
 	 */
-	inline T operator=(T Value)
+	FORCEINLINE T operator=(T Value)
 	{
 		this->Exchange(Value);
 		return Value;
