Coding standards pass on TValueOrError.
Added some extra UE_LIFETIMEBOUND and static assert checks.

diff --git a/Engine/Source/Runtime/Core/Public/Templates/ValueOrError.h b/Engine/Source/Runtime/Core/Public/Templates/ValueOrError.h
--- a/Engine/Source/Runtime/Core/Public/Templates/ValueOrError.h
+++ b/Engine/Source/Runtime/Core/Public/Templates/ValueOrError.h
@@ -8,28 +8,55 @@
 #include "Templates/Tuple.h"
 #include "Templates/UnrealTemplate.h"
 
+template <typename... ArgTypes> struct TValueOrError_ValueProxy;
+template <typename... ArgTypes> struct TValueOrError_ErrorProxy;
+template <typename ValueType, typename ErrorType> class TValueOrError;
+
+template <typename... ArgTypes> TValueOrError_ValueProxy<ArgTypes...> MakeValue(ArgTypes&&... Args);
+template <typename... ArgTypes> TValueOrError_ErrorProxy<ArgTypes...> MakeError(ArgTypes&&... Args);
+
 template <typename... ArgTypes>
 struct TValueOrError_ValueProxy
 {
-	[[nodiscard]] TValueOrError_ValueProxy(ArgTypes&&... InArgs) : Args(Forward<ArgTypes>(InArgs)...) {}
+	friend TValueOrError_ValueProxy<ArgTypes...> MakeValue<ArgTypes...>(ArgTypes&&... Args);
+
+	template <typename ValueType, typename ErrorType>
+	friend class TValueOrError;
+
+private:
+	[[nodiscard]] TValueOrError_ValueProxy(ArgTypes&&... InArgs UE_LIFETIMEBOUND)
+		: Args(Forward<ArgTypes>(InArgs)...)
+	{
+	}
+
 	TTuple<ArgTypes&&...> Args;
 };
 
 template <typename... ArgTypes>
 struct TValueOrError_ErrorProxy
 {
-	[[nodiscard]] TValueOrError_ErrorProxy(ArgTypes&&... InArgs) : Args(Forward<ArgTypes>(InArgs)...) {}
+	friend TValueOrError_ErrorProxy<ArgTypes...> MakeError<ArgTypes...>(ArgTypes&&... Args);
+
+	template <typename ValueType, typename ErrorType>
+	friend class TValueOrError;
+
+private:
+	[[nodiscard]] TValueOrError_ErrorProxy(ArgTypes&&... InArgs UE_LIFETIMEBOUND)
+		: Args(Forward<ArgTypes>(InArgs)...)
+	{
+	}
+
 	TTuple<ArgTypes&&...> Args;
 };
 
 template <typename... ArgTypes>
-[[nodiscard]] FORCEINLINE TValueOrError_ValueProxy<ArgTypes...> MakeValue(ArgTypes&&... Args)
+[[nodiscard]] FORCEINLINE TValueOrError_ValueProxy<ArgTypes...> MakeValue(ArgTypes&&... Args UE_LIFETIMEBOUND)
 {
 	return TValueOrError_ValueProxy<ArgTypes...>(Forward<ArgTypes>(Args)...);
 }
 
 template <typename... ArgTypes>
-[[nodiscard]] FORCEINLINE TValueOrError_ErrorProxy<ArgTypes...> MakeError(ArgTypes&&... Args)
+[[nodiscard]] FORCEINLINE TValueOrError_ErrorProxy<ArgTypes...> MakeError(ArgTypes&&... Args UE_LIFETIMEBOUND)
 {
 	return TValueOrError_ErrorProxy<ArgTypes...>(Forward<ArgTypes>(Args)...);
 }
@@ -47,6 +74,8 @@ template <typename... ArgTypes>
 template <typename ValueType, typename ErrorType>
 class TValueOrError
 {
+	static_assert(!std::is_reference_v<ValueType> && !std::is_reference_v<ErrorType>, "ValueType and ErrorType cannot be references");
+
 	/** Wrap the error type to allow ValueType and ErrorType to be the same type. */
 	struct FWrapErrorType
 	{
@@ -88,15 +117,30 @@ public:
 	}
 
 	/** Check whether a value is set. Prefer HasValue and HasError to this. !IsValid() does *not* imply HasError(). */
-	[[nodiscard]] inline bool IsValid() const { return Variant.template IsType<ValueType>(); }
+	[[nodiscard]] inline bool IsValid() const
+	{
+		return Variant.template IsType<ValueType>();
+	}
 
 	/** Whether the error is set. An error does imply no value. No error does *not* imply that a value is set. */
-	[[nodiscard]] inline bool HasError() const { return Variant.template IsType<FWrapErrorType>(); }
+	[[nodiscard]] inline bool HasError() const
+	{
+		return Variant.template IsType<FWrapErrorType>();
+	}
 
 	/** Access the error. Asserts if this does not have an error. */
-	[[nodiscard]] inline       ErrorType&  GetError()       &  UE_LIFETIMEBOUND { return Variant.template Get<FWrapErrorType>().Error; }
-	[[nodiscard]] inline const ErrorType&  GetError() const &  UE_LIFETIMEBOUND { return Variant.template Get<FWrapErrorType>().Error; }
-	[[nodiscard]] inline       ErrorType&& GetError()       && UE_LIFETIMEBOUND { return MoveTemp(Variant.template Get<FWrapErrorType>().Error); }
+	[[nodiscard]] inline ErrorType& GetError() & UE_LIFETIMEBOUND
+	{
+		return Variant.template Get<FWrapErrorType>().Error;
+	}
+	[[nodiscard]] inline const ErrorType& GetError() const & UE_LIFETIMEBOUND
+	{
+		return const_cast<TValueOrError*>(this)->GetError();
+	}
+	[[nodiscard]] inline ErrorType GetError() && UE_LIFETIMEBOUND
+	{
+		return MoveTempIfPossible(this->GetError());
+	}
 
 	/** Access the error if it is set. */
 	[[nodiscard]] inline ErrorType* TryGetError() UE_LIFETIMEBOUND
@@ -109,37 +153,51 @@ public:
 	}
 	[[nodiscard]] inline const ErrorType* TryGetError() const UE_LIFETIMEBOUND
 	{
-		if (const FWrapErrorType* Wrap = Variant.template TryGet<FWrapErrorType>())
-		{
-			return &Wrap->Error;
-		}
-		return nullptr;
+		return const_cast<TValueOrError*>(this)->TryGetError();
 	}
 
 	/** Steal the error. Asserts if this does not have an error. This causes the error to be unset. */
 	[[nodiscard]] inline ErrorType StealError()
 	{
-		ErrorType Temp = MoveTemp(GetError());
+		ErrorType Temp = MoveTempIfPossible(GetError());
 		Variant.template Emplace<FEmptyType>();
 		return Temp;
 	}
 
 	/** Whether the value is set. A value does imply no error. No value does *not* imply that an error is set. */
-	[[nodiscard]] inline bool HasValue() const { return Variant.template IsType<ValueType>(); }
+	[[nodiscard]] inline bool HasValue() const
+	{
+		return Variant.template IsType<ValueType>();
+	}
 
 	/** Access the value. Asserts if this does not have a value. */
-	[[nodiscard]] inline       ValueType&  GetValue()       &  UE_LIFETIMEBOUND { return Variant.template Get<ValueType>(); }
-	[[nodiscard]] inline const ValueType&  GetValue() const &  UE_LIFETIMEBOUND { return Variant.template Get<ValueType>(); }
-	[[nodiscard]] inline       ValueType&& GetValue()       && UE_LIFETIMEBOUND { return MoveTemp(Variant.template Get<ValueType>()); }
+	[[nodiscard]] inline ValueType& GetValue() & UE_LIFETIMEBOUND
+	{
+		return Variant.template Get<ValueType>();
+	}
+	[[nodiscard]] inline const ValueType& GetValue() const & UE_LIFETIMEBOUND
+	{
+		return const_cast<TValueOrError*>(this)->GetValue();
+	}
+	[[nodiscard]] inline ValueType GetValue() && UE_LIFETIMEBOUND
+	{
+		return MoveTempIfPossible(this->GetValue());
+	}
 
 	/** Access the value if it is set. */
-	[[nodiscard]] inline       ValueType* TryGetValue()       UE_LIFETIMEBOUND { return Variant.template TryGet<ValueType>(); }
-	[[nodiscard]] inline const ValueType* TryGetValue() const UE_LIFETIMEBOUND { return Variant.template TryGet<ValueType>(); }
+	[[nodiscard]] inline ValueType* TryGetValue() UE_LIFETIMEBOUND
+	{
+		return Variant.template TryGet<ValueType>();
+	}
+	[[nodiscard]] inline const ValueType* TryGetValue() const UE_LIFETIMEBOUND
+	{
+		return const_cast<TValueOrError*>(this)->TryGetValue();
+	}
 
 	/** Steal the value. Asserts if this does not have a value. This causes the value to be unset. */
 	[[nodiscard]] inline ValueType StealValue()
 	{
-		ValueType Temp = MoveTemp(GetValue());
+		ValueType Temp = MoveTempIfPossible(GetValue());
 		Variant.template Emplace<FEmptyType>();
 		return Temp;
 	}
@@ -151,6 +209,8 @@ private:
 template <typename ValueType>
 class TValueOrError<ValueType, void>
 {
+	static_assert(!std::is_reference_v<ValueType>, "ValueType cannot be a reference");
+
 	template <typename... ArgTypes, uint32... ArgIndices>
 	[[nodiscard]] inline TValueOrError(TValueOrError_ValueProxy<ArgTypes...>&& Proxy, TIntegerSequence<uint32, ArgIndices...>)
 		: Value(InPlace, MoveTemp(Proxy.Args).template Get<ArgIndices>()...)
@@ -171,24 +231,45 @@ public:
 	}
 
 	/** Whether the error is set. An error does imply no value. */
-	[[nodiscard]] inline bool HasError() const { return !Value.IsSet(); }
+	[[nodiscard]] inline bool HasError() const
+	{
+		return !Value.IsSet();
+	}
 
 	/** Whether the value is set. A value does imply no error. */
-	[[nodiscard]] inline bool HasValue() const { return Value.IsSet(); }
+	[[nodiscard]] inline bool HasValue() const
+	{
+		return Value.IsSet();
+	}
 
 	/** Access the value. Asserts if this does not have a value. */
-	[[nodiscard]] inline       ValueType&  GetValue()       &  UE_LIFETIMEBOUND { return Value.GetValue(); }
-	[[nodiscard]] inline const ValueType&  GetValue() const &  UE_LIFETIMEBOUND { return Value.GetValue(); }
-	[[nodiscard]] inline       ValueType&& GetValue()       && UE_LIFETIMEBOUND { return MoveTemp(Value.GetValue()); }
+	[[nodiscard]] inline ValueType& GetValue() & UE_LIFETIMEBOUND
+	{
+		return Value.GetValue();
+	}
+	[[nodiscard]] inline const ValueType& GetValue() const & UE_LIFETIMEBOUND
+	{
+		return const_cast<TValueOrError*>(this)->GetValue();
+	}
+	[[nodiscard]] inline ValueType GetValue() && UE_LIFETIMEBOUND
+	{
+		return MoveTempIfPossible(this->GetValue());
+	}
 
 	/** Access the value if it is set. */
-	[[nodiscard]] inline       ValueType* TryGetValue()       UE_LIFETIMEBOUND { return Value.GetPtrOrNull(); }
-	[[nodiscard]] inline const ValueType* TryGetValue() const UE_LIFETIMEBOUND { return Value.GetPtrOrNull(); }
+	[[nodiscard]] inline ValueType* TryGetValue() UE_LIFETIMEBOUND
+	{
+		return Value.GetPtrOrNull();
+	}
+	[[nodiscard]] inline const ValueType* TryGetValue() const UE_LIFETIMEBOUND
+	{
+		return const_cast<TValueOrError*>(this)->TryGetValue();
+	}
 
 	/** Steal the value. Asserts if this does not have a value. This causes the value to be unset. */
 	[[nodiscard]] inline ValueType StealValue()
 	{
-		ValueType Temp = MoveTemp(GetValue());
+		ValueType Temp = MoveTempIfPossible(GetValue());
 		Value.Reset();
 		return Temp;
 	}
@@ -200,6 +281,8 @@ private:
 template <typename ErrorType>
 class TValueOrError<void, ErrorType>
 {
+	static_assert(!std::is_reference_v<ErrorType>, "ErrorType cannot be a reference");
+
 	template <typename... ArgTypes, uint32... ArgIndices>
 	[[nodiscard]] inline TValueOrError(TValueOrError_ErrorProxy<ArgTypes...>&& Proxy, TIntegerSequence<uint32, ArgIndices...>)
 		: Error(InPlace, MoveTemp(Proxy.Args).template Get<ArgIndices>()...)
@@ -220,24 +303,45 @@ public:
 	}
 
 	/** Whether the value is set. A value does imply no error. */
-	[[nodiscard]] inline bool HasValue() const { return !Error.IsSet(); }
+	[[nodiscard]] inline bool HasValue() const
+	{
+		return !Error.IsSet();
+	}
 
 	/** Whether the error is set. An error does imply no value. */
-	[[nodiscard]] inline bool HasError() const { return Error.IsSet(); }
+	[[nodiscard]] inline bool HasError() const
+	{
+		return Error.IsSet();
+	}
 
 	/** Access the error. Asserts if this does not have an error. */
-	[[nodiscard]] inline       ErrorType&  GetError()       &  UE_LIFETIMEBOUND { return Error.GetValue(); }
-	[[nodiscard]] inline const ErrorType&  GetError() const &  UE_LIFETIMEBOUND { return Error.GetValue(); }
-	[[nodiscard]] inline       ErrorType&& GetError()       && UE_LIFETIMEBOUND { return MoveTemp(Error.GetValue()); }
+	[[nodiscard]] inline ErrorType& GetError() & UE_LIFETIMEBOUND
+	{
+		return Error.GetValue();
+	}
+	[[nodiscard]] inline const ErrorType& GetError() const & UE_LIFETIMEBOUND
+	{
+		return const_cast<TValueOrError*>(this)->GetError();
+	}
+	[[nodiscard]] inline ErrorType GetError() && UE_LIFETIMEBOUND
+	{
+		return MoveTempIfPossible(this->GetError());
+	}
 
 	/** Access the error if it is set. */
-	[[nodiscard]] inline       ErrorType* TryGetError()       UE_LIFETIMEBOUND { return Error.GetPtrOrNull(); }
-	[[nodiscard]] inline const ErrorType* TryGetError() const UE_LIFETIMEBOUND { return Error.GetPtrOrNull(); }
+	[[nodiscard]] inline ErrorType* TryGetError() UE_LIFETIMEBOUND
+	{
+		return Error.GetPtrOrNull();
+	}
+	[[nodiscard]] inline const ErrorType* TryGetError() const UE_LIFETIMEBOUND
+	{
+		return const_cast<TValueOrError*>(this)->TryGetError();
+	}
 
 	/** Steal the error. Asserts if this does not have an error. This causes the error to be unset. */
 	[[nodiscard]] inline ErrorType StealError()
 	{
-		ErrorType Temp = MoveTemp(GetError());
+		ErrorType Temp = MoveTempIfPossible(GetError());
 		Error.Reset();
 		return Temp;
 	}
@@ -263,10 +367,16 @@ public:
 	}
 
 	/** Whether the error is set. An error does imply no value. */
-	[[nodiscard]] inline bool HasError() const { return !bValue; }
+	[[nodiscard]] inline bool HasError() const
+	{
+		return !bValue;
+	}
 
 	/** Whether the value is set. A value does imply no error. */
-	[[nodiscard]] inline bool HasValue() const { return bValue; }
+	[[nodiscard]] inline bool HasValue() const
+	{
+		return bValue;
+	}
 
 private:
 	bool bValue;
