made TArray, TSet, and FPThreadsRecursiveMutex, default constructors constexpr

diff --git a/Engine/Source/Runtime/Core/Public/Containers/Array.h b/Engine/Source/Runtime/Core/Public/Containers/Array.h
--- a/Engine/Source/Runtime/Core/Public/Containers/Array.h
+++ b/Engine/Source/Runtime/Core/Public/Containers/Array.h
@@ -673,7 +673,7 @@ public:
 	/**
 	 * Constructor, initializes element number counters.
 	 */
-	[[nodiscard]] FORCEINLINE TArray()
+	[[nodiscard]] FORCEINLINE constexpr TArray()
 		: ArrayNum(0)
 		, ArrayMax(AllocatorInstance.GetInitialCapacity())
 	{
diff --git a/Engine/Source/Runtime/Core/Public/Containers/BitArray.h b/Engine/Source/Runtime/Core/Public/Containers/BitArray.h
--- a/Engine/Source/Runtime/Core/Public/Containers/BitArray.h
+++ b/Engine/Source/Runtime/Core/Public/Containers/BitArray.h
@@ -370,7 +370,7 @@ public:
 	template<typename,typename,bool>
 	friend class TConstDualSetBitIterator;
 
-	[[nodiscard]] TBitArray()
+	[[nodiscard]] constexpr TBitArray()
 	:	NumBits(0)
 	,	MaxBits(AllocatorInstance.GetInitialCapacity() * NumBitsPerDWORD)
 	{
diff --git a/Engine/Source/Runtime/Core/Public/Containers/ContainerAllocationPolicies.h b/Engine/Source/Runtime/Core/Public/Containers/ContainerAllocationPolicies.h
--- a/Engine/Source/Runtime/Core/Public/Containers/ContainerAllocationPolicies.h
+++ b/Engine/Source/Runtime/Core/Public/Containers/ContainerAllocationPolicies.h
@@ -639,9 +639,7 @@ public:
 
 	public:
 		/** Default constructor. */
-		ForAnyElementType()
-			: Data(nullptr)
-		{}
+		constexpr ForAnyElementType() = default;
 
 		/**
 		 * Moves the state of another allocator into this one.  The allocator can be different.
@@ -782,7 +780,7 @@ public:
 			return !!Data;
 		}
 
-		SizeType GetInitialCapacity() const
+		constexpr SizeType GetInitialCapacity() const
 		{
 			return 0;
 		}
@@ -806,7 +804,7 @@ public:
 		ForAnyElementType& operator=(const ForAnyElementType&);
 
 		/** A pointer to the container's elements. */
-		FScriptContainerElement* Data;
+		FScriptContainerElement* Data = nullptr;
 	};
 	
 	template<typename ElementType>
@@ -870,8 +868,15 @@ public:
 	public:
 
 		/** Default constructor. */
-		ForElementType()
+		constexpr ForElementType()
 		{
+			UE_IF_CONSTEVAL
+			{
+				for (auto& Element : InlineData)
+				{
+					Element.ConstInit();
+				}
+			}
 		}
 
 		/**
@@ -975,7 +980,7 @@ public:
 			return SecondaryData.HasAllocation();
 		}
 
-		SizeType GetInitialCapacity() const
+		constexpr SizeType GetInitialCapacity() const
 		{
 			return NumInlineElements;
 		}
diff --git a/Engine/Source/Runtime/Core/Public/Containers/Set.h b/Engine/Source/Runtime/Core/Public/Containers/Set.h
--- a/Engine/Source/Runtime/Core/Public/Containers/Set.h
+++ b/Engine/Source/Runtime/Core/Public/Containers/Set.h
@@ -174,7 +174,7 @@ private:
 
 public:
 	/** Initialization constructor. */
-	[[nodiscard]] FORCEINLINE TSet() = default;
+	[[nodiscard]] FORCEINLINE constexpr TSet() = default;
 
 	/** Copy constructor. */
 	[[nodiscard]] FORCEINLINE TSet(const TSet& Copy)
diff --git a/Engine/Source/Runtime/Core/Public/Containers/SparseArray.h b/Engine/Source/Runtime/Core/Public/Containers/SparseArray.h
--- a/Engine/Source/Runtime/Core/Public/Containers/SparseArray.h
+++ b/Engine/Source/Runtime/Core/Public/Containers/SparseArray.h
@@ -769,11 +769,7 @@ public:
 	}
 
 	/** Default constructor. */
-	[[nodiscard]] TSparseArray()
-		: FirstFreeIndex(-1)
-		, NumFreeIndices(0)
-	{
-	}
+	[[nodiscard]] constexpr TSparseArray() = default;
 
 	/** Move constructor. */
 	[[nodiscard]] TSparseArray(TSparseArray&& InCopy)
diff --git a/Engine/Source/Runtime/Core/Public/HAL/PThreadsRecursiveMutex.h b/Engine/Source/Runtime/Core/Public/HAL/PThreadsRecursiveMutex.h
--- a/Engine/Source/Runtime/Core/Public/HAL/PThreadsRecursiveMutex.h
+++ b/Engine/Source/Runtime/Core/Public/HAL/PThreadsRecursiveMutex.h
@@ -9,7 +9,10 @@
 #if PLATFORM_USE_PTHREADS
 
 #include <pthread.h>
-#include <errno.h>
+
+#if defined(PTHREAD_RECURSIVE_MUTEX_INITIALIZER) && !defined(PTHREAD_RECURSIVE_MUTEX_INITIALIZER_NP)
+	#define PTHREAD_RECURSIVE_MUTEX_INITIALIZER_NP PTHREAD_RECURSIVE_MUTEX_INITIALIZER
+#endif
 
 namespace UE
 {
@@ -25,6 +28,9 @@ public:
 	FPThreadsRecursiveMutex(const FPThreadsRecursiveMutex&) = delete;
 	FPThreadsRecursiveMutex& operator=(const FPThreadsRecursiveMutex&) = delete;
 
+#if defined(PTHREAD_RECURSIVE_MUTEX_INITIALIZER_NP)
+	constexpr FPThreadsRecursiveMutex() = default;
+#else
 	inline FPThreadsRecursiveMutex()
 	{
 		pthread_mutexattr_t MutexAttributes;
@@ -33,6 +39,7 @@ public:
 		pthread_mutex_init(&Mutex, &MutexAttributes);
 		pthread_mutexattr_destroy(&MutexAttributes);
 	}
+#endif
 
 	FORCEINLINE ~FPThreadsRecursiveMutex()
 	{
@@ -55,7 +62,11 @@ public:
 	}
 
 private:
+#if defined(PTHREAD_RECURSIVE_MUTEX_INITIALIZER_NP)
+	pthread_mutex_t Mutex = PTHREAD_RECURSIVE_MUTEX_INITIALIZER_NP;
+#else
 	pthread_mutex_t Mutex;
+#endif
 };
 
 } // UE
diff --git a/Engine/Source/Runtime/Core/Public/HAL/Platform.h b/Engine/Source/Runtime/Core/Public/HAL/Platform.h
--- a/Engine/Source/Runtime/Core/Public/HAL/Platform.h
+++ b/Engine/Source/Runtime/Core/Public/HAL/Platform.h
@@ -100,6 +100,7 @@
 #endif
 
 #if PLATFORM_APPLE
+	#define _DARWIN_C_SOURCE // include darwin extensions
 	#include <stddef.h> // needed for size_t
 #endif
 
diff --git a/Engine/Source/Runtime/Core/Public/Templates/TypeCompatibleBytes.h b/Engine/Source/Runtime/Core/Public/Templates/TypeCompatibleBytes.h
--- a/Engine/Source/Runtime/Core/Public/Templates/TypeCompatibleBytes.h
+++ b/Engine/Source/Runtime/Core/Public/Templates/TypeCompatibleBytes.h
@@ -62,6 +62,15 @@ struct TTypeCompatibleBytes
 		return (ElementType&&)*(ElementType*)this;
 	}
 
+	// Initializes to be constinit friendly
+	constexpr void ConstInit()
+	{
+		for (auto& Value : Pad)
+		{
+			Value = 0;
+		}
+	}
+	
 	// Emplaces an inner element.
 	// Note: no checks are possible to ensure that an element isn't already present.  DestroyUnchecked() must be called to end the element's lifetime.
 	template <typename... ArgTypes>
