Expanded FTextChar support

FTextChar now supports all FChar functions that have an ICU implementation

FUnicodeChar::CodepointToString implementation has moved to FTextChar::ConvertCodepointToString and no longer depends on ICU, and there is also now FTextChar::AppendCodepointToString

diff --git a/Engine/Source/Runtime/Core/Private/Internationalization/ICUText.cpp b/Engine/Source/Runtime/Core/Private/Internationalization/ICUText.cpp
--- a/Engine/Source/Runtime/Core/Private/Internationalization/ICUText.cpp
+++ b/Engine/Source/Runtime/Core/Private/Internationalization/ICUText.cpp
@@ -101,11 +101,56 @@ bool FTextChar::IsAlpha(const UTF32CHAR Codepoint)
 	return u_isalpha(Codepoint) != 0;
 }
 
+bool FTextChar::IsGraph(const UTF32CHAR Codepoint)
+{
+	return u_isgraph(Codepoint) != 0;
+}
+
+bool FTextChar::IsPrint(const UTF32CHAR Codepoint)
+{
+	return u_isprint(Codepoint) != 0;
+}
+
+bool FTextChar::IsPunct(const UTF32CHAR Codepoint)
+{
+	return u_ispunct(Codepoint) != 0;
+}
+
+bool FTextChar::IsAlnum(const UTF32CHAR Codepoint)
+{
+	return u_isalnum(Codepoint) != 0;
+}
+
+bool FTextChar::IsDigit(const UTF32CHAR Codepoint)
+{
+	return u_isdigit(Codepoint) != 0;
+}
+
+bool FTextChar::IsHexDigit(const UTF32CHAR Codepoint)
+{
+	return u_isxdigit(Codepoint) != 0;
+}
+
+bool FTextChar::IsControl(const UTF32CHAR Codepoint)
+{
+	return u_iscntrl(Codepoint) != 0;
+}
+
 bool FTextChar::IsWhitespace(const UTF32CHAR Codepoint)
 {
 	return u_isWhitespace(Codepoint) != 0;
 }
 
+bool FTextChar::IsUpper(const UTF32CHAR Codepoint)
+{
+	return u_isupper(Codepoint) != 0;
+}
+
+bool FTextChar::IsLower(const UTF32CHAR Codepoint)
+{
+	return u_islower(Codepoint) != 0;
+}
+
 UTF32CHAR FTextChar::ToLower(const UTF32CHAR Codepoint)
 {
 	return u_tolower(Codepoint);
@@ -194,14 +239,6 @@ bool FText::FSortPredicate::operator()(const FText& A, const FText& B) const
 	return Implementation->Compare(A, B);
 }
 
-bool FUnicodeChar::CodepointToString(const uint32 InCodepoint, FString& OutString)
-{
-	icu::UnicodeString CodepointString;
-	CodepointString.setTo((UChar32)InCodepoint);
-	ICUUtilities::ConvertString(CodepointString, OutString);
-	return true;
-}
-
 namespace TextBiDi
 {
 
diff --git a/Engine/Source/Runtime/Core/Private/Internationalization/LegacyText.cpp b/Engine/Source/Runtime/Core/Private/Internationalization/LegacyText.cpp
--- a/Engine/Source/Runtime/Core/Private/Internationalization/LegacyText.cpp
+++ b/Engine/Source/Runtime/Core/Private/Internationalization/LegacyText.cpp
@@ -51,11 +51,56 @@ bool FTextChar::IsAlpha(const UTF32CHAR Codepoint)
 	return FChar::IsAlpha(static_cast<TCHAR>(Codepoint));
 }
 
+bool FTextChar::IsGraph(const UTF32CHAR Codepoint)
+{
+	return FChar::IsGraph(static_cast<TCHAR>(Codepoint));
+}
+
+bool FTextChar::IsPrint(const UTF32CHAR Codepoint)
+{
+	return FChar::IsPrint(static_cast<TCHAR>(Codepoint));
+}
+
+bool FTextChar::IsPunct(const UTF32CHAR Codepoint)
+{
+	return FChar::IsPunct(static_cast<TCHAR>(Codepoint));
+}
+
+bool FTextChar::IsAlnum(const UTF32CHAR Codepoint)
+{
+	return FChar::IsAlnum(static_cast<TCHAR>(Codepoint));
+}
+
+bool FTextChar::IsDigit(const UTF32CHAR Codepoint)
+{
+	return FChar::IsDigit(static_cast<TCHAR>(Codepoint));
+}
+
+bool FTextChar::IsHexDigit(const UTF32CHAR Codepoint)
+{
+	return FChar::IsHexDigit(static_cast<TCHAR>(Codepoint));
+}
+
+bool FTextChar::IsControl(const UTF32CHAR Codepoint)
+{
+	return FChar::IsControl(static_cast<TCHAR>(Codepoint));
+}
+
 bool FTextChar::IsWhitespace(const UTF32CHAR Codepoint)
 {
 	return FChar::IsWhitespace(static_cast<TCHAR>(Codepoint));
 }
 
+bool FTextChar::IsUpper(const UTF32CHAR Codepoint)
+{
+	return FChar::IsUpper(static_cast<TCHAR>(Codepoint));
+}
+
+bool FTextChar::IsLower(const UTF32CHAR Codepoint)
+{
+	return FChar::IsLower(static_cast<TCHAR>(Codepoint));
+}
+
 UTF32CHAR FTextChar::ToLower(const UTF32CHAR Codepoint)
 {
 	return FChar::ToLower(static_cast<TCHAR>(Codepoint));
@@ -96,11 +141,6 @@ bool FText::FSortPredicate::operator()(const FText& A, const FText& B) const
 	return A.ToString() < B.ToString();
 }
 
-bool FUnicodeChar::CodepointToString(const uint32 InCodepoint, FString& OutString)
-{
-	return false;
-}
-
 namespace TextBiDi
 {
 
diff --git a/Engine/Source/Runtime/Core/Private/Internationalization/Text.cpp b/Engine/Source/Runtime/Core/Private/Internationalization/Text.cpp
--- a/Engine/Source/Runtime/Core/Private/Internationalization/Text.cpp
+++ b/Engine/Source/Runtime/Core/Private/Internationalization/Text.cpp
@@ -2126,7 +2126,7 @@ UTF32CHAR FTextChar::GetCodepoint(const TCHAR* Buffer, int32* OutNumTCHARInCodep
 {
 	static_assert(sizeof(TCHAR) == 2 || sizeof(TCHAR) == 4, "Expected TCHAR to be UTF-16 or UTF-32");
 
-	if (!Buffer)
+	if (!Buffer || !Buffer[0])
 	{
 		return 0;
 	}
@@ -2156,6 +2156,55 @@ UTF32CHAR FTextChar::GetCodepoint(const TCHAR* Buffer, int32* OutNumTCHARInCodep
     return Codepoint;
 }
 
+bool FTextChar::ConvertCodepointToString(const UTF32CHAR Codepoint, FString& OutString)
+{
+	OutString.Reset();
+	return FTextChar::AppendCodepointToString(Codepoint, OutString);
+}
+
+bool FTextChar::AppendCodepointToString(const UTF32CHAR Codepoint, FString& OutString)
+{
+	static_assert(sizeof(TCHAR) == 2 || sizeof(TCHAR) == 4, "Expected TCHAR to be UTF-16 or UTF-32");
+
+	if (!StringConv::IsValidCodepoint(Codepoint))
+	{
+		return false;
+	}
+	
+	if (StringConv::IsHighSurrogate(Codepoint) || StringConv::IsLowSurrogate(Codepoint)) // Surrogate pairs range should not be present in UTF-32
+	{
+		return false;
+	}
+
+#if !PLATFORM_TCHAR_IS_4_BYTES
+	if (StringConv::IsEncodedSurrogate(Codepoint))
+	{
+		uint16 HighSurrogate = 0;
+		uint16 LowSurrogate = 0;
+		StringConv::DecodeSurrogate(Codepoint, HighSurrogate, LowSurrogate);
+
+		OutString.AppendChar((TCHAR)HighSurrogate);
+		OutString.AppendChar((TCHAR)LowSurrogate);
+	}
+	else if (Codepoint > ENCODED_SURROGATE_END_CODEPOINT)
+	{
+		// Ignore values higher than the supplementary plane range
+		return false;
+	}
+	else
+#endif	// !PLATFORM_TCHAR_IS_4_BYTES
+	{
+		OutString.AppendChar((TCHAR)Codepoint);
+	}
+
+	return true;
+}
+
+bool FUnicodeChar::CodepointToString(const uint32 InCodepoint, FString& OutString)
+{
+	return FTextChar::ConvertCodepointToString(InCodepoint, OutString);
+}
+
 bool LexTryParseString(ETextGender& OutValue, const TCHAR* Buffer)
 {
 #define ENUM_CASE_FROM_STRING(Enum) if (FCString::Stricmp(Buffer, TEXT(#Enum)) == 0) { OutValue = ETextGender::Enum; return true; }
diff --git a/Engine/Source/Runtime/Core/Public/Internationalization/TextChar.h b/Engine/Source/Runtime/Core/Public/Internationalization/TextChar.h
--- a/Engine/Source/Runtime/Core/Public/Internationalization/TextChar.h
+++ b/Engine/Source/Runtime/Core/Public/Internationalization/TextChar.h
@@ -4,36 +4,94 @@
 
 #include "CoreTypes.h"
 
+class FString;
+
 /**
- * Utility for performing character checks.
+ * Utility for performing Unicode aware (in builds with ICU enabled) character checks.
+ * In builds without ICU enabled (such as server builds), this API uses the non-Unicode aware FChar API internally.
  * The implementation can be found in LegacyText.cpp and ICUText.cpp.
  */
 class FTextChar
 {
 public:
 	/**
-	* Return UTF32CHAR codepoint. Buffer must be a null terminated string and at least 1 character long
-	* EX: FString Name;
-	*     if (Name.Len() > 1)
-	*	  	  const UTF32CHAR FirstLetter = FTextChar::GetCodepoint(*Name);
-	*/
-	[[nodiscard]] static CORE_API UTF32CHAR GetCodepoint( const TCHAR* Buffer, int32* OutNumTCHARInCodepoint = nullptr );
+	 * Return the next codepoint from the given null terminated buffer.
+	 */
+	[[nodiscard]] static CORE_API UTF32CHAR GetCodepoint(const TCHAR* Buffer, int32* OutNumTCHARInCodepoint = nullptr);
+
+	/**
+	 * Attempt to convert the given codepoint to a string, replacing any existing content in the string.
+	 */
+	[[nodiscard]] static CORE_API bool ConvertCodepointToString(const UTF32CHAR Codepoint, FString& OutString);
 
 	/**
-	* Determine whether the specified code point is a letter character.
-	*/
-	[[nodiscard]] static CORE_API bool IsAlpha( const UTF32CHAR Codepoint );
+	 * Attempt to convert the given codepoint to a string, appending after any existing content in the string.
+	 */
+	[[nodiscard]] static CORE_API bool AppendCodepointToString(const UTF32CHAR Codepoint, FString& OutString);
 
 	/**
-	* Determine whether the Char code is a whitespace character.
-	* TCHAR should either be UTF-16 or UTF-32, so we should be fine to cast it to a UChar32 for the whitespace
-	* check, since whitespace is never a pair of UTF-16 characters
-	* EX: TChar Char;
-	*     IsWhitespace(static_cast<UTF32CHAR>(Char));
-	*/
-	[[nodiscard]] static CORE_API bool IsWhitespace( const UTF32CHAR Codepoint );
+	 * Determine whether the given codepoint is a letter character.
+	 */
+	[[nodiscard]] static CORE_API bool IsAlpha(const UTF32CHAR Codepoint);
 
-	[[nodiscard]] static CORE_API UTF32CHAR ToLower( const UTF32CHAR Codepoint );
+	/**
+	 * Determine whether the given codepoint is a "graphic" character (printable, excluding spaces).
+	 */
+	[[nodiscard]] static CORE_API bool IsGraph(const UTF32CHAR Codepoint);
 
-	[[nodiscard]] static CORE_API UTF32CHAR ToUpper( const UTF32CHAR Codepoint );
+	/**
+	 * Determine whether the given codepoint is a printable character.
+	 */
+	[[nodiscard]] static CORE_API bool IsPrint(const UTF32CHAR Codepoint);
+
+	/**
+	 * Determine whether the given codepoint is a punctuation character.
+	 */
+	[[nodiscard]] static CORE_API bool IsPunct(const UTF32CHAR Codepoint);
+
+	/**
+	 * Determine whether the given codepoint is an alphanumeric character (letter or digit).
+	 */
+	[[nodiscard]] static CORE_API bool IsAlnum(const UTF32CHAR Codepoint);
+
+	/**
+	 * Determine whether the given codepoint is a digit character.
+	 */
+	[[nodiscard]] static CORE_API bool IsDigit(const UTF32CHAR Codepoint);
+
+	/**
+	 * Determine whether the given codepoint is a hexadecimal digit character.
+	 */
+	[[nodiscard]] static CORE_API bool IsHexDigit(const UTF32CHAR Codepoint);
+
+	/**
+	 * Determine whether the given codepoint is a control character.
+	 */
+	[[nodiscard]] static CORE_API bool IsControl(const UTF32CHAR Codepoint);
+
+	/**
+	 * Determine whether the given codepoint is a whitespace character.
+	 * @note It is safe to pass a UTF-16 TCHAR to this function, since whitespace is never a pair of UTF-16 characters.
+	 */
+	[[nodiscard]] static CORE_API bool IsWhitespace(const UTF32CHAR Codepoint);
+
+	/**
+	 * Determine whether the given codepoint is an uppercase letter character.
+	 */
+	[[nodiscard]] static CORE_API bool IsUpper(const UTF32CHAR Codepoint);
+
+	/**
+	 * Determine whether the given codepoint is an lowercase letter character.
+	 */
+	[[nodiscard]] static CORE_API bool IsLower(const UTF32CHAR Codepoint);
+
+	/**
+	 * Convert the given codepoint to its lowercase version.
+	 */
+	[[nodiscard]] static CORE_API UTF32CHAR ToLower(const UTF32CHAR Codepoint);
+
+	/**
+	 * Convert the given codepoint to its uppercase version.
+	 */
+	[[nodiscard]] static CORE_API UTF32CHAR ToUpper(const UTF32CHAR Codepoint);
 };
