Removed code duplication in TValueOrError by implementing overloaded getters in terms of the primary overload.

diff --git a/Engine/Source/Runtime/Core/Public/Templates/ValueOrError.h b/Engine/Source/Runtime/Core/Public/Templates/ValueOrError.h
--- a/Engine/Source/Runtime/Core/Public/Templates/ValueOrError.h
+++ b/Engine/Source/Runtime/Core/Public/Templates/ValueOrError.h
@@ -114,11 +114,11 @@ public:
 	}
 	[[nodiscard]] inline const ErrorType& GetError() const & UE_LIFETIMEBOUND
 	{
-		return Variant.template Get<FWrapErrorType>().Error;
+		return const_cast<TValueOrError*>(this)->GetError();
 	}
 	[[nodiscard]] inline ErrorType&& GetError() && UE_LIFETIMEBOUND
 	{
-		return MoveTemp(Variant.template Get<FWrapErrorType>().Error);
+		return MoveTemp(this->GetError());
 	}
 
 	/** Access the error if it is set. */
@@ -132,11 +132,7 @@ public:
 	}
 	[[nodiscard]] inline const ErrorType* TryGetError() const UE_LIFETIMEBOUND
 	{
-		if (const FWrapErrorType* Wrap = Variant.template TryGet<FWrapErrorType>())
-		{
-			return &Wrap->Error;
-		}
-		return nullptr;
+		return const_cast<TValueOrError*>(this)->TryGetError();
 	}
 
 	/** Steal the error. Asserts if this does not have an error. This causes the error to be unset. */
@@ -158,13 +154,13 @@ public:
 	{
 		return Variant.template Get<ValueType>();
 	}
-	[[nodiscard]] inline const ValueType&  GetValue() const & UE_LIFETIMEBOUND
+	[[nodiscard]] inline const ValueType& GetValue() const & UE_LIFETIMEBOUND
 	{
-		return Variant.template Get<ValueType>();
+		return const_cast<TValueOrError*>(this)->GetValue();
 	}
 	[[nodiscard]] inline ValueType&& GetValue() && UE_LIFETIMEBOUND
 	{
-		return MoveTemp(Variant.template Get<ValueType>());
+		return MoveTemp(this->GetValue());
 	}
 
 	/** Access the value if it is set. */
@@ -174,7 +170,7 @@ public:
 	}
 	[[nodiscard]] inline const ValueType* TryGetValue() const UE_LIFETIMEBOUND
 	{
-		return Variant.template TryGet<ValueType>();
+		return const_cast<TValueOrError*>(this)->TryGetValue();
 	}
 
 	/** Steal the value. Asserts if this does not have a value. This causes the value to be unset. */
@@ -230,11 +226,11 @@ public:
 	}
 	[[nodiscard]] inline const ValueType& GetValue() const & UE_LIFETIMEBOUND
 	{
-		return Value.GetValue();
+		return const_cast<TValueOrError*>(this)->GetValue();
 	}
 	[[nodiscard]] inline ValueType&& GetValue() && UE_LIFETIMEBOUND
 	{
-		return MoveTemp(Value.GetValue());
+		return MoveTemp(this->GetValue());
 	}
 
 	/** Access the value if it is set. */
@@ -244,7 +240,7 @@ public:
 	}
 	[[nodiscard]] inline const ValueType* TryGetValue() const UE_LIFETIMEBOUND
 	{
-		return Value.GetPtrOrNull();
+		return const_cast<TValueOrError*>(this)->TryGetValue();
 	}
 
 	/** Steal the value. Asserts if this does not have a value. This causes the value to be unset. */
@@ -298,13 +294,13 @@ public:
 	{
 		return Error.GetValue();
 	}
-	[[nodiscard]] inline const ErrorType&  GetError() const & UE_LIFETIMEBOUND
+	[[nodiscard]] inline const ErrorType& GetError() const & UE_LIFETIMEBOUND
 	{
-		return Error.GetValue();
+		return const_cast<TValueOrError*>(this)->GetError();
 	}
 	[[nodiscard]] inline ErrorType&& GetError() && UE_LIFETIMEBOUND
 	{
-		return MoveTemp(Error.GetValue());
+		return MoveTemp(this->GetError());
 	}
 
 	/** Access the error if it is set. */
@@ -314,7 +310,7 @@ public:
 	}
 	[[nodiscard]] inline const ErrorType* TryGetError() const UE_LIFETIMEBOUND
 	{
-		return Error.GetPtrOrNull();
+		return const_cast<TValueOrError*>(this)->TryGetError();
 	}
 
 	/** Steal the error. Asserts if this does not have an error. This causes the error to be unset. */
