Deprecated some functions that are now implemented via FTextChar

* FText::IsWhitespace is deprecated in favor of FTextChar::IsWhitespace
* FUnicodeChar::CodepointToString is deprecated in favor of FTextChar::ConvertCodepointToString

diff --git a/Engine/Source/Runtime/Core/Public/Internationalization/Text.h b/Engine/Source/Runtime/Core/Public/Internationalization/Text.h
--- a/Engine/Source/Runtime/Core/Public/Internationalization/Text.h
+++ b/Engine/Source/Runtime/Core/Public/Internationalization/Text.h
@@ -628,6 +628,7 @@ public:
 	/**
 	 * Check to see if the given character is considered whitespace by the current culture
 	 */
+	UE_DEPRECATED(5.7, "FText::IsWhitespace is deprecated. Use FTextChar::IsWhitespace instead.")
 	[[nodiscard]] static CORE_API bool IsWhitespace( const TCHAR Char );
 
 	static CORE_API void GetFormatPatternParameters(const FTextFormat& Fmt, TArray<FString>& ParameterNames);
@@ -1393,9 +1394,10 @@ private:
 	int32 IndentCount = 0;
 };
 
-/** Unicode character helper functions */
+/** Unicode character helper functions (see FTextChar for more functionality) */
 struct FUnicodeChar
 {
+	UE_DEPRECATED(5.7, "FUnicodeChar::CodepointToString is deprecated. Use FTextChar::ConvertCodepointToString instead.")
 	static CORE_API bool CodepointToString(const uint32 InCodepoint, FString& OutString);
 };
 
