Substrate - fixed specified tangent broken on the ShadingModel legacy node for anisotropy.

#tests PC editor dx12
#rnx
[RN] Substrate - fixed specified tangent broken on the ShadingModel legacy node for anisotropy.

diff --git a/Engine/Source/Runtime/Core/Private/UObject/FortniteMainBranchObjectVersion.cpp b/Engine/Source/Runtime/Core/Private/UObject/FortniteMainBranchObjectVersion.cpp
--- a/Engine/Source/Runtime/Core/Private/UObject/FortniteMainBranchObjectVersion.cpp
+++ b/Engine/Source/Runtime/Core/Private/UObject/FortniteMainBranchObjectVersion.cpp
@@ -14,6 +14,6 @@ TMap<FGuid, FGuid> FFortniteMainBranchObjectVersion::GetSystemGuids()
 	SystemGuids.Add(DevGuids.Niagara_LatestScriptCompileVersion, FGuid("97F0FD7149D90D42B65D6557AF477D1A"));
 	SystemGuids.Add(DevGuids.SkeletalMeshDerivedDataVersion, FGuid("C998F2F3B2884A02B7E290B8C25835F0"));
 	SystemGuids.Add(DevGuids.STATICMESH_DERIVEDDATA_VER, FGuid("3ABF32149F9D4D838750368AB95573BE"));
-	SystemGuids.Add(DevGuids.MaterialTranslationDDCVersion, FGuid("9074F4446D3843F1ABD42CA551384212"));
+	SystemGuids.Add(DevGuids.MaterialTranslationDDCVersion, FGuid("BC673DED3CEC4D72B10DC7336930BF9B"));
 	return SystemGuids;
 }
diff --git a/Engine/Source/Runtime/Engine/Private/Materials/MaterialExpressionSubstrate.cpp b/Engine/Source/Runtime/Engine/Private/Materials/MaterialExpressionSubstrate.cpp
--- a/Engine/Source/Runtime/Engine/Private/Materials/MaterialExpressionSubstrate.cpp
+++ b/Engine/Source/Runtime/Engine/Private/Materials/MaterialExpressionSubstrate.cpp
@@ -250,7 +250,7 @@ int32 UMaterialExpressionSubstrateShadingModels::CompileCommon(class FMaterialCo
 	int32 TangentCodeChunk = INDEX_NONE;
 	if (bHasAnisotropy)
 	{
-		LEGACY_DIRECT_ATTRIBUTE_MAPPING(TangentCodeChunkTmp, MP_Tangent, CompileWithDefaultTangentWS(Compiler, Normal));
+		LEGACY_DIRECT_ATTRIBUTE_MAPPING(TangentCodeChunkTmp, MP_Tangent, CompileWithDefaultTangentWS(Compiler, Tangent));
 		TangentCodeChunk = TangentCodeChunkTmp;
 	}
 	const FSubstrateRegisteredSharedLocalBasis NewRegisteredSharedLocalBasis = SubstrateCompilationInfoCreateSharedLocalBasis(Compiler, NormalCodeChunk, TangentCodeChunk);
