Added AddHead and AddOrInsertBefore to intrusive list.

diff --git a/Engine/Source/Runtime/Core/Internal/IO/IoContainers.h b/Engine/Source/Runtime/Core/Internal/IO/IoContainers.h
--- a/Engine/Source/Runtime/Core/Internal/IO/IoContainers.h
+++ b/Engine/Source/Runtime/Core/Internal/IO/IoContainers.h
@@ -74,6 +74,22 @@ public:
 		return *this;
 	}
 
+	void AddHead(ElementType* Element)
+	{
+		check(Element != nullptr && TypeTraits::GetNext(Element) == nullptr);
+
+		if (Tail == nullptr)
+		{
+			check(Head == nullptr);
+			Tail = Element;
+		}
+		else
+		{
+			TypeTraits::SetNext(Element, Head);
+		}
+		Head = Element;
+	}
+
 	void AddTail(ElementType* Element)
 	{
 		check(Element != nullptr && TypeTraits::GetNext(Element) == nullptr);
@@ -176,6 +192,50 @@ public:
 		return true;
 	}
 
+	/**
+	 * Adds or inserts a new item before the existing item returned by the user defined predicate.
+	 *
+	 * @param Element	The new item to add
+	 * @param Predicate	User defined predicate that should return true to insert the new item before the existing item
+	 *
+	 * @note: The new item gets added to the tail if the list is empty or if the predicate
+	 *		  returns false for the last element.
+	 */
+	template<typename PredicateType>
+	void AddOrInsertBefore(ElementType* Element, PredicateType Predicate)
+	{
+		check(Element != nullptr && TypeTraits::GetNext(Element) == nullptr);
+
+		if (Tail == nullptr || !Predicate(*Element, *Tail))
+		{
+			AddTail(Element);
+			return;
+		}
+
+		if (Predicate(*Element, *Head))
+		{
+			AddHead(Element);
+			return;
+		}
+
+		ElementType* It = Head;
+		ElementType* NextElement = TypeTraits::GetNext(It);
+		while (NextElement != nullptr)
+		{
+			if (Predicate(*Element, *NextElement))
+			{
+				TypeTraits::SetNext(It, Element);
+				TypeTraits::SetNext(Element, NextElement);
+				return;
+			}
+
+			It = NextElement;
+			NextElement = TypeTraits::GetNext(It);
+		}
+
+		AddTail(Element);
+	}
+
 	bool				IsEmpty() const { return Head == nullptr; }
 	ElementType*		GetHead()		{ return Head; }
 	const ElementType*	GetHead() const { return Head; }
