Added checks for non-trivially-relocatable types being used in containers, particularly FGCObject-derived types.

diff --git a/Engine/Source/Runtime/Core/Public/Containers/Array.h b/Engine/Source/Runtime/Core/Public/Containers/Array.h
--- a/Engine/Source/Runtime/Core/Public/Containers/Array.h
+++ b/Engine/Source/Runtime/Core/Public/Containers/Array.h
@@ -29,6 +29,7 @@
 #include "Templates/Sorting.h"
 #include "Templates/AlignmentTemplates.h"
 #include "Traits/ElementType.h"
+#include "Traits/IsTriviallyRelocatable.h"
 
 #include <limits>
 #include <type_traits>
@@ -965,6 +966,8 @@ public:
 	/** Destructor. */
 	~TArray()
 	{
+		UE_STATIC_ASSERT_WARN(TIsTriviallyRelocatable_V<InElementType>, "TArray can only be used with trivially relocatable types");
+
 		DestructItems(GetData(), ArrayNum);
 
 		// note ArrayNum, ArrayMax and data pointer are not invalidated
diff --git a/Engine/Source/Runtime/Core/Public/Containers/ChunkedArray.h b/Engine/Source/Runtime/Core/Public/Containers/ChunkedArray.h
--- a/Engine/Source/Runtime/Core/Public/Containers/ChunkedArray.h
+++ b/Engine/Source/Runtime/Core/Public/Containers/ChunkedArray.h
@@ -7,7 +7,7 @@
 #include "Misc/IntrusiveUnsetOptionalState.h"
 #include "Templates/UnrealTypeTraits.h"
 #include "Containers/IndirectArray.h"
-
+#include "Traits/IsTriviallyRelocatable.h"
 
 namespace UE4ChunkedArray_Private
 {
@@ -98,6 +98,11 @@ public:
 		return *this;
 	}
 
+	~TChunkedArray()
+	{
+		UE_STATIC_ASSERT_WARN(TIsTriviallyRelocatable_V<InElementType>, "TChunkedArray can only be used with trivially relocatable types");
+	}
+
 	[[nodiscard]] TChunkedArray(const TChunkedArray&) = default;
 	TChunkedArray& operator=(const TChunkedArray&) = default;
 
diff --git a/Engine/Source/Runtime/Core/Public/Containers/CompactSet.h b/Engine/Source/Runtime/Core/Public/Containers/CompactSet.h
--- a/Engine/Source/Runtime/Core/Public/Containers/CompactSet.h
+++ b/Engine/Source/Runtime/Core/Public/Containers/CompactSet.h
@@ -13,6 +13,7 @@
 #include "Templates/AlignmentTemplates.h"
 #include "Templates/RetainedRef.h"
 #include "Templates/TypeHash.h"
+#include "Traits/IsTriviallyRelocatable.h"
 
 #include <initializer_list>
 #include <type_traits>
@@ -340,6 +341,8 @@ public:
 
 	FORCEINLINE ~TCompactSet()
 	{
+		UE_STATIC_ASSERT_WARN(TIsTriviallyRelocatable_V<InElementType>, "TMapBase can only be used with trivially relocatable types");
+
 		Empty(0);
 	}
 
diff --git a/Engine/Source/Runtime/Core/Public/Containers/Deque.h b/Engine/Source/Runtime/Core/Public/Containers/Deque.h
--- a/Engine/Source/Runtime/Core/Public/Containers/Deque.h
+++ b/Engine/Source/Runtime/Core/Public/Containers/Deque.h
@@ -8,6 +8,7 @@
 #include "Math/UnrealMathUtility.h"
 #include "Misc/AssertionMacros.h"
 #include "Templates/MemoryOps.h"
+#include "Traits/IsTriviallyRelocatable.h"
 
 #include <initializer_list>
 
@@ -126,6 +127,8 @@ public:
 
 	~TDeque()
 	{
+		UE_STATIC_ASSERT_WARN(TIsTriviallyRelocatable_V<InElementType>, "TArray can only be used with trivially relocatable types");
+
 		Empty();
 	}
 
diff --git a/Engine/Source/Runtime/Core/Public/Containers/Map.h b/Engine/Source/Runtime/Core/Public/Containers/Map.h
--- a/Engine/Source/Runtime/Core/Public/Containers/Map.h
+++ b/Engine/Source/Runtime/Core/Public/Containers/Map.h
@@ -15,6 +15,7 @@
 #include "Templates/Tuple.h"
 #include "Templates/UnrealTemplate.h"
 #include "Templates/UnrealTypeTraits.h"
+#include "Traits/IsTriviallyRelocatable.h"
 #include <type_traits>
 
 #define ExchangeB(A,B) {bool T=A; A=B; B=T;}
@@ -160,6 +161,11 @@ protected:
 	{
 	}
 
+	~TMapBase()
+	{
+		UE_STATIC_ASSERT_WARN(TIsTriviallyRelocatable_V<KeyType> && TIsTriviallyRelocatable_V<ValueType>, "TMapBase can only be used with trivially relocatable types");
+	}
+
 	/////////////////////////////////////////////////
 	// Start - intrusive TOptional<TMapBase> state //
 	/////////////////////////////////////////////////
diff --git a/Engine/Source/Runtime/Core/Public/Containers/PagedArray.h b/Engine/Source/Runtime/Core/Public/Containers/PagedArray.h
--- a/Engine/Source/Runtime/Core/Public/Containers/PagedArray.h
+++ b/Engine/Source/Runtime/Core/Public/Containers/PagedArray.h
@@ -10,6 +10,7 @@
 #include "Misc/AssertionMacros.h"
 #include "Serialization/StructuredArchive.h"
 #include "Templates/UnrealTemplate.h"  // For GetData, GetNum
+#include "Traits/IsTriviallyRelocatable.h"
 
 #include <cstring>
 #include <initializer_list>
@@ -218,7 +219,10 @@ public:
 		CopyToEmpty(InSource, InSize);
 	}
 
-	~TPagedArray() = default;
+	~TPagedArray()
+	{
+		UE_STATIC_ASSERT_WARN(TIsTriviallyRelocatable_V<InElementType>, "TArray can only be used with trivially relocatable types");
+	}
 
 	TPagedArray& operator=(TPagedArray&& Other)
 	{
diff --git a/Engine/Source/Runtime/Core/Public/Containers/Set.h b/Engine/Source/Runtime/Core/Public/Containers/Set.h
--- a/Engine/Source/Runtime/Core/Public/Containers/Set.h
+++ b/Engine/Source/Runtime/Core/Public/Containers/Set.h
@@ -17,6 +17,7 @@
 #include "Templates/Sorting.h"
 #include "Templates/TypeHash.h"
 #include "Templates/UnrealTemplate.h"
+#include "Traits/IsTriviallyRelocatable.h"
 #include <initializer_list>
 #include <type_traits>
 
@@ -195,6 +196,8 @@ public:
 	/** Destructor. */
 	FORCEINLINE ~TSet()
 	{
+		UE_STATIC_ASSERT_WARN(TIsTriviallyRelocatable_V<InElementType>, "TSet can only be used with trivially relocatable types");
+
 		HashSize = 0;
 	}
 
diff --git a/Engine/Source/Runtime/Core/Public/Containers/SparseArray.h b/Engine/Source/Runtime/Core/Public/Containers/SparseArray.h
--- a/Engine/Source/Runtime/Core/Public/Containers/SparseArray.h
+++ b/Engine/Source/Runtime/Core/Public/Containers/SparseArray.h
@@ -17,9 +17,9 @@
 #include "Serialization/MemoryImageWriter.h"
 #include "Containers/UnrealString.h"
 #include "Misc/IntrusiveUnsetOptionalState.h"
+#include "Traits/IsTriviallyRelocatable.h"
 #include <initializer_list>
 
-
 #if UE_BUILD_SHIPPING || UE_BUILD_TEST
 	#define TSPARSEARRAY_RANGED_FOR_CHECKS 0
 #else
@@ -82,6 +82,8 @@ public:
 	/** Destructor. */
 	~TSparseArray()
 	{
+		UE_STATIC_ASSERT_WARN(TIsTriviallyRelocatable_V<InElementType>, "TArray can only be used with trivially relocatable types");
+
 		// Destruct the elements in the array.
 		Empty();
 	}
diff --git a/Engine/Source/Runtime/Core/Public/Traits/IsTriviallyRelocatable.h b/Engine/Source/Runtime/Core/Public/Traits/IsTriviallyRelocatable.h
new file mode 100644
--- /dev/null
+++ b/Engine/Source/Runtime/Core/Public/Traits/IsTriviallyRelocatable.h
@@ -0,0 +1,35 @@
+// Copyright Epic Games, Inc. All Rights Reserved.
+
+#pragma once
+
+#include <type_traits>
+
+/**
+ * A trait which tests if a type is trivially relocatable.
+ *
+ * NOTE: In UE, all types are assumed to be trivially relocatable, so this defaults to true. The intent of this
+ *       trait is for types to specialize it as false when it is known that they are not trivially relocatable,
+ *       and for containers and other utility functions which use trivial relocation to to inspect this trait and
+ *       guard against use with types which are known to be not trivially relocatable.
+ */
+template <typename T>
+struct TIsTriviallyRelocatable
+{
+	// We allow bitwise relocation of types containing e.g. const and reference members, but we don't want to allow 
+	// relocation of types which are UE_NONCOPYABLE or equivalent.
+	//
+	// See note above as to why this is why it is. In C++26, we should be able to catch more
+	// by using this: https://www.open-std.org/jtc1/sc22/wg21/docs/papers/2024/p2786r4.pdf
+	//
+	// Should be the test below by default, but there are a lot of existing violations that will need to be fixed first.
+
+//	static inline constexpr bool Value = std::is_move_constructible_v<T>;
+	static inline constexpr bool Value = true;
+};
+
+template <typename T> struct TIsTriviallyRelocatable<const          T> : TIsTriviallyRelocatable<T> {};
+template <typename T> struct TIsTriviallyRelocatable<      volatile T> : TIsTriviallyRelocatable<T> {};
+template <typename T> struct TIsTriviallyRelocatable<const volatile T> : TIsTriviallyRelocatable<T> {};
+
+template <typename T>
+inline constexpr bool TIsTriviallyRelocatable_V = TIsTriviallyRelocatable<T>::Value;
diff --git a/Engine/Source/Runtime/CoreUObject/Public/UObject/GCObject.h b/Engine/Source/Runtime/CoreUObject/Public/UObject/GCObject.h
--- a/Engine/Source/Runtime/CoreUObject/Public/UObject/GCObject.h
+++ b/Engine/Source/Runtime/CoreUObject/Public/UObject/GCObject.h
@@ -9,6 +9,8 @@
 #pragma once
 
 #include "CoreMinimal.h"
+#include "Concepts/DerivedFrom.h"
+#include "Traits/IsTriviallyRelocatable.h"
 #include "UObject/ObjectMacros.h"
 #include "UObject/Object.h"
 #include "UObject/ObjectPtr.h"
@@ -219,3 +221,10 @@ FORCEINLINE_DEBUGGABLE FGCObject::FGCObject(EFlags Flags)
 		RegisterGCObject();
 	}
 }
+
+// Types derived from FGCObject are not trivially relocatable, so flag their use in containers and other contexts.
+template <UE::CDerivedFrom<FGCObject> T>
+struct TIsTriviallyRelocatable<T>
+{
+	static inline constexpr bool Value = false;
+};
