Added static_asserts to TValueOrError to ensure that it's not instantiated with references.

diff --git a/Engine/Source/Runtime/Core/Public/Templates/ValueOrError.h b/Engine/Source/Runtime/Core/Public/Templates/ValueOrError.h
--- a/Engine/Source/Runtime/Core/Public/Templates/ValueOrError.h
+++ b/Engine/Source/Runtime/Core/Public/Templates/ValueOrError.h
@@ -7,6 +7,7 @@
 #include "Misc/TVariant.h"
 #include "Templates/Tuple.h"
 #include "Templates/UnrealTemplate.h"
+#include <type_traits>
 
 template <typename... ArgTypes>
 struct TValueOrError_ValueProxy
@@ -55,6 +56,8 @@ template <typename... ArgTypes>
 template <typename ValueType, typename ErrorType>
 class TValueOrError
 {
+	static_assert(!std::is_reference_v<ValueType> && !std::is_reference_v<ErrorType>, "ValueType and ErrorType cannot be references");
+
 	/** Wrap the error type to allow ValueType and ErrorType to be the same type. */
 	struct FWrapErrorType
 	{
@@ -188,6 +191,8 @@ private:
 template <typename ValueType>
 class TValueOrError<ValueType, void>
 {
+	static_assert(!std::is_reference_v<ValueType>, "ValueType cannot be a reference");
+
 	template <typename... ArgTypes, uint32... ArgIndices>
 	[[nodiscard]] inline TValueOrError(TValueOrError_ValueProxy<ArgTypes...>&& Proxy, TIntegerSequence<uint32, ArgIndices...>)
 		: Value(InPlace, MoveTemp(Proxy.Args).template Get<ArgIndices>()...)
@@ -258,6 +263,8 @@ private:
 template <typename ErrorType>
 class TValueOrError<void, ErrorType>
 {
+	static_assert(!std::is_reference_v<ErrorType>, "ErrorType cannot be a reference");
+
 	template <typename... ArgTypes, uint32... ArgIndices>
 	[[nodiscard]] inline TValueOrError(TValueOrError_ErrorProxy<ArgTypes...>&& Proxy, TIntegerSequence<uint32, ArgIndices...>)
 		: Error(InPlace, MoveTemp(Proxy.Args).template Get<ArgIndices>()...)
