Added UE_REWRITE - a replacement for FORCEINLINE intended to be used on macro-like functions - rather than for optimization - by 'rewriting' a function call as a different expression.  Such functions are intended to remove a function entirely from the end user experience, including avoiding stepping into it and generating symbols for it.
Added some examples of its use.

diff --git a/Engine/Source/Runtime/Core/Public/Containers/Array.h b/Engine/Source/Runtime/Core/Public/Containers/Array.h
--- a/Engine/Source/Runtime/Core/Public/Containers/Array.h
+++ b/Engine/Source/Runtime/Core/Public/Containers/Array.h
@@ -1009,9 +1009,9 @@ public:
 	 *
 	 * @returns Pointer to first array entry or nullptr if ArrayMax == 0.
 	 */
-	UE_NODEBUG [[nodiscard]] FORCEINLINE const ElementType* GetData() const UE_LIFETIMEBOUND
+	[[nodiscard]] UE_REWRITE const ElementType* GetData() const UE_LIFETIMEBOUND
 	{
-		return (const ElementType*)AllocatorInstance.GetAllocation();
+		return const_cast<TArray*>(this)->GetData();
 	}
 
 	/** 
@@ -1105,7 +1105,7 @@ public:
 	 * @returns True if the array is empty.
 	 * @see Num
 	 */
-	[[nodiscard]] bool IsEmpty() const
+	[[nodiscard]] UE_REWRITE bool IsEmpty() const
 	{
 		return ArrayNum == 0;
 	}
@@ -1116,7 +1116,7 @@ public:
 	 * @returns Number of elements in array.
 	 * @see GetSlack
 	 */
-	UE_NODEBUG [[nodiscard]] FORCEINLINE SizeType Num() const
+	[[nodiscard]] UE_REWRITE SizeType Num() const
 	{
 		return ArrayNum;
 	}
@@ -1133,7 +1133,7 @@ public:
 	 * @returns Maximum number of elements in array.
 	 * @see GetSlack
 	 */
-	UE_NODEBUG [[nodiscard]] FORCEINLINE SizeType Max() const
+	[[nodiscard]] UE_REWRITE SizeType Max() const
 	{
 		return ArrayMax;
 	}
@@ -1156,10 +1156,9 @@ public:
 	 *
 	 * @returns Reference to indexed element.
 	 */
-	UE_NODEBUG [[nodiscard]] FORCEINLINE const ElementType& operator[](SizeType Index) const UE_LIFETIMEBOUND
+	[[nodiscard]] UE_REWRITE const ElementType& operator[](SizeType Index) const UE_LIFETIMEBOUND
 	{
-		RangeCheck(Index);
-		return GetData()[Index];
+		return (*const_cast<TArray*>(this))[Index];
 	}
 
 	/**
@@ -1224,18 +1223,9 @@ public:
 	{
 		return Last();
 	}
-
-	/**
-	 * Returns the top element, i.e. the last one.
-	 *
-	 * Const version of the above.
-	 *
-	 * @returns Reference to the top element.
-	 * @see Pop, Push
-	 */
-	UE_NODEBUG [[nodiscard]] FORCEINLINE const ElementType& Top() const UE_LIFETIMEBOUND
+	[[nodiscard]] UE_REWRITE const ElementType& Top() const UE_LIFETIMEBOUND
 	{
-		return Last();
+		return const_cast<TArray*>(this)->Top();
 	}
 
 	/**
@@ -1249,19 +1239,9 @@ public:
 		RangeCheck(ArrayNum - IndexFromTheEnd - 1);
 		return GetData()[ArrayNum - IndexFromTheEnd - 1];
 	}
-
-	/**
-	 * Returns n-th last element from the array.
-	 *
-	 * Const version of the above.
-	 *
-	 * @param IndexFromTheEnd (Optional) Index from the end of array (default = 0).
-	 * @returns Reference to n-th last element from the array.
-	 */
-	UE_NODEBUG [[nodiscard]] FORCEINLINE const ElementType& Last(SizeType IndexFromTheEnd = 0) const UE_LIFETIMEBOUND
+	[[nodiscard]] UE_REWRITE const ElementType& Last(SizeType IndexFromTheEnd = 0) const UE_LIFETIMEBOUND
 	{
-		RangeCheck(ArrayNum - IndexFromTheEnd - 1);
-		return GetData()[ArrayNum - IndexFromTheEnd - 1];
+		return const_cast<TArray*>(this)->Last(IndexFromTheEnd);
 	}
 
 	/**
@@ -1427,20 +1407,6 @@ public:
 		return INDEX_NONE;
 	}
 
-	/**
-	 * Finds an item by key (assuming the ElementType overloads operator== for
-	 * the comparison).
-	 *
-	 * @param Key The key to search by.
-	 * @returns Pointer to the first matching element, or nullptr if none is found.
-	 * @see Find
-	 */
-	template <typename KeyType>
-	UE_NODEBUG [[nodiscard]] FORCEINLINE const ElementType* FindByKey(const KeyType& Key) const
-	{
-		return const_cast<TArray*>(this)->FindByKey(Key);
-	}
-
 	/**
 	 * Finds an item by key (assuming the ElementType overloads operator== for
 	 * the comparison). Time Complexity: O(n), starts iteration from the beginning so better performance if Key is in the front
@@ -1462,24 +1428,17 @@ public:
 
 		return nullptr;
 	}
-
-	/**
-	 * Finds an element which matches a predicate functor.
-	 *
-	 * @param Pred The functor to apply to each element.
-	 * @returns Pointer to the first element for which the predicate returns true, or nullptr if none is found.
-	 * @see FilterByPredicate, ContainsByPredicate
-	 */
-	template <typename Predicate>
-	UE_NODEBUG [[nodiscard]] FORCEINLINE const ElementType* FindByPredicate(Predicate Pred) const
+	template <typename KeyType>
+	[[nodiscard]] UE_REWRITE const ElementType* FindByKey(const KeyType& Key) const
 	{
-		return const_cast<TArray*>(this)->FindByPredicate(Pred);
+		return const_cast<TArray*>(this)->FindByKey(Key);
 	}
 
 	/**
 	 * Finds an element which matches a predicate functor.
 	 *
-	 * @param Pred The functor to apply to each element. true, or nullptr if none is found.
+	 * @param Pred The functor to apply to each element.
+	 * @returns Pointer to the first element for which the predicate returns true, or nullptr if none is found.
 	 * @see FilterByPredicate, ContainsByPredicate
 	 */
 	template <typename Predicate>
@@ -1495,6 +1454,11 @@ public:
 
 		return nullptr;
 	}
+	template <typename Predicate>
+	[[nodiscard]] UE_REWRITE const ElementType* FindByPredicate(Predicate Pred) const
+	{
+		return const_cast<TArray*>(this)->FindByPredicate(Pred);
+	}
 
 	/**
 	 * Filters the elements in the array based on a predicate functor.
@@ -1570,7 +1534,7 @@ public:
 	 * @returns True if this array is NOT the same as OtherArray. False otherwise.
 	 */
 #if !PLATFORM_COMPILER_HAS_GENERATED_COMPARISON_OPERATORS
-	UE_NODEBUG [[nodiscard]] FORCEINLINE bool operator!=(const TArray& OtherArray) const
+	[[nodiscard]] UE_REWRITE bool operator!=(const TArray& OtherArray) const
 	{
 		return !(*this == OtherArray);
 	}
@@ -2868,7 +2832,7 @@ public:
 		TContainerElementTypeCompatibility<ElementType>::CopyingFromOtherType();
 		DestructItems(GetData(), ArrayNum);
 		MoveOrCopy(*this, Other, ArrayMax);
-		return *this; 
+		return *this;
 	}
 
 	/**
diff --git a/Engine/Source/Runtime/Core/Public/HAL/Platform.h b/Engine/Source/Runtime/Core/Public/HAL/Platform.h
--- a/Engine/Source/Runtime/Core/Public/HAL/Platform.h
+++ b/Engine/Source/Runtime/Core/Public/HAL/Platform.h
@@ -704,6 +704,30 @@
 	#define RESTRICT __restrict						/* no alias hint */
 #endif
 
+// UE_REWRITE should be placed on functions which rewrite a function call as another expression. Like a macro,
+// but without using the preprocessor.
+//
+// Its intent is different from FORCEINLINE, which can be used for optimization reasons. UE_REWRITE should be applied
+// to functions which need to exist for API reasons, but where we don't want any interactions with the user experience,
+// like being forced to step through them in the debugger, or having symbols generated for them.
+//
+// Having a dedicated macro for this specific use case, separate from FORCEINLINE, makes it easier to reason about
+// and spot when it's being misapplied. As a rule of thumb, a rewrite function should consist of a single return statement
+// with at most one call to a non-rewrite function.
+//
+// Examples:
+// - Forwarding calls.
+// - Trivial getters.
+// - Iterators.
+// - Related overloads implemented in terms of each other, like const/non-const, lvalue/rvalue, operator==/operator!=.
+// - Lvalue/rvalue overloads where one is implemented in terms of the other.
+// - Cast functions like MoveTemp, Forward or ImplicitConv.
+// - operator-> on smart pointers.
+// - View constructors.
+#ifndef UE_REWRITE
+	#define UE_REWRITE UE_NODEBUG FORCEINLINE
+#endif
+
 /* Use before a function declaration to warn that callers should not ignore the return value */
 #ifdef __has_cpp_attribute
 	#if __has_cpp_attribute(nodiscard)
diff --git a/Engine/Source/Runtime/Core/Public/Templates/UnrealTemplate.h b/Engine/Source/Runtime/Core/Public/Templates/UnrealTemplate.h
--- a/Engine/Source/Runtime/Core/Public/Templates/UnrealTemplate.h
+++ b/Engine/Source/Runtime/Core/Public/Templates/UnrealTemplate.h
@@ -377,18 +377,6 @@ struct TGuardValue : private FNoncopyable
 		RefValue = OriginalValue;
 	}
 
-	/**
-	 * Overloaded dereference operator.
-	 * Provides read-only access to the original value of the data being tracked by this struct
-	 *
-	 * @return	a const reference to the original data value
-	 */
-	UE_DEPRECATED(5.5, "Use GetOriginalValue() instead.")
-	FORCEINLINE const AssignedType& operator*() const
-	{
-		return OriginalValue;
-	}
-
 	/**
 	 * Provides read-only access to the original value of the data being tracked by this struct
 	 *
@@ -569,7 +557,7 @@ template <typename T> struct TRemovePointer<T*> { typedef T Type; };
  * const object, because we would prefer to be informed when MoveTemp will have no effect.
  */
 template <typename T>
-UE_INTRINSIC_CAST UE_NODEBUG FORCEINLINE constexpr std::remove_reference_t<T>&& MoveTemp(T&& Obj) noexcept
+UE_INTRINSIC_CAST UE_REWRITE constexpr std::remove_reference_t<T>&& MoveTemp(T&& Obj) noexcept
 {
 	using CastType = std::remove_reference_t<T>;
 
@@ -587,7 +575,7 @@ UE_INTRINSIC_CAST UE_NODEBUG FORCEINLINE constexpr std::remove_reference_t<T>&&
  * where you can but not stop compilation.
  */
 template <typename T>
-UE_INTRINSIC_CAST UE_NODEBUG FORCEINLINE constexpr std::remove_reference_t<T>&& MoveTempIfPossible(T&& Obj) noexcept
+UE_INTRINSIC_CAST UE_REWRITE constexpr std::remove_reference_t<T>&& MoveTempIfPossible(T&& Obj) noexcept
 {
 	using CastType = std::remove_reference_t<T>;
 	return (CastType&&)Obj;
@@ -603,13 +591,13 @@ UE_INTRINSIC_CAST UE_NODEBUG FORCEINLINE constexpr std::remove_reference_t<T>&&
  * argument is not known in advance.
  */
 template <typename T>
-FORCEINLINE T CopyTemp(T& Val)
+UE_REWRITE T CopyTemp(T& Val)
 {
 	return const_cast<const T&>(Val);
 }
 
 template <typename T>
-FORCEINLINE T CopyTemp(const T& Val)
+UE_REWRITE T CopyTemp(const T& Val)
 {
 	return Val;
 }
@@ -623,7 +611,7 @@ FORCEINLINE T CopyTemp(const T& Val)
  * create a PR value without stopping compilation.
  */
 template <typename T>
-FORCEINLINE constexpr std::decay_t<T> CopyTempIfNecessary(T&& Val)
+UE_REWRITE constexpr std::decay_t<T> CopyTempIfNecessary(T&& Val)
 {
 	return (T&&)Val;
 }
@@ -633,13 +621,13 @@ FORCEINLINE constexpr std::decay_t<T> CopyTempIfNecessary(T&& Val)
  * This is UE's equivalent of std::forward.
  */
 template <typename T>
-UE_INTRINSIC_CAST UE_NODEBUG FORCEINLINE constexpr T&& Forward(std::remove_reference_t<T>& Obj) noexcept
+UE_INTRINSIC_CAST UE_REWRITE constexpr T&& Forward(std::remove_reference_t<T>& Obj) noexcept
 {
 	return (T&&)Obj;
 }
 
 template <typename T>
-UE_INTRINSIC_CAST UE_NODEBUG FORCEINLINE constexpr T&& Forward(std::remove_reference_t<T>&& Obj) noexcept
+UE_INTRINSIC_CAST UE_REWRITE constexpr T&& Forward(std::remove_reference_t<T>&& Obj) noexcept
 {
 	return (T&&)Obj;
 }
@@ -676,7 +664,7 @@ constexpr inline void Swap(T& A, T& B)
 }
 
 template <typename T>
-constexpr void Exchange(T& A, T& B)
+UE_REWRITE constexpr void Exchange(T& A, T& B)
 {
 	Swap(A, B);
 }
@@ -687,7 +675,7 @@ constexpr void Exchange(T& A, T& B)
  * code to be generated.
  */
 template <typename T, typename ArgType>
-FORCEINLINE T StaticCast(ArgType&& Arg)
+UE_REWRITE T StaticCast(ArgType&& Arg)
 {
 	return static_cast<T>(Arg);
 }
@@ -792,7 +780,7 @@ T&& DeclVal();
  * @return The object converted to the specified type.
  */
 template <typename T>
-FORCEINLINE T ImplicitConv(typename TIdentity<T>::Type Obj)
+UE_REWRITE T ImplicitConv(typename TIdentity<T>::Type Obj)
 {
 	return Obj;
 }
@@ -806,7 +794,7 @@ template <
 	typename Base
 	UE_REQUIRES(std::is_convertible_v<std::remove_reference_t<T>*, const volatile Base*>)
 >
-UE_INTRINSIC_CAST UE_NODEBUG FORCEINLINE decltype(auto) ForwardAsBase(std::remove_reference_t<T>& Obj)
+UE_INTRINSIC_CAST UE_REWRITE decltype(auto) ForwardAsBase(std::remove_reference_t<T>& Obj)
 {
 	return (TCopyQualifiersAndRefsFromTo_T<T&&, Base>)Obj;
 }
@@ -816,7 +804,7 @@ template <
 	typename Base
 	UE_REQUIRES(std::is_convertible_v<std::remove_reference_t<T>*, const volatile Base*>)
 >
-UE_INTRINSIC_CAST UE_NODEBUG FORCEINLINE decltype(auto) ForwardAsBase(std::remove_reference_t<T>&& Obj)
+UE_INTRINSIC_CAST UE_REWRITE decltype(auto) ForwardAsBase(std::remove_reference_t<T>&& Obj)
 {
 	return (TCopyQualifiersAndRefsFromTo_T<T&&, Base>)Obj;
 }
diff --git a/Engine/Source/Runtime/Core/Public/Templates/ValueOrError.h b/Engine/Source/Runtime/Core/Public/Templates/ValueOrError.h
--- a/Engine/Source/Runtime/Core/Public/Templates/ValueOrError.h
+++ b/Engine/Source/Runtime/Core/Public/Templates/ValueOrError.h
@@ -32,13 +32,13 @@ struct TValueOrError_ErrorProxy
 };
 
 template <typename... ArgTypes>
-[[nodiscard]] FORCEINLINE TValueOrError_ValueProxy<ArgTypes...> MakeValue(ArgTypes&&... Args)
+[[nodiscard]] UE_REWRITE TValueOrError_ValueProxy<ArgTypes...> MakeValue(ArgTypes&&... Args)
 {
 	return TValueOrError_ValueProxy<ArgTypes...>(Forward<ArgTypes>(Args)...);
 }
 
 template <typename... ArgTypes>
-[[nodiscard]] FORCEINLINE TValueOrError_ErrorProxy<ArgTypes...> MakeError(ArgTypes&&... Args)
+[[nodiscard]] UE_REWRITE TValueOrError_ErrorProxy<ArgTypes...> MakeError(ArgTypes&&... Args)
 {
 	return TValueOrError_ErrorProxy<ArgTypes...>(Forward<ArgTypes>(Args)...);
 }
@@ -99,13 +99,13 @@ public:
 	}
 
 	/** Check whether a value is set. Prefer HasValue and HasError to this. !IsValid() does *not* imply HasError(). */
-	[[nodiscard]] inline bool IsValid() const
+	[[nodiscard]] UE_REWRITE bool IsValid() const
 	{
 		return Variant.template IsType<ValueType>();
 	}
 
 	/** Whether the error is set. An error does imply no value. No error does *not* imply that a value is set. */
-	[[nodiscard]] inline bool HasError() const
+	[[nodiscard]] UE_REWRITE bool HasError() const
 	{
 		return Variant.template IsType<FWrapErrorType>();
 	}
@@ -115,11 +115,11 @@ public:
 	{
 		return Variant.template Get<FWrapErrorType>().Error;
 	}
-	[[nodiscard]] inline const ErrorType& GetError() const & UE_LIFETIMEBOUND
+	[[nodiscard]] UE_REWRITE const ErrorType& GetError() const & UE_LIFETIMEBOUND
 	{
 		return const_cast<TValueOrError*>(this)->GetError();
 	}
-	[[nodiscard]] inline ErrorType&& GetError() && UE_LIFETIMEBOUND
+	[[nodiscard]] UE_REWRITE ErrorType&& GetError() && UE_LIFETIMEBOUND
 	{
 		return MoveTemp(this->GetError());
 	}
@@ -133,7 +133,7 @@ public:
 		}
 		return nullptr;
 	}
-	[[nodiscard]] inline const ErrorType* TryGetError() const UE_LIFETIMEBOUND
+	[[nodiscard]] UE_REWRITE const ErrorType* TryGetError() const UE_LIFETIMEBOUND
 	{
 		return const_cast<TValueOrError*>(this)->TryGetError();
 	}
@@ -147,7 +147,7 @@ public:
 	}
 
 	/** Whether the value is set. A value does imply no error. No value does *not* imply that an error is set. */
-	[[nodiscard]] inline bool HasValue() const
+	[[nodiscard]] UE_REWRITE bool HasValue() const
 	{
 		return Variant.template IsType<ValueType>();
 	}
@@ -157,21 +157,21 @@ public:
 	{
 		return Variant.template Get<ValueType>();
 	}
-	[[nodiscard]] inline const ValueType& GetValue() const & UE_LIFETIMEBOUND
+	[[nodiscard]] UE_REWRITE const ValueType& GetValue() const & UE_LIFETIMEBOUND
 	{
 		return const_cast<TValueOrError*>(this)->GetValue();
 	}
-	[[nodiscard]] inline ValueType&& GetValue() && UE_LIFETIMEBOUND
+	[[nodiscard]] UE_REWRITE ValueType&& GetValue() && UE_LIFETIMEBOUND
 	{
 		return MoveTemp(this->GetValue());
 	}
 
 	/** Access the value if it is set. */
-	[[nodiscard]] inline ValueType* TryGetValue() UE_LIFETIMEBOUND
+	[[nodiscard]] UE_REWRITE ValueType* TryGetValue() UE_LIFETIMEBOUND
 	{
 		return Variant.template TryGet<ValueType>();
 	}
-	[[nodiscard]] inline const ValueType* TryGetValue() const UE_LIFETIMEBOUND
+	[[nodiscard]] UE_REWRITE const ValueType* TryGetValue() const UE_LIFETIMEBOUND
 	{
 		return const_cast<TValueOrError*>(this)->TryGetValue();
 	}
@@ -213,13 +213,13 @@ public:
 	}
 
 	/** Whether the error is set. An error does imply no value. */
-	[[nodiscard]] inline bool HasError() const
+	[[nodiscard]] UE_REWRITE bool HasError() const
 	{
 		return !Value.IsSet();
 	}
 
 	/** Whether the value is set. A value does imply no error. */
-	[[nodiscard]] inline bool HasValue() const
+	[[nodiscard]] UE_REWRITE bool HasValue() const
 	{
 		return Value.IsSet();
 	}
@@ -229,11 +229,11 @@ public:
 	{
 		return Value.GetValue();
 	}
-	[[nodiscard]] inline const ValueType& GetValue() const & UE_LIFETIMEBOUND
+	[[nodiscard]] UE_REWRITE const ValueType& GetValue() const & UE_LIFETIMEBOUND
 	{
 		return const_cast<TValueOrError*>(this)->GetValue();
 	}
-	[[nodiscard]] inline ValueType&& GetValue() && UE_LIFETIMEBOUND
+	[[nodiscard]] UE_REWRITE ValueType&& GetValue() && UE_LIFETIMEBOUND
 	{
 		return MoveTemp(this->GetValue());
 	}
@@ -243,7 +243,7 @@ public:
 	{
 		return Value.GetPtrOrNull();
 	}
-	[[nodiscard]] inline const ValueType* TryGetValue() const UE_LIFETIMEBOUND
+	[[nodiscard]] UE_REWRITE const ValueType* TryGetValue() const UE_LIFETIMEBOUND
 	{
 		return const_cast<TValueOrError*>(this)->TryGetValue();
 	}
@@ -285,13 +285,13 @@ public:
 	}
 
 	/** Whether the value is set. A value does imply no error. */
-	[[nodiscard]] inline bool HasValue() const
+	[[nodiscard]] UE_REWRITE bool HasValue() const
 	{
 		return !Error.IsSet();
 	}
 
 	/** Whether the error is set. An error does imply no value. */
-	[[nodiscard]] inline bool HasError() const
+	[[nodiscard]] UE_REWRITE bool HasError() const
 	{
 		return Error.IsSet();
 	}
@@ -301,11 +301,11 @@ public:
 	{
 		return Error.GetValue();
 	}
-	[[nodiscard]] inline const ErrorType& GetError() const & UE_LIFETIMEBOUND
+	[[nodiscard]] UE_REWRITE const ErrorType& GetError() const & UE_LIFETIMEBOUND
 	{
 		return const_cast<TValueOrError*>(this)->GetError();
 	}
-	[[nodiscard]] inline ErrorType&& GetError() && UE_LIFETIMEBOUND
+	[[nodiscard]] UE_REWRITE ErrorType&& GetError() && UE_LIFETIMEBOUND
 	{
 		return MoveTemp(this->GetError());
 	}
@@ -315,7 +315,7 @@ public:
 	{
 		return Error.GetPtrOrNull();
 	}
-	[[nodiscard]] inline const ErrorType* TryGetError() const UE_LIFETIMEBOUND
+	[[nodiscard]] UE_REWRITE const ErrorType* TryGetError() const UE_LIFETIMEBOUND
 	{
 		return const_cast<TValueOrError*>(this)->TryGetError();
 	}
@@ -349,13 +349,13 @@ public:
 	}
 
 	/** Whether the error is set. An error does imply no value. */
-	[[nodiscard]] inline bool HasError() const
+	[[nodiscard]] UE_REWRITE bool HasError() const
 	{
 		return !bValue;
 	}
 
 	/** Whether the value is set. A value does imply no error. */
-	[[nodiscard]] inline bool HasValue() const
+	[[nodiscard]] UE_REWRITE bool HasValue() const
 	{
 		return bValue;
 	}
