Replaced some constraint workarounds with C++20 concepts.

diff --git a/Engine/Source/Runtime/Core/Public/Async/Future.h b/Engine/Source/Runtime/Core/Public/Async/Future.h
--- a/Engine/Source/Runtime/Core/Public/Async/Future.h
+++ b/Engine/Source/Runtime/Core/Public/Async/Future.h
@@ -13,7 +13,6 @@
 #include "HAL/PooledSyncEvent.h"
 #include "Misc/ScopeLock.h"
 #include "Templates/Models.h"
-#include "Templates/Requires.h"
 #include <type_traits>
 
 template<typename ResultType> class TFuture;
@@ -237,12 +236,6 @@ public:
 	 * @return The result as a const reference, or the same reference if the future holds a reference, or void if the future holds a void.
 	 * @note Not equivalent to std::future::get(). The future remains valid.
 	 */
-#ifndef __cpp_concepts
-	ConstResultType Get() const UE_LIFETIMEBOUND // this will give false positives when getting from an rvalue TFuture<T&> or TFuture<void> until C++20
-	{
-		return this->GetState()->GetResult();
-	}
-#else
 	ConstResultType Get() const
 		requires (!std::is_object_v<ResultType>)
 	{
@@ -253,7 +246,6 @@ public:
 	{
 		return this->GetState()->GetResult();
 	}
-#endif
 
 	/**
 	 * Checks whether this future object has its value set.
@@ -429,12 +421,6 @@ public:
 	 * @return The result as a non-const reference, or the same reference if the future holds a reference, or void if the future holds a void.
 	 * @note Not equivalent to std::future::get(). The future remains valid.
 	 */
-#ifndef __cpp_concepts
-	MutableResultType GetMutable() UE_LIFETIMEBOUND // this will give false positives when getting from an rvalue TFuture<T&> or TFuture<void> until C++20
-	{
-		return (MutableResultType)this->Get();
-	}
-#else
 	MutableResultType GetMutable()
 		requires (!std::is_object_v<ResultType>)
 	{
@@ -445,7 +431,6 @@ public:
 	{
 		return (MutableResultType)this->Get();
 	}
-#endif
 
 	/**
 	 * Consumes the future's result and invalidates the future.
@@ -615,42 +600,18 @@ public:
 	 *
 	 * @param Result The result value to set.
 	 */
-#ifndef __cpp_concepts
-	template <
-		typename DummyResultType = ResultType
-		UE_REQUIRES(!std::is_void_v<DummyResultType>)
-	>
-#endif
 	inline void SetValue(const SetType& Result)
-#ifdef __cpp_concepts
 		requires(!std::is_void_v<ResultType>)
-#endif
 	{
 		EmplaceValue(Result);
 	}
-#ifndef __cpp_concepts
-	template <
-		typename DummyResultType = ResultType
-		UE_REQUIRES(!std::is_void_v<DummyResultType> && !std::is_lvalue_reference_v<DummyResultType>)
-	>
-#endif
 	inline void SetValue(SetType&& Result)
-#ifdef __cpp_concepts
 		requires(!std::is_void_v<ResultType> && !std::is_lvalue_reference_v<ResultType>)
-#endif
 	{
 		EmplaceValue(MoveTemp(Result));
 	}
-#ifndef __cpp_concepts
-	template <
-		typename DummyResultType = ResultType
-		UE_REQUIRES(std::is_void_v<DummyResultType>)
-	>
-#endif
 	inline void SetValue()
-#ifdef __cpp_concepts
 		requires(std::is_void_v<ResultType>)
-#endif
 	{
 		EmplaceValue();
 	}
@@ -769,3 +730,7 @@ TPromise<ResultType> MakeFulfilledPromise(ArgTypes&&... Args)
 	Promise.EmplaceValue(Forward<ArgTypes>(Args)...);
 	return Promise;
 }
+
+#if UE_ENABLE_INCLUDE_ORDER_DEPRECATED_IN_5_7
+#include "Templates/Requires.h"
+#endif
\ No newline at end of file
diff --git a/Engine/Source/Runtime/Core/Public/Templates/MemoryOps.h b/Engine/Source/Runtime/Core/Public/Templates/MemoryOps.h
--- a/Engine/Source/Runtime/Core/Public/Templates/MemoryOps.h
+++ b/Engine/Source/Runtime/Core/Public/Templates/MemoryOps.h
@@ -4,9 +4,6 @@
 
 #include "CoreTypes.h"
 #include "HAL/UnrealMemory.h"
-#include "Templates/IsTriviallyCopyAssignable.h"
-#include "Templates/IsTriviallyCopyConstructible.h"
-#include "Templates/Requires.h"
 #include "Templates/UnrealTypeTraits.h"
 #include "Traits/UseBitwiseSwap.h"
 #include <new> // IWYU pragma: export
@@ -29,20 +26,14 @@ namespace UE::Core::Private::MemoryOps
  * @param	Elements	The address of the first memory location to construct at.
  * @param	Count		The number of elements to destruct.
  */
-template <
-	typename ElementType,
-	typename SizeType
-	UE_REQUIRES(sizeof(ElementType) > 0 && TIsZeroConstructType<ElementType>::Value) // the sizeof here should improve the error messages we get when we try to call this function with incomplete types
->
+template <typename ElementType, typename SizeType>
+	requires (sizeof(ElementType) > 0 && !!TIsZeroConstructType<ElementType>::Value) // the sizeof here should improve the error messages we get when we try to call this function with incomplete types
 FORCEINLINE void DefaultConstructItems(void* Address, SizeType Count)
 {
 	FMemory::Memset(Address, 0, sizeof(ElementType) * Count);
 }
-template <
-	typename ElementType,
-	typename SizeType
-	UE_REQUIRES(sizeof(ElementType) > 0 && !TIsZeroConstructType<ElementType>::Value) // the sizeof here should improve the error messages we get when we try to call this function with incomplete types
->
+template <typename ElementType, typename SizeType>
+	requires (sizeof(ElementType) > 0 && !TIsZeroConstructType<ElementType>::Value) // the sizeof here should improve the error messages we get when we try to call this function with incomplete types
 FORCENOINLINE void DefaultConstructItems(void* Address, SizeType Count)
 {
 	ElementType* Element = (ElementType*)Address;
@@ -85,19 +76,13 @@ FORCEINLINE constexpr void DestructItem(ElementType* Element)
  *
  * @note: This function is optimized for values of T, and so will not dynamically dispatch destructor calls if T's destructor is virtual.
  */
-template <
-	typename ElementType,
-	typename SizeType
-	UE_REQUIRES(sizeof(ElementType) > 0 && std::is_trivially_destructible_v<ElementType>) // the sizeof here should improve the error messages we get when we try to call this function with incomplete types
->
+template <typename ElementType, typename SizeType>
+	requires (sizeof(ElementType) > 0 && std::is_trivially_destructible_v<ElementType>) // the sizeof here should improve the error messages we get when we try to call this function with incomplete types
 FORCEINLINE constexpr void DestructItems(ElementType* Element, SizeType Count)
 {
 }
-template <
-	typename ElementType,
-	typename SizeType
-	UE_REQUIRES(sizeof(ElementType) > 0 && !std::is_trivially_destructible_v<ElementType>) // the sizeof here should improve the error messages we get when we try to call this function with incomplete types
->
+template <typename ElementType, typename SizeType>
+	requires (sizeof(ElementType) > 0 && !std::is_trivially_destructible_v<ElementType>) // the sizeof here should improve the error messages we get when we try to call this function with incomplete types
 FORCENOINLINE constexpr void DestructItems(ElementType* Element, SizeType Count)
 {
 	while (Count)
@@ -118,12 +103,8 @@ FORCENOINLINE constexpr void DestructItems(ElementType* Element, SizeType Count)
  * @param	Source		A pointer to the first argument to pass to the constructor.
  * @param	Count		The number of elements to copy.
  */
-template <
-	typename DestinationElementType,
-	typename SourceElementType,
-	typename SizeType
-	UE_REQUIRES(sizeof(DestinationElementType) > 0 && sizeof(SourceElementType) > 0 && TIsBitwiseConstructible<DestinationElementType, SourceElementType>::Value) // the sizeof here should improve the error messages we get when we try to call this function with incomplete types
->
+template <typename DestinationElementType, typename SourceElementType, typename SizeType>
+	requires (sizeof(DestinationElementType) > 0 && sizeof(SourceElementType) > 0 && !!TIsBitwiseConstructible<DestinationElementType, SourceElementType>::Value) // the sizeof here should improve the error messages we get when we try to call this function with incomplete types
 FORCEINLINE void ConstructItems(void* Dest, const SourceElementType* Source, SizeType Count)
 {
 	if (Count)
@@ -131,12 +112,8 @@ FORCEINLINE void ConstructItems(void* Dest, const SourceElementType* Source, Siz
 		FMemory::Memcpy(Dest, Source, sizeof(SourceElementType) * Count);
 	}
 }
-template <
-	typename DestinationElementType,
-	typename SourceElementType,
-	typename SizeType
-	UE_REQUIRES(sizeof(DestinationElementType) > 0 && sizeof(SourceElementType) > 0 && !TIsBitwiseConstructible<DestinationElementType, SourceElementType>::Value) // the sizeof here should improve the error messages we get when we try to call this function with incomplete types
->
+template <typename DestinationElementType, typename SourceElementType, typename SizeType>
+	requires (sizeof(DestinationElementType) > 0 && sizeof(SourceElementType) > 0 && !TIsBitwiseConstructible<DestinationElementType, SourceElementType>::Value) // the sizeof here should improve the error messages we get when we try to call this function with incomplete types
 FORCENOINLINE void ConstructItems(void* Dest, const SourceElementType* Source, SizeType Count)
 {
 	while (Count)
@@ -155,20 +132,14 @@ FORCENOINLINE void ConstructItems(void* Dest, const SourceElementType* Source, S
  * @param	Source		A pointer to the first item to assign.
  * @param	Count		The number of elements to assign.
  */
-template <
-	typename ElementType,
-	typename SizeType
-	UE_REQUIRES(sizeof(ElementType) > 0 && TIsTriviallyCopyAssignable<ElementType>::Value) // the sizeof here should improve the error messages we get when we try to call this function with incomplete types
->
+template <typename ElementType, typename SizeType>
+	requires (sizeof(ElementType) > 0 && std::is_trivially_copy_assignable_v<ElementType>) // the sizeof here should improve the error messages we get when we try to call this function with incomplete types
 FORCEINLINE void CopyAssignItems(ElementType* Dest, const ElementType* Source, SizeType Count)
 {
 	FMemory::Memcpy(Dest, Source, sizeof(ElementType) * Count);
 }
-template <
-	typename ElementType,
-	typename SizeType
-	UE_REQUIRES(sizeof(ElementType) > 0 && !TIsTriviallyCopyAssignable<ElementType>::Value) // the sizeof here should improve the error messages we get when we try to call this function with incomplete types
->
+template <typename ElementType, typename SizeType>
+	requires (sizeof(ElementType) > 0 && !std::is_trivially_copy_assignable_v<ElementType>) // the sizeof here should improve the error messages we get when we try to call this function with incomplete types
 FORCENOINLINE void CopyAssignItems(ElementType* Dest, const ElementType* Source, SizeType Count)
 {
 	while (Count)
@@ -225,12 +196,8 @@ FORCEINLINE void RelocateConstructItem(void* Dest, const SourceElementType* Sour
  * @param	Source		A pointer to the first item to relocate.
  * @param	Count		The number of elements to relocate.
  */
-template <
-	typename DestinationElementType,
-	typename SourceElementType,
-	typename SizeType
-	UE_REQUIRES(sizeof(DestinationElementType) > 0 && sizeof(SourceElementType) > 0 && UE::Core::Private::MemoryOps::TCanBitwiseRelocate_V<DestinationElementType, SourceElementType>) // the sizeof here should improve the error messages we get when we try to call this function with incomplete types
->
+template <typename DestinationElementType, typename SourceElementType, typename SizeType>
+	requires (sizeof(DestinationElementType) > 0 && sizeof(SourceElementType) > 0 && UE::Core::Private::MemoryOps::TCanBitwiseRelocate_V<DestinationElementType, SourceElementType>) // the sizeof here should improve the error messages we get when we try to call this function with incomplete types
 FORCEINLINE void RelocateConstructItems(void* Dest, SourceElementType* Source, SizeType Count)
 {
 	static_assert(!std::is_const_v<SourceElementType>, "RelocateConstructItems: Source cannot be const");
@@ -245,12 +212,8 @@ FORCEINLINE void RelocateConstructItems(void* Dest, SourceElementType* Source, S
 
 	FMemory::Memmove(Dest, Source, sizeof(SourceElementType) * Count);
 }
-template <
-	typename DestinationElementType,
-	typename SourceElementType,
-	typename SizeType
-	UE_REQUIRES(sizeof(DestinationElementType) > 0 && sizeof(SourceElementType) > 0 && !UE::Core::Private::MemoryOps::TCanBitwiseRelocate_V<DestinationElementType, SourceElementType>) // the sizeof here should improve the error messages we get when we try to call this function with incomplete types
->
+template <typename DestinationElementType, typename SourceElementType, typename SizeType>
+	requires (sizeof(DestinationElementType) > 0 && sizeof(SourceElementType) > 0 && !UE::Core::Private::MemoryOps::TCanBitwiseRelocate_V<DestinationElementType, SourceElementType>) // the sizeof here should improve the error messages we get when we try to call this function with incomplete types
 FORCENOINLINE void RelocateConstructItems(void* Dest, SourceElementType* Source, SizeType Count)
 {
 	static_assert(!std::is_const_v<SourceElementType>, "RelocateConstructItems: Source cannot be const");
@@ -274,20 +237,14 @@ FORCENOINLINE void RelocateConstructItems(void* Dest, SourceElementType* Source,
  * @param	Source		A pointer to the first item to move from.
  * @param	Count		The number of elements to move.
  */
-template <
-	typename ElementType,
-	typename SizeType
-	UE_REQUIRES(sizeof(ElementType) > 0 && TIsTriviallyCopyConstructible<ElementType>::Value) // the sizeof here should improve the error messages we get when we try to call this function with incomplete types
->
+template <typename ElementType, typename SizeType>
+	requires (sizeof(ElementType) > 0 && std::is_trivially_copy_constructible_v<ElementType>) // the sizeof here should improve the error messages we get when we try to call this function with incomplete types
 FORCEINLINE void MoveConstructItems(void* Dest, const ElementType* Source, SizeType Count)
 {
 	FMemory::Memmove(Dest, Source, sizeof(ElementType) * Count);
 }
-template <
-	typename ElementType,
-	typename SizeType
-	UE_REQUIRES(sizeof(ElementType) > 0 && !TIsTriviallyCopyConstructible<ElementType>::Value) // the sizeof here should improve the error messages we get when we try to call this function with incomplete types
->
+template <typename ElementType, typename SizeType>
+	requires (sizeof(ElementType) > 0 && !std::is_trivially_copy_constructible_v<ElementType>) // the sizeof here should improve the error messages we get when we try to call this function with incomplete types
 FORCENOINLINE void MoveConstructItems(void* Dest, const ElementType* Source, SizeType Count)
 {
 	while (Count)
@@ -306,20 +263,14 @@ FORCENOINLINE void MoveConstructItems(void* Dest, const ElementType* Source, Siz
  * @param	Source		A pointer to the first item to move assign.
  * @param	Count		The number of elements to move assign.
  */
-template <
-	typename ElementType,
-	typename SizeType
-	UE_REQUIRES(sizeof(ElementType) > 0 && TIsTriviallyCopyAssignable<ElementType>::Value) // the sizeof here should improve the error messages we get when we try to call this function with incomplete types
->
+template <typename ElementType, typename SizeType>
+	requires (sizeof(ElementType) > 0 && std::is_trivially_copy_assignable_v<ElementType>) // the sizeof here should improve the error messages we get when we try to call this function with incomplete types
 FORCEINLINE void MoveAssignItems(ElementType* Dest, const ElementType* Source, SizeType Count)
 {
 	FMemory::Memmove(Dest, Source, sizeof(ElementType) * Count);
 }
-template <
-	typename ElementType,
-	typename SizeType
-	UE_REQUIRES(sizeof(ElementType) > 0 && !TIsTriviallyCopyAssignable<ElementType>::Value) // the sizeof here should improve the error messages we get when we try to call this function with incomplete types
->
+template <typename ElementType, typename SizeType>
+	requires (sizeof(ElementType) > 0 && !std::is_trivially_copy_assignable_v<ElementType>) // the sizeof here should improve the error messages we get when we try to call this function with incomplete types
 FORCENOINLINE void MoveAssignItems(ElementType* Dest, const ElementType* Source, SizeType Count)
 {
 	while (Count)
@@ -331,20 +282,14 @@ FORCENOINLINE void MoveAssignItems(ElementType* Dest, const ElementType* Source,
 	}
 }
 
-template <
-	typename ElementType,
-	typename SizeType
-	UE_REQUIRES(sizeof(ElementType) > 0 && TTypeTraits<ElementType>::IsBytewiseComparable) // the sizeof here should improve the error messages we get when we try to call this function with incomplete types
->
+template <typename ElementType, typename SizeType>
+	requires (sizeof(ElementType) > 0 && !!TTypeTraits<ElementType>::IsBytewiseComparable) // the sizeof here should improve the error messages we get when we try to call this function with incomplete types
 FORCEINLINE bool CompareItems(const ElementType* A, const ElementType* B, SizeType Count)
 {
 	return !Count || !FMemory::Memcmp(A, B, sizeof(ElementType) * Count);
 }
-template <
-	typename ElementType,
-	typename SizeType
-	UE_REQUIRES(sizeof(ElementType) > 0 && !TTypeTraits<ElementType>::IsBytewiseComparable) // the sizeof here should improve the error messages we get when we try to call this function with incomplete types
->
+template <typename ElementType, typename SizeType>
+	requires (sizeof(ElementType) > 0 && !TTypeTraits<ElementType>::IsBytewiseComparable) // the sizeof here should improve the error messages we get when we try to call this function with incomplete types
 FORCENOINLINE bool CompareItems(const ElementType* A, const ElementType* B, SizeType Count)
 {
 	while (Count)
@@ -362,6 +307,8 @@ FORCENOINLINE bool CompareItems(const ElementType* A, const ElementType* B, Size
 	return true;
 }
 
-#if UE_ENABLE_INCLUDE_ORDER_DEPRECATED_IN_5_5
-#include "Templates/IsTriviallyDestructible.h"
+#if UE_ENABLE_INCLUDE_ORDER_DEPRECATED_IN_5_7
+#include "Templates/IsTriviallyCopyAssignable.h"
+#include "Templates/IsTriviallyCopyConstructible.h"
+#include "Templates/Requires.h"
 #endif
diff --git a/Engine/Source/Runtime/Core/Public/Templates/Tuple.h b/Engine/Source/Runtime/Core/Public/Templates/Tuple.h
--- a/Engine/Source/Runtime/Core/Public/Templates/Tuple.h
+++ b/Engine/Source/Runtime/Core/Public/Templates/Tuple.h
@@ -76,9 +76,7 @@ namespace UE::Core::Private::Tuple
 		}
 
 		TTupleBaseElement()
-#ifdef __cpp_concepts // A C++17 fix was deemed too complex
 			requires (std::is_default_constructible_v<T>)
-#endif
 			: Value()
 		{
 		}
@@ -104,9 +102,7 @@ namespace UE::Core::Private::Tuple
 		}
 
 		TTupleBaseElement()
-#ifdef __cpp_concepts // A C++17 fix was deemed too complex
 			requires (std::is_default_constructible_v<T>)
-#endif
 			: Key()
 		{
 		}
diff --git a/Engine/Source/Runtime/Core/Public/Traits/IsImplicitlyConstructible.h b/Engine/Source/Runtime/Core/Public/Traits/IsImplicitlyConstructible.h
--- a/Engine/Source/Runtime/Core/Public/Traits/IsImplicitlyConstructible.h
+++ b/Engine/Source/Runtime/Core/Public/Traits/IsImplicitlyConstructible.h
@@ -12,23 +12,7 @@ namespace UE::Core::Private
 		template <typename U>
 		static void TakesU(U&&);
 
-#if defined(__cpp_concepts)
 		static constexpr bool value = (!std::is_aggregate_v<T> || sizeof...(ArgTypes) == 0) && requires{ TakesU<T>({ std::declval<ArgTypes>()... }); };
-#else
-		template <typename U, typename, typename... OtherArgTypes>
-		struct TDetect
-		{
-			static constexpr bool value = false;
-		};
-		template <typename U, typename... OtherArgTypes>
-		struct TDetect<U, decltype(TakesU<U>({ std::declval<OtherArgTypes>()... })), OtherArgTypes...>
-		{
-			static constexpr bool value = true;
-		};
-
-		static constexpr bool value = (!std::is_aggregate_v<T> || sizeof...(ArgTypes) == 0) && TDetect<T, void, ArgTypes...>::value;
-		static constexpr bool Value = value;
-#endif
 	};
 	template <typename T, typename ArgType>
 	struct TIsImplicitlyConstructibleImpl<T, ArgType>
