[Backout] - CL43468202

While the change is correct, this unfortunately causes some linux linker relocation failures. Backing out to unblock builds.

#rnx
Original CL Desc
-----------------------------------------------------------------
made TArray, TSet, and FPThreadsRecursiveMutex, default constructors constexpr

diff --git a/Engine/Source/Runtime/Core/Public/Containers/Array.h b/Engine/Source/Runtime/Core/Public/Containers/Array.h
--- a/Engine/Source/Runtime/Core/Public/Containers/Array.h
+++ b/Engine/Source/Runtime/Core/Public/Containers/Array.h
@@ -674,7 +674,7 @@ public:
 	/**
 	 * Constructor, initializes element number counters.
 	 */
-	[[nodiscard]] FORCEINLINE constexpr TArray()
+	[[nodiscard]] FORCEINLINE TArray()
 		: ArrayNum(0)
 		, ArrayMax(AllocatorInstance.GetInitialCapacity())
 	{
diff --git a/Engine/Source/Runtime/Core/Public/Containers/BitArray.h b/Engine/Source/Runtime/Core/Public/Containers/BitArray.h
--- a/Engine/Source/Runtime/Core/Public/Containers/BitArray.h
+++ b/Engine/Source/Runtime/Core/Public/Containers/BitArray.h
@@ -370,7 +370,7 @@ public:
 	template<typename,typename,bool>
 	friend class TConstDualSetBitIterator;
 
-	[[nodiscard]] constexpr TBitArray()
+	[[nodiscard]] TBitArray()
 	:	NumBits(0)
 	,	MaxBits(AllocatorInstance.GetInitialCapacity() * NumBitsPerDWORD)
 	{
diff --git a/Engine/Source/Runtime/Core/Public/Containers/ContainerAllocationPolicies.h b/Engine/Source/Runtime/Core/Public/Containers/ContainerAllocationPolicies.h
--- a/Engine/Source/Runtime/Core/Public/Containers/ContainerAllocationPolicies.h
+++ b/Engine/Source/Runtime/Core/Public/Containers/ContainerAllocationPolicies.h
@@ -639,7 +639,9 @@ public:
 
 	public:
 		/** Default constructor. */
-		constexpr ForAnyElementType() = default;
+		ForAnyElementType()
+			: Data(nullptr)
+		{}
 
 		/**
 		 * Moves the state of another allocator into this one.  The allocator can be different.
@@ -780,7 +782,7 @@ public:
 			return !!Data;
 		}
 
-		constexpr SizeType GetInitialCapacity() const
+		SizeType GetInitialCapacity() const
 		{
 			return 0;
 		}
@@ -804,7 +806,7 @@ public:
 		ForAnyElementType& operator=(const ForAnyElementType&);
 
 		/** A pointer to the container's elements. */
-		FScriptContainerElement* Data = nullptr;
+		FScriptContainerElement* Data;
 	};
 	
 	template<typename ElementType>
@@ -868,15 +870,8 @@ public:
 	public:
 
 		/** Default constructor. */
-		constexpr ForElementType()
+		ForElementType()
 		{
-			UE_IF_CONSTEVAL
-			{
-				for (auto& Element : InlineData)
-				{
-					Element.ConstInit();
-				}
-			}
 		}
 
 		/**
@@ -980,7 +975,7 @@ public:
 			return SecondaryData.HasAllocation();
 		}
 
-		constexpr SizeType GetInitialCapacity() const
+		SizeType GetInitialCapacity() const
 		{
 			return NumInlineElements;
 		}
diff --git a/Engine/Source/Runtime/Core/Public/Containers/Set.h b/Engine/Source/Runtime/Core/Public/Containers/Set.h
--- a/Engine/Source/Runtime/Core/Public/Containers/Set.h
+++ b/Engine/Source/Runtime/Core/Public/Containers/Set.h
@@ -175,7 +175,7 @@ private:
 
 public:
 	/** Initialization constructor. */
-	[[nodiscard]] FORCEINLINE constexpr TSet() = default;
+	[[nodiscard]] FORCEINLINE TSet() = default;
 
 	/** Copy constructor. */
 	[[nodiscard]] FORCEINLINE TSet(const TSet& Copy)
diff --git a/Engine/Source/Runtime/Core/Public/Containers/SparseArray.h b/Engine/Source/Runtime/Core/Public/Containers/SparseArray.h
--- a/Engine/Source/Runtime/Core/Public/Containers/SparseArray.h
+++ b/Engine/Source/Runtime/Core/Public/Containers/SparseArray.h
@@ -771,7 +771,11 @@ public:
 	}
 
 	/** Default constructor. */
-	[[nodiscard]] constexpr TSparseArray() = default;
+	[[nodiscard]] TSparseArray()
+		: FirstFreeIndex(-1)
+		, NumFreeIndices(0)
+	{
+	}
 
 	/** Move constructor. */
 	[[nodiscard]] TSparseArray(TSparseArray&& InCopy)
diff --git a/Engine/Source/Runtime/Core/Public/HAL/PThreadsRecursiveMutex.h b/Engine/Source/Runtime/Core/Public/HAL/PThreadsRecursiveMutex.h
--- a/Engine/Source/Runtime/Core/Public/HAL/PThreadsRecursiveMutex.h
+++ b/Engine/Source/Runtime/Core/Public/HAL/PThreadsRecursiveMutex.h
@@ -9,10 +9,7 @@
 #if PLATFORM_USE_PTHREADS
 
 #include <pthread.h>
-
-#if defined(PTHREAD_RECURSIVE_MUTEX_INITIALIZER) && !defined(PTHREAD_RECURSIVE_MUTEX_INITIALIZER_NP)
-	#define PTHREAD_RECURSIVE_MUTEX_INITIALIZER_NP PTHREAD_RECURSIVE_MUTEX_INITIALIZER
-#endif
+#include <errno.h>
 
 namespace UE
 {
@@ -28,9 +25,6 @@ public:
 	FPThreadsRecursiveMutex(const FPThreadsRecursiveMutex&) = delete;
 	FPThreadsRecursiveMutex& operator=(const FPThreadsRecursiveMutex&) = delete;
 
-#if defined(PTHREAD_RECURSIVE_MUTEX_INITIALIZER_NP)
-	constexpr FPThreadsRecursiveMutex() = default;
-#else
 	inline FPThreadsRecursiveMutex()
 	{
 		pthread_mutexattr_t MutexAttributes;
@@ -39,7 +33,6 @@ public:
 		pthread_mutex_init(&Mutex, &MutexAttributes);
 		pthread_mutexattr_destroy(&MutexAttributes);
 	}
-#endif
 
 	FORCEINLINE ~FPThreadsRecursiveMutex()
 	{
@@ -62,11 +55,7 @@ public:
 	}
 
 private:
-#if defined(PTHREAD_RECURSIVE_MUTEX_INITIALIZER_NP)
-	pthread_mutex_t Mutex = PTHREAD_RECURSIVE_MUTEX_INITIALIZER_NP;
-#else
 	pthread_mutex_t Mutex;
-#endif
 };
 
 } // UE
diff --git a/Engine/Source/Runtime/Core/Public/HAL/Platform.h b/Engine/Source/Runtime/Core/Public/HAL/Platform.h
--- a/Engine/Source/Runtime/Core/Public/HAL/Platform.h
+++ b/Engine/Source/Runtime/Core/Public/HAL/Platform.h
@@ -100,7 +100,6 @@
 #endif
 
 #if PLATFORM_APPLE
-	#define _DARWIN_C_SOURCE // include darwin extensions
 	#include <stddef.h> // needed for size_t
 #endif
 
diff --git a/Engine/Source/Runtime/Core/Public/Templates/TypeCompatibleBytes.h b/Engine/Source/Runtime/Core/Public/Templates/TypeCompatibleBytes.h
--- a/Engine/Source/Runtime/Core/Public/Templates/TypeCompatibleBytes.h
+++ b/Engine/Source/Runtime/Core/Public/Templates/TypeCompatibleBytes.h
@@ -62,15 +62,6 @@ struct TTypeCompatibleBytes
 		return (ElementType&&)*(ElementType*)this;
 	}
 
-	// Initializes to be constinit friendly
-	constexpr void ConstInit()
-	{
-		for (auto& Value : Pad)
-		{
-			Value = 0;
-		}
-	}
-	
 	// Emplaces an inner element.
 	// Note: no checks are possible to ensure that an element isn't already present.  DestroyUnchecked() must be called to end the element's lifetime.
 	template <typename... ArgTypes>
