Changed AddOrInsertBefore to not assume list is sorted.

#rnx

diff --git a/Engine/Source/Runtime/Core/Internal/IO/IoContainers.h b/Engine/Source/Runtime/Core/Internal/IO/IoContainers.h
--- a/Engine/Source/Runtime/Core/Internal/IO/IoContainers.h
+++ b/Engine/Source/Runtime/Core/Internal/IO/IoContainers.h
@@ -197,22 +197,13 @@ public:
 	 *
 	 * @param Element	The new item to add
 	 * @param Predicate	User defined predicate that should return true to insert the new item before the existing item
-	 *
-	 * @note: The new item gets added to the tail if the list is empty or if the predicate
-	 *		  returns false for the last element.
 	 */
 	template<typename PredicateType>
 	void AddOrInsertBefore(ElementType* Element, PredicateType Predicate)
 	{
 		check(Element != nullptr && TypeTraits::GetNext(Element) == nullptr);
 
-		if (Tail == nullptr || !Predicate(*Element, *Tail))
-		{
-			AddTail(Element);
-			return;
-		}
-
-		if (Predicate(*Element, *Head))
+		if (IsEmpty() || Predicate(*Element, *Head))
 		{
 			AddHead(Element);
 			return;
