- making r.MobileAntialiasing configurable in scalability ini files
- improved logging for making easier to detect  similar issues

diff --git a/Engine/Source/Runtime/Core/Private/HAL/ConsoleManager.cpp b/Engine/Source/Runtime/Core/Private/HAL/ConsoleManager.cpp
--- a/Engine/Source/Runtime/Core/Private/HAL/ConsoleManager.cpp
+++ b/Engine/Source/Runtime/Core/Private/HAL/ConsoleManager.cpp
@@ -4461,7 +4461,7 @@ static TAutoConsoleVariable<int32> CVarMobileDefaultAntiAliasing(
 		 " 3: MSAA \n"
 		 " 4: TSR(it will fallback to TAA if SupportsGen5TAA is disabled) \n"
 		 " 5: SMAA"),
-	ECVF_RenderThreadSafe | ECVF_Preview);
+	ECVF_RenderThreadSafe | ECVF_Scalability);
 
 static TAutoConsoleVariable<int32> CVarMobileSupportsGen4TAA(
 	TEXT("r.Mobile.SupportsGen4TAA"),
diff --git a/Engine/Source/Runtime/Core/Private/Misc/ConfigUtilities.cpp b/Engine/Source/Runtime/Core/Private/Misc/ConfigUtilities.cpp
--- a/Engine/Source/Runtime/Core/Private/Misc/ConfigUtilities.cpp
+++ b/Engine/Source/Runtime/Core/Private/Misc/ConfigUtilities.cpp
@@ -299,7 +299,9 @@ void OnSetCVarFromIniEntry(const TCHAR *IniFile, const TCHAR *Key, const TCHAR*
 		{
 			if(!CVar->TestFlags(EConsoleVariableFlags::ECVF_Scalability) && !CVar->TestFlags(EConsoleVariableFlags::ECVF_ScalabilityGroup))
 			{
-				ensureMsgf(false, TEXT("Scalability.ini can only set ECVF_Scalability console variables ('%s'='%s' is ignored)"), Key, Value);
+				FString Msg =  FString::Printf(TEXT("Scalability.ini can only set ECVF_Scalability console variables ('%s'='%s' is ignored)"), Key, Value);
+				ensureMsgf(false, TEXT("%s"), *Msg);
+				UE_LOG(LogConfig, Log, TEXT("%s"), *Msg);
 				return;
 			}
 		}
