Changed FTextKeyState to allocate key string data inline, and use only 2 bytes for len, saving 10 bytes per FKeyData.
Use UE_TEXTKEY_USE_INLINE_ALLOCATOR to enable it.
Text Key sizes are limited by UE_TEXTKEY_MAX_SIZE_BYTES (can be overriden if needed as long as it fits inside TChunkedPackedArray::BytesPerChunk).

#tests Tested on editor + replays for several platforms

diff --git a/Engine/Source/Runtime/Core/Private/Internationalization/TextKey.cpp b/Engine/Source/Runtime/Core/Private/Internationalization/TextKey.cpp
--- a/Engine/Source/Runtime/Core/Private/Internationalization/TextKey.cpp
+++ b/Engine/Source/Runtime/Core/Private/Internationalization/TextKey.cpp
@@ -10,7 +10,6 @@
 #include "HAL/LowLevelMemTracker.h"
 #include "Hash/CityHash.h"
 #include "Logging/LogMacros.h"
-#include "Misc/ByteSwap.h"
 #include "Misc/Guid.h"
 #include "Misc/LazySingleton.h"
 #include "Misc/ScopeRWLock.h"
@@ -29,14 +28,18 @@ DEFINE_LOG_CATEGORY_STATIC(LogTextKey, Log, All);
 	#define UE_TEXTKEY_SPLIT_GUID (1)
 #endif
 
-#ifndef UE_TEXTKEY_ELEMENTS_CHUNK_SIZE_BYTES
-	#define UE_TEXTKEY_ELEMENTS_CHUNK_SIZE_BYTES (32768)
-#endif
-
 #ifndef UE_TEXTKEY_ELEMENTS_MIN_HASH_SIZE
 	#define UE_TEXTKEY_ELEMENTS_MIN_HASH_SIZE (32768)
 #endif
 
+#ifndef UE_TEXTKEY_USE_INLINE_ALLOCATOR
+	#define UE_TEXTKEY_USE_INLINE_ALLOCATOR 1 //(0)
+#endif
+
+#ifndef UE_TEXTKEY_MAX_SIZE_BYTES
+	#define UE_TEXTKEY_MAX_SIZE_BYTES (1000)
+#endif
+
 #if UE_TEXTKEY_USE_UTF8
 using FTextKeyCharType = UTF8CHAR;
 #else
@@ -44,6 +47,74 @@ using FTextKeyCharType = TCHAR;
 #endif
 using FTextKeyStringView = TStringView<FTextKeyCharType>;
 
+
+// FName-like block-based variable sized allocator
+template<typename InElementType, uint32 ChunkCursorBits = 16> // Default to 64K per chunk
+class TChunkedPackedArray
+{
+public:
+	static constexpr uint32 BytesPerChunk = (1 << ChunkCursorBits);
+
+	int32 Add(const uint32 ElementSize)
+	{
+		if (Chunks.Num() == 0 || ChunkCursors[Chunks.Num() - 1] + ElementSize > BytesPerChunk)
+		{
+			AllocateNewChunk();
+		}
+
+		const uint32 CurrentChunk = Chunks.Num() - 1; 
+		int32 ResultIndex = (CurrentChunk << ChunkCursorBits) | ChunkCursors[CurrentChunk];
+		ChunkCursors[CurrentChunk] += ElementSize;
+		Count++;
+		
+		return ResultIndex;
+	}
+
+	int32 GetFirstIndex() const
+	{
+		return Num() > 0 ? 0 : INDEX_NONE;
+	}
+
+	int32 GetNextIndex(const int32 ElementIndex, const uint32 ElementSize) const
+	{
+		uint32 ChunkIndex = ElementIndex >> ChunkCursorBits;
+		uint32 ChunkCursor = ElementIndex & ((1 << ChunkCursorBits) - 1);
+		
+		ChunkCursor += ElementSize;
+
+		if (ChunkCursor >= ChunkCursors[ChunkIndex])
+		{
+			ChunkIndex++;
+			ChunkCursor = 0;
+		}
+		return (ChunkIndex << ChunkCursorBits) | ChunkCursor;
+	}
+
+	InElementType& Get(const int32 ElementIndex) const
+	{
+		uint32 ChunkIndex = ElementIndex >> ChunkCursorBits;
+		uint32 ChunkCursor = ElementIndex & ((1 << ChunkCursorBits) - 1);
+
+		return *reinterpret_cast<InElementType*>(&Chunks[ChunkIndex][ChunkCursor]);
+	}
+
+	uint32 Num() const
+	{
+		return Count;
+	}
+	
+private:
+	TArray<uint8*> Chunks;
+	TArray<uint32> ChunkCursors;
+	uint32 Count = 0;
+
+	void AllocateNewChunk()
+	{
+		Chunks.Add(static_cast<uint8*>(FMemory::Malloc(BytesPerChunk, 8)));
+		ChunkCursors.Add(0);
+	}
+};
+
 class FTextKeyState
 {
 public:
@@ -193,90 +264,115 @@ public:
 	}
 
 private:
-	struct FKeyData
+	// Used to represent Key Data to find or insert into the allocator
+	struct FKeyDataView
 	{
-		FKeyData()
-#if UE_TEXTKEY_USE_UTF8
-			: DataPtr(UTF8TEXT(""))
-#else
-			: DataPtr(TEXT(""))
-#endif
-			, StrLen(0)
-			, StrHash(0)
-		{
-		}
-
-		FKeyData(FTextKeyStringView InStr, const uint32 InStrHash)
-			: DataPtr(InStr.GetData())
-			, StrLen(InStr.Len())
+		FKeyDataView(FTextKeyStringView InStr, const uint32 InStrHash)
+			: StrLen(InStr.Len())
 			, StrHash(InStrHash)
+			, DataPtr(InStr.GetData())
 		{
 		}
 
-		FKeyData(const FGuid& InGuid, const uint32 InStrHash)
-			: DataPtr(&InGuid)
-			, StrLen(INDEX_NONE)
+		FKeyDataView(const FGuid& InGuid, const uint32 InStrHash)
+			: StrLen(INDEX_NONE)
 			, StrHash(InStrHash)
+			, DataPtr(&InGuid)
 		{
 		}
-
+		
 		bool IsStringView() const
 		{
 			return StrLen >= 0;
 		}
 
-		FTextKeyStringView AsStringView() const
+		bool IsGuid() const
 		{
-			checkf(IsStringView(), TEXT("AsStringView called on a FKeyData that doesn't reference a string!"));
-			return FTextKeyStringView(static_cast<const FTextKeyCharType*>(DataPtr), StrLen);
+			return StrLen == INDEX_NONE;
 		}
 
-		bool IsGuid() const
+		FTextKeyStringView AsStringView() const
 		{
-			return StrLen == INDEX_NONE;
+			checkf(IsStringView(), TEXT("AsStringView called on a FKeyDataView that doesn't reference a string!"));
+			return FTextKeyStringView(static_cast<const FTextKeyCharType*>(DataPtr), StrLen);
 		}
 
 		const FGuid& AsGuid() const
 		{
-			checkf(IsGuid(), TEXT("AsGuid called on a FKeyData that doesn't reference a GUID!"));
+			checkf(IsGuid(), TEXT("AsGuid called on a FKeyDataView that doesn't reference a GUID!"));
 			return *static_cast<const FGuid*>(DataPtr);
 		}
 
-		friend FORCEINLINE bool operator==(const FKeyData& A, const FKeyData& B)
+		/** Length of the string (in characters) if DataPtr points to a string buffer, or INDEX_NONE if it points to a FGuid */
+		int32 StrLen;
+		/** Hash of the source string this data was created from */
+		uint32 StrHash;
+		/** Internal data pointer; if StrLen >= 0 this points to a string buffer (which may not be null-terminated), or if INDEX_NONE points to a FGuid */
+		const void* DataPtr;
+	};
+
+	// Used to represent Key Data that exists inside the allocator
+	struct FKeyData
+	{
+		bool IsStringView() const
 		{
-			if (A.StrLen == B.StrLen)
-			{
-				if (A.StrLen >= 0)
-				{
-					// We can use Memcmp here as we know we're comparing two blocks of the same size and don't care about lexical ordering
-					return FMemory::Memcmp(A.DataPtr, B.DataPtr, A.StrLen * sizeof(FTextKeyCharType)) == 0;
-				}
-				else
-				{
-					return *static_cast<const FGuid*>(A.DataPtr) == *static_cast<const FGuid*>(B.DataPtr);
-				}
-			}
-			return false;
+			return StrLen >= 0;
 		}
 
-		friend FORCEINLINE bool operator!=(const FKeyData& A, const FKeyData& B)
+		FTextKeyStringView AsStringView() const
 		{
-			return !(A == B);
+			checkf(IsStringView(), TEXT("AsStringView called on a FKeyData that doesn't reference a string!"));
+			return FTextKeyStringView(reinterpret_cast<const FTextKeyCharType*>(DataPtr), StrLen);
+		}
+
+		bool IsGuid() const
+		{
+			return StrLen == INDEX_NONE;
 		}
 
-		friend FORCEINLINE uint32 GetTypeHash(const FKeyData& A)
+		const FGuid& AsGuid() const
 		{
-			return A.StrHash;
+			checkf(IsGuid(), TEXT("AsGuid called on a FKeyData that doesn't reference a GUID!"));
+			return *reinterpret_cast<const FGuid*>(DataPtr);
 		}
 
+		/** Index of the next element in this hash bucket */
+		int32 NextElementIndex = INDEX_NONE;
+		/** Hash of the source string this data was created from */
+		uint32 StrHash;
+		
+#if UE_TEXTKEY_USE_INLINE_ALLOCATOR
+		/** Length of the string (in characters) if DataPtr points to a string buffer, or INDEX_NONE if it points to a FGuid */
+		int16 StrLen;
+
 		/** Internal data pointer; if StrLen >= 0 this points to a string buffer (which may not be null-terminated), or if INDEX_NONE points to a FGuid */
-		const void* DataPtr;
+		uint8 DataPtr[0];
+#else
 		/** Length of the string (in characters) if DataPtr points to a string buffer, or INDEX_NONE if it points to a FGuid */
 		int32 StrLen;
-		/** Hash of the source string this data was created from */
-		uint32 StrHash;
+
+		const void* DataPtr;
+#endif
 	};
 
+	static constexpr uint32 KeyDataHeaderSizeBytes = STRUCT_OFFSET(FKeyData, DataPtr);
+
+	template<class TKeyDataA, class TKeyDataB>
+	friend FORCEINLINE bool operator==(const TKeyDataA& A, const TKeyDataB& B)
+	{
+		if (A.StrLen == B.StrLen)
+		{
+			if (A.StrLen >= 0)
+			{
+				// We can use Memcmp here as we know we're comparing two blocks of the same size and don't care about lexical ordering
+				return FMemory::Memcmp(A.DataPtr, B.DataPtr, A.StrLen * sizeof(FTextKeyCharType)) == 0;
+			}
+				
+			return A.AsGuid() == B.AsGuid();
+		}
+		return false;
+	}
+
 	void FindOrAddImpl(FStringView InStr, const uint32 InStrHash, FTextKey& OutTextKey)
 	{
 		int32 Index = INDEX_NONE;
@@ -285,14 +381,14 @@ private:
 			FGuid KeyGuid;
 			if (FGuid::ParseExact(InStr, EGuidFormats::Digits, KeyGuid)) // Only checking for EGuidFormats::Digits as that's the default of FGuid::ToString() as used by FText keys
 			{
-				const FKeyData KeyData = FKeyData(KeyGuid, InStrHash);
+				const FKeyDataView KeyData = FKeyDataView(KeyGuid, InStrHash);
 				Index = FindOrAddData(KeyData);
 			}
 			else
 #endif
 			{
 				auto ConvertedString = StrCast<FTextKeyCharType>(InStr.GetData(), InStr.Len());
-				const FKeyData KeyData = FKeyData(FTextKeyStringView(ConvertedString.Get(), ConvertedString.Length()), InStrHash);
+				const FKeyDataView KeyData = FKeyDataView(FTextKeyStringView(ConvertedString.Get(), ConvertedString.Length()), InStrHash);
 				Index = FindOrAddData(KeyData);
 			}
 		}
@@ -304,7 +400,7 @@ private:
 #endif
 	}
 
-	int32 FindOrAddData(const FKeyData& KeyData)
+	int32 FindOrAddData(const FKeyDataView& KeyData)
 	{
 		// Read-only
 		int32 NumElementsOnRead = 0;
@@ -337,13 +433,11 @@ private:
 			LLM_SCOPE_BYNAME(TEXT("Localization/TextKeys"));
 			if (KeyData.IsStringView())
 			{
-				FTextKeyStringView NewStr = FStaticAllocator::Add(KeyData.AsStringView());
-				return KeyDataAllocations.Add(FKeyData(NewStr, KeyData.StrHash));
+				return KeyDataAllocations.Add(KeyData.AsStringView(), KeyData.StrHash);
 			}
 			else
 			{
-				const FGuid* NewGuid = FStaticAllocator::Add(KeyData.AsGuid());
-				return KeyDataAllocations.Add(FKeyData(*NewGuid, KeyData.StrHash));
+				return KeyDataAllocations.Add(KeyData.AsGuid(), KeyData.StrHash);
 			}
 		}
 	}
@@ -358,31 +452,60 @@ private:
 			FreeHash();
 		}
 
-		int32 Add(const FKeyData& InKeyData)
+		int32 Add(const FTextKeyStringView& InKeyStringView, uint32 StrHash)
 		{
-			ConditionalRehash(Elements.Num() + 1);
+			const uint32 DataSizeBytes = InKeyStringView.Len() * sizeof(FTextKeyCharType);
+#if UE_TEXTKEY_USE_INLINE_ALLOCATOR
+			checkf(DataSizeBytes <= UE_TEXTKEY_MAX_SIZE_BYTES, TEXT("Key string size greater than UE_TEXTKEY_MAX_SIZE_BYTES"));
+#endif
+			
+			const int32 NewElementIndex = Add(DataSizeBytes, StrHash);
+			FKeyData& KeyData = Elements.Get(NewElementIndex);
 
-			const uint32 KeyDataHash = GetTypeHash(InKeyData);
-			const uint32 HashIndex = KeyDataHash & (HashSize - 1);
+#if UE_TEXTKEY_USE_INLINE_ALLOCATOR
+			KeyData.StrLen = static_cast<int16>(InKeyStringView.Len());
+			FMemory::Memcpy(KeyData.DataPtr, InKeyStringView.GetData(), DataSizeBytes);
+#else
+			KeyData.StrLen = InKeyStringView.Len();
+			
+			void* DataBuffer = GetPersistentLinearAllocator().Allocate(DataSizeBytes, alignof(FTextKeyCharType));
+			FMemory::Memcpy(DataBuffer, InKeyStringView.GetData(), DataSizeBytes);
+			KeyData.DataPtr = DataBuffer;
+#endif
+			
+			return NewElementIndex;
+		}
+		
+		int32 Add(const FGuid& InKeyGuid, uint32 StrHash)
+		{
+			constexpr uint32 DataSizeBytes = sizeof(FGuid);
+			const int32 NewElementIndex = Add(DataSizeBytes, StrHash);
 
-			const int32 NewElementIndex = Elements.Num();
-			Elements.AddElement(FElement{ InKeyData, Hash[HashIndex] });
-			Hash[HashIndex] = NewElementIndex;
+			FKeyData& KeyData = Elements.Get(NewElementIndex);
+			KeyData.StrLen = INDEX_NONE;
+#if UE_TEXTKEY_USE_INLINE_ALLOCATOR
+			FGuid* GuidPtr = reinterpret_cast<FGuid*>(KeyData.DataPtr);
+#else
+			FGuid* GuidPtr = static_cast<FGuid*>(GetPersistentLinearAllocator().Allocate(sizeof(FGuid), alignof(FGuid)));
+			KeyData.DataPtr = GuidPtr;
+#endif
+			*GuidPtr = InKeyGuid;
+			
 			return NewElementIndex;
 		}
 
-		int32 Find(const FKeyData& InKeyData) const
+		int32 Find(const FKeyDataView& InKeyData) const
 		{
 			if (HashSize > 0)
 			{
-				const uint32 KeyDataHash = GetTypeHash(InKeyData);
+				const uint32 KeyDataHash = InKeyData.StrHash;
 				const uint32 HashIndex = KeyDataHash & (HashSize - 1);
 
 				int32 ElementIndex = Hash[HashIndex];
 				while (ElementIndex != INDEX_NONE)
 				{
-					const FElement& Element = Elements[ElementIndex];
-					if (Element.Value == InKeyData)
+					const FKeyData& Element = Elements.Get(ElementIndex);
+					if (Element == InKeyData)
 					{
 						return ElementIndex;
 					}
@@ -394,7 +517,7 @@ private:
 
 		const FKeyData& Get(int32 InIndex) const
 		{
-			return Elements[InIndex].Value;
+			return Elements.Get(InIndex);
 		}
 
 		int32 Num() const
@@ -427,27 +550,66 @@ private:
 				HashSize = NewHashSize;
 				AllocateHash();
 
-				int32 ElementIndex = 0;
-				for (FElement& Element : Elements)
+				int32 ElementIndex = Elements.GetFirstIndex();
+				for (uint32 ElementCount = 0; ElementCount < Elements.Num(); ++ElementCount)
 				{
-					const uint32 KeyDataHash = GetTypeHash(Element.Value);
+					FKeyData& Element = Elements.Get(ElementIndex);
+					
+					const uint32 KeyDataHash = Element.StrHash;
 					const uint32 HashIndex = KeyDataHash & (HashSize - 1);
 
 					Element.NextElementIndex = Hash[HashIndex];
-					Hash[HashIndex] = ElementIndex++;
+					Hash[HashIndex] = ElementIndex;
+					
+					ElementIndex = Elements.GetNextIndex(ElementIndex, GetAllocationSize(Element));
 				}
 			}
 		}
 
-		struct FElement
+		static uint32 GetAllocationSize(const FKeyData& KeyData)
 		{
-			FKeyData Value;
-			/** Index of the next element in this hash bucket */
-			int32 NextElementIndex = INDEX_NONE;
-		};
-		static constexpr int32 ElementChunkSize = UE_TEXTKEY_ELEMENTS_CHUNK_SIZE_BYTES;
+			uint32 DataSize;
+			if (KeyData.IsStringView())
+			{
+				DataSize = KeyData.StrLen * sizeof(FTextKeyCharType);
+			}
+			else
+			{
+				DataSize = sizeof(FGuid);
+			}
+			return GetAllocationSize(DataSize);
+		}
+
+		static uint32 GetAllocationSize(const uint32 DataSize)
+		{
+#if UE_TEXTKEY_USE_INLINE_ALLOCATOR
+			return Align(KeyDataHeaderSizeBytes + DataSize, alignof(FKeyData)); 
+#else
+			return sizeof(FKeyData);
+#endif
+		}
+
+		int32 Add(const uint32 DataSize, uint32 StrHash)
+		{
+			ConditionalRehash(Elements.Num() + 1);
+
+			const uint32 HashIndex = StrHash & (HashSize - 1);
+			const int32 NewElementIndex = Elements.Add(GetAllocationSize(DataSize));
+			
+			FKeyData& NewElement = Elements.Get(NewElementIndex);
+			NewElement.StrHash = StrHash;
+			NewElement.NextElementIndex = Hash[HashIndex];
+			Hash[HashIndex] = NewElementIndex;
+			
+			return NewElementIndex;
+		}
+
 		/** Values; indices are referenced by the hash and FTextKey */
-		TChunkedArray<FElement, ElementChunkSize> Elements;
+		TChunkedPackedArray<FKeyData> Elements;
+
+		// Add in a bit of overhead for the FKeyData header
+		static_assert(UE_TEXTKEY_MAX_SIZE_BYTES + KeyDataHeaderSizeBytes <= TChunkedPackedArray<FKeyData>::BytesPerChunk,
+			"UE_TEXTKEY_MAX_SIZE_BYTES is too large for the storage array chunk size");
 
 		static constexpr int32 MinHashSize = UE_TEXTKEY_ELEMENTS_MIN_HASH_SIZE;
 		/** Current size of the hash; if this changes the hash must be rebuilt */
@@ -456,25 +618,6 @@ private:
 		int32* Hash = nullptr;
 	};
 
-	struct FStaticAllocator
-	{
-		static FTextKeyStringView Add(FTextKeyStringView InStr)
-		{
-			const int32 NumBytes = InStr.Len() * sizeof(FTextKeyCharType);
-			void* StringPtr = GetPersistentLinearAllocator().Allocate(NumBytes, alignof(FTextKeyCharType));
-			FMemory::Memcpy(StringPtr, InStr.GetData(), NumBytes);
-
-			return FTextKeyStringView((FTextKeyCharType*)StringPtr, InStr.Len());
-		}
-
-		static FGuid* Add(const FGuid& InGuid)
-		{
-			FGuid* GuidPtr = (FGuid*)GetPersistentLinearAllocator().Allocate(sizeof(InGuid), alignof(FGuid));
-			*GuidPtr = InGuid;
-			return GuidPtr;
-		}
-	};
-
 	mutable FTransactionallySafeRWLock DataRW;
 	FKeyDataAllocator KeyDataAllocations;
 
