Move MSVC headers to Microsoft folder as these only apply to Microsoft platforms

#rnx

diff --git a/Engine/Source/Runtime/Core/Public/MSVC/MSVCPlatform.h b/Engine/Source/Runtime/Core/Public/MSVC/MSVCPlatform.h
--- a/Engine/Source/Runtime/Core/Public/MSVC/MSVCPlatform.h
+++ b/Engine/Source/Runtime/Core/Public/MSVC/MSVCPlatform.h
@@ -6,31 +6,10 @@
 
 #pragma once
 
-// HEADER_UNIT_SKIP - Not included directly
+#include "Misc/CoreMiscDefines.h"
 
-#if _MSC_VER < 1920
-	#error "Compiler is expected to support if constexpr"
-#endif
-
-#if !defined(__cpp_fold_expressions)
-	#error "Compiler is expected to support fold expressions"
-#endif
-
-#define PLATFORM_RETURN_ADDRESS()	        _ReturnAddress()
-#define PLATFORM_RETURN_ADDRESS_POINTER()	_AddressOfReturnAddress()
+UE_DEPRECATED_HEADER(5.7, "Use Microsoft/MSVCPlatform.h instead.")
 
-// https://devblogs.microsoft.com/cppblog/improving-the-state-of-debug-performance-in-c/
-#if __has_cpp_attribute(msvc::intrinsic)
-#define UE_INTRINSIC_CAST [[msvc::intrinsic]]
+#if PLATFORM_MICROSOFT
+#include "Microsoft/MSVCPlatform.h"
 #endif
-
-// Ensure we can use this builtin - seems to be present on Clang 9, GCC 11 and MSVC 19.26
-#define PLATFORM_COMPILER_SUPPORTS_BUILTIN_BITCAST (_MSC_VER >= 1926)
-
-#ifdef __has_cpp_attribute
-	#if __has_cpp_attribute(msvc::lifetimebound)
-		#define UE_LIFETIMEBOUND [[msvc::lifetimebound]]
-	#endif
-#endif
-
-#define UE_NO_PROFILE_ATTRIBUTE
\ No newline at end of file
diff --git a/Engine/Source/Runtime/Core/Public/MSVC/MSVCPlatformCompilerPreSetup.h b/Engine/Source/Runtime/Core/Public/MSVC/MSVCPlatformCompilerPreSetup.h
--- a/Engine/Source/Runtime/Core/Public/MSVC/MSVCPlatformCompilerPreSetup.h
+++ b/Engine/Source/Runtime/Core/Public/MSVC/MSVCPlatformCompilerPreSetup.h
@@ -2,290 +2,12 @@
 
 #pragma once
 
-// HEADER_UNIT_SKIP - Cause cirular dependencies for header units
+// HEADER_UNIT_SKIP - Cause circular dependencies for header units
 
-#if defined(__clang__)
-	// UE4 uses a struct packing of 4 for Win32 and 8 for Win64, 
-	// and the default packing is 8 for Win32 and 16 for Win64
-	#ifdef _WIN64
-		#ifndef PRAGMA_PUSH_PLATFORM_DEFAULT_PACKING
-			#define PRAGMA_PUSH_PLATFORM_DEFAULT_PACKING \
-				_Pragma("pack(push)") \
-				_Pragma("pack(16)")
-		#endif // PRAGMA_PUSH_PLATFORM_DEFAULT_PACKING
-	#else // _WIN64
-		#ifndef PRAGMA_PUSH_PLATFORM_DEFAULT_PACKING
-			#define PRAGMA_PUSH_PLATFORM_DEFAULT_PACKING \
-				_Pragma("pack(push)") \
-				_Pragma("pack(8)")
-		#endif // PRAGMA_PUSH_PLATFORM_DEFAULT_PACKING
-	#endif // _WIN64
-	
-	#ifndef PRAGMA_POP_PLATFORM_DEFAULT_PACKING
-		#define PRAGMA_POP_PLATFORM_DEFAULT_PACKING \
-			_Pragma("pack(pop)")
-	#endif // PRAGMA_POP_PLATFORM_DEFAULT_PACKING
+#include "Misc/CoreMiscDefines.h"
 
-	#include "Clang/ClangPlatformCompilerPreSetup.h"
-#else
-	#ifndef PRAGMA_DISABLE_DEPRECATION_WARNINGS
-		#define PRAGMA_DISABLE_DEPRECATION_WARNINGS \
-			__pragma (warning(push)) \
-			__pragma (warning(disable: 4995)) /* 'function': name was marked as #pragma deprecated */ \
-			__pragma (warning(disable: 4996)) /* The compiler encountered a deprecated declaration. */
-	#endif // PRAGMA_DISABLE_DEPRECATION_WARNINGS
+UE_DEPRECATED_HEADER(5.7, "Use Microsoft/MSVCPlatformCompilerPreSetup.h instead.")
 
-	#ifndef PRAGMA_ENABLE_DEPRECATION_WARNINGS
-		#define PRAGMA_ENABLE_DEPRECATION_WARNINGS \
-			__pragma (warning(pop))
-	#endif // PRAGMA_ENABLE_DEPRECATION_WARNINGS
-
-	#ifndef PRAGMA_DISABLE_BUFFER_OVERRUN_WARNING
-		#define PRAGMA_DISABLE_BUFFER_OVERRUN_WARNING \
-			__pragma (warning(push)) \
-			__pragma (warning(disable:4789)) /* buffer 'identifier' of size N bytes will be overrun; M bytes will be written starting at offset L */
-	#endif // PRAGMA_DISABLE_BUFFER_OVERRUN_WARNING
-
-	#ifndef PRAGMA_ENABLE_BUFFER_OVERRUN_WARNING
-		#define PRAGMA_ENABLE_BUFFER_OVERRUN_WARNING \
-			__pragma (warning(pop))
-	#endif // PRAGMA_ENABLE_BUFFER_OVERRUN_WARNING
-
-	#ifndef PRAGMA_DISABLE_SHADOW_VARIABLE_WARNINGS
-		#define PRAGMA_DISABLE_SHADOW_VARIABLE_WARNINGS \
-			__pragma (warning(push)) \
-			__pragma (warning(disable: 4456)) /* declaration of 'LocalVariable' hides previous local declaration */ \
-			__pragma (warning(disable: 4457)) /* declaration of 'LocalVariable' hides function parameter */ \
-			__pragma (warning(disable: 4458)) /* declaration of 'LocalVariable' hides class member */ \
-			__pragma (warning(disable: 4459)) /* declaration of 'LocalVariable' hides global declaration */ \
-			__pragma (warning(disable: 6244)) /* local declaration of <variable> hides previous declaration at <line> of <file> */
-	#endif // PRAGMA_DISABLE_SHADOW_VARIABLE_WARNINGS
-
-	#ifndef PRAGMA_ENABLE_SHADOW_VARIABLE_WARNINGS
-		#define PRAGMA_ENABLE_SHADOW_VARIABLE_WARNINGS \
-			__pragma(warning(pop))
-	#endif // PRAGMA_ENABLE_SHADOW_VARIABLE_WARNINGS
-
-	#ifndef PRAGMA_DISABLE_UNREACHABLE_CODE_WARNINGS
-		#define PRAGMA_DISABLE_UNREACHABLE_CODE_WARNINGS \
-			__pragma (warning(push)) \
-			__pragma (warning(disable: 4702)) /* unreachable code */
-	#endif // PRAGMA_DISABLE_UNREACHABLE_CODE_WARNINGS
-	
-	#ifndef PRAGMA_RESTORE_UNREACHABLE_CODE_WARNINGS
-		#define PRAGMA_RESTORE_UNREACHABLE_CODE_WARNINGS \
-			__pragma (warning(pop))
-	#endif // PRAGMA_RESTORE_UNREACHABLE_CODE_WARNINGS
-
-	#ifndef PRAGMA_DISABLE_UNSAFE_TYPECAST_WARNINGS
-		#define PRAGMA_DISABLE_UNSAFE_TYPECAST_WARNINGS \
-			__pragma (warning(push)) \
-			__pragma (warning(disable: 4244)) /* 'argument': conversion from 'type1' to 'type2', possible loss of data */	\
-			__pragma (warning(disable: 4838)) /* 'argument': conversion from 'type1' to 'type2' requires a narrowing conversion */
-	#endif // PRAGMA_DISABLE_UNSAFE_TYPECAST_WARNINGS
-
-	#ifndef PRAGMA_ENABLE_UNSAFE_TYPECAST_WARNINGS
-		#define PRAGMA_ENABLE_UNSAFE_TYPECAST_WARNINGS \
-			UE_DEPRECATED_MACRO(5.0, "The PRAGMA_ENABLE_UNSAFE_TYPECAST_WARNINGS macro has been deprecated in favor of PRAGMA_RESTORE_UNSAFE_TYPECAST_WARNINGS. To force enable warnings use PRAGMA_FORCE_UNSAFE_TYPECAST_WARNINGS.")
-	#endif // PRAGMA_ENABLE_UNSAFE_TYPECAST_WARNINGS
-
-	#ifndef PRAGMA_FORCE_UNSAFE_TYPECAST_WARNINGS
-		#define PRAGMA_FORCE_UNSAFE_TYPECAST_WARNINGS \
-			__pragma (warning(push)) \
-			__pragma (warning(error: 4244)) /* 'argument': conversion from 'type1' to 'type2', possible loss of data */	\
-			__pragma (warning(error: 4838)) /* 'argument': conversion from 'type1' to 'type2' requires a narrowing conversion */
-	#endif // PRAGMA_FORCE_UNSAFE_TYPECAST_WARNINGS
-
-	#ifndef PRAGMA_RESTORE_UNSAFE_TYPECAST_WARNINGS
-		#define PRAGMA_RESTORE_UNSAFE_TYPECAST_WARNINGS \
-			__pragma(warning(pop))
-	#endif // PRAGMA_RESTORE_UNSAFE_TYPECAST_WARNINGS
-
-	#ifndef PRAGMA_DISABLE_SWITCH_UNHANDLED_ENUM_CASE_WARNINGS
-		#define PRAGMA_DISABLE_SWITCH_UNHANDLED_ENUM_CASE_WARNINGS \
-			__pragma (warning(push)) \
-			__pragma (warning(disable: 4061)) /*  enumerator 'enumator' in switch of enum 'enumeration' is not explicitly handled by a case label */
-	#endif // PRAGMA_DISABLE_UNSAFE_TYPECAST_WARNINGS
-
-	#ifndef PRAGMA_RESTORE_SWITCH_UNHANDLED_ENUM_CASE_WARNINGS
-	#define PRAGMA_RESTORE_SWITCH_UNHANDLED_ENUM_CASE_WARNINGS \
-				__pragma(warning(pop))
-	#endif // PRAGMA_RESTORE_UNSAFE_TYPECAST_WARNINGS
-
-	#ifndef PRAGMA_DISABLE_UNDEFINED_IDENTIFIER_WARNINGS
-		#define PRAGMA_DISABLE_UNDEFINED_IDENTIFIER_WARNINGS \
-			__pragma(warning(push)) \
-			__pragma(warning(disable: 4668)) /* 'symbol' is not defined as a preprocessor macro, replacing with '0' for 'directives' */
-	#endif // PRAGMA_DISABLE_UNDEFINED_IDENTIFIER_WARNINGS
-
-	#ifndef PRAGMA_ENABLE_UNDEFINED_IDENTIFIER_WARNINGS
-		#define PRAGMA_ENABLE_UNDEFINED_IDENTIFIER_WARNINGS \
-			__pragma(warning(pop))
-	#endif // PRAGMA_ENABLE_UNDEFINED_IDENTIFIER_WARNINGS
-
-	#ifndef PRAGMA_DISABLE_MISSING_VIRTUAL_DESTRUCTOR_WARNINGS
-		#define PRAGMA_DISABLE_MISSING_VIRTUAL_DESTRUCTOR_WARNINGS \
-			__pragma(warning(push)) \
-			__pragma(warning(disable: 4265)) /* class' : class has virtual functions, but destructor is not virtual */
-	#endif // PRAGMA_DISABLE_MISSING_VIRTUAL_DESTRUCTOR_WARNINGS
-
-	#ifndef PRAGMA_ENABLE_MISSING_VIRTUAL_DESTRUCTOR_WARNINGS
-		#define PRAGMA_ENABLE_MISSING_VIRTUAL_DESTRUCTOR_WARNINGS \
-			__pragma(warning(pop))
-	#endif // PRAGMA_ENABLE_MISSING_VIRTUAL_DESTRUCTOR_WARNINGS
-
-	#ifndef PRAGMA_DISABLE_REORDER_WARNINGS
-		#define PRAGMA_DISABLE_REORDER_WARNINGS \
-			__pragma(warning(push)) \
-			__pragma(warning(disable: 5038)) /* data member 'member1' will be initialized after data member 'member2' data member 'member' will be initialized after base class 'base_class' */
-	#endif // PRAGMA_DISABLE_REORDER_WARNINGS
-
-	#ifndef PRAGMA_ENABLE_REORDER_WARNINGS
-		#define PRAGMA_ENABLE_REORDER_WARNINGS \
-			__pragma(warning(pop))
-	#endif // PRAGMA_ENABLE_REORDER_WARNINGS
-
-	#ifndef PRAGMA_DISABLE_REGISTER_WARNINGS
-		#define PRAGMA_DISABLE_REGISTER_WARNINGS \
-			__pragma(warning(push)) \
-			__pragma(warning(disable: 5033)) /* 'register' is no longer a supported storage class */
-	#endif // PRAGMA_DISABLE_REGISTER_WARNINGS
-
-	#ifndef PRAGMA_ENABLE_REGISTER_WARNINGS
-		#define PRAGMA_ENABLE_REGISTER_WARNINGS \
-			__pragma(warning(pop))
-	#endif // PRAGMA_ENABLE_REGISTER_WARNINGS
-
-	#ifndef PRAGMA_DISABLE_UNUSED_PRIVATE_FIELDS_WARNINGS
-	#define PRAGMA_DISABLE_UNUSED_PRIVATE_FIELDS_WARNINGS
-			// MSVC doesn't seem to have a warning similar to -Wunused-private-field
-	#endif // PRAGMA_DISABLE_UNUSED_PRIVATE_FIELDS_WARNINGS
-
-	#ifndef PRAGMA_ENABLE_UNUSED_PRIVATE_FIELDS_WARNINGS
-	#define PRAGMA_ENABLE_UNUSED_PRIVATE_FIELDS_WARNINGS
-			// MSVC doesn't seem to have a warning similar to -Wunused-private-field
-	#endif // PRAGMA_ENABLE_UNUSED_PRIVATE_FIELDS_WARNINGS
-
-	#ifndef PRAGMA_DISABLE_UNINITIALIZED_CONST_REFERENCE_WARNINGS
-	#define PRAGMA_DISABLE_UNINITIALIZED_CONST_REFERENCE_WARNINGS
-		// MSVC doesn't seem to have a warning similar to -Wuninitialized-const-reference
-	#endif // PRAGMA_DISABLE_UNINITIALIZED_CONST_REFERENCE_WARNINGS
-
-	#ifndef PRAGMA_ENABLE_UNINITIALIZED_CONST_REFERENCE_WARNINGS
-	#define PRAGMA_ENABLE_UNINITIALIZED_CONST_REFERENCE_WARNINGS
-		// MSVC doesn't seem to have a warning similar to -Wuninitialized-const-reference
-	#endif // PRAGMA_ENABLE_UNINITIALIZED_CONST_REFERENCE_WARNINGS
-
-	#ifndef PRAGMA_DISABLE_ENUM_CONSTEXPR_CONVERSION_WARNINGS
-	#define PRAGMA_DISABLE_ENUM_CONSTEXPR_CONVERSION_WARNINGS
-		// MSVC doesn't seem to have a warning similar to -Wenum-constexpr-conversion
-	#endif // PRAGMA_DISABLE_ENUM_CONSTEXPR_CONVERSION_WARNINGS
-
-	#ifndef PRAGMA_ENABLE_ENUM_CONSTEXPR_CONVERSION_WARNINGS
-	#define PRAGMA_ENABLE_ENUM_CONSTEXPR_CONVERSION_WARNINGS
-		// MSVC doesn't seem to have a warning similar to -Wenum-constexpr-conversion
-	#endif // PRAGMA_ENABLE_ENUM_CONSTEXPR_CONVERSION_WARNINGS
-
-	#ifndef PRAGMA_DISABLE_DEPRECATED_COPY_WITH_USER_DEFINED_COPY_WARNINGS
-	#define PRAGMA_DISABLE_DEPRECATED_COPY_WITH_USER_DEFINED_COPY_WARNINGS
-		// MSVC doesn't seem to have a warning similar to -Wdeprecated-copy-with-user-provided-copy
-	#endif // PRAGMA_DISABLE_DEPRECATED_COPY_WITH_USER_DEFINED_COPY_WARNINGS
-
-	#ifndef PRAGMA_ENABLE_DEPRECATED_COPY_WITH_USER_DEFINED_COPY_WARNINGS
-	#define PRAGMA_ENABLE_DEPRECATED_COPY_WITH_USER_DEFINED_COPY_WARNINGS
-		// MSVC doesn't seem to have a warning similar to -Wdeprecated-copy-with-user-provided-copy
-	#endif // PRAGMA_ENABLE_DEPRECATED_COPY_WITH_USER_DEFINED_COPY_WARNINGS
-
-	#ifndef PRAGMA_DISABLE_NONPORTABLE_INCLUDE_PATH_WARNING
-	#define PRAGMA_DISABLE_NONPORTABLE_INCLUDE_PATH_WARNING
-		// MSVC doesn't seem to have a warning similar to -Wnonportable-include-path
-	#endif // PRAGMA_DISABLE_NONPORTABLE_INCLUDE_PATH_WARNING
-
-	#ifndef PRAGMA_ENABLE_NONPORTABLE_INCLUDE_PATH_WARNING
-	#define PRAGMA_ENABLE_NONPORTABLE_INCLUDE_PATH_WARNING
-		// MSVC doesn't seem to have a warning similar to -Wnonportable-include-path
-	#endif // PRAGMA_ENABLE_NONPORTABLE_INCLUDE_PATH_WARNING
-
-	#ifndef PRAGMA_POP
-		#define PRAGMA_POP \
-			__pragma(warning(pop))
-	#endif // PRAGMA_POP
-
-	// UE4 uses a struct packing of 4 for Win32 and 8 for Win64, 
-	// and the default packing is 8 for Win32 and 16 for Win64
-	#ifdef _WIN64
-		#ifndef PRAGMA_PUSH_PLATFORM_DEFAULT_PACKING
-			#define PRAGMA_PUSH_PLATFORM_DEFAULT_PACKING \
-				__pragma(pack(push, 16))
-		#endif // PRAGMA_PUSH_PLATFORM_DEFAULT_PACKING
-	#else // _WIN64
-		#ifndef PRAGMA_PUSH_PLATFORM_DEFAULT_PACKING
-			#define PRAGMA_PUSH_PLATFORM_DEFAULT_PACKING \
-				__pragma(pack(push, 8))
-		#endif // PRAGMA_PUSH_PLATFORM_DEFAULT_PACKING
-	#endif // _WIN64
-	
-	#ifndef PRAGMA_POP_PLATFORM_DEFAULT_PACKING
-		#define PRAGMA_POP_PLATFORM_DEFAULT_PACKING \
-			__pragma(pack(pop))
-	#endif // PRAGMA_POP_PLATFORM_DEFAULT_PACKING
-
-	// Disable common CA warnings around SDK includes
-	#ifndef UE_COMPILER_THIRD_PARTY_INCLUDES_START
-		#define UE_COMPILER_THIRD_PARTY_INCLUDES_START \
-			__pragma(warning(push)) \
-			__pragma(warning(disable: 4125))  /* decimal digit terminates octal escape sequence. */ \
-			__pragma(warning(disable: 4510))  /* '<class>': default constructor could not be generated. */ \
-			__pragma(warning(disable: 4610))  /* object '<class>' can never be instantiated - user-defined constructor required. */ \
-			__pragma(warning(disable: 4800))  /* Implicit conversion from '<type>' to bool. Possible information loss. */ \
-			__pragma(warning(disable: 4946))  /* reinterpret_cast used between related classes: '<class1>' and '<class2>' */ \
-			__pragma(warning(disable: 4996))  /* '<obj>' was declared deprecated. */ \
-			__pragma(warning(disable: 6011))  /* Dereferencing NULL pointer '<ptr>'. */ \
-			__pragma(warning(disable: 6101))  /* Returning uninitialized memory '<expr>'.  A successful path through the function does not set the named _Out_ parameter. */ \
-			__pragma(warning(disable: 6287))  /* Redundant code:  the left and right sub-expressions are identical. */ \
-			__pragma(warning(disable: 6308))  /* 'realloc' might return null pointer: assigning null pointer to 'X', which is passed as an argument to 'realloc', will cause the original memory block to be leaked. */ \
-			__pragma(warning(disable: 6326))  /* Potential comparison of a constant with another constant. */ \
-			__pragma(warning(disable: 6340))  /* Mismatch on sign: Incorrect type passed as parameter in call to function. */ \
-			__pragma(warning(disable: 6385))  /* Reading invalid data from '<ptr>':  the readable size is '<num1>' bytes, but '<num2>' bytes may be read. */ \
-			__pragma(warning(disable: 6386))  /* Buffer overrun while writing to '<ptr>':  the writable size is '<num1>' bytes, but '<num2>' bytes might be written. */ \
-			__pragma(warning(disable: 6553))  /* The annotation for function '<func>' on _Param_(<num>) does not apply to a value type. */ \
-			__pragma(warning(disable: 28182)) /* Dereferencing NULL pointer. '<ptr1>' contains the same NULL value as '<ptr2>' did. */ \
-			__pragma(warning(disable: 28251)) /* Inconsistent annotation for '<func>': this instance has no annotations. */ \
-			__pragma(warning(disable: 28252)) /* Inconsistent annotation for '<func>': return/function has '<annotation>' on the prior instance. */ \
-			__pragma(warning(disable: 28253)) /* Inconsistent annotation for '<func>': _Param_(<num>) has '<annotation>' on the prior instance. */ \
-			__pragma(warning(disable: 28301)) /* No annotations for first declaration of '<func>'. */ \
-			PRAGMA_DISABLE_REORDER_WARNINGS \
-			PRAGMA_DISABLE_UNDEFINED_IDENTIFIER_WARNINGS \
-			PRAGMA_DISABLE_SHADOW_VARIABLE_WARNINGS \
-			PRAGMA_DISABLE_MISSING_VIRTUAL_DESTRUCTOR_WARNINGS \
-			PRAGMA_DISABLE_DEPRECATION_WARNINGS \
-			PRAGMA_DISABLE_UNINITIALIZED_CONST_REFERENCE_WARNINGS \
-			PRAGMA_DISABLE_UNREACHABLE_CODE_WARNINGS
-	#endif // THIRD_PARTY_INCLUDES_START
-
-	#ifndef UE_COMPILER_THIRD_PARTY_INCLUDES_END
-		#define UE_COMPILER_THIRD_PARTY_INCLUDES_END \
-			PRAGMA_RESTORE_UNREACHABLE_CODE_WARNINGS \
-			PRAGMA_ENABLE_UNINITIALIZED_CONST_REFERENCE_WARNINGS \
-			PRAGMA_ENABLE_DEPRECATION_WARNINGS \
-			PRAGMA_ENABLE_MISSING_VIRTUAL_DESTRUCTOR_WARNINGS \
-			PRAGMA_ENABLE_SHADOW_VARIABLE_WARNINGS \
-			PRAGMA_ENABLE_UNDEFINED_IDENTIFIER_WARNINGS \
-			PRAGMA_ENABLE_REORDER_WARNINGS \
-			__pragma(warning(pop))
-	#endif // THIRD_PARTY_INCLUDES_END
-
-	#define EMIT_CUSTOM_WARNING_AT_LINE(Line, Warning) \
-		__pragma(message(WARNING_LOCATION(Line) ": warning C4996: " Warning))
-
-	#define EMIT_CUSTOM_ERROR_AT_LINE(Line, Error) \
-		__pragma(message(WARNING_LOCATION(Line) ": error C4996: " Error))
-
-	#ifndef PGO_LINK_DISABLE_WARNINGS
-		#define PGO_LINK_DISABLE_WARNINGS PRAGMA_DISABLE_BUFFER_OVERRUN_WARNING
-	#endif
-
-	#ifndef PGO_LINK_ENABLE_WARNINGS
-		#define PGO_LINK_ENABLE_WARNINGS PRAGMA_ENABLE_BUFFER_OVERRUN_WARNING
-	#endif
-#endif // defined(__clang__)
+#if PLATFORM_MICROSOFT
+#include "Microsoft/MSVCPlatformCompilerPreSetup.h"
+#endif
diff --git a/Engine/Source/Runtime/Core/Public/Microsoft/MSVCPlatform.h b/Engine/Source/Runtime/Core/Public/Microsoft/MSVCPlatform.h
new file mode 100644
--- /dev/null
+++ b/Engine/Source/Runtime/Core/Public/Microsoft/MSVCPlatform.h
@@ -0,0 +1,34 @@
+// Copyright Epic Games, Inc. All Rights Reserved.
+
+/*================================================================================
+	MSVCPlatform.h: Setup for any MSVC-using platform
+==================================================================================*/
+
+#pragma once
+
+#if _MSC_VER < 1920
+	#error "Compiler is expected to support if constexpr"
+#endif
+
+#if !defined(__cpp_fold_expressions)
+	#error "Compiler is expected to support fold expressions"
+#endif
+
+#define PLATFORM_RETURN_ADDRESS()	        _ReturnAddress()
+#define PLATFORM_RETURN_ADDRESS_POINTER()	_AddressOfReturnAddress()
+
+// https://devblogs.microsoft.com/cppblog/improving-the-state-of-debug-performance-in-c/
+#if __has_cpp_attribute(msvc::intrinsic)
+#define UE_INTRINSIC_CAST [[msvc::intrinsic]]
+#endif
+
+// Ensure we can use this builtin - seems to be present on Clang 9, GCC 11 and MSVC 19.26
+#define PLATFORM_COMPILER_SUPPORTS_BUILTIN_BITCAST (_MSC_VER >= 1926)
+
+#ifdef __has_cpp_attribute
+	#if __has_cpp_attribute(msvc::lifetimebound)
+		#define UE_LIFETIMEBOUND [[msvc::lifetimebound]]
+	#endif
+#endif
+
+#define UE_NO_PROFILE_ATTRIBUTE
\ No newline at end of file
diff --git a/Engine/Source/Runtime/Core/Public/Microsoft/MSVCPlatformCompilerPreSetup.h b/Engine/Source/Runtime/Core/Public/Microsoft/MSVCPlatformCompilerPreSetup.h
new file mode 100644
--- /dev/null
+++ b/Engine/Source/Runtime/Core/Public/Microsoft/MSVCPlatformCompilerPreSetup.h
@@ -0,0 +1,291 @@
+// Copyright Epic Games, Inc. All Rights Reserved.
+
+#pragma once
+
+// HEADER_UNIT_SKIP - Cause cirular dependencies for header units
+
+#if defined(__clang__)
+	// UE4 uses a struct packing of 4 for Win32 and 8 for Win64, 
+	// and the default packing is 8 for Win32 and 16 for Win64
+	#ifdef _WIN64
+		#ifndef PRAGMA_PUSH_PLATFORM_DEFAULT_PACKING
+			#define PRAGMA_PUSH_PLATFORM_DEFAULT_PACKING \
+				_Pragma("pack(push)") \
+				_Pragma("pack(16)")
+		#endif // PRAGMA_PUSH_PLATFORM_DEFAULT_PACKING
+	#else // _WIN64
+		#ifndef PRAGMA_PUSH_PLATFORM_DEFAULT_PACKING
+			#define PRAGMA_PUSH_PLATFORM_DEFAULT_PACKING \
+				_Pragma("pack(push)") \
+				_Pragma("pack(8)")
+		#endif // PRAGMA_PUSH_PLATFORM_DEFAULT_PACKING
+	#endif // _WIN64
+	
+	#ifndef PRAGMA_POP_PLATFORM_DEFAULT_PACKING
+		#define PRAGMA_POP_PLATFORM_DEFAULT_PACKING \
+			_Pragma("pack(pop)")
+	#endif // PRAGMA_POP_PLATFORM_DEFAULT_PACKING
+
+	#include "Clang/ClangPlatformCompilerPreSetup.h"
+#else
+	#ifndef PRAGMA_DISABLE_DEPRECATION_WARNINGS
+		#define PRAGMA_DISABLE_DEPRECATION_WARNINGS \
+			__pragma (warning(push)) \
+			__pragma (warning(disable: 4995)) /* 'function': name was marked as #pragma deprecated */ \
+			__pragma (warning(disable: 4996)) /* The compiler encountered a deprecated declaration. */
+	#endif // PRAGMA_DISABLE_DEPRECATION_WARNINGS
+
+	#ifndef PRAGMA_ENABLE_DEPRECATION_WARNINGS
+		#define PRAGMA_ENABLE_DEPRECATION_WARNINGS \
+			__pragma (warning(pop))
+	#endif // PRAGMA_ENABLE_DEPRECATION_WARNINGS
+
+	#ifndef PRAGMA_DISABLE_BUFFER_OVERRUN_WARNING
+		#define PRAGMA_DISABLE_BUFFER_OVERRUN_WARNING \
+			__pragma (warning(push)) \
+			__pragma (warning(disable:4789)) /* buffer 'identifier' of size N bytes will be overrun; M bytes will be written starting at offset L */
+	#endif // PRAGMA_DISABLE_BUFFER_OVERRUN_WARNING
+
+	#ifndef PRAGMA_ENABLE_BUFFER_OVERRUN_WARNING
+		#define PRAGMA_ENABLE_BUFFER_OVERRUN_WARNING \
+			__pragma (warning(pop))
+	#endif // PRAGMA_ENABLE_BUFFER_OVERRUN_WARNING
+
+	#ifndef PRAGMA_DISABLE_SHADOW_VARIABLE_WARNINGS
+		#define PRAGMA_DISABLE_SHADOW_VARIABLE_WARNINGS \
+			__pragma (warning(push)) \
+			__pragma (warning(disable: 4456)) /* declaration of 'LocalVariable' hides previous local declaration */ \
+			__pragma (warning(disable: 4457)) /* declaration of 'LocalVariable' hides function parameter */ \
+			__pragma (warning(disable: 4458)) /* declaration of 'LocalVariable' hides class member */ \
+			__pragma (warning(disable: 4459)) /* declaration of 'LocalVariable' hides global declaration */ \
+			__pragma (warning(disable: 6244)) /* local declaration of <variable> hides previous declaration at <line> of <file> */
+	#endif // PRAGMA_DISABLE_SHADOW_VARIABLE_WARNINGS
+
+	#ifndef PRAGMA_ENABLE_SHADOW_VARIABLE_WARNINGS
+		#define PRAGMA_ENABLE_SHADOW_VARIABLE_WARNINGS \
+			__pragma(warning(pop))
+	#endif // PRAGMA_ENABLE_SHADOW_VARIABLE_WARNINGS
+
+	#ifndef PRAGMA_DISABLE_UNREACHABLE_CODE_WARNINGS
+		#define PRAGMA_DISABLE_UNREACHABLE_CODE_WARNINGS \
+			__pragma (warning(push)) \
+			__pragma (warning(disable: 4702)) /* unreachable code */
+	#endif // PRAGMA_DISABLE_UNREACHABLE_CODE_WARNINGS
+	
+	#ifndef PRAGMA_RESTORE_UNREACHABLE_CODE_WARNINGS
+		#define PRAGMA_RESTORE_UNREACHABLE_CODE_WARNINGS \
+			__pragma (warning(pop))
+	#endif // PRAGMA_RESTORE_UNREACHABLE_CODE_WARNINGS
+
+	#ifndef PRAGMA_DISABLE_UNSAFE_TYPECAST_WARNINGS
+		#define PRAGMA_DISABLE_UNSAFE_TYPECAST_WARNINGS \
+			__pragma (warning(push)) \
+			__pragma (warning(disable: 4244)) /* 'argument': conversion from 'type1' to 'type2', possible loss of data */	\
+			__pragma (warning(disable: 4838)) /* 'argument': conversion from 'type1' to 'type2' requires a narrowing conversion */
+	#endif // PRAGMA_DISABLE_UNSAFE_TYPECAST_WARNINGS
+
+	#ifndef PRAGMA_ENABLE_UNSAFE_TYPECAST_WARNINGS
+		#define PRAGMA_ENABLE_UNSAFE_TYPECAST_WARNINGS \
+			UE_DEPRECATED_MACRO(5.0, "The PRAGMA_ENABLE_UNSAFE_TYPECAST_WARNINGS macro has been deprecated in favor of PRAGMA_RESTORE_UNSAFE_TYPECAST_WARNINGS. To force enable warnings use PRAGMA_FORCE_UNSAFE_TYPECAST_WARNINGS.")
+	#endif // PRAGMA_ENABLE_UNSAFE_TYPECAST_WARNINGS
+
+	#ifndef PRAGMA_FORCE_UNSAFE_TYPECAST_WARNINGS
+		#define PRAGMA_FORCE_UNSAFE_TYPECAST_WARNINGS \
+			__pragma (warning(push)) \
+			__pragma (warning(error: 4244)) /* 'argument': conversion from 'type1' to 'type2', possible loss of data */	\
+			__pragma (warning(error: 4838)) /* 'argument': conversion from 'type1' to 'type2' requires a narrowing conversion */
+	#endif // PRAGMA_FORCE_UNSAFE_TYPECAST_WARNINGS
+
+	#ifndef PRAGMA_RESTORE_UNSAFE_TYPECAST_WARNINGS
+		#define PRAGMA_RESTORE_UNSAFE_TYPECAST_WARNINGS \
+			__pragma(warning(pop))
+	#endif // PRAGMA_RESTORE_UNSAFE_TYPECAST_WARNINGS
+
+	#ifndef PRAGMA_DISABLE_SWITCH_UNHANDLED_ENUM_CASE_WARNINGS
+		#define PRAGMA_DISABLE_SWITCH_UNHANDLED_ENUM_CASE_WARNINGS \
+			__pragma (warning(push)) \
+			__pragma (warning(disable: 4061)) /*  enumerator 'enumator' in switch of enum 'enumeration' is not explicitly handled by a case label */
+	#endif // PRAGMA_DISABLE_UNSAFE_TYPECAST_WARNINGS
+
+	#ifndef PRAGMA_RESTORE_SWITCH_UNHANDLED_ENUM_CASE_WARNINGS
+	#define PRAGMA_RESTORE_SWITCH_UNHANDLED_ENUM_CASE_WARNINGS \
+				__pragma(warning(pop))
+	#endif // PRAGMA_RESTORE_UNSAFE_TYPECAST_WARNINGS
+
+	#ifndef PRAGMA_DISABLE_UNDEFINED_IDENTIFIER_WARNINGS
+		#define PRAGMA_DISABLE_UNDEFINED_IDENTIFIER_WARNINGS \
+			__pragma(warning(push)) \
+			__pragma(warning(disable: 4668)) /* 'symbol' is not defined as a preprocessor macro, replacing with '0' for 'directives' */
+	#endif // PRAGMA_DISABLE_UNDEFINED_IDENTIFIER_WARNINGS
+
+	#ifndef PRAGMA_ENABLE_UNDEFINED_IDENTIFIER_WARNINGS
+		#define PRAGMA_ENABLE_UNDEFINED_IDENTIFIER_WARNINGS \
+			__pragma(warning(pop))
+	#endif // PRAGMA_ENABLE_UNDEFINED_IDENTIFIER_WARNINGS
+
+	#ifndef PRAGMA_DISABLE_MISSING_VIRTUAL_DESTRUCTOR_WARNINGS
+		#define PRAGMA_DISABLE_MISSING_VIRTUAL_DESTRUCTOR_WARNINGS \
+			__pragma(warning(push)) \
+			__pragma(warning(disable: 4265)) /* class' : class has virtual functions, but destructor is not virtual */
+	#endif // PRAGMA_DISABLE_MISSING_VIRTUAL_DESTRUCTOR_WARNINGS
+
+	#ifndef PRAGMA_ENABLE_MISSING_VIRTUAL_DESTRUCTOR_WARNINGS
+		#define PRAGMA_ENABLE_MISSING_VIRTUAL_DESTRUCTOR_WARNINGS \
+			__pragma(warning(pop))
+	#endif // PRAGMA_ENABLE_MISSING_VIRTUAL_DESTRUCTOR_WARNINGS
+
+	#ifndef PRAGMA_DISABLE_REORDER_WARNINGS
+		#define PRAGMA_DISABLE_REORDER_WARNINGS \
+			__pragma(warning(push)) \
+			__pragma(warning(disable: 5038)) /* data member 'member1' will be initialized after data member 'member2' data member 'member' will be initialized after base class 'base_class' */
+	#endif // PRAGMA_DISABLE_REORDER_WARNINGS
+
+	#ifndef PRAGMA_ENABLE_REORDER_WARNINGS
+		#define PRAGMA_ENABLE_REORDER_WARNINGS \
+			__pragma(warning(pop))
+	#endif // PRAGMA_ENABLE_REORDER_WARNINGS
+
+	#ifndef PRAGMA_DISABLE_REGISTER_WARNINGS
+		#define PRAGMA_DISABLE_REGISTER_WARNINGS \
+			__pragma(warning(push)) \
+			__pragma(warning(disable: 5033)) /* 'register' is no longer a supported storage class */
+	#endif // PRAGMA_DISABLE_REGISTER_WARNINGS
+
+	#ifndef PRAGMA_ENABLE_REGISTER_WARNINGS
+		#define PRAGMA_ENABLE_REGISTER_WARNINGS \
+			__pragma(warning(pop))
+	#endif // PRAGMA_ENABLE_REGISTER_WARNINGS
+
+	#ifndef PRAGMA_DISABLE_UNUSED_PRIVATE_FIELDS_WARNINGS
+	#define PRAGMA_DISABLE_UNUSED_PRIVATE_FIELDS_WARNINGS
+			// MSVC doesn't seem to have a warning similar to -Wunused-private-field
+	#endif // PRAGMA_DISABLE_UNUSED_PRIVATE_FIELDS_WARNINGS
+
+	#ifndef PRAGMA_ENABLE_UNUSED_PRIVATE_FIELDS_WARNINGS
+	#define PRAGMA_ENABLE_UNUSED_PRIVATE_FIELDS_WARNINGS
+			// MSVC doesn't seem to have a warning similar to -Wunused-private-field
+	#endif // PRAGMA_ENABLE_UNUSED_PRIVATE_FIELDS_WARNINGS
+
+	#ifndef PRAGMA_DISABLE_UNINITIALIZED_CONST_REFERENCE_WARNINGS
+	#define PRAGMA_DISABLE_UNINITIALIZED_CONST_REFERENCE_WARNINGS
+		// MSVC doesn't seem to have a warning similar to -Wuninitialized-const-reference
+	#endif // PRAGMA_DISABLE_UNINITIALIZED_CONST_REFERENCE_WARNINGS
+
+	#ifndef PRAGMA_ENABLE_UNINITIALIZED_CONST_REFERENCE_WARNINGS
+	#define PRAGMA_ENABLE_UNINITIALIZED_CONST_REFERENCE_WARNINGS
+		// MSVC doesn't seem to have a warning similar to -Wuninitialized-const-reference
+	#endif // PRAGMA_ENABLE_UNINITIALIZED_CONST_REFERENCE_WARNINGS
+
+	#ifndef PRAGMA_DISABLE_ENUM_CONSTEXPR_CONVERSION_WARNINGS
+	#define PRAGMA_DISABLE_ENUM_CONSTEXPR_CONVERSION_WARNINGS
+		// MSVC doesn't seem to have a warning similar to -Wenum-constexpr-conversion
+	#endif // PRAGMA_DISABLE_ENUM_CONSTEXPR_CONVERSION_WARNINGS
+
+	#ifndef PRAGMA_ENABLE_ENUM_CONSTEXPR_CONVERSION_WARNINGS
+	#define PRAGMA_ENABLE_ENUM_CONSTEXPR_CONVERSION_WARNINGS
+		// MSVC doesn't seem to have a warning similar to -Wenum-constexpr-conversion
+	#endif // PRAGMA_ENABLE_ENUM_CONSTEXPR_CONVERSION_WARNINGS
+
+	#ifndef PRAGMA_DISABLE_DEPRECATED_COPY_WITH_USER_DEFINED_COPY_WARNINGS
+	#define PRAGMA_DISABLE_DEPRECATED_COPY_WITH_USER_DEFINED_COPY_WARNINGS
+		// MSVC doesn't seem to have a warning similar to -Wdeprecated-copy-with-user-provided-copy
+	#endif // PRAGMA_DISABLE_DEPRECATED_COPY_WITH_USER_DEFINED_COPY_WARNINGS
+
+	#ifndef PRAGMA_ENABLE_DEPRECATED_COPY_WITH_USER_DEFINED_COPY_WARNINGS
+	#define PRAGMA_ENABLE_DEPRECATED_COPY_WITH_USER_DEFINED_COPY_WARNINGS
+		// MSVC doesn't seem to have a warning similar to -Wdeprecated-copy-with-user-provided-copy
+	#endif // PRAGMA_ENABLE_DEPRECATED_COPY_WITH_USER_DEFINED_COPY_WARNINGS
+
+	#ifndef PRAGMA_DISABLE_NONPORTABLE_INCLUDE_PATH_WARNING
+	#define PRAGMA_DISABLE_NONPORTABLE_INCLUDE_PATH_WARNING
+		// MSVC doesn't seem to have a warning similar to -Wnonportable-include-path
+	#endif // PRAGMA_DISABLE_NONPORTABLE_INCLUDE_PATH_WARNING
+
+	#ifndef PRAGMA_ENABLE_NONPORTABLE_INCLUDE_PATH_WARNING
+	#define PRAGMA_ENABLE_NONPORTABLE_INCLUDE_PATH_WARNING
+		// MSVC doesn't seem to have a warning similar to -Wnonportable-include-path
+	#endif // PRAGMA_ENABLE_NONPORTABLE_INCLUDE_PATH_WARNING
+
+	#ifndef PRAGMA_POP
+		#define PRAGMA_POP \
+			__pragma(warning(pop))
+	#endif // PRAGMA_POP
+
+	// UE4 uses a struct packing of 4 for Win32 and 8 for Win64, 
+	// and the default packing is 8 for Win32 and 16 for Win64
+	#ifdef _WIN64
+		#ifndef PRAGMA_PUSH_PLATFORM_DEFAULT_PACKING
+			#define PRAGMA_PUSH_PLATFORM_DEFAULT_PACKING \
+				__pragma(pack(push, 16))
+		#endif // PRAGMA_PUSH_PLATFORM_DEFAULT_PACKING
+	#else // _WIN64
+		#ifndef PRAGMA_PUSH_PLATFORM_DEFAULT_PACKING
+			#define PRAGMA_PUSH_PLATFORM_DEFAULT_PACKING \
+				__pragma(pack(push, 8))
+		#endif // PRAGMA_PUSH_PLATFORM_DEFAULT_PACKING
+	#endif // _WIN64
+	
+	#ifndef PRAGMA_POP_PLATFORM_DEFAULT_PACKING
+		#define PRAGMA_POP_PLATFORM_DEFAULT_PACKING \
+			__pragma(pack(pop))
+	#endif // PRAGMA_POP_PLATFORM_DEFAULT_PACKING
+
+	// Disable common CA warnings around SDK includes
+	#ifndef UE_COMPILER_THIRD_PARTY_INCLUDES_START
+		#define UE_COMPILER_THIRD_PARTY_INCLUDES_START \
+			__pragma(warning(push)) \
+			__pragma(warning(disable: 4125))  /* decimal digit terminates octal escape sequence. */ \
+			__pragma(warning(disable: 4510))  /* '<class>': default constructor could not be generated. */ \
+			__pragma(warning(disable: 4610))  /* object '<class>' can never be instantiated - user-defined constructor required. */ \
+			__pragma(warning(disable: 4800))  /* Implicit conversion from '<type>' to bool. Possible information loss. */ \
+			__pragma(warning(disable: 4946))  /* reinterpret_cast used between related classes: '<class1>' and '<class2>' */ \
+			__pragma(warning(disable: 4996))  /* '<obj>' was declared deprecated. */ \
+			__pragma(warning(disable: 6011))  /* Dereferencing NULL pointer '<ptr>'. */ \
+			__pragma(warning(disable: 6101))  /* Returning uninitialized memory '<expr>'.  A successful path through the function does not set the named _Out_ parameter. */ \
+			__pragma(warning(disable: 6287))  /* Redundant code:  the left and right sub-expressions are identical. */ \
+			__pragma(warning(disable: 6308))  /* 'realloc' might return null pointer: assigning null pointer to 'X', which is passed as an argument to 'realloc', will cause the original memory block to be leaked. */ \
+			__pragma(warning(disable: 6326))  /* Potential comparison of a constant with another constant. */ \
+			__pragma(warning(disable: 6340))  /* Mismatch on sign: Incorrect type passed as parameter in call to function. */ \
+			__pragma(warning(disable: 6385))  /* Reading invalid data from '<ptr>':  the readable size is '<num1>' bytes, but '<num2>' bytes may be read. */ \
+			__pragma(warning(disable: 6386))  /* Buffer overrun while writing to '<ptr>':  the writable size is '<num1>' bytes, but '<num2>' bytes might be written. */ \
+			__pragma(warning(disable: 6553))  /* The annotation for function '<func>' on _Param_(<num>) does not apply to a value type. */ \
+			__pragma(warning(disable: 28182)) /* Dereferencing NULL pointer. '<ptr1>' contains the same NULL value as '<ptr2>' did. */ \
+			__pragma(warning(disable: 28251)) /* Inconsistent annotation for '<func>': this instance has no annotations. */ \
+			__pragma(warning(disable: 28252)) /* Inconsistent annotation for '<func>': return/function has '<annotation>' on the prior instance. */ \
+			__pragma(warning(disable: 28253)) /* Inconsistent annotation for '<func>': _Param_(<num>) has '<annotation>' on the prior instance. */ \
+			__pragma(warning(disable: 28301)) /* No annotations for first declaration of '<func>'. */ \
+			PRAGMA_DISABLE_REORDER_WARNINGS \
+			PRAGMA_DISABLE_UNDEFINED_IDENTIFIER_WARNINGS \
+			PRAGMA_DISABLE_SHADOW_VARIABLE_WARNINGS \
+			PRAGMA_DISABLE_MISSING_VIRTUAL_DESTRUCTOR_WARNINGS \
+			PRAGMA_DISABLE_DEPRECATION_WARNINGS \
+			PRAGMA_DISABLE_UNINITIALIZED_CONST_REFERENCE_WARNINGS \
+			PRAGMA_DISABLE_UNREACHABLE_CODE_WARNINGS
+	#endif // THIRD_PARTY_INCLUDES_START
+
+	#ifndef UE_COMPILER_THIRD_PARTY_INCLUDES_END
+		#define UE_COMPILER_THIRD_PARTY_INCLUDES_END \
+			PRAGMA_RESTORE_UNREACHABLE_CODE_WARNINGS \
+			PRAGMA_ENABLE_UNINITIALIZED_CONST_REFERENCE_WARNINGS \
+			PRAGMA_ENABLE_DEPRECATION_WARNINGS \
+			PRAGMA_ENABLE_MISSING_VIRTUAL_DESTRUCTOR_WARNINGS \
+			PRAGMA_ENABLE_SHADOW_VARIABLE_WARNINGS \
+			PRAGMA_ENABLE_UNDEFINED_IDENTIFIER_WARNINGS \
+			PRAGMA_ENABLE_REORDER_WARNINGS \
+			__pragma(warning(pop))
+	#endif // THIRD_PARTY_INCLUDES_END
+
+	#define EMIT_CUSTOM_WARNING_AT_LINE(Line, Warning) \
+		__pragma(message(WARNING_LOCATION(Line) ": warning C4996: " Warning))
+
+	#define EMIT_CUSTOM_ERROR_AT_LINE(Line, Error) \
+		__pragma(message(WARNING_LOCATION(Line) ": error C4996: " Error))
+
+	#ifndef PGO_LINK_DISABLE_WARNINGS
+		#define PGO_LINK_DISABLE_WARNINGS PRAGMA_DISABLE_BUFFER_OVERRUN_WARNING
+	#endif
+
+	#ifndef PGO_LINK_ENABLE_WARNINGS
+		#define PGO_LINK_ENABLE_WARNINGS PRAGMA_ENABLE_BUFFER_OVERRUN_WARNING
+	#endif
+#endif // defined(__clang__)
diff --git a/Engine/Source/Runtime/Core/Public/Windows/WindowsPlatform.h b/Engine/Source/Runtime/Core/Public/Windows/WindowsPlatform.h
--- a/Engine/Source/Runtime/Core/Public/Windows/WindowsPlatform.h
+++ b/Engine/Source/Runtime/Core/Public/Windows/WindowsPlatform.h
@@ -11,7 +11,7 @@
 #elif defined(__clang__)
 	#include "Clang/ClangPlatform.h"
 #else
-	#include "MSVC/MSVCPlatform.h"
+	#include "Microsoft/MSVCPlatform.h"
 #endif
 
 #if WINVER <= 0x600
diff --git a/Engine/Source/Runtime/Core/Public/Windows/WindowsPlatformCompilerPreSetup.h b/Engine/Source/Runtime/Core/Public/Windows/WindowsPlatformCompilerPreSetup.h
--- a/Engine/Source/Runtime/Core/Public/Windows/WindowsPlatformCompilerPreSetup.h
+++ b/Engine/Source/Runtime/Core/Public/Windows/WindowsPlatformCompilerPreSetup.h
@@ -2,9 +2,9 @@
 
 #pragma once
 
-// HEADER_UNIT_SKIP - Cause cirular dependencies for header units
+// HEADER_UNIT_SKIP - Cause circular dependencies for header units
 
-#include "MSVC/MSVCPlatformCompilerPreSetup.h"
+#include "Microsoft/MSVCPlatformCompilerPreSetup.h"
 
 // This is needed when compiling with header units. By default _CRT_MEMCPY_S_INLINE is "static inline"
 // "static inline" means that function is private to module which causes compile errors when other windows includes uses memcpy_s and friends
