SharedPointer dtor is the same for every specialization, manually outline it into a helper.

Saves a _lot_ of code size (-5% for some builds) and is usually a bit faster due to better I$ hit rates.

diff --git a/Engine/Source/Runtime/Core/Public/Templates/SharedPointerInternals.h b/Engine/Source/Runtime/Core/Public/Templates/SharedPointerInternals.h
--- a/Engine/Source/Runtime/Core/Public/Templates/SharedPointerInternals.h
+++ b/Engine/Source/Runtime/Core/Public/Templates/SharedPointerInternals.h
@@ -248,6 +248,17 @@ namespace SharedPointerInternals
 			}
 		}
 
+		/** Release a shared reference, including nullptr check. Not inlined, for code size. */
+		static FORCENOINLINE void ReleaseSharedReferenceNoInline(TReferenceControllerBase<Mode>* ReferenceController)
+		{
+			if (ReferenceController != nullptr)
+			{
+				// Tell the reference counter object that we're no longer referencing the object with
+				// this shared pointer
+				ReferenceController->ReleaseSharedReference();
+			}
+		}
+
 		/** Adds a weak reference to this counter */
 		inline void AddWeakReference()
 		{
@@ -599,12 +610,7 @@ namespace SharedPointerInternals
 		/** Destructor. */
 		inline ~FSharedReferencer()
 		{
-			if( ReferenceController != nullptr )
-			{
-				// Tell the reference counter object that we're no longer referencing the object with
-				// this shared pointer
-				ReferenceController->ReleaseSharedReference();
-			}
+			TReferenceControllerBase<Mode>::ReleaseSharedReferenceNoInline(ReferenceController);
 		}
 
 		/** Assignment operator adds a reference to the assigned object.  If this counter was previously
