Expose function UE::ConfigUtilities::SaveCVarForNextBoot in case there is requirement to trigger CVar saving by user code other than hotfixed .ini file.

diff --git a/Engine/Source/Runtime/Core/Public/Misc/ConfigUtilities.h b/Engine/Source/Runtime/Core/Public/Misc/ConfigUtilities.h
--- a/Engine/Source/Runtime/Core/Public/Misc/ConfigUtilities.h
+++ b/Engine/Source/Runtime/Core/Public/Misc/ConfigUtilities.h
@@ -54,6 +54,9 @@ namespace UE::ConfigUtilities
 	 */
 	CORE_API void ApplyCVarSettingsFromIni(const TCHAR* InSectionBaseName, const TCHAR* InIniFilename, uint32 SetBy, bool bAllowCheating=false, FName Tag=NAME_None);
 
+
+	CORE_API void SaveCVarForNextBoot(const TCHAR* Key, const TCHAR* Value);
+
 	/**
 	 * Helper function to operate a user defined function for each CVar key/value pair in the specified section in an ini file
 	 * @param InSectionName - The name of the section to apply cvars from
