Move AudioStreamCache LLM catogery out from Audio top-level tag

diff --git a/Engine/Source/Runtime/Core/Public/HAL/LowLevelMemTracker.h b/Engine/Source/Runtime/Core/Public/HAL/LowLevelMemTracker.h
--- a/Engine/Source/Runtime/Core/Public/HAL/LowLevelMemTracker.h
+++ b/Engine/Source/Runtime/Core/Public/HAL/LowLevelMemTracker.h
@@ -199,7 +199,7 @@ ENUM_CLASS_FLAGS(ESizeParams);
 	macro(AudioDecompress,						"AudioDecompress",				GET_STATFNAME(STAT_AudioDecompressLLM),						GET_STATFNAME(STAT_AudioSummaryLLM),			ELLMTag::Audio)\
 	macro(AudioRealtimePrecache,				"AudioRealtimePrecache",		GET_STATFNAME(STAT_AudioRealtimePrecacheLLM),				GET_STATFNAME(STAT_AudioSummaryLLM),			ELLMTag::Audio)\
 	macro(AudioFullDecompress,					"AudioFullDecompress",			GET_STATFNAME(STAT_AudioFullDecompressLLM),					GET_STATFNAME(STAT_AudioSummaryLLM),			ELLMTag::Audio)\
-	macro(AudioStreamCache,						"AudioStreamCache",				GET_STATFNAME(STAT_AudioStreamCacheLLM),					GET_STATFNAME(STAT_AudioSummaryLLM),			ELLMTag::Audio)\
+	macro(AudioStreamCache,						"AudioStreamCache",				GET_STATFNAME(STAT_AudioStreamCacheLLM),					GET_STATFNAME(STAT_AudioSummaryLLM),			-1) \
 	macro(AudioStreamCacheCompressedData,		"AudioStreamCacheCompressedData",GET_STATFNAME(STAT_AudioStreamCacheCompressedDataLLM),		GET_STATFNAME(STAT_AudioSummaryLLM),			ELLMTag::Audio)\
 	macro(AudioSynthesis,						"AudioSynthesis",				GET_STATFNAME(STAT_AudioSynthesisLLM),						GET_STATFNAME(STAT_AudioSummaryLLM),			ELLMTag::Audio)\
 	macro(RealTimeCommunications,				"RealTimeCommunications",		GET_STATFNAME(STAT_RealTimeCommunicationsLLM),				NAME_None,										-1)\
