[verse] use lz4 to compress <predicts> network data

#rnx

diff --git a/Engine/Source/Runtime/Core/Public/Compression/lz4.h b/Engine/Source/Runtime/Core/Public/Compression/lz4.h
--- a/Engine/Source/Runtime/Core/Public/Compression/lz4.h
+++ b/Engine/Source/Runtime/Core/Public/Compression/lz4.h
@@ -80,26 +80,8 @@ extern "C" {
 /*^***************************************************************
 *  Export parameters
 *****************************************************************/
-/*
-*  LZ4_DLL_EXPORT :
-*  Enable exporting of functions when building a Windows DLL
-*  LZ4LIB_VISIBILITY :
-*  Control library symbols visibility.
-*/
-#ifndef LZ4LIB_VISIBILITY
-#  if defined(__GNUC__) && (__GNUC__ >= 4)
-#    define LZ4LIB_VISIBILITY __attribute__ ((visibility ("default")))
-#  else
-#    define LZ4LIB_VISIBILITY
-#  endif
-#endif
-#if defined(LZ4_DLL_EXPORT) && (LZ4_DLL_EXPORT==1)
-#  define LZ4LIB_API __declspec(dllexport) LZ4LIB_VISIBILITY
-#elif defined(LZ4_DLL_IMPORT) && (LZ4_DLL_IMPORT==1)
-#  define LZ4LIB_API __declspec(dllimport) LZ4LIB_VISIBILITY /* It isn't required but allows to generate better code, saving a function pointer load from the IAT and an indirect jump.*/
-#else
-#  define LZ4LIB_API LZ4LIB_VISIBILITY
-#endif
+
+#define LZ4LIB_API CORE_API
 
 /*! LZ4_FREESTANDING :
  *  When this macro is set to 1, it enables "freestanding mode" that is
