Add a less-templated base class to `TSparseArray` to shave some binary size.

This change adds a less-templated base class in `TSparseArray` that is templated on the size/alignment/allocator, and moves a bunch of methods that don't actually need the type in C++. Doing this shaves 6MB off of the binary size in a big compile target.

diff --git a/Engine/Source/Runtime/Core/Public/Containers/SparseArray.h b/Engine/Source/Runtime/Core/Public/Containers/SparseArray.h
--- a/Engine/Source/Runtime/Core/Public/Containers/SparseArray.h
+++ b/Engine/Source/Runtime/Core/Public/Containers/SparseArray.h
@@ -60,51 +60,27 @@
 
 DECLARE_TEMPLATE_INTRINSIC_TYPE_LAYOUT(template<typename ElementType>, TSparseArrayElementOrFreeListLink<ElementType>);
 
-template <typename AllocatorType, typename InDerivedType = void>
-class TScriptSparseArray;
-
-/**
- * A dynamically sized array where element indices aren't necessarily contiguous.  Memory is allocated for all 
- * elements in the array's index range, so it doesn't save memory; but it does allow O(1) element removal that 
- * doesn't invalidate the indices of subsequent elements.  It uses TArray to store the elements, and a TBitArray
- * to store whether each element index is allocated (for fast iteration over allocated elements).
- *
- **/
-template<typename InElementType,typename Allocator /*= FDefaultSparseArrayAllocator */>
-class TSparseArray
+template<size_t SizeOfElementType, size_t AlignOfElementType, typename Allocator> class TSparseArrayBase
 {
-	using ElementType = InElementType;
-
-	template <typename, typename>
-	friend class TScriptSparseArray;
-
-public:
-	/** Destructor. */
-	~TSparseArray()
-	{
-		UE_STATIC_ASSERT_WARN(TIsTriviallyRelocatable_V<InElementType>, "TArray can only be used with trivially relocatable types");
-
-		// Destruct the elements in the array.
-		Empty();
-	}
-
-	/////////////////////////////////////////////////////
-	// Start - intrusive TOptional<TSparseArray> state //
-	/////////////////////////////////////////////////////
-	constexpr static bool bHasIntrusiveUnsetOptionalState = true;
-	using IntrusiveUnsetOptionalStateType = TSparseArray;
-
-	[[nodiscard]] explicit TSparseArray(FIntrusiveUnsetOptionalState Tag)
+protected:
+	[[nodiscard]] explicit TSparseArrayBase(FIntrusiveUnsetOptionalState Tag)
 		: Data(Tag)
 	{
 	}
-	[[nodiscard]] bool operator==(FIntrusiveUnsetOptionalState Tag) const
+
+	/** Default constructor. */
+	[[nodiscard]] TSparseArrayBase()
+		: FirstFreeIndex(-1)
+		, NumFreeIndices(0)
 	{
-		return Data == Tag;
 	}
-	///////////////////////////////////////////////////
-	// End - intrusive TOptional<TSparseArray> state //
-	///////////////////////////////////////////////////
 
+	/** Copy constructor. */
+	[[nodiscard]] TSparseArrayBase(const TSparseArrayBase& InCopy)
+		: FirstFreeIndex(-1)
+		, NumFreeIndices(0)
+	{
+	}
+
 	/** Marks an index as allocated, and returns information about the allocation. */
 	FSparseArrayAllocationInfo AllocateIndex(int32 Index)
@@ -123,7 +99,7 @@
 
 		return Result;
 	}
-
+	
 	/**
 	 * Allocates space for an element in the array.  The element is not initialized, and you must use the corresponding placement new operator
 	 * to construct the element in the allocated memory.
@@ -154,22 +130,6 @@
 		return AllocateIndex(Index);
 	}
 
-	/** Adds an element to the array. */
-	int32 Add(const ElementType& Element)
-	{
-		FSparseArrayAllocationInfo Allocation = AddUninitialized();
-		new(Allocation) ElementType(Element);
-		return Allocation.Index;
-	}
-
-	/** Adds an element to the array. */
-	int32 Add(ElementType&& Element)
-	{
-		FSparseArrayAllocationInfo Allocation = AddUninitialized();
-		new(Allocation) ElementType(MoveTemp(Element));
-		return Allocation.Index;
-	}
-
 	FSparseArrayAllocationInfo AddUninitializedAtLowestFreeIndex(int32& LowestFreeIndexSearchStart)
 	{
 		FElementOrFreeListLink* DataPtr;
@@ -216,6 +176,141 @@
 		Result.Pointer = &DataPtr[Result.Index].ElementData;
 		return Result;
 	}
+	
+	/** Removes Count elements from the array, starting from Index, without destructing them. */
+	void RemoveAtUninitialized(int32 Index,int32 Count = 1)
+	{
+		FElementOrFreeListLink* DataPtr = (FElementOrFreeListLink*)Data.GetData();
+
+		for (; Count; --Count)
+		{
+			check(AllocationFlags[Index]);
+
+			// Mark the element as free and add it to the free element list.
+			if(NumFreeIndices)
+			{
+				DataPtr[FirstFreeIndex].PrevFreeIndex = Index;
+			}
+			DataPtr[Index].PrevFreeIndex = -1;
+			DataPtr[Index].NextFreeIndex = NumFreeIndices > 0 ? FirstFreeIndex : INDEX_NONE;
+			FirstFreeIndex = Index;
+			++NumFreeIndices;
+			AllocationFlags[Index] = false;
+
+			++Index;
+		}
+	}
+
+	[[nodiscard]] int32 GetMaxIndex() const
+	{
+		return Data.Num();
+	}
+	
+	[[nodiscard]] bool IsEmpty() const
+	{
+		return Data.Num() == NumFreeIndices;
+	}
+
+	[[nodiscard]] int32 Num() const
+	{
+		return Data.Num() - NumFreeIndices;
+	}
+
+	/** @return The number of elements the array can hold before reallocation. */
+	[[nodiscard]] int32 Max() const
+	{
+		return Data.Max();
+	}
 
 	/**
+	 * The element type stored is only indirectly related to the element type requested, to avoid instantiating TArray redundantly for
+	 * compatible types.
+	 */
+	typedef TSparseArrayElementOrFreeListLink<
+		TAlignedBytes<SizeOfElementType, AlignOfElementType>
+		> FElementOrFreeListLink;
+
+	typedef TArray<FElementOrFreeListLink, typename Allocator::ElementAllocator> DataType;
+	DataType Data;
+
+	typedef TBitArray<typename Allocator::BitArrayAllocator> AllocationBitArrayType;
+	AllocationBitArrayType AllocationFlags;
+
+	/** The index of an unallocated element in the array that currently contains the head of the linked list of free elements. */
+	int32 FirstFreeIndex = -1;
+
+	/** The number of elements in the free list. */
+	int32 NumFreeIndices = 0;
+};
+
+template <typename AllocatorType, typename InDerivedType = void>
+class TScriptSparseArray;
+
+/**
+ * A dynamically sized array where element indices aren't necessarily contiguous.  Memory is allocated for all 
+ * elements in the array's index range, so it doesn't save memory; but it does allow O(1) element removal that 
+ * doesn't invalidate the indices of subsequent elements.  It uses TArray to store the elements, and a TBitArray
+ * to store whether each element index is allocated (for fast iteration over allocated elements).
+ *
+ **/
+template<typename InElementType, typename Allocator /*= FDefaultSparseArrayAllocator */>
+class TSparseArray : public TSparseArrayBase<sizeof(InElementType), alignof(InElementType), Allocator>
+{
+	using SuperType = TSparseArrayBase<sizeof(InElementType), alignof(InElementType), Allocator>;
+	using ElementType = InElementType;
+
+	template <typename, typename>
+	friend class TScriptSparseArray;
+
+public:
+	/** Destructor. */
+	~TSparseArray()
+	{
+		UE_STATIC_ASSERT_WARN(TIsTriviallyRelocatable_V<InElementType>, "TArray can only be used with trivially relocatable types");
+
+		// Destruct the elements in the array.
+		Empty();
+	}
+
+	/////////////////////////////////////////////////////
+	// Start - intrusive TOptional<TSparseArray> state //
+	/////////////////////////////////////////////////////
+	constexpr static bool bHasIntrusiveUnsetOptionalState = true;
+	using IntrusiveUnsetOptionalStateType = TSparseArray;
+
+	[[nodiscard]] explicit TSparseArray(FIntrusiveUnsetOptionalState Tag)
+		: SuperType(Tag)
+	{
+	}
+
+	[[nodiscard]] bool operator==(FIntrusiveUnsetOptionalState Tag) const
+	{
+		return Data == Tag;
+	}
+	///////////////////////////////////////////////////
+	// End - intrusive TOptional<TSparseArray> state //
+	///////////////////////////////////////////////////
+
+	using SuperType::AllocateIndex;
+	using SuperType::AddUninitialized;
+
+	/** Adds an element to the array. */
+	int32 Add(const ElementType& Element)
+	{
+		FSparseArrayAllocationInfo Allocation = AddUninitialized();
+		new(Allocation) ElementType(Element);
+		return Allocation.Index;
+	}
+
+	/** Adds an element to the array. */
+	int32 Add(ElementType&& Element)
+	{
+		FSparseArrayAllocationInfo Allocation = AddUninitialized();
+		new(Allocation) ElementType(MoveTemp(Element));
+		return Allocation.Index;
+	}
+
+	using SuperType::AddUninitializedAtLowestFreeIndex;
+
+	/**
 	 * Constructs a new item at the last freed index of the array.
@@ -353,29 +448,7 @@
 		RemoveAtUninitialized(Index, Count);
 	}
 
-	/** Removes Count elements from the array, starting from Index, without destructing them. */
-	void RemoveAtUninitialized(int32 Index,int32 Count = 1)
-	{
-		FElementOrFreeListLink* DataPtr = (FElementOrFreeListLink*)Data.GetData();
+	using SuperType::RemoveAtUninitialized;
 
-		for (; Count; --Count)
-		{
-			check(AllocationFlags[Index]);
-
-			// Mark the element as free and add it to the free element list.
-			if(NumFreeIndices)
-			{
-				DataPtr[FirstFreeIndex].PrevFreeIndex = Index;
-			}
-			DataPtr[Index].PrevFreeIndex = -1;
-			DataPtr[Index].NextFreeIndex = NumFreeIndices > 0 ? FirstFreeIndex : INDEX_NONE;
-			FirstFreeIndex = Index;
-			++NumFreeIndices;
-			AllocationFlags[Index] = false;
-
-			++Index;
-		}
-	}
-
 	/**
 	 * Removes all elements from the array, potentially leaving space allocated for an expected number of elements about to be added.
@@ -771,11 +844,7 @@
 	}
 
 	/** Default constructor. */
-	[[nodiscard]] TSparseArray()
-		: FirstFreeIndex(-1)
-		, NumFreeIndices(0)
-	{
-	}
+	[[nodiscard]] TSparseArray() = default;
 
 	/** Move constructor. */
 	[[nodiscard]] TSparseArray(TSparseArray&& InCopy)
@@ -785,15 +854,14 @@
 
 	/** Copy constructor. */
 	[[nodiscard]] TSparseArray(const TSparseArray& InCopy)
-		: FirstFreeIndex(-1)
-		, NumFreeIndices(0)
+		: SuperType(InCopy)
 	{
 		*this = InCopy;
 	}
 
 	/** Initializer list constructor. */
 	[[nodiscard]] TSparseArray(std::initializer_list<ElementType> InList)
-		: TSparseArray()
+		: SuperType()
 	{
 		*this = InList;
 	}
@@ -921,25 +989,11 @@
 	{
 		return AllocationFlags[Index];
 	}
-	/** @return The non-inclusive maximum index of elements in the sparse array. */
-	[[nodiscard]] int32 GetMaxIndex() const
-	{
-		return Data.Num();
-	}
-	[[nodiscard]] bool IsEmpty() const
-	{
-		return Data.Num() == NumFreeIndices;
-	}
-	[[nodiscard]] int32 Num() const
-	{
-		return Data.Num() - NumFreeIndices;
-	}
 
-	/** @return The number of elements the array can hold before reallocation. */
-	[[nodiscard]] int32 Max() const
-	{
-		return Data.Max();
-	}
+	using SuperType::GetMaxIndex;
+	using SuperType::IsEmpty;
+	using SuperType::Num;
+	using SuperType::Max;
 
 	/**
 	 * Checks that the specified address is not part of an element within the container.  Used for implementations
@@ -1206,13 +1260,7 @@
 
 private:
 
-	/**
-	 * The element type stored is only indirectly related to the element type requested, to avoid instantiating TArray redundantly for
-	 * compatible types.
-	 */
-	typedef TSparseArrayElementOrFreeListLink<
-		TAlignedBytes<sizeof(ElementType), alignof(ElementType)>
-		> FElementOrFreeListLink;
+	using FElementOrFreeListLink = SuperType::FElementOrFreeListLink;
 
 	/** Extracts the element value from the array's element structure and passes it to the user provided comparison class. */
 	template <typename PREDICATE_CLASS>
@@ -1232,17 +1280,13 @@
 		}
 	};
 
-	typedef TArray<FElementOrFreeListLink,typename Allocator::ElementAllocator> DataType;
-	DataType Data;
+	using typename SuperType::DataType;
+	using typename SuperType::AllocationBitArrayType;
 
-	typedef TBitArray<typename Allocator::BitArrayAllocator> AllocationBitArrayType;
-	AllocationBitArrayType AllocationFlags;
+	using SuperType::Data;
+	using SuperType::AllocationFlags;
+	using SuperType::FirstFreeIndex;
+	using SuperType::NumFreeIndices;
 
-	/** The index of an unallocated element in the array that currently contains the head of the linked list of free elements. */
-	int32 FirstFreeIndex = -1;
-
-	/** The number of elements in the free list. */
-	int32 NumFreeIndices = 0;
-
 public:
 	void WriteMemoryImage(FMemoryImageWriter& Writer) const
