Mutex: Use Yield() instead of YieldThread() when spinning

Yield() maps to _mm_pause() on most platforms, which is meant for spin-wait loops.
YieldThread() maps to a function that can cause a context switch on many platforms. That behavior is problematic when the system is under heavy load.

diff --git a/Engine/Source/Runtime/Core/Private/Async/RecursiveMutex.cpp b/Engine/Source/Runtime/Core/Private/Async/RecursiveMutex.cpp
--- a/Engine/Source/Runtime/Core/Private/Async/RecursiveMutex.cpp
+++ b/Engine/Source/Runtime/Core/Private/Async/RecursiveMutex.cpp
@@ -86,7 +86,7 @@ FORCENOINLINE void FRecursiveMutex::LockSlow(uint32 CurrentState, const uint32 C
 		// Spin up to the spin limit while there are no waiting threads.
 		if (LIKELY(!(CurrentState & MayHaveWaitingLockFlag) && SpinCount < SpinLimit))
 		{
-			FPlatformProcess::YieldThread();
+			FPlatformProcess::Yield();
 			++SpinCount;
 			CurrentState = State.load(std::memory_order_relaxed);
 			continue;
diff --git a/Engine/Source/Runtime/Core/Private/Async/SharedMutex.cpp b/Engine/Source/Runtime/Core/Private/Async/SharedMutex.cpp
--- a/Engine/Source/Runtime/Core/Private/Async/SharedMutex.cpp
+++ b/Engine/Source/Runtime/Core/Private/Async/SharedMutex.cpp
@@ -53,7 +53,7 @@ void FSharedMutex::LockSharedSlow()
 		// Spin up to the spin limit while there are no waiting threads.
 		if (LIKELY(!(CurrentState & MayHaveWaitingLockFlag) && SpinCount < SpinLimit))
 		{
-			FPlatformProcess::YieldThread();
+			FPlatformProcess::Yield();
 			++SpinCount;
 			CurrentState = State.load(std::memory_order_relaxed);
 			continue;
diff --git a/Engine/Source/Runtime/Core/Private/Async/SharedRecursiveMutex.cpp b/Engine/Source/Runtime/Core/Private/Async/SharedRecursiveMutex.cpp
--- a/Engine/Source/Runtime/Core/Private/Async/SharedRecursiveMutex.cpp
+++ b/Engine/Source/Runtime/Core/Private/Async/SharedRecursiveMutex.cpp
@@ -155,7 +155,7 @@ FORCENOINLINE void FSharedRecursiveMutex::LockSlow(uint32 CurrentState, const ui
 		// Spin up to the spin limit while there are no waiting threads.
 		if (LIKELY(!(CurrentState & MayHaveWaitingLockFlag) && SpinCount < SpinLimit))
 		{
-			FPlatformProcess::YieldThread();
+			FPlatformProcess::Yield();
 			++SpinCount;
 			CurrentState = State.load(std::memory_order_relaxed);
 			continue;
@@ -283,7 +283,7 @@ FORCENOINLINE void FSharedRecursiveMutex::LockSharedSlow(Core::Private::FSharedR
 		// Spin up to the spin limit while there are no waiting threads.
 		if (LIKELY(!(CurrentState & MayHaveWaitingLockFlag) && SpinCount < SpinLimit))
 		{
-			FPlatformProcess::YieldThread();
+			FPlatformProcess::Yield();
 			++SpinCount;
 			CurrentState = State.load(std::memory_order_relaxed);
 			continue;
diff --git a/Engine/Source/Runtime/Core/Private/Async/WordMutex.cpp b/Engine/Source/Runtime/Core/Private/Async/WordMutex.cpp
--- a/Engine/Source/Runtime/Core/Private/Async/WordMutex.cpp
+++ b/Engine/Source/Runtime/Core/Private/Async/WordMutex.cpp
@@ -45,7 +45,7 @@ void FWordMutex::LockSlow()
 		// Spin up to the spin limit while there is no queue.
 		if (!(CurrentState & QueueMask) && SpinCount < SpinLimit)
 		{
-			FPlatformProcess::YieldThread();
+			FPlatformProcess::Yield();
 			++SpinCount;
 			continue;
 		}
diff --git a/Engine/Source/Runtime/Core/Public/Async/IntrusiveMutex.h b/Engine/Source/Runtime/Core/Public/Async/IntrusiveMutex.h
--- a/Engine/Source/Runtime/Core/Public/Async/IntrusiveMutex.h
+++ b/Engine/Source/Runtime/Core/Public/Async/IntrusiveMutex.h
@@ -151,7 +151,7 @@ public:
 			// Spin up to the spin limit while there are no waiting threads.
 			if (LIKELY(!(CurrentState & MayHaveWaitingLockFlag) && SpinCount < SpinLimit))
 			{
-				FPlatformProcess::YieldThread();
+				FPlatformProcess::Yield();
 				++SpinCount;
 				CurrentState = State.load(std::memory_order_relaxed);
 				continue;
