[PlatformFile]

Allow read requests for more bytes than are less than within the file to not result in failure but instead return all bytes left in the file and a success return value. Over-reading is common when using flags such as FILE_FLAG_NO_BUFFERING where read requests must be volume sector size multiples.

ReadAt explicitly early outs as a failure if more bytes than the file contains are requested so the while loop used there has been unchanged.

#rnx

diff --git a/Engine/Source/Runtime/Core/Private/Windows/WindowsPlatformFile.cpp b/Engine/Source/Runtime/Core/Private/Windows/WindowsPlatformFile.cpp
--- a/Engine/Source/Runtime/Core/Private/Windows/WindowsPlatformFile.cpp
+++ b/Engine/Source/Runtime/Core/Private/Windows/WindowsPlatformFile.cpp
@@ -811,7 +811,7 @@ public:
 				TRACE_PLATFORMFILE_END_READ(&OverlappedIO, TotalNumRead);
 				return false;
 			}
-		} while (BytesToRead > 0);
+		} while (BytesToRead > 0 && FilePos < FileSize);
 		TRACE_PLATFORMFILE_END_READ(&OverlappedIO, TotalNumRead);
 		return true;
 	}
