Fix console variables not being included in Help command

diff --git a/Engine/Source/Runtime/Core/Private/Misc/Parse.cpp b/Engine/Source/Runtime/Core/Private/Misc/Parse.cpp
--- a/Engine/Source/Runtime/Core/Private/Misc/Parse.cpp
+++ b/Engine/Source/Runtime/Core/Private/Misc/Parse.cpp
@@ -62,27 +62,6 @@ ConsoleCommandLibrary::~ConsoleCommandLibrary()
 	GConsoleCommandLibrary = 0;
 }
 
-
-
-class FConsoleCommandDumpVisitor
-{
-public:
-	/**
-	 * @param Name must not be null
-	 * @param Object must not be null
-	 * @param Sink must not be null
-	 */
-	static void OnConsoleObject(const TCHAR *Name, IConsoleObject* Object, TSet<FString>* Sink)
-	{
-		if (!Object->AsCommand() || Object->TestFlags(ECVF_Unregistered))
-		{
-			return;
-		}
-
-		Sink->Add(Name);
-	}
-};
-
 bool ConsoleCommandLibrary_DumpLibrary(UWorld* InWorld, FExec& SubSystem, const FString& Pattern, FOutputDevice& Ar)
 {
 	TStringBuilder<32> Prefix;
@@ -105,12 +84,16 @@ bool ConsoleCommandLibrary_DumpLibrary(UWorld* InWorld, FExec& SubSystem, const
 		bExecuted = SubSystem.Exec(InWorld, *FakeCmd, Null);
 	}
 
+	auto VisitConsoleCommand = [&Sink = LocalConsoleCommandLibrary.KnownNames](const TCHAR *Name, IConsoleObject* Object)
+	{
+		if (!Object->TestFlags(ECVF_Unregistered) && Object->AsCommand())
+		{
+			Sink.Add(Name);
+		}
+	};
+
 	// Gather registered commands
-	IConsoleManager::Get().ForEachConsoleObjectThatStartsWith(
-		FConsoleObjectVisitor::CreateStatic(
-			&FConsoleCommandDumpVisitor::OnConsoleObject,
-			&LocalConsoleCommandLibrary.KnownNames),
-		Prefix);
+	IConsoleManager::Get().ForEachConsoleObjectThatStartsWith(FConsoleObjectVisitor::CreateLambda(VisitConsoleCommand), Prefix);
 
 	LocalConsoleCommandLibrary.KnownNames.Sort( TLess<FString>() );
 
@@ -148,11 +131,16 @@ bool ConsoleCommandLibrary_DumpLibraryHTML(UWorld* InWorld, FExec& SubSystem, co
 		bExecuted = SubSystem.Exec(InWorld, FakeCmd, Null);
 	}
 
+	auto VisitConsoleObject = [&Sink = LocalConsoleCommandLibrary.KnownNames](const TCHAR *Name, IConsoleObject* Object)
+	{
+		if (!Object->TestFlags(ECVF_Unregistered))
+		{
+			Sink.Add(Name);
+		}
+	};
+
 	// Gather registered variables and commands
-	IConsoleManager::Get().ForEachConsoleObjectThatStartsWith(
-		FConsoleObjectVisitor::CreateStatic(
-			&FConsoleCommandDumpVisitor::OnConsoleObject,
-			&LocalConsoleCommandLibrary.KnownNames));
+	IConsoleManager::Get().ForEachConsoleObjectThatStartsWith(FConsoleObjectVisitor::CreateLambda(VisitConsoleObject));
 
 	LocalConsoleCommandLibrary.KnownNames.Sort( TLess<FString>() );
 
diff --git a/Engine/Source/Runtime/Core/Public/Misc/Parse.h b/Engine/Source/Runtime/Core/Public/Misc/Parse.h
--- a/Engine/Source/Runtime/Core/Public/Misc/Parse.h
+++ b/Engine/Source/Runtime/Core/Public/Misc/Parse.h
@@ -235,9 +235,9 @@ ENUM_CLASS_FLAGS(FParse::ELineExtendedFlags)
 UE_DEPRECATED("5.6", "Use the overload that takes the entire command instead of a pattern.")
 CORE_API bool ConsoleCommandLibrary_DumpLibrary(class UWorld* InWorld, FExec& SubSystem, const FString& Cmd, FOutputDevice& Ar);
 
-/** Needed for the console command "DumpConsoleCommands" */
+/** Needed for the console command "DumpConsoleCommands". Outputs registered and unregistered console commands to the supplied device. */
 CORE_API bool ConsoleCommandLibrary_DumpLibrary(class UWorld* InWorld, FExec& SubSystem, const TCHAR* Cmd, FOutputDevice& Ar);
 
-/** Needed for the console command "Help" */
+/** Needed for the console command "Help". Outputs registered and unregistered console commands and console variables the supplied file path. */
 CORE_API bool ConsoleCommandLibrary_DumpLibraryHTML(class UWorld* InWorld, FExec& SubSystem, const FString& OutPath);
 #endif // !UE_BUILD_SHIPPING
diff --git a/Engine/Source/Runtime/Engine/Private/UnrealEngine.cpp b/Engine/Source/Runtime/Engine/Private/UnrealEngine.cpp
--- a/Engine/Source/Runtime/Engine/Private/UnrealEngine.cpp
+++ b/Engine/Source/Runtime/Engine/Private/UnrealEngine.cpp
@@ -5890,7 +5890,7 @@ static void DumpHelp(UWorld* InWorld)
 		TEXT("A console command has no state associated with it and gets executed immediately.\n")
 		TEXT("\n")
 		TEXT("DumpCVars                   Lists all console variables\n")
-		TEXT("DumpCCmds                   Lists all registerd console commands. Not all commands are registered.\n")
+		TEXT("DumpCCmds                   Lists all registered console commands. Not all commands are registered.\n")
 		TEXT("DumpConsoleCommands         Lists all console commands, registered and unregistered\n")
 		TEXT("\n")
 		TEXT("<Console variable>          Get the console variable value and how it was last set\n")
