Make CVar registration ctors non-inlined

We have a lot of CVars, and even things like the innocent-looking Singleton Get()
end up showing up for every single CVar ctor, which ends up being a major code
size waste.

diff --git a/Engine/Source/Runtime/Core/Private/HAL/ConsoleManager.cpp b/Engine/Source/Runtime/Core/Private/HAL/ConsoleManager.cpp
--- a/Engine/Source/Runtime/Core/Private/HAL/ConsoleManager.cpp
+++ b/Engine/Source/Runtime/Core/Private/HAL/ConsoleManager.cpp
@@ -5134,5 +5134,131 @@ namespace ConsoleVariablePlatformMapping
 	}
 }
 
+#if !NO_CVARS
+
+FAutoConsoleVariable::FAutoConsoleVariable(const TCHAR* Name, bool DefaultValue, const TCHAR* Help, uint32 Flags)
+	: FAutoConsoleObject(IConsoleManager::Get().RegisterConsoleVariable(Name, DefaultValue, Help, Flags))
+{
+}
+
+FAutoConsoleVariable::FAutoConsoleVariable(const TCHAR* Name, int32 DefaultValue, const TCHAR* Help, uint32 Flags)
+	: FAutoConsoleObject(IConsoleManager::Get().RegisterConsoleVariable(Name, DefaultValue, Help, Flags))
+{
+}
+
+FAutoConsoleVariable::FAutoConsoleVariable(const TCHAR* Name, float DefaultValue, const TCHAR* Help, uint32 Flags)
+	: FAutoConsoleObject(IConsoleManager::Get().RegisterConsoleVariable(Name, DefaultValue, Help, Flags))
+{
+}
+
+FAutoConsoleVariable::FAutoConsoleVariable(const TCHAR* Name, const TCHAR* DefaultValue, const TCHAR* Help, uint32 Flags)
+	: FAutoConsoleObject(IConsoleManager::Get().RegisterConsoleVariable(Name, DefaultValue, Help, Flags))
+{
+}
+
+FAutoConsoleVariable::FAutoConsoleVariable(const TCHAR* Name, bool DefaultValue, const TCHAR* Help, const FConsoleVariableDelegate& Callback, uint32 Flags)
+	: FAutoConsoleObject(IConsoleManager::Get().RegisterConsoleVariable(Name, DefaultValue, Help, Flags))
+{
+	AsVariable()->SetOnChangedCallback(Callback);
+}
+
+FAutoConsoleVariable::FAutoConsoleVariable(const TCHAR* Name, int32 DefaultValue, const TCHAR* Help, const FConsoleVariableDelegate& Callback, uint32 Flags)
+	: FAutoConsoleObject(IConsoleManager::Get().RegisterConsoleVariable(Name, DefaultValue, Help, Flags))
+{
+	AsVariable()->SetOnChangedCallback(Callback);
+}
+
+FAutoConsoleVariable::FAutoConsoleVariable(const TCHAR* Name, float DefaultValue, const TCHAR* Help, const FConsoleVariableDelegate& Callback, uint32 Flags)
+	: FAutoConsoleObject(IConsoleManager::Get().RegisterConsoleVariable(Name, DefaultValue, Help, Flags))
+{
+	AsVariable()->SetOnChangedCallback(Callback);
+}
+
+FAutoConsoleVariable::FAutoConsoleVariable(const TCHAR* Name, const TCHAR* DefaultValue, const TCHAR* Help, const FConsoleVariableDelegate& Callback, uint32 Flags)
+	: FAutoConsoleObject(IConsoleManager::Get().RegisterConsoleVariable(Name, DefaultValue, Help, Flags))
+{
+	AsVariable()->SetOnChangedCallback(Callback);
+}
+
+FAutoConsoleVariableRef::FAutoConsoleVariableRef(const TCHAR* Name, int32& RefValue, const TCHAR* Help, uint32 Flags)
+	: FAutoConsoleObject(IConsoleManager::Get().RegisterConsoleVariableRef(Name, RefValue, Help, Flags))
+{
+}
+
+FAutoConsoleVariableRef::FAutoConsoleVariableRef(const TCHAR* Name, float& RefValue, const TCHAR* Help, uint32 Flags)
+	: FAutoConsoleObject(IConsoleManager::Get().RegisterConsoleVariableRef(Name, RefValue, Help, Flags))
+{
+}
+
+FAutoConsoleVariableRef::FAutoConsoleVariableRef(const TCHAR* Name, bool& RefValue, const TCHAR* Help, uint32 Flags)
+	: FAutoConsoleObject(IConsoleManager::Get().RegisterConsoleVariableRef(Name, RefValue, Help, Flags))
+{
+}
+
+FAutoConsoleVariableRef::FAutoConsoleVariableRef(const TCHAR* Name, FString& RefValue, const TCHAR* Help, uint32 Flags)
+	: FAutoConsoleObject(IConsoleManager::Get().RegisterConsoleVariableRef(Name, RefValue, Help, Flags))
+{
+}
+
+FAutoConsoleVariableRef::FAutoConsoleVariableRef(const TCHAR* Name, FName& RefValue, const TCHAR* Help, uint32 Flags)
+	: FAutoConsoleObject(IConsoleManager::Get().RegisterConsoleVariableRef(Name, RefValue, Help, Flags))
+{
+}
+
+FAutoConsoleVariableRef::FAutoConsoleVariableRef(const TCHAR* Name, int32& RefValue, const TCHAR* Help, const FConsoleVariableDelegate& Callback, uint32 Flags)
+	: FAutoConsoleObject(IConsoleManager::Get().RegisterConsoleVariableRef(Name, RefValue, Help, Flags))
+{
+	AsVariable()->SetOnChangedCallback(Callback);
+}
+
+FAutoConsoleVariableRef::FAutoConsoleVariableRef(const TCHAR* Name, float& RefValue, const TCHAR* Help, const FConsoleVariableDelegate& Callback, uint32 Flags)
+	: FAutoConsoleObject(IConsoleManager::Get().RegisterConsoleVariableRef(Name, RefValue, Help, Flags))
+{
+	AsVariable()->SetOnChangedCallback(Callback);
+}
+
+FAutoConsoleVariableRef::FAutoConsoleVariableRef(const TCHAR* Name, bool& RefValue, const TCHAR* Help, const FConsoleVariableDelegate& Callback, uint32 Flags)
+	: FAutoConsoleObject(IConsoleManager::Get().RegisterConsoleVariableRef(Name, RefValue, Help, Flags))
+{
+	AsVariable()->SetOnChangedCallback(Callback);
+}
+
+FAutoConsoleVariableRef::FAutoConsoleVariableRef(const TCHAR* Name, FString& RefValue, const TCHAR* Help, const FConsoleVariableDelegate& Callback, uint32 Flags)
+	: FAutoConsoleObject(IConsoleManager::Get().RegisterConsoleVariableRef(Name, RefValue, Help, Flags))
+{
+	AsVariable()->SetOnChangedCallback(Callback);
+}
+
+FAutoConsoleCommand::FAutoConsoleCommand(const TCHAR* Name, const TCHAR* Help, const FConsoleCommandDelegate& Command, uint32 Flags)
+	: FAutoConsoleObject(IConsoleManager::Get().RegisterConsoleCommand(Name, Help, Command, Flags))
+{
+}
+
+FAutoConsoleCommand::FAutoConsoleCommand(const TCHAR* Name, const TCHAR* Help, const FConsoleCommandWithArgsDelegate& Command, uint32 Flags)
+	: FAutoConsoleObject(IConsoleManager::Get().RegisterConsoleCommand(Name, Help, Command, Flags))
+{
+}
+
+FAutoConsoleCommand::FAutoConsoleCommand(const TCHAR* Name, const TCHAR* Help, const FConsoleCommandWithOutputDeviceDelegate& Command, uint32 Flags)
+	: FAutoConsoleObject(IConsoleManager::Get().RegisterConsoleCommand(Name, Help, Command, Flags))
+{
+}
+
+FAutoConsoleCommand::FAutoConsoleCommand(const TCHAR* Name, const TCHAR* Help, const FConsoleCommandWithArgsAndOutputDeviceDelegate& Command, uint32 Flags)
+	: FAutoConsoleObject(IConsoleManager::Get().RegisterConsoleCommand(Name, Help, Command, Flags))
+{
+}
+
+FAutoConsoleCommand::FAutoConsoleCommand(const TCHAR* Name, const TCHAR* Help, const FConsoleCommandWithWorldDelegate& Command, uint32 Flags)
+	: FAutoConsoleObject(IConsoleManager::Get().RegisterConsoleCommand(Name, Help, Command, Flags))
+{
+}
+
+FAutoConsoleCommand::FAutoConsoleCommand(const TCHAR* Name, const TCHAR* Help, const FConsoleCommandWithWorldArgsAndOutputDeviceDelegate& Command, uint32 Flags)
+	: FAutoConsoleObject(IConsoleManager::Get().RegisterConsoleCommand(Name, Help, Command, Flags))
+{
+}
+
+#endif
 
 #undef LOCTEXT_NAMESPACE
diff --git a/Engine/Source/Runtime/Core/Public/HAL/IConsoleManager.h b/Engine/Source/Runtime/Core/Public/HAL/IConsoleManager.h
--- a/Engine/Source/Runtime/Core/Public/HAL/IConsoleManager.h
+++ b/Engine/Source/Runtime/Core/Public/HAL/IConsoleManager.h
@@ -1379,40 +1379,28 @@ public:
 	 * @param Help must not be 0
 	 * @param Flags bitmask combined from EConsoleVariableFlags
 	 */
-	FAutoConsoleVariable(const TCHAR* Name, bool DefaultValue, const TCHAR* Help, uint32 Flags = ECVF_Default)
-		: FAutoConsoleObject(IConsoleManager::Get().RegisterConsoleVariable(Name, DefaultValue, Help, Flags))
-	{
-	}
+	CORE_API FAutoConsoleVariable(const TCHAR* Name, bool DefaultValue, const TCHAR* Help, uint32 Flags = ECVF_Default);
 	/**
 	 * Create a int console variable
 	 * @param Name must not be 0
 	 * @param Help must not be 0
 	 * @param Flags bitmask combined from EConsoleVariableFlags
 	 */
-	FAutoConsoleVariable(const TCHAR* Name, int32 DefaultValue, const TCHAR* Help, uint32 Flags = ECVF_Default)
-		: FAutoConsoleObject(IConsoleManager::Get().RegisterConsoleVariable(Name, DefaultValue, Help, Flags))
-	{
-	}
+	CORE_API FAutoConsoleVariable(const TCHAR* Name, int32 DefaultValue, const TCHAR* Help, uint32 Flags = ECVF_Default);
 	/**
 	 * Create a float console variable
 	 * @param Name must not be 0
 	 * @param Help must not be 0
 	 * @param Flags bitmask combined from EConsoleVariableFlags
 	 */
-	FAutoConsoleVariable(const TCHAR* Name, float DefaultValue, const TCHAR* Help, uint32 Flags = ECVF_Default)
-		: FAutoConsoleObject(IConsoleManager::Get().RegisterConsoleVariable(Name, DefaultValue, Help, Flags))
-	{
-	}
+	CORE_API FAutoConsoleVariable(const TCHAR* Name, float DefaultValue, const TCHAR* Help, uint32 Flags = ECVF_Default);
 	/**
 	 * Create a string console variable
 	 * @param Name must not be 0
 	 * @param Help must not be 0
 	 * @param Flags bitmask combined from EConsoleVariableFlags
 	 */
-	FAutoConsoleVariable(const TCHAR* Name, const TCHAR* DefaultValue, const TCHAR* Help, uint32 Flags = ECVF_Default)
-		: FAutoConsoleObject(IConsoleManager::Get().RegisterConsoleVariable(Name, DefaultValue, Help, Flags))
-	{
-	}
+	CORE_API FAutoConsoleVariable(const TCHAR* Name, const TCHAR* DefaultValue, const TCHAR* Help, uint32 Flags = ECVF_Default);
 
 	/**
 	 * Create a bool console variable
@@ -1421,11 +1409,7 @@ public:
 	 * @param Callback Delegate called when the variable changes. @see IConsoleVariable::SetOnChangedCallback
 	 * @param Flags bitmask combined from EConsoleVariableFlags
 	 */
-	FAutoConsoleVariable(const TCHAR* Name, bool DefaultValue, const TCHAR* Help, const FConsoleVariableDelegate& Callback, uint32 Flags = ECVF_Default)
-		: FAutoConsoleObject(IConsoleManager::Get().RegisterConsoleVariable(Name, DefaultValue, Help, Flags))
-	{
-		AsVariable()->SetOnChangedCallback(Callback);
-	}
+	CORE_API FAutoConsoleVariable(const TCHAR* Name, bool DefaultValue, const TCHAR* Help, const FConsoleVariableDelegate& Callback, uint32 Flags = ECVF_Default);
 	
 	/**
 	 * Create a int console variable
@@ -1434,11 +1418,7 @@ public:
 	 * @param Callback Delegate called when the variable changes. @see IConsoleVariable::SetOnChangedCallback
 	 * @param Flags bitmask combined from EConsoleVariableFlags
 	 */
-	FAutoConsoleVariable(const TCHAR* Name, int32 DefaultValue, const TCHAR* Help, const FConsoleVariableDelegate& Callback, uint32 Flags = ECVF_Default)
-		: FAutoConsoleObject(IConsoleManager::Get().RegisterConsoleVariable(Name, DefaultValue, Help, Flags))
-	{
-		AsVariable()->SetOnChangedCallback(Callback);
-	}
+	CORE_API FAutoConsoleVariable(const TCHAR* Name, int32 DefaultValue, const TCHAR* Help, const FConsoleVariableDelegate& Callback, uint32 Flags = ECVF_Default);
 
 	/**
 	 * Create a float console variable
@@ -1447,11 +1427,7 @@ public:
 	 * @param Callback Delegate called when the variable changes. @see IConsoleVariable::SetOnChangedCallback
 	 * @param Flags bitmask combined from EConsoleVariableFlags
 	 */
-	FAutoConsoleVariable(const TCHAR* Name, float DefaultValue, const TCHAR* Help, const FConsoleVariableDelegate& Callback, uint32 Flags = ECVF_Default)
-		: FAutoConsoleObject(IConsoleManager::Get().RegisterConsoleVariable(Name, DefaultValue, Help, Flags))
-	{
-		AsVariable()->SetOnChangedCallback(Callback);
-	}
+	CORE_API FAutoConsoleVariable(const TCHAR* Name, float DefaultValue, const TCHAR* Help, const FConsoleVariableDelegate& Callback, uint32 Flags = ECVF_Default);
 
 	/**
 	 * Create a string console variable
@@ -1460,11 +1436,7 @@ public:
 	 * @param Callback Delegate called when the variable changes. @see IConsoleVariable::SetOnChangedCallback
 	 * @param Flags bitmask combined from EConsoleVariableFlags
 	 */
-	FAutoConsoleVariable(const TCHAR* Name, const TCHAR* DefaultValue, const TCHAR* Help, const FConsoleVariableDelegate& Callback, uint32 Flags = ECVF_Default)
-		: FAutoConsoleObject(IConsoleManager::Get().RegisterConsoleVariable(Name, DefaultValue, Help, Flags))
-	{
-		AsVariable()->SetOnChangedCallback(Callback);
-	}
+	CORE_API FAutoConsoleVariable(const TCHAR* Name, const TCHAR* DefaultValue, const TCHAR* Help, const FConsoleVariableDelegate& Callback, uint32 Flags = ECVF_Default);
 
 	/** Dereference back to a console variable**/
 	FORCEINLINE IConsoleVariable& operator*()
@@ -1516,40 +1488,28 @@ public:
 	 * @param Help must not be 0
 	 * @param Flags bitmask combined from EConsoleVariableFlags
 	 */
-	FAutoConsoleVariableRef(const TCHAR* Name, int32& RefValue, const TCHAR* Help, uint32 Flags = ECVF_Default)
-		: FAutoConsoleObject(IConsoleManager::Get().RegisterConsoleVariableRef(Name, RefValue, Help, Flags))
-	{
-	}
+	CORE_API FAutoConsoleVariableRef(const TCHAR* Name, int32& RefValue, const TCHAR* Help, uint32 Flags = ECVF_Default);
 	/**
 	 * Create a reference to a float console variable
 	 * @param Name must not be 0
 	 * @param Help must not be 0
 	 * @param Flags bitmask combined from EConsoleVariableFlags
 	 */
-	FAutoConsoleVariableRef(const TCHAR* Name, float& RefValue, const TCHAR* Help, uint32 Flags = ECVF_Default)
-		: FAutoConsoleObject(IConsoleManager::Get().RegisterConsoleVariableRef(Name, RefValue, Help, Flags))
-	{
-	}
+	CORE_API FAutoConsoleVariableRef(const TCHAR* Name, float& RefValue, const TCHAR* Help, uint32 Flags = ECVF_Default);
 	/**
 	 * Create a reference to a bool console variable
 	 * @param Name must not be 0
 	 * @param Help must not be 0
 	 * @param Flags bitmask combined from EConsoleVariableFlags
 	 */
-	FAutoConsoleVariableRef(const TCHAR* Name, bool& RefValue, const TCHAR* Help, uint32 Flags = ECVF_Default)
-		: FAutoConsoleObject(IConsoleManager::Get().RegisterConsoleVariableRef(Name, RefValue, Help, Flags))
-	{
-	}
+	CORE_API FAutoConsoleVariableRef(const TCHAR* Name, bool& RefValue, const TCHAR* Help, uint32 Flags = ECVF_Default);
 	/**
 	 * Create a reference to a FString console variable
 	 * @param Name must not be 0
 	 * @param Help must not be 0
 	 * @param Flags bitmask combined from EConsoleVariableFlags
 	 */
-	FAutoConsoleVariableRef(const TCHAR* Name, FString& RefValue, const TCHAR* Help, uint32 Flags = ECVF_Default)
-		: FAutoConsoleObject(IConsoleManager::Get().RegisterConsoleVariableRef(Name, RefValue, Help, Flags))
-	{
-	}
+	CORE_API FAutoConsoleVariableRef(const TCHAR* Name, FString& RefValue, const TCHAR* Help, uint32 Flags = ECVF_Default);
 
 	/**
 	 * Create a reference to a FName console variable
@@ -1557,10 +1517,7 @@ public:
 	 * @param Help must not be 0
 	 * @param Flags bitmask combined from EConsoleVariableFlags
 	 */
-	FAutoConsoleVariableRef(const TCHAR* Name, FName& RefValue, const TCHAR* Help, uint32 Flags = ECVF_Default)
-		: FAutoConsoleObject(IConsoleManager::Get().RegisterConsoleVariableRef(Name, RefValue, Help, Flags))
-	{
-	}
+	CORE_API FAutoConsoleVariableRef(const TCHAR* Name, FName& RefValue, const TCHAR* Help, uint32 Flags = ECVF_Default);
 
 	/**
 	 * Create a reference to a int console variable
@@ -1569,11 +1526,7 @@ public:
 	 * @param Callback Delegate called when the variable changes. @see IConsoleVariable::SetOnChangedCallback
 	 * @param Flags bitmask combined from EConsoleVariableFlags
 	 */
-	FAutoConsoleVariableRef(const TCHAR* Name, int32& RefValue, const TCHAR* Help, const FConsoleVariableDelegate& Callback, uint32 Flags = ECVF_Default)
-		: FAutoConsoleObject(IConsoleManager::Get().RegisterConsoleVariableRef(Name, RefValue, Help, Flags))
-	{
-		AsVariable()->SetOnChangedCallback(Callback);
-	}
+	CORE_API FAutoConsoleVariableRef(const TCHAR* Name, int32& RefValue, const TCHAR* Help, const FConsoleVariableDelegate& Callback, uint32 Flags = ECVF_Default);
 
 	/**
 	 * Create a reference to a float console variable
@@ -1582,11 +1535,7 @@ public:
 	 * @param Callback Delegate called when the variable changes. @see IConsoleVariable::SetOnChangedCallback
 	 * @param Flags bitmask combined from EConsoleVariableFlags
 	 */
-	FAutoConsoleVariableRef(const TCHAR* Name, float& RefValue, const TCHAR* Help, const FConsoleVariableDelegate& Callback, uint32 Flags = ECVF_Default)
-		: FAutoConsoleObject(IConsoleManager::Get().RegisterConsoleVariableRef(Name, RefValue, Help, Flags))
-	{
-		AsVariable()->SetOnChangedCallback(Callback);
-	}
+	CORE_API FAutoConsoleVariableRef(const TCHAR* Name, float& RefValue, const TCHAR* Help, const FConsoleVariableDelegate& Callback, uint32 Flags = ECVF_Default);
 
 	/**
 	 * Create a reference to a bool console variable
@@ -1595,11 +1544,7 @@ public:
 	 * @param Callback Delegate called when the variable changes. @see IConsoleVariable::SetOnChangedCallback
 	 * @param Flags bitmask combined from EConsoleVariableFlags
 	 */
-	FAutoConsoleVariableRef(const TCHAR* Name, bool& RefValue, const TCHAR* Help, const FConsoleVariableDelegate& Callback, uint32 Flags = ECVF_Default)
-		: FAutoConsoleObject(IConsoleManager::Get().RegisterConsoleVariableRef(Name, RefValue, Help, Flags))
-	{
-		AsVariable()->SetOnChangedCallback(Callback);
-	}
+	CORE_API FAutoConsoleVariableRef(const TCHAR* Name, bool& RefValue, const TCHAR* Help, const FConsoleVariableDelegate& Callback, uint32 Flags = ECVF_Default);
 
 	/**
 	 * Create a reference to a FString console variable
@@ -1608,11 +1553,7 @@ public:
 	 * @param Callback Delegate called when the variable changes. @see IConsoleVariable::SetOnChangedCallback
 	 * @param Flags bitmask combined from EConsoleVariableFlags
 	 */
-	FAutoConsoleVariableRef(const TCHAR* Name, FString& RefValue, const TCHAR* Help, const FConsoleVariableDelegate& Callback, uint32 Flags = ECVF_Default)
-		: FAutoConsoleObject(IConsoleManager::Get().RegisterConsoleVariableRef(Name, RefValue, Help, Flags))
-	{
-		AsVariable()->SetOnChangedCallback(Callback);
-	}
+	CORE_API FAutoConsoleVariableRef(const TCHAR* Name, FString& RefValue, const TCHAR* Help, const FConsoleVariableDelegate& Callback, uint32 Flags = ECVF_Default);
 
 	virtual ~FAutoConsoleVariableRef()
 	{
@@ -1785,7 +1726,7 @@ public:
 	 * @param Help must not be 0
 	 * @param Flags bitmask combined from EConsoleVariableFlags
 	 */
-	TAutoConsoleVariable(const TCHAR* Name, const T& DefaultValue, const TCHAR* Help, uint32 Flags = ECVF_Default)
+	FORCENOINLINE TAutoConsoleVariable(const TCHAR* Name, const T& DefaultValue, const TCHAR* Help, uint32 Flags = ECVF_Default)
 		: TAutoConsoleVariable(IConsoleManager::Get().RegisterConsoleVariable(Name, DefaultValue, Help, Flags))
 	{
 	}
@@ -1797,7 +1738,7 @@ public:
 	 * @param Callback Delegate called when the variable changes. @see IConsoleVariable::SetOnChangedCallback
 	 * @param Flags bitmask combined from EConsoleVariableFlags
 	 */
-	TAutoConsoleVariable(const TCHAR* Name, const T& DefaultValue, const TCHAR* Help, const FConsoleVariableDelegate& Callback, uint32 Flags = ECVF_Default)
+	FORCENOINLINE TAutoConsoleVariable(const TCHAR* Name, const T& DefaultValue, const TCHAR* Help, const FConsoleVariableDelegate& Callback, uint32 Flags = ECVF_Default)
 		: TAutoConsoleVariable(Name, DefaultValue, Help, Flags)
 	{
 		AsVariable()->SetOnChangedCallback(Callback);
@@ -1995,10 +1936,7 @@ public:
 	 * @param	Command		The user function to call when this command is executed
 	 * @param	Flags		Optional flags bitmask
 	 */
-	FAutoConsoleCommand(const TCHAR* Name, const TCHAR* Help, const FConsoleCommandDelegate& Command, uint32 Flags = ECVF_Default)
-		: FAutoConsoleObject(IConsoleManager::Get().RegisterConsoleCommand(Name, Help, Command, Flags))
-	{
-	}
+	CORE_API FAutoConsoleCommand(const TCHAR* Name, const TCHAR* Help, const FConsoleCommandDelegate& Command, uint32 Flags = ECVF_Default);
 
 	/**
 	 * Register a console command that takes arguments
@@ -2008,10 +1946,7 @@ public:
 	 * @param	Command		The user function to call when this command is executed
 	 * @param	Flags		Optional flags bitmask
 	 */
-	FAutoConsoleCommand(const TCHAR* Name, const TCHAR* Help, const FConsoleCommandWithArgsDelegate& Command, uint32 Flags = ECVF_Default)
-		: FAutoConsoleObject(IConsoleManager::Get().RegisterConsoleCommand(Name, Help, Command, Flags))
-	{
-	}
+	CORE_API FAutoConsoleCommand(const TCHAR* Name, const TCHAR* Help, const FConsoleCommandWithArgsDelegate& Command, uint32 Flags = ECVF_Default);
 
 	/**
 	* Register a console command that takes an output device
@@ -2021,10 +1956,7 @@ public:
 	* @param	Command		The user function to call when this command is executed
 	* @param	Flags		Optional flags bitmask
 	*/
-	FAutoConsoleCommand(const TCHAR* Name, const TCHAR* Help, const FConsoleCommandWithOutputDeviceDelegate& Command, uint32 Flags = ECVF_Default)
-		: FAutoConsoleObject(IConsoleManager::Get().RegisterConsoleCommand(Name, Help, Command, Flags))
-	{
-	}
+	CORE_API FAutoConsoleCommand(const TCHAR* Name, const TCHAR* Help, const FConsoleCommandWithOutputDeviceDelegate& Command, uint32 Flags = ECVF_Default);
 
     /**
     * Register a console command that takes an output device
@@ -2034,10 +1966,7 @@ public:
     * @param	Command		The user function to call when this command is executed
     * @param	Flags		Optional flags bitmask
     */
-    FAutoConsoleCommand(const TCHAR* Name, const TCHAR* Help, const FConsoleCommandWithArgsAndOutputDeviceDelegate& Command, uint32 Flags = ECVF_Default)
-        : FAutoConsoleObject(IConsoleManager::Get().RegisterConsoleCommand(Name, Help, Command, Flags))
-    {
-    }
+    CORE_API FAutoConsoleCommand(const TCHAR* Name, const TCHAR* Help, const FConsoleCommandWithArgsAndOutputDeviceDelegate& Command, uint32 Flags = ECVF_Default);
 
 	/**
 	* Register a console command that takes a world argument
@@ -2047,10 +1976,7 @@ public:
 	* @param	Command		The user function to call when this command is executed
 	* @param	Flags		Optional flags bitmask
 	*/
-	FAutoConsoleCommand(const TCHAR* Name, const TCHAR* Help, const FConsoleCommandWithWorldDelegate& Command, uint32 Flags = ECVF_Default)
-		: FAutoConsoleObject(IConsoleManager::Get().RegisterConsoleCommand(Name, Help, Command, Flags))
-	{
-	}
+	CORE_API FAutoConsoleCommand(const TCHAR* Name, const TCHAR* Help, const FConsoleCommandWithWorldDelegate& Command, uint32 Flags = ECVF_Default);
 
 	/**
 	* Register a console command that takes arguments, a world argument and an output device
@@ -2060,10 +1986,7 @@ public:
 	* @param	Command		The user function to call when this command is executed
 	* @param	Flags		Optional flags bitmask
 	*/
-	FAutoConsoleCommand(const TCHAR* Name, const TCHAR* Help, const FConsoleCommandWithWorldArgsAndOutputDeviceDelegate& Command, uint32 Flags = ECVF_Default)
-		: FAutoConsoleObject(IConsoleManager::Get().RegisterConsoleCommand(Name, Help, Command, Flags))
-	{
-	}
+	CORE_API FAutoConsoleCommand(const TCHAR* Name, const TCHAR* Help, const FConsoleCommandWithWorldArgsAndOutputDeviceDelegate& Command, uint32 Flags = ECVF_Default);
 };
 #else
 class FAutoConsoleCommand
