ParkingLot: Simplify Wait when there is no timeout

diff --git a/Engine/Source/Runtime/Core/Private/Async/ParkingLot.cpp b/Engine/Source/Runtime/Core/Private/Async/ParkingLot.cpp
--- a/Engine/Source/Runtime/Core/Private/Async/ParkingLot.cpp
+++ b/Engine/Source/Runtime/Core/Private/Async/ParkingLot.cpp
@@ -500,13 +500,56 @@ FThreadLocalData::~FThreadLocalData()
 
 ///////////////////////////////////////////////////////////////////////////////////////////////////////////////////////
 
-FWaitState Wait(
+UE_AUTORTFM_NOAUTORTFM
+FWaitState Wait(const void* Address, bool (*CanWait)(void*), void* CanWaitContext, void (*BeforeWait)(void*), void* BeforeWaitContext)
+{
+	using namespace Private;
+
+	TRefCountPtr<FThread> Self = FThreadLocalData::Get();
+
+	checkfSlow(!Self->WaitAddress.load(std::memory_order_relaxed), TEXT("WaitAddress must be null. This can happen if Wait is called by BeforeWait."));
+	checkfSlow(Self->WakeToken == 0, TEXT("WakeToken must be 0. This is an error in ParkingLot."));
+
+	FWaitState State;
+
+	// Enqueue the thread if CanWait returns true while the bucket is locked.
+	{
+		TDynamicUniqueLock<FWordMutex> BucketLock;
+		FBucket& Bucket = FTable::FindOrCreateBucket(Address, BucketLock);
+		State.bDidWait = !CanWait || CanWait(CanWaitContext);
+		if (!State.bDidWait)
+		{
+			return State;
+		}
+		Self->WaitAddress.store(Address, std::memory_order_relaxed);
+		Self->Event.Reset();
+		Bucket.Enqueue(&*Self);
+	}
+
+	// BeforeWait must be invoked after the bucket is unlocked.
+	if (BeforeWait)
+	{
+		BeforeWait(BeforeWaitContext);
+	}
+
+	// Wait until the thread has been dequeued.
+	Self->Event.Wait();
+
+	// WaitAddress is reset when the thread is dequeued.
+	verifySlow(!Self->WaitAddress.load(std::memory_order_relaxed));
+	State.bDidWake = true;
+	State.WakeToken = Self->WakeToken;
+	Self->WakeToken = 0;
+	return State;
+}
+
+static FWaitState TimedWait(
 	const void* Address,
 	bool (*CanWait)(void*),
 	void* CanWaitContext,
 	void (*BeforeWait)(void*),
 	void* BeforeWaitContext,
-	void (*WaitOnEvent)(FPlatformManualResetEvent&, void*),
+	void (*WaitOnEvent)(void*, FPlatformManualResetEvent&),
 	void* WaitOnEventContext)
 {
 	using namespace Private;
@@ -539,7 +582,7 @@ FWaitState Wait(
 	}
 
 	// Wait until the timeout or until the thread has been dequeued.
-	WaitOnEvent(Self->Event, WaitOnEventContext);
+	WaitOnEvent(WaitOnEventContext, Self->Event);
 
 	// WaitAddress is reset when the thread is dequeued.
 	if (!Self->WaitAddress.load(std::memory_order_relaxed))
@@ -580,27 +623,28 @@ FWaitState Wait(
 	return State;
 }
 
-UE_AUTORTFM_NOAUTORTFM
-FWaitState Wait(const void* Address, bool (*CanWait)(void*), void* CanWaitContext, void (*BeforeWait)(void*), void* BeforeWaitContext)
+static void WaitForTime(void* WaitTime, FPlatformManualResetEvent& Event)
 {
-	return Wait(Address, CanWait, CanWaitContext, BeforeWait, BeforeWaitContext,
-		[](FPlatformManualResetEvent& Event, void*) { Event.Wait(); }, nullptr);
+	Event.WaitFor(*(FMonotonicTimeSpan*)WaitTime);
 }
 
 UE_AUTORTFM_NOAUTORTFM
 FWaitState WaitFor(const void* Address, bool (*CanWait)(void*), void* CanWaitContext, void (*BeforeWait)(void*), void* BeforeWaitContext, FMonotonicTimeSpan WaitTime)
 {
 	check(!WaitTime.IsNaN());
-	return Wait(Address, CanWait, CanWaitContext, BeforeWait, BeforeWaitContext,
-		[](FPlatformManualResetEvent& Event, void* Context) { Event.WaitFor(*(FMonotonicTimeSpan*)Context); }, &WaitTime);
+	return TimedWait(Address, CanWait, CanWaitContext, BeforeWait, BeforeWaitContext, WaitForTime, &WaitTime);
+}
+
+static void WaitUntilTime(void* WaitTime, FPlatformManualResetEvent& Event)
+{
+	Event.WaitUntil(*(FMonotonicTimePoint*)WaitTime);
 }
 
 UE_AUTORTFM_NOAUTORTFM
 FWaitState WaitUntil(const void* Address, bool (*CanWait)(void*), void* CanWaitContext, void (*BeforeWait)(void*), void* BeforeWaitContext, FMonotonicTimePoint WaitTime)
 {
 	check(!WaitTime.IsNaN());
-	return Wait(Address, CanWait, CanWaitContext, BeforeWait, BeforeWaitContext,
-		[](FPlatformManualResetEvent& Event, void* Context) { Event.WaitUntil(*(FMonotonicTimePoint*)Context); }, &WaitTime);
+	return TimedWait(Address, CanWait, CanWaitContext, BeforeWait, BeforeWaitContext, WaitUntilTime, &WaitTime);
 }
 
 UE_AUTORTFM_NOAUTORTFM
@@ -631,13 +675,18 @@ void WakeOne(const void* Address, uint64 (*OnWakeState)(void*, FWakeState), void
 
 	if (WakeThread)
 	{
-		checkSlow(WakeThread->WaitAddress.load(std::memory_order_relaxed) == Address);
 		WakeThread->WakeToken = WakeToken;
 		WakeThread->WaitAddress.store(nullptr, std::memory_order_relaxed);
 		WakeThread->Event.Notify();
 	}
 }
 
+static uint64 CopyWakeState(void* OutState, FWakeState State)
+{
+	*(FWakeState*)OutState = State;
+	return 0;
+}
+
 } // UE::ParkingLot::Private
 
 namespace UE::ParkingLot
@@ -646,11 +695,7 @@ namespace UE::ParkingLot
 FWakeState WakeOne(const void* Address)
 {
 	FWakeState OutState;
-	WakeOne(Address, [&OutState](FWakeState State) -> uint64
-	{
-		OutState = State;
-		return 0;
-	});
+	Private::WakeOne(Address, Private::CopyWakeState, &OutState);
 	return OutState;
 }
 
@@ -676,7 +721,6 @@ uint32 WakeMultiple(const void* const Address, const uint32 WakeCount)
 
 	for (FThread* WakeThread : WakeThreads)
 	{
-		checkSlow(WakeThread->WaitAddress.load(std::memory_order_relaxed) == Address);
 		WakeThread->WaitAddress.store(nullptr, std::memory_order_relaxed);
 		WakeThread->Event.Notify();
 	}
