Free allocated chunks in TChunkedPackedArray destructor to avoid false positives in LeakDetector. Also fixed default value for UE_TEXTKEY_USE_INLINE_ALLOCATOR

#rnx

diff --git a/Engine/Source/Runtime/Core/Private/Internationalization/TextKey.cpp b/Engine/Source/Runtime/Core/Private/Internationalization/TextKey.cpp
--- a/Engine/Source/Runtime/Core/Private/Internationalization/TextKey.cpp
+++ b/Engine/Source/Runtime/Core/Private/Internationalization/TextKey.cpp
@@ -33,7 +33,7 @@ DEFINE_LOG_CATEGORY_STATIC(LogTextKey, Log, All);
 #endif
 
 #ifndef UE_TEXTKEY_USE_INLINE_ALLOCATOR
-	#define UE_TEXTKEY_USE_INLINE_ALLOCATOR 1 //(0)
+	#define UE_TEXTKEY_USE_INLINE_ALLOCATOR (0)
 #endif
 
 #ifndef UE_TEXTKEY_MAX_SIZE_BYTES
@@ -103,6 +103,14 @@ public:
 		return Count;
 	}
 	
+	~TChunkedPackedArray()
+	{
+		for(uint8* Chunk : Chunks)
+		{
+			FMemory::Free(Chunk);
+		}
+	}
+	
 private:
 	TArray<uint8*> Chunks;
 	TArray<uint32> ChunkCursors;
