[AutoRTFM] Implement FTransactionallySafeCriticalSection in terms of recursive mutexes.

This has a few benefits.
- Recursive mutexes are smaller objects (8 bytes); this makes our State object only twelve bytes in total.
- Recursive mutexes are slightly faster than critical sections.
- Critical sections need to be created and destroyed by the OS, and require an Open block during FState ctor/dtor. Recursive mutex initialization is just atomically writing two zeros, and destruction is a no-op.
- The advantage offered by critical sections--a stronger fairness guarantee--is not interesting to us. We give the main thread preferential treatment anyway.

This change has been on hold for a long while because it significantly hurt performance in Verse PresubmitTest.py (50x+ slowdown). This performance issue has recently been fixed in CL https://p4-swarm.epicgames.net/changes/43576941 .

diff --git a/Engine/Source/Runtime/Core/Public/Async/TransactionallySafeRecursiveMutex.h b/Engine/Source/Runtime/Core/Public/Async/TransactionallySafeRecursiveMutex.h
--- a/Engine/Source/Runtime/Core/Public/Async/TransactionallySafeRecursiveMutex.h
+++ b/Engine/Source/Runtime/Core/Public/Async/TransactionallySafeRecursiveMutex.h
@@ -2,20 +2,126 @@
 
 #pragma once
 
-#include "Async/RecursiveMutex.h"
 #include "AutoRTFM.h"
+#include "Async/RecursiveMutex.h"
 
 #if UE_AUTORTFM
-#include "Misc/TransactionallySafeCriticalSection.h"
+#include "Templates/SharedPointer.h"
 #endif
 
 namespace UE
 {
 
 #if UE_AUTORTFM
-using FTransactionallySafeRecursiveMutex = ::FTransactionallySafeCriticalSection;
+
+// A transactionally safe recursive mutex that works in the following novel ways:
+// - In the open (non-transactional):
+//   - Take the lock like before. Simple!
+//   - Free the lock like before too.
+// - In the closed (transactional):
+//   - During locking we query `TransactionalLockCount`:
+//     - 0 means we haven't taken the lock within our transaction nest and need to acquire the lock.
+//     - Otherwise we already have the lock (and are preventing non-transactional code seeing any
+//       modifications we've made while holding the lock), so just bump `TransactionalLockCount`.
+//     - We also register an on-abort handler to release the lock should we abort (but we need to
+//       query `TransactionalLockCount` even there because we could be aborting an inner transaction
+//       and the parent transaction still wants to have the lock held!).
+//   - During unlocking we defer doing the unlock until the transaction commits.
+//
+// Thus with this approach we will hold this lock for the *entirety* of the transactional nest should
+// we take the lock during the transaction, thus preventing non-transactional code from seeing any
+// modifications we should make.
+struct FTransactionallySafeRecursiveMutexDefinition final
+{
+	FTransactionallySafeRecursiveMutexDefinition() : State(MakeShared<FState>())
+	{
+	}
+
+	void Lock()
+	{
+		if (AutoRTFM::IsTransactional() || AutoRTFM::IsCommittingOrAborting())
+		{
+			UE_AUTORTFM_OPEN
+			{
+				// The transactional system which can increment TransactionalLockCount
+				// is always single-threaded, thus this is safe to check without atomicity.
+				if (0 == State->TransactionalLockCount)
+				{
+					State->Mutex.Lock();
+				}
+
+				State->TransactionalLockCount += 1;
+			};
+
+			// We explicitly copy the state here for the case that `this` was stack
+			// allocated and has already died before the on-abort is hit.
+			AutoRTFM::OnAbort([State = this->State]
+			{
+				ensure(0 != State->TransactionalLockCount);
+
+				State->TransactionalLockCount -= 1;
+
+				if (0 == State->TransactionalLockCount)
+				{
+					State->Mutex.Unlock();
+				}
+			});
+		}
+		else
+		{
+			State->Mutex.Lock();
+			ensure(0 == State->TransactionalLockCount);
+		}
+	}
+
+	void Unlock()
+	{
+		if (AutoRTFM::IsTransactional() || AutoRTFM::IsCommittingOrAborting())
+		{
+			// We explicitly copy the state here for the case that `this` was stack
+			// allocated and has already died before the on-commit is hit.
+			AutoRTFM::OnCommit([State = this->State]
+			{
+				ensure(0 != State->TransactionalLockCount);
+
+				State->TransactionalLockCount -= 1;
+
+				if (0 == State->TransactionalLockCount)
+				{
+					State->Mutex.Unlock();
+				}
+			});
+		}
+		else
+		{
+			ensure(0 == State->TransactionalLockCount);
+			State->Mutex.Unlock();
+		}
+	}
+
+private:
+	UE_NONCOPYABLE(FTransactionallySafeRecursiveMutexDefinition)
+
+	struct FState final
+	{
+		UE::FRecursiveMutex Mutex;
+		uint32 TransactionalLockCount = 0;
+
+		FState() = default;
+		
+		~FState()
+		{
+			ensure(0 == TransactionalLockCount);
+		}
+	};
+
+	TSharedPtr<FState> State;
+};
+
+using FTransactionallySafeRecursiveMutex = ::UE::FTransactionallySafeRecursiveMutexDefinition;
+
 #else
 using FTransactionallySafeRecursiveMutex = ::UE::FRecursiveMutex;
-#endif // UE_AUTORTFM
+#endif
 
-} // UE
+} // namespace UE
diff --git a/Engine/Source/Runtime/Core/Public/Misc/TransactionallySafeCriticalSection.h b/Engine/Source/Runtime/Core/Public/Misc/TransactionallySafeCriticalSection.h
--- a/Engine/Source/Runtime/Core/Public/Misc/TransactionallySafeCriticalSection.h
+++ b/Engine/Source/Runtime/Core/Public/Misc/TransactionallySafeCriticalSection.h
@@ -2,122 +2,17 @@
 
 #pragma once
 
-#include "HAL/CriticalSection.h"
 #include "AutoRTFM.h"
-#include "Templates/SharedPointer.h"
 
 #if UE_AUTORTFM
+#include "Async/TransactionallySafeRecursiveMutex.h"
+#else
+#include "HAL/CriticalSection.h"
+#endif // UE_AUTORTFM
 
-// A transactionally safe critical section that works in the following novel ways:
-// - In the open (non-transactional):
-//   - Take the lock like before. Simple!
-//   - Free the lock like before too.
-// - In the closed (transactional):
-//   - During locking we query `TransactionalLockCount`:
-//     - 0 means we haven't taken the lock within our transaction nest and need to acquire the lock.
-//     - Otherwise we already have the lock (and are preventing non-transactional code seeing any
-//       modifications we've made while holding the lock), so just bump `TransactionalLockCount`.
-//     - We also register an on-abort handler to release the lock should we abort (but we need to
-//       query `TransactionalLockCount` even there because we could be aborting an inner transaction
-//       and the parent transaction still wants to have the lock held!).
-//   - During unlocking we defer doing the unlock until the transaction commits.
-//
-// Thus with this approach we will hold this lock for the *entirety* of the transactional nest should
-// we take the lock during the transaction, thus preventing non-transactional code from seeing any
-// modifications we should make.
-struct FTransactionallySafeCriticalSectionDefinition final
-{
-	FTransactionallySafeCriticalSectionDefinition() : State(MakeShared<FState>())
-	{
-	}
-
-	void Lock()
-	{
-		if (AutoRTFM::IsTransactional() || AutoRTFM::IsCommittingOrAborting())
-		{
-			UE_AUTORTFM_OPEN
-				{
-					// The transactional system which can increment TransactionalLockCount
-					// is always single-threaded, thus this is safe to check without atomicity.
-					if (0 == State->TransactionalLockCount)
-					{
-						State->CriticalSection.Lock();
-					}
-
-					State->TransactionalLockCount += 1;
-				};
-
-			// We explicitly copy the state here for the case that `this` was stack
-			// allocated and has already died before the on-abort is hit.
-			AutoRTFM::OnAbort([State = this->State]
-				{
-					ensure(0 != State->TransactionalLockCount);
-
-					State->TransactionalLockCount -= 1;
-
-					if (0 == State->TransactionalLockCount)
-					{
-						State->CriticalSection.Unlock();
-					}
-				});
-		}
-		else
-		{
-			State->CriticalSection.Lock();
-			ensure(0 == State->TransactionalLockCount);
-		}
-	}
-
-	void Unlock()
-	{
-		if (AutoRTFM::IsTransactional() || AutoRTFM::IsCommittingOrAborting())
-		{
-			// We explicitly copy the state here for the case that `this` was stack
-			// allocated and has already died before the on-commit is hit.
-			AutoRTFM::OnCommit([State = this->State]
-				{
-					ensure(0 != State->TransactionalLockCount);
-
-					State->TransactionalLockCount -= 1;
-
-					if (0 == State->TransactionalLockCount)
-					{
-						State->CriticalSection.Unlock();
-					}
-				});
-		}
-		else
-		{
-			ensure(0 == State->TransactionalLockCount);
-			State->CriticalSection.Unlock();
-		}
-	}
-
-private:
-	UE_NONCOPYABLE(FTransactionallySafeCriticalSectionDefinition)
-
-	struct FState final
-	{
-		FCriticalSection CriticalSection;
-		uint32 TransactionalLockCount = 0;
-
-		// Constructor is always open because FCriticalSection is not transactionally safe.
-		UE_AUTORTFM_ALWAYS_OPEN
-		FState() = default;
-		
-		// Destructor is always open because FCriticalSection is not transactionally safe.
-		UE_AUTORTFM_ALWAYS_OPEN
-		~FState()
-		{
-			ensure(0 == TransactionalLockCount);
-		}
-	};
-
-	TSharedPtr<FState> State;
-};
-
-using FTransactionallySafeCriticalSection = FTransactionallySafeCriticalSectionDefinition;
-
+#if UE_AUTORTFM
+using FTransactionallySafeCriticalSection = ::UE::FTransactionallySafeRecursiveMutex;
 #else
-using FTransactionallySafeCriticalSection = FCriticalSection;
-#endif
+using FTransactionallySafeCriticalSection = ::FCriticalSection;
+#endif // UE_AUTORTFM
+
