Add csv metadata to indicate when merged modules is enabled.

#rnx

diff --git a/Engine/Source/Runtime/Core/Private/ProfilingDebugging/CsvProfiler.cpp b/Engine/Source/Runtime/Core/Private/ProfilingDebugging/CsvProfiler.cpp
--- a/Engine/Source/Runtime/Core/Private/ProfilingDebugging/CsvProfiler.cpp
+++ b/Engine/Source/Runtime/Core/Private/ProfilingDebugging/CsvProfiler.cpp
@@ -3479,6 +3479,7 @@ FCsvProfiler::FCsvProfiler()
 	SetMetadataInternal(TEXT("PGOProfilingEnabled"), PLATFORM_COMPILER_OPTIMIZATION_PG_PROFILING ? TEXT("1") : TEXT("0"));//True if Profile Guided Optimisation Instrumentation is enabled 
 	SetMetadataInternal(TEXT("LTOEnabled"), PLATFORM_COMPILER_OPTIMIZATION_LTCG ? TEXT("1") : TEXT("0"));//True if Link Time Optimisation is enabled 
 	SetMetadataInternal(TEXT("ASan"), USING_ADDRESS_SANITISER ? TEXT("1") : TEXT("0"));
+	SetMetadataInternal(TEXT("MergedModules"), UE_MERGED_MODULES ? TEXT("1") : TEXT("0"));
 
 	FCoreDelegates::OnSystemResolutionChanged.AddLambda([](uint32 ResX, uint32 ResY)
 		{
