Added CContiguousRange concept, based on syntactic constraints instead of a trait specialization.
Added CPointer concept.

diff --git a/Engine/Plugins/Experimental/PlainProps/Source/Public/PlainPropsWrite.h b/Engine/Plugins/Experimental/PlainProps/Source/Public/PlainPropsWrite.h
--- a/Engine/Plugins/Experimental/PlainProps/Source/Public/PlainPropsWrite.h
+++ b/Engine/Plugins/Experimental/PlainProps/Source/Public/PlainPropsWrite.h
@@ -3,6 +3,7 @@
 #pragma once
 
 #include "PlainPropsTypes.h"
+#include "Concepts/ContiguousRange.h"
 #include "Containers/Array.h"
 #include "Templates/UniquePtr.h"
 #include <type_traits>
@@ -43,8 +44,8 @@ inline void WriteData(TArray64<uint8>& Out, const void* Data, int64 Size)
 	Out.Append(static_cast<const uint8*>(Data), Size);
 }
 
-template<typename ArrayType>
-void WriteArray(TArray64<uint8>& Out, const ArrayType& In) requires (!!TIsContiguousContainer<ArrayType>::Value)
+template<::UE::CContiguousRange ArrayType>
+void WriteArray(TArray64<uint8>& Out, const ArrayType& In)
 {
 	WriteData(Out, In.GetData(), sizeof(typename ArrayType::ElementType) * In.Num());
 }
diff --git a/Engine/Source/Runtime/Core/Public/Concepts/Pointer.h b/Engine/Source/Runtime/Core/Public/Concepts/Pointer.h
new file mode 100644
--- /dev/null
+++ b/Engine/Source/Runtime/Core/Public/Concepts/Pointer.h
@@ -0,0 +1,14 @@
+// Copyright Epic Games, Inc. All Rights Reserved.
+
+#pragma once
+
+#include <type_traits>
+
+namespace UE
+{
+	/**
+	 * Concept which describes a pointer type.
+	 */
+	template <typename T>
+	concept CPointer = std::is_pointer_v<T>;
+}
diff --git a/Engine/Source/Runtime/Core/Public/Traits/IsContiguousContainer.h b/Engine/Source/Runtime/Core/Public/Traits/IsContiguousContainer.h
--- a/Engine/Source/Runtime/Core/Public/Traits/IsContiguousContainer.h
+++ b/Engine/Source/Runtime/Core/Public/Traits/IsContiguousContainer.h
@@ -17,7 +17,7 @@ struct TIsContiguousContainer
 	UE_STATIC_ASSERT_COMPLETE_TYPE(T, "TIsContiguousContainer instantiated with an incomplete type");
 	// This sizeof(T) == 0 test will let Clang treat zero-sized arrays as contiguous containers.
 	// An alternate fix using a specialization of TIsContiguousContainer (which Clang doesn't support) can be found below.
-	enum { Value = sizeof(T) == 0 };
+	static constexpr bool Value = sizeof(T) == 0;
 };
 
 template <typename T> struct TIsContiguousContainer<             T& > : TIsContiguousContainer<T> {};
@@ -29,28 +29,28 @@ template <typename T> struct TIsContiguousContainer<const volatile T> : TIsConti
 /**
  * Specialization for C arrays (always contiguous)
  */
-template <typename T, size_t N> struct TIsContiguousContainer<               T[N]> { enum { Value = true }; };
-template <typename T, size_t N> struct TIsContiguousContainer<const          T[N]> { enum { Value = true }; };
-template <typename T, size_t N> struct TIsContiguousContainer<      volatile T[N]> { enum { Value = true }; };
-template <typename T, size_t N> struct TIsContiguousContainer<const volatile T[N]> { enum { Value = true }; };
+template <typename T, size_t N> struct TIsContiguousContainer<               T[N]> { static constexpr bool Value = true; };
+template <typename T, size_t N> struct TIsContiguousContainer<const          T[N]> { static constexpr bool Value = true; };
+template <typename T, size_t N> struct TIsContiguousContainer<      volatile T[N]> { static constexpr bool Value = true; };
+template <typename T, size_t N> struct TIsContiguousContainer<const volatile T[N]> { static constexpr bool Value = true; };
 
 /**
  * Specialization for zero-sized arrays - which are not handled by the above Clang-specific fix in the primary TIsContiguousContainer template.
  */
 #ifndef __clang__
-	template <typename T> struct TIsContiguousContainer<               T[0]> { enum { Value = true }; };
-	template <typename T> struct TIsContiguousContainer<const          T[0]> { enum { Value = true }; };
-	template <typename T> struct TIsContiguousContainer<      volatile T[0]> { enum { Value = true }; };
-	template <typename T> struct TIsContiguousContainer<const volatile T[0]> { enum { Value = true }; };
+	template <typename T> struct TIsContiguousContainer<               T[0]> { static constexpr bool Value = true; };
+	template <typename T> struct TIsContiguousContainer<const          T[0]> { static constexpr bool Value = true; };
+	template <typename T> struct TIsContiguousContainer<      volatile T[0]> { static constexpr bool Value = true; };
+	template <typename T> struct TIsContiguousContainer<const volatile T[0]> { static constexpr bool Value = true; };
 #endif
 
 /**
  * Specialization for unbounded C arrays (never contiguous - should be treated as pointers which are not regarded as contiguous containers)
  */
-template <typename T> struct TIsContiguousContainer<               T[]> { enum { Value = false }; };
-template <typename T> struct TIsContiguousContainer<const          T[]> { enum { Value = false }; };
-template <typename T> struct TIsContiguousContainer<      volatile T[]> { enum { Value = false }; };
-template <typename T> struct TIsContiguousContainer<const volatile T[]> { enum { Value = false }; };
+template <typename T> struct TIsContiguousContainer<               T[]> { static constexpr bool Value = false; };
+template <typename T> struct TIsContiguousContainer<const          T[]> { static constexpr bool Value = false; };
+template <typename T> struct TIsContiguousContainer<      volatile T[]> { static constexpr bool Value = false; };
+template <typename T> struct TIsContiguousContainer<const volatile T[]> { static constexpr bool Value = false; };
 
 /**
  * Specialization for initializer lists (also always contiguous)
@@ -58,5 +58,8 @@ template <typename T> struct TIsContiguousContainer<const volatile T[]> { enum {
 template <typename T>
 struct TIsContiguousContainer<std::initializer_list<T>>
 {
-	enum { Value = true };
+	static constexpr bool Value = true;
 };
+
+template <typename T>
+inline constexpr bool TIsContiguousContainer_V = TIsContiguousContainer<T>::Value;
