Missing header.

diff --git a/Engine/Source/Runtime/Core/Public/Concepts/ContiguousRange.h b/Engine/Source/Runtime/Core/Public/Concepts/ContiguousRange.h
new file mode 100644
--- /dev/null
+++ b/Engine/Source/Runtime/Core/Public/Concepts/ContiguousRange.h
@@ -0,0 +1,21 @@
+// Copyright Epic Games, Inc. All Rights Reserved. ** 
+
+#pragma once
+
+#include "Concepts/Integral.h"
+#include "Concepts/Pointer.h"
+#include "Templates/UnrealTemplate.h"
+
+namespace UE
+{
+	/**
+	 * Concept which describes a contiguous range of elements.
+	 */
+	template <typename T>
+	concept CContiguousRange =
+		requires (const T& Range)
+		{
+			{ GetData(Range) } -> CPointer;
+			{ GetNum(Range) } -> CIntegral;
+		};
+}
