Moved LexTryParseString into the per-char-type header instead of being a template.
Added CArithmetic concept.

diff --git a/Engine/Source/Runtime/Core/Public/Concepts/Arithmetic.h b/Engine/Source/Runtime/Core/Public/Concepts/Arithmetic.h
new file mode 100644
--- /dev/null
+++ b/Engine/Source/Runtime/Core/Public/Concepts/Arithmetic.h
@@ -0,0 +1,14 @@
+// Copyright Epic Games, Inc. All Rights Reserved.
+
+#pragma once
+
+#include <type_traits>
+
+namespace UE
+{
+	/**
+	 * Concept which describes an arithmetic type.
+	 */
+	template <typename T>
+	concept CArithmetic = std::is_arithmetic_v<T>;
+}
diff --git a/Engine/Source/Runtime/Core/Public/Containers/UnrealString.h b/Engine/Source/Runtime/Core/Public/Containers/UnrealString.h
--- a/Engine/Source/Runtime/Core/Public/Containers/UnrealString.h
+++ b/Engine/Source/Runtime/Core/Public/Containers/UnrealString.h
@@ -213,54 +213,6 @@ struct TTypeToString
 	}
 };
 
-/** Parse a string into this type, returning whether it was successful */
-/** Specialization for arithmetic types */
-template <
-	typename T,
-	typename CharType
-	UE_REQUIRES(std::is_arithmetic_v<T>)
->
-bool LexTryParseString(T& OutValue, const CharType* Buffer)
-{
-	if (Buffer[0] == '\0')
-	{
-		OutValue = 0;
-		return false;
-	}
-
-	LexFromString(OutValue, Buffer);
-	if (OutValue == 0 && FMath::IsFinite((float)OutValue)) //@TODO:FLOATPRECISION: ? huh ?
-	{
-		bool bSawZero = false;
-		CharType C = *Buffer;
-		while (C != CHARTEXT(CharType, '\0') && (C == CHARTEXT(CharType, '+') || C == CHARTEXT(CharType, '-') || TChar<CharType>::IsWhitespace(C)))
-		{
-			C = *(++Buffer);
-		}
-
-		while (C != CHARTEXT(CharType, '\0') && !TChar<CharType>::IsWhitespace(C) && (std::is_floating_point_v<T> || C != CHARTEXT(CharType, '.')))
-		{
-			bSawZero = bSawZero || (C == CHARTEXT(CharType, '0'));
-			if (!bSawZero && C != CHARTEXT(CharType, '.'))
-			{
-				return false;
-			}
-
-			C = *(++Buffer);
-		}
-		return bSawZero;
-	}
-
-	return true;
-}
-
-/** Try and parse a bool - always returns true */
-template <typename CharType>
-inline bool LexTryParseString(bool& OutValue, const CharType* Buffer)
-{
-	LexFromString(OutValue, Buffer);
-	return true;
-}
 
 template<typename T>
 struct TTypeFromString
diff --git a/Engine/Source/Runtime/Core/Public/Containers/UnrealString.h.inl b/Engine/Source/Runtime/Core/Public/Containers/UnrealString.h.inl
--- a/Engine/Source/Runtime/Core/Public/Containers/UnrealString.h.inl
+++ b/Engine/Source/Runtime/Core/Public/Containers/UnrealString.h.inl
@@ -2431,6 +2431,49 @@ inline void LexFromString(double& OutValue, 			const UE_STRING_CHARTYPE* Buffer)
 inline void LexFromString(bool& OutValue, 				const UE_STRING_CHARTYPE* Buffer)	{	OutValue = TCString<UE_STRING_CHARTYPE>::ToBool(Buffer);	}
 inline void LexFromString(UE_STRING_CLASS& OutValue,	const UE_STRING_CHARTYPE* Buffer)	{	OutValue = Buffer;						}
 
+/** Parse a string into this type, returning whether it was successful */
+/** Specialization for arithmetic types */
+template <UE::CArithmetic T>
+bool LexTryParseString(T& OutValue, const UE_STRING_CHARTYPE* Buffer)
+{
+	if (Buffer[0] == '\0')
+	{
+		OutValue = 0;
+		return false;
+	}
+
+	LexFromString(OutValue, Buffer);
+	if (OutValue == 0 && FMath::IsFinite((float)OutValue)) //@TODO:FLOATPRECISION: ? huh ?
+	{
+		bool bSawZero = false;
+		UE_STRING_CHARTYPE C = *Buffer;
+		while (C != CHARTEXT(UE_STRING_CHARTYPE, '\0') && (C == CHARTEXT(UE_STRING_CHARTYPE, '+') || C == CHARTEXT(UE_STRING_CHARTYPE, '-') || TChar<UE_STRING_CHARTYPE>::IsWhitespace(C)))
+		{
+			C = *(++Buffer);
+		}
+
+		while (C != CHARTEXT(UE_STRING_CHARTYPE, '\0') && !TChar<UE_STRING_CHARTYPE>::IsWhitespace(C) && (std::is_floating_point_v<T> || C != CHARTEXT(UE_STRING_CHARTYPE, '.')))
+		{
+			bSawZero = bSawZero || (C == CHARTEXT(UE_STRING_CHARTYPE, '0'));
+			if (!bSawZero && C != CHARTEXT(UE_STRING_CHARTYPE, '.'))
+			{
+				return false;
+			}
+
+			C = *(++Buffer);
+		}
+		return bSawZero;
+	}
+
+	return true;
+}
+
+/** Try and parse a bool - always returns true */
+inline bool LexTryParseString(bool& OutValue, const UE_STRING_CHARTYPE* Buffer)
+{
+	LexFromString(OutValue, Buffer);
+	return true;
+}
 template <typename StringType = FString>
 [[nodiscard]] FORCEINLINE StringType LexToString(UE_STRING_CLASS&& Str)
 {
diff --git a/Engine/Source/Runtime/Core/Public/Containers/UnrealStringIncludes.h.inl b/Engine/Source/Runtime/Core/Public/Containers/UnrealStringIncludes.h.inl
--- a/Engine/Source/Runtime/Core/Public/Containers/UnrealStringIncludes.h.inl
+++ b/Engine/Source/Runtime/Core/Public/Containers/UnrealStringIncludes.h.inl
@@ -26,6 +26,7 @@
 #include "Templates/UnrealTypeTraits.h"
 #include "Templates/UnrealTemplate.h"
 #include "Math/NumericLimits.h"
+#include "Concepts/Arithmetic.h"
 #include "Containers/Array.h"
 #include "Misc/CString.h"
 #include "Misc/Crc.h"
