Added CCharType concept and applied it in TStringView.

diff --git a/Engine/Source/Runtime/Core/Public/Concepts/CharType.h b/Engine/Source/Runtime/Core/Public/Concepts/CharType.h
new file mode 100644
--- /dev/null
+++ b/Engine/Source/Runtime/Core/Public/Concepts/CharType.h
@@ -0,0 +1,14 @@
+// Copyright Epic Games, Inc. All Rights Reserved.
+
+#pragma once
+
+#include "Traits/IsCharType.h"
+
+namespace UE
+{
+	/**
+	 * Concept which describes a character encoding type.
+	 */
+	template <typename T>
+	concept CCharType = TIsCharType_V<T>;
+}
diff --git a/Engine/Source/Runtime/Core/Public/Containers/StringFwd.h b/Engine/Source/Runtime/Core/Public/Containers/StringFwd.h
--- a/Engine/Source/Runtime/Core/Public/Containers/StringFwd.h
+++ b/Engine/Source/Runtime/Core/Public/Containers/StringFwd.h
@@ -3,6 +3,7 @@
 #pragma once
 
 #include "CoreTypes.h"
+#include "Concepts/CharType.h"
 #include "Traits/ElementType.h"
 #include "Traits/IsContiguousContainer.h"
 
@@ -34,9 +35,10 @@ template <int32 BufferSize> using TUtf8StringBuilder = TStringBuilderWithBuffer<
 
 // String View
 
-template <typename CharType> class TStringView;
+template <UE::CCharType T>
+class TStringView;
 
-template <typename CharType> struct TIsContiguousContainer<TStringView<CharType>> { static constexpr bool Value = true; };
+template <UE::CCharType CharType> struct TIsContiguousContainer<TStringView<CharType>> { static constexpr bool Value = true; };
 
 using FStringView = TStringView<TCHAR>;
 using FAnsiStringView = TStringView<ANSICHAR>;
diff --git a/Engine/Source/Runtime/Core/Public/Containers/StringView.h b/Engine/Source/Runtime/Core/Public/Containers/StringView.h
--- a/Engine/Source/Runtime/Core/Public/Containers/StringView.h
+++ b/Engine/Source/Runtime/Core/Public/Containers/StringView.h
@@ -3,6 +3,7 @@
 #pragma once
 
 #include "CoreTypes.h"
+#include "Concepts/CharType.h"
 #include "Containers/StringFwd.h" // IWYU pragma: export
 #include "HAL/UnrealMemory.h"
 #include "Math/NumericLimits.h"
@@ -28,7 +29,7 @@ namespace UE::Core::Private
 		return GetData(Forward<ArgTypes>(Args)...);
 	}
 
-	template<typename CharType>
+	template<UE::CCharType CharType>
 	constexpr int32 StringLength(const CharType* InString)
 	{
 		int32 Count = 0;
@@ -96,7 +97,7 @@ namespace UE::Core::Private
  *	UseString(TEXTVIEW("ABC"));
  * @endcode
  */
-template <typename CharType>
+template <UE::CCharType CharType>
 class TStringView
 {
 public:
@@ -122,15 +123,8 @@ public:
 	}
 
 	/** Construct a view of the null-terminated string pointed to by InData. */
-	template <
-		typename OtherCharType
-		UE_REQUIRES(
-			TAnd<
-				TIsCharType<OtherCharType>,
-				TIsCharEncodingCompatibleWith<OtherCharType, CharType>
-			>::Value
-		)
-	>
+	template <UE::CCharType OtherCharType>
+		requires (TIsCharEncodingCompatibleWith<OtherCharType, CharType>::Value)
 	[[nodiscard]] inline TStringView(const OtherCharType* InData UE_LIFETIMEBOUND)
 		: DataPtr((const CharType*)InData)
 		, Size(UE::Core::Private::StringLength((const CharType*)InData))
@@ -138,15 +132,8 @@ public:
 	}
 
 	/** Construct a view of InSize characters beginning at InData. */
-	template <
-		typename OtherCharType
-		UE_REQUIRES(
-			TAnd<
-				TIsCharType<OtherCharType>,
-				TIsCharEncodingCompatibleWith<OtherCharType, CharType>
-			>::Value
-		)
-	>
+	template <UE::CCharType OtherCharType>
+		requires (TIsCharEncodingCompatibleWith<OtherCharType, CharType>::Value)
 	[[nodiscard]] constexpr inline TStringView(const OtherCharType* InData UE_LIFETIMEBOUND, int32 InSize)
 		: DataPtr((const CharType*)InData)
 		, Size(InSize)
@@ -333,10 +320,7 @@ public:
 	 * @param Other        A string that is comparable with the character type of this view.
 	 * @param SearchCase Whether the comparison should ignore case.
 	 */
-	template <
-		typename OtherCharType
-		UE_REQUIRES(TIsCharType_V<OtherCharType>)
-	>
+	template <UE::CCharType OtherCharType>
 	[[nodiscard]] inline bool Equals(const OtherCharType* Other, ESearchCase::Type SearchCase = ESearchCase::CaseSensitive) const
 	{
 		return PrivateEquals(Other, SearchCase);
@@ -362,10 +346,7 @@ public:
 	 * @param SearchCase   Whether the comparison should ignore case.
 	 * @return 0 is equal, negative if this view is less, positive if this view is greater.
 	 */
-	template <
-		typename OtherCharType
-		UE_REQUIRES(TIsCharType_V<OtherCharType>)
-	>
+	template <UE::CCharType OtherCharType>
 	[[nodiscard]] inline int32 Compare(TStringView<OtherCharType> Other, ESearchCase::Type SearchCase = ESearchCase::CaseSensitive) const
 	{
 		return PrivateCompare(Other, SearchCase);
@@ -378,7 +359,7 @@ public:
 	 * @param SearchCase Whether the comparison should ignore case.
 	 * @return 0 is equal, negative if this view is less, positive if this view is greater.
 	 */
-	template <typename OtherCharType>
+	template <UE::CCharType OtherCharType>
 	[[nodiscard]] inline int32 Compare(const OtherCharType* Other, ESearchCase::Type SearchCase = ESearchCase::CaseSensitive) const;
 
 	/** Returns whether this view starts with the prefix character compared case-sensitively. */
@@ -450,10 +431,10 @@ private:
 	static bool PrivateEquals(TStringView Lhs, TStringView Rhs);
 	static bool PrivateLess(TStringView Lhs, TStringView Rhs);
 
-	template <typename OtherCharType>
+	template <UE::CCharType OtherCharType>
 	inline bool PrivateEquals(const OtherCharType* Other, ESearchCase::Type SearchCase) const;
 
-	template <typename OtherCharType>
+	template <UE::CCharType OtherCharType>
 	inline int32 PrivateCompare(TStringView<OtherCharType> Other, ESearchCase::Type SearchCase) const;
 
 public:
@@ -577,17 +558,14 @@ template <typename CharRangeType>
 TStringView(CharRangeType&& Range)
 	-> TStringView<TElementType_T<CharRangeType>>;
 
-template <typename CharPtrOrRangeType,
-	// This uses enable_if instead of concepts because it tries to instantiate TStringView<void> with concepts.
-	// This can use concepts once TStringView constrains its CharType to character types.
-	std::enable_if_t<!std::is_pointer_v<CharPtrOrRangeType> || TIsCharType_V<std::remove_pointer_t<CharPtrOrRangeType>>>* = nullptr>
+template <typename CharPtrOrRangeType>
+	requires (!std::is_pointer_v<CharPtrOrRangeType> || TIsCharType_V<std::remove_pointer_t<CharPtrOrRangeType>>)
 [[nodiscard]] constexpr inline auto MakeStringView(const CharPtrOrRangeType& CharPtrOrRange UE_LIFETIMEBOUND) -> decltype(TStringView(CharPtrOrRange))
 {
 	return TStringView(CharPtrOrRange);
 }
 
-template <typename CharType
-	UE_REQUIRES(TIsCharType_V<CharType>)>
+template <UE::CCharType CharType>
 [[nodiscard]] constexpr inline auto MakeStringView(const CharType* CharPtr UE_LIFETIMEBOUND, int32 Size) -> decltype(TStringView(CharPtr, Size))
 {
 	return TStringView(CharPtr, Size);
@@ -635,14 +613,14 @@ template <typename CharType
 
 //////////////////////////////////////////////////////////////////////////
 
-template <typename CharType>
+template <UE::CCharType CharType>
 [[nodiscard]] inline const CharType& TStringView<CharType>::operator[](int32 Index) const
 {
 	checkf(Index >= 0 && Index < Size, TEXT("Index out of bounds on StringView: index %i on a view with a length of %i"), Index, Size);
 	return DataPtr[Index];
 }
 
-template <typename CharType>
+template <UE::CCharType CharType>
 inline int32 TStringView<CharType>::CopyString(CharType* Dest, int32 CharCount, int32 Position) const
 {
 	const int32 CopyCount = FMath::Min(Size - Position, CharCount);
@@ -653,33 +631,33 @@ inline int32 TStringView<CharType>::CopyString(CharType* Dest, int32 CharCount,
 	return CopyCount;
 }
 
-template <typename CharType>
+template <UE::CCharType CharType>
 [[nodiscard]] inline TStringView<CharType> TStringView<CharType>::Left(int32 CharCount) const
 {
 	return ViewType(DataPtr, FMath::Clamp(CharCount, 0, Size));
 }
 
-template <typename CharType>
+template <UE::CCharType CharType>
 [[nodiscard]] inline TStringView<CharType> TStringView<CharType>::LeftChop(int32 CharCount) const
 {
 	return ViewType(DataPtr, FMath::Clamp(Size - CharCount, 0, Size));
 }
 
-template <typename CharType>
+template <UE::CCharType CharType>
 [[nodiscard]] inline TStringView<CharType> TStringView<CharType>::Right(int32 CharCount) const
 {
 	const int32 OutLen = FMath::Clamp(CharCount, 0, Size);
 	return ViewType(DataPtr + Size - OutLen, OutLen);
 }
 
-template <typename CharType>
+template <UE::CCharType CharType>
 [[nodiscard]] inline TStringView<CharType> TStringView<CharType>::RightChop(int32 CharCount) const
 {
 	const int32 OutLen = FMath::Clamp(Size - CharCount, 0, Size);
 	return ViewType(DataPtr + Size - OutLen, OutLen);
 }
 
-template <typename CharType>
+template <UE::CCharType CharType>
 [[nodiscard]] inline TStringView<CharType> TStringView<CharType>::Mid(int32 Position, int32 CharCount) const
 {
 	const CharType* CurrentStart = GetData();
@@ -700,13 +678,13 @@ template <typename CharType>
 	return Result;
 }
 
-template <typename CharType>
+template <UE::CCharType CharType>
 [[nodiscard]] inline TStringView<CharType> TStringView<CharType>::TrimStartAndEnd() const
 {
 	return TrimStart().TrimEnd();
 }
 
-template <typename CharType>
+template <UE::CCharType CharType>
 [[nodiscard]] inline TStringView<CharType> TStringView<CharType>::TrimStart() const
 {
 	int32 SpaceCount = 0;
@@ -721,7 +699,7 @@ template <typename CharType>
 	return ViewType(DataPtr + SpaceCount, Size - SpaceCount);
 }
 
-template <typename CharType>
+template <UE::CCharType CharType>
 [[nodiscard]] inline TStringView<CharType> TStringView<CharType>::TrimEnd() const
 {
 	int32 NewSize = Size;
@@ -732,8 +710,8 @@ template <typename CharType>
 	return ViewType(DataPtr, NewSize);
 }
 
-template <typename CharType>
-template <typename OtherCharType>
+template <UE::CCharType CharType>
+template <UE::CCharType OtherCharType>
 [[nodiscard]] inline bool TStringView<CharType>::PrivateEquals(const OtherCharType* const Other, ESearchCase::Type SearchCase) const
 {
 	if (IsEmpty())
@@ -751,8 +729,8 @@ template <typename OtherCharType>
 	}
 }
 
-template <typename CharType>
-template <typename OtherCharType>
+template <UE::CCharType CharType>
+template <UE::CCharType OtherCharType>
 [[nodiscard]] inline int32 TStringView<CharType>::PrivateCompare(TStringView<OtherCharType> OtherView, ESearchCase::Type SearchCase) const
 {
 	const int32 SelfLen = Len();
@@ -782,8 +760,8 @@ template <typename OtherCharType>
 	return SelfLen - OtherLen;
 }
 
-template <typename CharType>
-template <typename OtherCharType>
+template <UE::CCharType CharType>
+template <UE::CCharType OtherCharType>
 [[nodiscard]] inline int32 TStringView<CharType>::Compare(const OtherCharType* const Other, ESearchCase::Type SearchCase) const
 {
 	if (IsEmpty())
@@ -810,52 +788,52 @@ template <typename OtherCharType>
 	return -!!Other[Len()];
 }
 
-template <typename CharType>
+template <UE::CCharType CharType>
 [[nodiscard]] inline bool TStringView<CharType>::StartsWith(ViewType Prefix, ESearchCase::Type SearchCase) const
 {
 	return Prefix.Equals(Left(Prefix.Len()), SearchCase);
 }
 
-template <typename CharType>
+template <UE::CCharType CharType>
 [[nodiscard]] inline bool TStringView<CharType>::EndsWith(ViewType Suffix, ESearchCase::Type SearchCase) const
 {
 	return Suffix.Equals(Right(Suffix.Len()), SearchCase);
 }
 
-template <typename CharType>
+template <UE::CCharType CharType>
 [[nodiscard]] inline int32 TStringView<CharType>::Find(const ViewType Search, const int32 StartPosition, ESearchCase::Type SearchCase) const
 {
 	const int32 Index = UE::String::FindFirst(RightChop(StartPosition), Search, SearchCase);
 	return Index == INDEX_NONE ? INDEX_NONE : Index + StartPosition;
 }
 
-template <typename CharType>
+template <UE::CCharType CharType>
 inline bool TStringView<CharType>::FindChar(const CharType Search, int32& OutIndex) const
 {
 	OutIndex = UE::String::FindFirstChar(*this, Search);
 	return OutIndex != INDEX_NONE;
 }
 
-template <typename CharType>
+template <UE::CCharType CharType>
 inline bool TStringView<CharType>::FindLastChar(const CharType Search, int32& OutIndex) const
 {
 	OutIndex = UE::String::FindLastChar(*this, Search);
 	return OutIndex != INDEX_NONE;
 }
 
-template <typename CharType>
+template <UE::CCharType CharType>
 [[nodiscard]] inline bool TStringView<CharType>::PrivateEquals(TStringView Lhs, const CharType* Rhs)
 {
 	return Lhs.Equals(Rhs, ESearchCase::IgnoreCase);
 }
 
-template <typename CharType>
+template <UE::CCharType CharType>
 [[nodiscard]] inline bool TStringView<CharType>::PrivateEquals(TStringView Lhs, TStringView Rhs)
 {
 	return Lhs.Equals(Rhs, ESearchCase::IgnoreCase);
 }
 
-template <typename CharType>
+template <UE::CCharType CharType>
 [[nodiscard]] inline bool TStringView<CharType>::PrivateLess(TStringView Lhs, TStringView Rhs)
 {
 	return Lhs.Compare(Rhs, ESearchCase::IgnoreCase) < 0;
