Added UE::CCompatibleCharType and some usage in TStringView.

diff --git a/Engine/Source/Runtime/Core/Public/Concepts/CompatibleCharType.h b/Engine/Source/Runtime/Core/Public/Concepts/CompatibleCharType.h
new file mode 100644
--- /dev/null
+++ b/Engine/Source/Runtime/Core/Public/Concepts/CompatibleCharType.h
@@ -0,0 +1,22 @@
+// Copyright Epic Games, Inc. All Rights Reserved.
+
+#pragma once
+
+#include "Concepts/CharType.h"
+#include "Traits/IsCharEncodingCompatibleWith.h"
+
+namespace UE
+{
+	/**
+	 * Concept which describes a character encoding type which is binary compatible with a given character encoding type.
+	 *
+	 * Examples:
+	 * UE::CCompatibleCharType<ANSICHAR, ANSICHAR> == true  // Any character type is compatible with itself
+	 * UE::CCompatibleCharType<int, int>           == false // ints are not character types
+	 * UE::CCompatibleCharType<ANSICHAR, WIDECHAR> == false // ASCII is not binary compatible with wide strings (different sizes)
+	 * UE::CCompatibleCharType<ANSICHAR, UTF8CHAR> == true  // ASCII is binary compatible with UTF-8
+	 * UE::CCompatibleCharType<UTF8CHAR, ANSICHAR> == false // UTF-8 is not binary compatible with ANSI
+	 */
+	template <typename T, typename CharType>
+	concept CCompatibleCharType = UE::CCharType<T> && TIsCharEncodingCompatibleWith_V<T, CharType>;
+}
diff --git a/Engine/Source/Runtime/Core/Public/Containers/StringView.h b/Engine/Source/Runtime/Core/Public/Containers/StringView.h
--- a/Engine/Source/Runtime/Core/Public/Containers/StringView.h
+++ b/Engine/Source/Runtime/Core/Public/Containers/StringView.h
@@ -4,6 +4,7 @@
 
 #include "CoreTypes.h"
 #include "Concepts/CharType.h"
+#include "Concepts/CompatibleCharType.h"
 #include "Containers/StringFwd.h" // IWYU pragma: export
 #include "HAL/UnrealMemory.h"
 #include "Math/NumericLimits.h"
@@ -123,8 +124,7 @@ public:
 	}
 
 	/** Construct a view of the null-terminated string pointed to by InData. */
-	template <UE::CCharType OtherCharType>
-		requires (TIsCharEncodingCompatibleWith<OtherCharType, CharType>::Value)
+	template <UE::CCompatibleCharType<CharType> OtherCharType>
 	[[nodiscard]] inline TStringView(const OtherCharType* InData UE_LIFETIMEBOUND)
 		: DataPtr((const CharType*)InData)
 		, Size(UE::Core::Private::StringLength((const CharType*)InData))
@@ -132,8 +132,7 @@ public:
 	}
 
 	/** Construct a view of InSize characters beginning at InData. */
-	template <UE::CCharType OtherCharType>
-		requires (TIsCharEncodingCompatibleWith<OtherCharType, CharType>::Value)
+	template <UE::CCompatibleCharType<CharType> OtherCharType>
 	[[nodiscard]] constexpr inline TStringView(const OtherCharType* InData UE_LIFETIMEBOUND, int32 InSize)
 		: DataPtr((const CharType*)InData)
 		, Size(InSize)
