Added Linux fork support for MallocBinned2 large allocations

diff --git a/Engine/Source/Runtime/Core/Private/HAL/MallocBinned2.cpp b/Engine/Source/Runtime/Core/Private/HAL/MallocBinned2.cpp
--- a/Engine/Source/Runtime/Core/Private/HAL/MallocBinned2.cpp
+++ b/Engine/Source/Runtime/Core/Private/HAL/MallocBinned2.cpp
@@ -638,6 +638,9 @@ void* FMallocBinned2::MallocExternalLarge(SIZE_T Size, uint32 Alignment)
 		}
 
 		Pool = Internal::GetOrCreatePoolInfo(*this, Result, FPoolInfo::ECanary::FirstFreeBlockIsOSAllocSize);
+#if BINNED2_FORK_SUPPORT
+		Pool->bAllocatedPreFork = CurrentCanary == EBlockCanary::PreFork ? 1 : 0;
+#endif
 	}
 
 	UE_CLOG(!IsAligned(Result, Alignment) ,LogMemory, Fatal, TEXT("FMallocBinned2 alignment was too large for OS. Alignment=%d   Ptr=%p"), Alignment, Result);
@@ -713,6 +716,17 @@ void* FMallocBinned2::ReallocExternal(void* Ptr, SIZE_T NewSize, uint32 Alignmen
 		bUseSmallMalloc = PromoteToLargerBin(NewSize, Alignment, *this);
 	}
 
+#if BINNED2_FORK_SUPPORT
+	if (Pool->bAllocatedPreFork)
+	{
+		// Can't dirty this page, should create a new allocation instead
+		void* Result = MallocExternalLarge(NewSize, Alignment);
+		ExternalAllocMutex.Unlock();
+		FMemory::Memcpy(Result, Ptr, FMath::Min(NewSize, PoolOSRequestedBytes));
+		return Result;
+	}
+	else
+#endif
 	if (NewSize > PoolOsBytes || // can't fit in the old block
 		bUseSmallMalloc || // can switch to the small bin allocator
 		Align(NewSize, OsAllocationGranularity) < PoolOsBytes) // we can get some pages back
@@ -809,6 +823,13 @@ void FMallocBinned2::FreeExternal(void* Ptr)
 		{
 			UE_LOG(LogMemory, Fatal, TEXT("FMallocBinned2 Attempt to free an unrecognized pointer %p"), Ptr);
 		}
+#if BINNED2_FORK_SUPPORT
+		if (Pool->bAllocatedPreFork)
+		{
+			// avoid freeing it as it might increase unique memory
+			return;
+		}
+#endif
 		const UPTRINT PoolOsBytes = Pool->GetOsAllocatedBytes();
 		const SIZE_T PoolOSRequestedBytes = Pool->GetOSRequestedBytes();
 
diff --git a/Engine/Source/Runtime/Core/Public/HAL/MallocBinned2.h b/Engine/Source/Runtime/Core/Public/HAL/MallocBinned2.h
--- a/Engine/Source/Runtime/Core/Public/HAL/MallocBinned2.h
+++ b/Engine/Source/Runtime/Core/Public/HAL/MallocBinned2.h
@@ -83,12 +83,17 @@ public:
 			FirstFreeBlockIsPtr = 0xf317
 		};
 
-		uint16      Taken;          // Number of allocated elements in this pool, when counts down to zero can free the entire pool	
-		ECanary     Canary;         // See ECanary
-		uint32      AllocSize;      // Number of bytes allocated
-		FFreeBlock* FirstFreeBlock; // Pointer to first free memory in this pool or the OS Allocation Size in bytes if this allocation is not binned
-		FPoolInfo*  Next;           // Pointer to next pool
-		FPoolInfo** PtrToPrevNext;  // Pointer to whichever pointer points to this pool
+		uint16      Taken;					// Number of allocated elements in this pool, when counts down to zero can free the entire pool	
+		ECanary     Canary;					// See ECanary
+#if BINNED2_FORK_SUPPORT
+		uint32      AllocSize : 31;			// Number of bytes allocated
+		bool		bAllocatedPreFork : 1;	// Set to true if this was allocated pre fork, so we don't free it post fork
+#else
+		uint32      AllocSize;				// Number of bytes allocated
+#endif
+		FFreeBlock* FirstFreeBlock;			// Pointer to first free memory in this pool or the OS Allocation Size in bytes if this allocation is not binned
+		FPoolInfo*  Next;					// Pointer to next pool
+		FPoolInfo** PtrToPrevNext;			// Pointer to whichever pointer points to this pool
 
 		FPoolInfo();
 
