[CoreDelegates]
- Cleanup deprecations

diff --git a/Engine/Source/Runtime/ApplicationCore/Private/GenericPlatform/GenericPlatformInputDeviceMapper.cpp b/Engine/Source/Runtime/ApplicationCore/Private/GenericPlatform/GenericPlatformInputDeviceMapper.cpp
--- a/Engine/Source/Runtime/ApplicationCore/Private/GenericPlatform/GenericPlatformInputDeviceMapper.cpp
+++ b/Engine/Source/Runtime/ApplicationCore/Private/GenericPlatform/GenericPlatformInputDeviceMapper.cpp
@@ -346,14 +346,6 @@ bool IPlatformInputDeviceMapper::Internal_MapInputDeviceToUser(FInputDeviceId De
 	
 	// Broadcast delegates to let listeners know that the platform user has had an input device change
 	OnInputDeviceConnectionChange.Broadcast(ConnectionState, UserId, DeviceId);
-
-	if (ShouldBroadcastLegacyDelegates())
-	{
-		const bool bIsConnected = (ConnectionState == EInputDeviceConnectionState::Connected);
-		PRAGMA_DISABLE_DEPRECATION_WARNINGS
-		FCoreDelegates::OnControllerConnectionChange.Broadcast(bIsConnected, UserId, DeviceId.GetId());
-		PRAGMA_ENABLE_DEPRECATION_WARNINGS
-	}
 	
 	return true;
 }
@@ -390,15 +382,6 @@ bool IPlatformInputDeviceMapper::Internal_ChangeInputDeviceUserMapping(FInputDev
 	// Broadcast the delegates letting listeners know that the input device has changed owners
 	OnInputDevicePairingChange.Broadcast(DeviceId, NewUserId, OldUserId);
 	
-	if (ShouldBroadcastLegacyDelegates())
-	{
-		// Remap the DeviceId to the older int32 "ControllerId" format for the legacy delegates
-		int32 LegacyControllerId = INDEX_NONE;
-		RemapUserAndDeviceToControllerId(NewUserId, LegacyControllerId, DeviceId);
-		PRAGMA_DISABLE_DEPRECATION_WARNINGS
-		FCoreDelegates::OnControllerPairingChange.Broadcast(LegacyControllerId, NewUserId, OldUserId);
-		PRAGMA_ENABLE_DEPRECATION_WARNINGS
-	}
 	return true;
 }
 
diff --git a/Engine/Source/Runtime/Core/Private/Misc/CoreDelegates.cpp b/Engine/Source/Runtime/Core/Private/Misc/CoreDelegates.cpp
--- a/Engine/Source/Runtime/Core/Private/Misc/CoreDelegates.cpp
+++ b/Engine/Source/Runtime/Core/Private/Misc/CoreDelegates.cpp
@@ -36,10 +36,6 @@ TDelegate<int32(const FString&)> FCoreDelegates::PakChunkIdOverrideDelegate;
 
 TDelegate<void()> FCoreDelegates::OnOptimizeMemoryUsageForMountedPaks;
 
-PRAGMA_DISABLE_DEPRECATION_WARNINGS
-TMulticastDelegate<void(const IPakFile&)> FCoreDelegates::OnPakFileMounted2;
-PRAGMA_ENABLE_DEPRECATION_WARNINGS
-
 TTSMulticastDelegate<void(const IPakFile&)>& FCoreDelegates::GetOnPakFileMounted2()
 {
 	static TTSMulticastDelegate<void(const IPakFile&)> Singleton;
@@ -49,10 +45,6 @@ TTSMulticastDelegate<void(const IPakFile&)>& FCoreDelegates::GetOnPakFileMounted
 TMulticastDelegate<void(const FString&)> FCoreDelegates::NewFileAddedDelegate;
 TMulticastDelegate<void()> FCoreDelegates::NoPakFilesMountedDelegate;
 
-PRAGMA_DISABLE_DEPRECATION_WARNINGS
-TMulticastDelegate<void(const TCHAR*, const TCHAR*), FNotThreadSafeNotCheckedDelegateUserPolicy> FCoreDelegates::OnFileOpenedForReadFromPakFile;
-PRAGMA_ENABLE_DEPRECATION_WARNINGS
-
 TTSMulticastDelegate<void(const TCHAR*, const TCHAR*)>& FCoreDelegates::GetOnFileOpenedForReadFromPakFile()
 {
 	static TTSMulticastDelegate<void(const TCHAR*, const TCHAR*)> Singleton;
@@ -77,10 +69,6 @@ FCurrentlyMountedPaksDelegate& FCoreInternalDelegates::GetCurrentlyMountedPaksDe
 //} // namespace UE::Core
 
 TMulticastDelegate<void(bool, int32, int32)> FCoreDelegates::OnUserLoginChangedEvent;
-PRAGMA_DISABLE_DEPRECATION_WARNINGS
-TMulticastDelegate<void(bool, FPlatformUserId, int32)> FCoreDelegates::OnControllerConnectionChange;
-TMulticastDelegate<void(int32, FPlatformUserId, FPlatformUserId)> FCoreDelegates::OnControllerPairingChange;
-PRAGMA_ENABLE_DEPRECATION_WARNINGS
 TMulticastDelegate<void()> FCoreDelegates::OnSafeFrameChangedEvent;
 TMulticastDelegate<void()> FCoreDelegates::OnHandleSystemEnsure;
 TMulticastDelegate<void()> FCoreDelegates::OnHandleSystemError;
@@ -149,10 +137,6 @@ TMulticastDelegate<void()> FCoreDelegates::ApplicationHasReactivatedDelegate;
 TMulticastDelegate<void()> FCoreDelegates::ApplicationWillEnterBackgroundDelegate;
 TMulticastDelegate<void()> FCoreDelegates::ApplicationHasEnteredForegroundDelegate;
 
-PRAGMA_DISABLE_DEPRECATION_WARNINGS
-TMulticastDelegate<void()> FCoreDelegates::ApplicationWillTerminateDelegate;
-PRAGMA_ENABLE_DEPRECATION_WARNINGS
-
 TTSMulticastDelegate<void()>& FCoreDelegates::GetApplicationWillTerminateDelegate()
 {
 	static TTSMulticastDelegate<void()> Singleton;
@@ -276,10 +260,6 @@ TMulticastDelegate<void(bool)> FCoreDelegates::IsVanillaProductChanged;
 TMulticastDelegate<void()> FCoreDelegates::OnAsyncLoadingFlush;
 TMulticastDelegate<void()> FCoreDelegates::OnAsyncLoadingFlushUpdate;
 
-PRAGMA_DISABLE_DEPRECATION_WARNINGS
-TMulticastDelegate<void(const FString&)> FCoreDelegates::OnAsyncLoadPackage;
-PRAGMA_ENABLE_DEPRECATION_WARNINGS
-
 TTSMulticastDelegate<void(FStringView)>& FCoreDelegates::GetOnAsyncLoadPackage()
 {
 	static TTSMulticastDelegate<void(FStringView)> Singleton;
@@ -317,10 +297,6 @@ TDelegate<void(TArrayView<const FSHAHash>, FCoreDelegates::FAttachShaderReadRequ
 TDelegate<void(TArrayView<const FSHAHash>)> FCoreDelegates::ReleasePreloadedPackageShaderMaps;
 TMulticastDelegate<void(const FLogCategoryName&, ELogVerbosity::Type, ELogVerbosity::Type)> FCoreDelegates::OnLogVerbosityChanged;
 
-PRAGMA_DISABLE_DEPRECATION_WARNINGS
-TDelegate<TSharedPtr<class IPackageStore>()> FCoreDelegates::CreatePackageStore;
-PRAGMA_ENABLE_DEPRECATION_WARNINGS
-
 TMulticastDelegate<void(bool)> FCoreDelegates::ApplicationNetworkInitializationChanged;
 TMulticastDelegate<void(ENetworkConnectionStatus, ENetworkConnectionStatus)> FCoreDelegates::OnNetworkConnectionStatusChanged;
 
diff --git a/Engine/Source/Runtime/Core/Private/Windows/WindowsPlatformMisc.cpp b/Engine/Source/Runtime/Core/Private/Windows/WindowsPlatformMisc.cpp
--- a/Engine/Source/Runtime/Core/Private/Windows/WindowsPlatformMisc.cpp
+++ b/Engine/Source/Runtime/Core/Private/Windows/WindowsPlatformMisc.cpp
@@ -992,10 +992,6 @@ static BOOL WINAPI ConsoleCtrlHandler(DWORD CtrlType)
 	{
 		RequestEngineExit(TEXT("ConsoleCtrl RequestExit"));
 
-		PRAGMA_DISABLE_DEPRECATION_WARNINGS
-		FCoreDelegates::ApplicationWillTerminateDelegate.Broadcast();
-		PRAGMA_ENABLE_DEPRECATION_WARNINGS
-
 		FCoreDelegates::GetApplicationWillTerminateDelegate().Broadcast();
 
 		AppTermDelegateBroadcast = true;
@@ -1367,10 +1363,6 @@ void FWindowsPlatformMisc::RequestExitWithStatus(bool Force, uint8 ReturnCode, c
 
 	RequestEngineExit(TEXT("Win RequestExit"));
 
-	PRAGMA_DISABLE_DEPRECATION_WARNINGS
-	FCoreDelegates::ApplicationWillTerminateDelegate.Broadcast();
-	PRAGMA_ENABLE_DEPRECATION_WARNINGS
-
 	FCoreDelegates::GetApplicationWillTerminateDelegate().Broadcast();
 
 	if (Force)
diff --git a/Engine/Source/Runtime/Core/Public/Misc/CoreDelegates.h b/Engine/Source/Runtime/Core/Public/Misc/CoreDelegates.h
--- a/Engine/Source/Runtime/Core/Public/Misc/CoreDelegates.h
+++ b/Engine/Source/Runtime/Core/Public/Misc/CoreDelegates.h
@@ -107,10 +107,6 @@ public:
 	// Callback for a game thread interruption point when a async load flushing. Used to updating UI during long loads.
 	static CORE_API TMulticastDelegate<void()> OnAsyncLoadingFlushUpdate;
 
-	// Callback on the game thread when an async load is started. This goes off before the packages has finished loading
-	UE_DEPRECATED(5.3, "This delegate is not thread-safe, please use GetOnAsyncLoadPackage().")
-	static CORE_API TMulticastDelegate<void(const FString&)> OnAsyncLoadPackage;
-
 	// Thread-safe callback that is called on the same thread that LoadPackageAsync is issued from.
 	static CORE_API TTSMulticastDelegate<void(FStringView)>& GetOnAsyncLoadPackage();
 
@@ -122,14 +118,6 @@ public:
 	// Callback when a user logs in/out of the platform.
 	static CORE_API TMulticastDelegate<void(bool, int32, int32)> OnUserLoginChangedEvent;
 
-	// Callback when controllers disconnected / reconnected
-	UE_DEPRECATED(5.1, "OnControllerConnectionChange, use IPlatformInputDeviceMapper::GetOnInputDeviceConnectionChange() instead")
-	static CORE_API TMulticastDelegate<void(bool, FPlatformUserId, int32)> OnControllerConnectionChange;
-
-	// Callback when a single controller pairing changes
-	UE_DEPRECATED(5.1, "OnControllerPairingChange, use IPlatformInputDeviceMapper::GetOnInputDevicePairingChange() instead")
-	static CORE_API TMulticastDelegate<void(int32 ControllerIndex, FPlatformUserId NewUserPlatformId, FPlatformUserId OldUserPlatformId)> OnControllerPairingChange;
-
 	// Callback when a user changes the safe frame size
 	static CORE_API TMulticastDelegate<void()> OnSafeFrameChangedEvent;
 
@@ -153,8 +141,6 @@ public:
 	static CORE_API TDelegate<void()> OnOptimizeMemoryUsageForMountedPaks;
 
 	// After a pakfile is mounted this is called
-	UE_DEPRECATED(5.3, "This delegate is not thread-safe, please use GetOnPakFileMounted2().")
-	static CORE_API TMulticastDelegate<void(const IPakFile&)> OnPakFileMounted2;
 	static CORE_API TTSMulticastDelegate<void(const IPakFile&)>& GetOnPakFileMounted2();
 
 	// After a file is added this is called
@@ -164,8 +150,6 @@ public:
 	static CORE_API TMulticastDelegate<void()> NoPakFilesMountedDelegate;
 
 	// When a file is opened for read from a pak file.
-	UE_DEPRECATED(5.3, "This delegate is not thread-safe, please use Use FCoreDelegates::GetOnFileOpenedForReadFromPakFile()")
-	static CORE_API TMulticastDelegate<void(const TCHAR* PakFile, const TCHAR* FileName), FNotThreadSafeNotCheckedDelegateUserPolicy> OnFileOpenedForReadFromPakFile;
 	static CORE_API TTSMulticastDelegate<void(const TCHAR* PakFile, const TCHAR* FileName)>& GetOnFileOpenedForReadFromPakFile();
 
 	typedef TSharedPtr<class IMovieStreamer, ESPMode::ThreadSafe> FMovieStreamerPtr;
@@ -346,9 +330,6 @@ public:
 	// This *may* be called when the application is getting terminated by the OS.
 	// There is no guarantee that this will ever be called on a mobile device,
 	// save state when ApplicationWillEnterBackgroundDelegate is called instead.
-	UE_DEPRECATED(5.3, "This delegate is not thread safe while it can be used concurrently. Please use to GetApplicationWillTerminateDelegate() instead.")
-	static CORE_API TMulticastDelegate<void()> ApplicationWillTerminateDelegate;
-
 	static CORE_API TTSMulticastDelegate<void()>& GetApplicationWillTerminateDelegate();
 
 	// Some platform have a System UI Overlay that can draw on top of the application.
@@ -587,11 +568,6 @@ public:
 	/** Called when the verbosity of a log category is changed */
 	static CORE_API TMulticastDelegate<void(const FLogCategoryName& CategoryName, ELogVerbosity::Type OldVerbosity, ELogVerbosity::Type NewVerbosity)> OnLogVerbosityChanged;
 
-	PRAGMA_DISABLE_DEPRECATION_WARNINGS
-	UE_DEPRECATED(5.1, "Use FPackageStore::Mount() instead")
-	static CORE_API TDelegate<TSharedPtr<class IPackageStore>()> CreatePackageStore;
-	PRAGMA_ENABLE_DEPRECATION_WARNINGS
-
 	// Called immediately before the parent process will start responding to signals to fork
 	static CORE_API FSimpleMulticastDelegate OnParentBeginFork;
 	// Called each time immediately before the parent process forks itself
@@ -630,107 +606,4 @@ private:
 
 	// This class is only for namespace use
 	FCoreDelegates() = default;
-
-public:	// deprecated delegate type aliases
-	using FHotFixDelegate UE_DEPRECATED(5.2, "Use template instantiation instead as `TDelegate<void(void*, int32)>`") = TDelegate<void(void*, int32)>;
-	using FOnActorLabelChanged UE_DEPRECATED(5.2, "Use template instantiation instead as `TMulticastDelegate<void(AActor*)>`") = TMulticastDelegate<void(AActor*)>;
-	using FOnMountAllPakFiles UE_DEPRECATED(5.2, "Use template instantiation instead as `TDelegate<int32(const TArray<FString>&)>`") = TDelegate<int32(const TArray<FString>&)>;
-	using FMountPak UE_DEPRECATED(5.2, "Use template instantiation instead as `TDelegate<IPakFile*(const FString&, int32)>`") = TDelegate<IPakFile*(const FString&, int32)>;
-	using FOnUnmountPak UE_DEPRECATED(5.2, "Use template instantiation instead as `TDelegate<bool(const FString&)>`") = TDelegate<bool(const FString&)>;
-	using FOnOptimizeMemoryUsageForMountedPaks UE_DEPRECATED(5.2, "Use template instantiation instead as `TDelegate<void()>`") = TDelegate<void()>;
-	using FOnPakFileMounted2 UE_DEPRECATED(5.2, "Use template instantiation instead as `TMulticastDelegate<void(const IPakFile&)>`") = TMulticastDelegate<void(const IPakFile&)>;
-	using FNoPakFilesMountedDelegate UE_DEPRECATED(5.2, "Use template instantiation instead as `TMulticastDelegate<void()>`") = TMulticastDelegate<void()>;
-	using FOnModalMessageBox UE_DEPRECATED(5.2, "Use template instantiation instead as `TDelegate<EAppReturnType::Type(EAppMsgType::Type, const FText&, const FText&)>`") = TDelegate<EAppReturnType::Type(EAppMsgType::Type, const FText&, const FText&)>;
-	using FOnHandleSystemEnsure UE_DEPRECATED(5.2, "Use template instantiation instead as `TMulticastDelegate<void()>`") = TMulticastDelegate<void()>;
-	using FOnHandleSystemError UE_DEPRECATED(5.2, "Use template instantiation instead as `TMulticastDelegate<void()>`") = TMulticastDelegate<void()>;
-	using FRegisterMovieStreamerDelegate UE_DEPRECATED(5.2, "Use template instantiation instead as `TMulticastDelegate<void(FMovieStreamerPtr)>`") = TMulticastDelegate<void(FMovieStreamerPtr)>;
-	using FUnRegisterMovieStreamerDelegate UE_DEPRECATED(5.2, "Use template instantiation instead as `TMulticastDelegate<void(FMovieStreamerPtr)>`") = TMulticastDelegate<void(FMovieStreamerPtr)>;
-	using FOnUserLoginChangedEvent UE_DEPRECATED(5.2, "Use template instantiation instead as `TMulticastDelegate<void(bool, int32, int32)>`") = TMulticastDelegate<void(bool, int32, int32)>;
-	using FOnSafeFrameChangedEvent UE_DEPRECATED(5.2, "Use template instantiation instead as `TMulticastDelegate<void()>`") = TMulticastDelegate<void()>;
-	using FOnInviteAccepted UE_DEPRECATED(5.2, "Use template instantiation instead as `TMulticastDelegate<void(const FString&, const FString&)>`") = TMulticastDelegate<void(const FString&, const FString&)>;
-	using FRegisterEncryptionKeyMulticastDelegate UE_DEPRECATED(5.2, "Use template instantiation instead as `TMulticastDelegate<void(const FGuid&, const FAES::FAESKey&)>`") = TMulticastDelegate<void(const FGuid&, const FAES::FAESKey&)>;
-	using FPakEncryptionKeyDelegate UE_DEPRECATED(5.2, "Use template instantiation instead as `TDelegate<void(uint8[32])>`") = TDelegate<void(uint8[32])>;
-	using FPakSigningKeysDelegate UE_DEPRECATED(5.2, "Use template instantiation instead as `TDelegate<void(TArray<uint8>&, TArray<uint8>&)>`") = TDelegate<void(TArray<uint8>&, TArray<uint8>&)>;
-	using FOnUserControllerConnectionChange UE_DEPRECATED(5.2, "Use template instantiation instead as `TMulticastDelegate<void(bool, FPlatformUserId, int32)>`") = TMulticastDelegate<void(bool, FPlatformUserId, int32)>;
-	using FOnUserControllerPairingChange UE_DEPRECATED(5.2, "Use template instantiation instead as `TMulticastDelegate<void(int32 ControllerIndex, FPlatformUserId NewUserPlatformId, FPlatformUserId OldUserPlatformId)>`") = TMulticastDelegate<void(int32 ControllerIndex, FPlatformUserId NewUserPlatformId, FPlatformUserId OldUserPlatformId)>;
-	using FOnAsyncLoadingFlush UE_DEPRECATED(5.2, "Use template instantiation instead as `TMulticastDelegate<void()>`") = TMulticastDelegate<void()>;
-	using FOnAsyncLoadingFlushUpdate UE_DEPRECATED(5.2, "Use template instantiation instead as `TMulticastDelegate<void()>`") = TMulticastDelegate<void()>;
-	using FOnAsyncLoadPackage UE_DEPRECATED(5.2, "Use template instantiation instead as `TMulticastDelegate<void(const FString&)>`") = TMulticastDelegate<void(const FString&)>;
-	using FOnSyncLoadPackage UE_DEPRECATED(5.2, "Use template instantiation instead as `TMulticastDelegate<void(const FString&)>`") = TMulticastDelegate<void(const FString&)>;
-	using FNewFileAddedDelegate UE_DEPRECATED(5.2, "Use template instantiation instead as `TMulticastDelegate<void(const FString&)>`") = TMulticastDelegate<void(const FString&)>;
-	using FOnFileOpenedForReadFromPakFile UE_DEPRECATED(5.2, "Use template instantiation instead as `TMulticastDelegate<void(const TCHAR* PakFile, const TCHAR* FileName)>`") = TMulticastDelegate<void(const TCHAR* PakFile, const TCHAR* FileName)>;
-#if ALLOW_OTHER_PLATFORM_CONFIG
-	using FGatherDeviceProfileCVars UE_DEPRECATED(5.2, "Use template instantiation instead as `TDelegate<FCVarKeyValueMap(const FString& DeviceProfileName)>`") = TDelegate<FCVarKeyValueMap(const FString& DeviceProfileName)>;
-#endif
-	using FGatherAdditionalLocResPathsDelegate UE_DEPRECATED(5.2, "Use template instantiation instead as `TMulticastDelegate<void(TArray<FString>&)>`") = TMulticastDelegate<void(TArray<FString>&)>;
-	using FWorldOriginOffset UE_DEPRECATED(5.2, "Use template instantiation instead as `TMulticastDelegate<void(class UWorld*, FIntVector, FIntVector)>`") = TMulticastDelegate<void(class UWorld*, FIntVector, FIntVector)>;
-	using FStarvedGameLoop UE_DEPRECATED(5.2, "Use template instantiation instead as `TDelegate<void()>`") = TDelegate<void()>;
-	using FOnTemperatureChange UE_DEPRECATED(5.2, "Use template instantiation instead as `TMulticastDelegate<void(ETemperatureSeverity)>`") = TMulticastDelegate<void(ETemperatureSeverity)>;
-	using FOnLowPowerMode UE_DEPRECATED(5.2, "Use template instantiation instead as `TMulticastDelegate<void(bool)>`") = TMulticastDelegate<void(bool)>;
-	using FCountPreLoadConfigFileRespondersDelegate UE_DEPRECATED(5.2, "Use template instantiation instead as `TMulticastDelegate<void(const TCHAR* IniFilename, int32& ResponderCount)>`") = TMulticastDelegate<void(const TCHAR* IniFilename, int32& ResponderCount)>;
-	using FPreLoadConfigFileDelegate UE_DEPRECATED(5.2, "Use template instantiation instead as `TMulticastDelegate<void(const TCHAR* IniFilename, FString& LoadedContents)>`") = TMulticastDelegate<void(const TCHAR* IniFilename, FString& LoadedContents)>;
-	using FPreSaveConfigFileDelegate UE_DEPRECATED(5.2, "Use template instantiation instead as `TMulticastDelegate<void(const TCHAR* IniFilename, const FString& ContentsToSave, int32& SavedCount)>`") = TMulticastDelegate<void(const TCHAR* IniFilename, const FString& ContentsToSave, int32& SavedCount)>;
-	using FOnFConfigFileCreated UE_DEPRECATED(5.2, "Use template instantiation instead as `TMulticastDelegate<void(const FConfigFile*)>`") = TMulticastDelegate<void(const FConfigFile*)>;
-	using FOnFConfigFileDeleted UE_DEPRECATED(5.2, "Use template instantiation instead as `TMulticastDelegate<void(const FConfigFile*)>`") = TMulticastDelegate<void(const FConfigFile*)>;
-	using FOnConfigValueRead UE_DEPRECATED(5.2, "Use template instantiation instead as `TMulticastDelegate<void(const TCHAR* IniFilename, const TCHAR* SectionName, const TCHAR* Key)>`") = TMulticastDelegate<void(const TCHAR* IniFilename, const TCHAR* SectionName, const TCHAR* Key)>;
-	using FOnConfigSectionRead UE_DEPRECATED(5.2, "Use template instantiation instead as `TMulticastDelegate<void(const TCHAR* IniFilename, const TCHAR* SectionName)>`") = TMulticastDelegate<void(const TCHAR* IniFilename, const TCHAR* SectionName)>;
-	using FOnConfigSectionsChanged UE_DEPRECATED(5.2, "Use template instantiation instead as `TMulticastDelegate<void(const FString& IniFilename, const TSet<FString>& SectionNames)>`") = TMulticastDelegate<void(const FString& IniFilename, const TSet<FString>& SectionNames)>;
-	using FOnApplyCVarFromIni UE_DEPRECATED(5.2, "Use template instantiation instead as `TMulticastDelegate<void(const TCHAR* SectionName, const TCHAR* IniFilename, uint32 SetBy, bool bAllowCheating)>`") = TMulticastDelegate<void(const TCHAR* SectionName, const TCHAR* IniFilename, uint32 SetBy, bool bAllowCheating)>;
-	using FOnSystemResolutionChanged UE_DEPRECATED(5.2, "Use template instantiation instead as `TMulticastDelegate<void(uint32 ResX, uint32 ResY)>`") = TMulticastDelegate<void(uint32 ResX, uint32 ResY)>;
-	using FOnTargetPlatformChangedSupportedFormats UE_DEPRECATED(5.2, "Use template instantiation instead as `TMulticastDelegate<void(const ITargetPlatform*)>`") = TMulticastDelegate<void(const ITargetPlatform*)>;
-	using FOnFeatureLevelDisabled UE_DEPRECATED(5.2, "Use template instantiation instead as `TMulticastDelegate<void(int, const FName&)>`") = TMulticastDelegate<void(int, const FName&)>;
-	using FApplicationLifetimeDelegate UE_DEPRECATED(5.2, "Use template instantiation instead as `TMulticastDelegate<void()>`") = TMulticastDelegate<void()>;
-	using FBackgroundTickDelegate UE_DEPRECATED(5.2, "Use template instantiation instead as `TMulticastDelegate<void(float DeltaTime)>`") = TMulticastDelegate<void(float DeltaTime)>;
-	using FUserMusicInterruptDelegate UE_DEPRECATED(5.2, "Use template instantiation instead as `TMulticastDelegate<void(bool)>`") = TMulticastDelegate<void(bool)>;
-	using FAudioMuteDelegate UE_DEPRECATED(5.2, "Use template instantiation instead as `TMulticastDelegate<void(bool, int)>`") = TMulticastDelegate<void(bool, int)>;
-	using FAudioRouteChangedDelegate UE_DEPRECATED(5.2, "Use template instantiation instead as `TMulticastDelegate<void(bool)>`") = TMulticastDelegate<void(bool)>;
-	using FApplicationRequestAudioState UE_DEPRECATED(5.2, "Use template instantiation instead as `TMulticastDelegate<void()>`") = TMulticastDelegate<void()>;
-	using FApplicationStartupArgumentsDelegate UE_DEPRECATED(5.2, "Use template instantiation instead as `TMulticastDelegate<void(const TArray<FString>&)>`") = TMulticastDelegate<void(const TArray<FString>&)>;
-	using FApplicationRegisteredForRemoteNotificationsDelegate UE_DEPRECATED(5.2, "Use template instantiation instead as `TMulticastDelegate<void(TArray<uint8>)>`") = TMulticastDelegate<void(TArray<uint8>)>;
-	using FApplicationRegisteredForUserNotificationsDelegate UE_DEPRECATED(5.2, "Use template instantiation instead as `TMulticastDelegate<void(int)>`") = TMulticastDelegate<void(int)>;
-	using FApplicationFailedToRegisterForRemoteNotificationsDelegate UE_DEPRECATED(5.2, "Use template instantiation instead as `TMulticastDelegate<void(FString)>`") = TMulticastDelegate<void(FString)>;
-	using FApplicationReceivedRemoteNotificationDelegate UE_DEPRECATED(5.2, "Use template instantiation instead as `TMulticastDelegate<void(FString, int)>`") = TMulticastDelegate<void(FString, int)>;
-	using FApplicationReceivedLocalNotificationDelegate UE_DEPRECATED(5.2, "Use template instantiation instead as `TMulticastDelegate<void(FString, int, int)>`") = TMulticastDelegate<void(FString, int, int)>;
-	using FApplicationPerformFetchDelegate UE_DEPRECATED(5.2, "Use template instantiation instead as `TMulticastDelegate<void()>`") = TMulticastDelegate<void()>;
-	using FApplicationBackgroundSessionEventDelegate UE_DEPRECATED(5.2, "Use template instantiation instead as `TMulticastDelegate<void(FString)>`") = TMulticastDelegate<void(FString)>;
-	using FApplicationReceivedOnScreenOrientationChangedNotificationDelegate UE_DEPRECATED(5.2, "Use template instantiation instead as `TMulticastDelegate<void(int32)>`") = TMulticastDelegate<void(int32)>;
-	using FStatCheckEnabled UE_DEPRECATED(5.2, "Use template instantiation instead as `TMulticastDelegate<void(const TCHAR*, bool&, bool&)>`") = TMulticastDelegate<void(const TCHAR*, bool&, bool&)>;
-	using FStatEnabled UE_DEPRECATED(5.2, "Use template instantiation instead as `TMulticastDelegate<void(const TCHAR*)>`") = TMulticastDelegate<void(const TCHAR*)>;
-	using FStatDisabled UE_DEPRECATED(5.2, "Use template instantiation instead as `TMulticastDelegate<void(const TCHAR*)>`") = TMulticastDelegate<void(const TCHAR*)>;
-	using FStatDisableAll UE_DEPRECATED(5.2, "Use template instantiation instead as `TMulticastDelegate<void(const bool)>`") = TMulticastDelegate<void(const bool)>;
-	using FApplicationLicenseChange UE_DEPRECATED(5.2, "Use template instantiation instead as `TMulticastDelegate<void()>`") = TMulticastDelegate<void()>;
-	using FPlatformChangedLaptopMode UE_DEPRECATED(5.2, "Use template instantiation instead as `TMulticastDelegate<void(EConvertibleLaptopMode)>`") = TMulticastDelegate<void(EConvertibleLaptopMode)>;
-	using FVRHeadsetTrackingInitializingAndNeedsHMDToBeTrackedDelegate UE_DEPRECATED(5.2, "Use template instantiation instead as `TMulticastDelegate<void()>`") = TMulticastDelegate<void()>;
-	using FVRHeadsetTrackingInitializedDelegate UE_DEPRECATED(5.2, "Use template instantiation instead as `TMulticastDelegate<void()>`") = TMulticastDelegate<void()>;
-	using FVRHeadsetRecenter UE_DEPRECATED(5.2, "Use template instantiation instead as `TMulticastDelegate<void()>`") = TMulticastDelegate<void()>;
-	using FVRHeadsetLost UE_DEPRECATED(5.2, "Use template instantiation instead as `TMulticastDelegate<void()>`") = TMulticastDelegate<void()>;
-	using FVRHeadsetReconnected UE_DEPRECATED(5.2, "Use template instantiation instead as `TMulticastDelegate<void()>`") = TMulticastDelegate<void()>;
-	using FVRHeadsetConnectCanceled UE_DEPRECATED(5.2, "Use template instantiation instead as `TMulticastDelegate<void()>`") = TMulticastDelegate<void()>;
-	using FVRHeadsetPutOnHead UE_DEPRECATED(5.2, "Use template instantiation instead as `TMulticastDelegate<void()>`") = TMulticastDelegate<void()>;
-	using FVRHeadsetRemovedFromHead UE_DEPRECATED(5.2, "Use template instantiation instead as `TMulticastDelegate<void()>`") = TMulticastDelegate<void()>;
-	using FVRControllerRecentered UE_DEPRECATED(5.2, "Use template instantiation instead as `TMulticastDelegate<void()>`") = TMulticastDelegate<void()>;
-	using FOnUserActivityStringChanged UE_DEPRECATED(5.2, "Use template instantiation instead as `TMulticastDelegate<void(const FString&)>`") = TMulticastDelegate<void(const FString&)>;
-	using FOnGameSessionIDChange UE_DEPRECATED(5.2, "Use template instantiation instead as `TMulticastDelegate<void(const FString&)>`") = TMulticastDelegate<void(const FString&)>;
-	using FOnGameStateClassChange UE_DEPRECATED(5.2, "Use template instantiation instead as `TMulticastDelegate<void(const FString&)>`") = TMulticastDelegate<void(const FString&)>;
-	using FOnCrashOverrideParamsChanged UE_DEPRECATED(5.2, "Use template instantiation instead as `TMulticastDelegate<void(const FCrashOverrideParameters&)>`") = TMulticastDelegate<void(const FCrashOverrideParameters&)>;
-	using FOnIsVanillaProductChanged UE_DEPRECATED(5.2, "Use template instantiation instead as `TMulticastDelegate<void(bool)>`") = TMulticastDelegate<void(bool)>;
-	using FOnPreMainInit UE_DEPRECATED(5.2, "Use template instantiation instead as `TMulticastDelegate<void()>`") = TMulticastDelegate<void()>;
-	using FConfigReadyForUse UE_DEPRECATED(5.2, "Use template instantiation instead as `TMulticastDelegate<void()>`") = TMulticastDelegate<void()>;
-	using FRenderingThreadChanged UE_DEPRECATED(5.2, "Use template instantiation instead as `TMulticastDelegate<void()>`") = TMulticastDelegate<void()>;
-	using FResolvePackageNameDelegate UE_DEPRECATED(5.2, "Use template instantiation instead as `TMulticastDelegate<bool(const FString&, FString&)>`") = TMulticastDelegate<bool(const FString&, FString&)>;
-	using FGPUOutOfMemoryDelegate UE_DEPRECATED(5.2, "Use template instantiation instead as `TMulticastDelegate<void(const uint64, const uint64)>`") = TMulticastDelegate<void(const uint64, const uint64)>;
-	using FGetOnScreenMessagesDelegate UE_DEPRECATED(5.2, "Use template instantiation instead as `TMulticastDelegate<void(FSeverityMessageMap&)>`") = TMulticastDelegate<void(FSeverityMessageMap&)>;
-	using FIsLoadingMovieCurrentlyPlaying UE_DEPRECATED(5.2, "Use template instantiation instead as `TDelegate<bool()>`") = TDelegate<bool()>;
-	using FShouldLaunchUrl UE_DEPRECATED(5.2, "Use template instantiation instead as `TDelegate<bool(const TCHAR* URL)>`") = TDelegate<bool(const TCHAR* URL)>;
-	using FOnActivatedByProtocol UE_DEPRECATED(5.2, "Use template instantiation instead as `TMulticastDelegate<void(const FString& Parameter, FPlatformUserId UserId /*= PLATFORMUSERID_NONE*/)>`") = TMulticastDelegate<void(const FString& Parameter, FPlatformUserId UserId /*= PLATFORMUSERID_NONE*/)>;
-	using FOnGCFinishDestroyTimeExtended UE_DEPRECATED(5.2, "Use template instantiation instead as `TMulticastDelegate<void(const FString&)>`") = TMulticastDelegate<void(const FString&)>;
-	using FApplicationNetworkInitializationChanged UE_DEPRECATED(5.2, "Use template instantiation instead as `TMulticastDelegate<void(bool bIsNetworkInitialized)>`") = TMulticastDelegate<void(bool bIsNetworkInitialized)>;
-	using FAccesExtraBinaryConfigData UE_DEPRECATED(5.2, "Use template instantiation instead as `TMulticastDelegate<void(FExtraBinaryConfigData&)>`") = TMulticastDelegate<void(FExtraBinaryConfigData&)>;
-	using FPreloadPackageShaderMaps UE_DEPRECATED(5.2, "Use template instantiation instead as `TDelegate<void(TArrayView<const FSHAHash>, FAttachShaderReadRequestFunc)>`") = TDelegate<void(TArrayView<const FSHAHash>, FAttachShaderReadRequestFunc)>;
-	using FReleasePreloadedPackageShaderMaps UE_DEPRECATED(5.2, "Use template instantiation instead as `TDelegate<void(TArrayView<const FSHAHash>)>`") = TDelegate<void(TArrayView<const FSHAHash>)>;
-	using FOnLogVerbosityChanged UE_DEPRECATED(5.2, "Use template instantiation instead as `TMulticastDelegate<void(const FLogCategoryName& CategoryName, ELogVerbosity::Type OldVerbosity, ELogVerbosity::Type NewVerbosity)>`") = TMulticastDelegate<void(const FLogCategoryName& CategoryName, ELogVerbosity::Type OldVerbosity, ELogVerbosity::Type NewVerbosity)>;
-PRAGMA_DISABLE_DEPRECATION_WARNINGS
-	using FCreatePackageStore UE_DEPRECATED(5.2, "Use template instantiation instead as `TDelegate<TSharedPtr<class IPackageStore>()>`") = TDelegate<TSharedPtr<class IPackageStore>()>;
-PRAGMA_ENABLE_DEPRECATION_WARNINGS
-	using FProcessForkDelegate UE_DEPRECATED(5.2, "Use template instantiation instead as `TMulticastDelegate<void(EForkProcessRole ProcessRole)>`") = TMulticastDelegate<void(EForkProcessRole ProcessRole)>;
 };
diff --git a/Engine/Source/Runtime/CoreUObject/Private/Serialization/AsyncLoading.cpp b/Engine/Source/Runtime/CoreUObject/Private/Serialization/AsyncLoading.cpp
--- a/Engine/Source/Runtime/CoreUObject/Private/Serialization/AsyncLoading.cpp
+++ b/Engine/Source/Runtime/CoreUObject/Private/Serialization/AsyncLoading.cpp
@@ -7480,13 +7480,6 @@ int32 FAsyncLoadingThread::LoadPackage(const FPackagePath& InPackagePath, FName
 	FName PackageName = InCustomName.IsNone() ? InPackagePath.GetPackageFName() : InCustomName;
 	UE_SCOPED_COOK_STAT(PackageName, EPackageEventStatType::LoadPackage);
 
-PRAGMA_DISABLE_DEPRECATION_WARNINGS
-	if ( FCoreDelegates::OnAsyncLoadPackage.IsBound() )
-	{
-		FCoreDelegates::OnAsyncLoadPackage.Broadcast(PackageName.ToString());
-	}
-PRAGMA_ENABLE_DEPRECATION_WARNINGS
-
 	if (FCoreDelegates::GetOnAsyncLoadPackage().IsBound())
 	{
 		FCoreDelegates::GetOnAsyncLoadPackage().Broadcast(PackageName.ToString());
diff --git a/Engine/Source/Runtime/CoreUObject/Private/Serialization/AsyncLoading2.cpp b/Engine/Source/Runtime/CoreUObject/Private/Serialization/AsyncLoading2.cpp
--- a/Engine/Source/Runtime/CoreUObject/Private/Serialization/AsyncLoading2.cpp
+++ b/Engine/Source/Runtime/CoreUObject/Private/Serialization/AsyncLoading2.cpp
@@ -11411,15 +11411,6 @@ int32 FAsyncLoadingThread2::LoadPackageInternal(const FPackagePath& InPackagePat
 	SCOPED_CUSTOM_LOADTIMER(LoadAsyncPackageInternal)
 		ADD_CUSTOM_LOADTIMER_META(LoadAsyncPackageInternal, PackageName, *WriteToString<256>(PackageNameToLoad));
 
-PRAGMA_DISABLE_DEPRECATION_WARNINGS
-	if (FCoreDelegates::OnAsyncLoadPackage.IsBound())
-	{
-		checkf(IsInGameThread(), TEXT("FCoreDelegates::OnAsyncLoadPackage is not thread-safe and deprecated, update the callees to be thread-safe and register to FCoreDelegates::GetOnAsyncLoadPackage() instead before calling LoadPackageAsync from any other thread than the game-thread."));
-		const FName PackageName = InCustomName.IsNone() ? PackageNameToLoad : InCustomName;
-		FCoreDelegates::OnAsyncLoadPackage.Broadcast(PackageName.ToString());
-	}
-PRAGMA_ENABLE_DEPRECATION_WARNINGS
-
 	if (FCoreDelegates::GetOnAsyncLoadPackage().IsBound())
 	{
 		const FName PackageName = InCustomName.IsNone() ? PackageNameToLoad : InCustomName;
diff --git a/Engine/Source/Runtime/PakFile/Private/IPlatformFilePak.cpp b/Engine/Source/Runtime/PakFile/Private/IPlatformFilePak.cpp
--- a/Engine/Source/Runtime/PakFile/Private/IPlatformFilePak.cpp
+++ b/Engine/Source/Runtime/PakFile/Private/IPlatformFilePak.cpp
@@ -6122,15 +6122,6 @@ bool FPakPlatformFile::Mount(const FPakMountArgs& MountArgs, FIoStatus* OutIoMou
 				{
 					FScopedDurationTimer Timer(OnPakFileMounted2Time);
 
-					PRAGMA_DISABLE_DEPRECATION_WARNINGS
-					if (FCoreDelegates::OnPakFileMounted2.IsBound())
-					{
-						// Avoid calling Broadcast if not in use; Broadcast even on an unsubscribed
-						// non-threadsafe delegate is not threadsafe.
-						FCoreDelegates::OnPakFileMounted2.Broadcast(*Pak);
-					}
-					PRAGMA_ENABLE_DEPRECATION_WARNINGS
-
 					FCoreDelegates::GetOnPakFileMounted2().Broadcast(*Pak);
 				}
 
@@ -6600,9 +6591,6 @@ IFileHandle* FPakPlatformFile::OpenRead(const TCHAR* Filename, bool bAllowWrite)
 
 		if (Result)
 		{
-PRAGMA_DISABLE_DEPRECATION_WARNINGS
-			FCoreDelegates::OnFileOpenedForReadFromPakFile.Broadcast(*PakFile->GetFilename(), Filename);
-PRAGMA_ENABLE_DEPRECATION_WARNINGS
 			FCoreDelegates::GetOnFileOpenedForReadFromPakFile().Broadcast(*PakFile->GetFilename(), Filename);
 		}
 	}
