[CoreDelegates]
- Cleanup deprecations

diff --git a/Engine/Source/Runtime/ApplicationCore/Private/IOS/IOSAppDelegate.cpp b/Engine/Source/Runtime/ApplicationCore/Private/IOS/IOSAppDelegate.cpp
--- a/Engine/Source/Runtime/ApplicationCore/Private/IOS/IOSAppDelegate.cpp
+++ b/Engine/Source/Runtime/ApplicationCore/Private/IOS/IOSAppDelegate.cpp
@@ -1726,10 +1726,6 @@ extern double GCStartTime;
 	 Save data if appropriate.
 	 See also applicationDidEnterBackground:.
 	 */
-	PRAGMA_DISABLE_DEPRECATION_WARNINGS
-	FCoreDelegates::ApplicationWillTerminateDelegate.Broadcast();
-	PRAGMA_ENABLE_DEPRECATION_WARNINGS
-	
 	FCoreDelegates::GetApplicationWillTerminateDelegate().Broadcast();
     
     // note that we are shutting down
diff --git a/Engine/Source/Runtime/Core/Private/Mac/MacPlatformMisc.cpp b/Engine/Source/Runtime/Core/Private/Mac/MacPlatformMisc.cpp
--- a/Engine/Source/Runtime/Core/Private/Mac/MacPlatformMisc.cpp
+++ b/Engine/Source/Runtime/Core/Private/Mac/MacPlatformMisc.cpp
@@ -719,10 +719,6 @@ void FMacPlatformMisc::RequestExit( bool Force, const TCHAR* CallSite)
 	UE_LOG(LogMac, Log,  TEXT("FPlatformMisc::RequestExit(%i, %s)"), Force,
 		CallSite ? CallSite : TEXT("<NoCallSiteInfo>"));
 
-	PRAGMA_DISABLE_DEPRECATION_WARNINGS
-	FCoreDelegates::ApplicationWillTerminateDelegate.Broadcast();
-	PRAGMA_ENABLE_DEPRECATION_WARNINGS
-
 	FCoreDelegates::GetApplicationWillTerminateDelegate().Broadcast();
 
 	notify_cancel(GMacAppInfo.PowerSourceNotification);
