Cleanup of UnifiedError.
- Completly moved over to structured logging for error messaging.
- Added support for localized structured logs.
- Removed old code paths which are no longer used.
- Removed commented code.

diff --git a/Engine/Source/Runtime/Core/Private/Experimental/IO/IoStatusError.cpp b/Engine/Source/Runtime/Core/Private/Experimental/IO/IoStatusError.cpp
--- a/Engine/Source/Runtime/Core/Private/Experimental/IO/IoStatusError.cpp
+++ b/Engine/Source/Runtime/Core/Private/Experimental/IO/IoStatusError.cpp
@@ -45,24 +45,6 @@ namespace UE::UnifiedError::IoStore
 			return FAnsiString(UE::UnifiedError::IoStore::StaticModuleName);
 		}
 
-		/*
-
-		virtual UE::UnifiedError::FAnyString GetErrorCodeString(const FError& Error) const final
-		{
-			UE::UnifiedError::FAnyString Result; 
-			Result.CopyString(TCHAR_TO_UTF8(GetIoErrorText(EIoErrorCode(Error.GetErrorCode()))));
-			return Result;
-		}
-
-		virtual UE::UnifiedError::FAnyString GetModuleIdString(const FError& Error) const final
-		{
-			UE::UnifiedError::FAnyString Result;
-			Result.ReferenceString(UE::UnifiedError::IoStore::StaticModuleName);
-			return Result;
-
-		}
-*/
-
 		virtual const FText GetErrorFormatString(const FError& Error) const final override
 		{
 			return GenericErrorFormatString;
@@ -84,15 +66,6 @@ namespace UE::UnifiedError::IoStore
 			Writer.EndObject();
 		}
 
-		virtual void GetErrorProperties(const FError& Error, IErrorPropertyExtractor& OutProperties) const override
-		{
-			OutProperties.AddProperty(UTF8TEXTVIEW("ErrorCodeString"), GetErrorName(Error));
-			OutProperties.AddProperty(UTF8TEXTVIEW("ModuleIdString"), UE::UnifiedError::IoStore::StaticModuleName);
-			OutProperties.AddProperty(UTF8TEXTVIEW("ErrorCode"), Error.GetErrorCode());
-			OutProperties.AddProperty(UTF8TEXTVIEW("ModuleId"), Error.GetModuleId());
-			OutProperties.AddProperty(UTF8TEXTVIEW("IoStoreErrorMessage"), ErrorMessage);
-		}
-
 		virtual const FAnsiStringView GetErrorDetailsTypeName() const override
 		{
 			return ANSITEXTVIEW("FIoStoreErrorDetails");
diff --git a/Engine/Source/Runtime/Core/Private/Experimental/UnifiedError/UnifiedError.cpp b/Engine/Source/Runtime/Core/Private/Experimental/UnifiedError/UnifiedError.cpp
--- a/Engine/Source/Runtime/Core/Private/Experimental/UnifiedError/UnifiedError.cpp
+++ b/Engine/Source/Runtime/Core/Private/Experimental/UnifiedError/UnifiedError.cpp
@@ -65,19 +65,6 @@ const FManditoryErrorDetails* FError::GetManditoryErrorDetails() const
 	}
 	return (const FManditoryErrorDetails*)(CurrentIt.GetReference());
 
-
-	/*const TRefCountPtr<IErrorDetails> It = ErrorDetails;
-	while (const auto Next = It->GetInnerErrorDetails())
-	{
-		It = Next;
-	}
-
-	
-	if(!ensure(It.GetReference()))
-	{
-		return nullptr;
-	}
-	return (const FManditoryErrorDetails*)(It.GetReference());*/
 }
 
 FAnsiString FError::GetModuleIdAndErrorCodeString() const
@@ -102,7 +89,6 @@ FAnsiString FError::GetModuleIdString() const
 
 
 
-#define USE_STRUCTURED_LOG_FOR_FERRORMESSAGE 1
 #define USE_LOCALIZED_STRUCTURED_LOG_FOR_FERRORMESSAGE 1
 
 
@@ -110,7 +96,6 @@ FAnsiString FError::GetModuleIdString() const
 FText FError::GetErrorMessage(bool bIncludeContext) const
 {
 
-#if USE_STRUCTURED_LOG_FOR_FERRORMESSAGE
 
 	FText FormatText = FText::Join(FText::FromString(": "), NSLOCTEXT("UnifiedError", "DefaultLogPrepend", "{Root/ModuleIdString}.{Root/ErrorCodeString}"), GetFormatErrorText());
 
@@ -150,12 +135,6 @@ FText FError::GetErrorMessage(bool bIncludeContext) const
 	FUtf8StringBuilderBase OutputMessage;
 
 	return Template.FormatToText(Writer.Save());
-#else
-	FFormatNamedArguments Args;
-	FTextFormatArgsPropertyExtractor PropertyExtractor(Args);
-	ErrorDetails->GetErrorProperties(*this, PropertyExtractor);
-	return FText::Format(ErrorDetails->GetErrorFormatString(*this), Args);
-#endif
 }
 
 
@@ -188,99 +167,6 @@ int32 FError::GetModuleId() const
 	return ModuleId;
 }
 
-class FNullErrorPropertyExtractor : public IErrorPropertyExtractor
-{
-public:
-	virtual void AddProperty(const FUtf8StringView& PropertyName, const FStringView& PropertyValue) override { } 
-	virtual void AddProperty(const FUtf8StringView& PropertyName, const FUtf8StringView& PropertyValue) override { }
-	virtual void AddProperty(const FUtf8StringView& PropertyName, const FText& PropertyValue) override { }
-	virtual void AddProperty(const FUtf8StringView& PropertyName, const int64 PropertyValue) override { }
-	virtual void AddProperty(const FUtf8StringView& PropertyName, const int32 PropertyValue) override { }
-	virtual void AddProperty(const FUtf8StringView& PropertyName, const float PropertyValue) override { }
-	virtual void AddProperty(const FUtf8StringView& PropertyName, const double PropertyValue) override { }
-};
-
-template<typename SearchType, typename ResultType>
-class FPropertySearchVisitor : public FNullErrorPropertyExtractor
-{
-public:
-	FPropertySearchVisitor(const FUtf8StringView& InSearchName, ResultType* InResult) : SearchName(InSearchName), Result(InResult)
-	{
-	}
-
-	virtual void AddProperty(const FUtf8StringView& PropertyName, SearchType PropertyValue) override
-	{
-		if ((PropertyName == SearchName) && (Result != nullptr))
-		{
-			*Result = PropertyValue;
-			bFound = true;
-		}
-	}
-
-	bool WasFound() const { return bFound; }
-private:
-
-	const FUtf8StringView& SearchName;
-	ResultType* Result;
-	mutable bool bFound = false;
-};
-
-bool FError::GetDetailByKey(const FUtf8StringView& KeyName, FString& Result) const
-{
-	FPropertySearchVisitor<const FStringView&, FString> Visitor(KeyName, &Result);
-	ErrorDetails->GetErrorProperties(*this, Visitor);
-	return Visitor.WasFound();
-}
-
-bool FError::GetDetailByKey(const FUtf8StringView& KeyName, FUtf8String& Result) const
-{
-	FPropertySearchVisitor<const FUtf8StringView&, FUtf8String> Visitor(KeyName, &Result);
-	ErrorDetails->GetErrorProperties(*this, Visitor);
-	return Visitor.WasFound();
-}
-
-bool FError::GetDetailByKey(const FUtf8StringView& KeyName, FText& Result) const
-{
-	FPropertySearchVisitor<const FText&, FText> Visitor(KeyName, &Result);
-	ErrorDetails->GetErrorProperties(*this, Visitor);
-	return Visitor.WasFound();
-}
-
-bool FError::GetDetailByKey(const FUtf8StringView& KeyName, int64& Result) const
-{
-	FPropertySearchVisitor<const int64, int64> Visitor(KeyName, &Result);
-	ErrorDetails->GetErrorProperties(*this, Visitor);
-	return Visitor.WasFound();
-}
-
-bool FError::GetDetailByKey(const FUtf8StringView& KeyName, int32& Result) const
-{
-	FPropertySearchVisitor<const int32, int32> Visitor(KeyName, &Result);
-	ErrorDetails->GetErrorProperties(*this, Visitor);
-	return Visitor.WasFound();
-}
-
-bool FError::GetDetailByKey(const FUtf8StringView& KeyName, double& Result) const
-{
-	FPropertySearchVisitor<const double, double> Visitor(KeyName, &Result);
-	ErrorDetails->GetErrorProperties(*this, Visitor);
-	return Visitor.WasFound();
-}
-
-bool FError::GetDetailByKey(const FUtf8StringView& KeyName, float& Result) const
-{
-	FPropertySearchVisitor<const float, float> Visitor(KeyName, &Result);
-	ErrorDetails->GetErrorProperties(*this, Visitor);
-	return Visitor.WasFound();
-}
-
-void FError::GetErrorProperties(IErrorPropertyExtractor& Visitor) const
-{
-	ErrorDetails->GetErrorProperties(*this, Visitor);
-}
-
-
-
 class FAppendFormatStringDetails : public FDynamicErrorDetails
 {
 	DECLARE_FERROR_DETAILS(UnifiedError, FAppendFormatStringDetails);
@@ -322,54 +208,11 @@ void FError::AppendFormatString(FText&& InFormatString)
 }
 
 
-// Structured logging FError integration
-class FCbWriterErrorPropertyExtractor : public IErrorPropertyExtractor
-{
-private:
-	FCbWriter& Writer;
-public:
-	FCbWriterErrorPropertyExtractor(FCbWriter& InWriter) : Writer(InWriter) {}
-
-
-
-public:
-	virtual void AddProperty(const FUtf8StringView& PropertyName, const FStringView& PropertyValue) override 
-	{ 
-		Writer.AddString(PropertyName, PropertyValue);
-	}
-	virtual void AddProperty(const FUtf8StringView& PropertyName, const FUtf8StringView& PropertyValue) override
-	{
-		Writer.AddString(PropertyName, PropertyValue);
-	}
-	virtual void AddProperty(const FUtf8StringView& PropertyName, const FText& PropertyValue) override
-	{
-		Writer.AddString(PropertyName, PropertyValue.ToString());
-	}
-	virtual void AddProperty(const FUtf8StringView& PropertyName, const int64 PropertyValue) override
-	{
-		Writer.AddInteger(PropertyName, PropertyValue);
-	}
-	virtual void AddProperty(const FUtf8StringView& PropertyName, const int32 PropertyValue) override 
-	{
-		Writer.AddInteger(PropertyName, PropertyValue);
-	}
-	virtual void AddProperty(const FUtf8StringView& PropertyName, const float PropertyValue) override
-	{
-		Writer.AddFloat(PropertyName, PropertyValue);
-	}
-	virtual void AddProperty(const FUtf8StringView& PropertyName, const double PropertyValue) override 
-	{ 
-		Writer.AddFloat(PropertyName, PropertyValue);
-	}
-};
-
 void SerializeForLog(FCbWriter& Writer, const FError& Error)
 {
 	Writer.BeginObject();
 	Writer.AddString(ANSITEXTVIEW("$type"), ANSITEXTVIEW("FError"));
 
-#if USE_STRUCTURED_LOG_FOR_FERRORMESSAGE
-	
 	// what I have to do 
 	TCbWriter<1024> InnerWriter;
 
@@ -394,27 +237,10 @@ void SerializeForLog(FCbWriter& Writer, const FError& Error)
 
 	Writer.AddString(ANSITEXTVIEW("$text"), OutputMessage.ToString());
 
-#else
-
-	FCbWriterErrorPropertyExtractor ExtractToCbWriter(Writer);
-	Error.GetErrorProperties(ExtractToCbWriter);
-
-
-	// build a format string
-	const FText StructuredLogFormatString = NSLOCTEXT("UnifiedError", "StructuredLogErrorMessage", "{ModuleIdString}.{ErrorCodeString}({ModuleId}.{ErrorCode}): ");
-	FText ErrorFormatText = FText::Join(FText::FromString(TEXT("")), StructuredLogFormatString, Error.GetFormatErrorText());
-	Writer.AddString(ANSITEXTVIEW("ErrorFormatString"), ErrorFormatText.ToString());
-
-	// build up a formatted error string including some additional information about the error
-	FFormatNamedArguments ErrorMessageArgs;
-	FTextFormatArgsPropertyExtractor ExtractToFTextArgs(ErrorMessageArgs);
-	Error.GetErrorProperties(ExtractToFTextArgs);
-	Writer.AddString("$text", FText::Format(ErrorFormatText, MoveTemp(ErrorMessageArgs)).ToString());
-#endif
 	Writer.EndObject();
 }
 
-// FRefCountedErrorDetails fucntions
+// FRefCountedErrorDetails functions
 FRefCountedErrorDetails::~FRefCountedErrorDetails()
 {
 }
@@ -434,34 +260,6 @@ const FText FDynamicErrorDetails::GetErrorFormatString(const FError & Error) con
 	check(InnerErrorDetails);
 	return InnerErrorDetails->GetErrorFormatString(Error);
 }
-/*
-
-const FUtf8String& FDynamicErrorDetails::GetErrorName(const FError& Error) const
-{
-	check(InnerErrorDetails);
-	return InnerErrorDetails->GetErrorName(Error);
-}
-*/
-/*
-
-
-UE::UnifiedError::FAnyString FDynamicErrorDetails::GetErrorCodeString(const FError& Error) const
-{
-	if (InnerErrorDetails.IsValid())
-	{
-		return InnerErrorDetails->GetErrorCodeString(Error);
-	}
-	return UE::UnifiedError::FAnyString();
-}
-
-UE::UnifiedError::FAnyString FDynamicErrorDetails::GetModuleIdString(const FError& Error) const
-{
-	if (InnerErrorDetails.IsValid())
-	{
-		return InnerErrorDetails->GetModuleIdString(Error);
-	}
-	return UE::UnifiedError::FAnyString();
-}*/
 
 
 void FDynamicErrorDetails::SetInnerErrorDetails(TRefCountPtr<const IErrorDetails> InInnerErrorDetails)
@@ -475,26 +273,19 @@ FStaticErrorDetails::FStaticErrorDetails(const FAnsiStringView& InErrorName, con
 {
 }
 
-
-/*
-
-UE::UnifiedError::FAnyString FStaticErrorDetails::GetErrorCodeString(const FError& Error) const
+const FText FStaticErrorDetails::GetErrorFormatString(const FError& Error) const
 {
-	UE::UnifiedError::FAnyString ReturnValue;
-	ReturnValue.ReferenceString(ErrorName);
-	return ReturnValue;
+	return ErrorFormatString;
 }
 
-UE::UnifiedError::FAnyString FStaticErrorDetails::GetModuleIdString(const FError& Error) const
+FAnsiString FStaticErrorDetails::GetErrorCodeString(const FError& Error) const
 {
-	UE::UnifiedError::FAnyString ReturnValue;
-	ReturnValue.ReferenceString(ModuleName);
-	return ReturnValue;
-}*/
+	return FAnsiString(ErrorName);
+}
 
-const FText FStaticErrorDetails::GetErrorFormatString(const FError& Error) const
+FAnsiString FStaticErrorDetails::GetModuleIdString(const FError& Error) const 
 {
-	return ErrorFormatString;
+	return FAnsiString(ModuleName);
 }
 
 // FErrorDetailsRegistry
diff --git a/Engine/Source/Runtime/Core/Public/Experimental/UnifiedError/UnifiedError.h b/Engine/Source/Runtime/Core/Public/Experimental/UnifiedError/UnifiedError.h
--- a/Engine/Source/Runtime/Core/Public/Experimental/UnifiedError/UnifiedError.h
+++ b/Engine/Source/Runtime/Core/Public/Experimental/UnifiedError/UnifiedError.h
@@ -2,12 +2,12 @@
 #pragma once
 
 #include "CoreTypes.h"
-#include "Containers/AnsiString.h"
 #include "Containers/Utf8String.h"
 #include "Containers/StringView.h"
 #include "Internationalization/Internationalization.h"
 #include "Internationalization/Text.h"
 #include "Serialization/CompactBinaryWriter.h"
+#include "Logging/StructuredLogFormat.h"
 #include "Misc/TVariant.h"
 #include "Templates/RefCounting.h"
 
@@ -171,54 +171,9 @@ namespace UE::UnifiedError
 		inline constexpr int32 FIRST_CUSTOM_DETAILS_TYPE = 4;
 	}
 
-	/// <summary>
-	/// IErrorPropertyExtractor is an interface used to visit properties exposed by IErrorDetails implementations
-	/// See also IErrorDetails::GetErrorDetails
-	/// </summary>
-	class IErrorPropertyExtractor
-	{
-	public:
-		virtual void AddProperty(const FUtf8StringView& PropertyName, const FStringView& PropertyValue) = 0;
-		virtual void AddProperty(const FUtf8StringView& PropertyName, const FUtf8StringView& PropertyValue) = 0;
-		virtual void AddProperty(const FUtf8StringView& PropertyName, const FText& PropertyValue) = 0;
-		virtual void AddProperty(const FUtf8StringView& PropertyName, const int64 PropertyValue) = 0;
-		virtual void AddProperty(const FUtf8StringView& PropertyName, const int32 PropertyValue) = 0;
-		virtual void AddProperty(const FUtf8StringView& PropertyName, const float PropertyValue) = 0;
-		virtual void AddProperty(const FUtf8StringView& PropertyName, const double PropertyValue) = 0;
-	};
 
 	class FError;
 
-	// helper class to make pass by string view effiecent while also allowing pass by other strings
-	/*class FAnyString
-	{
-	private:
-		TVariant<FUtf8String, FUtf8StringView> VariantString;
-	public:
-		
-		void ReferenceString(FUtf8StringView InStringView)
-		{
-			VariantString.Set<FUtf8StringView>(InStringView);
-		}
-		void CopyString(FUtf8StringView InString)
-		{
-			VariantString.Set<FUtf8String>(FUtf8String(InString));
-		}
-
-		FUtf8StringView GetUtf8StringView() const
-		{
-			if (VariantString.IsType<FUtf8StringView>())
-			{
-				return VariantString.Get<FUtf8StringView>();
-			}
-			else if (VariantString.IsType<FUtf8String>())
-			{
-				return VariantString.Get<FUtf8String>();
-			}
-			return FUtf8StringView();
-		}
-	};*/
-
 	class IErrorDetails : public IRefCountedObject
 	{
 	public:
@@ -233,13 +188,6 @@ namespace UE::UnifiedError
 		/// <returns></returns>
 		virtual const FText GetErrorFormatString(const FError& Error) const = 0;
 
-		/// <summary>
-		/// GetErrorProperties; used to expose error properties in name:value format, Error properties can be used for string formating functions, searching, exposing context to higher level stacks
-		///	 See also: FError::GetErrorMessage, IErrorDetails::GetErrorFormatString, FError::GetDetailByKey
-		/// </summary>
-		/// <param name="Error"></param>
-		/// <param name="OutProperties"></param>
-		virtual void GetErrorProperties(const FError& Error, IErrorPropertyExtractor& OutProperties) const = 0;
 
 		/// <summary>
 		/// GetInnerErrorDetails; Exposes inner error details to FError, if this ErrorDetails allows inner details
@@ -260,14 +208,8 @@ namespace UE::UnifiedError
 		/// <returns></returns>
 		virtual uint64 GetErrorDetailsTypeId() const = 0;
 
-		/*
-
-		virtual FAnyString GetErrorCodeString(const FError& Error) const = 0;
-		virtual FAnyString GetModuleIdString(const FError& Error) const = 0;
-*/
-
-		virtual const FAnsiStringView GetErrorDetailsTypeName() const = 0;
-		virtual const FAnsiStringView GetErrorDetailsTypeNameForLog() const = 0;
+		CORE_API virtual const FAnsiStringView GetErrorDetailsTypeName() const = 0;
+		CORE_API virtual const FAnsiStringView GetErrorDetailsTypeNameForLog() const = 0;
 		virtual void SerializeToCb(FCbWriter& Writer, const FError& Error) const =0;
 
 
@@ -284,8 +226,8 @@ namespace UE::UnifiedError
 		DECLARE_FERROR_DETAILS_ABSTRACT(UnifiedError, FManditoryErrorDetails);
 	private:
 	public:
-		virtual FAnsiString GetErrorCodeString(const FError& Error) const = 0;
-		virtual FAnsiString GetModuleIdString(const FError& Error) const = 0;
+		CORE_API virtual FAnsiString GetErrorCodeString(const FError& Error) const = 0;
+		CORE_API virtual FAnsiString GetModuleIdString(const FError& Error) const = 0;
 	};
 
 	class FErrorDetailsRegistry
@@ -318,10 +260,6 @@ namespace UE::UnifiedError
 		virtual FReturnedRefCountValue AddRef() const final { return FRefCountBase::AddRef(); }
 		virtual uint32 Release() const final { return FRefCountBase::Release(); }
 		virtual uint32 GetRefCount() const final { return FRefCountBase::GetRefCount(); }
-
-		/*using FRefCountBase::AddRef;
-		using FRefCountBase::Release;
-		using FRefCountBase::GetRefCount;*/
 	};
 
 
@@ -342,12 +280,6 @@ namespace UE::UnifiedError
 		CORE_API virtual void SetInnerErrorDetails(TRefCountPtr<const IErrorDetails> InInnerErrorDetails);
 		virtual TRefCountPtr<const IErrorDetails> GetInnerErrorDetails() const { return InnerErrorDetails; }
 
-/*
-		virtual FAnyString GetErrorCodeString(const FError& Error) const override;
-		virtual FAnyString GetModuleIdString(const FError& Error) const override;
-*/
-
-
 		/// <summary>
 		/// GetErrorFormatString; Pass through to the InnerErrorDetails.
 		/// </summary>
@@ -355,31 +287,9 @@ namespace UE::UnifiedError
 		/// <returns></returns>
 		CORE_API virtual const FText GetErrorFormatString(const FError& Error) const override;
 
-		/// <summary>
-		/// GetErrorProperties; by default pass through to the InnerErrorDetails.
-		///  It's expected FDynamicErrorDetails will be derived from and implement GetErrorProperties.  Declared in header to avoid issues with dll exports
-		/// </summary>
-		/// <param name="Error"></param>
-		/// <param name="OutProperties"></param>
-		virtual void GetErrorProperties(const FError& Error, IErrorPropertyExtractor& OutProperties) const override
-	{
-			InnerErrorDetails->GetErrorProperties(Error, OutProperties);
-		}
 	};
 
 
-	class FError;
-} // UE::UnifiedError
-
-template <typename ValueType>
-inline void GatherPropertiesForError(const UE::UnifiedError::FError& Writer, const ValueType& Value, UE::UnifiedError::IErrorPropertyExtractor& Extractor)
-{
-	static_assert(deferred_false<ValueType>::value, "undefined function, define your own function pls");
-}
-
-namespace UE::UnifiedError
-{
-
 	template<typename T>
 	class TErrorDetails : public FDynamicErrorDetails
 	{
@@ -438,19 +348,13 @@ namespace UE::UnifiedError
 	private:
 		T ErrorDetail;
 		bool bAppendContextToMessage = false;
-
+		
 	public:
 		const T& GetErrorContext() const
 		{
 			return ErrorDetail;
 		}
 
-		virtual void GetErrorProperties(const FError& Error, IErrorPropertyExtractor& OutProperties) const override
-		{
-			GatherPropertiesForError(Error, ErrorDetail, OutProperties);
-			
-			FDynamicErrorDetails::GetErrorProperties(Error, OutProperties);
-		}
 		virtual void SerializeToCb(FCbWriter& Writer, const FError& Error) const override
 		{
 			SerializeForLog(Writer, ErrorDetail);
@@ -511,7 +415,6 @@ namespace UE::UnifiedError
 			Invalidate();
 		}
 
-		void GetErrorProperties(IErrorPropertyExtractor& Visitor) const;
 
 		CORE_API int32 GetErrorCode() const;
 		CORE_API int32 GetModuleId() const;
@@ -536,17 +439,6 @@ namespace UE::UnifiedError
 			}
 			return nullptr;
 
-			/*TRefCountPtr<const IErrorDetails> CurrentIt = ErrorDetails;
-			while (CurrentIt != nullptr)
-			{
-				if (CurrentIt->GetErrorDetailsTypeId() == GetErrorDetails<>)
-				{
-					const TErrorDetails<const T>* Details = (const TErrorDetails<const T>*)(CurrentIt.GetReference());
-					return &Details->GetErrorContext();
-				}
-				CurrentIt = CurrentIt->GetInnerErrorDetails();
-			}
-			return nullptr;*/
 		}
 
 		CORE_API FAnsiString GetErrorCodeString() const;
@@ -563,14 +455,6 @@ namespace UE::UnifiedError
 
 		CORE_API void SerializeDetailsForLog(FCbWriter& Writer) const;
 
-		CORE_API bool GetDetailByKey(const FUtf8StringView& KeyName, FString& Result) const;
-		CORE_API bool GetDetailByKey(const FUtf8StringView& KeyName, FUtf8String& Result) const;
-		CORE_API bool GetDetailByKey(const FUtf8StringView& KeyName, FText& Result) const;
-		CORE_API bool GetDetailByKey(const FUtf8StringView& KeyName, int64& Result) const;
-		CORE_API bool GetDetailByKey(const FUtf8StringView& KeyName, int32& Result) const;
-		CORE_API bool GetDetailByKey(const FUtf8StringView& KeyName, double& Result) const;
-		CORE_API bool GetDetailByKey(const FUtf8StringView& KeyName, float& Result) const;
-
 		FORCEINLINE bool IsValid() const
 		{
 			return (ErrorCode != 0) && (ModuleId != 0);
@@ -680,10 +564,7 @@ namespace UE::UnifiedError
 		///  See also: FError::GetErrorDetails
 		/// </summary>
 		/// <returns></returns>
-		virtual FAnsiString GetErrorCodeString(const FError& Error) const override
-		{
-			return FAnsiString(ErrorName);
-		}
+		CORE_API virtual FAnsiString GetErrorCodeString(const FError& Error) const override;
 
 
 		/// <summary>
@@ -691,28 +572,10 @@ namespace UE::UnifiedError
 		///  Can be called directly on FStaticErrorDetails object
 		///  See also: FError::GetErrorDetails
 		/// </summary>
-		virtual FAnsiString GetModuleIdString(const FError& Error) const override
-		{
-			return FAnsiString(ModuleName);
-		}
+		CORE_API virtual FAnsiString GetModuleIdString(const FError& Error) const override;
 
 		// IErrorDetails implementation
 
-		/// <summary>
-		/// GetErrorProperties;  Adds error properties exposed by this error details object to the OutProperties value
-		///  FStaticErrorDetails includes error details all errors have ErrorCode, ModuleId, ErrorCodeString, ModuleIdString
-		///  These can be used in any error format string see also GetErrorFormatString
-		/// </summary>
-		/// <param name="Error"></param>
-		/// <param name="OutProperties"></param>
-		virtual void GetErrorProperties(const FError& Error, IErrorPropertyExtractor& OutProperties) const
-		{
-			OutProperties.AddProperty(UTF8TEXTVIEW("ErrorCodeString"), ErrorName);
-			OutProperties.AddProperty(UTF8TEXTVIEW("ModuleIdString"), ModuleName);
-			OutProperties.AddProperty(UTF8TEXTVIEW("ErrorCode"), Error.GetErrorCode());
-			OutProperties.AddProperty(UTF8TEXTVIEW("ModuleId"), Error.GetModuleId());
-		}
-
 		/// <summary>
 		/// GetErrorFormatString; return the localized format text generated in DECLARE_ERROR macro.
 		/// </summary>
@@ -741,50 +604,6 @@ namespace UE::UnifiedError
 		virtual bool ShouldIncludeInContextMessage() const { return false; }
 	};
 
-	class FTextFormatArgsPropertyExtractor : public IErrorPropertyExtractor
-	{
-	private:
-		FFormatNamedArguments& Arguments;
-	public:
-		inline  FTextFormatArgsPropertyExtractor(FFormatNamedArguments& InArguments) : Arguments(InArguments)
-		{
-		}
-
-		inline virtual void AddProperty(const FUtf8StringView& PropertyName, const FStringView& PropertyValue) override
-		{
-			Arguments.Add(UTF8_TO_TCHAR(PropertyName.GetData()), FFormatArgumentValue(FText::FromString(PropertyValue.GetData())));
-		}
-
-		inline virtual void AddProperty(const FUtf8StringView& PropertyName, const FUtf8StringView& PropertyValue)  override
-		{
-			Arguments.Add(UTF8_TO_TCHAR(PropertyName.GetData()), FFormatArgumentValue(FText::FromStringView(UTF8_TO_TCHAR(PropertyValue.GetData()))));
-		}
-
-		inline virtual void AddProperty(const FUtf8StringView& PropertyName, const int64 PropertyValue) override
-		{
-			Arguments.Add(UTF8_TO_TCHAR(PropertyName.GetData()), FFormatArgumentValue(PropertyValue));
-		}
-
-		inline virtual void AddProperty(const FUtf8StringView& PropertyName, const int32 PropertyValue) override
-		{
-			Arguments.Add(UTF8_TO_TCHAR(PropertyName.GetData()), FFormatArgumentValue(PropertyValue));
-		}
-
-		inline virtual void AddProperty(const FUtf8StringView& PropertyName, const float PropertyValue) override
-		{
-			Arguments.Add(UTF8_TO_TCHAR(PropertyName.GetData()), FFormatArgumentValue(PropertyValue));
-		}
-
-		inline virtual void AddProperty(const FUtf8StringView& PropertyName, const double PropertyValue) override
-		{
-			Arguments.Add(UTF8_TO_TCHAR(PropertyName.GetData()), FFormatArgumentValue(PropertyValue));
-		}
-
-		inline virtual void AddProperty(const FUtf8StringView& PropertyName, const FText& PropertyValue) override
-		{
-			Arguments.Add(UTF8_TO_TCHAR(PropertyName.GetData()), FFormatArgumentValue(PropertyValue));
-		}
-	};
 
 	inline void SerializeBasicToCb(FCbWriter& Writer, const FUtf8StringView& PropertyName, const FStringView& Value)
 	{
@@ -869,8 +688,6 @@ CORE_API constexpr uint32 GetTypeIdHash(const TCHAR* TypeName)
 	}}}
 
 
-	// TODO: Daniel, $format below should be $formatns and $format
-
 #define DECLARE_ERROR_ONEPARAM(ErrorName, ErrorCode, ModuleName, FormatString, ParamOneType, ParamOneName, ParamOneDefault) \
 	DECLARE_ERROR_INTERNAL(ErrorName, ErrorCode, ModuleName, FText::FromString(UTF8TEXT("{"#ModuleName"_F"#ErrorName"}")) ) \
 	namespace UE::UnifiedError { namespace ModuleName { \
@@ -884,15 +701,11 @@ CORE_API constexpr uint32 GetTypeIdHash(const TCHAR* TypeName)
 	{\
 		Writer.BeginObject();\
 		Writer.AddString(ANSITEXTVIEW("$type"), TErrorStructFeatures<UE::UnifiedError::ModuleName::F##ErrorName>::GetErrorContextTypeNameAsString());\
-		Writer.AddString(ANSITEXTVIEW("$format"), FormatString.ToString());\
+		UE::SerializeLogFormat(Writer, FormatString); \
 		UE::UnifiedError::SerializeBasicToCb(Writer, ANSITEXTVIEW(#ParamOneName), Context.ParamOneName);\
 		Writer.EndObject();\
 	}\
 	namespace UE::UnifiedError { namespace ModuleName { \
-	inline void GatherPropertiesForError(const UE::UnifiedError::FError& Error, const UE::UnifiedError::ModuleName::F##ErrorName& Context, UE::UnifiedError::IErrorPropertyExtractor& PropertyExtractor) \
-	{ \
-		PropertyExtractor.AddProperty(ANSITEXTVIEW(#ParamOneName), Context.ParamOneName);\
-	}\
 	namespace ErrorName { \
 			inline FError MakeError(ParamOneType ParamOneName = ParamOneDefault) \
 			{ \
diff --git a/Engine/Source/Runtime/Engine/Classes/Engine/Experimental/StreamableManagerError.h b/Engine/Source/Runtime/Engine/Classes/Engine/Experimental/StreamableManagerError.h
--- a/Engine/Source/Runtime/Engine/Classes/Engine/Experimental/StreamableManagerError.h
+++ b/Engine/Source/Runtime/Engine/Classes/Engine/Experimental/StreamableManagerError.h
@@ -7,51 +7,6 @@
 
 DECLARE_ERROR_MODULE(StreamableManager, 0x10);
 
-
-//DECLARE_ERROR_ONEPARAM(PackageLoadFailed, 1, StreamableManager, NSLOCTEXT("StreamableManager", "PackageLoadError", "Failed to load package {PackageName}"), FString, PackageName, TEXT("Unknown"));
-
-/*
-
-
-	DECLARE_ERROR_INTERNAL(PackageLoadFailed, 1, StreamableManager, NSLOCTEXT("StreamableManager", "PackageLoadError", "Failed to load package {PackageName}")) \
-	namespace UE::UnifiedError { namespace StreamableManager { \
-			struct FPackageLoadFailed \
-			{\
-				FString PackageName;\
-			};\
-	}}\
-	DECLARE_ERRORSTRUCT_FEATURES(StreamableManager, FPackageLoadFailed);\
-	inline void SerializeForLog(FCbWriter& Writer, const UE::UnifiedError::StreamableManager::FPackageLoadFailed& Context)\
-	{\
-		Writer.BeginObject();\
-		Writer.AddString(ANSITEXTVIEW("$type"), TErrorStructFeatures<UE::UnifiedError::StreamableManager::FPackageLoadFailed>::GetErrorContextTypeNameAsString());\
-		Writer.AddString(ANSITEXTVIEW("$format"),  NSLOCTEXT("StreamableManager", "PackageLoadError", "Failed to load package {PackageName}").ToString());\
-		Writer.AddString(ANSITEXTVIEW("PackageName"), Context.PackageName);\
-		Writer.EndObject();\
-	}\
-	namespace UE::UnifiedError { namespace StreamableManager { namespace PackageLoadError { \
-			inline const int32 ErrorCodeId = UE::UnifiedError::FErrorRegistry::Get().RegisterErrorCode(TEXTVIEW("PackageLoadFailed"), UE::UnifiedError::StreamableManager::StaticModuleId, 1); \
-			FORCEINLINE int32 GetErrorCodeId() \
-			{\
-				return ErrorCodeId; \
-			}\
-			FORCEINLINE FError MakeError(FString PackageName = TEXT("Unknown")) \
-			{ \
-				FError Error = FError(UE::UnifiedError::StreamableManager::StaticModuleId, 1, UE::UnifiedError::StreamableManager::PackageLoadFailed::GetStaticErrorDetails().GetReference()); \
-				FPackageLoadFailed Context = { PackageName };\
-				Error.PushErrorContext(Context); \
-				return Error; \
-			} \
-			FORCEINLINE const FError& GetStaticError() \
-			{\
-				static FError StaticError = MakeError();\
-				return StaticError;\
-			}\
-	}}}
-
-*/
-
-
 DECLARE_ERROR_ONEPARAM(PackageLoadFailed, 1, StreamableManager, NSLOCTEXT("StreamableManager", "PackageLoadError", "Failed to load package {PackageName}"), FString, PackageName, TEXT("Unknown"));
 DECLARE_ERROR_ONEPARAM(PackageLoadCanceled, 2, StreamableManager, NSLOCTEXT("StreamableManager", "PackageLoadCancelled", "Async load canceled {PackageName}"), FString, PackageName, TEXT("Unknown"));
 DECLARE_ERROR(DownloadError, 3, StreamableManager, NSLOCTEXT("StreamableManager", "DownloadError", "Failed to download"));
@@ -62,120 +17,22 @@ DECLARE_ERROR_ONEPARAM(SyncLoadIncomplete, 7, StreamableManager, NSLOCTEXT("Stre
 
 DECLARE_ERROR(AsyncLoadFailed, 8, StreamableManager, NSLOCTEXT("StreamableManager", "AsyncLoadFailed", "Async load failed"));
 DECLARE_ERROR(AsyncLoadCancelled, 9, StreamableManager, NSLOCTEXT("StreamableManager", "AsyncLoadCancelled", "Async load cancelled"));
-//DECLARE_ERROR_ONEPARAM(AsyncLoadUnknownError, 10, StreamableManager, NSLOCTEXT("StreamableManager", "AsyncLoadUnknownError", "Unknown async loading error {AsyncLoadingErrorId}."), int32, AsyncLoadingErrorId, -1);
-
-	DECLARE_ERROR_INTERNAL(AsyncLoadUnknownError, 10, StreamableManager, NSLOCTEXT("StreamableManager", "AsyncLoadUnknownError", "Unknown async loading error {AsyncLoadingErrorId}.")) \
-	namespace UE::UnifiedError { namespace StreamableManager { \
-			struct FAsyncLoadUnknownError \
-			{\
-				int32 AsyncLoadingErrorId;\
-			};\
-	}}\
-	DECLARE_ERRORSTRUCT_FEATURES(StreamableManager, FAsyncLoadUnknownError);\
-	inline void SerializeForLog(FCbWriter& Writer, const UE::UnifiedError::StreamableManager::FAsyncLoadUnknownError& Context)\
-	{\
-		Writer.BeginObject();\
-		Writer.AddString(ANSITEXTVIEW("$type"), TErrorStructFeatures<UE::UnifiedError::StreamableManager::FAsyncLoadUnknownError>::GetErrorContextTypeNameAsString());\
-		Writer.AddString(ANSITEXTVIEW("$format"), NSLOCTEXT("StreamableManager", "AsyncLoadUnknownError", "Unknown async loading error {AsyncLoadingErrorId}.").ToString());\
-		UE::UnifiedError::SerializeBasicToCb(Writer, ANSITEXTVIEW("AsyncLoadingErrorId"), Context.AsyncLoadingErrorId); \
-		Writer.EndObject(); \
-	}\
-	
-		namespace UE::UnifiedError {
-		namespace StreamableManager {
-			
-			inline void GatherPropertiesForError(const UE::UnifiedError::FError& Error, const UE::UnifiedError::StreamableManager::FAsyncLoadUnknownError& Context, UE::UnifiedError::IErrorPropertyExtractor& PropertyExtractor) \
-			{ \
-				PropertyExtractor.AddProperty(ANSITEXTVIEW("AsyncLoadingErrorId"), Context.AsyncLoadingErrorId); \
-			}\
-	namespace AsyncLoadUnknownError { \
-		
-			FORCEINLINE FError MakeError(int32 AsyncLoadingErrorId = -1) \
-			{ \
-				FError Error = FError(UE::UnifiedError::StreamableManager::StaticModuleId, UE::UnifiedError::StreamableManager::AsyncLoadUnknownError::GetErrorCodeId(), UE::UnifiedError::StreamableManager::AsyncLoadUnknownError::GetStaticErrorDetails().GetReference()); \
-				UE::UnifiedError::StreamableManager::FAsyncLoadUnknownError Context = {AsyncLoadingErrorId}; \
-				Error.PushErrorContext(MoveTemp(Context)); \
-				return Error; \
-			} \
-			FORCEINLINE const FError& GetStaticError() \
-			{\
-				static FError StaticError = MakeError();\
-				return StaticError;\
-			}\
-	}}}
-
-
+DECLARE_ERROR_ONEPARAM(AsyncLoadUnknownError, 10, StreamableManager, NSLOCTEXT("StreamableManager", "AsyncLoadUnknownError", "Unknown async loading error {AsyncLoadingErrorId}."), int32, AsyncLoadingErrorId, -1);
 DECLARE_ERROR(UnknownError, 11, StreamableManager, NSLOCTEXT("StreamableManager", "UnknownError", "Unknown error occured while streaming asset"));
 DECLARE_ERROR(AsyncLoadNotInstalled, 12, StreamableManager, NSLOCTEXT("StreamableManager", "AsyncLoadNotInstalled", "Async load failed because the package is not installed."));
 
 namespace UE::UnifiedError::StreamableManager
 {
-	class FStreamableManagerErrorDetails : public FDynamicErrorDetails
-	{
-		DECLARE_FERROR_DETAILS(StreamableManager, FStreamableManagerErrorDetails);
-	private:
-		FString SoftObjectPath;
-		FText ErrorFormatString;
-		FStreamableManagerErrorDetails() = default;
-	public:
-
-		FStreamableManagerErrorDetails(const FString& InSoftObjectPath) 
-			: SoftObjectPath(InSoftObjectPath)
-			, ErrorFormatString(NSLOCTEXT("StreamableManager", "StreamableManagerErrorDetails", "{InnerErrorMessage} Target Path: {SoftObjectPath}"))
-		{
-		}
-
-		virtual void GetErrorProperties(const FError& Error, IErrorPropertyExtractor& OutProperties) const override
-		{
-			OutProperties.AddProperty(UTF8TEXTVIEW("SoftObjectPath"), SoftObjectPath);
-			FFormatNamedArguments ErrorMessageArgs;
-			UE::UnifiedError::FTextFormatArgsPropertyExtractor ExtractToFTextArgs(ErrorMessageArgs);
-			GetInnerErrorDetails()->GetErrorProperties(Error, ExtractToFTextArgs);
-			const FText& FormatString = GetInnerErrorDetails()->GetErrorFormatString(Error);
-			FText Result =FText::Format(FormatString, ErrorMessageArgs);
-			OutProperties.AddProperty(UTF8TEXTVIEW("InnerErrorMessage"), Result);
-			FDynamicErrorDetails::GetErrorProperties(Error, OutProperties);
-		}
-
-		virtual const FText GetErrorFormatString(const FError& Error) const override
-		{
-			const FText InnerFormatString = GetInnerErrorDetails()->GetErrorFormatString(Error);
-
-			FText MyFormatString = NSLOCTEXT("StreamableManager", "DetailsFormatString", "SoftObjectPath:{StreamableManager_FStreamableManagerErrorDetails/SoftObjectPath}");
-
-			return FText::Join(NSLOCTEXT("StreamableManager", "DetailsFormatStringDelimiter", " "), InnerFormatString, MyFormatString);
-		}
-
-		virtual void SerializeToCb(FCbWriter& Writer, const FError& Error) const
-		{
-			Writer.BeginObject();
-			Writer.AddString(UTF8TEXTVIEW("$type"), UTF8TEXTVIEW("FStreamableManagerErrorDetails"));
-			Writer.AddString(UTF8TEXTVIEW("$format"), UTF8TEXTVIEW("(SoftObjectPath:{SoftObjectPath})"));
-			Writer.AddString(UTF8TEXTVIEW("SoftObjectPath"), SoftObjectPath);
-			Writer.EndObject();
-		}
-
-		const FString& GetSoftObjectPath() const { return SoftObjectPath; }
-
-		
-	};
 	UE::UnifiedError::FError GetStreamableError(EAsyncLoadingResult::Type Result);
 
 
-
 	class FStreamableManagerAdditionalContext
 	{
 	public:
-		/*friend void SerializeForLog(FCbWriter& Writer, const UE::UnifiedError::StreamableManager::FStreamableManagerAdditionalContext& Context);
-	private:*/
 		FString SoftObjectPath;
 	};
 
 
-	inline void GatherPropertiesForError(const UE::UnifiedError::FError& Error, const UE::UnifiedError::StreamableManager::FStreamableManagerAdditionalContext& Context, UE::UnifiedError::IErrorPropertyExtractor& PropertyExtractor)
-	{
-		PropertyExtractor.AddProperty(ANSITEXTVIEW("SoftObjectPath"), Context.SoftObjectPath);
-	}
 	
 }
 DECLARE_ERRORSTRUCT_FEATURES(StreamableManager, FStreamableManagerAdditionalContext);
