Removed some deprecated functions from UnrealTemplate.h.

diff --git a/Engine/Source/Runtime/Core/Public/Templates/UnrealTemplate.h b/Engine/Source/Runtime/Core/Public/Templates/UnrealTemplate.h
--- a/Engine/Source/Runtime/Core/Public/Templates/UnrealTemplate.h
+++ b/Engine/Source/Runtime/Core/Public/Templates/UnrealTemplate.h
@@ -19,34 +19,6 @@
 	Standard templates.
 -----------------------------------------------------------------------------*/
 
-/**
- * Chooses between the two parameters based on whether the first is nullptr or not.
- * @return If the first parameter provided is non-nullptr, it is returned; otherwise the second parameter is returned.
- */
-template<typename ReferencedType>
-UE_DEPRECATED(5.4, "IfAThenAElseB has been deprecated - use regular C++ constructs instead")
-FORCEINLINE ReferencedType* IfAThenAElseB(ReferencedType* A,ReferencedType* B)
-{
-	return A ? A : B;
-}
-
-/** branchless pointer selection based on predicate
-* return PTRINT(Predicate) ? A : B;
-**/
-template<typename PredicateType,typename ReferencedType>
-UE_DEPRECATED(5.4, "IfPThenAElseB has been deprecated - use regular C++ constructs instead")
-FORCEINLINE ReferencedType* IfPThenAElseB(PredicateType Predicate,ReferencedType* A,ReferencedType* B)
-{
-	return Predicate ? A : B;
-}
-
-/** A logical exclusive or function. */
-UE_DEPRECATED(5.4, "XOR has been deprecated - use regular C++ constructs instead")
-inline bool XOR(bool A, bool B)
-{
-	return A != B;
-}
-
 /** This is used to provide type specific behavior for a copy which cannot change the value of B. */
 template<typename T>
 inline void Move(T& A,typename TMoveSupportTraits<T>::Copy B)
@@ -419,18 +391,6 @@ struct TOptionalGuardValue : private FNoncopyable
 		}
 	}
 
-	/**
-	 * Overloaded dereference operator.
-	 * Provides read-only access to the original value of the data being tracked by this struct
-	 *
-	 * @return	a const reference to the original data value
-	 */
-	UE_DEPRECATED(5.5, "Use GetOriginalValue() instead.")
-	FORCEINLINE const AssignedType& operator*() const
-	{
-		return OriginalValue;
-	}
-
 	/**
 	 * Provides read-only access to the original value of the data being tracked by this struct
 	 *
