Logging: Support UTF-8 format strings structured log formatting

diff --git a/Engine/Source/Runtime/Core/Private/Logging/StructuredLog.cpp b/Engine/Source/Runtime/Core/Private/Logging/StructuredLog.cpp
--- a/Engine/Source/Runtime/Core/Private/Logging/StructuredLog.cpp
+++ b/Engine/Source/Runtime/Core/Private/Logging/StructuredLog.cpp
@@ -142,7 +142,7 @@ static void LogFieldValue(TStringBuilderBase<CharType>& Out, const FCbFieldView&
 			// Use $format for formatting if present.
 			if (FCbFieldView FormatField = Object.FindViewIgnoreCase(ANSITEXTVIEW("$format")); FormatField.IsString())
 			{
-				TStringBuilder<128> Format(InPlace, FormatField.AsString());
+				TUtf8StringBuilder<128> Format(InPlace, FormatField.AsString());
 				FInlineLogTemplate Template(*Format, {.bAllowSubObjectReferences = true});
 				Template.FormatTo(Out, Object.CreateViewIterator());
 				break;
@@ -156,7 +156,7 @@ static void LogFieldValue(TStringBuilderBase<CharType>& Out, const FCbFieldView&
 			{
 				TStringBuilder<32> Namespace(InPlace, LocNamespaceField.AsString());
 				TStringBuilder<32> Key(InPlace, LocKeyField.AsString());
-				TStringBuilder<128> Format(InPlace, LocFormatField.AsString());
+				TUtf8StringBuilder<128> Format(InPlace, LocFormatField.AsString());
 				FInlineLogTemplate Template(*Namespace, *Key, *Format, {.bAllowSubObjectReferences = true});
 				Template.FormatTo(Out, Object.CreateViewIterator());
 				break;
@@ -257,7 +257,7 @@ static void LogFieldValue(TStringBuilderBase<CharType>& Out, const FCbFieldView&
 	}
 }
 
-static void AddFieldValue(FFormatNamedArguments& Out, FAnsiStringView FieldPath, const FCbFieldView& Field)
+static void AddFieldValue(FFormatNamedArguments& Out, FUtf8StringView FieldPath, const FCbFieldView& Field)
 {
 	FString FieldName(FieldPath);
 
@@ -285,16 +285,17 @@ static void AddFieldValue(FFormatNamedArguments& Out, FAnsiStringView FieldPath,
 	Out.Emplace(MoveTemp(FieldName), FText::FromString(FString(Text)));
 }
 
-class FFieldFinder
+template <typename FormatCharType>
+class TFieldFinder
 {
 public:
-	inline FFieldFinder(const TCHAR* InFormat, const FCbFieldViewIterator& InFields)
+	inline TFieldFinder(const FormatCharType* InFormat, const FCbFieldViewIterator& InFields)
 		: Format(InFormat)
 		, Fields(InFields)
 	{
 	}
 
-	const FCbFieldView& Find(FAnsiStringView Name, int32 IndexHint = -1)
+	const FCbFieldView& Find(FUtf8StringView Name, int32 IndexHint = -1)
 	{
 		if (IndexHint >= 0)
 		{
@@ -327,27 +328,29 @@ public:
 				return It;
 			}
 		}
-		checkf(false, TEXT("Log format requires field '%.*hs' which was not provided. [[%s]]"), Name.Len(), Name.GetData(), Format);
+		checkf(false, TEXT("Log format requires field '%s' which was not provided. [[%s]]"),
+			*WriteToString<32>(Name), StringCast<TCHAR>(Format).Get());
 		return It;
 	}
 
-	FCbFieldView FindByPath(FAnsiStringView Path, int32 IndexHint = -1)
+	FCbFieldView FindByPath(FUtf8StringView Path, int32 IndexHint = -1)
 	{
-		FAnsiStringView Name = Path;
-		bool bMore = String::SplitFirstChar(Path, GFieldPathDelimiter, Name, Path);
+		FUtf8StringView Name = Path;
+		bool bMore = String::SplitFirstChar(Path, UTF8CHAR(GFieldPathDelimiter), Name, Path);
 		FCbFieldView Field = Find(Name, IndexHint);
 		while (bMore)
 		{
 			Name = Path;
-			bMore = String::SplitFirstChar(Path, GFieldPathDelimiter, Name, Path);
+			bMore = String::SplitFirstChar(Path, UTF8CHAR(GFieldPathDelimiter), Name, Path);
 			Field = Field.AsObjectView().FindView(Name);
-			checkf(Field, TEXT("Log format requires field '%.*hs' which was not provided. [[%s]]"), Path.Len(), Path.GetData(), Format);
+			checkf(Field, TEXT("Log format requires field '%s' which was not provided. [[%s]]"),
+				*WriteToString<32>(Path), StringCast<TCHAR>(Format).Get());
 		}
 		return Field;
 	};
 
 private:
-	const TCHAR* Format;
+	const FormatCharType* Format;
 	const FCbFieldViewIterator Fields;
 	FCbFieldViewIterator It{Fields};
 	int32 Index{0};
@@ -402,24 +405,27 @@ class FLogTemplate
 	using FLogField = Logging::Private::FLogField;
 
 public:
+	template <typename FormatCharType>
 	static FLogTemplate* Create(
-		const TCHAR* Format,
+		const FormatCharType* Format,
 		const FLogTemplateOptions& Options,
 		const FLogField* Fields,
 		int32 FieldCount,
-		TFunctionRef<void* (int32)> Allocate);
+		TFunctionWithContext<void* (int32)> Allocate);
 
+	template <typename FormatCharType>
 	static FLogTemplate* CreateLocalized(
 		const FText& FormatText,
-		const TCHAR* Format,
+		const FormatCharType* Format,
 		const FLogTemplateOptions& Options,
 		const FLogField* Fields,
 		int32 FieldCount,
-		TFunctionRef<void* (int32)> Allocate);
+		TFunctionWithContext<void* (int32)> Allocate);
 
 	static void Destroy(FLogTemplate* Template);
 
 	const TCHAR* GetFormat() const { return StaticFormat; }
+	const UTF8CHAR* GetUtf8Format() const { return StaticFormatUtf8; }
 
 	uint8* GetOpData() { return (uint8*)(this + 1); }
 	const uint8* GetOpData() const { return (const uint8*)(this + 1); }
@@ -427,33 +433,60 @@ public:
 	template <typename CharType>
 	void FormatTo(TStringBuilderBase<CharType>& Out, const FCbFieldViewIterator& Fields) const;
 
+	template <typename FormatCharType, typename OutputCharType>
+	static void FormatTo(
+		const FormatCharType* Format,
+		TStringBuilderBase<OutputCharType>& Out,
+		const FCbFieldViewIterator& Fields,
+		const uint8* FirstOp);
+
 	FText FormatToText(const FCbFieldViewIterator& Fields) const;
 
 private:
 	template <typename CharType>
 	void FormatLocalizedTo(TStringBuilderBase<CharType>& Out, const FCbFieldViewIterator& Fields) const;
 
+	template <typename FormatCharType, typename OutputCharType>
+	static void FormatLocalizedTo(
+		const FormatCharType* const Format,
+		TStringBuilderBase<OutputCharType>& Out,
+		const FCbFieldViewIterator& Fields,
+		const uint8* const FirstOp);
+
 	FText FormatLocalizedToText(const FCbFieldViewIterator& Fields) const;
 
+	template <typename FormatCharType>
+	static FText FormatLocalizedToText(
+		const FormatCharType* Format,
+		const FCbFieldViewIterator& Fields,
+		const uint8* FirstOp);
+
 	inline constexpr explicit FLogTemplate(const TCHAR* Format)
 		: StaticFormat(Format)
 	{
 	}
 
+	inline constexpr explicit FLogTemplate(const UTF8CHAR* Format)
+		: StaticFormatUtf8(Format)
+	{
+	}
+
 	FLogTemplate(const FLogTemplate&) = delete;
 	FLogTemplate& operator=(const FLogTemplate&) = delete;
 
 	const TCHAR* StaticFormat = nullptr;
+	const UTF8CHAR* StaticFormatUtf8 = nullptr;
 };
 
 static_assert(std::is_trivially_destructible_v<FLogTemplate>);
 
+template <typename FormatCharType>
 FLogTemplate* FLogTemplate::Create(
-	const TCHAR* Format,
+	const FormatCharType* Format,
 	const FLogTemplateOptions& Options,
 	const FLogField* Fields,
 	const int32 FieldCount,
-	TFunctionRef<void* (int32)> Allocate)
+	TFunctionWithContext<void* (int32)> Allocate)
 {
 	using namespace Logging::Private;
 
@@ -461,10 +494,10 @@ FLogTemplate* FLogTemplate::Create(
 	const bool bFindFields = !!Fields;
 	const bool bPositional = !FieldCount || Algo::NoneOf(FieldsView, &FLogField::Name);
 	checkf(bPositional || Algo::AllOf(FieldsView, &FLogField::Name),
-		TEXT("Log fields must be entirely named or entirely anonymous. [[%s]]"), Format);
+		TEXT("Log fields must be entirely named or entirely anonymous. [[%s]]"), StringCast<TCHAR>(Format).Get());
 	checkf(bPositional || Algo::AllOf(FieldsView,
 		[](const FLogField& Field) { return *Field.Name && *Field.Name != '_' && FAsciiSet::HasOnly(Field.Name, GValidLogFieldName); }),
-		TEXT("Log field names must match \"[A-Za-z0-9][A-Za-z0-9_]*\" in [[%s]]."), Format);
+		TEXT("Log field names must match \"[A-Za-z0-9][A-Za-z0-9_]*\" in [[%s]]."), StringCast<TCHAR>(Format).Get());
 
 	TArray<FLogTemplateOp, TInlineAllocator<16>> Ops;
 
@@ -474,15 +507,15 @@ FLogTemplate* FLogTemplate::Create(
 	int32 FieldSearchIndex = -1;
 	int32 FormatFieldCount = 0;
 	int32 SymbolSearchOffset = 0;
-	for (const TCHAR* TextStart = Format;;)
+	for (const FormatCharType* TextStart = Format;;)
 	{
 		constexpr FAsciiSet Brackets("{}");
-		const TCHAR* const TextEnd = FAsciiSet::FindFirstOrEnd(TextStart + SymbolSearchOffset, Brackets);
+		const FormatCharType* const TextEnd = FAsciiSet::FindFirstOrEnd(TextStart + SymbolSearchOffset, Brackets);
 		SymbolSearchOffset = 0;
 
 		// Escaped "{{" or "}}"
-		if ((TextEnd[0] == TEXT('{') && TextEnd[1] == TEXT('{')) ||
-			(TextEnd[0] == TEXT('}') && TextEnd[1] == TEXT('}')))
+		if ((TextEnd[0] == '{' && TextEnd[1] == '{') ||
+			(TextEnd[0] == '}' && TextEnd[1] == '}'))
 		{
 			// Only "{{" or "}}"
 			if (TextStart == TextEnd)
@@ -515,16 +548,20 @@ FLogTemplate* FLogTemplate::Create(
 		}
 
 		// Parse and validate the field path.
-		const TCHAR* const FieldStart = TextEnd;
-		checkf(*FieldStart == TEXT('{'), TEXT("Log format has an unexpected '%c' character. Use '%c%c' to escape it. [[%s]]"), *FieldStart, *FieldStart, *FieldStart, Format);
-		const TCHAR* const FieldPathEnd = FAsciiSet::Skip(FieldStart + 1, GValidLogFieldPath);
-		checkf(*FieldPathEnd, TEXT("Log format has an unterminated field reference. Use '{{' to escape '{' if needed. [[%s]]"), Format);
-		checkf(*FieldPathEnd == TEXT('}'), TEXT("Log format has invalid character '%c' in field name. "
-			"Use '{{' to escape '{' if needed. Names must match \"[A-Za-z0-9][A-Za-z0-9_]*\". [[%s]]"), *FieldPathEnd, Format);
-		const TCHAR* const FieldEnd = FieldPathEnd + 1;
+		const FormatCharType* const FieldStart = TextEnd;
+		checkf(*FieldStart == '{', TEXT("Log format has an unexpected '%c' character. Use '%c%c' to escape it. [[%s]]"),
+			*FieldStart, *FieldStart, *FieldStart, StringCast<TCHAR>(Format).Get());
+		const FormatCharType* const FieldPathEnd = FAsciiSet::Skip(FieldStart + 1, GValidLogFieldPath);
+		checkf(*FieldPathEnd, TEXT("Log format has an unterminated field reference. Use '{{' to escape '{' if needed. [[%s]]"),
+			StringCast<TCHAR>(Format).Get());
+		checkf(*FieldPathEnd == '}', TEXT("Log format has invalid character '%c' in field name. "
+			"Use '{{' to escape '{' if needed. Names must match \"[A-Za-z0-9][A-Za-z0-9_]*\". [[%s]]"),
+			*FieldPathEnd, StringCast<TCHAR>(Format).Get());
+		const FormatCharType* const FieldEnd = FieldPathEnd + 1;
 		const int32 FieldLen = UE_PTRDIFF_TO_INT32(FieldEnd - FieldStart);
-		checkf(FieldStart[1] != TEXT('_'), TEXT("Log format uses reserved field name '%.*s' with leading '_'. "
-			"Names must match \"[A-Za-z0-9][A-Za-z0-9_]*\". [[%s]]"), FieldLen - 2, FieldStart + 1, Format);
+		checkf(FieldStart[1] != '_', TEXT("Log format uses reserved field name '%s' with leading '_'. "
+			"Names must match \"[A-Za-z0-9][A-Za-z0-9_]*\". [[%s]]"),
+			*WriteToString<32>(MakeStringView(FieldStart + 1, FieldLen - 2)), StringCast<TCHAR>(Format).Get());
 
 		const int32 FieldPathIndex = FieldPathData.Len();
 		FieldPathData.Append(FieldStart + 1, FieldLen - 2);
@@ -532,8 +569,9 @@ FLogTemplate* FLogTemplate::Create(
 		FieldPathSizes.Add(FieldPath.Len());
 
 		const bool bHasSubObjectReference = !!Algo::Find(FieldPath, GFieldPathDelimiter);
-		checkf(!bHasSubObjectReference || Options.bAllowSubObjectReferences, TEXT("Log format has a sub-object reference (%c) "
-			"in a context that does not allow them. [[%s]]"), GFieldPathDelimiter, Format);
+		checkf(!bHasSubObjectReference || Options.bAllowSubObjectReferences,
+			TEXT("Log format has a sub-object reference (%c) in a context that does not allow them. [[%s]]"),
+			GFieldPathDelimiter, StringCast<TCHAR>(Format).Get());
 
 		if (bFindFields && !bPositional)
 		{
@@ -549,7 +587,7 @@ FLogTemplate* FLogTemplate::Create(
 				}
 			}
 			checkf(bFoundField, TEXT("Log format requires field '%.*hs' which was not provided. [[%s]]"),
-				FieldPath.Len(), FieldPath.GetData(), Format);
+				FieldPath.Len(), FieldPath.GetData(), StringCast<TCHAR>(Format).Get());
 		}
 
 		Ops.Add({bHasSubObjectReference ? FLogTemplateOp::OpPath : FLogTemplateOp::OpName, FieldLen});
@@ -559,7 +597,8 @@ FLogTemplate* FLogTemplate::Create(
 	}
 
 	checkf(!bFindFields || !bPositional || FormatFieldCount == FieldCount,
-		TEXT("Log format requires %d fields and %d were provided. [[%s]]"), FormatFieldCount, FieldCount, Format);
+		TEXT("Log format requires %d fields and %d were provided. [[%s]]"),
+		FormatFieldCount, FieldCount, StringCast<TCHAR>(Format).Get());
 
 	const uint32 TotalSize = sizeof(FLogTemplate) + Algo::TransformAccumulate(Ops, FLogTemplateOp::SaveSize, 0);
 	FLogTemplate* const Template = new(Allocate(TotalSize)) FLogTemplate(Format);
@@ -571,13 +610,14 @@ FLogTemplate* FLogTemplate::Create(
 	return Template;
 }
 
+template <typename FormatCharType>
 FLogTemplate* FLogTemplate::CreateLocalized(
 	const FText& FormatText,
-	const TCHAR* Format,
+	const FormatCharType* Format,
 	const FLogTemplateOptions& Options,
 	const FLogField* Fields,
 	const int32 FieldCount,
-	TFunctionRef<void* (int32)> Allocate)
+	TFunctionWithContext<void* (int32)> Allocate)
 {
 	// A localized format string consists of an OpLocalized op followed by a sequence of OpSkip and OpName/OpPath ops
 	// that are terminated by an OpEnd op. Only the first occurrence of each name/path is included and everything else
@@ -587,7 +627,7 @@ FLogTemplate* FLogTemplate::CreateLocalized(
 
 	const bool bFindFields = !!Fields;
 	checkf(!bFindFields || !Options.bAllowSubObjectReferences,
-		TEXT("Validation of field names is not compatible with sub-object references. [[%s]]"), Format);
+		TEXT("Validation of field names is not compatible with sub-object references. [[%s]]"), StringCast<TCHAR>(Format).Get());
 
 	TArray<FLogTemplateOp, TInlineAllocator<16>> Ops;
 	Ops.Add({FLogTemplateOp::OpLocalized, GetLocalizedLogFormatTable().Add(FTextFormat(FormatText))});
@@ -599,14 +639,14 @@ FLogTemplate* FLogTemplate::CreateLocalized(
 	// Parse the format string to find unique field names and optionally validate that required fields are present.
 	int32 FieldSearchIndex = -1;
 	int32 SymbolSearchOffset = 0;
-	for (const TCHAR* TextStart = Format;;)
+	for (const FormatCharType* TextStart = Format;;)
 	{
 		constexpr FAsciiSet Symbols("`{}");
-		const TCHAR* const TextEnd = FAsciiSet::FindFirstOrEnd(TextStart + SymbolSearchOffset, Symbols);
+		const FormatCharType* const TextEnd = FAsciiSet::FindFirstOrEnd(TextStart + SymbolSearchOffset, Symbols);
 		SymbolSearchOffset = 0;
 
 		// Escaped "``" or "`{" or "`}"
-		if (TextEnd[0] == TEXT('`') && (TextEnd[1] == TEXT('`') || TextEnd[1] == TEXT('{') || TextEnd[1] == TEXT('}')))
+		if (TextEnd[0] == '`' && (TextEnd[1] == '`' || TextEnd[1] == '{' || TextEnd[1] == '}'))
 		{
 			// Continue the search after the escaped symbol.
 			SymbolSearchOffset = UE_PTRDIFF_TO_INT32(2 + TextEnd - TextStart);
@@ -621,24 +661,29 @@ FLogTemplate* FLogTemplate::CreateLocalized(
 		}
 
 		// Parse and validate the field path.
-		const TCHAR* const FieldStart = TextEnd;
-		checkf(*FieldStart == TEXT('{'), TEXT("Log format has an unexpected '%c' character. Use '`%c' to escape it. [[%s]]"), *FieldStart, *FieldStart, Format);
-		const TCHAR* const FieldPathEnd = FAsciiSet::Skip(FieldStart + 1, GValidLogFieldPath);
-		checkf(*FieldPathEnd, TEXT("Log format has an unterminated field reference. Use '`{' to escape '{' if needed. [[%s]]"), Format);
-		checkf(*FieldPathEnd == TEXT('}'), TEXT("Log format has invalid character '%c' in field name. "
-			"Use '`{' to escape '{' if needed. Names must match \"[A-Za-z0-9][A-Za-z0-9_]*\". [[%s]]"), *FieldPathEnd, Format);
-		const TCHAR* const FieldEnd = FieldPathEnd + 1;
+		const FormatCharType* const FieldStart = TextEnd;
+		checkf(*FieldStart == '{', TEXT("Log format has an unexpected '%c' character. Use '`%c' to escape it. [[%s]]"),
+			*FieldStart, *FieldStart, StringCast<TCHAR>(Format).Get());
+		const FormatCharType* const FieldPathEnd = FAsciiSet::Skip(FieldStart + 1, GValidLogFieldPath);
+		checkf(*FieldPathEnd, TEXT("Log format has an unterminated field reference. Use '`{' to escape '{' if needed. [[%s]]"),
+			StringCast<TCHAR>(Format).Get());
+		checkf(*FieldPathEnd == '}', TEXT("Log format has invalid character '%c' in field name. "
+			"Use '`{' to escape '{' if needed. Names must match \"[A-Za-z0-9][A-Za-z0-9_]*\". [[%s]]"),
+			*FieldPathEnd, StringCast<TCHAR>(Format).Get());
+		const FormatCharType* const FieldEnd = FieldPathEnd + 1;
 		const int32 FieldLen = UE_PTRDIFF_TO_INT32(FieldEnd - FieldStart);
-		checkf(FieldStart[1] != TEXT('_'), TEXT("Log format uses reserved field name '%.*s' with leading '_'. "
-			"Names must match \"[A-Za-z0-9][A-Za-z0-9_]*\". [[%s]]"), FieldLen - 2, FieldStart + 1, Format);
+		checkf(FieldStart[1] != '_', TEXT("Log format uses reserved field name '%s' with leading '_'. "
+			"Names must match \"[A-Za-z0-9][A-Za-z0-9_]*\". [[%s]]"),
+			*WriteToString<32>(MakeStringView(FieldStart + 1, FieldLen - 2)), StringCast<TCHAR>(Format).Get());
 
 		const int32 FieldPathIndex = FieldPathData.Len();
 		FieldPathData.Append(FieldStart + 1, FieldLen - 2);
 		const FAnsiStringView FieldPath = FieldPathData.ToView().RightChop(FieldPathIndex);
 
 		const bool bHasSubObjectReference = !!Algo::Find(FieldPath, GFieldPathDelimiter);
-		checkf(!bHasSubObjectReference || Options.bAllowSubObjectReferences, TEXT("Log format has a sub-object reference (%c) "
-			"in a context that does not allow them. [[%s]]"), GFieldPathDelimiter, Format);
+		checkf(!bHasSubObjectReference || Options.bAllowSubObjectReferences,
+			TEXT("Log format has a sub-object reference (%c) in a context that does not allow them. [[%s]]"),
+			GFieldPathDelimiter, StringCast<TCHAR>(Format).Get());
 
 		// Check if the field path has been seen and skip it if it has.
 		const bool bExistingField = !!Algo::FindByPredicate(FieldPathSizes, [FieldPath, Data = FieldPathData.GetData()](int32 Size) mutable
@@ -675,7 +720,7 @@ FLogTemplate* FLogTemplate::CreateLocalized(
 				}
 			}
 			checkf(bFoundField, TEXT("Log format requires field '%.*hs' which was not provided. [[%s]]"),
-				FieldPath.Len(), FieldPath.GetData(), Format);
+				FieldPath.Len(), FieldPath.GetData(), StringCast<TCHAR>(Format).Get());
 		}
 
 		Ops.Add({bHasSubObjectReference ? FLogTemplateOp::OpPath : FLogTemplateOp::OpName, FieldLen});
@@ -706,20 +751,37 @@ void FLogTemplate::Destroy(FLogTemplate* Template)
 
 template <typename CharType>
 void FLogTemplate::FormatTo(TStringBuilderBase<CharType>& Out, const FCbFieldViewIterator& Fields) const
+{
+	if (StaticFormatUtf8)
+	{
+		FormatTo(StaticFormatUtf8, Out, Fields, GetOpData());
+	}
+	else
+	{
+		FormatTo(StaticFormat, Out, Fields, GetOpData());
+	}
+}
+
+template <typename FormatCharType, typename OutputCharType>
+void FLogTemplate::FormatTo(
+	const FormatCharType* const Format,
+	TStringBuilderBase<OutputCharType>& Out,
+	const FCbFieldViewIterator& Fields,
+	const uint8* const FirstOp)
 {
 	using namespace Logging::Private;
 
 	int32 FieldIndexHint = -1;
-	const uint8* NextOp = GetOpData();
-	const TCHAR* NextFormat = StaticFormat;
-	FFieldFinder FieldFinder(StaticFormat, Fields);
+	const uint8* NextOp = FirstOp;
+	const FormatCharType* NextFormat = Format;
+	TFieldFinder FieldFinder(Format, Fields);
 	for (;;)
 	{
 		const FLogTemplateOp Op = FLogTemplateOp::Load(NextOp);
 		switch (Op.Code)
 		{
 		case FLogTemplateOp::OpLocalized:
-			return FormatLocalizedTo(Out, Fields);
+			return FormatLocalizedTo(Format, Out, Fields, FirstOp);
 		case FLogTemplateOp::OpEnd:
 			return;
 		case FLogTemplateOp::OpText:
@@ -730,14 +792,14 @@ void FLogTemplate::FormatTo(TStringBuilderBase<CharType>& Out, const FCbFieldVie
 			break;
 		case FLogTemplateOp::OpName:
 			{
-				const auto Name = WriteToAnsiString<32>(MakeStringView(NextFormat + 1, Op.Value - 2));
+				const auto Name = StringCast<UTF8CHAR, 32>(NextFormat + 1, Op.Value - 2);
 				LogFieldValue(Out, FieldFinder.Find(Name, FieldIndexHint));
 				FieldIndexHint = -1;
 			}
 			break;
 		case FLogTemplateOp::OpPath:
 			{
-				const auto Path = WriteToAnsiString<32>(MakeStringView(NextFormat + 1, Op.Value - 2));
+				const auto Path = StringCast<UTF8CHAR, 32>(NextFormat + 1, Op.Value - 2);
 				LogFieldValue(Out, FieldFinder.FindByPath(Path, FieldIndexHint));
 				FieldIndexHint = -1;
 			}
@@ -770,7 +832,33 @@ FORCENOINLINE void FLogTemplate::FormatLocalizedTo(TStringBuilderBase<CharType>&
 	Out.Append(FormatLocalizedToText(Fields).ToString());
 }
 
+template <typename FormatCharType, typename OutputCharType>
+FORCENOINLINE void FLogTemplate::FormatLocalizedTo(
+	const FormatCharType* const Format,
+	TStringBuilderBase<OutputCharType>& Out,
+	const FCbFieldViewIterator& Fields,
+	const uint8* const FirstOp)
+{
+	Out.Append(FormatLocalizedToText(Format, Fields, FirstOp).ToString());
+}
+
 FText FLogTemplate::FormatLocalizedToText(const FCbFieldViewIterator& Fields) const
+{
+	if (StaticFormatUtf8)
+	{
+		return FormatLocalizedToText(StaticFormatUtf8, Fields, GetOpData());
+	}
+	else
+	{
+		return FormatLocalizedToText(StaticFormat, Fields, GetOpData());
+	}
+}
+
+template <typename FormatCharType>
+FText FLogTemplate::FormatLocalizedToText(
+	const FormatCharType* const Format,
+	const FCbFieldViewIterator& Fields,
+	const uint8* const FirstOp)
 {
 	using namespace Logging::Private;
 
@@ -778,9 +866,9 @@ FText FLogTemplate::FormatLocalizedToText(const FCbFieldViewIterator& Fields) co
 	FFormatNamedArguments TextFormatArguments;
 
 	int32 FieldIndexHint = -1;
-	const uint8* NextOp = GetOpData();
-	const TCHAR* NextFormat = StaticFormat;
-	FFieldFinder FieldFinder(StaticFormat, Fields);
+	const uint8* NextOp = FirstOp;
+	const FormatCharType* NextFormat = Format;
+	TFieldFinder FieldFinder(Format, Fields);
 	while (NextOp)
 	{
 		const FLogTemplateOp Op = FLogTemplateOp::Load(NextOp);
@@ -797,14 +885,14 @@ FText FLogTemplate::FormatLocalizedToText(const FCbFieldViewIterator& Fields) co
 			break;
 		case FLogTemplateOp::OpName:
 			{
-				const auto Name = WriteToAnsiString<32>(MakeStringView(NextFormat + 1, Op.Value - 2));
+				const auto Name = StringCast<UTF8CHAR, 32>(NextFormat + 1, Op.Value - 2);
 				AddFieldValue(TextFormatArguments, Name, FieldFinder.Find(Name, FieldIndexHint));
 				FieldIndexHint = -1;
 			}
 			break;
 		case FLogTemplateOp::OpPath:
 			{
-				const auto Path = WriteToAnsiString<32>(MakeStringView(NextFormat + 1, Op.Value - 2));
+				const auto Path = StringCast<UTF8CHAR, 32>(NextFormat + 1, Op.Value - 2);
 				AddFieldValue(TextFormatArguments, Path, FieldFinder.FindByPath(Path, FieldIndexHint));
 				FieldIndexHint = -1;
 			}
@@ -813,23 +901,34 @@ FText FLogTemplate::FormatLocalizedToText(const FCbFieldViewIterator& Fields) co
 		NextFormat += Op.GetSkipSize();
 	}
 
-	checkf(TextFormat, TEXT("Missing text format when formatting localized template. [[%s]]"), StaticFormat);
+	checkf(TextFormat, TEXT("Missing text format when formatting localized template. [[%s]]"), StringCast<TCHAR>(Format).Get());
 	return FText::Format(*TextFormat, TextFormatArguments);
 }
 
 ///////////////////////////////////////////////////////////////////////////////////////////////////////////////////////
 
-void Logging::Private::CreateLogTemplate(const TCHAR* Format, const FLogTemplateOptions& Options, const FLogField* Fields, const int32 FieldCount, TFunctionRef<void* (int32)> Allocate)
+void Logging::Private::CreateLogTemplate(const TCHAR* Format, const FLogTemplateOptions& Options, const FLogField* Fields, const int32 FieldCount, TFunctionWithContext<void* (int32)> Allocate)
 {
 	FLogTemplate::Create(Format, Options, Fields, FieldCount, Allocate);
 }
 
-void Logging::Private::CreateLocalizedLogTemplate(const FText& Format, const FLogTemplateOptions& Options, const FLogField* Fields, const int32 FieldCount, TFunctionRef<void* (int32)> Allocate)
+void Logging::Private::CreateLogTemplate(const UTF8CHAR* Format, const FLogTemplateOptions& Options, const FLogField* Fields, const int32 FieldCount, TFunctionWithContext<void* (int32)> Allocate)
+{
+	FLogTemplate::Create(Format, Options, Fields, FieldCount, Allocate);
+}
+
+void Logging::Private::CreateLocalizedLogTemplate(const FText& Format, const FLogTemplateOptions& Options, const FLogField* Fields, const int32 FieldCount, TFunctionWithContext<void* (int32)> Allocate)
 {
 	FLogTemplate::CreateLocalized(Format, *Format.ToString(), Options, Fields, FieldCount, Allocate);
 }
 
-void Logging::Private::CreateLocalizedLogTemplate(const TCHAR* TextNamespace, const TCHAR* TextKey, const TCHAR* Format, const FLogTemplateOptions& Options, const FLogField* Fields, const int32 FieldCount, TFunctionRef<void* (int32)> Allocate)
+void Logging::Private::CreateLocalizedLogTemplate(const TCHAR* TextNamespace, const TCHAR* TextKey, const TCHAR* Format, const FLogTemplateOptions& Options, const FLogField* Fields, const int32 FieldCount, TFunctionWithContext<void* (int32)> Allocate)
+{
+	const FText FormatText = FText::AsLocalizable_Advanced(TextNamespace, TextKey, Format);
+	FLogTemplate::CreateLocalized(FormatText, Format, Options, Fields, FieldCount, Allocate);
+}
+
+void Logging::Private::CreateLocalizedLogTemplate(const TCHAR* TextNamespace, const TCHAR* TextKey, const UTF8CHAR* Format, const FLogTemplateOptions& Options, const FLogField* Fields, const int32 FieldCount, TFunctionWithContext<void* (int32)> Allocate)
 {
 	const FText FormatText = FText::AsLocalizable_Advanced(TextNamespace, TextKey, Format);
 	FLogTemplate::CreateLocalized(FormatText, Format, Options, Fields, FieldCount, Allocate);
diff --git a/Engine/Source/Runtime/Core/Public/Logging/StructuredLogFormat.h b/Engine/Source/Runtime/Core/Public/Logging/StructuredLogFormat.h
--- a/Engine/Source/Runtime/Core/Public/Logging/StructuredLogFormat.h
+++ b/Engine/Source/Runtime/Core/Public/Logging/StructuredLogFormat.h
@@ -6,13 +6,13 @@
 #include "Containers/StringFwd.h"
 #include "Internationalization/Text.h"
 #include "Misc/ScopeExit.h"
+#include "Templates/FunctionWithContext.h"
 #include "Templates/UniquePtr.h"
 
 #define UE_API CORE_API
 
 class FCbFieldViewIterator;
 class FCbWriter;
-template <typename FuncType> class TFunctionRef;
 
 namespace UE { class FLogTemplate; }
 namespace UE { struct FLogTemplateOptions; }
@@ -107,9 +107,11 @@ private:
 	void* Data = nullptr;
 };
 
-UE_API void CreateLogTemplate(const TCHAR* Format, const FLogTemplateOptions& Options, const FLogField* Fields, int32 FieldCount, TFunctionRef<void* (int32)> Allocate);
-UE_API void CreateLocalizedLogTemplate(const FText& Format, const FLogTemplateOptions& Options, const FLogField* Fields, int32 FieldCount, TFunctionRef<void* (int32)> Allocate);
-UE_API void CreateLocalizedLogTemplate(const TCHAR* TextNamespace, const TCHAR* TextKey, const TCHAR* Format, const FLogTemplateOptions& Options, const FLogField* Fields, int32 FieldCount, TFunctionRef<void* (int32)> Allocate);
+UE_API void CreateLogTemplate(const TCHAR* Format, const FLogTemplateOptions& Options, const FLogField* Fields, int32 FieldCount, TFunctionWithContext<void* (int32)> Allocate);
+UE_API void CreateLogTemplate(const UTF8CHAR* Format, const FLogTemplateOptions& Options, const FLogField* Fields, int32 FieldCount, TFunctionWithContext<void* (int32)> Allocate);
+UE_API void CreateLocalizedLogTemplate(const FText& Format, const FLogTemplateOptions& Options, const FLogField* Fields, int32 FieldCount, TFunctionWithContext<void* (int32)> Allocate);
+UE_API void CreateLocalizedLogTemplate(const TCHAR* TextNamespace, const TCHAR* TextKey, const TCHAR* Format, const FLogTemplateOptions& Options, const FLogField* Fields, int32 FieldCount, TFunctionWithContext<void* (int32)> Allocate);
+UE_API void CreateLocalizedLogTemplate(const TCHAR* TextNamespace, const TCHAR* TextKey, const UTF8CHAR* Format, const FLogTemplateOptions& Options, const FLogField* Fields, int32 FieldCount, TFunctionWithContext<void* (int32)> Allocate);
 
 UE_API void DestroyLogTemplate(FLogTemplate* Template);
 
@@ -138,6 +140,11 @@ public:
 	{
 		Logging::Private::CreateLogTemplate(Format, Options, Fields, FieldCount, [this](int32 Size) { return Storage.Allocate(Size); });
 	}
+public:
+	inline explicit TLogTemplate(const UTF8CHAR* Format, const FLogTemplateOptions& Options = {}, const Logging::Private::FLogField* Fields = nullptr, int32 FieldCount = 0)
+	{
+		Logging::Private::CreateLogTemplate(Format, Options, Fields, FieldCount, [this](int32 Size) { return Storage.Allocate(Size); });
+	}
 
 	inline explicit TLogTemplate(const FText& Format, const FLogTemplateOptions& Options = {}, const Logging::Private::FLogField* Fields = nullptr, int32 FieldCount = 0)
 	{
@@ -149,6 +156,11 @@ public:
 		Logging::Private::CreateLocalizedLogTemplate(TextNamespace, TextKey, Format, Options, Fields, FieldCount, [this](int32 Size) { return Storage.Allocate(Size); });
 	}
 
+	inline explicit TLogTemplate(const TCHAR* TextNamespace, const TCHAR* TextKey, const UTF8CHAR* Format, const FLogTemplateOptions& Options = {}, const Logging::Private::FLogField* Fields = nullptr, int32 FieldCount = 0)
+	{
+		Logging::Private::CreateLocalizedLogTemplate(TextNamespace, TextKey, Format, Options, Fields, FieldCount, [this](int32 Size) { return Storage.Allocate(Size); });
+	}
+
 	inline void FormatTo(FUtf8StringBuilderBase& Out, const FCbFieldViewIterator& Fields)
 	{
 		FormatLogTo(Out, Get(), Fields);
diff --git a/Engine/Source/Runtime/Core/Tests/Logging/StructuredLogFormatTest.cpp b/Engine/Source/Runtime/Core/Tests/Logging/StructuredLogFormatTest.cpp
--- a/Engine/Source/Runtime/Core/Tests/Logging/StructuredLogFormatTest.cpp
+++ b/Engine/Source/Runtime/Core/Tests/Logging/StructuredLogFormatTest.cpp
@@ -174,12 +174,23 @@ TEST_CASE_NAMED(FStructuredLogFormatTest, "Core::Logging::StructuredLogFormat",
 
 	const auto Test = [&Fields](const TCHAR* Format, FUtf8StringView Expected, const FLogTemplateOptions& Options = {})
 	{
-		FInlineLogTemplate Template(Format, Options);
-		TUtf8StringBuilder<1024> Message;
-		Template.FormatTo(Message, Fields);
-		// TODO: Requires a fix for CaptureExpressionsAndValues in LowLevelTestAdapter.h
-		//CAPTURE(Format, Expected, Message);
-		CHECK(Message.ToView().Equals(Expected));
+		{
+			FInlineLogTemplate Template(Format, Options);
+			TUtf8StringBuilder<1024> Message;
+			Template.FormatTo(Message, Fields);
+			// TODO: Requires a fix for CaptureExpressionsAndValues in LowLevelTestAdapter.h
+			//CAPTURE(Format, Expected, Message);
+			CHECK(Message.ToView().Equals(Expected));
+		}
+		{
+			TUtf8StringBuilder<64> Utf8Format(InPlace, Format);
+			FInlineLogTemplate Template(*Utf8Format, Options);
+			TUtf8StringBuilder<1024> Message;
+			Template.FormatTo(Message, Fields);
+			// TODO: Requires a fix for CaptureExpressionsAndValues in LowLevelTestAdapter.h
+			//CAPTURE(Format, Expected, Message);
+			CHECK(Message.ToView().Equals(Expected));
+		}
 	};
 
 	// Invalid
