[Core]
* Added override keyword on a few methods
* Removed "class" in friend statement to fix some warning

diff --git a/Engine/Source/Runtime/Core/Public/Serialization/ArchiveProxy.h b/Engine/Source/Runtime/Core/Public/Serialization/ArchiveProxy.h
--- a/Engine/Source/Runtime/Core/Public/Serialization/ArchiveProxy.h
+++ b/Engine/Source/Runtime/Core/Public/Serialization/ArchiveProxy.h
@@ -287,26 +287,26 @@ public:
 	}
 
 	/** Pushes editor-only marker to the stack of currently serialized properties */
-	virtual inline void PushSerializedProperty(class FProperty* InProperty, const bool bIsEditorOnlyProperty)
+	virtual inline void PushSerializedProperty(class FProperty* InProperty, const bool bIsEditorOnlyProperty) override
 	{
 		FArchive::PushSerializedProperty(InProperty, bIsEditorOnlyProperty);
 		InnerArchive.PushSerializedProperty(InProperty, bIsEditorOnlyProperty);
 	}
 	/** Pops editor-only marker from the stack of currently serialized properties */
-	virtual inline void PopSerializedProperty(class FProperty* InProperty, const bool bIsEditorOnlyProperty)
+	virtual inline void PopSerializedProperty(class FProperty* InProperty, const bool bIsEditorOnlyProperty) override
 	{
 		FArchive::PopSerializedProperty(InProperty, bIsEditorOnlyProperty);
 		InnerArchive.PopSerializedProperty(InProperty, bIsEditorOnlyProperty);
 	}
 #if WITH_EDITORONLY_DATA
 	/** Returns true if the stack of currently serialized properties contains an editor-only property */
-	virtual FORCEINLINE bool IsEditorOnlyPropertyOnTheStack() const
+	virtual FORCEINLINE bool IsEditorOnlyPropertyOnTheStack() const override
 	{
 		return InnerArchive.IsEditorOnlyPropertyOnTheStack();
 	}
 #endif
 
-	virtual FORCEINLINE bool IsProxyOf(FArchive* InOther) const
+	virtual FORCEINLINE bool IsProxyOf(FArchive* InOther) const override
 	{
 		return InOther == this || InOther == &InnerArchive || InnerArchive.IsProxyOf(InOther);
 	}
diff --git a/Engine/Source/Runtime/Core/Public/Serialization/Formatters/BinaryArchiveFormatter.h b/Engine/Source/Runtime/Core/Public/Serialization/Formatters/BinaryArchiveFormatter.h
--- a/Engine/Source/Runtime/Core/Public/Serialization/Formatters/BinaryArchiveFormatter.h
+++ b/Engine/Source/Runtime/Core/Public/Serialization/Formatters/BinaryArchiveFormatter.h
@@ -31,7 +31,7 @@ public:
 	virtual void LeaveRecord() override;
 	virtual void EnterField(FArchiveFieldName Name) override;
 	virtual void LeaveField() override;
-	CORE_API virtual bool TryEnterField(FArchiveFieldName Name, bool bEnterIfWriting);
+	CORE_API virtual bool TryEnterField(FArchiveFieldName Name, bool bEnterIfWriting) override;
 
 	virtual void EnterArray(int32& NumElements) override;
 	virtual void LeaveArray() override;
diff --git a/Engine/Source/Runtime/Core/Public/Serialization/StructuredArchiveSlotBase.h b/Engine/Source/Runtime/Core/Public/Serialization/StructuredArchiveSlotBase.h
--- a/Engine/Source/Runtime/Core/Public/Serialization/StructuredArchiveSlotBase.h
+++ b/Engine/Source/Runtime/Core/Public/Serialization/StructuredArchiveSlotBase.h
@@ -46,7 +46,7 @@ namespace UE::StructuredArchive::Private
 	// Represents a position of a slot within the hierarchy.
 	class FSlotPosition
 	{
-		friend class FStructuredArchive;
+		friend FStructuredArchive;
 
 	public:
 		int32 Depth;
@@ -65,7 +65,7 @@ namespace UE::StructuredArchive::Private
 		: protected FSlotPosition
 #endif
 	{
-		friend class FStructuredArchive;
+		friend FStructuredArchive;
 
 	protected:
 		// Token which needs to be passed to the constructor, so only friends are able to do so, but
