Added UEOps.h to aid the correct generation of all binary comparison operators based on a couple of simple member function definitions.

diff --git a/Engine/Source/Runtime/Core/Private/UObject/UnrealNames.cpp b/Engine/Source/Runtime/Core/Private/UObject/UnrealNames.cpp
--- a/Engine/Source/Runtime/Core/Private/UObject/UnrealNames.cpp
+++ b/Engine/Source/Runtime/Core/Private/UObject/UnrealNames.cpp
@@ -4546,29 +4546,29 @@ uint32 FLazyName::ParseNumber(const WIDECHAR* Str, int32 Len)
 	return FNameHelper::ParseNumber(Str, Len);
 }
 
-bool operator==(const FLazyName& A, const FLazyName& B)
+bool FLazyName::UEOpEquals(const FLazyName& Rhs) const
 {
 	// If we have started creating FNames we might as well resolve and cache both lazy names
-	if (A.Either.IsName() || B.Either.IsName())
+	if (Either.IsName() || Rhs.Either.IsName())
 	{
-		return A.Resolve() == B.Resolve();
+		return Resolve() == Rhs.Resolve();
 	}
 
 	// Literal pointer comparison, can ignore width
-	if (A.Either.AsAnsiLiteral() == B.Either.AsAnsiLiteral())
+	if (Either.AsAnsiLiteral() == Rhs.Either.AsAnsiLiteral())
 	{
 		return true;
 	}
 
-	if (A.bLiteralIsWide)
+	if (bLiteralIsWide)
 	{
-		return B.bLiteralIsWide ? FPlatformString::Stricmp(A.Either.AsWideLiteral(), B.Either.AsWideLiteral()) == 0
-								: FPlatformString::Stricmp(A.Either.AsWideLiteral(), B.Either.AsAnsiLiteral()) == 0;
+		return Rhs.bLiteralIsWide ? FPlatformString::Stricmp(Either.AsWideLiteral(), Rhs.Either.AsWideLiteral()) == 0
+								: FPlatformString::Stricmp(Either.AsWideLiteral(), Rhs.Either.AsAnsiLiteral()) == 0;
 	}
 	else
 	{
-		return B.bLiteralIsWide ? FPlatformString::Stricmp(A.Either.AsAnsiLiteral(), B.Either.AsWideLiteral()) == 0
-								: FPlatformString::Stricmp(A.Either.AsAnsiLiteral(), B.Either.AsAnsiLiteral()) == 0;	
+		return Rhs.bLiteralIsWide ? FPlatformString::Stricmp(Either.AsAnsiLiteral(), Rhs.Either.AsWideLiteral()) == 0
+								: FPlatformString::Stricmp(Either.AsAnsiLiteral(), Rhs.Either.AsAnsiLiteral()) == 0;	
 	}
 }
 
diff --git a/Engine/Source/Runtime/Core/Public/Containers/ArrayView.h b/Engine/Source/Runtime/Core/Public/Containers/ArrayView.h
--- a/Engine/Source/Runtime/Core/Public/Containers/ArrayView.h
+++ b/Engine/Source/Runtime/Core/Public/Containers/ArrayView.h
@@ -7,6 +7,7 @@
 #include "Misc/AssertionMacros.h"
 #include "Misc/ReverseIterate.h"
 #include "Misc/IntrusiveUnsetOptionalState.h"
+#include "Misc/UEOps.h"
 #include "Templates/Invoke.h"
 #include "Templates/UnrealTypeTraits.h"
 #include "Traits/ElementType.h"
@@ -28,7 +29,7 @@ namespace UE::Core::ArrayView::Private
 	// Simply forwards to an unqualified GetData(), but can be called from within TArrayView
 	// where GetData() is already a member and so hides any others.
 	template <typename T>
-	FORCEINLINE decltype(auto) GetDataHelper(T&& Arg)
+	FORCEINLINE auto GetDataHelper(T&& Arg) -> decltype(GetData(Forward<T>(Arg)))
 	{
 		return GetData(Forward<T>(Arg));
 	}
@@ -276,7 +277,7 @@ public:
 		, ArrayNum(-1)
 	{
 	}
-	[[nodiscard]] bool operator==(FIntrusiveUnsetOptionalState) const
+	[[nodiscard]] bool UEOpEquals(FIntrusiveUnsetOptionalState) const
 	{
 		return ArrayNum == -1;
 	}
@@ -822,6 +823,26 @@ public:
 		Algo::StableSort(*this, PredicateWrapper);
 	}
 
+	// Comparison of array views to each other is not implemented because it
+	// is not obvious whether the caller wants an exact match of the data
+	// pointer and size, or to compare the objects being pointed to.
+	template <typename OtherElementType, typename OtherSizeType>
+	bool UEOpEquals(TArrayView<OtherElementType, OtherSizeType>) const = delete;
+
+	/**
+	 * Equality operator.
+	 *
+	 * @param Lhs Another ranged type to compare.
+	 * @returns True if this array view's contents and the other ranged type match. False otherwise.
+	 */
+	template <typename RangeType>
+		requires (std::is_convertible_v<decltype(UE::Core::ArrayView::Private::GetDataHelper(std::declval<RangeType&>())), const ElementType*>)
+	[[nodiscard]] bool UEOpEquals(RangeType&& Rhs) const
+	{
+		auto Count = GetNum(Rhs);
+		return Count == ArrayNum && CompareItems(DataPtr, UE::Core::ArrayView::Private::GetDataHelper(Rhs), Count);
+	}
+
 private:
 	ElementType* DataPtr;
 	SizeType ArrayNum;
@@ -905,69 +926,6 @@ template <typename ElementType>
 
 //////////////////////////////////////////////////////////////////////////
 
-// Comparison of array views to each other is not implemented because it
-// is not obvious whether the caller wants an exact match of the data
-// pointer and size, or to compare the objects being pointed to.
-template <typename ElementType, typename SizeType, typename OtherElementType, typename OtherSizeType>
-bool operator==(TArrayView<ElementType, SizeType>, TArrayView<OtherElementType, OtherSizeType>) = delete;
-template <typename ElementType, typename SizeType, typename OtherElementType, typename OtherSizeType>
-bool operator!=(TArrayView<ElementType, SizeType>, TArrayView<OtherElementType, OtherSizeType>) = delete;
-
-/**
- * Equality operator.
- *
- * @param Lhs Another ranged type to compare.
- * @returns True if this array view's contents and the other ranged type match. False otherwise.
- */
-template <
-	typename RangeType,
-	typename ElementType
-	UE_REQUIRES(std::is_convertible_v<decltype(GetData(std::declval<RangeType&>())), const ElementType*>)
->
-[[nodiscard]] bool operator==(RangeType&& Lhs, TArrayView<ElementType> Rhs)
-{
-	auto Count = GetNum(Lhs);
-	return Count == Rhs.Num() && CompareItems(GetData(Lhs), Rhs.GetData(), Count);
-}
-
-template <
-	typename RangeType,
-	typename ElementType
-	UE_REQUIRES(std::is_convertible_v<decltype(GetData(std::declval<RangeType&>())), const ElementType*>)
->
-[[nodiscard]] bool operator==(TArrayView<ElementType> Lhs, RangeType&& Rhs)
-{
-	return (Rhs == Lhs);
-}
-
-#if !PLATFORM_COMPILER_HAS_GENERATED_COMPARISON_OPERATORS
-/**
- * Inequality operator.
- *
- * @param Lhs Another ranged type to compare.
- * @returns False if this array view's contents and the other ranged type match. True otherwise.
- */
-template <
-	typename RangeType,
-	typename ElementType
-	UE_REQUIRES(std::is_convertible_v<decltype(GetData(std::declval<RangeType&>())), const ElementType*>)
->
-[[nodiscard]] bool operator!=(RangeType&& Lhs, TArrayView<ElementType> Rhs)
-{
-	return !(Lhs == Rhs);
-}
-
-template <
-	typename RangeType,
-	typename ElementType
-	UE_REQUIRES(std::is_convertible_v<decltype(GetData(std::declval<RangeType&>())), const ElementType*>)
->
-[[nodiscard]] bool operator!=(TArrayView<ElementType> Lhs, RangeType&& Rhs)
-{
-	return !(Rhs == Lhs);
-}
-#endif
-
 template<typename InElementType, typename InAllocatorType>
 template<typename OtherElementType, typename OtherSizeType>
 [[nodiscard]] FORCEINLINE TArray<InElementType, InAllocatorType>::TArray(const TArrayView<OtherElementType, OtherSizeType>& Other)
diff --git a/Engine/Source/Runtime/Core/Public/Containers/StringView.h b/Engine/Source/Runtime/Core/Public/Containers/StringView.h
--- a/Engine/Source/Runtime/Core/Public/Containers/StringView.h
+++ b/Engine/Source/Runtime/Core/Public/Containers/StringView.h
@@ -12,6 +12,7 @@
 #include "Misc/Crc.h"
 #include "Misc/CString.h"
 #include "Misc/ReverseIterate.h"
+#include "Misc/UEOps.h"
 #include "String/Find.h"
 #include "Templates/Requires.h"
 #include "Templates/UnrealTemplate.h"
@@ -449,93 +450,37 @@ public:
 		return FCrc::Strihash_DEPRECATED(View.Len(), View.GetData());
 	}
 
-	[[nodiscard]] friend inline bool operator==(TStringView Lhs, TStringView Rhs)
+	[[nodiscard]] inline bool UEOpEquals(TStringView Rhs) const
 	{
-		return Lhs.Equals(Rhs, ESearchCase::IgnoreCase);
+		return this->Equals(Rhs, ESearchCase::IgnoreCase);
 	}
 
-	[[nodiscard]] friend inline bool operator!=(TStringView Lhs, TStringView Rhs)
+	[[nodiscard]] inline bool UEOpLessThan(TStringView Rhs) const
 	{
-		return !(Lhs == Rhs);
-	}
-
-	[[nodiscard]] friend inline bool operator<(TStringView Lhs, TStringView Rhs)
-	{
-		return Lhs.Compare(Rhs, ESearchCase::IgnoreCase) < 0;
-	}
-
-	template <
-		typename CharRangeType
-		UE_REQUIRES(std::is_convertible_v<CharRangeType&&, TStringView>)
-	>
-	[[nodiscard]] friend inline auto operator==(TStringView Lhs, CharRangeType&& Rhs)
-	{
-		return TStringView::PrivateEquals(Lhs, ImplicitConv<TStringView>(Forward<CharRangeType>(Rhs)));
-	}
-
-	template <
-		typename CharRangeType
-		UE_REQUIRES(std::is_convertible_v<CharRangeType&&, TStringView>)
-	>
-	[[nodiscard]] friend inline auto operator==(CharRangeType&& Lhs, TStringView Rhs)
-	{
-		return TStringView::PrivateEquals(ImplicitConv<TStringView>(Forward<CharRangeType>(Lhs)), Rhs);
-	}
-
-	template <
-		typename CharRangeType
-		UE_REQUIRES(std::is_convertible_v<CharRangeType&&, TStringView>)
-	>
-	[[nodiscard]] friend inline auto operator!=(TStringView Lhs, CharRangeType&& Rhs)
-	{
-		return !TStringView::PrivateEquals(Lhs, ImplicitConv<TStringView>(Forward<CharRangeType>(Rhs)));
-	}
-
-	template <
-		typename CharRangeType
-		UE_REQUIRES(std::is_convertible_v<CharRangeType&&, TStringView>)
-	>
-	[[nodiscard]] friend inline auto operator!=(CharRangeType&& Lhs, TStringView Rhs)
-	{
-		return !TStringView::PrivateEquals(ImplicitConv<TStringView>(Forward<CharRangeType>(Lhs)), Rhs);
+		return this->Compare(Rhs, ESearchCase::IgnoreCase) < 0;
 	}
 
 	template <
 		typename CharRangeType
 		UE_REQUIRES(std::is_convertible_v<CharRangeType&&, TStringView>)
 	>
-	[[nodiscard]] friend inline auto operator<(TStringView Lhs, CharRangeType&& Rhs)
+	[[nodiscard]] inline auto UEOpEquals(CharRangeType&& Rhs) const
 	{
-		return TStringView::PrivateLess(Lhs, ImplicitConv<TStringView>(Forward<CharRangeType>(Rhs)));
+		return TStringView::PrivateEquals(*this, ImplicitConv<TStringView>(Forward<CharRangeType>(Rhs)));
 	}
 
 	template <
 		typename CharRangeType
 		UE_REQUIRES(std::is_convertible_v<CharRangeType&&, TStringView>)
 	>
-	[[nodiscard]] friend inline auto operator<(CharRangeType&& Lhs, TStringView Rhs)
-	{
-		return TStringView::PrivateLess(ImplicitConv<TStringView>(Forward<CharRangeType>(Lhs)), Rhs);
-	}
-
-	[[nodiscard]] friend inline bool operator==(TStringView Lhs, const CharType* Rhs)
-	{
-		return TStringView::PrivateEquals(Lhs, Rhs);
-	}
-
-	[[nodiscard]] friend inline bool operator==(const CharType* Lhs, TStringView Rhs)
-	{
-		return Rhs == Lhs;
-	}
-
-	[[nodiscard]] friend inline bool operator!=(TStringView Lhs, const CharType* Rhs)
+	[[nodiscard]] inline auto UEOpLessThan(CharRangeType&& Rhs)
 	{
-		return !(Lhs == Rhs);
+		return TStringView::PrivateLess(*this, ImplicitConv<TStringView>(Forward<CharRangeType>(Rhs)));
 	}
 
-	[[nodiscard]] friend inline bool operator!=(const CharType* Lhs, TStringView Rhs)
+	[[nodiscard]] inline bool UEOpEquals(const CharType* Rhs) const
 	{
-		return !(Lhs == Rhs);
+		return TStringView::PrivateEquals(*this, Rhs);
 	}
 
 public:
diff --git a/Engine/Source/Runtime/Core/Public/Containers/UnrealString.h b/Engine/Source/Runtime/Core/Public/Containers/UnrealString.h
--- a/Engine/Source/Runtime/Core/Public/Containers/UnrealString.h
+++ b/Engine/Source/Runtime/Core/Public/Containers/UnrealString.h
@@ -246,17 +246,12 @@ struct FTextRange
 
 	}
 
-	inline bool operator==(const FTextRange& Other) const
+	inline bool UEOpEquals(const FTextRange& Other) const
 	{
 		return BeginIndex == Other.BeginIndex
 			&& EndIndex == Other.EndIndex;
 	}
 
-	FORCEINLINE bool operator!=(const FTextRange& Other) const
-	{
-		return !(*this == Other);
-	}
-
 	friend inline uint32 GetTypeHash(const FTextRange& Key)
 	{
 		uint32 KeyHash = 0;
diff --git a/Engine/Source/Runtime/Core/Public/Containers/UnrealString.h.inl b/Engine/Source/Runtime/Core/Public/Containers/UnrealString.h.inl
--- a/Engine/Source/Runtime/Core/Public/Containers/UnrealString.h.inl
+++ b/Engine/Source/Runtime/Core/Public/Containers/UnrealString.h.inl
@@ -212,7 +212,7 @@ public:
 		: Data(Tag)
 	{
 	}
-	[[nodiscard]] bool operator==(FIntrusiveUnsetOptionalState Tag) const
+	[[nodiscard]] bool UEOpEquals(FIntrusiveUnsetOptionalState Tag) const
 	{
 		return Data == Tag;
 	}
@@ -878,60 +878,6 @@ public:
 		return Result;
 	}
 
-	/**
-	 * Lexicographically test whether the left string is <= the right string
-	 *
-	 * @param Lhs String to compare against.
-	 * @param Rhs String to compare against.
-	 * @return true if the left string is lexicographically <= the right string, otherwise false
-	 * @note case insensitive
-	 */
-	[[nodiscard]] FORCEINLINE friend bool operator<=(const UE_STRING_CLASS& Lhs, const UE_STRING_CLASS& Rhs)
-	{
-		return FPlatformString::Stricmp(*Lhs, *Rhs) <= 0;
-	}
-
-	/**
-	 * Lexicographically test whether the left string is <= the right string
-	 *
-	 * @param Lhs String to compare against.
-	 * @param Rhs String to compare against.
-	 * @return true if the left string is lexicographically <= the right string, otherwise false
-	 * @note case insensitive
-	 */
-	template <typename CharType>
-	[[nodiscard]] FORCEINLINE friend bool operator<=(const UE_STRING_CLASS& Lhs, const CharType* Rhs)
-	{
-		return FPlatformString::Stricmp(*Lhs, Rhs) <= 0;
-	}
-
-	/**
-	 * Lexicographically test whether the left string is <= the right string
-	 *
-	 * @param Lhs String to compare against.
-	 * @param Rhs String to compare against.
-	 * @return true if the left string is lexicographically <= the right string, otherwise false
-	 * @note case insensitive
-	 */
-	template <typename CharType>
-	[[nodiscard]] FORCEINLINE friend bool operator<=(const CharType* Lhs, const UE_STRING_CLASS& Rhs)
-	{
-		return FPlatformString::Stricmp(Lhs, *Rhs) <= 0;
-	}
-
-	/**
-	 * Lexicographically test whether the left string is < the right string
-	 *
-	 * @param Lhs String to compare against.
-	 * @param Rhs String to compare against.
-	 * @return true if the left string is lexicographically < the right string, otherwise false
-	 * @note case insensitive
-	 */
-	[[nodiscard]] FORCEINLINE friend bool operator<(const UE_STRING_CLASS& Lhs, const UE_STRING_CLASS& Rhs)
-	{
-		return FPlatformString::Stricmp(*Lhs, *Rhs) < 0;
-	}
-
 	/**
 	 * Lexicographically test whether the left string is < the right string
 	 *
@@ -940,10 +886,9 @@ public:
 	 * @return true if the left string is lexicographically < the right string, otherwise false
 	 * @note case insensitive
 	 */
-	template <typename CharType>
-	[[nodiscard]] FORCEINLINE friend bool operator<(const UE_STRING_CLASS& Lhs, const CharType* Rhs)
+	[[nodiscard]] FORCEINLINE bool UEOpLessThan(const UE_STRING_CLASS& Rhs) const
 	{
-		return FPlatformString::Stricmp(*Lhs, Rhs) < 0;
+		return FPlatformString::Stricmp(**this, *Rhs) < 0;
 	}
 
 	/**
@@ -955,91 +900,9 @@ public:
 	 * @note case insensitive
 	 */
 	template <typename CharType>
-	[[nodiscard]] FORCEINLINE friend bool operator<(const CharType* Lhs, const UE_STRING_CLASS& Rhs)
-	{
-		return FPlatformString::Stricmp(Lhs, *Rhs) < 0;
-	}
-
-	/**
-	 * Lexicographically test whether the left string is >= the right string
-	 *
-	 * @param Lhs String to compare against.
-	 * @param Rhs String to compare against.
-	 * @return true if the left string is lexicographically >= the right string, otherwise false
-	 * @note case insensitive
-	 */
-	[[nodiscard]] FORCEINLINE friend bool operator>=(const UE_STRING_CLASS& Lhs, const UE_STRING_CLASS& Rhs)
-	{
-		return FPlatformString::Stricmp(*Lhs, *Rhs) >= 0;
-	}
-
-	/**
-	 * Lexicographically test whether the left string is >= the right string
-	 *
-	 * @param Lhs String to compare against.
-	 * @param Rhs String to compare against.
-	 * @return true if the left string is lexicographically >= the right string, otherwise false
-	 * @note case insensitive
-	 */
-	template <typename CharType>
-	[[nodiscard]] FORCEINLINE friend bool operator>=(const UE_STRING_CLASS& Lhs, const CharType* Rhs)
-	{
-		return FPlatformString::Stricmp(*Lhs, Rhs) >= 0;
-	}
-
-	/**
-	 * Lexicographically test whether the left string is >= the right string
-	 *
-	 * @param Lhs String to compare against.
-	 * @param Rhs String to compare against.
-	 * @return true if the left string is lexicographically >= the right string, otherwise false
-	 * @note case insensitive
-	 */
-	template <typename CharType>
-	[[nodiscard]] FORCEINLINE friend bool operator>=(const CharType* Lhs, const UE_STRING_CLASS& Rhs)
+	[[nodiscard]] FORCEINLINE bool UEOpLessThan(const CharType* Rhs) const
 	{
-		return FPlatformString::Stricmp(Lhs, *Rhs) >= 0;
-	}
-
-	/**
-	 * Lexicographically test whether the left string is > the right string
-	 *
-	 * @param Lhs String to compare against.
-	 * @param Rhs String to compare against.
-	 * @return true if the left string is lexicographically > the right string, otherwise false
-	 * @note case insensitive
-	 */
-	[[nodiscard]] FORCEINLINE friend bool operator>(const UE_STRING_CLASS& Lhs, const UE_STRING_CLASS& Rhs)
-	{
-		return FPlatformString::Stricmp(*Lhs, *Rhs) > 0;
-	}
-
-	/**
-	 * Lexicographically test whether the left string is > the right string
-	 *
-	 * @param Lhs String to compare against.
-	 * @param Rhs String to compare against.
-	 * @return true if the left string is lexicographically > the right string, otherwise false
-	 * @note case insensitive
-	 */
-	template <typename CharType>
-	[[nodiscard]] FORCEINLINE friend bool operator>(const UE_STRING_CLASS& Lhs, const CharType* Rhs)
-	{
-		return FPlatformString::Stricmp(*Lhs, Rhs) > 0;
-	}
-
-	/**
-	 * Lexicographically test whether the left string is > the right string
-	 *
-	 * @param Lhs String to compare against.
-	 * @param Rhs String to compare against.
-	 * @return true if the left string is lexicographically > the right string, otherwise false
-	 * @note case insensitive
-	 */
-	template <typename CharType>
-	[[nodiscard]] FORCEINLINE friend bool operator>(const CharType* Lhs, const UE_STRING_CLASS& Rhs)
-	{
-		return FPlatformString::Stricmp(Lhs, *Rhs) > 0;
+		return FPlatformString::Stricmp(**this, Rhs) < 0;
 	}
 
 	/**
@@ -1050,7 +913,7 @@ public:
 	 * @return true if the left string is lexicographically == the right string, otherwise false
 	 * @note case insensitive
 	 */
-	[[nodiscard]] FORCEINLINE bool operator==(const UE_STRING_CLASS& Rhs) const
+	[[nodiscard]] FORCEINLINE bool UEOpEquals(const UE_STRING_CLASS& Rhs) const
 	{
 		return Equals(Rhs, ESearchCase::IgnoreCase);
 	}
@@ -1064,66 +927,11 @@ public:
 	 * @note case insensitive
 	 */
 	template <typename CharType>
-	[[nodiscard]] FORCEINLINE bool operator==(const CharType* Rhs) const
+	[[nodiscard]] FORCEINLINE bool UEOpEquals(const CharType* Rhs) const
 	{
 		return FPlatformString::Stricmp(**this, Rhs) == 0;
 	}
 
-	/**
-	 * Lexicographically test whether the left string is == the right string
-	 *
-	 * @param Lhs String to compare against.
-	 * @param Rhs String to compare against.
-	 * @return true if the left string is lexicographically == the right string, otherwise false
-	 * @note case insensitive
-	 */
-	template <typename CharType>
-	[[nodiscard]] FORCEINLINE friend bool operator==(const CharType* Lhs, const UE_STRING_CLASS& Rhs)
-	{
-		return FPlatformString::Stricmp(Lhs, *Rhs) == 0;
-	}
-
-	/**
-	 * Lexicographically test whether the left string is != the right string
-	 *
-	 * @param Lhs String to compare against.
-	 * @param Rhs String to compare against.
-	 * @return true if the left string is lexicographically != the right string, otherwise false
-	 * @note case insensitive
-	 */
-	[[nodiscard]] FORCEINLINE bool operator!=(const UE_STRING_CLASS& Rhs) const
-	{
-		return !Equals(Rhs, ESearchCase::IgnoreCase);
-	}
-
-	/**
-	 * Lexicographically test whether the left string is != the right string
-	 *
-	 * @param Lhs String to compare against.
-	 * @param Rhs String to compare against.
-	 * @return true if the left string is lexicographically != the right string, otherwise false
-	 * @note case insensitive
-	 */
-	template <typename CharType>
-	[[nodiscard]] FORCEINLINE bool operator!=(const CharType* Rhs) const
-	{
-		return FPlatformString::Stricmp(**this, Rhs) != 0;
-	}
-
-	/**
-	 * Lexicographically test whether the left string is != the right string
-	 *
-	 * @param Lhs String to compare against.
-	 * @param Rhs String to compare against.
-	 * @return true if the left string is lexicographically != the right string, otherwise false
-	 * @note case insensitive
-	 */
-	template <typename CharType>
-	[[nodiscard]] FORCEINLINE friend bool operator!=(const CharType* Lhs, const UE_STRING_CLASS& Rhs)
-	{
-		return FPlatformString::Stricmp(Lhs, *Rhs) != 0;
-	}
-
 	/** Get the length of the string, excluding terminating character */
 	[[nodiscard]] FORCEINLINE int32 Len() const
 	{
diff --git a/Engine/Source/Runtime/Core/Public/Containers/UnrealStringIncludes.h.inl b/Engine/Source/Runtime/Core/Public/Containers/UnrealStringIncludes.h.inl
--- a/Engine/Source/Runtime/Core/Public/Containers/UnrealStringIncludes.h.inl
+++ b/Engine/Source/Runtime/Core/Public/Containers/UnrealStringIncludes.h.inl
@@ -20,6 +20,7 @@
 #include "CoreTypes.h"
 #include "Misc/VarArgs.h"
 #include "Misc/AssertionMacros.h"
+#include "Misc/UEOps.h"
 #include "HAL/UnrealMemory.h"
 #include "Templates/IsArithmetic.h"
 #include "Templates/IsArray.h"
diff --git a/Engine/Source/Runtime/Core/Public/Internationalization/LocalizedTextSourceTypes.h b/Engine/Source/Runtime/Core/Public/Internationalization/LocalizedTextSourceTypes.h
--- a/Engine/Source/Runtime/Core/Public/Internationalization/LocalizedTextSourceTypes.h
+++ b/Engine/Source/Runtime/Core/Public/Internationalization/LocalizedTextSourceTypes.h
@@ -4,6 +4,7 @@
 
 #include "CoreTypes.h"
 #include "Misc/EnumClassFlags.h"
+#include "Misc/UEOps.h"
 #include "Templates/RefCounting.h"
 #include "Templates/Requires.h"
 #include "Templates/SharedPointer.h"
@@ -144,6 +145,12 @@ public:
 		return DisplayStringPtr;
 	}
 
+	template <typename OtherObjectType>
+	[[nodiscard]] bool UEOpEquals(const TDisplayStringPtrBase<OtherObjectType>& Rhs) const
+	{
+		return DisplayStringPtr == Rhs.DisplayStringPtr;
+	}
+
 protected:
 	[[nodiscard]] FString& GetDisplayString() const
 	{
@@ -281,54 +288,6 @@ public:
 
 } // namespace UE::Text::Private
 
-template <typename ObjectTypeA, typename ObjectTypeB>
-[[nodiscard]] bool operator==(const UE::Text::Private::TDisplayStringRef<ObjectTypeA>& A, const UE::Text::Private::TDisplayStringRef<ObjectTypeB>& B)
-{
-	return &A.Get() == &B.Get();
-}
-
-template <typename ObjectTypeA, typename ObjectTypeB>
-[[nodiscard]] bool operator==(const UE::Text::Private::TDisplayStringRef<ObjectTypeA>& A, const UE::Text::Private::TDisplayStringPtr<ObjectTypeB>& B)
-{
-	return &A.Get() == B.Get();
-}
-
-template <typename ObjectTypeA, typename ObjectTypeB>
-[[nodiscard]] bool operator==(const UE::Text::Private::TDisplayStringPtr<ObjectTypeA>& A, const UE::Text::Private::TDisplayStringRef<ObjectTypeB>& B)
-{
-	return A.Get() == &B.Get();
-}
-
-template <typename ObjectTypeA, typename ObjectTypeB>
-[[nodiscard]] bool operator==(const UE::Text::Private::TDisplayStringPtr<ObjectTypeA>& A, const UE::Text::Private::TDisplayStringPtr<ObjectTypeB>& B)
-{
-	return A.Get() == B.Get();
-}
-
-template <typename ObjectTypeA, typename ObjectTypeB>
-[[nodiscard]] bool operator!=(const UE::Text::Private::TDisplayStringRef<ObjectTypeA>& A, const UE::Text::Private::TDisplayStringRef<ObjectTypeB>& B)
-{
-	return &A.Get() != &B.Get();
-}
-
-template <typename ObjectTypeA, typename ObjectTypeB>
-[[nodiscard]] bool operator!=(const UE::Text::Private::TDisplayStringRef<ObjectTypeA>& A, const UE::Text::Private::TDisplayStringPtr<ObjectTypeB>& B)
-{
-	return &A.Get() != B.Get();
-}
-
-template <typename ObjectTypeA, typename ObjectTypeB>
-[[nodiscard]] bool operator!=(const UE::Text::Private::TDisplayStringPtr<ObjectTypeA>& A, const UE::Text::Private::TDisplayStringRef<ObjectTypeB>& B)
-{
-	return A.Get() != &B.Get();
-}
-
-template <typename ObjectTypeA, typename ObjectTypeB>
-[[nodiscard]] bool operator!=(const UE::Text::Private::TDisplayStringPtr<ObjectTypeA>& A, const UE::Text::Private::TDisplayStringPtr<ObjectTypeB>& B)
-{
-	return A.Get() != B.Get();
-}
-
 template <typename ObjectType>
 [[nodiscard]] uint32 GetTypeHash(const UE::Text::Private::TDisplayStringRef<ObjectType>& A)
 {
diff --git a/Engine/Source/Runtime/Core/Public/Misc/NotNull.h b/Engine/Source/Runtime/Core/Public/Misc/NotNull.h
--- a/Engine/Source/Runtime/Core/Public/Misc/NotNull.h
+++ b/Engine/Source/Runtime/Core/Public/Misc/NotNull.h
@@ -4,6 +4,7 @@
 
 #include "CoreTypes.h"
 #include "Misc/IntrusiveUnsetOptionalState.h"
+#include "Misc/UEOps.h"
 #include "Templates/Requires.h"
 #include "Templates/TypeHash.h"
 #include "Traits/IsImplicitlyConstructible.h"
@@ -243,6 +244,22 @@
 		{
 			return GetTypeHashHelper(NotNull.Val);
 		}
+
+		// Allow comparison, if it's supported
+		template <typename U>
+		[[nodiscard]] auto UEOpEquals(const TNotNull<U>& Rhs) const -> decltype(Val == Rhs.Val)
+		{
+			return Val == Rhs.Val;
+		}
+		template <typename U>
+		[[nodiscard]] auto UEOpEquals(const U& Rhs) const -> decltype(Val == Rhs)
+		{
+			static_assert(!std::is_same_v<U, TYPE_OF_NULLPTR>, "Comparing a TNotNull to nullptr is illegal");
+			return Val == Rhs;
+		}
+
+		// Disallow comparison against nullptr
+		bool UEOpEquals(TYPE_OF_NULLPTR Rhs) const = delete;
 	};
 
 	// Allow serialization, if it's supported
@@ -286,46 +303,6 @@
 		return (TNotNull<T>&&)NotNull;
 	}
 
-	// Allow comparison, if it's supported
-	template <typename T, typename U>
-	[[nodiscard]] auto operator==(const TNotNull<T>& Lhs, const TNotNull<U>& Rhs) -> decltype((T)Lhs == (U)Rhs)
-	{
-		return (T)Lhs == (U)Rhs;
-	}
-	template <typename T, typename U>
-	[[nodiscard]] auto operator==(const TNotNull<T>& Lhs, const U& Rhs) -> decltype((T)Lhs == Rhs)
-	{
-		static_assert(!std::is_same_v<U, TYPE_OF_NULLPTR>, "Comparing a TNotNull to nullptr is illegal");
-		return (T)Lhs == Rhs;
-	}
-
-	// Disallow comparison against nullptr
-	template <typename T>
-	bool operator==(const TNotNull<T>& Lhs, TYPE_OF_NULLPTR Rhs) = delete;
-
-#if !PLATFORM_COMPILER_HAS_GENERATED_COMPARISON_OPERATORS
-	template <typename T, typename U>
-	[[nodiscard]] auto operator!=(const TNotNull<T>& Lhs, const TNotNull<U>& Rhs) -> decltype(!(Lhs == Rhs))
-	{
-		return !(Lhs == Rhs);
-	}
-	template <typename T, typename U>
-	[[nodiscard]] auto operator!=(const TNotNull<T>& Lhs, const U& Rhs) -> decltype(!(Lhs == Rhs))
-	{
-		return !(Lhs == Rhs);
-	}
-	template <typename T, typename U>
-	[[nodiscard]] auto operator==(const U& Lhs, const TNotNull<T>& Rhs) -> decltype(Rhs == Lhs)
-	{
-		return Rhs == Lhs;
-	}
-	template <typename T, typename U>
-	[[nodiscard]] auto operator!=(const U& Lhs, const TNotNull<T>& Rhs) -> decltype(!(Rhs == Lhs))
-	{
-		return !(Rhs == Lhs);
-	}
-#endif // !PLATFORM_COMPILER_HAS_GENERATED_COMPARISON_OPERATORS
-
 #undef UE_NOTNULL_FUNCTION_NON_NULL_RETURN_END
 #undef UE_NOTNULL_FUNCTION_NON_NULL_RETURN_START
 
diff --git a/Engine/Source/Runtime/Core/Public/Misc/StringOutputDevice.h b/Engine/Source/Runtime/Core/Public/Misc/StringOutputDevice.h
--- a/Engine/Source/Runtime/Core/Public/Misc/StringOutputDevice.h
+++ b/Engine/Source/Runtime/Core/Public/Misc/StringOutputDevice.h
@@ -44,6 +44,16 @@ public:
 	{
 		return FString::operator+=(Other);
 	}
+
+	bool UEOpEquals(const FString& Rhs) const
+	{
+		return FString::UEOpEquals(Rhs);
+	}
+
+	bool UEOpLessThan(const FString& Rhs) const
+	{
+		return FString::UEOpLessThan(Rhs);
+	}
 };
 
 template <>
diff --git a/Engine/Source/Runtime/Core/Public/Templates/SharedPointer.h b/Engine/Source/Runtime/Core/Public/Templates/SharedPointer.h
--- a/Engine/Source/Runtime/Core/Public/Templates/SharedPointer.h
+++ b/Engine/Source/Runtime/Core/Public/Templates/SharedPointer.h
@@ -10,6 +10,7 @@
 #include "Containers/Array.h"
 #include "Containers/Map.h"
 #include "CoreGlobals.h"
+#include "Misc/UEOps.h"
 
 
 /**
@@ -538,7 +539,7 @@ public:
 		: Object(nullptr)
 	{
 	}
-	bool operator==(FIntrusiveUnsetOptionalState) const
+	bool UEOpEquals(FIntrusiveUnsetOptionalState) const
 	{
 		return !IsValid();
 	}
@@ -546,6 +547,32 @@ public:
 	// End - intrusive TOptional<TSharedRef> state //
 	/////////////////////////////////////////////////
 
+	template <typename OtherType>
+	[[nodiscard]] bool UEOpEquals(const TSharedRef<OtherType, Mode>& Rhs) const
+	{
+		return this->Object == &Rhs.Get();
+	}
+
+	template <typename OtherType>
+	[[nodiscard]] bool UEOpEquals(const TSharedPtr<OtherType, Mode>& Rhs) const
+	{
+		// This comparison against null is maintained as existing behavior, but isn't consistent with TWeakPtr comparison.
+		OtherType* RhsPtr = Rhs.Get();
+		return RhsPtr && RhsPtr == this->Object;
+	}
+
+	template <typename OtherType>
+	[[nodiscard]] bool UEOpEquals(const TWeakPtr<OtherType, Mode>& Rhs) const
+	{
+		return this->Object == Rhs.Pin().Get();
+	}
+
+	template <typename OtherType>
+	[[nodiscard]] bool UEOpLessThan(const TSharedRef<OtherType, Mode>& Rhs) const
+	{
+		return this->Object < &Rhs.Get();
+	}
+
 private:
 	template<class OtherType>
 	void Init(OtherType* InObject)
@@ -1128,15 +1155,38 @@ public:
 		: Object((ObjectType*)-1)
 	{
 	}
-	bool operator==(FIntrusiveUnsetOptionalState) const
+	bool UEOpEquals(FIntrusiveUnsetOptionalState) const
 	{
 		return Object == (ObjectType*)-1;
 	}
 	/////////////////////////////////////////////////
 	// End - intrusive TOptional<TSharedPtr> state //
 	/////////////////////////////////////////////////
-private:
 
+	template <typename OtherType>
+	[[nodiscard]] inline bool UEOpEquals(const TSharedPtr<OtherType, Mode>& Rhs) const
+	{
+		return this->Object == Rhs.Get();
+	}
+
+	template <typename OtherType>
+	[[nodiscard]] UE_REWRITE bool UEOpEquals(const TWeakPtr<OtherType, Mode>& Rhs) const
+	{
+		return this->Object == Rhs.Pin().Get();
+	}
+
+	[[nodiscard]] UE_REWRITE bool UEOpEquals(TYPE_OF_NULLPTR) const
+	{
+		return !this->Object;
+	}
+
+	template <typename OtherType>
+	[[nodiscard]] inline bool UEOpLessThan(const TSharedPtr<OtherType, Mode>& Rhs) const
+	{
+		return this->Object < Rhs.Get();
+	}
+
+private:
 	/**
 	 * Constructs a shared pointer from a weak pointer, allowing you to access the object (if it
 	 * hasn't expired yet.)  Remember, if there are no more shared references to the object, the
@@ -1519,15 +1569,32 @@ public:
 		: Object((ObjectType*)-1)
 	{
 	}
-	bool operator==(FIntrusiveUnsetOptionalState) const
+	bool UEOpEquals(FIntrusiveUnsetOptionalState) const
 	{
 		return Object == (ObjectType*)-1;
 	}
 	///////////////////////////////////////////////
 	// End - intrusive TOptional<TWeakPtr> state //
 	///////////////////////////////////////////////
+
+	template <typename OtherType>
+	[[nodiscard]] inline bool UEOpEquals(const TWeakPtr<OtherType, Mode>& Rhs) const
+	{
+		return this->Pin().Get() == Rhs.Pin().Get();
+	}
+
+	[[nodiscard]] UE_REWRITE bool UEOpEquals(TYPE_OF_NULLPTR) const
+	{
+		return !this->IsValid();
+	}
+
+	template <typename OtherType>
+	[[nodiscard]] inline bool UEOpLessThan(const TWeakPtr<OtherType, Mode>& Rhs) const
+	{
+		return this->Pin().Get() < Rhs.Pin().Get();
+	}
+
 private:
-	
 	// We declare ourselves as a friend (templated using OtherType) so we can access members as needed
     template< class OtherType, ESPMode OtherMode > friend class TWeakPtr;
 
@@ -1535,7 +1602,6 @@ private:
     template< class OtherType, ESPMode OtherMode > friend class TSharedPtr;
 
 private:
-
 	/** The object we have a weak reference to.  Can be nullptr.  Also, it's important to note that because
 	    this is a weak reference, the object this pointer points to may have already been destroyed. */
 	ObjectType* Object;
@@ -1811,462 +1877,6 @@ private:
 };
 
 
-/**
- * Global equality operator for TSharedRef
- *
- * @return  True if the two shared references are equal
- */
-template< class ObjectTypeA, class ObjectTypeB, ESPMode Mode >
-[[nodiscard]] FORCEINLINE bool operator==( TSharedRef< ObjectTypeA, Mode > const& InSharedRefA, TSharedRef< ObjectTypeB, Mode > const& InSharedRefB )
-{
-	return &( InSharedRefA.Get() ) == &( InSharedRefB.Get() );
-}
-
-
-/**
- * Global less than operator for TSharedRef
- *
- * @return  True if InSharedRefA is less than InSharedRefB
- */
-template< class ObjectTypeA, class ObjectTypeB, ESPMode Mode >
-[[nodiscard]] FORCEINLINE bool operator<( TSharedRef< ObjectTypeA, Mode > const& InSharedRefA, TSharedRef< ObjectTypeB, Mode > const& InSharedRefB )
-{
-	return &( InSharedRefA.Get() ) < &( InSharedRefB.Get() );
-}
-
-
-/**
- * Global greater than operator for TSharedRef
- *
- * @return  True if InSharedRefA is greater than InSharedRefB
- */
-template< class ObjectTypeA, class ObjectTypeB, ESPMode Mode >
-[[nodiscard]] FORCEINLINE bool operator>( TSharedRef< ObjectTypeA, Mode > const& InSharedRefA, TSharedRef< ObjectTypeB, Mode > const& InSharedRefB )
-{
-	return &( InSharedRefA.Get() ) > &( InSharedRefB.Get() );
-}
-
-
-/**
- * Global less than or equal to operator for TSharedRef
- *
- * @return  True if InSharedRefA is less than or equal to InSharedRefB
- */
-template< class ObjectTypeA, class ObjectTypeB, ESPMode Mode >
-[[nodiscard]] FORCEINLINE bool operator<=( TSharedRef< ObjectTypeA, Mode > const& InSharedRefA, TSharedRef< ObjectTypeB, Mode > const& InSharedRefB )
-{
-	return &( InSharedRefA.Get() ) <= &( InSharedRefB.Get() );
-}
-
-
-/**
- * Global greater than or equal to operator for TSharedRef
- *
- * @return  True if InSharedRefA is greater than or equal to InSharedRefB
- */
-template< class ObjectTypeA, class ObjectTypeB, ESPMode Mode >
-[[nodiscard]] FORCEINLINE bool operator>=( TSharedRef< ObjectTypeA, Mode > const& InSharedRefA, TSharedRef< ObjectTypeB, Mode > const& InSharedRefB )
-{
-	return &( InSharedRefA.Get() ) >= &( InSharedRefB.Get() );
-}
-
-
-/**
- * Global inequality operator for TSharedRef
- *
- * @return  True if the two shared references are not equal
- */
-template< class ObjectTypeA, class ObjectTypeB, ESPMode Mode >
-[[nodiscard]] FORCEINLINE bool operator!=( TSharedRef< ObjectTypeA, Mode > const& InSharedRefA, TSharedRef< ObjectTypeB, Mode > const& InSharedRefB )
-{
-	return &( InSharedRefA.Get() ) != &( InSharedRefB.Get() );
-}
-
-
-/**
- * Global equality operator for TSharedPtr
- *
- * @return  True if the two shared pointers are equal
- */
-template< class ObjectTypeA, class ObjectTypeB, ESPMode Mode >
-[[nodiscard]] FORCEINLINE bool operator==( TSharedPtr< ObjectTypeA, Mode > const& InSharedPtrA, TSharedPtr< ObjectTypeB, Mode > const& InSharedPtrB )
-{
-	return InSharedPtrA.Get() == InSharedPtrB.Get();
-}
-
-
-/**
- * Global less than operator for TSharedPtr
- *
- * @return  True if InSharedPtrA is less than InSharedPtrB
- */
-template< class ObjectTypeA, class ObjectTypeB, ESPMode Mode >
-[[nodiscard]] FORCEINLINE bool operator<( TSharedPtr< ObjectTypeA, Mode > const& InSharedPtrA, TSharedPtr< ObjectTypeB, Mode > const& InSharedPtrB )
-{
-	return InSharedPtrA.Get() < InSharedPtrB.Get();
-}
-
-
-/**
- * Global greater than operator for TSharedPtr
- *
- * @return  True if InSharedPtrA is greater than InSharedPtrB
- */
-template< class ObjectTypeA, class ObjectTypeB, ESPMode Mode >
-[[nodiscard]] FORCEINLINE bool operator>( TSharedPtr< ObjectTypeA, Mode > const& InSharedPtrA, TSharedPtr< ObjectTypeB, Mode > const& InSharedPtrB )
-{
-	return InSharedPtrA.Get() > InSharedPtrB.Get();
-}
-
-
-/**
- * Global less than or equal to operator for TSharedPtr
- *
- * @return  True if InSharedPtrA is less than or equal to InSharedPtrB
- */
-template< class ObjectTypeA, class ObjectTypeB, ESPMode Mode >
-[[nodiscard]] FORCEINLINE bool operator<=( TSharedPtr< ObjectTypeA, Mode > const& InSharedPtrA, TSharedPtr< ObjectTypeB, Mode > const& InSharedPtrB )
-{
-	return InSharedPtrA.Get() <= InSharedPtrB.Get();
-}
-
-
-/**
- * Global greater than or equal to operator for TSharedPtr
- *
- * @return  True if InSharedPtrA is greater than or equal to InSharedPtrB
- */
-template< class ObjectTypeA, class ObjectTypeB, ESPMode Mode >
-[[nodiscard]] FORCEINLINE bool operator>=( TSharedPtr< ObjectTypeA, Mode > const& InSharedPtrA, TSharedPtr< ObjectTypeB, Mode > const& InSharedPtrB )
-{
-	return InSharedPtrA.Get() >= InSharedPtrB.Get();
-}
-
-
-/**
- * Global equality operator for TSharedPtr
- *
- * @return  True if the shared pointer is null
- */
-template< class ObjectTypeA, ESPMode Mode >
-[[nodiscard]] FORCEINLINE bool operator==( TSharedPtr< ObjectTypeA, Mode > const& InSharedPtrA, TYPE_OF_NULLPTR )
-{
-	return !InSharedPtrA.IsValid();
-}
-
-
-/**
- * Global equality operator for TSharedPtr
- *
- * @return  True if the shared pointer is null
- */
-template< class ObjectTypeB, ESPMode Mode >
-[[nodiscard]] FORCEINLINE bool operator==( TYPE_OF_NULLPTR, TSharedPtr< ObjectTypeB, Mode > const& InSharedPtrB )
-{
-	return !InSharedPtrB.IsValid();
-}
-
-
-/**
- * Global inequality operator for TSharedPtr
- *
- * @return  True if the two shared pointers are not equal
- */
-template< class ObjectTypeA, class ObjectTypeB, ESPMode Mode >
-[[nodiscard]] FORCEINLINE bool operator!=( TSharedPtr< ObjectTypeA, Mode > const& InSharedPtrA, TSharedPtr< ObjectTypeB, Mode > const& InSharedPtrB )
-{
-	return InSharedPtrA.Get() != InSharedPtrB.Get();
-}
-
-
-/**
- * Global inequality operator for TSharedPtr
- *
- * @return  True if the shared pointer is not null
- */
-template< class ObjectTypeA, ESPMode Mode >
-[[nodiscard]] FORCEINLINE bool operator!=( TSharedPtr< ObjectTypeA, Mode > const& InSharedPtrA, TYPE_OF_NULLPTR )
-{
-	return InSharedPtrA.IsValid();
-}
-
-
-/**
- * Global inequality operator for TSharedPtr
- *
- * @return  True if the shared pointer is not null
- */
-template< class ObjectTypeB, ESPMode Mode >
-[[nodiscard]] FORCEINLINE bool operator!=( TYPE_OF_NULLPTR, TSharedPtr< ObjectTypeB, Mode > const& InSharedPtrB )
-{
-	return InSharedPtrB.IsValid();
-}
-
-
-/**
- * Tests to see if a TSharedRef is "equal" to a TSharedPtr (both are valid and refer to the same object)
- *
- * @return  True if the shared reference and shared pointer are "equal"
- */
-template< class ObjectTypeA, class ObjectTypeB, ESPMode Mode >
-[[nodiscard]] FORCEINLINE bool operator==( TSharedRef< ObjectTypeA, Mode > const& InSharedRef, TSharedPtr< ObjectTypeB, Mode > const& InSharedPtr )
-{
-	return InSharedPtr.IsValid() && InSharedPtr.Get() == &( InSharedRef.Get() );
-}
-
-
-/**
- * Tests to see if a TSharedRef is not "equal" to a TSharedPtr (shared pointer is invalid, or both refer to different objects)
- *
- * @return  True if the shared reference and shared pointer are not "equal"
- */
-template< class ObjectTypeA, class ObjectTypeB, ESPMode Mode >
-[[nodiscard]] FORCEINLINE bool operator!=( TSharedRef< ObjectTypeA, Mode > const& InSharedRef, TSharedPtr< ObjectTypeB, Mode > const& InSharedPtr )
-{
-	return !InSharedPtr.IsValid() || ( InSharedPtr.Get() != &( InSharedRef.Get() ) );
-}
-
-
-/**
- * Tests to see if a TSharedRef is "equal" to a TSharedPtr (both are valid and refer to the same object) (reverse)
- *
- * @return  True if the shared reference and shared pointer are "equal"
- */
-template< class ObjectTypeA, class ObjectTypeB, ESPMode Mode >
-[[nodiscard]] FORCEINLINE bool operator==( TSharedPtr< ObjectTypeB, Mode > const& InSharedPtr, TSharedRef< ObjectTypeA, Mode > const& InSharedRef )
-{
-	return InSharedRef == InSharedPtr;
-}
-
-
-/**
- * Tests to see if a TSharedRef is not "equal" to a TSharedPtr (shared pointer is invalid, or both refer to different objects) (reverse)
- *
- * @return  True if the shared reference and shared pointer are not "equal"
- */
-template< class ObjectTypeA, class ObjectTypeB, ESPMode Mode >
-[[nodiscard]] FORCEINLINE bool operator!=( TSharedPtr< ObjectTypeB, Mode > const& InSharedPtr, TSharedRef< ObjectTypeA, Mode > const& InSharedRef )
-{
-	return InSharedRef != InSharedPtr;
-}
-
-
-/**
- * Global equality operator for TWeakPtr
- *
- * @return  True if the two weak pointers are equal
- */
-template< class ObjectTypeA, class ObjectTypeB, ESPMode Mode >
-[[nodiscard]] FORCEINLINE bool operator==( TWeakPtr< ObjectTypeA, Mode > const& InWeakPtrA, TWeakPtr< ObjectTypeB, Mode > const& InWeakPtrB )
-{
-	return InWeakPtrA.Pin().Get() == InWeakPtrB.Pin().Get();
-}
-
-
-/**
- * Global less than operator for TWeakPtr
- *
- * @return  True if InWeakPtrA is less than InWeakPtrB
- */
-template< class ObjectTypeA, class ObjectTypeB, ESPMode Mode >
-[[nodiscard]] FORCEINLINE bool operator<( TWeakPtr< ObjectTypeA, Mode > const& InWeakPtrA, TWeakPtr< ObjectTypeB, Mode > const& InWeakPtrB )
-{
-	return InWeakPtrA.Pin().Get() < InWeakPtrB.Pin().Get();
-}
-
-
-/**
- * Global greater than operator for TWeakPtr
- *
- * @return  True if InWeakPtrA is greater than InWeakPtrB
- */
-template< class ObjectTypeA, class ObjectTypeB, ESPMode Mode >
-[[nodiscard]] FORCEINLINE bool operator>( TWeakPtr< ObjectTypeA, Mode > const& InWeakPtrA, TWeakPtr< ObjectTypeB, Mode > const& InWeakPtrB )
-{
-	return InWeakPtrA.Pin().Get() > InWeakPtrB.Pin().Get();
-}
-
-
-/**
- * Global less than or equal to operator for TWeakPtr
- *
- * @return  True if InWeakPtrA is less than or equal to InWeakPtrB
- */
-template< class ObjectTypeA, class ObjectTypeB, ESPMode Mode >
-[[nodiscard]] FORCEINLINE bool operator<=( TWeakPtr< ObjectTypeA, Mode > const& InWeakPtrA, TWeakPtr< ObjectTypeB, Mode > const& InWeakPtrB )
-{
-	return InWeakPtrA.Pin().Get() <= InWeakPtrB.Pin().Get();
-}
-
-
-/**
- * Global greater than or equal to operator for TWeakPtr
- *
- * @return  True if InWeakPtrA is greater than or equal to InWeakPtrB
- */
-template< class ObjectTypeA, class ObjectTypeB, ESPMode Mode >
-[[nodiscard]] FORCEINLINE bool operator>=( TWeakPtr< ObjectTypeA, Mode > const& InWeakPtrA, TWeakPtr< ObjectTypeB, Mode > const& InWeakPtrB )
-{
-	return InWeakPtrA.Pin().Get() >= InWeakPtrB.Pin().Get();
-}
-
-
-/**
- * Global equality operator for TWeakPtr
- *
- * @return  True if the weak pointer and the shared ref are equal
- */
-template< class ObjectTypeA, class ObjectTypeB, ESPMode Mode >
-[[nodiscard]] FORCEINLINE bool operator==( TWeakPtr< ObjectTypeA, Mode > const& InWeakPtrA, TSharedRef< ObjectTypeB, Mode > const& InSharedRefB )
-{
-	return InWeakPtrA.Pin().Get() == &InSharedRefB.Get();
-}
-
-
-/**
- * Global equality operator for TWeakPtr
- *
- * @return  True if the weak pointer and the shared ptr are equal
- */
-template< class ObjectTypeA, class ObjectTypeB, ESPMode Mode >
-[[nodiscard]] FORCEINLINE bool operator==( TWeakPtr< ObjectTypeA, Mode > const& InWeakPtrA, TSharedPtr< ObjectTypeB, Mode > const& InSharedPtrB )
-{
-	return InWeakPtrA.Pin().Get() == InSharedPtrB.Get();
-}
-
-
-/**
- * Global equality operator for TWeakPtr
- *
- * @return  True if the weak pointer and the shared ref are equal
- */
-template< class ObjectTypeA, class ObjectTypeB, ESPMode Mode >
-[[nodiscard]] FORCEINLINE bool operator==( TSharedRef< ObjectTypeA, Mode > const& InSharedRefA, TWeakPtr< ObjectTypeB, Mode > const& InWeakPtrB )
-{
-	return &InSharedRefA.Get() == InWeakPtrB.Pin().Get();
-}
-
-
-/**
- * Global equality operator for TWeakPtr
- *
- * @return  True if the weak pointer and the shared ptr are equal
- */
-template< class ObjectTypeA, class ObjectTypeB, ESPMode Mode >
-[[nodiscard]] FORCEINLINE bool operator==( TSharedPtr< ObjectTypeA, Mode > const& InSharedPtrA, TWeakPtr< ObjectTypeB, Mode > const& InWeakPtrB )
-{
-	return InSharedPtrA.Get() == InWeakPtrB.Pin().Get();
-}
-
-
-/**
- * Global equality operator for TWeakPtr
- *
- * @return  True if the weak pointer is null
- */
-template< class ObjectTypeA, ESPMode Mode >
-[[nodiscard]] FORCEINLINE bool operator==( TWeakPtr< ObjectTypeA, Mode > const& InWeakPtrA, TYPE_OF_NULLPTR )
-{
-	return !InWeakPtrA.IsValid();
-}
-
-
-/**
- * Global equality operator for TWeakPtr
- *
- * @return  True if the weak pointer is null
- */
-template< class ObjectTypeB, ESPMode Mode >
-[[nodiscard]] FORCEINLINE bool operator==( TYPE_OF_NULLPTR, TWeakPtr< ObjectTypeB, Mode > const& InWeakPtrB )
-{
-	return !InWeakPtrB.IsValid();
-}
-
-
-/**
- * Global inequality operator for TWeakPtr
- *
- * @return  True if the two weak pointers are not equal
- */
-template< class ObjectTypeA, class ObjectTypeB, ESPMode Mode >
-[[nodiscard]] FORCEINLINE bool operator!=( TWeakPtr< ObjectTypeA, Mode > const& InWeakPtrA, TWeakPtr< ObjectTypeB, Mode > const& InWeakPtrB )
-{
-	return InWeakPtrA.Pin().Get() != InWeakPtrB.Pin().Get();
-}
-
-
-/**
- * Global equality operator for TWeakPtr
- *
- * @return  True if the weak pointer and the shared ref are not equal
- */
-template< class ObjectTypeA, class ObjectTypeB, ESPMode Mode >
-[[nodiscard]] FORCEINLINE bool operator!=( TWeakPtr< ObjectTypeA, Mode > const& InWeakPtrA, TSharedRef< ObjectTypeB, Mode > const& InSharedRefB )
-{
-	return InWeakPtrA.Pin().Get() != &InSharedRefB.Get();
-}
-
-
-/**
- * Global equality operator for TWeakPtr
- *
- * @return  True if the weak pointer and the shared ptr are not equal
- */
-template< class ObjectTypeA, class ObjectTypeB, ESPMode Mode >
-[[nodiscard]] FORCEINLINE bool operator!=( TWeakPtr< ObjectTypeA, Mode > const& InWeakPtrA, TSharedPtr< ObjectTypeB, Mode > const& InSharedPtrB )
-{
-	return InWeakPtrA.Pin().Get() != InSharedPtrB.Get();
-}
-
-
-/**
- * Global equality operator for TWeakPtr
- *
- * @return  True if the weak pointer and the shared ref are not equal
- */
-template< class ObjectTypeA, class ObjectTypeB, ESPMode Mode >
-[[nodiscard]] FORCEINLINE bool operator!=( TSharedRef< ObjectTypeA, Mode > const& InSharedRefA, TWeakPtr< ObjectTypeB, Mode > const& InWeakPtrB )
-{
-	return &InSharedRefA.Get() != InWeakPtrB.Pin().Get();
-}
-
-
-/**
- * Global equality operator for TWeakPtr
- *
- * @return  True if the weak pointer and the shared ptr are not equal
- */
-template< class ObjectTypeA, class ObjectTypeB, ESPMode Mode >
-[[nodiscard]] FORCEINLINE bool operator!=( TSharedPtr< ObjectTypeA, Mode > const& InSharedPtrA, TWeakPtr< ObjectTypeB, Mode > const& InWeakPtrB )
-{
-	return InSharedPtrA.Get() != InWeakPtrB.Pin().Get();
-}
-
-
-/**
- * Global inequality operator for TWeakPtr
- *
- * @return  True if the weak pointer is not null
- */
-template< class ObjectTypeA, ESPMode Mode >
-[[nodiscard]] FORCEINLINE bool operator!=( TWeakPtr< ObjectTypeA, Mode > const& InWeakPtrA, TYPE_OF_NULLPTR )
-{
-	return InWeakPtrA.IsValid();
-}
-
-
-/**
- * Global inequality operator for TWeakPtr
- *
- * @return  True if the weak pointer is not null
- */
-template< class ObjectTypeB, ESPMode Mode >
-[[nodiscard]] FORCEINLINE bool operator!=( TYPE_OF_NULLPTR, TWeakPtr< ObjectTypeB, Mode > const& InWeakPtrB )
-{
-	return InWeakPtrB.IsValid();
-}
-
-
 /**
  * Casts a shared pointer of one type to another type. (static_cast)  Useful for down-casting.
  *
diff --git a/Engine/Source/Runtime/Core/Public/UObject/NameTypes.h b/Engine/Source/Runtime/Core/Public/UObject/NameTypes.h
--- a/Engine/Source/Runtime/Core/Public/UObject/NameTypes.h
+++ b/Engine/Source/Runtime/Core/Public/UObject/NameTypes.h
@@ -16,6 +16,7 @@
 #include "Serialization/MemoryLayout.h"
 #include "Misc/IntrusiveUnsetOptionalState.h"
 #include "Misc/StringBuilder.h"
+#include "Misc/UEOps.h"
 #include "Trace/Trace.h"
 
 /*----------------------------------------------------------------------------
@@ -80,15 +81,13 @@ struct FNameEntryId
 	bool FastLess(FNameEntryId Rhs) const { return CompareFast(Rhs) < 0; }
 
 	/** Fast non-alphabetical order that is only stable during this process' lifetime */
-	bool operator<(FNameEntryId Rhs) const { return Value < Rhs.Value; }
+	bool UEOpLessThan(FNameEntryId Rhs) const { return Value < Rhs.Value; }
 
 	/** Fast non-alphabetical order that is only stable during this process' lifetime */
-	bool operator>(FNameEntryId Rhs) const { return Rhs.Value < Value; }
-	bool operator==(FNameEntryId Rhs) const { return Value == Rhs.Value; }
-	bool operator!=(FNameEntryId Rhs) const { return Value != Rhs.Value; }
+	bool UEOpEquals(FNameEntryId Rhs) const { return Value == Rhs.Value; }
 	
 	/** Comparison against special type for checking TOptional<FName>:IsSet */
-	bool operator==(FIntrusiveUnsetOptionalState) const
+	bool UEOpEquals(FIntrusiveUnsetOptionalState) const
 	{
 		return Value == ~0u;
 	}
@@ -112,9 +111,9 @@ struct FNameEntryId
 		return Ename == NAME_None ? FNameEntryId() : FromValidEName(Ename);
 	}
 
-	friend inline bool operator==(FNameEntryId Id, EName Ename)
+	inline bool UEOpEquals(EName Ename) const
 	{
-		return Ename == NAME_None ? !Id : Id == FromValidENamePostInit(Ename);
+		return Ename == NAME_None ? Value == 0 : Value == FromValidENamePostInit(Ename).Value;
 	}
 
 private:
@@ -124,9 +123,6 @@ private:
 	CORE_API static FNameEntryId FromValidENamePostInit(EName Ename);
 
 public:
-	friend inline bool operator==(EName Ename, FNameEntryId Id) { return Id == Ename; }
-	friend inline bool operator!=(EName Ename, FNameEntryId Id) { return !(Id == Ename); }
-	friend inline bool operator!=(FNameEntryId Id, EName Ename) { return !(Id == Ename); }
 	friend CORE_API uint32 GetTypeHash(FNameEntryId Id);
 
 	/** Serialize as process specific unstable int */
@@ -397,7 +393,7 @@ struct FMinimalName
 
 	FORCEINLINE explicit FMinimalName(const FName& Name);
 	FORCEINLINE bool IsNone() const;
-	FORCEINLINE bool operator<(FMinimalName Rhs) const;
+	FORCEINLINE bool UEOpLessThan(FMinimalName Rhs) const;
 	
 private:
 	/** Index into the Names array (used to find String portion of the string/number pair) */
@@ -407,29 +403,26 @@ private:
 	int32			Number = NAME_NO_NUMBER_INTERNAL;
 #endif // UE_FNAME_OUTLINE_NUMBE
 
+public:
 #if UE_FNAME_OUTLINE_NUMBER
-	friend FORCEINLINE bool operator==(FMinimalName Lhs, FMinimalName Rhs)
+	FORCEINLINE bool UEOpEquals(FMinimalName Rhs) const
 	{
-		return Lhs.Index == Rhs.Index;
+		return Index == Rhs.Index;
 	}
 	friend FORCEINLINE uint32 GetTypeHash(FMinimalName Name)
 	{
 		return GetTypeHash(Name.Index);
 	}
 #else
-	friend FORCEINLINE bool operator==(FMinimalName Lhs, FMinimalName Rhs)
+	FORCEINLINE bool UEOpEquals(FMinimalName Rhs) const
 	{
-		return Lhs.Index == Rhs.Index && Lhs.Number == Rhs.Number;
+		return Index == Rhs.Index && Number == Rhs.Number;
 	}
 	friend FORCEINLINE uint32 GetTypeHash(FMinimalName Name)
 	{
 		return GetTypeHash(Name.Index) + Name.Number;
 	}
 #endif
-	friend FORCEINLINE bool operator!=(FMinimalName Lhs, FMinimalName Rhs)
-	{
-		return !operator==(Lhs, Rhs);
-	}
 };
 
 /**
@@ -453,7 +446,7 @@ struct FScriptName
 
 	FORCEINLINE explicit FScriptName(const FName& Name);
 	FORCEINLINE bool IsNone() const;
-	inline bool operator==(EName Name) { return *this == FScriptName(Name); }
+	inline bool UEOpEquals(EName Name) const { return UEOpEquals(FScriptName(Name)); }
 
 	CORE_API FString ToString() const;
 	CORE_API FUtf8String ToUtf8String() const;
@@ -471,30 +464,26 @@ private:
 	uint32			Number = NAME_NO_NUMBER_INTERNAL;
 #endif // UE_FNAME_OUTLINE_NUMBER
 
-
+public:
 #if UE_FNAME_OUTLINE_NUMBER
-	friend FORCEINLINE bool operator==(FScriptName Lhs, FScriptName Rhs)
+	FORCEINLINE bool UEOpEquals(FScriptName Rhs) const
 	{
-		return Lhs.ComparisonIndex == Rhs.ComparisonIndex;
+		return ComparisonIndex == Rhs.ComparisonIndex;
 	}
 	friend FORCEINLINE uint32 GetTypeHash(FScriptName Name)
 	{
 		return GetTypeHash(Name.ComparisonIndex);
 	}
 #else
-	friend FORCEINLINE bool operator==(FScriptName Lhs, FScriptName Rhs)
+	FORCEINLINE bool UEOpEquals(FScriptName Rhs) const
 	{
-		return Lhs.ComparisonIndex == Rhs.ComparisonIndex && Lhs.Number == Rhs.Number;
+		return ComparisonIndex == Rhs.ComparisonIndex && Number == Rhs.Number;
 	}
 	friend FORCEINLINE uint32 GetTypeHash(FScriptName Name)
 	{
 		return GetTypeHash(Name.ComparisonIndex) + Name.Number;
 	}
 #endif
-	friend FORCEINLINE bool operator!=(FScriptName Lhs, FScriptName Rhs)
-	{
-		return !(Lhs == Rhs);
-	}
 };
 
 struct FMemoryImageName;
@@ -522,7 +511,7 @@ struct FMemoryImageName
 	FMemoryImageName(EName Name);
 	FORCEINLINE FMemoryImageName(const FName& Name);
 
-	inline bool operator==(EName Name) const { return *this == FMemoryImageName(Name); }
+	inline bool UEOpEquals(EName Name) const { return UEOpEquals(FMemoryImageName(Name)); }
 
 	FORCEINLINE bool IsNone() const;
 	CORE_API FString ToString() const;
@@ -542,29 +531,26 @@ private:
 	FNameEntryId	DisplayIndex;
 #endif
 
+public:
 #if UE_FNAME_OUTLINE_NUMBER
-	friend FORCEINLINE bool operator==(FMemoryImageName Lhs, FMemoryImageName Rhs)
+	FORCEINLINE bool UEOpEquals(FMemoryImageName Rhs) const
 	{
-		return Lhs.ComparisonIndex == Rhs.ComparisonIndex;
+		return ComparisonIndex == Rhs.ComparisonIndex;
 	}
 	friend FORCEINLINE uint32 GetTypeHash(FMemoryImageName Name)
 	{
 		return GetTypeHash(Name.ComparisonIndex);
 	}
 #else
-	friend FORCEINLINE bool operator==(FMemoryImageName Lhs, FMemoryImageName Rhs)
+	FORCEINLINE bool UEOpEquals(FMemoryImageName Rhs) const
 	{
-		return Lhs.ComparisonIndex == Rhs.ComparisonIndex && Lhs.Number == Rhs.Number;
+		return ComparisonIndex == Rhs.ComparisonIndex && Number == Rhs.Number;
 	}
 	friend FORCEINLINE uint32 GetTypeHash(FMemoryImageName Name)
 	{
 		return GetTypeHash(Name.ComparisonIndex) + Name.Number;
 	}
 #endif
-	friend FORCEINLINE bool operator!=(FMemoryImageName Lhs, FMemoryImageName Rhs)
-	{
-		return !(Lhs == Rhs);
-	}
 };
 
 /**
@@ -759,18 +745,13 @@ public:
 	 */
 	[[nodiscard]] FORCEINLINE bool IsEqual(const FName& Other, const ENameCase CompareMethod = ENameCase::IgnoreCase, const bool bCompareNumber = true) const;
 
-	[[nodiscard]] FORCEINLINE bool operator==(FName Other) const
+	[[nodiscard]] FORCEINLINE bool UEOpEquals(FName Other) const
 	{
 		return ToUnstableInt() == Other.ToUnstableInt();
 	}
 
-	[[nodiscard]] FORCEINLINE bool operator!=(FName Other) const
-	{
-		return !(*this == Other);
-	}
-
 	/** Special comparison operator for TOptional<FName>::IsSet */
-	[[nodiscard]] inline bool operator==(FIntrusiveUnsetOptionalState I) const
+	[[nodiscard]] inline bool UEOpEquals(FIntrusiveUnsetOptionalState I) const
 	{
 		return ComparisonIndex == I;
 	}
@@ -1330,31 +1311,12 @@ private:
 	}
 #endif
 
-    template <typename T>
-    [[nodiscard]] friend FORCEINLINE auto operator==(FName N, T O) -> decltype(FName::Equals(N, O))
-    {
-    	return FName::Equals(N, O);
-    }
-
-    template <typename T>
-    [[nodiscard]] friend FORCEINLINE auto operator==(T O, FName N) -> decltype(FName::Equals(N, O))
-    {
-    	return FName::Equals(N, O);
-    }
-
-#if !PLATFORM_COMPILER_HAS_GENERATED_COMPARISON_OPERATORS
-    template <typename T>
-    [[nodiscard]] friend FORCEINLINE auto operator!=(FName N, T O) -> decltype(FName::Equals(N, O))
-    {
-    	return !FName::Equals(N, O);
-    }
-
-    template <typename T>
-    [[nodiscard]] friend FORCEINLINE auto operator!=(T O, FName N) -> decltype(FName::Equals(N, O))
-    {
-    	return !FName::Equals(N, O);
-    }
-#endif
+public:
+	template <typename T>
+	[[nodiscard]] FORCEINLINE auto UEOpEquals(T Rhs) const -> decltype(FName::Equals(*this, Rhs))
+	{
+		return FName::Equals(*this, Rhs);
+	}
 };
 
 template<> struct TIsZeroConstructType<class FName> { enum { Value = true }; };
@@ -1465,7 +1427,7 @@ FORCEINLINE bool FMinimalName::IsNone() const
 	return Index.IsNone();
 }
 
-FORCEINLINE bool FMinimalName::operator<(FMinimalName Rhs) const
+FORCEINLINE bool FMinimalName::UEOpLessThan(FMinimalName Rhs) const
 {
 	return Index < Rhs.Index;
 }
@@ -1545,7 +1507,7 @@ FORCEINLINE bool FMinimalName::IsNone() const
 	return Index.IsNone() && Number == NAME_NO_NUMBER_INTERNAL;
 }
 
-FORCEINLINE bool FMinimalName::operator<(FMinimalName Rhs) const
+FORCEINLINE bool FMinimalName::UEOpLessThan(FMinimalName Rhs) const
 {
 	return Index == Rhs.Index ? Number < Rhs.Number : Index < Rhs.Index;
 }
@@ -1764,31 +1726,25 @@ private:
 
 public:
 
-	friend bool operator==(FName Name, const FLazyName& Lazy)
+	bool UEOpEquals(FName Rhs) const
 	{
-		// If !Name.IsNone(), we have started creating FNames
-		// and might as well resolve and cache Lazy
-		if (Lazy.Either.IsName() || !Name.IsNone())
+		// If !Rhs.IsNone(), we have started creating FNames
+		// and might as well resolve and cache Lhs
+		if (Either.IsName() || !Rhs.IsNone())
 		{
-			return Name == Lazy.Resolve();
+			return Rhs == Resolve();
 		}
-		else if (!Lazy.bLiteralIsWide)
+		else if (!bLiteralIsWide)
 		{
-			return Name == Lazy.Either.AsAnsiLiteral();
+			return Rhs == Either.AsAnsiLiteral();
 		}
 		else
 		{
-			return Name == Lazy.Either.AsWideLiteral();
+			return Rhs == Either.AsWideLiteral();
 		}
 	}
 
-	friend bool operator==(const FLazyName& Lazy, FName Name)
-	{
-		return Name == Lazy;
-	}
-
-	CORE_API friend bool operator==(const FLazyName& A, const FLazyName& B);
-
+	CORE_API bool UEOpEquals(const FLazyName& Rhs) const;
 };
 
 /**
@@ -1816,17 +1772,18 @@ private:
 	FORCEINLINE bool SameIds() const { return (Value & DifferentIdsFlag) == 0; }
 	FORCEINLINE FNameEntryId GetDisplayId() const { return FNameEntryId::FromUnstableInt(Value & DisplayIdMask); }
 	FORCEINLINE FNameEntryId GetComparisonId() const { return SameIds() ? GetDisplayId() : FName::GetComparisonIdFromDisplayId(GetDisplayId()); }
-	friend bool operator==(FDisplayNameEntryId A, FDisplayNameEntryId B) { return A.Value == B.Value; }
+public:
+	bool UEOpEquals(FDisplayNameEntryId Rhs) const { return Value == Rhs.Value; }
 #else
 	FNameEntryId Id;
 
 	FDisplayNameEntryId(FNameEntryId InId, FNameEntryId) : Id(InId) {}
 	FORCEINLINE FNameEntryId GetDisplayId() const { return Id; }
 	FORCEINLINE FNameEntryId GetComparisonId() const { return Id; }
-	friend bool operator==(FDisplayNameEntryId A, FDisplayNameEntryId B) { return A.Id == B.Id; }
+public:
+	bool UEOpEquals(FDisplayNameEntryId Rhs) const { return Id == Rhs.Id; }
 #endif
-	friend bool operator==(FNameEntryId A, FDisplayNameEntryId B) { return A == B.GetDisplayId(); }
-	friend bool operator==(FDisplayNameEntryId A, FNameEntryId B) { return A.GetDisplayId() == B; }
+	bool UEOpEquals(FNameEntryId Rhs) const { return GetDisplayId() == Rhs; }
 	[[nodiscard]] friend uint32 GetTypeHash(FDisplayNameEntryId InId) { return GetTypeHash(InId.GetDisplayId()); }
 
 public: // Internal functions for batch serialization code - intentionally lacking CORE_API
diff --git a/Engine/Source/Runtime/Core/Public/UObject/WeakObjectPtrTemplates.h b/Engine/Source/Runtime/Core/Public/UObject/WeakObjectPtrTemplates.h
--- a/Engine/Source/Runtime/Core/Public/UObject/WeakObjectPtrTemplates.h
+++ b/Engine/Source/Runtime/Core/Public/UObject/WeakObjectPtrTemplates.h
@@ -6,6 +6,7 @@
 #include "Templates/LosesQualifiersFromTo.h"
 #include "Templates/Requires.h"
 #include "Containers/Map.h"
+#include "Misc/UEOps.h"
 #include "UObject/WeakObjectPtrTemplatesFwd.h"
 #include "UObject/StrongObjectPtrTemplatesFwd.h"
 
@@ -319,19 +320,19 @@ public:
 	 * @param Other weak pointer to compare to
 	 */
 	template <typename RhsT, typename = decltype((T*)nullptr == (RhsT*)nullptr)>
-	[[nodiscard]] FORCENOINLINE bool operator==(const TWeakObjectPtr<RhsT, TWeakObjectPtrBase>& Rhs) const
+	[[nodiscard]] FORCENOINLINE bool UEOpEquals(const TWeakObjectPtr<RhsT, TWeakObjectPtrBase>& Rhs) const
 	{
 		return WeakPtr == Rhs.WeakPtr;
 	}
 
 	template <typename RhsT, typename = decltype((T*)nullptr == (RhsT*)nullptr)>
-	[[nodiscard]] FORCENOINLINE bool operator==(const RhsT* Rhs) const
+	[[nodiscard]] FORCENOINLINE bool UEOpEquals(const RhsT* Rhs) const
 	{
 		// NOTE: this constructs a TWeakObjectPtrBase, which has some amount of overhead, so this may not be an efficient operation
 		return WeakPtr == TWeakObjectPtrBase(Rhs);
 	}
 
-	[[nodiscard]] FORCENOINLINE bool operator==(TYPE_OF_NULLPTR) const
+	[[nodiscard]] FORCENOINLINE bool UEOpEquals(TYPE_OF_NULLPTR) const
 	{
 		return !IsValid();
 	}
@@ -374,39 +375,6 @@ TWeakObjectPtr(T*) -> TWeakObjectPtr<T>;
 template <typename T>
 TWeakObjectPtr(const TWeakObjectPtr<T>&) -> TWeakObjectPtr<T>;
 
-#if !PLATFORM_COMPILER_HAS_GENERATED_COMPARISON_OPERATORS
-/**
- * Compare weak pointers for equality.
- * If both pointers would return nullptr from Get() they count as equal even if they were not initialized to the same object.
- * @param Other weak pointer to compare to
- */
-template <typename LhsT, typename RhsT, typename OtherTWeakObjectPtrBase, typename = decltype((LhsT*)nullptr == (RhsT*)nullptr)>
-[[nodiscard]] FORCEINLINE bool operator==(const LhsT* Lhs, const TWeakObjectPtr<RhsT, OtherTWeakObjectPtrBase>& Rhs)
-{
-	return Rhs == Lhs;
-}
-template <typename RhsT, typename OtherTWeakObjectPtrBase>
-[[nodiscard]] FORCEINLINE bool operator==(TYPE_OF_NULLPTR, const TWeakObjectPtr<RhsT, OtherTWeakObjectPtrBase>& Rhs)
-{
-	return Rhs == nullptr;
-}
-
-/**
- * Compare weak pointers for inequality
- * @param Other weak pointer to compare to
- */
-template <typename LhsT, typename RhsT, typename OtherTWeakObjectPtrBase, typename = decltype((LhsT*)nullptr != (RhsT*)nullptr)>
-[[nodiscard]] FORCEINLINE bool operator!=(const LhsT* Lhs, const TWeakObjectPtr<RhsT, OtherTWeakObjectPtrBase>& Rhs)
-{
-	return !(Rhs == Lhs);
-}
-template <typename RhsT, typename OtherTWeakObjectPtrBase>
-[[nodiscard]] FORCEINLINE bool operator!=(TYPE_OF_NULLPTR, const TWeakObjectPtr<RhsT, OtherTWeakObjectPtrBase>& Rhs)
-{
-	return !(Rhs == nullptr);
-}
-#endif
-
 // Helper function which deduces the type of the initializer
 template <typename T>
 [[nodiscard]] FORCEINLINE TWeakObjectPtr<T> MakeWeakObjectPtr(T* Ptr)
diff --git a/Engine/Source/Runtime/CoreUObject/Public/UObject/ObjectPtr.h b/Engine/Source/Runtime/CoreUObject/Public/UObject/ObjectPtr.h
--- a/Engine/Source/Runtime/CoreUObject/Public/UObject/ObjectPtr.h
+++ b/Engine/Source/Runtime/CoreUObject/Public/UObject/ObjectPtr.h
@@ -3,12 +3,13 @@
 #pragma once
 
 #include "HAL/Platform.h"
+#include "Misc/UEOps.h"
 #include "Serialization/StructuredArchive.h"
 #include "Templates/IsTObjectPtr.h"
+#include "Templates/NonNullPointer.h"
 #include "UObject/GarbageCollectionGlobals.h"
 #include "UObject/ObjectHandle.h"
 #include "UObject/UObjectGlobals.h"
-#include "Templates/NonNullPointer.h"
 
 #include <type_traits>
 
@@ -166,8 +167,10 @@ public:
 		return *this;
 	}
 
-	[[nodiscard]] FORCEINLINE bool operator==(FObjectPtr Other) const { return (Handle == Other.Handle); }
-	[[nodiscard]] FORCEINLINE bool operator!=(FObjectPtr Other) const { return (Handle != Other.Handle); }
+	[[nodiscard]] FORCEINLINE bool UEOpEquals(FObjectPtr Other) const
+	{
+		return (Handle == Other.Handle);
+	}
 
 	UE_OBJPTR_DEPRECATED(5.0, "Use of ToTObjectPtr is unsafe and is deprecated.")
 	[[nodiscard]] FORCEINLINE TObjectPtr<UObject>& ToTObjectPtr();
@@ -608,7 +611,7 @@ public:
 		typename U,
 		typename Base = std::common_type_t<T*, U*>
 	>
-	[[nodiscard]] FORCEINLINE bool operator==(const TObjectPtr<U>& Other) const
+	[[nodiscard]] FORCEINLINE bool UEOpEquals(const TObjectPtr<U>& Other) const
 	{
 #if UE_WITH_OBJECT_HANDLE_TYPE_SAFETY
 		// Do a NULL test first before comparing the underlying handles, in case either side is
@@ -626,7 +629,7 @@ public:
 	}
 
 	// Equality/Inequality comparisons against nullptr
-	[[nodiscard]] FORCEINLINE bool operator==(TYPE_OF_NULLPTR) const
+	[[nodiscard]] FORCEINLINE bool UEOpEquals(TYPE_OF_NULLPTR) const
 	{
 		return ObjectPtr_Private::IsObjectPtrNull<T>(ObjectPtr);
 	}
@@ -636,19 +639,11 @@ public:
 		typename U,
 		typename = decltype(ObjectPtr_Private::IsObjectPtrEqual(std::declval<const TObjectPtr<T>&>(), std::declval<const U&>()))
 	>
-	[[nodiscard]] FORCEINLINE bool operator==(const U& Other) const
+	[[nodiscard]] FORCEINLINE bool UEOpEquals(const U& Other) const
 	{
 		return ObjectPtr_Private::IsObjectPtrEqual(*this, Other);
 	}
 
-#if __cplusplus < 202002L
-	template <typename U>
-	[[nodiscard]] FORCEINLINE auto operator!=(const U& Other) const -> decltype(!(*this == Other))
-	{
-		return !(*this == Other);
-	}
-#endif
-
 	// @TODO: OBJPTR: There is a risk that the FObjectPtr is storing a reference to the wrong type.  This could
 	//			happen if data was serialized at a time when a pointer field was declared to be of type A, but then the declaration
 	//			changed and the pointer field is now of type B.  Upon deserialization of pre-existing data, we'll be holding
@@ -950,28 +945,6 @@ FORCEINLINE void operator<<(FStructuredArchiveSlot Slot, TObjectPtr<T>& InObject
 	ObjectPtr_Private::Friend::SerializePtrStructured(Slot, InObjectPtr);
 }
 
-#if !PLATFORM_COMPILER_HAS_GENERATED_COMPARISON_OPERATORS
-// Equality/Inequality comparisons against another type that can be implicitly converted to the pointer type kept in a TObjectPtr
-template <
-	typename T,
-	typename U,
-	decltype(ObjectPtr_Private::IsObjectPtrEqual(std::declval<const TObjectPtr<T>&>(), std::declval<const U&>()))* = nullptr
->
-[[nodiscard]] FORCEINLINE bool operator==(const U& Other, const TObjectPtr<T>& Ptr)
-{
-	return ObjectPtr_Private::IsObjectPtrEqual(Ptr, Other);
-}
-template <
-	typename T,
-	typename U,
-	decltype(ObjectPtr_Private::IsObjectPtrEqual(std::declval<const TObjectPtr<T>&>(), std::declval<const U&>()))* = nullptr
->
-[[nodiscard]] FORCEINLINE bool operator!=(const U& Other, const TObjectPtr<T>& Ptr)
-{
-	return !ObjectPtr_Private::IsObjectPtrEqual(Ptr, Other);
-}
-#endif
-
 template <typename T>
 TPrivateObjectPtr<T> MakeObjectPtrUnsafe(const UObject* Obj);
 
