Added missing header.

diff --git a/Engine/Source/Runtime/Core/Public/Misc/UEOps.h b/Engine/Source/Runtime/Core/Public/Misc/UEOps.h
new file mode 100644
--- /dev/null
+++ b/Engine/Source/Runtime/Core/Public/Misc/UEOps.h
@@ -0,0 +1,177 @@
+// Copyright Epic Games, Inc. All Rights Reserved.
+
+#pragma once
+
+#include "Templates/Requires.h"
+
+namespace UE::Core::Private
+{
+	// This only exists to bind to NULL, which is not perfectly-forwardable but which is commonly used in comparison operators.
+	// The type is unique so that it can't be confused with a type that will actually be used by real user code.
+	struct FIncomplete;
+
+	template <typename LhsType, typename RhsType>
+	auto EnableUEOpEqualsImpl(const LhsType& Lhs, const RhsType& Rhs, int) -> decltype(Lhs.UEOpEquals(Rhs));
+	template <typename LhsType, typename RhsType>
+	void EnableUEOpEqualsImpl(const LhsType& Lhs, const RhsType& Rhs, ...) = delete;
+
+	template <typename LhsType, typename RhsType>
+	auto DisableUEOpEqualsImpl(const LhsType& Lhs, const RhsType& Rhs, int) -> decltype(Lhs.UEOpEquals(Rhs)) = delete;
+	template <typename LhsType, typename RhsType>
+	void DisableUEOpEqualsImpl(const LhsType& Lhs, const RhsType& Rhs, ...);
+
+	template <typename LhsType, typename RhsType>
+	auto EnableUEOpLessThanImpl(const LhsType& Lhs, const RhsType& Rhs, int) -> decltype(Lhs.UEOpLessThan(Rhs));
+	template <typename LhsType, typename RhsType>
+	void EnableUEOpLessThanImpl(const LhsType& Lhs, const RhsType& Rhs, ...) = delete;
+
+	template <typename LhsType, typename RhsType>
+	auto DisableUEOpLessThanImpl(const LhsType& Lhs, const RhsType& Rhs, int) -> decltype(Lhs.UEOpLessThan(Rhs)) = delete;
+	template <typename LhsType, typename RhsType>
+	void DisableUEOpLessThanImpl(const LhsType& Lhs, const RhsType& Rhs, ...);
+
+	// SFINAE-compatible concept checks for the existence of various UEOp member functions.
+	//
+	// We can't use metaprogramming classes like TModels or std::enable_if because they end up being instantiated with incomplete types and
+	// end up with incorrect overload resolution later in the compilation.
+	//
+	// Use:
+	// decltype(UE::Core::Private::UEOpEquals(A, B)) should be present in the signature of the function to enable it for comparisons between
+	// the two operands.
+	// This will cause SFINAE failure (and thus 'enable' the overload) if A.UEOpEquals(B) is a valid expression. If the function is enabled,
+	// the return type of this function will be the same as the return type as A.UEOpEquals(B).
+	//
+	// decltype(UE::Core::Private::DisableUEOpEquals(A, B)) will cause SFINAE failure (and thus 'disable' the overload) if A.UEOpEquals(B) is
+	// a valid expression.
+	// The typical use is in conjunction with decltype(UE::Core::Private::UEOpEquals(B, A)) to cause mutually exclusive overloads with a
+	// decltype(UE::Core::Private::UEOpEquals(A, B)) overload to prevent ambiguity when A and B are comparable in both directions, e.g.
+	// comparing TSharedPtr<X> with TSharedPtr<Y>.
+	//
+	// UE::Core::Private::UEOpLessThan and UE::Core::Private::DisableUEOpLessThan work similarly, but for UEOpLessThan.
+
+	template <typename LhsType, typename RhsType>
+	auto UEOpEquals(const LhsType& Lhs, const RhsType& Rhs) -> decltype(::UE::Core::Private::EnableUEOpEqualsImpl(Lhs, Rhs, 0));
+
+	template <typename LhsType, typename RhsType>
+	auto DisableUEOpEquals(const LhsType& Lhs, const RhsType& Rhs) -> decltype(::UE::Core::Private::DisableUEOpEqualsImpl(Lhs, Rhs, 0));
+
+	template <typename LhsType, typename RhsType>
+	auto UEOpLessThan(const LhsType& Lhs, const RhsType& Rhs) -> decltype(::UE::Core::Private::EnableUEOpLessThanImpl(Lhs, Rhs, 0));
+
+	template <typename LhsType, typename RhsType>
+	auto DisableUEOpLessThan(const LhsType& Lhs, const RhsType& Rhs) -> decltype(::UE::Core::Private::DisableUEOpLessThanImpl(Lhs, Rhs, 0));
+}
+
+/**
+ * This file defines global comparison operators for all types containing certain named member functions.
+ * These ensure that the operators and their variants are defined correctly based on a single customization
+ * point and should result in faster overload resolution.
+ *
+ * bool FMyType::UEOpEquals(const OtherType&) const;
+ *   Should be defined to opt into operator== and operator!=.
+ *
+ * bool FMyType::UEOpLessThan(const OtherType&) const;
+ *   Should be defined to opt into operator<, operator<=, operator> and operator>=.
+ */
+
+template <typename LhsType, typename RhsType>
+[[nodiscard]] UE_REWRITE constexpr auto operator==(const LhsType& Lhs, const RhsType& Rhs) -> decltype(UE::Core::Private::UEOpEquals(Lhs, Rhs))
+{
+	return Lhs.UEOpEquals(Rhs);
+}
+
+template <typename LhsType, typename RhsType>
+[[nodiscard]] UE_REWRITE constexpr auto operator==(const LhsType& Lhs, const RhsType& Rhs) -> decltype(UE::Core::Private::DisableUEOpEquals(Lhs, Rhs), UE::Core::Private::UEOpEquals(Rhs, Lhs))
+{
+	return Rhs.UEOpEquals(Lhs);
+}
+
+// This overload exists to support A == NULL, which is common in comparison operators, but NULL is not perfectly-forwardable.
+template <typename LhsType>
+[[nodiscard]] UE_REWRITE constexpr auto operator==(const LhsType& Lhs, UE::Core::Private::FIncomplete* Rhs) -> decltype(UE::Core::Private::UEOpEquals(Lhs, nullptr))
+{
+	return Lhs.UEOpEquals(nullptr);
+}
+
+// This overload exists to support NULL == A, which is common in comparison operators, but NULL is not perfectly-forwardable.
+template <typename RhsType>
+[[nodiscard]] UE_REWRITE constexpr auto operator==(UE::Core::Private::FIncomplete* Lhs, const RhsType& Rhs) -> decltype(UE::Core::Private::UEOpEquals(Rhs, nullptr))
+{
+	return Rhs.UEOpEquals(nullptr);
+}
+
+#if !PLATFORM_COMPILER_HAS_GENERATED_COMPARISON_OPERATORS
+template <typename LhsType, typename RhsType>
+[[nodiscard]] UE_REWRITE constexpr auto operator!=(const LhsType& Lhs, const RhsType& Rhs) -> decltype(UE::Core::Private::UEOpEquals(Lhs, Rhs))
+{
+	return !Lhs.UEOpEquals(Rhs);
+}
+
+template <typename LhsType, typename RhsType>
+[[nodiscard]] UE_REWRITE constexpr auto operator!=(const LhsType& Lhs, const RhsType& Rhs) -> decltype(UE::Core::Private::DisableUEOpEquals(Lhs, Rhs), UE::Core::Private::UEOpEquals(Rhs, Lhs))
+{
+	return !Rhs.UEOpEquals(Lhs);
+}
+
+// This overload exists to support A != NULL, which is common in comparison operators, but NULL is not perfectly-forwardable.
+template <typename LhsType>
+[[nodiscard]] UE_REWRITE constexpr auto operator!=(const LhsType& Lhs, UE::Core::Private::FIncomplete* Rhs) -> decltype(UE::Core::Private::UEOpEquals(Lhs, nullptr))
+{
+	return !Lhs.UEOpEquals(nullptr);
+}
+
+// This overload exists to support NULL != A, which is common in comparison operators, but NULL is not perfectly-forwardable.
+template <typename RhsType>
+[[nodiscard]] UE_REWRITE constexpr auto operator!=(UE::Core::Private::FIncomplete* Lhs, const RhsType& Rhs) -> decltype(UE::Core::Private::UEOpEquals(Rhs, nullptr))
+{
+	return !Rhs.UEOpEquals(nullptr);
+}
+#endif
+
+template <typename LhsType, typename RhsType>
+[[nodiscard]] UE_REWRITE constexpr auto operator<(const LhsType& Lhs, const RhsType& Rhs) -> decltype(UE::Core::Private::UEOpLessThan(Lhs, Rhs))
+{
+	return Lhs.UEOpLessThan(Rhs);
+}
+
+template <typename LhsType, typename RhsType>
+[[nodiscard]] UE_REWRITE constexpr auto operator<(const LhsType& Lhs, const RhsType& Rhs) -> decltype(UE::Core::Private::DisableUEOpLessThan(Lhs, Rhs), UE::Core::Private::UEOpLessThan(Rhs, Lhs))
+{
+	return Rhs.UEOpLessThan(Lhs);
+}
+
+template <typename LhsType, typename RhsType>
+[[nodiscard]] UE_REWRITE constexpr auto operator>(const LhsType& Lhs, const RhsType& Rhs) -> decltype(UE::Core::Private::UEOpEquals(Rhs, Lhs))
+{
+	return Rhs.UEOpLessThan(Lhs);
+}
+
+template <typename LhsType, typename RhsType>
+[[nodiscard]] UE_REWRITE constexpr auto operator>(const LhsType& Lhs, const RhsType& Rhs) -> decltype(UE::Core::Private::DisableUEOpEquals(Rhs, Lhs), UE::Core::Private::UEOpEquals(Lhs, Rhs))
+{
+	return Lhs.UEOpLessThan(Rhs);
+}
+
+template <typename LhsType, typename RhsType>
+[[nodiscard]] UE_REWRITE constexpr auto operator>=(const LhsType& Lhs, const RhsType& Rhs) -> decltype(UE::Core::Private::UEOpLessThan(Lhs, Rhs))
+{
+	return !Lhs.UEOpLessThan(Rhs);
+}
+
+template <typename LhsType, typename RhsType>
+[[nodiscard]] UE_REWRITE constexpr auto operator>=(const LhsType& Lhs, const RhsType& Rhs) -> decltype(UE::Core::Private::DisableUEOpLessThan(Lhs, Rhs), UE::Core::Private::UEOpLessThan(Rhs, Lhs))
+{
+	return !Rhs.UEOpLessThan(Lhs);
+}
+
+template <typename LhsType, typename RhsType>
+[[nodiscard]] UE_REWRITE constexpr auto operator<=(const LhsType& Lhs, const RhsType& Rhs) -> decltype(UE::Core::Private::UEOpEquals(Rhs, Lhs))
+{
+	return !Rhs.UEOpLessThan(Lhs);
+}
+
+template <typename LhsType, typename RhsType>
+[[nodiscard]] UE_REWRITE constexpr auto operator<=(const LhsType& Lhs, const RhsType& Rhs) -> decltype(UE::Core::Private::DisableUEOpEquals(Rhs, Lhs), UE::Core::Private::UEOpEquals(Lhs, Rhs))
+{
+	return !Lhs.UEOpLessThan(Rhs);
+}
