Updated comments and implementation around context switch events.

diff --git a/Engine/Source/Runtime/Core/Private/ProfilingDebugging/Microsoft/EventTracingForWindows.cpp b/Engine/Source/Runtime/Core/Private/ProfilingDebugging/Microsoft/EventTracingForWindows.cpp
--- a/Engine/Source/Runtime/Core/Private/ProfilingDebugging/Microsoft/EventTracingForWindows.cpp
+++ b/Engine/Source/Runtime/Core/Private/ProfilingDebugging/Microsoft/EventTracingForWindows.cpp
@@ -80,6 +80,7 @@ struct FContextSwitchEvent
 	uint8  OldThreadWaitIdealProcessor;	// 15
 	uint32 NewThreadWaitTime;			// 16
 	uint32 Reserved;					// 20
+	uint32 Reserved2;					// 24 (undocumented)
 };
 
 // https://docs.microsoft.com/en-us/windows/win32/etw/stackwalk-event
@@ -113,8 +114,10 @@ void FPlatformEvents::TraceEventCallback(PEVENT_RECORD Event)
 	{
 		const FContextSwitchEvent* Payload = (const FContextSwitchEvent*)Event->UserData;
 		
-		// On some platforms we are getting Event->UserDataLength > sizeof(FContextSwitchEvent), the delta in size is undocumented 
-		if (Payload == nullptr || Event->UserDataLength < sizeof(FContextSwitchEvent))
+		// Recent updates of ETW changes the size of the payload somewhat.
+		constexpr uint32 UserDataSizeOld = sizeof(FContextSwitchEvent) - 4;
+		constexpr uint32 UserDataSizeNew = sizeof(FContextSwitchEvent);
+		if (Payload == nullptr || (Event->UserDataLength != UserDataSizeOld && Event->UserDataLength != UserDataSizeNew))
 		{
 			return;
 		}
