Replaced the UTF8CHAR enum with a char8_t typedef.

diff --git a/Engine/Source/Runtime/Core/Private/String/Find.cpp b/Engine/Source/Runtime/Core/Private/String/Find.cpp
--- a/Engine/Source/Runtime/Core/Private/String/Find.cpp
+++ b/Engine/Source/Runtime/Core/Private/String/Find.cpp
@@ -442,11 +442,6 @@ int32 FindLastOfAny(FWideStringView View, TConstArrayView<FWideStringView> Searc
 	return Private::FindLastOfAny(View, Search, SearchCase, OutMatchIndex);
 }
 
-int32 FindFirstChar(FUtf8StringView View, ANSICHAR Search, ESearchCase::Type SearchCase)
-{
-	return Private::FindFirstChar(View, UTF8CHAR(Search), SearchCase);
-}
-
 int32 FindFirstChar(FUtf8StringView View, UTF8CHAR Search, ESearchCase::Type SearchCase)
 {
 	return Private::FindFirstChar(View, Search, SearchCase);
@@ -457,11 +452,6 @@ int32 FindFirstChar(FWideStringView View, WIDECHAR Search, ESearchCase::Type Sea
 	return Private::FindFirstChar(View, Search, SearchCase);
 }
 
-int32 FindLastChar(FUtf8StringView View, ANSICHAR Search, ESearchCase::Type SearchCase)
-{
-	return Private::FindLastChar(View, UTF8CHAR(Search), SearchCase);
-}
-
 int32 FindLastChar(FUtf8StringView View, UTF8CHAR Search, ESearchCase::Type SearchCase)
 {
 	return Private::FindLastChar(View, Search, SearchCase);
@@ -472,6 +462,11 @@ int32 FindLastChar(FWideStringView View, WIDECHAR Search, ESearchCase::Type Sear
 	return Private::FindLastChar(View, Search, SearchCase);
 }
 
+int32 FindFirstOfAnyChar(FUtf8StringView View, std::initializer_list<UTF8CHAR> Search, ESearchCase::Type SearchCase, int32* OutMatchIndex)
+{
+	return Private::FindFirstOfAnyChar(View, MakeArrayView(Search), SearchCase, OutMatchIndex);
+}
+
 int32 FindFirstOfAnyChar(FUtf8StringView View, TConstArrayView<ANSICHAR> Search, ESearchCase::Type SearchCase, int32* OutMatchIndex)
 {
 	return Private::FindFirstOfAnyChar(View, MakeArrayView((const UTF8CHAR*)Search.GetData(), Search.Num()), SearchCase, OutMatchIndex);
@@ -487,6 +482,11 @@ int32 FindFirstOfAnyChar(FWideStringView View, TConstArrayView<WIDECHAR> Search,
 	return Private::FindFirstOfAnyChar(View, Search, SearchCase, OutMatchIndex);
 }
 
+int32 FindLastOfAnyChar(FUtf8StringView View, std::initializer_list<UTF8CHAR> Search, ESearchCase::Type SearchCase, int32* OutMatchIndex)
+{
+	return Private::FindLastOfAnyChar(View, MakeArrayView(Search), SearchCase, OutMatchIndex);
+}
+
 int32 FindLastOfAnyChar(FUtf8StringView View, TConstArrayView<ANSICHAR> Search, ESearchCase::Type SearchCase, int32* OutMatchIndex)
 {
 	return Private::FindLastOfAnyChar(View, MakeArrayView((const UTF8CHAR*)Search.GetData(), Search.Num()), SearchCase, OutMatchIndex);
diff --git a/Engine/Source/Runtime/Core/Public/Containers/StringView.h b/Engine/Source/Runtime/Core/Public/Containers/StringView.h
--- a/Engine/Source/Runtime/Core/Public/Containers/StringView.h
+++ b/Engine/Source/Runtime/Core/Public/Containers/StringView.h
@@ -53,6 +53,14 @@ namespace UE::Core::Private
 	}
 } // UE::Core::Private
 
+template <typename CharPtrOrRangeType>
+	requires (!std::is_pointer_v<CharPtrOrRangeType> || TIsCharType_V<std::remove_pointer_t<CharPtrOrRangeType>>)
+[[nodiscard]] constexpr inline auto MakeStringView(const CharPtrOrRangeType& CharPtrOrRange UE_LIFETIMEBOUND) -> decltype(TStringView(CharPtrOrRange));
+
+template <UE::CCharType CharType>
+[[nodiscard]] constexpr inline auto MakeStringView(const CharType* CharPtr UE_LIFETIMEBOUND, int32 Size) -> decltype(TStringView(CharPtr, Size));
+
+
 /**
  * A string view is a non-owning view of a range of characters.
  *
diff --git a/Engine/Source/Runtime/Core/Public/GenericPlatform/GenericPlatform.h b/Engine/Source/Runtime/Core/Public/GenericPlatform/GenericPlatform.h
--- a/Engine/Source/Runtime/Core/Public/GenericPlatform/GenericPlatform.h
+++ b/Engine/Source/Runtime/Core/Public/GenericPlatform/GenericPlatform.h
@@ -67,8 +67,8 @@ struct FGenericPlatformTypes
 	// A wide character. In-memory only. ?-bit fixed-width representation of the platform's natural wide character set. Could be different sizes on different platforms.
 	typedef wchar_t				WIDECHAR;
 
-	// An 8-bit character type. In-memory only. 8-bit representation. Should really be char8_t but making this the generic option is easier for compilers which don't fully support C++20 yet.
-	enum UTF8CHAR : unsigned char {};
+	// An 8-bit character type. In-memory only. 8-bit representation.
+	using UTF8CHAR = char8_t;
 
 	// A 16-bit character type. In-memory only.  16-bit representation. Should really be char16_t but making this the generic option is easier for compilers which don't fully support C++11 yet (i.e. MSVC).
 	typedef uint16				CHAR16;		
@@ -98,30 +98,3 @@ struct FGenericPlatformTypes
 	typedef int32					TYPE_OF_NULL;
 	typedef decltype(nullptr)		TYPE_OF_NULLPTR;
 };
-
-constexpr FGenericPlatformTypes::UTF8CHAR& operator++(FGenericPlatformTypes::UTF8CHAR& Ch     ) {                                              ++*(unsigned char*)&Ch; return Ch;     }
-constexpr FGenericPlatformTypes::UTF8CHAR& operator--(FGenericPlatformTypes::UTF8CHAR& Ch     ) {                                              --*(unsigned char*)&Ch; return Ch;     }
-constexpr FGenericPlatformTypes::UTF8CHAR  operator++(FGenericPlatformTypes::UTF8CHAR& Ch, int) { FGenericPlatformTypes::UTF8CHAR Result = Ch; ++*(unsigned char*)&Ch; return Result; }
-constexpr FGenericPlatformTypes::UTF8CHAR  operator--(FGenericPlatformTypes::UTF8CHAR& Ch, int) { FGenericPlatformTypes::UTF8CHAR Result = Ch; --*(unsigned char*)&Ch; return Result; }
-
-template <typename T> inline auto operator+= (FGenericPlatformTypes::UTF8CHAR& Ch, T&& InValue) -> decltype((FGenericPlatformTypes::UTF8CHAR&)(*(unsigned char*)&Ch +=  (T&&)InValue)) { return (FGenericPlatformTypes::UTF8CHAR&)(*(unsigned char*)&Ch +=  (T&&)InValue); }
-template <typename T> inline auto operator-= (FGenericPlatformTypes::UTF8CHAR& Ch, T&& InValue) -> decltype((FGenericPlatformTypes::UTF8CHAR&)(*(unsigned char*)&Ch -=  (T&&)InValue)) { return (FGenericPlatformTypes::UTF8CHAR&)(*(unsigned char*)&Ch -=  (T&&)InValue); }
-template <typename T> inline auto operator/= (FGenericPlatformTypes::UTF8CHAR& Ch, T&& InValue) -> decltype((FGenericPlatformTypes::UTF8CHAR&)(*(unsigned char*)&Ch /=  (T&&)InValue)) { return (FGenericPlatformTypes::UTF8CHAR&)(*(unsigned char*)&Ch /=  (T&&)InValue); }
-template <typename T> inline auto operator*= (FGenericPlatformTypes::UTF8CHAR& Ch, T&& InValue) -> decltype((FGenericPlatformTypes::UTF8CHAR&)(*(unsigned char*)&Ch *=  (T&&)InValue)) { return (FGenericPlatformTypes::UTF8CHAR&)(*(unsigned char*)&Ch *=  (T&&)InValue); }
-template <typename T> inline auto operator%= (FGenericPlatformTypes::UTF8CHAR& Ch, T&& InValue) -> decltype((FGenericPlatformTypes::UTF8CHAR&)(*(unsigned char*)&Ch %=  (T&&)InValue)) { return (FGenericPlatformTypes::UTF8CHAR&)(*(unsigned char*)&Ch %=  (T&&)InValue); }
-template <typename T> inline auto operator&= (FGenericPlatformTypes::UTF8CHAR& Ch, T&& InValue) -> decltype((FGenericPlatformTypes::UTF8CHAR&)(*(unsigned char*)&Ch &=  (T&&)InValue)) { return (FGenericPlatformTypes::UTF8CHAR&)(*(unsigned char*)&Ch &=  (T&&)InValue); }
-template <typename T> inline auto operator|= (FGenericPlatformTypes::UTF8CHAR& Ch, T&& InValue) -> decltype((FGenericPlatformTypes::UTF8CHAR&)(*(unsigned char*)&Ch |=  (T&&)InValue)) { return (FGenericPlatformTypes::UTF8CHAR&)(*(unsigned char*)&Ch |=  (T&&)InValue); }
-template <typename T> inline auto operator^= (FGenericPlatformTypes::UTF8CHAR& Ch, T&& InValue) -> decltype((FGenericPlatformTypes::UTF8CHAR&)(*(unsigned char*)&Ch ^=  (T&&)InValue)) { return (FGenericPlatformTypes::UTF8CHAR&)(*(unsigned char*)&Ch ^=  (T&&)InValue); }
-template <typename T> inline auto operator<<=(FGenericPlatformTypes::UTF8CHAR& Ch, T&& InValue) -> decltype((FGenericPlatformTypes::UTF8CHAR&)(*(unsigned char*)&Ch <<= (T&&)InValue)) { return (FGenericPlatformTypes::UTF8CHAR&)(*(unsigned char*)&Ch <<= (T&&)InValue); }
-template <typename T> inline auto operator>>=(FGenericPlatformTypes::UTF8CHAR& Ch, T&& InValue) -> decltype((FGenericPlatformTypes::UTF8CHAR&)(*(unsigned char*)&Ch >>= (T&&)InValue)) { return (FGenericPlatformTypes::UTF8CHAR&)(*(unsigned char*)&Ch >>= (T&&)InValue); }
-
-inline FGenericPlatformTypes::UTF8CHAR& operator+= (FGenericPlatformTypes::UTF8CHAR& Ch, FGenericPlatformTypes::UTF8CHAR InValue) { return (FGenericPlatformTypes::UTF8CHAR&)(*(unsigned char*)&Ch +=  *(unsigned char*)&InValue); }
-inline FGenericPlatformTypes::UTF8CHAR& operator-= (FGenericPlatformTypes::UTF8CHAR& Ch, FGenericPlatformTypes::UTF8CHAR InValue) { return (FGenericPlatformTypes::UTF8CHAR&)(*(unsigned char*)&Ch -=  *(unsigned char*)&InValue); }
-inline FGenericPlatformTypes::UTF8CHAR& operator/= (FGenericPlatformTypes::UTF8CHAR& Ch, FGenericPlatformTypes::UTF8CHAR InValue) { return (FGenericPlatformTypes::UTF8CHAR&)(*(unsigned char*)&Ch /=  *(unsigned char*)&InValue); }
-inline FGenericPlatformTypes::UTF8CHAR& operator*= (FGenericPlatformTypes::UTF8CHAR& Ch, FGenericPlatformTypes::UTF8CHAR InValue) { return (FGenericPlatformTypes::UTF8CHAR&)(*(unsigned char*)&Ch *=  *(unsigned char*)&InValue); }
-inline FGenericPlatformTypes::UTF8CHAR& operator%= (FGenericPlatformTypes::UTF8CHAR& Ch, FGenericPlatformTypes::UTF8CHAR InValue) { return (FGenericPlatformTypes::UTF8CHAR&)(*(unsigned char*)&Ch %=  *(unsigned char*)&InValue); }
-inline FGenericPlatformTypes::UTF8CHAR& operator&= (FGenericPlatformTypes::UTF8CHAR& Ch, FGenericPlatformTypes::UTF8CHAR InValue) { return (FGenericPlatformTypes::UTF8CHAR&)(*(unsigned char*)&Ch &=  *(unsigned char*)&InValue); }
-inline FGenericPlatformTypes::UTF8CHAR& operator|= (FGenericPlatformTypes::UTF8CHAR& Ch, FGenericPlatformTypes::UTF8CHAR InValue) { return (FGenericPlatformTypes::UTF8CHAR&)(*(unsigned char*)&Ch |=  *(unsigned char*)&InValue); }
-inline FGenericPlatformTypes::UTF8CHAR& operator^= (FGenericPlatformTypes::UTF8CHAR& Ch, FGenericPlatformTypes::UTF8CHAR InValue) { return (FGenericPlatformTypes::UTF8CHAR&)(*(unsigned char*)&Ch ^=  *(unsigned char*)&InValue); }
-inline FGenericPlatformTypes::UTF8CHAR& operator<<=(FGenericPlatformTypes::UTF8CHAR& Ch, FGenericPlatformTypes::UTF8CHAR InValue) { return (FGenericPlatformTypes::UTF8CHAR&)(*(unsigned char*)&Ch <<= *(unsigned char*)&InValue); }
-inline FGenericPlatformTypes::UTF8CHAR& operator>>=(FGenericPlatformTypes::UTF8CHAR& Ch, FGenericPlatformTypes::UTF8CHAR InValue) { return (FGenericPlatformTypes::UTF8CHAR&)(*(unsigned char*)&Ch >>= *(unsigned char*)&InValue); }
diff --git a/Engine/Source/Runtime/Core/Public/HAL/Platform.h b/Engine/Source/Runtime/Core/Public/HAL/Platform.h
--- a/Engine/Source/Runtime/Core/Public/HAL/Platform.h
+++ b/Engine/Source/Runtime/Core/Public/HAL/Platform.h
@@ -1282,88 +1282,10 @@ namespace TypeTests
 	#define TEXT(x) TEXT_PASTE(x)
 #endif
 
-// UE_UTF8TEXT_IS_CONSTEXPR is only enabled because FormatStringSanTests.cpp needs to use UTF8TEXT in static_asserts.
-// It is not enabled in general because of a concern for compile times and PDB sizes.  This can be removed when we can guarantee C++20 and
-// when UTF8CHAR == char8_t and UTF8TEXT simply prepends u8 to the literal.
-#if WITH_LOW_LEVEL_TESTS && defined(__cpp_nontype_template_args)
-	#ifdef __clang__
-		// Clang supports the necessary language features to support this in C++20, but hasn't updated the version
-		// number of the feature test macro due to unrelated deficiencies, so we just assume we can enable it under C++20 in general:
-		//
-		// [REDACTED_PROFILE_URL]/llvm-project/issues/54297#issuecomment-1408671519
-		#define UE_UTF8TEXT_IS_CONSTEXPR (__cplusplus >= 202002L)
-	#else
-		#define UE_UTF8TEXT_IS_CONSTEXPR (__cpp_nontype_template_args >= 201911L)
-	#endif
-#else
-	#define UE_UTF8TEXT_IS_CONSTEXPR 0
-#endif
+#define UE_UTF8TEXT_IS_CONSTEXPR 1
 
 namespace UE::Core::Private
 {
-#if UE_UTF8TEXT_IS_CONSTEXPR
-	template <SIZE_T N = 1>
-	struct TUtf8StringLiteral
-	{
-#if defined(__cpp_char8_t)
-		using Utf8CharLiteralType = char8_t;
-#else
-		using Utf8CharLiteralType = char;
-#endif
-
-		constexpr TUtf8StringLiteral(const Utf8CharLiteralType(&InArray)[N])
-		{
-			for (SIZE_T Index = 0; Index != N; ++Index)
-			{
-				Array[Index] = (UTF8CHAR)InArray[Index];
-			}
-		}
-
-		constexpr TUtf8StringLiteral(unsigned long long Val)
-		{
-			Array[0] = (UTF8CHAR)Val;
-		}
-
-		UTF8CHAR Array[N];
-	};
-
-	template <typename ArgType, TUtf8StringLiteral Literal, typename ArraySelector = void>
-	struct TUtf8StringLiteralStatic
-	{
-		static constexpr auto StaticLiteral = Literal;
-		static constexpr const auto& Value = Literal.Array;
-	};
-	template <typename ArgType, TUtf8StringLiteral Literal>
-	struct TUtf8StringLiteralStatic<ArgType, Literal, std::enable_if_t<!std::is_array_v<std::remove_reference_t<ArgType>>>>
-	{
-		static constexpr auto StaticLiteral = Literal;
-		static constexpr auto Value = Literal.Array[0];
-	};
-
-	#define UTF8TEXT(x) (UE::Core::Private::TUtf8StringLiteralStatic<decltype(x), UTF8TEXT_PASTE(x)>::Value)
-#else
-	// Can't be constexpr because it involves casts
-	template <SIZE_T N>
-	FORCEINLINE auto ToUTF8Literal(const char(&Array)[N]) -> const UTF8CHAR(&)[N]
-	{
-		return (const UTF8CHAR(&)[N])Array;
-	}
-	FORCEINLINE constexpr UTF8CHAR ToUTF8Literal(unsigned long long Ch)
-	{
-		return (UTF8CHAR)Ch;
-	}
-#if defined(__cpp_char8_t)
-	// Can't be constexpr because it involves casts
-	template <SIZE_T N>
-	FORCEINLINE auto ToUTF8Literal(const char8_t (&Array)[N]) -> const UTF8CHAR(&)[N]
-	{
-		return (const UTF8CHAR (&)[N])Array;
-	}
-#endif
-
-	#define UTF8TEXT(x) (UE::Core::Private::ToUTF8Literal(UTF8TEXT_PASTE(x)))
-#endif
-
 	template <typename CharType>
 	void CharTextStaticAssert()
 	{
@@ -1371,6 +1293,7 @@ namespace UE::Core::Private
 	}
 }
 
+#define UTF8TEXT(x) UTF8TEXT_PASTE(x)
 #define UTF16TEXT(x) UTF16TEXT_PASTE(x)
 #define WIDETEXT(str) WIDETEXT_PASTE(str)
 
diff --git a/Engine/Source/Runtime/Core/Public/String/Find.h b/Engine/Source/Runtime/Core/Public/String/Find.h
--- a/Engine/Source/Runtime/Core/Public/String/Find.h
+++ b/Engine/Source/Runtime/Core/Public/String/Find.h
@@ -67,7 +67,6 @@ CORE_API int32 FindLastOfAny(FWideStringView View, TConstArrayView<FWideStringVi
  * @param SearchCase Whether the comparison should ignore case.
  * @return The position at which the search character was found, or INDEX_NONE if not found.
  */
-CORE_API int32 FindFirstChar(FUtf8StringView View, ANSICHAR Search, ESearchCase::Type SearchCase = ESearchCase::CaseSensitive);
 CORE_API int32 FindFirstChar(FUtf8StringView View, UTF8CHAR Search, ESearchCase::Type SearchCase = ESearchCase::CaseSensitive);
 CORE_API int32 FindFirstChar(FWideStringView View, WIDECHAR Search, ESearchCase::Type SearchCase = ESearchCase::CaseSensitive);
 
@@ -79,7 +78,6 @@ CORE_API int32 FindFirstChar(FWideStringView View, WIDECHAR Search, ESearchCase:
  * @param SearchCase Whether the comparison should ignore case.
  * @return The position at which the search character was found, or INDEX_NONE if not found.
  */
-CORE_API int32 FindLastChar(FUtf8StringView View, ANSICHAR Search, ESearchCase::Type SearchCase = ESearchCase::CaseSensitive);
 CORE_API int32 FindLastChar(FUtf8StringView View, UTF8CHAR Search, ESearchCase::Type SearchCase = ESearchCase::CaseSensitive);
 CORE_API int32 FindLastChar(FWideStringView View, WIDECHAR Search, ESearchCase::Type SearchCase = ESearchCase::CaseSensitive);
 
@@ -92,6 +90,7 @@ CORE_API int32 FindLastChar(FWideStringView View, WIDECHAR Search, ESearchCase::
  * @param OutMatchIndex Receives the index of the search character that matched.
  * @return The position at which any search character was found, or INDEX_NONE if not found.
  */
+CORE_API int32 FindFirstOfAnyChar(FUtf8StringView View, std::initializer_list<UTF8CHAR> Search, ESearchCase::Type SearchCase = ESearchCase::CaseSensitive, int32* OutMatchIndex = nullptr);
 CORE_API int32 FindFirstOfAnyChar(FUtf8StringView View, TConstArrayView<ANSICHAR> Search, ESearchCase::Type SearchCase = ESearchCase::CaseSensitive, int32* OutMatchIndex = nullptr);
 CORE_API int32 FindFirstOfAnyChar(FUtf8StringView View, TConstArrayView<UTF8CHAR> Search, ESearchCase::Type SearchCase = ESearchCase::CaseSensitive, int32* OutMatchIndex = nullptr);
 CORE_API int32 FindFirstOfAnyChar(FWideStringView View, TConstArrayView<WIDECHAR> Search, ESearchCase::Type SearchCase = ESearchCase::CaseSensitive, int32* OutMatchIndex = nullptr);
@@ -105,6 +104,7 @@ CORE_API int32 FindFirstOfAnyChar(FWideStringView View, TConstArrayView<WIDECHAR
  * @param OutMatchIndex Receives the index of the search character that matched.
  * @return The position at which any search character was found, or INDEX_NONE if not found.
  */
+CORE_API int32 FindLastOfAnyChar(FUtf8StringView View, std::initializer_list<UTF8CHAR> Search, ESearchCase::Type SearchCase = ESearchCase::CaseSensitive, int32* OutMatchIndex = nullptr);
 CORE_API int32 FindLastOfAnyChar(FUtf8StringView View, TConstArrayView<ANSICHAR> Search, ESearchCase::Type SearchCase = ESearchCase::CaseSensitive, int32* OutMatchIndex = nullptr);
 CORE_API int32 FindLastOfAnyChar(FUtf8StringView View, TConstArrayView<UTF8CHAR> Search, ESearchCase::Type SearchCase = ESearchCase::CaseSensitive, int32* OutMatchIndex = nullptr);
 CORE_API int32 FindLastOfAnyChar(FWideStringView View, TConstArrayView<WIDECHAR> Search, ESearchCase::Type SearchCase = ESearchCase::CaseSensitive, int32* OutMatchIndex = nullptr);
