Deprecated UE_UTF8TEXT_IS_CONSTEXPR.

diff --git a/Engine/Source/Runtime/Core/Private/Misc/SourceLocationUtils.cpp b/Engine/Source/Runtime/Core/Private/Misc/SourceLocationUtils.cpp
--- a/Engine/Source/Runtime/Core/Private/Misc/SourceLocationUtils.cpp
+++ b/Engine/Source/Runtime/Core/Private/Misc/SourceLocationUtils.cpp
@@ -89,9 +89,7 @@ namespace UE::SourceLocation
 	void FFullAdapter::ToString(FUtf8String& Out) const
 	{
 		Out.Reset();
-		// Using ANSICHAR format string - FUtf8String::Appendf's use of TCheckedFormatString
-		// fails if UE_UTF8TEXT_IS_CONSTEXPR == 0 and UTF8TEXT is used.
-		Private::AppendFull<ANSICHAR>(Out, Location);
+		Private::AppendFull<UTF8CHAR>(Out, Location);
 	}
 
 	void FFullAdapter::AppendString(FWideStringBuilderBase& Out) const
@@ -127,9 +125,7 @@ namespace UE::SourceLocation
 	void FFileAndLineAdapter::ToString(FUtf8String& Out) const
 	{
 		Out.Reset();
-		// Using ANSICHAR format string - FUtf8String::Appendf's use of TCheckedFormatString
-		// fails if UE_UTF8TEXT_IS_CONSTEXPR == 0 and UTF8TEXT is used.
-		Private::AppendFileAndLine<ANSICHAR>(Out, Location);
+		Private::AppendFileAndLine<UTF8CHAR>(Out, Location);
 	}
 
 	void FFileAndLineAdapter::AppendString(FWideStringBuilderBase& Out) const
diff --git a/Engine/Source/Runtime/Core/Private/ProfilingDebugging/MemoryTrace.cpp b/Engine/Source/Runtime/Core/Private/ProfilingDebugging/MemoryTrace.cpp
--- a/Engine/Source/Runtime/Core/Private/ProfilingDebugging/MemoryTrace.cpp
+++ b/Engine/Source/Runtime/Core/Private/ProfilingDebugging/MemoryTrace.cpp
@@ -314,7 +314,7 @@ static EMemoryTraceInit MemoryTrace_ShouldEnable(int32 ArgC, const ArgCharType*
 		// the multiple arguments are squashed into one single argument. Due to this we also need
 		// to parse each argument as if it contains multiple sub arguments.
 		int32 ArgStart;
-		while (ArgView.FindChar('-', ArgStart))
+		while (ArgView.FindChar(CHARTEXT(ArgCharType, '-'), ArgStart))
 		{
 			ArgView.RightChopInline(ArgStart);
 
diff --git a/Engine/Source/Runtime/Core/Public/HAL/Platform.h b/Engine/Source/Runtime/Core/Public/HAL/Platform.h
--- a/Engine/Source/Runtime/Core/Public/HAL/Platform.h
+++ b/Engine/Source/Runtime/Core/Public/HAL/Platform.h
@@ -1282,7 +1282,7 @@ namespace TypeTests
 	#define TEXT(x) TEXT_PASTE(x)
 #endif
 
-#define UE_UTF8TEXT_IS_CONSTEXPR 1
+#define UE_UTF8TEXT_IS_CONSTEXPR 1 UE_DEPRECATED_MACRO(5.7, "UE_UTF8TEXT_IS_CONSTEXPR has been deprecated.")
 
 namespace UE::Core::Private
 {
diff --git a/Engine/Source/Runtime/Core/Tests/String/FormatStringSanTest.cpp b/Engine/Source/Runtime/Core/Tests/String/FormatStringSanTest.cpp
--- a/Engine/Source/Runtime/Core/Tests/String/FormatStringSanTest.cpp
+++ b/Engine/Source/Runtime/Core/Tests/String/FormatStringSanTest.cpp
@@ -2,8 +2,7 @@
 
 #include "CoreTypes.h"
 
-// Tests need UE_UTF8TEXT_IS_CONSTEXPR because of the use of STATIC_CHECK with UTFTEXT.
-#if WITH_TESTS && UE_UTF8TEXT_IS_CONSTEXPR
+#if WITH_TESTS
 
 #include "String/FormatStringSan.h"
 #include "Containers/AnsiString.h"
