Static analysis fix (false positive?).

#rnx

diff --git a/Engine/Source/Runtime/Core/Private/ProfilingDebugging/MemoryTrace.cpp b/Engine/Source/Runtime/Core/Private/ProfilingDebugging/MemoryTrace.cpp
--- a/Engine/Source/Runtime/Core/Private/ProfilingDebugging/MemoryTrace.cpp
+++ b/Engine/Source/Runtime/Core/Private/ProfilingDebugging/MemoryTrace.cpp
@@ -313,7 +313,7 @@ static EMemoryTraceInit MemoryTrace_ShouldEnable(int32 ArgC, const ArgCharType*
 		// In certain scenarios (for example when launching a packaged title on some platforms)
 		// the multiple arguments are squashed into one single argument. Due to this we also need
 		// to parse each argument as if it contains multiple sub arguments.
-		int32 ArgStart;
+		int32 ArgStart = -1;
 		while (ArgView.FindChar(CHARTEXT(ArgCharType, '-'), ArgStart))
 		{
 			ArgView.RightChopInline(ArgStart);
