iOS: Add a check to determime if running in Designed for iPad mode.

diff --git a/Engine/Source/Runtime/Core/Private/IOS/IOSPlatformMisc.cpp b/Engine/Source/Runtime/Core/Private/IOS/IOSPlatformMisc.cpp
--- a/Engine/Source/Runtime/Core/Private/IOS/IOSPlatformMisc.cpp
+++ b/Engine/Source/Runtime/Core/Private/IOS/IOSPlatformMisc.cpp
@@ -410,6 +410,26 @@ bool FIOSPlatformMisc::IsInLowPowerMode()
     return false;
 }
 
+bool FIOSPlatformMisc::IsDesignedForIpadOnVisionOS()
+{
+	static BOOL bIsVision;
+	static dispatch_once_t onceToken;
+	dispatch_once(&onceToken, ^{
+		bIsVision = (NSClassFromString(@"UIWindowSceneGeometryPreferencesVision") != nil);
+	});
+	return bIsVision;
+}
+
+bool FIOSPlatformMisc:: IsDesignedForIpadOnMacOS()
+{
+	static BOOL bIsMac;
+	static dispatch_once_t onceToken;
+	dispatch_once(&onceToken, ^{
+		bIsMac = (NSClassFromString(@"UIWindowSceneGeometryPreferencesMac") != nil);
+	});
+	return bIsMac;
+}
+
 
 #if !PLATFORM_TVOS && !PLATFORM_VISIONOS
 UIInterfaceOrientationMask GetUIInterfaceOrientationMask(EDeviceScreenOrientation ScreenOrientation)
diff --git a/Engine/Source/Runtime/Core/Public/IOS/IOSPlatformMisc.h b/Engine/Source/Runtime/Core/Public/IOS/IOSPlatformMisc.h
--- a/Engine/Source/Runtime/Core/Public/IOS/IOSPlatformMisc.h
+++ b/Engine/Source/Runtime/Core/Public/IOS/IOSPlatformMisc.h
@@ -66,6 +66,11 @@ struct CORE_API FIOSPlatformMisc : public FApplePlatformMisc
 	static bool SupportsBrightness() { return true; }
     static bool IsInLowPowerMode();
 
+	// These two functions can be used to see if we're in designed for ipad mode explicltly
+	// on vision os or mac. This check may not be valid in the future, but will do the trick for now
+	static bool IsDesignedForIpadOnVisionOS();
+	static bool IsDesignedForIpadOnMacOS();
+
 	//////// Notifications
 	static void RegisterForRemoteNotifications();
 	static bool IsRegisteredForRemoteNotifications();
