Outline more methods in TSparseArray and TSet

diff --git a/Engine/Source/Runtime/Core/Public/Containers/Set.h b/Engine/Source/Runtime/Core/Public/Containers/Set.h
--- a/Engine/Source/Runtime/Core/Public/Containers/Set.h
+++ b/Engine/Source/Runtime/Core/Public/Containers/Set.h
@@ -129,6 +129,39 @@ public:
 namespace UE::Core::Private
 {
 	[[noreturn]] CORE_API void OnInvalidSetNum(unsigned long long NewNum);
+
+	template<typename HashType>
+	void CopyHash(HashType& Hash, int32& HashSize, const HashType& Copy, int32 HashSizeCopy)
+	{
+		DestructItems((FSetElementId*)Hash.GetAllocation(), HashSize);
+		Hash.ResizeAllocation(0, HashSizeCopy, sizeof(FSetElementId));
+		ConstructItems<FSetElementId>(Hash.GetAllocation(), (FSetElementId*)Copy.GetAllocation(), HashSizeCopy);
+		HashSize = HashSizeCopy;
+	}
+
+	template<typename HashType>
+	[[nodiscard]] FORCEINLINE FSetElementId& GetTypedHash(HashType& Hash, int32 HashIndex, int32 HashSize)
+	{
+		return ((FSetElementId*)Hash.GetAllocation())[HashIndex & (HashSize - 1)];
+	}
+
+	template<typename HashType>
+	void Rehash(HashType& Hash, int32 HashSize)
+	{
+		// Free the old hash.
+		Hash.ResizeAllocation(0, 0, sizeof(FSetElementId));
+
+		if (HashSize)
+		{
+			// Allocate the new hash.
+			checkSlow(FMath::IsPowerOfTwo(HashSize));
+			Hash.ResizeAllocation(0, HashSize, sizeof(FSetElementId));
+			for (int32 HashIndex = 0; HashIndex < HashSize; ++HashIndex)
+			{
+				GetTypedHash(Hash, HashIndex, HashSize) = FSetElementId();
+			}
+		}
+	}
 }
 
 /**
@@ -224,12 +257,7 @@ public:
 	{
 		if (this != &Copy)
 		{
-			int32 CopyHashSize = Copy.HashSize;
-
-			DestructItems((FSetElementId*)Hash.GetAllocation(), HashSize);
-			Hash.ResizeAllocation(0, CopyHashSize, sizeof(FSetElementId));
-			ConstructItems<FSetElementId>(Hash.GetAllocation(), (FSetElementId*)Copy.Hash.GetAllocation(), CopyHashSize);
-			HashSize = CopyHashSize;
+			UE::Core::Private::CopyHash(Hash, HashSize, Copy.Hash, Copy.HashSize);
 
 			Elements = Copy.Elements;
 		}
@@ -1435,20 +1463,10 @@ private:
 	/** Resizes the hash. */
 	void Rehash()
 	{
-		// Free the old hash.
-		Hash.ResizeAllocation(0,0,sizeof(FSetElementId));
+		UE::Core::Private::Rehash(Hash, HashSize);
 
-		int32 LocalHashSize = HashSize;
-		if (LocalHashSize)
+		if (HashSize)
 		{
-			// Allocate the new hash.
-			checkSlow(FMath::IsPowerOfTwo(HashSize));
-			Hash.ResizeAllocation(0, LocalHashSize, sizeof(FSetElementId));
-			for (int32 HashIndex = 0; HashIndex < LocalHashSize; ++HashIndex)
-			{
-				GetTypedHash(HashIndex) = FSetElementId();
-			}
-
 			// Add the existing elements to the new hash.
 			for(typename ElementArrayType::TConstIterator ElementIt(Elements);ElementIt;++ElementIt)
 			{
diff --git a/Engine/Source/Runtime/Core/Public/Containers/SparseArray.h b/Engine/Source/Runtime/Core/Public/Containers/SparseArray.h
--- a/Engine/Source/Runtime/Core/Public/Containers/SparseArray.h
+++ b/Engine/Source/Runtime/Core/Public/Containers/SparseArray.h
@@ -60,7 +60,8 @@
 
 DECLARE_TEMPLATE_INTRINSIC_TYPE_LAYOUT(template<typename ElementType>, TSparseArrayElementOrFreeListLink<ElementType>);
 
-template<size_t SizeOfElementType, size_t AlignOfElementType, typename Allocator> class TSparseArrayBase
+template<size_t SizeOfElementType, size_t AlignOfElementType, typename Allocator>
+class TSparseArrayBase
 {
 protected:
 	[[nodiscard]] explicit TSparseArrayBase(FIntrusiveUnsetOptionalState Tag)
@@ -201,6 +202,252 @@
 		}
 	}
 
+	/**
+	 * Preallocates enough memory to contain the specified number of elements.
+	 *
+	 * @param	ExpectedNumElements		the total number of elements that the array will have
+	 */
+	void Reserve(int32 ExpectedNumElements)
+	{
+		if (ExpectedNumElements > Data.Num())
+		{
+			const int32 ElementsToAdd = ExpectedNumElements - Data.Num();
+
+			// allocate memory in the array itself - reserve first to ensure that resizing-for-growth doesn't occur
+			Data.Reserve(ExpectedNumElements);
+			int32 ElementIndex = Data.AddUninitialized(ElementsToAdd);
+
+			FElementOrFreeListLink* DataPtr = (FElementOrFreeListLink*)Data.GetData();
+
+			// now mark the new elements as free
+			for (int32 FreeIndex = ExpectedNumElements - 1; FreeIndex >= ElementIndex; --FreeIndex)
+			{
+				if (NumFreeIndices)
+				{
+					DataPtr[FirstFreeIndex].PrevFreeIndex = FreeIndex;
+				}
+				DataPtr[FreeIndex].PrevFreeIndex = -1;
+				DataPtr[FreeIndex].NextFreeIndex = NumFreeIndices > 0 ? FirstFreeIndex : INDEX_NONE;
+				FirstFreeIndex = FreeIndex;
+				++NumFreeIndices;
+			}
+
+			if (ElementsToAdd == ExpectedNumElements)
+			{
+				AllocationFlags.Init(false, ElementsToAdd);
+			}
+			else
+			{
+				AllocationFlags.Add(false, ElementsToAdd);
+			}
+		}
+	}
+
+	/** Shrinks the array's storage to avoid slack. */
+	void Shrink()
+	{
+		// Determine the highest allocated index in the data array.
+		int32 MaxAllocatedIndex = AllocationFlags.FindLast(true);
+
+		const int32 FirstIndexToRemove = MaxAllocatedIndex + 1;
+		if (FirstIndexToRemove < Data.Num())
+		{
+			if (NumFreeIndices > 0)
+			{
+				FElementOrFreeListLink* DataPtr = (FElementOrFreeListLink*)Data.GetData();
+
+				// Look for elements in the free list that are in the memory to be freed.
+				int32 FreeIndex = FirstFreeIndex;
+				while (FreeIndex != INDEX_NONE)
+				{
+					if (FreeIndex >= FirstIndexToRemove)
+					{
+						const int32 PrevFreeIndex = DataPtr[FreeIndex].PrevFreeIndex;
+						const int32 NextFreeIndex = DataPtr[FreeIndex].NextFreeIndex;
+						if (NextFreeIndex != -1)
+						{
+							DataPtr[NextFreeIndex].PrevFreeIndex = PrevFreeIndex;
+						}
+						if (PrevFreeIndex != -1)
+						{
+							DataPtr[PrevFreeIndex].NextFreeIndex = NextFreeIndex;
+						}
+						else
+						{
+							FirstFreeIndex = NextFreeIndex;
+						}
+						--NumFreeIndices;
+
+						FreeIndex = NextFreeIndex;
+					}
+					else
+					{
+						FreeIndex = DataPtr[FreeIndex].NextFreeIndex;
+					}
+				}
+			}
+
+			// Truncate unallocated elements at the end of the data array.
+			Data.RemoveAt(FirstIndexToRemove, Data.Num() - FirstIndexToRemove, EAllowShrinking::No);
+			AllocationFlags.RemoveAt(FirstIndexToRemove, AllocationFlags.Num() - FirstIndexToRemove);
+		}
+
+		// Shrink the data array.
+		Data.Shrink();
+	}
+
+	/**
+	* Sort the free element list so that subsequent allocations will occur in the lowest available
+	* position resulting in tighter packing without moving any existing items. This also means
+	* that assigned indices no longer depend on the order in which old items were removed, making
+	* it easier to use the container when determinism is required (without a container reset).
+	*
+	* E.g., call SortFreeList() each frame to make the container assign the same indices when we
+	* perform the following operations:
+	*	Frame1: Add(A) -> [0]; Add(B) -> [1]; Remove(A); Remove(B); (Free list is now {[1],[0],...})
+	*	Frame2: Add(A) -> [1]; Add(B) -> [0]; Remove(A); Remove(B); (Free list is now {[0],[1],...})
+	*
+	* NOTE: This is operation is currently O(N) with N = GetMaxIndex(). This could be improved for
+	* large mostly-full arrays if necessary.
+	*/
+	void SortFreeList()
+	{
+		FElementOrFreeListLink* DataPtr = (FElementOrFreeListLink*)Data.GetData();
+		int32 CurrentHeadIndex = INDEX_NONE;
+		int32 NumFreeIndicesProcessed = 0;
+
+		// Reverse iteration to build the list from low to high indices
+		for (int32 Index = Data.Num() - 1; NumFreeIndicesProcessed < NumFreeIndices; --Index)
+		{
+			// If we have an unused element, add it to the free list
+			if (!IsValidIndex(Index))
+			{
+				DataPtr[Index].PrevFreeIndex = INDEX_NONE;
+				DataPtr[Index].NextFreeIndex = INDEX_NONE;
+
+				if (CurrentHeadIndex != INDEX_NONE)
+				{
+					DataPtr[CurrentHeadIndex].PrevFreeIndex = Index;
+					DataPtr[Index].NextFreeIndex = CurrentHeadIndex;
+				}
+
+				CurrentHeadIndex = Index;
+				++NumFreeIndicesProcessed;
+			}
+		}
+
+		// Set up the new head
+		FirstFreeIndex = CurrentHeadIndex;
+	}
+
+	/**
+	 * Allocates space for an element in the array at a given index.  The element is not initialized, and you must use the corresponding placement new operator
+	 * to construct the element in the allocated memory.
+	 */
+	FSparseArrayAllocationInfo InsertUninitialized(int32 Index)
+	{
+		FElementOrFreeListLink* DataPtr;
+
+		// Enlarge the array to include the given index.
+		if (Index >= Data.Num())
+		{
+			Data.AddUninitialized(Index + 1 - Data.Num());
+
+			// Defer getting the data pointer until after a possible reallocation
+			DataPtr = (FElementOrFreeListLink*)Data.GetData();
+
+			while (AllocationFlags.Num() < Data.Num())
+			{
+				const int32 FreeIndex = AllocationFlags.Num();
+				DataPtr[FreeIndex].PrevFreeIndex = -1;
+				DataPtr[FreeIndex].NextFreeIndex = FirstFreeIndex;
+				if (NumFreeIndices)
+				{
+					DataPtr[FirstFreeIndex].PrevFreeIndex = FreeIndex;
+				}
+				FirstFreeIndex = FreeIndex;
+				verify(AllocationFlags.Add(false) == FreeIndex);
+				++NumFreeIndices;
+			};
+		}
+		else
+		{
+			DataPtr = (FElementOrFreeListLink*)Data.GetData();
+		}
+
+		// Verify that the specified index is free.
+		check(!AllocationFlags[Index]);
+
+		// Remove the index from the list of free elements.
+		--NumFreeIndices;
+		const int32 PrevFreeIndex = DataPtr[Index].PrevFreeIndex;
+		const int32 NextFreeIndex = DataPtr[Index].NextFreeIndex;
+		if (PrevFreeIndex != -1)
+		{
+			DataPtr[PrevFreeIndex].NextFreeIndex = NextFreeIndex;
+		}
+		else
+		{
+			FirstFreeIndex = NextFreeIndex;
+		}
+		if (NextFreeIndex != -1)
+		{
+			DataPtr[NextFreeIndex].PrevFreeIndex = PrevFreeIndex;
+		}
+
+		return AllocateIndex(Index);
+	}
+
+	/** Compacts the allocated elements into a contiguous index range. */
+	/** Returns true if any elements were relocated, false otherwise. */
+	bool Compact()
+	{
+		int32 NumFree = NumFreeIndices;
+		if (NumFree == 0)
+		{
+			return false;
+		}
+
+		bool bResult = false;
+
+		FElementOrFreeListLink* DataPtr = (FElementOrFreeListLink*)Data.GetData();
+
+		int32 EndIndex    = Data.Num();
+		int32 TargetIndex = EndIndex - NumFree;
+		int32 FreeIndex   = FirstFreeIndex;
+		while (FreeIndex != -1)
+		{
+			int32 NextFreeIndex = DataPtr[FreeIndex].NextFreeIndex;
+			if (FreeIndex < TargetIndex)
+			{
+				// We need an element here
+				do
+				{
+					--EndIndex;
+				}
+				while (!AllocationFlags[EndIndex]);
+
+				RelocateConstructItems<FElementOrFreeListLink>(DataPtr + FreeIndex, DataPtr + EndIndex, 1);
+				AllocationFlags[FreeIndex] = true;
+
+				bResult = true;
+			}
+
+			FreeIndex = NextFreeIndex;
+		}
+
+		Data.RemoveAt(TargetIndex, NumFree, EAllowShrinking::No);
+		AllocationFlags.RemoveAt(TargetIndex, NumFree);
+
+		NumFreeIndices = 0;
+		FirstFreeIndex = -1;
+
+		// Shrink the data array.
+		Data.Shrink();
+
+		return bResult;
+	}
+
 	[[nodiscard]] int32 GetMaxIndex() const
 	{
 		return Data.Num();
@@ -222,6 +469,16 @@
 		return Data.Max();
 	}
 
+	[[nodiscard]] bool IsValidIndex(int32 Index) const
+	{
+		return AllocationFlags.IsValidIndex(Index) && AllocationFlags[Index];
+	}
+
+	[[nodiscard]] bool IsAllocated(int32 Index) const
+	{
+		return AllocationFlags[Index];
+	}
+
 	/**
 	 * The element type stored is only indirectly related to the element type requested, to avoid instantiating TArray redundantly for
 	 * compatible types.
@@ -367,63 +624,7 @@
 		return Allocation.Index;
 	}
 
-	/**
-	 * Allocates space for an element in the array at a given index.  The element is not initialized, and you must use the corresponding placement new operator
-	 * to construct the element in the allocated memory.
-	 */
-	FSparseArrayAllocationInfo InsertUninitialized(int32 Index)
-	{
-		FElementOrFreeListLink* DataPtr;
-
-		// Enlarge the array to include the given index.
-		if(Index >= Data.Num())
-		{
-			Data.AddUninitialized(Index + 1 - Data.Num());
-
-			// Defer getting the data pointer until after a possible reallocation
-			DataPtr = (FElementOrFreeListLink*)Data.GetData();
-
-			while(AllocationFlags.Num() < Data.Num())
-			{
-				const int32 FreeIndex = AllocationFlags.Num();
-				DataPtr[FreeIndex].PrevFreeIndex = -1;
-				DataPtr[FreeIndex].NextFreeIndex = FirstFreeIndex;
-				if(NumFreeIndices)
-				{
-					DataPtr[FirstFreeIndex].PrevFreeIndex = FreeIndex;
-				}
-				FirstFreeIndex = FreeIndex;
-				verify(AllocationFlags.Add(false) == FreeIndex);
-				++NumFreeIndices;
-			};
-		}
-		else
-		{
-			DataPtr = (FElementOrFreeListLink*)Data.GetData();
-		}
-
-		// Verify that the specified index is free.
-		check(!AllocationFlags[Index]);
-
-		// Remove the index from the list of free elements.
-		--NumFreeIndices;
-		const int32 PrevFreeIndex = DataPtr[Index].PrevFreeIndex;
-		const int32 NextFreeIndex = DataPtr[Index].NextFreeIndex;
-		if(PrevFreeIndex != -1)
-		{
-			DataPtr[PrevFreeIndex].NextFreeIndex = NextFreeIndex;
-		}
-		else
-		{
-			FirstFreeIndex = NextFreeIndex;
-		}
-		if(NextFreeIndex != -1)
-		{
-			DataPtr[NextFreeIndex].PrevFreeIndex = PrevFreeIndex;
-		}
-
-		return AllocateIndex(Index);
-	}
+	using SuperType::InsertUninitialized;
 
 	/**
 	 * Inserts an element to the array.
@@ -438,11 +639,7 @@
 	{
 		if constexpr (!std::is_trivially_destructible_v<ElementType>)
 		{
-			FElementOrFreeListLink* DataPtr = (FElementOrFreeListLink*)Data.GetData();
-			for (int32 It = Index, ItCount = Count; ItCount; ++It, --ItCount)
-			{
-				((ElementType&)DataPtr[It].ElementData).~ElementType();
-			}
+			DestroyElements(Index, Count);
 		}
 
 		RemoveAtUninitialized(Index, Count);
@@ -459,11 +656,7 @@
 		// Destruct the allocated elements.
 		if constexpr (!std::is_trivially_destructible_v<ElementType>)
 		{
-			for (TIterator It(*this); It; ++It)
-			{
-				ElementType& Element = *It;
-				Element.~ElementType();
-			}
+			DestroyElements();
 		}
 
 		// Free the allocated elements.
@@ -479,11 +672,7 @@
 		// Destruct the allocated elements.
 		if constexpr (!std::is_trivially_destructible_v<ElementType>)
 		{
-			for (TIterator It(*this); It; ++It)
-			{
-				ElementType& Element = *It;
-				Element.~ElementType();
-			}
+			DestroyElements();
 		}
 
 		// Free the allocated elements.
@@ -493,149 +682,9 @@
 		AllocationFlags.Reset();
 	}
 
-	/**
-	 * Preallocates enough memory to contain the specified number of elements.
-	 *
-	 * @param	ExpectedNumElements		the total number of elements that the array will have
-	 */
-	void Reserve(int32 ExpectedNumElements)
-	{
-		if ( ExpectedNumElements > Data.Num() )
-		{
-			const int32 ElementsToAdd = ExpectedNumElements - Data.Num();
+	using SuperType::Reserve;
+	using SuperType::Shrink;
+	using SuperType::Compact;
 
-			// allocate memory in the array itself - reserve first to ensure that resizing-for-growth doesn't occur
-			Data.Reserve(ExpectedNumElements);
-			int32 ElementIndex = Data.AddUninitialized(ElementsToAdd);
-
-			FElementOrFreeListLink* DataPtr = (FElementOrFreeListLink*)Data.GetData();
-
-			// now mark the new elements as free
-			for ( int32 FreeIndex = ExpectedNumElements - 1; FreeIndex >= ElementIndex; --FreeIndex )
-			{
-				if(NumFreeIndices)
-				{
-					DataPtr[FirstFreeIndex].PrevFreeIndex = FreeIndex;
-				}
-				DataPtr[FreeIndex].PrevFreeIndex = -1;
-				DataPtr[FreeIndex].NextFreeIndex = NumFreeIndices > 0 ? FirstFreeIndex : INDEX_NONE;
-				FirstFreeIndex = FreeIndex;
-				++NumFreeIndices;
-			}
-
-			if (ElementsToAdd == ExpectedNumElements)
-			{
-				AllocationFlags.Init(false, ElementsToAdd);
-			}
-			else
-			{
-				AllocationFlags.Add(false, ElementsToAdd);
-			}
-		}
-	}
-
-	/** Shrinks the array's storage to avoid slack. */
-	void Shrink()
-	{
-		// Determine the highest allocated index in the data array.
-		int32 MaxAllocatedIndex = AllocationFlags.FindLast(true);
-
-		const int32 FirstIndexToRemove = MaxAllocatedIndex + 1;
-		if(FirstIndexToRemove < Data.Num())
-		{
-			if(NumFreeIndices > 0)
-			{
-				FElementOrFreeListLink* DataPtr = (FElementOrFreeListLink*)Data.GetData();
-
-				// Look for elements in the free list that are in the memory to be freed.
-				int32 FreeIndex = FirstFreeIndex;
-				while(FreeIndex != INDEX_NONE)
-				{
-					if(FreeIndex >= FirstIndexToRemove)
-					{
-						const int32 PrevFreeIndex = DataPtr[FreeIndex].PrevFreeIndex;
-						const int32 NextFreeIndex = DataPtr[FreeIndex].NextFreeIndex;
-						if(NextFreeIndex != -1)
-						{
-							DataPtr[NextFreeIndex].PrevFreeIndex = PrevFreeIndex;
-						}
-						if(PrevFreeIndex != -1)
-						{
-							DataPtr[PrevFreeIndex].NextFreeIndex = NextFreeIndex;
-						}
-						else
-						{
-							FirstFreeIndex = NextFreeIndex;
-						}
-						--NumFreeIndices;
-
-						FreeIndex = NextFreeIndex;
-					}
-					else
-					{
-						FreeIndex = DataPtr[FreeIndex].NextFreeIndex;
-					}
-				}
-			}
-
-			// Truncate unallocated elements at the end of the data array.
-			Data.RemoveAt(FirstIndexToRemove, Data.Num() - FirstIndexToRemove, EAllowShrinking::No);
-			AllocationFlags.RemoveAt(FirstIndexToRemove,AllocationFlags.Num() - FirstIndexToRemove);
-		}
-
-		// Shrink the data array.
-		Data.Shrink();
-	}
-
-	/** Compacts the allocated elements into a contiguous index range. */
-	/** Returns true if any elements were relocated, false otherwise. */
-	bool Compact()
-	{
-		int32 NumFree = NumFreeIndices;
-		if (NumFree == 0)
-		{
-			return false;
-		}
-
-		bool bResult = false;
-
-		FElementOrFreeListLink* DataPtr = (FElementOrFreeListLink*)Data.GetData();
-
-		int32 EndIndex    = Data.Num();
-		int32 TargetIndex = EndIndex - NumFree;
-		int32 FreeIndex   = FirstFreeIndex;
-		while (FreeIndex != -1)
-		{
-			int32 NextFreeIndex = DataPtr[FreeIndex].NextFreeIndex;
-			if (FreeIndex < TargetIndex)
-			{
-				// We need an element here
-				do
-				{
-					--EndIndex;
-				}
-				while (!AllocationFlags[EndIndex]);
-
-				RelocateConstructItems<FElementOrFreeListLink>(DataPtr + FreeIndex, DataPtr + EndIndex, 1);
-				AllocationFlags[FreeIndex] = true;
-
-				bResult = true;
-			}
-
-			FreeIndex = NextFreeIndex;
-		}
-
-		Data.RemoveAt(TargetIndex, NumFree, EAllowShrinking::No);
-		AllocationFlags.RemoveAt(TargetIndex, NumFree);
-
-		NumFreeIndices = 0;
-		FirstFreeIndex = -1;
-
-		// Shrink the data array.
-		Data.Shrink();
-
-		return bResult;
-	}
-
 	/** Compacts the allocated elements into a contiguous index range. Does not change the iteration order of the elements. */
 	/** Returns true if any elements were relocated, false otherwise. */
@@ -700,49 +749,7 @@
 		StableSort(TLess< ElementType >());
 	}
 
-	/**
-	* Sort the free element list so that subsequent allocations will occur in the lowest available 
-	* position resulting in tighter packing without moving any existing items. This also means 
-	* that assigned indices no longer depend on the order in which old items were removed, making
-	* it easier to use the container when determinism is required (without a container reset).
-	* 
-	* E.g., call SortFreeList() each frame to make the container assign the same indices when we 
-	* perform the following operations:
-	*	Frame1: Add(A) -> [0]; Add(B) -> [1]; Remove(A); Remove(B); (Free list is now {[1],[0],...})
-	*	Frame2: Add(A) -> [1]; Add(B) -> [0]; Remove(A); Remove(B); (Free list is now {[0],[1],...})
-	* 
-	* NOTE: This is operation is currently O(N) with N = GetMaxIndex(). This could be improved for
-	* large mostly-full arrays if necessary.
-	*/
-	void SortFreeList()
-	{
-		FElementOrFreeListLink* DataPtr = (FElementOrFreeListLink*)Data.GetData();
-		int32 CurrentHeadIndex = INDEX_NONE;
-		int32 NumFreeIndicesProcessed = 0;
-
-		// Reverse iteration to build the list from low to high indices
-		for (int32 Index = Data.Num() - 1; NumFreeIndicesProcessed < NumFreeIndices; --Index)
-		{
-			// If we have an unused element, add it to the free list
-			if (!IsValidIndex(Index))
-			{
-				DataPtr[Index].PrevFreeIndex = INDEX_NONE;
-				DataPtr[Index].NextFreeIndex = INDEX_NONE;
-
-				if (CurrentHeadIndex != INDEX_NONE)
-				{
-					DataPtr[CurrentHeadIndex].PrevFreeIndex = Index;
-					DataPtr[Index].NextFreeIndex = CurrentHeadIndex;
-				}
-
-				CurrentHeadIndex = Index;
-				++NumFreeIndicesProcessed;
-			}
-		}
-
-		// Set up the new head
-		FirstFreeIndex = CurrentHeadIndex;
-	}
+	using SuperType::SortFreeList;
 	
 	/**
 	 * Finds an item by predicate.
@@ -809,22 +816,22 @@
 	 */
 	[[nodiscard]] bool operator==(const TSparseArray& B) const
 	{
-		if(GetMaxIndex() != B.GetMaxIndex())
+		if (GetMaxIndex() != B.GetMaxIndex())
 		{
 			return false;
 		}
 
-		for(int32 ElementIndex = 0;ElementIndex < GetMaxIndex();ElementIndex++)
+		for (int32 ElementIndex = 0; ElementIndex < GetMaxIndex(); ElementIndex++)
 		{
 			const bool bIsAllocatedA = IsAllocated(ElementIndex);
 			const bool bIsAllocatedB = B.IsAllocated(ElementIndex);
-			if(bIsAllocatedA != bIsAllocatedB)
+			if (bIsAllocatedA != bIsAllocatedB)
 			{
 				return false;
 			}
-			else if(bIsAllocatedA)
+			else if (bIsAllocatedA)
 			{
-				if((*this)[ElementIndex] != B[ElementIndex])
+				if ((*this)[ElementIndex] != B[ElementIndex])
 				{
 					return false;
 				}
@@ -945,10 +952,7 @@
 		// Destruct the allocated elements.
 		if constexpr (!std::is_trivially_destructible_v<ElementType>)
 		{
-			for (ElementType& Element : ToArray)
-			{
-				DestructItem(&Element);
-			}
+			ToArray.DestroyElements();
 		}
 
 		ToArray.Data            = (DataType&&)FromArray.Data;
@@ -960,6 +964,24 @@
 		FromArray.NumFreeIndices = 0;
 	}
 
+	FORCEINLINE void DestroyElements()
+	{
+		DestroyElements(0, Data.Num());
+	}
+
+	FORCENOINLINE void DestroyElements(int32 Index, int32 Count)
+	{
+		const int32 EndIndex = Index + Count;
+		FElementOrFreeListLink* DataPtr = (FElementOrFreeListLink*)Data.GetData();
+		for (int32 i = Index; i < EndIndex; ++i)
+		{
+			if (IsValidIndex(i))
+			{
+				((ElementType&)DataPtr[i].ElementData).~ElementType();
+			}
+		}
+	}
+
 public:
 	// Accessors.
 	[[nodiscard]] ElementType& operator[](int32 Index)
@@ -981,15 +1003,8 @@
 		checkSlow(Index >= 0 && Index < Data.Num() && Index < AllocationFlags.Num() && AllocationFlags[Index]);
 		return Index;
 	}
-	[[nodiscard]] bool IsValidIndex(int32 Index) const
-	{
-		return AllocationFlags.IsValidIndex(Index) && AllocationFlags[Index];
-	}
-	[[nodiscard]] bool IsAllocated(int32 Index) const
-	{
-		return AllocationFlags[Index];
-	}
-
+	using SuperType::IsValidIndex;
+	using SuperType::IsAllocated;
 	using SuperType::GetMaxIndex;
 	using SuperType::IsEmpty;
 	using SuperType::Num;
