When creating console variable, LoadAllPlatformCVars is called and the config setting tracker registers all the settings for the first package being cooked.
To prevent this add FIgnoreScope in FConsoleManager::LoadAllPlatformCVars.

#incrementalcook
#rnx

diff --git a/Engine/Source/Runtime/Core/Private/HAL/ConsoleManager.cpp b/Engine/Source/Runtime/Core/Private/HAL/ConsoleManager.cpp
--- a/Engine/Source/Runtime/Core/Private/HAL/ConsoleManager.cpp
+++ b/Engine/Source/Runtime/Core/Private/HAL/ConsoleManager.cpp
@@ -3419,6 +3419,8 @@ void FConsoleManager::BatchUpdateTag(FName Tag, const TMap<FName, FString>& CVar
 
 void FConsoleManager::LoadAllPlatformCVars(FName PlatformName, const FString& DeviceProfileName)
 {
+	UE::ConfigAccessTracking::FIgnoreScope IgnoreScope;
+
 	FName PlatformKey = MakePlatformKey(PlatformName, DeviceProfileName);
 	
 	// protect the cached CVar info from two threads trying to get a platform CVar at once, and both attempting to load all of the cvars at the same time
