[PlatformFile]

Expose FFileJournalFileHandle ToString()

#rnx

diff --git a/Engine/Source/Runtime/Core/Public/GenericPlatform/GenericPlatformFile.h b/Engine/Source/Runtime/Core/Public/GenericPlatform/GenericPlatformFile.h
--- a/Engine/Source/Runtime/Core/Public/GenericPlatform/GenericPlatformFile.h
+++ b/Engine/Source/Runtime/Core/Public/GenericPlatform/GenericPlatformFile.h
@@ -254,7 +254,7 @@ struct FFileJournalFileHandle
 {
 	bool operator==(const FFileJournalFileHandle& Other) const;
 	bool operator!=(const FFileJournalFileHandle& Other) const;
-	FString ToString();
+	CORE_API FString ToString();
 
 	uint8 Bytes[20];
 };
