Move the #if SUBSTRATE_USE_FULLYSIMPLIFIED_MATERIAL shader statement to outside of the RegisteredOperators for loop because there's several circumstances where the loop is skipped and we write  relevant parameter entries outside of the loop too. This leads to a failed compilation as the if declaration is missing.

diff --git a/Engine/Source/Runtime/Core/Private/UObject/FortniteMainBranchObjectVersion.cpp b/Engine/Source/Runtime/Core/Private/UObject/FortniteMainBranchObjectVersion.cpp
--- a/Engine/Source/Runtime/Core/Private/UObject/FortniteMainBranchObjectVersion.cpp
+++ b/Engine/Source/Runtime/Core/Private/UObject/FortniteMainBranchObjectVersion.cpp
@@ -14,6 +14,6 @@ TMap<FGuid, FGuid> FFortniteMainBranchObjectVersion::GetSystemGuids()
 	SystemGuids.Add(DevGuids.Niagara_LatestScriptCompileVersion, FGuid("97F0FD7149D90D42B65D6557AF477D1A"));
 	SystemGuids.Add(DevGuids.SkeletalMeshDerivedDataVersion, FGuid("C998F2F3B2884A02B7E290B8C25835F0"));
 	SystemGuids.Add(DevGuids.STATICMESH_DERIVEDDATA_VER, FGuid("3ABF32149F9D4D838750368AB95573BE"));
-	SystemGuids.Add(DevGuids.MaterialTranslationDDCVersion, FGuid("05D49B4F487C48F29A0B6AB785F506AF"));
+	SystemGuids.Add(DevGuids.MaterialTranslationDDCVersion, FGuid("AC9EE1DC667C477C10F748FE00FD61FB"));
 	return SystemGuids;
 }
