Added UE_LIFETIMEBOUND to MakeValue and MakeError.

diff --git a/Engine/Source/Runtime/Core/Public/Templates/ValueOrError.h b/Engine/Source/Runtime/Core/Public/Templates/ValueOrError.h
--- a/Engine/Source/Runtime/Core/Public/Templates/ValueOrError.h
+++ b/Engine/Source/Runtime/Core/Public/Templates/ValueOrError.h
@@ -12,7 +12,7 @@
 template <typename... ArgTypes>
 struct TValueOrError_ValueProxy
 {
-	[[nodiscard]] TValueOrError_ValueProxy(ArgTypes&&... InArgs)
+	[[nodiscard]] explicit TValueOrError_ValueProxy(ArgTypes&&... InArgs UE_LIFETIMEBOUND)
 		: Args(Forward<ArgTypes>(InArgs)...)
 	{
 	}
@@ -23,7 +23,7 @@ struct TValueOrError_ValueProxy
 template <typename... ArgTypes>
 struct TValueOrError_ErrorProxy
 {
-	[[nodiscard]] TValueOrError_ErrorProxy(ArgTypes&&... InArgs)
+	[[nodiscard]] explicit TValueOrError_ErrorProxy(ArgTypes&&... InArgs UE_LIFETIMEBOUND)
 		: Args(Forward<ArgTypes>(InArgs)...)
 	{
 	}
@@ -32,13 +32,13 @@ struct TValueOrError_ErrorProxy
 };
 
 template <typename... ArgTypes>
-[[nodiscard]] UE_REWRITE TValueOrError_ValueProxy<ArgTypes...> MakeValue(ArgTypes&&... Args)
+[[nodiscard]] UE_REWRITE TValueOrError_ValueProxy<ArgTypes...> MakeValue(ArgTypes&&... Args UE_LIFETIMEBOUND)
 {
 	return TValueOrError_ValueProxy<ArgTypes...>(Forward<ArgTypes>(Args)...);
 }
 
 template <typename... ArgTypes>
-[[nodiscard]] UE_REWRITE TValueOrError_ErrorProxy<ArgTypes...> MakeError(ArgTypes&&... Args)
+[[nodiscard]] UE_REWRITE TValueOrError_ErrorProxy<ArgTypes...> MakeError(ArgTypes&&... Args UE_LIFETIMEBOUND)
 {
 	return TValueOrError_ErrorProxy<ArgTypes...>(Forward<ArgTypes>(Args)...);
 }
