IOS 26 dynamic linker crash fix

New version of the MLAssetsIO private framework, shipped with IOS 26, has a new static initializer, which is executed very early - earlier than your code.
As we provide our oeprator new/delete for dynamic linker (as exported symbols, by design) - those are used to allocate within that static initializer. However, our allocator initialization allocates, causing a re-entrant initalization. Hence, it crashes.

The initialization code, which allocates, is NSString regex replacement, which goes to another C++ library inside of ObjC runtime and also tries to allocate via our allocator. The fix is to do the parsing without ObjC and additional C++ allocations. This way allocator initializes just fine.

#tests Local/horde builds on real devices, IOS 18 & IOS 26 beta, IOS18 & IOS26 beta
#rnx

diff --git a/Engine/Source/Runtime/Core/Private/IOS/IOSPlatformMisc.cpp b/Engine/Source/Runtime/Core/Private/IOS/IOSPlatformMisc.cpp
--- a/Engine/Source/Runtime/Core/Private/IOS/IOSPlatformMisc.cpp
+++ b/Engine/Source/Runtime/Core/Private/IOS/IOSPlatformMisc.cpp
@@ -34,7 +34,9 @@
 #include "Misc/SecureHash.h"
 #include "Modules/ModuleManager.h"
 #include "Templates/Function.h"
+#include "Templates/UniquePtr.h"
 #include "Containers/StringConv.h"
+#include "Containers/StringView.h"
 #include "Apple/PreAppleSystemHeaders.h"
 
 #if !PLATFORM_TVOS && !PLATFORM_VISIONOS
@@ -1502,216 +1504,287 @@ FString FIOSPlatformMisc::GetPendingActivationProtocol()
 	return FString(ActivationProtocol);
 }
 
+namespace UE::FIOSPlatformMisc::Entitlements
+{
+
 // See for more information about the Blobs
 // https://opensource.apple.com/source/xnu/xnu-4570.61.1/osfmk/kern/cs_blobs.h.auto.html
 
 typedef struct __Blob {
-    uint32_t magic;
-    uint32_t length;
-    char data[];
+	uint32_t magic;
+	uint32_t length;
+	char data[];
 } CS_GenericBlob;
 
 typedef struct __BlobIndex {
-    uint32_t type;
-    uint32_t offset;
+	uint32_t type;
+	uint32_t offset;
 } CS_BlobIndex;
 
 typedef struct __MultiBlob {
-    uint32_t magic;
-    uint32_t length;
-    uint32_t count;
-    CS_BlobIndex index[];
+	uint32_t magic;
+	uint32_t length;
+	uint32_t count;
+	CS_BlobIndex index[];
 } CS_MultiBlob;
 
-extern NSString *EntitlementsData(void)
+struct FEntitlementsData
 {
-    // iterate through the headers to find the executable, since only the executable has the entitlements
-    const struct mach_header_64* executableHeader = nullptr;
-    char* ImageName = nullptr;
-    for (uint32_t i = 0; i < _dyld_image_count() && executableHeader == nullptr; i++)
-    {
-        const struct mach_header_64 *header = (struct mach_header_64 *)_dyld_get_image_header(i);
-        if (header->filetype == MH_EXECUTE)
-        {
-            ImageName = (char *)_dyld_get_image_name(i);
-            executableHeader = header;
-        }
-    }
+	FEntitlementsData() = default;
+	
+	FEntitlementsData(void* InData, uint32 InByteSize)
+		: Data(InData)
+		, ByteSize(InByteSize)
+	{
+	}
+	
+	struct FDeleteByFree
+	{
+		void operator()(void* Ptr) const
+		{
+			free(Ptr);
+		}
+	};
+	
+	TUniquePtr<void, FDeleteByFree> Data;
+	uint32 ByteSize = 0;
+};
 
-    if (executableHeader == nullptr)
-    {
-        return nil;
-    }
+static FEntitlementsData TryFindEntitlementsData()
+{
+	// iterate through the headers to find the executable, since only the executable has the entitlements
+	const struct mach_header_64* executableHeader = nullptr;
+	char* ImageName = nullptr;
+	for (uint32_t i = 0; i < _dyld_image_count() && executableHeader == nullptr; i++)
+	{
+		const struct mach_header_64 *header = (struct mach_header_64 *)_dyld_get_image_header(i);
+		if (header->filetype == MH_EXECUTE)
+		{
+			ImageName = (char *)_dyld_get_image_name(i);
+			executableHeader = header;
+		}
+	}
 
-    // verify that it's a 64bit app
-    if (executableHeader->magic != MH_MAGIC_64)
-    {
+	if (executableHeader == nullptr)
+	{
+		return {};
+	}
+
+	// verify that it's a 64bit app
+	if (executableHeader->magic != MH_MAGIC_64)
+	{
 		FPlatformMisc::LowLevelOutputDebugString(TEXT("Executable is NOT 64bit. Entitlement retrieval not supported.\n"));
-        return nil;
-    }
-    uintptr_t cursor = (uintptr_t)executableHeader + sizeof(struct mach_header_64);
-    const struct linkedit_data_command *segmentCommand = NULL;
+		return {};
+	}
+	uintptr_t cursor = (uintptr_t)executableHeader + sizeof(struct mach_header_64);
+	const struct linkedit_data_command *segmentCommand = NULL;
 
-    for (uint32_t i = 0; i < executableHeader->ncmds; i++, cursor += segmentCommand->cmdsize)
-    {
-        segmentCommand = (struct linkedit_data_command *)cursor;
-    
-        switch (segmentCommand->cmd)
-        {
-            case LC_CODE_SIGNATURE:
+	for (uint32_t i = 0; i < executableHeader->ncmds; i++, cursor += segmentCommand->cmdsize)
+	{
+		segmentCommand = (struct linkedit_data_command *)cursor;
+	
+		switch (segmentCommand->cmd)
+		{
+			case LC_CODE_SIGNATURE:
 				FPlatformMisc::LowLevelOutputDebugString(TEXT("LC_CODE_SIGNATURE found\n"));
-                break;
-            default:
-                continue;
-        }
+				break;
+			default:
+				continue;
+		}
 
-        const struct linkedit_data_command *dataCommand = (const struct linkedit_data_command *)segmentCommand;
+		const struct linkedit_data_command *dataCommand = (const struct linkedit_data_command *)segmentCommand;
 
-        FILE* file = fopen(ImageName, "rb");
-        if (file == NULL)
-        {
+		FILE* file = fopen(ImageName, "rb");
+		if (file == nullptr)
+		{
 			FPlatformMisc::LowLevelOutputDebugString(TEXT("Could not open binary file\n"));
-            return nil;
-        }
-        CS_MultiBlob multiBlob;
-        int isSuccess = fseek(file, UInt64(dataCommand->dataoff), SEEK_SET);
-        size_t Count = fread(&multiBlob, sizeof(CS_MultiBlob), 1, file);
-        uint32 multiBlobCount = ntohl(multiBlob.count);
-        CS_BlobIndex multiBlobIndex[multiBlobCount];
-        fread(&multiBlobIndex[0], sizeof(CS_BlobIndex) * multiBlobCount, 1, file);
-        
-        if (__builtin_bswap32(multiBlob.magic) != 0xfade0cc0)
-        {
-            fclose(file);
-            return nil;
-        }
-        
-        uint32 multiBlobSize = sizeof(multiBlob);
-        uint32 blobSize = sizeof(CS_GenericBlob);
+			return {};
+		}
+		CS_MultiBlob multiBlob;
+		int isSuccess = fseek(file, UInt64(dataCommand->dataoff), SEEK_SET);
+		size_t Count = fread(&multiBlob, sizeof(CS_MultiBlob), 1, file);
+		uint32 multiBlobCount = ntohl(multiBlob.count);
+		CS_BlobIndex multiBlobIndex[multiBlobCount];
+		fread(&multiBlobIndex[0], sizeof(CS_BlobIndex) * multiBlobCount, 1, file);
+		
+		if (__builtin_bswap32(multiBlob.magic) != 0xfade0cc0)
+		{
+			fclose(file);
+			return {};
+		}
+		
+		uint32 multiBlobSize = sizeof(multiBlob);
+		uint32 blobSize = sizeof(CS_GenericBlob);
 
-        for (int j = 0; j < multiBlobCount; j++)
-        {
-            uint32_t currentOffset = dataCommand->dataoff;
-            uint32_t blobOffset = ntohl(multiBlobIndex[j].offset);
-            CS_GenericBlob blob;
-            isSuccess = fseek(file, currentOffset + blobOffset, SEEK_SET);
-            Count = fread(&blob, sizeof(CS_GenericBlob), 1, file);
-            
-
-            if (__builtin_bswap32(blob.magic) == 0xfade7171)
-            {
-                uint32 blobLength = ntohl(blob.length);
-                char data[blobLength];
-                fread(&data[0], sizeof(char) * blobLength, 1, file);
-                NSString *stringFromData = [NSString stringWithFormat: @"%s", data];
-                NSLog(@"%@", stringFromData);
-                fclose(file);
-                return stringFromData;
-            }
-            else
-            {
-                continue;
-            }
-        }
-        fclose(file);
-    }
-    return nil;
+		for (int j = 0; j < multiBlobCount; j++)
+		{
+			uint32_t currentOffset = dataCommand->dataoff;
+			uint32_t blobOffset = ntohl(multiBlobIndex[j].offset);
+			CS_GenericBlob blob;
+			isSuccess = fseek(file, currentOffset + blobOffset, SEEK_SET);
+			Count = fread(&blob, sizeof(CS_GenericBlob), 1, file);
+
+			if (__builtin_bswap32(blob.magic) == 0xfade7171)
+			{
+				// Blob's length is magic+length+data
+				uint32 dataLength = ntohl(blob.length) - sizeof(uint32_t) * 2;
+				char* data = (char*)malloc(dataLength);
+				fread(&data[0], sizeof(char) * dataLength, 1, file);
+				NSLog(@"%s", data);
+				fclose(file);
+				return FEntitlementsData(data, dataLength);
+			}
+			else
+			{
+				continue;
+			}
+		}
+		fclose(file);
+	}
+	return {};
 }
 
-extern bool IsEntitlementPresentInEmbeddedProvision(const char *EntitlementsToFind)
+static bool IsPresentInEmbeddedProvision(const char *EntitlementsToFind)
 {
-    NSString* mobileprovisionPath = [[[NSBundle mainBundle] bundlePath] stringByAppendingPathComponent:@"embedded.mobileprovision"];
+	NSString* mobileprovisionPath = [[[NSBundle mainBundle] bundlePath] stringByAppendingPathComponent:@"embedded.mobileprovision"];
   
-    FILE* file = fopen([mobileprovisionPath cStringUsingEncoding:1],"rb");
+	FILE* file = fopen([mobileprovisionPath cStringUsingEncoding:1],"rb");
 
-    size_t entitlements_len = strlen(EntitlementsToFind);
-    
-    if (file == NULL)
-    {
-        NSLog(@"Mobile Provision not found");
-        return false;
-    }
+	size_t entitlements_len = strlen(EntitlementsToFind);
+	
+	if (file == NULL)
+	{
+		NSLog(@"Mobile Provision not found");
+		return false;
+	}
 
-    size_t readcount;
-    char buffer[1025]; //room for null
-    while(true)
-    {
-        readcount = fread(buffer, 1, 1024, file);
-        if (readcount <= 0)
-        {
-            break;
-        }
-        for(size_t i = 0; i < readcount; i++)
-        {
-            if(buffer[i] == 0)
-            {
-                buffer[i] = ' ';  // replace any null terminators that might be in the binary data
-            }
-        }
-        
-        buffer[readcount] = 0; // null terminate buffer!
-        char* entitlementskeyptr = strstr(buffer, EntitlementsToFind);
-        
-        if(entitlementskeyptr)
-        {
-            fseek(file, -readcount + entitlementskeyptr - buffer, SEEK_CUR); // seek to immediately after the entitlements key
-            readcount = fread(buffer, 1, 1024, file); // read 1024 bytes immediately after the entitlements key so we definitely get the value untruncated
-            // we don't expect any \0 characters between the key and the value, so we don't do the replace thing.
-            buffer[readcount] = 0; // null terminate buffer!
-            
-            // there could be keys following our one with their own true or false values that we could accidentally match if we just look for 'true' and return immediately.
-            char* trueptr = strstr(buffer, "true");
-            char* falseptr = strstr(buffer, "false");
-            
-            if(trueptr == NULL && falseptr == NULL)
-            {
+	size_t readcount;
+	char buffer[1025]; //room for null
+	while(true)
+	{
+		readcount = fread(buffer, 1, 1024, file);
+		if (readcount <= 0)
+		{
+			break;
+		}
+		for(size_t i = 0; i < readcount; i++)
+		{
+			if(buffer[i] == 0)
+			{
+				buffer[i] = ' ';  // replace any null terminators that might be in the binary data
+			}
+		}
+		
+		buffer[readcount] = 0; // null terminate buffer!
+		char* entitlementskeyptr = strstr(buffer, EntitlementsToFind);
+		
+		if(entitlementskeyptr)
+		{
+			fseek(file, -readcount + entitlementskeyptr - buffer, SEEK_CUR); // seek to immediately after the entitlements key
+			readcount = fread(buffer, 1, 1024, file); // read 1024 bytes immediately after the entitlements key so we definitely get the value untruncated
+			// we don't expect any \0 characters between the key and the value, so we don't do the replace thing.
+			buffer[readcount] = 0; // null terminate buffer!
+			
+			// there could be keys following our one with their own true or false values that we could accidentally match if we just look for 'true' and return immediately.
+			char* trueptr = strstr(buffer, "true");
+			char* falseptr = strstr(buffer, "false");
+			
+			if(trueptr == NULL && falseptr == NULL)
+			{
 				FPlatformMisc::LowLevelOutputDebugString(TEXT("Unexpected Behaviour. The entitlement key is found but its value is not set.\n"));
-                return false;
-            }
-            if(trueptr && falseptr == NULL)  // only true
-            {
+				return false;
+			}
+			if(trueptr && falseptr == NULL)  // only true
+			{
 				FPlatformMisc::LowLevelOutputDebugString(TEXT("Entitlements found in embedded mobile provision file.\n"));
-                return true;
-            }
-            if(trueptr == NULL && falseptr) // only false
-            {
+				return true;
+			}
+			if(trueptr == NULL && falseptr) // only false
+			{
 				FPlatformMisc::LowLevelOutputDebugString(TEXT("Entitlements found but set to false.\n"));
-                return false;
-            }
-            return (trueptr < falseptr); // return true if true comes before false
-        }
-        
-        if(readcount < 1024)
-        {
-            break; // end of file
-        }
-        
-        // seek back in case the entitlements was truncated.
-        fseek(file, -entitlements_len, SEEK_CUR);
-    }
-    fclose(file);
-    return false;
+				return false;
+			}
+			return (trueptr < falseptr); // return true if true comes before false
+		}
+		
+		if(readcount < 1024)
+		{
+			break; // end of file
+		}
+		
+		// seek back in case the entitlements was truncated.
+		fseek(file, -entitlements_len, SEEK_CUR);
+	}
+	fclose(file);
+	return false;
 }
 
-bool FIOSPlatformMisc::IsEntitlementEnabled(const char * EntitlementToCheck)
+static uint32 StripWhitespaceInline(ANSICHAR* Data, uint32 Length)
 {
-    NSString* TrueFlag = @"</key><true/>";
-    NSString* EntitlementsToFind = [NSString stringWithFormat: @"%s%@" , EntitlementToCheck, TrueFlag];
-    
-    NSString* CleanedEntitlementData = [EntitlementsData() stringByReplacingOccurrencesOfString: @"\\s+" withString: @"" options: NSRegularExpressionSearch range: NSMakeRange(0, EntitlementsToFind.length)];
+	check(Data);
+
+	ANSICHAR* Write = Data;
+	const ANSICHAR* Read    = Data;
+	const ANSICHAR* ReadEnd = Read + Length;
+	
+	while (Read != ReadEnd)
+	{
+		if (!TChar<ANSICHAR>::IsWhitespace(*Read))
+		{
+			*Write++ = *Read;
+		}
+		++Read;
+	}
+	
+	check(Read >= Write);
+	uint32 NewLength = Length - (Read - Write);
+	
+	return NewLength;
+}
 
-    CleanedEntitlementData = [[CleanedEntitlementData componentsSeparatedByCharactersInSet:[NSCharacterSet characterSetWithCharactersInString:@"\t"]] componentsJoinedByString:@""];
-    CleanedEntitlementData = [[CleanedEntitlementData componentsSeparatedByCharactersInSet:[NSCharacterSet newlineCharacterSet]] componentsJoinedByString:@""];
+static bool IsPresentInCodeSignature(const char* EntitlementToCheck)
+{
+	check(EntitlementToCheck);
+	
+	FEntitlementsData RawEntitlements = TryFindEntitlementsData();
+	
+	bool bFound = false;
 
-    if ([CleanedEntitlementData rangeOfString: (@"%s", EntitlementsToFind)].location == NSNotFound)
+	if (RawEntitlements.Data)
+	{
+		ANSICHAR* Data = reinterpret_cast<ANSICHAR*>(RawEntitlements.Data.Get());
+		uint32 Length = StripWhitespaceInline(Data, RawEntitlements.ByteSize);
+		
+		static_assert(sizeof(ANSICHAR) == sizeof(UTF8CHAR));
+		FUtf8StringView Entitlements = FUtf8StringView(reinterpret_cast<UTF8CHAR*>(Data), Length);
+		
+		const FUtf8StringView Needle = EntitlementToCheck;
+		const int32 Index = Entitlements.Find(Needle);
+		if (Index != INDEX_NONE)
+		{
+			Entitlements.RemovePrefix(Index + Needle.Len());
+			bFound = Entitlements.StartsWith("</key><true/>");
+		}
+	}
+	
+	return bFound;
+}
+
+} // namespace UE::FIOSPlatformMisc::Entitlements
+
+bool FIOSPlatformMisc::IsEntitlementEnabled(const char* EntitlementToCheck)
+{
+	check(EntitlementToCheck);
+	
+	bool bFound = UE::FIOSPlatformMisc::Entitlements::IsPresentInCodeSignature(EntitlementToCheck);
+	
+    if (!bFound)
     {
 		FPlatformMisc::LowLevelOutputDebugString(TEXT("Entitlements not found in binary Mach-O header. Looking at the embedded mobile provision file.\n"));
-        return IsEntitlementPresentInEmbeddedProvision(EntitlementToCheck);
-    }
-    else
-    {
-        return true;
+        return UE::FIOSPlatformMisc::Entitlements::IsPresentInEmbeddedProvision(EntitlementToCheck);
     }
+
+	return true;
 }
 
 void FIOSPlatformMisc::GetValidTargetPlatforms(TArray<FString>& TargetPlatformNames)
