Logging: Exclude trace state when log trace is disabled and optionally compile out source location

When both features are disabled this saves 17 bytes for each UE_LOG statement and 16 bytes for each UE_LOGFMT statement, plus the size of the file name strings if they were only used by logging.

diff --git a/Engine/Source/Runtime/Core/Private/AutoRTFM/AutoRTFMUE.cpp b/Engine/Source/Runtime/Core/Private/AutoRTFM/AutoRTFMUE.cpp
--- a/Engine/Source/Runtime/Core/Private/AutoRTFM/AutoRTFMUE.cpp
+++ b/Engine/Source/Runtime/Core/Private/AutoRTFM/AutoRTFMUE.cpp
@@ -289,7 +289,7 @@ void AutoRTFMLog(const char* File, int Line, void* ProgramCounter, autortfm_log_
 		/* File */ File,
 		/* Line */ Line,
 		/* Verbosity */ SeverityToVerbosity(Severity),
-		/* DynamicData */ DynamicData
+		/* DynamicData */ &DynamicData
 	);
 
 	switch (Severity)
@@ -323,7 +323,7 @@ void AutoRTFMLogWithCallstack(void* ProgramCounter, autortfm_log_severity Severi
 		/* File */ "",
 		/* Line */ 0,
 		/* Verbosity */ SeverityToVerbosity(Severity),
-		/* DynamicData */ DynamicData
+		/* DynamicData */ &DynamicData
 	);
 
 	switch (Severity)
diff --git a/Engine/Source/Runtime/Core/Public/Logging/LogMacros.h b/Engine/Source/Runtime/Core/Public/Logging/LogMacros.h
--- a/Engine/Source/Runtime/Core/Public/Logging/LogMacros.h
+++ b/Engine/Source/Runtime/Core/Public/Logging/LogMacros.h
@@ -20,6 +20,10 @@
 
 #include <type_traits>
 
+#ifndef UE_LOG_INCLUDE_SOURCE_LOCATION
+#define UE_LOG_INCLUDE_SOURCE_LOCATION 1
+#endif
+
 /*----------------------------------------------------------------------------
 	Logging
 ----------------------------------------------------------------------------*/
@@ -86,7 +90,7 @@ private:
 	#endif
 #endif
 
-#define UE_LOG_EXPAND_IS_FATAL(Verbosity, ActiveBlock, InactiveBlock) PREPROCESSOR_JOIN(UE_LOG_EXPAND_IS_FATAL_, Verbosity)(ActiveBlock, InactiveBlock)
+#define UE_LOG_EXPAND_IS_FATAL(Verbosity, ActiveBlock, InactiveBlock) UE_JOIN(UE_LOG_EXPAND_IS_FATAL_, Verbosity)(ActiveBlock, InactiveBlock)
 
 #define UE_LOG_EXPAND_IS_FATAL_Fatal(      ActiveBlock, InactiveBlock) ActiveBlock
 #define UE_LOG_EXPAND_IS_FATAL_Error(      ActiveBlock, InactiveBlock) InactiveBlock
@@ -110,17 +114,21 @@ namespace UE::Logging::Private
 /** Data about a static basic log that is created on-demand. */
 struct FStaticBasicLogDynamicData
 {
+#if LOGTRACE_ENABLED
 	std::atomic<bool> bInitializedTrace = false;
+#endif
 };
 
 /** Data about a static basic log that is constant for every occurrence. */
 struct FStaticBasicLogRecord
 {
 	const TCHAR* Format = nullptr;
-	const ANSICHAR* File = nullptr;
-	int32 Line = 0;
+	UE_IF(UE_LOG_INCLUDE_SOURCE_LOCATION,, inline static constexpr) const ANSICHAR* File = nullptr;
+	UE_IF(UE_LOG_INCLUDE_SOURCE_LOCATION,, inline static constexpr) int32 Line = 0;
 	ELogVerbosity::Type Verbosity = ELogVerbosity::Log;
+#if LOGTRACE_ENABLED
 	FStaticBasicLogDynamicData& DynamicData;
+#endif
 
 	// Workaround for https://developercommunity.visualstudio.com/t/Incorrect-warning-C4700-with-unrelated-s/10285950
 	constexpr FStaticBasicLogRecord(
@@ -128,12 +136,16 @@ struct FStaticBasicLogRecord
 		const ANSICHAR* InFile,
 		int32 InLine,
 		ELogVerbosity::Type InVerbosity,
-		FStaticBasicLogDynamicData& InDynamicData)
+		FStaticBasicLogDynamicData* InDynamicData)
 		: Format(InFormat)
+	#if UE_LOG_INCLUDE_SOURCE_LOCATION
 		, File(InFile)
 		, Line(InLine)
+	#endif
 		, Verbosity(InVerbosity)
-		, DynamicData(InDynamicData)
+	#if LOGTRACE_ENABLED
+		, DynamicData(*InDynamicData)
+	#endif
 	{
 	}
 };
@@ -238,7 +250,7 @@ CORE_API void BasicFatalLogWithProgramCounter(const FLogCategoryBase& Category,
 	 * @param Format         Format string literal in the style of printf.
 	 */
 	#define UE_LOG(CategoryName, Verbosity, Format, ...) \
-		UE_PRIVATE_LOG(PREPROCESSOR_NOTHING, constexpr, CategoryName, Verbosity, Format, ##__VA_ARGS__)
+		UE_PRIVATE_LOG(UE_EMPTY, constexpr, CategoryName, Verbosity, Format, ##__VA_ARGS__)
 
 	/**
 	 * DO NOT USE. A macro that logs a formatted message if the log category is active at the requested verbosity level.
@@ -250,7 +262,7 @@ CORE_API void BasicFatalLogWithProgramCounter(const FLogCategoryBase& Category,
 	 * @param Format        Format string literal in the style of printf.
 	 */
 	#define UE_LOG_REF(CategoryRef, Verbosity, Format, ...) \
-		UE_PRIVATE_LOG(PREPROCESSOR_NOTHING, PREPROCESSOR_NOTHING, CategoryRef, Verbosity, Format, ##__VA_ARGS__)
+		UE_PRIVATE_LOG(UE_EMPTY, UE_EMPTY, CategoryRef, Verbosity, Format, ##__VA_ARGS__)
 
 	// DO NOT USE. Use UE_LOG because this will be deprecated in a future release.
 	#define UE_LOG_CLINKAGE UE_LOG
@@ -274,9 +286,8 @@ CORE_API void BasicFatalLogWithProgramCounter(const FLogCategoryBase& Category,
 		static_assert(std::is_const_v<std::remove_reference_t<decltype(Format)>>, "Formatting string must be a const TCHAR array."); \
 		static_assert(TIsArrayOrRefOfTypeByPredicate<decltype(Format), TIsCharEncodingCompatibleWithTCHAR>::Value, "Formatting string must be a TCHAR array."); \
 		UE_VALIDATE_FORMAT_STRING(Format, ##__VA_ARGS__); \
-		static ::UE::Logging::Private::FStaticBasicLogDynamicData LOG_Dynamic; \
-		/* This variable can only be constexpr if the __builtin_FILE() and __builtin_LINE() intrinsic functions are constexpr - otherwise make it plain const */ \
-		static PREPROCESSOR_IF(PLATFORM_COMPILER_SUPPORTS_CONSTEXPR_BUILTIN_FILE_AND_LINE, constexpr, const) ::UE::Logging::Private::FStaticBasicLogRecord LOG_Static(Format, __builtin_FILE(), __builtin_LINE(), ::ELogVerbosity::Verbosity, LOG_Dynamic); \
+		UE_IF(LOGTRACE_ENABLED, static ::UE::Logging::Private::FStaticBasicLogDynamicData LOG_Dynamic, ); \
+		static constexpr ::UE::Logging::Private::FStaticBasicLogRecord LOG_Static(Format, __builtin_FILE(), __builtin_LINE(), ::ELogVerbosity::Verbosity, UE_IF(LOGTRACE_ENABLED, &LOG_Dynamic, nullptr)); \
 		static_assert((::ELogVerbosity::Verbosity & ::ELogVerbosity::VerbosityMask) < ::ELogVerbosity::NumVerbosity && ::ELogVerbosity::Verbosity > 0, "Verbosity must be constant and in range."); \
 		if constexpr ((::ELogVerbosity::Verbosity & ELogVerbosity::VerbosityMask) == ::ELogVerbosity::Fatal) \
 		{ \
diff --git a/Engine/Source/Runtime/Core/Public/Logging/StructuredLog.h b/Engine/Source/Runtime/Core/Public/Logging/StructuredLog.h
--- a/Engine/Source/Runtime/Core/Public/Logging/StructuredLog.h
+++ b/Engine/Source/Runtime/Core/Public/Logging/StructuredLog.h
@@ -328,15 +328,17 @@ namespace UE::Logging::Private
 struct FStaticLogDynamicData
 {
 	std::atomic<FLogTemplate*> Template = nullptr;
+#if LOGTRACE_ENABLED
 	std::atomic<bool> bInitializedTrace = false;
+#endif
 };
 
 /** Data about a static log that is constant for every occurrence. */
 struct FStaticLogRecord
 {
 	const TCHAR* Format = nullptr;
-	const ANSICHAR* File = nullptr;
-	int32 Line = 0;
+	UE_IF(UE_LOG_INCLUDE_SOURCE_LOCATION,, inline static constexpr) const ANSICHAR* File = nullptr;
+	UE_IF(UE_LOG_INCLUDE_SOURCE_LOCATION,, inline static constexpr) int32 Line = 0;
 	ELogVerbosity::Type Verbosity = ELogVerbosity::Log;
 	FStaticLogDynamicData& DynamicData;
 
@@ -348,8 +350,10 @@ struct FStaticLogRecord
 		ELogVerbosity::Type InVerbosity,
 		FStaticLogDynamicData& InDynamicData)
 		: Format(InFormat)
+	#if UE_LOG_INCLUDE_SOURCE_LOCATION
 		, File(InFile)
 		, Line(InLine)
+	#endif
 		, Verbosity(InVerbosity)
 		, DynamicData(InDynamicData)
 	{
@@ -362,8 +366,8 @@ struct FStaticLocalizedLogRecord
 	const TCHAR* TextNamespace = nullptr;
 	const TCHAR* TextKey = nullptr;
 	const TCHAR* Format = nullptr;
-	const ANSICHAR* File = nullptr;
-	int32 Line = 0;
+	UE_IF(UE_LOG_INCLUDE_SOURCE_LOCATION,, inline static constexpr) const ANSICHAR* File = nullptr;
+	UE_IF(UE_LOG_INCLUDE_SOURCE_LOCATION,, inline static constexpr) int32 Line = 0;
 	ELogVerbosity::Type Verbosity = ELogVerbosity::Log;
 	FStaticLogDynamicData& DynamicData;
 
@@ -379,8 +383,10 @@ struct FStaticLocalizedLogRecord
 		: TextNamespace(InTextNamespace)
 		, TextKey(InTextKey)
 		, Format(InFormat)
+	#if UE_LOG_INCLUDE_SOURCE_LOCATION
 		, File(InFile)
 		, Line(InLine)
+	#endif
 		, Verbosity(InVerbosity)
 		, DynamicData(InDynamicData)
 	{
@@ -595,7 +601,7 @@ struct FLogContext
 
 // This macro expands `Arg1, Arg2` to `UE_PRIVATE_LOGFMT_FIELD(Arg1), UE_PRIVATE_LOGFMT_FIELD(Arg2), ...`
 // This macro expands `("Name1", Arg1), ("Name2", Arg2)` to `UE_PRIVATE_LOGFMT_FIELD(("Name1", Arg1)), UE_PRIVATE_LOGFMT_FIELD(("Name2", Arg2)), ...
-#define UE_PRIVATE_LOGFMT_FIELDS(...) UE_PRIVATE_LOGFMT_CALL(PREPROCESSOR_JOIN(UE_PRIVATE_LOGFMT_FIELDS_, UE_PRIVATE_LOGFMT_COUNT(__VA_ARGS__)), (__VA_ARGS__))
+#define UE_PRIVATE_LOGFMT_FIELDS(...) UE_PRIVATE_LOGFMT_CALL(UE_JOIN(UE_PRIVATE_LOGFMT_FIELDS_, UE_PRIVATE_LOGFMT_COUNT(__VA_ARGS__)), (__VA_ARGS__))
 
 #define UE_PRIVATE_LOGFMT_FIELDS_0()
 #define UE_PRIVATE_LOGFMT_FIELDS_1(A)                 , UE_PRIVATE_LOGFMT_FIELD(A)
