Deprecated PLATFORM_COMPILER_SUPPORTS_CONSTEXPR_BUILTIN_FILE_AND_LINE

Cleaned up old deprecated macros in Platform.h.

diff --git a/Engine/Source/Programs/Shared/EpicGames.UHT/Tokenizer/UhtTokenReaderGeneralExtensions.cs b/Engine/Source/Programs/Shared/EpicGames.UHT/Tokenizer/UhtTokenReaderGeneralExtensions.cs
--- a/Engine/Source/Programs/Shared/EpicGames.UHT/Tokenizer/UhtTokenReaderGeneralExtensions.cs
+++ b/Engine/Source/Programs/Shared/EpicGames.UHT/Tokenizer/UhtTokenReaderGeneralExtensions.cs
@@ -271,16 +271,11 @@ namespace EpicGames.UHT.Tokenizer
 				{
 					tokenReader.ConsumeToken();
 				}
-				else if (token.IsIdentifier("UE_NODISCARD") || token.IsIdentifier("UE_NODISCARD_CTOR"))
+				else if (token.IsIdentifier("UE_NODISCARD_CTOR"))
 				{
 					tokenReader.ConsumeToken();
 					attributeAction?.Invoke("nodiscard");
 				}
-				else if (token.IsIdentifier("UE_NORETURN"))
-				{
-					tokenReader.ConsumeToken();
-					attributeAction?.Invoke("noreturn");
-				}
 				else if (token.IsIdentifier("UE_NO_UNIQUE_ADDRESS"))
 				{
 					tokenReader.ConsumeToken();
diff --git a/Engine/Source/Runtime/Core/Public/HAL/Platform.h b/Engine/Source/Runtime/Core/Public/HAL/Platform.h
--- a/Engine/Source/Runtime/Core/Public/HAL/Platform.h
+++ b/Engine/Source/Runtime/Core/Public/HAL/Platform.h
@@ -261,15 +261,11 @@
 #ifndef PLATFORM_COMPILER_HAS_TCHAR_WMAIN
 	#define PLATFORM_COMPILER_HAS_TCHAR_WMAIN 0
 #endif
-#define PLATFORM_COMPILER_HAS_DECLTYPE_AUTO 1 UE_DEPRECATED_MACRO(5.4, "PLATFORM_COMPILER_HAS_DECLTYPE_AUTO has been deprecated and should be replaced with 1.")
-#define PLATFORM_COMPILER_HAS_FOLD_EXPRESSIONS 1 UE_DEPRECATED_MACRO(5.4, "PLATFORM_COMPILER_HAS_FOLD_EXPRESSIONS has been deprecated and should be replaced with 1.")
 #ifndef PLATFORM_COMPILER_HAS_GENERATED_COMPARISON_OPERATORS
 	#define PLATFORM_COMPILER_HAS_GENERATED_COMPARISON_OPERATORS (__cplusplus >= 202002L)
 #endif
 // Feature test macro for constexpr __builtin_FILE() and __builtin_LINE()
-#ifndef PLATFORM_COMPILER_SUPPORTS_CONSTEXPR_BUILTIN_FILE_AND_LINE
-	#define PLATFORM_COMPILER_SUPPORTS_CONSTEXPR_BUILTIN_FILE_AND_LINE 1
-#endif
+#define PLATFORM_COMPILER_SUPPORTS_CONSTEXPR_BUILTIN_FILE_AND_LINE 1 UE_DEPRECATED_MACRO(5.7, "PLATFORM_COMPILER_SUPPORTS_CONSTEXPR_BUILTIN_FILE_AND_LINE has been deprecated and should be replaced with 1.")
 #ifndef PLATFORM_COMPILER_SUPPORTS_BUILTIN_BITCAST
 	#error PLATFORM_COMPILER_SUPPORTS_BUILTIN_BITCAST should have been defined by now
 #endif
@@ -579,10 +575,6 @@
 	#define	PLATFORM_USE_SHOWFLAGS_ALWAYS_BITFIELD				1
 #endif
 
-#ifndef PLATFORM_SUPPORTS_LANDSCAPE_VISUAL_MESH_LOD_STREAMING
-	#define PLATFORM_SUPPORTS_LANDSCAPE_VISUAL_MESH_LOD_STREAMING 0 UE_DEPRECATED_MACRO(5.1, "PLATFORM_SUPPORTS_LANDSCAPE_VISUAL_MESH_LOD_STREAMING has been deprecated and should be replaced with 0.")
-#endif
-
 #ifndef PLATFORM_USE_GENERIC_LAUNCH_IMPLEMENTATION
 	#define PLATFORM_USE_GENERIC_LAUNCH_IMPLEMENTATION			0
 #endif
@@ -742,25 +734,6 @@
 	#define UE_REWRITE UE_NODEBUG FORCEINLINE
 #endif
 
-/* Use before a function declaration to warn that callers should not ignore the return value */
-#ifdef __has_cpp_attribute
-	#if __has_cpp_attribute(nodiscard)
-		// Define this in when all UE_NODISCARD usage has been replaced and we are ready to deprecate
-		#ifdef _MSC_VER
-			#define UE_NODISCARD [[nodiscard]]
-			#ifndef __clang__
-				#pragma deprecated("UE_NODISCARD")
-			#endif
-		#else
-			#define UE_NODISCARD [[nodiscard]] UE_DEPRECATED_MACRO(5.4, "UE_NODISCARD has been deprecated and should be replaced with [[nodiscard]].")
-		#endif
-	#else
-		#error "Compiler is expected to support [[nodiscard]]"
-	#endif
-#else
-	#error "Compiler is expected to support [[nodiscard]]"
-#endif
-
 // Use before a constructor declaration to warn when an unnamed temporary object is ignored, e.g. FScopeLock(CS); instead of FScopeLock Lock(CS);
 // We can't just use UE_NODISCARD in this case because older compilers don't like [[nodiscard]] on constructors.
 #if !defined(UE_NODISCARD_CTOR) && defined(__has_cpp_attribute)
@@ -775,24 +748,6 @@
 	#define UE_NODISCARD_CTOR
 #endif
 
-/* Use before a function declaration to indicate that the function never returns */
-#ifdef __has_cpp_attribute
-	#if __has_cpp_attribute(noreturn)
-		#ifdef _MSC_VER
-			#define UE_NORETURN [[noreturn]]
-			#ifndef __clang__
-				#pragma deprecated("UE_NORETURN")
-			#endif
-		#else
-			#define UE_NORETURN [[noreturn]] UE_DEPRECATED_MACRO(5.4, "UE_NORETURN has been deprecated and should be replaced with [[noreturn]].")
-		#endif
-	#else
-		#error "Compiler is expected to support [[noreturn]]"
-	#endif
-#else
-	#error "Compiler is expected to support [[noreturn]]"
-#endif
-
 /* Use in front of initialized global variable to have link once-semantics
    When multiple definitions of the variable are encountered by the linker, the first is selected and the remainder are discarded */
 #ifdef _MSC_VER
@@ -922,7 +877,6 @@
 #ifndef ABSTRACT
 	#define ABSTRACT
 #endif
-#define CONSTEXPR constexpr UE_DEPRECATED_MACRO(5.4, "CONSTEXPR has been deprecated and should be replaced with constexpr.")
 #ifndef IN
 	#define IN
 #endif
