Replace some FORCEINLINE usage in Core/CoreUObject with UE_FORCEINLINE_HINT.

We cannot replace these with 'inline' as it heavily affects performance in targets not using PGO/LTO. UE_FORCEINLINE_HINT allows conditionally defining these based on the target to what is most optimal for size/speed.

Depends on CL 43555442 which adds UE_FORCEINLINE_HINT.

#rnx

diff --git a/Engine/Source/Runtime/Core/Public/Android/AndroidJava.h b/Engine/Source/Runtime/Core/Public/Android/AndroidJava.h
--- a/Engine/Source/Runtime/Core/Public/Android/AndroidJava.h
+++ b/Engine/Source/Runtime/Core/Public/Android/AndroidJava.h
@@ -37,7 +37,7 @@ public:
 	template<typename ReturnType>
 	ReturnType CallMethod(FJavaClassMethod Method, ...);
 
-	FORCEINLINE jobject GetJObject() const
+	UE_FORCEINLINE_HINT jobject GetJObject() const
 	{
 		return Object;
 	}
diff --git a/Engine/Source/Runtime/Core/Public/Android/AndroidPlatformMemory.h b/Engine/Source/Runtime/Core/Public/Android/AndroidPlatformMemory.h
--- a/Engine/Source/Runtime/Core/Public/Android/AndroidPlatformMemory.h
+++ b/Engine/Source/Runtime/Core/Public/Android/AndroidPlatformMemory.h
@@ -52,27 +52,27 @@ struct FAndroidPlatformMemory : public FGenericPlatformMemory
 		void Decommit(size_t InOffset, size_t InSize);
 		void FreeVirtual();
 
-		FORCEINLINE void CommitByPtr(void *InPtr, size_t InSize)
+		UE_FORCEINLINE_HINT void CommitByPtr(void *InPtr, size_t InSize)
 		{
 			Commit(size_t(((uint8*)InPtr) - ((uint8*)Ptr)), InSize);
 		}
 
-		FORCEINLINE void DecommitByPtr(void *InPtr, size_t InSize)
+		UE_FORCEINLINE_HINT void DecommitByPtr(void *InPtr, size_t InSize)
 		{
 			Decommit(size_t(((uint8*)InPtr) - ((uint8*)Ptr)), InSize);
 		}
 
-		FORCEINLINE void Commit()
+		UE_FORCEINLINE_HINT void Commit()
 		{
 			Commit(0, GetActualSize());
 		}
 
-		FORCEINLINE void Decommit()
+		UE_FORCEINLINE_HINT void Decommit()
 		{
 			Decommit(0, GetActualSize());
 		}
 
-		FORCEINLINE size_t GetActualSize() const
+		UE_FORCEINLINE_HINT size_t GetActualSize() const
 		{
 			return VMSizeDivVirtualSizeAlignment * GetVirtualSizeAlignment();
 		}
diff --git a/Engine/Source/Runtime/Core/Public/Android/AndroidPlatformMisc.h b/Engine/Source/Runtime/Core/Public/Android/AndroidPlatformMisc.h
--- a/Engine/Source/Runtime/Core/Public/Android/AndroidPlatformMisc.h
+++ b/Engine/Source/Runtime/Core/Public/Android/AndroidPlatformMisc.h
@@ -48,7 +48,7 @@ struct FAndroidMisc : public FGenericPlatformMisc
 	static CORE_API void SetCellularPreference(int32 Value);
 	static CORE_API int32 GetCellularPreference();
     
-	FORCEINLINE static int32 GetMaxPathLength()
+	UE_FORCEINLINE_HINT static int32 GetMaxPathLength()
 	{
 		return ANDROID_MAX_PATH;
 	}
@@ -243,7 +243,7 @@ public:
 	static CORE_API bool IsDebuggerPresent();
 #endif
 
-	FORCEINLINE static void MemoryBarrier()
+	UE_FORCEINLINE_HINT static void MemoryBarrier()
 	{
 		__sync_synchronize();
 	}
diff --git a/Engine/Source/Runtime/Core/Public/Android/AndroidPlatformProcess.h b/Engine/Source/Runtime/Core/Public/Android/AndroidPlatformProcess.h
--- a/Engine/Source/Runtime/Core/Public/Android/AndroidPlatformProcess.h
+++ b/Engine/Source/Runtime/Core/Public/Android/AndroidPlatformProcess.h
@@ -13,12 +13,12 @@ struct FProcHandle : public TProcHandle<void*, nullptr>
 {
 public:
 	/** Default constructor. */
-	FORCEINLINE FProcHandle()
+	UE_FORCEINLINE_HINT FProcHandle()
 		: TProcHandle()
 	{}
 
 	/** Initialization constructor. */
-	FORCEINLINE explicit FProcHandle( HandleType Other )
+	UE_FORCEINLINE_HINT explicit FProcHandle( HandleType Other )
 		: TProcHandle( Other )
 	{}
 };
diff --git a/Engine/Source/Runtime/Core/Public/Android/AndroidPlatformProperties.h b/Engine/Source/Runtime/Core/Public/Android/AndroidPlatformProperties.h
--- a/Engine/Source/Runtime/Core/Public/Android/AndroidPlatformProperties.h
+++ b/Engine/Source/Runtime/Core/Public/Android/AndroidPlatformProperties.h
@@ -20,117 +20,117 @@
 struct FAndroidPlatformProperties
 	: public FGenericPlatformProperties
 {
-	static constexpr FORCEINLINE bool HasEditorOnlyData()
+	static constexpr UE_FORCEINLINE_HINT bool HasEditorOnlyData()
 	{
 		return false;
 	}
 
-	static constexpr FORCEINLINE const char* PlatformName()
+	static constexpr UE_FORCEINLINE_HINT const char* PlatformName()
 	{
 		return "Android";
 	}
 
-	static constexpr FORCEINLINE const char* IniPlatformName()
+	static constexpr UE_FORCEINLINE_HINT const char* IniPlatformName()
 	{
 		return "Android";
 	}
 
-	static constexpr FORCEINLINE const TCHAR* GetRuntimeSettingsClassName()
+	static constexpr UE_FORCEINLINE_HINT const TCHAR* GetRuntimeSettingsClassName()
 	{
 		return TEXT("/Script/AndroidRuntimeSettings.AndroidRuntimeSettings");
 	}
 
-	static constexpr FORCEINLINE bool IsGameOnly()
+	static constexpr UE_FORCEINLINE_HINT bool IsGameOnly()
 	{
 		return true;
 	}
 
-	static constexpr FORCEINLINE bool RequiresCookedData()
+	static constexpr UE_FORCEINLINE_HINT bool RequiresCookedData()
 	{
 		return true;
 	}
 
-	static FORCEINLINE bool SupportsBuildTarget(EBuildTargetType TargetType)
+	static UE_FORCEINLINE_HINT bool SupportsBuildTarget(EBuildTargetType TargetType)
 	{
 		return (TargetType == EBuildTargetType::Game);
 	}
 
-	static constexpr FORCEINLINE bool SupportsAutoSDK()
+	static constexpr UE_FORCEINLINE_HINT bool SupportsAutoSDK()
 	{
 		return true;
 	}
 
-	static constexpr FORCEINLINE bool SupportsHighQualityLightmaps()
+	static constexpr UE_FORCEINLINE_HINT bool SupportsHighQualityLightmaps()
 	{
 		return true; // always true because of Vulkan
 	}
 
-	static constexpr FORCEINLINE bool SupportsLowQualityLightmaps()
+	static constexpr UE_FORCEINLINE_HINT bool SupportsLowQualityLightmaps()
 	{
 		return true;
 	}
 
-	static constexpr FORCEINLINE bool SupportsDistanceFieldShadows()
+	static constexpr UE_FORCEINLINE_HINT bool SupportsDistanceFieldShadows()
 	{
 		return true;
 	}
 
-	static constexpr FORCEINLINE bool SupportsTextureStreaming()
+	static constexpr UE_FORCEINLINE_HINT bool SupportsTextureStreaming()
 	{
 		return true;
 	}
 
-	static constexpr FORCEINLINE bool SupportsMinimize()
+	static constexpr UE_FORCEINLINE_HINT bool SupportsMinimize()
 	{
 		return true;
 	}
 
-	static constexpr FORCEINLINE bool SupportsQuit()
+	static constexpr UE_FORCEINLINE_HINT bool SupportsQuit()
 	{
 		return true;
 	}
 
-	static constexpr FORCEINLINE bool HasFixedResolution()
+	static constexpr UE_FORCEINLINE_HINT bool HasFixedResolution()
 	{
 		return true;
 	}
 
-	static constexpr FORCEINLINE bool AllowsFramerateSmoothing()
+	static constexpr UE_FORCEINLINE_HINT bool AllowsFramerateSmoothing()
 	{
 		return true;
 	}
 
-	static constexpr FORCEINLINE bool AllowsCallStackDumpDuringAssert()
+	static constexpr UE_FORCEINLINE_HINT bool AllowsCallStackDumpDuringAssert()
 	{
 		return true;
 	}
 
-	static constexpr FORCEINLINE bool SupportsAudioStreaming()
+	static constexpr UE_FORCEINLINE_HINT bool SupportsAudioStreaming()
 	{
 		return true;
 	}
 
-	static constexpr FORCEINLINE bool SupportsMeshLODStreaming()
+	static constexpr UE_FORCEINLINE_HINT bool SupportsMeshLODStreaming()
 	{
 		return true;
 	}
 
-	static constexpr FORCEINLINE bool SupportsMemoryMappedFiles()
+	static constexpr UE_FORCEINLINE_HINT bool SupportsMemoryMappedFiles()
 	{
 		return true;
 	}
 
-	static constexpr FORCEINLINE bool SupportsMemoryMappedAudio()
+	static constexpr UE_FORCEINLINE_HINT bool SupportsMemoryMappedAudio()
 	{
 		return true;
 	}
 
-	static constexpr FORCEINLINE bool SupportsMemoryMappedAnimation()
+	static constexpr UE_FORCEINLINE_HINT bool SupportsMemoryMappedAnimation()
 	{
 		return true;
 	}
 
-	static constexpr FORCEINLINE int64 GetMemoryMappingAlignment()
+	static constexpr UE_FORCEINLINE_HINT int64 GetMemoryMappingAlignment()
 	{
 		// Cook for largest page size available.
 		// Data cooked for 16kb page size will work for 4kb page size firmware.
diff --git a/Engine/Source/Runtime/Core/Public/Android/AndroidPlatformString.h b/Engine/Source/Runtime/Core/Public/Android/AndroidPlatformString.h
--- a/Engine/Source/Runtime/Core/Public/Android/AndroidPlatformString.h
+++ b/Engine/Source/Runtime/Core/Public/Android/AndroidPlatformString.h
@@ -77,7 +77,7 @@ struct FAndroidPlatformString : public FGenericPlatformString
 	// This version was deprecated because Strcpy taking a size field does not match the CRT strcpy, and on some platforms the size field was being ignored.
 	// If the memzeroing causes a performance problem, request a new function StrcpyTruncate.
 	UE_DEPRECATED(5.6, "Use Strncpy instead. Note that Strncpy has a behavior difference from Strcpy: it memzeroes the entire DestCount-sized buffer after the end of string.")
-	static FORCEINLINE WIDECHAR* Strcpy(WIDECHAR* Dest, SIZE_T DestCount, const WIDECHAR* Src)
+	static UE_FORCEINLINE_HINT WIDECHAR* Strcpy(WIDECHAR* Dest, SIZE_T DestCount, const WIDECHAR* Src)
 	{
 		return Strcpy(Dest, Src);
 	}
@@ -128,7 +128,7 @@ struct FAndroidPlatformString : public FGenericPlatformString
 
 	// This version was deprecated because Strcat taking a size field does not match the CRT strcat, and on some platforms the size field was being ignored.
 	UE_DEPRECATED(5.6, "Use Strncat instead. !!NOTE THAT STRNCAT takes SrcLen rather than DestCount. You must call Strncat(Dest, Src, DestCount - Strlen(Dest) - 1).")
-	static FORCEINLINE WIDECHAR* Strcat(WIDECHAR* Dest, SIZE_T DestCount, const WIDECHAR* Src)
+	static UE_FORCEINLINE_HINT WIDECHAR* Strcat(WIDECHAR* Dest, SIZE_T DestCount, const WIDECHAR* Src)
 	{
 		return Strcat(Dest, Src);
 	}
@@ -504,7 +504,7 @@ struct FAndroidPlatformString : public FGenericPlatformString
 	/** 
 	 * Ansi implementation 
 	 **/
-	static FORCEINLINE ANSICHAR* Strcpy(ANSICHAR* Dest, const ANSICHAR* Src)
+	static UE_FORCEINLINE_HINT ANSICHAR* Strcpy(ANSICHAR* Dest, const ANSICHAR* Src)
 	{
 		return strcpy( Dest, Src );
 	}
@@ -512,7 +512,7 @@ struct FAndroidPlatformString : public FGenericPlatformString
 	// This version was deprecated because Strcpy taking a size field does not match the CRT strcpy, and on some platforms the size field was being ignored.
 	// If the memzeroing causes a performance problem, request a new function StrcpyTruncate.
 	UE_DEPRECATED(5.6, "Use Strncpy instead. Note that Strncpy has a behavior difference from Strcpy: it memzeroes the entire DestCount-sized buffer after the end of string.")
-	static FORCEINLINE ANSICHAR* Strcpy(ANSICHAR* Dest, SIZE_T DestCount, const ANSICHAR* Src)
+	static UE_FORCEINLINE_HINT ANSICHAR* Strcpy(ANSICHAR* Dest, SIZE_T DestCount, const ANSICHAR* Src)
 	{
 		return Strcpy(Dest, Src);
 	}
@@ -524,94 +524,94 @@ struct FAndroidPlatformString : public FGenericPlatformString
 		return Dest;
 	}
 
-	static FORCEINLINE ANSICHAR* Strcat(ANSICHAR* Dest, const ANSICHAR* Src)
+	static UE_FORCEINLINE_HINT ANSICHAR* Strcat(ANSICHAR* Dest, const ANSICHAR* Src)
 	{
 		return strcat( Dest, Src );
 	}
 
 	// This version was deprecated because Strcat taking a size field does not match the CRT strcat, and on some platforms the size field was being ignored.
 	UE_DEPRECATED(5.6, "Use Strncat instead. !!NOTE THAT STRNCAT takes SrcLen rather than DestCount. You must call Strncat(Dest, Src, DestCount - Strlen(Dest) - 1).")
-	static FORCEINLINE ANSICHAR* Strcat(ANSICHAR* Dest, SIZE_T DestCount, const ANSICHAR* Src)
+	static UE_FORCEINLINE_HINT ANSICHAR* Strcat(ANSICHAR* Dest, SIZE_T DestCount, const ANSICHAR* Src)
 	{
 		return Strcat(Dest, Src);
 	}
 
-	static FORCEINLINE ANSICHAR* Strncat(ANSICHAR* Dest, const ANSICHAR* Src, SIZE_T SrcLen)
+	static UE_FORCEINLINE_HINT ANSICHAR* Strncat(ANSICHAR* Dest, const ANSICHAR* Src, SIZE_T SrcLen)
 	{
 		return strncat(Dest, Src, SrcLen);
 	}
 
-	static FORCEINLINE int32 Strcmp( const ANSICHAR* String1, const ANSICHAR* String2 )
+	static UE_FORCEINLINE_HINT int32 Strcmp( const ANSICHAR* String1, const ANSICHAR* String2 )
 	{
 		return strcmp(String1, String2);
 	}
 
-	static FORCEINLINE int32 Strncmp( const ANSICHAR* String1, const ANSICHAR* String2, SIZE_T Count )
+	static UE_FORCEINLINE_HINT int32 Strncmp( const ANSICHAR* String1, const ANSICHAR* String2, SIZE_T Count )
 	{
 		return strncmp( String1, String2, Count );
 	}
 
-	static FORCEINLINE int32 Strlen( const ANSICHAR* String )
+	static UE_FORCEINLINE_HINT int32 Strlen( const ANSICHAR* String )
 	{
 		return strlen( String ); 
 	}
 
-	static FORCEINLINE int32 Strnlen( const ANSICHAR* String, SIZE_T StringSize )
+	static UE_FORCEINLINE_HINT int32 Strnlen( const ANSICHAR* String, SIZE_T StringSize )
 	{
 		return strnlen( String, StringSize );
 	}
 
-	static FORCEINLINE const ANSICHAR* Strstr( const ANSICHAR* String, const ANSICHAR* Find)
+	static UE_FORCEINLINE_HINT const ANSICHAR* Strstr( const ANSICHAR* String, const ANSICHAR* Find)
 	{
 		return strstr(String, Find);
 	}
 
-	static FORCEINLINE const ANSICHAR* Strchr( const ANSICHAR* String, ANSICHAR C)
+	static UE_FORCEINLINE_HINT const ANSICHAR* Strchr( const ANSICHAR* String, ANSICHAR C)
 	{
 		return strchr(String, C);
 	}
 
-	static FORCEINLINE const ANSICHAR* Strrchr( const ANSICHAR* String, ANSICHAR C)
+	static UE_FORCEINLINE_HINT const ANSICHAR* Strrchr( const ANSICHAR* String, ANSICHAR C)
 	{
 		return strrchr(String, C);
 	}
 
-	static FORCEINLINE int32 Atoi(const ANSICHAR* String)
+	static UE_FORCEINLINE_HINT int32 Atoi(const ANSICHAR* String)
 	{
 		return atoi( String ); 
 	}
 
-	static FORCEINLINE int64 Atoi64(const ANSICHAR* String)
+	static UE_FORCEINLINE_HINT int64 Atoi64(const ANSICHAR* String)
 	{
 		return strtoll( String, NULL, 10 );
 	}
 
-	static FORCEINLINE float Atof(const ANSICHAR* String)
+	static UE_FORCEINLINE_HINT float Atof(const ANSICHAR* String)
 	{
 		return (float)atof( String ); 
 	}
 
-	static FORCEINLINE double Atod(const ANSICHAR* String)
+	static UE_FORCEINLINE_HINT double Atod(const ANSICHAR* String)
 	{
 		return atof( String ); 
 	}
 
-	static FORCEINLINE int32 Strtoi( const ANSICHAR* Start, ANSICHAR** End, int32 Base ) 
+	static UE_FORCEINLINE_HINT int32 Strtoi( const ANSICHAR* Start, ANSICHAR** End, int32 Base ) 
 	{
 		return strtol( Start, End, Base ); 
 	}
 
-	static FORCEINLINE int64 Strtoi64( const ANSICHAR* Start, ANSICHAR** End, int32 Base ) 
+	static UE_FORCEINLINE_HINT int64 Strtoi64( const ANSICHAR* Start, ANSICHAR** End, int32 Base ) 
 	{
 		return strtoll(Start, End, Base);;
 	}
 
-	static FORCEINLINE uint64 Strtoui64( const ANSICHAR* Start, ANSICHAR** End, int32 Base ) 
+	static UE_FORCEINLINE_HINT uint64 Strtoui64( const ANSICHAR* Start, ANSICHAR** End, int32 Base ) 
 	{
 		return strtoull(Start, End, Base);;
 	}
 
-	static FORCEINLINE ANSICHAR* Strtok(ANSICHAR* StrToken, const ANSICHAR* Delim, ANSICHAR** Context)
+	static UE_FORCEINLINE_HINT ANSICHAR* Strtok(ANSICHAR* StrToken, const ANSICHAR* Delim, ANSICHAR** Context)
 	{
 		return strtok(StrToken, Delim);
 	}
@@ -625,7 +625,7 @@ struct FAndroidPlatformString : public FGenericPlatformString
 	/**
 	 * UTF8CHAR implementation.
 	 */
-	static FORCEINLINE UTF8CHAR* Strcpy(UTF8CHAR* Dest, const UTF8CHAR* Src)
+	static UE_FORCEINLINE_HINT UTF8CHAR* Strcpy(UTF8CHAR* Dest, const UTF8CHAR* Src)
 	{
 		return (UTF8CHAR*)Strcpy((ANSICHAR*)Dest, (const ANSICHAR*)Src);
 	}
@@ -639,12 +639,12 @@ struct FAndroidPlatformString : public FGenericPlatformString
 		PRAGMA_ENABLE_DEPRECATION_WARNINGS
 	}
 
-	static FORCEINLINE UTF8CHAR* Strncpy(UTF8CHAR* Dest, const UTF8CHAR* Src, SIZE_T MaxLen)
+	static UE_FORCEINLINE_HINT UTF8CHAR* Strncpy(UTF8CHAR* Dest, const UTF8CHAR* Src, SIZE_T MaxLen)
 	{
 		return (UTF8CHAR*)Strncpy((ANSICHAR*)Dest, (const ANSICHAR*)Src, MaxLen);
 	}
 
-	static FORCEINLINE UTF8CHAR* Strcat(UTF8CHAR* Dest, const UTF8CHAR* Src)
+	static UE_FORCEINLINE_HINT UTF8CHAR* Strcat(UTF8CHAR* Dest, const UTF8CHAR* Src)
 	{
 		return (UTF8CHAR*)Strcat((ANSICHAR*)Dest, (const ANSICHAR*)Src);
 	}
@@ -658,87 +658,87 @@ struct FAndroidPlatformString : public FGenericPlatformString
 		PRAGMA_ENABLE_DEPRECATION_WARNINGS
 	}
 
-	static FORCEINLINE UTF8CHAR* Strncat(UTF8CHAR* Dest, const UTF8CHAR* Src, SIZE_T SrcLen)
+	static UE_FORCEINLINE_HINT UTF8CHAR* Strncat(UTF8CHAR* Dest, const UTF8CHAR* Src, SIZE_T SrcLen)
 	{
 		return (UTF8CHAR*)Strncat((ANSICHAR*)Dest, (const ANSICHAR*)Src, SrcLen);
 	}
 
-	static FORCEINLINE int32 Strcmp(const UTF8CHAR* String1, const UTF8CHAR* String2)
+	static UE_FORCEINLINE_HINT int32 Strcmp(const UTF8CHAR* String1, const UTF8CHAR* String2)
 	{
 		return Strcmp((const ANSICHAR*)String1, (const ANSICHAR*)String2);
 	}
 
-	static FORCEINLINE int32 Strncmp(const UTF8CHAR* String1, const UTF8CHAR* String2, SIZE_T Count)
+	static UE_FORCEINLINE_HINT int32 Strncmp(const UTF8CHAR* String1, const UTF8CHAR* String2, SIZE_T Count)
 	{
 		return Strncmp((const ANSICHAR*)String1, (const ANSICHAR*)String2, Count);
 	}
 
-	static FORCEINLINE int32 Strlen(const UTF8CHAR* String)
+	static UE_FORCEINLINE_HINT int32 Strlen(const UTF8CHAR* String)
 	{
 		return Strlen((const ANSICHAR*)String);
 	}
 
-	static FORCEINLINE int32 Strnlen( const UTF8CHAR* String, SIZE_T StringSize )
+	static UE_FORCEINLINE_HINT int32 Strnlen( const UTF8CHAR* String, SIZE_T StringSize )
 	{
 		return Strnlen((const ANSICHAR*)String, StringSize);
 	}
 
-	static FORCEINLINE const UTF8CHAR* Strstr(const UTF8CHAR* String, const UTF8CHAR* Find)
+	static UE_FORCEINLINE_HINT const UTF8CHAR* Strstr(const UTF8CHAR* String, const UTF8CHAR* Find)
 	{
 		return (const UTF8CHAR*)Strstr((const ANSICHAR*)String, (const ANSICHAR*)Find);
 	}
 
-	static FORCEINLINE const UTF8CHAR* Strchr(const UTF8CHAR* String, UTF8CHAR C)
+	static UE_FORCEINLINE_HINT const UTF8CHAR* Strchr(const UTF8CHAR* String, UTF8CHAR C)
 	{
 		return (const UTF8CHAR*)Strchr((const ANSICHAR*)String, (ANSICHAR)C);
 	}
 
-	static FORCEINLINE const UTF8CHAR* Strrchr(const UTF8CHAR* String, UTF8CHAR C)
+	static UE_FORCEINLINE_HINT const UTF8CHAR* Strrchr(const UTF8CHAR* String, UTF8CHAR C)
 	{
 		return (const UTF8CHAR*)Strrchr((const ANSICHAR*)String, (ANSICHAR)C);
 	}
 
-	static FORCEINLINE int32 Atoi(const UTF8CHAR* String)
+	static UE_FORCEINLINE_HINT int32 Atoi(const UTF8CHAR* String)
 	{
 		return Atoi((const ANSICHAR*)String);
 	}
 
-	static FORCEINLINE int64 Atoi64(const UTF8CHAR* String)
+	static UE_FORCEINLINE_HINT int64 Atoi64(const UTF8CHAR* String)
 	{
 		return Atoi64((const ANSICHAR*)String);
 	}
 
-	static FORCEINLINE float Atof(const UTF8CHAR* String)
+	static UE_FORCEINLINE_HINT float Atof(const UTF8CHAR* String)
 	{
 		return Atof((const ANSICHAR*)String);
 	}
 
-	static FORCEINLINE double Atod(const UTF8CHAR* String)
+	static UE_FORCEINLINE_HINT double Atod(const UTF8CHAR* String)
 	{
 		return Atod((const ANSICHAR*)String);
 	}
 
-	static FORCEINLINE int32 Strtoi(const UTF8CHAR* Start, UTF8CHAR** End, int32 Base)
+	static UE_FORCEINLINE_HINT int32 Strtoi(const UTF8CHAR* Start, UTF8CHAR** End, int32 Base)
 	{
 		return Strtoi((const ANSICHAR*)Start, (ANSICHAR**)End, Base);
 	}
 
-	static FORCEINLINE int64 Strtoi64(const UTF8CHAR* Start, UTF8CHAR** End, int32 Base)
+	static UE_FORCEINLINE_HINT int64 Strtoi64(const UTF8CHAR* Start, UTF8CHAR** End, int32 Base)
 	{
 		return Strtoi64((const ANSICHAR*)Start, (ANSICHAR**)End, Base);
 	}
 
-	static FORCEINLINE uint64 Strtoui64(const UTF8CHAR* Start, UTF8CHAR** End, int32 Base)
+	static UE_FORCEINLINE_HINT uint64 Strtoui64(const UTF8CHAR* Start, UTF8CHAR** End, int32 Base)
 	{
 		return Strtoui64((const ANSICHAR*)Start, (ANSICHAR**)End, Base);
 	}
 
-	static FORCEINLINE UTF8CHAR* Strtok(UTF8CHAR* StrToken, const UTF8CHAR* Delim, UTF8CHAR** Context)
+	static UE_FORCEINLINE_HINT UTF8CHAR* Strtok(UTF8CHAR* StrToken, const UTF8CHAR* Delim, UTF8CHAR** Context)
 	{
 		return (UTF8CHAR*)Strtok((ANSICHAR*)StrToken, (const ANSICHAR*)Delim, (ANSICHAR**)Context);
 	}
 
-	static FORCEINLINE int32 GetVarArgs(UTF8CHAR* Dest, SIZE_T DestSize, const UTF8CHAR*& Fmt, va_list ArgPtr)
+	static UE_FORCEINLINE_HINT int32 GetVarArgs(UTF8CHAR* Dest, SIZE_T DestSize, const UTF8CHAR*& Fmt, va_list ArgPtr)
 	{
 		return GetVarArgs((ANSICHAR*)Dest, DestSize, *(const ANSICHAR**)&Fmt, ArgPtr);
 	}
diff --git a/Engine/Source/Runtime/Core/Public/Android/AndroidPlatformTLS.h b/Engine/Source/Runtime/Core/Public/Android/AndroidPlatformTLS.h
--- a/Engine/Source/Runtime/Core/Public/Android/AndroidPlatformTLS.h
+++ b/Engine/Source/Runtime/Core/Public/Android/AndroidPlatformTLS.h
@@ -44,7 +44,7 @@ struct FAndroidTLS : public FGenericPlatformTLS
 	 * @param SlotIndex the TLS index to store it in
 	 * @param Value the value to store in the slot
 	 */
-	static FORCEINLINE void SetTlsValue(uint32 SlotIndex,void* Value)
+	static UE_FORCEINLINE_HINT void SetTlsValue(uint32 SlotIndex,void* Value)
 	{
 		pthread_setspecific((pthread_key_t)SlotIndex, Value);
 	}
@@ -54,7 +54,7 @@ struct FAndroidTLS : public FGenericPlatformTLS
 	 *
 	 * @return the value stored in the slot
 	 */
-	static FORCEINLINE void* GetTlsValue(uint32 SlotIndex)
+	static UE_FORCEINLINE_HINT void* GetTlsValue(uint32 SlotIndex)
 	{
 		return pthread_getspecific((pthread_key_t)SlotIndex);
 	}
@@ -64,7 +64,7 @@ struct FAndroidTLS : public FGenericPlatformTLS
 	 *
 	 * @param SlotIndex the TLS index to store it in
 	 */
-	static FORCEINLINE void FreeTlsSlot(uint32 SlotIndex)
+	static UE_FORCEINLINE_HINT void FreeTlsSlot(uint32 SlotIndex)
 	{
 		pthread_key_delete((pthread_key_t)SlotIndex);
 	}
diff --git a/Engine/Source/Runtime/Core/Public/Apple/ApplePlatformMemory.h b/Engine/Source/Runtime/Core/Public/Apple/ApplePlatformMemory.h
--- a/Engine/Source/Runtime/Core/Public/Apple/ApplePlatformMemory.h
+++ b/Engine/Source/Runtime/Core/Public/Apple/ApplePlatformMemory.h
@@ -124,27 +124,27 @@ struct CORE_API FApplePlatformMemory : public FGenericPlatformMemory
 		void Decommit(size_t InOffset, size_t InSize);
 		void FreeVirtual();
 
-		FORCEINLINE void CommitByPtr(void *InPtr, size_t InSize)
+		UE_FORCEINLINE_HINT void CommitByPtr(void *InPtr, size_t InSize)
 		{
 			Commit(size_t(((uint8*)InPtr) - ((uint8*)Ptr)), InSize);
 		}
 
-		FORCEINLINE void DecommitByPtr(void *InPtr, size_t InSize)
+		UE_FORCEINLINE_HINT void DecommitByPtr(void *InPtr, size_t InSize)
 		{
 			Decommit(size_t(((uint8*)InPtr) - ((uint8*)Ptr)), InSize);
 		}
 
-		FORCEINLINE void Commit()
+		UE_FORCEINLINE_HINT void Commit()
 		{
 			Commit(0, GetActualSize());
 		}
 
-		FORCEINLINE void Decommit()
+		UE_FORCEINLINE_HINT void Decommit()
 		{
 			Decommit(0, GetActualSize());
 		}
 
-		FORCEINLINE size_t GetActualSize() const
+		UE_FORCEINLINE_HINT size_t GetActualSize() const
 		{
 			return VMSizeDivVirtualSizeAlignment * GetVirtualSizeAlignment();
 		}
diff --git a/Engine/Source/Runtime/Core/Public/Apple/ApplePlatformMisc.h b/Engine/Source/Runtime/Core/Public/Apple/ApplePlatformMisc.h
--- a/Engine/Source/Runtime/Core/Public/Apple/ApplePlatformMisc.h
+++ b/Engine/Source/Runtime/Core/Public/Apple/ApplePlatformMisc.h
@@ -65,7 +65,7 @@ struct CORE_API FApplePlatformMisc : public FGenericPlatformMisc
 	static bool IsDebuggerPresent();
 #endif
 
-	FORCEINLINE static void MemoryBarrier()
+	UE_FORCEINLINE_HINT static void MemoryBarrier()
 	{
 		__sync_synchronize();
 	}
diff --git a/Engine/Source/Runtime/Core/Public/Apple/ApplePlatformSymbolication.h b/Engine/Source/Runtime/Core/Public/Apple/ApplePlatformSymbolication.h
--- a/Engine/Source/Runtime/Core/Public/Apple/ApplePlatformSymbolication.h
+++ b/Engine/Source/Runtime/Core/Public/Apple/ApplePlatformSymbolication.h
@@ -46,7 +46,7 @@ struct FApplePlatformSymbolDatabaseKeyFuncs
 	/**
 	 * @return The key used to index the given element.
 	 */
-	static FORCEINLINE KeyInitType GetSetKey(ElementInitType Element)
+	static UE_FORCEINLINE_HINT KeyInitType GetSetKey(ElementInitType Element)
 	{
 		return Element.GenericDB->Signature;
 	}
@@ -54,13 +54,13 @@ struct FApplePlatformSymbolDatabaseKeyFuncs
 	/**
 	 * @return True if the keys match.
 	 */
-	static FORCEINLINE bool Matches(KeyInitType A,KeyInitType B)
+	static UE_FORCEINLINE_HINT bool Matches(KeyInitType A,KeyInitType B)
 	{
 		return A == B;
 	}
 
 	/** Calculates a hash index for a key. */
-	static FORCEINLINE uint32 GetKeyHash(KeyInitType Key)
+	static UE_FORCEINLINE_HINT uint32 GetKeyHash(KeyInitType Key)
 	{
 		return GetTypeHash(Key);
 	}
diff --git a/Engine/Source/Runtime/Core/Public/Apple/ApplePlatformTLS.h b/Engine/Source/Runtime/Core/Public/Apple/ApplePlatformTLS.h
--- a/Engine/Source/Runtime/Core/Public/Apple/ApplePlatformTLS.h
+++ b/Engine/Source/Runtime/Core/Public/Apple/ApplePlatformTLS.h
@@ -21,7 +21,7 @@ struct CORE_API FApplePlatformTLS : public FGenericPlatformTLS
 	/**
 	 * Returns the currently executing thread's id
 	 */
-	static FORCEINLINE uint32 GetCurrentThreadId(void)
+	static UE_FORCEINLINE_HINT uint32 GetCurrentThreadId(void)
 	{
 		return (uint32)pthread_mach_thread_np(pthread_self());
 	}
@@ -46,7 +46,7 @@ struct CORE_API FApplePlatformTLS : public FGenericPlatformTLS
 	 * @param SlotIndex the TLS index to store it in
 	 * @param Value the value to store in the slot
 	 */
-	static FORCEINLINE void SetTlsValue(uint32 SlotIndex,void* Value)
+	static UE_FORCEINLINE_HINT void SetTlsValue(uint32 SlotIndex,void* Value)
 	{
 		pthread_setspecific((pthread_key_t)SlotIndex, Value);
 	}
@@ -56,7 +56,7 @@ struct CORE_API FApplePlatformTLS : public FGenericPlatformTLS
 	 *
 	 * @return the value stored in the slot
 	 */
-	static FORCEINLINE void* GetTlsValue(uint32 SlotIndex)
+	static UE_FORCEINLINE_HINT void* GetTlsValue(uint32 SlotIndex)
 	{
 		return pthread_getspecific((pthread_key_t)SlotIndex);
 	}
@@ -66,7 +66,7 @@ struct CORE_API FApplePlatformTLS : public FGenericPlatformTLS
 	 *
 	 * @param SlotIndex the TLS index to store it in
 	 */
-	static FORCEINLINE void FreeTlsSlot(uint32 SlotIndex)
+	static UE_FORCEINLINE_HINT void FreeTlsSlot(uint32 SlotIndex)
 	{
 		pthread_key_delete((pthread_key_t)SlotIndex);
 	}
diff --git a/Engine/Source/Runtime/Core/Public/Async/AsyncFileHandle.h b/Engine/Source/Runtime/Core/Public/Async/AsyncFileHandle.h
--- a/Engine/Source/Runtime/Core/Public/Async/AsyncFileHandle.h
+++ b/Engine/Source/Runtime/Core/Public/Async/AsyncFileHandle.h
@@ -83,7 +83,7 @@ public:
 	* Nonblocking poll of the state of completion.
 	* @return true if the request is complete
 	**/
-	FORCEINLINE bool PollCompletion()
+	UE_FORCEINLINE_HINT bool PollCompletion()
 	{
 		return bCompleteAndCallbackCalled;
 	}
diff --git a/Engine/Source/Runtime/Core/Public/Async/Fundamental/Scheduler.h b/Engine/Source/Runtime/Core/Public/Async/Fundamental/Scheduler.h
--- a/Engine/Source/Runtime/Core/Public/Async/Fundamental/Scheduler.h
+++ b/Engine/Source/Runtime/Core/Public/Async/Fundamental/Scheduler.h
@@ -110,7 +110,7 @@ namespace LowLevelTasks
 		using FConditional = TTaskDelegate<bool(), 16>;
 
 	public: // Public Interface of the Scheduler
-		FORCEINLINE_DEBUGGABLE static FScheduler& Get();
+		inline static FScheduler& Get();
 
 		//start number of workers where 0 is the system default
 		CORE_API void StartWorkers(uint32 NumForegroundWorkers = 0, uint32 NumBackgroundWorkers = 0, FThread::EForkable IsForkable = FThread::NonForkable, EThreadPriority InWorkerPriority = EThreadPriority::TPri_Normal, EThreadPriority InBackgroundPriority = EThreadPriority::TPri_BelowNormal, uint64 InWorkerAffinity = 0, uint64 InBackgroundAffinity = 0);
@@ -175,7 +175,7 @@ namespace LowLevelTasks
 		FOversubscriptionLimitReached                  OversubscriptionLimitReachedEvent;
 	};
 
-	FORCEINLINE_DEBUGGABLE bool TryLaunch(FTask& Task, EQueuePreference QueuePreference = EQueuePreference::DefaultPreference, bool bWakeUpWorker = true)
+	UE_FORCEINLINE_HINT bool TryLaunch(FTask& Task, EQueuePreference QueuePreference = EQueuePreference::DefaultPreference, bool bWakeUpWorker = true)
 	{
 		return FScheduler::Get().TryLaunch(Task, QueuePreference, bWakeUpWorker);
 	}
diff --git a/Engine/Source/Runtime/Core/Public/Async/Fundamental/Task.h b/Engine/Source/Runtime/Core/Public/Async/Fundamental/Task.h
--- a/Engine/Source/Runtime/Core/Public/Async/Fundamental/Task.h
+++ b/Engine/Source/Runtime/Core/Public/Async/Fundamental/Task.h
@@ -366,7 +366,7 @@ namespace LowLevelTasks
 		*/
 		CORE_API static const FTask* GetActiveTask();
 #else
-		FORCEINLINE static const FTask* GetActiveTask()
+		UE_FORCEINLINE_HINT static const FTask* GetActiveTask()
 		{
 			return ActiveTask;
 		}
diff --git a/Engine/Source/Runtime/Core/Public/Async/MappedFileHandle.h b/Engine/Source/Runtime/Core/Public/Async/MappedFileHandle.h
--- a/Engine/Source/Runtime/Core/Public/Async/MappedFileHandle.h
+++ b/Engine/Source/Runtime/Core/Public/Async/MappedFileHandle.h
@@ -20,7 +20,7 @@ class IMappedFileRegion
 	FString DebugFilename;
 	size_t DebugOffsetRelativeToFile;
 
-	FORCEINLINE void CheckInvariants()
+	UE_FORCEINLINE_HINT void CheckInvariants()
 	{
 		check(MappedPtr && MappedSize);
 	}
@@ -130,7 +130,7 @@ public:
 	/**
 	* Return the size of the mapped file.
 	**/
-	FORCEINLINE int64 GetFileSize() const
+	UE_FORCEINLINE_HINT int64 GetFileSize() const
 	{
 		return MappedFileSize;
 	}
diff --git a/Engine/Source/Runtime/Core/Public/Async/TaskGraphInterfaces.h b/Engine/Source/Runtime/Core/Public/Async/TaskGraphInterfaces.h
--- a/Engine/Source/Runtime/Core/Public/Async/TaskGraphInterfaces.h
+++ b/Engine/Source/Runtime/Core/Public/Async/TaskGraphInterfaces.h
@@ -120,22 +120,22 @@ namespace ENamedThreads
 		friend void SetRenderThread_Local(Type Thread);
 	};
 
-	FORCEINLINE Type GetRenderThread()
+	UE_FORCEINLINE_HINT Type GetRenderThread()
 	{
 		return FRenderThreadStatics::RenderThread.Load(EMemoryOrder::Relaxed);
 	}
 
-	FORCEINLINE Type GetRenderThread_Local()
+	UE_FORCEINLINE_HINT Type GetRenderThread_Local()
 	{
 		return FRenderThreadStatics::RenderThread_Local.Load(EMemoryOrder::Relaxed);
 	}
 
-	FORCEINLINE void SetRenderThread(Type Thread)
+	UE_FORCEINLINE_HINT void SetRenderThread(Type Thread)
 	{
 		FRenderThreadStatics::RenderThread.Store(Thread, EMemoryOrder::Relaxed);
 	}
 
-	FORCEINLINE void SetRenderThread_Local(Type Thread)
+	UE_FORCEINLINE_HINT void SetRenderThread_Local(Type Thread)
 	{
 		FRenderThreadStatics::RenderThread_Local.Store(Thread, EMemoryOrder::Relaxed);
 	}
@@ -145,17 +145,17 @@ namespace ENamedThreads
 	extern CORE_API int32 bHasBackgroundThreads; 
 	extern CORE_API int32 bHasHighPriorityThreads;
 
-	FORCEINLINE Type GetThreadIndex(Type ThreadAndIndex)
+	UE_FORCEINLINE_HINT Type GetThreadIndex(Type ThreadAndIndex)
 	{
 		return ((ThreadAndIndex & ThreadIndexMask) == AnyThread) ? AnyThread : Type(ThreadAndIndex & ThreadIndexMask);
 	}
 
-	FORCEINLINE int32 GetQueueIndex(Type ThreadAndIndex)
+	UE_FORCEINLINE_HINT int32 GetQueueIndex(Type ThreadAndIndex)
 	{
 		return (ThreadAndIndex & QueueIndexMask) >> QueueIndexShift;
 	}
 
-	FORCEINLINE int32 GetTaskPriority(Type ThreadAndIndex)
+	UE_FORCEINLINE_HINT int32 GetTaskPriority(Type ThreadAndIndex)
 	{
 		return (ThreadAndIndex & TaskPriorityMask) >> TaskPriorityShift;
 	}
@@ -771,7 +771,7 @@ public:
 		checkThreadGraph(ThreadToReturnFrom != ENamedThreads::AnyThread); // doesn't make any sense to return from any thread
 	}
 
-	FORCEINLINE TStatId GetStatId() const
+	UE_FORCEINLINE_HINT TStatId GetStatId() const
 	{
 		return GET_STATID(STAT_FReturnGraphTask);
 	}
@@ -904,7 +904,7 @@ public:
 		check(Event);
 	}
 
-	FORCEINLINE TStatId GetStatId() const
+	UE_FORCEINLINE_HINT TStatId GetStatId() const
 	{
 		return GET_STATID(STAT_FTriggerEventGraphTask);
 	}
@@ -1089,25 +1089,25 @@ public:
 		return SubsequentsMode;
 	}
 
-	FORCEINLINE void DoTask(ENamedThreads::Type CurrentThread, const FGraphEventRef& MyCompletionGraphEvent)
+	UE_FORCEINLINE_HINT void DoTask(ENamedThreads::Type CurrentThread, const FGraphEventRef& MyCompletionGraphEvent)
 	{
 		DoTaskImpl(Function, CurrentThread, MyCompletionGraphEvent);
 	}
 
 private:
-	FORCEINLINE static void DoTaskImpl(TUniqueFunction<void()>& Function, ENamedThreads::Type CurrentThread,
+	UE_FORCEINLINE_HINT static void DoTaskImpl(TUniqueFunction<void()>& Function, ENamedThreads::Type CurrentThread,
 		const FGraphEventRef& MyCompletionGraphEvent)
 	{
 		Function();
 	}
 
-	FORCEINLINE static void DoTaskImpl(TUniqueFunction<void(const FGraphEventRef&)>& Function,
+	UE_FORCEINLINE_HINT static void DoTaskImpl(TUniqueFunction<void(const FGraphEventRef&)>& Function,
 		ENamedThreads::Type CurrentThread, const FGraphEventRef& MyCompletionGraphEvent)
 	{
 		Function(MyCompletionGraphEvent);
 	}
 
-	FORCEINLINE static void DoTaskImpl(TUniqueFunction<void(ENamedThreads::Type, const FGraphEventRef&)>& Function,
+	UE_FORCEINLINE_HINT static void DoTaskImpl(TUniqueFunction<void(ENamedThreads::Type, const FGraphEventRef&)>& Function,
 		ENamedThreads::Type CurrentThread, const FGraphEventRef& MyCompletionGraphEvent)
 	{
 		Function(CurrentThread, MyCompletionGraphEvent);
diff --git a/Engine/Source/Runtime/Core/Public/Clang/ClangPlatformMath.h b/Engine/Source/Runtime/Core/Public/Clang/ClangPlatformMath.h
--- a/Engine/Source/Runtime/Core/Public/Clang/ClangPlatformMath.h
+++ b/Engine/Source/Runtime/Core/Public/Clang/ClangPlatformMath.h
@@ -25,7 +25,7 @@ struct FClangPlatformMath : public FGenericPlatformMath
 	 *
 	 * @return the number of zeros before the first "on" bit
 	 */
-	static FORCEINLINE constexpr uint8 CountLeadingZeros8(uint8 Value)
+	static UE_FORCEINLINE_HINT constexpr uint8 CountLeadingZeros8(uint8 Value)
 	{
 		return uint8(__builtin_clz((uint32(Value) << 1) | 1) - 23);
 	}
@@ -37,7 +37,7 @@ struct FClangPlatformMath : public FGenericPlatformMath
 	 *
 	 * @return the number of zeros before the first "on" bit
 	 */
-	static FORCEINLINE constexpr uint32 CountLeadingZeros(uint32 Value)
+	static UE_FORCEINLINE_HINT constexpr uint32 CountLeadingZeros(uint32 Value)
 	{
 		return __builtin_clzll((uint64(Value) << 1) | 1) - 31;
 	}
@@ -76,7 +76,7 @@ struct FClangPlatformMath : public FGenericPlatformMath
 		return (uint32)__builtin_ctz(Value);
 	}
 
-	static FORCEINLINE constexpr uint32 CountTrailingZerosConstExpr(uint32 Value)
+	static UE_FORCEINLINE_HINT constexpr uint32 CountTrailingZerosConstExpr(uint32 Value)
 	{
 		return CountTrailingZeros(Value);
 	}
@@ -98,28 +98,28 @@ struct FClangPlatformMath : public FGenericPlatformMath
 		return (uint64)__builtin_ctzll(Value);
 	}
 
-	static FORCEINLINE constexpr uint64 CountTrailingZeros64ConstExpr(uint64 Value)
+	static UE_FORCEINLINE_HINT constexpr uint64 CountTrailingZeros64ConstExpr(uint64 Value)
 	{
 		return CountTrailingZeros64(Value);
 	}
 	
 
-	static FORCEINLINE constexpr uint32 FloorLog2(uint32 Value)
+	static UE_FORCEINLINE_HINT constexpr uint32 FloorLog2(uint32 Value)
 	{
 		return 31 - __builtin_clz(Value | 1);
 	}
 
-	static FORCEINLINE constexpr uint32 FloorLog2NonZero(uint32 Value)
+	static UE_FORCEINLINE_HINT constexpr uint32 FloorLog2NonZero(uint32 Value)
 	{
 		return 31 - __builtin_clz(Value);
 	}
 
-	static FORCEINLINE constexpr uint64 FloorLog2_64(uint64 Value)
+	static UE_FORCEINLINE_HINT constexpr uint64 FloorLog2_64(uint64 Value)
 	{
 		return 63 - __builtin_clzll(Value | 1);
 	}
 
-	static FORCEINLINE constexpr uint64 FloorLog2NonZero_64(uint64 Value)
+	static UE_FORCEINLINE_HINT constexpr uint64 FloorLog2NonZero_64(uint64 Value)
 	{
 		return 63 - __builtin_clzll(Value);
 	}
@@ -130,7 +130,7 @@ struct FClangPlatformMath : public FGenericPlatformMath
 	 * If there wasn't overflow, it returns true, and the result of the addition is written to OutResult.
 	 */
 	template <UE::CIntegral IntType>
-	static FORCEINLINE bool AddAndCheckForOverflow(IntType A, IntType B, IntType& OutResult)
+	static UE_FORCEINLINE_HINT bool AddAndCheckForOverflow(IntType A, IntType B, IntType& OutResult)
 	{
 		return !__builtin_add_overflow(A, B, &OutResult);
 	}
@@ -141,7 +141,7 @@ struct FClangPlatformMath : public FGenericPlatformMath
 	 * If there wasn't overflow, it returns true, and the result of the subtraction is written to OutResult.
 	 */
 	template <UE::CIntegral IntType>
-	static FORCEINLINE bool SubtractAndCheckForOverflow(IntType A, IntType B, IntType& OutResult)
+	static UE_FORCEINLINE_HINT bool SubtractAndCheckForOverflow(IntType A, IntType B, IntType& OutResult)
 	{
 		return !__builtin_sub_overflow(A, B, &OutResult);
 	}
@@ -152,7 +152,7 @@ struct FClangPlatformMath : public FGenericPlatformMath
 	 * If there wasn't overflow, it returns true, and the result of the multiplication is written to OutResult.
 	 */
 	template <UE::CIntegral IntType>
-	static FORCEINLINE bool MultiplyAndCheckForOverflow(IntType A, IntType B, IntType& OutResult)
+	static UE_FORCEINLINE_HINT bool MultiplyAndCheckForOverflow(IntType A, IntType B, IntType& OutResult)
 	{
 		return !__builtin_mul_overflow(A, B, &OutResult);
 	}
diff --git a/Engine/Source/Runtime/Core/Public/Containers/Array.h b/Engine/Source/Runtime/Core/Public/Containers/Array.h
--- a/Engine/Source/Runtime/Core/Public/Containers/Array.h
+++ b/Engine/Source/Runtime/Core/Public/Containers/Array.h
@@ -119,18 +119,18 @@ public:
 		return Tmp -= Offset;
 	}
 
-	UE_NODEBUG [[nodiscard]] FORCEINLINE ElementType& operator* () const
+	UE_NODEBUG [[nodiscard]] UE_FORCEINLINE_HINT ElementType& operator* () const
 	{
 		return Container[ Index ];
 	}
 
-	UE_NODEBUG [[nodiscard]] FORCEINLINE ElementType* operator->() const
+	UE_NODEBUG [[nodiscard]] UE_FORCEINLINE_HINT ElementType* operator->() const
 	{
 		return &Container[ Index ];
 	}
 
 	/** conversion to "bool" returning true if the iterator has not reached the last element. */
-	UE_NODEBUG [[nodiscard]] FORCEINLINE explicit operator bool() const
+	UE_NODEBUG [[nodiscard]] UE_FORCEINLINE_HINT explicit operator bool() const
 	{
 		return Container.IsValidIndex(Index);
 	}
@@ -171,12 +171,12 @@ public:
 		Index--;
 	}
 
-	UE_NODEBUG [[nodiscard]] FORCEINLINE bool operator==(const TIndexedContainerIterator& Rhs) const
+	UE_NODEBUG [[nodiscard]] UE_FORCEINLINE_HINT bool operator==(const TIndexedContainerIterator& Rhs) const
 	{
 		return &Container == &Rhs.Container && Index == Rhs.Index;
 	}
 #if !PLATFORM_COMPILER_HAS_GENERATED_COMPARISON_OPERATORS
-	UE_NODEBUG [[nodiscard]] FORCEINLINE bool operator!=(const TIndexedContainerIterator& Rhs) const
+	UE_NODEBUG [[nodiscard]] UE_FORCEINLINE_HINT bool operator!=(const TIndexedContainerIterator& Rhs) const
 	{
 		return &Container != &Rhs.Container || Index != Rhs.Index;
 	}
@@ -190,7 +190,7 @@ private:
 
 /** operator + */
 template <typename ContainerType, typename ElementType, typename SizeType>
-UE_NODEBUG [[nodiscard]] FORCEINLINE TIndexedContainerIterator<ContainerType, ElementType, SizeType> operator+(SizeType Offset, TIndexedContainerIterator<ContainerType, ElementType, SizeType> RHS)
+UE_NODEBUG [[nodiscard]] UE_FORCEINLINE_HINT TIndexedContainerIterator<ContainerType, ElementType, SizeType> operator+(SizeType Offset, TIndexedContainerIterator<ContainerType, ElementType, SizeType> RHS)
 {
 	return RHS + Offset;
 }
@@ -216,7 +216,7 @@ UE_NODEBUG [[nodiscard]] FORCEINLINE TIndexedContainerIterator<ContainerType, El
 		{
 		}
 
-		UE_NODEBUG [[nodiscard]] FORCEINLINE ElementType& operator*() const
+		UE_NODEBUG [[nodiscard]] UE_FORCEINLINE_HINT ElementType& operator*() const
 		{
 			if constexpr (bReverse)
 			{
@@ -228,7 +228,7 @@ UE_NODEBUG [[nodiscard]] FORCEINLINE TIndexedContainerIterator<ContainerType, El
 			}
 		}
 
-		UE_NODEBUG FORCEINLINE TCheckedPointerIterator& operator++()
+		UE_NODEBUG UE_FORCEINLINE_HINT TCheckedPointerIterator& operator++()
 		{
 			if constexpr (bReverse)
 			{
@@ -241,7 +241,7 @@ UE_NODEBUG [[nodiscard]] FORCEINLINE TIndexedContainerIterator<ContainerType, El
 			return *this;
 		}
 
-		UE_NODEBUG FORCEINLINE TCheckedPointerIterator& operator--()
+		UE_NODEBUG UE_FORCEINLINE_HINT TCheckedPointerIterator& operator--()
 		{
 			if constexpr (bReverse)
 			{
@@ -254,7 +254,7 @@ UE_NODEBUG [[nodiscard]] FORCEINLINE TIndexedContainerIterator<ContainerType, El
 			return *this;
 		}
 
-		UE_NODEBUG [[nodiscard]] FORCEINLINE bool operator!=(const TCheckedPointerIterator& Rhs) const
+		UE_NODEBUG [[nodiscard]] UE_FORCEINLINE_HINT bool operator!=(const TCheckedPointerIterator& Rhs) const
 		{
 			// We only need to do the check in this operator, because no other operator will be
 			// called until after this one returns.
@@ -281,18 +281,18 @@ struct TDereferencingIterator
 	{
 	}
 
-	UE_NODEBUG [[nodiscard]] FORCEINLINE ElementType& operator*() const
+	UE_NODEBUG [[nodiscard]] UE_FORCEINLINE_HINT ElementType& operator*() const
 	{
 		return *(ElementType*)*Iter;
 	}
 
-	UE_NODEBUG FORCEINLINE TDereferencingIterator& operator++()
+	UE_NODEBUG UE_FORCEINLINE_HINT TDereferencingIterator& operator++()
 	{
 		++Iter;
 		return *this;
 	}
 
-	UE_NODEBUG [[nodiscard]] FORCEINLINE bool operator!=(const TDereferencingIterator& Rhs) const
+	UE_NODEBUG [[nodiscard]] UE_FORCEINLINE_HINT bool operator!=(const TDereferencingIterator& Rhs) const
 	{
 		return Iter != Rhs.Iter;
 	}
@@ -306,7 +306,7 @@ namespace UE::Core::Private
 	// Simply forwards to an unqualified GetData(), but can be called from within a container or view
 	// where GetData() is already a member and so hides any others.
 	template <typename T>
-	[[nodiscard]] FORCEINLINE decltype(auto) GetDataHelper(T&& Arg)
+	[[nodiscard]] UE_FORCEINLINE_HINT decltype(auto) GetDataHelper(T&& Arg)
 	{
 		return GetData(Forward<T>(Arg));
 	}
@@ -390,7 +390,7 @@ namespace UE::Core::Private
 	// hot location in the code.
 	// This returns the old ArrayMax in order to save a register clobber/reload.
 	template <uint32 Flags, typename AllocatorInstanceType>
-	FORCEINLINE TAllocatorSizeType_T<AllocatorInstanceType> ReallocGrow1_DoAlloc_Impl(
+	UE_FORCEINLINE_HINT TAllocatorSizeType_T<AllocatorInstanceType> ReallocGrow1_DoAlloc_Impl(
 		uint32                                       ElementSize,
 		uint32                                       ElementAlignment,
 		AllocatorInstanceType&                       AllocatorInstance,
@@ -674,7 +674,7 @@ public:
 	/**
 	 * Constructor, initializes element number counters.
 	 */
-	[[nodiscard]] FORCEINLINE TArray()
+	[[nodiscard]] UE_FORCEINLINE_HINT TArray()
 		: ArrayNum(0)
 		, ArrayMax(AllocatorInstance.GetInitialCapacity())
 	{
@@ -687,7 +687,7 @@ public:
 	 * @param Count The number of elements to copy from Ptr.
 	 * @see Append
 	 */
-	[[nodiscard]] FORCEINLINE TArray(const ElementType* Ptr, SizeType Count)
+	[[nodiscard]] UE_FORCEINLINE_HINT TArray(const ElementType* Ptr, SizeType Count)
 	{
 		if (Count < 0)
 		{
@@ -724,7 +724,7 @@ public:
 		typename OtherAllocator
 		UE_REQUIRES(UE::Core::Private::TArrayElementsAreCompatible_V<ElementType, const OtherElementType&>)
 	>
-	[[nodiscard]] FORCEINLINE explicit TArray(const TArray<OtherElementType, OtherAllocator>& Other)
+	[[nodiscard]] UE_FORCEINLINE_HINT explicit TArray(const TArray<OtherElementType, OtherAllocator>& Other)
 	{
 		CopyToEmpty(Other.GetData(), Other.Num(), 0);
 	}
@@ -734,7 +734,7 @@ public:
 	 *
 	 * @param Other The source array to copy.
 	 */
-	[[nodiscard]] FORCEINLINE TArray(const TArray& Other)
+	[[nodiscard]] UE_FORCEINLINE_HINT TArray(const TArray& Other)
 	{
 		CopyToEmpty(Other.GetData(), Other.Num(), 0);
 	}
@@ -746,7 +746,7 @@ public:
 	 * @param ExtraSlack Tells how much extra memory should be preallocated
 	 *                   at the end of the array in the number of elements.
 	 */
-	[[nodiscard]] FORCEINLINE TArray(const TArray& Other, SizeType ExtraSlack)
+	[[nodiscard]] UE_FORCEINLINE_HINT TArray(const TArray& Other, SizeType ExtraSlack)
 	{
 		CopyToEmptyWithSlack(Other.GetData(), Other.Num(), 0, ExtraSlack);
 	}
@@ -804,7 +804,7 @@ public:
 
 private:
 
-	FORCEINLINE void SlackTrackerNumChanged()
+	UE_FORCEINLINE_HINT void SlackTrackerNumChanged()
 	{
 #if UE_ENABLE_ARRAY_SLACK_TRACKING
 		if constexpr (TAllocatorTraits<InAllocatorType>::SupportsSlackTracking)
@@ -822,7 +822,7 @@ private:
 	 * @param PrevMax The previous allocated size.
 	 */
 	template <typename FromArrayType, typename ToArrayType>
-	static FORCEINLINE void MoveOrCopy(ToArrayType& ToArray, FromArrayType& FromArray, SizeType PrevMax)
+	static UE_FORCEINLINE_HINT void MoveOrCopy(ToArrayType& ToArray, FromArrayType& FromArray, SizeType PrevMax)
 	{
 		if constexpr (UE::Core::Private::CanMoveTArrayPointersBetweenArrayTypes<FromArrayType, ToArrayType>())
 		{
@@ -879,7 +879,7 @@ private:
 	 *                   at the end of the array in the number of elements.
 	 */
 	template <typename FromArrayType, typename ToArrayType>
-	static FORCEINLINE void MoveOrCopyWithSlack(ToArrayType& ToArray, FromArrayType& FromArray, SizeType PrevMax, SizeType ExtraSlack)
+	static UE_FORCEINLINE_HINT void MoveOrCopyWithSlack(ToArrayType& ToArray, FromArrayType& FromArray, SizeType PrevMax, SizeType ExtraSlack)
 	{
 		if constexpr (UE::Core::Private::CanMoveTArrayPointersBetweenArrayTypes<FromArrayType, ToArrayType>())
 		{
@@ -912,7 +912,7 @@ public:
 	 *
 	 * @param Other Array to move from.
 	 */
-	[[nodiscard]] FORCEINLINE TArray(TArray&& Other)
+	[[nodiscard]] UE_FORCEINLINE_HINT TArray(TArray&& Other)
 	{
 		MoveOrCopy(*this, Other, 0);
 	}
@@ -927,7 +927,7 @@ public:
 		typename OtherAllocator
 		UE_REQUIRES(UE::Core::Private::TArrayElementsAreCompatible_V<ElementType, OtherElementType&&>)
 	>
-	[[nodiscard]] FORCEINLINE explicit TArray(TArray<OtherElementType, OtherAllocator>&& Other)
+	[[nodiscard]] UE_FORCEINLINE_HINT explicit TArray(TArray<OtherElementType, OtherAllocator>&& Other)
 	{
 		MoveOrCopy(*this, Other, 0);
 	}
@@ -999,7 +999,7 @@ public:
 	 *
 	 * @returns Pointer to first array entry or nullptr if ArrayMax == 0.
 	 */
-	UE_NODEBUG [[nodiscard]] FORCEINLINE ElementType* GetData() UE_LIFETIMEBOUND
+	UE_NODEBUG [[nodiscard]] UE_FORCEINLINE_HINT ElementType* GetData() UE_LIFETIMEBOUND
 	{
 		return (ElementType*)AllocatorInstance.GetAllocation();
 	}
@@ -1019,7 +1019,7 @@ public:
 	 *
 	 * @returns Size in bytes of array type.
 	 */
-	UE_NODEBUG [[nodiscard]] FORCEINLINE static constexpr uint32 GetTypeSize()
+	UE_NODEBUG [[nodiscard]] UE_FORCEINLINE_HINT static constexpr uint32 GetTypeSize()
 	{
 		return sizeof(ElementType);
 	}
@@ -1031,7 +1031,7 @@ public:
 	 *
 	 * @returns Number of bytes allocated by this container.
 	 */
-	UE_NODEBUG [[nodiscard]] FORCEINLINE SIZE_T GetAllocatedSize(void) const
+	UE_NODEBUG [[nodiscard]] UE_FORCEINLINE_HINT SIZE_T GetAllocatedSize(void) const
 	{
 		return AllocatorInstance.GetAllocatedSize(ArrayMax, sizeof(ElementType));
 	}
@@ -1041,7 +1041,7 @@ public:
 	 *
 	 * @see Num, Shrink
 	 */
-	UE_NODEBUG [[nodiscard]] FORCEINLINE SizeType GetSlack() const
+	UE_NODEBUG [[nodiscard]] UE_FORCEINLINE_HINT SizeType GetSlack() const
 	{
 		return ArrayMax - ArrayNum;
 	}
@@ -1050,7 +1050,7 @@ public:
 	 * Checks array invariants: if array size is greater than or equal to zero and less
 	 * than or equal to the maximum.
 	 */
-	UE_NODEBUG FORCEINLINE void CheckInvariants() const
+	UE_NODEBUG UE_FORCEINLINE_HINT void CheckInvariants() const
 	{
 		checkSlow((ArrayNum >= 0) & (ArrayMax >= ArrayNum)); // & for one branch
 	}
@@ -1060,7 +1060,7 @@ public:
 	 *
 	 * @param Index Index to check.
 	 */
-	FORCEINLINE void RangeCheck(SizeType Index) const
+	UE_FORCEINLINE_HINT void RangeCheck(SizeType Index) const
 	{
 		CheckInvariants();
 
@@ -1077,7 +1077,7 @@ public:
 	 * @param Index Index of the start of the range to check.
 	 * @param Count Number of elements in the range.
 	 */
-	FORCEINLINE void RangeCheck(SizeType Index, SizeType Count) const
+	UE_FORCEINLINE_HINT void RangeCheck(SizeType Index, SizeType Count) const
 	{
 		CheckInvariants();
 
@@ -1094,7 +1094,7 @@ public:
 	 * @param Index Index to test.
 	 * @returns True if index is valid. False otherwise.
 	 */
-	UE_NODEBUG [[nodiscard]] FORCEINLINE bool IsValidIndex(SizeType Index) const
+	UE_NODEBUG [[nodiscard]] UE_FORCEINLINE_HINT bool IsValidIndex(SizeType Index) const
 	{
 		return Index >= 0 && Index < ArrayNum;
 	}
@@ -1122,7 +1122,7 @@ public:
 	}
 
 	/** @returns Number of bytes used, excluding slack */
-	UE_NODEBUG [[nodiscard]] FORCEINLINE SIZE_T NumBytes() const
+	UE_NODEBUG [[nodiscard]] UE_FORCEINLINE_HINT SIZE_T NumBytes() const
 	{
 		return static_cast<SIZE_T>(ArrayNum) * sizeof(ElementType);
 	}
@@ -1143,7 +1143,7 @@ public:
 	 *
 	 * @returns Reference to indexed element.
 	 */
-	UE_NODEBUG [[nodiscard]] FORCEINLINE ElementType& operator[](SizeType Index) UE_LIFETIMEBOUND
+	UE_NODEBUG [[nodiscard]] UE_FORCEINLINE_HINT ElementType& operator[](SizeType Index) UE_LIFETIMEBOUND
 	{
 		RangeCheck(Index);
 		return GetData()[Index];
@@ -1185,7 +1185,7 @@ public:
 		return Result;
 	}
 	UE_ALLOWSHRINKING_BOOL_DEPRECATED("Pop")
-	UE_NODEBUG FORCEINLINE ElementType Pop(bool bAllowShrinking)
+	UE_NODEBUG UE_FORCEINLINE_HINT ElementType Pop(bool bAllowShrinking)
 	{
 		return Pop(bAllowShrinking ? EAllowShrinking::Yes : EAllowShrinking::No);
 	}
@@ -1195,7 +1195,7 @@ public:
 	 *
 	 * @param Item Item to push.
 	 */
-	UE_NODEBUG FORCEINLINE void Push(ElementType&& Item)
+	UE_NODEBUG UE_FORCEINLINE_HINT void Push(ElementType&& Item)
 	{
 		Add(MoveTempIfPossible(Item));
 	}
@@ -1208,7 +1208,7 @@ public:
 	 * @param Item Item to push.
 	 * @see Pop, Top
 	 */
-	UE_NODEBUG FORCEINLINE void Push(const ElementType& Item)
+	UE_NODEBUG UE_FORCEINLINE_HINT void Push(const ElementType& Item)
 	{
 		Add(Item);
 	}
@@ -1219,7 +1219,7 @@ public:
 	 * @returns Reference to the top element.
 	 * @see Pop, Push
 	 */
-	UE_NODEBUG [[nodiscard]] FORCEINLINE ElementType& Top() UE_LIFETIMEBOUND
+	UE_NODEBUG [[nodiscard]] UE_FORCEINLINE_HINT ElementType& Top() UE_LIFETIMEBOUND
 	{
 		return Last();
 	}
@@ -1234,7 +1234,7 @@ public:
 	 * @param IndexFromTheEnd (Optional) Index from the end of array (default = 0).
 	 * @returns Reference to n-th last element from the array.
 	 */
-	UE_NODEBUG [[nodiscard]] FORCEINLINE ElementType& Last(SizeType IndexFromTheEnd = 0) UE_LIFETIMEBOUND
+	UE_NODEBUG [[nodiscard]] UE_FORCEINLINE_HINT ElementType& Last(SizeType IndexFromTheEnd = 0) UE_LIFETIMEBOUND
 	{
 		RangeCheck(ArrayNum - IndexFromTheEnd - 1);
 		return GetData()[ArrayNum - IndexFromTheEnd - 1];
@@ -1249,7 +1249,7 @@ public:
 	 *
 	 * @see Slack
 	 */
-	FORCEINLINE void Shrink()
+	UE_FORCEINLINE_HINT void Shrink()
 	{
 		CheckInvariants();
 		if (ArrayMax != ArrayNum)
@@ -1273,7 +1273,7 @@ public:
 	 * @returns True if found. False otherwise.
 	 * @see FindLast, FindLastByPredicate
 	 */
-	UE_NODEBUG FORCEINLINE bool Find(const ElementType& Item, SizeType& Index) const
+	UE_NODEBUG UE_FORCEINLINE_HINT bool Find(const ElementType& Item, SizeType& Index) const
 	{
 		Index = this->Find(Item);
 		return Index != INDEX_NONE;
@@ -1307,7 +1307,7 @@ public:
 	 * @returns True if found. False otherwise.
 	 * @see Find, FindLastByPredicate
 	 */
-	UE_NODEBUG FORCEINLINE bool FindLast(const ElementType& Item, SizeType& Index) const
+	UE_NODEBUG UE_FORCEINLINE_HINT bool FindLast(const ElementType& Item, SizeType& Index) const
 	{
 		Index = this->FindLast(Item);
 		return Index != INDEX_NONE;
@@ -1361,7 +1361,7 @@ public:
 	 * @returns Index of the found element. INDEX_NONE otherwise.
 	 */
 	template <typename Predicate>
-	UE_NODEBUG [[nodiscard]] FORCEINLINE SizeType FindLastByPredicate(Predicate Pred) const
+	UE_NODEBUG [[nodiscard]] UE_FORCEINLINE_HINT SizeType FindLastByPredicate(Predicate Pred) const
 	{
 		return FindLastByPredicate(Pred, ArrayNum);
 	}
@@ -1509,7 +1509,7 @@ public:
 	 * @see Contains, Find
 	 */
 	template <typename Predicate>
-	UE_NODEBUG [[nodiscard]] FORCEINLINE bool ContainsByPredicate(Predicate Pred) const
+	UE_NODEBUG [[nodiscard]] UE_FORCEINLINE_HINT bool ContainsByPredicate(Predicate Pred) const
 	{
 		return FindByPredicate(Pred) != nullptr;
 	}
@@ -1635,7 +1635,7 @@ public:
 	 * @param Count Number of elements to add.
 	 * @returns Number of elements in array before addition.
 	 */
-	FORCEINLINE SizeType AddUninitialized()
+	UE_FORCEINLINE_HINT SizeType AddUninitialized()
 	{
 		// Begin sensitive code! Single element additions/insertions get inlined _everywhere_ and thus
 		// punch above their weight with respect to exe size, so we pay a lot of attention to every instruction
@@ -1682,7 +1682,7 @@ public:
 		ArrayNum++;
 		return OldArrayNum;	
 	}
-	FORCEINLINE SizeType AddUninitialized(SizeType Count)
+	UE_FORCEINLINE_HINT SizeType AddUninitialized(SizeType Count)
 	{
 		// Should be SetNumUninitialized?
 		return UE::Core::Private::ReallocGrow<UE::Core::Private::GetAllocatorFlags<AllocatorType>()>(
@@ -1753,11 +1753,11 @@ public:
 	 * @param Count Number of elements to add.
 	 * @see Insert, InsertZeroed, InsertDefaulted
 	 */
-	UE_NODEBUG FORCEINLINE void InsertUninitialized(SizeType Index)
+	UE_NODEBUG UE_FORCEINLINE_HINT void InsertUninitialized(SizeType Index)
 	{
 		InsertUninitializedImpl(Index);
 	}
-	UE_NODEBUG FORCEINLINE void InsertUninitialized(SizeType Index, SizeType Count)
+	UE_NODEBUG UE_FORCEINLINE_HINT void InsertUninitialized(SizeType Index, SizeType Count)
 	{
 		InsertUninitializedImpl(Index, Count);
 	}
@@ -1926,7 +1926,7 @@ public:
 	 * @param Addr The address to check.
 	 * @see Add, Remove
 	 */
-	FORCEINLINE void CheckAddress(const ElementType* Addr) const
+	UE_FORCEINLINE_HINT void CheckAddress(const ElementType* Addr) const
 	{
 		checkf(Addr < GetData() || Addr >= (GetData() + ArrayMax), TEXT("Attempting to use a container element (%p) which already comes from the container being modified (%p, ArrayMax: %lld, ArrayNum: %lld, SizeofElement: %zu)!"), Addr, GetData(), (long long)ArrayMax, (long long)ArrayNum, sizeof(ElementType));
 	}
@@ -2079,7 +2079,7 @@ public:
 	 * @param AllowShrinking (Optional) Tells if this call can shrink array if suitable after remove. Default is yes.
 	 */
 	template <UE::CIntegral CountType>
-	FORCEINLINE void RemoveAt(SizeType Index, CountType Count, EAllowShrinking AllowShrinking = EAllowShrinking::Default)
+	UE_FORCEINLINE_HINT void RemoveAt(SizeType Index, CountType Count, EAllowShrinking AllowShrinking = EAllowShrinking::Default)
 	{
 		static_assert(!std::is_same_v<CountType, bool>, "TArray::RemoveAt: unexpected bool passed as the Count argument");
 		RangeCheck(Index, Count);
@@ -2100,7 +2100,7 @@ public:
 	}
 	template <typename CountType>
 	UE_ALLOWSHRINKING_BOOL_DEPRECATED("RemoveAt")
-	FORCEINLINE void RemoveAt(SizeType Index, CountType Count, bool bAllowShrinking)
+	UE_FORCEINLINE_HINT void RemoveAt(SizeType Index, CountType Count, bool bAllowShrinking)
 	{
 		RemoveAt(Index, Count, bAllowShrinking ? EAllowShrinking::Yes : EAllowShrinking::No);
 	}
@@ -2156,7 +2156,7 @@ public:
 	 * @param AllowShrinking (Optional) Tells if this call can shrink array if
 	 *                        suitable after remove. Default is yes.
 	 */
-	FORCEINLINE void RemoveAtSwap(SizeType Index, EAllowShrinking AllowShrinking = EAllowShrinking::Default)
+	UE_FORCEINLINE_HINT void RemoveAtSwap(SizeType Index, EAllowShrinking AllowShrinking = EAllowShrinking::Default)
 	{
 		RangeCheck(Index);
 		RemoveAtSwapImpl(Index);
@@ -2185,7 +2185,7 @@ public:
 	 *                        suitable after remove. Default is yes.
 	 */
 	template <UE::CIntegral CountType>
-	FORCEINLINE void RemoveAtSwap(SizeType Index, CountType Count, EAllowShrinking AllowShrinking = EAllowShrinking::Default)
+	UE_FORCEINLINE_HINT void RemoveAtSwap(SizeType Index, CountType Count, EAllowShrinking AllowShrinking = EAllowShrinking::Default)
 	{
 		static_assert(!std::is_same_v<CountType, bool>, "TArray::RemoveAtSwap: unexpected bool passed as the Count argument");
 		RangeCheck(Index, Count);
@@ -2206,7 +2206,7 @@ public:
 	}
 	template <typename CountType>
 	UE_ALLOWSHRINKING_BOOL_DEPRECATED("RemoveAtSwap")
-	UE_NODEBUG FORCEINLINE void RemoveAtSwap(SizeType Index, CountType Count, bool bAllowShrinking)
+	UE_NODEBUG UE_FORCEINLINE_HINT void RemoveAtSwap(SizeType Index, CountType Count, bool bAllowShrinking)
 	{
 		RemoveAtSwap(Index, Count, bAllowShrinking ? EAllowShrinking::Yes : EAllowShrinking::No);
 	}
@@ -2297,7 +2297,7 @@ public:
 		}
 	}
 	UE_ALLOWSHRINKING_BOOL_DEPRECATED("SetNum")
-	UE_NODEBUG FORCEINLINE void SetNum(SizeType NewNum, bool bAllowShrinking)
+	UE_NODEBUG UE_FORCEINLINE_HINT void SetNum(SizeType NewNum, bool bAllowShrinking)
 	{
 		SetNum(NewNum, bAllowShrinking ? EAllowShrinking::Yes : EAllowShrinking::No);
 	}
@@ -2326,7 +2326,7 @@ public:
 		}
 	}
 	UE_ALLOWSHRINKING_BOOL_DEPRECATED("SetNumZeroed")
-	UE_NODEBUG FORCEINLINE void SetNumZeroed(SizeType NewNum, bool bAllowShrinking)
+	UE_NODEBUG UE_FORCEINLINE_HINT void SetNumZeroed(SizeType NewNum, bool bAllowShrinking)
 	{
 		SetNumZeroed(NewNum, bAllowShrinking ? EAllowShrinking::Yes : EAllowShrinking::No);
 	}
@@ -2354,7 +2354,7 @@ public:
 		}
 	}
 	UE_ALLOWSHRINKING_BOOL_DEPRECATED("SetNumUninitialized")
-	UE_NODEBUG FORCEINLINE void SetNumUninitialized(SizeType NewNum, bool bAllowShrinking)
+	UE_NODEBUG UE_FORCEINLINE_HINT void SetNumUninitialized(SizeType NewNum, bool bAllowShrinking)
 	{
 		SetNumUninitialized(NewNum, bAllowShrinking ? EAllowShrinking::Yes : EAllowShrinking::No);
 	}
@@ -2477,7 +2477,7 @@ public:
 	 * @param InitList The initializer list of elements to add.
 	 * @see Add, Insert
 	 */
-	FORCEINLINE void Append(std::initializer_list<ElementType> InitList)
+	UE_FORCEINLINE_HINT void Append(std::initializer_list<ElementType> InitList)
 	{
 		SizeType Count = (SizeType)InitList.size();
 
@@ -2529,7 +2529,7 @@ public:
 	 * @return		Index to the new item
 	 */
 	template <typename... ArgsType>
-	FORCEINLINE SizeType Emplace(ArgsType&&... Args)
+	UE_FORCEINLINE_HINT SizeType Emplace(ArgsType&&... Args)
 	{
 		// If this fails to compile when trying to call Emplace with a non-public constructor,
 		// do not make TArray a friend.
@@ -2581,7 +2581,7 @@ public:
 	 * @return A reference to the newly-inserted element.
 	 */
 	template <typename... ArgsType>
-	[[nodiscard]] FORCEINLINE ElementType& Emplace_GetRef(ArgsType&&... Args) UE_LIFETIMEBOUND
+	[[nodiscard]] UE_FORCEINLINE_HINT ElementType& Emplace_GetRef(ArgsType&&... Args) UE_LIFETIMEBOUND
 	{
 		// If this fails to compile when trying to call Emplace with a non-public constructor,
 		// do not make TArray a friend.
@@ -2633,7 +2633,7 @@ public:
 	 * @param Args	The arguments to forward to the constructor of the new item.
 	 */
 	template <typename... ArgsType>
-	FORCEINLINE void EmplaceAt(SizeType Index, ArgsType&&... Args)
+	UE_FORCEINLINE_HINT void EmplaceAt(SizeType Index, ArgsType&&... Args)
 	{
 		InsertUninitializedImpl(Index, 1);
 		::new((void*)(GetData() + Index)) ElementType(Forward<ArgsType>(Args)...);
@@ -2647,7 +2647,7 @@ public:
 	 * @return A reference to the newly-inserted element.
 	 */
 	template <typename... ArgsType>
-	[[nodiscard]] FORCEINLINE ElementType& EmplaceAt_GetRef(SizeType Index, ArgsType&&... Args) UE_LIFETIMEBOUND
+	[[nodiscard]] UE_FORCEINLINE_HINT ElementType& EmplaceAt_GetRef(SizeType Index, ArgsType&&... Args) UE_LIFETIMEBOUND
 	{
 		InsertUninitializedImpl(Index, 1);
 		ElementType* Ptr = GetData() + Index;
@@ -2664,7 +2664,7 @@ public:
 	 * @return Index to the new item
 	 * @see AddDefaulted, AddUnique, AddZeroed, Append, Insert
 	 */
-	UE_NODEBUG FORCEINLINE SizeType Add(ElementType&& Item)
+	UE_NODEBUG UE_FORCEINLINE_HINT SizeType Add(ElementType&& Item)
 	{
 		CheckAddress(&Item);
 		return Emplace(MoveTempIfPossible(Item));
@@ -2677,7 +2677,7 @@ public:
 	 * @return Index to the new item
 	 * @see AddDefaulted, AddUnique, AddZeroed, Append, Insert
 	 */
-	UE_NODEBUG FORCEINLINE SizeType Add(const ElementType& Item)
+	UE_NODEBUG UE_FORCEINLINE_HINT SizeType Add(const ElementType& Item)
 	{
 		CheckAddress(&Item);
 		return Emplace(Item);
@@ -2692,7 +2692,7 @@ public:
 	 * @return A reference to the newly-inserted element.
 	 * @see AddDefaulted_GetRef, AddUnique_GetRef, AddZeroed_GetRef, Insert_GetRef
 	 */
-	UE_NODEBUG [[nodiscard]] FORCEINLINE ElementType& Add_GetRef(ElementType&& Item) UE_LIFETIMEBOUND
+	UE_NODEBUG [[nodiscard]] UE_FORCEINLINE_HINT ElementType& Add_GetRef(ElementType&& Item) UE_LIFETIMEBOUND
 	{
 		CheckAddress(&Item);
 		return Emplace_GetRef(MoveTempIfPossible(Item));
@@ -2705,7 +2705,7 @@ public:
 	 * @return A reference to the newly-inserted element.
 	 * @see AddDefaulted_GetRef, AddUnique_GetRef, AddZeroed_GetRef, Insert_GetRef
 	 */
-	UE_NODEBUG [[nodiscard]] FORCEINLINE ElementType& Add_GetRef(const ElementType& Item) UE_LIFETIMEBOUND
+	UE_NODEBUG [[nodiscard]] UE_FORCEINLINE_HINT ElementType& Add_GetRef(const ElementType& Item) UE_LIFETIMEBOUND
 	{
 		CheckAddress(&Item);
 		return Emplace_GetRef(Item);
@@ -2961,7 +2961,7 @@ public:
 	 * @returns Index of the element in the array.
 	 * @see Add, AddDefaulted, AddZeroed, Append, Insert
 	 */
-	FORCEINLINE SizeType AddUnique(ElementType&& Item)
+	UE_FORCEINLINE_HINT SizeType AddUnique(ElementType&& Item)
 	{
 		return AddUniqueImpl(MoveTempIfPossible(Item));
 	}
@@ -2973,7 +2973,7 @@ public:
 	 * @returns Index of the element in the array.
 	 * @see Add, AddDefaulted, AddZeroed, Append, Insert
 	 */
-	FORCEINLINE SizeType AddUnique(const ElementType& Item)
+	UE_FORCEINLINE_HINT SizeType AddUnique(const ElementType& Item)
 	{
 		return AddUniqueImpl(Item);
 	}
@@ -2984,7 +2984,7 @@ public:
 	 * @param Number The number of elements that the array should be able to contain after allocation.
 	 * @see Shrink
 	 */
-	FORCEINLINE void Reserve(SizeType Number)
+	UE_FORCEINLINE_HINT void Reserve(SizeType Number)
 	{
 		checkSlow(Number >= 0);
 		if (Number < 0)
@@ -3162,7 +3162,7 @@ public:
 	}
 	template <class PREDICATE_CLASS>
 	UE_ALLOWSHRINKING_BOOL_DEPRECATED("RemoveAllSwap")
-	UE_NODEBUG FORCEINLINE SizeType RemoveAllSwap(const PREDICATE_CLASS& Predicate, bool bAllowShrinking)
+	UE_NODEBUG UE_FORCEINLINE_HINT SizeType RemoveAllSwap(const PREDICATE_CLASS& Predicate, bool bAllowShrinking)
 	{
 		return RemoveAllSwap(Predicate, bAllowShrinking ? EAllowShrinking::Yes : EAllowShrinking::No);
 	}
@@ -3191,7 +3191,7 @@ public:
 		return 1;
 	}
 	UE_ALLOWSHRINKING_BOOL_DEPRECATED("RemoveSingleSwap")
-	UE_NODEBUG FORCEINLINE SizeType RemoveSingleSwap(const ElementType& Item, bool bAllowShrinking)
+	UE_NODEBUG UE_FORCEINLINE_HINT SizeType RemoveSingleSwap(const ElementType& Item, bool bAllowShrinking)
 	{
 		return RemoveSingleSwap(Item, bAllowShrinking ? EAllowShrinking::Yes : EAllowShrinking::No);
 	}
@@ -3238,7 +3238,7 @@ public:
 		return OriginalNum - ArrayNum;
 	}
 	UE_ALLOWSHRINKING_BOOL_DEPRECATED("RemoveSwap")
-	UE_NODEBUG FORCEINLINE SizeType RemoveSwap(const ElementType& Item, bool bAllowShrinking)
+	UE_NODEBUG UE_FORCEINLINE_HINT SizeType RemoveSwap(const ElementType& Item, bool bAllowShrinking)
 	{
 		return RemoveSwap(Item, bAllowShrinking ? EAllowShrinking::Yes : EAllowShrinking::No);
 	}
@@ -3249,7 +3249,7 @@ public:
 	 * @param FirstIndexToSwap Position of the first element to swap.
 	 * @param SecondIndexToSwap Position of the second element to swap.
 	 */
-	FORCEINLINE void SwapMemory(SizeType FirstIndexToSwap, SizeType SecondIndexToSwap)
+	UE_FORCEINLINE_HINT void SwapMemory(SizeType FirstIndexToSwap, SizeType SecondIndexToSwap)
 	{
 		::Swap(
 			*(ElementType*)((uint8*)AllocatorInstance.GetAllocation() + (sizeof(ElementType)*FirstIndexToSwap)),
@@ -3265,7 +3265,7 @@ public:
 	 * @param FirstIndexToSwap Position of the first element to swap.
 	 * @param SecondIndexToSwap Position of the second element to swap.
 	 */
-	FORCEINLINE void Swap(SizeType FirstIndexToSwap, SizeType SecondIndexToSwap)
+	UE_FORCEINLINE_HINT void Swap(SizeType FirstIndexToSwap, SizeType SecondIndexToSwap)
 	{
 		check((FirstIndexToSwap >= 0) && (SecondIndexToSwap >= 0));
 		check((ArrayNum > FirstIndexToSwap) && (ArrayNum > SecondIndexToSwap));
@@ -3354,23 +3354,23 @@ public:
 	 * STL-like iterators to enable range-based for loop support.
 	 */
 	#if TARRAY_RANGED_FOR_CHECKS
-		[[nodiscard]] UE_NODEBUG FORCEINLINE RangedForIteratorType             begin ()       { return RangedForIteratorType            (ArrayNum, GetData()); }
-		[[nodiscard]] UE_NODEBUG FORCEINLINE RangedForConstIteratorType        begin () const { return RangedForConstIteratorType       (ArrayNum, GetData()); }
-		[[nodiscard]] UE_NODEBUG FORCEINLINE RangedForIteratorType             end   ()       { return RangedForIteratorType            (ArrayNum, GetData() + Num()); }
-		[[nodiscard]] UE_NODEBUG FORCEINLINE RangedForConstIteratorType        end   () const { return RangedForConstIteratorType       (ArrayNum, GetData() + Num()); }
-		[[nodiscard]] UE_NODEBUG FORCEINLINE RangedForReverseIteratorType      rbegin()       { return RangedForReverseIteratorType     (ArrayNum, GetData() + Num()); }
-		[[nodiscard]] UE_NODEBUG FORCEINLINE RangedForConstReverseIteratorType rbegin() const { return RangedForConstReverseIteratorType(ArrayNum, GetData() + Num()); }
-		[[nodiscard]] UE_NODEBUG FORCEINLINE RangedForReverseIteratorType      rend  ()       { return RangedForReverseIteratorType     (ArrayNum, GetData()); }
-		[[nodiscard]] UE_NODEBUG FORCEINLINE RangedForConstReverseIteratorType rend  () const { return RangedForConstReverseIteratorType(ArrayNum, GetData()); }
+		[[nodiscard]] UE_NODEBUG UE_FORCEINLINE_HINT RangedForIteratorType             begin ()       { return RangedForIteratorType            (ArrayNum, GetData()); }
+		[[nodiscard]] UE_NODEBUG UE_FORCEINLINE_HINT RangedForConstIteratorType        begin () const { return RangedForConstIteratorType       (ArrayNum, GetData()); }
+		[[nodiscard]] UE_NODEBUG UE_FORCEINLINE_HINT RangedForIteratorType             end   ()       { return RangedForIteratorType            (ArrayNum, GetData() + Num()); }
+		[[nodiscard]] UE_NODEBUG UE_FORCEINLINE_HINT RangedForConstIteratorType        end   () const { return RangedForConstIteratorType       (ArrayNum, GetData() + Num()); }
+		[[nodiscard]] UE_NODEBUG UE_FORCEINLINE_HINT RangedForReverseIteratorType      rbegin()       { return RangedForReverseIteratorType     (ArrayNum, GetData() + Num()); }
+		[[nodiscard]] UE_NODEBUG UE_FORCEINLINE_HINT RangedForConstReverseIteratorType rbegin() const { return RangedForConstReverseIteratorType(ArrayNum, GetData() + Num()); }
+		[[nodiscard]] UE_NODEBUG UE_FORCEINLINE_HINT RangedForReverseIteratorType      rend  ()       { return RangedForReverseIteratorType     (ArrayNum, GetData()); }
+		[[nodiscard]] UE_NODEBUG UE_FORCEINLINE_HINT RangedForConstReverseIteratorType rend  () const { return RangedForConstReverseIteratorType(ArrayNum, GetData()); }
 	#else
-		[[nodiscard]] UE_NODEBUG FORCEINLINE RangedForIteratorType             begin ()       { return                                   GetData(); }
-		[[nodiscard]] UE_NODEBUG FORCEINLINE RangedForConstIteratorType        begin () const { return                                   GetData(); }
-		[[nodiscard]] UE_NODEBUG FORCEINLINE RangedForIteratorType             end   ()       { return                                   GetData() + Num(); }
-		[[nodiscard]] UE_NODEBUG FORCEINLINE RangedForConstIteratorType        end   () const { return                                   GetData() + Num(); }
-		[[nodiscard]] UE_NODEBUG FORCEINLINE RangedForReverseIteratorType      rbegin()       { return RangedForReverseIteratorType     (GetData() + Num()); }
-		[[nodiscard]] UE_NODEBUG FORCEINLINE RangedForConstReverseIteratorType rbegin() const { return RangedForConstReverseIteratorType(GetData() + Num()); }
-		[[nodiscard]] UE_NODEBUG FORCEINLINE RangedForReverseIteratorType      rend  ()       { return RangedForReverseIteratorType     (GetData()); }
-		[[nodiscard]] UE_NODEBUG FORCEINLINE RangedForConstReverseIteratorType rend  () const { return RangedForConstReverseIteratorType(GetData()); }
+		[[nodiscard]] UE_NODEBUG UE_FORCEINLINE_HINT RangedForIteratorType             begin ()       { return                                   GetData(); }
+		[[nodiscard]] UE_NODEBUG UE_FORCEINLINE_HINT RangedForConstIteratorType        begin () const { return                                   GetData(); }
+		[[nodiscard]] UE_NODEBUG UE_FORCEINLINE_HINT RangedForIteratorType             end   ()       { return                                   GetData() + Num(); }
+		[[nodiscard]] UE_NODEBUG UE_FORCEINLINE_HINT RangedForConstIteratorType        end   () const { return                                   GetData() + Num(); }
+		[[nodiscard]] UE_NODEBUG UE_FORCEINLINE_HINT RangedForReverseIteratorType      rbegin()       { return RangedForReverseIteratorType     (GetData() + Num()); }
+		[[nodiscard]] UE_NODEBUG UE_FORCEINLINE_HINT RangedForConstReverseIteratorType rbegin() const { return RangedForConstReverseIteratorType(GetData() + Num()); }
+		[[nodiscard]] UE_NODEBUG UE_FORCEINLINE_HINT RangedForReverseIteratorType      rend  ()       { return RangedForReverseIteratorType     (GetData()); }
+		[[nodiscard]] UE_NODEBUG UE_FORCEINLINE_HINT RangedForConstReverseIteratorType rend  () const { return RangedForConstReverseIteratorType(GetData()); }
 	#endif
 
 public:
@@ -3605,7 +3605,7 @@ public:
 	 *        The auto-dereferencing behavior does not occur with smart pointers.
 	 */
 	template <class PREDICATE_CLASS>
-	UE_NODEBUG FORCEINLINE void Heapify(const PREDICATE_CLASS& Predicate)
+	UE_NODEBUG UE_FORCEINLINE_HINT void Heapify(const PREDICATE_CLASS& Predicate)
 	{
 		TDereferenceWrapper<ElementType, PREDICATE_CLASS> PredicateWrapper(Predicate);
 		Algo::Heapify(*this, PredicateWrapper);
@@ -3722,7 +3722,7 @@ public:
 	}
 	template <class PREDICATE_CLASS>
 	UE_ALLOWSHRINKING_BOOL_DEPRECATED("HeapPop")
-	UE_NODEBUG FORCEINLINE void HeapPop(ElementType& OutItem, const PREDICATE_CLASS& Predicate, bool bAllowShrinking)
+	UE_NODEBUG UE_FORCEINLINE_HINT void HeapPop(ElementType& OutItem, const PREDICATE_CLASS& Predicate, bool bAllowShrinking)
 	{
 		HeapPop(OutItem, Predicate, bAllowShrinking ? EAllowShrinking::Yes : EAllowShrinking::No);
 	}
@@ -3743,7 +3743,7 @@ public:
 		HeapPop(OutItem, TLess<ElementType>(), AllowShrinking);
 	}
 	UE_ALLOWSHRINKING_BOOL_DEPRECATED("HeapPop")
-	UE_NODEBUG FORCEINLINE void HeapPop(ElementType& OutItem, bool bAllowShrinking)
+	UE_NODEBUG UE_FORCEINLINE_HINT void HeapPop(ElementType& OutItem, bool bAllowShrinking)
 	{
 		HeapPop(OutItem, bAllowShrinking ? EAllowShrinking::Yes : EAllowShrinking::No);
 	}
@@ -3778,7 +3778,7 @@ public:
 	}
 	template <class PREDICATE_CLASS>
 	UE_ALLOWSHRINKING_BOOL_DEPRECATED("HeapPopDiscard")
-	UE_NODEBUG FORCEINLINE void HeapPopDiscard(const PREDICATE_CLASS& Predicate, bool bAllowShrinking)
+	UE_NODEBUG UE_FORCEINLINE_HINT void HeapPopDiscard(const PREDICATE_CLASS& Predicate, bool bAllowShrinking)
 	{
 		HeapPopDiscard(Predicate, bAllowShrinking ? EAllowShrinking::Yes : EAllowShrinking::No);
 	}
@@ -3798,7 +3798,7 @@ public:
 		HeapPopDiscard(TLess<ElementType>(), AllowShrinking);
 	}
 	UE_ALLOWSHRINKING_BOOL_DEPRECATED("HeapPopDiscard")
-	UE_NODEBUG FORCEINLINE void HeapPopDiscard(bool bAllowShrinking)
+	UE_NODEBUG UE_FORCEINLINE_HINT void HeapPopDiscard(bool bAllowShrinking)
 	{
 		HeapPopDiscard(TLess<ElementType>(), bAllowShrinking ? EAllowShrinking::Yes : EAllowShrinking::No);
 	}
@@ -3869,7 +3869,7 @@ public:
 		HeapRemoveAt(Index, TLess< ElementType >(), AllowShrinking);
 	}
 	UE_ALLOWSHRINKING_BOOL_DEPRECATED("HeapRemoveAt")
-	UE_NODEBUG FORCEINLINE void HeapRemoveAt(SizeType Index, bool bAllowShrinking)
+	UE_NODEBUG UE_FORCEINLINE_HINT void HeapRemoveAt(SizeType Index, bool bAllowShrinking)
 	{
 		HeapRemoveAt(Index, bAllowShrinking ? EAllowShrinking::Yes : EAllowShrinking::No);
 	}
diff --git a/Engine/Source/Runtime/Core/Public/Containers/ArrayView.h b/Engine/Source/Runtime/Core/Public/Containers/ArrayView.h
--- a/Engine/Source/Runtime/Core/Public/Containers/ArrayView.h
+++ b/Engine/Source/Runtime/Core/Public/Containers/ArrayView.h
@@ -29,7 +29,7 @@ namespace UE::Core::ArrayView::Private
 	// Simply forwards to an unqualified GetData(), but can be called from within TArrayView
 	// where GetData() is already a member and so hides any others.
 	template <typename T>
-	FORCEINLINE auto GetDataHelper(T&& Arg) -> decltype(GetData(Forward<T>(Arg)))
+	UE_FORCEINLINE_HINT auto GetDataHelper(T&& Arg) -> decltype(GetData(Forward<T>(Arg)))
 	{
 		return GetData(Forward<T>(Arg));
 	}
@@ -292,7 +292,7 @@ public:
 	 *
 	 * @returns Pointer to first array entry.
 	 */
-	[[nodiscard]] FORCEINLINE ElementType* GetData() const
+	[[nodiscard]] UE_FORCEINLINE_HINT ElementType* GetData() const
 	{
 		return DataPtr;
 	}
@@ -302,7 +302,7 @@ public:
 	 *
 	 * @returns Size in bytes of array type.
 	 */
-	[[nodiscard]] FORCEINLINE static constexpr size_t GetTypeSize()
+	[[nodiscard]] UE_FORCEINLINE_HINT static constexpr size_t GetTypeSize()
 	{
 		return sizeof(ElementType);
 	}
@@ -310,7 +310,7 @@ public:
 	/** 
 	 * Helper function returning the alignment of the inner type.
 	 */
-	[[nodiscard]] FORCEINLINE static constexpr size_t GetTypeAlignment()
+	[[nodiscard]] UE_FORCEINLINE_HINT static constexpr size_t GetTypeAlignment()
 	{
 		return alignof(ElementType);
 	}
@@ -318,7 +318,7 @@ public:
 	/**
 	 * Checks array invariants: if array size is greater than or equal to zero.
 	 */
-	FORCEINLINE void CheckInvariants() const
+	UE_FORCEINLINE_HINT void CheckInvariants() const
 	{
 		checkSlow(ArrayNum >= 0);
 	}
@@ -356,7 +356,7 @@ public:
 	 *
 	 * @returns True if index is valid. False otherwise.
 	 */
-	[[nodiscard]] FORCEINLINE bool IsValidIndex(SizeType Index) const
+	[[nodiscard]] UE_FORCEINLINE_HINT bool IsValidIndex(SizeType Index) const
 	{
 		return (Index >= 0) && (Index < ArrayNum);
 	}
@@ -377,13 +377,13 @@ public:
 	 *
 	 * @returns Number of elements in array.
 	 */
-	[[nodiscard]] FORCEINLINE SizeType Num() const
+	[[nodiscard]] UE_FORCEINLINE_HINT SizeType Num() const
 	{
 		return ArrayNum;
 	}
 
 	/** @returns Number of bytes used, excluding slack */
-	[[nodiscard]] FORCEINLINE SIZE_T NumBytes() const
+	[[nodiscard]] UE_FORCEINLINE_HINT SIZE_T NumBytes() const
 	{
 		return static_cast<SIZE_T>(ArrayNum) * sizeof(ElementType);
 	}
@@ -604,7 +604,7 @@ public:
 	* @returns Index of the found element. INDEX_NONE otherwise.
 	*/
 	template <typename Predicate>
-	[[nodiscard]] FORCEINLINE SizeType FindLastByPredicate(Predicate Pred) const
+	[[nodiscard]] UE_FORCEINLINE_HINT SizeType FindLastByPredicate(Predicate Pred) const
 	{
 		return FindLastByPredicate(Pred, ArrayNum);
 	}
@@ -746,7 +746,7 @@ public:
 	 * @returns	True if found. False otherwise.
 	 */
 	template <typename Predicate>
-	[[nodiscard]] FORCEINLINE bool ContainsByPredicate(Predicate Pred) const
+	[[nodiscard]] UE_FORCEINLINE_HINT bool ContainsByPredicate(Predicate Pred) const
 	{
 		return FindByPredicate(Pred) != nullptr;
 	}
@@ -756,10 +756,10 @@ public:
 	 * DO NOT USE DIRECTLY
 	 * STL-like iterators to enable range-based for loop support.
 	 */
-	[[nodiscard]] FORCEINLINE ElementType*                         begin () const { return GetData(); }
-	[[nodiscard]] FORCEINLINE ElementType*                         end   () const { return GetData() + Num(); }
-	[[nodiscard]] FORCEINLINE TReversePointerIterator<ElementType> rbegin() const { return TReversePointerIterator<ElementType>(GetData() + Num()); }
-	[[nodiscard]] FORCEINLINE TReversePointerIterator<ElementType> rend  () const { return TReversePointerIterator<ElementType>(GetData()); }
+	[[nodiscard]] UE_FORCEINLINE_HINT ElementType*                         begin () const { return GetData(); }
+	[[nodiscard]] UE_FORCEINLINE_HINT ElementType*                         end   () const { return GetData() + Num(); }
+	[[nodiscard]] UE_FORCEINLINE_HINT TReversePointerIterator<ElementType> rbegin() const { return TReversePointerIterator<ElementType>(GetData() + Num()); }
+	[[nodiscard]] UE_FORCEINLINE_HINT TReversePointerIterator<ElementType> rend  () const { return TReversePointerIterator<ElementType>(GetData()); }
 
 public:
 	/**
@@ -928,7 +928,7 @@ template <typename ElementType>
 
 template<typename InElementType, typename InAllocatorType>
 template<typename OtherElementType, typename OtherSizeType>
-[[nodiscard]] FORCEINLINE TArray<InElementType, InAllocatorType>::TArray(const TArrayView<OtherElementType, OtherSizeType>& Other)
+[[nodiscard]] UE_FORCEINLINE_HINT TArray<InElementType, InAllocatorType>::TArray(const TArrayView<OtherElementType, OtherSizeType>& Other)
 {
 	CopyToEmpty(Other.GetData(), Other.Num(), 0);
 }
diff --git a/Engine/Source/Runtime/Core/Public/Containers/BasicArray.h b/Engine/Source/Runtime/Core/Public/Containers/BasicArray.h
--- a/Engine/Source/Runtime/Core/Public/Containers/BasicArray.h
+++ b/Engine/Source/Runtime/Core/Public/Containers/BasicArray.h
@@ -105,17 +105,17 @@ public:
 		return LocalData ? &LocalData->Data[0] : nullptr;
 	}
 
-	[[nodiscard]] FORCEINLINE const T* GetData() const
+	[[nodiscard]] UE_FORCEINLINE_HINT const T* GetData() const
 	{
 		return const_cast<TBasicArray*>(this)->GetData();
 	}
 
-	[[nodiscard]] FORCEINLINE T& operator[](int32 Index)
+	[[nodiscard]] UE_FORCEINLINE_HINT T& operator[](int32 Index)
 	{
 		return GetData()[Index];
 	}
 
-	[[nodiscard]] FORCEINLINE const T& operator[](int32 Index) const
+	[[nodiscard]] UE_FORCEINLINE_HINT const T& operator[](int32 Index) const
 	{
 		return GetData()[Index];
 	}
@@ -123,7 +123,7 @@ public:
 private:
 	static constexpr int32 InitialReservationSize = 16;
 
-	static FORCEINLINE int32 ApplyGrowthFactor(int32 CurrentNum)
+	static UE_FORCEINLINE_HINT int32 ApplyGrowthFactor(int32 CurrentNum)
 	{
 		return CurrentNum + CurrentNum / 2;
 	}
diff --git a/Engine/Source/Runtime/Core/Public/Containers/BinaryHeap.h b/Engine/Source/Runtime/Core/Public/Containers/BinaryHeap.h
--- a/Engine/Source/Runtime/Core/Public/Containers/BinaryHeap.h
+++ b/Engine/Source/Runtime/Core/Public/Containers/BinaryHeap.h
@@ -72,7 +72,7 @@ protected:
 };
 
 template< typename KeyType, typename IndexType >
-[[nodiscard]] FORCEINLINE FBinaryHeap< KeyType, IndexType >::FBinaryHeap()
+[[nodiscard]] UE_FORCEINLINE_HINT FBinaryHeap< KeyType, IndexType >::FBinaryHeap()
 {
 	ResetInternal();
 }
@@ -105,7 +105,7 @@ template< typename KeyType, typename IndexType >
 }
 
 template< typename KeyType, typename IndexType >
-FORCEINLINE FBinaryHeap< KeyType, IndexType >::~FBinaryHeap()
+UE_FORCEINLINE_HINT FBinaryHeap< KeyType, IndexType >::~FBinaryHeap()
 {
 	Free();
 }
