Replace some FORCEINLINE usage in Core/CoreUObject with UE_FORCEINLINE_HINT.

diff --git a/Engine/Source/Runtime/Core/Public/Containers/BitArray.h b/Engine/Source/Runtime/Core/Public/Containers/BitArray.h
--- a/Engine/Source/Runtime/Core/Public/Containers/BitArray.h
+++ b/Engine/Source/Runtime/Core/Public/Containers/BitArray.h
@@ -25,7 +25,7 @@ template<typename Allocator > class TBitArray;
 struct FBitSet
 {
 	/** Clears the next set bit in the mask and returns its index. */
-	static FORCEINLINE uint32 GetAndClearNextBit(uint32& Mask)
+	static UE_FORCEINLINE_HINT uint32 GetAndClearNextBit(uint32& Mask)
 	{
 		const uint32 LowestBitMask = (Mask) & (-(int32)Mask);
 		const uint32 BitIndex = FMath::FloorLog2(LowestBitMask);
@@ -36,7 +36,7 @@ struct FBitSet
 	// Clang generates 7 instructions for int32 DivideAndRoundUp but only 2 for uint32
 	static constexpr uint32 BitsPerWord = NumBitsPerDWORD;
 
-	FORCEINLINE static uint32 CalculateNumWords(int32 NumBits)
+	UE_FORCEINLINE_HINT static uint32 CalculateNumWords(int32 NumBits)
 	{
 		checkSlow(NumBits >= 0);
 		return FMath::DivideAndRoundUp(static_cast<uint32>(NumBits), BitsPerWord);
@@ -189,16 +189,16 @@ class FBitReference
 {
 public:
 
-	FORCEINLINE explicit FBitReference(uint32& InData UE_LIFETIMEBOUND,uint32 InMask)
+	UE_FORCEINLINE_HINT explicit FBitReference(uint32& InData UE_LIFETIMEBOUND,uint32 InMask)
 	:	Data(InData)
 	,	Mask(InMask)
 	{}
 
-	[[nodiscard]] FORCEINLINE operator bool() const
+	[[nodiscard]] UE_FORCEINLINE_HINT operator bool() const
 	{
 		 return (Data & Mask) != 0;
 	}
-	FORCEINLINE void operator=(const bool NewValue)
+	UE_FORCEINLINE_HINT void operator=(const bool NewValue)
 	{
 		if(NewValue)
 		{
@@ -209,21 +209,21 @@ public:
 			Data &= ~Mask;
 		}
 	}
-	FORCEINLINE void operator|=(const bool NewValue)
+	UE_FORCEINLINE_HINT void operator|=(const bool NewValue)
 	{
 		if (NewValue)
 		{
 			Data |= Mask;
 		}
 	}
-	FORCEINLINE void operator&=(const bool NewValue)
+	UE_FORCEINLINE_HINT void operator&=(const bool NewValue)
 	{
 		if (!NewValue)
 		{
 			Data &= ~Mask;
 		}
 	}
-	FORCEINLINE void AtomicSet(const bool NewValue)
+	UE_FORCEINLINE_HINT void AtomicSet(const bool NewValue)
 	{
 		uint32 Current = static_cast<uint32>(FPlatformAtomics::AtomicRead_Relaxed((const volatile int32*)&Data));
 		if(NewValue)
@@ -241,7 +241,7 @@ public:
 			}
 		}
 	}
-	FORCEINLINE FBitReference& operator=(const FBitReference& Copy)
+	UE_FORCEINLINE_HINT FBitReference& operator=(const FBitReference& Copy)
 	{
 		// As this is emulating a reference, assignment should not rebind,
 		// it should write to the referenced bit.
@@ -260,12 +260,12 @@ class FConstBitReference
 {
 public:
 
-	FORCEINLINE explicit FConstBitReference(const uint32& InData UE_LIFETIMEBOUND,uint32 InMask)
+	UE_FORCEINLINE_HINT explicit FConstBitReference(const uint32& InData UE_LIFETIMEBOUND,uint32 InMask)
 	:	Data(InData)
 	,	Mask(InMask)
 	{}
 
-	[[nodiscard]] FORCEINLINE operator bool() const
+	[[nodiscard]] UE_FORCEINLINE_HINT operator bool() const
 	{
 		 return (Data & Mask) != 0;
 	}
@@ -280,7 +280,7 @@ private:
 class FRelativeBitReference
 {
 public:
-	FORCEINLINE explicit FRelativeBitReference(int32 BitIndex)
+	UE_FORCEINLINE_HINT explicit FRelativeBitReference(int32 BitIndex)
 		: WordIndex(BitIndex >> NumBitsPerDWORDLogTwo)
 		, Mask(1 << (BitIndex & (NumBitsPerDWORD - 1)))
 	{
@@ -289,12 +289,12 @@ public:
 	int32  WordIndex;
 	uint32 Mask;
 
-	[[nodiscard]] FORCEINLINE bool operator==(FRelativeBitReference Other) const
+	[[nodiscard]] UE_FORCEINLINE_HINT bool operator==(FRelativeBitReference Other) const
 	{
 		return (WordIndex == Other.WordIndex) & (Mask == Other.Mask);
 	}
 
-	[[nodiscard]] FORCEINLINE bool operator!=(FRelativeBitReference Other) const
+	[[nodiscard]] UE_FORCEINLINE_HINT bool operator!=(FRelativeBitReference Other) const
 	{
 		return !(Other == *this);
 	}
@@ -382,7 +382,7 @@ public:
 	 * @param Value - The value to initial the bits to.
 	 * @param InNumBits - The initial number of bits in the array.
 	 */
-	[[nodiscard]] FORCEINLINE explicit TBitArray(bool bValue, int32 InNumBits)
+	[[nodiscard]] UE_FORCEINLINE_HINT explicit TBitArray(bool bValue, int32 InNumBits)
 	:	MaxBits(AllocatorInstance.GetInitialCapacity() * NumBitsPerDWORD)
 	{
 		Init(bValue, InNumBits);
@@ -391,7 +391,7 @@ public:
 	/**
 	 * Move constructor.
 	 */
-	[[nodiscard]] FORCEINLINE TBitArray(TBitArray&& Other)
+	[[nodiscard]] UE_FORCEINLINE_HINT TBitArray(TBitArray&& Other)
 	{
 		this->Move(*this, Other);
 	}
@@ -399,7 +399,7 @@ public:
 	/**
 	 * Copy constructor.
 	 */
-	[[nodiscard]] FORCEINLINE TBitArray(const TBitArray& Copy)
+	[[nodiscard]] UE_FORCEINLINE_HINT TBitArray(const TBitArray& Copy)
 	:	NumBits(0)
 	,	MaxBits(0)
 	{
@@ -407,7 +407,7 @@ public:
 	}
 
 	template<typename OtherAllocator>
-	[[nodiscard]] FORCEINLINE explicit TBitArray(const TBitArray<OtherAllocator> & Copy)
+	[[nodiscard]] UE_FORCEINLINE_HINT explicit TBitArray(const TBitArray<OtherAllocator> & Copy)
 		: NumBits(0)
 		, MaxBits(0)
 	{
@@ -417,7 +417,7 @@ public:
 	/**
 	 * Move assignment.
 	 */
-	FORCEINLINE TBitArray& operator=(TBitArray&& Other)
+	UE_FORCEINLINE_HINT TBitArray& operator=(TBitArray&& Other)
 	{
 		if (this != &Other)
 		{
@@ -430,7 +430,7 @@ public:
 	/**
 	 * Assignment operator.
 	 */
-	FORCEINLINE TBitArray& operator=(const TBitArray& Copy)
+	UE_FORCEINLINE_HINT TBitArray& operator=(const TBitArray& Copy)
 	{
 		// check for self assignment since we don't use swap() mechanic
 		if (this != &Copy)
@@ -442,7 +442,7 @@ public:
 
 
 	template<typename OtherAllocator>
-	FORCEINLINE TBitArray& operator=(const TBitArray< OtherAllocator >& Copy)
+	UE_FORCEINLINE_HINT TBitArray& operator=(const TBitArray< OtherAllocator >& Copy)
 	{
 		// No need for a self-assignment check. If we were the same, we'd be using
 		// the default assignment operator.
@@ -450,7 +450,7 @@ public:
 		return *this;
 	}
 
-	[[nodiscard]] FORCEINLINE bool operator==(const TBitArray<Allocator>& Other) const
+	[[nodiscard]] UE_FORCEINLINE_HINT bool operator==(const TBitArray<Allocator>& Other) const
 	{
 		if (Num() != Other.Num())
 		{
@@ -460,7 +460,7 @@ public:
 		return FMemory::Memcmp(GetData(), Other.GetData(), GetNumWords() * sizeof(uint32)) == 0;
 	}
 
-	[[nodiscard]] FORCEINLINE bool operator<(const TBitArray<Allocator>& Other) const
+	[[nodiscard]] UE_FORCEINLINE_HINT bool operator<(const TBitArray<Allocator>& Other) const
 	{
 		//sort by length
 		if (Num() != Other.Num())
@@ -484,30 +484,30 @@ public:
 	}
 
 #if !PLATFORM_COMPILER_HAS_GENERATED_COMPARISON_OPERATORS
-	[[nodiscard]] FORCEINLINE bool operator!=(const TBitArray<Allocator>& Other) const
+	[[nodiscard]] UE_FORCEINLINE_HINT bool operator!=(const TBitArray<Allocator>& Other) const
 	{
 		return !(*this == Other);
 	}
 #endif
 
 private:
-	[[nodiscard]] FORCEINLINE uint32 GetNumWords() const
+	[[nodiscard]] UE_FORCEINLINE_HINT uint32 GetNumWords() const
 	{
 		return FBitSet::CalculateNumWords(NumBits);
 	}
 
-	[[nodiscard]] FORCEINLINE uint32 GetMaxWords() const
+	[[nodiscard]] UE_FORCEINLINE_HINT uint32 GetMaxWords() const
 	{
 		return FBitSet::CalculateNumWords(MaxBits);
 	}
 
-	[[nodiscard]] FORCEINLINE uint32 GetLastWordMask(int32 EndIndexExclusive) const
+	[[nodiscard]] UE_FORCEINLINE_HINT uint32 GetLastWordMask(int32 EndIndexExclusive) const
 	{
 		const uint32 UnusedBits = (FBitSet::BitsPerWord - static_cast<uint32>(EndIndexExclusive) % FBitSet::BitsPerWord) % FBitSet::BitsPerWord;
 		return ~0u >> UnusedBits;
 	}
 
-	FORCEINLINE static void SetWords(uint32* Words, int32 NumWords, bool bValue)
+	UE_FORCEINLINE_HINT static void SetWords(uint32* Words, int32 NumWords, bool bValue)
 	{
 		if (NumWords > 8)
 		{
@@ -524,7 +524,7 @@ private:
 	}
 
 	template <typename BitArrayType>
-	static FORCEINLINE void Move(BitArrayType& ToArray, BitArrayType& FromArray)
+	static UE_FORCEINLINE_HINT void Move(BitArrayType& ToArray, BitArrayType& FromArray)
 	{
 		ToArray.AllocatorInstance.MoveToEmpty(FromArray.AllocatorInstance);
 
@@ -818,7 +818,7 @@ public:
 	 * @param Value - The value to initial the bits to.
 	 * @param NumBits - The number of bits in the array.
 	 */
-	FORCEINLINE void Init(bool bValue, int32 InNumBits)
+	UE_FORCEINLINE_HINT void Init(bool bValue, int32 InNumBits)
 	{
 		NumBits = InNumBits;
 		
@@ -963,7 +963,7 @@ public:
 	 * @param ReadOffsetBits Number of bits into ReadBits at which to start reading.
 	 */
 	template <typename OtherAllocator>
-	FORCEINLINE void SetRangeFromRange(int32 Index, int32 NumBitsToSet, const TBitArray<OtherAllocator>& ReadBits, int32 ReadOffsetBits = 0)
+	UE_FORCEINLINE_HINT void SetRangeFromRange(int32 Index, int32 NumBitsToSet, const TBitArray<OtherAllocator>& ReadBits, int32 ReadOffsetBits = 0)
 	{
 		check(Index >= 0 && NumBitsToSet >= 0 && Index + NumBitsToSet <= NumBits);
 		check(0 <= ReadOffsetBits && ReadOffsetBits + NumBitsToSet <= ReadBits.NumBits);
@@ -979,7 +979,7 @@ public:
 	 * @param WriteOffsetBits Number of bits into WriteBits at which to start writing.
 	 */
 	template <typename InWordType>
-	FORCEINLINE void GetRange(int32 Index, int32 NumBitsToGet, InWordType* WriteBits, int32 WriteOffsetBits = 0) const
+	UE_FORCEINLINE_HINT void GetRange(int32 Index, int32 NumBitsToGet, InWordType* WriteBits, int32 WriteOffsetBits = 0) const
 	{
 		check(Index >= 0 && NumBitsToGet >= 0 && Index + NumBitsToGet <= NumBits);
 		check(NumBitsToGet == 0 || WriteBits != nullptr);
@@ -1166,7 +1166,7 @@ public:
 	 * @return The index of the first occurrence of the specified value (true/false) from StartIndex, or INDEX_NONE if not found.
 	 */
 	template <typename IndexType>
-	[[nodiscard]] FORCEINLINE int32 FindFrom(bool bValue, IndexType StartIndex) const
+	[[nodiscard]] UE_FORCEINLINE_HINT int32 FindFrom(bool bValue, IndexType StartIndex) const
 	{
 		static_assert(!std::is_same_v<IndexType, bool>, "TBitArray::FindFrom: unexpected bool passed as the StartIndex argument");
 		checkSlow((StartIndex >= 0) && (StartIndex <= NumBits));
@@ -1187,7 +1187,7 @@ public:
 	* @return The index of the first occurrence of the specified value (true/false) from StartIndex to EndIndexExclusive, or INDEX_NONE if not found.
 	*/
 	template <typename IndexType>
-	[[nodiscard]] FORCEINLINE int32 FindFrom(bool bValue, IndexType StartIndex, IndexType EndIndexExclusive) const
+	[[nodiscard]] UE_FORCEINLINE_HINT int32 FindFrom(bool bValue, IndexType StartIndex, IndexType EndIndexExclusive) const
 	{
 		static_assert(!std::is_same_v<IndexType, bool>, "TBitArray::FindFrom: unexpected bool passed as the StartIndex argument");
 		checkSlow((StartIndex >= 0) && (StartIndex <= NumBits));
@@ -1235,7 +1235,7 @@ public:
 	 *
 	 * @return true if the array contains the specified value, false otherwise.
 	 */
-	[[nodiscard]] FORCEINLINE bool Contains(bool bValue) const
+	[[nodiscard]] UE_FORCEINLINE_HINT bool Contains(bool bValue) const
 	{
 		return Find(bValue) != INDEX_NONE;
 	}
@@ -1440,7 +1440,7 @@ public:
 	}
 
 	// Accessors.
-	[[nodiscard]] FORCEINLINE bool IsValidIndex(int32 InIndex) const
+	[[nodiscard]] UE_FORCEINLINE_HINT bool IsValidIndex(int32 InIndex) const
 	{
 		return InIndex >= 0 && InIndex < NumBits;
 	}
@@ -1456,9 +1456,9 @@ public:
 		return NumBits == 0; 
 	}
 
-	[[nodiscard]] FORCEINLINE int32 Num() const { return NumBits; }
-	[[nodiscard]] FORCEINLINE int32 Max() const { return MaxBits; }
-	[[nodiscard]] FORCEINLINE FBitReference operator[](int32 Index)
+	[[nodiscard]] UE_FORCEINLINE_HINT int32 Num() const { return NumBits; }
+	[[nodiscard]] UE_FORCEINLINE_HINT int32 Max() const { return MaxBits; }
+	[[nodiscard]] UE_FORCEINLINE_HINT FBitReference operator[](int32 Index)
 	{
 		check(Index>=0 && Index<NumBits);
 		return FBitReference(
@@ -1466,7 +1466,7 @@ public:
 			1 << (Index & (NumBitsPerDWORD - 1))
 			);
 	}
-	[[nodiscard]] FORCEINLINE const FConstBitReference operator[](int32 Index) const
+	[[nodiscard]] UE_FORCEINLINE_HINT const FConstBitReference operator[](int32 Index) const
 	{
 		check(Index>=0 && Index<NumBits);
 		return FConstBitReference(
@@ -1474,7 +1474,7 @@ public:
 			1 << (Index & (NumBitsPerDWORD - 1))
 			);
 	}
-	[[nodiscard]] FORCEINLINE FBitReference AccessCorrespondingBit(const FRelativeBitReference& RelativeReference)
+	[[nodiscard]] UE_FORCEINLINE_HINT FBitReference AccessCorrespondingBit(const FRelativeBitReference& RelativeReference)
 	{
 		checkSlow(RelativeReference.Mask);
 		checkSlow(RelativeReference.WordIndex >= 0);
@@ -1484,7 +1484,7 @@ public:
 			RelativeReference.Mask
 			);
 	}
-	[[nodiscard]] FORCEINLINE const FConstBitReference AccessCorrespondingBit(const FRelativeBitReference& RelativeReference) const
+	[[nodiscard]] UE_FORCEINLINE_HINT const FConstBitReference AccessCorrespondingBit(const FRelativeBitReference& RelativeReference) const
 	{
 		checkSlow(RelativeReference.Mask);
 		checkSlow(RelativeReference.WordIndex >= 0);
@@ -1499,13 +1499,13 @@ public:
 	class FIterator : public FRelativeBitReference
 	{
 	public:
-		[[nodiscard]] FORCEINLINE explicit FIterator(TBitArray<Allocator>& InArray UE_LIFETIMEBOUND,int32 StartIndex = 0)
+		[[nodiscard]] UE_FORCEINLINE_HINT explicit FIterator(TBitArray<Allocator>& InArray UE_LIFETIMEBOUND,int32 StartIndex = 0)
 		:	FRelativeBitReference(StartIndex)
 		,	Array(InArray)
 		,	Index(StartIndex)
 		{
 		}
-		FORCEINLINE FIterator& operator++()
+		UE_FORCEINLINE_HINT FIterator& operator++()
 		{
 			++Index;
 			this->Mask <<= 1;
@@ -1518,25 +1518,25 @@ public:
 			return *this;
 		}
 
-		[[nodiscard]] FORCEINLINE FBitReference operator*() const
+		[[nodiscard]] UE_FORCEINLINE_HINT FBitReference operator*() const
 		{
 			return GetValue();
 		}
 
 		/** conversion to "bool" returning true if the iterator is valid. */
-		[[nodiscard]] FORCEINLINE explicit operator bool() const
+		[[nodiscard]] UE_FORCEINLINE_HINT explicit operator bool() const
 		{ 
 			return Index < Array.Num(); 
 		}
 
 		/** inverse of the "bool" operator */
-		[[nodiscard]] FORCEINLINE bool operator !() const
+		[[nodiscard]] UE_FORCEINLINE_HINT bool operator !() const
 		{
 			return !(bool)*this;
 		}
 
-		[[nodiscard]] FORCEINLINE FBitReference GetValue() const { return FBitReference(Array.GetData()[this->WordIndex],this->Mask); }
-		[[nodiscard]] FORCEINLINE int32 GetIndex() const { return Index; }
+		[[nodiscard]] UE_FORCEINLINE_HINT FBitReference GetValue() const { return FBitReference(Array.GetData()[this->WordIndex],this->Mask); }
+		[[nodiscard]] UE_FORCEINLINE_HINT int32 GetIndex() const { return Index; }
 	private:
 		TBitArray<Allocator>& Array;
 		int32 Index;
@@ -1546,21 +1546,21 @@ public:
 	class FReverseIterator : public FRelativeBitReference
 	{
 	public:
-		[[nodiscard]] FORCEINLINE explicit FReverseIterator(TBitArray<Allocator>& InArray UE_LIFETIMEBOUND)
+		[[nodiscard]] UE_FORCEINLINE_HINT explicit FReverseIterator(TBitArray<Allocator>& InArray UE_LIFETIMEBOUND)
 			: FRelativeBitReference(InArray.Num() - 1)
 			, Array(InArray)
 			, Index(InArray.Num() - 1)
 		{
 		}
 
-		[[nodiscard]] FORCEINLINE explicit FReverseIterator(TBitArray<Allocator>& InArray UE_LIFETIMEBOUND, int32 StartIndex)
+		[[nodiscard]] UE_FORCEINLINE_HINT explicit FReverseIterator(TBitArray<Allocator>& InArray UE_LIFETIMEBOUND, int32 StartIndex)
 			: FRelativeBitReference(-1)
 			, Array(InArray)
 			, Index(-1)
 		{
 		}
 
-		FORCEINLINE FReverseIterator& operator++()
+		UE_FORCEINLINE_HINT FReverseIterator& operator++()
 		{
 			--Index;
 			this->Mask >>= 1;
@@ -1573,25 +1573,25 @@ public:
 			return *this;
 		}
 
-		[[nodiscard]] FORCEINLINE FBitReference operator*() const
+		[[nodiscard]] UE_FORCEINLINE_HINT FBitReference operator*() const
 		{
 			return GetValue();
 		}
 
 		/** conversion to "bool" returning true if the iterator is valid. */
-		[[nodiscard]] FORCEINLINE explicit operator bool() const
+		[[nodiscard]] UE_FORCEINLINE_HINT explicit operator bool() const
 		{
 			return Index >= 0;
 		}
 
 		/** inverse of the "bool" operator */
-		[[nodiscard]] FORCEINLINE bool operator !() const
+		[[nodiscard]] UE_FORCEINLINE_HINT bool operator !() const
 		{
 			return !(bool)*this;
 		}
 
-		[[nodiscard]] FORCEINLINE FBitReference GetValue() const { return FBitReference(Array.GetData()[this->WordIndex], this->Mask); }
-		[[nodiscard]] FORCEINLINE int32 GetIndex() const { return Index; }
+		[[nodiscard]] UE_FORCEINLINE_HINT FBitReference GetValue() const { return FBitReference(Array.GetData()[this->WordIndex], this->Mask); }
+		[[nodiscard]] UE_FORCEINLINE_HINT int32 GetIndex() const { return Index; }
 	private:
 		TBitArray<Allocator>& Array;
 		int32 Index;
@@ -1601,13 +1601,13 @@ public:
 	class FConstIterator : public FRelativeBitReference
 	{
 	public:
-		[[nodiscard]] FORCEINLINE explicit FConstIterator(const TBitArray<Allocator>& InArray UE_LIFETIMEBOUND,int32 StartIndex = 0)
+		[[nodiscard]] UE_FORCEINLINE_HINT explicit FConstIterator(const TBitArray<Allocator>& InArray UE_LIFETIMEBOUND,int32 StartIndex = 0)
 		:	FRelativeBitReference(StartIndex)
 		,	Array(InArray)
 		,	Index(StartIndex)
 		{
 		}
-		FORCEINLINE FConstIterator& operator++()
+		UE_FORCEINLINE_HINT FConstIterator& operator++()
 		{
 			++Index;
 			this->Mask <<= 1;
@@ -1620,24 +1620,24 @@ public:
 			return *this;
 		}
 
-		[[nodiscard]] FORCEINLINE FConstBitReference operator*() const
+		[[nodiscard]] UE_FORCEINLINE_HINT FConstBitReference operator*() const
 		{
 			return GetValue();
 		}
 
 		/** conversion to "bool" returning true if the iterator is valid. */
-		[[nodiscard]] FORCEINLINE explicit operator bool() const
+		[[nodiscard]] UE_FORCEINLINE_HINT explicit operator bool() const
 		{
 			return Index < Array.Num(); 
 		}
 		/** inverse of the "bool" operator */
-		[[nodiscard]] FORCEINLINE bool operator !() const 
+		[[nodiscard]] UE_FORCEINLINE_HINT bool operator !() const 
 		{
 			return !(bool)*this;
 		}
 
-		[[nodiscard]] FORCEINLINE FConstBitReference GetValue() const { return FConstBitReference(Array.GetData()[this->WordIndex],this->Mask); }
-		[[nodiscard]] FORCEINLINE int32 GetIndex() const { return Index; }
+		[[nodiscard]] UE_FORCEINLINE_HINT FConstBitReference GetValue() const { return FConstBitReference(Array.GetData()[this->WordIndex],this->Mask); }
+		[[nodiscard]] UE_FORCEINLINE_HINT int32 GetIndex() const { return Index; }
 	private:
 		const TBitArray<Allocator>& Array;
 		int32 Index;
@@ -1647,19 +1647,19 @@ public:
 	class FConstReverseIterator : public FRelativeBitReference
 	{
 	public:
-		[[nodiscard]] FORCEINLINE explicit FConstReverseIterator(const TBitArray<Allocator>& InArray UE_LIFETIMEBOUND)
+		[[nodiscard]] UE_FORCEINLINE_HINT explicit FConstReverseIterator(const TBitArray<Allocator>& InArray UE_LIFETIMEBOUND)
 			:	FRelativeBitReference(InArray.Num() - 1)
 			,	Array(InArray)
 			,	Index(InArray.Num() - 1)
 		{
 		}
-		[[nodiscard]] FORCEINLINE explicit FConstReverseIterator(const TBitArray<Allocator>& InArray UE_LIFETIMEBOUND, int32 StartIndex)
+		[[nodiscard]] UE_FORCEINLINE_HINT explicit FConstReverseIterator(const TBitArray<Allocator>& InArray UE_LIFETIMEBOUND, int32 StartIndex)
 			: FRelativeBitReference(-1)
 			, Array(InArray)
 			, Index(-1)
 		{
 		}
-		FORCEINLINE FConstReverseIterator& operator++()
+		UE_FORCEINLINE_HINT FConstReverseIterator& operator++()
 		{
 			--Index;
 			this->Mask >>= 1u;
@@ -1672,35 +1672,35 @@ public:
 			return *this;
 		}
 
-		FORCEINLINE FConstBitReference operator*() const
+		UE_FORCEINLINE_HINT FConstBitReference operator*() const
 		{
 			return GetValue();
 		}
 
 		/** conversion to "bool" returning true if the iterator is valid. */
-		[[nodiscard]] FORCEINLINE explicit operator bool() const
+		[[nodiscard]] UE_FORCEINLINE_HINT explicit operator bool() const
 		{
 			return Index >= 0; 
 		}
 		/** inverse of the "bool" operator */
-		[[nodiscard]] FORCEINLINE bool operator !() const 
+		[[nodiscard]] UE_FORCEINLINE_HINT bool operator !() const 
 		{
 			return !(bool)*this;
 		}
 
-		[[nodiscard]] FORCEINLINE FConstBitReference GetValue() const { return FConstBitReference(Array.GetData()[this->WordIndex],this->Mask); }
-		[[nodiscard]] FORCEINLINE int32 GetIndex() const { return Index; }
+		[[nodiscard]] UE_FORCEINLINE_HINT FConstBitReference GetValue() const { return FConstBitReference(Array.GetData()[this->WordIndex],this->Mask); }
+		[[nodiscard]] UE_FORCEINLINE_HINT int32 GetIndex() const { return Index; }
 	private:
 		const TBitArray<Allocator>& Array;
 		int32 Index;
 	};
 
-	[[nodiscard]] FORCEINLINE const uint32* GetData() const
+	[[nodiscard]] UE_FORCEINLINE_HINT const uint32* GetData() const
 	{
 		return (uint32*)AllocatorInstance.GetAllocation();
 	}
 
-	[[nodiscard]] FORCEINLINE uint32* GetData()
+	[[nodiscard]] UE_FORCEINLINE_HINT uint32* GetData()
 	{
 		return (uint32*)AllocatorInstance.GetAllocation();
 	}
@@ -1829,15 +1829,15 @@ public:
 
 	/** Enables range-based for loops, DO NOT USE DIRECTLY. */
 
-	[[nodiscard]] FORCEINLINE FIterator      begin()       { return FIterator(*this); }
-	[[nodiscard]] FORCEINLINE FConstIterator begin() const { return FConstIterator(*this); }
-	[[nodiscard]] FORCEINLINE FIterator      end()         { return FIterator(*this, NumBits); }
-	[[nodiscard]] FORCEINLINE FConstIterator end()   const { return FConstIterator(*this, NumBits); }
+	[[nodiscard]] UE_FORCEINLINE_HINT FIterator      begin()       { return FIterator(*this); }
+	[[nodiscard]] UE_FORCEINLINE_HINT FConstIterator begin() const { return FConstIterator(*this); }
+	[[nodiscard]] UE_FORCEINLINE_HINT FIterator      end()         { return FIterator(*this, NumBits); }
+	[[nodiscard]] UE_FORCEINLINE_HINT FConstIterator end()   const { return FConstIterator(*this, NumBits); }
 
-	[[nodiscard]] FORCEINLINE FReverseIterator      rbegin()       { return FReverseIterator(*this); }
-	[[nodiscard]] FORCEINLINE FConstReverseIterator rbegin() const { return FConstReverseIterator(*this); }
-	[[nodiscard]] FORCEINLINE FReverseIterator      rend()         { return FReverseIterator(*this, 0); }
-	[[nodiscard]] FORCEINLINE FConstReverseIterator rend()   const { return FConstReverseIterator(*this, 0); }
+	[[nodiscard]] UE_FORCEINLINE_HINT FReverseIterator      rbegin()       { return FReverseIterator(*this); }
+	[[nodiscard]] UE_FORCEINLINE_HINT FConstReverseIterator rbegin() const { return FConstReverseIterator(*this); }
+	[[nodiscard]] UE_FORCEINLINE_HINT FReverseIterator      rend()         { return FReverseIterator(*this, 0); }
+	[[nodiscard]] UE_FORCEINLINE_HINT FConstReverseIterator rend()   const { return FConstReverseIterator(*this, 0); }
 
 private:
 	AllocatorType AllocatorInstance;
@@ -1918,7 +1918,7 @@ namespace Freeze
 DECLARE_TEMPLATE_INTRINSIC_TYPE_LAYOUT(template<typename Allocator>, TBitArray<Allocator>);
 
 template<typename Allocator>
-[[nodiscard]] FORCEINLINE uint32 GetTypeHash(const TBitArray<Allocator>& BitArray)
+[[nodiscard]] UE_FORCEINLINE_HINT uint32 GetTypeHash(const TBitArray<Allocator>& BitArray)
 {
 	uint32 NumWords = FBitSet::CalculateNumWords(BitArray.Num());
 	uint32 Hash = NumWords;
@@ -1965,7 +1965,7 @@ public:
 	}
 
 	/** Forwards iteration operator. */
-	FORCEINLINE TConstSetBitIterator& operator++()
+	UE_FORCEINLINE_HINT TConstSetBitIterator& operator++()
 	{
 		// Mark the current bit as visited.
 		UnvisitedBitMask &= ~this->Mask;
@@ -1976,32 +1976,32 @@ public:
 		return *this;
 	}
 
-	[[nodiscard]] FORCEINLINE bool operator==(const TConstSetBitIterator& Rhs) const
+	[[nodiscard]] UE_FORCEINLINE_HINT bool operator==(const TConstSetBitIterator& Rhs) const
 	{
 		// We only need to compare the bit index and the array... all the rest of the state is unobservable.
 		return CurrentBitIndex == Rhs.CurrentBitIndex && &Array == &Rhs.Array;
 	}
 
 #if !PLATFORM_COMPILER_HAS_GENERATED_COMPARISON_OPERATORS
-	[[nodiscard]] FORCEINLINE bool operator!=(const TConstSetBitIterator& Rhs) const
+	[[nodiscard]] UE_FORCEINLINE_HINT bool operator!=(const TConstSetBitIterator& Rhs) const
 	{ 
 		return !(*this == Rhs);
 	}
 #endif
 
 	/** conversion to "bool" returning true if the iterator is valid. */
-	[[nodiscard]] FORCEINLINE explicit operator bool() const
+	[[nodiscard]] UE_FORCEINLINE_HINT explicit operator bool() const
 	{ 
 		return CurrentBitIndex < Array.Num(); 
 	}
 	/** inverse of the "bool" operator */
-	[[nodiscard]] FORCEINLINE bool operator !() const
+	[[nodiscard]] UE_FORCEINLINE_HINT bool operator !() const
 	{
 		return !(bool)*this;
 	}
 
 	/** Index accessor. */
-	[[nodiscard]] FORCEINLINE int32 GetIndex() const
+	[[nodiscard]] UE_FORCEINLINE_HINT int32 GetIndex() const
 	{
 		return CurrentBitIndex;
 	}
@@ -2066,7 +2066,7 @@ class TConstDualSetBitIterator : public FRelativeBitReference
 public:
 
 	/** Constructor. */
-	[[nodiscard]] FORCEINLINE TConstDualSetBitIterator(
+	[[nodiscard]] UE_FORCEINLINE_HINT TConstDualSetBitIterator(
 		const TBitArray<Allocator>& InArrayA,
 		const TBitArray<OtherAllocator>& InArrayB,
 		int32 StartIndex = 0
@@ -2084,7 +2084,7 @@ public:
 	}
 
 	/** Advancement operator. */
-	FORCEINLINE TConstDualSetBitIterator& operator++()
+	UE_FORCEINLINE_HINT TConstDualSetBitIterator& operator++()
 	{
 		checkSlow(ArrayA.Num() == ArrayB.Num());
 
@@ -2099,18 +2099,18 @@ public:
 	}
 
 	/** conversion to "bool" returning true if the iterator is valid. */
-	[[nodiscard]] FORCEINLINE explicit operator bool() const
+	[[nodiscard]] UE_FORCEINLINE_HINT explicit operator bool() const
 	{ 
 		return CurrentBitIndex < ArrayA.Num(); 
 	}
 	/** inverse of the "bool" operator */
-	[[nodiscard]] FORCEINLINE bool operator !() const
+	[[nodiscard]] UE_FORCEINLINE_HINT bool operator !() const
 	{
 		return !(bool)*this;
 	}
 
 	/** Index accessor. */
-	[[nodiscard]] FORCEINLINE int32 GetIndex() const
+	[[nodiscard]] UE_FORCEINLINE_HINT int32 GetIndex() const
 	{
 		return CurrentBitIndex;
 	}
@@ -2305,12 +2305,12 @@ private:
 		static_assert(STRUCT_OFFSET(ScriptType, MaxBits)           == STRUCT_OFFSET(RealType, MaxBits),           "TScriptBitArray's MaxBits member offset does not match TBitArray's");
 	}
 
-	[[nodiscard]] FORCEINLINE uint32* GetData()
+	[[nodiscard]] UE_FORCEINLINE_HINT uint32* GetData()
 	{
 		return (uint32*)AllocatorInstance.GetAllocation();
 	}
 
-	[[nodiscard]] FORCEINLINE const uint32* GetData() const
+	[[nodiscard]] UE_FORCEINLINE_HINT const uint32* GetData() const
 	{
 		return (const uint32*)AllocatorInstance.GetAllocation();
 	}
diff --git a/Engine/Source/Runtime/Core/Public/Containers/ChunkedArray.h b/Engine/Source/Runtime/Core/Public/Containers/ChunkedArray.h
--- a/Engine/Source/Runtime/Core/Public/Containers/ChunkedArray.h
+++ b/Engine/Source/Runtime/Core/Public/Containers/ChunkedArray.h
@@ -179,7 +179,7 @@ public:
 	*
 	* @returns True if index is valid. False otherwise.
 	*/
-	[[nodiscard]] FORCEINLINE bool IsValidIndex(int32 Index) const
+	[[nodiscard]] UE_FORCEINLINE_HINT bool IsValidIndex(int32 Index) const
 	{
 		return Index >= 0 && Index < NumElements;
 	}
diff --git a/Engine/Source/Runtime/Core/Public/Containers/CircularBuffer.h b/Engine/Source/Runtime/Core/Public/Containers/CircularBuffer.h
--- a/Engine/Source/Runtime/Core/Public/Containers/CircularBuffer.h
+++ b/Engine/Source/Runtime/Core/Public/Containers/CircularBuffer.h
@@ -54,7 +54,7 @@ public:
 	 *
 	 * @param Index The index of the element to return.
 	 */
-	[[nodiscard]] FORCEINLINE ElementType& operator[](uint32 Index)
+	[[nodiscard]] UE_FORCEINLINE_HINT ElementType& operator[](uint32 Index)
 	{
 		return Elements[Index & IndexMask];
 	}
@@ -64,7 +64,7 @@ public:
 	 *
 	 * @param Index The index of the element to return.
 	 */
-	[[nodiscard]] FORCEINLINE const ElementType& operator[](uint32 Index) const
+	[[nodiscard]] UE_FORCEINLINE_HINT const ElementType& operator[](uint32 Index) const
 	{
 		return Elements[Index & IndexMask];
 	}
@@ -76,7 +76,7 @@ public:
 	 *
 	 * @return Buffer capacity.
 	 */
-	[[nodiscard]] FORCEINLINE uint32 Capacity() const
+	[[nodiscard]] UE_FORCEINLINE_HINT uint32 Capacity() const
 	{
 		return Elements.Num();
 	}
@@ -87,7 +87,7 @@ public:
 	 * @param CurrentIndex The current index.
 	 * @return The next index.
 	 */
-	[[nodiscard]] FORCEINLINE uint32 GetNextIndex(uint32 CurrentIndex) const
+	[[nodiscard]] UE_FORCEINLINE_HINT uint32 GetNextIndex(uint32 CurrentIndex) const
 	{
 		return ((CurrentIndex + 1) & IndexMask);
 	}
@@ -98,7 +98,7 @@ public:
 	 * @param CurrentIndex The current index.
 	 * @return The previous index.
 	 */
-	[[nodiscard]] FORCEINLINE uint32 GetPreviousIndex(uint32 CurrentIndex) const
+	[[nodiscard]] UE_FORCEINLINE_HINT uint32 GetPreviousIndex(uint32 CurrentIndex) const
 	{
 		return ((CurrentIndex - 1) & IndexMask);
 	}
diff --git a/Engine/Source/Runtime/Core/Public/Containers/CircularQueue.h b/Engine/Source/Runtime/Core/Public/Containers/CircularQueue.h
--- a/Engine/Source/Runtime/Core/Public/Containers/CircularQueue.h
+++ b/Engine/Source/Runtime/Core/Public/Containers/CircularQueue.h
@@ -168,7 +168,7 @@ public:
 	 * @return true if the queue is empty, false otherwise.
 	 * @see Empty, IsFull
 	 */
-	FORCEINLINE bool IsEmpty() const
+	UE_FORCEINLINE_HINT bool IsEmpty() const
 	{
 		return (Head.Load() == Tail.Load());
 	}
diff --git a/Engine/Source/Runtime/Core/Public/Containers/CompactHashTable.h b/Engine/Source/Runtime/Core/Public/Containers/CompactHashTable.h
--- a/Engine/Source/Runtime/Core/Public/Containers/CompactHashTable.h
+++ b/Engine/Source/Runtime/Core/Public/Containers/CompactHashTable.h
@@ -51,22 +51,22 @@ namespace UE::Core::CompactHashTable::Private
 /* Helpers for interacting with compact hash table memory layouts */
 namespace UE::Core::CompactHashTable
 {
-	[[nodiscard]] FORCEINLINE constexpr uint32 GetTypeSize(uint32 IndexCount)
+	[[nodiscard]] UE_FORCEINLINE_HINT constexpr uint32 GetTypeSize(uint32 IndexCount)
 	{
 		return 1 + int(IndexCount > 0xff) + int(IndexCount > 0xffff) * 2;
 	}
 
-	[[nodiscard]] FORCEINLINE constexpr uint32 GetTypeShift(uint32 IndexCount)
+	[[nodiscard]] UE_FORCEINLINE_HINT constexpr uint32 GetTypeShift(uint32 IndexCount)
 	{
 		return 0 + int(IndexCount > 0xff) + int(IndexCount > 0xffff);
 	}
 
-	[[nodiscard]] FORCEINLINE constexpr size_t GetMemoryRequiredInBytes(uint32 IndexCount, uint32 HashCount)
+	[[nodiscard]] UE_FORCEINLINE_HINT constexpr size_t GetMemoryRequiredInBytes(uint32 IndexCount, uint32 HashCount)
 	{
 		return ((size_t)IndexCount + HashCount) << GetTypeShift(IndexCount);
 	}
 
-	[[nodiscard]] FORCEINLINE constexpr size_t GetMemoryAlignment()
+	[[nodiscard]] UE_FORCEINLINE_HINT constexpr size_t GetMemoryAlignment()
 	{
 		return 4; // only support uint32 for now, pick highest alignment so we don't change alignment between allocations
 	}
@@ -234,17 +234,17 @@ public:
 	}
 	TStaticCompactHashTable(ENoInit) {}
 	
-	FORCEINLINE void Reset()
+	UE_FORCEINLINE_HINT void Reset()
 	{
 		FMemory::Memset(HashData, 0xff, sizeof(HashData));
 	}
 
-	[[nodiscard]] FORCEINLINE uint32 GetFirst(uint32 Key) const
+	[[nodiscard]] UE_FORCEINLINE_HINT uint32 GetFirst(uint32 Key) const
 	{
 		return UE::Core::CompactHashTable::GetFirst(Key, HashData, HashCount);
 	}
 
-	[[nodiscard]] FORCEINLINE uint32 GetFirstByIndex(uint32 HashIndex) const
+	[[nodiscard]] UE_FORCEINLINE_HINT uint32 GetFirstByIndex(uint32 HashIndex) const
 	{
 		return UE::Core::CompactHashTable::GetFirstByIndex(HashIndex, HashData, HashCount);
 	}
@@ -268,7 +268,7 @@ public:
 		UE::Core::CompactHashTable::Add(CurrentCount, Key, HashData, HashCount, NextIndexData, CurrentCount + 1);
 	}
 
-	FORCEINLINE void Remove(uint32 Index, uint32 Key, uint32 LastIndex, uint32 OptLastKey)
+	UE_FORCEINLINE_HINT void Remove(uint32 Index, uint32 Key, uint32 LastIndex, uint32 OptLastKey)
 	{
 		UE::Core::CompactHashTable::Remove(Index, Key, LastIndex, OptLastKey, HashData, HashCount, NextIndexData, ElementCount);
 	}
@@ -349,7 +349,7 @@ protected:
 class FCompactHashTableView : public FConstCompactHashTableView
 {
 public:
-	FORCEINLINE explicit FCompactHashTableView(uint8* Memory, uint32 InNextIndexCount, uint32 InHashCount, size_t MemorySize)
+	UE_FORCEINLINE_HINT explicit FCompactHashTableView(uint8* Memory, uint32 InNextIndexCount, uint32 InHashCount, size_t MemorySize)
 	: FConstCompactHashTableView(Memory, InNextIndexCount, InHashCount, MemorySize)
 	{
 	}
@@ -370,14 +370,14 @@ public:
 		#undef UE_COMPACTHASHTABLE_EXECUTEBYTYPE
 	}
 
-	FORCEINLINE void Remove(uint32 Index, uint32 Key, uint32 LastIndex, uint32 OptLastKey) const
+	UE_FORCEINLINE_HINT void Remove(uint32 Index, uint32 Key, uint32 LastIndex, uint32 OptLastKey) const
 	{
 		#define UE_COMPACTHASHTABLE_EXECUTEBYTYPE(Type) return UE::Core::CompactHashTable::Remove(Index, Key, LastIndex, OptLastKey, (Type*)(HashData), HashCount, (Type*)(NextIndexData), NextIndexCount)
 		UE_COMPACTHASHTABLE_CALLBYTYPE(NextIndexCount);
 		#undef UE_COMPACTHASHTABLE_EXECUTEBYTYPE
 	}
 
-	FORCEINLINE void RemoveStable(uint32 Index, uint32 Key) const
+	UE_FORCEINLINE_HINT void RemoveStable(uint32 Index, uint32 Key) const
 	{
 		#define UE_COMPACTHASHTABLE_EXECUTEBYTYPE(Type) return UE::Core::CompactHashTable::RemoveStable(Index, Key, (Type*)(HashData), HashCount, (Type*)(NextIndexData), NextIndexCount)
 		UE_COMPACTHASHTABLE_CALLBYTYPE(NextIndexCount);
diff --git a/Engine/Source/Runtime/Core/Public/Containers/CompactSet.h b/Engine/Source/Runtime/Core/Public/Containers/CompactSet.h
--- a/Engine/Source/Runtime/Core/Public/Containers/CompactSet.h
+++ b/Engine/Source/Runtime/Core/Public/Containers/CompactSet.h
@@ -82,7 +82,7 @@ private:
 	public:
 		using ElementItType = std::conditional_t<bConst, const ElementType, ElementType>;
 
-		[[nodiscard]] FORCEINLINE TBaseIterator(SetType& InSet)
+		[[nodiscard]] UE_FORCEINLINE_HINT TBaseIterator(SetType& InSet)
 		: TBaseIterator(InSet, 0)
 		{
 		}
@@ -99,24 +99,24 @@ private:
 #endif
 		}
 
-		[[nodiscard]] FORCEINLINE ElementItType & operator*() const
+		[[nodiscard]] UE_FORCEINLINE_HINT ElementItType & operator*() const
 		{
 			return Set[FSetElementId::FromInteger(Index)];
 		}
 
-		[[nodiscard]] FORCEINLINE ElementItType* operator->() const
+		[[nodiscard]] UE_FORCEINLINE_HINT ElementItType* operator->() const
 		{
 			return &Set[FSetElementId::FromInteger(Index)];
 		}
 
-		FORCEINLINE TBaseIterator& operator++()
+		UE_FORCEINLINE_HINT TBaseIterator& operator++()
 		{
 			++Index;
 			return *this;
 		}
 
 		/** conversion to "bool" returning true if the iterator is valid. */
-		[[nodiscard]] FORCEINLINE explicit operator bool() const
+		[[nodiscard]] UE_FORCEINLINE_HINT explicit operator bool() const
 		{
 			return Set.IsValidId(this->GetId());
 		}
@@ -130,13 +130,13 @@ private:
 			return Index == Rhs.Index;
 		}
 
-		[[nodiscard]] FORCEINLINE bool operator!=(const TBaseIterator& Rhs) const
+		[[nodiscard]] UE_FORCEINLINE_HINT bool operator!=(const TBaseIterator& Rhs) const
 		{
 			return !(*this == Rhs);
 		}
 
 		// Accessors.
-		[[nodiscard]] FORCEINLINE FSetElementId GetId() const
+		[[nodiscard]] UE_FORCEINLINE_HINT FSetElementId GetId() const
 		{
 			return FSetElementId::FromInteger(Index);
 		}
@@ -207,23 +207,23 @@ private:
 		}
 
 		/** conversion to "bool" returning true if the iterator is valid. */
-		[[nodiscard]] FORCEINLINE explicit operator bool() const
+		[[nodiscard]] UE_FORCEINLINE_HINT explicit operator bool() const
 		{
 			return Index != INDEX_NONE;
 		}
 
 		// Accessors.
-		[[nodiscard]] FORCEINLINE FSetElementId GetId() const
+		[[nodiscard]] UE_FORCEINLINE_HINT FSetElementId GetId() const
 		{
 			return FSetElementId::FromInteger(this->Index);
 		}
 
-		[[nodiscard]] FORCEINLINE ItElementType* operator->() const
+		[[nodiscard]] UE_FORCEINLINE_HINT ItElementType* operator->() const
 		{
 			return &Set[GetId()];
 		}
 
-		[[nodiscard]] FORCEINLINE ItElementType& operator*() const
+		[[nodiscard]] UE_FORCEINLINE_HINT ItElementType& operator*() const
 		{
 			return Set[GetId()];
 		}
@@ -298,19 +298,19 @@ public:
 
 	/* Constructors */
 
-	[[nodiscard]] FORCEINLINE TCompactSet() = default;
+	[[nodiscard]] UE_FORCEINLINE_HINT TCompactSet() = default;
 
-	[[nodiscard]] FORCEINLINE TCompactSet(const TCompactSet& Copy)
+	[[nodiscard]] UE_FORCEINLINE_HINT TCompactSet(const TCompactSet& Copy)
 	{
 		*this = Copy;
 	}
 
-	[[nodiscard]] FORCEINLINE explicit TCompactSet(TArrayView<const ElementType> InArrayView)
+	[[nodiscard]] UE_FORCEINLINE_HINT explicit TCompactSet(TArrayView<const ElementType> InArrayView)
 	{
 		Append(InArrayView);
 	}
 
-	[[nodiscard]] FORCEINLINE explicit TCompactSet(TArray<ElementType>&& InArray)
+	[[nodiscard]] UE_FORCEINLINE_HINT explicit TCompactSet(TArray<ElementType>&& InArray)
 	{
 		Append(MoveTemp(InArray));
 	}
@@ -339,7 +339,7 @@ public:
 		Append(Other);
 	}
 
-	FORCEINLINE ~TCompactSet()
+	UE_FORCEINLINE_HINT ~TCompactSet()
 	{
 		UE_STATIC_ASSERT_WARN(TIsTriviallyRelocatable_V<InElementType>, "TMapBase can only be used with trivially relocatable types");
 
@@ -551,7 +551,7 @@ public:
 	 * Only returns the size of allocations made directly by the container, not the elements themselves.
 	 * @return number of bytes allocated by this container
 	 */
-	[[nodiscard]] FORCEINLINE SIZE_T GetAllocatedSize() const
+	[[nodiscard]] UE_FORCEINLINE_HINT SIZE_T GetAllocatedSize() const
 	{
 		return this->Elements.GetAllocationSize();
 	}
@@ -564,24 +564,24 @@ public:
 	}
 
 	/* Calculate the size of the hash table from the number of elements in the set assuming the default number of hash elements */
-	[[nodiscard]] FORCEINLINE static constexpr size_t GetTotalMemoryRequiredInBytes(uint32 NumElements)
+	[[nodiscard]] UE_FORCEINLINE_HINT static constexpr size_t GetTotalMemoryRequiredInBytes(uint32 NumElements)
 	{
 		return Super::GetTotalMemoryRequiredInBytes(NumElements, GetSetLayout());
 	}
 
 
 	/* Get the alignment required for the allocation */
-	[[nodiscard]] FORCEINLINE static constexpr size_t GetAlignment()
+	[[nodiscard]] UE_FORCEINLINE_HINT static constexpr size_t GetAlignment()
 	{
 		return FGenericPlatformMath::Max<size_t>(alignof(ElementType), UE::Core::CompactHashTable::GetMemoryAlignment());
 	}
 
-	[[nodiscard]] FORCEINLINE ElementType* GetData()
+	[[nodiscard]] UE_FORCEINLINE_HINT ElementType* GetData()
 	{
 		return (ElementType*)this->Elements.GetAllocation();
 	}
 
-	[[nodiscard]] FORCEINLINE const ElementType* GetData() const
+	[[nodiscard]] UE_FORCEINLINE_HINT const ElementType* GetData() const
 	{
 		return (const ElementType*)this->Elements.GetAllocation();
 	}
@@ -591,7 +591,7 @@ public:
 	 * @param Id - The element id to check.
 	 * @return true if the element identifier refers to a valid element in this set.
 	 */
-	[[nodiscard]] FORCEINLINE bool IsValidId(FSetElementId Id) const
+	[[nodiscard]] UE_FORCEINLINE_HINT bool IsValidId(FSetElementId Id) const
 	{
 		return Id.AsInteger() >= 0 && Id.AsInteger() < this->NumElements;
 	}
@@ -600,7 +600,7 @@ public:
 	 * Checks array invariants: if array size is greater than or equal to zero and less
 	 * than or equal to the maximum.
 	 */
-	FORCEINLINE void CheckInvariants() const
+	UE_FORCEINLINE_HINT void CheckInvariants() const
 	{
 		checkSlow((this->NumElements >= 0) & (this->MaxElements >= this->NumElements)); // & for one branch
 	}
@@ -656,11 +656,11 @@ public:
 	 * @param	bIsAlreadyInSetPtr	[out]	Optional pointer to bool that will be set depending on whether element is already in set
 	 * @return	A pointer to the element stored in the set.
 	 */
-	FORCEINLINE FSetElementId Add(const ElementType&  InElement, bool* bIsAlreadyInSetPtr = nullptr)
+	UE_FORCEINLINE_HINT FSetElementId Add(const ElementType&  InElement, bool* bIsAlreadyInSetPtr = nullptr)
 	{
 		return EmplaceByHash(KeyFuncs::GetKeyHash(KeyFuncs::GetSetKey(InElement)), InElement, bIsAlreadyInSetPtr);
 	}
-	FORCEINLINE FSetElementId Add(ElementType && InElement, bool* bIsAlreadyInSetPtr = nullptr)
+	UE_FORCEINLINE_HINT FSetElementId Add(ElementType && InElement, bool* bIsAlreadyInSetPtr = nullptr)
 	{
 		return EmplaceByHash(KeyFuncs::GetKeyHash(KeyFuncs::GetSetKey(InElement)), MoveTempIfPossible(InElement), bIsAlreadyInSetPtr);
 	}
@@ -673,11 +673,11 @@ public:
 	 * @param	bIsAlreadyInSetPtr	[out]	Optional pointer to bool that will be set depending on whether element is already in set
      * @return  A handle to the element stored in the set
      */
-	FORCEINLINE FSetElementId AddByHash(uint32 KeyHash, const ElementType& InElement, bool* bIsAlreadyInSetPtr = nullptr)
+	UE_FORCEINLINE_HINT FSetElementId AddByHash(uint32 KeyHash, const ElementType& InElement, bool* bIsAlreadyInSetPtr = nullptr)
 	{
 		return EmplaceByHash(KeyHash, InElement, bIsAlreadyInSetPtr);
 	}
-	FORCEINLINE FSetElementId AddByHash(uint32 KeyHash,		 ElementType && InElement, bool* bIsAlreadyInSetPtr = nullptr)
+	UE_FORCEINLINE_HINT FSetElementId AddByHash(uint32 KeyHash,		 ElementType && InElement, bool* bIsAlreadyInSetPtr = nullptr)
 	{
 		return EmplaceByHash(KeyHash, MoveTempIfPossible(InElement), bIsAlreadyInSetPtr);
 	}
@@ -771,11 +771,11 @@ public:
 		* @param	bIsAlreadyInSetPtr	[out]	Optional pointer to bool that will be set depending on whether element is already in set
 		* @return	A reference to the element stored in the set.
 		*/
-	FORCEINLINE ElementType& FindOrAdd(const InElementType& InElement, bool* bIsAlreadyInSetPtr = nullptr)
+	UE_FORCEINLINE_HINT ElementType& FindOrAdd(const InElementType& InElement, bool* bIsAlreadyInSetPtr = nullptr)
 	{
 		return FindOrAddByHash(KeyFuncs::GetKeyHash(KeyFuncs::GetSetKey(InElement)), InElement, bIsAlreadyInSetPtr);
 	}
-	FORCEINLINE ElementType& FindOrAdd(InElementType && InElement, bool* bIsAlreadyInSetPtr = nullptr)
+	UE_FORCEINLINE_HINT ElementType& FindOrAdd(InElementType && InElement, bool* bIsAlreadyInSetPtr = nullptr)
 	{
 		return FindOrAddByHash(KeyFuncs::GetKeyHash(KeyFuncs::GetSetKey(InElement)), MoveTempIfPossible(InElement), bIsAlreadyInSetPtr);
 	}
@@ -962,7 +962,7 @@ public:
 	 * @param Key - The key to search for.
 	 * @return A const pointer to an element with the given key.  If no element in the set has the given key, this will return NULL.
 	 */
-	[[nodiscard]] FORCEINLINE const ElementType* Find(KeyInitType Key) const
+	[[nodiscard]] UE_FORCEINLINE_HINT const ElementType* Find(KeyInitType Key) const
 	{
 		return const_cast<TCompactSet*>(this)->Find(Key);
 	}
@@ -1065,7 +1065,7 @@ public:
 	 * @param Key - The key to check for.
 	 * @return true if the set contains an element with the given key.
 	 */
-	[[nodiscard]] FORCEINLINE bool Contains(KeyInitType Key) const
+	[[nodiscard]] UE_FORCEINLINE_HINT bool Contains(KeyInitType Key) const
 	{
 		return FindIndexByHash(KeyFuncs::GetKeyHash(Key), Key) != INDEX_NONE;
 	}
@@ -1232,13 +1232,13 @@ public:
 	}
 
 	/** Creates an iterator for the contents of this set */
-	[[nodiscard]] FORCEINLINE TIterator CreateIterator()
+	[[nodiscard]] UE_FORCEINLINE_HINT TIterator CreateIterator()
 	{
 		return TIterator(*this);
 	}
 
 	/** Creates a const iterator for the contents of this set */
-	[[nodiscard]] FORCEINLINE TConstIterator CreateConstIterator() const
+	[[nodiscard]] UE_FORCEINLINE_HINT TConstIterator CreateConstIterator() const
 	{
 		return TConstIterator(*this);
 	}
@@ -1251,18 +1251,18 @@ public:
 	using TRangedForIterator             = TCheckedPointerIterator<      ElementType, SizeType, false>;
 	using TRangedForConstIterator        = TCheckedPointerIterator<const ElementType, SizeType, false>;
 
-	[[nodiscard]] FORCEINLINE TRangedForIterator      begin() { return TRangedForIterator(this->NumElements, GetData()); }
-	[[nodiscard]] FORCEINLINE TRangedForConstIterator begin() const { return TRangedForConstIterator(this->NumElements, GetData()); }
-	[[nodiscard]] FORCEINLINE TRangedForIterator      end() { return TRangedForIterator(this->NumElements, GetData() + this->NumElements); }
-	[[nodiscard]] FORCEINLINE TRangedForConstIterator end() const { return TRangedForConstIterator(this->NumElements, GetData() + this->NumElements); }
+	[[nodiscard]] UE_FORCEINLINE_HINT TRangedForIterator      begin() { return TRangedForIterator(this->NumElements, GetData()); }
+	[[nodiscard]] UE_FORCEINLINE_HINT TRangedForConstIterator begin() const { return TRangedForConstIterator(this->NumElements, GetData()); }
+	[[nodiscard]] UE_FORCEINLINE_HINT TRangedForIterator      end() { return TRangedForIterator(this->NumElements, GetData() + this->NumElements); }
+	[[nodiscard]] UE_FORCEINLINE_HINT TRangedForConstIterator end() const { return TRangedForConstIterator(this->NumElements, GetData() + this->NumElements); }
 	#else
 	using TRangedForIterator = ElementType*;
 	using TRangedForConstIterator = const ElementType*;
 
-	[[nodiscard]] FORCEINLINE ElementType*       begin() { return GetData(); }
-	[[nodiscard]] FORCEINLINE const ElementType* begin() const { return GetData(); }
-	[[nodiscard]] FORCEINLINE ElementType*       end() { return GetData() + this->NumElements; }
-	[[nodiscard]] FORCEINLINE const ElementType* end() const { return GetData() + this->NumElements; }
+	[[nodiscard]] UE_FORCEINLINE_HINT ElementType*       begin() { return GetData(); }
+	[[nodiscard]] UE_FORCEINLINE_HINT const ElementType* begin() const { return GetData(); }
+	[[nodiscard]] UE_FORCEINLINE_HINT ElementType*       end() { return GetData() + this->NumElements; }
+	[[nodiscard]] UE_FORCEINLINE_HINT const ElementType* end() const { return GetData() + this->NumElements; }
 	#endif
 
 	// Sets are deliberately prevented from being hashed or compared, because this would hide potentially major performance problems behind default operations.
@@ -1347,23 +1347,23 @@ public:
 	}
 
 private:
-	[[nodiscard]] FORCEINLINE static constexpr FCompactSetLayout GetSetLayout()
+	[[nodiscard]] UE_FORCEINLINE_HINT static constexpr FCompactSetLayout GetSetLayout()
 	{
 		return { sizeof(ElementType), GetAlignment() };
 	}
 
-	[[nodiscard]] FORCEINLINE FCompactHashTableView GetHashTableView()
+	[[nodiscard]] UE_FORCEINLINE_HINT FCompactHashTableView GetHashTableView()
 	{
 		return Super::GetHashTableView(GetSetLayout());
 	}
 
-	[[nodiscard]] FORCEINLINE FConstCompactHashTableView GetConstHashTableView() const
+	[[nodiscard]] UE_FORCEINLINE_HINT FConstCompactHashTableView GetConstHashTableView() const
 	{
 		return Super::GetConstHashTableView(GetSetLayout());
 	}
 
 	// Use if you're going to reset/rehash regardless of the results
-	FORCEINLINE void ResizeAllocation(SizeType NewMaxElements)
+	UE_FORCEINLINE_HINT void ResizeAllocation(SizeType NewMaxElements)
 	{
 		(void)Super::ResizeAllocationPreserveData(NewMaxElements, GetSetLayout(), false);
 	}
@@ -1532,7 +1532,7 @@ private:
 	}
 
 	template<typename ComparableKey>
-	FORCEINLINE int32 RemoveImplStable(uint32 KeyHash, const ComparableKey& Key)
+	UE_FORCEINLINE_HINT int32 RemoveImplStable(uint32 KeyHash, const ComparableKey& Key)
 	{
 		return RemoveImpl<ComparableKey, true>(KeyHash, Key);
 	}
diff --git a/Engine/Source/Runtime/Core/Public/Containers/CompactSetBase.h b/Engine/Source/Runtime/Core/Public/Containers/CompactSetBase.h
--- a/Engine/Source/Runtime/Core/Public/Containers/CompactSetBase.h
+++ b/Engine/Source/Runtime/Core/Public/Containers/CompactSetBase.h
@@ -109,18 +109,18 @@ public:
 		return MaxElements == INDEX_NONE;
 	}
 
-	[[nodiscard]] FORCEINLINE bool IsEmpty() const
+	[[nodiscard]] UE_FORCEINLINE_HINT bool IsEmpty() const
 	{
 		return NumElements == 0;
 	}
 
-	[[nodiscard]] FORCEINLINE int32 Num() const
+	[[nodiscard]] UE_FORCEINLINE_HINT int32 Num() const
 	{
 		return NumElements;
 	}
 
 	/** Deprecated - unnecessary, keeping this here so TCompactSet can be swapped with TSet without changing code */
-	[[nodiscard]] FORCEINLINE int32 GetMaxIndex() const
+	[[nodiscard]] UE_FORCEINLINE_HINT int32 GetMaxIndex() const
 	{
 		return NumElements;
 	}
@@ -133,7 +133,7 @@ protected:
 	static_assert(std::is_same_v<SizeType, int32>, "TCompactSet currently only supports 32-bit allocators");
 	static_assert(sizeof(HashCountType) == UE::Core::CompactHashTable::GetMemoryAlignment(), "Hashtable alignment changed, need to make sure HashCountType is stored correctly");
 
-	[[nodiscard]] FORCEINLINE TCompactSetBase() = default;	
+	[[nodiscard]] UE_FORCEINLINE_HINT TCompactSetBase() = default;	
 
 	[[nodiscard]] explicit TCompactSetBase(FIntrusiveUnsetOptionalState Tag)
 	: NumElements(0)
@@ -142,7 +142,7 @@ protected:
 	}
 
 	/* Calculate the size of the hash table from the number of elements in the set */
-	[[nodiscard]] FORCEINLINE const HashCountType* GetHashTableMemory(const FCompactSetLayout Layout) const
+	[[nodiscard]] UE_FORCEINLINE_HINT const HashCountType* GetHashTableMemory(const FCompactSetLayout Layout) const
 	{
 		return (const HashCountType*)(Elements.GetAllocation() + GetElementsSizeInBytes(MaxElements, Layout));
 	}
@@ -162,25 +162,25 @@ protected:
 	}
 
 	/* Calculate the size of the hash table from the number of elements in the set */
-	[[nodiscard]] FORCEINLINE static constexpr SizeType GetHashCount(uint32 NumElements)
+	[[nodiscard]] UE_FORCEINLINE_HINT static constexpr SizeType GetHashCount(uint32 NumElements)
 	{
 		return UE::Core::CompactHashTable::GetHashCount(NumElements);
 	}
 
 	/* Calculate bytes required to store the elements, includes any padding required for the hash table */
-	[[nodiscard]] FORCEINLINE static constexpr size_t GetElementsSizeInBytes(uint32 NumElements, const FCompactSetLayout Layout)
+	[[nodiscard]] UE_FORCEINLINE_HINT static constexpr size_t GetElementsSizeInBytes(uint32 NumElements, const FCompactSetLayout Layout)
 	{
 		return Align(Layout.Size * NumElements, UE::Core::CompactHashTable::GetMemoryAlignment());
 	}
 
 	/* Get the total memory required for the compact set for the given number of elements and size of the number of hash elements */
-	[[nodiscard]] FORCEINLINE static constexpr size_t GetTotalMemoryRequiredInBytes(uint32 NumElements, uint32 HashCount, const FCompactSetLayout Layout)
+	[[nodiscard]] UE_FORCEINLINE_HINT static constexpr size_t GetTotalMemoryRequiredInBytes(uint32 NumElements, uint32 HashCount, const FCompactSetLayout Layout)
 	{
 		return NumElements ? GetElementsSizeInBytes(NumElements, Layout) + UE::Core::CompactHashTable::GetMemoryRequiredInBytes(NumElements, HashCount) + HashCountSize : 0;
 	}
 
 	/* Calculate the size of the hash table from the number of elements in the set assuming the default number of hash elements */
-	[[nodiscard]] FORCEINLINE static constexpr size_t GetTotalMemoryRequiredInBytes(uint32 NumElements, const FCompactSetLayout Layout)
+	[[nodiscard]] UE_FORCEINLINE_HINT static constexpr size_t GetTotalMemoryRequiredInBytes(uint32 NumElements, const FCompactSetLayout Layout)
 	{
 		return NumElements ? GetElementsSizeInBytes(NumElements, Layout) + UE::Core::CompactHashTable::GetMemoryRequiredInBytes(NumElements, GetHashCount(NumElements)) + HashCountSize : 0;
 	}
diff --git a/Engine/Source/Runtime/Core/Public/Containers/ContainerAllocationPolicies.h b/Engine/Source/Runtime/Core/Public/Containers/ContainerAllocationPolicies.h
--- a/Engine/Source/Runtime/Core/Public/Containers/ContainerAllocationPolicies.h
+++ b/Engine/Source/Runtime/Core/Public/Containers/ContainerAllocationPolicies.h
@@ -69,7 +69,7 @@ struct FArraySlackTrackingHeader
 		}
 	}
 
-	static FORCEINLINE void UpdateNumUsed(void* Ptr, int64 NewNumUsed)
+	static UE_FORCEINLINE_HINT void UpdateNumUsed(void* Ptr, int64 NewNumUsed)
 	{
 		if (Ptr)
 		{
@@ -79,7 +79,7 @@ struct FArraySlackTrackingHeader
 		}
 	}
 
-	static FORCEINLINE void DisableTracking(void* Ptr)
+	static UE_FORCEINLINE_HINT void DisableTracking(void* Ptr)
 	{
 		if (Ptr)
 		{
@@ -94,7 +94,7 @@ struct FArraySlackTrackingHeader
 		}
 	}
 
-	FORCEINLINE int64 SlackSizeInBytes() const
+	UE_FORCEINLINE_HINT int64 SlackSizeInBytes() const
 	{
 		return (ArrayMax - ArrayNum) * ElemSize;
 	}
@@ -141,7 +141,7 @@ template<int IndexSize> class TSizedDefaultAllocator;
 using FDefaultAllocator = TSizedDefaultAllocator<32>;
 
 template <typename SizeType>
-FORCEINLINE SizeType DefaultCalculateSlackShrink(SizeType NewMax, SizeType CurrentMax, SIZE_T BytesPerElement, bool bAllowQuantize, uint32 Alignment = DEFAULT_ALIGNMENT)
+UE_FORCEINLINE_HINT SizeType DefaultCalculateSlackShrink(SizeType NewMax, SizeType CurrentMax, SIZE_T BytesPerElement, bool bAllowQuantize, uint32 Alignment = DEFAULT_ALIGNMENT)
 {
 	SizeType Retval;
 	checkSlow(NewMax < CurrentMax);
@@ -171,7 +171,7 @@ FORCEINLINE SizeType DefaultCalculateSlackShrink(SizeType NewMax, SizeType Curre
 }
 
 template <typename SizeType>
-FORCEINLINE SizeType DefaultCalculateSlackGrow(SizeType NewMax, SizeType CurrentMax, SIZE_T BytesPerElement, bool bAllowQuantize, uint32 Alignment = DEFAULT_ALIGNMENT)
+UE_FORCEINLINE_HINT SizeType DefaultCalculateSlackGrow(SizeType NewMax, SizeType CurrentMax, SIZE_T BytesPerElement, bool bAllowQuantize, uint32 Alignment = DEFAULT_ALIGNMENT)
 {
 #if !defined(AGGRESSIVE_MEMORY_SAVING)
 	#error "AGGRESSIVE_MEMORY_SAVING must be defined"
@@ -225,7 +225,7 @@ FORCEINLINE SizeType DefaultCalculateSlackGrow(SizeType NewMax, SizeType Current
 }
 
 template <typename SizeType>
-FORCEINLINE SizeType DefaultCalculateSlackReserve(SizeType NewMax, SIZE_T BytesPerElement, bool bAllowQuantize, uint32 Alignment = DEFAULT_ALIGNMENT)
+UE_FORCEINLINE_HINT SizeType DefaultCalculateSlackReserve(SizeType NewMax, SIZE_T BytesPerElement, bool bAllowQuantize, uint32 Alignment = DEFAULT_ALIGNMENT)
 {
 	SizeType Retval = NewMax;
 	checkSlow(NewMax > 0);
@@ -470,7 +470,7 @@ public:
 		 * Assumes that the allocator is currently empty, i.e. memory may be allocated but any existing elements have already been destructed (if necessary).
 		 * @param Other - The allocator to move the state from.  This allocator should be left in a valid empty state.
 		 */
-		FORCEINLINE void MoveToEmpty(ForAnyElementType& Other)
+		UE_FORCEINLINE_HINT void MoveToEmpty(ForAnyElementType& Other)
 		{
 			checkSlow(this != &Other);
 
@@ -488,7 +488,7 @@ public:
 		}
 
 		/** Destructor. */
-		FORCEINLINE ~ForAnyElementType()
+		UE_FORCEINLINE_HINT ~ForAnyElementType()
 		{
 			if(Data)
 			{
@@ -501,7 +501,7 @@ public:
 		}
 
 		// FContainerAllocatorInterface
-		FORCEINLINE FScriptContainerElement* GetAllocation() const
+		UE_FORCEINLINE_HINT FScriptContainerElement* GetAllocation() const
 		{
 			return Data;
 		}
@@ -529,15 +529,15 @@ public:
 #endif
 			}
 		}
-		FORCEINLINE SizeType CalculateSlackReserve(SizeType NewMax, SIZE_T NumBytesPerElement) const
+		UE_FORCEINLINE_HINT SizeType CalculateSlackReserve(SizeType NewMax, SIZE_T NumBytesPerElement) const
 		{
 			return DefaultCalculateSlackReserve(NewMax, NumBytesPerElement, true, Alignment);
 		}
-		FORCEINLINE SizeType CalculateSlackShrink(SizeType NewMax, SizeType CurrentMax, SIZE_T NumBytesPerElement) const
+		UE_FORCEINLINE_HINT SizeType CalculateSlackShrink(SizeType NewMax, SizeType CurrentMax, SIZE_T NumBytesPerElement) const
 		{
 			return DefaultCalculateSlackShrink(NewMax, CurrentMax, NumBytesPerElement, true, Alignment);
 		}
-		FORCEINLINE SizeType CalculateSlackGrow(SizeType NewMax, SizeType CurrentMax, SIZE_T NumBytesPerElement) const
+		UE_FORCEINLINE_HINT SizeType CalculateSlackGrow(SizeType NewMax, SizeType CurrentMax, SIZE_T NumBytesPerElement) const
 		{
 			return DefaultCalculateSlackGrow(NewMax, CurrentMax, NumBytesPerElement, true, Alignment);
 		}
@@ -558,14 +558,14 @@ public:
 		}
 
 #if UE_ENABLE_ARRAY_SLACK_TRACKING
-		FORCEINLINE void SlackTrackerLogNum(SizeType NewNumUsed)
+		UE_FORCEINLINE_HINT void SlackTrackerLogNum(SizeType NewNumUsed)
 		{
 			FArraySlackTrackingHeader::UpdateNumUsed(Data, (int64)NewNumUsed);
 		}
 
 		// Suppress slack tracking on an allocation -- should be called whenever the container may have been resized.
 		// Useful for debug allocations you don't want to show up in slack reports.
-		FORCEINLINE void DisableSlackTracking()
+		UE_FORCEINLINE_HINT void DisableSlackTracking()
 		{
 			FArraySlackTrackingHeader::DisableTracking(Data);
 		}
@@ -591,7 +591,7 @@ public:
 			static_assert(alignof(ElementType) <= MinimumAlignment, "Using TAlignedHeapAllocator with an alignment lower than the element type's alignment - please update the alignment parameter");
 		}
 
-		FORCEINLINE ElementType* GetAllocation() const
+		UE_FORCEINLINE_HINT ElementType* GetAllocation() const
 		{
 			return (ElementType*)ForAnyElementType::GetAllocation();
 		}
@@ -649,7 +649,7 @@ public:
 		 * @param Other - The allocator to move the state from.  This allocator should be left in a valid empty state.
 		 */
 		template <typename OtherAllocator>
-		FORCEINLINE void MoveToEmptyFromOtherAllocator(typename OtherAllocator::ForAnyElementType& Other)
+		UE_FORCEINLINE_HINT void MoveToEmptyFromOtherAllocator(typename OtherAllocator::ForAnyElementType& Other)
 		{
 			checkSlow((void*)this != (void*)&Other);
 
@@ -673,13 +673,13 @@ public:
 		 * Assumes that the allocator is currently empty, i.e. memory may be allocated but any existing elements have already been destructed (if necessary).
 		 * @param Other - The allocator to move the state from.  This allocator should be left in a valid empty state.
 		 */
-		FORCEINLINE void MoveToEmpty(ForAnyElementType& Other)
+		UE_FORCEINLINE_HINT void MoveToEmpty(ForAnyElementType& Other)
 		{
 			this->MoveToEmptyFromOtherAllocator<TSizedHeapAllocator>(Other);
 		}
 
 		/** Destructor. */
-		FORCEINLINE ~ForAnyElementType()
+		UE_FORCEINLINE_HINT ~ForAnyElementType()
 		{
 			if(Data)
 			{
@@ -692,7 +692,7 @@ public:
 		}
 
 		// FContainerAllocatorInterface
-		FORCEINLINE FScriptContainerElement* GetAllocation() const
+		UE_FORCEINLINE_HINT FScriptContainerElement* GetAllocation() const
 		{
 			return Data;
 		}
@@ -746,27 +746,27 @@ public:
 #endif
 			}
 		}
-		FORCEINLINE SizeType CalculateSlackReserve(SizeType NewMax, SIZE_T NumBytesPerElement) const
+		UE_FORCEINLINE_HINT SizeType CalculateSlackReserve(SizeType NewMax, SIZE_T NumBytesPerElement) const
 		{
 			return DefaultCalculateSlackReserve(NewMax, NumBytesPerElement, true);
 		}
-		FORCEINLINE SizeType CalculateSlackReserve(SizeType NewMax, SIZE_T NumBytesPerElement, uint32 AlignmentOfElement) const
+		UE_FORCEINLINE_HINT SizeType CalculateSlackReserve(SizeType NewMax, SIZE_T NumBytesPerElement, uint32 AlignmentOfElement) const
 		{
 			return DefaultCalculateSlackReserve(NewMax, NumBytesPerElement, true, (uint32)AlignmentOfElement);
 		}
-		FORCEINLINE SizeType CalculateSlackShrink(SizeType NewMax, SizeType CurrentMax, SIZE_T NumBytesPerElement) const
+		UE_FORCEINLINE_HINT SizeType CalculateSlackShrink(SizeType NewMax, SizeType CurrentMax, SIZE_T NumBytesPerElement) const
 		{
 			return DefaultCalculateSlackShrink(NewMax, CurrentMax, NumBytesPerElement, true);
 		}
-		FORCEINLINE SizeType CalculateSlackShrink(SizeType NewMax, SizeType CurrentMax, SIZE_T NumBytesPerElement, uint32 AlignmentOfElement) const
+		UE_FORCEINLINE_HINT SizeType CalculateSlackShrink(SizeType NewMax, SizeType CurrentMax, SIZE_T NumBytesPerElement, uint32 AlignmentOfElement) const
 		{
 			return DefaultCalculateSlackShrink(NewMax, CurrentMax, NumBytesPerElement, true, (uint32)AlignmentOfElement);
 		}
-		FORCEINLINE SizeType CalculateSlackGrow(SizeType NewMax, SizeType CurrentMax, SIZE_T NumBytesPerElement) const
+		UE_FORCEINLINE_HINT SizeType CalculateSlackGrow(SizeType NewMax, SizeType CurrentMax, SIZE_T NumBytesPerElement) const
 		{
 			return DefaultCalculateSlackGrow(NewMax, CurrentMax, NumBytesPerElement, true);
 		}
-		FORCEINLINE SizeType CalculateSlackGrow(SizeType NewMax, SizeType CurrentMax, SIZE_T NumBytesPerElement, uint32 AlignmentOfElement) const
+		UE_FORCEINLINE_HINT SizeType CalculateSlackGrow(SizeType NewMax, SizeType CurrentMax, SIZE_T NumBytesPerElement, uint32 AlignmentOfElement) const
 		{
 			return DefaultCalculateSlackGrow(NewMax, CurrentMax, NumBytesPerElement, true, (uint32)AlignmentOfElement);
 		}
@@ -787,14 +787,14 @@ public:
 		}
 
 #if UE_ENABLE_ARRAY_SLACK_TRACKING
-		FORCEINLINE void SlackTrackerLogNum(SizeType NewNumUsed)
+		UE_FORCEINLINE_HINT void SlackTrackerLogNum(SizeType NewNumUsed)
 		{
 			FArraySlackTrackingHeader::UpdateNumUsed(Data, (int64)NewNumUsed);
 		}
 
 		// Suppress slack tracking on an allocation -- should be called whenever the container may have been resized.
 		// Useful for debug allocations you don't want to show up in slack reports.
-		FORCEINLINE void DisableSlackTracking()
+		UE_FORCEINLINE_HINT void DisableSlackTracking()
 		{
 			FArraySlackTrackingHeader::DisableTracking(Data);
 		}
@@ -815,7 +815,7 @@ public:
 		/** Default constructor. */
 		ForElementType() = default;
 
-		FORCEINLINE ElementType* GetAllocation() const
+		UE_FORCEINLINE_HINT ElementType* GetAllocation() const
 		{
 			return (ElementType*)ForAnyElementType::GetAllocation();
 		}
@@ -878,7 +878,7 @@ public:
 		 * Assumes that the allocator is currently empty, i.e. memory may be allocated but any existing elements have already been destructed (if necessary).
 		 * @param Other - The allocator to move the state from.  This allocator should be left in a valid empty state.
 		 */
-		FORCEINLINE void MoveToEmpty(ForElementType& Other)
+		UE_FORCEINLINE_HINT void MoveToEmpty(ForElementType& Other)
 		{
 			checkSlow(this != &Other);
 
@@ -894,7 +894,7 @@ public:
 		}
 
 		// FContainerAllocatorInterface
-		FORCEINLINE ElementType* GetAllocation() const
+		UE_FORCEINLINE_HINT ElementType* GetAllocation() const
 		{
 			if (ElementType* Result = SecondaryData.GetAllocation())
 			{
@@ -935,21 +935,21 @@ public:
 			}
 		}
 
-		FORCEINLINE SizeType CalculateSlackReserve(SizeType NewMax, SIZE_T NumBytesPerElement) const
+		UE_FORCEINLINE_HINT SizeType CalculateSlackReserve(SizeType NewMax, SIZE_T NumBytesPerElement) const
 		{
 			// If the elements use less space than the inline allocation, only use the inline allocation as slack.
 			return NewMax <= NumInlineElements ?
 				NumInlineElements :
 				SecondaryData.CalculateSlackReserve(NewMax, NumBytesPerElement);
 		}
-		FORCEINLINE SizeType CalculateSlackShrink(SizeType NewMax, SizeType CurrentMax, SIZE_T NumBytesPerElement) const
+		UE_FORCEINLINE_HINT SizeType CalculateSlackShrink(SizeType NewMax, SizeType CurrentMax, SIZE_T NumBytesPerElement) const
 		{
 			// If the elements use less space than the inline allocation, only use the inline allocation as slack.
 			return NewMax <= NumInlineElements ?
 				NumInlineElements :
 				SecondaryData.CalculateSlackShrink(NewMax, CurrentMax, NumBytesPerElement);
 		}
-		FORCEINLINE SizeType CalculateSlackGrow(SizeType NewMax, SizeType CurrentMax, SIZE_T NumBytesPerElement) const
+		UE_FORCEINLINE_HINT SizeType CalculateSlackGrow(SizeType NewMax, SizeType CurrentMax, SIZE_T NumBytesPerElement) const
 		{
 			// If the elements use less space than the inline allocation, only use the inline allocation as slack.
 			// Also, when computing slack growth, don't count inline elements -- the slack algorithm has a special
@@ -980,7 +980,7 @@ public:
 		}
 
 #if UE_ENABLE_ARRAY_SLACK_TRACKING
-		FORCEINLINE void SlackTrackerLogNum(SizeType NewNumUsed)
+		UE_FORCEINLINE_HINT void SlackTrackerLogNum(SizeType NewNumUsed)
 		{
 			if constexpr (TAllocatorTraits<SecondaryAllocator>::SupportsSlackTracking)
 			{
@@ -1065,7 +1065,7 @@ public:
 		 * Assumes that the allocator is currently empty, i.e. memory may be allocated but any existing elements have already been destructed (if necessary).
 		 * @param Other - The allocator to move the state from.  This allocator should be left in a valid empty state.
 		 */
-		FORCEINLINE void MoveToEmpty(ForElementType& Other)
+		UE_FORCEINLINE_HINT void MoveToEmpty(ForElementType& Other)
 		{
 			checkSlow(this != &Other);
 
@@ -1091,7 +1091,7 @@ public:
 		}
 
 		// FContainerAllocatorInterface
-		FORCEINLINE ElementType* GetAllocation() const
+		UE_FORCEINLINE_HINT ElementType* GetAllocation() const
 		{
 			return Data;
 		}
@@ -1139,19 +1139,19 @@ public:
 			}
 		}
 
-		FORCEINLINE SizeType CalculateSlackReserve(SizeType NewMax, SIZE_T NumBytesPerElement) const
+		UE_FORCEINLINE_HINT SizeType CalculateSlackReserve(SizeType NewMax, SIZE_T NumBytesPerElement) const
 		{
 			// If the elements use less space than the inline allocation, only use the inline allocation as slack.
 			return (NewMax <= NumInlineElements) ? NumInlineElements : DefaultCalculateSlackReserve(NewMax, NumBytesPerElement, true);
 		}
 
-		FORCEINLINE SizeType CalculateSlackShrink(SizeType NewMax, SizeType CurrentMax, SIZE_T NumBytesPerElement) const
+		UE_FORCEINLINE_HINT SizeType CalculateSlackShrink(SizeType NewMax, SizeType CurrentMax, SIZE_T NumBytesPerElement) const
 		{
 			// If the elements use less space than the inline allocation, only use the inline allocation as slack.
 			return (NewMax <= NumInlineElements) ? NumInlineElements : DefaultCalculateSlackShrink(NewMax, CurrentMax, NumBytesPerElement, true);
 		}
 
-		FORCEINLINE SizeType CalculateSlackGrow(SizeType NewMax, SizeType CurrentMax, SIZE_T NumBytesPerElement) const
+		UE_FORCEINLINE_HINT SizeType CalculateSlackGrow(SizeType NewMax, SizeType CurrentMax, SIZE_T NumBytesPerElement) const
 		{
 			// If the elements use less space than the inline allocation, only use the inline allocation as slack.
 			return (NewMax <= NumInlineElements) ? NumInlineElements : DefaultCalculateSlackGrow(NewMax, CurrentMax, NumBytesPerElement, true);
@@ -1162,7 +1162,7 @@ public:
 			return HasAllocation() ? (CurrentMax * NumBytesPerElement) : 0;
 		}
 
-		FORCEINLINE bool HasAllocation() const
+		UE_FORCEINLINE_HINT bool HasAllocation() const
 		{
 			return Data != GetInlineElements();
 		}
@@ -1173,7 +1173,7 @@ public:
 		}
 
 #if UE_ENABLE_ARRAY_SLACK_TRACKING
-		FORCEINLINE void SlackTrackerLogNum(SizeType NewNumUsed)
+		UE_FORCEINLINE_HINT void SlackTrackerLogNum(SizeType NewNumUsed)
 		{
 			if (HasAllocation())
 			{
@@ -1195,7 +1195,7 @@ public:
 		TTypeCompatibleBytes<ElementType> InlineData[NumInlineElements];
 
 		/** @return the base of the aligned inline element data */
-		FORCEINLINE ElementType* GetInlineElements() const
+		UE_FORCEINLINE_HINT ElementType* GetInlineElements() const
 		{
 			return (ElementType*)InlineData;
 		}
@@ -1238,7 +1238,7 @@ public:
 		 * Assumes that the allocator is currently empty, i.e. memory may be allocated but any existing elements have already been destructed (if necessary).
 		 * @param Other - The allocator to move the state from.  This allocator should be left in a valid empty state.
 		 */
-		FORCEINLINE void MoveToEmpty(ForElementType& Other)
+		UE_FORCEINLINE_HINT void MoveToEmpty(ForElementType& Other)
 		{
 			checkSlow(this != &Other);
 
@@ -1247,7 +1247,7 @@ public:
 		}
 
 		// FContainerAllocatorInterface
-		FORCEINLINE ElementType* GetAllocation() const
+		UE_FORCEINLINE_HINT ElementType* GetAllocation() const
 		{
 			return GetInlineElements();
 		}
@@ -1258,19 +1258,19 @@ public:
 			check(NewMax >= 0 && NewMax <= NumInlineElements);
 		}
 
-		FORCEINLINE SizeType CalculateSlackReserve(SizeType NewMax, SIZE_T NumBytesPerElement) const
+		UE_FORCEINLINE_HINT SizeType CalculateSlackReserve(SizeType NewMax, SIZE_T NumBytesPerElement) const
 		{
 			// Ensure the requested allocation will fit in the inline data area.
 			check(NewMax <= NumInlineElements);
 			return NumInlineElements;
 		}
-		FORCEINLINE SizeType CalculateSlackShrink(SizeType NewMax, SizeType CurrentMax, SIZE_T NumBytesPerElement) const
+		UE_FORCEINLINE_HINT SizeType CalculateSlackShrink(SizeType NewMax, SizeType CurrentMax, SIZE_T NumBytesPerElement) const
 		{
 			// Ensure the requested allocation will fit in the inline data area.
 			check(CurrentMax <= NumInlineElements);
 			return NumInlineElements;
 		}
-		FORCEINLINE SizeType CalculateSlackGrow(SizeType NewMax, SizeType CurrentMax, SIZE_T NumBytesPerElement) const
+		UE_FORCEINLINE_HINT SizeType CalculateSlackGrow(SizeType NewMax, SizeType CurrentMax, SIZE_T NumBytesPerElement) const
 		{
 			// Ensure the requested allocation will fit in the inline data area.
 			check(NewMax <= NumInlineElements);
@@ -1393,7 +1393,7 @@ class TSetAllocator
 public:
 
 	/** Computes the number of hash buckets to use for a given number of elements. */
-	static FORCEINLINE uint32 GetNumberOfHashBuckets(uint32 NumHashedElements)
+	static UE_FORCEINLINE_HINT uint32 GetNumberOfHashBuckets(uint32 NumHashedElements)
 	{
 		if(NumHashedElements >= MinNumberOfHashedElements)
 		{
@@ -1441,7 +1441,7 @@ private:
 public:
 
 	/** Computes the number of hash buckets to use for a given number of elements. */
-	static FORCEINLINE uint32 GetNumberOfHashBuckets(uint32 NumHashedElements)
+	static UE_FORCEINLINE_HINT uint32 GetNumberOfHashBuckets(uint32 NumHashedElements)
 	{
 		const uint32 NumDesiredHashBuckets = FPlatformMath::RoundUpToPowerOfTwo(NumHashedElements / AverageNumberOfElementsPerHashBucket);
 		if (NumDesiredHashBuckets < NumInlineHashBuckets)
@@ -1478,7 +1478,7 @@ private:
 public:
 
 	/** Computes the number of hash buckets to use for a given number of elements. */
-	static FORCEINLINE uint32 GetNumberOfHashBuckets(uint32 NumHashedElements)
+	static UE_FORCEINLINE_HINT uint32 GetNumberOfHashBuckets(uint32 NumHashedElements)
 	{
 		const uint32 NumDesiredHashBuckets = FPlatformMath::RoundUpToPowerOfTwo(NumHashedElements / AverageNumberOfElementsPerHashBucket);
 		if (NumDesiredHashBuckets < NumInlineHashBuckets)
diff --git a/Engine/Source/Runtime/Core/Public/Containers/Deque.h b/Engine/Source/Runtime/Core/Public/Containers/Deque.h
--- a/Engine/Source/Runtime/Core/Public/Containers/Deque.h
+++ b/Engine/Source/Runtime/Core/Public/Containers/Deque.h
@@ -26,7 +26,7 @@ namespace Private
  * avoid the modulo operator or a power of 2 range value requirement and have an efficient wrap around function.
  */
 template <typename InSizeType>
-FORCEINLINE InSizeType WrapAround(InSizeType Index, InSizeType Range)
+UE_FORCEINLINE_HINT InSizeType WrapAround(InSizeType Index, InSizeType Range)
 {
 	return (Index < Range) ? Index : Index - Range;
 }
@@ -52,7 +52,7 @@ protected:
 	{
 	}
 
-	[[nodiscard]] FORCEINLINE ElementType& Dereference() const
+	[[nodiscard]] UE_FORCEINLINE_HINT ElementType& Dereference() const
 	{
 		return *(Data + WrapAround(Offset, Range));
 	}
@@ -63,7 +63,7 @@ protected:
 		Offset++;
 	}
 
-	[[nodiscard]] FORCEINLINE bool Equals(const TIteratorBase& Other) const
+	[[nodiscard]] UE_FORCEINLINE_HINT bool Equals(const TIteratorBase& Other) const
 	{
 		return Data + Offset == Other.Data + Other.Offset;
 	}
@@ -195,17 +195,17 @@ public:
 		return GetData()[Head];
 	}
 
-	[[nodiscard]] FORCEINLINE bool IsEmpty() const
+	[[nodiscard]] UE_FORCEINLINE_HINT bool IsEmpty() const
 	{
 		return !Count;
 	}
 
-	[[nodiscard]] FORCEINLINE SizeType Max() const
+	[[nodiscard]] UE_FORCEINLINE_HINT SizeType Max() const
 	{
 		return Capacity;
 	}
 
-	[[nodiscard]] FORCEINLINE SizeType Num() const
+	[[nodiscard]] UE_FORCEINLINE_HINT SizeType Num() const
 	{
 		return Count;
 	}
@@ -216,7 +216,7 @@ public:
 	 *
 	 * @returns Number of bytes allocated by this container.
 	 */
-	[[nodiscard]] FORCEINLINE SIZE_T GetAllocatedSize() const
+	[[nodiscard]] UE_FORCEINLINE_HINT SIZE_T GetAllocatedSize() const
 	{
 		return Storage.GetAllocatedSize(Capacity, sizeof(ElementType));
 	}
@@ -254,12 +254,12 @@ public:
 	/*
 	 * Adds the parameter element at the back of the queue.
 	 */
-	FORCEINLINE void PushLast(const ElementType& Element)
+	UE_FORCEINLINE_HINT void PushLast(const ElementType& Element)
 	{
 		EmplaceLast(Element);
 	}
 
-	FORCEINLINE void PushLast(ElementType&& Element)
+	UE_FORCEINLINE_HINT void PushLast(ElementType&& Element)
 	{
 		EmplaceLast(MoveTempIfPossible(Element));
 	}
@@ -267,12 +267,12 @@ public:
 	/*
 	 * Adds the parameter element at the front of the queue.
 	 */
-	FORCEINLINE void PushFirst(const ElementType& Element)
+	UE_FORCEINLINE_HINT void PushFirst(const ElementType& Element)
 	{
 		EmplaceFirst(Element);
 	}
 
-	FORCEINLINE void PushFirst(ElementType&& Element)
+	UE_FORCEINLINE_HINT void PushFirst(ElementType&& Element)
 	{
 		EmplaceFirst(MoveTempIfPossible(Element));
 	}
@@ -384,19 +384,19 @@ public:
 	/*
 	 * STL IteratorType model compliance methods.
 	 */
-	[[nodiscard]] FORCEINLINE ConstIteratorType begin() const
+	[[nodiscard]] UE_FORCEINLINE_HINT ConstIteratorType begin() const
 	{
 		return GetIterator(Head);
 	}
-	[[nodiscard]] FORCEINLINE IteratorType begin()
+	[[nodiscard]] UE_FORCEINLINE_HINT IteratorType begin()
 	{
 		return GetIterator(Head);
 	}
-	[[nodiscard]] FORCEINLINE ConstIteratorType end() const
+	[[nodiscard]] UE_FORCEINLINE_HINT ConstIteratorType end() const
 	{
 		return GetIterator(Head + Count);
 	}
-	[[nodiscard]] FORCEINLINE IteratorType end()
+	[[nodiscard]] UE_FORCEINLINE_HINT IteratorType end()
 	{
 		return GetIterator(Head + Count);
 	}
@@ -408,22 +408,22 @@ private:
 	SizeType Head = 0;
 	SizeType Tail = 0;
 
-	[[nodiscard]] FORCEINLINE const ElementType* GetData() const
+	[[nodiscard]] UE_FORCEINLINE_HINT const ElementType* GetData() const
 	{
 		return (const ElementType*)Storage.GetAllocation();
 	}
 
-	[[nodiscard]] FORCEINLINE ElementType* GetData()
+	[[nodiscard]] UE_FORCEINLINE_HINT ElementType* GetData()
 	{
 		return (ElementType*)Storage.GetAllocation();
 	}
 
-	[[nodiscard]] FORCEINLINE ConstIteratorType GetIterator(SizeType HeadOffset) const
+	[[nodiscard]] UE_FORCEINLINE_HINT ConstIteratorType GetIterator(SizeType HeadOffset) const
 	{
 		return ConstIteratorType(InPlace, GetData(), Max(), HeadOffset);
 	}
 
-	[[nodiscard]] FORCEINLINE IteratorType GetIterator(SizeType HeadOffset)
+	[[nodiscard]] UE_FORCEINLINE_HINT IteratorType GetIterator(SizeType HeadOffset)
 	{
 		return IteratorType(InPlace, GetData(), Max(), HeadOffset);
 	}
diff --git a/Engine/Source/Runtime/Core/Public/Containers/EnumAsByte.h b/Engine/Source/Runtime/Core/Public/Containers/EnumAsByte.h
--- a/Engine/Source/Runtime/Core/Public/Containers/EnumAsByte.h
+++ b/Engine/Source/Runtime/Core/Public/Containers/EnumAsByte.h
@@ -34,7 +34,7 @@ public:
 	 *
 	 * @param InValue value to construct with.
 	 */
-	[[nodiscard]] FORCEINLINE TEnumAsByte( EnumType InValue )
+	[[nodiscard]] UE_FORCEINLINE_HINT TEnumAsByte( EnumType InValue )
 		: Value(static_cast<uint8>(InValue))
 	{
 	}
@@ -44,7 +44,7 @@ public:
 	 *
 	 * @param InValue value to construct with.
 	 */
-	[[nodiscard]] explicit FORCEINLINE TEnumAsByte( int32 InValue )
+	[[nodiscard]] explicit UE_FORCEINLINE_HINT TEnumAsByte( int32 InValue )
 		: Value(static_cast<uint8>(InValue))
 	{
 	}
@@ -54,7 +54,7 @@ public:
 	 *
 	 * @param InValue value to construct with.
 	 */
-	[[nodiscard]] explicit FORCEINLINE TEnumAsByte( uint8 InValue )
+	[[nodiscard]] explicit UE_FORCEINLINE_HINT TEnumAsByte( uint8 InValue )
 		: Value(InValue)
 	{
 	}
@@ -117,7 +117,7 @@ private:
 };
 
 template<class T>
-[[nodiscard]] FORCEINLINE uint32 GetTypeHash(const TEnumAsByte<T>& Enum)
+[[nodiscard]] UE_FORCEINLINE_HINT uint32 GetTypeHash(const TEnumAsByte<T>& Enum)
 {
 	return GetTypeHash((uint8)Enum.GetValue());
 }
diff --git a/Engine/Source/Runtime/Core/Public/Containers/HashTable.h b/Engine/Source/Runtime/Core/Public/Containers/HashTable.h
--- a/Engine/Source/Runtime/Core/Public/Containers/HashTable.h
+++ b/Engine/Source/Runtime/Core/Public/Containers/HashTable.h
@@ -128,7 +128,7 @@ inline TStaticHashTable< HashSize, IndexSize >::TStaticHashTable(ENoInit)
 }
 
 template< uint16 HashSize, uint16 IndexSize >
-FORCEINLINE void TStaticHashTable< HashSize, IndexSize >::Clear()
+UE_FORCEINLINE_HINT void TStaticHashTable< HashSize, IndexSize >::Clear()
 {
 	FMemory::Memset( Hash, 0xff, HashSize * 2 );
 }
@@ -150,7 +150,7 @@ inline uint16 TStaticHashTable< HashSize, IndexSize >::Next( uint16 Index ) cons
 }
 
 template< uint16 HashSize, uint16 IndexSize >
-FORCEINLINE bool TStaticHashTable< HashSize, IndexSize >::IsValid( uint16 Index ) const
+UE_FORCEINLINE_HINT bool TStaticHashTable< HashSize, IndexSize >::IsValid( uint16 Index ) const
 {
 	return Index != 0xffff;
 }
@@ -342,7 +342,7 @@ inline FHashTable& FHashTable::operator=(FHashTable&& Other)
 	return *this;
 }
 
-FORCEINLINE FHashTable::~FHashTable()
+UE_FORCEINLINE_HINT FHashTable::~FHashTable()
 {
 	Free();
 }
@@ -406,7 +406,7 @@ inline uint32 FHashTable::Next( uint32 Index ) const
 	return NextIndex[ Index ];
 }
 
-FORCEINLINE bool FHashTable::IsValid( uint32 Index ) const
+UE_FORCEINLINE_HINT bool FHashTable::IsValid( uint32 Index ) const
 {
 	return Index != ~0u;
 }
@@ -497,10 +497,10 @@ public:
 	void			Remove(uint16 Key, uint32 Index);
 
 private:
-	FORCEINLINE uint32 HashAt(uint32 Index) const { return ((uint32*)Hash.GetAllocation())[Index]; }
-	FORCEINLINE uint32 NextIndexAt(uint32 Index) const { return ((uint32*)NextIndex.GetAllocation())[Index]; }
-	FORCEINLINE uint32& HashAt(uint32 Index) { return ((uint32*)Hash.GetAllocation())[Index]; }
-	FORCEINLINE uint32& NextIndexAt(uint32 Index) { return ((uint32*)NextIndex.GetAllocation())[Index]; }
+	UE_FORCEINLINE_HINT uint32 HashAt(uint32 Index) const { return ((uint32*)Hash.GetAllocation())[Index]; }
+	UE_FORCEINLINE_HINT uint32 NextIndexAt(uint32 Index) const { return ((uint32*)NextIndex.GetAllocation())[Index]; }
+	UE_FORCEINLINE_HINT uint32& HashAt(uint32 Index) { return ((uint32*)Hash.GetAllocation())[Index]; }
+	UE_FORCEINLINE_HINT uint32& NextIndexAt(uint32 Index) { return ((uint32*)NextIndex.GetAllocation())[Index]; }
 
 	ElementAllocatorType	Hash;
 	ElementAllocatorType	NextIndex;
@@ -556,7 +556,7 @@ inline THashTable<InAllocator>::THashTable(uint32 InHashSize, uint32 InIndexSize
 }
 
 template<typename InAllocator>
-FORCEINLINE THashTable<InAllocator>::~THashTable()
+UE_FORCEINLINE_HINT THashTable<InAllocator>::~THashTable()
 {
 }
 
@@ -600,7 +600,7 @@ inline uint32 THashTable<InAllocator>::Next(uint32 Index) const
 }
 
 template<typename InAllocator>
-FORCEINLINE bool THashTable<InAllocator>::IsValid(uint32 Index) const
+UE_FORCEINLINE_HINT bool THashTable<InAllocator>::IsValid(uint32 Index) const
 {
 	return Index != ~0u;
 }
diff --git a/Engine/Source/Runtime/Core/Public/Containers/IndirectArray.h b/Engine/Source/Runtime/Core/Public/Containers/IndirectArray.h
--- a/Engine/Source/Runtime/Core/Public/Containers/IndirectArray.h
+++ b/Engine/Source/Runtime/Core/Public/Containers/IndirectArray.h
@@ -91,7 +91,7 @@ public:
 	 *
 	 * @returns Number of elements in array.
 	 */
-	[[nodiscard]] FORCEINLINE int32 Num() const
+	[[nodiscard]] UE_FORCEINLINE_HINT int32 Num() const
 	{
 		return Array.Num();
 	}
@@ -101,7 +101,7 @@ public:
 	 *
 	 * @returns Pointer to first array entry or nullptr if this->ArrayMax == 0.
 	 */
-	[[nodiscard]] FORCEINLINE T** GetData()
+	[[nodiscard]] UE_FORCEINLINE_HINT T** GetData()
 	{
 		return (T**)Array.GetData();
 	}
@@ -111,7 +111,7 @@ public:
 	 *
 	 * @returns Pointer to first array entry or nullptr if this->ArrayMax == 0.
 	 */
-	[[nodiscard]] FORCEINLINE const T** GetData() const
+	[[nodiscard]] UE_FORCEINLINE_HINT const T** GetData() const
 	{
 		return (const T**)Array.GetData();
 	}
@@ -132,7 +132,7 @@ public:
 	 * @param Index Position of element to return.
 	 * @returns Reference to element in array at given position.
 	 */
-	[[nodiscard]] FORCEINLINE T& operator[](int32 Index)
+	[[nodiscard]] UE_FORCEINLINE_HINT T& operator[](int32 Index)
 	{
 		return *(T*)Array[Index];
 	}
@@ -145,7 +145,7 @@ public:
 	 * @param Index Position of element to return.
 	 * @returns Reference to element in array at given position.
 	 */
-	[[nodiscard]] FORCEINLINE const T& operator[](int32 Index) const
+	[[nodiscard]] UE_FORCEINLINE_HINT const T& operator[](int32 Index) const
 	{
 		return *(T*)Array[Index];
 	}
@@ -156,7 +156,7 @@ public:
 	 * @param IndexFromTheEnd (Optional) Index from the end of array (default = 0).
 	 * @returns Reference to n-th last element from the array.
 	 */
-	[[nodiscard]] FORCEINLINE ElementType& Last(int32 IndexFromTheEnd = 0)
+	[[nodiscard]] UE_FORCEINLINE_HINT ElementType& Last(int32 IndexFromTheEnd = 0)
 	{
 		return *(T*)Array.Last(IndexFromTheEnd);
 	}
@@ -169,7 +169,7 @@ public:
 	 * @param IndexFromTheEnd (Optional) Index from the end of array (default = 0).
 	 * @returns Reference to n-th last element from the array.
 	 */
-	[[nodiscard]] FORCEINLINE const ElementType& Last(int32 IndexFromTheEnd = 0) const
+	[[nodiscard]] UE_FORCEINLINE_HINT const ElementType& Last(int32 IndexFromTheEnd = 0) const
 	{
 		return *(T*)Array.Last(IndexFromTheEnd);
 	}
@@ -283,7 +283,7 @@ public:
 		Array.RemoveAt(Index, Count, AllowShrinking);
 	}
 	UE_ALLOWSHRINKING_BOOL_DEPRECATED("RemoveAt")
-	FORCEINLINE void RemoveAt(int32 Index, int32 Count, bool bAllowShrinking)
+	UE_FORCEINLINE_HINT void RemoveAt(int32 Index, int32 Count, bool bAllowShrinking)
 	{
 		RemoveAt(Index, Count, bAllowShrinking ? EAllowShrinking::Yes : EAllowShrinking::No);
 	}
@@ -334,7 +334,7 @@ public:
 		Array.RemoveAtSwap(Index, Count, AllowShrinking);
 	}
 	UE_ALLOWSHRINKING_BOOL_DEPRECATED("RemoveAtSwap")
-	FORCEINLINE void RemoveAtSwap(int32 Index, int32 Count, bool bAllowShrinking)
+	UE_FORCEINLINE_HINT void RemoveAtSwap(int32 Index, int32 Count, bool bAllowShrinking)
 	{
 		RemoveAtSwap(Index, Count, bAllowShrinking ? EAllowShrinking::Yes : EAllowShrinking::No);
 	}
@@ -371,7 +371,7 @@ public:
 	 * @param Item The item to add.
 	 * @returns Index to the new item.
 	 */
-	FORCEINLINE int32 Add(T* Item)
+	UE_FORCEINLINE_HINT int32 Add(T* Item)
 	{
 		return Array.Add(Item);
 	}
@@ -382,7 +382,7 @@ public:
 	 * @param Item The element to insert.
 	 * @param Index Tells where to insert the new elements.
 	 */
-	FORCEINLINE void Insert(T* Item, int32 Index)
+	UE_FORCEINLINE_HINT void Insert(T* Item, int32 Index)
 	{
 		Array.Insert(Item, Index);
 	}
@@ -393,7 +393,7 @@ public:
 	 * @param Number The number of elements that the array should be able to
 	 *               contain after allocation.
 	 */
-	FORCEINLINE void Reserve(int32 Number)
+	UE_FORCEINLINE_HINT void Reserve(int32 Number)
 	{
 		Array.Reserve(Number);
 	}
@@ -405,7 +405,7 @@ public:
 	 * @param Index Index to test.
 	 * @returns True if index is valid. False otherwise.
 	 */
-	[[nodiscard]] FORCEINLINE bool IsValidIndex(int32 Index) const
+	[[nodiscard]] UE_FORCEINLINE_HINT bool IsValidIndex(int32 Index) const
 	{
 		return Array.IsValidIndex(Index);
 	}
@@ -464,10 +464,10 @@ public:
 	 * DO NOT USE DIRECTLY
 	 * STL-like iterators to enable range-based for loop support.
 	 */
-	[[nodiscard]] FORCEINLINE TDereferencingIterator<      ElementType, typename InternalArrayType::RangedForIteratorType     > begin()       { return TDereferencingIterator<      ElementType, typename InternalArrayType::RangedForIteratorType     >(Array.begin()); }
-	[[nodiscard]] FORCEINLINE TDereferencingIterator<const ElementType, typename InternalArrayType::RangedForConstIteratorType> begin() const { return TDereferencingIterator<const ElementType, typename InternalArrayType::RangedForConstIteratorType>(Array.begin()); }
-	[[nodiscard]] FORCEINLINE TDereferencingIterator<      ElementType, typename InternalArrayType::RangedForIteratorType     > end  ()       { return TDereferencingIterator<      ElementType, typename InternalArrayType::RangedForIteratorType     >(Array.end()); }
-	[[nodiscard]] FORCEINLINE TDereferencingIterator<const ElementType, typename InternalArrayType::RangedForConstIteratorType> end  () const { return TDereferencingIterator<const ElementType, typename InternalArrayType::RangedForConstIteratorType>(Array.end()); }
+	[[nodiscard]] UE_FORCEINLINE_HINT TDereferencingIterator<      ElementType, typename InternalArrayType::RangedForIteratorType     > begin()       { return TDereferencingIterator<      ElementType, typename InternalArrayType::RangedForIteratorType     >(Array.begin()); }
+	[[nodiscard]] UE_FORCEINLINE_HINT TDereferencingIterator<const ElementType, typename InternalArrayType::RangedForConstIteratorType> begin() const { return TDereferencingIterator<const ElementType, typename InternalArrayType::RangedForConstIteratorType>(Array.begin()); }
+	[[nodiscard]] UE_FORCEINLINE_HINT TDereferencingIterator<      ElementType, typename InternalArrayType::RangedForIteratorType     > end  ()       { return TDereferencingIterator<      ElementType, typename InternalArrayType::RangedForIteratorType     >(Array.end()); }
+	[[nodiscard]] UE_FORCEINLINE_HINT TDereferencingIterator<const ElementType, typename InternalArrayType::RangedForConstIteratorType> end  () const { return TDereferencingIterator<const ElementType, typename InternalArrayType::RangedForConstIteratorType>(Array.end()); }
 
 private:
 
diff --git a/Engine/Source/Runtime/Core/Public/Containers/IntrusiveDoubleLinkedList.h b/Engine/Source/Runtime/Core/Public/Containers/IntrusiveDoubleLinkedList.h
--- a/Engine/Source/Runtime/Core/Public/Containers/IntrusiveDoubleLinkedList.h
+++ b/Engine/Source/Runtime/Core/Public/Containers/IntrusiveDoubleLinkedList.h
@@ -32,20 +32,20 @@ public:
 		, Prev(GetThisElement())
 	{}
 
-	FORCEINLINE void Reset()
+	UE_FORCEINLINE_HINT void Reset()
 	{
 		Next = Prev = GetThisElement();
 	}
 
-	[[nodiscard]] FORCEINLINE bool IsInList() const
+	[[nodiscard]] UE_FORCEINLINE_HINT bool IsInList() const
 	{
 		return Next != GetThisElement();
 	}
-	[[nodiscard]] FORCEINLINE ElementType* GetNext() const
+	[[nodiscard]] UE_FORCEINLINE_HINT ElementType* GetNext() const
 	{
 		return Next;
 	}
-	[[nodiscard]] FORCEINLINE ElementType* GetPrev() const
+	[[nodiscard]] UE_FORCEINLINE_HINT ElementType* GetPrev() const
 	{
 		return Prev;
 	}
@@ -88,8 +88,8 @@ protected:
 	friend class TIntrusiveDoubleLinkedListIterator<TIntrusiveDoubleLinkedListNode>;
 	friend class TIntrusiveDoubleLinkedList<ElementType, ContainerType>;
 
-	[[nodiscard]] FORCEINLINE ElementType*       GetThisElement()       { return static_cast<ElementType*>(this); }
-	[[nodiscard]] FORCEINLINE const ElementType* GetThisElement() const { return static_cast<const ElementType*>(this); }
+	[[nodiscard]] UE_FORCEINLINE_HINT ElementType*       GetThisElement()       { return static_cast<ElementType*>(this); }
+	[[nodiscard]] UE_FORCEINLINE_HINT const ElementType* GetThisElement() const { return static_cast<const ElementType*>(this); }
 
 	ElementType* Next;
 	ElementType* Prev;
@@ -105,7 +105,7 @@ class TIntrusiveDoubleLinkedListIterator
 public:
 	using ElementType = typename NodeType::ElementType;
 
-	[[nodiscard]] FORCEINLINE explicit TIntrusiveDoubleLinkedListIterator(ElementType* Node)
+	[[nodiscard]] UE_FORCEINLINE_HINT explicit TIntrusiveDoubleLinkedListIterator(ElementType* Node)
 		: CurrentNode(Node)
 	{
 	}
@@ -157,8 +157,8 @@ public:
 		return CurrentNode;
 	}
 
-	[[nodiscard]] FORCEINLINE bool operator==(const TIntrusiveDoubleLinkedListIterator& Other) const { return CurrentNode == Other.CurrentNode; }
-	[[nodiscard]] FORCEINLINE bool operator!=(const TIntrusiveDoubleLinkedListIterator& Other) const { return CurrentNode != Other.CurrentNode; }
+	[[nodiscard]] UE_FORCEINLINE_HINT bool operator==(const TIntrusiveDoubleLinkedListIterator& Other) const { return CurrentNode == Other.CurrentNode; }
+	[[nodiscard]] UE_FORCEINLINE_HINT bool operator!=(const TIntrusiveDoubleLinkedListIterator& Other) const { return CurrentNode != Other.CurrentNode; }
 
 private:
 	ElementType* CurrentNode;
@@ -178,37 +178,37 @@ public:
 	using ElementType = InElementType;
 	using NodeType = TIntrusiveDoubleLinkedListNode<ElementType, ContainerType>;
 
-	[[nodiscard]] FORCEINLINE TIntrusiveDoubleLinkedList()
+	[[nodiscard]] UE_FORCEINLINE_HINT TIntrusiveDoubleLinkedList()
 	{
 	}
 
-	FORCEINLINE TIntrusiveDoubleLinkedList(const TIntrusiveDoubleLinkedList&) = delete;
-	FORCEINLINE TIntrusiveDoubleLinkedList& operator=(const TIntrusiveDoubleLinkedList&) = delete;
+	UE_FORCEINLINE_HINT TIntrusiveDoubleLinkedList(const TIntrusiveDoubleLinkedList&) = delete;
+	UE_FORCEINLINE_HINT TIntrusiveDoubleLinkedList& operator=(const TIntrusiveDoubleLinkedList&) = delete;
 
 	/**
 	 * Fast empty that clears this list *without* changing the links in any elements.
 	 * 
 	 * @see IsEmpty, IsFilled
 	 */
-	FORCEINLINE void Reset()
+	UE_FORCEINLINE_HINT void Reset()
 	{
 		Sentinel.Reset();
 	}
 
 	// Accessors.
 
-	[[nodiscard]] FORCEINLINE bool IsEmpty() const
+	[[nodiscard]] UE_FORCEINLINE_HINT bool IsEmpty() const
 	{
 		return Sentinel.Next == GetSentinel();
 	}
-	[[nodiscard]] FORCEINLINE bool IsFilled() const
+	[[nodiscard]] UE_FORCEINLINE_HINT bool IsFilled() const
 	{
 		return Sentinel.Next != GetSentinel();
 	}
 
 	// Adding/Removing methods
 
-	FORCEINLINE void AddHead(ElementType* Element)
+	UE_FORCEINLINE_HINT void AddHead(ElementType* Element)
 	{
 		static_cast<NodeType*>(Element)->InsertAfter(GetSentinel());
 	}
@@ -225,7 +225,7 @@ public:
 		}
 	}
 
-	FORCEINLINE void AddTail(ElementType* Element)
+	UE_FORCEINLINE_HINT void AddTail(ElementType* Element)
 	{
 		static_cast<NodeType*>(Element)->InsertBefore(GetSentinel());
 	}
@@ -242,12 +242,12 @@ public:
 		}
 	}
 
-	[[nodiscard]] FORCEINLINE ElementType* GetHead()
+	[[nodiscard]] UE_FORCEINLINE_HINT ElementType* GetHead()
 	{
 		return IsFilled() ? Sentinel.Next : nullptr;
 	}
 
-	[[nodiscard]] FORCEINLINE ElementType* GetTail()
+	[[nodiscard]] UE_FORCEINLINE_HINT ElementType* GetTail()
 	{
 		return IsFilled() ? Sentinel.Prev : nullptr;
 	}
@@ -276,17 +276,17 @@ public:
 		return Tail;
 	}
 
-	FORCEINLINE static void Remove(ElementType* Element)
+	UE_FORCEINLINE_HINT static void Remove(ElementType* Element)
 	{
 		static_cast<NodeType*>(Element)->Remove();
 	}
 
-	FORCEINLINE static void InsertAfter(ElementType* InsertThis, ElementType* AfterThis)
+	UE_FORCEINLINE_HINT static void InsertAfter(ElementType* InsertThis, ElementType* AfterThis)
 	{
 		static_cast<NodeType*>(InsertThis)->InsertAfter(AfterThis);
 	}
 
-	FORCEINLINE static void InsertBefore(ElementType* InsertThis, ElementType* BeforeThis)
+	UE_FORCEINLINE_HINT static void InsertBefore(ElementType* InsertThis, ElementType* BeforeThis)
 	{
 		static_cast<NodeType*>(InsertThis)->InsertBefore(BeforeThis);
 	}
@@ -294,15 +294,15 @@ public:
 	using TIterator = TIntrusiveDoubleLinkedListIterator<NodeType>;
 	using TConstIterator = TIntrusiveDoubleLinkedListIterator<const NodeType>;
 
-	[[nodiscard]] FORCEINLINE TIterator      begin()       { return TIterator(Sentinel.Next); }
-	[[nodiscard]] FORCEINLINE TConstIterator begin() const { return TConstIterator(Sentinel.Next); }
-	[[nodiscard]] FORCEINLINE TIterator      end()         { return TIterator(GetSentinel()); }
-	[[nodiscard]] FORCEINLINE TConstIterator end() const   { return TConstIterator(GetSentinel()); }
+	[[nodiscard]] UE_FORCEINLINE_HINT TIterator      begin()       { return TIterator(Sentinel.Next); }
+	[[nodiscard]] UE_FORCEINLINE_HINT TConstIterator begin() const { return TConstIterator(Sentinel.Next); }
+	[[nodiscard]] UE_FORCEINLINE_HINT TIterator      end()         { return TIterator(GetSentinel()); }
+	[[nodiscard]] UE_FORCEINLINE_HINT TConstIterator end() const   { return TConstIterator(GetSentinel()); }
 
 private:
 
-	[[nodiscard]] FORCEINLINE ElementType*       GetSentinel()       { return static_cast<ElementType*>(&Sentinel); }       //-V717
-	[[nodiscard]] FORCEINLINE const ElementType* GetSentinel() const { return static_cast<const ElementType*>(&Sentinel); } //-V717
+	[[nodiscard]] UE_FORCEINLINE_HINT ElementType*       GetSentinel()       { return static_cast<ElementType*>(&Sentinel); }       //-V717
+	[[nodiscard]] UE_FORCEINLINE_HINT const ElementType* GetSentinel() const { return static_cast<const ElementType*>(&Sentinel); } //-V717
 
 	NodeType Sentinel;
 
diff --git a/Engine/Source/Runtime/Core/Public/Containers/LinkedListBuilder.h b/Engine/Source/Runtime/Core/Public/Containers/LinkedListBuilder.h
--- a/Engine/Source/Runtime/Core/Public/Containers/LinkedListBuilder.h
+++ b/Engine/Source/Runtime/Core/Public/Containers/LinkedListBuilder.h
@@ -10,7 +10,7 @@ struct TLinkedListBuilderNextLink
 {
 	using ElementType = InElementType;
 
-	[[nodiscard]] FORCEINLINE static ElementType** GetNextPtr(ElementType& Element)
+	[[nodiscard]] UE_FORCEINLINE_HINT static ElementType** GetNextPtr(ElementType& Element)
 	{
 		return &(Element.Next);
 	}
@@ -21,7 +21,7 @@ struct TLinkedListBuilderNextLinkMemberVar
 {
 	using ElementType = InElementType;
 
-	[[nodiscard]] FORCEINLINE static ElementType** GetNextPtr(ElementType& Element)
+	[[nodiscard]] UE_FORCEINLINE_HINT static ElementType** GetNextPtr(ElementType& Element)
 	{
 		return &(Element.*NextLink);
 	}
@@ -60,13 +60,13 @@ public:
 
 	// Move builder back to start and prepare for overwriting
 	// It only changes state of builder, use NullTerminate() to mark list as empty!
-	FORCEINLINE void Restart()
+	UE_FORCEINLINE_HINT void Restart()
 	{
 		EndPtr = StartPtr;
 	}
 
 	UE_DEPRECATED(5.6, "Append is deprecated. Please use AppendTerminated instead.")
-	FORCEINLINE void Append(ElementType& Element)
+	UE_FORCEINLINE_HINT void Append(ElementType& Element)
 	{
 		AppendTerminated(Element);
 	}
@@ -103,7 +103,7 @@ public:
 	}
 
 	// Mark end of the list
-	FORCEINLINE void NullTerminate()
+	UE_FORCEINLINE_HINT void NullTerminate()
 	{
 		WriteEndPtr(nullptr);
 	}
@@ -127,23 +127,23 @@ public:
 		return false;
 	}
 
-	[[nodiscard]] FORCEINLINE ElementType* GetNext(ElementType& Element) const
+	[[nodiscard]] UE_FORCEINLINE_HINT ElementType* GetNext(ElementType& Element) const
 	{
 		return GetNextRef(Element);
 	}
 
-	[[nodiscard]] FORCEINLINE ElementType* GetListStart() const
+	[[nodiscard]] UE_FORCEINLINE_HINT ElementType* GetListStart() const
 	{
 		return *StartPtr;
 	}
 
-	[[nodiscard]] FORCEINLINE ElementType* GetListEnd() const
+	[[nodiscard]] UE_FORCEINLINE_HINT ElementType* GetListEnd() const
 	{
 		return *EndPtr;
 	}
 
 private:
-	[[nodiscard]] FORCEINLINE ElementType*& GetNextRef(ElementType& Element) const
+	[[nodiscard]] UE_FORCEINLINE_HINT ElementType*& GetNextRef(ElementType& Element) const
 	{
 		return *LinkAccessor::GetNextPtr(Element);
 	}
diff --git a/Engine/Source/Runtime/Core/Public/Containers/List.h b/Engine/Source/Runtime/Core/Public/Containers/List.h
--- a/Engine/Source/Runtime/Core/Public/Containers/List.h
+++ b/Engine/Source/Runtime/Core/Public/Containers/List.h
@@ -37,13 +37,13 @@ public:
 	}
 
 	/** conversion to "bool" returning true if the iterator is valid. */
-	[[nodiscard]] FORCEINLINE explicit operator bool() const
+	[[nodiscard]] UE_FORCEINLINE_HINT explicit operator bool() const
 	{ 
 		return CurrentLink != nullptr;
 	}
 
-	[[nodiscard]] FORCEINLINE bool operator==(const TLinkedListIteratorBase& Rhs) const { return CurrentLink == Rhs.CurrentLink; }
-	[[nodiscard]] FORCEINLINE bool operator!=(const TLinkedListIteratorBase& Rhs) const { return CurrentLink != Rhs.CurrentLink; }
+	[[nodiscard]] UE_FORCEINLINE_HINT bool operator==(const TLinkedListIteratorBase& Rhs) const { return CurrentLink == Rhs.CurrentLink; }
+	[[nodiscard]] UE_FORCEINLINE_HINT bool operator!=(const TLinkedListIteratorBase& Rhs) const { return CurrentLink != Rhs.CurrentLink; }
 protected:
 
 	ContainerType* CurrentLink;
@@ -240,22 +240,22 @@ public:
 	 *
 	 * @return true if currently linked, false otherwise
 	 */
-	[[nodiscard]] FORCEINLINE bool IsLinked()
+	[[nodiscard]] UE_FORCEINLINE_HINT bool IsLinked()
 	{
 		return PrevLink != nullptr;
 	}
 
-	[[nodiscard]] FORCEINLINE ContainerType** GetPrevLink() const
+	[[nodiscard]] UE_FORCEINLINE_HINT ContainerType** GetPrevLink() const
 	{
 		return PrevLink;
 	}
 
-	[[nodiscard]] FORCEINLINE ContainerType* GetNextLink() const
+	[[nodiscard]] UE_FORCEINLINE_HINT ContainerType* GetNextLink() const
 	{
 		return NextLink;
 	}
 
-	[[nodiscard]] FORCEINLINE ContainerType* Next()
+	[[nodiscard]] UE_FORCEINLINE_HINT ContainerType* Next()
 	{
 		return NextLink;
 	}
@@ -268,10 +268,10 @@ private:
 	ContainerType** PrevLink;
 
 
-	[[nodiscard]] FORCEINLINE friend TIterator      begin(      ContainerType& List) { return TIterator     (&List); }
-	[[nodiscard]] FORCEINLINE friend TConstIterator begin(const ContainerType& List) { return TConstIterator(const_cast<ContainerType*>(&List)); }
-	[[nodiscard]] FORCEINLINE friend TIterator      end  (      ContainerType& List) { return TIterator     (nullptr); }
-	[[nodiscard]] FORCEINLINE friend TConstIterator end  (const ContainerType& List) { return TConstIterator(nullptr); }
+	[[nodiscard]] UE_FORCEINLINE_HINT friend TIterator      begin(      ContainerType& List) { return TIterator     (&List); }
+	[[nodiscard]] UE_FORCEINLINE_HINT friend TConstIterator begin(const ContainerType& List) { return TConstIterator(const_cast<ContainerType*>(&List)); }
+	[[nodiscard]] UE_FORCEINLINE_HINT friend TIterator      end  (      ContainerType& List) { return TIterator     (nullptr); }
+	[[nodiscard]] UE_FORCEINLINE_HINT friend TConstIterator end  (const ContainerType& List) { return TConstIterator(nullptr); }
 };
 
 /**
@@ -304,19 +304,19 @@ public:
 
 
 	// Accessors.
-	[[nodiscard]] FORCEINLINE ElementType* operator->()
+	[[nodiscard]] UE_FORCEINLINE_HINT ElementType* operator->()
 	{
 		return &Element;
 	}
-	[[nodiscard]] FORCEINLINE const ElementType* operator->() const
+	[[nodiscard]] UE_FORCEINLINE_HINT const ElementType* operator->() const
 	{
 		return &Element;
 	}
-	[[nodiscard]] FORCEINLINE ElementType& operator*()
+	[[nodiscard]] UE_FORCEINLINE_HINT ElementType& operator*()
 	{
 		return Element;
 	}
-	[[nodiscard]] FORCEINLINE const ElementType& operator*() const
+	[[nodiscard]] UE_FORCEINLINE_HINT const ElementType& operator*() const
 	{
 		return Element;
 	}
@@ -368,7 +368,7 @@ public:
 	}
 
 	/** conversion to "bool" returning true if the iterator is valid. */
-	[[nodiscard]] FORCEINLINE explicit operator bool() const
+	[[nodiscard]] UE_FORCEINLINE_HINT explicit operator bool() const
 	{ 
 		return CurrentNode != nullptr; 
 	}
