Replace some FORCEINLINE usage in Core/CoreUObject with UE_FORCEINLINE_HINT.

diff --git a/Engine/Source/Runtime/Core/Public/Containers/LockFreeList.h b/Engine/Source/Runtime/Core/Public/Containers/LockFreeList.h
--- a/Engine/Source/Runtime/Core/Public/Containers/LockFreeList.h
+++ b/Engine/Source/Runtime/Core/Public/Containers/LockFreeList.h
@@ -42,7 +42,7 @@ CORE_API DECLARE_LOG_CATEGORY_EXTERN(LogLockFreeList, Log, All);
 		}
 	}
 #else
-	FORCEINLINE void TestCriticalStall()
+	UE_FORCEINLINE_HINT void TestCriticalStall()
 	{
 	}
 #endif
@@ -149,22 +149,22 @@ struct FIndexedPointer
 		Ptrs.store(uint64(Ptr) | (CounterAndState << MAX_LOCK_FREE_LINKS_AS_BITS), std::memory_order_relaxed);
 	}
 
-	[[nodiscard]] FORCEINLINE uint32 GetPtr() const
+	[[nodiscard]] UE_FORCEINLINE_HINT uint32 GetPtr() const
 	{
 		return uint32(Ptrs.load(std::memory_order_relaxed) & (MAX_LOCK_FREE_LINKS - 1));
 	}
 
-	FORCEINLINE void SetPtr(uint32 To)
+	UE_FORCEINLINE_HINT void SetPtr(uint32 To)
 	{
 		SetAll(To, GetCounterAndState());
 	}
 
-	[[nodiscard]] FORCEINLINE uint64 GetCounterAndState() const
+	[[nodiscard]] UE_FORCEINLINE_HINT uint64 GetCounterAndState() const
 	{
 		return (Ptrs.load(std::memory_order_relaxed) >> MAX_LOCK_FREE_LINKS_AS_BITS);
 	}
 
-	FORCEINLINE void SetCounterAndState(uint64 To)
+	UE_FORCEINLINE_HINT void SetCounterAndState(uint64 To)
 	{
 		SetAll(GetPtr(), To);
 	}
@@ -180,7 +180,7 @@ struct FIndexedPointer
 	}
 
 	template<uint64 TABAInc>
-	[[nodiscard]] FORCEINLINE uint64 GetState() const
+	[[nodiscard]] UE_FORCEINLINE_HINT uint64 GetState() const
 	{
 		return GetCounterAndState() & (TABAInc - 1);
 	}
@@ -206,11 +206,11 @@ struct FIndexedPointer
 		return Ptrs.compare_exchange_strong(Expected, Exchange.Ptrs.load(std::memory_order_relaxed), std::memory_order_acq_rel, std::memory_order_relaxed);
 	}
 
-	[[nodiscard]] FORCEINLINE bool operator==(const FIndexedPointer& Other) const
+	[[nodiscard]] UE_FORCEINLINE_HINT bool operator==(const FIndexedPointer& Other) const
 	{
 		return Ptrs.load(std::memory_order_relaxed) == Other.Ptrs.load(std::memory_order_relaxed);
 	}
-	[[nodiscard]] FORCEINLINE bool operator!=(const FIndexedPointer& Other) const
+	[[nodiscard]] UE_FORCEINLINE_HINT bool operator!=(const FIndexedPointer& Other) const
 	{
 		return Ptrs.load(std::memory_order_relaxed) != Other.Ptrs.load(std::memory_order_relaxed);
 	}
@@ -239,15 +239,15 @@ struct FLockFreeLinkPolicy
 	typedef uint32 TLinkPtr;
 	typedef TLockFreeAllocOnceIndexedAllocator<FIndexedLockFreeLink, MAX_LOCK_FREE_LINKS, 16384> TAllocator;
 
-	[[nodiscard]] static FORCEINLINE FIndexedLockFreeLink* DerefLink(uint32 Ptr)
+	[[nodiscard]] static UE_FORCEINLINE_HINT FIndexedLockFreeLink* DerefLink(uint32 Ptr)
 	{
 		return LinkAllocator.GetItem(Ptr);
 	}
-	[[nodiscard]] static FORCEINLINE FIndexedLockFreeLink* IndexToLink(uint32 Index)
+	[[nodiscard]] static UE_FORCEINLINE_HINT FIndexedLockFreeLink* IndexToLink(uint32 Index)
 	{
 		return LinkAllocator.GetItem(Index);
 	}
-	[[nodiscard]] static FORCEINLINE uint32 IndexToPtr(uint32 Index)
+	[[nodiscard]] static UE_FORCEINLINE_HINT uint32 IndexToPtr(uint32 Index)
 	{
 		return Index;
 	}
@@ -394,7 +394,7 @@ public:
 		return Item;
 	}
 
-	[[nodiscard]] FORCEINLINE bool IsEmpty() const
+	[[nodiscard]] UE_FORCEINLINE_HINT bool IsEmpty() const
 	{
 		return !Head.GetPtr();
 	}
@@ -524,12 +524,12 @@ public:
 		}
 	}
 
-	[[nodiscard]] FORCEINLINE bool IsEmpty() const
+	[[nodiscard]] UE_FORCEINLINE_HINT bool IsEmpty() const
 	{
 		return RootList.IsEmpty();
 	}
 
-	[[nodiscard]] FORCEINLINE uint64 GetState() const
+	[[nodiscard]] UE_FORCEINLINE_HINT uint64 GetState() const
 	{
 		return RootList.GetState();
 	}
@@ -887,7 +887,7 @@ public:
 	*	CAUTION: This methods safety depends on external assumptions. For example, if another thread could add to the list at any time, the return value is no better than a best guess.
 	*	As typically used, the list is not being access concurrently when this is called.
 	*/
-	[[nodiscard]] FORCEINLINE bool IsEmpty() const
+	[[nodiscard]] UE_FORCEINLINE_HINT bool IsEmpty() const
 	{
 		return FLockFreePointerListLIFOBase<T, TPaddingForCacheContention>::IsEmpty();
 	}
@@ -947,7 +947,7 @@ public:
 	*	CAUTION: This methods safety depends on external assumptions. For example, if another thread could add to the list at any time, the return value is no better than a best guess.
 	*	As typically used, the list is not being access concurrently when this is called.
 	*/
-	[[nodiscard]] FORCEINLINE bool IsEmpty() const
+	[[nodiscard]] UE_FORCEINLINE_HINT bool IsEmpty() const
 	{
 		return FLockFreePointerFIFOBase<T, TPaddingForCacheContention>::IsEmpty();
 	}
diff --git a/Engine/Source/Runtime/Core/Public/Containers/LruCache.h b/Engine/Source/Runtime/Core/Public/Containers/LruCache.h
--- a/Engine/Source/Runtime/Core/Public/Containers/LruCache.h
+++ b/Engine/Source/Runtime/Core/Public/Containers/LruCache.h
@@ -15,13 +15,13 @@
 template<typename KeyType>
 struct DefaultKeyComparer
 {
-	[[nodiscard]] static FORCEINLINE bool Matches(KeyType A, KeyType B)
+	[[nodiscard]] static UE_FORCEINLINE_HINT bool Matches(KeyType A, KeyType B)
 	{
 		return (A == B);
 	}
 
 	/** Calculates a hash index for a key. */
-	[[nodiscard]] static FORCEINLINE uint32 GetKeyHash(KeyType Key)
+	[[nodiscard]] static UE_FORCEINLINE_HINT uint32 GetKeyHash(KeyType Key)
 	{
 		return GetTypeHash(Key);
 	}
@@ -112,17 +112,17 @@ class TLruCache
 	/** Lookup set key functions. */
 	struct FKeyFuncs : public BaseKeyFuncs<FCacheEntry*, KeyType>
 	{
-		[[nodiscard]] FORCEINLINE static const KeyType& GetSetKey(const FCacheEntry* Entry)
+		[[nodiscard]] UE_FORCEINLINE_HINT static const KeyType& GetSetKey(const FCacheEntry* Entry)
 		{
 			return Entry->Key;
 		}
 
-		[[nodiscard]] FORCEINLINE static bool Matches(KeyType A, KeyType B)
+		[[nodiscard]] UE_FORCEINLINE_HINT static bool Matches(KeyType A, KeyType B)
 		{
 			return KeyComp::Matches(A, B);
 		}
 
-		[[nodiscard]] FORCEINLINE static uint32 GetKeyHash(KeyType Key)
+		[[nodiscard]] UE_FORCEINLINE_HINT static uint32 GetKeyHash(KeyType Key)
 		{
 			return KeyComp::GetKeyHash(Key);
 		}
@@ -263,7 +263,7 @@ public:
 	 * @return true if the entry is in the cache, false otherwise.
 	 * @see Add, ContainsByPredicate, Empty, FilterByPredicate, Find, GetKeys, Remove
 	 */
-	[[nodiscard]] FORCEINLINE bool Contains(const KeyType& Key) const
+	[[nodiscard]] UE_FORCEINLINE_HINT bool Contains(const KeyType& Key) const
 	{
 		return LookupSet.Contains(Key);
 	}
@@ -517,7 +517,7 @@ public:
 	 * @return Maximum number of entries.
 	 * @see Empty, Num
 	 */
-	[[nodiscard]] FORCEINLINE int32 Max() const
+	[[nodiscard]] UE_FORCEINLINE_HINT int32 Max() const
 	{
 		return MaxNumElements;
 	}
@@ -539,7 +539,7 @@ public:
 	 * @return Number of entries.
 	 * @see Empty, Max
 	 */
-	[[nodiscard]] FORCEINLINE int32 Num() const
+	[[nodiscard]] UE_FORCEINLINE_HINT int32 Num() const
 	{
 		return LookupSet.Num();
 	}
@@ -634,12 +634,12 @@ public:
 	{
 	public:
 
-		[[nodiscard]] FORCEINLINE TBaseIterator()
+		[[nodiscard]] UE_FORCEINLINE_HINT TBaseIterator()
 			: CurrentEntry(nullptr)
 		{
 		}
 
-		[[nodiscard]] FORCEINLINE TBaseIterator(const TLruCache& Cache)
+		[[nodiscard]] UE_FORCEINLINE_HINT TBaseIterator(const TLruCache& Cache)
 			: CurrentEntry(Cache.MostRecent)
 		{
 		}
@@ -652,12 +652,12 @@ public:
 			return *this;
 		}
 
-		[[nodiscard]] FORCEINLINE bool operator==(const TBaseIterator& Rhs) const
+		[[nodiscard]] UE_FORCEINLINE_HINT bool operator==(const TBaseIterator& Rhs) const
 		{
 			return CurrentEntry == Rhs.CurrentEntry;
 		}
 
-		[[nodiscard]] FORCEINLINE bool operator!=(const TBaseIterator& Rhs) const
+		[[nodiscard]] UE_FORCEINLINE_HINT bool operator!=(const TBaseIterator& Rhs) const
 		{
 			return CurrentEntry != Rhs.CurrentEntry;
 		}
@@ -674,12 +674,12 @@ public:
 			return CurrentEntry->Value;
 		}
 
-		[[nodiscard]] FORCEINLINE explicit operator bool() const
+		[[nodiscard]] UE_FORCEINLINE_HINT explicit operator bool() const
 		{
 			return (CurrentEntry != nullptr);
 		}
 
-		[[nodiscard]] FORCEINLINE bool operator!() const
+		[[nodiscard]] UE_FORCEINLINE_HINT bool operator!() const
 		{
 			return !(bool)*this;
 		}
@@ -725,12 +725,12 @@ public:
 	{
 	public:
 
-		[[nodiscard]] FORCEINLINE TConstIterator()
+		[[nodiscard]] UE_FORCEINLINE_HINT TConstIterator()
 			: TBaseIterator<true>()
 		{
 		}
 
-		[[nodiscard]] FORCEINLINE TConstIterator(const TLruCache& Cache)
+		[[nodiscard]] UE_FORCEINLINE_HINT TConstIterator(const TLruCache& Cache)
 			: TBaseIterator<true>(Cache)
 		{
 		}
diff --git a/Engine/Source/Runtime/Core/Public/Containers/Map.h b/Engine/Source/Runtime/Core/Public/Containers/Map.h
--- a/Engine/Source/Runtime/Core/Public/Containers/Map.h
+++ b/Engine/Source/Runtime/Core/Public/Containers/Map.h
@@ -59,7 +59,7 @@ public:
 	std::conditional_t<std::is_rvalue_reference_v<KeyInitType>, KeyInitType&, KeyInitType> Key;
 
 	/** Initialization constructor. */
-	[[nodiscard]] FORCEINLINE explicit TKeyInitializer(KeyInitType InKey)
+	[[nodiscard]] UE_FORCEINLINE_HINT explicit TKeyInitializer(KeyInitType InKey)
 		: Key(InKey)
 	{
 	}
@@ -78,29 +78,29 @@ struct TDefaultMapKeyFuncs : BaseKeyFuncs<TPair<KeyType, ValueType>, KeyType, bI
 	typedef typename TTypeTraits<KeyType>::ConstPointerType KeyInitType;
 	typedef const TPairInitializer<typename TTypeTraits<KeyType>::ConstInitType, typename TTypeTraits<ValueType>::ConstInitType>& ElementInitType;
 
-	[[nodiscard]] static FORCEINLINE KeyInitType GetSetKey(ElementInitType Element)
+	[[nodiscard]] static UE_FORCEINLINE_HINT KeyInitType GetSetKey(ElementInitType Element)
 	{
 		return Element.Key;
 	}
 
-	[[nodiscard]] static FORCEINLINE bool Matches(KeyInitType A, KeyInitType B)
+	[[nodiscard]] static UE_FORCEINLINE_HINT bool Matches(KeyInitType A, KeyInitType B)
 	{
 		return A == B;
 	}
 
 	template<typename ComparableKey>
-	[[nodiscard]] static FORCEINLINE bool Matches(KeyInitType A, ComparableKey B)
+	[[nodiscard]] static UE_FORCEINLINE_HINT bool Matches(KeyInitType A, ComparableKey B)
 	{
 		return A == B;
 	}
 
-	[[nodiscard]] static FORCEINLINE uint32 GetKeyHash(KeyInitType Key)
+	[[nodiscard]] static UE_FORCEINLINE_HINT uint32 GetKeyHash(KeyInitType Key)
 	{
 		return GetTypeHash(Key);
 	}
 
 	template<typename ComparableKey>
-	[[nodiscard]] static FORCEINLINE uint32 GetKeyHash(ComparableKey Key)
+	[[nodiscard]] static UE_FORCEINLINE_HINT uint32 GetKeyHash(ComparableKey Key)
 	{
 		return GetTypeHash(Key);
 	}
@@ -244,37 +244,37 @@ public:
 	 *
 	 * @param ExpectedNumElements The number of elements about to be added to the set.
 	 */
-	FORCEINLINE void Empty(int32 ExpectedNumElements = 0)
+	UE_FORCEINLINE_HINT void Empty(int32 ExpectedNumElements = 0)
 	{
 		Pairs.Empty(ExpectedNumElements);
 	}
 
 	/** Efficiently empties out the map but preserves all allocations and capacities */
-	FORCEINLINE void Reset()
+	UE_FORCEINLINE_HINT void Reset()
 	{
 		Pairs.Reset();
 	}
 
 	/** Shrinks the pair set to avoid slack. */
-	FORCEINLINE void Shrink()
+	UE_FORCEINLINE_HINT void Shrink()
 	{
 		Pairs.Shrink();
 	}
 
 	/** Compacts the pair set to remove holes */
-	FORCEINLINE void Compact()
+	UE_FORCEINLINE_HINT void Compact()
 	{
 		Pairs.Compact();
 	}
 
 	/** Compacts the pair set to remove holes. Does not change the iteration order of the elements. */
-	FORCEINLINE void CompactStable()
+	UE_FORCEINLINE_HINT void CompactStable()
 	{
 		Pairs.CompactStable();
 	}
 
 	/** Preallocates enough memory to contain Number elements */
-	FORCEINLINE void Reserve(int32 Number)
+	UE_FORCEINLINE_HINT void Reserve(int32 Number)
 	{
 		Pairs.Reserve(Number);
 	}
@@ -291,19 +291,19 @@ public:
 	}
 
 	/** @return The number of elements in the map. */
-	[[nodiscard]] FORCEINLINE int32 Num() const
+	[[nodiscard]] UE_FORCEINLINE_HINT int32 Num() const
 	{
 		return Pairs.Num();
 	}
 
 	/** @return The number of elements the map can hold before reallocation. */
-	[[nodiscard]] FORCEINLINE int32 Max() const
+	[[nodiscard]] UE_FORCEINLINE_HINT int32 Max() const
 	{
 		return Pairs.Max();
 	}
 
 	/** @return The non-inclusive maximum index of elements in the map. */
-	[[nodiscard]] FORCEINLINE int32 GetMaxIndex() const
+	[[nodiscard]] UE_FORCEINLINE_HINT int32 GetMaxIndex() const
 	{
 		return Pairs.GetMaxIndex();
 	}
@@ -313,19 +313,19 @@ public:
 	 * @param Id - The element id to check.
 	 * @return true if the element identifier refers to a valid element in this map.
 	 */
-	[[nodiscard]] FORCEINLINE bool IsValidId(FSetElementId Id) const
+	[[nodiscard]] UE_FORCEINLINE_HINT bool IsValidId(FSetElementId Id) const
 	{
 		return Pairs.IsValidId(Id);
 	}
 
 	/** Return a mapped pair by internal identifier. Element must be valid (see @IsValidId). */
-	[[nodiscard]] FORCEINLINE ElementType& Get(FSetElementId Id)
+	[[nodiscard]] UE_FORCEINLINE_HINT ElementType& Get(FSetElementId Id)
 	{
 		return Pairs[Id];
 	}
 
 	/** Return a mapped pair by internal identifier.  Element must be valid (see @IsValidId).*/
-	[[nodiscard]] FORCEINLINE const ElementType& Get(FSetElementId Id) const
+	[[nodiscard]] UE_FORCEINLINE_HINT const ElementType& Get(FSetElementId Id) const
 	{
 		return Pairs[Id];
 	}
@@ -394,7 +394,7 @@ public:
 	 * @return Number of bytes allocated by this container.
 	 * @see CountBytes
 	 */
-	[[nodiscard]] FORCEINLINE SIZE_T GetAllocatedSize() const
+	[[nodiscard]] UE_FORCEINLINE_HINT SIZE_T GetAllocatedSize() const
 	{
 		return Pairs.GetAllocatedSize();
 	}
@@ -405,7 +405,7 @@ public:
 	 * @param Ar The archive to use.
 	 * @see GetAllocatedSize
 	 */
-	FORCEINLINE void CountBytes(FArchive& Ar) const
+	UE_FORCEINLINE_HINT void CountBytes(FArchive& Ar) const
 	{
 		Pairs.CountBytes(Ar);
 	}
@@ -417,16 +417,16 @@ public:
 	 * @param InValue The value to associate with the key.
 	 * @return A reference to the value as stored in the map. The reference is only valid until the next change to any key in the map.
 	 */
-	FORCEINLINE ValueType& Add(const KeyType&  InKey, const ValueType&  InValue) { return Emplace(InKey, InValue); }
-	FORCEINLINE ValueType& Add(const KeyType&  InKey,		ValueType&& InValue) { return Emplace(InKey, MoveTempIfPossible(InValue)); }
-	FORCEINLINE ValueType& Add(		 KeyType&& InKey, const ValueType&  InValue) { return Emplace(MoveTempIfPossible(InKey), InValue); }
-	FORCEINLINE ValueType& Add(		 KeyType&& InKey,		ValueType&& InValue) { return Emplace(MoveTempIfPossible(InKey), MoveTempIfPossible(InValue)); }
+	UE_FORCEINLINE_HINT ValueType& Add(const KeyType&  InKey, const ValueType&  InValue) { return Emplace(InKey, InValue); }
+	UE_FORCEINLINE_HINT ValueType& Add(const KeyType&  InKey,		ValueType&& InValue) { return Emplace(InKey, MoveTempIfPossible(InValue)); }
+	UE_FORCEINLINE_HINT ValueType& Add(		 KeyType&& InKey, const ValueType&  InValue) { return Emplace(MoveTempIfPossible(InKey), InValue); }
+	UE_FORCEINLINE_HINT ValueType& Add(		 KeyType&& InKey,		ValueType&& InValue) { return Emplace(MoveTempIfPossible(InKey), MoveTempIfPossible(InValue)); }
 
 	/** See Add() and class documentation section on ByHash() functions */
-	FORCEINLINE ValueType& AddByHash(uint32 KeyHash, const KeyType&  InKey, const ValueType&  InValue) { return EmplaceByHash(KeyHash, InKey, InValue); }
-	FORCEINLINE ValueType& AddByHash(uint32 KeyHash, const KeyType&  InKey,		  ValueType&& InValue) { return EmplaceByHash(KeyHash, InKey, MoveTempIfPossible(InValue)); }
-	FORCEINLINE ValueType& AddByHash(uint32 KeyHash,	   KeyType&& InKey, const ValueType&  InValue) { return EmplaceByHash(KeyHash, MoveTempIfPossible(InKey), InValue); }
-	FORCEINLINE ValueType& AddByHash(uint32 KeyHash,	   KeyType&& InKey,		  ValueType&& InValue) { return EmplaceByHash(KeyHash, MoveTempIfPossible(InKey), MoveTempIfPossible(InValue)); }
+	UE_FORCEINLINE_HINT ValueType& AddByHash(uint32 KeyHash, const KeyType&  InKey, const ValueType&  InValue) { return EmplaceByHash(KeyHash, InKey, InValue); }
+	UE_FORCEINLINE_HINT ValueType& AddByHash(uint32 KeyHash, const KeyType&  InKey,		  ValueType&& InValue) { return EmplaceByHash(KeyHash, InKey, MoveTempIfPossible(InValue)); }
+	UE_FORCEINLINE_HINT ValueType& AddByHash(uint32 KeyHash,	   KeyType&& InKey, const ValueType&  InValue) { return EmplaceByHash(KeyHash, MoveTempIfPossible(InKey), InValue); }
+	UE_FORCEINLINE_HINT ValueType& AddByHash(uint32 KeyHash,	   KeyType&& InKey,		  ValueType&& InValue) { return EmplaceByHash(KeyHash, MoveTempIfPossible(InKey), MoveTempIfPossible(InValue)); }
 
 	/**
 	 * Set a default value associated with a key.
@@ -434,12 +434,12 @@ public:
 	 * @param InKey The key to associate the value with.
 	 * @return A reference to the value as stored in the map. The reference is only valid until the next change to any key in the map.
 	 */
-	FORCEINLINE ValueType& Add(const KeyType&  InKey) { return Emplace(InKey); }
-	FORCEINLINE ValueType& Add(		 KeyType&& InKey) { return Emplace(MoveTempIfPossible(InKey)); }
+	UE_FORCEINLINE_HINT ValueType& Add(const KeyType&  InKey) { return Emplace(InKey); }
+	UE_FORCEINLINE_HINT ValueType& Add(		 KeyType&& InKey) { return Emplace(MoveTempIfPossible(InKey)); }
 
 	/** See Add() and class documentation section on ByHash() functions */
-	FORCEINLINE ValueType& AddByHash(uint32 KeyHash, const KeyType&  InKey) { return EmplaceByHash(KeyHash, InKey); }
-	FORCEINLINE ValueType& AddByHash(uint32 KeyHash,	   KeyType&& InKey)	{ return EmplaceByHash(KeyHash, MoveTempIfPossible(InKey)); }
+	UE_FORCEINLINE_HINT ValueType& AddByHash(uint32 KeyHash, const KeyType&  InKey) { return EmplaceByHash(KeyHash, InKey); }
+	UE_FORCEINLINE_HINT ValueType& AddByHash(uint32 KeyHash,	   KeyType&& InKey)	{ return EmplaceByHash(KeyHash, MoveTempIfPossible(InKey)); }
 
 	/**
 	 * Set the value associated with a key.
@@ -447,8 +447,8 @@ public:
 	 * @param InKeyValue A Tuple containing the Key and Value to associate together
 	 * @return A reference to the value as stored in the map. The reference is only valid until the next change to any key in the map.
 	 */
-	FORCEINLINE ValueType& Add(const TTuple<KeyType, ValueType>&  InKeyValue) { return Emplace(InKeyValue.Key, InKeyValue.Value); }
-	FORCEINLINE ValueType& Add(		 TTuple<KeyType, ValueType>&& InKeyValue) { return Emplace(MoveTempIfPossible(InKeyValue.Key), MoveTempIfPossible(InKeyValue.Value)); }
+	UE_FORCEINLINE_HINT ValueType& Add(const TTuple<KeyType, ValueType>&  InKeyValue) { return Emplace(InKeyValue.Key, InKeyValue.Value); }
+	UE_FORCEINLINE_HINT ValueType& Add(		 TTuple<KeyType, ValueType>&& InKeyValue) { return Emplace(MoveTempIfPossible(InKeyValue.Key), MoveTempIfPossible(InKeyValue.Value)); }
 
 	/**
 	 * Sets the value associated with a key.
@@ -520,7 +520,7 @@ public:
 	 * Removes the element at the specified index. The caller has to ensure that the index is valid.
 	 * @param Id The index of the element to remove.
 	 */
-	FORCEINLINE void Remove(FSetElementId Id)
+	UE_FORCEINLINE_HINT void Remove(FSetElementId Id)
 	{
 		Pairs.Remove(Id);
 	}
@@ -585,7 +585,7 @@ public:
 
 		return nullptr;
 	}
-	[[nodiscard]] FORCEINLINE const ValueType* Find(KeyConstPointerType Key) const
+	[[nodiscard]] UE_FORCEINLINE_HINT const ValueType* Find(KeyConstPointerType Key) const
 	{
 		return const_cast<TMapBase*>(this)->Find(Key);
 	}
@@ -602,7 +602,7 @@ public:
 		return nullptr;
 	}
 	template<typename ComparableKey>
-	[[nodiscard]] FORCEINLINE const ValueType* FindByHash(uint32 KeyHash, const ComparableKey& Key) const
+	[[nodiscard]] UE_FORCEINLINE_HINT const ValueType* FindByHash(uint32 KeyHash, const ComparableKey& Key) const
 	{
 		return const_cast<TMapBase*>(this)->FindByHash(KeyHash, Key);
 	}
@@ -615,7 +615,7 @@ public:
 		return Pair->Value;
 	}
 	template<typename ComparableKey>
-	[[nodiscard]] FORCEINLINE const ValueType& FindByHashChecked(uint32 KeyHash, const ComparableKey& Key) const
+	[[nodiscard]] UE_FORCEINLINE_HINT const ValueType& FindByHashChecked(uint32 KeyHash, const ComparableKey& Key) const
 	{
 		return const_cast<TMapBase*>(this)->FindByHashChecked(KeyHash, Key);
 	}
@@ -626,20 +626,20 @@ public:
 	 *
 	 * @param Key The key to search for.
 	 */
-	FORCEINLINE FSetElementId FindId(KeyInitType Key) const
+	UE_FORCEINLINE_HINT FSetElementId FindId(KeyInitType Key) const
 	{
 		return Pairs.FindId(Key);
 	}
 
 	/** See FindId() and class documentation section on ByHash() functions */
 	template<typename ComparableKey>
-	FORCEINLINE FSetElementId FindIdByHash(uint32 KeyHash, const ComparableKey& Key) const
+	UE_FORCEINLINE_HINT FSetElementId FindIdByHash(uint32 KeyHash, const ComparableKey& Key) const
 	{
 		return Pairs.FindIdByHash(KeyHash, Key);
 	}
 
 private:
-	[[nodiscard]] FORCEINLINE static uint32 HashKey(const KeyType& Key)
+	[[nodiscard]] UE_FORCEINLINE_HINT static uint32 HashKey(const KeyType& Key)
 	{
 		return KeyFuncs::GetKeyHash(Key);
 	}
@@ -690,12 +690,12 @@ public:
 	 * @param Key The key to search for.
 	 * @return A reference to the value associated with the specified key.
 	 */
-	FORCEINLINE ValueType& FindOrAdd(const KeyType&  Key) { return FindOrAddImpl(HashKey(Key),					  Key); }
-	FORCEINLINE ValueType& FindOrAdd(      KeyType&& Key) { return FindOrAddImpl(HashKey(Key), MoveTempIfPossible(Key)); }
+	UE_FORCEINLINE_HINT ValueType& FindOrAdd(const KeyType&  Key) { return FindOrAddImpl(HashKey(Key),					  Key); }
+	UE_FORCEINLINE_HINT ValueType& FindOrAdd(      KeyType&& Key) { return FindOrAddImpl(HashKey(Key), MoveTempIfPossible(Key)); }
 
 	/** See FindOrAdd() and class documentation section on ByHash() functions */
-	FORCEINLINE ValueType& FindOrAddByHash(uint32 KeyHash, const KeyType&  Key) { return FindOrAddImpl(KeyHash,                    Key); }
-	FORCEINLINE ValueType& FindOrAddByHash(uint32 KeyHash,       KeyType&& Key) { return FindOrAddImpl(KeyHash, MoveTempIfPossible(Key)); }
+	UE_FORCEINLINE_HINT ValueType& FindOrAddByHash(uint32 KeyHash, const KeyType&  Key) { return FindOrAddImpl(KeyHash,                    Key); }
+	UE_FORCEINLINE_HINT ValueType& FindOrAddByHash(uint32 KeyHash,       KeyType&& Key) { return FindOrAddImpl(KeyHash, MoveTempIfPossible(Key)); }
 
 	/**
 	 * Find the value associated with a specified key, or if none exists, 
@@ -705,16 +705,16 @@ public:
 	 * @param Value The value to associate with the key.
 	 * @return A reference to the value associated with the specified key.
 	 */
-	FORCEINLINE ValueType& FindOrAdd(const KeyType&  Key, const ValueType&  Value) { return FindOrAddImpl(HashKey(Key),						Key,                    Value  ); }
-	FORCEINLINE ValueType& FindOrAdd(const KeyType&  Key, ValueType&&       Value) { return FindOrAddImpl(HashKey(Key),						Key, MoveTempIfPossible(Value) ); }
-	FORCEINLINE ValueType& FindOrAdd(      KeyType&& Key, const ValueType&  Value) { return FindOrAddImpl(HashKey(Key), MoveTempIfPossible(Key),                    Value  ); }
-	FORCEINLINE ValueType& FindOrAdd(      KeyType&& Key, ValueType&&       Value) { return FindOrAddImpl(HashKey(Key), MoveTempIfPossible(Key), MoveTempIfPossible(Value) ); }
+	UE_FORCEINLINE_HINT ValueType& FindOrAdd(const KeyType&  Key, const ValueType&  Value) { return FindOrAddImpl(HashKey(Key),						Key,                    Value  ); }
+	UE_FORCEINLINE_HINT ValueType& FindOrAdd(const KeyType&  Key, ValueType&&       Value) { return FindOrAddImpl(HashKey(Key),						Key, MoveTempIfPossible(Value) ); }
+	UE_FORCEINLINE_HINT ValueType& FindOrAdd(      KeyType&& Key, const ValueType&  Value) { return FindOrAddImpl(HashKey(Key), MoveTempIfPossible(Key),                    Value  ); }
+	UE_FORCEINLINE_HINT ValueType& FindOrAdd(      KeyType&& Key, ValueType&&       Value) { return FindOrAddImpl(HashKey(Key), MoveTempIfPossible(Key), MoveTempIfPossible(Value) ); }
 
 	/** See FindOrAdd() and class documentation section on ByHash() functions */
-	FORCEINLINE ValueType& FindOrAddByHash(uint32 KeyHash, const KeyType&  Key, const ValueType&  Value) { return FindOrAddImpl(KeyHash,                     Key,                    Value); }
-	FORCEINLINE ValueType& FindOrAddByHash(uint32 KeyHash, const KeyType&  Key,       ValueType&& Value) { return FindOrAddImpl(KeyHash,                     Key, MoveTempIfPossible(Value)); }
-	FORCEINLINE ValueType& FindOrAddByHash(uint32 KeyHash,       KeyType&& Key, const ValueType&  Value) { return FindOrAddImpl(KeyHash, MoveTempIfPossible(Key),                    Value); }
-	FORCEINLINE ValueType& FindOrAddByHash(uint32 KeyHash,       KeyType&& Key,       ValueType&& Value) { return FindOrAddImpl(KeyHash, MoveTempIfPossible(Key), MoveTempIfPossible(Value)); }
+	UE_FORCEINLINE_HINT ValueType& FindOrAddByHash(uint32 KeyHash, const KeyType&  Key, const ValueType&  Value) { return FindOrAddImpl(KeyHash,                     Key,                    Value); }
+	UE_FORCEINLINE_HINT ValueType& FindOrAddByHash(uint32 KeyHash, const KeyType&  Key,       ValueType&& Value) { return FindOrAddImpl(KeyHash,                     Key, MoveTempIfPossible(Value)); }
+	UE_FORCEINLINE_HINT ValueType& FindOrAddByHash(uint32 KeyHash,       KeyType&& Key, const ValueType&  Value) { return FindOrAddImpl(KeyHash, MoveTempIfPossible(Key),                    Value); }
+	UE_FORCEINLINE_HINT ValueType& FindOrAddByHash(uint32 KeyHash,       KeyType&& Key,       ValueType&& Value) { return FindOrAddImpl(KeyHash, MoveTempIfPossible(Key), MoveTempIfPossible(Value)); }
 
 	/**
 	 * Find a reference to the value associated with a specified key.
@@ -797,14 +797,14 @@ public:
 	 * @param Key The key to check for.
 	 * @return true if the map contains the key.
 	 */
-	[[nodiscard]] FORCEINLINE bool Contains(KeyConstPointerType Key) const
+	[[nodiscard]] UE_FORCEINLINE_HINT bool Contains(KeyConstPointerType Key) const
 	{
 		return Pairs.Contains(Key);
 	}
 
 	/** See Contains() and class documentation section on ByHash() functions */
 	template<typename ComparableKey>
-	[[nodiscard]] FORCEINLINE bool ContainsByHash(uint32 KeyHash, const ComparableKey& Key) const
+	[[nodiscard]] UE_FORCEINLINE_HINT bool ContainsByHash(uint32 KeyHash, const ComparableKey& Key) const
 	{
 		return Pairs.ContainsByHash(KeyHash, Key);
 	}
@@ -873,56 +873,56 @@ protected:
 		typedef std::conditional_t<bConst, const typename ElementSetType::ElementType, typename ElementSetType::ElementType> PairType;
 
 	public:
-		[[nodiscard]] FORCEINLINE TBaseIterator(const PairItType& InElementIt)
+		[[nodiscard]] UE_FORCEINLINE_HINT TBaseIterator(const PairItType& InElementIt)
 			: PairIt(InElementIt)
 		{
 		}
 
-		FORCEINLINE TBaseIterator& operator++()
+		UE_FORCEINLINE_HINT TBaseIterator& operator++()
 		{
 			++PairIt;
 			return *this;
 		}
 
 		/** conversion to "bool" returning true if the iterator is valid. */
-		[[nodiscard]] FORCEINLINE explicit operator bool() const
+		[[nodiscard]] UE_FORCEINLINE_HINT explicit operator bool() const
 		{ 
 			return !!PairIt; 
 		}
 		/** inverse of the "bool" operator */
-		[[nodiscard]] FORCEINLINE bool operator !() const
+		[[nodiscard]] UE_FORCEINLINE_HINT bool operator !() const
 		{
 			return !(bool)*this;
 		}
 
-		[[nodiscard]] FORCEINLINE bool operator==(const TBaseIterator& Rhs) const
+		[[nodiscard]] UE_FORCEINLINE_HINT bool operator==(const TBaseIterator& Rhs) const
 		{
 			return PairIt == Rhs.PairIt;
 		}
-		[[nodiscard]] FORCEINLINE bool operator!=(const TBaseIterator& Rhs) const
+		[[nodiscard]] UE_FORCEINLINE_HINT bool operator!=(const TBaseIterator& Rhs) const
 		{
 			return PairIt != Rhs.PairIt;
 		}
 
-		[[nodiscard]] FORCEINLINE ItKeyType& Key() const
+		[[nodiscard]] UE_FORCEINLINE_HINT ItKeyType& Key() const
 		{
 			return PairIt->Key;
 		}
-		[[nodiscard]] FORCEINLINE ItValueType& Value() const
+		[[nodiscard]] UE_FORCEINLINE_HINT ItValueType& Value() const
 		{
 			return PairIt->Value;
 		}
 
-		[[nodiscard]] FORCEINLINE FSetElementId GetId() const
+		[[nodiscard]] UE_FORCEINLINE_HINT FSetElementId GetId() const
 		{
 			return PairIt.GetId();
 		}
 
-		[[nodiscard]] FORCEINLINE PairType& operator*() const
+		[[nodiscard]] UE_FORCEINLINE_HINT PairType& operator*() const
 		{
 			return *PairIt;
 		}
-		[[nodiscard]] FORCEINLINE PairType* operator->() const
+		[[nodiscard]] UE_FORCEINLINE_HINT PairType* operator->() const
 		{
 			return &*PairIt;
 		}
@@ -942,46 +942,46 @@ protected:
 
 	public:
 		/** Initialization constructor. */
-		[[nodiscard]] FORCEINLINE TBaseKeyIterator(const SetItType& InSetIt)
+		[[nodiscard]] UE_FORCEINLINE_HINT TBaseKeyIterator(const SetItType& InSetIt)
 			: SetIt(InSetIt)
 		{
 		}
 
-		FORCEINLINE TBaseKeyIterator& operator++()
+		UE_FORCEINLINE_HINT TBaseKeyIterator& operator++()
 		{
 			++SetIt;
 			return *this;
 		}
 
 		/** conversion to "bool" returning true if the iterator is valid. */
-		[[nodiscard]] FORCEINLINE explicit operator bool() const
+		[[nodiscard]] UE_FORCEINLINE_HINT explicit operator bool() const
 		{
 			return !!SetIt; 
 		}
 		/** inverse of the "bool" operator */
-		[[nodiscard]] FORCEINLINE bool operator !() const 
+		[[nodiscard]] UE_FORCEINLINE_HINT bool operator !() const 
 		{
 			return !(bool)*this;
 		}
 
-		[[nodiscard]] FORCEINLINE FSetElementId GetId() const
+		[[nodiscard]] UE_FORCEINLINE_HINT FSetElementId GetId() const
 		{
 			return SetIt.GetId();
 		}
-		[[nodiscard]] FORCEINLINE ItKeyType& Key() const
+		[[nodiscard]] UE_FORCEINLINE_HINT ItKeyType& Key() const
 		{
 			return SetIt->Key;
 		}
-		[[nodiscard]] FORCEINLINE ItValueType& Value() const
+		[[nodiscard]] UE_FORCEINLINE_HINT ItValueType& Value() const
 		{
 			return SetIt->Value;
 		}
 
-		[[nodiscard]] FORCEINLINE decltype(auto) operator*() const
+		[[nodiscard]] UE_FORCEINLINE_HINT decltype(auto) operator*() const
 		{
 			return SetIt.operator*();
 		}
-		[[nodiscard]] FORCEINLINE decltype(auto) operator->() const
+		[[nodiscard]] UE_FORCEINLINE_HINT decltype(auto) operator->() const
 		{
 			return SetIt.operator->();
 		}
@@ -1052,7 +1052,7 @@ public:
 	class TConstIterator : public TBaseIterator<true>
 	{
 	public:
-		[[nodiscard]] FORCEINLINE TConstIterator(const TMapBase& InMap)
+		[[nodiscard]] UE_FORCEINLINE_HINT TConstIterator(const TMapBase& InMap)
 			: TBaseIterator<true>(InMap.Pairs.CreateConstIterator())
 		{
 		}
@@ -1071,7 +1071,7 @@ public:
 	public:
 		using KeyArgumentType = typename IteratorType::KeyArgumentType;
 
-		[[nodiscard]] FORCEINLINE TConstKeyIterator(const TMapBase& InMap, KeyArgumentType InKey)
+		[[nodiscard]] UE_FORCEINLINE_HINT TConstKeyIterator(const TMapBase& InMap, KeyArgumentType InKey)
 			: Super(IteratorType(InMap.Pairs, InKey))
 		{
 		}
@@ -1087,38 +1087,38 @@ public:
 	public:
 		using KeyArgumentType = typename IteratorType::KeyArgumentType;
 
-		[[nodiscard]] FORCEINLINE TKeyIterator(TMapBase& InMap, KeyArgumentType InKey)
+		[[nodiscard]] UE_FORCEINLINE_HINT TKeyIterator(TMapBase& InMap, KeyArgumentType InKey)
 			: Super(IteratorType(InMap.Pairs, InKey))
 		{
 		}
 
 		/** Removes the current key-value pair from the map. */
-		FORCEINLINE void RemoveCurrent()
+		UE_FORCEINLINE_HINT void RemoveCurrent()
 		{
 			TBaseKeyIterator<false>::SetIt.RemoveCurrent();
 		}
 	};
 
 	/** Creates an iterator over all the pairs in this map */
-	[[nodiscard]] FORCEINLINE TIterator CreateIterator()
+	[[nodiscard]] UE_FORCEINLINE_HINT TIterator CreateIterator()
 	{
 		return TIterator(*this);
 	}
 
 	/** Creates a const iterator over all the pairs in this map */
-	[[nodiscard]] FORCEINLINE TConstIterator CreateConstIterator() const
+	[[nodiscard]] UE_FORCEINLINE_HINT TConstIterator CreateConstIterator() const
 	{
 		return TConstIterator(*this);
 	}
 
 	/** Creates an iterator over the values associated with a specified key in a map */
-	[[nodiscard]] FORCEINLINE TKeyIterator CreateKeyIterator(typename TKeyIterator::KeyArgumentType InKey)
+	[[nodiscard]] UE_FORCEINLINE_HINT TKeyIterator CreateKeyIterator(typename TKeyIterator::KeyArgumentType InKey)
 	{
 		return TKeyIterator(*this, InKey);
 	}
 
 	/** Creates a const iterator over the values associated with a specified key in a map */
-	[[nodiscard]] FORCEINLINE TConstKeyIterator CreateConstKeyIterator(typename TConstKeyIterator::KeyArgumentType InKey) const
+	[[nodiscard]] UE_FORCEINLINE_HINT TConstKeyIterator CreateConstKeyIterator(typename TConstKeyIterator::KeyArgumentType InKey) const
 	{
 		return TConstKeyIterator(*this, InKey);
 	}
@@ -1130,10 +1130,10 @@ public:
 	 * DO NOT USE DIRECTLY
 	 * STL-like iterators to enable range-based for loop support.
 	 */
-	[[nodiscard]] FORCEINLINE TRangedForIterator      begin() { return TRangedForIterator(Pairs.begin()); }
-	[[nodiscard]] FORCEINLINE TRangedForConstIterator begin() const { return TRangedForConstIterator(Pairs.begin()); }
-	[[nodiscard]] FORCEINLINE TRangedForIterator      end() { return TRangedForIterator(Pairs.end()); }
-	[[nodiscard]] FORCEINLINE TRangedForConstIterator end() const { return TRangedForConstIterator(Pairs.end()); }
+	[[nodiscard]] UE_FORCEINLINE_HINT TRangedForIterator      begin() { return TRangedForIterator(Pairs.begin()); }
+	[[nodiscard]] UE_FORCEINLINE_HINT TRangedForConstIterator begin() const { return TRangedForConstIterator(Pairs.begin()); }
+	[[nodiscard]] UE_FORCEINLINE_HINT TRangedForIterator      end() { return TRangedForIterator(Pairs.end()); }
+	[[nodiscard]] UE_FORCEINLINE_HINT TRangedForConstIterator end() const { return TRangedForConstIterator(Pairs.end()); }
 
 	// Maps are deliberately prevented from being hashed or compared, because this would hide potentially major performance problems behind default operations.
 	friend uint32 GetTypeHash(const TMapBase& Map) = delete;
@@ -1208,7 +1208,7 @@ public:
 	 * Invoked using "MyMapVar.KeySort( PREDICATE_CLASS() );"
 	 */
 	template<typename PREDICATE_CLASS>
-	FORCEINLINE void KeySort(const PREDICATE_CLASS& Predicate)
+	UE_FORCEINLINE_HINT void KeySort(const PREDICATE_CLASS& Predicate)
 	{
 		Super::Pairs.Sort(FKeyComparisonClass<PREDICATE_CLASS>(Predicate));
 	}
@@ -1218,7 +1218,7 @@ public:
 	 * Invoked using "MyMapVar.KeySort( PREDICATE_CLASS() );"
 	 */
 	template<typename PREDICATE_CLASS>
-	FORCEINLINE void KeyStableSort(const PREDICATE_CLASS& Predicate)
+	UE_FORCEINLINE_HINT void KeyStableSort(const PREDICATE_CLASS& Predicate)
 	{
 		Super::Pairs.StableSort(FKeyComparisonClass<PREDICATE_CLASS>(Predicate));
 	}
@@ -1228,7 +1228,7 @@ public:
 	 * Invoked using "MyMapVar.ValueSort( PREDICATE_CLASS() );"
 	 */
 	template<typename PREDICATE_CLASS>
-	FORCEINLINE void ValueSort(const PREDICATE_CLASS& Predicate)
+	UE_FORCEINLINE_HINT void ValueSort(const PREDICATE_CLASS& Predicate)
 	{
 		Super::Pairs.Sort(FValueComparisonClass<PREDICATE_CLASS>(Predicate));
 	}
@@ -1238,7 +1238,7 @@ public:
 	 * Invoked using "MyMapVar.ValueSort( PREDICATE_CLASS() );"
 	 */
 	template<typename PREDICATE_CLASS>
-	FORCEINLINE void ValueStableSort(const PREDICATE_CLASS& Predicate)
+	UE_FORCEINLINE_HINT void ValueStableSort(const PREDICATE_CLASS& Predicate)
 	{
 		Super::Pairs.StableSort(FValueComparisonClass<PREDICATE_CLASS>(Predicate));
 	}
@@ -1263,7 +1263,7 @@ private:
 
 	public:
 
-		[[nodiscard]] FORCEINLINE FKeyComparisonClass(const PREDICATE_CLASS& InPredicate)
+		[[nodiscard]] UE_FORCEINLINE_HINT FKeyComparisonClass(const PREDICATE_CLASS& InPredicate)
 			: Predicate(InPredicate)
 		{
 		}
@@ -1282,7 +1282,7 @@ private:
 
 	public:
 
-		[[nodiscard]] FORCEINLINE FValueComparisonClass(const PREDICATE_CLASS& InPredicate)
+		[[nodiscard]] UE_FORCEINLINE_HINT FValueComparisonClass(const PREDICATE_CLASS& InPredicate)
 			: Predicate(InPredicate)
 		{
 		}
@@ -1478,11 +1478,11 @@ public:
 		}
 	}
 
-	[[nodiscard]] FORCEINLINE ValueType& operator[](KeyConstPointerType Key)
+	[[nodiscard]] UE_FORCEINLINE_HINT ValueType& operator[](KeyConstPointerType Key)
 	{
 		return this->FindChecked(Key);
 	}
-	[[nodiscard]] FORCEINLINE const ValueType& operator[](KeyConstPointerType Key) const
+	[[nodiscard]] UE_FORCEINLINE_HINT const ValueType& operator[](KeyConstPointerType Key) const
 	{
 		return this->FindChecked(Key);
 	}
@@ -1662,10 +1662,10 @@ public:
 	 * @param InValue The value to associate.
 	 * @return A reference to the value as stored in the map; the reference is only valid until the next change to any key in the map.
 	 */
-	FORCEINLINE ValueType& AddUnique(const KeyType&  InKey, const ValueType&  InValue) { return EmplaceUnique(InKey, InValue); }
-	FORCEINLINE ValueType& AddUnique(const KeyType&  InKey, ValueType&& InValue) { return EmplaceUnique(InKey, MoveTempIfPossible(InValue)); }
-	FORCEINLINE ValueType& AddUnique(KeyType&& InKey, const ValueType&  InValue) { return EmplaceUnique(MoveTempIfPossible(InKey), InValue); }
-	FORCEINLINE ValueType& AddUnique(KeyType&& InKey, ValueType&& InValue) { return EmplaceUnique(MoveTempIfPossible(InKey), MoveTempIfPossible(InValue)); }
+	UE_FORCEINLINE_HINT ValueType& AddUnique(const KeyType&  InKey, const ValueType&  InValue) { return EmplaceUnique(InKey, InValue); }
+	UE_FORCEINLINE_HINT ValueType& AddUnique(const KeyType&  InKey, ValueType&& InValue) { return EmplaceUnique(InKey, MoveTempIfPossible(InValue)); }
+	UE_FORCEINLINE_HINT ValueType& AddUnique(KeyType&& InKey, const ValueType&  InValue) { return EmplaceUnique(MoveTempIfPossible(InKey), InValue); }
+	UE_FORCEINLINE_HINT ValueType& AddUnique(KeyType&& InKey, ValueType&& InValue) { return EmplaceUnique(MoveTempIfPossible(InKey), MoveTempIfPossible(InValue)); }
 
 	/**
 	 * Add a key-value association to the map.
@@ -1696,7 +1696,7 @@ public:
 	 * @param InKey The key to remove associated values for.
 	 * @return The number of values that were associated with the key.
 	 */
-	FORCEINLINE int32 Remove(KeyConstPointerType InKey)
+	UE_FORCEINLINE_HINT int32 Remove(KeyConstPointerType InKey)
 	{
 		return Super::Remove(InKey);
 	}
@@ -1758,7 +1758,7 @@ public:
 	 * @return If the map contains a matching association, a pointer to the value in the map is returned.  Otherwise nullptr is returned.
 	 *			The pointer is only valid as long as the map isn't changed.
 	 */
-	[[nodiscard]] FORCEINLINE const ValueType* FindPair(KeyInitType Key, ValueInitType Value) const
+	[[nodiscard]] UE_FORCEINLINE_HINT const ValueType* FindPair(KeyInitType Key, ValueInitType Value) const
 	{
 		return const_cast<TMultiMap*>(this)->FindPair(Key, Value);
 	}
@@ -1799,7 +1799,7 @@ public:
 	}
 
 	// Since we implement an overloaded Num() function in TMultiMap, we need to reimplement TMapBase::Num to make it visible.
-	[[nodiscard]] FORCEINLINE int32 Num() const
+	[[nodiscard]] UE_FORCEINLINE_HINT int32 Num() const
 	{
 		return Super::Num();
 	}
@@ -2160,7 +2160,7 @@ struct TMapPrivateFriend
 	}
 
 	template <typename KeyType, typename ValueType, typename SetAllocator, typename KeyFuncs>
-	FORCEINLINE static void SerializeStructured(FStructuredArchive::FSlot Slot, TMapBase<KeyType, ValueType, SetAllocator, KeyFuncs>& InMap)
+	UE_FORCEINLINE_HINT static void SerializeStructured(FStructuredArchive::FSlot Slot, TMapBase<KeyType, ValueType, SetAllocator, KeyFuncs>& InMap)
 	{
 		/*
 		if (Slot.GetUnderlyingArchive().IsTextFormat())
@@ -2207,14 +2207,14 @@ struct TMapPrivateFriend
 
 /** Serializer. */
 template <typename KeyType, typename ValueType, typename SetAllocator, typename KeyFuncs>
-FORCEINLINE FArchive& operator<<(FArchive& Ar, TMapBase<KeyType, ValueType, SetAllocator, KeyFuncs>& Map)
+UE_FORCEINLINE_HINT FArchive& operator<<(FArchive& Ar, TMapBase<KeyType, ValueType, SetAllocator, KeyFuncs>& Map)
 {
 	return TMapPrivateFriend::Serialize(Ar, Map);
 }
 
 /** Structured archive serializer. */
 template <typename KeyType, typename ValueType, typename SetAllocator, typename KeyFuncs>
-FORCEINLINE void operator<<(FStructuredArchive::FSlot Slot, TMapBase<KeyType, ValueType, SetAllocator, KeyFuncs>& InMap)
+UE_FORCEINLINE_HINT void operator<<(FStructuredArchive::FSlot Slot, TMapBase<KeyType, ValueType, SetAllocator, KeyFuncs>& InMap)
 {
 	TMapPrivateFriend::SerializeStructured(Slot, InMap);
 }
diff --git a/Engine/Source/Runtime/Core/Public/Containers/PagedArray.h b/Engine/Source/Runtime/Core/Public/Containers/PagedArray.h
--- a/Engine/Source/Runtime/Core/Public/Containers/PagedArray.h
+++ b/Engine/Source/Runtime/Core/Public/Containers/PagedArray.h
@@ -59,7 +59,7 @@ protected:
 	{
 	}
 
-	[[nodiscard]] FORCEINLINE ElementType& Dereference() const
+	[[nodiscard]] UE_FORCEINLINE_HINT ElementType& Dereference() const
 	{
 		return Data[PageIndex][PageOffset];
 	}
@@ -73,7 +73,7 @@ protected:
 		}
 	}
 
-	[[nodiscard]] FORCEINLINE bool Equals(const TIteratorBase& Other) const
+	[[nodiscard]] UE_FORCEINLINE_HINT bool Equals(const TIteratorBase& Other) const
 	{
 		return Data == Other.Data && PageIndex == Other.PageIndex && PageOffset == Other.PageOffset;
 	}
@@ -282,27 +282,27 @@ public:
 		Ar.CountBytes(Num() * sizeof(ElementType), Max() * sizeof(ElementType));
 	}
 
-	[[nodiscard]] FORCEINLINE SizeType Max() const
+	[[nodiscard]] UE_FORCEINLINE_HINT SizeType Max() const
 	{
 		return Pages.Num() * PageTraits::Capacity;
 	}
 
-	[[nodiscard]] FORCEINLINE SizeType Num() const
+	[[nodiscard]] UE_FORCEINLINE_HINT SizeType Num() const
 	{
 		return Count;
 	}
 
-	[[nodiscard]] FORCEINLINE SizeType NumPages() const
+	[[nodiscard]] UE_FORCEINLINE_HINT SizeType NumPages() const
 	{
 		return Pages.Num();
 	}
 
-	[[nodiscard]] FORCEINLINE bool IsValidIndex(SizeType Index) const
+	[[nodiscard]] UE_FORCEINLINE_HINT bool IsValidIndex(SizeType Index) const
 	{
 		return Index >= 0 && Index < Count;
 	}
 
-	[[nodiscard]] FORCEINLINE bool IsEmpty() const
+	[[nodiscard]] UE_FORCEINLINE_HINT bool IsEmpty() const
 	{
 		return !Count;
 	}
@@ -356,7 +356,7 @@ public:
 		Count = NewNum;
 	}
 	UE_ALLOWSHRINKING_BOOL_DEPRECATED("SetNum")
-	FORCEINLINE void SetNum(SizeType NewNum, bool bAllowShrinking)
+	UE_FORCEINLINE_HINT void SetNum(SizeType NewNum, bool bAllowShrinking)
 	{
 		SetNum(NewNum, bAllowShrinking ? EAllowShrinking::Yes : EAllowShrinking::No);
 	}
@@ -379,7 +379,7 @@ public:
 	/*
 	 * Sets the memory used by the container elements to zero.
 	 */
-	FORCEINLINE void SetZero()
+	UE_FORCEINLINE_HINT void SetZero()
 	{
 		SetMem(0);
 	}
@@ -403,7 +403,7 @@ public:
 		typename ContainerType
 		UE_REQUIRES(TIsContiguousContainer<ContainerType>::Value)
 	>
-	FORCEINLINE void Append(ContainerType&& Container)
+	UE_FORCEINLINE_HINT void Append(ContainerType&& Container)
 	{
 		Append(GetData(Container), GetNum(Container));
 	}
@@ -494,32 +494,32 @@ public:
 	/*
 	 * Adds the parameter element at the back of the container.
 	 */
-	FORCEINLINE SizeType Add(const ElementType& Element)
+	UE_FORCEINLINE_HINT SizeType Add(const ElementType& Element)
 	{
 		return Emplace(Element);
 	}
 
-	FORCEINLINE SizeType Add(ElementType&& Element)
+	UE_FORCEINLINE_HINT SizeType Add(ElementType&& Element)
 	{
 		return Emplace(MoveTempIfPossible(Element));
 	}
 
-	[[nodiscard]] FORCEINLINE ElementType& Add_GetRef(const ElementType& Element)
+	[[nodiscard]] UE_FORCEINLINE_HINT ElementType& Add_GetRef(const ElementType& Element)
 	{
 		return Emplace_GetRef(Element);
 	}
 
-	[[nodiscard]] FORCEINLINE ElementType& Add_GetRef(ElementType&& Element)
+	[[nodiscard]] UE_FORCEINLINE_HINT ElementType& Add_GetRef(ElementType&& Element)
 	{
 		return Emplace_GetRef(MoveTempIfPossible(Element));
 	}
 
-	FORCEINLINE void Push(const ElementType& Element)
+	UE_FORCEINLINE_HINT void Push(const ElementType& Element)
 	{
 		Emplace(Element);
 	}
 
-	FORCEINLINE void Push(ElementType&& Element)
+	UE_FORCEINLINE_HINT void Push(ElementType&& Element)
 	{
 		Emplace(MoveTempIfPossible(Element));
 	}
@@ -541,7 +541,7 @@ public:
 		--Count;
 	}
 	UE_ALLOWSHRINKING_BOOL_DEPRECATED("Pop")
-	FORCEINLINE void Pop(bool bAllowShrinking)
+	UE_FORCEINLINE_HINT void Pop(bool bAllowShrinking)
 	{
 		Pop(bAllowShrinking ? EAllowShrinking::Yes : EAllowShrinking::No);
 	}
@@ -575,7 +575,7 @@ public:
 		--Count;
 	}
 	UE_ALLOWSHRINKING_BOOL_DEPRECATED("RemoveAtSwap")
-	FORCEINLINE void RemoveAtSwap(SizeType Index, bool bAllowShrinking)
+	UE_FORCEINLINE_HINT void RemoveAtSwap(SizeType Index, bool bAllowShrinking)
 	{
 		RemoveAtSwap(Index, bAllowShrinking ? EAllowShrinking::Yes : EAllowShrinking::No);
 	}
@@ -675,19 +675,19 @@ public:
 	}
 
 	/** STL IteratorType model compliance methods.*/
-	[[nodiscard]] FORCEINLINE ConstIteratorType begin() const
+	[[nodiscard]] UE_FORCEINLINE_HINT ConstIteratorType begin() const
 	{
 		return GetIterator(0);
 	}
-	[[nodiscard]] FORCEINLINE IteratorType begin()
+	[[nodiscard]] UE_FORCEINLINE_HINT IteratorType begin()
 	{
 		return GetIterator(0);
 	}
-	[[nodiscard]] FORCEINLINE ConstIteratorType end() const
+	[[nodiscard]] UE_FORCEINLINE_HINT ConstIteratorType end() const
 	{
 		return GetIterator(Count);
 	}
-	[[nodiscard]] FORCEINLINE IteratorType end()
+	[[nodiscard]] UE_FORCEINLINE_HINT IteratorType end()
 	{
 		return GetIterator(Count);
 	}
@@ -710,7 +710,7 @@ public:
 		return true;
 	}
 
-	[[nodiscard]] FORCEINLINE bool operator!=(const TPagedArray& Right) const
+	[[nodiscard]] UE_FORCEINLINE_HINT bool operator!=(const TPagedArray& Right) const
 	{
 		return !(*this == Right);
 	}
@@ -719,12 +719,12 @@ private:
 	TArray<PageType> Pages;
 	SizeType Count = 0;
 
-	[[nodiscard]] FORCEINLINE ConstIteratorType GetIterator(SizeType Offset) const
+	[[nodiscard]] UE_FORCEINLINE_HINT ConstIteratorType GetIterator(SizeType Offset) const
 	{
 		return ConstIteratorType(InPlace, Pages.GetData(), Offset);
 	}
 
-	[[nodiscard]] FORCEINLINE IteratorType GetIterator(SizeType Offset)
+	[[nodiscard]] UE_FORCEINLINE_HINT IteratorType GetIterator(SizeType Offset)
 	{
 		return IteratorType(InPlace, Pages.GetData(), Offset);
 	}
@@ -823,7 +823,7 @@ private:
 		}
 	}
 
-	FORCEINLINE void CheckValidIndex(SizeType Index) const
+	UE_FORCEINLINE_HINT void CheckValidIndex(SizeType Index) const
 	{
 		checkf((Index >= 0) & (Index < Count), TEXT("Parameter index %d exceeds container size %d"), Index, Count);
 	}
@@ -831,14 +831,14 @@ private:
 
 /** Serializer. */
 template <typename ElementType, int32 PageSizeInBytes, typename AllocatorType>
-FORCEINLINE FArchive& operator<<(FArchive& Ar, TPagedArray<ElementType, PageSizeInBytes, AllocatorType>& InOutPagedArray)
+UE_FORCEINLINE_HINT FArchive& operator<<(FArchive& Ar, TPagedArray<ElementType, PageSizeInBytes, AllocatorType>& InOutPagedArray)
 {
 	return UE::Core::PagedArray::Private::Serialize(Ar, InOutPagedArray);
 }
 
 /** Structured archive serializer. */
 template <typename ElementType, int32 PageSizeInBytes, typename AllocatorType>
-FORCEINLINE void operator<<(FStructuredArchive::FSlot Slot, TPagedArray<ElementType, PageSizeInBytes, AllocatorType>& InOutPagedArray)
+UE_FORCEINLINE_HINT void operator<<(FStructuredArchive::FSlot Slot, TPagedArray<ElementType, PageSizeInBytes, AllocatorType>& InOutPagedArray)
 {
 	UE::Core::PagedArray::Private::SerializeStructured(Slot, InOutPagedArray);
 }
\ No newline at end of file
diff --git a/Engine/Source/Runtime/Core/Public/Containers/Queue.h b/Engine/Source/Runtime/Core/Public/Containers/Queue.h
--- a/Engine/Source/Runtime/Core/Public/Containers/Queue.h
+++ b/Engine/Source/Runtime/Core/Public/Containers/Queue.h
@@ -244,7 +244,7 @@ public:
 		return &Tail->NextNode->Item;
 	}
 
-	[[nodiscard]] FORCEINLINE const FElementType* Peek() const
+	[[nodiscard]] UE_FORCEINLINE_HINT const FElementType* Peek() const
 	{
 		return const_cast<TQueue*>(this)->Peek();
 	}
diff --git a/Engine/Source/Runtime/Core/Public/Containers/RingBuffer.h b/Engine/Source/Runtime/Core/Public/Containers/RingBuffer.h
--- a/Engine/Source/Runtime/Core/Public/Containers/RingBuffer.h
+++ b/Engine/Source/Runtime/Core/Public/Containers/RingBuffer.h
@@ -71,18 +71,18 @@ public:
 		return Tmp -= Offset;
 	}
 
-	[[nodiscard]] FORCEINLINE ElementType& operator*() const
+	[[nodiscard]] UE_FORCEINLINE_HINT ElementType& operator*() const
 	{
 		return Container[Index];
 	}
 
-	[[nodiscard]] FORCEINLINE ElementType* operator->() const
+	[[nodiscard]] UE_FORCEINLINE_HINT ElementType* operator->() const
 	{
 		return &Container[Index];
 	}
 
 	/** conversion to "bool" returning true if the iterator has not reached the last element. */
-	[[nodiscard]] FORCEINLINE explicit operator bool() const
+	[[nodiscard]] UE_FORCEINLINE_HINT explicit operator bool() const
 	{
 		return Container.IsValidIndex(Index);
 	}
@@ -112,8 +112,8 @@ public:
 		--Index;
 	}
 
-	[[nodiscard]] FORCEINLINE bool operator==(const TRingBufferIterator& Rhs) const { return &Container == &Rhs.Container && Index == Rhs.Index; }
-	[[nodiscard]] FORCEINLINE bool operator!=(const TRingBufferIterator& Rhs) const { return &Container != &Rhs.Container || Index != Rhs.Index; }
+	[[nodiscard]] UE_FORCEINLINE_HINT bool operator==(const TRingBufferIterator& Rhs) const { return &Container == &Rhs.Container && Index == Rhs.Index; }
+	[[nodiscard]] UE_FORCEINLINE_HINT bool operator!=(const TRingBufferIterator& Rhs) const { return &Container != &Rhs.Container || Index != Rhs.Index; }
 
 private:
 
@@ -976,7 +976,7 @@ private:
 	}
 
 	/** Return a pointer to the underlying storage of the RingBuffer */
-	[[nodiscard]] FORCEINLINE const ElementType* GetStorage() const
+	[[nodiscard]] UE_FORCEINLINE_HINT const ElementType* GetStorage() const
 	{
 		return const_cast<TRingBuffer*>(this)->GetStorage();
 	}
diff --git a/Engine/Source/Runtime/Core/Public/Containers/Set.h b/Engine/Source/Runtime/Core/Public/Containers/Set.h
--- a/Engine/Source/Runtime/Core/Public/Containers/Set.h
+++ b/Engine/Source/Runtime/Core/Public/Containers/Set.h
@@ -37,7 +37,7 @@ public:
 		typename InitType
 		UE_REQUIRES(!std::is_same_v<TSetElementBase, std::decay_t<InitType>>)
 	>
-	[[nodiscard]] explicit FORCEINLINE TSetElementBase(InitType&& InValue)
+	[[nodiscard]] explicit UE_FORCEINLINE_HINT TSetElementBase(InitType&& InValue)
 		: Value(Forward<InitType>(InValue))
 	{
 	}
@@ -71,7 +71,7 @@ public:
 		typename InitType
 		UE_REQUIRES(!std::is_same_v<TSetElementBase, std::decay_t<InitType>>)
 	>
-	[[nodiscard]] explicit FORCEINLINE TSetElementBase(InitType&& InValue)
+	[[nodiscard]] explicit UE_FORCEINLINE_HINT TSetElementBase(InitType&& InValue)
 		: Value(Forward<InitType>(InValue))
 	{
 	}
@@ -105,7 +105,7 @@ public:
 		typename InitType
 		UE_REQUIRES(!std::is_same_v<TSetElement, std::decay_t<InitType>>)
 	>
-	[[nodiscard]] explicit FORCEINLINE TSetElement(InitType&& InValue)
+	[[nodiscard]] explicit UE_FORCEINLINE_HINT TSetElement(InitType&& InValue)
 		: Super(Forward<InitType>(InValue))
 	{
 	}
@@ -116,11 +116,11 @@ public:
 	TSetElement& operator=(const TSetElement&) = default;
 
 	// Comparison operators
-	[[nodiscard]] FORCEINLINE bool operator==(const TSetElement& Other) const
+	[[nodiscard]] UE_FORCEINLINE_HINT bool operator==(const TSetElement& Other) const
 	{
 		return this->Value == Other.Value;
 	}
-	[[nodiscard]] FORCEINLINE bool operator!=(const TSetElement& Other) const
+	[[nodiscard]] UE_FORCEINLINE_HINT bool operator!=(const TSetElement& Other) const
 	{
 		return this->Value != Other.Value;
 	}
@@ -208,26 +208,26 @@ private:
 
 public:
 	/** Initialization constructor. */
-	[[nodiscard]] FORCEINLINE TSet() = default;
+	[[nodiscard]] UE_FORCEINLINE_HINT TSet() = default;
 
 	/** Copy constructor. */
-	[[nodiscard]] FORCEINLINE TSet(const TSet& Copy)
+	[[nodiscard]] UE_FORCEINLINE_HINT TSet(const TSet& Copy)
 	{
 		*this = Copy;
 	}
 
-	[[nodiscard]] FORCEINLINE explicit TSet(TArrayView<const ElementType> InArrayView)
+	[[nodiscard]] UE_FORCEINLINE_HINT explicit TSet(TArrayView<const ElementType> InArrayView)
 	{
 		Append(InArrayView);
 	}
 
-	[[nodiscard]] FORCEINLINE explicit TSet(TArray<ElementType>&& InArray)
+	[[nodiscard]] UE_FORCEINLINE_HINT explicit TSet(TArray<ElementType>&& InArray)
 	{
 		Append(MoveTemp(InArray));
 	}
 
 	/** Destructor. */
-	FORCEINLINE ~TSet()
+	UE_FORCEINLINE_HINT ~TSet()
 	{
 		UE_STATIC_ASSERT_WARN(TIsTriviallyRelocatable_V<InElementType>, "TSet can only be used with trivially relocatable types");
 
@@ -440,7 +440,7 @@ public:
 	}
 
 	/** Relaxes the set's hash to a size strictly bounded by the number of elements in the set. */
-	FORCEINLINE void Relax()
+	UE_FORCEINLINE_HINT void Relax()
 	{
 		ConditionalRehash(Elements.Num(), EAllowShrinking::Yes);
 	}
@@ -450,7 +450,7 @@ public:
 	 * Only returns the size of allocations made directly by the container, not the elements themselves.
 	 * @return number of bytes allocated by this container
 	 */
-	[[nodiscard]] FORCEINLINE SIZE_T GetAllocatedSize( void ) const
+	[[nodiscard]] UE_FORCEINLINE_HINT SIZE_T GetAllocatedSize( void ) const
 	{
 		return Elements.GetAllocatedSize() + Hash.GetAllocatedSize(HashSize, sizeof(FSetElementId));
 	}
@@ -474,19 +474,19 @@ public:
 	}
 
 	/** @return the number of elements. */
-	[[nodiscard]] FORCEINLINE int32 Num() const
+	[[nodiscard]] UE_FORCEINLINE_HINT int32 Num() const
 	{
 		return Elements.Num();
 	}
 
 	/** @return The number of elements the set can hold before reallocation. */
-	[[nodiscard]] FORCEINLINE int32 Max() const
+	[[nodiscard]] UE_FORCEINLINE_HINT int32 Max() const
 	{
 		return Elements.Max();
 	}
 
 	/** @return The non-inclusive maximum index of elements in the set. */
-	[[nodiscard]] FORCEINLINE int32 GetMaxIndex() const
+	[[nodiscard]] UE_FORCEINLINE_HINT int32 GetMaxIndex() const
 	{
 		return Elements.GetMaxIndex();
 	}
@@ -506,25 +506,25 @@ public:
 	}
 
 	/** Accesses the identified element's value. Element must be valid (see @IsValidId). */
-	[[nodiscard]] FORCEINLINE ElementType& operator[](FSetElementId Id)
+	[[nodiscard]] UE_FORCEINLINE_HINT ElementType& operator[](FSetElementId Id)
 	{
 		return Elements[Id.AsInteger()].Value;
 	}
 
 	/** Accesses the identified element's value. Element must be valid (see @IsValidId). */
-	[[nodiscard]] FORCEINLINE const ElementType& operator[](FSetElementId Id) const
+	[[nodiscard]] UE_FORCEINLINE_HINT const ElementType& operator[](FSetElementId Id) const
 	{
 		return Elements[Id.AsInteger()].Value;
 	}
 
 	/** Accesses the identified element's value. Element must be valid (see @IsValidId). */
-	[[nodiscard]] FORCEINLINE ElementType& Get(FSetElementId Id)
+	[[nodiscard]] UE_FORCEINLINE_HINT ElementType& Get(FSetElementId Id)
 	{
 		return Elements[Id.AsInteger()].Value;
 	}
 
 	/** Accesses the identified element's value. Element must be valid (see @IsValidId). */
-	[[nodiscard]] FORCEINLINE const ElementType& Get(FSetElementId Id) const
+	[[nodiscard]] UE_FORCEINLINE_HINT const ElementType& Get(FSetElementId Id) const
 	{
 		return Elements[Id.AsInteger()].Value;
 	}
@@ -536,11 +536,11 @@ public:
 	 * @param	bIsAlreadyInSetPtr	[out]	Optional pointer to bool that will be set depending on whether element is already in set
 	 * @return	A pointer to the element stored in the set.
 	 */
-	FORCEINLINE FSetElementId Add(const InElementType&  InElement, bool* bIsAlreadyInSetPtr = nullptr)
+	UE_FORCEINLINE_HINT FSetElementId Add(const InElementType&  InElement, bool* bIsAlreadyInSetPtr = nullptr)
 	{
 		return Emplace(InElement, bIsAlreadyInSetPtr);
 	}
-	FORCEINLINE FSetElementId Add(InElementType&& InElement, bool* bIsAlreadyInSetPtr = nullptr)
+	UE_FORCEINLINE_HINT FSetElementId Add(InElementType&& InElement, bool* bIsAlreadyInSetPtr = nullptr)
 	{
 		return Emplace(MoveTempIfPossible(InElement), bIsAlreadyInSetPtr);
 	}
@@ -552,11 +552,11 @@ public:
 	 * @param	bIsAlreadyInSetPtr	[out]	Optional pointer to bool that will be set depending on whether element is already in set
 	 * @return	A reference to the element stored in the set.
 	 */
-	FORCEINLINE ElementType& FindOrAdd(const InElementType& InElement, bool* bIsAlreadyInSetPtr = nullptr)
+	UE_FORCEINLINE_HINT ElementType& FindOrAdd(const InElementType& InElement, bool* bIsAlreadyInSetPtr = nullptr)
 	{
 		return FindOrAddByHash(KeyFuncs::GetKeyHash(KeyFuncs::GetSetKey(InElement)), InElement, bIsAlreadyInSetPtr);
 	}
-	FORCEINLINE ElementType& FindOrAdd(InElementType&& InElement, bool* bIsAlreadyInSetPtr = nullptr)
+	UE_FORCEINLINE_HINT ElementType& FindOrAdd(InElementType&& InElement, bool* bIsAlreadyInSetPtr = nullptr)
 	{
 		return FindOrAddByHash(KeyFuncs::GetKeyHash(KeyFuncs::GetSetKey(InElement)), MoveTempIfPossible(InElement), bIsAlreadyInSetPtr);
 	}
@@ -569,11 +569,11 @@ public:
 	 * @param	bIsAlreadyInSetPtr	[out]	Optional pointer to bool that will be set depending on whether element is already in set
      * @return  A handle to the element stored in the set
      */
-	FORCEINLINE FSetElementId AddByHash(uint32 KeyHash, const InElementType& InElement, bool* bIsAlreadyInSetPtr = nullptr)
+	UE_FORCEINLINE_HINT FSetElementId AddByHash(uint32 KeyHash, const InElementType& InElement, bool* bIsAlreadyInSetPtr = nullptr)
 	{
 		return EmplaceByHash(KeyHash, InElement, bIsAlreadyInSetPtr);
 	}
-	FORCEINLINE FSetElementId AddByHash(uint32 KeyHash,		 InElementType&& InElement, bool* bIsAlreadyInSetPtr = nullptr)
+	UE_FORCEINLINE_HINT FSetElementId AddByHash(uint32 KeyHash,		 InElementType&& InElement, bool* bIsAlreadyInSetPtr = nullptr)
 	{
 		return EmplaceByHash(KeyHash, MoveTempIfPossible(InElement), bIsAlreadyInSetPtr);
 	}
@@ -891,7 +891,7 @@ public:
 	 * @param Key - The key to search for.
 	 * @return A const pointer to an element with the given key.  If no element in the set has the given key, this will return NULL.
 	 */
-	[[nodiscard]] FORCEINLINE const ElementType* Find(KeyInitType Key) const
+	[[nodiscard]] UE_FORCEINLINE_HINT const ElementType* Find(KeyInitType Key) const
 	{
 		return const_cast<TSet*>(this)->Find(Key);
 	}
@@ -999,7 +999,7 @@ public:
 	 * @param Key - The key to check for.
 	 * @return true if the set contains an element with the given key.
 	 */
-	[[nodiscard]] FORCEINLINE bool Contains(KeyInitType Key) const
+	[[nodiscard]] UE_FORCEINLINE_HINT bool Contains(KeyInitType Key) const
 	{
 		return FindIndexByHash(KeyFuncs::GetKeyHash(Key), Key) != INDEX_NONE;
 	}
@@ -1223,7 +1223,7 @@ public:
 	 *
 	 * @param Addr The address to check.
 	 */
-	FORCEINLINE void CheckAddress(const ElementType* Addr) const
+	UE_FORCEINLINE_HINT void CheckAddress(const ElementType* Addr) const
 	{
 		Elements.CheckAddress(Addr);
 	}
@@ -1318,12 +1318,12 @@ private:
 		TDereferenceWrapper< ElementType, PREDICATE_CLASS > Predicate;
 
 	public:
-		[[nodiscard]] FORCEINLINE FElementCompareClass( const PREDICATE_CLASS& InPredicate )
+		[[nodiscard]] UE_FORCEINLINE_HINT FElementCompareClass( const PREDICATE_CLASS& InPredicate )
 			: Predicate( InPredicate )
 		{
 		}
 
-		[[nodiscard]] FORCEINLINE bool operator()( const SetElementType& A,const SetElementType& B ) const
+		[[nodiscard]] UE_FORCEINLINE_HINT bool operator()( const SetElementType& A,const SetElementType& B ) const
 		{
 			return Predicate( A.Value, B.Value );
 		}
@@ -1378,7 +1378,7 @@ public:
 
 private:
 
-	[[nodiscard]] FORCEINLINE FSetElementId& GetTypedHash(int32 HashIndex) const
+	[[nodiscard]] UE_FORCEINLINE_HINT FSetElementId& GetTypedHash(int32 HashIndex) const
 	{
 		return ((FSetElementId*)Hash.GetAllocation())[HashIndex & (HashSize - 1)];
 	}
@@ -1396,7 +1396,7 @@ private:
 	}
 
 	/** Hashes and links an added element to the hash chain. */
-	FORCEINLINE void HashElement(SizeType ElementIndex, const SetElementType& Element) const
+	UE_FORCEINLINE_HINT void HashElement(SizeType ElementIndex, const SetElementType& Element) const
 	{
 		LinkElement(ElementIndex, Element, KeyFuncs::GetKeyHash(KeyFuncs::GetSetKey(Element.Value)));
 	}
@@ -1432,7 +1432,7 @@ private:
 	 * @param AllowShrinking - If the hash is allowed to shrink.
 	 * @return true if the set should berehashed.
 	 */
-	[[nodiscard]] FORCEINLINE bool ShouldRehash(int32 NumHashedElements, int32 DesiredHashSize, EAllowShrinking AllowShrinking) const
+	[[nodiscard]] UE_FORCEINLINE_HINT bool ShouldRehash(int32 NumHashedElements, int32 DesiredHashSize, EAllowShrinking AllowShrinking) const
 	{
 		// If the hash hasn't been created yet, or is smaller than the desired hash size, rehash.
 		// If shrinking is allowed and the hash is bigger than the desired hash size, rehash.
@@ -1491,45 +1491,45 @@ private:
 			std::conditional_t<bRangedFor, typename ElementArrayType::TRangedForIterator,      typename ElementArrayType::TIterator     >
 		> ElementItType;
 
-		[[nodiscard]] FORCEINLINE TBaseIterator(const ElementItType& InElementIt)
+		[[nodiscard]] UE_FORCEINLINE_HINT TBaseIterator(const ElementItType& InElementIt)
 			: ElementIt(InElementIt)
 		{
 		}
 
 		/** Advances the iterator to the next element. */
-		FORCEINLINE TBaseIterator& operator++()
+		UE_FORCEINLINE_HINT TBaseIterator& operator++()
 		{
 			++ElementIt;
 			return *this;
 		}
 
 		/** conversion to "bool" returning true if the iterator is valid. */
-		[[nodiscard]] FORCEINLINE explicit operator bool() const
+		[[nodiscard]] UE_FORCEINLINE_HINT explicit operator bool() const
 		{ 
 			return !!ElementIt; 
 		}
 		/** inverse of the "bool" operator */
-		[[nodiscard]] FORCEINLINE bool operator!() const 
+		[[nodiscard]] UE_FORCEINLINE_HINT bool operator!() const 
 		{
 			return !(bool)*this;
 		}
 
 		// Accessors.
-		[[nodiscard]] FORCEINLINE FSetElementId GetId() const
+		[[nodiscard]] UE_FORCEINLINE_HINT FSetElementId GetId() const
 		{
 			return FSetElementId::FromInteger(ElementIt.GetIndex());
 		}
-		[[nodiscard]] FORCEINLINE ItElementType* operator->() const
+		[[nodiscard]] UE_FORCEINLINE_HINT ItElementType* operator->() const
 		{
 			return &ElementIt->Value;
 		}
-		[[nodiscard]] FORCEINLINE ItElementType& operator*() const
+		[[nodiscard]] UE_FORCEINLINE_HINT ItElementType& operator*() const
 		{
 			return ElementIt->Value;
 		}
 
-		[[nodiscard]] FORCEINLINE bool operator!=(const TBaseIterator& Rhs) const { return ElementIt != Rhs.ElementIt; }
-		[[nodiscard]] FORCEINLINE bool operator==(const TBaseIterator& Rhs) const { return ElementIt == Rhs.ElementIt; }
+		[[nodiscard]] UE_FORCEINLINE_HINT bool operator!=(const TBaseIterator& Rhs) const { return ElementIt != Rhs.ElementIt; }
+		[[nodiscard]] UE_FORCEINLINE_HINT bool operator==(const TBaseIterator& Rhs) const { return ElementIt == Rhs.ElementIt; }
 
 		ElementItType ElementIt;
 	};
@@ -1589,26 +1589,26 @@ private:
 		}
 
 		/** conversion to "bool" returning true if the iterator is valid. */
-		[[nodiscard]] FORCEINLINE explicit operator bool() const
+		[[nodiscard]] UE_FORCEINLINE_HINT explicit operator bool() const
 		{ 
 			return Index != INDEX_NONE;
 		}
 		/** inverse of the "bool" operator */
-		[[nodiscard]] FORCEINLINE bool operator!() const
+		[[nodiscard]] UE_FORCEINLINE_HINT bool operator!() const
 		{
 			return !(bool)*this;
 		}
 
 		// Accessors.
-		[[nodiscard]] FORCEINLINE FSetElementId GetId() const
+		[[nodiscard]] UE_FORCEINLINE_HINT FSetElementId GetId() const
 		{
 			return FSetElementId::FromInteger(Index);
 		}
-		[[nodiscard]] FORCEINLINE ItElementType* operator->() const
+		[[nodiscard]] UE_FORCEINLINE_HINT ItElementType* operator->() const
 		{
 			return &Set.Elements[Index].Value;
 		}
-		[[nodiscard]] FORCEINLINE ItElementType& operator*() const
+		[[nodiscard]] UE_FORCEINLINE_HINT ItElementType& operator*() const
 		{
 			return Set.Elements[Index].Value;
 		}
@@ -1628,7 +1628,7 @@ public:
 		friend class TSet;
 
 	public:
-		[[nodiscard]] FORCEINLINE TConstIterator(const TSet& InSet)
+		[[nodiscard]] UE_FORCEINLINE_HINT TConstIterator(const TSet& InSet)
 			: TBaseIterator<true>(InSet.Elements.begin())
 		{
 		}
@@ -1647,7 +1647,7 @@ public:
 		}
 
 		/** Removes the current element from the set. */
-		FORCEINLINE void RemoveCurrent()
+		UE_FORCEINLINE_HINT void RemoveCurrent()
 		{
 			Set.RemoveByIndex(TBaseIterator<false>::ElementIt.GetIndex());
 		}
@@ -1668,7 +1668,7 @@ public:
 	public:
 		using KeyArgumentType = typename Super::KeyArgumentType;
 
-		[[nodiscard]] FORCEINLINE TConstKeyIterator(const TSet& InSet, KeyArgumentType InKey)
+		[[nodiscard]] UE_FORCEINLINE_HINT TConstKeyIterator(const TSet& InSet, KeyArgumentType InKey)
 			: Super(InSet, InKey)
 		{
 		}
@@ -1683,7 +1683,7 @@ public:
 	public:
 		using KeyArgumentType = typename Super::KeyArgumentType;
 
-		[[nodiscard]] FORCEINLINE TKeyIterator(TSet& InSet, KeyArgumentType InKey)
+		[[nodiscard]] UE_FORCEINLINE_HINT TKeyIterator(TSet& InSet, KeyArgumentType InKey)
 			: Super(InSet, InKey)
 		{
 		}
@@ -1697,13 +1697,13 @@ public:
 	};
 
 	/** Creates an iterator for the contents of this set */
-	[[nodiscard]] FORCEINLINE TIterator CreateIterator()
+	[[nodiscard]] UE_FORCEINLINE_HINT TIterator CreateIterator()
 	{
 		return TIterator(*this);
 	}
 
 	/** Creates a const iterator for the contents of this set */
-	[[nodiscard]] FORCEINLINE TConstIterator CreateConstIterator() const
+	[[nodiscard]] UE_FORCEINLINE_HINT TConstIterator CreateConstIterator() const
 	{
 		return TConstIterator(*this);
 	}
@@ -1715,19 +1715,19 @@ public:
 	 * DO NOT USE DIRECTLY
 	 * STL-like iterators to enable range-based for loop support.
 	 */
-	[[nodiscard]] FORCEINLINE TRangedForIterator begin()
+	[[nodiscard]] UE_FORCEINLINE_HINT TRangedForIterator begin()
 	{
 		return TRangedForIterator(Elements.begin());
 	}
-	[[nodiscard]] FORCEINLINE TRangedForConstIterator begin() const
+	[[nodiscard]] UE_FORCEINLINE_HINT TRangedForConstIterator begin() const
 	{
 		return TRangedForConstIterator(Elements.begin());
 	}
-	[[nodiscard]] FORCEINLINE TRangedForIterator end()
+	[[nodiscard]] UE_FORCEINLINE_HINT TRangedForIterator end()
 	{
 		return TRangedForIterator(Elements.end());
 	}
-	[[nodiscard]] FORCEINLINE TRangedForConstIterator end() const
+	[[nodiscard]] UE_FORCEINLINE_HINT TRangedForConstIterator end() const
 	{
 		return TRangedForConstIterator(Elements.end());
 	}
@@ -2087,7 +2087,7 @@ private:
 	HashType         Hash;
 	int32            HashSize;
 
-	[[nodiscard]] FORCEINLINE FSetElementId& GetTypedHash(int32 HashIndex) const
+	[[nodiscard]] UE_FORCEINLINE_HINT FSetElementId& GetTypedHash(int32 HashIndex) const
 	{
 		return ((FSetElementId*)Hash.GetAllocation())[HashIndex & (HashSize - 1)];
 	}
@@ -2211,14 +2211,14 @@ struct TSetPrivateFriend
 
 /** Serializer. */
 template <typename ElementType>
-FORCEINLINE FArchive& operator<<(FArchive& Ar, TSetElement<ElementType>& Element)
+UE_FORCEINLINE_HINT FArchive& operator<<(FArchive& Ar, TSetElement<ElementType>& Element)
 {
 	return Ar << Element.Value;
 }
 
 /** Structured archive serializer. */
 template <typename ElementType>
-FORCEINLINE void operator<<(FStructuredArchive::FSlot& Ar, TSetElement<ElementType>& Element)
+UE_FORCEINLINE_HINT void operator<<(FStructuredArchive::FSlot& Ar, TSetElement<ElementType>& Element)
 {
 	Ar << Element.Value;
 }
diff --git a/Engine/Source/Runtime/Core/Public/Containers/SetUtilities.h b/Engine/Source/Runtime/Core/Public/Containers/SetUtilities.h
--- a/Engine/Source/Runtime/Core/Public/Containers/SetUtilities.h
+++ b/Engine/Source/Runtime/Core/Public/Containers/SetUtilities.h
@@ -34,7 +34,7 @@ struct DefaultKeyFuncs : BaseKeyFuncs<ElementType,ElementType,bInAllowDuplicateK
 	/**
 	 * @return The key used to index the given element.
 	 */
-	[[nodiscard]] static FORCEINLINE KeyInitType GetSetKey(ElementInitType Element)
+	[[nodiscard]] static UE_FORCEINLINE_HINT KeyInitType GetSetKey(ElementInitType Element)
 	{
 		return Element;
 	}
@@ -42,7 +42,7 @@ struct DefaultKeyFuncs : BaseKeyFuncs<ElementType,ElementType,bInAllowDuplicateK
 	/**
 	 * @return True if the keys match.
 	 */
-	[[nodiscard]] static FORCEINLINE bool Matches(KeyInitType A, KeyInitType B)
+	[[nodiscard]] static UE_FORCEINLINE_HINT bool Matches(KeyInitType A, KeyInitType B)
 	{
 		return A == B;
 	}
@@ -51,20 +51,20 @@ struct DefaultKeyFuncs : BaseKeyFuncs<ElementType,ElementType,bInAllowDuplicateK
 	 * @return True if the keys match.
 	 */
 	template<typename ComparableKey>
-	[[nodiscard]] static FORCEINLINE bool Matches(KeyInitType A, ComparableKey B)
+	[[nodiscard]] static UE_FORCEINLINE_HINT bool Matches(KeyInitType A, ComparableKey B)
 	{
 		return A == B;
 	}
 
 	/** Calculates a hash index for a key. */
-	[[nodiscard]] static FORCEINLINE uint32 GetKeyHash(KeyInitType Key)
+	[[nodiscard]] static UE_FORCEINLINE_HINT uint32 GetKeyHash(KeyInitType Key)
 	{
 		return GetTypeHash(Key);
 	}
 
 	/** Calculates a hash index for a key. */
 	template<typename ComparableKey>
-	[[nodiscard]] static FORCEINLINE uint32 GetKeyHash(ComparableKey Key)
+	[[nodiscard]] static UE_FORCEINLINE_HINT uint32 GetKeyHash(ComparableKey Key)
 	{
 		return GetTypeHash(Key);
 	}
@@ -89,30 +89,30 @@ class FSetElementId
 {
 public:
 	/** Default constructor. */
-	FORCEINLINE FSetElementId() = default;
+	UE_FORCEINLINE_HINT FSetElementId() = default;
 
 	/** @return a boolean value representing whether the id is NULL. */
-	[[nodiscard]] FORCEINLINE bool IsValidId() const
+	[[nodiscard]] UE_FORCEINLINE_HINT bool IsValidId() const
 	{
 		return Index != INDEX_NONE;
 	}
 
 	/** Comparison operators. */
-	[[nodiscard]] FORCEINLINE friend bool operator==(const FSetElementId& A, const FSetElementId& B)
+	[[nodiscard]] UE_FORCEINLINE_HINT friend bool operator==(const FSetElementId& A, const FSetElementId& B)
 	{
 		return A.Index == B.Index;
 	}
-	[[nodiscard]] FORCEINLINE friend bool operator!=(const FSetElementId& A, const FSetElementId& B)
+	[[nodiscard]] UE_FORCEINLINE_HINT friend bool operator!=(const FSetElementId& A, const FSetElementId& B)
 	{
 		return A.Index != B.Index;
 	}
 
-	[[nodiscard]] FORCEINLINE int32 AsInteger() const
+	[[nodiscard]] UE_FORCEINLINE_HINT int32 AsInteger() const
 	{
 		return Index;
 	}
 
-	[[nodiscard]] FORCEINLINE static FSetElementId FromInteger(int32 Integer)
+	[[nodiscard]] UE_FORCEINLINE_HINT static FSetElementId FromInteger(int32 Integer)
 	{
 		return FSetElementId(Integer);
 	}
@@ -122,7 +122,7 @@ private:
 	int32 Index = INDEX_NONE;
 
 	/** Initialization constructor. */
-	[[nodiscard]] FORCEINLINE explicit FSetElementId(int32 InIndex)
+	[[nodiscard]] UE_FORCEINLINE_HINT explicit FSetElementId(int32 InIndex)
 	: Index(InIndex)
 	{
 	}
diff --git a/Engine/Source/Runtime/Core/Public/Containers/SortedMap.h b/Engine/Source/Runtime/Core/Public/Containers/SortedMap.h
--- a/Engine/Source/Runtime/Core/Public/Containers/SortedMap.h
+++ b/Engine/Source/Runtime/Core/Public/Containers/SortedMap.h
@@ -105,13 +105,13 @@ public:
 	}
 
 	/** Equality operator. This is efficient because pairs are always sorted. */
-	FORCEINLINE bool operator==(const TSortedMap& Other) const
+	UE_FORCEINLINE_HINT bool operator==(const TSortedMap& Other) const
 	{
 		return Pairs == Other.Pairs;
 	}
 
 	/** Inequality operator. This is efficient because pairs are always sorted. */
-	FORCEINLINE bool operator!=(const TSortedMap& Other) const
+	UE_FORCEINLINE_HINT bool operator!=(const TSortedMap& Other) const
 	{
 		return Pairs != Other.Pairs;
 	}
@@ -121,25 +121,25 @@ public:
 	 *
 	 * @param ExpectedNumElements The number of elements about to be added to the set.
 	 */
-	FORCEINLINE void Empty(int32 ExpectedNumElements = 0)
+	UE_FORCEINLINE_HINT void Empty(int32 ExpectedNumElements = 0)
 	{
 		Pairs.Empty(ExpectedNumElements);
 	}
 
     /** Efficiently empties out the map but preserves all allocations and capacities. */
-    FORCEINLINE void Reset()
+    UE_FORCEINLINE_HINT void Reset()
     {
         Pairs.Reset();
     }
 
 	/** Shrinks the pair set to avoid slack. */
-	FORCEINLINE void Shrink()
+	UE_FORCEINLINE_HINT void Shrink()
 	{
 		Pairs.Shrink();
 	}
 
 	/** Preallocates enough memory to contain Number elements. */
-	FORCEINLINE void Reserve(int32 Number)
+	UE_FORCEINLINE_HINT void Reserve(int32 Number)
 	{
 		Pairs.Reserve(Number);
 	}
@@ -156,13 +156,13 @@ public:
 	}
 
 	/** @return The number of elements in the map. */
-	[[nodiscard]] FORCEINLINE int32 Num() const
+	[[nodiscard]] UE_FORCEINLINE_HINT int32 Num() const
 	{
 		return Pairs.Num();
 	}
 
 	/** @return The number of elements the map can hold before reallocation. */
-	[[nodiscard]] FORCEINLINE int32 Max() const
+	[[nodiscard]] UE_FORCEINLINE_HINT int32 Max() const
 	{
 		return Pairs.Max();
 	}
@@ -173,13 +173,13 @@ public:
 	 *
 	 * @return number of bytes allocated by this container.
 	 */
-	[[nodiscard]] FORCEINLINE SIZE_T GetAllocatedSize() const
+	[[nodiscard]] UE_FORCEINLINE_HINT SIZE_T GetAllocatedSize() const
 	{
 		return Pairs.GetAllocatedSize();
 	}
 
 	/** Tracks the container's memory use through an archive. */
-	FORCEINLINE void CountBytes(FArchive& Ar) const
+	UE_FORCEINLINE_HINT void CountBytes(FArchive& Ar) const
 	{
 		Pairs.CountBytes(Ar);
 	}
@@ -191,10 +191,10 @@ public:
 	 * @param InValue The value to associate with the key.
 	 * @return A reference to the value as stored in the map (only valid until the next change to any key in the map).
 	 */
-	FORCEINLINE ValueType& Add(const KeyType&  InKey, const ValueType&  InValue) { return Emplace(         InKey ,          InValue ); }
-	FORCEINLINE ValueType& Add(const KeyType&  InKey,       ValueType&& InValue) { return Emplace(         InKey , MoveTemp(InValue)); }
-	FORCEINLINE ValueType& Add(      KeyType&& InKey, const ValueType&  InValue) { return Emplace(MoveTemp(InKey),          InValue ); }
-	FORCEINLINE ValueType& Add(      KeyType&& InKey,       ValueType&& InValue) { return Emplace(MoveTemp(InKey), MoveTemp(InValue)); }
+	UE_FORCEINLINE_HINT ValueType& Add(const KeyType&  InKey, const ValueType&  InValue) { return Emplace(         InKey ,          InValue ); }
+	UE_FORCEINLINE_HINT ValueType& Add(const KeyType&  InKey,       ValueType&& InValue) { return Emplace(         InKey , MoveTemp(InValue)); }
+	UE_FORCEINLINE_HINT ValueType& Add(      KeyType&& InKey, const ValueType&  InValue) { return Emplace(MoveTemp(InKey),          InValue ); }
+	UE_FORCEINLINE_HINT ValueType& Add(      KeyType&& InKey,       ValueType&& InValue) { return Emplace(MoveTemp(InKey), MoveTemp(InValue)); }
 
 	/**
 	 * Sets a default value associated with a key.
@@ -202,8 +202,8 @@ public:
 	 * @param InKey The key to associate the value with.
 	 * @return A reference to the value as stored in the map.  The reference is only valid until the next change to any key in the map.
 	 */
-	FORCEINLINE ValueType& Add(const KeyType&  InKey) { return Emplace(         InKey ); }
-	FORCEINLINE ValueType& Add(      KeyType&& InKey) { return Emplace(MoveTemp(InKey)); }
+	UE_FORCEINLINE_HINT ValueType& Add(const KeyType&  InKey) { return Emplace(         InKey ); }
+	UE_FORCEINLINE_HINT ValueType& Add(      KeyType&& InKey) { return Emplace(MoveTemp(InKey)); }
 
 	/**
 	 * Sets the value associated with a key.
@@ -294,7 +294,7 @@ public:
 		return nullptr;
 	}
 
-	[[nodiscard]] FORCEINLINE const ValueType* Find(KeyConstPointerType Key) const
+	[[nodiscard]] UE_FORCEINLINE_HINT const ValueType* Find(KeyConstPointerType Key) const
 	{
 		return const_cast<TSortedMap*>(this)->Find(Key);
 	}
@@ -306,8 +306,8 @@ public:
 	 * @param Key The key to search for.
 	 * @return A reference to the value associated with the specified key.
 	 */
-	FORCEINLINE ValueType& FindOrAdd(const KeyType&  Key) { return FindOrAddImpl(         Key ); }
-	FORCEINLINE ValueType& FindOrAdd(      KeyType&& Key) { return FindOrAddImpl(MoveTemp(Key)); }
+	UE_FORCEINLINE_HINT ValueType& FindOrAdd(const KeyType&  Key) { return FindOrAddImpl(         Key ); }
+	UE_FORCEINLINE_HINT ValueType& FindOrAdd(      KeyType&& Key) { return FindOrAddImpl(MoveTemp(Key)); }
 
 	/**
 	 * Returns a reference to the value associated with a specified key.
@@ -512,11 +512,11 @@ public:
 		}
 	}
 
-	[[nodiscard]] FORCEINLINE ValueType& operator[](KeyConstPointerType Key)
+	[[nodiscard]] UE_FORCEINLINE_HINT ValueType& operator[](KeyConstPointerType Key)
 	{
 		return this->FindChecked(Key);
 	}
-	[[nodiscard]] FORCEINLINE const ValueType& operator[](KeyConstPointerType Key) const
+	[[nodiscard]] UE_FORCEINLINE_HINT const ValueType& operator[](KeyConstPointerType Key) const
 	{
 		return this->FindChecked(Key);
 	}
@@ -525,7 +525,7 @@ public:
 
 	/** @return The non-inclusive maximum index of elements in the map (was: the max valid index of the elements in the sparse storage). */
 	UE_DEPRECATED(5.7 until 5.9, "GetMaxIndex() should be replaced with Num() - 1.")
-	[[nodiscard]] FORCEINLINE int32 GetMaxIndex() const
+	[[nodiscard]] UE_FORCEINLINE_HINT int32 GetMaxIndex() const
 	{
 		// We don't want to deprecate the function entirely, because it should mirror the the existing TMap API, but it's been
 		// returning the wrong value and we want to fix that.
@@ -539,19 +539,19 @@ public:
 	 * @param Id - The element id to check.
 	 * @return true if the element identifier refers to a valid element in this map.
 	 */
-	[[nodiscard]] FORCEINLINE bool IsValidId(FSetElementId Id) const
+	[[nodiscard]] UE_FORCEINLINE_HINT bool IsValidId(FSetElementId Id) const
 	{
 		return Pairs.IsValidIndex(Id.AsInteger());
 	}
 
 	/** Return a mapped pair by internal identifier. Element must be valid (see @IsValidId). */
-	[[nodiscard]] FORCEINLINE ElementType& Get(FSetElementId Id)
+	[[nodiscard]] UE_FORCEINLINE_HINT ElementType& Get(FSetElementId Id)
 	{
 		return Pairs[Id.AsInteger()];
 	}
 
 	/** Return a mapped pair by internal identifier.  Element must be valid (see @IsValidId).*/
-	[[nodiscard]] FORCEINLINE const ElementType& Get(FSetElementId Id) const
+	[[nodiscard]] UE_FORCEINLINE_HINT const ElementType& Get(FSetElementId Id) const
 	{
 		return Pairs[Id.AsInteger()];
 	}
@@ -572,7 +572,7 @@ private:
 	}
 
 	/** Find index of key */
-	[[nodiscard]] FORCEINLINE int32 FindIndex(KeyConstPointerType Key) const
+	[[nodiscard]] UE_FORCEINLINE_HINT int32 FindIndex(KeyConstPointerType Key) const
 	{
 		return Algo::BinarySearchBy(Pairs, Key, FKeyForward(), SortPredicate());
 	}
@@ -605,7 +605,7 @@ private:
 	/** Forwards sorting into Key of pair */
 	struct FKeyForward
 	{
-		FORCEINLINE const KeyType& operator()(const ElementType& Pair) const
+		UE_FORCEINLINE_HINT const KeyType& operator()(const ElementType& Pair) const
 		{
 			return Pair.Key;
 		}
@@ -624,52 +624,52 @@ private:
 		typedef std::conditional_t<bConst,const typename ElementArrayType::ElementType, typename ElementArrayType::ElementType> PairType;
 
 	protected:
-		[[nodiscard]] FORCEINLINE TBaseIterator(const PairItType& InElementIt)
+		[[nodiscard]] UE_FORCEINLINE_HINT TBaseIterator(const PairItType& InElementIt)
 			: PairIt(InElementIt)
 		{
 		}
 
 	public:
-		FORCEINLINE TBaseIterator& operator++()
+		UE_FORCEINLINE_HINT TBaseIterator& operator++()
 		{
 			++PairIt;
 			return *this;
 		}
 
 		/** conversion to "bool" returning true if the iterator is valid. */
-		[[nodiscard]] FORCEINLINE explicit operator bool() const
+		[[nodiscard]] UE_FORCEINLINE_HINT explicit operator bool() const
 		{
 			return !!PairIt; 
 		}
 
-		[[nodiscard]] FORCEINLINE bool operator==(const TBaseIterator& Rhs) const
+		[[nodiscard]] UE_FORCEINLINE_HINT bool operator==(const TBaseIterator& Rhs) const
 		{
 			return PairIt == Rhs.PairIt;
 		}
-		[[nodiscard]] FORCEINLINE bool operator!=(const TBaseIterator& Rhs) const
+		[[nodiscard]] UE_FORCEINLINE_HINT bool operator!=(const TBaseIterator& Rhs) const
 		{
 			return PairIt != Rhs.PairIt;
 		}
 
-		FORCEINLINE ItKeyType& Key() const
+		UE_FORCEINLINE_HINT ItKeyType& Key() const
 		{
 			return PairIt->Key;
 		}
-		FORCEINLINE ItValueType& Value() const
+		UE_FORCEINLINE_HINT ItValueType& Value() const
 		{
 			return PairIt->Value;
 		}
 
-		[[nodiscard]] FORCEINLINE FSetElementId GetId() const
+		[[nodiscard]] UE_FORCEINLINE_HINT FSetElementId GetId() const
 		{
 			return FSetElementId::FromInteger(PairIt.GetIndex());
 		}
 
-		FORCEINLINE PairType& operator*() const
+		UE_FORCEINLINE_HINT PairType& operator*() const
 		{
 			return  *PairIt;
 		}
-		FORCEINLINE PairType* operator->() const
+		UE_FORCEINLINE_HINT PairType* operator->() const
 		{
 			return &*PairIt;
 		}
@@ -693,7 +693,7 @@ private:
 		typedef std::conditional_t<bConst, const typename ElementArrayType::ElementType, typename ElementArrayType::ElementType> PairType;
 
 	protected:
-		[[nodiscard]] FORCEINLINE TBaseReverseIterator(PairType* InData, SizeType InNum)
+		[[nodiscard]] UE_FORCEINLINE_HINT TBaseReverseIterator(PairType* InData, SizeType InNum)
 			: Data(InData), Index(InNum-1)
 		{
 		}
@@ -707,39 +707,39 @@ private:
 		}
 
 		/** conversion to "bool" returning true if the iterator is valid. */
-		[[nodiscard]] FORCEINLINE explicit operator bool() const
+		[[nodiscard]] UE_FORCEINLINE_HINT explicit operator bool() const
 		{
 			return Index != static_cast<SizeType>(-1);
 		}
 
-		[[nodiscard]] FORCEINLINE bool operator==(const TBaseReverseIterator& Rhs) const
+		[[nodiscard]] UE_FORCEINLINE_HINT bool operator==(const TBaseReverseIterator& Rhs) const
 		{
 			return Index == Rhs.Index;
 		}
-		[[nodiscard]] FORCEINLINE bool operator!=(const TBaseReverseIterator& Rhs) const
+		[[nodiscard]] UE_FORCEINLINE_HINT bool operator!=(const TBaseReverseIterator& Rhs) const
 		{
 			return Index != Rhs.Index;
 		}
 
-		FORCEINLINE ItKeyType& Key() const
+		UE_FORCEINLINE_HINT ItKeyType& Key() const
 		{
 			return Data[Index].Key;
 		}
-		FORCEINLINE ItValueType& Value() const
+		UE_FORCEINLINE_HINT ItValueType& Value() const
 		{
 			return Data[Index].Value;
 		}
 
-		[[nodiscard]] FORCEINLINE FSetElementId GetId() const
+		[[nodiscard]] UE_FORCEINLINE_HINT FSetElementId GetId() const
 		{
 			return FSetElementId::FromInteger(Index);
 		}
 
-		FORCEINLINE PairType& operator*() const
+		UE_FORCEINLINE_HINT PairType& operator*() const
 		{
 			return  Data[Index];
 		}
-		FORCEINLINE PairType* operator->() const
+		UE_FORCEINLINE_HINT PairType* operator->() const
 		{
 			return &Data[Index];
 		}
@@ -759,18 +759,18 @@ public:
 	{
 	public:
 
-		[[nodiscard]] FORCEINLINE TIterator(TSortedMap& InMap)
+		[[nodiscard]] UE_FORCEINLINE_HINT TIterator(TSortedMap& InMap)
 			: TBaseIterator<false>(InMap.Pairs.CreateIterator())
 		{
 		}
 
-		[[nodiscard]] FORCEINLINE TIterator(const typename TBaseIterator<false>::PairItType& InPairIt)
+		[[nodiscard]] UE_FORCEINLINE_HINT TIterator(const typename TBaseIterator<false>::PairItType& InPairIt)
 			: TBaseIterator<false>(InPairIt)
 		{
 		}
 
 		/** Removes the current pair from the map */
-		FORCEINLINE void RemoveCurrent()
+		UE_FORCEINLINE_HINT void RemoveCurrent()
 		{
 			TBaseIterator<false>::PairIt.RemoveCurrent();
 		}
@@ -780,12 +780,12 @@ public:
 	class TConstIterator : public TBaseIterator<true>
 	{
 	public:
-		[[nodiscard]] FORCEINLINE TConstIterator(const TSortedMap& InMap)
+		[[nodiscard]] UE_FORCEINLINE_HINT TConstIterator(const TSortedMap& InMap)
 			: TBaseIterator<true>(InMap.Pairs.CreateConstIterator())
 		{
 		}
 
-		[[nodiscard]] FORCEINLINE TConstIterator(const typename TBaseIterator<true>::PairItType& InPairIt)
+		[[nodiscard]] UE_FORCEINLINE_HINT TConstIterator(const typename TBaseIterator<true>::PairItType& InPairIt)
 			: TBaseIterator<true>(InPairIt)
 		{
 		}
@@ -795,7 +795,7 @@ public:
 	class TReverseIterator : public TBaseReverseIterator<false>
 	{
 	public:
-		[[nodiscard]] FORCEINLINE TReverseIterator(TSortedMap& InMap)
+		[[nodiscard]] UE_FORCEINLINE_HINT TReverseIterator(TSortedMap& InMap)
 			: TBaseReverseIterator<false>(InMap.Pairs.GetData(), InMap.Pairs.Num())
 		{
 		}
@@ -807,7 +807,7 @@ public:
 	class TConstReverseIterator : public TBaseReverseIterator<true>
 	{
 	public:
-		[[nodiscard]] FORCEINLINE TConstReverseIterator(const TSortedMap& InMap)
+		[[nodiscard]] UE_FORCEINLINE_HINT TConstReverseIterator(const TSortedMap& InMap)
 			: TBaseReverseIterator<true>(InMap.Pairs.GetData(), InMap.Pairs.Num())
 		{
 		}
@@ -879,25 +879,25 @@ public:
 	};
 
 	/** Creates an iterator over all the pairs in this map */
-	[[nodiscard]] FORCEINLINE TIterator CreateIterator()
+	[[nodiscard]] UE_FORCEINLINE_HINT TIterator CreateIterator()
 	{
 		return TIterator(*this);
 	}
 
 	/** Creates a const iterator over all the pairs in this map */
-	[[nodiscard]] FORCEINLINE TConstIterator CreateConstIterator() const
+	[[nodiscard]] UE_FORCEINLINE_HINT TConstIterator CreateConstIterator() const
 	{
 		return TConstIterator(*this);
 	}
 
 	/** Creates an iterator over the values associated with a specified key in a map. This will be at most one value because keys must be unique */
-	[[nodiscard]] FORCEINLINE TKeyIterator CreateKeyIterator(KeyInitType InKey)
+	[[nodiscard]] UE_FORCEINLINE_HINT TKeyIterator CreateKeyIterator(KeyInitType InKey)
 	{
 		return TKeyIterator(*this, InKey);
 	}
 
 	/** Creates a const iterator over the values associated with a specified key in a map. This will be at most one value because keys must be unique */
-	[[nodiscard]] FORCEINLINE TConstKeyIterator CreateConstKeyIterator(KeyInitType InKey) const
+	[[nodiscard]] UE_FORCEINLINE_HINT TConstKeyIterator CreateConstKeyIterator(KeyInitType InKey) const
 	{
 		return TConstKeyIterator(*this, InKey);
 	}
@@ -911,10 +911,10 @@ public:
 	 * DO NOT USE DIRECTLY
 	 * STL-like iterators to enable range-based for loop support.
 	 */
-	[[nodiscard]] FORCEINLINE RangedForIteratorType		begin()	      { return Pairs.begin(); }
-	[[nodiscard]] FORCEINLINE RangedForConstIteratorType	begin() const { return Pairs.begin(); }
-	[[nodiscard]] FORCEINLINE RangedForIteratorType		end()         { return Pairs.end(); }
-	[[nodiscard]] FORCEINLINE RangedForConstIteratorType	end() const   { return Pairs.end(); }
+	[[nodiscard]] UE_FORCEINLINE_HINT RangedForIteratorType		begin()	      { return Pairs.begin(); }
+	[[nodiscard]] UE_FORCEINLINE_HINT RangedForConstIteratorType	begin() const { return Pairs.begin(); }
+	[[nodiscard]] UE_FORCEINLINE_HINT RangedForIteratorType		end()         { return Pairs.end(); }
+	[[nodiscard]] UE_FORCEINLINE_HINT RangedForConstIteratorType	end() const   { return Pairs.end(); }
 
 	friend struct TSortedMapPrivateFriend;
 };
diff --git a/Engine/Source/Runtime/Core/Public/Containers/SortedSet.h b/Engine/Source/Runtime/Core/Public/Containers/SortedSet.h
--- a/Engine/Source/Runtime/Core/Public/Containers/SortedSet.h
+++ b/Engine/Source/Runtime/Core/Public/Containers/SortedSet.h
@@ -102,13 +102,13 @@ public:
 	}
 
 	/** Equality operator. This is efficient because elements are always sorted. */
-	FORCEINLINE bool operator==(const TSortedSet& Other) const
+	UE_FORCEINLINE_HINT bool operator==(const TSortedSet& Other) const
 	{
 		return Elements == Other.Elements;
 	}
 
 	/** Inequality operator. This is efficient because elements are always sorted. */
-	FORCEINLINE bool operator!=(const TSortedSet& Other) const
+	UE_FORCEINLINE_HINT bool operator!=(const TSortedSet& Other) const
 	{
 		return Elements != Other.Elements;
 	}
@@ -118,25 +118,25 @@ public:
 	 *
 	 * @param ExpectedNumElements The number of elements about to be added to the set.
 	 */
-	FORCEINLINE void Empty(int32 ExpectedNumElements = 0)
+	UE_FORCEINLINE_HINT void Empty(int32 ExpectedNumElements = 0)
 	{
 		Elements.Empty(ExpectedNumElements);
 	}
 
 	/** Efficiently empties out the set but preserves all allocations and capacities. */
-	FORCEINLINE void Reset()
+	UE_FORCEINLINE_HINT void Reset()
 	{
 		Elements.Reset();
 	}
 
 	/** Shrinks the set to avoid slack. */
-	FORCEINLINE void Shrink()
+	UE_FORCEINLINE_HINT void Shrink()
 	{
 		Elements.Shrink();
 	}
 
 	/** Preallocates enough memory to contain Number elements. */
-	FORCEINLINE void Reserve(int32 Number)
+	UE_FORCEINLINE_HINT void Reserve(int32 Number)
 	{
 		Elements.Reserve(Number);
 	}
@@ -153,19 +153,19 @@ public:
 	}
 
 	/** @return The number of elements in the set. */
-	[[nodiscard]] FORCEINLINE int32 Num() const
+	[[nodiscard]] UE_FORCEINLINE_HINT int32 Num() const
 	{
 		return Elements.Num();
 	}
 
 	/** @return The number of elements the set can hold before reallocation. */
-	[[nodiscard]] FORCEINLINE int32 Max() const
+	[[nodiscard]] UE_FORCEINLINE_HINT int32 Max() const
 	{
 		return Elements.Max();
 	}
 
 	/** @return The non-inclusive maximum index of elements in the set. */
-	[[nodiscard]] FORCEINLINE int32 GetMaxIndex() const
+	[[nodiscard]] UE_FORCEINLINE_HINT int32 GetMaxIndex() const
 	{
 		return Elements.Num();
 	}
@@ -176,13 +176,13 @@ public:
 	 *
 	 * @return number of bytes allocated by this container.
 	 */
-	[[nodiscard]] FORCEINLINE SIZE_T GetAllocatedSize() const
+	[[nodiscard]] UE_FORCEINLINE_HINT SIZE_T GetAllocatedSize() const
 	{
 		return Elements.GetAllocatedSize();
 	}
 
 	/** Tracks the container's memory use through an archive. */
-	FORCEINLINE void CountBytes(FArchive& Ar) const
+	UE_FORCEINLINE_HINT void CountBytes(FArchive& Ar) const
 	{
 		Elements.CountBytes(Ar);
 	}
@@ -194,11 +194,11 @@ public:
 	 * @param	bIsAlreadyInSetPtr	[out]	Optional pointer to bool that will be set depending on whether element is already in set
 	 * @return A handle to the value in the set (only valid until the next change to the set).
 	 */
-	FORCEINLINE FSetElementId Add(const ElementType& InValue, bool* bIsAlreadyInSetPtr = nullptr)
+	UE_FORCEINLINE_HINT FSetElementId Add(const ElementType& InValue, bool* bIsAlreadyInSetPtr = nullptr)
 	{
 		return Emplace(InValue, bIsAlreadyInSetPtr);
 	}
-	FORCEINLINE FSetElementId Add(ElementType&& InValue, bool* bIsAlreadyInSetPtr = nullptr)
+	UE_FORCEINLINE_HINT FSetElementId Add(ElementType&& InValue, bool* bIsAlreadyInSetPtr = nullptr)
 	{
 		return Emplace(MoveTemp(InValue), bIsAlreadyInSetPtr);
 	}
@@ -226,7 +226,7 @@ public:
 	 * @param	InKey	The key to match elements against.
 	 * @return	The number of elements removed.
 	 */
-	FORCEINLINE int32 Remove(KeyInitType InKey)
+	UE_FORCEINLINE_HINT int32 Remove(KeyInitType InKey)
 	{
 		int32 RemoveIndex = FindIndex(InKey);
 		if (RemoveIndex == INDEX_NONE)
@@ -244,7 +244,7 @@ public:
 	 * @param Key The key to search for.
 	 * @return A pointer to the value associated with the specified key, or nullptr if the key isn't contained in this set.  The pointer (only valid until the next change to any key in the set).
 	 */
-	[[nodiscard]] FORCEINLINE ElementType* Find(KeyInitType Key)
+	[[nodiscard]] UE_FORCEINLINE_HINT ElementType* Find(KeyInitType Key)
 	{
 		int32 FoundIndex = FindIndex(Key);
 		if (FoundIndex != INDEX_NONE)
@@ -254,7 +254,7 @@ public:
 
 		return nullptr;
 	}
-	[[nodiscard]] FORCEINLINE const ElementType* Find(KeyInitType Key) const
+	[[nodiscard]] UE_FORCEINLINE_HINT const ElementType* Find(KeyInitType Key) const
 	{
 		return const_cast<TSortedSet*>(this)->Find(Key);
 	}
@@ -266,11 +266,11 @@ public:
 	 * @param	bIsAlreadyInSetPtr	[out]	Optional pointer to bool that will be set depending on whether element is already in set
 	 * @return A handle to the value in the set (only valid until the next change to the set).
 	 */
-	FORCEINLINE ElementType& FindOrAdd(const ElementType& InValue, bool* bIsAlreadyInSetPtr = nullptr)
+	UE_FORCEINLINE_HINT ElementType& FindOrAdd(const ElementType& InValue, bool* bIsAlreadyInSetPtr = nullptr)
 	{
 		return FindOrAddImpl(InValue, bIsAlreadyInSetPtr);
 	}
-	FORCEINLINE ElementType& FindOrAdd(ElementType&& InValue, bool* bIsAlreadyInSetPtr = nullptr)
+	UE_FORCEINLINE_HINT ElementType& FindOrAdd(ElementType&& InValue, bool* bIsAlreadyInSetPtr = nullptr)
 	{
 		return FindOrAddImpl(MoveTemp(InValue), bIsAlreadyInSetPtr);
 	}
@@ -297,7 +297,7 @@ public:
 	 * @param Key The key to check for.
 	 * @return true if the set contains the key.
 	 */
-	[[nodiscard]] FORCEINLINE bool Contains(KeyInitType Key) const
+	[[nodiscard]] UE_FORCEINLINE_HINT bool Contains(KeyInitType Key) const
 	{
 		return FindIndex(Key) != INDEX_NONE;
 	}
@@ -473,24 +473,24 @@ public:
 	 * @param Id - The element id to check.
 	 * @return true if the element identifier refers to a valid element in this set.
 	 */
-	[[nodiscard]] FORCEINLINE bool IsValidId(FSetElementId Id) const
+	[[nodiscard]] UE_FORCEINLINE_HINT bool IsValidId(FSetElementId Id) const
 	{
 		return Elements.IsValidIndex(Id.AsInteger());
 	}
 
-	[[nodiscard]] FORCEINLINE ElementType& operator[](FSetElementId Id)
+	[[nodiscard]] UE_FORCEINLINE_HINT ElementType& operator[](FSetElementId Id)
 	{
 		return Elements[Id.AsInteger()];
 	}
-	[[nodiscard]] FORCEINLINE const ElementType& operator[](FSetElementId Id) const
+	[[nodiscard]] UE_FORCEINLINE_HINT const ElementType& operator[](FSetElementId Id) const
 	{
 		return (*const_cast<TSortedSet*>(this))[Id];
 	}
-	[[nodiscard]] FORCEINLINE ElementType& Get(FSetElementId Id)
+	[[nodiscard]] UE_FORCEINLINE_HINT ElementType& Get(FSetElementId Id)
 	{
 		return (*this)[Id];
 	}
-	[[nodiscard]] FORCEINLINE const ElementType& Get(FSetElementId Id) const
+	[[nodiscard]] UE_FORCEINLINE_HINT const ElementType& Get(FSetElementId Id) const
 	{
 		return (*this)[Id];
 	}
@@ -500,7 +500,7 @@ private:
 
 	/** Implementation of find and add */
 	template <typename InitArgType>
-	[[nodiscard]] FORCEINLINE ElementType& FindOrAddImpl(InitArgType&& InInitArg, bool* bIsAlreadyInSetPtr)
+	[[nodiscard]] UE_FORCEINLINE_HINT ElementType& FindOrAddImpl(InitArgType&& InInitArg, bool* bIsAlreadyInSetPtr)
 	{
 		int32 InsertIndex = Algo::LowerBound(Elements, InInitArg, SortPredicate());
 
@@ -522,14 +522,14 @@ private:
 	}
 
 	/** Find index of key */
-	[[nodiscard]] FORCEINLINE int32 FindIndex(KeyInitType Key) const
+	[[nodiscard]] UE_FORCEINLINE_HINT int32 FindIndex(KeyInitType Key) const
 	{
 		return Algo::BinarySearch(Elements, Key, SortPredicate());
 	}
 
 	/** Allocates raw memory for emplacing */
 	template <typename InitArgType>
-	FORCEINLINE ElementType* AllocateMemoryForEmplace(InitArgType&& InInitArg, bool* bIsAlreadyInSetPtr = nullptr)
+	UE_FORCEINLINE_HINT ElementType* AllocateMemoryForEmplace(InitArgType&& InInitArg, bool* bIsAlreadyInSetPtr = nullptr)
 	{
 		int32 InsertIndex = Algo::LowerBound(Elements, InInitArg, SortPredicate());
 
@@ -567,43 +567,43 @@ private:
 		using ItElementType = std::conditional_t<bConst, const ElementType, ElementType>;
 
 	protected:
-		[[nodiscard]] FORCEINLINE TBaseIterator(const ArrayIteratorType& InArrayIt)
+		[[nodiscard]] UE_FORCEINLINE_HINT TBaseIterator(const ArrayIteratorType& InArrayIt)
 			: ArrayIt(InArrayIt)
 		{
 		}
 
 	public:
-		FORCEINLINE TBaseIterator& operator++()
+		UE_FORCEINLINE_HINT TBaseIterator& operator++()
 		{
 			++ArrayIt;
 			return *this;
 		}
 
 		/** conversion to "bool" returning true if the iterator is valid. */
-		[[nodiscard]] FORCEINLINE explicit operator bool() const
+		[[nodiscard]] UE_FORCEINLINE_HINT explicit operator bool() const
 		{
 			return !!ArrayIt; 
 		}
 
-		[[nodiscard]] FORCEINLINE bool operator==(const TBaseIterator& Rhs) const
+		[[nodiscard]] UE_FORCEINLINE_HINT bool operator==(const TBaseIterator& Rhs) const
 		{
 			return ArrayIt == Rhs.ArrayIt;
 		}
-		[[nodiscard]] FORCEINLINE bool operator!=(const TBaseIterator& Rhs) const
+		[[nodiscard]] UE_FORCEINLINE_HINT bool operator!=(const TBaseIterator& Rhs) const
 		{
 			return !(*this == Rhs);
 		}
 
-		[[nodiscard]] FORCEINLINE FSetElementId GetId() const
+		[[nodiscard]] UE_FORCEINLINE_HINT FSetElementId GetId() const
 		{
 			return FSetElementId::FromInteger(ArrayIt.GetIndex());
 		}
 
-		FORCEINLINE ItElementType& operator*() const
+		UE_FORCEINLINE_HINT ItElementType& operator*() const
 		{
 			return *ArrayIt;
 		}
-		FORCEINLINE ItElementType* operator->() const
+		UE_FORCEINLINE_HINT ItElementType* operator->() const
 		{
 			return &*ArrayIt;
 		}
@@ -620,13 +620,13 @@ public:
 	class TIterator : public TBaseIterator<false>
 	{
 	public:
-		[[nodiscard]] FORCEINLINE explicit TIterator(TSortedSet& InSet)
+		[[nodiscard]] UE_FORCEINLINE_HINT explicit TIterator(TSortedSet& InSet)
 			: TBaseIterator<false>(InSet.Elements.CreateIterator())
 		{
 		}
 
 		/** Removes the current element from the set */
-		FORCEINLINE void RemoveCurrent()
+		UE_FORCEINLINE_HINT void RemoveCurrent()
 		{
 			TBaseIterator<false>::ArrayIt.RemoveCurrent();
 		}
@@ -636,20 +636,20 @@ public:
 	class TConstIterator : public TBaseIterator<true>
 	{
 	public:
-		[[nodiscard]] FORCEINLINE explicit TConstIterator(const TSortedSet& InSet)
+		[[nodiscard]] UE_FORCEINLINE_HINT explicit TConstIterator(const TSortedSet& InSet)
 			: TBaseIterator<true>(InSet.Elements.CreateConstIterator())
 		{
 		}
 	};
 
 	/** Creates an iterator over all the elements in this set */
-	[[nodiscard]] FORCEINLINE TIterator CreateIterator()
+	[[nodiscard]] UE_FORCEINLINE_HINT TIterator CreateIterator()
 	{
 		return TIterator(*this);
 	}
 
 	/** Creates a const iterator over all the elements in this set */
-	[[nodiscard]] FORCEINLINE TConstIterator CreateConstIterator() const
+	[[nodiscard]] UE_FORCEINLINE_HINT TConstIterator CreateConstIterator() const
 	{
 		return TConstIterator(*this);
 	}
@@ -659,19 +659,19 @@ public:
 	 * DO NOT USE DIRECTLY
 	 * STL-like iterators to enable range-based for loop support.
 	 */
-	[[nodiscard]] FORCEINLINE auto begin()
+	[[nodiscard]] UE_FORCEINLINE_HINT auto begin()
 	{
 		return Elements.begin();
 	}
-	[[nodiscard]] FORCEINLINE auto begin() const
+	[[nodiscard]] UE_FORCEINLINE_HINT auto begin() const
 	{
 		return Elements.begin();
 	}
-	[[nodiscard]] FORCEINLINE auto end()
+	[[nodiscard]] UE_FORCEINLINE_HINT auto end()
 	{
 		return Elements.end();
 	}
-	[[nodiscard]] FORCEINLINE auto end() const
+	[[nodiscard]] UE_FORCEINLINE_HINT auto end() const
 	{
 		return Elements.end();
 	}
@@ -680,7 +680,7 @@ public:
 DECLARE_TEMPLATE_INTRINSIC_TYPE_LAYOUT((template <typename ElementType, typename ArrayAllocator, typename SortPredicate>), (TSortedSet<ElementType, ArrayAllocator, SortPredicate>));
 
 template <typename ElementType, typename ArrayAllocator, typename SortPredicate>
-FORCEINLINE FArchive& operator<<(FArchive& Ar, TSortedSet<ElementType, ArrayAllocator, SortPredicate>& Set)
+UE_FORCEINLINE_HINT FArchive& operator<<(FArchive& Ar, TSortedSet<ElementType, ArrayAllocator, SortPredicate>& Set)
 {
 	Ar << Set.Elements;
 
