Replace some FORCEINLINE usage in Core/CoreUObject with UE_FORCEINLINE_HINT.

diff --git a/Engine/Source/Runtime/Core/Public/GenericPlatform/GenericPlatformProperties.h b/Engine/Source/Runtime/Core/Public/GenericPlatform/GenericPlatformProperties.h
--- a/Engine/Source/Runtime/Core/Public/GenericPlatform/GenericPlatformProperties.h
+++ b/Engine/Source/Runtime/Core/Public/GenericPlatform/GenericPlatformProperties.h
@@ -20,7 +20,7 @@ struct FGenericPlatformProperties
 	 *
 	 * @return The physics format name.
 	 */
-	static constexpr FORCEINLINE const char* GetPhysicsFormat()
+	static constexpr UE_FORCEINLINE_HINT const char* GetPhysicsFormat()
 	{
 		return "Chaos";
 	}
@@ -30,7 +30,7 @@ struct FGenericPlatformProperties
 	 *
 	 * @return true if the platform has Editor-only data, false otherwise.
 	 */
-	static constexpr FORCEINLINE bool HasEditorOnlyData()
+	static constexpr UE_FORCEINLINE_HINT bool HasEditorOnlyData()
 	{
 		return WITH_EDITORONLY_DATA;
 	}
@@ -49,7 +49,7 @@ struct FGenericPlatformProperties
 	 *
 	 * @return true if this is a game only platform, false otherwise.
 	 */
-	static constexpr FORCEINLINE bool IsGameOnly()
+	static constexpr UE_FORCEINLINE_HINT bool IsGameOnly()
 	{
 		return UE_GAME;
 	}
@@ -59,7 +59,7 @@ struct FGenericPlatformProperties
 	 *
 	 * @return true if this is a server only platform, false otherwise.
 	 */
-	static constexpr FORCEINLINE bool IsServerOnly()
+	static constexpr UE_FORCEINLINE_HINT bool IsServerOnly()
 	{
 		return UE_SERVER;
 	}
@@ -69,7 +69,7 @@ struct FGenericPlatformProperties
 	 *
 	 * @return true if this is a client only platform, false otherwise.
 	 */
-	static constexpr FORCEINLINE bool IsClientOnly()
+	static constexpr UE_FORCEINLINE_HINT bool IsClientOnly()
 	{
 		return !WITH_SERVER_CODE;
 	}
@@ -77,7 +77,7 @@ struct FGenericPlatformProperties
 	/**
 	 *	Gets whether this was a monolithic build or not
 	 */
-	static constexpr FORCEINLINE bool IsMonolithicBuild()
+	static constexpr UE_FORCEINLINE_HINT bool IsMonolithicBuild()
 	{
 		return IS_MONOLITHIC;
 	}
@@ -85,7 +85,7 @@ struct FGenericPlatformProperties
 	/**
 	 *	Gets whether this was a program or not
 	 */
-	static constexpr FORCEINLINE bool IsProgram()
+	static constexpr UE_FORCEINLINE_HINT bool IsProgram()
 	{
 		return IS_PROGRAM;
 	}
@@ -95,7 +95,7 @@ struct FGenericPlatformProperties
 	 *
 	 * @return true if the platform is Little Endian, false otherwise.
 	 */
-	static constexpr FORCEINLINE bool IsLittleEndian()
+	static constexpr UE_FORCEINLINE_HINT bool IsLittleEndian()
 	{
 		return true;
 	}
@@ -107,7 +107,7 @@ struct FGenericPlatformProperties
 	 *
 	 * @return Platform Name.
 	 */
-	static FORCEINLINE const char* PlatformName();
+	static UE_FORCEINLINE_HINT const char* PlatformName();
 
 	/**
 	  * Get the name of the hardware variant of the current platform.
@@ -117,7 +117,7 @@ struct FGenericPlatformProperties
 	  * 
 	  * @return Name of the platform variant.
 	  */
-	static FORCEINLINE const char* PlatformVariantName()
+	static UE_FORCEINLINE_HINT const char* PlatformVariantName()
 	{
 		return "";
 	}
@@ -127,7 +127,7 @@ struct FGenericPlatformProperties
 	 *
 	 * @return true if cooked data is required, false otherwise.
 	 */
-	static constexpr FORCEINLINE bool RequiresCookedData()
+	static constexpr UE_FORCEINLINE_HINT bool RequiresCookedData()
 	{
 		return !HasEditorOnlyData();
 	}
@@ -137,7 +137,7 @@ struct FGenericPlatformProperties
 	*
 	* @return true if packaged data is considered secure, false otherwise.
 	*/
-	static constexpr FORCEINLINE bool HasSecurePackageFormat()
+	static constexpr UE_FORCEINLINE_HINT bool HasSecurePackageFormat()
 	{
 		return false;
 	}
@@ -147,7 +147,7 @@ struct FGenericPlatformProperties
 	 *
 	 * @return true if this platform requires user credentials, false otherwise.
 	 */
-	static constexpr FORCEINLINE bool RequiresUserCredentials()
+	static constexpr UE_FORCEINLINE_HINT bool RequiresUserCredentials()
 	{
 		return false;
 	}
@@ -158,7 +158,7 @@ struct FGenericPlatformProperties
 	 * @param TargetType The build target to check.
 	 * @return true if the build target is supported, false otherwise.
 	 */
-	static FORCEINLINE bool SupportsBuildTarget(EBuildTargetType TargetType)
+	static UE_FORCEINLINE_HINT bool SupportsBuildTarget(EBuildTargetType TargetType)
 	{
 		return true;
 	}
@@ -166,7 +166,7 @@ struct FGenericPlatformProperties
 	/**
 	 * Returns true if platform supports the AutoSDK system
 	 */
-	static constexpr FORCEINLINE bool SupportsAutoSDK()
+	static constexpr UE_FORCEINLINE_HINT bool SupportsAutoSDK()
 	{
 		return false;
 	}
@@ -176,7 +176,7 @@ struct FGenericPlatformProperties
 	 *
 	 * @return true if gray scale sRGB texture formats are supported.
 	 */
-	static constexpr FORCEINLINE bool SupportsGrayscaleSRGB()
+	static constexpr UE_FORCEINLINE_HINT bool SupportsGrayscaleSRGB()
 	{
 		return true;
 	}
@@ -186,7 +186,7 @@ struct FGenericPlatformProperties
 	 *
 	 * @return true if multiple instances are supported, false otherwise.
 	 */
-	static constexpr FORCEINLINE bool SupportsMultipleGameInstances()
+	static constexpr UE_FORCEINLINE_HINT bool SupportsMultipleGameInstances()
 	{
 		return false;
 	}
@@ -196,7 +196,7 @@ struct FGenericPlatformProperties
 	 *
 	 * @return true if windowed mode is supported.
 	 */
-	static constexpr FORCEINLINE bool SupportsWindowedMode()
+	static constexpr UE_FORCEINLINE_HINT bool SupportsWindowedMode()
 	{
 		return false;
 	}
@@ -204,7 +204,7 @@ struct FGenericPlatformProperties
 	/**
 	 * Whether this platform wants to allow framerate smoothing or not.
 	 */
-	static constexpr FORCEINLINE bool AllowsFramerateSmoothing()
+	static constexpr UE_FORCEINLINE_HINT bool AllowsFramerateSmoothing()
 	{
 		return true;
 	}
@@ -212,78 +212,78 @@ struct FGenericPlatformProperties
 	/**
 	 * Whether this platform supports streaming audio
 	 */
-	static constexpr FORCEINLINE bool SupportsAudioStreaming()
+	static constexpr UE_FORCEINLINE_HINT bool SupportsAudioStreaming()
 	{
 		return false;
 	}
 
-	static constexpr FORCEINLINE bool SupportsHighQualityLightmaps()
+	static constexpr UE_FORCEINLINE_HINT bool SupportsHighQualityLightmaps()
 	{
 		return true;
 	}
 
-	static constexpr FORCEINLINE bool SupportsLowQualityLightmaps()
+	static constexpr UE_FORCEINLINE_HINT bool SupportsLowQualityLightmaps()
 	{
 		return true;
 	}
 
-	static constexpr FORCEINLINE bool SupportsDistanceFieldShadows()
+	static constexpr UE_FORCEINLINE_HINT bool SupportsDistanceFieldShadows()
 	{
 		return true;
 	}
 
-	static constexpr FORCEINLINE bool SupportsDistanceFieldAO()
+	static constexpr UE_FORCEINLINE_HINT bool SupportsDistanceFieldAO()
 	{
 		return true;
 	}
 
-	static constexpr FORCEINLINE bool SupportsTextureStreaming()
+	static constexpr UE_FORCEINLINE_HINT bool SupportsTextureStreaming()
 	{
 		return true;
 	}
 
-	static constexpr FORCEINLINE bool SupportsMeshLODStreaming()
+	static constexpr UE_FORCEINLINE_HINT bool SupportsMeshLODStreaming()
 	{
 		return false;
 	}
 
-	static constexpr FORCEINLINE bool SupportsMemoryMappedFiles()
+	static constexpr UE_FORCEINLINE_HINT bool SupportsMemoryMappedFiles()
 	{
 		return false;
 	}
 
-	static constexpr FORCEINLINE bool SupportsMemoryMappedAudio()
+	static constexpr UE_FORCEINLINE_HINT bool SupportsMemoryMappedAudio()
 	{
 		return false;
 	}
 
-	static constexpr FORCEINLINE bool SupportsMemoryMappedAnimation()
+	static constexpr UE_FORCEINLINE_HINT bool SupportsMemoryMappedAnimation()
 	{
 		return false;
 	}
 
-	static constexpr FORCEINLINE int64 GetMemoryMappingAlignment()
+	static constexpr UE_FORCEINLINE_HINT int64 GetMemoryMappingAlignment()
 	{
 		return 0;
 	}
 
 	// Guaranteed virtual memory alignment on a given platform, regardless of a specific device
-	static constexpr FORCEINLINE int GetMaxSupportedVirtualMemoryAlignment()
+	static constexpr UE_FORCEINLINE_HINT int GetMaxSupportedVirtualMemoryAlignment()
 	{
 		return 4096;
 	}
 	
-	static constexpr FORCEINLINE bool SupportsRayTracing()
+	static constexpr UE_FORCEINLINE_HINT bool SupportsRayTracing()
 	{
 		return false;
 	}
 
-	static constexpr FORCEINLINE bool SupportsLumenGI()
+	static constexpr UE_FORCEINLINE_HINT bool SupportsLumenGI()
 	{
 		return true;
 	}
 
-	static constexpr FORCEINLINE bool SupportsHardwareLZDecompression()
+	static constexpr UE_FORCEINLINE_HINT bool SupportsHardwareLZDecompression()
 	{
 		return false;
 	}
@@ -291,37 +291,37 @@ struct FGenericPlatformProperties
 	/**
 	 * Gets whether user settings should override the resolution or not
 	 */
-	static constexpr FORCEINLINE bool HasFixedResolution()
+	static constexpr UE_FORCEINLINE_HINT bool HasFixedResolution()
 	{
 		return true;
 	}
 
-	static constexpr FORCEINLINE bool SupportsMinimize()
+	static constexpr UE_FORCEINLINE_HINT bool SupportsMinimize()
 	{
 		return false;
 	}
 
 	// Whether the platform allows an application to quit to the OS
-	static constexpr FORCEINLINE bool SupportsQuit()
+	static constexpr UE_FORCEINLINE_HINT bool SupportsQuit()
 	{
 		return false;
 	}
 
 	// Whether the platform allows the call stack to be dumped during an assert
-	static constexpr FORCEINLINE bool AllowsCallStackDumpDuringAssert()
+	static constexpr UE_FORCEINLINE_HINT bool AllowsCallStackDumpDuringAssert()
 	{
 		return IsProgram();
 	}
 
 	// If this platform wants to replace Zlib with a platform-specific version, set the name of the compression format 
 	// plugin (matching its GetCompressionFormatName() function) in an override of this function
-	static constexpr FORCEINLINE const char* GetZlibReplacementFormat()
+	static constexpr UE_FORCEINLINE_HINT const char* GetZlibReplacementFormat()
 	{
 		return nullptr;
 	}
 	 
 	// Whether the platform requires an original release version to make a patch
-	static constexpr FORCEINLINE bool RequiresOriginalReleaseVersionForPatch()
+	static constexpr UE_FORCEINLINE_HINT bool RequiresOriginalReleaseVersionForPatch()
 	{
 		return false;
 	}
diff --git a/Engine/Source/Runtime/Core/Public/GenericPlatform/GenericPlatformSymbolication.h b/Engine/Source/Runtime/Core/Public/GenericPlatform/GenericPlatformSymbolication.h
--- a/Engine/Source/Runtime/Core/Public/GenericPlatform/GenericPlatformSymbolication.h
+++ b/Engine/Source/Runtime/Core/Public/GenericPlatform/GenericPlatformSymbolication.h
@@ -92,7 +92,7 @@ struct FGenericPlatformSymbolDatabaseKeyFuncs
 	/**
 	 * @return The key used to index the given element.
 	 */
-	static FORCEINLINE KeyInitType GetSetKey(ElementInitType Element)
+	static UE_FORCEINLINE_HINT KeyInitType GetSetKey(ElementInitType Element)
 	{
 		return Element.Signature;
 	}
@@ -100,13 +100,13 @@ struct FGenericPlatformSymbolDatabaseKeyFuncs
 	/**
 	 * @return True if the keys match.
 	 */
-	static FORCEINLINE bool Matches(KeyInitType A,KeyInitType B)
+	static UE_FORCEINLINE_HINT bool Matches(KeyInitType A,KeyInitType B)
 	{
 		return A == B;
 	}
 
 	/** Calculates a hash index for a key. */
-	static FORCEINLINE uint32 GetKeyHash(KeyInitType Key)
+	static UE_FORCEINLINE_HINT uint32 GetKeyHash(KeyInitType Key)
 	{
 		return GetTypeHash(Key);
 	}
diff --git a/Engine/Source/Runtime/Core/Public/GenericPlatform/GenericPlatformTLS.h b/Engine/Source/Runtime/Core/Public/GenericPlatform/GenericPlatformTLS.h
--- a/Engine/Source/Runtime/Core/Public/GenericPlatform/GenericPlatformTLS.h
+++ b/Engine/Source/Runtime/Core/Public/GenericPlatform/GenericPlatformTLS.h
@@ -17,7 +17,7 @@ struct FGenericPlatformTLS
 	 * @param SlotIndex the TLS index to check
 	 * @return true if this looks like a valid slot
 	 */
-	static FORCEINLINE bool IsValidTlsSlot(uint32 SlotIndex)
+	static UE_FORCEINLINE_HINT bool IsValidTlsSlot(uint32 SlotIndex)
 	{
 		return SlotIndex != InvalidTlsSlot;
 	}
@@ -26,14 +26,14 @@ struct FGenericPlatformTLS
 	/**
 	 * Returns the currently executing thread's id
 	 */
-	static FORCEINLINE uint32 GetCurrentThreadId(void)
+	static UE_FORCEINLINE_HINT uint32 GetCurrentThreadId(void)
 	{
 	}
 
 	/**
 	 * Allocates a thread local store slot
 	 */
-	static FORCEINLINE uint32 AllocTlsSlot(void)
+	static UE_FORCEINLINE_HINT uint32 AllocTlsSlot(void)
 	{
 	}
 
@@ -43,7 +43,7 @@ struct FGenericPlatformTLS
 	 * @param SlotIndex the TLS index to store it in
 	 * @param Value the value to store in the slot
 	 */
-	static FORCEINLINE void SetTlsValue(uint32 SlotIndex,void* Value)
+	static UE_FORCEINLINE_HINT void SetTlsValue(uint32 SlotIndex,void* Value)
 	{
 	}
 
@@ -52,7 +52,7 @@ struct FGenericPlatformTLS
 	 *
 	 * @return the value stored in the slot
 	 */
-	static FORCEINLINE void* GetTlsValue(uint32 SlotIndex)
+	static UE_FORCEINLINE_HINT void* GetTlsValue(uint32 SlotIndex)
 	{
 	}
 
@@ -61,7 +61,7 @@ struct FGenericPlatformTLS
 	 *
 	 * @param SlotIndex the TLS index to store it in
 	 */
-	static FORCEINLINE void FreeTlsSlot(uint32 SlotIndex)
+	static UE_FORCEINLINE_HINT void FreeTlsSlot(uint32 SlotIndex)
 	{
 	}
 #endif
diff --git a/Engine/Source/Runtime/Core/Public/GenericPlatform/GenericWidePlatformString.h b/Engine/Source/Runtime/Core/Public/GenericPlatform/GenericWidePlatformString.h
--- a/Engine/Source/Runtime/Core/Public/GenericPlatform/GenericWidePlatformString.h
+++ b/Engine/Source/Runtime/Core/Public/GenericPlatform/GenericWidePlatformString.h
@@ -193,12 +193,12 @@ public:
 	CORE_API static float Atof(const WIDECHAR* String);
 	CORE_API static double Atod(const WIDECHAR* String);
 
-	static FORCEINLINE int32 Atoi(const WIDECHAR* String)
+	static UE_FORCEINLINE_HINT int32 Atoi(const WIDECHAR* String)
 	{
 		return Strtoi( String, NULL, 10 );
 	}
 	
-	static FORCEINLINE int64 Atoi64(const WIDECHAR* String)
+	static UE_FORCEINLINE_HINT int64 Atoi64(const WIDECHAR* String)
 	{
 		return Strtoi64( String, NULL, 10 );
 	}
@@ -223,7 +223,7 @@ PRAGMA_ENABLE_DEPRECATION_WARNINGS
 	// This version was deprecated because Strcpy taking a size field does not match the CRT strcpy, and on some platforms the size field was being ignored.
 	// If the memzeroing causes a performance problem, request a new function StrcpyTruncate.
 	UE_DEPRECATED(5.6, "Use Strncpy instead. Note that Strncpy has a behavior difference from Strcpy: it memzeroes the entire DestCount-sized buffer after the end of string.")
-	static FORCEINLINE ANSICHAR* Strcpy(ANSICHAR* Dest, SIZE_T DestCount, const ANSICHAR* Src)
+	static UE_FORCEINLINE_HINT ANSICHAR* Strcpy(ANSICHAR* Dest, SIZE_T DestCount, const ANSICHAR* Src)
 	{
 		return Strcpy(Dest, Src);
 	}
@@ -248,7 +248,7 @@ PRAGMA_ENABLE_DEPRECATION_WARNINGS
 
 	// This version was deprecated because Strcat taking a size field does not match the CRT strcat, and on some platforms the size field was being ignored.
 	UE_DEPRECATED(5.6, "Use Strncat instead. !!NOTE THAT STRNCAT takes SrcLen rather than DestCount. You must call Strncat(Dest, Src, DestCount - Strlen(Dest) - 1).")
-	static FORCEINLINE ANSICHAR* Strcat(ANSICHAR* Dest, SIZE_T DestCount, const ANSICHAR* Src)
+	static UE_FORCEINLINE_HINT ANSICHAR* Strcat(ANSICHAR* Dest, SIZE_T DestCount, const ANSICHAR* Src)
 	{
 		return Strcat(Dest, Src);
 	}
@@ -266,72 +266,72 @@ PRAGMA_DISABLE_DEPRECATION_WARNINGS
 PRAGMA_ENABLE_DEPRECATION_WARNINGS
 	}
 
-	static FORCEINLINE int32 Strcmp( const ANSICHAR* String1, const ANSICHAR* String2 )
+	static UE_FORCEINLINE_HINT int32 Strcmp( const ANSICHAR* String1, const ANSICHAR* String2 )
 	{
 		return strcmp(String1, String2);
 	}
 
-	static FORCEINLINE int32 Strncmp( const ANSICHAR* String1, const ANSICHAR* String2, SIZE_T Count )
+	static UE_FORCEINLINE_HINT int32 Strncmp( const ANSICHAR* String1, const ANSICHAR* String2, SIZE_T Count )
 	{
 		return strncmp( String1, String2, Count );
 	}
 
-	static FORCEINLINE int32 Strlen( const ANSICHAR* String )
+	static UE_FORCEINLINE_HINT int32 Strlen( const ANSICHAR* String )
 	{
 		return (int32)strlen( String );
 	}
 
-	static FORCEINLINE int32 Strnlen( const ANSICHAR* String, SIZE_T StringSize )
+	static UE_FORCEINLINE_HINT int32 Strnlen( const ANSICHAR* String, SIZE_T StringSize )
 	{
 		return (int32)strnlen( String, StringSize );
 	}
 
-	static FORCEINLINE const ANSICHAR* Strstr( const ANSICHAR* String, const ANSICHAR* Find)
+	static UE_FORCEINLINE_HINT const ANSICHAR* Strstr( const ANSICHAR* String, const ANSICHAR* Find)
 	{
 		return strstr(String, Find);
 	}
 
-	static FORCEINLINE const ANSICHAR* Strchr( const ANSICHAR* String, ANSICHAR C)
+	static UE_FORCEINLINE_HINT const ANSICHAR* Strchr( const ANSICHAR* String, ANSICHAR C)
 	{
 		return strchr(String, C);
 	}
 
-	static FORCEINLINE const ANSICHAR* Strrchr( const ANSICHAR* String, ANSICHAR C)
+	static UE_FORCEINLINE_HINT const ANSICHAR* Strrchr( const ANSICHAR* String, ANSICHAR C)
 	{
 		return strrchr(String, C);
 	}
 
-	static FORCEINLINE int32 Atoi(const ANSICHAR* String)
+	static UE_FORCEINLINE_HINT int32 Atoi(const ANSICHAR* String)
 	{
 		return atoi( String );
 	}
 
-	static FORCEINLINE int64 Atoi64(const ANSICHAR* String)
+	static UE_FORCEINLINE_HINT int64 Atoi64(const ANSICHAR* String)
 	{
 		return strtoll( String, NULL, 10 );
 	}
 
-	static FORCEINLINE float Atof(const ANSICHAR* String)
+	static UE_FORCEINLINE_HINT float Atof(const ANSICHAR* String)
 	{
 		return (float)atof( String );
 	}
 
-	static FORCEINLINE double Atod(const ANSICHAR* String)
+	static UE_FORCEINLINE_HINT double Atod(const ANSICHAR* String)
 	{
 		return atof( String );
 	}
 
-	static FORCEINLINE int32 Strtoi( const ANSICHAR* Start, ANSICHAR** End, int32 Base )
+	static UE_FORCEINLINE_HINT int32 Strtoi( const ANSICHAR* Start, ANSICHAR** End, int32 Base )
 	{
 		return (int32)strtol( Start, End, Base );
 	}
 
-	static FORCEINLINE int64 Strtoi64( const ANSICHAR* Start, ANSICHAR** End, int32 Base )
+	static UE_FORCEINLINE_HINT int64 Strtoi64( const ANSICHAR* Start, ANSICHAR** End, int32 Base )
 	{
 		return strtoll(Start, End, Base);
 	}
 
-	static FORCEINLINE uint64 Strtoui64( const ANSICHAR* Start, ANSICHAR** End, int32 Base )
+	static UE_FORCEINLINE_HINT uint64 Strtoui64( const ANSICHAR* Start, ANSICHAR** End, int32 Base )
 	{
 		return strtoull(Start, End, Base);
 	}
@@ -382,7 +382,7 @@ PRAGMA_ENABLE_DEPRECATION_WARNINGS
 	/**
 	 * UTF8CHAR implementation.
 	 */
-	static FORCEINLINE UTF8CHAR* Strcpy(UTF8CHAR* Dest, const UTF8CHAR* Src)
+	static UE_FORCEINLINE_HINT UTF8CHAR* Strcpy(UTF8CHAR* Dest, const UTF8CHAR* Src)
 	{
 		return (UTF8CHAR*)Strcpy((ANSICHAR*)Dest, (const ANSICHAR*)Src);
 	}
@@ -390,24 +390,24 @@ PRAGMA_ENABLE_DEPRECATION_WARNINGS
 	// This version was deprecated because Strcpy taking a size field does not match the CRT strcpy, and on some platforms the size field was being ignored.
 	// If the memzeroing causes a performance problem, request a new function StrcpyTruncate.
 	UE_DEPRECATED(5.6, "Use Strncpy instead. Note that Strncpy has a behavior difference from Strcpy: it memzeroes the entire DestCount-sized buffer after the end of string.")
-	static FORCEINLINE UTF8CHAR* Strcpy(UTF8CHAR* Dest, SIZE_T DestCount, const UTF8CHAR* Src)
+	static UE_FORCEINLINE_HINT UTF8CHAR* Strcpy(UTF8CHAR* Dest, SIZE_T DestCount, const UTF8CHAR* Src)
 	{
 		return (UTF8CHAR*)Strcpy((ANSICHAR*)Dest, DestCount, (const ANSICHAR*)Src);
 	}
 
-	static FORCEINLINE UTF8CHAR* Strncpy(UTF8CHAR* Dest, const UTF8CHAR* Src, SIZE_T MaxLen)
+	static UE_FORCEINLINE_HINT UTF8CHAR* Strncpy(UTF8CHAR* Dest, const UTF8CHAR* Src, SIZE_T MaxLen)
 	{
 		return (UTF8CHAR*)Strncpy((ANSICHAR*)Dest, (const ANSICHAR*)Src, MaxLen);
 	}
 
-	static FORCEINLINE UTF8CHAR* Strcat(UTF8CHAR* Dest, const UTF8CHAR* Src)
+	static UE_FORCEINLINE_HINT UTF8CHAR* Strcat(UTF8CHAR* Dest, const UTF8CHAR* Src)
 	{
 		return (UTF8CHAR*)Strcat((ANSICHAR*)Dest, (const ANSICHAR*)Src);
 	}
 
 	// This version was deprecated because Strcat taking a size field does not match the CRT strcat, and on some platforms the size field was being ignored.
 	UE_DEPRECATED(5.6, "Use Strncat instead. !!NOTE THAT STRNCAT takes SrcLen rather than DestCount. You must call Strncat(Dest, Src, DestCount - Strlen(Dest) - 1).")
-	static FORCEINLINE UTF8CHAR* Strcat(UTF8CHAR* Dest, SIZE_T DestCount, const UTF8CHAR* Src)
+	static UE_FORCEINLINE_HINT UTF8CHAR* Strcat(UTF8CHAR* Dest, SIZE_T DestCount, const UTF8CHAR* Src)
 	{
 		return Strcat(Dest, Src);
 	}
@@ -417,87 +417,87 @@ PRAGMA_ENABLE_DEPRECATION_WARNINGS
 	 * If the length of the string in Src is less than SrcLen, only the content up to the terminating
 	 * null-character is copied.
 	 */
-	static FORCEINLINE UTF8CHAR* Strncat(UTF8CHAR* Dest, const UTF8CHAR* Src, SIZE_T SrcLen)
+	static UE_FORCEINLINE_HINT UTF8CHAR* Strncat(UTF8CHAR* Dest, const UTF8CHAR* Src, SIZE_T SrcLen)
 	{
 		return (UTF8CHAR*)Strncat((ANSICHAR*)Dest, (const ANSICHAR*)Src, SrcLen);
 	}
 
-	static FORCEINLINE int32 Strcmp(const UTF8CHAR* String1, const UTF8CHAR* String2)
+	static UE_FORCEINLINE_HINT int32 Strcmp(const UTF8CHAR* String1, const UTF8CHAR* String2)
 	{
 		return Strcmp((const ANSICHAR*)String1, (const ANSICHAR*)String2);
 	}
 
-	static FORCEINLINE int32 Strncmp(const UTF8CHAR* String1, const UTF8CHAR* String2, SIZE_T Count)
+	static UE_FORCEINLINE_HINT int32 Strncmp(const UTF8CHAR* String1, const UTF8CHAR* String2, SIZE_T Count)
 	{
 		return Strncmp((const ANSICHAR*)String1, (const ANSICHAR*)String2, Count);
 	}
 
-	static FORCEINLINE int32 Strlen(const UTF8CHAR* String)
+	static UE_FORCEINLINE_HINT int32 Strlen(const UTF8CHAR* String)
 	{
 		return Strlen((const ANSICHAR*)String);
 	}
 
-	static FORCEINLINE int32 Strnlen(const UTF8CHAR* String, SIZE_T StringSize)
+	static UE_FORCEINLINE_HINT int32 Strnlen(const UTF8CHAR* String, SIZE_T StringSize)
 	{
 		return Strnlen((const ANSICHAR*)String, StringSize);
 	}
 
-	static FORCEINLINE const UTF8CHAR* Strstr(const UTF8CHAR* String, const UTF8CHAR* Find)
+	static UE_FORCEINLINE_HINT const UTF8CHAR* Strstr(const UTF8CHAR* String, const UTF8CHAR* Find)
 	{
 		return (const UTF8CHAR*)Strstr((const ANSICHAR*)String, (const ANSICHAR*)Find);
 	}
 
-	static FORCEINLINE const UTF8CHAR* Strchr(const UTF8CHAR* String, UTF8CHAR C)
+	static UE_FORCEINLINE_HINT const UTF8CHAR* Strchr(const UTF8CHAR* String, UTF8CHAR C)
 	{
 		return (const UTF8CHAR*)Strchr((const ANSICHAR*)String, (ANSICHAR)C);
 	}
 
-	static FORCEINLINE const UTF8CHAR* Strrchr(const UTF8CHAR* String, UTF8CHAR C)
+	static UE_FORCEINLINE_HINT const UTF8CHAR* Strrchr(const UTF8CHAR* String, UTF8CHAR C)
 	{
 		return (const UTF8CHAR*)Strrchr((const ANSICHAR*)String, (ANSICHAR)C);
 	}
 
-	static FORCEINLINE int32 Atoi(const UTF8CHAR* String)
+	static UE_FORCEINLINE_HINT int32 Atoi(const UTF8CHAR* String)
 	{
 		return Atoi((const ANSICHAR*)String);
 	}
 
-	static FORCEINLINE int64 Atoi64(const UTF8CHAR* String)
+	static UE_FORCEINLINE_HINT int64 Atoi64(const UTF8CHAR* String)
 	{
 		return Atoi64((const ANSICHAR*)String);
 	}
 
-	static FORCEINLINE float Atof(const UTF8CHAR* String)
+	static UE_FORCEINLINE_HINT float Atof(const UTF8CHAR* String)
 	{
 		return Atof((const ANSICHAR*)String);
 	}
 
-	static FORCEINLINE double Atod(const UTF8CHAR* String)
+	static UE_FORCEINLINE_HINT double Atod(const UTF8CHAR* String)
 	{
 		return Atod((const ANSICHAR*)String);
 	}
 
-	static FORCEINLINE int32 Strtoi(const UTF8CHAR* Start, UTF8CHAR** End, int32 Base)
+	static UE_FORCEINLINE_HINT int32 Strtoi(const UTF8CHAR* Start, UTF8CHAR** End, int32 Base)
 	{
 		return Strtoi((const ANSICHAR*)Start, (ANSICHAR**)End, Base);
 	}
 
-	static FORCEINLINE int64 Strtoi64(const UTF8CHAR* Start, UTF8CHAR** End, int32 Base)
+	static UE_FORCEINLINE_HINT int64 Strtoi64(const UTF8CHAR* Start, UTF8CHAR** End, int32 Base)
 	{
 		return Strtoi64((const ANSICHAR*)Start, (ANSICHAR**)End, Base);
 	}
 
-	static FORCEINLINE uint64 Strtoui64(const UTF8CHAR* Start, UTF8CHAR** End, int32 Base)
+	static UE_FORCEINLINE_HINT uint64 Strtoui64(const UTF8CHAR* Start, UTF8CHAR** End, int32 Base)
 	{
 		return Strtoui64((const ANSICHAR*)Start, (ANSICHAR**)End, Base);
 	}
 
-	static FORCEINLINE UTF8CHAR* Strtok(UTF8CHAR* StrToken, const UTF8CHAR* Delim, UTF8CHAR** Context)
+	static UE_FORCEINLINE_HINT UTF8CHAR* Strtok(UTF8CHAR* StrToken, const UTF8CHAR* Delim, UTF8CHAR** Context)
 	{
 		return (UTF8CHAR*)Strtok((ANSICHAR*)StrToken, (const ANSICHAR*)Delim, (ANSICHAR**)Context);
 	}
 
-	static FORCEINLINE int32 GetVarArgs(UTF8CHAR* Dest, SIZE_T DestSize, const UTF8CHAR*& Fmt, va_list ArgPtr)
+	static UE_FORCEINLINE_HINT int32 GetVarArgs(UTF8CHAR* Dest, SIZE_T DestSize, const UTF8CHAR*& Fmt, va_list ArgPtr)
 	{
 		return GetVarArgs((ANSICHAR*)Dest, DestSize, *(const ANSICHAR**)&Fmt, ArgPtr);
 	}
diff --git a/Engine/Source/Runtime/Core/Public/GenericPlatform/StandardPlatformString.h b/Engine/Source/Runtime/Core/Public/GenericPlatform/StandardPlatformString.h
--- a/Engine/Source/Runtime/Core/Public/GenericPlatform/StandardPlatformString.h
+++ b/Engine/Source/Runtime/Core/Public/GenericPlatform/StandardPlatformString.h
@@ -104,72 +104,72 @@ PRAGMA_ENABLE_DEPRECATION_WARNINGS
 		return Dest;
 	}
 
-	static FORCEINLINE int32 Strcmp( const WIDECHAR* String1, const WIDECHAR* String2 )
+	static UE_FORCEINLINE_HINT int32 Strcmp( const WIDECHAR* String1, const WIDECHAR* String2 )
 	{
 		return wcscmp(String1, String2);
 	}
 
-	static FORCEINLINE int32 Strncmp( const WIDECHAR* String1, const WIDECHAR* String2, SIZE_T Count )
+	static UE_FORCEINLINE_HINT int32 Strncmp( const WIDECHAR* String1, const WIDECHAR* String2, SIZE_T Count )
 	{
 		return wcsncmp( String1, String2, Count );
 	}
 
-	static FORCEINLINE int32 Strlen( const WIDECHAR* String )
+	static UE_FORCEINLINE_HINT int32 Strlen( const WIDECHAR* String )
 	{
 		return (int32)wcslen( String );
 	}
 
-	static FORCEINLINE int32 Strnlen( const WIDECHAR* String, SIZE_T StringSize )
+	static UE_FORCEINLINE_HINT int32 Strnlen( const WIDECHAR* String, SIZE_T StringSize )
 	{
 		return (int32)wcsnlen_s( String, StringSize );
 	}
 
-	static FORCEINLINE const WIDECHAR* Strstr( const WIDECHAR* String, const WIDECHAR* Find)
+	static UE_FORCEINLINE_HINT const WIDECHAR* Strstr( const WIDECHAR* String, const WIDECHAR* Find)
 	{
 		return wcsstr( String, Find );
 	}
 
-	static FORCEINLINE const WIDECHAR* Strchr( const WIDECHAR* String, WIDECHAR C)
+	static UE_FORCEINLINE_HINT const WIDECHAR* Strchr( const WIDECHAR* String, WIDECHAR C)
 	{
 		return wcschr( String, C );
 	}
 
-	static FORCEINLINE const WIDECHAR* Strrchr( const WIDECHAR* String, WIDECHAR C)
+	static UE_FORCEINLINE_HINT const WIDECHAR* Strrchr( const WIDECHAR* String, WIDECHAR C)
 	{
 		return wcsrchr( String, C );
 	}
 
-	static FORCEINLINE int32 Atoi(const WIDECHAR* String)
+	static UE_FORCEINLINE_HINT int32 Atoi(const WIDECHAR* String)
 	{
 		return (int32)wcstol( String, NULL, 10 );
 	}
 
-	static FORCEINLINE int64 Atoi64(const WIDECHAR* String)
+	static UE_FORCEINLINE_HINT int64 Atoi64(const WIDECHAR* String)
 	{
 		return wcstoll( String, NULL, 10 );
 	}
 
-	static FORCEINLINE float Atof(const WIDECHAR* String)
+	static UE_FORCEINLINE_HINT float Atof(const WIDECHAR* String)
 	{
 		return wcstof( String, NULL );
 	}
 
-	static FORCEINLINE double Atod(const WIDECHAR* String)
+	static UE_FORCEINLINE_HINT double Atod(const WIDECHAR* String)
 	{
 		return wcstod( String, NULL );
 	}
 
-	static FORCEINLINE int32 Strtoi( const WIDECHAR* Start, WIDECHAR** End, int32 Base )
+	static UE_FORCEINLINE_HINT int32 Strtoi( const WIDECHAR* Start, WIDECHAR** End, int32 Base )
 	{
 		return (int32)wcstol( Start, End, Base );
 	}
 
-	static FORCEINLINE int64 Strtoi64( const WIDECHAR* Start, WIDECHAR** End, int32 Base )
+	static UE_FORCEINLINE_HINT int64 Strtoi64( const WIDECHAR* Start, WIDECHAR** End, int32 Base )
 	{
 		return wcstoll( Start, End, Base );
 	}
 
-	static FORCEINLINE uint64 Strtoui64( const WIDECHAR* Start, WIDECHAR** End, int32 Base )
+	static UE_FORCEINLINE_HINT uint64 Strtoui64( const WIDECHAR* Start, WIDECHAR** End, int32 Base )
 	{
 		return wcstoull( Start, End, Base );
 	}
@@ -296,77 +296,77 @@ PRAGMA_DISABLE_DEPRECATION_WARNINGS
 PRAGMA_ENABLE_DEPRECATION_WARNINGS
 	}
 
-	static FORCEINLINE ANSICHAR* Strncat(ANSICHAR* Dest, const ANSICHAR* Src, SIZE_T SrcLen)
+	static UE_FORCEINLINE_HINT ANSICHAR* Strncat(ANSICHAR* Dest, const ANSICHAR* Src, SIZE_T SrcLen)
 	{
 		return (ANSICHAR*)strncat(Dest, Src, SrcLen);
 	}
 
-	static FORCEINLINE int32 Strcmp( const ANSICHAR* String1, const ANSICHAR* String2 )
+	static UE_FORCEINLINE_HINT int32 Strcmp( const ANSICHAR* String1, const ANSICHAR* String2 )
 	{
 		return strcmp(String1, String2);
 	}
 
-	static FORCEINLINE int32 Strncmp( const ANSICHAR* String1, const ANSICHAR* String2, SIZE_T Count )
+	static UE_FORCEINLINE_HINT int32 Strncmp( const ANSICHAR* String1, const ANSICHAR* String2, SIZE_T Count )
 	{
 		return strncmp( String1, String2, Count );
 	}
 
-	static FORCEINLINE int32 Strlen( const ANSICHAR* String )
+	static UE_FORCEINLINE_HINT int32 Strlen( const ANSICHAR* String )
 	{
 		return (int32)strlen( String );
 	}
 
-	static FORCEINLINE int32 Strnlen( const ANSICHAR* String, SIZE_T StringSize )
+	static UE_FORCEINLINE_HINT int32 Strnlen( const ANSICHAR* String, SIZE_T StringSize )
 	{
 		return (int32)strnlen_s( String, StringSize );
 	}
 
-	static FORCEINLINE const ANSICHAR* Strstr( const ANSICHAR* String, const ANSICHAR* Find)
+	static UE_FORCEINLINE_HINT const ANSICHAR* Strstr( const ANSICHAR* String, const ANSICHAR* Find)
 	{
 		return strstr(String, Find);
 	}
 
-	static FORCEINLINE const ANSICHAR* Strchr( const ANSICHAR* String, ANSICHAR C)
+	static UE_FORCEINLINE_HINT const ANSICHAR* Strchr( const ANSICHAR* String, ANSICHAR C)
 	{
 		return strchr(String, C);
 	}
 
-	static FORCEINLINE const ANSICHAR* Strrchr( const ANSICHAR* String, ANSICHAR C)
+	static UE_FORCEINLINE_HINT const ANSICHAR* Strrchr( const ANSICHAR* String, ANSICHAR C)
 	{
 		return strrchr(String, C);
 	}
 
-	static FORCEINLINE int32 Atoi(const ANSICHAR* String)
+	static UE_FORCEINLINE_HINT int32 Atoi(const ANSICHAR* String)
 	{
 		return atoi( String );
 	}
 
-	static FORCEINLINE int64 Atoi64(const ANSICHAR* String)
+	static UE_FORCEINLINE_HINT int64 Atoi64(const ANSICHAR* String)
 	{
 		return strtoll( String, NULL, 10 );
 	}
 
-	static FORCEINLINE float Atof(const ANSICHAR* String)
+	static UE_FORCEINLINE_HINT float Atof(const ANSICHAR* String)
 	{
 		return (float)atof( String );
 	}
 
-	static FORCEINLINE double Atod(const ANSICHAR* String)
+	static UE_FORCEINLINE_HINT double Atod(const ANSICHAR* String)
 	{
 		return atof( String );
 	}
 
-	static FORCEINLINE int32 Strtoi( const ANSICHAR* Start, ANSICHAR** End, int32 Base )
+	static UE_FORCEINLINE_HINT int32 Strtoi( const ANSICHAR* Start, ANSICHAR** End, int32 Base )
 	{
 		return (int32)strtol( Start, End, Base );
 	}
 
-	static FORCEINLINE int64 Strtoi64( const ANSICHAR* Start, ANSICHAR** End, int32 Base )
+	static UE_FORCEINLINE_HINT int64 Strtoi64( const ANSICHAR* Start, ANSICHAR** End, int32 Base )
 	{
 		return strtoll(Start, End, Base);
 	}
 
-	static FORCEINLINE uint64 Strtoui64( const ANSICHAR* Start, ANSICHAR** End, int32 Base )
+	static UE_FORCEINLINE_HINT uint64 Strtoui64( const ANSICHAR* Start, ANSICHAR** End, int32 Base )
 	{
 		return strtoull(Start, End, Base);
 	}
@@ -388,7 +388,7 @@ PRAGMA_ENABLE_DEPRECATION_WARNINGS
 	/**
 	 * UTF8CHAR implementation.
 	 */
-	static FORCEINLINE UTF8CHAR* Strcpy(UTF8CHAR* Dest, const UTF8CHAR* Src)
+	static UE_FORCEINLINE_HINT UTF8CHAR* Strcpy(UTF8CHAR* Dest, const UTF8CHAR* Src)
 	{
 		return (UTF8CHAR*)Strcpy((ANSICHAR*)Dest, (const ANSICHAR*)Src);
 	}
@@ -403,12 +403,12 @@ PRAGMA_DISABLE_DEPRECATION_WARNINGS
 PRAGMA_ENABLE_DEPRECATION_WARNINGS
 	}
 
-	static FORCEINLINE UTF8CHAR* Strncpy(UTF8CHAR* Dest, const UTF8CHAR* Src, SIZE_T MaxLen)
+	static UE_FORCEINLINE_HINT UTF8CHAR* Strncpy(UTF8CHAR* Dest, const UTF8CHAR* Src, SIZE_T MaxLen)
 	{
 		return (UTF8CHAR*)Strncpy((ANSICHAR*)Dest, (const ANSICHAR*)Src, MaxLen);
 	}
 
-	static FORCEINLINE UTF8CHAR* Strcat(UTF8CHAR* Dest, const UTF8CHAR* Src)
+	static UE_FORCEINLINE_HINT UTF8CHAR* Strcat(UTF8CHAR* Dest, const UTF8CHAR* Src)
 	{
 		return (UTF8CHAR*)Strcat((ANSICHAR*)Dest, (const ANSICHAR*)Src);
 	}
@@ -422,87 +422,87 @@ PRAGMA_DISABLE_DEPRECATION_WARNINGS
 PRAGMA_ENABLE_DEPRECATION_WARNINGS
 	}
 
-	static FORCEINLINE UTF8CHAR* Strncat(UTF8CHAR* Dest, const UTF8CHAR* Src, SIZE_T SrcLen)
+	static UE_FORCEINLINE_HINT UTF8CHAR* Strncat(UTF8CHAR* Dest, const UTF8CHAR* Src, SIZE_T SrcLen)
 	{
 		return (UTF8CHAR*)Strncat((ANSICHAR*)Dest, (const ANSICHAR*)Src, SrcLen);
 	}
 
-	static FORCEINLINE int32 Strcmp(const UTF8CHAR* String1, const UTF8CHAR* String2)
+	static UE_FORCEINLINE_HINT int32 Strcmp(const UTF8CHAR* String1, const UTF8CHAR* String2)
 	{
 		return Strcmp((const ANSICHAR*)String1, (const ANSICHAR*)String2);
 	}
 
-	static FORCEINLINE int32 Strncmp(const UTF8CHAR* String1, const UTF8CHAR* String2, SIZE_T Count)
+	static UE_FORCEINLINE_HINT int32 Strncmp(const UTF8CHAR* String1, const UTF8CHAR* String2, SIZE_T Count)
 	{
 		return Strncmp((const ANSICHAR*)String1, (const ANSICHAR*)String2, Count);
 	}
 
-	static FORCEINLINE int32 Strlen(const UTF8CHAR* String)
+	static UE_FORCEINLINE_HINT int32 Strlen(const UTF8CHAR* String)
 	{
 		return Strlen((const ANSICHAR*)String);
 	}
 
-	static FORCEINLINE int32 Strnlen(const UTF8CHAR* String, SIZE_T StringSize)
+	static UE_FORCEINLINE_HINT int32 Strnlen(const UTF8CHAR* String, SIZE_T StringSize)
 	{
 		return Strnlen((const ANSICHAR*)String, StringSize);
 	}
 
-	static FORCEINLINE const UTF8CHAR* Strstr(const UTF8CHAR* String, const UTF8CHAR* Find)
+	static UE_FORCEINLINE_HINT const UTF8CHAR* Strstr(const UTF8CHAR* String, const UTF8CHAR* Find)
 	{
 		return (const UTF8CHAR*)Strstr((const ANSICHAR*)String, (const ANSICHAR*)Find);
 	}
 
-	static FORCEINLINE const UTF8CHAR* Strchr(const UTF8CHAR* String, UTF8CHAR C)
+	static UE_FORCEINLINE_HINT const UTF8CHAR* Strchr(const UTF8CHAR* String, UTF8CHAR C)
 	{
 		return (const UTF8CHAR*)Strchr((const ANSICHAR*)String, (ANSICHAR)C);
 	}
 
-	static FORCEINLINE const UTF8CHAR* Strrchr(const UTF8CHAR* String, UTF8CHAR C)
+	static UE_FORCEINLINE_HINT const UTF8CHAR* Strrchr(const UTF8CHAR* String, UTF8CHAR C)
 	{
 		return (const UTF8CHAR*)Strrchr((const ANSICHAR*)String, (ANSICHAR)C);
 	}
 
-	static FORCEINLINE int32 Atoi(const UTF8CHAR* String)
+	static UE_FORCEINLINE_HINT int32 Atoi(const UTF8CHAR* String)
 	{
 		return Atoi((const ANSICHAR*)String);
 	}
 
-	static FORCEINLINE int64 Atoi64(const UTF8CHAR* String)
+	static UE_FORCEINLINE_HINT int64 Atoi64(const UTF8CHAR* String)
 	{
 		return Atoi64((const ANSICHAR*)String);
 	}
 
-	static FORCEINLINE float Atof(const UTF8CHAR* String)
+	static UE_FORCEINLINE_HINT float Atof(const UTF8CHAR* String)
 	{
 		return Atof((const ANSICHAR*)String);
 	}
 
-	static FORCEINLINE double Atod(const UTF8CHAR* String)
+	static UE_FORCEINLINE_HINT double Atod(const UTF8CHAR* String)
 	{
 		return Atod((const ANSICHAR*)String);
 	}
 
-	static FORCEINLINE int32 Strtoi(const UTF8CHAR* Start, UTF8CHAR** End, int32 Base)
+	static UE_FORCEINLINE_HINT int32 Strtoi(const UTF8CHAR* Start, UTF8CHAR** End, int32 Base)
 	{
 		return Strtoi((const ANSICHAR*)Start, (ANSICHAR**)End, Base);
 	}
 
-	static FORCEINLINE int64 Strtoi64(const UTF8CHAR* Start, UTF8CHAR** End, int32 Base)
+	static UE_FORCEINLINE_HINT int64 Strtoi64(const UTF8CHAR* Start, UTF8CHAR** End, int32 Base)
 	{
 		return Strtoi64((const ANSICHAR*)Start, (ANSICHAR**)End, Base);
 	}
 
-	static FORCEINLINE uint64 Strtoui64(const UTF8CHAR* Start, UTF8CHAR** End, int32 Base)
+	static UE_FORCEINLINE_HINT uint64 Strtoui64(const UTF8CHAR* Start, UTF8CHAR** End, int32 Base)
 	{
 		return Strtoui64((const ANSICHAR*)Start, (ANSICHAR**)End, Base);
 	}
 
-	static FORCEINLINE UTF8CHAR* Strtok(UTF8CHAR* StrToken, const UTF8CHAR* Delim, UTF8CHAR** Context)
+	static UE_FORCEINLINE_HINT UTF8CHAR* Strtok(UTF8CHAR* StrToken, const UTF8CHAR* Delim, UTF8CHAR** Context)
 	{
 		return (UTF8CHAR*)Strtok((ANSICHAR*)StrToken, (const ANSICHAR*)Delim, (ANSICHAR**)Context);
 	}
 
-	static FORCEINLINE int32 GetVarArgs(UTF8CHAR* Dest, SIZE_T DestSize, const UTF8CHAR*& Fmt, va_list ArgPtr)
+	static UE_FORCEINLINE_HINT int32 GetVarArgs(UTF8CHAR* Dest, SIZE_T DestSize, const UTF8CHAR*& Fmt, va_list ArgPtr)
 	{
 		return GetVarArgs((ANSICHAR*)Dest, DestSize, *(const ANSICHAR**)&Fmt, ArgPtr);
 	}
diff --git a/Engine/Source/Runtime/Core/Public/HAL/Allocators/CachedOSPageAllocator.h b/Engine/Source/Runtime/Core/Public/HAL/Allocators/CachedOSPageAllocator.h
--- a/Engine/Source/Runtime/Core/Public/HAL/Allocators/CachedOSPageAllocator.h
+++ b/Engine/Source/Runtime/Core/Public/HAL/Allocators/CachedOSPageAllocator.h
@@ -41,7 +41,7 @@ struct TCachedOSPageAllocator : private FCachedOSPageAllocator
 	{
 	}
 
-	FORCEINLINE void* Allocate(SIZE_T Size, uint32 AllocationHint = 0, UE::FPlatformRecursiveMutex* Mutex = nullptr)
+	UE_FORCEINLINE_HINT void* Allocate(SIZE_T Size, uint32 AllocationHint = 0, UE::FPlatformRecursiveMutex* Mutex = nullptr)
 	{
 		return AllocateImpl(Size, CachedByteLimit, FreedPageBlocks, FreedPageBlocks + FreedPageBlocksNum, FreedPageBlocksNum, CachedTotal, Mutex);
 	}
diff --git a/Engine/Source/Runtime/Core/Public/HAL/Allocators/PooledVirtualMemoryAllocator.h b/Engine/Source/Runtime/Core/Public/HAL/Allocators/PooledVirtualMemoryAllocator.h
--- a/Engine/Source/Runtime/Core/Public/HAL/Allocators/PooledVirtualMemoryAllocator.h
+++ b/Engine/Source/Runtime/Core/Public/HAL/Allocators/PooledVirtualMemoryAllocator.h
@@ -111,7 +111,7 @@ private:
 	 *
 	 * This function should not be passed 0 as Size!
 	 */
-	FORCEINLINE int32 GetAllocationSizeClass(SIZE_T Size)
+	UE_FORCEINLINE_HINT int32 GetAllocationSizeClass(SIZE_T Size)
 	{
 		return static_cast<int32>(Size >> 16) + ((Size & 0xFFFF) ? 1 : 0) - 1;
 	}
@@ -119,7 +119,7 @@ private:
 	/**
 	* Returns allocation size for a class.
 	*/
-	FORCEINLINE SIZE_T CalculateAllocationSizeFromClass(int32 Class)
+	UE_FORCEINLINE_HINT SIZE_T CalculateAllocationSizeFromClass(int32 Class)
 	{
 		return (static_cast<SIZE_T>(Class) + 1) * 65536;
 	}
diff --git a/Engine/Source/Runtime/Core/Public/HAL/CriticalSection.h b/Engine/Source/Runtime/Core/Public/HAL/CriticalSection.h
--- a/Engine/Source/Runtime/Core/Public/HAL/CriticalSection.h
+++ b/Engine/Source/Runtime/Core/Public/HAL/CriticalSection.h
@@ -13,32 +13,32 @@ namespace UE
 class FPlatformRWLock final
 {
 public:
-	FORCEINLINE bool TryWriteLock()
+	UE_REWRITE bool TryWriteLock()
 	{
 		return Mutex.TryLock();
 	}
 
-	FORCEINLINE void WriteLock()
+	UE_REWRITE void WriteLock()
 	{
 		Mutex.Lock();
 	}
 
-	FORCEINLINE void WriteUnlock()
+	UE_REWRITE void WriteUnlock()
 	{
 		Mutex.Unlock();
 	}
 
-	FORCEINLINE bool TryReadLock()
+	UE_REWRITE bool TryReadLock()
 	{
 		return Mutex.TryLockShared();
 	}
 
-	FORCEINLINE void ReadLock()
+	UE_REWRITE void ReadLock()
 	{
 		Mutex.LockShared();
 	}
 
-	FORCEINLINE void ReadUnlock()
+	UE_REWRITE void ReadUnlock()
 	{
 		Mutex.UnlockShared();
 	}
diff --git a/Engine/Source/Runtime/Core/Public/HAL/FileManagerGeneric.h b/Engine/Source/Runtime/Core/Public/HAL/FileManagerGeneric.h
--- a/Engine/Source/Runtime/Core/Public/HAL/FileManagerGeneric.h
+++ b/Engine/Source/Runtime/Core/Public/HAL/FileManagerGeneric.h
@@ -36,7 +36,7 @@ class FFileManagerGeneric
 	: public IFileManager
 {
 	// instead of caching the LowLevel, we call the singleton each time to never be incorrect
-	FORCEINLINE IPlatformFile& GetLowLevel() const
+	UE_FORCEINLINE_HINT IPlatformFile& GetLowLevel() const
 	{
 		return FPlatformFileManager::Get().GetPlatformFile();
 	}
diff --git a/Engine/Source/Runtime/Core/Public/HAL/IConsoleManager.h b/Engine/Source/Runtime/Core/Public/HAL/IConsoleManager.h
--- a/Engine/Source/Runtime/Core/Public/HAL/IConsoleManager.h
+++ b/Engine/Source/Runtime/Core/Public/HAL/IConsoleManager.h
@@ -1439,20 +1439,20 @@ public:
 	CORE_API FAutoConsoleVariable(const TCHAR* Name, const TCHAR* DefaultValue, const TCHAR* Help, const FConsoleVariableDelegate& Callback, uint32 Flags = ECVF_Default);
 
 	/** Dereference back to a console variable**/
-	FORCEINLINE IConsoleVariable& operator*()
+	UE_FORCEINLINE_HINT IConsoleVariable& operator*()
 	{
 		return *AsVariable();
 	}
-	FORCEINLINE const IConsoleVariable& operator*() const
+	UE_FORCEINLINE_HINT const IConsoleVariable& operator*() const
 	{
 		return *AsVariable();
 	}
 	/** Dereference back to a console variable**/
-	FORCEINLINE IConsoleVariable* operator->()
+	UE_FORCEINLINE_HINT IConsoleVariable* operator->()
 	{
 		return AsVariable();
 	}
-	FORCEINLINE const IConsoleVariable* operator->() const
+	UE_FORCEINLINE_HINT const IConsoleVariable* operator->() const
 	{
 		return AsVariable();
 	}
@@ -1559,20 +1559,20 @@ public:
 	{
 	}
 	/** Dereference back to a variable**/
-	FORCEINLINE IConsoleVariable& operator*()
+	UE_FORCEINLINE_HINT IConsoleVariable& operator*()
 	{
 		return *AsVariable();
 	}
-	FORCEINLINE const IConsoleVariable& operator*() const
+	UE_FORCEINLINE_HINT const IConsoleVariable& operator*() const
 	{
 		return *AsVariable();
 	}
 	/** Dereference back to a variable**/
-	FORCEINLINE IConsoleVariable* operator->()
+	UE_FORCEINLINE_HINT IConsoleVariable* operator->()
 	{
 		return AsVariable();
 	}
-	FORCEINLINE const IConsoleVariable* operator->() const
+	UE_FORCEINLINE_HINT const IConsoleVariable* operator->() const
 	{
 		return AsVariable();
 	}
@@ -1760,20 +1760,20 @@ public:
 	}
 	
 	/** Dereference back to a variable**/
-	FORCEINLINE IConsoleVariable& operator*()
+	UE_FORCEINLINE_HINT IConsoleVariable& operator*()
 	{
 		return *AsVariable();
 	}
-	FORCEINLINE const IConsoleVariable& operator*() const
+	UE_FORCEINLINE_HINT const IConsoleVariable& operator*() const
 	{
 		return *AsVariable();
 	}
 	/** Dereference back to a variable**/
-	FORCEINLINE IConsoleVariable* operator->()
+	UE_FORCEINLINE_HINT IConsoleVariable* operator->()
 	{
 		return AsVariable();
 	}
-	FORCEINLINE const IConsoleVariable* operator->() const
+	UE_FORCEINLINE_HINT const IConsoleVariable* operator->() const
 	{
 		return AsVariable();
 	}
diff --git a/Engine/Source/Runtime/Core/Public/HAL/LowLevelMemTracker.h b/Engine/Source/Runtime/Core/Public/HAL/LowLevelMemTracker.h
--- a/Engine/Source/Runtime/Core/Public/HAL/LowLevelMemTracker.h
+++ b/Engine/Source/Runtime/Core/Public/HAL/LowLevelMemTracker.h
@@ -770,12 +770,12 @@ public:
 
 	CORE_API void OnPreFork();
 
-	FORCEINLINE bool IsInitialized() const
+	UE_FORCEINLINE_HINT bool IsInitialized() const
 	{
 		return bFullyInitialised;
 	}
 
-	FORCEINLINE bool IsConfigured() const
+	UE_FORCEINLINE_HINT bool IsConfigured() const
 	{
 		return bConfigurationComplete;
 	}
diff --git a/Engine/Source/Runtime/Core/Public/HAL/MallocBinnedCommon.h b/Engine/Source/Runtime/Core/Public/HAL/MallocBinnedCommon.h
--- a/Engine/Source/Runtime/Core/Public/HAL/MallocBinnedCommon.h
+++ b/Engine/Source/Runtime/Core/Public/HAL/MallocBinnedCommon.h
@@ -221,7 +221,7 @@ protected:
 			OutPoolIndex = uint32((Ptr >> PtrToPoolPageBitShift) & PoolMask);
 		}
 
-		FORCEINLINE uint64 GetMaxHashBuckets() const
+		UE_FORCEINLINE_HINT uint64 GetMaxHashBuckets() const
 		{
 			return MaxHashBuckets;
 		}
@@ -246,7 +246,7 @@ protected:
 private:
 	struct FBundle
 	{
-		FORCEINLINE FBundle()
+		UE_FORCEINLINE_HINT FBundle()
 		{
 			Reset();
 		}
@@ -299,7 +299,7 @@ protected:
 			return true;
 		}
 
-		FORCEINLINE bool CanPushToFront(uint32 InPoolIndex, uint32 InBinSize) const
+		UE_FORCEINLINE_HINT bool CanPushToFront(uint32 InPoolIndex, uint32 InBinSize) const
 		{
 			return !((!!FullBundle.Head) & ((PartialBundle.Count >= (uint32)GMallocBinnedBundleCount) | (PartialBundle.Count * InBinSize >= (uint32)GMallocBinnedBundleSize)));
 		}
@@ -539,19 +539,19 @@ protected:
 			FPlatformTLS::SetTlsValue(BinnedTlsSlot, nullptr);
 		}
 
-		FORCEINLINE void* Malloc(uint32 InPoolIndex)
+		UE_FORCEINLINE_HINT void* Malloc(uint32 InPoolIndex)
 		{
 			return FreeLists[InPoolIndex].PopFromFront(InPoolIndex);
 		}
 
 		// return true if the pointer was pushed
-		FORCEINLINE bool Free(void* InPtr, uint32 InPoolIndex, uint32 InBinSize)
+		UE_FORCEINLINE_HINT bool Free(void* InPtr, uint32 InPoolIndex, uint32 InBinSize)
 		{
 			return FreeLists[InPoolIndex].PushToFront(InPtr, InPoolIndex, InBinSize);
 		}
 
 		// return true if a pointer can be pushed
-		FORCEINLINE bool CanFree(uint32 InPoolIndex, uint32 InBinSize) const
+		UE_FORCEINLINE_HINT bool CanFree(uint32 InPoolIndex, uint32 InBinSize) const
 		{
 			return FreeLists[InPoolIndex].CanPushToFront(InPoolIndex, InBinSize);
 		}
diff --git a/Engine/Source/Runtime/Core/Public/HAL/MallocBinnedGPU.h b/Engine/Source/Runtime/Core/Public/HAL/MallocBinnedGPU.h
--- a/Engine/Source/Runtime/Core/Public/HAL/MallocBinnedGPU.h
+++ b/Engine/Source/Runtime/Core/Public/HAL/MallocBinnedGPU.h
@@ -115,11 +115,11 @@ class FMallocBinnedGPU final : public FMalloc
 			NumFreeBlocks = InPageSize / InBlockSize;
 		}
 
-		FORCEINLINE uint32 GetNumFreeRegularBlocks() const
+		UE_FORCEINLINE_HINT uint32 GetNumFreeRegularBlocks() const
 		{
 			return NumFreeBlocks;
 		}
-		FORCEINLINE bool IsCanaryOk() const
+		UE_FORCEINLINE_HINT bool IsCanaryOk() const
 		{
 			return Canary == FFreeBlock::CANARY_VALUE;
 		}
@@ -192,7 +192,7 @@ class FMallocBinnedGPU final : public FMalloc
 			OutPoolIndex = ((UPTRINT)InPtr >> PtrToPoolPageBitShift) & PoolMask;
 		}
 
-		FORCEINLINE uint64 GetMaxHashBuckets() const
+		UE_FORCEINLINE_HINT uint64 GetMaxHashBuckets() const
 		{
 			return MaxHashBuckets;
 		}
@@ -220,7 +220,7 @@ class FMallocBinnedGPU final : public FMalloc
 
 	struct FBundle
 	{
-		FORCEINLINE FBundle()
+		UE_FORCEINLINE_HINT FBundle()
 		{
 			Reset();
 		}
@@ -272,7 +272,7 @@ class FMallocBinnedGPU final : public FMalloc
 			MBG_STAT(Allocator.GPUProxyMemory += sizeof(FGPUMemoryBlockProxy);)
 			return true;
 		}
-		FORCEINLINE bool CanPushToFront(uint32 InPoolIndex, uint32 InBlockSize, const FArenaParams& LocalArenaParams)
+		UE_FORCEINLINE_HINT bool CanPushToFront(uint32 InPoolIndex, uint32 InBlockSize, const FArenaParams& LocalArenaParams)
 		{
 			return !((!!FullBundle.Head) & ((PartialBundle.Count >= (uint32)LocalArenaParams.MaxBlocksPerBundle) | (PartialBundle.Count * InBlockSize >= (uint32)LocalArenaParams.MaxSizePerBundle)));
 		}
@@ -306,7 +306,7 @@ class FMallocBinnedGPU final : public FMalloc
 
 	struct FPerThreadFreeBlockLists
 	{
-		FORCEINLINE static FPerThreadFreeBlockLists* Get(uint32 BinnedGPUTlsSlot)
+		UE_FORCEINLINE_HINT static FPerThreadFreeBlockLists* Get(uint32 BinnedGPUTlsSlot)
 		{
 			return FPlatformTLS::IsValidTlsSlot(BinnedGPUTlsSlot) ? (FPerThreadFreeBlockLists*)FPlatformTLS::GetTlsValue(BinnedGPUTlsSlot) : nullptr;
 		}
@@ -319,17 +319,17 @@ class FMallocBinnedGPU final : public FMalloc
 			FreeLists.AddDefaulted(PoolCount);
 		}
 
-		FORCEINLINE void* Malloc(FMallocBinnedGPU& Allocator, uint32 InPoolIndex)
+		UE_FORCEINLINE_HINT void* Malloc(FMallocBinnedGPU& Allocator, uint32 InPoolIndex)
 		{
 			return FreeLists[InPoolIndex].PopFromFront(Allocator, InPoolIndex);
 		}
 		// return true if the pointer was pushed
-		FORCEINLINE bool Free(FMallocBinnedGPU& Allocator, void* InPtr, uint32 InPoolIndex, uint32 InBlockSize, const FArenaParams& LocalArenaParams)
+		UE_FORCEINLINE_HINT bool Free(FMallocBinnedGPU& Allocator, void* InPtr, uint32 InPoolIndex, uint32 InBlockSize, const FArenaParams& LocalArenaParams)
 		{
 			return FreeLists[InPoolIndex].PushToFront(Allocator, InPtr, InPoolIndex, InBlockSize, LocalArenaParams);
 		}
 		// return true if a pointer can be pushed
-		FORCEINLINE bool CanFree(uint32 InPoolIndex, uint32 InBlockSize, const FArenaParams& LocalArenaParams)
+		UE_FORCEINLINE_HINT bool CanFree(uint32 InPoolIndex, uint32 InBlockSize, const FArenaParams& LocalArenaParams)
 		{
 			return FreeLists[InPoolIndex].CanPushToFront(InPoolIndex, InBlockSize, LocalArenaParams);
 		}
@@ -443,7 +443,7 @@ class FMallocBinnedGPU final : public FMalloc
 		return PoolIndex;
 	}
 
-	FORCEINLINE uint8* PoolBasePtr(uint32 InPoolIndex)
+	UE_FORCEINLINE_HINT uint8* PoolBasePtr(uint32 InPoolIndex)
 	{
 		return PoolBaseVMPtr[InPoolIndex];
 	}
@@ -454,7 +454,7 @@ class FMallocBinnedGPU final : public FMalloc
 		return Result;
 	}
 
-	FORCEINLINE bool IsOSAllocation(const void* Ptr)
+	UE_FORCEINLINE_HINT bool IsOSAllocation(const void* Ptr)
 	{
 		return PoolIndexFromPtr(Ptr) >= ArenaParams.PoolCount;
 	}
@@ -630,7 +630,7 @@ public:
 		checkSlow(PoolIndex >= 0 && PoolIndex < ArenaParams.PoolCount);
 		return PoolIndex;
 	}
-	FORCEINLINE uint32 PoolIndexToBlockSize(uint32 PoolIndex)
+	UE_FORCEINLINE_HINT uint32 PoolIndexToBlockSize(uint32 PoolIndex)
 	{
 		return uint32(SmallBlockSizesReversedShifted[ArenaParams.PoolCount - PoolIndex - 1]) << ArenaParams.MinimumAlignmentShift;
 	}
diff --git a/Engine/Source/Runtime/Core/Public/HAL/MallocStomp2.h b/Engine/Source/Runtime/Core/Public/HAL/MallocStomp2.h
--- a/Engine/Source/Runtime/Core/Public/HAL/MallocStomp2.h
+++ b/Engine/Source/Runtime/Core/Public/HAL/MallocStomp2.h
@@ -28,7 +28,7 @@
 class FMallocStomp2 final : public FMalloc
 {
 private:
-	FORCEINLINE SIZE_T GetAlignedSize(SIZE_T InSize, SIZE_T InAlignment)
+	UE_FORCEINLINE_HINT SIZE_T GetAlignedSize(SIZE_T InSize, SIZE_T InAlignment)
 	{
 		return (((InSize)+(InAlignment - 1)) & ~(InAlignment - 1));
 	}
diff --git a/Engine/Source/Runtime/Core/Public/HAL/MallocTimer.h b/Engine/Source/Runtime/Core/Public/HAL/MallocTimer.h
--- a/Engine/Source/Runtime/Core/Public/HAL/MallocTimer.h
+++ b/Engine/Source/Runtime/Core/Public/HAL/MallocTimer.h
@@ -54,13 +54,13 @@ struct FScopedVirtualMallocTimer
 
 	static CORE_API void UpdateStats();
 #else	//UE_TIME_VIRTUALMALLOC
-	FORCEINLINE FScopedVirtualMallocTimer(int32 InIndex = 0, int32 InPlatformTypeIndex = 0)
+	UE_FORCEINLINE_HINT FScopedVirtualMallocTimer(int32 InIndex = 0, int32 InPlatformTypeIndex = 0)
 	{
 	}
-	FORCEINLINE ~FScopedVirtualMallocTimer()
+	UE_FORCEINLINE_HINT ~FScopedVirtualMallocTimer()
 	{
 	}
-	static FORCEINLINE void UpdateStats()
+	static UE_FORCEINLINE_HINT void UpdateStats()
 	{
 	}
 #endif	//UE_TIME_VIRTUALMALLOC
diff --git a/Engine/Source/Runtime/Core/Public/HAL/PThreadsRecursiveMutex.h b/Engine/Source/Runtime/Core/Public/HAL/PThreadsRecursiveMutex.h
--- a/Engine/Source/Runtime/Core/Public/HAL/PThreadsRecursiveMutex.h
+++ b/Engine/Source/Runtime/Core/Public/HAL/PThreadsRecursiveMutex.h
@@ -34,22 +34,22 @@ public:
 		pthread_mutexattr_destroy(&MutexAttributes);
 	}
 
-	FORCEINLINE ~FPThreadsRecursiveMutex()
+	UE_FORCEINLINE_HINT ~FPThreadsRecursiveMutex()
 	{
 		pthread_mutex_destroy(&Mutex);
 	}
 
-	FORCEINLINE bool TryLock()
+	UE_FORCEINLINE_HINT bool TryLock()
 	{
 		return pthread_mutex_trylock(&Mutex) == 0;
 	}
 
-	FORCEINLINE void Lock()
+	UE_FORCEINLINE_HINT void Lock()
 	{
 		pthread_mutex_lock(&Mutex);
 	}
 
-	FORCEINLINE void Unlock()
+	UE_FORCEINLINE_HINT void Unlock()
 	{
 		pthread_mutex_unlock(&Mutex);
 	}
diff --git a/Engine/Source/Runtime/Core/Public/HAL/ThreadHeartBeat.h b/Engine/Source/Runtime/Core/Public/HAL/ThreadHeartBeat.h
--- a/Engine/Source/Runtime/Core/Public/HAL/ThreadHeartBeat.h
+++ b/Engine/Source/Runtime/Core/Public/HAL/ThreadHeartBeat.h
@@ -415,11 +415,11 @@ typedef FGameThreadHitchHeartBeatThreaded FGameThreadHitchHeartBeat;
 /** Suspends hitch detection in the current scope */
 struct FDisableHitchDetectorScope
 {
-	FORCEINLINE FDisableHitchDetectorScope()
+	UE_FORCEINLINE_HINT FDisableHitchDetectorScope()
 	{
 		FGameThreadHitchHeartBeat::Get().SuspendHeartBeat();
 	}
-	FORCEINLINE ~FDisableHitchDetectorScope()
+	UE_FORCEINLINE_HINT ~FDisableHitchDetectorScope()
 	{
 		FGameThreadHitchHeartBeat::Get().ResumeHeartBeat();
 	}
diff --git a/Engine/Source/Runtime/Core/Public/HAL/ThreadSafeBool.h b/Engine/Source/Runtime/Core/Public/HAL/ThreadSafeBool.h
--- a/Engine/Source/Runtime/Core/Public/HAL/ThreadSafeBool.h
+++ b/Engine/Source/Runtime/Core/Public/HAL/ThreadSafeBool.h
@@ -27,7 +27,7 @@ public:
 	/**
 	 * Operator to use this struct as a bool with thread safety
 	 */
-	FORCEINLINE operator bool() const
+	UE_FORCEINLINE_HINT operator bool() const
 	{
 		return GetValue() != 0;
 	}
@@ -47,7 +47,7 @@ public:
 	 * @param bNewValue   Value to set
 	 * @return The old value
 	 */
-	FORCEINLINE bool AtomicSet(bool bNewValue)
+	UE_FORCEINLINE_HINT bool AtomicSet(bool bNewValue)
 	{
 		return Set(bNewValue ? 1 : 0) != 0;
 	}
diff --git a/Engine/Source/Runtime/Core/Public/HAL/UnrealMemory.h b/Engine/Source/Runtime/Core/Public/HAL/UnrealMemory.h
--- a/Engine/Source/Runtime/Core/Public/HAL/UnrealMemory.h
+++ b/Engine/Source/Runtime/Core/Public/HAL/UnrealMemory.h
@@ -29,7 +29,7 @@
 #if MALLOC_GT_HOOKS
 CORE_API void DoGamethreadHook(int32 Index);
 #else
-FORCEINLINE void DoGamethreadHook(int32 Index)
+UE_FORCEINLINE_HINT void DoGamethreadHook(int32 Index)
 { 
 }
 #endif
@@ -64,7 +64,7 @@ struct FScopedMallocTimer
 			Spew();
 		}
 	}
-	static FORCEINLINE void Miss(int32 InIndex)
+	static UE_FORCEINLINE_HINT void Miss(int32 InIndex)
 	{
 		FPlatformAtomics::InterlockedAdd((volatile int64*)&GTotalMisses[InIndex], 1);
 	}
@@ -73,13 +73,13 @@ struct FScopedMallocTimer
 #else
 struct FScopedMallocTimer
 {
-	FORCEINLINE FScopedMallocTimer(int32 InIndex)
+	UE_FORCEINLINE_HINT FScopedMallocTimer(int32 InIndex)
 	{
 	}
-	FORCEINLINE ~FScopedMallocTimer()
+	UE_FORCEINLINE_HINT ~FScopedMallocTimer()
 	{
 	}
-	FORCEINLINE void Hit(int32 InIndex)
+	UE_FORCEINLINE_HINT void Hit(int32 InIndex)
 	{
 	}
 }; 
@@ -106,17 +106,17 @@ struct FMemory
 
 	/** @name Memory functions (wrapper for FPlatformMemory) */
 
-	static FORCEINLINE void* Memmove( void* Dest, const void* Src, SIZE_T Count )
+	static UE_FORCEINLINE_HINT void* Memmove( void* Dest, const void* Src, SIZE_T Count )
 	{
 		return FPlatformMemory::Memmove( Dest, Src, Count );
 	}
 
-	static FORCEINLINE int32 Memcmp( const void* Buf1, const void* Buf2, SIZE_T Count )
+	static UE_FORCEINLINE_HINT int32 Memcmp( const void* Buf1, const void* Buf2, SIZE_T Count )
 	{
 		return FPlatformMemory::Memcmp( Buf1, Buf2, Count );
 	}
 
-	static FORCEINLINE void* Memset(void* Dest, uint8 Char, SIZE_T Count)
+	static UE_FORCEINLINE_HINT void* Memset(void* Dest, uint8 Char, SIZE_T Count)
 	{
 		return FPlatformMemory::Memset( Dest, Char, Count );
 	}
@@ -128,7 +128,7 @@ struct FMemory
 		Memset( &Src, ValueToSet, sizeof( T ) );
 	}
 
-	static FORCEINLINE void* Memzero(void* Dest, SIZE_T Count)
+	static UE_FORCEINLINE_HINT void* Memzero(void* Dest, SIZE_T Count)
 	{
 		return FPlatformMemory::Memzero( Dest, Count );
 	}
@@ -157,7 +157,7 @@ struct FMemory
 		Memzero( &Src, sizeof( T ) );
 	}
 
-	static FORCEINLINE void* Memcpy(void* Dest, const void* Src, SIZE_T Count)
+	static UE_FORCEINLINE_HINT void* Memcpy(void* Dest, const void* Src, SIZE_T Count)
 	{
 		return FPlatformMemory::Memcpy(Dest,Src,Count);
 	}
@@ -169,22 +169,22 @@ struct FMemory
 		Memcpy( &Dest, &Src, sizeof( T ) );
 	}
 
-	static FORCEINLINE void* BigBlockMemcpy(void* Dest, const void* Src, SIZE_T Count)
+	static UE_FORCEINLINE_HINT void* BigBlockMemcpy(void* Dest, const void* Src, SIZE_T Count)
 	{
 		return FPlatformMemory::BigBlockMemcpy(Dest,Src,Count);
 	}
 
-	static FORCEINLINE void* StreamingMemcpy(void* Dest, const void* Src, SIZE_T Count)
+	static UE_FORCEINLINE_HINT void* StreamingMemcpy(void* Dest, const void* Src, SIZE_T Count)
 	{
 		return FPlatformMemory::StreamingMemcpy(Dest,Src,Count);
 	}
 
-	static FORCEINLINE void* ParallelMemcpy(void* Dest, const void* Src, SIZE_T Count, EMemcpyCachePolicy Policy = EMemcpyCachePolicy::StoreCached)
+	static UE_FORCEINLINE_HINT void* ParallelMemcpy(void* Dest, const void* Src, SIZE_T Count, EMemcpyCachePolicy Policy = EMemcpyCachePolicy::StoreCached)
 	{
 		return FPlatformMemory::ParallelMemcpy(Dest, Src, Count, Policy);
 	}
 
-	static FORCEINLINE void Memswap( void* Ptr1, void* Ptr2, SIZE_T Size )
+	static UE_FORCEINLINE_HINT void Memswap( void* Ptr1, void* Ptr2, SIZE_T Size )
 	{
 		FPlatformMemory::Memswap(Ptr1,Ptr2,Size);
 	}
diff --git a/Engine/Source/Runtime/Core/Public/Hash/BuzHash.h b/Engine/Source/Runtime/Core/Public/Hash/BuzHash.h
--- a/Engine/Source/Runtime/Core/Public/Hash/BuzHash.h
+++ b/Engine/Source/Runtime/Core/Public/Hash/BuzHash.h
@@ -13,13 +13,13 @@ struct FBuzHash
 {
 public:
 	/** Reset the state of the hasher. */
-	FORCEINLINE void Reset() { *this = FBuzHash(); }
+	UE_FORCEINLINE_HINT void Reset() { *this = FBuzHash(); }
 
 	/** Gets the current hash value. */
-	FORCEINLINE uint32 Get() const { return State; }
+	UE_FORCEINLINE_HINT uint32 Get() const { return State; }
 
 	/** Gets the size of the window. */
-	FORCEINLINE uint64 GetWindowSize() const { return Count; }
+	UE_FORCEINLINE_HINT uint64 GetWindowSize() const { return Count; }
 
 	/** Appends a byte to the start of the window and updates the hash. */
 	inline void Add(uint8 X)
@@ -29,7 +29,7 @@ public:
 	}
 
 	/** Adds an array of bytes to the start of the window. */
-	FORCEINLINE void Add(FMemoryView View)
+	UE_FORCEINLINE_HINT void Add(FMemoryView View)
 	{
 		Add((const uint8*)View.GetData(), View.GetSize());
 	}
diff --git a/Engine/Source/Runtime/Core/Public/IOS/IOSPlatformProcess.h b/Engine/Source/Runtime/Core/Public/IOS/IOSPlatformProcess.h
--- a/Engine/Source/Runtime/Core/Public/IOS/IOSPlatformProcess.h
+++ b/Engine/Source/Runtime/Core/Public/IOS/IOSPlatformProcess.h
@@ -13,12 +13,12 @@ struct FProcHandle : public TProcHandle<void*, nullptr>
 {
 public:
 	/** Default constructor. */
-	FORCEINLINE FProcHandle()
+	UE_FORCEINLINE_HINT FProcHandle()
 		: TProcHandle()
 	{}
 
 	/** Initialization constructor. */
-	FORCEINLINE explicit FProcHandle( HandleType Other )
+	UE_FORCEINLINE_HINT explicit FProcHandle( HandleType Other )
 		: TProcHandle( Other )
 	{}
 };
diff --git a/Engine/Source/Runtime/Core/Public/IOS/IOSPlatformProperties.h b/Engine/Source/Runtime/Core/Public/IOS/IOSPlatformProperties.h
--- a/Engine/Source/Runtime/Core/Public/IOS/IOSPlatformProperties.h
+++ b/Engine/Source/Runtime/Core/Public/IOS/IOSPlatformProperties.h
@@ -18,102 +18,102 @@
 struct FIOSPlatformProperties
 	: public FGenericPlatformProperties
 {
-	static constexpr FORCEINLINE bool HasEditorOnlyData()
+	static constexpr UE_FORCEINLINE_HINT bool HasEditorOnlyData()
 	{
 		return false;
 	}
 
-	static constexpr FORCEINLINE const char* PlatformName()
+	static constexpr UE_FORCEINLINE_HINT const char* PlatformName()
 	{
 		return "IOS";
 	}
 
-	static constexpr FORCEINLINE const char* IniPlatformName()
+	static constexpr UE_FORCEINLINE_HINT const char* IniPlatformName()
 	{
 		return "IOS";
 	}
 
-	static constexpr FORCEINLINE const TCHAR* GetRuntimeSettingsClassName()
+	static constexpr UE_FORCEINLINE_HINT const TCHAR* GetRuntimeSettingsClassName()
 	{
 		return TEXT("/Script/IOSRuntimeSettings.IOSRuntimeSettings");
 	}
 
-	static constexpr FORCEINLINE bool IsGameOnly()
+	static constexpr UE_FORCEINLINE_HINT bool IsGameOnly()
 	{
 		return true;
 	}
 	
-	static constexpr FORCEINLINE bool RequiresCookedData()
+	static constexpr UE_FORCEINLINE_HINT bool RequiresCookedData()
 	{
 		return true;
 	}
     
-	static FORCEINLINE bool SupportsBuildTarget(EBuildTargetType TargetType)
+	static UE_FORCEINLINE_HINT bool SupportsBuildTarget(EBuildTargetType TargetType)
 	{
 		return (TargetType == EBuildTargetType::Game);
 	}
 
-	static constexpr FORCEINLINE bool SupportsLowQualityLightmaps()
+	static constexpr UE_FORCEINLINE_HINT bool SupportsLowQualityLightmaps()
 	{
 		return true;
 	}
 
-	static constexpr FORCEINLINE bool SupportsHighQualityLightmaps()
+	static constexpr UE_FORCEINLINE_HINT bool SupportsHighQualityLightmaps()
 	{
 		return true;
 	}
 
-	static constexpr FORCEINLINE bool SupportsTextureStreaming()
+	static constexpr UE_FORCEINLINE_HINT bool SupportsTextureStreaming()
 	{
 		return true;
 	}
 
-	static constexpr FORCEINLINE bool SupportsMemoryMappedFiles()
+	static constexpr UE_FORCEINLINE_HINT bool SupportsMemoryMappedFiles()
 	{
 		return true;
 	}
 
-	static constexpr FORCEINLINE bool SupportsMemoryMappedAudio()
+	static constexpr UE_FORCEINLINE_HINT bool SupportsMemoryMappedAudio()
 	{
 		return true;
 	}
 
-	static constexpr FORCEINLINE bool SupportsMemoryMappedAnimation()
+	static constexpr UE_FORCEINLINE_HINT bool SupportsMemoryMappedAnimation()
 	{
 		return true;
 	}
 
-	static constexpr FORCEINLINE int64 GetMemoryMappingAlignment()
+	static constexpr UE_FORCEINLINE_HINT int64 GetMemoryMappingAlignment()
 	{
 		return 16384;
 	}
 
-	static constexpr FORCEINLINE int GetMaxSupportedVirtualMemoryAlignment()
+	static constexpr UE_FORCEINLINE_HINT int GetMaxSupportedVirtualMemoryAlignment()
 	{
 		return 16384;
 	}
 
-	static constexpr FORCEINLINE bool HasFixedResolution()
+	static constexpr UE_FORCEINLINE_HINT bool HasFixedResolution()
 	{
 		return true;
 	}
 
-	static constexpr FORCEINLINE bool AllowsFramerateSmoothing()
+	static constexpr UE_FORCEINLINE_HINT bool AllowsFramerateSmoothing()
 	{
 		return true;
 	}
 
-	static constexpr FORCEINLINE bool SupportsAudioStreaming()
+	static constexpr UE_FORCEINLINE_HINT bool SupportsAudioStreaming()
 	{
 		return true;
 	}
 	
-	static constexpr FORCEINLINE bool SupportsMeshLODStreaming()
+	static constexpr UE_FORCEINLINE_HINT bool SupportsMeshLODStreaming()
 	{
 		return true;
 	}
 
-	static constexpr FORCEINLINE bool SupportsQuit()
+	static constexpr UE_FORCEINLINE_HINT bool SupportsQuit()
 	{
 		return true;
 	}
@@ -122,12 +122,12 @@ struct FIOSPlatformProperties
 struct FTVOSPlatformProperties : public FIOSPlatformProperties
 {
 	// @todo breaking change here!
-	static constexpr FORCEINLINE const char* PlatformName()
+	static constexpr UE_FORCEINLINE_HINT const char* PlatformName()
 	{
 		return "TVOS";
 	}
 
-	static constexpr FORCEINLINE const char* IniPlatformName()
+	static constexpr UE_FORCEINLINE_HINT const char* IniPlatformName()
 	{
 		return "TVOS";
 	}
@@ -135,12 +135,12 @@ struct FTVOSPlatformProperties : public FIOSPlatformProperties
 
 struct FVisionOSPlatformProperties : public FIOSPlatformProperties
 {
-	static constexpr FORCEINLINE const char* PlatformName()
+	static constexpr UE_FORCEINLINE_HINT const char* PlatformName()
 	{
 		return "IOS";
 	}
 
-	static constexpr FORCEINLINE const char* IniPlatformName()
+	static constexpr UE_FORCEINLINE_HINT const char* IniPlatformName()
 	{
 		return "VisionOS";
 	}
diff --git a/Engine/Source/Runtime/Core/Public/Internationalization/FastDecimalFormat.h b/Engine/Source/Runtime/Core/Public/Internationalization/FastDecimalFormat.h
--- a/Engine/Source/Runtime/Core/Public/Internationalization/FastDecimalFormat.h
+++ b/Engine/Source/Runtime/Core/Public/Internationalization/FastDecimalFormat.h
@@ -156,7 +156,7 @@ inline FString NumberToString(const T InVal, const FDecimalNumberFormattingRules
 		OutVal *= (bIsNegative ? -1 : 1);																																																	\
 		return bResult;																																																						\
 	}																																																										\
-	FORCEINLINE bool StringToNumber(const TCHAR* InStr, const FDecimalNumberFormattingRules& InFormattingRules, const FNumberParsingOptions& InParsingOptions, NUMBER_TYPE& OutVal, int32* OutParsedLen = nullptr)							\
+	UE_FORCEINLINE_HINT bool StringToNumber(const TCHAR* InStr, const FDecimalNumberFormattingRules& InFormattingRules, const FNumberParsingOptions& InParsingOptions, NUMBER_TYPE& OutVal, int32* OutParsedLen = nullptr)							\
 	{																																																										\
 		return StringToNumber(InStr, FCString::Strlen(InStr), InFormattingRules, InParsingOptions, OutVal, OutParsedLen);																													\
 	}
@@ -169,7 +169,7 @@ inline FString NumberToString(const T InVal, const FDecimalNumberFormattingRules
 		OutVal = static_cast<NUMBER_TYPE>(Val);																																																\
 		return bResult;																																																						\
 	}																																																										\
-	FORCEINLINE bool StringToNumber(const TCHAR* InStr, const FDecimalNumberFormattingRules& InFormattingRules, const FNumberParsingOptions& InParsingOptions, NUMBER_TYPE& OutVal, int32* OutParsedLen = nullptr)							\
+	UE_FORCEINLINE_HINT bool StringToNumber(const TCHAR* InStr, const FDecimalNumberFormattingRules& InFormattingRules, const FNumberParsingOptions& InParsingOptions, NUMBER_TYPE& OutVal, int32* OutParsedLen = nullptr)							\
 	{																																																										\
 		return StringToNumber(InStr, FCString::Strlen(InStr), InFormattingRules, InParsingOptions, OutVal, OutParsedLen);																													\
 	}
diff --git a/Engine/Source/Runtime/Core/Public/Internationalization/ILocalizedTextSource.h b/Engine/Source/Runtime/Core/Public/Internationalization/ILocalizedTextSource.h
--- a/Engine/Source/Runtime/Core/Public/Internationalization/ILocalizedTextSource.h
+++ b/Engine/Source/Runtime/Core/Public/Internationalization/ILocalizedTextSource.h
@@ -53,37 +53,37 @@ public:
 	}
 
 	/** Should we load native data based on the given load flags and environment? */
-	static FORCEINLINE bool ShouldLoadNative(const ELocalizationLoadFlags InLoadFlags)
+	static UE_FORCEINLINE_HINT bool ShouldLoadNative(const ELocalizationLoadFlags InLoadFlags)
 	{
 		return EnumHasAnyFlags(InLoadFlags, ELocalizationLoadFlags::Native);
 	}
 
 	/** Should we load editor data based on the given load flags and environment? */
-	static FORCEINLINE bool ShouldLoadEditor(const ELocalizationLoadFlags InLoadFlags)
+	static UE_FORCEINLINE_HINT bool ShouldLoadEditor(const ELocalizationLoadFlags InLoadFlags)
 	{
 		return EnumHasAnyFlags(InLoadFlags, ELocalizationLoadFlags::Editor);
 	}
 
 	/** Should we load game data based on the given load flags and environment? */
-	static FORCEINLINE bool ShouldLoadGame(const ELocalizationLoadFlags InLoadFlags)
+	static UE_FORCEINLINE_HINT bool ShouldLoadGame(const ELocalizationLoadFlags InLoadFlags)
 	{
 		return EnumHasAnyFlags(InLoadFlags, ELocalizationLoadFlags::Game | ELocalizationLoadFlags::ForceLocalizedGame);
 	}
 
 	/** Should we load engine data based on the given load flags and environment? */
-	static FORCEINLINE bool ShouldLoadEngine(const ELocalizationLoadFlags InLoadFlags)
+	static UE_FORCEINLINE_HINT bool ShouldLoadEngine(const ELocalizationLoadFlags InLoadFlags)
 	{
 		return EnumHasAnyFlags(InLoadFlags, ELocalizationLoadFlags::Engine);
 	}
 
 	/** Should we load additional (eg, plugin) data based on the given load flags and environment? */
-	static FORCEINLINE bool ShouldLoadAdditional(const ELocalizationLoadFlags InLoadFlags)
+	static UE_FORCEINLINE_HINT bool ShouldLoadAdditional(const ELocalizationLoadFlags InLoadFlags)
 	{
 		return EnumHasAnyFlags(InLoadFlags, ELocalizationLoadFlags::Additional);
 	}
 
 	/** Should we load native game data based on the given load flags and environment? */
-	static FORCEINLINE bool ShouldLoadNativeGameData(const ELocalizationLoadFlags InLoadFlags)
+	static UE_FORCEINLINE_HINT bool ShouldLoadNativeGameData(const ELocalizationLoadFlags InLoadFlags)
 	{
 		// The editor loads native game data by default to prevent authoring issues
 		// It will load localized data only if the request is forced (eg, when entering game localization preview mode)
diff --git a/Engine/Source/Runtime/Core/Public/Internationalization/Internationalization.h b/Engine/Source/Runtime/Core/Public/Internationalization/Internationalization.h
--- a/Engine/Source/Runtime/Core/Public/Internationalization/Internationalization.h
+++ b/Engine/Source/Runtime/Core/Public/Internationalization/Internationalization.h
@@ -33,7 +33,7 @@ public:
 	static CORE_API void TearDown();
 
 	UE_DEPRECATED(5.5, "FInternationalization::ForUseOnlyByLocMacroAndGraphNodeTextLiterals_CreateText is deprecated in favor of FText::AsLocalizable_Advanced (NOTE: FText::AsLocalizable_Advanced takes its arguments in a different order!).")
-	static FORCEINLINE FText ForUseOnlyByLocMacroAndGraphNodeTextLiterals_CreateText(const TCHAR* InTextLiteral, const TCHAR* InNamespace, const TCHAR* InKey)
+	static UE_FORCEINLINE_HINT FText ForUseOnlyByLocMacroAndGraphNodeTextLiterals_CreateText(const TCHAR* InTextLiteral, const TCHAR* InNamespace, const TCHAR* InKey)
 	{
 		return FText::AsLocalizable_Advanced(InNamespace, InKey, InTextLiteral);
 	}
diff --git a/Engine/Source/Runtime/Core/Public/Internationalization/LocKeyFuncs.h b/Engine/Source/Runtime/Core/Public/Internationalization/LocKeyFuncs.h
--- a/Engine/Source/Runtime/Core/Public/Internationalization/LocKeyFuncs.h
+++ b/Engine/Source/Runtime/Core/Public/Internationalization/LocKeyFuncs.h
@@ -70,32 +70,32 @@ public:
 		return *this;
 	}
 
-	FORCEINLINE bool operator==(const FLocKey& Other) const
+	UE_FORCEINLINE_HINT bool operator==(const FLocKey& Other) const
 	{
 		return Hash == Other.Hash && Compare(Other) == 0;
 	}
 
-	FORCEINLINE bool operator!=(const FLocKey& Other) const
+	UE_FORCEINLINE_HINT bool operator!=(const FLocKey& Other) const
 	{
 		return Hash != Other.Hash || Compare(Other) != 0;
 	}
 
-	FORCEINLINE bool operator<(const FLocKey& Other) const
+	UE_FORCEINLINE_HINT bool operator<(const FLocKey& Other) const
 	{
 		return Compare(Other) < 0;
 	}
 
-	FORCEINLINE bool operator<=(const FLocKey& Other) const
+	UE_FORCEINLINE_HINT bool operator<=(const FLocKey& Other) const
 	{
 		return Compare(Other) <= 0;
 	}
 
-	FORCEINLINE bool operator>(const FLocKey& Other) const
+	UE_FORCEINLINE_HINT bool operator>(const FLocKey& Other) const
 	{
 		return Compare(Other) > 0;
 	}
 
-	FORCEINLINE bool operator>=(const FLocKey& Other) const
+	UE_FORCEINLINE_HINT bool operator>=(const FLocKey& Other) const
 	{
 		return Compare(Other) >= 0;
 	}
@@ -105,27 +105,27 @@ public:
 		return Id.Hash;
 	}
 
-	FORCEINLINE bool IsEmpty() const
+	UE_FORCEINLINE_HINT bool IsEmpty() const
 	{
 		return Str.IsEmpty();
 	}
 
-	FORCEINLINE bool Equals(const FLocKey& Other) const
+	UE_FORCEINLINE_HINT bool Equals(const FLocKey& Other) const
 	{
 		return *this == Other;
 	}
 
-	FORCEINLINE int32 Compare(const FLocKey& Other) const
+	UE_FORCEINLINE_HINT int32 Compare(const FLocKey& Other) const
 	{
 		return FCString::Strcmp(*Str, *Other.Str);
 	}
 
-	FORCEINLINE const FString& GetString() const
+	UE_FORCEINLINE_HINT const FString& GetString() const
 	{
 		return Str;
 	}
 
-	static FORCEINLINE uint32 ProduceHash(const FString& InStr, const uint32 InBaseHash = 0)
+	static UE_FORCEINLINE_HINT uint32 ProduceHash(const FString& InStr, const uint32 InBaseHash = 0)
 	{
 		return FCrc::StrCrc32<TCHAR>(*InStr, InBaseHash);
 	}
@@ -141,15 +141,15 @@ private:
 /** Case sensitive hashing function for TSet */
 struct FLocKeySetFuncs : BaseKeyFuncs<FString, FString>
 {
-	static FORCEINLINE const FString& GetSetKey(const FString& Element)
+	static UE_FORCEINLINE_HINT const FString& GetSetKey(const FString& Element)
 	{
 		return Element;
 	}
-	static FORCEINLINE bool Matches(const FString& A, const FString& B)
+	static UE_FORCEINLINE_HINT bool Matches(const FString& A, const FString& B)
 	{
 		return A.Equals(B, ESearchCase::CaseSensitive);
 	}
-	static FORCEINLINE uint32 GetKeyHash(const FString& Key)
+	static UE_FORCEINLINE_HINT uint32 GetKeyHash(const FString& Key)
 	{
 		return FLocKey::ProduceHash(Key);
 	}
@@ -159,15 +159,15 @@ struct FLocKeySetFuncs : BaseKeyFuncs<FString, FString>
 template <typename ValueType>
 struct FLocKeyMapFuncs : BaseKeyFuncs<ValueType, FString, /*bInAllowDuplicateKeys*/false>
 {
-	static FORCEINLINE const FString& GetSetKey(const TPair<FString, ValueType>& Element)
+	static UE_FORCEINLINE_HINT const FString& GetSetKey(const TPair<FString, ValueType>& Element)
 	{
 		return Element.Key;
 	}
-	static FORCEINLINE bool Matches(const FString& A, const FString& B)
+	static UE_FORCEINLINE_HINT bool Matches(const FString& A, const FString& B)
 	{
 		return A.Equals(B, ESearchCase::CaseSensitive);
 	}
-	static FORCEINLINE uint32 GetKeyHash(const FString& Key)
+	static UE_FORCEINLINE_HINT uint32 GetKeyHash(const FString& Key)
 	{
 		return FLocKey::ProduceHash(Key);
 	}
@@ -177,15 +177,15 @@ struct FLocKeyMapFuncs : BaseKeyFuncs<ValueType, FString, /*bInAllowDuplicateKey
 template <typename ValueType>
 struct FLocKeyMultiMapFuncs : BaseKeyFuncs<ValueType, FString, /*bInAllowDuplicateKeys*/true>
 {
-	static FORCEINLINE const FString& GetSetKey(const TPair<FString, ValueType>& Element)
+	static UE_FORCEINLINE_HINT const FString& GetSetKey(const TPair<FString, ValueType>& Element)
 	{
 		return Element.Key;
 	}
-	static FORCEINLINE bool Matches(const FString& A, const FString& B)
+	static UE_FORCEINLINE_HINT bool Matches(const FString& A, const FString& B)
 	{
 		return A.Equals(B, ESearchCase::CaseSensitive);
 	}
-	static FORCEINLINE uint32 GetKeyHash(const FString& Key)
+	static UE_FORCEINLINE_HINT uint32 GetKeyHash(const FString& Key)
 	{
 		return FLocKey::ProduceHash(Key);
 	}
@@ -194,7 +194,7 @@ struct FLocKeyMultiMapFuncs : BaseKeyFuncs<ValueType, FString, /*bInAllowDuplica
 /** Case sensitive comparison function for TSortedMap */
 struct FLocKeySortedMapLess
 {
-	FORCEINLINE bool operator()(const FString& A, const FString& B) const
+	UE_FORCEINLINE_HINT bool operator()(const FString& A, const FString& B) const
 	{
 		return A.Compare(B, ESearchCase::CaseSensitive) < 0;
 	}
diff --git a/Engine/Source/Runtime/Core/Public/Internationalization/Text.h b/Engine/Source/Runtime/Core/Public/Internationalization/Text.h
--- a/Engine/Source/Runtime/Core/Public/Internationalization/Text.h
+++ b/Engine/Source/Runtime/Core/Public/Internationalization/Text.h
@@ -553,7 +553,7 @@ public:
 	static CORE_API bool FindTextInLiveTable_Advanced(const FTextKey& Namespace, const FTextKey& Key, FText& OutText, const FString* const SourceString = nullptr);
 
 	UE_DEPRECATED(5.5, "FText::FindText is deprecated as it can no longer find untranslated text. Consider FText::FindTextInLiveTable_Advanced or FText::AsLocalizable_Advanced as a replacement, depending on your use-case.")
-	static FORCEINLINE bool FindText(const FTextKey& Namespace, const FTextKey& Key, FText& OutText, const FString* const SourceString = nullptr)
+	static UE_FORCEINLINE_HINT bool FindText(const FTextKey& Namespace, const FTextKey& Key, FText& OutText, const FString* const SourceString = nullptr)
 	{
 		return FText::FindTextInLiveTable_Advanced(Namespace, Key, OutText, SourceString);
 	}
@@ -1000,7 +1000,7 @@ public:
 	CORE_API void ToExportedString(FString& OutResult, const bool bStripPackageNamespace = false) const;
 	CORE_API const TCHAR* FromExportedString(const TCHAR* InBuffer);
 
-	FORCEINLINE EFormatArgumentType::Type GetType() const
+	UE_FORCEINLINE_HINT EFormatArgumentType::Type GetType() const
 	{
 		return Type;
 	}
diff --git a/Engine/Source/Runtime/Core/Public/Internationalization/TextKey.h b/Engine/Source/Runtime/Core/Public/Internationalization/TextKey.h
--- a/Engine/Source/Runtime/Core/Public/Internationalization/TextKey.h
+++ b/Engine/Source/Runtime/Core/Public/Internationalization/TextKey.h
@@ -17,7 +17,7 @@ namespace TextKeyUtil
 {
 	/** Utility to produce a hash for a UTF-16 string (as used by FTextKey) */
 	CORE_API uint32 HashString(const FTCHARToUTF16& InStr);
-	FORCEINLINE uint32 HashString(const FTCHARToUTF16& InStr, const uint32 InBaseHash)
+	UE_FORCEINLINE_HINT uint32 HashString(const FTCHARToUTF16& InStr, const uint32 InBaseHash)
 	{
 		return HashCombine(HashString(InStr), InBaseHash);
 	}
@@ -47,21 +47,21 @@ namespace TextKeyUtil
 	}
 
 	/** Utility to produce a hash for a string (as used by FTextKey) */
-	FORCEINLINE uint32 HashString(const FString& InStr)
+	UE_FORCEINLINE_HINT uint32 HashString(const FString& InStr)
 	{
 		return HashString(*InStr, InStr.Len());
 	}
-	FORCEINLINE uint32 HashString(const FString& InStr, const uint32 InBaseHash)
+	UE_FORCEINLINE_HINT uint32 HashString(const FString& InStr, const uint32 InBaseHash)
 	{
 		return HashString(*InStr, InStr.Len(), InBaseHash);
 	}
 
 	/** Utility to produce a hash for a string (as used by FTextKey) */
-	FORCEINLINE uint32 HashString(FStringView InStr)
+	UE_FORCEINLINE_HINT uint32 HashString(FStringView InStr)
 	{
 		return HashString(InStr.GetData(), InStr.Len());
 	}
-	FORCEINLINE uint32 HashString(FStringView InStr, const uint32 InBaseHash)
+	UE_FORCEINLINE_HINT uint32 HashString(FStringView InStr, const uint32 InBaseHash)
 	{
 		return HashString(InStr.GetData(), InStr.Len(), InBaseHash);
 	}
@@ -91,13 +91,13 @@ public:
 	CORE_API void AppendString(FStringBuilderBase& Out) const;
 
 	/** Compare for equality */
-	friend FORCEINLINE bool operator==(const FTextKey& A, const FTextKey& B)
+	friend UE_FORCEINLINE_HINT bool operator==(const FTextKey& A, const FTextKey& B)
 	{
 		return A.Index == B.Index;
 	}
 
 	/** Compare for inequality */
-	friend FORCEINLINE bool operator!=(const FTextKey& A, const FTextKey& B)
+	friend UE_FORCEINLINE_HINT bool operator!=(const FTextKey& A, const FTextKey& B)
 	{
 		return A.Index != B.Index;
 	}
@@ -124,7 +124,7 @@ public:
 	CORE_API void SerializeDiscardHash(FStructuredArchiveSlot Slot);
 
 	/** Is this text key empty? */
-	FORCEINLINE bool IsEmpty() const
+	UE_FORCEINLINE_HINT bool IsEmpty() const
 	{
 		return Index == INDEX_NONE;
 	}
@@ -171,31 +171,31 @@ public:
 	}
 
 	/** Get the namespace component of this text identity */
-	FORCEINLINE FTextKey GetNamespace() const
+	UE_FORCEINLINE_HINT FTextKey GetNamespace() const
 	{
 		return Namespace;
 	}
 
 	/** Get the key component of this text identity */
-	FORCEINLINE FTextKey GetKey() const
+	UE_FORCEINLINE_HINT FTextKey GetKey() const
 	{
 		return Key;
 	}
 
 	/** Compare for equality */
-	friend FORCEINLINE bool operator==(const FTextId& A, const FTextId& B)
+	friend UE_FORCEINLINE_HINT bool operator==(const FTextId& A, const FTextId& B)
 	{
 		return A.Namespace == B.Namespace && A.Key == B.Key;
 	}
 
 	/** Compare for inequality */
-	friend FORCEINLINE bool operator!=(const FTextId& A, const FTextId& B)
+	friend UE_FORCEINLINE_HINT bool operator!=(const FTextId& A, const FTextId& B)
 	{
 		return A.Namespace != B.Namespace || A.Key != B.Key;
 	}
 
 	/** Get the hash of this text identity */
-	friend FORCEINLINE uint32 GetTypeHash(const FTextId& A)
+	friend UE_FORCEINLINE_HINT uint32 GetTypeHash(const FTextId& A)
 	{
 		return HashCombine(GetTypeHash(A.Namespace), GetTypeHash(A.Key));
 	}
@@ -249,7 +249,7 @@ public:
 	}
 
 	/** Is this text identity empty? */
-	FORCEINLINE bool IsEmpty() const
+	UE_FORCEINLINE_HINT bool IsEmpty() const
 	{
 		return Namespace.IsEmpty() && Key.IsEmpty();
 	}
diff --git a/Engine/Source/Runtime/Core/Public/Internationalization/TextLocalizationResource.h b/Engine/Source/Runtime/Core/Public/Internationalization/TextLocalizationResource.h
--- a/Engine/Source/Runtime/Core/Public/Internationalization/TextLocalizationResource.h
+++ b/Engine/Source/Runtime/Core/Public/Internationalization/TextLocalizationResource.h
@@ -74,13 +74,13 @@ public:
 	FEntriesTable Entries;
 
 	/** Utility to produce a hash for a string (as used by SourceStringHash) */
-	static FORCEINLINE uint32 HashString(const TCHAR* InStr, const uint32 InBaseHash = 0)
+	static UE_FORCEINLINE_HINT uint32 HashString(const TCHAR* InStr, const uint32 InBaseHash = 0)
 	{
 		return FCrc::StrCrc32<TCHAR>(InStr, InBaseHash);
 	}
 
 	/** Utility to produce a hash for a string (as used by SourceStringHash) */
-	static FORCEINLINE uint32 HashString(const FString& InStr, const uint32 InBaseHash = 0)
+	static UE_FORCEINLINE_HINT uint32 HashString(const FString& InStr, const uint32 InBaseHash = 0)
 	{
 		return FCrc::StrCrc32<TCHAR>(*InStr, InBaseHash);
 	}
diff --git a/Engine/Source/Runtime/Core/Public/Linux/LinuxPlatformProperties.h b/Engine/Source/Runtime/Core/Public/Linux/LinuxPlatformProperties.h
--- a/Engine/Source/Runtime/Core/Public/Linux/LinuxPlatformProperties.h
+++ b/Engine/Source/Runtime/Core/Public/Linux/LinuxPlatformProperties.h
@@ -20,37 +20,37 @@ template<bool HAS_EDITOR_DATA, bool IS_DEDICATED_SERVER, bool IS_CLIENT_ONLY, bo
 struct FLinuxPlatformProperties
 	: public FGenericPlatformProperties
 {
-	static constexpr FORCEINLINE bool HasEditorOnlyData()
+	static constexpr UE_FORCEINLINE_HINT bool HasEditorOnlyData()
 	{
 		return HAS_EDITOR_DATA;
 	}
 
-	static constexpr FORCEINLINE const char* IniPlatformName()
+	static constexpr UE_FORCEINLINE_HINT const char* IniPlatformName()
 	{
 		return IS_ARM64 ? "LinuxArm64" : "Linux";
 	}
 
-	static constexpr FORCEINLINE const TCHAR* GetRuntimeSettingsClassName()
+	static constexpr UE_FORCEINLINE_HINT const TCHAR* GetRuntimeSettingsClassName()
 	{
 		return TEXT("/Script/LinuxTargetPlatform.LinuxTargetSettings");
 	}
 
-	static constexpr FORCEINLINE bool IsGameOnly()
+	static constexpr UE_FORCEINLINE_HINT bool IsGameOnly()
 	{
 		return UE_GAME;
 	}
 
-	static constexpr FORCEINLINE bool IsServerOnly()
+	static constexpr UE_FORCEINLINE_HINT bool IsServerOnly()
 	{
 		return IS_DEDICATED_SERVER;
 	}
 
-	static constexpr FORCEINLINE bool IsClientOnly()
+	static constexpr UE_FORCEINLINE_HINT bool IsClientOnly()
 	{
 		return IS_CLIENT_ONLY;
 	}
 
-	static constexpr FORCEINLINE bool IsArm64()
+	static constexpr UE_FORCEINLINE_HINT bool IsArm64()
 	{
 		return IS_ARM64;
 	}
@@ -75,17 +75,17 @@ struct FLinuxPlatformProperties
 		return IS_ARM64 ? "LinuxArm64" : "Linux";
 	}
 
-	static constexpr FORCEINLINE bool RequiresCookedData()
+	static constexpr UE_FORCEINLINE_HINT bool RequiresCookedData()
 	{
 		return !HAS_EDITOR_DATA;
 	}
 
-	static constexpr FORCEINLINE bool HasSecurePackageFormat()
+	static constexpr UE_FORCEINLINE_HINT bool HasSecurePackageFormat()
 	{
 		return IS_DEDICATED_SERVER;
 	}
 
-	static constexpr FORCEINLINE bool RequiresUserCredentials()
+	static constexpr UE_FORCEINLINE_HINT bool RequiresUserCredentials()
 	{
 		return true;
 	}
@@ -100,32 +100,32 @@ struct FLinuxPlatformProperties
 #endif
 	}
 
-	static constexpr FORCEINLINE bool SupportsMultipleGameInstances()
+	static constexpr UE_FORCEINLINE_HINT bool SupportsMultipleGameInstances()
 	{
 		return true;
 	}
 
-	static constexpr FORCEINLINE bool HasFixedResolution()
+	static constexpr UE_FORCEINLINE_HINT bool HasFixedResolution()
 	{
 		return false;
 	}
 
-	static constexpr FORCEINLINE bool SupportsWindowedMode()
+	static constexpr UE_FORCEINLINE_HINT bool SupportsWindowedMode()
 	{
 		return !IS_DEDICATED_SERVER;
 	}
 
-	static constexpr FORCEINLINE bool AllowsFramerateSmoothing()
+	static constexpr UE_FORCEINLINE_HINT bool AllowsFramerateSmoothing()
 	{
 		return true;
 	}
 
-	static constexpr FORCEINLINE bool SupportsRayTracing()
+	static constexpr UE_FORCEINLINE_HINT bool SupportsRayTracing()
 	{
 		return true;
 	}
 
-	static constexpr FORCEINLINE bool SupportsQuit()
+	static constexpr UE_FORCEINLINE_HINT bool SupportsQuit()
 	{
 		return true;
 	}
@@ -150,17 +150,17 @@ struct FLinuxPlatformProperties
 		return 1.0f;
 	}
 
-	static constexpr FORCEINLINE bool AllowsCallStackDumpDuringAssert()
+	static constexpr UE_FORCEINLINE_HINT bool AllowsCallStackDumpDuringAssert()
 	{
 		return true;
 	}
 
-	static constexpr FORCEINLINE bool SupportsAudioStreaming()
+	static constexpr UE_FORCEINLINE_HINT bool SupportsAudioStreaming()
 	{
 		return !IsServerOnly();
 	}
 
-	static constexpr FORCEINLINE int64 GetMemoryMappingAlignment()
+	static constexpr UE_FORCEINLINE_HINT int64 GetMemoryMappingAlignment()
 	{
 		return 4096;
 	}
diff --git a/Engine/Source/Runtime/Core/Public/Logging/LogCategory.h b/Engine/Source/Runtime/Core/Public/Logging/LogCategory.h
--- a/Engine/Source/Runtime/Core/Public/Logging/LogCategory.h
+++ b/Engine/Source/Runtime/Core/Public/Logging/LogCategory.h
@@ -30,7 +30,7 @@ struct FLogCategoryBase
 	UE_API ~FLogCategoryBase();
 
 	/** Should not generally be used directly. Tests the runtime verbosity and maybe triggers a debug break, etc. */
-	FORCEINLINE constexpr bool IsSuppressed(ELogVerbosity::Type VerbosityLevel) const
+	UE_FORCEINLINE_HINT constexpr bool IsSuppressed(ELogVerbosity::Type VerbosityLevel) const
 	{
 		return !((VerbosityLevel & ELogVerbosity::VerbosityMask) <= Verbosity);
 	}
@@ -80,7 +80,7 @@ struct FLogCategory : public FLogCategoryBase
 
 	inline constexpr ELogVerbosity::Type GetCompileTimeVerbosity() const { return CompileTimeVerbosity; }
 
-	FORCEINLINE FLogCategory(const FLogCategoryName& InCategoryName)
+	UE_FORCEINLINE_HINT FLogCategory(const FLogCategoryName& InCategoryName)
 		: FLogCategoryBase(InCategoryName, InDefaultVerbosity, CompileTimeVerbosity)
 	{
 	}
diff --git a/Engine/Source/Runtime/Core/Public/Logging/LogMacros.h b/Engine/Source/Runtime/Core/Public/Logging/LogMacros.h
--- a/Engine/Source/Runtime/Core/Public/Logging/LogMacros.h
+++ b/Engine/Source/Runtime/Core/Public/Logging/LogMacros.h
@@ -215,7 +215,7 @@ CORE_API void BasicFatalLogWithProgramCounter(const FLogCategoryBase& Category,
 	namespace UEAsserts_Private
 	{
 		template <int32 VerbosityToCheck, typename CategoryType>
-		FORCEINLINE bool IsLogActive(const CategoryType& Category)
+		UE_FORCEINLINE_HINT bool IsLogActive(const CategoryType& Category)
 		{
 			if constexpr (((VerbosityToCheck & ELogVerbosity::VerbosityMask) <= CategoryType::CompileTimeVerbosity &&
 				(VerbosityToCheck & ELogVerbosity::VerbosityMask) <= ELogVerbosity::COMPILED_IN_MINIMUM_VERBOSITY))
@@ -343,7 +343,7 @@ CORE_API void BasicFatalLogWithProgramCounter(const FLogCategoryBase& Category,
 	#define DECLARE_LOG_CATEGORY_EXTERN(CategoryName, DefaultVerbosity, CompileTimeVerbosity) \
 		extern struct FLogCategory##CategoryName : public FLogCategory<ELogVerbosity::DefaultVerbosity, ELogVerbosity::CompileTimeVerbosity> \
 		{ \
-			FORCEINLINE FLogCategory##CategoryName() : FLogCategory(TEXT(#CategoryName)) {} \
+			UE_FORCEINLINE_HINT FLogCategory##CategoryName() : FLogCategory(TEXT(#CategoryName)) {} \
 		} CategoryName;
 
 	/** 
@@ -361,7 +361,7 @@ CORE_API void BasicFatalLogWithProgramCounter(const FLogCategoryBase& Category,
 	#define DEFINE_LOG_CATEGORY_STATIC(CategoryName, DefaultVerbosity, CompileTimeVerbosity) \
 		static struct FLogCategory##CategoryName : public FLogCategory<ELogVerbosity::DefaultVerbosity, ELogVerbosity::CompileTimeVerbosity> \
 		{ \
-			FORCEINLINE FLogCategory##CategoryName() : FLogCategory(TEXT(#CategoryName)) {} \
+			UE_FORCEINLINE_HINT FLogCategory##CategoryName() : FLogCategory(TEXT(#CategoryName)) {} \
 		} CategoryName;
 
 	/** 
diff --git a/Engine/Source/Runtime/Core/Public/Mac/MacPlatformMath.h b/Engine/Source/Runtime/Core/Public/Mac/MacPlatformMath.h
--- a/Engine/Source/Runtime/Core/Public/Mac/MacPlatformMath.h
+++ b/Engine/Source/Runtime/Core/Public/Mac/MacPlatformMath.h
@@ -22,16 +22,16 @@
         /**
          * Use the SSE instruction to count bits
          */
-        static FORCEINLINE int32 CountBits(uint64 Bits)
+        static UE_FORCEINLINE_HINT int32 CountBits(uint64 Bits)
         {
             return __builtin_popcountll(Bits);
         }
     #endif
 
-        static FORCEINLINE bool IsNaN( float A ) { return isnan(A) != 0; }
-		static FORCEINLINE bool IsNaN(double A) { return isnan(A) != 0; }
-		static FORCEINLINE bool IsFinite(float A) { return isfinite(A); }
-		static FORCEINLINE bool IsFinite(double A) { return isfinite(A); }
+        static UE_FORCEINLINE_HINT bool IsNaN( float A ) { return isnan(A) != 0; }
+		static UE_FORCEINLINE_HINT bool IsNaN(double A) { return isnan(A) != 0; }
+		static UE_FORCEINLINE_HINT bool IsFinite(float A) { return isfinite(A); }
+		static UE_FORCEINLINE_HINT bool IsFinite(double A) { return isfinite(A); }
 	};
 
     typedef FMacPlatformMath FPlatformMath;
diff --git a/Engine/Source/Runtime/Core/Public/Mac/MacPlatformMisc.h b/Engine/Source/Runtime/Core/Public/Mac/MacPlatformMisc.h
--- a/Engine/Source/Runtime/Core/Public/Mac/MacPlatformMisc.h
+++ b/Engine/Source/Runtime/Core/Public/Mac/MacPlatformMisc.h
@@ -19,7 +19,7 @@ struct CORE_API FMacPlatformMisc : public FApplePlatformMisc
 	static void PlatformTearDown();
 	static void SetEnvironmentVar(const TCHAR* VariableName, const TCHAR* Value);
 
-	FORCEINLINE static constexpr int32 GetMaxPathLength()
+	UE_FORCEINLINE_HINT static constexpr int32 GetMaxPathLength()
 	{
 		return MAC_MAX_PATH;
 	}
diff --git a/Engine/Source/Runtime/Core/Public/Mac/MacPlatformProcess.h b/Engine/Source/Runtime/Core/Public/Mac/MacPlatformProcess.h
--- a/Engine/Source/Runtime/Core/Public/Mac/MacPlatformProcess.h
+++ b/Engine/Source/Runtime/Core/Public/Mac/MacPlatformProcess.h
@@ -29,7 +29,7 @@ struct FProcState
 	~FProcState();
 
 	/** Getter for process id */
-	FORCEINLINE pid_t GetProcessId() const
+	UE_FORCEINLINE_HINT pid_t GetProcessId() const
 	{
 		return ProcessId;
 	}
@@ -125,7 +125,7 @@ struct FProcHandle
 	}
 
 	/** Accessors. */
-	FORCEINLINE pid_t Get() const
+	UE_FORCEINLINE_HINT pid_t Get() const
 	{
 		return ProcInfo ? ProcInfo->GetProcessId() : OpenedPid;
 	}
@@ -138,13 +138,13 @@ struct FProcHandle
 	}
 
 	/** Checks the validity of handle */
-	FORCEINLINE bool IsValid() const
+	UE_FORCEINLINE_HINT bool IsValid() const
 	{
 		return ProcInfo != nullptr || OpenedPid != -1;
 	}
 
 	// the below is not part of FProcHandle API and is specific to Unix implementation
-	FORCEINLINE FProcState* GetProcessInfo() const
+	UE_FORCEINLINE_HINT FProcState* GetProcessInfo() const
 	{
 		return ProcInfo;
 	}
diff --git a/Engine/Source/Runtime/Core/Public/Mac/MacPlatformProperties.h b/Engine/Source/Runtime/Core/Public/Mac/MacPlatformProperties.h
--- a/Engine/Source/Runtime/Core/Public/Mac/MacPlatformProperties.h
+++ b/Engine/Source/Runtime/Core/Public/Mac/MacPlatformProperties.h
@@ -19,32 +19,32 @@ template<bool HAS_EDITOR_DATA, bool IS_DEDICATED_SERVER, bool IS_CLIENT_ONLY>
 struct FMacPlatformProperties
 	: public FGenericPlatformProperties
 {
-	constexpr static FORCEINLINE bool HasEditorOnlyData()
+	constexpr static UE_FORCEINLINE_HINT bool HasEditorOnlyData()
 	{
 		return HAS_EDITOR_DATA;
 	}
 
-	constexpr static FORCEINLINE const char* IniPlatformName()
+	constexpr static UE_FORCEINLINE_HINT const char* IniPlatformName()
 	{
 		return "Mac";
 	}
 
-	constexpr static FORCEINLINE const TCHAR* GetRuntimeSettingsClassName()
+	constexpr static UE_FORCEINLINE_HINT const TCHAR* GetRuntimeSettingsClassName()
 	{
 		return TEXT("/Script/MacTargetPlatform.MacTargetSettings");
 	}
 
-	constexpr static FORCEINLINE bool IsGameOnly()
+	constexpr static UE_FORCEINLINE_HINT bool IsGameOnly()
 	{
 		return UE_GAME;
 	}
 
-	constexpr static FORCEINLINE bool IsServerOnly()
+	constexpr static UE_FORCEINLINE_HINT bool IsServerOnly()
 	{
 		return IS_DEDICATED_SERVER;
 	}
 
-	constexpr static FORCEINLINE bool IsClientOnly()
+	constexpr static UE_FORCEINLINE_HINT bool IsClientOnly()
 	{
 		return IS_CLIENT_ONLY;
 	}
@@ -69,37 +69,37 @@ struct FMacPlatformProperties
 		return "Mac";
 	}
 
-	constexpr static FORCEINLINE bool RequiresCookedData()
+	constexpr static UE_FORCEINLINE_HINT bool RequiresCookedData()
 	{
 		return !HAS_EDITOR_DATA;
 	}
 
-	constexpr static FORCEINLINE bool HasSecurePackageFormat()
+	constexpr static UE_FORCEINLINE_HINT bool HasSecurePackageFormat()
 	{
 		return IS_DEDICATED_SERVER;
 	}
 
-	constexpr static FORCEINLINE bool SupportsMultipleGameInstances()
+	constexpr static UE_FORCEINLINE_HINT bool SupportsMultipleGameInstances()
 	{
 		return false;
 	}
 	
-	constexpr static FORCEINLINE bool SupportsWindowedMode()
+	constexpr static UE_FORCEINLINE_HINT bool SupportsWindowedMode()
 	{
 		return true;
 	}
 
-	constexpr static FORCEINLINE bool AllowsFramerateSmoothing()
+	constexpr static UE_FORCEINLINE_HINT bool AllowsFramerateSmoothing()
 	{
 		return true;
 	}
 
-	constexpr static FORCEINLINE bool HasFixedResolution()
+	constexpr static UE_FORCEINLINE_HINT bool HasFixedResolution()
 	{
 		return false;
 	}
 
-	constexpr static FORCEINLINE bool SupportsQuit()
+	constexpr static UE_FORCEINLINE_HINT bool SupportsQuit()
 	{
 		return true;
 	}
@@ -124,12 +124,12 @@ struct FMacPlatformProperties
 		return 1.0f;
 	}
 	
-	constexpr static FORCEINLINE bool SupportsAudioStreaming()
+	constexpr static UE_FORCEINLINE_HINT bool SupportsAudioStreaming()
 	{
 		return !IsServerOnly();
 	}
 
-	constexpr static FORCEINLINE bool SupportsMeshLODStreaming()
+	constexpr static UE_FORCEINLINE_HINT bool SupportsMeshLODStreaming()
 	{
 		return !IsServerOnly() && !HasEditorOnlyData();
 	}
diff --git a/Engine/Source/Runtime/Core/Public/Math/BigInt.h b/Engine/Source/Runtime/Core/Public/Math/BigInt.h
--- a/Engine/Source/Runtime/Core/Public/Math/BigInt.h
+++ b/Engine/Source/Runtime/Core/Public/Math/BigInt.h
@@ -79,18 +79,18 @@ private:
 
 public:
 
-	[[nodiscard]] FORCEINLINE uint32* GetBits()
+	[[nodiscard]] UE_FORCEINLINE_HINT uint32* GetBits()
 	{
 		return Bits;
 	}
 
-	[[nodiscard]] FORCEINLINE const uint32* GetBits() const
+	[[nodiscard]] UE_FORCEINLINE_HINT const uint32* GetBits() const
 	{
 		return Bits;
 	}
 
 	/** Sets this integer to 0. */
-	FORCEINLINE void Zero()
+	UE_FORCEINLINE_HINT void Zero()
 	{
 		FMemory::Memset(Bits, 0, sizeof(Bits));
 	}
diff --git a/Engine/Source/Runtime/Core/Public/Math/Bounds.h b/Engine/Source/Runtime/Core/Public/Math/Bounds.h
--- a/Engine/Source/Runtime/Core/Public/Math/Bounds.h
+++ b/Engine/Source/Runtime/Core/Public/Math/Bounds.h
@@ -38,7 +38,7 @@ struct TBounds
 		return *this;
 	}
 
-	[[nodiscard]] FORCEINLINE TBounds<T> operator+( const TBounds<T>& Other ) const
+	[[nodiscard]] UE_FORCEINLINE_HINT TBounds<T> operator+( const TBounds<T>& Other ) const
 	{
 		return TBounds<T>(*this) += Other;
 	}
@@ -68,17 +68,17 @@ struct TBounds
 		return VectorDot3Scalar( rClosest, rClosest );
 	}
 
-	[[nodiscard]] FORCEINLINE TVector<T> GetCenter() const
+	[[nodiscard]] UE_FORCEINLINE_HINT TVector<T> GetCenter() const
 	{
 		return (Max + Min) * 0.5f;
 	}
 
-	[[nodiscard]] FORCEINLINE TVector<T> GetExtent() const
+	[[nodiscard]] UE_FORCEINLINE_HINT TVector<T> GetExtent() const
 	{
 		return (Max - Min) * 0.5f;
 	}
 
-	[[nodiscard]] FORCEINLINE TVector<T> GetSize() const
+	[[nodiscard]] UE_FORCEINLINE_HINT TVector<T> GetSize() const
 	{
 		return (Max - Min);
 	}
