Replace some FORCEINLINE usage in Core/CoreUObject with UE_FORCEINLINE_HINT.

diff --git a/Engine/Source/Runtime/Core/Public/Math/Box.h b/Engine/Source/Runtime/Core/Public/Math/Box.h
--- a/Engine/Source/Runtime/Core/Public/Math/Box.h
+++ b/Engine/Source/Runtime/Core/Public/Math/Box.h
@@ -106,7 +106,7 @@ public:
 		}
 	}
 #else
-	FORCEINLINE void DiagnosticCheck() const {}
+	UE_FORCEINLINE_HINT void DiagnosticCheck() const {}
 #endif
 
 	[[nodiscard]] bool ContainsNaN() const
@@ -184,7 +184,7 @@ public:
 	 *
 	 * @return true if the boxes are equal, false otherwise.
 	 */
-	[[nodiscard]] FORCEINLINE bool operator==( const TBox<T>& Other ) const
+	[[nodiscard]] UE_FORCEINLINE_HINT bool operator==( const TBox<T>& Other ) const
 	{
 		return (!IsValid && !Other.IsValid) || ((IsValid && Other.IsValid) && (Min == Other.Min) && (Max == Other.Max));
 	}
@@ -194,7 +194,7 @@ public:
 	 * 
 	 * @return false if the boxes are equal, true otherwise.
 	 */
-	[[nodiscard]] FORCEINLINE bool operator!=( const TBox<T>& Other) const
+	[[nodiscard]] UE_FORCEINLINE_HINT bool operator!=( const TBox<T>& Other) const
 	{
 		return !(*this == Other);
 	}
@@ -219,7 +219,7 @@ public:
 	 * @param Other the point to increase the bounding volume to.
 	 * @return Reference to this bounding box after resizing to include the other point.
 	 */
-	FORCEINLINE TBox<T>& operator+=( const TVector<T> &Other );
+	inline TBox<T>& operator+=( const TVector<T> &Other );
 
 	/**
 	 * Gets the result of addition to this bounding volume.
@@ -227,7 +227,7 @@ public:
 	 * @param Other The other point to add to this.
 	 * @return A new bounding volume.
 	 */
-	[[nodiscard]] FORCEINLINE TBox<T> operator+( const TVector<T>& Other ) const
+	[[nodiscard]] UE_FORCEINLINE_HINT TBox<T> operator+( const TVector<T>& Other ) const
 	{
 		return TBox<T>(*this) += Other;
 	}
@@ -238,7 +238,7 @@ public:
 	 * @param Other the bounding volume to increase the bounding volume to.
 	 * @return Reference to this bounding volume after resizing to include the other bounding volume.
 	 */
-	FORCEINLINE TBox<T>& operator+=( const TBox<T>& Other );
+	inline TBox<T>& operator+=( const TBox<T>& Other );
 
 	/**
 	 * Gets the result of addition to this bounding volume.
@@ -246,7 +246,7 @@ public:
 	 * @param Other The other volume to add to this.
 	 * @return A new bounding volume.
 	 */
-	[[nodiscard]] FORCEINLINE TBox<T> operator+( const TBox<T>& Other ) const
+	[[nodiscard]] UE_FORCEINLINE_HINT TBox<T> operator+( const TBox<T>& Other ) const
 	{
 		return TBox<T>(*this) += Other;
 	}
@@ -294,7 +294,7 @@ public:
 	 * @param Point The point.
 	 * @return The distance.
 	 */
-	[[nodiscard]] FORCEINLINE T ComputeSquaredDistanceToPoint( const TVector<T>& Point ) const
+	[[nodiscard]] UE_FORCEINLINE_HINT T ComputeSquaredDistanceToPoint( const TVector<T>& Point ) const
 	{
 		return ComputeSquaredDistanceFromBoxToPoint(Min, Max, Point);
 	}
@@ -315,7 +315,7 @@ public:
 	 * @param W The size to increase the volume by.
 	 * @return A new bounding box.
 	 */
-	[[nodiscard]] FORCEINLINE TBox<T> ExpandBy(T W) const
+	[[nodiscard]] UE_FORCEINLINE_HINT TBox<T> ExpandBy(T W) const
 	{
 		return TBox<T>(Min - TVector<T>(W, W, W), Max + TVector<T>(W, W, W));
 	}
@@ -326,7 +326,7 @@ public:
 	* @param V The size to increase the volume by.
 	* @return A new bounding box.
 	*/
-	[[nodiscard]] FORCEINLINE TBox<T> ExpandBy(const TVector<T>& V) const
+	[[nodiscard]] UE_FORCEINLINE_HINT TBox<T> ExpandBy(const TVector<T>& V) const
 	{
 		return TBox<T>(Min - V, Max + V);
 	}
@@ -349,7 +349,7 @@ public:
 	 * @param Offset The vector to shift the box by.
 	 * @return A new bounding box.
 	 */
-	[[nodiscard]] FORCEINLINE TBox<T> ShiftBy( const TVector<T>& Offset ) const
+	[[nodiscard]] UE_FORCEINLINE_HINT TBox<T> ShiftBy( const TVector<T>& Offset ) const
 	{
 		return TBox<T>(Min + Offset, Max + Offset);
 	}
@@ -397,7 +397,7 @@ public:
 	 * @param Point The point in space.
 	 * @return The closest point on or inside the box.
 	 */
-	FORCEINLINE TVector<T> GetClosestPointTo( const TVector<T>& Point ) const;
+	inline TVector<T> GetClosestPointTo( const TVector<T>& Point ) const;
 
 	/**
 	 * Gets the extents of this box.
@@ -453,7 +453,7 @@ public:
 	 * @note  This function assumes boxes have closed bounds, i.e. boxes with
 	 *        coincident borders on any edge will overlap.
 	 */
-	FORCEINLINE bool Intersect( const TBox<T>& Other ) const;
+	inline bool Intersect( const TBox<T>& Other ) const;
 
 	/**
 	 * Checks whether the given bounding box intersects this bounding box in the XY plane.
@@ -464,7 +464,7 @@ public:
 	 * @note  This function assumes boxes have closed bounds, i.e. boxes with
 	 *        coincident borders on any edge will overlap.
 	 */
-	FORCEINLINE bool IntersectXY( const TBox<T>& Other ) const;
+	inline bool IntersectXY( const TBox<T>& Other ) const;
 
 	/**
 	 * Returns the overlap TBox<T> of two box
@@ -524,7 +524,7 @@ public:
 	 * @note  This function assumes boxes have open bounds, i.e. boxes with
 	 *        coincident borders on any edge are not encapsulated.
 	 */
-	FORCEINLINE bool IsInside( const TBox<T>& Other ) const
+	UE_FORCEINLINE_HINT bool IsInside( const TBox<T>& Other ) const
 	{
 		return (IsInside(Other.Min) && IsInside(Other.Max));
 	}
@@ -539,7 +539,7 @@ public:
 	 * @note  This function assumes boxes have closed bounds, i.e. boxes with
 	 *        coincident borders on any edge are encapsulated.
 	 */
-	FORCEINLINE bool IsInsideOrOn( const TBox<T>& Other ) const
+	UE_FORCEINLINE_HINT bool IsInsideOrOn( const TBox<T>& Other ) const
 	{
 		return (IsInsideOrOn(Other.Min) && IsInsideOrOn(Other.Max));
 	}
@@ -586,7 +586,7 @@ public:
 	 * @note  This function assumes boxes have open bounds, i.e. boxes with
 	 *        coincident borders on any edge are not encapsulated.
 	 */
-	FORCEINLINE bool IsInsideXY( const TBox<T>& Other ) const
+	UE_FORCEINLINE_HINT bool IsInsideXY( const TBox<T>& Other ) const
 	{
 		return (IsInsideXY(Other.Min) && IsInsideXY(Other.Max));
 	}
@@ -855,13 +855,13 @@ inline bool TBox<T>::IntersectXY( const TBox<T>& Other ) const
 
 
 template<typename T>
-FORCEINLINE FString TBox<T>::ToString() const
+UE_FORCEINLINE_HINT FString TBox<T>::ToString() const
 {
 	return FString::Printf(TEXT("IsValid=%s, Min=(%s), Max=(%s)"), IsValid ? TEXT("true") : TEXT("false"), *Min.ToString(), *Max.ToString());
 }
 
 template<typename T>
-FORCEINLINE FString TBox<T>::ToCompactString() const
+UE_FORCEINLINE_HINT FString TBox<T>::ToCompactString() const
 {
 	return IsValid ? FString::Printf(TEXT("Min=(%s), Max=(%s)"), *Min.ToString(), *Max.ToString()) : TEXT("IsValid=false");
 }
diff --git a/Engine/Source/Runtime/Core/Public/Math/Box2D.h b/Engine/Source/Runtime/Core/Public/Math/Box2D.h
--- a/Engine/Source/Runtime/Core/Public/Math/Box2D.h
+++ b/Engine/Source/Runtime/Core/Public/Math/Box2D.h
@@ -126,7 +126,7 @@ public:
 	 * @param Other The point to increase the bounding volume to.
 	 * @return Reference to this bounding box after resizing to include the other point.
 	 */
-	FORCEINLINE TBox2<T>& operator+=( const TVector2<T> &Other );
+	inline TBox2<T>& operator+=( const TVector2<T> &Other );
 
 	/**
 	 * Gets the result of addition to this bounding volume.
@@ -145,7 +145,7 @@ public:
 	 * @param Other The bounding volume to increase the bounding volume to.
 	 * @return Reference to this bounding volume after resizing to include the other bounding volume.
 	 */
-	FORCEINLINE TBox2<T>& operator+=( const TBox2<T>& Other );
+	inline TBox2<T>& operator+=( const TBox2<T>& Other );
 
 	/**
 	 * Gets the result of addition to this bounding volume.
@@ -293,7 +293,7 @@ public:
 	 *
 	 * @return The closest point on or inside the box.
 	 */
-	[[nodiscard]] FORCEINLINE TVector2<T> GetClosestPointTo( const TVector2<T>& Point ) const;
+	[[nodiscard]] inline TVector2<T> GetClosestPointTo( const TVector2<T>& Point ) const;
 
 	/**
 	 * Gets the box extents around the center.
@@ -341,7 +341,7 @@ public:
 	 * @param other bounding box to test intersection
 	 * @return true if boxes intersect, false otherwise.
 	 */
-	[[nodiscard]] FORCEINLINE bool Intersect( const TBox2<T> & other ) const;
+	[[nodiscard]] inline bool Intersect( const TBox2<T> & other ) const;
 
 	/**
 	 * Checks whether the given point is inside this box.
@@ -611,7 +611,7 @@ inline bool TBox2<T>::Intersect( const TBox2<T> & Other ) const
 }
 
 template<typename T>
-FORCEINLINE FString TBox2<T>::ToString() const
+UE_FORCEINLINE_HINT FString TBox2<T>::ToString() const
 {
 	return FString::Printf(TEXT("bIsValid=%s, Min=(%s), Max=(%s)"), bIsValid ? TEXT("true") : TEXT("false"), *Min.ToString(), *Max.ToString());
 }
diff --git a/Engine/Source/Runtime/Core/Public/Math/BoxSphereBounds.h b/Engine/Source/Runtime/Core/Public/Math/BoxSphereBounds.h
--- a/Engine/Source/Runtime/Core/Public/Math/BoxSphereBounds.h
+++ b/Engine/Source/Runtime/Core/Public/Math/BoxSphereBounds.h
@@ -159,7 +159,7 @@ public:
 	 * @param Other The other bounding volume.
 	 * @return The combined bounding volume.
 	 */
-	[[nodiscard]] FORCEINLINE TBoxSphereBounds<T, TExtent> operator+( const TBoxSphereBounds<T, TExtent>& Other ) const;
+	[[nodiscard]] inline TBoxSphereBounds<T, TExtent> operator+( const TBoxSphereBounds<T, TExtent>& Other ) const;
 
 	/**
 	 * Compare bounding volume this and Other.
@@ -167,7 +167,7 @@ public:
 	 * @param Other The other bounding volume.
 	 * @return true of they match.
 	 */
-	[[nodiscard]] FORCEINLINE bool operator==(const TBoxSphereBounds<T, TExtent>& Other) const;
+	[[nodiscard]] UE_FORCEINLINE_HINT bool operator==(const TBoxSphereBounds<T, TExtent>& Other) const;
 	
 	/**
 	 * Compare bounding volume this and Other.
@@ -175,7 +175,7 @@ public:
 	 * @param Other The other bounding volume.
 	 * @return true of they do not match.
 	 */	
-	[[nodiscard]] FORCEINLINE bool operator!=(const TBoxSphereBounds<T, TExtent>& Other) const;
+	[[nodiscard]] UE_FORCEINLINE_HINT bool operator!=(const TBoxSphereBounds<T, TExtent>& Other) const;
 
 public:
 
@@ -201,7 +201,7 @@ public:
 	 * @param  Tolerance Error tolerance added to test distance.
 	 * @return true if spheres intersect, false otherwise.
 	 */
-	[[nodiscard]] FORCEINLINE static bool SpheresIntersect(const TBoxSphereBounds<T, TExtent>& A, const TBoxSphereBounds<T, TExtent>& B, TExtent Tolerance = UE_KINDA_SMALL_NUMBER)
+	[[nodiscard]] UE_FORCEINLINE_HINT static bool SpheresIntersect(const TBoxSphereBounds<T, TExtent>& A, const TBoxSphereBounds<T, TExtent>& B, TExtent Tolerance = UE_KINDA_SMALL_NUMBER)
 	{
 		return (A.Origin - B.Origin).SizeSquared() <= FMath::Square(FMath::Max<TExtent>(0, A.SphereRadius + B.SphereRadius + Tolerance));
 	}
@@ -213,7 +213,7 @@ public:
 	 * @param  B Second BoxSphereBounds to test.
 	 * @return true if boxes intersect, false otherwise.
 	 */
-	[[nodiscard]] FORCEINLINE static bool BoxesIntersect(const TBoxSphereBounds<T, TExtent>& A, const TBoxSphereBounds<T, TExtent>& B)
+	[[nodiscard]] UE_FORCEINLINE_HINT static bool BoxesIntersect(const TBoxSphereBounds<T, TExtent>& A, const TBoxSphereBounds<T, TExtent>& B)
 	{
 		return A.GetBox().Intersect(B.GetBox());
 	}
@@ -223,7 +223,7 @@ public:
 	 *
 	 * @return The bounding box.
 	 */
-	[[nodiscard]] FORCEINLINE TBox<T> GetBox() const
+	[[nodiscard]] UE_FORCEINLINE_HINT TBox<T> GetBox() const
 	{
 		return TBox<T>(Origin - BoxExtent,Origin + BoxExtent);
 	}
@@ -249,7 +249,7 @@ public:
 	 *
 	 * @return The bounding sphere.
 	 */
-	[[nodiscard]] FORCEINLINE TSphere<T> GetSphere() const
+	[[nodiscard]] UE_FORCEINLINE_HINT TSphere<T> GetSphere() const
 	{
 		return TSphere<T>(Origin,SphereRadius);
 	}
@@ -260,7 +260,7 @@ public:
 	 * @param ExpandAmount The size to increase by.
 	 * @return A new box with the expanded size.
 	 */
-	[[nodiscard]] FORCEINLINE TBoxSphereBounds<T, TExtent> ExpandBy(TExtent ExpandAmount ) const
+	[[nodiscard]] UE_FORCEINLINE_HINT TBoxSphereBounds<T, TExtent> ExpandBy(TExtent ExpandAmount ) const
 	{
 		return TBoxSphereBounds(Origin, BoxExtent + ExpandAmount, SphereRadius + ExpandAmount);
 	}
@@ -318,7 +318,7 @@ public:
 		}
 	}
 #else
-	FORCEINLINE void DiagnosticCheckNaN() const {}
+	UE_FORCEINLINE_HINT void DiagnosticCheckNaN() const {}
 #endif
 
 	[[nodiscard]] inline bool ContainsNaN() const
@@ -334,32 +334,32 @@ public:
 	 */
 	struct Builder
 	{
-		FORCEINLINE Builder& operator+=(const TBoxSphereBounds<T, TExtent>& Other)
+		UE_FORCEINLINE_HINT Builder& operator+=(const TBoxSphereBounds<T, TExtent>& Other)
 		{
 			return Append(Other);
 		}
 
-		FORCEINLINE Builder& operator+=(const TBox<T>& Box)
+		UE_FORCEINLINE_HINT Builder& operator+=(const TBox<T>& Box)
 		{
 			return Append(Box);
 		}
 
-		FORCEINLINE Builder& operator+=(const TSphere<T>& Sphere)
+		UE_FORCEINLINE_HINT Builder& operator+=(const TSphere<T>& Sphere)
 		{
 			return Append(Sphere);
 		}
 
-		FORCEINLINE Builder& operator+=(const TVector<T>& Point)
+		UE_FORCEINLINE_HINT Builder& operator+=(const TVector<T>& Point)
 		{
 			return Append(TBoxSphereBounds<T, TExtent>(Point, TVector<T>::ZeroVector, T(0)));
 		}
 
-		FORCEINLINE bool IsValid() const
+		UE_FORCEINLINE_HINT bool IsValid() const
 		{
 			return BoxSphereBounds.IsSet();
 		}
 
-		FORCEINLINE operator TBoxSphereBounds<T, TExtent>() const
+		UE_FORCEINLINE_HINT operator TBoxSphereBounds<T, TExtent>() const
 		{
 			return BoxSphereBounds.Get(TBoxSphereBounds<T, TExtent>(ForceInitToZero));
 		}
@@ -473,13 +473,13 @@ inline TBoxSphereBounds<T, TExtent> TBoxSphereBounds<T, TExtent>::operator+( con
 }
 
 template<typename T, typename TExtent>
-FORCEINLINE bool TBoxSphereBounds<T, TExtent>::operator==(const TBoxSphereBounds<T, TExtent>& Other) const
+UE_FORCEINLINE_HINT bool TBoxSphereBounds<T, TExtent>::operator==(const TBoxSphereBounds<T, TExtent>& Other) const
 {
 	return Origin == Other.Origin && BoxExtent == Other.BoxExtent &&  SphereRadius == Other.SphereRadius;
 }
 
 template<typename T, typename TExtent>
-FORCEINLINE bool TBoxSphereBounds<T, TExtent>::operator!=(const TBoxSphereBounds<T, TExtent>& Other) const
+UE_FORCEINLINE_HINT bool TBoxSphereBounds<T, TExtent>::operator!=(const TBoxSphereBounds<T, TExtent>& Other) const
 {
 	return !(*this == Other);
 }
@@ -492,7 +492,7 @@ inline bool TBoxSphereBounds<T, TExtent>::Serialize(FArchive &Ar)
 }
 
 template<typename T, typename TExtent>
-FORCEINLINE FString TBoxSphereBounds<T, TExtent>::ToString() const
+UE_FORCEINLINE_HINT FString TBoxSphereBounds<T, TExtent>::ToString() const
 {
 	return FString::Printf(TEXT("Origin=%s, BoxExtent=(%s), SphereRadius=(%f)"), *Origin.ToString(), *BoxExtent.ToString(), SphereRadius);
 }
diff --git a/Engine/Source/Runtime/Core/Public/Math/ClipProjectionMatrix.h b/Engine/Source/Runtime/Core/Public/Math/ClipProjectionMatrix.h
--- a/Engine/Source/Runtime/Core/Public/Math/ClipProjectionMatrix.h
+++ b/Engine/Source/Runtime/Core/Public/Math/ClipProjectionMatrix.h
@@ -37,7 +37,7 @@ public:
 	
 private:
 	/** return sign of a number */
-	FORCEINLINE T sgn(T a );
+	inline T sgn(T a );
 };
 
 
diff --git a/Engine/Source/Runtime/Core/Public/Math/Color.h b/Engine/Source/Runtime/Core/Public/Math/Color.h
--- a/Engine/Source/Runtime/Core/Public/Math/Color.h
+++ b/Engine/Source/Runtime/Core/Public/Math/Color.h
@@ -67,17 +67,17 @@ struct FLinearColor
 	static CORE_API float sRGBToLinearTable[256];
 
 	FLinearColor() = default;
-	FORCEINLINE explicit FLinearColor(EForceInit)
+	UE_FORCEINLINE_HINT explicit FLinearColor(EForceInit)
 	: R(0), G(0), B(0), A(0)
 	{}
-	constexpr FORCEINLINE FLinearColor(float InR,float InG,float InB,float InA = 1.0f): R(InR), G(InG), B(InB), A(InA) {}
+	constexpr UE_FORCEINLINE_HINT FLinearColor(float InR,float InG,float InB,float InA = 1.0f): R(InR), G(InG), B(InB), A(InA) {}
 	
 	/**
 	 * Converts an FColor which is assumed to be in sRGB space, into linear color space.
 	 * @param Color The sRGB color that needs to be converted into linear space.
 	 * to get direct conversion use ReinterpretAsLinear
 	 */
-	constexpr FORCEINLINE FLinearColor(const FColor& Color);
+	constexpr UE_FORCEINLINE_HINT FLinearColor(const FColor& Color);
 
 	CORE_API FLinearColor(const FVector3f& Vector);
 	CORE_API explicit FLinearColor(const FVector3d& Vector); // Warning: keep this explicit, or FVector4f will be implicitly created from FVector3d via FLinearColor
@@ -120,7 +120,7 @@ struct FLinearColor
 	 * Converts an FColor coming from an observed sRGB output, into a linear color.
 	 * @param Color The sRGB color that needs to be converted into linear space.
 	 */
-	[[nodiscard]] FORCEINLINE static FLinearColor FromSRGBColor(const FColor& Color)
+	[[nodiscard]] UE_FORCEINLINE_HINT static FLinearColor FromSRGBColor(const FColor& Color)
 	{
 		return FLinearColor(Color);
 	}
@@ -272,17 +272,17 @@ struct FLinearColor
 	}
 
 	/** Comparison operators */
-	[[nodiscard]] FORCEINLINE bool operator==(const FLinearColor& ColorB) const
+	[[nodiscard]] UE_FORCEINLINE_HINT bool operator==(const FLinearColor& ColorB) const
 	{
 		return this->R == ColorB.R && this->G == ColorB.G && this->B == ColorB.B && this->A == ColorB.A;
 	}
-	[[nodiscard]] FORCEINLINE bool operator!=(const FLinearColor& Other) const
+	[[nodiscard]] UE_FORCEINLINE_HINT bool operator!=(const FLinearColor& Other) const
 	{
 		return this->R != Other.R || this->G != Other.G || this->B != Other.B || this->A != Other.A;
 	}
 
 	// Error-tolerant comparison.
-	[[nodiscard]] FORCEINLINE bool Equals(const FLinearColor& ColorB, float Tolerance=UE_KINDA_SMALL_NUMBER) const
+	[[nodiscard]] UE_FORCEINLINE_HINT bool Equals(const FLinearColor& ColorB, float Tolerance=UE_KINDA_SMALL_NUMBER) const
 	{
 		return FMath::Abs(this->R - ColorB.R) < Tolerance && FMath::Abs(this->G - ColorB.G) < Tolerance && FMath::Abs(this->B - ColorB.B) < Tolerance && FMath::Abs(this->A - ColorB.A) < Tolerance;
 	}
@@ -354,14 +354,14 @@ struct FLinearColor
 	* QuantizeRound can be dequantized back to linear with FColor::ReinterpretAsLinear (just /255.f)
 	* this matches the GPU UNORM<->float conversion spec and should be preferred
 	*/
-	[[nodiscard]] FORCEINLINE FColor QuantizeRound() const;
+	[[nodiscard]] inline FColor QuantizeRound() const;
 	
 	/** Quantizes the linear color and returns the result as a FColor.  This bypasses the SRGB conversion.
 	* Uses floor quantization, which does not match the GPU standard conversion.
 	* Restoration to float should be done with a +0.5 bias to restore to centered buckets.
 	* Do NOT use this for graphics or textures or images, use QuantizeRound instead.
 	*/
-	[[nodiscard]] FORCEINLINE FColor QuantizeFloor() const;
+	[[nodiscard]] inline FColor QuantizeFloor() const;
 
 	/** Quantizes the linear color and returns the result as a FColor with optional sRGB conversion. 
 	* Clamps in [0,1] range before conversion.
@@ -369,7 +369,7 @@ struct FLinearColor
 	*/
 	[[nodiscard]] CORE_API FColor ToFColorSRGB() const;
 	
-	[[nodiscard]] FORCEINLINE FColor ToFColor(const bool bSRGB) const;
+	[[nodiscard]] inline FColor ToFColor(const bool bSRGB) const;
 	
 	/**
 	 * Returns a desaturated color, with 0 meaning no desaturation and 1 == full desaturation
@@ -387,7 +387,7 @@ struct FLinearColor
 	 *
 	 * @return The maximum color channel value
 	 */
-	[[nodiscard]] FORCEINLINE float GetMax() const
+	[[nodiscard]] UE_FORCEINLINE_HINT float GetMax() const
 	{
 		return FMath::Max( FMath::Max( FMath::Max( R, G ), B ), A );
 	}
@@ -403,7 +403,7 @@ struct FLinearColor
 	 *
 	 * @return The minimum color channel value
 	 */
-	[[nodiscard]] FORCEINLINE float GetMin() const
+	[[nodiscard]] UE_FORCEINLINE_HINT float GetMin() const
 	{
 		return FMath::Min( FMath::Min( FMath::Min( R, G ), B ), A );
 	}
@@ -462,7 +462,7 @@ struct FLinearColor
 	static CORE_API const FLinearColor Blue;
 	static CORE_API const FLinearColor Yellow;
 
-	[[nodiscard]] friend FORCEINLINE uint32 GetTypeHash( const FLinearColor& LinearColor )
+	[[nodiscard]] friend UE_FORCEINLINE_HINT uint32 GetTypeHash( const FLinearColor& LinearColor )
 	{
 		// Note: this assumes there's no padding in FLinearColor that could contain uncompared data.
 		return FCrc::MemCrc_DEPRECATED(&LinearColor, sizeof(FLinearColor));
@@ -470,7 +470,7 @@ struct FLinearColor
 };
 DECLARE_INTRINSIC_TYPE_LAYOUT(FLinearColor);
 
-FORCEINLINE FLinearColor operator*(float Scalar,const FLinearColor& Color)
+UE_FORCEINLINE_HINT FLinearColor operator*(float Scalar,const FLinearColor& Color)
 {
 	return Color.operator*( Scalar );
 }
@@ -497,7 +497,7 @@ public:
 
 	// Constructors.
 	FColor() = default;
-	FORCEINLINE explicit FColor(EForceInit)
+	UE_FORCEINLINE_HINT explicit FColor(EForceInit)
 	{
 		// put these into the body for proper ordering with INTEL vs non-INTEL_BYTE_ORDER
 		R = G = B = A = 0;
@@ -511,7 +511,7 @@ public:
 #endif
 	{}
 
-	FORCEINLINE explicit FColor( uint32 InColor )
+	UE_FORCEINLINE_HINT explicit FColor( uint32 InColor )
 	{ 
 		DWColor() = InColor; 
 	}
@@ -541,12 +541,12 @@ public:
 	}
 
 	// Operators.
-	[[nodiscard]] FORCEINLINE bool operator==( const FColor &C ) const
+	[[nodiscard]] UE_FORCEINLINE_HINT bool operator==( const FColor &C ) const
 	{
 		return DWColor() == C.DWColor();
 	}
 
-	[[nodiscard]] FORCEINLINE bool operator!=( const FColor& C ) const
+	[[nodiscard]] UE_FORCEINLINE_HINT bool operator!=( const FColor& C ) const
 	{
 		return DWColor() != C.DWColor();
 	}
@@ -716,7 +716,7 @@ public:
 	/**
 	 * Gets the color in a packed uint32 format packed in the order ARGB.
 	 */
-	[[nodiscard]] FORCEINLINE uint32 ToPackedARGB() const
+	[[nodiscard]] UE_FORCEINLINE_HINT uint32 ToPackedARGB() const
 	{
 		return ( A << 24 ) | ( R << 16 ) | ( G << 8 ) | ( B << 0 );
 	}
@@ -724,7 +724,7 @@ public:
 	/**
 	 * Gets the color in a packed uint32 format packed in the order ABGR.
 	 */
-	[[nodiscard]] FORCEINLINE uint32 ToPackedABGR() const
+	[[nodiscard]] UE_FORCEINLINE_HINT uint32 ToPackedABGR() const
 	{
 		return ( A << 24 ) | ( B << 16 ) | ( G << 8 ) | ( R << 0 );
 	}
@@ -732,7 +732,7 @@ public:
 	/**
 	 * Gets the color in a packed uint32 format packed in the order RGBA.
 	 */
-	[[nodiscard]] FORCEINLINE uint32 ToPackedRGBA() const
+	[[nodiscard]] UE_FORCEINLINE_HINT uint32 ToPackedRGBA() const
 	{
 		return ( R << 24 ) | ( G << 16 ) | ( B << 8 ) | ( A << 0 );
 	}
@@ -740,7 +740,7 @@ public:
 	/**
 	 * Gets the color in a packed uint32 format packed in the order BGRA.
 	 */
-	[[nodiscard]] FORCEINLINE uint32 ToPackedBGRA() const
+	[[nodiscard]] UE_FORCEINLINE_HINT uint32 ToPackedBGRA() const
 	{
 		return ( B << 24 ) | ( G << 16 ) | ( R << 8 ) | ( A << 0 );
 	}
@@ -761,7 +761,7 @@ public:
 	static CORE_API const FColor Silver;
 	static CORE_API const FColor Emerald;
 
-	[[nodiscard]] friend FORCEINLINE uint32 GetTypeHash( const FColor& Color )
+	[[nodiscard]] friend UE_FORCEINLINE_HINT uint32 GetTypeHash( const FColor& Color )
 	{
 		return Color.DWColor();
 	}
@@ -778,7 +778,7 @@ private:
 };
 DECLARE_INTRINSIC_TYPE_LAYOUT(FColor);
 
-constexpr FORCEINLINE FLinearColor::FLinearColor(const FColor& Color)
+constexpr UE_FORCEINLINE_HINT FLinearColor::FLinearColor(const FColor& Color)
 	: R(sRGBToLinearTable[Color.R])
     , G(sRGBToLinearTable[Color.G])
     , B(sRGBToLinearTable[Color.B])
diff --git a/Engine/Source/Runtime/Core/Public/Math/Float16.h b/Engine/Source/Runtime/Core/Public/Math/Float16.h
--- a/Engine/Source/Runtime/Core/Public/Math/Float16.h
+++ b/Engine/Source/Runtime/Core/Public/Math/Float16.h
@@ -120,18 +120,18 @@ template<> struct TCanBulkSerialize<FFloat16> { enum { Value = true }; };
 
 DECLARE_INTRINSIC_TYPE_LAYOUT(FFloat16);
 
-FORCEINLINE FFloat16::FFloat16()
+UE_FORCEINLINE_HINT FFloat16::FFloat16()
 	:	Encoded(0)
 { }
 
 
-FORCEINLINE FFloat16::FFloat16(const FFloat16& FP16Value)
+UE_FORCEINLINE_HINT FFloat16::FFloat16(const FFloat16& FP16Value)
 {
 	Encoded = FP16Value.Encoded;
 }
 
 
-FORCEINLINE FFloat16::FFloat16(float FP32Value)
+UE_FORCEINLINE_HINT FFloat16::FFloat16(float FP32Value)
 {
 	Set(FP32Value);
 }	
@@ -151,14 +151,14 @@ inline FFloat16& FFloat16::operator=(const FFloat16& FP16Value)
 }
 
 
-FORCEINLINE FFloat16::operator float() const
+UE_FORCEINLINE_HINT FFloat16::operator float() const
 {
 	return GetFloat();
 }
 
 
 // NOTE: Set() on values out of F16 max range store them as +-Inf
-FORCEINLINE void FFloat16::Set(float FP32Value)
+UE_FORCEINLINE_HINT void FFloat16::Set(float FP32Value)
 {
 	// FPlatformMath::StoreHalf follows RTNE (round-to-nearest-even) rounding default convention
 	FPlatformMath::StoreHalf(&Encoded, FP32Value);
@@ -166,7 +166,7 @@ FORCEINLINE void FFloat16::Set(float FP32Value)
 
 
 
-FORCEINLINE float FFloat16::GetFloat() const
+UE_FORCEINLINE_HINT float FFloat16::GetFloat() const
 {
 	return FPlatformMath::LoadHalf(&Encoded);
 }
diff --git a/Engine/Source/Runtime/Core/Public/Math/Float16Color.h b/Engine/Source/Runtime/Core/Public/Math/Float16Color.h
--- a/Engine/Source/Runtime/Core/Public/Math/Float16Color.h
+++ b/Engine/Source/Runtime/Core/Public/Math/Float16Color.h
@@ -52,7 +52,7 @@ public:
 };
 
 
-FORCEINLINE FFloat16Color::FFloat16Color() { }
+UE_FORCEINLINE_HINT FFloat16Color::FFloat16Color() { }
 
 
 inline FFloat16Color::FFloat16Color(const FFloat16Color& Src)
@@ -64,7 +64,7 @@ inline FFloat16Color::FFloat16Color(const FFloat16Color& Src)
 }
 
 
-FORCEINLINE FFloat16Color::FFloat16Color(const FLinearColor& Src)
+UE_FORCEINLINE_HINT FFloat16Color::FFloat16Color(const FLinearColor& Src)
 {
 	FPlatformMath::VectorStoreHalf( GetFourHalves(), (const float *)&Src );
 }
diff --git a/Engine/Source/Runtime/Core/Public/Math/Float32.h b/Engine/Source/Runtime/Core/Public/Math/Float32.h
--- a/Engine/Source/Runtime/Core/Public/Math/Float32.h
+++ b/Engine/Source/Runtime/Core/Public/Math/Float32.h
@@ -38,6 +38,6 @@ public:
 };
 
 
-FORCEINLINE FFloat32::FFloat32(float InValue)
+UE_FORCEINLINE_HINT FFloat32::FFloat32(float InValue)
 	: FloatValue(InValue)
 { }
diff --git a/Engine/Source/Runtime/Core/Public/Math/GenericOctree.h b/Engine/Source/Runtime/Core/Public/Math/GenericOctree.h
--- a/Engine/Source/Runtime/Core/Public/Math/GenericOctree.h
+++ b/Engine/Source/Runtime/Core/Public/Math/GenericOctree.h
@@ -164,33 +164,33 @@ public:
 	}
 
 	/** Advances the reference to the next child node.  If this was the last node remain, Index will be 8 which represents null. */
-	FORCEINLINE void Advance()
+	UE_FORCEINLINE_HINT void Advance()
 	{
 		++Index;
 	}
 
 	/** @return true if the reference isn't set. */
-	[[nodiscard]] FORCEINLINE bool IsNULL() const
+	[[nodiscard]] UE_FORCEINLINE_HINT bool IsNULL() const
 	{
 		return Index >= 8;
 	}
 
-	FORCEINLINE void SetNULL()
+	UE_FORCEINLINE_HINT void SetNULL()
 	{
 		Index = 8;
 	}
 
-	[[nodiscard]] FORCEINLINE int32 X() const
+	[[nodiscard]] UE_FORCEINLINE_HINT int32 X() const
 	{
 		return (Index >> 0) & 1;
 	}
 
-	[[nodiscard]] FORCEINLINE int32 Y() const
+	[[nodiscard]] UE_FORCEINLINE_HINT int32 Y() const
 	{
 		return (Index >> 1) & 1;
 	}
 
-	[[nodiscard]] FORCEINLINE int32 Z() const
+	[[nodiscard]] UE_FORCEINLINE_HINT int32 Z() const
 	{
 		return (Index >> 2) & 1;
 	}
@@ -1115,31 +1115,31 @@ public:
 		}
 
 		// Accessors.
-		[[nodiscard]] FORCEINLINE ElementConstIt GetElementIt() const
+		[[nodiscard]] UE_FORCEINLINE_HINT ElementConstIt GetElementIt() const
 		{
 			return ElementConstIt(Elements);
 		}
-		[[nodiscard]] FORCEINLINE bool IsLeaf() const
+		[[nodiscard]] UE_FORCEINLINE_HINT bool IsLeaf() const
 		{
 			return bIsLeaf;
 		}
-		[[nodiscard]] FORCEINLINE bool HasChild(FOctreeChildNodeRef ChildRef) const
+		[[nodiscard]] UE_FORCEINLINE_HINT bool HasChild(FOctreeChildNodeRef ChildRef) const
 		{
 			return Children[ChildRef.Index] != NULL && Children[ChildRef.Index]->InclusiveNumElements > 0;
 		}
-		[[nodiscard]] FORCEINLINE FNode* GetChild(FOctreeChildNodeRef ChildRef) const
+		[[nodiscard]] UE_FORCEINLINE_HINT FNode* GetChild(FOctreeChildNodeRef ChildRef) const
 		{
 			return Children[ChildRef.Index];
 		}
-		[[nodiscard]] FORCEINLINE int32 GetElementCount() const
+		[[nodiscard]] UE_FORCEINLINE_HINT int32 GetElementCount() const
 		{
 			return Elements.Num();
 		}
-		[[nodiscard]] FORCEINLINE int32 GetInclusiveElementCount() const
+		[[nodiscard]] UE_FORCEINLINE_HINT int32 GetInclusiveElementCount() const
 		{
 			return InclusiveNumElements;
 		}
-		[[nodiscard]] FORCEINLINE const ElementArrayType& GetElements() const
+		[[nodiscard]] UE_FORCEINLINE_HINT const ElementArrayType& GetElements() const
 		{
 			return Elements;
 		}
diff --git a/Engine/Source/Runtime/Core/Public/Math/Int128.h b/Engine/Source/Runtime/Core/Public/Math/Int128.h
--- a/Engine/Source/Runtime/Core/Public/Math/Int128.h
+++ b/Engine/Source/Runtime/Core/Public/Math/Int128.h
@@ -63,7 +63,7 @@ public:
 	}
 
 	/** Sets this number to 0. */
-	FORCEINLINE void Zero()
+	UE_FORCEINLINE_HINT void Zero()
 	{
 		Hi = Lo = 0;
 	}
@@ -76,10 +76,10 @@ public:
 	}
 
 	/** Default constructors. */
-    [[nodiscard]] FORCEINLINE FUInt128(const FUInt128&) = default;
-    [[nodiscard]] FORCEINLINE FUInt128(FUInt128&&) = default;
-    FORCEINLINE FUInt128& operator=(FUInt128 const&) = default;
-    FORCEINLINE FUInt128& operator=(FUInt128&&) = default;
+    [[nodiscard]] UE_FORCEINLINE_HINT FUInt128(const FUInt128&) = default;
+    [[nodiscard]] UE_FORCEINLINE_HINT FUInt128(FUInt128&&) = default;
+    UE_FORCEINLINE_HINT FUInt128& operator=(FUInt128 const&) = default;
+    UE_FORCEINLINE_HINT FUInt128& operator=(FUInt128&&) = default;
 
 	/** Default constructor. Initializes the number to zero. */
 	[[nodiscard]] inline FUInt128()
@@ -146,7 +146,7 @@ public:
 	}
 
 	/** this == Other */
-	[[nodiscard]] FORCEINLINE bool IsEqual(const FUInt128& Other) const
+	[[nodiscard]] UE_FORCEINLINE_HINT bool IsEqual(const FUInt128& Other) const
 	{
 		return (Hi == Other.Hi) && (Lo == Other.Lo);
 	}
@@ -209,37 +209,37 @@ public:
 	/**
 	 * Comparison operators
 	 */
-	[[nodiscard]] FORCEINLINE bool operator>(const FUInt128& Other) const
+	[[nodiscard]] UE_FORCEINLINE_HINT bool operator>(const FUInt128& Other) const
 	{
 		return IsGreater(Other);
 	}
 
-	[[nodiscard]] FORCEINLINE bool operator>=(const FUInt128& Other) const
+	[[nodiscard]] UE_FORCEINLINE_HINT bool operator>=(const FUInt128& Other) const
 	{
 		return IsGreaterOrEqual(Other);
 	}
 
-	[[nodiscard]] FORCEINLINE bool operator==(const FUInt128& Other) const
+	[[nodiscard]] UE_FORCEINLINE_HINT bool operator==(const FUInt128& Other) const
 	{
 		return IsEqual(Other);
 	}
 
-	[[nodiscard]] FORCEINLINE bool operator<(const FUInt128& Other) const
+	[[nodiscard]] UE_FORCEINLINE_HINT bool operator<(const FUInt128& Other) const
 	{
 		return IsLess(Other);
 	}
 
-	[[nodiscard]] FORCEINLINE bool operator<=(const FUInt128& Other) const
+	[[nodiscard]] UE_FORCEINLINE_HINT bool operator<=(const FUInt128& Other) const
 	{
 		return IsLessOrEqual(Other);
 	}
 
-	[[nodiscard]] FORCEINLINE bool IsZero() const
+	[[nodiscard]] UE_FORCEINLINE_HINT bool IsZero() const
 	{
 		return !(Hi | Lo);
 	}
 
-	[[nodiscard]] FORCEINLINE bool IsGreaterThanZero() const
+	[[nodiscard]] UE_FORCEINLINE_HINT bool IsGreaterThanZero() const
 	{
 		return !IsZero();
 	}
diff --git a/Engine/Source/Runtime/Core/Public/Math/InterpCurvePoint.h b/Engine/Source/Runtime/Core/Public/Math/InterpCurvePoint.h
--- a/Engine/Source/Runtime/Core/Public/Math/InterpCurvePoint.h
+++ b/Engine/Source/Runtime/Core/Public/Math/InterpCurvePoint.h
@@ -98,12 +98,12 @@ public:
 	 *
 	 * @param EForceInit Force init enum
 	 */
-	[[nodiscard]] explicit FORCEINLINE FInterpCurvePoint(EForceInit);
+	[[nodiscard]] explicit UE_FORCEINLINE_HINT FInterpCurvePoint(EForceInit);
 
 public:
 
 	/** @return true if the key value is using a curve interp mode, otherwise false */
-	[[nodiscard]] FORCEINLINE bool IsCurveKey() const;
+	[[nodiscard]] UE_FORCEINLINE_HINT bool IsCurveKey() const;
 
 public:
 
@@ -180,7 +180,7 @@ inline FInterpCurvePoint<T>::FInterpCurvePoint(EForceInit)
 }
 
 template< class T > 
-FORCEINLINE bool FInterpCurvePoint<T>::IsCurveKey() const
+UE_FORCEINLINE_HINT bool FInterpCurvePoint<T>::IsCurveKey() const
 {
 	return ((InterpMode == CIM_CurveAuto) || (InterpMode == CIM_CurveAutoClamped) || (InterpMode == CIM_CurveUser) || (InterpMode == CIM_CurveBreak));
 }
diff --git a/Engine/Source/Runtime/Core/Public/Math/Matrix.h b/Engine/Source/Runtime/Core/Public/Math/Matrix.h
--- a/Engine/Source/Runtime/Core/Public/Math/Matrix.h
+++ b/Engine/Source/Runtime/Core/Public/Math/Matrix.h
@@ -62,7 +62,7 @@ public:
 		}
 	}
 #else
-	FORCEINLINE void DiagnosticCheckNaN() const {}
+	void DiagnosticCheckNaN() const {}
 #endif
 
 
@@ -74,7 +74,7 @@ public:
 	 *
 	 * @param EForceInit Force Init Enum.
 	 */
-	[[nodiscard]] explicit FORCEINLINE TMatrix(EForceInit)
+	[[nodiscard]] explicit TMatrix(EForceInit)
 	{
 		FMemory::Memzero(this, sizeof(*this));
 	}
@@ -87,7 +87,7 @@ public:
 	 * @param InZ Z plane
 	 * @param InW W plane
 	 */
-	[[nodiscard]] FORCEINLINE TMatrix(const TPlane<T>& InX, const TPlane<T>& InY, const TPlane<T>& InZ, const TPlane<T>& InW);
+	[[nodiscard]] TMatrix(const TPlane<T>& InX, const TPlane<T>& InY, const TPlane<T>& InZ, const TPlane<T>& InW);
 
 	/**
 	 * Constructor.
@@ -97,7 +97,7 @@ public:
 	 * @param InZ Z vector
 	 * @param InW W vector
 	 */
-	[[nodiscard]] FORCEINLINE TMatrix(const TVector<T>& InX, const TVector<T>& InY, const TVector<T>& InZ, const TVector<T>& InW);
+	[[nodiscard]] TMatrix(const TVector<T>& InX, const TVector<T>& InY, const TVector<T>& InZ, const TVector<T>& InW);
 
 	// Set this to the identity matrix
 	inline void SetIdentity();
@@ -108,7 +108,7 @@ public:
 	 * @param Other The matrix to multiply this by.
 	 * @return The result of multiplication.
 	 */
-	[[nodiscard]] FORCEINLINE TMatrix<T> operator* (const TMatrix<T>& Other) const;
+	[[nodiscard]] TMatrix<T> operator* (const TMatrix<T>& Other) const;
 
 	/**
 	 * Multiply this by a matrix.
@@ -116,7 +116,7 @@ public:
 	 * @param Other the matrix to multiply by this.
 	 * @return reference to this after multiply.
 	 */
-	FORCEINLINE void operator*=(const TMatrix<T>& Other);
+	void operator*=(const TMatrix<T>& Other);
 
 	/**
 	 * Gets the result of adding a matrix to this.
@@ -124,7 +124,7 @@ public:
 	 * @param Other The Matrix to add.
 	 * @return The result of addition.
 	 */
-	[[nodiscard]] FORCEINLINE TMatrix<T> operator+ (const TMatrix<T>& Other) const;
+	[[nodiscard]] TMatrix<T> operator+ (const TMatrix<T>& Other) const;
 
 	/**
 	 * Adds to this matrix.
@@ -132,12 +132,12 @@ public:
 	 * @param Other The matrix to add to this.
 	 * @return Reference to this after addition.
 	 */
-	FORCEINLINE void operator+=(const TMatrix<T>& Other);
+	void operator+=(const TMatrix<T>& Other);
 
 	/**
 	  * This isn't applying SCALE, just multiplying the value to all members - i.e. weighting
 	  */
-	[[nodiscard]] FORCEINLINE TMatrix<T> operator* (T Other) const;
+	[[nodiscard]] TMatrix<T> operator* (T Other) const;
 
 	/**
 	 * Multiply this matrix by a weighting factor.
@@ -145,7 +145,7 @@ public:
 	 * @param other The weight.
 	 * @return a reference to this after weighting.
 	 */
-	FORCEINLINE void operator*=(T Other);
+	void operator*=(T Other);
 
 	/**
 	 * Checks whether two matrix are identical.
@@ -173,29 +173,29 @@ public:
 	[[nodiscard]] inline bool operator!=(const TMatrix<T>& Other) const;
 
 	// Homogeneous transform.
-	[[nodiscard]] FORCEINLINE TVector4<T> TransformFVector4(const TVector4<T>& V) const;
+	[[nodiscard]] TVector4<T> TransformFVector4(const TVector4<T>& V) const;
 
 	/** Transform a location - will take into account translation part of the TMatrix<T>. */
-	[[nodiscard]] FORCEINLINE TVector4<T> TransformPosition(const TVector<T>& V) const;
+	[[nodiscard]] TVector4<T> TransformPosition(const TVector<T>& V) const;
 
 	/** Inverts the matrix and then transforms V - correctly handles scaling in this matrix. */
-	[[nodiscard]] FORCEINLINE TVector<T> InverseTransformPosition(const TVector<T>& V) const;
+	[[nodiscard]] TVector<T> InverseTransformPosition(const TVector<T>& V) const;
 
 	/**
 	 *	Transform a direction vector - will not take into account translation part of the TMatrix<T>.
 	 *	If you want to transform a surface normal (or plane) and correctly account for non-uniform scaling you should use TransformByUsingAdjointT.
 	 */
-	[[nodiscard]] FORCEINLINE TVector4<T> TransformVector(const TVector<T>& V) const;
+	[[nodiscard]] TVector4<T> TransformVector(const TVector<T>& V) const;
 
 	/**
 	 *	Transform a direction vector by the inverse of this matrix - will not take into account translation part.
 	 *	If you want to transform a surface normal (or plane) and correctly account for non-uniform scaling you should use TransformByUsingAdjointT with adjoint of matrix inverse.
 	 */
-	[[nodiscard]] FORCEINLINE TVector<T> InverseTransformVector(const TVector<T>& V) const;
+	[[nodiscard]] TVector<T> InverseTransformVector(const TVector<T>& V) const;
 
 	// Transpose.
 
-	[[nodiscard]] FORCEINLINE TMatrix<T> GetTransposed() const;
+	[[nodiscard]] TMatrix<T> GetTransposed() const;
 
 	// @return determinant of this matrix.
 
@@ -333,26 +333,26 @@ public:
 	/**
 	 * Convert this Atom to the 3x4 transpose of the transformation matrix.
 	 */
-	FORCEINLINE void To3x4MatrixTranspose(T* Out) const;
+	void To3x4MatrixTranspose(T* Out) const;
 
 	// Frustum plane extraction.
 	/** @param OuTPln the near plane of the Frustum of this matrix */
-	FORCEINLINE bool GetFrustumNearPlane(TPlane<T>& OuTPln) const;
+	bool GetFrustumNearPlane(TPlane<T>& OuTPln) const;
 
 	/** @param OuTPln the far plane of the Frustum of this matrix */
-	FORCEINLINE bool GetFrustumFarPlane(TPlane<T>& OuTPln) const;
+	bool GetFrustumFarPlane(TPlane<T>& OuTPln) const;
 
 	/** @param OuTPln the left plane of the Frustum of this matrix */
-	FORCEINLINE bool GetFrustumLeftPlane(TPlane<T>& OuTPln) const;
+	bool GetFrustumLeftPlane(TPlane<T>& OuTPln) const;
 
 	/** @param OuTPln the right plane of the Frustum of this matrix */
-	FORCEINLINE bool GetFrustumRightPlane(TPlane<T>& OuTPln) const;
+	bool GetFrustumRightPlane(TPlane<T>& OuTPln) const;
 
 	/** @param OuTPln the top plane of the Frustum of this matrix */
-	FORCEINLINE bool GetFrustumTopPlane(TPlane<T>& OuTPln) const;
+	bool GetFrustumTopPlane(TPlane<T>& OuTPln) const;
 
 	/** @param OuTPln the bottom plane of the Frustum of this matrix */
-	FORCEINLINE bool GetFrustumBottomPlane(TPlane<T>& OuTPln) const;
+	bool GetFrustumBottomPlane(TPlane<T>& OuTPln) const;
 
 	/**
 	 * Utility for mirroring this transform across a certain plane, and flipping one of the axis as well.
diff --git a/Engine/Source/Runtime/Core/Public/Math/Matrix.inl b/Engine/Source/Runtime/Core/Public/Math/Matrix.inl
--- a/Engine/Source/Runtime/Core/Public/Math/Matrix.inl
+++ b/Engine/Source/Runtime/Core/Public/Math/Matrix.inl
@@ -181,7 +181,7 @@ inline TVector4<T> TMatrix<T>::TransformFVector4(const TVector4<T>& P) const
 
 /** Transform a location - will take into account translation part of the TMatrix<T>. */
 template<typename T>
-FORCEINLINE TVector4<T> TMatrix<T>::TransformPosition(const TVector<T>& V) const
+UE_FORCEINLINE_HINT TVector4<T> TMatrix<T>::TransformPosition(const TVector<T>& V) const
 {
 	return TransformFVector4(TVector4<T>(V.X, V.Y, V.Z, 1.0f));
 }
@@ -201,7 +201,7 @@ inline TVector<T> TMatrix<T>::InverseTransformPosition(const TVector<T>& V) cons
  *	If you want to transform a surface normal (or plane) and correctly account for non-uniform scaling you should use TransformByUsingAdjointT.
  */
 template<typename T>
-FORCEINLINE TVector4<T> TMatrix<T>::TransformVector(const TVector<T>& V) const
+UE_FORCEINLINE_HINT TVector4<T> TMatrix<T>::TransformVector(const TVector<T>& V) const
 {
 	return TransformFVector4(TVector4<T>(V.X, V.Y, V.Z, 0.0f));
 }
@@ -977,7 +977,7 @@ inline TLookFromMatrix<T>::TLookFromMatrix(const TVector<T>& EyePosition, const
 
 
 template<typename T>
-FORCEINLINE TLookAtMatrix<T>::TLookAtMatrix(const TVector<T>& EyePosition, const TVector<T>& LookAtPosition, const TVector<T>& UpVector) :
+UE_FORCEINLINE_HINT TLookAtMatrix<T>::TLookAtMatrix(const TVector<T>& EyePosition, const TVector<T>& LookAtPosition, const TVector<T>& UpVector) :
 	TLookFromMatrix<T>(EyePosition, LookAtPosition - EyePosition, UpVector)
 {
 }
diff --git a/Engine/Source/Runtime/Core/Public/Math/OrientedBox.h b/Engine/Source/Runtime/Core/Public/Math/OrientedBox.h
--- a/Engine/Source/Runtime/Core/Public/Math/OrientedBox.h
+++ b/Engine/Source/Runtime/Core/Public/Math/OrientedBox.h
@@ -56,14 +56,14 @@ public:
 	 *
 	 * @param Verts The array to fill in with the vertices.
 	 */
-	FORCEINLINE void CalcVertices(FVector* Verts) const;
+	inline void CalcVertices(FVector* Verts) const;
 
 	/**
 	 * Finds the projection interval of the box when projected onto Axis.
 	 *
 	 * @param Axis The unit vector defining the axis to project the box onto.
 	 */
-	[[nodiscard]] FORCEINLINE FFloatInterval Project(const FVector& Axis) const;
+	[[nodiscard]] inline FFloatInterval Project(const FVector& Axis) const;
 };
 
 
diff --git a/Engine/Source/Runtime/Core/Public/Math/Plane.h b/Engine/Source/Runtime/Core/Public/Math/Plane.h
--- a/Engine/Source/Runtime/Core/Public/Math/Plane.h
+++ b/Engine/Source/Runtime/Core/Public/Math/Plane.h
@@ -54,7 +54,7 @@ public:
 		}
 	}
 #else
-	FORCEINLINE void DiagnosticCheckNaN() const {}
+	UE_FORCEINLINE_HINT void DiagnosticCheckNaN() const {}
 #endif
 
 public:
@@ -67,7 +67,7 @@ public:
 	 *
 	 * @param V 4D vector to set up plane.
 	 */
-	[[nodiscard]] FORCEINLINE TPlane(const TVector4<T>& V);
+	[[nodiscard]] UE_FORCEINLINE_HINT TPlane(const TVector4<T>& V);
 
 	/**
 	 * Constructor.
@@ -77,7 +77,7 @@ public:
 	 * @param InZ Z-coefficient.
 	 * @param InW W-coefficient.
 	 */
-	[[nodiscard]] FORCEINLINE TPlane(T InX, T InY, T InZ, T InW);
+	[[nodiscard]] inline TPlane(T InX, T InY, T InZ, T InW);
 
 	/**
 	 * Constructor.
@@ -85,7 +85,7 @@ public:
 	 * @param InNormal Plane Normal Vector.
 	 * @param InW Plane W-coefficient.
 	 */
-	[[nodiscard]] FORCEINLINE TPlane(TVector<T> InNormal, T InW);
+	[[nodiscard]] inline TPlane(TVector<T> InNormal, T InW);
 
 	/**
 	 * Constructor.
@@ -93,7 +93,7 @@ public:
 	 * @param InBase Base point in plane.
 	 * @param InNormal Plane Normal Vector.
 	 */
-	[[nodiscard]] FORCEINLINE TPlane(TVector<T> InBase, const TVector<T>& InNormal);
+	[[nodiscard]] inline TPlane(TVector<T> InBase, const TVector<T>& InNormal);
 
 	/**
 	 * Constructor.
@@ -109,7 +109,7 @@ public:
 	 *
 	 * @param EForceInit Force Init Enum.
 	 */
-	[[nodiscard]] explicit FORCEINLINE TPlane(EForceInit);
+	[[nodiscard]] explicit UE_FORCEINLINE_HINT TPlane(EForceInit);
 
 	// Functions.
 
@@ -118,21 +118,21 @@ public:
 	 *
 	 * @return true if the plane is well-defined (has a non-zero normal), otherwise false.
 	 */
-	[[nodiscard]] FORCEINLINE bool IsValid() const;
+	[[nodiscard]] UE_FORCEINLINE_HINT bool IsValid() const;
 
 	/**
 	 * Get the origin of this plane.
 	 *
 	 * @return The origin (base point) of this plane.
 	 */
-	[[nodiscard]] FORCEINLINE TVector<T> GetOrigin() const;
+	[[nodiscard]] UE_FORCEINLINE_HINT TVector<T> GetOrigin() const;
 
 	/**
 	 * Get the normal of this plane.
 	 *
 	 * @return The normal of this plane.
 	 */
-	[[nodiscard]] FORCEINLINE const TVector<T>& GetNormal() const;
+	[[nodiscard]] UE_FORCEINLINE_HINT const TVector<T>& GetNormal() const;
 
 
 	/**
@@ -141,7 +141,7 @@ public:
 	 * @param P The other point.
 	 * @return The distance from the plane to the point. 0: Point is on the plane. >0: Point is in front of the plane. <0: Point is behind the plane.
 	 */
-	[[nodiscard]] FORCEINLINE T PlaneDot(const TVector<T>& P) const;
+	[[nodiscard]] UE_FORCEINLINE_HINT T PlaneDot(const TVector<T>& P) const;
 
 	/**
 	 * Normalize this plane in-place if it is larger than a given tolerance. Leaves it unchanged if not.
@@ -217,7 +217,7 @@ public:
 	 * @param V The other plane.
 	 * @return The dot product.
 	 */
-	[[nodiscard]] FORCEINLINE T operator|(const TPlane<T>& V) const;
+	[[nodiscard]] UE_FORCEINLINE_HINT T operator|(const TPlane<T>& V) const;
 
 	/**
 	 * Gets result of adding a plane to this.
@@ -407,7 +407,7 @@ inline FArchive& operator<<(FArchive& Ar, TPlane<double>& P)
  *****************************************************************************/
 
 template<typename T>
-FORCEINLINE TPlane<T>::TPlane(const TVector4<T>& V)
+UE_FORCEINLINE_HINT TPlane<T>::TPlane(const TVector4<T>& V)
 	: TVector<T>(V)
 	, W(V.W)
 {}
@@ -448,30 +448,30 @@ inline TPlane<T>::TPlane(TVector<T> A, TVector<T> B, TVector<T> C)
 
 
 template<typename T>
-FORCEINLINE TPlane<T>::TPlane(EForceInit)
+UE_FORCEINLINE_HINT TPlane<T>::TPlane(EForceInit)
 	: TVector<T>(ForceInit), W(0.f)
 {}
 
 template<typename T>
-FORCEINLINE bool TPlane<T>::IsValid() const
+UE_FORCEINLINE_HINT bool TPlane<T>::IsValid() const
 {
 	return !this->IsNearlyZero();
 }
 
 template<typename T>
-FORCEINLINE const TVector<T>& TPlane<T>::GetNormal() const
+UE_FORCEINLINE_HINT const TVector<T>& TPlane<T>::GetNormal() const
 {
 	return *this;
 }
 
 template<typename T>
-FORCEINLINE TVector<T> TPlane<T>::GetOrigin() const
+UE_FORCEINLINE_HINT TVector<T> TPlane<T>::GetOrigin() const
 {
 	return GetNormal() * W;
 }
 
 template<typename T>
-FORCEINLINE T TPlane<T>::PlaneDot(const TVector<T>& P) const
+UE_FORCEINLINE_HINT T TPlane<T>::PlaneDot(const TVector<T>& P) const
 {
 	return X * P.X + Y * P.Y + Z * P.Z - W;
 }
@@ -490,54 +490,54 @@ inline bool TPlane<T>::Normalize(T Tolerance)
 }
 
 template<typename T>
-FORCEINLINE TPlane<T> TPlane<T>::Flip() const
+UE_FORCEINLINE_HINT TPlane<T> TPlane<T>::Flip() const
 {
 	return TPlane<T>(-X, -Y, -Z, -W);
 }
 
 template<typename T>
-FORCEINLINE TPlane<T> TPlane<T>::TranslateBy(const TVector<T>& V) const
+UE_FORCEINLINE_HINT TPlane<T> TPlane<T>::TranslateBy(const TVector<T>& V) const
 {
 	return TPlane<T>(GetOrigin() + V, GetNormal());
 }
 
 template<typename T>
-FORCEINLINE bool TPlane<T>::operator==(const TPlane<T>& V) const
+UE_FORCEINLINE_HINT bool TPlane<T>::operator==(const TPlane<T>& V) const
 {
 	return (X == V.X) && (Y == V.Y) && (Z == V.Z) && (W == V.W);
 }
 
 
 template<typename T>
-FORCEINLINE bool TPlane<T>::operator!=(const TPlane<T>& V) const
+UE_FORCEINLINE_HINT bool TPlane<T>::operator!=(const TPlane<T>& V) const
 {
 	return (X != V.X) || (Y != V.Y) || (Z != V.Z) || (W != V.W);
 }
 
 
 template<typename T>
-FORCEINLINE bool TPlane<T>::Equals(const TPlane<T>& V, T Tolerance) const
+UE_FORCEINLINE_HINT bool TPlane<T>::Equals(const TPlane<T>& V, T Tolerance) const
 {
 	return (FMath::Abs(X - V.X) < Tolerance) && (FMath::Abs(Y - V.Y) < Tolerance) && (FMath::Abs(Z - V.Z) < Tolerance) && (FMath::Abs(W - V.W) < Tolerance);
 }
 
 
 template<typename T>
-FORCEINLINE T TPlane<T>::operator|(const TPlane<T>& V) const
+UE_FORCEINLINE_HINT T TPlane<T>::operator|(const TPlane<T>& V) const
 {
 	return X * V.X + Y * V.Y + Z * V.Z + W * V.W;
 }
 
 
 template<typename T>
-FORCEINLINE TPlane<T> TPlane<T>::operator+(const TPlane<T>& V) const
+UE_FORCEINLINE_HINT TPlane<T> TPlane<T>::operator+(const TPlane<T>& V) const
 {
 	return TPlane<T>(X + V.X, Y + V.Y, Z + V.Z, W + V.W);
 }
 
 
 template<typename T>
-FORCEINLINE TPlane<T> TPlane<T>::operator-(const TPlane<T>& V) const
+UE_FORCEINLINE_HINT TPlane<T> TPlane<T>::operator-(const TPlane<T>& V) const
 {
 	return TPlane<T>(X - V.X, Y - V.Y, Z - V.Z, W - V.W);
 }
@@ -552,14 +552,14 @@ inline TPlane<T> TPlane<T>::operator/(T Scale) const
 
 
 template<typename T>
-FORCEINLINE TPlane<T> TPlane<T>::operator*(T Scale) const
+UE_FORCEINLINE_HINT TPlane<T> TPlane<T>::operator*(T Scale) const
 {
 	return TPlane<T>(X * Scale, Y * Scale, Z * Scale, W * Scale);
 }
 
 
 template<typename T>
-FORCEINLINE TPlane<T> TPlane<T>::operator*(const TPlane<T>& V) const
+UE_FORCEINLINE_HINT TPlane<T> TPlane<T>::operator*(const TPlane<T>& V) const
 {
 	return TPlane<T>(X * V.X, Y * V.Y, Z * V.Z, W * V.W);
 }
diff --git a/Engine/Source/Runtime/Core/Public/Math/Quat.h b/Engine/Source/Runtime/Core/Public/Math/Quat.h
--- a/Engine/Source/Runtime/Core/Public/Math/Quat.h
+++ b/Engine/Source/Runtime/Core/Public/Math/Quat.h
@@ -72,7 +72,7 @@ public:
 	 *
 	 * @param EForceInit Force init enum: if equal to ForceInitToZero then W is 0, otherwise W = 1 (creating an identity transform)
 	 */
-	[[nodiscard]] explicit FORCEINLINE TQuat(EForceInit);
+	[[nodiscard]] explicit UE_FORCEINLINE_HINT TQuat(EForceInit);
 
 	/**
 	 * Constructor.
@@ -82,7 +82,7 @@ public:
 	 * @param InZ Z component of the quaternion
 	 * @param InW W component of the quaternion
 	 */
-	[[nodiscard]] FORCEINLINE TQuat(T InX, T InY, T InZ, T InW);
+	[[nodiscard]] inline TQuat(T InX, T InY, T InZ, T InW);
 
 	/**
 	 * Initializes all elements to V
@@ -107,7 +107,7 @@ protected:
 
 public:
 
-	[[nodiscard]] FORCEINLINE static TQuat<T> MakeFromVectorRegister(QuatVectorRegister V)
+	[[nodiscard]] UE_FORCEINLINE_HINT static TQuat<T> MakeFromVectorRegister(QuatVectorRegister V)
 	{
 		return TQuat<T>(V);
 	}
@@ -119,7 +119,7 @@ public:
 	 */
 	explicit TQuat(const TRotator<T>& R);
 
-	[[nodiscard]] FORCEINLINE static TQuat<T> MakeFromRotator(const TRotator<T>& R)
+	[[nodiscard]] UE_FORCEINLINE_HINT static TQuat<T> MakeFromRotator(const TRotator<T>& R)
 	{
 		return TQuat<T>(R);
 	}
@@ -148,7 +148,7 @@ public:
 	 * @param Q The Quaternion to add.
 	 * @return The result of addition.
 	 */
-	[[nodiscard]] FORCEINLINE TQuat<T> operator+(const TQuat<T>& Q) const;
+	[[nodiscard]] inline TQuat<T> operator+(const TQuat<T>& Q) const;
 
 	/**
 	 * Adds to this quaternion.
@@ -157,7 +157,7 @@ public:
 	 * @param Other The quaternion to add to this.
 	 * @return Result after addition.
 	 */
-	FORCEINLINE TQuat<T> operator+=(const TQuat<T>& Q);
+	inline TQuat<T> operator+=(const TQuat<T>& Q);
 
 	/**
 	 * Gets the result of subtracting a Quaternion to this.
@@ -166,14 +166,14 @@ public:
 	 * @param Q The Quaternion to subtract.
 	 * @return The result of subtraction.
 	 */
-	[[nodiscard]] FORCEINLINE TQuat<T> operator-(const TQuat<T>& Q) const;
+	[[nodiscard]] inline TQuat<T> operator-(const TQuat<T>& Q) const;
 
 	/**
 	* Negates the quaternion. Note that this represents the same rotation.
 	*
 	* @return The result of negation.
 	*/
-	[[nodiscard]] FORCEINLINE TQuat<T> operator-() const;
+	[[nodiscard]] inline TQuat<T> operator-() const;
 
 	/**
 	 * Checks whether another Quaternion is equal to this, within specified tolerance.
@@ -182,7 +182,7 @@ public:
 	 * @param Tolerance Error tolerance for comparison with other Quaternion.
 	 * @return true if two Quaternions are equal, within specified tolerance, otherwise false.
 	 */
-	[[nodiscard]] FORCEINLINE bool Equals(const TQuat<T>& Q, T Tolerance=UE_KINDA_SMALL_NUMBER) const;
+	[[nodiscard]] inline bool Equals(const TQuat<T>& Q, T Tolerance=UE_KINDA_SMALL_NUMBER) const;
 
 	/**
 	 * Checks whether this Quaternion is an Identity Quaternion.
@@ -191,7 +191,7 @@ public:
 	 * @param Tolerance Error tolerance for comparison with Identity Quaternion.
 	 * @return true if Quaternion is a normalized Identity Quaternion.
 	 */
-	[[nodiscard]] FORCEINLINE bool IsIdentity(T Tolerance=UE_SMALL_NUMBER) const;
+	[[nodiscard]] UE_FORCEINLINE_HINT bool IsIdentity(T Tolerance=UE_SMALL_NUMBER) const;
 
 	/**
 	 * Subtracts another quaternion from this.
@@ -200,7 +200,7 @@ public:
 	 * @param Q The other quaternion.
 	 * @return reference to this after subtraction.
 	 */
-	FORCEINLINE TQuat<T> operator-=(const TQuat<T>& Q);
+	inline TQuat<T> operator-=(const TQuat<T>& Q);
 
 	/**
 	 * Gets the result of multiplying this by another quaternion (this * Q).
@@ -211,7 +211,7 @@ public:
 	 * @param Q The Quaternion to multiply this by.
 	 * @return The result of multiplication (this * Q).
 	 */
-	[[nodiscard]] FORCEINLINE TQuat<T> operator*(const TQuat<T>& Q) const;
+	[[nodiscard]] inline TQuat<T> operator*(const TQuat<T>& Q) const;
 
 	/**
 	 * Multiply this by a quaternion (this = this * Q).
@@ -222,7 +222,7 @@ public:
 	 * @param Q the quaternion to multiply this by.
 	 * @return The result of multiplication (this * Q).
 	 */
-	FORCEINLINE TQuat<T> operator*=(const TQuat<T>& Q);
+	inline TQuat<T> operator*=(const TQuat<T>& Q);
 
 	/**
 	 * Rotate a vector by this quaternion.
@@ -382,7 +382,7 @@ public:
 	 *
 	 * @param Tolerance Minimum squared length of quaternion for normalization.
 	 */
-	FORCEINLINE void Normalize(T Tolerance = UE_SMALL_NUMBER);
+	inline void Normalize(T Tolerance = UE_SMALL_NUMBER);
 
 	/**
 	 * Get a normalized copy of this quaternion.
@@ -390,7 +390,7 @@ public:
 	 *
 	 * @param Tolerance Minimum squared length of quaternion for normalization.
 	 */
-	[[nodiscard]] FORCEINLINE TQuat<T> GetNormalized(T Tolerance = UE_SMALL_NUMBER) const;
+	[[nodiscard]] inline TQuat<T> GetNormalized(T Tolerance = UE_SMALL_NUMBER) const;
 
 	// Return true if this quaternion is normalized
 	[[nodiscard]] bool IsNormalized() const;
@@ -400,21 +400,21 @@ public:
 	 *
 	 * @return The length of this quaternion.
 	 */
-	[[nodiscard]] FORCEINLINE T Size() const;
+	[[nodiscard]] UE_FORCEINLINE_HINT T Size() const;
 
 	/**
 	 * Get the length squared of this quaternion.
 	 *
 	 * @return The length of this quaternion.
 	 */
-	[[nodiscard]] FORCEINLINE T SizeSquared() const;
+	[[nodiscard]] UE_FORCEINLINE_HINT T SizeSquared() const;
 
 	/** 
 	 * Get the angle in radians of this quaternion 
 	 * @warning : The rotation returned may not be the shortest version.
 	 *            You may wish to call GetShortestArcWith(FQuat::Identity) first.
 	 */
-	[[nodiscard]] FORCEINLINE T GetAngle() const;
+	[[nodiscard]] UE_FORCEINLINE_HINT T GetAngle() const;
 
 	/** 
 	 * get the axis and angle of rotation of this quaternion
@@ -498,7 +498,7 @@ public:
 	 * @return inverse of this quaternion
 	 * @warning : Requires this quaternion to be normalized.
 	 */
-	[[nodiscard]] FORCEINLINE TQuat<T> Inverse() const;
+	[[nodiscard]] TQuat<T> Inverse() const;
 
 	/**
 	 * Enforce that the delta between this Quaternion and another one represents
@@ -513,31 +513,31 @@ public:
 	[[nodiscard]] TQuat<T> GetShortestArcWith(const TQuat<T>& OtherQuat) const;
 
 	/** Get the forward direction (X axis) after it has been rotated by this Quaternion. */
-	[[nodiscard]] FORCEINLINE TVector<T> GetAxisX() const;
+	[[nodiscard]] UE_FORCEINLINE_HINT TVector<T> GetAxisX() const;
 
 	/** Get the right direction (Y axis) after it has been rotated by this Quaternion. */
-	[[nodiscard]] FORCEINLINE TVector<T> GetAxisY() const;
+	[[nodiscard]] UE_FORCEINLINE_HINT TVector<T> GetAxisY() const;
 
 	/** Get the up direction (Z axis) after it has been rotated by this Quaternion. */
-	[[nodiscard]] FORCEINLINE TVector<T> GetAxisZ() const;
+	[[nodiscard]] UE_FORCEINLINE_HINT TVector<T> GetAxisZ() const;
 
 	/** Get the forward direction (X axis) after it has been rotated by this Quaternion. */
-	[[nodiscard]] FORCEINLINE TVector<T> GetForwardVector() const;
+	[[nodiscard]] UE_FORCEINLINE_HINT TVector<T> GetForwardVector() const;
 
 	/** Get the right direction (Y axis) after it has been rotated by this Quaternion. */
-	[[nodiscard]] FORCEINLINE TVector<T> GetRightVector() const;
+	[[nodiscard]] UE_FORCEINLINE_HINT TVector<T> GetRightVector() const;
 
 	/** Get the up direction (Z axis) after it has been rotated by this Quaternion. */
-	[[nodiscard]] FORCEINLINE TVector<T> GetUpVector() const;
+	[[nodiscard]] UE_FORCEINLINE_HINT TVector<T> GetUpVector() const;
 
 	/** Convert a rotation into a unit vector facing in its direction. Equivalent to GetForwardVector(). */
-	[[nodiscard]] FORCEINLINE TVector<T> Vector() const;
+	[[nodiscard]] UE_FORCEINLINE_HINT TVector<T> Vector() const;
 
 	/** Get the TRotator<T> representation of this Quaternion. */
 	[[nodiscard]] CORE_API TRotator<T> Rotator() const;
 
 	/** Get the TMatrix<T> representation of this Quaternion. */
-	[[nodiscard]] FORCEINLINE TMatrix<T> ToMatrix() const;
+	[[nodiscard]] UE_FORCEINLINE_HINT TMatrix<T> ToMatrix() const;
 
 	/** Get the TMatrix<T> representation of this Quaternion and store it in Mat */
 	CORE_API void ToMatrix(TMatrix<T>& Mat) const;
@@ -547,10 +547,10 @@ public:
 	 * This is the axis around which rotation occurs to transform the canonical coordinate system to the target orientation.
 	 * For the identity Quaternion which has no such rotation, TVector<T>(1,0,0) is returned.
 	 */
-	[[nodiscard]] FORCEINLINE TVector<T> GetRotationAxis() const;
+	[[nodiscard]] inline TVector<T> GetRotationAxis() const;
 
 	/** Find the angular distance between two rotation quaternions (in radians) */
-	[[nodiscard]] FORCEINLINE T AngularDistance(const TQuat<T>& Q) const;
+	[[nodiscard]] T AngularDistance(const TQuat<T>& Q) const;
 
 	/**
 	 * Serializes the vector compressed for e.g. network transmission.
@@ -604,8 +604,8 @@ public:
 		}
 	}
 #else
-	FORCEINLINE void DiagnosticCheckNaN() const {}
-	FORCEINLINE void DiagnosticCheckNaN(const TCHAR* Message) const {}
+	UE_FORCEINLINE_HINT void DiagnosticCheckNaN() const {}
+	UE_FORCEINLINE_HINT void DiagnosticCheckNaN(const TCHAR* Message) const {}
 #endif
 
 public:
@@ -613,7 +613,7 @@ public:
 	/**
 	 * Generates the 'smallest' (geodesic) rotation between two vectors of arbitrary length.
 	 */
-	[[nodiscard]] static FORCEINLINE TQuat<T> FindBetween(const TVector<T>& Vector1, const TVector<T>& Vector2)
+	[[nodiscard]] static UE_FORCEINLINE_HINT TQuat<T> FindBetween(const TVector<T>& Vector1, const TVector<T>& Vector2)
 	{
 		return FindBetweenVectors(Vector1, Vector2);
 	}
@@ -633,31 +633,31 @@ public:
 	 * Returns the hypersphere-angle between two quaternions; alignment shouldn't matter, though 
 	 * @note normalized input is expected.
 	 */
-	[[nodiscard]] static FORCEINLINE T Error(const TQuat<T>& Q1, const TQuat<T>& Q2);
+	[[nodiscard]] static inline T Error(const TQuat<T>& Q1, const TQuat<T>& Q2);
 
 	/**
 	 * TQuat<T>::Error with auto-normalization.
 	 */
-	[[nodiscard]] static FORCEINLINE T ErrorAutoNormalize(const TQuat<T>& A, const TQuat<T>& B);
+	[[nodiscard]] static inline T ErrorAutoNormalize(const TQuat<T>& A, const TQuat<T>& B);
 
 	/** 
 	 * Fast Linear Quaternion Interpolation.
 	 * Result is NOT normalized.
 	 */
-	[[nodiscard]] static FORCEINLINE TQuat<T> FastLerp(const TQuat<T>& A, const TQuat<T>& B, const T Alpha);
+	[[nodiscard]] static inline TQuat<T> FastLerp(const TQuat<T>& A, const TQuat<T>& B, const T Alpha);
 
 	/** 
 	 * Bi-Linear Quaternion interpolation.
 	 * Result is NOT normalized.
 	 */
-	[[nodiscard]] static FORCEINLINE TQuat<T> FastBilerp(const TQuat<T>& P00, const TQuat<T>& P10, const TQuat<T>& P01, const TQuat<T>& P11, T FracX, T FracY);
+	[[nodiscard]] static inline TQuat<T> FastBilerp(const TQuat<T>& P00, const TQuat<T>& P10, const TQuat<T>& P01, const TQuat<T>& P11, T FracX, T FracY);
 
 
 	/** Spherical interpolation. Will correct alignment. Result is NOT normalized. */
 	[[nodiscard]] static CORE_API TQuat<T> Slerp_NotNormalized(const TQuat<T>& Quat1, const TQuat<T>& Quat2, T Slerp);
 
 	/** Spherical interpolation. Will correct alignment. Result is normalized. */
-	[[nodiscard]] static FORCEINLINE TQuat<T> Slerp(const TQuat<T>& Quat1, const TQuat<T>& Quat2, T Slerp)
+	[[nodiscard]] static UE_FORCEINLINE_HINT TQuat<T> Slerp(const TQuat<T>& Quat1, const TQuat<T>& Quat2, T Slerp)
 	{
 		return Slerp_NotNormalized(Quat1, Quat2, Slerp).GetNormalized();
 	}
@@ -674,7 +674,7 @@ public:
 	 * We need this for the cubic interpolation stuff so that the multiple Slerps dont go in different directions.
 	 * Result is normalized.
 	 */
-	[[nodiscard]] static FORCEINLINE TQuat<T> SlerpFullPath(const TQuat<T>& quat1, const TQuat<T>& quat2, T Alpha)
+	[[nodiscard]] static UE_FORCEINLINE_HINT TQuat<T> SlerpFullPath(const TQuat<T>& quat1, const TQuat<T>& quat2, T Alpha)
 	{
 		return SlerpFullPath_NotNormalized(quat1, quat2, Alpha).GetNormalized();
 	}
@@ -857,7 +857,7 @@ inline TQuat<T>::TQuat(const TRotator<T>& R)
 }
 
 template<typename T>
-FORCEINLINE TVector<T> TQuat<T>::operator*(const TVector<T>& V) const
+UE_FORCEINLINE_HINT TVector<T> TQuat<T>::operator*(const TVector<T>& V) const
 {
 	return RotateVector(V);
 }
@@ -866,7 +866,7 @@ FORCEINLINE TVector<T> TQuat<T>::operator*(const TVector<T>& V) const
  *****************************************************************************/
 
 template<typename T>
-FORCEINLINE TQuat<T>::TQuat(EForceInit ZeroOrNot)
+UE_FORCEINLINE_HINT TQuat<T>::TQuat(EForceInit ZeroOrNot)
 	:	X(0), Y(0), Z(0), W(ZeroOrNot == ForceInitToZero ? 0.0f : 1.0f)
 { }
 
@@ -890,7 +890,7 @@ inline TQuat<T>::TQuat(QuatVectorRegister V)
 
 
 template<typename T>
-FORCEINLINE FString TQuat<T>::ToString() const
+UE_FORCEINLINE_HINT FString TQuat<T>::ToString() const
 {
 	return FString::Printf(TEXT("X=%.9f Y=%.9f Z=%.9f W=%.9f"), X, Y, Z, W);
 }
@@ -999,7 +999,7 @@ inline bool TQuat<T>::Equals(const TQuat<T>& Q, T Tolerance) const
 }
 
 template<typename T>
-FORCEINLINE bool TQuat<T>::IsIdentity(T Tolerance) const
+UE_FORCEINLINE_HINT bool TQuat<T>::IsIdentity(T Tolerance) const
 {
 	return Equals(TQuat<T>::Identity, Tolerance);
 }
@@ -1052,20 +1052,20 @@ inline TQuat<T> TQuat<T>::operator*=(const TQuat<T>& Q)
 
 // Global operator for (float * Quat)
 template<typename T>
-FORCEINLINE TQuat<T> operator*(const float Scale, const TQuat<T>& Q)
+UE_FORCEINLINE_HINT TQuat<T> operator*(const float Scale, const TQuat<T>& Q)
 {
 	return Q.operator*(Scale);
 }
 
 // Global operator for (double * Quat)
 template<typename T>
-FORCEINLINE TQuat<T> operator*(const double Scale, const TQuat<T>& Q)
+UE_FORCEINLINE_HINT TQuat<T> operator*(const double Scale, const TQuat<T>& Q)
 {
 	return Q.operator*(Scale);
 }
 
 template<typename T>
-FORCEINLINE bool TQuat<T>::Identical(const TQuat* Q, const uint32 PortFlags) const
+UE_FORCEINLINE_HINT bool TQuat<T>::Identical(const TQuat* Q, const uint32 PortFlags) const
 {
 	return X == Q->X && Y == Q->Y && Z == Q->Z && W == Q->W;
 }
@@ -1097,7 +1097,7 @@ inline bool TQuat<T>::operator!=(const TQuat<T>& Q) const
 
 
 template<typename T>
-FORCEINLINE T TQuat<T>::operator|(const TQuat<T>& Q) const
+UE_FORCEINLINE_HINT T TQuat<T>::operator|(const TQuat<T>& Q) const
 {
 	return X * Q.X + Y * Q.Y + Z * Q.Z + W * Q.W;
 }
@@ -1159,19 +1159,19 @@ inline bool TQuat<T>::IsNormalized() const
 
 
 template<typename T>
-FORCEINLINE T TQuat<T>::Size() const
+UE_FORCEINLINE_HINT T TQuat<T>::Size() const
 {
 	return FMath::Sqrt(X * X + Y * Y + Z * Z + W * W);
 }
 
 template<typename T>
-FORCEINLINE T TQuat<T>::SizeSquared() const
+UE_FORCEINLINE_HINT T TQuat<T>::SizeSquared() const
 {
 	return (X * X + Y * Y + Z * Z + W * W);
 }
 
 template<typename T>
-FORCEINLINE T TQuat<T>::GetAngle() const
+UE_FORCEINLINE_HINT T TQuat<T>::GetAngle() const
 {
 	return T(2.0) * FMath::Acos(W);
 }
@@ -1293,46 +1293,46 @@ inline TQuat<T> TQuat<T>::GetShortestArcWith(const TQuat<T>& OtherQuat) const
 }
 
 template<typename T>
-FORCEINLINE TVector<T> TQuat<T>::GetAxisX() const
+UE_FORCEINLINE_HINT TVector<T> TQuat<T>::GetAxisX() const
 {
 	return RotateVector(TVector<T>(1.f, 0.f, 0.f));
 }
 
 
 template<typename T>
-FORCEINLINE TVector<T> TQuat<T>::GetAxisY() const
+UE_FORCEINLINE_HINT TVector<T> TQuat<T>::GetAxisY() const
 {
 	return RotateVector(TVector<T>(0.f, 1.f, 0.f));
 }
 
 
 template<typename T>
-FORCEINLINE TVector<T> TQuat<T>::GetAxisZ() const
+UE_FORCEINLINE_HINT TVector<T> TQuat<T>::GetAxisZ() const
 {
 	return RotateVector(TVector<T>(0.f, 0.f, 1.f));
 }
 
 
 template<typename T>
-FORCEINLINE TVector<T> TQuat<T>::GetForwardVector() const
+UE_FORCEINLINE_HINT TVector<T> TQuat<T>::GetForwardVector() const
 {
 	return GetAxisX();
 }
 
 template<typename T>
-FORCEINLINE TVector<T> TQuat<T>::GetRightVector() const
+UE_FORCEINLINE_HINT TVector<T> TQuat<T>::GetRightVector() const
 {
 	return GetAxisY();
 }
 
 template<typename T>
-FORCEINLINE TVector<T> TQuat<T>::GetUpVector() const
+UE_FORCEINLINE_HINT TVector<T> TQuat<T>::GetUpVector() const
 {
 	return GetAxisZ();
 }
 
 template<typename T>
-FORCEINLINE TVector<T> TQuat<T>::Vector() const
+UE_FORCEINLINE_HINT TVector<T> TQuat<T>::Vector() const
 {
 	return GetAxisX();
 }
@@ -1479,7 +1479,7 @@ struct TCustomLerp< UE::Math::TQuat<T> >
 	using QuatType = UE::Math::TQuat<T>;
 
 	template<class U>
-	static FORCEINLINE_DEBUGGABLE QuatType Lerp(const QuatType& A, const QuatType& B, const U& Alpha)
+	static UE_FORCEINLINE_HINT QuatType Lerp(const QuatType& A, const QuatType& B, const U& Alpha)
 	{
 		return QuatType::Slerp(A, B, (T)Alpha);
 	}
@@ -1499,7 +1499,7 @@ struct TCustomLerp< UE::Math::TQuat<T> >
 	}
 
 	template<class U>
-	static FORCEINLINE_DEBUGGABLE QuatType CubicInterp(const QuatType& P0, const QuatType& T0, const QuatType& P1, const QuatType& T1, const U& A)
+	static UE_FORCEINLINE_HINT QuatType CubicInterp(const QuatType& P0, const QuatType& T0, const QuatType& P1, const QuatType& T1, const U& A)
 	{
 		return QuatType::Squad(P0, T0, P1, T1, (T)A);
 	}
diff --git a/Engine/Source/Runtime/Core/Public/Math/RandomStream.h b/Engine/Source/Runtime/Core/Public/Math/RandomStream.h
--- a/Engine/Source/Runtime/Core/Public/Math/RandomStream.h
+++ b/Engine/Source/Runtime/Core/Public/Math/RandomStream.h
@@ -173,7 +173,7 @@ public:
 	 *
 	 * @return Random number.
 	 */
-	FORCEINLINE float FRand() const
+	UE_FORCEINLINE_HINT float FRand() const
 	{
 		return GetFraction();
 	}
@@ -183,7 +183,7 @@ public:
 	 *
 	 * @return A random number in [0..A)
 	 */
-	FORCEINLINE int32 RandHelper( int32 A ) const
+	UE_FORCEINLINE_HINT int32 RandHelper( int32 A ) const
 	{
 		// GetFraction guarantees a result in the [0,1) range.
 		return ((A > 0) ? FMath::TruncToInt(GetFraction() * float(A)) : 0);
@@ -206,7 +206,7 @@ public:
 	 *
 	 * @return A random number >= Min and <= Max
 	 */
-	FORCEINLINE FVector::FReal FRandRange( FVector::FReal InMin, FVector::FReal InMax ) const
+	UE_FORCEINLINE_HINT FVector::FReal FRandRange( FVector::FReal InMin, FVector::FReal InMax ) const
 	{
 		return InMin + (InMax - InMin) * FRand();
 	}
@@ -216,7 +216,7 @@ public:
 	 *
 	 * @return Random unit vector.
 	 */
-	FORCEINLINE FVector VRand() const
+	UE_FORCEINLINE_HINT FVector VRand() const
 	{
 		return GetUnitVector();
 	}
diff --git a/Engine/Source/Runtime/Core/Public/Math/Range.h b/Engine/Source/Runtime/Core/Public/Math/Range.h
--- a/Engine/Source/Runtime/Core/Public/Math/Range.h
+++ b/Engine/Source/Runtime/Core/Public/Math/Range.h
@@ -633,7 +633,7 @@ public:
 	 * @return A new range.
 	 * @see AtLeast, AtMost, Empty, Exclusive, GreaterThan, Inclusive, LessThan
 	 */
-	[[nodiscard]] static FORCEINLINE TRange All()
+	[[nodiscard]] static UE_FORCEINLINE_HINT TRange All()
 	{
 		return TRange(BoundsType::Open(), BoundsType::Open());
 	}
@@ -645,7 +645,7 @@ public:
 	 * @return A new range.
 	 * @see All, AtMost, Empty, Exclusive, GreaterThan, Inclusive, LessThan
 	 */
-	[[nodiscard]] static FORCEINLINE TRange AtLeast(ElementValueOrConstRef Value)
+	[[nodiscard]] static UE_FORCEINLINE_HINT TRange AtLeast(ElementValueOrConstRef Value)
 	{
 		return TRange(BoundsType::Inclusive(Value), BoundsType::Open());
 	}
@@ -657,7 +657,7 @@ public:
 	 * @return A new range.
 	 * @see All, AtLeast, Empty, Exclusive, GreaterThan, Inclusive, LessThan
 	 */
-	[[nodiscard]] static FORCEINLINE TRange AtMost(ElementValueOrConstRef Value)
+	[[nodiscard]] static UE_FORCEINLINE_HINT TRange AtMost(ElementValueOrConstRef Value)
 	{
 		return TRange(BoundsType::Open(), BoundsType::Inclusive(Value));
 	}
@@ -668,7 +668,7 @@ public:
 	 * @return Empty range.
 	 * @see All, AtLeast, AtMost, Exclusive, GreaterThan, Inclusive, LessThan
 	 */
-	[[nodiscard]] static FORCEINLINE TRange Empty()
+	[[nodiscard]] static UE_FORCEINLINE_HINT TRange Empty()
 	{
 		return TRange(BoundsType::Exclusive(ElementType()), BoundsType::Exclusive(ElementType()));
 	}
@@ -681,7 +681,7 @@ public:
 	 * @return A new range.
 	 * @see All, AtLeast, AtMost, Empty, Exclusive, GreaterThan, Inclusive, LessThan
 	 */
-	[[nodiscard]] static FORCEINLINE TRange Exclusive(ElementValueOrConstRef Min, ElementValueOrConstRef Max)
+	[[nodiscard]] static UE_FORCEINLINE_HINT TRange Exclusive(ElementValueOrConstRef Min, ElementValueOrConstRef Max)
 	{
 		return TRange(BoundsType::Exclusive(Min), BoundsType::Exclusive(Max));
 	}
@@ -693,7 +693,7 @@ public:
 	 * @return A new range.
 	 * @see All, AtLeast, AtMost, Empty, Exclusive, Inclusive, LessThan
 	 */
-	[[nodiscard]] static FORCEINLINE TRange GreaterThan(ElementValueOrConstRef Value)
+	[[nodiscard]] static UE_FORCEINLINE_HINT TRange GreaterThan(ElementValueOrConstRef Value)
 	{
 		return TRange(BoundsType::Exclusive(Value), BoundsType::Open());
 	}
@@ -706,7 +706,7 @@ public:
 	 * @return A new range.
 	 * @see All, AtLeast, AtMost, Empty, Exclusive, GreaterThan, LessThan
 	 */
-	[[nodiscard]] static FORCEINLINE TRange Inclusive(ElementValueOrConstRef Min, ElementValueOrConstRef Max)
+	[[nodiscard]] static UE_FORCEINLINE_HINT TRange Inclusive(ElementValueOrConstRef Min, ElementValueOrConstRef Max)
 	{
 		return TRange(BoundsType::Inclusive(Min), BoundsType::Inclusive(Max));
 	}
@@ -718,7 +718,7 @@ public:
 	 * @return A new range.
 	 * @see All, AtLeast, AtMost, Empty, Exclusive, GreaterThan, Inclusive
 	 */
-	[[nodiscard]] static FORCEINLINE TRange LessThan(ElementValueOrConstRef Value)
+	[[nodiscard]] static UE_FORCEINLINE_HINT TRange LessThan(ElementValueOrConstRef Value)
 	{
 		return TRange(BoundsType::Open(), BoundsType::Exclusive(Value));
 	}
@@ -798,62 +798,62 @@ private:
 			return TArray<Name>(Super::Split(Element)); \
 		} \
 		 \
-		static FORCEINLINE TArray<Name> Difference(const Name& X, const Name& Y) \
+		static UE_FORCEINLINE_HINT TArray<Name> Difference(const Name& X, const Name& Y) \
 		{ \
 			return TArray<Name>(Super::Difference(X, Y)); \
 		} \
 		 \
-		static FORCEINLINE Name Empty() \
+		static UE_FORCEINLINE_HINT Name Empty() \
 		{ \
 			return Super::Empty(); \
 		} \
 		 \
-		static FORCEINLINE Name Hull(const Name& X, const Name& Y) \
+		static UE_FORCEINLINE_HINT Name Hull(const Name& X, const Name& Y) \
 		{ \
 			return Super::Hull(X, Y); \
 		} \
 		 \
-		static FORCEINLINE Name Hull(const TArray<Name>& Ranges) \
+		static UE_FORCEINLINE_HINT Name Hull(const TArray<Name>& Ranges) \
 		{ \
 			return Super::Hull(reinterpret_cast<const TArray<Super>&>(Ranges)); \
 		} \
 		 \
-		static FORCEINLINE Name Intersection(const Name& X, const Name& Y) \
+		static UE_FORCEINLINE_HINT Name Intersection(const Name& X, const Name& Y) \
 		{ \
 			return Super::Intersection(X, Y); \
 		} \
 		 \
-		static FORCEINLINE Name Intersection(const TArray<Name>& Ranges) \
+		static UE_FORCEINLINE_HINT Name Intersection(const TArray<Name>& Ranges) \
 		{ \
 			return Super::Intersection(reinterpret_cast<const TArray<Super>&>(Ranges)); \
 		} \
 		 \
-		static FORCEINLINE TArray<Name> Union(const Name& X, const Name& Y) \
+		static UE_FORCEINLINE_HINT TArray<Name> Union(const Name& X, const Name& Y) \
 		{ \
 			return TArray<Name>(Super::Union(X, Y)); \
 		} \
 		 \
-		static FORCEINLINE Name All() \
+		static UE_FORCEINLINE_HINT Name All() \
 		{ \
 			return Super::All(); \
 		} \
 		 \
-		static FORCEINLINE Name AtLeast(ElementValueOrConstRef Value) \
+		static UE_FORCEINLINE_HINT Name AtLeast(ElementValueOrConstRef Value) \
 		{ \
 			return Super::AtLeast(Value); \
 		} \
 		 \
-		static FORCEINLINE Name AtMost(ElementValueOrConstRef Value) \
+		static UE_FORCEINLINE_HINT Name AtMost(ElementValueOrConstRef Value) \
 		{ \
 			return Super::AtMost(Value); \
 		} \
 		 \
-		static FORCEINLINE TRange GreaterThan(ElementValueOrConstRef Value) \
+		static UE_FORCEINLINE_HINT TRange GreaterThan(ElementValueOrConstRef Value) \
 		{ \
 			return Super::GreaterThan(Value); \
 		} \
 		 \
-		static FORCEINLINE TRange LessThan(ElementValueOrConstRef Value) \
+		static UE_FORCEINLINE_HINT TRange LessThan(ElementValueOrConstRef Value) \
 		{ \
 			return Super::LessThan(Value); \
 		} \
diff --git a/Engine/Source/Runtime/Core/Public/Math/RangeBound.h b/Engine/Source/Runtime/Core/Public/Math/RangeBound.h
--- a/Engine/Source/Runtime/Core/Public/Math/RangeBound.h
+++ b/Engine/Source/Runtime/Core/Public/Math/RangeBound.h
@@ -122,7 +122,7 @@ public:
 	 *
 	 * @return true if the bound is closed, false otherwise.
 	 */
-	[[nodiscard]] FORCEINLINE_DEBUGGABLE bool IsClosed() const
+	[[nodiscard]] UE_FORCEINLINE_HINT bool IsClosed() const
 	{
 		return (Type != ERangeBoundTypes::Open);
 	}
@@ -132,7 +132,7 @@ public:
 	 *
 	 * @return true if the bound is exclusive, false otherwise.
 	 */
-	[[nodiscard]] FORCEINLINE_DEBUGGABLE bool IsExclusive() const
+	[[nodiscard]] UE_FORCEINLINE_HINT bool IsExclusive() const
 	{
 		return (Type == ERangeBoundTypes::Exclusive);
 	}
@@ -142,7 +142,7 @@ public:
 	 *
 	 * @return true if the bound is inclusive, false otherwise.
 	 */
-	[[nodiscard]] FORCEINLINE_DEBUGGABLE bool IsInclusive() const
+	[[nodiscard]] UE_FORCEINLINE_HINT bool IsInclusive() const
 	{
 		return (Type == ERangeBoundTypes::Inclusive);
 	}
@@ -152,7 +152,7 @@ public:
 	 *
 	 * @return true if the bound is open, false otherwise.
 	 */
-	[[nodiscard]] FORCEINLINE_DEBUGGABLE bool IsOpen() const
+	[[nodiscard]] UE_FORCEINLINE_HINT bool IsOpen() const
 	{
 		return (Type == ERangeBoundTypes::Open);
 	}
@@ -358,42 +358,42 @@ private:
 			: Super(InValue) \
 		{ } \
 		 \
-		static FORCEINLINE Name Exclusive(ElementValueOrConstRef Value) \
+		static UE_FORCEINLINE_HINT Name Exclusive(ElementValueOrConstRef Value) \
 		{ \
 			return static_cast<const Name&>(Super::Exclusive(Value)); \
 		} \
 		 \
-		static FORCEINLINE Name Inclusive(ElementValueOrConstRef Value) \
+		static UE_FORCEINLINE_HINT Name Inclusive(ElementValueOrConstRef Value) \
 		{ \
 			return static_cast<const Name&>(Super::Inclusive(Value)); \
 		} \
 		 \
-		static FORCEINLINE Name Open() \
+		static UE_FORCEINLINE_HINT Name Open() \
 		{ \
 			return static_cast<const Name&>(Super::Open()); \
 		} \
 		 \
-		static FORCEINLINE Name FlipInclusion(const Name& Bound) \
+		static UE_FORCEINLINE_HINT Name FlipInclusion(const Name& Bound) \
 		{ \
 			return static_cast<const Name&>(Super::FlipInclusion(Bound)); \
 		} \
 		 \
-		static FORCEINLINE const Name& MaxLower(const Name& A, const Name& B) \
+		static UE_FORCEINLINE_HINT const Name& MaxLower(const Name& A, const Name& B) \
 		{ \
 			return static_cast<const Name&>(Super::MaxLower(A, B)); \
 		} \
 		 \
-		static FORCEINLINE const Name& MaxUpper(const Name& A, const Name& B) \
+		static UE_FORCEINLINE_HINT const Name& MaxUpper(const Name& A, const Name& B) \
 		{ \
 			return static_cast<const Name&>(Super::MaxUpper(A, B)); \
 		} \
 		 \
-		static FORCEINLINE const Name& MinLower(const Name& A, const Name& B) \
+		static UE_FORCEINLINE_HINT const Name& MinLower(const Name& A, const Name& B) \
 		{ \
 			return static_cast<const Name&>(Super::MinLower(A, B)); \
 		} \
 		 \
-		static FORCEINLINE const Name& MinUpper(const Name& A, const Name& B) \
+		static UE_FORCEINLINE_HINT const Name& MinUpper(const Name& A, const Name& B) \
 		{ \
 			return static_cast<const Name&>(Super::MinUpper(A, B)); \
 		} \
diff --git a/Engine/Source/Runtime/Core/Public/Math/Rotator.h b/Engine/Source/Runtime/Core/Public/Math/Rotator.h
--- a/Engine/Source/Runtime/Core/Public/Math/Rotator.h
+++ b/Engine/Source/Runtime/Core/Public/Math/Rotator.h
@@ -77,8 +77,8 @@ public:
 		}
 	}
 #else
-	FORCEINLINE void DiagnosticCheckNaN() const {}
-	FORCEINLINE void DiagnosticCheckNaN(const TCHAR* Message) const {}
+	UE_FORCEINLINE_HINT void DiagnosticCheckNaN() const {}
+	UE_FORCEINLINE_HINT void DiagnosticCheckNaN(const TCHAR* Message) const {}
 #endif
 
 	/**
@@ -91,7 +91,7 @@ public:
 	 *
 	 * @param InF Value to set all components to.
 	 */
-	[[nodiscard]] explicit FORCEINLINE TRotator(T InF);
+	[[nodiscard]] explicit UE_FORCEINLINE_HINT TRotator(T InF);
 
 	/**
 	 * Constructor.
@@ -100,14 +100,14 @@ public:
 	 * @param InYaw Yaw in degrees.
 	 * @param InRoll Roll in degrees.
 	 */
-	[[nodiscard]] FORCEINLINE TRotator( T InPitch, T InYaw, T InRoll );
+	[[nodiscard]] UE_FORCEINLINE_HINT TRotator( T InPitch, T InYaw, T InRoll );
 
 	/**
 	 * Constructor.
 	 *
 	 * @param EForceInit Force Init Enum.
 	 */
-	[[nodiscard]] explicit FORCEINLINE TRotator( EForceInit );
+	[[nodiscard]] explicit UE_FORCEINLINE_HINT TRotator( EForceInit );
 
 	/**
 	 * Constructor.
@@ -143,7 +143,7 @@ public:
 	 * @return The result of scaling.
 	 */
 	template<typename FArg UE_REQUIRES(std::is_arithmetic_v<FArg>)>
-	[[nodiscard]] FORCEINLINE TRotator operator*( FArg Scale ) const
+	[[nodiscard]] UE_FORCEINLINE_HINT TRotator operator*( FArg Scale ) const
 	{
 		return TRotator(Pitch * Scale, Yaw * Scale, Roll * Scale);
 	}
@@ -552,7 +552,7 @@ inline FArchive& operator<<(FArchive& Ar, TRotator<double>& R)
  * @return Scaled rotator.
  */
 template<typename T, typename FArg UE_REQUIRES(std::is_arithmetic_v<FArg>)>
-FORCEINLINE TRotator<T> operator*(FArg Scale, const TRotator<T>& R )
+UE_FORCEINLINE_HINT TRotator<T> operator*(FArg Scale, const TRotator<T>& R )
 {
 	return R.operator*( Scale );
 }
@@ -572,30 +572,30 @@ inline TRotator<T>::TRotator( T InPitch, T InYaw, T InRoll )
 }
 
 template<typename T>
-FORCEINLINE TRotator<T>::TRotator(EForceInit)
+UE_FORCEINLINE_HINT TRotator<T>::TRotator(EForceInit)
 	: Pitch(0), Yaw(0), Roll(0)
 {}
 
 template<typename T>
-FORCEINLINE TRotator<T> TRotator<T>::operator+( const TRotator<T>& R ) const
+UE_FORCEINLINE_HINT TRotator<T> TRotator<T>::operator+( const TRotator<T>& R ) const
 {
 	return TRotator( Pitch+R.Pitch, Yaw+R.Yaw, Roll+R.Roll );
 }
 
 template<typename T>
-FORCEINLINE TRotator<T> TRotator<T>::operator-( const TRotator<T>& R ) const
+UE_FORCEINLINE_HINT TRotator<T> TRotator<T>::operator-( const TRotator<T>& R ) const
 {
 	return TRotator( Pitch-R.Pitch, Yaw-R.Yaw, Roll-R.Roll );
 }
 
 template<typename T>
-FORCEINLINE bool TRotator<T>::operator==( const TRotator<T>& R ) const
+UE_FORCEINLINE_HINT bool TRotator<T>::operator==( const TRotator<T>& R ) const
 {
 	return Pitch==R.Pitch && Yaw==R.Yaw && Roll==R.Roll;
 }
 
 template<typename T>
-FORCEINLINE bool TRotator<T>::operator!=( const TRotator<T>& V ) const
+UE_FORCEINLINE_HINT bool TRotator<T>::operator!=( const TRotator<T>& V ) const
 {
 	return Pitch!=V.Pitch || Yaw!=V.Yaw || Roll!=V.Roll;
 }
@@ -633,7 +633,7 @@ inline bool TRotator<T>::IsNearlyZero(T Tolerance) const
 }
 
 template<typename T>
-FORCEINLINE bool TRotator<T>::IsZero() const
+UE_FORCEINLINE_HINT bool TRotator<T>::IsZero() const
 {
 	return (ClampAxis(Pitch)==0.f) && (ClampAxis(Yaw)==0.f) && (ClampAxis(Roll)==0.f);
 }
@@ -682,7 +682,7 @@ inline TRotator<T> TRotator<T>::GridSnap( const TRotator<T>& RotGrid ) const
 }
 
 template<typename T>
-FORCEINLINE TRotator<T> TRotator<T>::Clamp() const
+UE_FORCEINLINE_HINT TRotator<T> TRotator<T>::Clamp() const
 {
 	return TRotator(ClampAxis(Pitch), ClampAxis(Yaw), ClampAxis(Roll));
 }
@@ -718,28 +718,28 @@ inline T TRotator<T>::NormalizeAxis( T Angle )
 }
 
 template<typename T>
-FORCEINLINE uint8 TRotator<T>::CompressAxisToByte( T Angle )
+UE_FORCEINLINE_HINT uint8 TRotator<T>::CompressAxisToByte( T Angle )
 {
 	// map [0->360) to [0->256) and mask off any winding
 	return FMath::RoundToInt(Angle * (T)256.f / (T)360.f) & 0xFF;
 }
 
 template<typename T>
-FORCEINLINE T TRotator<T>::DecompressAxisFromByte( uint8 Angle )
+UE_FORCEINLINE_HINT T TRotator<T>::DecompressAxisFromByte( uint8 Angle )
 {
 	// map [0->256) to [0->360)
 	return (Angle * (T)360.f / (T)256.f);
 }
 
 template<typename T>
-FORCEINLINE uint16 TRotator<T>::CompressAxisToShort( T Angle )
+UE_FORCEINLINE_HINT uint16 TRotator<T>::CompressAxisToShort( T Angle )
 {
 	// map [0->360) to [0->65536) and mask off any winding
 	return FMath::RoundToInt(Angle * (T)65536.f / (T)360.f) & 0xFFFF;
 }
 
 template<typename T>
-FORCEINLINE T TRotator<T>::DecompressAxisFromShort( uint16 Angle )
+UE_FORCEINLINE_HINT T TRotator<T>::DecompressAxisFromShort( uint16 Angle )
 {
 	// map [0->65536) to [0->360)
 	return (Angle * (T)360.f / (T)65536.f);
@@ -816,7 +816,7 @@ inline void TRotator<T>::SetComponentForAxis(EAxis::Type Axis, T Component)
 }
 
 template<typename T>
-FORCEINLINE FString TRotator<T>::ToString() const
+UE_FORCEINLINE_HINT FString TRotator<T>::ToString() const
 {
 	return FString::Printf(TEXT("P=%f Y=%f R=%f"), Pitch, Yaw, Roll );
 }
@@ -878,13 +878,13 @@ inline bool TRotator<T>::ContainsNaN() const
 }
 
 template<typename T>
-FORCEINLINE T TRotator<T>::GetManhattanDistance(const TRotator<T> & Rotator) const
+UE_FORCEINLINE_HINT T TRotator<T>::GetManhattanDistance(const TRotator<T> & Rotator) const
 {
 	return FMath::Abs<T>(Yaw - Rotator.Yaw) + FMath::Abs<T>(Pitch - Rotator.Pitch) + FMath::Abs<T>(Roll - Rotator.Roll);
 }
 
 template<typename T>
-FORCEINLINE TRotator<T> TRotator<T>::GetEquivalentRotator() const
+UE_FORCEINLINE_HINT TRotator<T> TRotator<T>::GetEquivalentRotator() const
 {
 	return TRotator(180.0f - Pitch,Yaw + 180.0f, Roll + 180.0f);
 }
@@ -930,14 +930,14 @@ struct TCustomLerp<UE::Math::TRotator<T>>
 	using RotatorType = UE::Math::TRotator<T>;
 
 	template<class U>
-	static FORCEINLINE_DEBUGGABLE RotatorType Lerp(const RotatorType& A, const RotatorType& B, const U& Alpha)
+	static UE_FORCEINLINE_HINT RotatorType Lerp(const RotatorType& A, const RotatorType& B, const U& Alpha)
 	{
 		return A + (B - A).GetNormalized() * Alpha;
 	}
 };
 
 template< typename T, typename U >
-FORCEINLINE_DEBUGGABLE UE::Math::TRotator<T> FMath::LerpRange(const UE::Math::TRotator<T>& A, const UE::Math::TRotator<T>& B, U Alpha)
+UE_FORCEINLINE_HINT UE::Math::TRotator<T> FMath::LerpRange(const UE::Math::TRotator<T>& A, const UE::Math::TRotator<T>& B, U Alpha)
 {
 	// Similar to Lerp, but does not take the shortest path. Allows interpolation over more than 180 degrees.
 	return (A * ((T)1.0 - (T)Alpha) + B * Alpha).GetNormalized();
diff --git a/Engine/Source/Runtime/Core/Public/Math/SHMath.h b/Engine/Source/Runtime/Core/Public/Math/SHMath.h
--- a/Engine/Source/Runtime/Core/Public/Math/SHMath.h
+++ b/Engine/Source/Runtime/Core/Public/Math/SHMath.h
@@ -24,7 +24,7 @@ extern CORE_API float LegendrePolynomial(int32 L, int32 M, float X);
 extern CORE_API FLinearColor SHGetLuminanceFactors();
 
 /** Returns the basis index of the SH basis L,M. */
-FORCEINLINE int32 SHGetBasisIndex(int32 L,int32 M)
+UE_FORCEINLINE_HINT int32 SHGetBasisIndex(int32 L,int32 M)
 {
 	return L * (L + 1) + M;
 }
diff --git a/Engine/Source/Runtime/Core/Public/Math/ScalarRegister.h b/Engine/Source/Runtime/Core/Public/Math/ScalarRegister.h
--- a/Engine/Source/Runtime/Core/Public/Math/ScalarRegister.h
+++ b/Engine/Source/Runtime/Core/Public/Math/ScalarRegister.h
@@ -32,19 +32,19 @@ public:
 	RegisterType Value;
 
 	/** default constructor */
-	[[nodiscard]] FORCEINLINE ScalarRegister();
+	[[nodiscard]] UE_FORCEINLINE_HINT ScalarRegister();
 
 	/** Copy Constructor */
-	[[nodiscard]] FORCEINLINE ScalarRegister(const ScalarRegister& VectorValue);
+	[[nodiscard]] UE_FORCEINLINE_HINT ScalarRegister(const ScalarRegister& VectorValue);
 
 	/** Constructor using float value */
-	[[nodiscard]] explicit FORCEINLINE ScalarRegister(const float& ScalarValue);
+	[[nodiscard]] explicit UE_FORCEINLINE_HINT ScalarRegister(const float& ScalarValue);
 
 	/** Constructor
 	 *
 	 * @param RegisterType float4 vector register type
 	 */
-	[[nodiscard]] explicit FORCEINLINE ScalarRegister(RegisterType VectorValue);
+	[[nodiscard]] explicit UE_FORCEINLINE_HINT ScalarRegister(RegisterType VectorValue);
 
 	/**
 	 * Gets the result of multiplying a scalar register to this.
@@ -53,7 +53,7 @@ public:
 	 *
 	 * @return The result of multiplication.
 	 */
-	[[nodiscard]] FORCEINLINE ScalarRegister operator*(const ScalarRegister& RHS) const;
+	[[nodiscard]] UE_FORCEINLINE_HINT ScalarRegister operator*(const ScalarRegister& RHS) const;
 
 	/**
 	 * Gets the result of adding a scalar register to this.
@@ -62,7 +62,7 @@ public:
 	 *
 	 * @return The result of addition.
 	 */
-	[[nodiscard]] FORCEINLINE ScalarRegister operator+(const ScalarRegister& RHS) const;
+	[[nodiscard]] UE_FORCEINLINE_HINT ScalarRegister operator+(const ScalarRegister& RHS) const;
 
 	/**
 	 * Adds to this scalar register.
@@ -71,7 +71,7 @@ public:
 	 *
 	 * @return Reference to this after addition.
 	 */
-	FORCEINLINE ScalarRegister& operator+=(const ScalarRegister& RHS);
+	inline ScalarRegister& operator+=(const ScalarRegister& RHS);
 
 	/**
 	 * Subtracts another scalar register from this.
@@ -80,7 +80,7 @@ public:
 	 *
 	 * @return reference to this after subtraction.
 	 */
-	FORCEINLINE ScalarRegister& operator-=(const ScalarRegister& RHS);
+	inline ScalarRegister& operator-=(const ScalarRegister& RHS);
 
 	/**
 	 * Gets the result of subtracting a scalar register to this.
@@ -89,52 +89,52 @@ public:
 	 *
 	 * @return The result of subtraction.
 	 */
-	[[nodiscard]] FORCEINLINE ScalarRegister operator-(const ScalarRegister& RHS) const;
+	[[nodiscard]] UE_FORCEINLINE_HINT ScalarRegister operator-(const ScalarRegister& RHS) const;
 
 	
 	/** assignment operator
 	 *
 	 * @param RHS a ScalarRegister
 	 */
-	FORCEINLINE ScalarRegister& operator=(const ScalarRegister& RHS);
+	inline ScalarRegister& operator=(const ScalarRegister& RHS);
 
 	/** assignment operator
 	 *
 	 * @param RHS a RegisterType
 	 */
-	FORCEINLINE ScalarRegister& operator=(const RegisterType& RHS);
+	inline ScalarRegister& operator=(const RegisterType& RHS);
 
 	/**
 	 * ScalarRegister to RegisterType conversion operator.
 	 */
-	FORCEINLINE operator RegisterType() const;
+	UE_FORCEINLINE_HINT operator RegisterType() const;
 };
 
-FORCEINLINE ScalarRegister::ScalarRegister()
+UE_FORCEINLINE_HINT ScalarRegister::ScalarRegister()
 {
 }
 
-FORCEINLINE ScalarRegister::ScalarRegister(const ScalarRegister& VectorValue)
+UE_FORCEINLINE_HINT ScalarRegister::ScalarRegister(const ScalarRegister& VectorValue)
 {
 	Value = VectorValue.Value;
 }
 
-FORCEINLINE ScalarRegister::ScalarRegister(const float& ScalarValue)
+UE_FORCEINLINE_HINT ScalarRegister::ScalarRegister(const float& ScalarValue)
 {
 	Value = VectorLoadFloat1(&ScalarValue);
 }
 
-FORCEINLINE ScalarRegister::ScalarRegister(RegisterType VectorValue)
+UE_FORCEINLINE_HINT ScalarRegister::ScalarRegister(RegisterType VectorValue)
 {
 	Value = VectorValue;
 }
 
-FORCEINLINE ScalarRegister ScalarRegister::operator*(const ScalarRegister& RHS) const
+UE_FORCEINLINE_HINT ScalarRegister ScalarRegister::operator*(const ScalarRegister& RHS) const
 {
 	return ScalarRegister(VectorMultiply(Value, RHS.Value));
 }
 
-FORCEINLINE ScalarRegister ScalarRegister::operator+(const ScalarRegister& RHS) const
+UE_FORCEINLINE_HINT ScalarRegister ScalarRegister::operator+(const ScalarRegister& RHS) const
 {
 	return ScalarRegister(VectorAdd(Value, RHS.Value));
 }
@@ -151,7 +151,7 @@ inline ScalarRegister& ScalarRegister::operator-=(const ScalarRegister& RHS)
 	return *this;
 }
 
-FORCEINLINE ScalarRegister ScalarRegister::operator-(const ScalarRegister& RHS) const
+UE_FORCEINLINE_HINT ScalarRegister ScalarRegister::operator-(const ScalarRegister& RHS) const
 {
 	return ScalarRegister(VectorSubtract(Value, RHS.Value));
 }
@@ -169,7 +169,7 @@ inline ScalarRegister& ScalarRegister::operator=(const RegisterType& RHS)
 	return *this;
 }
 
-FORCEINLINE ScalarRegister::operator RegisterType() const
+UE_FORCEINLINE_HINT ScalarRegister::operator RegisterType() const
 {
 	return Value;
 }
@@ -182,13 +182,13 @@ FORCEINLINE ScalarRegister::operator RegisterType() const
 ----------------------------------------------------------------------------*/
 
 /** Returns the smaller of the two values */
-FORCEINLINE ScalarRegister ScalarMin(const ScalarRegister& A, const ScalarRegister& B)
+UE_FORCEINLINE_HINT ScalarRegister ScalarMin(const ScalarRegister& A, const ScalarRegister& B)
 {
 	return ScalarRegister(VectorMin(A.Value, B.Value));
 }
 
 /** Returns the larger of the two values */
-FORCEINLINE ScalarRegister ScalarMax(const ScalarRegister& A, const ScalarRegister& B)
+UE_FORCEINLINE_HINT ScalarRegister ScalarMax(const ScalarRegister& A, const ScalarRegister& B)
 {
 	return ScalarRegister(VectorMax(A.Value, B.Value));
 }
@@ -212,7 +212,7 @@ template<> struct TCustomLerp<ScalarRegister>
  * @param A	1st scalar
  * @return		ScalarRegister( 1.0f / A.x, 1.0f / A.y, 1.0f / A.z, 1.0f / A.w )
  */
-FORCEINLINE ScalarRegister ScalarReciprocal(const ScalarRegister& A)
+UE_FORCEINLINE_HINT ScalarRegister ScalarReciprocal(const ScalarRegister& A)
 {
 	return ScalarRegister(VectorReciprocalAccurate(A.Value));
 }
diff --git a/Engine/Source/Runtime/Core/Public/Math/Sphere.h b/Engine/Source/Runtime/Core/Public/Math/Sphere.h
--- a/Engine/Source/Runtime/Core/Public/Math/Sphere.h
+++ b/Engine/Source/Runtime/Core/Public/Math/Sphere.h
@@ -183,7 +183,7 @@ public:
 	 * @param  Tolerance Error tolerance.
 	 * @return true if spheres intersect, false otherwise.
 	 */
-	[[nodiscard]] FORCEINLINE bool Intersects(const TSphere<T>& Other, T Tolerance = UE_KINDA_SMALL_NUMBER) const
+	[[nodiscard]] UE_FORCEINLINE_HINT bool Intersects(const TSphere<T>& Other, T Tolerance = UE_KINDA_SMALL_NUMBER) const
 	{
 		return (Center - Other.Center).SizeSquared() <= FMath::Square(FMath::Max(0.f, Other.W + W + Tolerance));
 	}
diff --git a/Engine/Source/Runtime/Core/Public/Math/TransformCalculus2D.h b/Engine/Source/Runtime/Core/Public/Math/TransformCalculus2D.h
--- a/Engine/Source/Runtime/Core/Public/Math/TransformCalculus2D.h
+++ b/Engine/Source/Runtime/Core/Public/Math/TransformCalculus2D.h
@@ -616,7 +616,7 @@ public:
 		}
 	}
 #else
-	FORCEINLINE void DiagnosticCheckNaN() const {}
+	UE_FORCEINLINE_HINT void DiagnosticCheckNaN() const {}
 #endif
 
 private:
