Replace some FORCEINLINE usage in Core/CoreUObject with UE_FORCEINLINE_HINT.

diff --git a/Engine/Source/Runtime/Core/Public/Math/TransformNonVectorized.h b/Engine/Source/Runtime/Core/Public/Math/TransformNonVectorized.h
--- a/Engine/Source/Runtime/Core/Public/Math/TransformNonVectorized.h
+++ b/Engine/Source/Runtime/Core/Public/Math/TransformNonVectorized.h
@@ -102,11 +102,11 @@ public:
 
 	}
 #else
-	FORCEINLINE void DiagnosticCheckNaN_Translate() const {}
-	FORCEINLINE void DiagnosticCheckNaN_Rotate() const {}
-	FORCEINLINE void DiagnosticCheckNaN_Scale3D() const {}
-	FORCEINLINE void DiagnosticCheckNaN_All() const {}
-	FORCEINLINE void DiagnosticCheck_IsValid() const {}
+	UE_FORCEINLINE_HINT void DiagnosticCheckNaN_Translate() const {}
+	UE_FORCEINLINE_HINT void DiagnosticCheckNaN_Rotate() const {}
+	UE_FORCEINLINE_HINT void DiagnosticCheckNaN_Scale3D() const {}
+	UE_FORCEINLINE_HINT void DiagnosticCheckNaN_All() const {}
+	UE_FORCEINLINE_HINT void DiagnosticCheck_IsValid() const {}
 #endif
 
 	/** Default constructor. */
@@ -295,7 +295,7 @@ public:
 	/**
 	* Convert this Transform to matrix with scaling and compute the inverse of that.
 	*/
-	[[nodiscard]] FORCEINLINE TMatrix<T> ToInverseMatrixWithScale() const
+	[[nodiscard]] UE_FORCEINLINE_HINT TMatrix<T> ToInverseMatrixWithScale() const
 	{
 		// todo: optimize
 		return ToMatrixWithScale().Inverse();
@@ -435,7 +435,7 @@ public:
 	* Use only within blends!!
 	* Rotation part is NOT normalized!!
 	*/
-	[[nodiscard]] FORCEINLINE TTransform<T> operator+(const TTransform<T>& Atom) const
+	[[nodiscard]] UE_FORCEINLINE_HINT TTransform<T> operator+(const TTransform<T>& Atom) const
 	{
 		return TTransform(Rotation + Atom.Rotation, Translation + Atom.Translation, Scale3D + Atom.Scale3D);
 	}
@@ -454,7 +454,7 @@ public:
 		return *this;
 	}
 
-	[[nodiscard]] FORCEINLINE TTransform<T> operator*(T Mult) const
+	[[nodiscard]] UE_FORCEINLINE_HINT TTransform<T> operator*(T Mult) const
 	{
 		return TTransform(Rotation * Mult, Translation * Mult, Scale3D * Mult);
 	}
@@ -478,7 +478,7 @@ public:
 	* @param  Other other transform by which to multiply.
 	* @return new transform: this * Other
 	*/
-	[[nodiscard]] FORCEINLINE TTransform<T> operator*(const TTransform<T>& Other) const;
+	[[nodiscard]] inline TTransform<T> operator*(const TTransform<T>& Other) const;
 
 
 	/**
@@ -487,7 +487,7 @@ public:
 	*
 	* @param  Other other transform by which to multiply.
 	*/
-	FORCEINLINE void operator*=(const TTransform<T>& Other);
+	UE_FORCEINLINE_HINT void operator*=(const TTransform<T>& Other);
 
 	/**
 	* Return a transform that is the result of this multiplied by another transform (made only from a rotation).
@@ -496,7 +496,7 @@ public:
 	* @param  Other other quaternion rotation by which to multiply.
 	* @return new transform: this * TTransform(Other)
 	*/
-	[[nodiscard]] FORCEINLINE TTransform<T> operator*(const TQuat<T>& Other) const;
+	[[nodiscard]] inline TTransform<T> operator*(const TQuat<T>& Other) const;
 
 	/**
 	* Sets this transform to the result of this multiplied by another transform (made only from a rotation).
@@ -504,14 +504,14 @@ public:
 	*
 	* @param  Other other quaternion rotation by which to multiply.
 	*/
-	FORCEINLINE void operator*=(const TQuat<T>& Other);
+	inline void operator*=(const TQuat<T>& Other);
 
-	[[nodiscard]] FORCEINLINE static bool AnyHasNegativeScale(const TVector<T>& InScale3D, const TVector<T>& InOtherScale3D);
-	FORCEINLINE void ScaleTranslation(const TVector<T>& InScale3D);
-	FORCEINLINE void ScaleTranslation(const FReal& Scale);
-	FORCEINLINE void RemoveScaling(FReal Tolerance = UE_SMALL_NUMBER);
-	[[nodiscard]] FORCEINLINE T GetMaximumAxisScale() const;
-	[[nodiscard]] FORCEINLINE T GetMinimumAxisScale() const;
+	[[nodiscard]] inline static bool AnyHasNegativeScale(const TVector<T>& InScale3D, const TVector<T>& InOtherScale3D);
+	inline void ScaleTranslation(const TVector<T>& InScale3D);
+	inline void ScaleTranslation(const FReal& Scale);
+	inline void RemoveScaling(FReal Tolerance = UE_SMALL_NUMBER);
+	[[nodiscard]] inline T GetMaximumAxisScale() const;
+	[[nodiscard]] inline T GetMinimumAxisScale() const;
 
 	// Inverse does not work well with VQS format(in particular non-uniform), so removing it, but made two below functions to be used instead. 
 
@@ -530,56 +530,56 @@ public:
 	*/
 	CORE_API void SetToRelativeTransform(const TTransform<T>& ParentTransform);
 
-	[[nodiscard]] FORCEINLINE TVector4<T> TransformFVector4(const TVector4<T>& V) const;
-	[[nodiscard]] FORCEINLINE TVector4<T> TransformFVector4NoScale(const TVector4<T>& V) const;
-	[[nodiscard]] FORCEINLINE TVector<T> TransformPosition(const TVector<T>& V) const;
-	[[nodiscard]] FORCEINLINE TVector<T> TransformPositionNoScale(const TVector<T>& V) const;
+	[[nodiscard]] inline TVector4<T> TransformFVector4(const TVector4<T>& V) const;
+	[[nodiscard]] inline TVector4<T> TransformFVector4NoScale(const TVector4<T>& V) const;
+	[[nodiscard]] UE_FORCEINLINE_HINT TVector<T> TransformPosition(const TVector<T>& V) const;
+	[[nodiscard]] inline TVector<T> TransformPositionNoScale(const TVector<T>& V) const;
 
 	/** Inverts the transform and then transforms V - correctly handles scaling in this transform. */
-	[[nodiscard]] FORCEINLINE TVector<T> InverseTransformPosition(const TVector<T> &V) const;
-	[[nodiscard]] FORCEINLINE TVector<T> InverseTransformPositionNoScale(const TVector<T> &V) const;
-	[[nodiscard]] FORCEINLINE TVector<T> TransformVector(const TVector<T>& V) const;
-	[[nodiscard]] FORCEINLINE TVector<T> TransformVectorNoScale(const TVector<T>& V) const;
+	[[nodiscard]] inline TVector<T> InverseTransformPosition(const TVector<T> &V) const;
+	[[nodiscard]] inline TVector<T> InverseTransformPositionNoScale(const TVector<T> &V) const;
+	[[nodiscard]] UE_FORCEINLINE_HINT TVector<T> TransformVector(const TVector<T>& V) const;
+	[[nodiscard]] inline TVector<T> TransformVectorNoScale(const TVector<T>& V) const;
 
 	/**
 	*	Transform a direction vector by the inverse of this transform - will not take into account translation part.
 	*	If you want to transform a surface normal (or plane) and correctly account for non-uniform scaling you should use TransformByUsingAdjointT with adjoint of matrix inverse.
 	*/
-	[[nodiscard]] FORCEINLINE TVector<T> InverseTransformVector(const TVector<T> &V) const;
-	[[nodiscard]] FORCEINLINE TVector<T> InverseTransformVectorNoScale(const TVector<T> &V) const;
+	[[nodiscard]] inline TVector<T> InverseTransformVector(const TVector<T> &V) const;
+	[[nodiscard]] inline TVector<T> InverseTransformVectorNoScale(const TVector<T> &V) const;
 
 	/**
 	 * Transform a rotation.
 	 * For example if this is a LocalToWorld transform, TransformRotation(Q) would transform Q from local to world space.
 	 */
-	[[nodiscard]] FORCEINLINE TQuat<T> TransformRotation(const TQuat<T>& Q) const;
+	[[nodiscard]] UE_FORCEINLINE_HINT TQuat<T> TransformRotation(const TQuat<T>& Q) const;
 	
 	/**
 	* Inverse transform a rotation.
 	* For example if this is a LocalToWorld transform, InverseTransformRotation(Q) would transform Q from world to local space.
 	*/
-	[[nodiscard]] FORCEINLINE TQuat<T> InverseTransformRotation(const TQuat<T>& Q) const;
+	[[nodiscard]] UE_FORCEINLINE_HINT TQuat<T> InverseTransformRotation(const TQuat<T>& Q) const;
 
-	[[nodiscard]] FORCEINLINE TTransform<T> GetScaled(T Scale) const;
-	[[nodiscard]] FORCEINLINE TTransform<T> GetScaled(TVector<T> Scale) const;
-	[[nodiscard]] FORCEINLINE TVector<T> GetScaledAxis(EAxis::Type InAxis) const;
-	[[nodiscard]] FORCEINLINE TVector<T> GetUnitAxis(EAxis::Type InAxis) const;
-	[[nodiscard]] FORCEINLINE void Mirror(EAxis::Type MirrorAxis, EAxis::Type FlipAxis);
-	[[nodiscard]] FORCEINLINE static TVector<T> GetSafeScaleReciprocal(const TVector<T>& InScale, FReal Tolerance = UE_SMALL_NUMBER);
+	[[nodiscard]] inline TTransform<T> GetScaled(T Scale) const;
+	[[nodiscard]] inline TTransform<T> GetScaled(TVector<T> Scale) const;
+	[[nodiscard]] inline TVector<T> GetScaledAxis(EAxis::Type InAxis) const;
+	[[nodiscard]] inline TVector<T> GetUnitAxis(EAxis::Type InAxis) const;
+	[[nodiscard]] inline void Mirror(EAxis::Type MirrorAxis, EAxis::Type FlipAxis);
+	[[nodiscard]] inline static TVector<T> GetSafeScaleReciprocal(const TVector<T>& InScale, FReal Tolerance = UE_SMALL_NUMBER);
 
 	// temp function for easy conversion
-	[[nodiscard]] FORCEINLINE TVector<T> GetLocation() const
+	[[nodiscard]] UE_FORCEINLINE_HINT TVector<T> GetLocation() const
 	{
 		return GetTranslation();
 	}
 
-	[[nodiscard]] FORCEINLINE TRotator<T> Rotator() const
+	[[nodiscard]] UE_FORCEINLINE_HINT TRotator<T> Rotator() const
 	{
 		return Rotation.Rotator();
 	}
 
 	/** Calculate the determinant of this transformation */
-	[[nodiscard]] FORCEINLINE FReal GetDeterminant() const
+	[[nodiscard]] UE_FORCEINLINE_HINT FReal GetDeterminant() const
 	{
 		return (float)(Scale3D.X * Scale3D.Y * Scale3D.Z);
 	}
@@ -662,17 +662,17 @@ public:
 
 private:
 
-	FORCEINLINE bool Private_RotationEquals(const TQuat<T>& InRotation, const FReal Tolerance = UE_KINDA_SMALL_NUMBER) const
+	UE_FORCEINLINE_HINT bool Private_RotationEquals(const TQuat<T>& InRotation, const FReal Tolerance = UE_KINDA_SMALL_NUMBER) const
 	{
 		return Rotation.Equals(InRotation, Tolerance);
 	}
 
-	FORCEINLINE bool Private_TranslationEquals(const TVector<T>& InTranslation, const FReal Tolerance = UE_KINDA_SMALL_NUMBER) const
+	UE_FORCEINLINE_HINT bool Private_TranslationEquals(const TVector<T>& InTranslation, const FReal Tolerance = UE_KINDA_SMALL_NUMBER) const
 	{
 		return Translation.Equals(InTranslation, Tolerance);
 	}
 
-	FORCEINLINE bool Private_Scale3DEquals(const TVector<T>& InScale3D, const FReal Tolerance = UE_KINDA_SMALL_NUMBER) const
+	UE_FORCEINLINE_HINT bool Private_Scale3DEquals(const TVector<T>& InScale3D, const FReal Tolerance = UE_KINDA_SMALL_NUMBER) const
 	{
 		return Scale3D.Equals(InScale3D, Tolerance);
 	}
@@ -680,19 +680,19 @@ private:
 public:
 
 	// Test if A's rotation equals B's rotation, within a tolerance. Preferred over "A.GetRotation().Equals(B.GetRotation())" because it is faster on some platforms.
-	[[nodiscard]] FORCEINLINE static bool AreRotationsEqual(const TTransform<T>& A, const TTransform<T>& B, FReal Tolerance = UE_KINDA_SMALL_NUMBER)
+	[[nodiscard]] UE_FORCEINLINE_HINT static bool AreRotationsEqual(const TTransform<T>& A, const TTransform<T>& B, FReal Tolerance = UE_KINDA_SMALL_NUMBER)
 	{
 		return A.Private_RotationEquals(B.Rotation, Tolerance);
 	}
 
 	// Test if A's translation equals B's translation, within a tolerance. Preferred over "A.GetTranslation().Equals(B.GetTranslation())" because it is faster on some platforms.
-	[[nodiscard]] FORCEINLINE static bool AreTranslationsEqual(const TTransform<T>& A, const TTransform<T>& B, FReal Tolerance = UE_KINDA_SMALL_NUMBER)
+	[[nodiscard]] UE_FORCEINLINE_HINT static bool AreTranslationsEqual(const TTransform<T>& A, const TTransform<T>& B, FReal Tolerance = UE_KINDA_SMALL_NUMBER)
 	{
 		return A.Private_TranslationEquals(B.Translation, Tolerance);
 	}
 
 	// Test if A's scale equals B's scale, within a tolerance. Preferred over "A.GetScale3D().Equals(B.GetScale3D())" because it is faster on some platforms.
-	[[nodiscard]] FORCEINLINE static bool AreScale3DsEqual(const TTransform<T>& A, const TTransform<T>& B, FReal Tolerance = UE_KINDA_SMALL_NUMBER)
+	[[nodiscard]] UE_FORCEINLINE_HINT static bool AreScale3DsEqual(const TTransform<T>& A, const TTransform<T>& B, FReal Tolerance = UE_KINDA_SMALL_NUMBER)
 	{
 		return A.Private_Scale3DEquals(B.Scale3D, Tolerance);
 	}
@@ -700,37 +700,37 @@ public:
 
 
 	// Test if this Transform's rotation equals another's rotation, within a tolerance. Preferred over "GetRotation().Equals(Other.GetRotation())" because it is faster on some platforms.
-	[[nodiscard]] FORCEINLINE bool RotationEquals(const TTransform<T>& Other, FReal Tolerance = UE_KINDA_SMALL_NUMBER) const
+	[[nodiscard]] UE_FORCEINLINE_HINT bool RotationEquals(const TTransform<T>& Other, FReal Tolerance = UE_KINDA_SMALL_NUMBER) const
 	{
 		return AreRotationsEqual(*this, Other, Tolerance);
 	}
 
 	// Test if this Transform's translation equals another's translation, within a tolerance. Preferred over "GetTranslation().Equals(Other.GetTranslation())" because it is faster on some platforms.
-	[[nodiscard]] FORCEINLINE bool TranslationEquals(const TTransform<T>& Other, FReal Tolerance = UE_KINDA_SMALL_NUMBER) const
+	[[nodiscard]] UE_FORCEINLINE_HINT bool TranslationEquals(const TTransform<T>& Other, FReal Tolerance = UE_KINDA_SMALL_NUMBER) const
 	{
 		return AreTranslationsEqual(*this, Other, Tolerance);
 	}
 
 	// Test if this Transform's scale equals another's scale, within a tolerance. Preferred over "GetScale3D().Equals(Other.GetScale3D())" because it is faster on some platforms.
-	[[nodiscard]] FORCEINLINE bool Scale3DEquals(const TTransform<T>& Other, FReal Tolerance = UE_KINDA_SMALL_NUMBER) const
+	[[nodiscard]] UE_FORCEINLINE_HINT bool Scale3DEquals(const TTransform<T>& Other, FReal Tolerance = UE_KINDA_SMALL_NUMBER) const
 	{
 		return AreScale3DsEqual(*this, Other, Tolerance);
 	}
 
 	// Test if all components of the transforms are equal, within a tolerance.
-	[[nodiscard]] FORCEINLINE bool Equals(const TTransform<T>& Other, FReal Tolerance = UE_KINDA_SMALL_NUMBER) const
+	[[nodiscard]] UE_FORCEINLINE_HINT bool Equals(const TTransform<T>& Other, FReal Tolerance = UE_KINDA_SMALL_NUMBER) const
 	{
 		return Private_TranslationEquals(Other.Translation, Tolerance) && Private_RotationEquals(Other.Rotation, Tolerance) && Private_Scale3DEquals(Other.Scale3D, Tolerance);
 	}
 
 	// Test if all components of the transform property are equal.
-	[[nodiscard]] FORCEINLINE bool Identical(const TTransform<T>* Other, uint32 PortFlags) const
+	[[nodiscard]] UE_FORCEINLINE_HINT bool Identical(const TTransform<T>* Other, uint32 PortFlags) const
 	{
 		return Equals(*Other, 0.f);
 	}
 
 	// Test if rotation and translation components of the transforms are equal, within a tolerance.
-	[[nodiscard]] FORCEINLINE bool EqualsNoScale(const TTransform<T>& Other, FReal Tolerance = UE_KINDA_SMALL_NUMBER) const
+	[[nodiscard]] UE_FORCEINLINE_HINT bool EqualsNoScale(const TTransform<T>& Other, FReal Tolerance = UE_KINDA_SMALL_NUMBER) const
 	{
 		return Private_TranslationEquals(Other.Translation, Tolerance) && Private_RotationEquals(Other.Rotation, Tolerance);
 	}
@@ -744,7 +744,7 @@ public:
 	* @param  A Transform A.
 	* @param  B Transform B.
 	*/
-	FORCEINLINE static void Multiply(TTransform<T>* OutTransform, const TTransform<T>* A, const TTransform<T>* B);
+	UE_FORCEINLINE_HINT static void Multiply(TTransform<T>* OutTransform, const TTransform<T>* A, const TTransform<T>* B);
 
 	/**
 	* Sets the components
@@ -808,7 +808,7 @@ public:
 	}
 
 	/** Copy translation from another TTransform<T>. */
-	FORCEINLINE void CopyTranslation(const TTransform<T>& Other)
+	UE_FORCEINLINE_HINT void CopyTranslation(const TTransform<T>& Other)
 	{
 		Translation = Other.Translation;
 	}
@@ -837,7 +837,7 @@ public:
 	* Add the translations from two TTransform<T>s and return the result.
 	* @return A.Translation + B.Translation
 	*/
-	[[nodiscard]] FORCEINLINE static TVector<T> AddTranslations(const TTransform<T>& A, const TTransform<T>& B)
+	[[nodiscard]] UE_FORCEINLINE_HINT static TVector<T> AddTranslations(const TTransform<T>& A, const TTransform<T>& B)
 	{
 		return A.Translation + B.Translation;
 	}
@@ -846,7 +846,7 @@ public:
 	* Subtract translations from two TTransform<T>s and return the difference.
 	* @return A.Translation - B.Translation.
 	*/
-	[[nodiscard]] FORCEINLINE static TVector<T> SubtractTranslations(const TTransform<T>& A, const TTransform<T>& B)
+	[[nodiscard]] UE_FORCEINLINE_HINT static TVector<T> SubtractTranslations(const TTransform<T>& A, const TTransform<T>& B)
 	{
 		return A.Translation - B.Translation;
 	}
@@ -862,7 +862,7 @@ public:
 	}
 
 	/** Copy rotation from another TTransform<T>. */
-	FORCEINLINE void CopyRotation(const TTransform<T>& Other)
+	UE_FORCEINLINE_HINT void CopyRotation(const TTransform<T>& Other)
 	{
 		Rotation = Other.Rotation;
 	}
@@ -878,7 +878,7 @@ public:
 	}
 
 	/** Copy scale from another TTransform<T>. */
-	FORCEINLINE void CopyScale3D(const TTransform<T>& Other)
+	UE_FORCEINLINE_HINT void CopyScale3D(const TTransform<T>& Other)
 	{
 		Scale3D = Other.Scale3D;
 	}
@@ -1065,7 +1065,7 @@ public:
 	*
 	* @return true if the rotation component is normalized, and false otherwise.
 	*/
-	[[nodiscard]] FORCEINLINE bool IsRotationNormalized() const
+	[[nodiscard]] UE_FORCEINLINE_HINT bool IsRotationNormalized() const
 	{
 		return Rotation.IsNormalized();
 	}
@@ -1196,7 +1196,7 @@ private:
 	* @param  A Transform A.
 	* @param  B Transform B.
 	*/
-	FORCEINLINE static void MultiplyUsingMatrixWithScale(TTransform<T>* OutTransform, const TTransform<T>* A, const TTransform<T>* B);
+	UE_FORCEINLINE_HINT static void MultiplyUsingMatrixWithScale(TTransform<T>* OutTransform, const TTransform<T>* A, const TTransform<T>* B);
 	/**
 	* Create a new transform from multiplications of given to matrices (AMatrix*BMatrix) using desired scale
 	* This is used by MultiplyUsingMatrixWithScale and GetRelativeTransformUsingMatrixWithScale
@@ -1207,7 +1207,7 @@ private:
 	* @param	DesiredScale - there is no check on if the magnitude is correct here. It assumes that is correct.
 	* @param	OutTransform the constructed transform
 	*/
-	FORCEINLINE static void ConstructTransformFromMatrixWithDesiredScale(const TMatrix<T>& AMatrix, const TMatrix<T>& BMatrix, const TVector<T>& DesiredScale, TTransform<T>& OutTransform);
+	UE_FORCEINLINE_HINT static void ConstructTransformFromMatrixWithDesiredScale(const TMatrix<T>& AMatrix, const TMatrix<T>& BMatrix, const TVector<T>& DesiredScale, TTransform<T>& OutTransform);
 	/**
 	* Create a new transform: OutTransform = Base * Relative(-1) using the matrix while keeping the scale that's given by Base and Relative
 	* Please note that this operation is a lot more expensive than normal GetRelativeTrnasform
@@ -1269,7 +1269,7 @@ inline void TTransform<T>::RemoveScaling(T Tolerance/*=SMALL_NUMBER*/)
 }
 
 template<typename T>
-FORCEINLINE void TTransform<T>::MultiplyUsingMatrixWithScale(TTransform<T>* OutTransform, const TTransform<T>* A, const TTransform<T>* B)
+UE_FORCEINLINE_HINT void TTransform<T>::MultiplyUsingMatrixWithScale(TTransform<T>* OutTransform, const TTransform<T>* A, const TTransform<T>* B)
 {
 	// the goal of using M is to get the correct orientation
 	// but for translation, we still need scale
@@ -1487,13 +1487,13 @@ inline TVector<T> TTransform<T>::InverseTransformVectorNoScale(const TVector<T>
 }
 
 template<typename T>
-FORCEINLINE TQuat<T> TTransform<T>::TransformRotation(const TQuat<T>& Q) const
+UE_FORCEINLINE_HINT TQuat<T> TTransform<T>::TransformRotation(const TQuat<T>& Q) const
 {
 	return GetRotation() * Q;
 }
 
 template<typename T>
-FORCEINLINE TQuat<T> TTransform<T>::InverseTransformRotation(const TQuat<T>& Q) const
+UE_FORCEINLINE_HINT TQuat<T> TTransform<T>::InverseTransformRotation(const TQuat<T>& Q) const
 {
 	return GetRotation().Inverse() * Q;
 }
@@ -1508,7 +1508,7 @@ inline TTransform<T> TTransform<T>::operator*(const TTransform<T>& Other) const
 
 
 template<typename T>
-FORCEINLINE void TTransform<T>::operator*=(const TTransform<T>& Other)
+UE_FORCEINLINE_HINT void TTransform<T>::operator*=(const TTransform<T>& Other)
 {
 	Multiply(this, this, &Other);
 }
diff --git a/Engine/Source/Runtime/Core/Public/Math/TransformVectorized.h b/Engine/Source/Runtime/Core/Public/Math/TransformVectorized.h
--- a/Engine/Source/Runtime/Core/Public/Math/TransformVectorized.h
+++ b/Engine/Source/Runtime/Core/Public/Math/TransformVectorized.h
@@ -128,11 +128,11 @@ public:
 		
 	}
 #else
-	FORCEINLINE_DEBUGGABLE void DiagnosticCheckNaN_Translate() const {}
-	FORCEINLINE_DEBUGGABLE void DiagnosticCheckNaN_Rotate() const {}
-	FORCEINLINE_DEBUGGABLE void DiagnosticCheckNaN_Scale3D() const {}
-	FORCEINLINE_DEBUGGABLE void DiagnosticCheckNaN_All() const {}
-	FORCEINLINE_DEBUGGABLE void DiagnosticCheck_IsValid() const {}
+	UE_FORCEINLINE_HINT void DiagnosticCheckNaN_Translate() const {}
+	UE_FORCEINLINE_HINT void DiagnosticCheckNaN_Rotate() const {}
+	UE_FORCEINLINE_HINT void DiagnosticCheckNaN_Scale3D() const {}
+	UE_FORCEINLINE_HINT void DiagnosticCheckNaN_All() const {}
+	UE_FORCEINLINE_HINT void DiagnosticCheck_IsValid() const {}
 #endif
 
 	/**
@@ -354,7 +354,7 @@ public:
 	/**
 	* Convert this Transform to matrix with scaling and compute the inverse of that.
 	*/
-	[[nodiscard]] FORCEINLINE_DEBUGGABLE TMatrix<T> ToInverseMatrixWithScale() const
+	[[nodiscard]] UE_FORCEINLINE_HINT TMatrix<T> ToInverseMatrixWithScale() const
 	{
 		// todo: optimize
 		return ToMatrixWithScale().Inverse();
@@ -502,7 +502,7 @@ public:
 	 * Use only within blends!!
 	 * Rotation part is NOT normalized!!
 	 */
-	[[nodiscard]] FORCEINLINE_DEBUGGABLE TTransform<T> operator+(const TTransform<T>& Atom) const
+	[[nodiscard]] UE_FORCEINLINE_HINT TTransform<T> operator+(const TTransform<T>& Atom) const
 	{
 		return TTransform<T>( VectorAdd(Rotation, Atom.Rotation), VectorAdd(Translation, Atom.Translation ), VectorAdd( Scale3D, Atom.Scale3D ) );
 	}
@@ -532,17 +532,17 @@ public:
 		return *this;
 	}
 
-	[[nodiscard]] FORCEINLINE_DEBUGGABLE TTransform<T>	operator*(const TTransform<T>& Other) const;
-	FORCEINLINE_DEBUGGABLE void			operator*=(const TTransform<T>& Other);
-	[[nodiscard]] FORCEINLINE_DEBUGGABLE TTransform<T>	operator*(const TQuat<T>& Other) const;
-	FORCEINLINE_DEBUGGABLE void			operator*=(const TQuat<T>& Other);
+	[[nodiscard]] inline TTransform<T>	operator*(const TTransform<T>& Other) const;
+	UE_FORCEINLINE_HINT void			operator*=(const TTransform<T>& Other);
+	[[nodiscard]] inline TTransform<T>	operator*(const TQuat<T>& Other) const;
+	inline void							operator*=(const TQuat<T>& Other);
 
-	[[nodiscard]] FORCEINLINE_DEBUGGABLE static bool AnyHasNegativeScale(const TVector<T>& InScale3D, const TVector<T>& InOtherScale3D);
-	FORCEINLINE_DEBUGGABLE void ScaleTranslation(const TVector<T>& InScale3D);
-	FORCEINLINE_DEBUGGABLE void ScaleTranslation(const T& Scale);
-	FORCEINLINE_DEBUGGABLE void RemoveScaling(T Tolerance=UE_SMALL_NUMBER);
-	[[nodiscard]] FORCEINLINE_DEBUGGABLE float GetMaximumAxisScale() const;
-	[[nodiscard]] FORCEINLINE_DEBUGGABLE float GetMinimumAxisScale() const;
+	[[nodiscard]] inline static bool AnyHasNegativeScale(const TVector<T>& InScale3D, const TVector<T>& InOtherScale3D);
+	inline void ScaleTranslation(const TVector<T>& InScale3D);
+	inline void ScaleTranslation(const T& Scale);
+	inline void RemoveScaling(T Tolerance=UE_SMALL_NUMBER);
+	[[nodiscard]] inline float GetMaximumAxisScale() const;
+	[[nodiscard]] inline float GetMinimumAxisScale() const;
 	// Inverse does not work well with VQS format(in particular non-uniform), so removing it, but made two below functions to be used instead. 
 
 	/*******************************************************************************************
@@ -560,46 +560,46 @@ public:
 	 */
 	CORE_API void		SetToRelativeTransform(const TTransform<T>& ParentTransform);
 
-	[[nodiscard]] FORCEINLINE_DEBUGGABLE TVector4<T>	TransformFVector4(const TVector4<T>& V) const;
-	[[nodiscard]] FORCEINLINE_DEBUGGABLE TVector4<T>	TransformFVector4NoScale(const TVector4<T>& V) const;
-	[[nodiscard]] FORCEINLINE_DEBUGGABLE TVector<T>		TransformPosition(const TVector<T>& V) const;
-	[[nodiscard]] FORCEINLINE_DEBUGGABLE TVector<T>		TransformPositionNoScale(const TVector<T>& V) const;
+	[[nodiscard]] inline TVector4<T>	TransformFVector4(const TVector4<T>& V) const;
+	[[nodiscard]] inline TVector4<T>	TransformFVector4NoScale(const TVector4<T>& V) const;
+	[[nodiscard]] inline TVector<T>		TransformPosition(const TVector<T>& V) const;
+	[[nodiscard]] inline TVector<T>		TransformPositionNoScale(const TVector<T>& V) const;
 
 
 	/** Inverts the transform and then transforms V - correctly handles scaling in this transform. */
-	[[nodiscard]] FORCEINLINE_DEBUGGABLE TVector<T>		InverseTransformPosition(const TVector<T> &V) const;
-	[[nodiscard]] FORCEINLINE_DEBUGGABLE TVector<T>		InverseTransformPositionNoScale(const TVector<T> &V) const;
-	[[nodiscard]] FORCEINLINE_DEBUGGABLE TVector<T>		TransformVector(const TVector<T>& V) const;
-	[[nodiscard]] FORCEINLINE_DEBUGGABLE TVector<T>		TransformVectorNoScale(const TVector<T>& V) const;
+	[[nodiscard]] inline TVector<T>		InverseTransformPosition(const TVector<T> &V) const;
+	[[nodiscard]] inline TVector<T>		InverseTransformPositionNoScale(const TVector<T> &V) const;
+	[[nodiscard]] inline TVector<T>		TransformVector(const TVector<T>& V) const;
+	[[nodiscard]] inline TVector<T>		TransformVectorNoScale(const TVector<T>& V) const;
 
 	/** 
 	 *	Transform a direction vector by the inverse of this matrix - will not take into account translation part.
 	 *	If you want to transform a surface normal (or plane) and correctly account for non-uniform scaling you should use TransformByUsingAdjointT with adjoint of matrix inverse.
 	 */
-	[[nodiscard]] FORCEINLINE_DEBUGGABLE TVector<T> InverseTransformVector(const TVector<T> &V) const;
-	[[nodiscard]] FORCEINLINE_DEBUGGABLE TVector<T> InverseTransformVectorNoScale(const TVector<T> &V) const;
+	[[nodiscard]] inline TVector<T> InverseTransformVector(const TVector<T> &V) const;
+	[[nodiscard]] inline TVector<T> InverseTransformVectorNoScale(const TVector<T> &V) const;
 
 	/**
 	* Transform a rotation.
 	* For example if this is a LocalToWorld transform, TransformRotation(Q) would transform Q from local to world space.
 	*/
-	[[nodiscard]] FORCEINLINE_DEBUGGABLE TQuat<T> TransformRotation(const TQuat<T>& Q) const;
+	[[nodiscard]] UE_FORCEINLINE_HINT TQuat<T> TransformRotation(const TQuat<T>& Q) const;
 
 	/**
 	* Inverse transform a rotation.
 	* For example if this is a LocalToWorld transform, InverseTransformRotation(Q) would transform Q from world to local space.
 	*/
-	[[nodiscard]] FORCEINLINE_DEBUGGABLE TQuat<T> InverseTransformRotation(const TQuat<T>& Q) const;
+	[[nodiscard]] UE_FORCEINLINE_HINT TQuat<T> InverseTransformRotation(const TQuat<T>& Q) const;
 
-	[[nodiscard]] FORCEINLINE_DEBUGGABLE TTransform<T>	GetScaled(T Scale) const;
-	[[nodiscard]] FORCEINLINE_DEBUGGABLE TTransform<T>	GetScaled(TVector<T> Scale) const;
-	[[nodiscard]] FORCEINLINE_DEBUGGABLE TVector<T>		GetScaledAxis(EAxis::Type InAxis) const;
-	[[nodiscard]] FORCEINLINE_DEBUGGABLE TVector<T>		GetUnitAxis(EAxis::Type InAxis) const;
-	FORCEINLINE_DEBUGGABLE void		Mirror(EAxis::Type MirrorAxis, EAxis::Type FlipAxis);
-	[[nodiscard]] FORCEINLINE_DEBUGGABLE static TVector<T>	GetSafeScaleReciprocal(const TVector<T>& InScale, T Tolerance=UE_SMALL_NUMBER);
+	[[nodiscard]] inline TTransform<T>	GetScaled(T Scale) const;
+	[[nodiscard]] inline TTransform<T>	GetScaled(TVector<T> Scale) const;
+	[[nodiscard]] inline TVector<T>		GetScaledAxis(EAxis::Type InAxis) const;
+	[[nodiscard]] inline TVector<T>		GetUnitAxis(EAxis::Type InAxis) const;
+	inline void		Mirror(EAxis::Type MirrorAxis, EAxis::Type FlipAxis);
+	[[nodiscard]] inline static TVector<T>	GetSafeScaleReciprocal(const TVector<T>& InScale, T Tolerance=UE_SMALL_NUMBER);
 
 
-	[[nodiscard]] FORCEINLINE_DEBUGGABLE TVector<T> GetLocation() const
+	[[nodiscard]] UE_FORCEINLINE_HINT TVector<T> GetLocation() const
 	{
 		return GetTranslation();
 	}
@@ -717,7 +717,7 @@ public:
 
 private:
 
-	FORCEINLINE_DEBUGGABLE static bool Private_AnyHasNegativeScale(TransformVectorRegister InScale3D, TransformVectorRegister InOtherScale3D)
+	UE_FORCEINLINE_HINT static bool Private_AnyHasNegativeScale(TransformVectorRegister InScale3D, TransformVectorRegister InOtherScale3D)
 	{
 		return !!VectorAnyLesserThan(VectorMin(InScale3D, InOtherScale3D), (TransformVectorRegister)GlobalVectorConstants::FloatZero);
 	}
@@ -752,61 +752,61 @@ private:
 public:
 
 	// Test if A's rotation equals B's rotation, within a tolerance. Preferred over "A.GetRotation().Equals(B.GetRotation())" because it is faster on some platforms.
-	[[nodiscard]] FORCEINLINE_DEBUGGABLE static bool AreRotationsEqual(const TTransform<T>& A, const TTransform<T>& B, FReal Tolerance=UE_KINDA_SMALL_NUMBER)
+	[[nodiscard]] UE_FORCEINLINE_HINT static bool AreRotationsEqual(const TTransform<T>& A, const TTransform<T>& B, FReal Tolerance=UE_KINDA_SMALL_NUMBER)
 	{
 		return A.Private_RotationEquals(B.Rotation, Tolerance);
 	}
 
 	// Test if A's translation equals B's translation, within a tolerance. Preferred over "A.GetTranslation().Equals(B.GetTranslation())" because it avoids TransformVectorRegister->TVector<T> conversion.
-	[[nodiscard]] FORCEINLINE_DEBUGGABLE static bool AreTranslationsEqual(const TTransform<T>& A, const TTransform<T>& B, FReal Tolerance=UE_KINDA_SMALL_NUMBER)
+	[[nodiscard]] UE_FORCEINLINE_HINT static bool AreTranslationsEqual(const TTransform<T>& A, const TTransform<T>& B, FReal Tolerance=UE_KINDA_SMALL_NUMBER)
 	{
 		return A.Private_TranslationEquals(B.Translation, Tolerance);
 	}
 
 	// Test if A's scale equals B's scale, within a tolerance. Preferred over "A.GetScale3D().Equals(B.GetScale3D())" because it avoids TransformVectorRegister->TVector<T> conversion.
-	[[nodiscard]] FORCEINLINE_DEBUGGABLE static bool AreScale3DsEqual(const TTransform<T>& A, const TTransform<T>& B, FReal Tolerance=UE_KINDA_SMALL_NUMBER)
+	[[nodiscard]] UE_FORCEINLINE_HINT static bool AreScale3DsEqual(const TTransform<T>& A, const TTransform<T>& B, FReal Tolerance=UE_KINDA_SMALL_NUMBER)
 	{
 		return A.Private_Scale3DEquals(B.Scale3D, Tolerance);
 	}
 
 
 	// Test if this Transform's rotation equals another's rotation, within a tolerance. Preferred over "GetRotation().Equals(Other.GetRotation())" because it is faster on some platforms.
-	[[nodiscard]] FORCEINLINE_DEBUGGABLE bool RotationEquals(const TTransform<T>& Other, FReal Tolerance=UE_KINDA_SMALL_NUMBER) const
+	[[nodiscard]] UE_FORCEINLINE_HINT bool RotationEquals(const TTransform<T>& Other, FReal Tolerance=UE_KINDA_SMALL_NUMBER) const
 	{
 		return AreRotationsEqual(*this, Other, Tolerance);
 	}
 
 	// Test if this Transform's translation equals another's translation, within a tolerance. Preferred over "GetTranslation().Equals(Other.GetTranslation())" because it avoids TransformVectorRegister->TVector<T> conversion.
-	[[nodiscard]] FORCEINLINE_DEBUGGABLE bool TranslationEquals(const TTransform<T>& Other, FReal Tolerance=UE_KINDA_SMALL_NUMBER) const
+	[[nodiscard]] UE_FORCEINLINE_HINT bool TranslationEquals(const TTransform<T>& Other, FReal Tolerance=UE_KINDA_SMALL_NUMBER) const
 	{
 		return AreTranslationsEqual(*this, Other, Tolerance);
 	}
 
 	// Test if this Transform's scale equals another's scale, within a tolerance. Preferred over "GetScale3D().Equals(Other.GetScale3D())" because it avoids TransformVectorRegister->TVector<T> conversion.
-	[[nodiscard]] FORCEINLINE_DEBUGGABLE bool Scale3DEquals(const TTransform<T>& Other, FReal Tolerance=UE_KINDA_SMALL_NUMBER) const
+	[[nodiscard]] UE_FORCEINLINE_HINT bool Scale3DEquals(const TTransform<T>& Other, FReal Tolerance=UE_KINDA_SMALL_NUMBER) const
 	{
 		return AreScale3DsEqual(*this, Other, Tolerance);
 	}
 
 	// Test if all components of the transforms are equal, within a tolerance.
-	[[nodiscard]] FORCEINLINE_DEBUGGABLE bool Equals(const TTransform<T>& Other, FReal Tolerance=UE_KINDA_SMALL_NUMBER) const
+	[[nodiscard]] UE_FORCEINLINE_HINT bool Equals(const TTransform<T>& Other, FReal Tolerance=UE_KINDA_SMALL_NUMBER) const
 	{
 		return Private_TranslationEquals(Other.Translation, Tolerance) && Private_RotationEquals(Other.Rotation, Tolerance) && Private_Scale3DEquals(Other.Scale3D, Tolerance);
 	}
 
 	// Test if all components of the transform property are equal.
-	[[nodiscard]] FORCEINLINE_DEBUGGABLE bool Identical(const TTransform<T>* Other, uint32 PortFlags) const
+	[[nodiscard]] UE_FORCEINLINE_HINT bool Identical(const TTransform<T>* Other, uint32 PortFlags) const
 	{
 		return Equals(*Other, 0.f);
 	}
 
 	// Test if rotation and translation components of the transforms are equal, within a tolerance.
-	[[nodiscard]] FORCEINLINE_DEBUGGABLE bool EqualsNoScale(const TTransform<T>& Other, FReal Tolerance=UE_KINDA_SMALL_NUMBER) const
+	[[nodiscard]] UE_FORCEINLINE_HINT bool EqualsNoScale(const TTransform<T>& Other, FReal Tolerance=UE_KINDA_SMALL_NUMBER) const
 	{
 		return Private_TranslationEquals(Other.Translation, Tolerance) && Private_RotationEquals(Other.Rotation, Tolerance);
 	}
 
-	FORCEINLINE_DEBUGGABLE static void Multiply(TTransform<T>* OutTransform, const TTransform<T>* A, const TTransform<T>* B);
+	UE_FORCEINLINE_HINT static void Multiply(TTransform<T>* OutTransform, const TTransform<T>* A, const TTransform<T>* B);
 	/**
 	 * Sets the components
 	 * @param InRotation The new value for the Rotation component
@@ -875,7 +875,7 @@ public:
 	}
 
 	/** Copy translation from another TTransform<T>. */
-	FORCEINLINE_DEBUGGABLE void CopyTranslation(const TTransform<T>& Other)
+	UE_FORCEINLINE_HINT void CopyTranslation(const TTransform<T>& Other)
 	{
 		Translation = Other.Translation;
 	}
@@ -933,7 +933,7 @@ public:
 	}
 
 	/** Copy rotation from another TTransform<T>. */
-	FORCEINLINE_DEBUGGABLE void CopyRotation(const TTransform<T>& Other)
+	UE_FORCEINLINE_HINT void CopyRotation(const TTransform<T>& Other)
 	{
 		Rotation = Other.Rotation;
 	}
@@ -949,7 +949,7 @@ public:
 	}
 
 	/** Copy scale from another TTransform<T>. */
-	FORCEINLINE_DEBUGGABLE void CopyScale3D(const TTransform<T>& Other)
+	UE_FORCEINLINE_HINT void CopyScale3D(const TTransform<T>& Other)
 	{
 		Scale3D = Other.Scale3D;
 	}
@@ -1464,7 +1464,7 @@ private:
 	* @param  A Transform A.
 	* @param  B Transform B.
 	*/
-	FORCEINLINE_DEBUGGABLE static void MultiplyUsingMatrixWithScale(TTransform<T>* OutTransform, const TTransform<T>* A, const TTransform<T>* B);
+	UE_FORCEINLINE_HINT static void MultiplyUsingMatrixWithScale(TTransform<T>* OutTransform, const TTransform<T>* A, const TTransform<T>* B);
 	/**
 	* Create a new transform from multiplications of given to matrices (AMatrix*BMatrix) using desired scale
 	* This is used by MultiplyUsingMatrixWithScale and GetRelativeTransformUsingMatrixWithScale
@@ -1475,7 +1475,7 @@ private:
 	* @param	DesiredScale - there is no check on if the magnitude is correct here. It assumes that is correct. 
 	* @param	OutTransform the constructed transform 
 	*/
-	FORCEINLINE_DEBUGGABLE static void ConstructTransformFromMatrixWithDesiredScale(const TMatrix<T>& AMatrix, const TMatrix<T>& BMatrix, TransformVectorRegister DesiredScale, TTransform<T>& OutTransform);
+	inline static void ConstructTransformFromMatrixWithDesiredScale(const TMatrix<T>& AMatrix, const TMatrix<T>& BMatrix, TransformVectorRegister DesiredScale, TTransform<T>& OutTransform);
 	/**
 	* Create a new transform: OutTransform = Base * Relative(-1) using the matrix while keeping the scale that's given by Base and Relative
 	* Please note that this operation is a lot more expensive than normal GetRelativeTrnasform
@@ -1530,7 +1530,7 @@ inline void TTransform<T>::ScaleTranslation(const TVector<T>& InScale3D)
 
 /** Scale the translation part of the Transform by the supplied value. */
 template<typename T>
-FORCEINLINE_DEBUGGABLE void TTransform<T>::ScaleTranslation(const T& InScale)
+UE_FORCEINLINE_HINT void TTransform<T>::ScaleTranslation(const T& InScale)
 {
 	ScaleTranslation( TVector<T>(InScale) );
 }
@@ -1547,7 +1547,7 @@ inline void TTransform<T>::RemoveScaling(T Tolerance/*=UE_SMALL_NUMBER*/)
 }
 
 template<typename T>
-FORCEINLINE_DEBUGGABLE void TTransform<T>::MultiplyUsingMatrixWithScale(TTransform<T>* OutTransform, const TTransform<T>* A, const TTransform<T>* B)
+UE_FORCEINLINE_HINT void TTransform<T>::MultiplyUsingMatrixWithScale(TTransform<T>* OutTransform, const TTransform<T>* A, const TTransform<T>* B)
 {
 	ConstructTransformFromMatrixWithDesiredScale(A->ToMatrixWithScale(), B->ToMatrixWithScale(), VectorMultiply(A->Scale3D, B->Scale3D), *OutTransform);
 }
@@ -1881,13 +1881,13 @@ inline TVector<T> TTransform<T>::InverseTransformVectorNoScale(const TVector<T>
 }
 
 template<typename T>
-FORCEINLINE_DEBUGGABLE TQuat<T> TTransform<T>::TransformRotation(const TQuat<T>& Q) const
+UE_FORCEINLINE_HINT TQuat<T> TTransform<T>::TransformRotation(const TQuat<T>& Q) const
 {
 	return GetRotation() * Q;
 }
 
 template<typename T>
-FORCEINLINE_DEBUGGABLE TQuat<T> TTransform<T>::InverseTransformRotation(const TQuat<T>& Q) const
+UE_FORCEINLINE_HINT TQuat<T> TTransform<T>::InverseTransformRotation(const TQuat<T>& Q) const
 {
 	return GetRotation().Inverse() * Q;
 }
@@ -1901,7 +1901,7 @@ inline TTransform<T> TTransform<T>::operator*(const TTransform<T>& Other) const
 }
 
 template<typename T>
-FORCEINLINE_DEBUGGABLE void TTransform<T>::operator*=(const TTransform<T>& Other)
+UE_FORCEINLINE_HINT void TTransform<T>::operator*=(const TTransform<T>& Other)
 {
 	Multiply(this, this, &Other);
 }
diff --git a/Engine/Source/Runtime/Core/Public/Math/TwoVectors.h b/Engine/Source/Runtime/Core/Public/Math/TwoVectors.h
--- a/Engine/Source/Runtime/Core/Public/Math/TwoVectors.h
+++ b/Engine/Source/Runtime/Core/Public/Math/TwoVectors.h
@@ -24,7 +24,7 @@ public:
 public:
 
 	/** Default constructor. */
-	[[nodiscard]] FORCEINLINE	FTwoVectors();
+	[[nodiscard]] UE_FORCEINLINE_HINT	FTwoVectors();
 
 	/**
 	 * Creates and initializes a new instance with the specified vectors.
@@ -32,14 +32,14 @@ public:
 	 * @param In1 The first Vector.
 	 * @param In2 The second Vector.
 	 */
-	[[nodiscard]] FORCEINLINE	FTwoVectors( FVector In1, FVector In2 );
+	[[nodiscard]] UE_FORCEINLINE_HINT	FTwoVectors( FVector In1, FVector In2 );
 
 	/**
 	 * Constructor.
 	 *
 	 * @param EForceInit Force Init Enum
 	 */
-	[[nodiscard]] explicit FORCEINLINE FTwoVectors( EForceInit );
+	[[nodiscard]] explicit UE_FORCEINLINE_HINT FTwoVectors( EForceInit );
 
 public:
 
@@ -49,7 +49,7 @@ public:
 	 * @param V The pair to add.
 	 * @return Result of addition.
 	 */
-	[[nodiscard]] FORCEINLINE FTwoVectors operator+( const FTwoVectors& V ) const;
+	[[nodiscard]] UE_FORCEINLINE_HINT FTwoVectors operator+( const FTwoVectors& V ) const;
 
 	/**
 	 * Gets result of subtraction of two pairs of vectors.
@@ -57,7 +57,7 @@ public:
 	 * @param V The pair to subtract.
 	 * @return Result of subtraction.
 	 */
-	[[nodiscard]] FORCEINLINE FTwoVectors operator-( const FTwoVectors& V ) const;
+	[[nodiscard]] UE_FORCEINLINE_HINT FTwoVectors operator-( const FTwoVectors& V ) const;
 
 	/**
 	 * Gets result of scaling pair of vectors.
@@ -65,7 +65,7 @@ public:
 	 * @param Scale The scaling factor.
 	 * @return Result of Scaling.
 	 */
-	[[nodiscard]] FORCEINLINE FTwoVectors operator*( float Scale ) const;
+	[[nodiscard]] UE_FORCEINLINE_HINT FTwoVectors operator*( float Scale ) const;
 
 	/**
 	 * Gets result of dividing pair of vectors.
@@ -81,7 +81,7 @@ public:
 	 * @param V The pair to multiply with.
 	 * @return Result of multiplication.
 	 */
-	[[nodiscard]] FORCEINLINE FTwoVectors operator*( const FTwoVectors& V ) const;
+	[[nodiscard]] UE_FORCEINLINE_HINT FTwoVectors operator*( const FTwoVectors& V ) const;
 
 	/**
 	 * Gets result of division of two pairs of vectors.
@@ -89,7 +89,7 @@ public:
 	 * @param V The pair to divide by.
 	 * @return Result of division.
 	 */
-	[[nodiscard]] FORCEINLINE FTwoVectors operator/( const FTwoVectors& V ) const;
+	[[nodiscard]] UE_FORCEINLINE_HINT FTwoVectors operator/( const FTwoVectors& V ) const;
 
 	// Binary comparison operators.
 
@@ -125,7 +125,7 @@ public:
 	 *
 	 * @return A negated copy of the pair.
 	 */
-	[[nodiscard]] FORCEINLINE FTwoVectors operator-() const;
+	[[nodiscard]] UE_FORCEINLINE_HINT FTwoVectors operator-() const;
 
 	// Assignment operators.
 
@@ -135,7 +135,7 @@ public:
 	 * @param The pair to add.
 	 * @return Copy of the pair after addition.
 	 */
-	FORCEINLINE FTwoVectors operator+=( const FTwoVectors& V );
+	UE_FORCEINLINE_HINT FTwoVectors operator+=( const FTwoVectors& V );
 
 	/**
 	 * Subtract a pair from this.
@@ -143,7 +143,7 @@ public:
 	 * @param The pair to subtract.
 	 * @return Copy of the pair after subtraction.
 	 */
-	FORCEINLINE FTwoVectors operator-=( const FTwoVectors& V );
+	UE_FORCEINLINE_HINT FTwoVectors operator-=( const FTwoVectors& V );
 
 	/**
 	 * Scale the pair.
@@ -151,7 +151,7 @@ public:
 	 * @param Scale What to scale by.
 	 * @return Copy of the pair after scaling.
 	 */
-	FORCEINLINE FTwoVectors operator*=( float Scale );
+	UE_FORCEINLINE_HINT FTwoVectors operator*=( float Scale );
 
 	/**
 	 * Divide the pair.
@@ -234,7 +234,7 @@ public:
 /* FTwoVectors inline functions
  *****************************************************************************/
 
-FORCEINLINE FTwoVectors operator*( float Scale, const FTwoVectors& V )
+UE_FORCEINLINE_HINT FTwoVectors operator*( float Scale, const FTwoVectors& V )
 {
 	return V.operator*( Scale );
 }
@@ -314,19 +314,19 @@ inline FTwoVectors FTwoVectors::operator/( const FTwoVectors& V ) const
 }
 
 
-FORCEINLINE bool FTwoVectors::operator==( const FTwoVectors& V ) const
+UE_FORCEINLINE_HINT bool FTwoVectors::operator==( const FTwoVectors& V ) const
 {
 	return ((v1 == V.v1) && (v2 == V.v2));
 }
 
 
-FORCEINLINE bool FTwoVectors::operator!=( const FTwoVectors& V ) const
+UE_FORCEINLINE_HINT bool FTwoVectors::operator!=( const FTwoVectors& V ) const
 {
 	return ((v1 != V.v1) || (v2 != V.v2));
 }
 
 
-FORCEINLINE bool FTwoVectors::Equals( const FTwoVectors& V, float Tolerance ) const
+UE_FORCEINLINE_HINT bool FTwoVectors::Equals( const FTwoVectors& V, float Tolerance ) const
 {
 	return v1.Equals(V.v1, Tolerance) && v2.Equals(V.v2, Tolerance);
 }
@@ -431,7 +431,7 @@ inline FVector::FReal& FTwoVectors::operator[]( int32 i )
 }
 
 
-FORCEINLINE FString FTwoVectors::ToString() const
+UE_FORCEINLINE_HINT FString FTwoVectors::ToString() const
 {
 	return FString::Printf(TEXT("V1=(%s) V2=(%s)"), *v1.ToString(), *v2.ToString());
 }
diff --git a/Engine/Source/Runtime/Core/Public/Math/UnrealMathUtility.h b/Engine/Source/Runtime/Core/Public/Math/UnrealMathUtility.h
--- a/Engine/Source/Runtime/Core/Public/Math/UnrealMathUtility.h
+++ b/Engine/Source/Runtime/Core/Public/Math/UnrealMathUtility.h
@@ -271,14 +271,14 @@ struct FMath : public FPlatformMath
 	// Random Number Functions
 
 	/** Helper function for rand implementations. Returns a random number in [0..A) */
-	[[nodiscard]] static FORCEINLINE int32 RandHelper(int32 A)
+	[[nodiscard]] static UE_FORCEINLINE_HINT int32 RandHelper(int32 A)
 	{
 		// Note that on some platforms RAND_MAX is a large number so we cannot do ((rand()/(RAND_MAX+1)) * A)
 		// or else we may include the upper bound results, which should be excluded.
 		return A > 0 ? Min(TruncToInt(FRand() * (float)A), A - 1) : 0;
 	}
 
-	[[nodiscard]] static FORCEINLINE int64 RandHelper64(int64 A)
+	[[nodiscard]] static UE_FORCEINLINE_HINT int64 RandHelper64(int64 A)
 	{
 		// Note that on some platforms RAND_MAX is a large number so we cannot do ((rand()/(RAND_MAX+1)) * A)
 		// or else we may include the upper bound results, which should be excluded.
@@ -299,24 +299,24 @@ struct FMath : public FPlatformMath
 	}
 
 	/** Util to generate a random number in a range. Overloaded to distinguish from int32 version, where passing a float is typically a mistake. */
-	[[nodiscard]] static FORCEINLINE float RandRange(float InMin, float InMax)
+	[[nodiscard]] static UE_FORCEINLINE_HINT float RandRange(float InMin, float InMax)
 	{
 		return FRandRange(InMin, InMax);
 	}
 
-	[[nodiscard]] static FORCEINLINE double RandRange(double InMin, double InMax)
+	[[nodiscard]] static UE_FORCEINLINE_HINT double RandRange(double InMin, double InMax)
 	{
 		return FRandRange(InMin, InMax);
 	}
 
 	/** Util to generate a random number in a range. */
-	[[nodiscard]] static FORCEINLINE float FRandRange(float InMin, float InMax)
+	[[nodiscard]] static UE_FORCEINLINE_HINT float FRandRange(float InMin, float InMax)
 	{
 		return InMin + (InMax - InMin) * FRand();
 	}
 
 	/** Util to generate a random number in a range. */
-	[[nodiscard]] static FORCEINLINE double FRandRange(double InMin, double InMax)
+	[[nodiscard]] static UE_FORCEINLINE_HINT double FRandRange(double InMin, double InMax)
 	{
 		return InMin + (InMax - InMin) * FRand();	// LWC_TODO: Implement FRandDbl() for increased precision
 	}
@@ -324,7 +324,7 @@ struct FMath : public FPlatformMath
 	RESOLVE_FLOAT_AMBIGUITY_2_ARGS(FRandRange);
 
 	/** Util to generate a random boolean. */
-	[[nodiscard]] static FORCEINLINE bool RandBool()
+	[[nodiscard]] static UE_FORCEINLINE_HINT bool RandBool()
 	{
 		return (RandRange(0,1) == 1) ? true : false;
 	}
@@ -365,7 +365,7 @@ struct FMath : public FPlatformMath
 
 	/** Checks if value is within a range, exclusive on MaxValue) */
 	template< class T, class U> 
-	[[nodiscard]] static constexpr FORCEINLINE bool IsWithin(const T& TestValue, const U& MinValue, const U& MaxValue)
+	[[nodiscard]] static constexpr UE_FORCEINLINE_HINT bool IsWithin(const T& TestValue, const U& MinValue, const U& MaxValue)
 	{
 		return ((TestValue >= MinValue) && (TestValue < MaxValue));
 	}
@@ -373,7 +373,7 @@ struct FMath : public FPlatformMath
 
 	/** Checks if value is within a range, inclusive on MaxValue) */
 	template< class T, class U> 
-	[[nodiscard]] static constexpr FORCEINLINE bool IsWithinInclusive(const T& TestValue, const U& MinValue, const U& MaxValue)
+	[[nodiscard]] static constexpr UE_FORCEINLINE_HINT bool IsWithinInclusive(const T& TestValue, const U& MinValue, const U& MaxValue)
 	{
 		return ((TestValue>=MinValue) && (TestValue <= MaxValue));
 	}
@@ -385,12 +385,12 @@ struct FMath : public FPlatformMath
 	 *	@param ErrorTolerance	Maximum allowed difference for considering them as 'nearly equal'
 	 *	@return					true if A and B are nearly equal
 	 */
-	[[nodiscard]] static FORCEINLINE bool IsNearlyEqual(float A, float B, float ErrorTolerance = UE_SMALL_NUMBER)
+	[[nodiscard]] static UE_FORCEINLINE_HINT bool IsNearlyEqual(float A, float B, float ErrorTolerance = UE_SMALL_NUMBER)
 	{
 		return Abs<float>( A - B ) <= ErrorTolerance;
 	}
 
-	[[nodiscard]] static FORCEINLINE bool IsNearlyEqual(double A, double B, double ErrorTolerance = UE_DOUBLE_SMALL_NUMBER)
+	[[nodiscard]] static UE_FORCEINLINE_HINT bool IsNearlyEqual(double A, double B, double ErrorTolerance = UE_DOUBLE_SMALL_NUMBER)
 	{
 		return Abs<double>(A - B) <= ErrorTolerance;
 	}
@@ -404,7 +404,7 @@ struct FMath : public FPlatformMath
 	 *	@param ErrorTolerance	Maximum allowed difference for considering Value as 'nearly zero'
 	 *	@return					true if Value is nearly zero
 	 */
-	[[nodiscard]] static FORCEINLINE bool IsNearlyZero(float Value, float ErrorTolerance = UE_SMALL_NUMBER)
+	[[nodiscard]] static UE_FORCEINLINE_HINT bool IsNearlyZero(float Value, float ErrorTolerance = UE_SMALL_NUMBER)
 	{
 		return Abs<float>( Value ) <= ErrorTolerance;
 	}
@@ -415,7 +415,7 @@ struct FMath : public FPlatformMath
 	 *	@param ErrorTolerance	Maximum allowed difference for considering Value as 'nearly zero'
 	 *	@return					true if Value is nearly zero
 	 */
-	[[nodiscard]] static FORCEINLINE bool IsNearlyZero(double Value, double ErrorTolerance = UE_DOUBLE_SMALL_NUMBER)
+	[[nodiscard]] static UE_FORCEINLINE_HINT bool IsNearlyZero(double Value, double ErrorTolerance = UE_DOUBLE_SMALL_NUMBER)
 	{
 		return Abs<double>( Value ) <= ErrorTolerance;
 	}
@@ -483,7 +483,7 @@ public:
 	 *	                        numbers are allowed to differ.
 	 *	@return					true if the two values are nearly equal.
 	 */
-	[[nodiscard]] static FORCEINLINE bool IsNearlyEqualByULP(float A, float B, int32 MaxUlps = 4)
+	[[nodiscard]] static UE_FORCEINLINE_HINT bool IsNearlyEqualByULP(float A, float B, int32 MaxUlps = 4)
 	{
 		return TIsNearlyEqualByULP<float, uint32, uint32(1U << 31)>(A, B, MaxUlps);
 	}
@@ -505,7 +505,7 @@ public:
 	 *	                        numbers are allowed to differ.
 	 *	@return					true if the two values are nearly equal.
 	 */
-	[[nodiscard]] static FORCEINLINE bool IsNearlyEqualByULP(double A, double B, int32 MaxUlps = 4)
+	[[nodiscard]] static UE_FORCEINLINE_HINT bool IsNearlyEqualByULP(double A, double B, int32 MaxUlps = 4)
 	{
 		return TIsNearlyEqualByULP<double, uint64, uint64(1ULL << 63)>(A, B, MaxUlps);
 	}
@@ -516,19 +516,19 @@ public:
 	*	@return			true if Value is a power of two
 	*/
 	template <typename T>
-	[[nodiscard]] static constexpr FORCEINLINE bool IsPowerOfTwo( T Value )
+	[[nodiscard]] static constexpr UE_FORCEINLINE_HINT bool IsPowerOfTwo( T Value )
 	{
 		return ((Value & (Value - 1)) == (T)0);
 	}
 
 	/** Converts a float to a nearest less or equal integer. */
-	[[nodiscard]] static FORCEINLINE float Floor(float F)
+	[[nodiscard]] static UE_FORCEINLINE_HINT float Floor(float F)
 	{
 		return FloorToFloat(F);
 	}
 
 	/** Converts a double to a nearest less or equal integer. */
-	[[nodiscard]] static FORCEINLINE double Floor(double F)
+	[[nodiscard]] static UE_FORCEINLINE_HINT double Floor(double F)
 	{
 		return FloorToDouble(F);
 	}
@@ -538,7 +538,7 @@ public:
 	 * Unlike std::floor, it returns an IntegralType.
 	 */
 	template <UE::CIntegral IntegralType>
-	[[nodiscard]] static constexpr FORCEINLINE IntegralType Floor(IntegralType I)
+	[[nodiscard]] static constexpr UE_FORCEINLINE_HINT IntegralType Floor(IntegralType I)
 	{
 		return I;
 	}
@@ -548,48 +548,48 @@ public:
 
 	/** Returns highest of 3 values */
 	template< class T > 
-	[[nodiscard]] static constexpr FORCEINLINE T Max3( const T A, const T B, const T C )
+	[[nodiscard]] static constexpr UE_FORCEINLINE_HINT T Max3( const T A, const T B, const T C )
 	{
 		return Max ( Max( A, B ), C );
 	}
 
 	/** Returns lowest of 3 values */
 	template< class T > 
-	[[nodiscard]] static constexpr FORCEINLINE T Min3( const T A, const T B, const T C )
+	[[nodiscard]] static constexpr UE_FORCEINLINE_HINT T Min3( const T A, const T B, const T C )
 	{
 		return Min ( Min( A, B ), C );
 	}
 
 	template< class T > 
-	[[nodiscard]] static constexpr FORCEINLINE int32 Max3Index( const T A, const T B, const T C )
+	[[nodiscard]] static constexpr UE_FORCEINLINE_HINT int32 Max3Index( const T A, const T B, const T C )
 	{
 		return ( A > B ) ? ( ( A > C ) ? 0 : 2 ) : ( ( B > C ) ? 1 : 2 );
 	}
 
 	/** Returns index of the lowest value */
 	template< class T > 
-	[[nodiscard]] static constexpr FORCEINLINE int32 Min3Index( const T A, const T B, const T C )
+	[[nodiscard]] static constexpr UE_FORCEINLINE_HINT int32 Min3Index( const T A, const T B, const T C )
 	{
 		return ( A < B ) ? ( ( A < C ) ? 0 : 2 ) : ( ( B < C ) ? 1 : 2 );
 	}
 
 	/** Multiples value by itself */
 	template< class T > 
-	[[nodiscard]] static constexpr FORCEINLINE T Square( const T A )
+	[[nodiscard]] static constexpr UE_FORCEINLINE_HINT T Square( const T A )
 	{
 		return A*A;
 	}
 
 	/** Cubes the value */
 	template< class T > 
-    [[nodiscard]] static constexpr FORCEINLINE T Cube( const T A )
+    [[nodiscard]] static constexpr UE_FORCEINLINE_HINT T Cube( const T A )
 	{
 		return A*A*A;
 	}
 
 	/** Clamps X to be between Min and Max, inclusive */
 	template< class T >
-	[[nodiscard]] static constexpr FORCEINLINE T Clamp(const T X, const T MinValue, const T MaxValue)
+	[[nodiscard]] static constexpr UE_FORCEINLINE_HINT T Clamp(const T X, const T MinValue, const T MaxValue)
 	{
 		return Max(Min(X, MaxValue), MinValue);
 	}
@@ -597,11 +597,11 @@ public:
 	MIX_FLOATS_3_ARGS(Clamp);
 	
 	/** Clamps X to be between Min and Max, inclusive. Explicitly defined here for floats/doubles because static analysis gets confused between template and int versions. */
-	[[nodiscard]] static constexpr FORCEINLINE float Clamp(const float X, const float Min, const float Max) { return Clamp<float>(X, Min, Max); }
-	[[nodiscard]] static constexpr FORCEINLINE double Clamp(const double X, const double Min, const double Max) { return Clamp<double>(X, Min, Max); }
+	[[nodiscard]] static constexpr UE_FORCEINLINE_HINT float Clamp(const float X, const float Min, const float Max) { return Clamp<float>(X, Min, Max); }
+	[[nodiscard]] static constexpr UE_FORCEINLINE_HINT double Clamp(const double X, const double Min, const double Max) { return Clamp<double>(X, Min, Max); }
 
 	/** Clamps X to be between Min and Max, inclusive. Overload to support mixed int64/int32 types. */
-	[[nodiscard]] static constexpr FORCEINLINE int64 Clamp(const int64 X, const int32 Min, const int32 Max) { return Clamp<int64>(X, Min, Max); }
+	[[nodiscard]] static constexpr UE_FORCEINLINE_HINT int64 Clamp(const int64 X, const int32 Min, const int32 Max) { return Clamp<int64>(X, Min, Max); }
 
 private:
 	template< class T >
@@ -644,7 +644,7 @@ public:
 		typename T
 		UE_REQUIRES(std::is_floating_point_v<T>)
 	>
-	[[nodiscard]] static constexpr FORCEINLINE T Wrap(const T X, const T Min, const T Max)
+	[[nodiscard]] static constexpr UE_FORCEINLINE_HINT T Wrap(const T X, const T Min, const T Max)
 	{
 		return WrapImpl(X, Min, Max);
 	}
@@ -653,7 +653,7 @@ public:
 		UE_REQUIRES(!std::is_floating_point_v<T>)
 	>
 	UE_DEPRECATED(5.5, "Use of FMath::Wrap with non-floating point arguments is deprecated - consider FMath::WrapExclusive instead.")
-	[[nodiscard]] static constexpr FORCEINLINE T Wrap(const T X, const T Min, const T Max)
+	[[nodiscard]] static constexpr UE_FORCEINLINE_HINT T Wrap(const T X, const T Min, const T Max)
 	{
 		return WrapImpl(X, Min, Max);
 	}
@@ -707,7 +707,7 @@ public:
 
 	/** Snaps a value to the nearest grid multiple */
 	template< class T >
-	[[nodiscard]] static constexpr FORCEINLINE T GridSnap(T Location, T Grid)
+	[[nodiscard]] static constexpr UE_FORCEINLINE_HINT T GridSnap(T Location, T Grid)
 	{
 		return (Grid == T{}) ? Location : (Floor((Location + (Grid/(T)2)) / Grid) * Grid);
 	}
@@ -716,14 +716,14 @@ public:
 
 	/** Divides two integers and rounds up */
 	template <class T>
-	[[nodiscard]] static constexpr FORCEINLINE T DivideAndRoundUp(T Dividend, T Divisor)
+	[[nodiscard]] static constexpr UE_FORCEINLINE_HINT T DivideAndRoundUp(T Dividend, T Divisor)
 	{
 		return (Dividend + Divisor - 1) / Divisor;
 	}
 
 	/** Divides two integers and rounds down */
 	template <class T>
-	[[nodiscard]] static constexpr FORCEINLINE T DivideAndRoundDown(T Dividend, T Divisor)
+	[[nodiscard]] static constexpr UE_FORCEINLINE_HINT T DivideAndRoundDown(T Dividend, T Divisor)
 	{
 		return Dividend / Divisor;
 	}
@@ -831,7 +831,7 @@ public:
 		typename U
 		UE_REQUIRES(!std::is_same_v<T, U>)
 	>
-	static FORCEINLINE void SinCos(T* ScalarSin, T* ScalarCos, U Value)
+	static UE_FORCEINLINE_HINT void SinCos(T* ScalarSin, T* ScalarCos, U Value)
 	{
 		SinCos(ScalarSin, ScalarCos, T(Value));
 	}
@@ -868,7 +868,7 @@ public:
 	}
 #undef FASTASIN_HALF_PI
 
-	[[nodiscard]] static FORCEINLINE double FastAsin(double Value)
+	[[nodiscard]] static UE_FORCEINLINE_HINT double FastAsin(double Value)
 	{
 		// TODO: add fast approximation
 		return FMath::Asin(Value);
@@ -882,13 +882,13 @@ public:
 	 * @return					Value in degrees.
 	 */
 	template<class T>
-	[[nodiscard]] static constexpr FORCEINLINE auto RadiansToDegrees(T const& RadVal) -> decltype(RadVal * (180.f / UE_PI))
+	[[nodiscard]] static constexpr UE_FORCEINLINE_HINT auto RadiansToDegrees(T const& RadVal) -> decltype(RadVal * (180.f / UE_PI))
 	{
 		return RadVal * (180.f / UE_PI);
 	}
 
-	static constexpr FORCEINLINE float RadiansToDegrees(float const& RadVal) { return RadVal * (180.f / UE_PI); }
-	static constexpr FORCEINLINE double RadiansToDegrees(double const& RadVal) { return RadVal * (180.0 / UE_DOUBLE_PI); }
+	static constexpr UE_FORCEINLINE_HINT float RadiansToDegrees(float const& RadVal) { return RadVal * (180.f / UE_PI); }
+	static constexpr UE_FORCEINLINE_HINT double RadiansToDegrees(double const& RadVal) { return RadVal * (180.0 / UE_DOUBLE_PI); }
 
 	/** 
 	 * Converts degrees to radians.
@@ -896,13 +896,13 @@ public:
 	 * @return					Value in radians.
 	 */
 	template<class T>
-	[[nodiscard]] static constexpr FORCEINLINE auto DegreesToRadians(T const& DegVal) -> decltype(DegVal * (UE_PI / 180.f))
+	[[nodiscard]] static constexpr UE_FORCEINLINE_HINT auto DegreesToRadians(T const& DegVal) -> decltype(DegVal * (UE_PI / 180.f))
 	{
 		return DegVal * (UE_PI / 180.f);
 	}
 
-	static constexpr FORCEINLINE float DegreesToRadians(float const& DegVal) { return DegVal * (UE_PI / 180.f); }
-	static constexpr FORCEINLINE double DegreesToRadians(double const& DegVal) { return DegVal * (UE_DOUBLE_PI / 180.0); }
+	static constexpr UE_FORCEINLINE_HINT float DegreesToRadians(float const& DegVal) { return DegVal * (UE_PI / 180.f); }
+	static constexpr UE_FORCEINLINE_HINT double DegreesToRadians(double const& DegVal) { return DegVal * (UE_DOUBLE_PI / 180.0); }
 
 	/** 
 	 * Clamps an arbitrary angle to be between the given angles.  Will clamp to nearest boundary.
@@ -1040,7 +1040,7 @@ public:
 	}
 	/** Converts given Cartesian coordinate pair to Polar coordinate system. */
 	template<typename T>
-	static FORCEINLINE void CartesianToPolar(const UE::Math::TVector2<T> InCart, UE::Math::TVector2<T>& OutPolar)
+	static UE_FORCEINLINE_HINT void CartesianToPolar(const UE::Math::TVector2<T> InCart, UE::Math::TVector2<T>& OutPolar)
 	{
 		CartesianToPolar(InCart.X, InCart.Y, OutPolar.X, OutPolar.Y);
 	}
@@ -1054,7 +1054,7 @@ public:
 	}
 	/** Converts given Polar coordinate pair to Cartesian coordinate system. */
 	template<typename T>
-	static FORCEINLINE void PolarToCartesian(const UE::Math::TVector2<T> InPolar, UE::Math::TVector2<T>& OutCart)
+	static UE_FORCEINLINE_HINT void PolarToCartesian(const UE::Math::TVector2<T> InPolar, UE::Math::TVector2<T>& OutCart)
 	{
 		PolarToCartesian(InPolar.X, InPolar.Y, OutCart.X, OutCart.Y);
 	}
@@ -1151,19 +1151,19 @@ public:
 
 	/** Transform the given Value relative to the input range to the Output Range. */
 	template<typename T, typename T2>
-	[[nodiscard]] static FORCEINLINE auto GetMappedRangeValueUnclamped(const UE::Math::TVector2<T>& InputRange, const UE::Math::TVector2<T>& OutputRange, const T2 Value)
+	[[nodiscard]] static UE_FORCEINLINE_HINT auto GetMappedRangeValueUnclamped(const UE::Math::TVector2<T>& InputRange, const UE::Math::TVector2<T>& OutputRange, const T2 Value)
 	{
 		return GetRangeValue(OutputRange, GetRangePct(InputRange, Value));
 	}
 
 	template<class T>
-	[[nodiscard]] static FORCEINLINE double GetRangePct(TRange<T> const& Range, T Value)
+	[[nodiscard]] static UE_FORCEINLINE_HINT double GetRangePct(TRange<T> const& Range, T Value)
 	{
 		return GetRangePct(Range.GetLowerBoundValue(), Range.GetUpperBoundValue(), Value);
 	}
 
 	template<class T>
-	[[nodiscard]] static FORCEINLINE T GetRangeValue(TRange<T> const& Range, T Pct)
+	[[nodiscard]] static UE_FORCEINLINE_HINT T GetRangeValue(TRange<T> const& Range, T Pct)
 	{
 		return FMath::Lerp<T>(Range.GetLowerBoundValue(), Range.GetUpperBoundValue(), Pct);
 	}
@@ -1181,7 +1181,7 @@ public:
 		typename U
 		UE_REQUIRES(!TCustomLerp<T>::Value && (std::is_floating_point_v<U> || std::is_same_v<T, U>) && !std::is_same_v<T, bool>)
 	>
-	[[nodiscard]] static constexpr FORCEINLINE_DEBUGGABLE T Lerp( const T& A, const T& B, const U& Alpha )
+	[[nodiscard]] static constexpr UE_FORCEINLINE_HINT T Lerp( const T& A, const T& B, const U& Alpha )
 	{
 		return (T)(A + Alpha * (B-A));
 	}
@@ -1192,7 +1192,7 @@ public:
 		typename U
 		UE_REQUIRES(TCustomLerp<T>::Value)
 	>
-	[[nodiscard]] static FORCEINLINE_DEBUGGABLE T Lerp(const T& A, const T& B, const U& Alpha)
+	[[nodiscard]] static UE_FORCEINLINE_HINT T Lerp(const T& A, const T& B, const U& Alpha)
 	{
 		return TCustomLerp<T>::Lerp(A, B, Alpha);
 	}
@@ -1212,14 +1212,14 @@ public:
 
 	/** Performs a linear interpolation between two values, Alpha ranges from 0-1. Handles full numeric range of T */
 	template< class T > 
-	[[nodiscard]] static constexpr FORCEINLINE_DEBUGGABLE T LerpStable( const T& A, const T& B, double Alpha )
+	[[nodiscard]] static constexpr UE_FORCEINLINE_HINT T LerpStable( const T& A, const T& B, double Alpha )
 	{
 		return (T)((A * (1.0 - Alpha)) + (B * Alpha));
 	}
 
 	/** Performs a linear interpolation between two values, Alpha ranges from 0-1. Handles full numeric range of T */
 	template< class T >
-	[[nodiscard]] static constexpr FORCEINLINE_DEBUGGABLE T LerpStable(const T& A, const T& B, float Alpha)
+	[[nodiscard]] static constexpr UE_FORCEINLINE_HINT T LerpStable(const T& A, const T& B, float Alpha)
 	{
 		return (T)((A * (1.0f - Alpha)) + (B * Alpha));
 	}
@@ -1258,7 +1258,7 @@ public:
 		typename U
 		UE_REQUIRES(TCustomLerp<T>::Value)
 	>
-	[[nodiscard]] static FORCEINLINE_DEBUGGABLE T BiLerp(const T& P00, const T& P10, const T& P01, const T& P11, const U& FracX, const U& FracY)
+	[[nodiscard]] static UE_FORCEINLINE_HINT T BiLerp(const T& P00, const T& P10, const T& P01, const T& P11, const U& FracX, const U& FracY)
 	{
 		return TCustomLerp<T>::BiLerp(P00, P10, P01, P11, FracX, FracY);
 	}
@@ -1291,7 +1291,7 @@ public:
 		typename U
 		UE_REQUIRES(TCustomLerp<T>::Value)
 	>
-	[[nodiscard]] static FORCEINLINE_DEBUGGABLE T CubicInterp(const T& P0, const T& T0, const T& P1, const T& T1, const U& A)
+	[[nodiscard]] static UE_FORCEINLINE_HINT T CubicInterp(const T& P0, const T& T0, const T& P1, const T& T1, const U& A)
 	{
 		return TCustomLerp<T>::CubicInterp(P0, T0, P1, T1, A);
 	}
@@ -2263,12 +2263,12 @@ public:
 	* @param F		Floating point value to convert
 	* @return		The rounded integer
 	*/
-	[[nodiscard]] static FORCEINLINE float RoundFromZero(float F)
+	[[nodiscard]] static UE_FORCEINLINE_HINT float RoundFromZero(float F)
 	{
 		return (F < 0.0f) ? FloorToFloat(F) : CeilToFloat(F);
 	}
 
-	[[nodiscard]] static FORCEINLINE double RoundFromZero(double F)
+	[[nodiscard]] static UE_FORCEINLINE_HINT double RoundFromZero(double F)
 	{
 		return (F < 0.0) ? FloorToDouble(F) : CeilToDouble(F);
 	}
@@ -2278,12 +2278,12 @@ public:
 	* @param F		Floating point value to convert
 	* @return		The rounded integer
 	*/
-	[[nodiscard]] static FORCEINLINE float RoundToZero(float F)
+	[[nodiscard]] static UE_FORCEINLINE_HINT float RoundToZero(float F)
 	{
 		return (F < 0.0f) ? CeilToFloat(F) : FloorToFloat(F);
 	}
 
-	[[nodiscard]] static FORCEINLINE double RoundToZero(double F)
+	[[nodiscard]] static UE_FORCEINLINE_HINT double RoundToZero(double F)
 	{
 		return (F < 0.0) ? CeilToDouble(F) : FloorToDouble(F);
 	}
@@ -2293,12 +2293,12 @@ public:
 	* @param F		Floating point value to convert
 	* @return		The rounded integer
 	*/
-	[[nodiscard]] static FORCEINLINE float RoundToNegativeInfinity(float F)
+	[[nodiscard]] static UE_FORCEINLINE_HINT float RoundToNegativeInfinity(float F)
 	{
 		return FloorToFloat(F);
 	}
 
-	[[nodiscard]] static FORCEINLINE double RoundToNegativeInfinity(double F)
+	[[nodiscard]] static UE_FORCEINLINE_HINT double RoundToNegativeInfinity(double F)
 	{
 		return FloorToDouble(F);
 	}
@@ -2308,12 +2308,12 @@ public:
 	* @param F		Floating point value to convert
 	* @return		The rounded integer
 	*/
-	[[nodiscard]] static FORCEINLINE float RoundToPositiveInfinity(float F)
+	[[nodiscard]] static UE_FORCEINLINE_HINT float RoundToPositiveInfinity(float F)
 	{
 		return CeilToFloat(F);
 	}
 
-	[[nodiscard]] static FORCEINLINE double RoundToPositiveInfinity(double F)
+	[[nodiscard]] static UE_FORCEINLINE_HINT double RoundToPositiveInfinity(double F)
 	{
 		return CeilToDouble(F);
 	}
diff --git a/Engine/Source/Runtime/Core/Public/Math/UnrealPlatformMathSSE4.h b/Engine/Source/Runtime/Core/Public/Math/UnrealPlatformMathSSE4.h
--- a/Engine/Source/Runtime/Core/Public/Math/UnrealPlatformMathSSE4.h
+++ b/Engine/Source/Runtime/Core/Public/Math/UnrealPlatformMathSSE4.h
@@ -19,42 +19,42 @@ namespace UE4
 {
 namespace SSE4
 {
-	FORCEINLINE float TruncToFloat(float F)
+	UE_FORCEINLINE_HINT float TruncToFloat(float F)
 	{
 		return _mm_cvtss_f32(_mm_round_ps(_mm_set_ss(F), 3));
 	}
 
-	FORCEINLINE double TruncToDouble(double F)
+	UE_FORCEINLINE_HINT double TruncToDouble(double F)
 	{
 		return _mm_cvtsd_f64(_mm_round_pd(_mm_set_sd(F), 3));
 	}
 
-	FORCEINLINE float FloorToFloat(float F)
+	UE_FORCEINLINE_HINT float FloorToFloat(float F)
 	{
 		return _mm_cvtss_f32(_mm_floor_ps(_mm_set_ss(F)));
 	}
 
-	FORCEINLINE double FloorToDouble(double F)
+	UE_FORCEINLINE_HINT double FloorToDouble(double F)
 	{
 		return _mm_cvtsd_f64(_mm_floor_pd(_mm_set_sd(F)));
 	}
 
-	FORCEINLINE float RoundToFloat(float F)
+	UE_FORCEINLINE_HINT float RoundToFloat(float F)
 	{
 		return FloorToFloat(F + 0.5f);
 	}
 
-	FORCEINLINE double RoundToDouble(double F)
+	UE_FORCEINLINE_HINT double RoundToDouble(double F)
 	{
 		return FloorToDouble(F + 0.5);
 	}
 
-	FORCEINLINE float CeilToFloat(float F)
+	UE_FORCEINLINE_HINT float CeilToFloat(float F)
 	{
 		return _mm_cvtss_f32(_mm_ceil_ps(_mm_set_ss(F)));
 	}
 
-	FORCEINLINE double CeilToDouble(double F)
+	UE_FORCEINLINE_HINT double CeilToDouble(double F)
 	{
 		return _mm_cvtsd_f64(_mm_ceil_pd(_mm_set_sd(F)));
 	}
@@ -72,48 +72,48 @@ struct TUnrealPlatformMathSSE4Base : public TUnrealPlatformMathSSEBase<Base>
 
 	// Truncate
 
-	static FORCEINLINE float TruncToFloat(float F)
+	static UE_FORCEINLINE_HINT float TruncToFloat(float F)
 	{
 		return UE4::SSE4::TruncToFloat(F);
 	}
 
-	static FORCEINLINE double TruncToDouble(double F)
+	static UE_FORCEINLINE_HINT double TruncToDouble(double F)
 	{
 		return UE4::SSE4::TruncToDouble(F);
 	}
 
 	// Round
 
-	static FORCEINLINE float RoundToFloat(float F)
+	static UE_FORCEINLINE_HINT float RoundToFloat(float F)
 	{
 		return UE4::SSE4::RoundToFloat(F);
 	}
 
-	static FORCEINLINE double RoundToDouble(double F)
+	static UE_FORCEINLINE_HINT double RoundToDouble(double F)
 	{
 		return UE4::SSE4::RoundToDouble(F);
 	}
 
 	// Floor
 
-	static FORCEINLINE float FloorToFloat(float F)
+	static UE_FORCEINLINE_HINT float FloorToFloat(float F)
 	{
 		return UE4::SSE4::FloorToFloat(F);
 	}
 
-	static FORCEINLINE double FloorToDouble(double F)
+	static UE_FORCEINLINE_HINT double FloorToDouble(double F)
 	{
 		return UE4::SSE4::FloorToDouble(F);
 	}
 
 	// Ceil
 
-	static FORCEINLINE float CeilToFloat(float F)
+	static UE_FORCEINLINE_HINT float CeilToFloat(float F)
 	{
 		return UE4::SSE4::CeilToFloat(F);
 	}
 
-	static FORCEINLINE double CeilToDouble(double F)
+	static UE_FORCEINLINE_HINT double CeilToDouble(double F)
 	{
 		return UE4::SSE4::CeilToDouble(F);
 	}
@@ -123,10 +123,10 @@ struct TUnrealPlatformMathSSE4Base : public TUnrealPlatformMathSSEBase<Base>
 	// Wrappers for overloads in the base, required since calls declared in base struct won't redirect back to this class
 	//
 
-	static FORCEINLINE double TruncToFloat(double F) { return TruncToDouble(F); }
-	static FORCEINLINE double RoundToFloat(double F) { return RoundToDouble(F); }
-	static FORCEINLINE double FloorToFloat(double F) { return FloorToDouble(F); }
-	static FORCEINLINE double CeilToFloat(double F) { return CeilToDouble(F); }
+	static UE_FORCEINLINE_HINT double TruncToFloat(double F) { return TruncToDouble(F); }
+	static UE_FORCEINLINE_HINT double RoundToFloat(double F) { return RoundToDouble(F); }
+	static UE_FORCEINLINE_HINT double FloorToFloat(double F) { return FloorToDouble(F); }
+	static UE_FORCEINLINE_HINT double CeilToFloat(double F) { return CeilToDouble(F); }
 
 #endif // UNREALPLATFORMMATH_SSE4_ENABLED
 };
diff --git a/Engine/Source/Runtime/Core/Public/Math/Vector.h b/Engine/Source/Runtime/Core/Public/Math/Vector.h
--- a/Engine/Source/Runtime/Core/Public/Math/Vector.h
+++ b/Engine/Source/Runtime/Core/Public/Math/Vector.h
@@ -144,8 +144,8 @@ public:
         }
     }
 #else
-    FORCEINLINE void DiagnosticCheckNaN() const {}
-    FORCEINLINE void DiagnosticCheckNaN(const TCHAR* Message) const {}
+    UE_FORCEINLINE_HINT void DiagnosticCheckNaN() const {}
+    UE_FORCEINLINE_HINT void DiagnosticCheckNaN(const TCHAR* Message) const {}
 #endif
 
     /** Default constructor (no initialization). */
@@ -156,9 +156,9 @@ public:
      *
      * @param InF Value to set all components to.
      */
-    [[nodiscard]] explicit FORCEINLINE TVector(T InF);
+    [[nodiscard]] explicit inline TVector(T InF);
 
-	[[nodiscard]] FORCEINLINE constexpr TVector(T InF, TVectorConstInit);
+	[[nodiscard]] UE_FORCEINLINE_HINT constexpr TVector(T InF, TVectorConstInit);
 
     /**
      * Constructor using initial values for each component.
@@ -167,7 +167,7 @@ public:
      * @param InY Y Coordinate.
      * @param InZ Z Coordinate.
      */
-    [[nodiscard]] FORCEINLINE TVector(T InX, T InY, T InZ);
+    [[nodiscard]] inline TVector(T InX, T InY, T InZ);
 
     /**
      * Constructs a vector from an TVector2<T> and Z value.
@@ -175,14 +175,14 @@ public:
      * @param V Vector to copy from.
      * @param InZ Z Coordinate.
      */
-    [[nodiscard]] explicit FORCEINLINE TVector(const TVector2<T> V, T InZ);
+    [[nodiscard]] explicit inline TVector(const TVector2<T> V, T InZ);
 
     /**
      * Constructor using the XYZ components from a 4D vector.
      *
      * @param V 4D Vector to copy from.
      */
-    [[nodiscard]] FORCEINLINE TVector(const UE::Math::TVector4<T>& V);
+    [[nodiscard]] inline TVector(const UE::Math::TVector4<T>& V);
 
     /**
      * Constructs a vector from an FLinearColor.
@@ -212,7 +212,7 @@ public:
      *
      * @param EForceInit Force init enum
      */
-    [[nodiscard]] explicit FORCEINLINE TVector(EForceInit);
+    [[nodiscard]] explicit inline TVector(EForceInit);
 
     /**
      * Calculate cross product between this and another vector.
@@ -220,7 +220,7 @@ public:
      * @param V The other vector.
      * @return The cross product.
      */
-    [[nodiscard]] FORCEINLINE TVector<T> operator^(const TVector<T>& V) const;
+    [[nodiscard]] inline TVector<T> operator^(const TVector<T>& V) const;
 
 	/**
 	 * Calculate cross product between this and another vector.
@@ -228,7 +228,7 @@ public:
 	 * @param V The other vector.
 	 * @return The cross product.
 	 */
-	[[nodiscard]] FORCEINLINE TVector<T> Cross(const TVector<T>& V2) const;
+	[[nodiscard]] UE_FORCEINLINE_HINT TVector<T> Cross(const TVector<T>& V2) const;
 
     /**
      * Calculate the cross product of two vectors.
@@ -237,7 +237,7 @@ public:
      * @param B The second vector.
      * @return The cross product.
      */
-    [[nodiscard]] FORCEINLINE static TVector<T> CrossProduct(const TVector<T>& A, const TVector<T>& B);
+    [[nodiscard]] UE_FORCEINLINE_HINT static TVector<T> CrossProduct(const TVector<T>& A, const TVector<T>& B);
 
     /**
      * Calculate the dot product between this and another vector.
@@ -245,7 +245,7 @@ public:
      * @param V The other vector.
      * @return The dot product.
      */
-    [[nodiscard]] FORCEINLINE T operator|(const TVector<T>& V) const;
+    [[nodiscard]] UE_FORCEINLINE_HINT T operator|(const TVector<T>& V) const;
 
 	/**
 	 * Calculate the dot product between this and another vector.
@@ -253,7 +253,7 @@ public:
 	 * @param V The other vector.
 	 * @return The dot product.
 	 */
-	[[nodiscard]] FORCEINLINE T Dot(const TVector<T>& V) const;
+	[[nodiscard]] UE_FORCEINLINE_HINT T Dot(const TVector<T>& V) const;
 
     /**
      * Calculate the dot product of two vectors.
@@ -262,7 +262,7 @@ public:
      * @param B The second vector.
      * @return The dot product.
      */
-    [[nodiscard]] FORCEINLINE static T DotProduct(const TVector<T>& A, const TVector<T>& B);
+    [[nodiscard]] UE_FORCEINLINE_HINT static T DotProduct(const TVector<T>& A, const TVector<T>& B);
 
     /**
      * Gets the result of component-wise addition of this and another vector.
@@ -270,7 +270,7 @@ public:
      * @param V The vector to add to this.
      * @return The result of vector addition.
      */
-    [[nodiscard]] FORCEINLINE TVector<T> operator+(const TVector<T>& V) const;
+    [[nodiscard]] UE_FORCEINLINE_HINT TVector<T> operator+(const TVector<T>& V) const;
 
     /**
      * Gets the result of component-wise subtraction of this by another vector.
@@ -278,7 +278,7 @@ public:
      * @param V The vector to subtract from this.
      * @return The result of vector subtraction.
      */
-    [[nodiscard]] FORCEINLINE TVector<T> operator-(const TVector<T>& V) const;
+    [[nodiscard]] UE_FORCEINLINE_HINT TVector<T> operator-(const TVector<T>& V) const;
 
     /**
      * Gets the result of subtracting from each component of the vector.
@@ -287,7 +287,7 @@ public:
      * @return The result of subtraction.
      */
 	template<typename FArg UE_REQUIRES(std::is_arithmetic_v<FArg>)>
-	[[nodiscard]] FORCEINLINE TVector<T> operator-(FArg Bias) const
+	[[nodiscard]] UE_FORCEINLINE_HINT TVector<T> operator-(FArg Bias) const
 	{
 		return TVector<T>(X - (T)Bias, Y - (T)Bias, Z - (T)Bias);
 	}
@@ -299,7 +299,7 @@ public:
      * @return The result of addition.
      */
 	template<typename FArg UE_REQUIRES(std::is_arithmetic_v<FArg>)>
-	[[nodiscard]] FORCEINLINE TVector<T> operator+(FArg Bias) const
+	[[nodiscard]] UE_FORCEINLINE_HINT TVector<T> operator+(FArg Bias) const
 	{
 		return TVector<T>(X + (T)Bias, Y + (T)Bias, Z + (T)Bias);
 	}
@@ -311,7 +311,7 @@ public:
      * @return The result of multiplication.
      */
 	template<typename FArg UE_REQUIRES(std::is_arithmetic_v<FArg>)>
-	[[nodiscard]] FORCEINLINE TVector<T> operator*(FArg Scale) const
+	[[nodiscard]] UE_FORCEINLINE_HINT TVector<T> operator*(FArg Scale) const
 	{
 		return TVector<T>(X * (T)Scale, Y * (T)Scale, Z * (T)Scale);
 	}
@@ -335,7 +335,7 @@ public:
      * @param V The vector to multiply with.
      * @return The result of multiplication.
      */
-    [[nodiscard]] FORCEINLINE TVector<T> operator*(const TVector<T>& V) const;
+    [[nodiscard]] UE_FORCEINLINE_HINT TVector<T> operator*(const TVector<T>& V) const;
 
     /**
      * Gets the result of component-wise division of this vector by another.
@@ -343,7 +343,7 @@ public:
      * @param V The vector to divide by.
      * @return The result of division.
      */
-    [[nodiscard]] FORCEINLINE TVector<T> operator/(const TVector<T>& V) const;
+    [[nodiscard]] UE_FORCEINLINE_HINT TVector<T> operator/(const TVector<T>& V) const;
 
     // Binary comparison operators.
 
@@ -353,7 +353,7 @@ public:
      * @param V The vector to check against.
      * @return true if the vectors are equal, false otherwise.
      */
-   [[nodiscard]] bool operator==(const TVector<T>& V) const;
+   [[nodiscard]] UE_FORCEINLINE_HINT bool operator==(const TVector<T>& V) const;
 
     /**
      * Check against another vector for inequality.
@@ -361,7 +361,7 @@ public:
      * @param V The vector to check against.
      * @return true if the vectors are not equal, false otherwise.
      */
-    [[nodiscard]] bool operator!=(const TVector<T>& V) const;
+    [[nodiscard]] UE_FORCEINLINE_HINT bool operator!=(const TVector<T>& V) const;
 
     /**
      * Check against another vector for equality, within specified error limits.
@@ -370,7 +370,7 @@ public:
      * @param Tolerance Error tolerance.
      * @return true if the vectors are equal within tolerance limits, false otherwise.
      */
-    [[nodiscard]] bool Equals(const TVector<T>& V, T Tolerance=UE_KINDA_SMALL_NUMBER) const;
+    [[nodiscard]] UE_FORCEINLINE_HINT bool Equals(const TVector<T>& V, T Tolerance=UE_KINDA_SMALL_NUMBER) const;
 
     /**
      * Checks whether all components of this vector are the same, within a tolerance.
@@ -378,14 +378,14 @@ public:
      * @param Tolerance Error tolerance.
      * @return true if the vectors are equal within tolerance limits, false otherwise.
      */
-    [[nodiscard]] bool AllComponentsEqual(T Tolerance=UE_KINDA_SMALL_NUMBER) const;
+    [[nodiscard]] UE_FORCEINLINE_HINT bool AllComponentsEqual(T Tolerance=UE_KINDA_SMALL_NUMBER) const;
 
     /**
      * Get a negated copy of the vector.
      *
      * @return A negated copy of the vector.
      */
-    [[nodiscard]] FORCEINLINE TVector<T> operator-() const;
+    [[nodiscard]] UE_FORCEINLINE_HINT TVector<T> operator-() const;
 
     /**
      * Adds another vector to this.
@@ -394,7 +394,7 @@ public:
      * @param V Vector to add to this.
      * @return Copy of the vector after addition.
      */
-    FORCEINLINE TVector<T> operator+=(const TVector<T>& V);
+	inline TVector<T> operator+=(const TVector<T>& V);
 
     /**
      * Subtracts another vector from this.
@@ -403,7 +403,7 @@ public:
      * @param V Vector to subtract from this.
      * @return Copy of the vector after subtraction.
      */
-    FORCEINLINE TVector<T> operator-=(const TVector<T>& V);
+    inline TVector<T> operator-=(const TVector<T>& V);
 
     /**
      * Scales the vector.
@@ -614,7 +614,7 @@ public:
      * @param LengthSquaredTolerance Tolerance against squared length.
      * @return true if the vector is a unit vector within the specified tolerance.
      */
-    [[nodiscard]] FORCEINLINE bool IsUnit(T LengthSquaredTolerance = UE_KINDA_SMALL_NUMBER) const;
+    [[nodiscard]] UE_FORCEINLINE_HINT bool IsUnit(T LengthSquaredTolerance = UE_KINDA_SMALL_NUMBER) const;
 
     /**
      * Checks whether vector is normalized.
@@ -637,7 +637,7 @@ public:
      * @return Normalized version of vector.
      * @see GetSafeNormal()
      */
-    [[nodiscard]] FORCEINLINE TVector<T> GetUnsafeNormal() const;
+    [[nodiscard]] inline TVector<T> GetUnsafeNormal() const;
 
     /**
      * Gets a normalized copy of the vector, checking it is safe to do so based on the length.
@@ -672,7 +672,7 @@ public:
      *
      * @param A copy of the vector with each component set to +1 or -1
      */
-    [[nodiscard]] FORCEINLINE TVector<T> GetSignVector() const;
+    [[nodiscard]] inline TVector<T> GetSignVector() const;
 
     /**
      * Projects 2D components of vector based on Z.
@@ -687,7 +687,7 @@ public:
     * @return Normalized version of vector.
     * @see GetSafeNormal2D()
     */
-    [[nodiscard]] FORCEINLINE TVector<T> GetUnsafeNormal2D() const;
+    [[nodiscard]] inline TVector<T> GetUnsafeNormal2D() const;
 
     /**
      * Gets a copy of this vector snapped to a grid.
@@ -785,7 +785,7 @@ public:
      * @param B the other vector to find the 2D cosine of the angle with.
      * @return The cosine.
      */
-    [[nodiscard]] FORCEINLINE T CosineAngle2D(TVector<T> B) const;
+    [[nodiscard]] inline T CosineAngle2D(TVector<T> B) const;
 
     /**
      * Gets a copy of this vector projected onto the input vector.
@@ -793,7 +793,7 @@ public:
      * @param A	Vector to project onto, does not assume it is normalized.
      * @return Projected vector.
      */
-    [[nodiscard]] FORCEINLINE TVector<T> ProjectOnTo(const TVector<T>& A) const ;
+    [[nodiscard]] UE_FORCEINLINE_HINT TVector<T> ProjectOnTo(const TVector<T>& A) const ;
 
     /**
      * Gets a copy of this vector projected onto the input vector, which is assumed to be unit length.
@@ -801,7 +801,7 @@ public:
      * @param  Normal Vector to project onto (assumed to be unit length).
      * @return Projected vector.
      */
-    [[nodiscard]] FORCEINLINE TVector<T> ProjectOnToNormal(const TVector<T>& Normal) const;
+    [[nodiscard]] UE_FORCEINLINE_HINT TVector<T> ProjectOnToNormal(const TVector<T>& Normal) const;
 
     /**
      * Return the TRotator orientation corresponding to the direction in which the vector points.
@@ -831,7 +831,7 @@ public:
      * @return UE::Math::TRotator<T> from the Vector's direction.
      * @see ToOrientationRotator(), ToOrientationQuat()
      */
-	[[nodiscard]] FORCEINLINE UE::Math::TRotator<T> Rotation() const
+	[[nodiscard]] UE_FORCEINLINE_HINT UE::Math::TRotator<T> Rotation() const
 	{
 		return ToOrientationRotator();
 	}
@@ -1014,8 +1014,8 @@ public:
      * @param V2 The second point.
      * @return The distance between two points.
      */
-    [[nodiscard]] static FORCEINLINE T Dist(const TVector<T> &V1, const TVector<T> &V2);
-    [[nodiscard]] static FORCEINLINE T Distance(const TVector<T> &V1, const TVector<T> &V2)
+    [[nodiscard]] static UE_FORCEINLINE_HINT T Dist(const TVector<T> &V1, const TVector<T> &V2);
+    [[nodiscard]] static UE_FORCEINLINE_HINT T Distance(const TVector<T> &V1, const TVector<T> &V2)
 	{
 		return Dist(V1, V2);
 	}
@@ -1027,8 +1027,8 @@ public:
     * @param V2 The second point.
     * @return The distance between two points in the XY plane.
     */
-    [[nodiscard]] static FORCEINLINE T DistXY(const TVector<T> &V1, const TVector<T> &V2);
-    [[nodiscard]] static FORCEINLINE T Dist2D(const TVector<T> &V1, const TVector<T> &V2)
+    [[nodiscard]] static UE_FORCEINLINE_HINT T DistXY(const TVector<T> &V1, const TVector<T> &V2);
+    [[nodiscard]] static UE_FORCEINLINE_HINT T Dist2D(const TVector<T> &V1, const TVector<T> &V2)
 	{
 		return DistXY(V1, V2);
 	}
@@ -1040,7 +1040,7 @@ public:
      * @param V2 The second point.
      * @return The squared distance between two points.
      */
-    [[nodiscard]] static FORCEINLINE T DistSquared(const TVector<T> &V1, const TVector<T> &V2);
+    [[nodiscard]] static UE_FORCEINLINE_HINT T DistSquared(const TVector<T> &V1, const TVector<T> &V2);
 
     /**
      * Squared distance between two points in the XY plane only.
@@ -1049,8 +1049,8 @@ public:
      * @param V2 The second point.
      * @return The squared distance between two points in the XY plane
      */
-    [[nodiscard]] static FORCEINLINE T DistSquaredXY(const TVector<T> &V1, const TVector<T> &V2);
-    [[nodiscard]] static FORCEINLINE T DistSquared2D(const TVector<T> &V1, const TVector<T> &V2)
+    [[nodiscard]] static UE_FORCEINLINE_HINT T DistSquaredXY(const TVector<T> &V1, const TVector<T> &V2);
+    [[nodiscard]] static UE_FORCEINLINE_HINT T DistSquared2D(const TVector<T> &V1, const TVector<T> &V2)
 	{
 		return DistSquaredXY(V1, V2);
 	}
@@ -1062,16 +1062,16 @@ public:
      * @param Size The size of the box.
      * @return Pushout required.
      */
-    [[nodiscard]] static FORCEINLINE T BoxPushOut(const TVector<T>& Normal, const TVector<T>& Size);
+    [[nodiscard]] static UE_FORCEINLINE_HINT T BoxPushOut(const TVector<T>& Normal, const TVector<T>& Size);
 
     /**
      * Min, Max, Min3, Max3 like FMath
      */
-    [[nodiscard]] static FORCEINLINE TVector<T> Min( const TVector<T>& A, const TVector<T>& B );
-    [[nodiscard]] static FORCEINLINE TVector<T> Max( const TVector<T>& A, const TVector<T>& B );
+    [[nodiscard]] static inline TVector<T> Min( const TVector<T>& A, const TVector<T>& B );
+    [[nodiscard]] static inline TVector<T> Max( const TVector<T>& A, const TVector<T>& B );
 
-    [[nodiscard]] static FORCEINLINE TVector<T> Min3( const TVector<T>& A, const TVector<T>& B, const TVector<T>& C );
-    [[nodiscard]] static FORCEINLINE TVector<T> Max3( const TVector<T>& A, const TVector<T>& B, const TVector<T>& C );
+    [[nodiscard]] static inline TVector<T> Min3( const TVector<T>& A, const TVector<T>& B, const TVector<T>& C );
+    [[nodiscard]] static inline TVector<T> Max3( const TVector<T>& A, const TVector<T>& B, const TVector<T>& C );
 
     /**
      * See if two normal vectors are nearly parallel, meaning the angle between them is close to 0 degrees.
@@ -1478,7 +1478,7 @@ inline TVector<T>::TVector(T InF)
 }
 
 template<typename T>
-FORCEINLINE constexpr TVector<T>::TVector(T InF, TVectorConstInit)
+UE_FORCEINLINE_HINT constexpr TVector<T>::TVector(T InF, TVectorConstInit)
 	: X(InF), Y(InF), Z(InF)
 {
 }
@@ -1532,86 +1532,86 @@ inline TVector<T> TVector<T>::operator^(const TVector<T>& V) const
 }
 
 template<typename T>
-FORCEINLINE TVector<T> TVector<T>::Cross(const TVector<T>& V) const
+UE_FORCEINLINE_HINT TVector<T> TVector<T>::Cross(const TVector<T>& V) const
 {
 	return *this ^ V;
 }
 
 template<typename T>
-FORCEINLINE TVector<T> TVector<T>::CrossProduct(const TVector<T>& A, const TVector<T>& B)
+UE_FORCEINLINE_HINT TVector<T> TVector<T>::CrossProduct(const TVector<T>& A, const TVector<T>& B)
 {
     return A ^ B;
 }
 
 template<typename T>
-FORCEINLINE T TVector<T>::operator|(const TVector<T>& V) const
+UE_FORCEINLINE_HINT T TVector<T>::operator|(const TVector<T>& V) const
 {
     return X*V.X + Y*V.Y + Z*V.Z;
 }
 
 template<typename T>
-FORCEINLINE T TVector<T>::Dot(const TVector<T>& V) const
+UE_FORCEINLINE_HINT T TVector<T>::Dot(const TVector<T>& V) const
 {
 	return *this | V;
 }
 
 template<typename T>
-FORCEINLINE T TVector<T>::DotProduct(const TVector<T>& A, const TVector<T>& B)
+UE_FORCEINLINE_HINT T TVector<T>::DotProduct(const TVector<T>& A, const TVector<T>& B)
 {
     return A | B;
 }
 
 template<typename T>
-FORCEINLINE TVector<T> TVector<T>::operator+(const TVector<T>& V) const
+UE_FORCEINLINE_HINT TVector<T> TVector<T>::operator+(const TVector<T>& V) const
 {
     return TVector<T>(X + V.X, Y + V.Y, Z + V.Z);
 }
 
 template<typename T>
-FORCEINLINE TVector<T> TVector<T>::operator-(const TVector<T>& V) const
+UE_FORCEINLINE_HINT TVector<T> TVector<T>::operator-(const TVector<T>& V) const
 {
     return TVector<T>(X - V.X, Y - V.Y, Z - V.Z);
 }
 
 template<typename T>
-FORCEINLINE TVector<T> TVector<T>::operator*(const TVector<T>& V) const
+UE_FORCEINLINE_HINT TVector<T> TVector<T>::operator*(const TVector<T>& V) const
 {
     return TVector<T>(X * V.X, Y * V.Y, Z * V.Z);
 }
 
 template<typename T>
-FORCEINLINE TVector<T> TVector<T>::operator/(const TVector<T>& V) const
+UE_FORCEINLINE_HINT TVector<T> TVector<T>::operator/(const TVector<T>& V) const
 {
     return TVector<T>(X / V.X, Y / V.Y, Z / V.Z);
 }
 
 template<typename T>
-FORCEINLINE bool TVector<T>::operator==(const TVector<T>& V) const
+UE_FORCEINLINE_HINT bool TVector<T>::operator==(const TVector<T>& V) const
 {
     return X==V.X && Y==V.Y && Z==V.Z;
 }
 
 template<typename T>
-FORCEINLINE bool TVector<T>::operator!=(const TVector<T>& V) const
+UE_FORCEINLINE_HINT bool TVector<T>::operator!=(const TVector<T>& V) const
 {
     return X!=V.X || Y!=V.Y || Z!=V.Z;
 }
 
 template<typename T>
-FORCEINLINE bool TVector<T>::Equals(const TVector<T>& V, T Tolerance) const
+UE_FORCEINLINE_HINT bool TVector<T>::Equals(const TVector<T>& V, T Tolerance) const
 {
     return FMath::Abs(X-V.X) <= Tolerance && FMath::Abs(Y-V.Y) <= Tolerance && FMath::Abs(Z-V.Z) <= Tolerance;
 }
 
 template<typename T>
-FORCEINLINE bool TVector<T>::AllComponentsEqual(T Tolerance) const
+UE_FORCEINLINE_HINT bool TVector<T>::AllComponentsEqual(T Tolerance) const
 {
     return FMath::Abs(X - Y) <= Tolerance && FMath::Abs(X - Z) <= Tolerance && FMath::Abs(Y - Z) <= Tolerance;
 }
 
 
 template<typename T>
-FORCEINLINE TVector<T> TVector<T>::operator-() const
+UE_FORCEINLINE_HINT TVector<T> TVector<T>::operator-() const
 {
     return TVector<T>(-X, -Y, -Z);
 }
@@ -1650,13 +1650,13 @@ inline TVector<T> TVector<T>::operator/=(const TVector<T>& V)
 }
 
 template<typename T>
-FORCEINLINE T& TVector<T>::operator[](int32 Index)
+UE_FORCEINLINE_HINT T& TVector<T>::operator[](int32 Index)
 {
 	return Component(Index);
 }
 
 template<typename T>
-FORCEINLINE T TVector<T>::operator[](int32 Index)const
+UE_FORCEINLINE_HINT T TVector<T>::operator[](int32 Index)const
 {
 	return Component(Index);
 }
@@ -1671,79 +1671,79 @@ inline void TVector<T>::Set(T InX, T InY, T InZ)
 }
 
 template<typename T>
-FORCEINLINE T TVector<T>::GetMax() const
+UE_FORCEINLINE_HINT T TVector<T>::GetMax() const
 {
     return FMath::Max(FMath::Max(X,Y),Z);
 }
 
 template<typename T>
-FORCEINLINE T TVector<T>::GetAbsMax() const
+UE_FORCEINLINE_HINT T TVector<T>::GetAbsMax() const
 {
     return FMath::Max(FMath::Max(FMath::Abs(X),FMath::Abs(Y)),FMath::Abs(Z));
 }
 
 template<typename T>
-FORCEINLINE T TVector<T>::GetMin() const
+UE_FORCEINLINE_HINT T TVector<T>::GetMin() const
 {
     return FMath::Min(FMath::Min(X,Y),Z);
 }
 
 template<typename T>
-FORCEINLINE T TVector<T>::GetAbsMin() const
+UE_FORCEINLINE_HINT T TVector<T>::GetAbsMin() const
 {
     return FMath::Min(FMath::Min(FMath::Abs(X),FMath::Abs(Y)),FMath::Abs(Z));
 }
 
 template<typename T>
-FORCEINLINE TVector<T> TVector<T>::ComponentMin(const TVector<T>& Other) const
+UE_FORCEINLINE_HINT TVector<T> TVector<T>::ComponentMin(const TVector<T>& Other) const
 {
     return TVector<T>(FMath::Min(X, Other.X), FMath::Min(Y, Other.Y), FMath::Min(Z, Other.Z));
 }
 
 template<typename T>
-FORCEINLINE TVector<T> TVector<T>::ComponentMax(const TVector<T>& Other) const
+UE_FORCEINLINE_HINT TVector<T> TVector<T>::ComponentMax(const TVector<T>& Other) const
 {
     return TVector<T>(FMath::Max(X, Other.X), FMath::Max(Y, Other.Y), FMath::Max(Z, Other.Z));
 }
 
 template<typename T>
-FORCEINLINE TVector<T> TVector<T>::GetAbs() const
+UE_FORCEINLINE_HINT TVector<T> TVector<T>::GetAbs() const
 {
     return TVector<T>(FMath::Abs(X), FMath::Abs(Y), FMath::Abs(Z));
 }
 
 template<typename T>
-FORCEINLINE T TVector<T>::Size() const
+UE_FORCEINLINE_HINT T TVector<T>::Size() const
 {
     return FMath::Sqrt(X*X + Y*Y + Z*Z);
 }
 
 template<typename T>
-FORCEINLINE T TVector<T>::Length() const
+UE_FORCEINLINE_HINT T TVector<T>::Length() const
 {
 	return Size();
 }
 
 template<typename T>
-FORCEINLINE T TVector<T>::SizeSquared() const
+UE_FORCEINLINE_HINT T TVector<T>::SizeSquared() const
 {
     return X*X + Y*Y + Z*Z;
 }
 
 template<typename T>
-FORCEINLINE T TVector<T>::SquaredLength() const
+UE_FORCEINLINE_HINT T TVector<T>::SquaredLength() const
 {
 	return SizeSquared();
 }
 
 template<typename T>
-FORCEINLINE T TVector<T>::Size2D() const
+UE_FORCEINLINE_HINT T TVector<T>::Size2D() const
 {
     return FMath::Sqrt(X*X + Y*Y);
 }
 
 template<typename T>
-FORCEINLINE T TVector<T>::SizeSquared2D() const
+UE_FORCEINLINE_HINT T TVector<T>::SizeSquared2D() const
 {
     return X*X + Y*Y;
 }
@@ -1758,7 +1758,7 @@ inline bool TVector<T>::IsNearlyZero(T Tolerance) const
 }
 
 template<typename T>
-FORCEINLINE bool TVector<T>::IsZero() const
+UE_FORCEINLINE_HINT bool TVector<T>::IsZero() const
 {
     return X==0.f && Y==0.f && Z==0.f;
 }
@@ -1777,13 +1777,13 @@ inline bool TVector<T>::Normalize(T Tolerance)
 }
 
 template<typename T>
-FORCEINLINE bool TVector<T>::IsUnit(T LengthSquaredTolerance) const
+UE_FORCEINLINE_HINT bool TVector<T>::IsUnit(T LengthSquaredTolerance) const
 {
     return FMath::Abs(1.0f - SizeSquared()) < LengthSquaredTolerance;
 }
 
 template<typename T>
-FORCEINLINE bool TVector<T>::IsNormalized() const
+UE_FORCEINLINE_HINT bool TVector<T>::IsNormalized() const
 {
     return (FMath::Abs(1.f - SizeSquared()) < UE_THRESH_VECTOR_NORMALIZED);
 }
@@ -1852,7 +1852,7 @@ inline TVector<T> TVector<T>::GetUnsafeNormal2D() const
 }
 
 template<typename T>
-FORCEINLINE TVector<T> TVector<T>::GridSnap(const T& GridSz) const
+UE_FORCEINLINE_HINT TVector<T> TVector<T>::GridSnap(const T& GridSz) const
 {
     return TVector<T>(FMath::GridSnap(X, GridSz),FMath::GridSnap(Y, GridSz),FMath::GridSnap(Z, GridSz));
 }
@@ -1943,7 +1943,7 @@ inline TVector<T> TVector<T>::GetClampedToMaxSize2D(T MaxSize) const
 }
 
 template<typename T>
-FORCEINLINE void TVector<T>::AddBounded(const TVector<T>& V, T Radius)
+UE_FORCEINLINE_HINT void TVector<T>::AddBounded(const TVector<T>& V, T Radius)
 {
     *this = (*this + V).BoundToCube(Radius);
 }
@@ -2045,13 +2045,13 @@ inline TVector<T> TVector<T>::Reciprocal() const
 
 
 template<typename T>
-FORCEINLINE bool TVector<T>::IsUniform(T Tolerance) const
+UE_FORCEINLINE_HINT bool TVector<T>::IsUniform(T Tolerance) const
 {
     return AllComponentsEqual(Tolerance);
 }
 
 template<typename T>
-FORCEINLINE TVector<T> TVector<T>::MirrorByVector(const TVector<T>& MirrorNormal) const
+UE_FORCEINLINE_HINT TVector<T> TVector<T>::MirrorByVector(const TVector<T>& MirrorNormal) const
 {
     return *this - MirrorNormal * (2.f * (*this | MirrorNormal));
 }
@@ -2112,13 +2112,13 @@ inline T TVector<T>::CosineAngle2D(TVector<T> B) const
 }
 
 template<typename T>
-FORCEINLINE TVector<T> TVector<T>::ProjectOnTo(const TVector<T>& A) const
+UE_FORCEINLINE_HINT TVector<T> TVector<T>::ProjectOnTo(const TVector<T>& A) const
 { 
     return (A * ((*this | A) / (A | A))); 
 }
 
 template<typename T>
-FORCEINLINE TVector<T> TVector<T>::ProjectOnToNormal(const TVector<T>& Normal) const
+UE_FORCEINLINE_HINT TVector<T> TVector<T>::ProjectOnToNormal(const TVector<T>& Normal) const
 {
     return (Normal * (*this | Normal));
 }
@@ -2301,7 +2301,7 @@ inline bool TVector<T>::ContainsNaN() const
 }
 
 template<typename T>
-FORCEINLINE FString TVector<T>::ToString() const
+UE_FORCEINLINE_HINT FString TVector<T>::ToString() const
 {
     return FString::Printf(TEXT("X=%3.3f Y=%3.3f Z=%3.3f"), X, Y, Z);
 }
@@ -2463,31 +2463,31 @@ inline T TVector<T>::HeadingAngle() const
 
 
 template<typename T>
-FORCEINLINE T TVector<T>::Dist(const TVector<T>& V1, const TVector<T>& V2)
+UE_FORCEINLINE_HINT T TVector<T>::Dist(const TVector<T>& V1, const TVector<T>& V2)
 {
     return FMath::Sqrt(TVector<T>::DistSquared(V1, V2));
 }
 
 template<typename T>
-FORCEINLINE T TVector<T>::DistXY(const TVector<T>& V1, const TVector<T>& V2)
+UE_FORCEINLINE_HINT T TVector<T>::DistXY(const TVector<T>& V1, const TVector<T>& V2)
 {
     return FMath::Sqrt(TVector<T>::DistSquaredXY(V1, V2));
 }
 
 template<typename T>
-FORCEINLINE T TVector<T>::DistSquared(const TVector<T>& V1, const TVector<T>& V2)
+UE_FORCEINLINE_HINT T TVector<T>::DistSquared(const TVector<T>& V1, const TVector<T>& V2)
 {
     return FMath::Square(V2.X-V1.X) + FMath::Square(V2.Y-V1.Y) + FMath::Square(V2.Z-V1.Z);
 }
 
 template<typename T>
-FORCEINLINE T TVector<T>::DistSquaredXY(const TVector<T>& V1, const TVector<T>& V2)
+UE_FORCEINLINE_HINT T TVector<T>::DistSquaredXY(const TVector<T>& V1, const TVector<T>& V2)
 {
     return FMath::Square(V2.X-V1.X) + FMath::Square(V2.Y-V1.Y);
 }
 
 template<typename T>
-FORCEINLINE T TVector<T>::BoxPushOut(const TVector<T>& Normal, const TVector<T>& Size)
+UE_FORCEINLINE_HINT T TVector<T>::BoxPushOut(const TVector<T>& Normal, const TVector<T>& Size)
 {
     return FMath::Abs(Normal.X*Size.X) + FMath::Abs(Normal.Y*Size.Y) + FMath::Abs(Normal.Z*Size.Z);
 }
@@ -2565,7 +2565,7 @@ template<> CORE_API const FVector3d FVector3d::ZAxisVector;
  * @return Result of multiplication.
  */
 template<typename T, typename T2 UE_REQUIRES(std::is_arithmetic_v<T2>)>
-FORCEINLINE TVector<T> operator*(T2 Scale, const TVector<T>& V)
+UE_FORCEINLINE_HINT TVector<T> operator*(T2 Scale, const TVector<T>& V)
 {
 	return V.operator*(Scale);
 }
@@ -2577,7 +2577,7 @@ FORCEINLINE TVector<T> operator*(T2 Scale, const TVector<T>& V)
  * @return The hash value from the components
  */
 template<typename T>
-FORCEINLINE uint32 GetTypeHash(const TVector<T>& Vector)
+UE_FORCEINLINE_HINT uint32 GetTypeHash(const TVector<T>& Vector)
 {
 	// Note: this assumes there's no padding in Vector that could contain uncompared data.
 	return FCrc::MemCrc_DEPRECATED(&Vector, sizeof(Vector));
@@ -2611,7 +2611,7 @@ inline UE::Math::TVector<T> ClampVector(const UE::Math::TVector<T>& V, const UE:
 #define INTEL_ORDER_VECTOR(x) (x)
 #else
 template<typename T>
-static FORCEINLINE UE::Math::TVector<T> INTEL_ORDER_VECTOR(UE::Math::TVector<T> v)
+static UE_FORCEINLINE_HINT UE::Math::TVector<T> INTEL_ORDER_VECTOR(UE::Math::TVector<T> v)
 {
 	return UE::Math::TVector<T>(INTEL_ORDERF(v.X), INTEL_ORDERF(v.Y), INTEL_ORDERF(v.Z));
 }
