Replace some FORCEINLINE usage in Core/CoreUObject with UE_FORCEINLINE_HINT.

diff --git a/Engine/Source/Runtime/Core/Public/Math/Vector2D.h b/Engine/Source/Runtime/Core/Public/Math/Vector2D.h
--- a/Engine/Source/Runtime/Core/Public/Math/Vector2D.h
+++ b/Engine/Source/Runtime/Core/Public/Math/Vector2D.h
@@ -92,14 +92,14 @@ public:
 	* @param InX X coordinate.
 	* @param InY Y coordinate.
 	*/
-	[[nodiscard]] FORCEINLINE TVector2<T>(T InX, T InY);
+	[[nodiscard]] UE_FORCEINLINE_HINT TVector2<T>(T InX, T InY);
 
 	/**
 	* Constructor initializing both components to a single T value.
 	*
 	* @param InF Value to set both components to.
 	*/
-	[[nodiscard]] explicit FORCEINLINE TVector2<T>(T InF);
+	[[nodiscard]] explicit UE_FORCEINLINE_HINT TVector2<T>(T InF);
 
 	/**
 	* Constructs a vector from an FIntPoint.
@@ -107,7 +107,7 @@ public:
 	* @param InPos Integer point used to set this vector.
 	*/
 	template <typename IntType>
-	[[nodiscard]] FORCEINLINE TVector2<T>(TIntPoint<IntType> InPos);
+	[[nodiscard]] UE_FORCEINLINE_HINT TVector2<T>(TIntPoint<IntType> InPos);
 
 	template <typename IntType>
     [[nodiscard]] explicit  TVector2<T>(TIntVector2<IntType> V);
@@ -117,14 +117,14 @@ public:
 	*
 	* @param EForceInit Force init enum
 	*/
-	[[nodiscard]] explicit FORCEINLINE TVector2<T>(EForceInit);
+	[[nodiscard]] explicit UE_FORCEINLINE_HINT TVector2<T>(EForceInit);
 
 	/**
 	* Constructor that does not initialize.  More explicit than the default constructor.
 	*
 	* @param ENoInit Don't init
 	*/
-	[[nodiscard]] explicit FORCEINLINE TVector2<T>(ENoInit) { }
+	[[nodiscard]] explicit UE_FORCEINLINE_HINT TVector2<T>(ENoInit) { }
 
 	/**
 	* Constructs a vector from an FVector.
@@ -132,7 +132,7 @@ public:
 	*
 	* @param V Vector to copy from.
 	*/
-	[[nodiscard]] explicit FORCEINLINE TVector2<T>(const TVector<T>& V);
+	[[nodiscard]] explicit inline TVector2<T>(const TVector<T>& V);
 
 	/**
 	* Constructs a vector from an FVector4.
@@ -140,7 +140,7 @@ public:
 	*
 	* @param V Vector to copy from.
 	*/
-	[[nodiscard]] explicit FORCEINLINE TVector2<T>(const TVector4<T>& V);
+	[[nodiscard]] explicit inline TVector2<T>(const TVector4<T>& V);
 
 public:
 
@@ -150,7 +150,7 @@ public:
 	* @param V The other vector to add to this.
 	* @return The result of adding the vectors together.
 	*/
-	[[nodiscard]] FORCEINLINE TVector2<T> operator+(const TVector2<T>& V) const;
+	[[nodiscard]] UE_FORCEINLINE_HINT TVector2<T> operator+(const TVector2<T>& V) const;
 
 	/**
 	* Gets the result of subtracting a vector from this one.
@@ -158,7 +158,7 @@ public:
 	* @param V The other vector to subtract from this.
 	* @return The result of the subtraction.
 	*/
-	[[nodiscard]] FORCEINLINE TVector2<T> operator-(const TVector2<T>& V) const;
+	[[nodiscard]] UE_FORCEINLINE_HINT TVector2<T> operator-(const TVector2<T>& V) const;
 
 	/**
 	* Gets the result of scaling the vector (multiplying each component by a value).
@@ -166,7 +166,7 @@ public:
 	* @param Scale How much to scale the vector by.
 	* @return The result of scaling this vector.
 	*/
-	[[nodiscard]] FORCEINLINE TVector2<T> operator*(T Scale) const;
+	[[nodiscard]] UE_FORCEINLINE_HINT TVector2<T> operator*(T Scale) const;
 
 	/**
 	* Gets the result of dividing each component of the vector by a value.
@@ -182,7 +182,7 @@ public:
 	* @param A T to add to each component.
 	* @return The result of adding A to each component.
 	*/
-	[[nodiscard]] FORCEINLINE TVector2<T> operator+(T A) const;
+	[[nodiscard]] UE_FORCEINLINE_HINT TVector2<T> operator+(T A) const;
 
 	/**
 	* Gets the result of subtracting A from each component of the vector.
@@ -190,7 +190,7 @@ public:
 	* @param A T to subtract from each component
 	* @return The result of subtracting A from each component.
 	*/
-	[[nodiscard]] FORCEINLINE TVector2<T> operator-(T A) const;
+	[[nodiscard]] UE_FORCEINLINE_HINT TVector2<T> operator-(T A) const;
 
 	/**
 	* Gets the result of component-wise multiplication of this vector by another.
@@ -198,7 +198,7 @@ public:
 	* @param V The other vector to multiply this by.
 	* @return The result of the multiplication.
 	*/
-	[[nodiscard]] FORCEINLINE TVector2<T> operator*(const TVector2<T>& V) const;
+	[[nodiscard]] UE_FORCEINLINE_HINT TVector2<T> operator*(const TVector2<T>& V) const;
 
 	/**
 	* Gets the result of component-wise division of this vector by another.
@@ -206,7 +206,7 @@ public:
 	* @param V The other vector to divide this by.
 	* @return The result of the division.
 	*/
-	[[nodiscard]] TVector2<T> operator/(const TVector2<T>& V) const;
+	[[nodiscard]] UE_FORCEINLINE_HINT TVector2<T> operator/(const TVector2<T>& V) const;
 
 	/**
 	* Calculates dot product of this vector and another.
@@ -214,7 +214,7 @@ public:
 	* @param V The other vector.
 	* @return The dot product.
 	*/
-	[[nodiscard]] FORCEINLINE T operator|(const TVector2<T>& V) const;
+	[[nodiscard]] UE_FORCEINLINE_HINT T operator|(const TVector2<T>& V) const;
 
 	/**
 	* Calculates cross product of this vector and another.
@@ -222,7 +222,7 @@ public:
 	* @param V The other vector.
 	* @return The cross product.
 	*/
-	[[nodiscard]] FORCEINLINE T operator^(const TVector2<T>& V) const;
+	[[nodiscard]] UE_FORCEINLINE_HINT T operator^(const TVector2<T>& V) const;
 
 public:
 
@@ -280,7 +280,7 @@ public:
 	*
 	* @return A negated copy of the vector.
 	*/
-	[[nodiscard]] FORCEINLINE TVector2<T> operator-() const;
+	[[nodiscard]] UE_FORCEINLINE_HINT TVector2<T> operator-() const;
 
 	/**
 	* Adds another vector to this.
@@ -288,7 +288,7 @@ public:
 	* @param V The other vector to add.
 	* @return Copy of the vector after addition.
 	*/
-	FORCEINLINE TVector2<T> operator+=(const TVector2<T>& V);
+	inline TVector2<T> operator+=(const TVector2<T>& V);
 
 	/**
 	* Subtracts another vector from this.
@@ -296,7 +296,7 @@ public:
 	* @param V The other vector to subtract.
 	* @return Copy of the vector after subtraction.
 	*/
-	FORCEINLINE TVector2<T> operator-=(const TVector2<T>& V);
+	inline TVector2<T> operator-=(const TVector2<T>& V);
 
 	/**
 	* Scales this vector.
@@ -304,7 +304,7 @@ public:
 	* @param Scale The scale to multiply vector by.
 	* @return Copy of the vector after scaling.
 	*/
-	FORCEINLINE TVector2<T> operator*=(T Scale);
+	inline TVector2<T> operator*=(T Scale);
 
 	/**
 	* Divides this vector.
@@ -379,7 +379,7 @@ public:
 	* @param B The second vector.
 	* @return The dot product.
 	*/
-	[[nodiscard]] FORCEINLINE static T DotProduct(const TVector2<T>& A, const TVector2<T>& B);
+	[[nodiscard]] UE_FORCEINLINE_HINT static T DotProduct(const TVector2<T>& A, const TVector2<T>& B);
 
 	/**
 	* Squared distance between two 2D points.
@@ -388,7 +388,7 @@ public:
 	* @param V2 The second point.
 	* @return The squared distance between two 2D points.
 	*/
-	[[nodiscard]] FORCEINLINE static T DistSquared(const TVector2<T>& V1, const TVector2<T>& V2);
+	[[nodiscard]] UE_FORCEINLINE_HINT static T DistSquared(const TVector2<T>& V1, const TVector2<T>& V2);
 
 	/**
 	* Distance between two 2D points.
@@ -397,7 +397,7 @@ public:
 	* @param V2 The second point.
 	* @return The distance between two 2D points.
 	*/
-	[[nodiscard]] FORCEINLINE static T Distance(const TVector2<T>& V1, const TVector2<T>& V2);
+	[[nodiscard]] UE_FORCEINLINE_HINT static T Distance(const TVector2<T>& V1, const TVector2<T>& V2);
 
 	/**
 	* Calculate the cross product of two vectors.
@@ -406,7 +406,7 @@ public:
 	* @param B The second vector.
 	* @return The cross product.
 	*/
-	[[nodiscard]] FORCEINLINE static T CrossProduct(const TVector2<T>& A, const TVector2<T>& B);
+	[[nodiscard]] UE_FORCEINLINE_HINT static T CrossProduct(const TVector2<T>& A, const TVector2<T>& B);
 
 	/**
 	* Returns a vector with the maximum component for each dimension from the pair of vectors.
@@ -415,7 +415,7 @@ public:
 	* @param B The second vector.
 	* @return The max vector.
 	*/
-	[[nodiscard]] FORCEINLINE static TVector2<T> Max(const TVector2<T>& A, const TVector2<T>& B);
+	[[nodiscard]] UE_FORCEINLINE_HINT static TVector2<T> Max(const TVector2<T>& A, const TVector2<T>& B);
 
 	/**
 	* Returns a vector with the minimum component for each dimension from the pair of vectors.
@@ -424,7 +424,7 @@ public:
 	* @param B The second vector.
 	* @return The min vector.
 	*/
-	[[nodiscard]] FORCEINLINE static TVector2<T> Min(const TVector2<T>& A, const TVector2<T>& B);
+	[[nodiscard]] UE_FORCEINLINE_HINT static TVector2<T> Min(const TVector2<T>& A, const TVector2<T>& B);
 
 	/**
 	* Returns a vector with each component clamped between a minimum and a maximum.
@@ -434,7 +434,7 @@ public:
 	* @param MaxValue The maximum vector.
 	* @return The clamped vector.
 	*/
-	[[nodiscard]] FORCEINLINE static TVector2<T> Clamp(const TVector2<T>& V, const TVector2<T>& MinValue, const TVector2<T>& MaxValue);
+	[[nodiscard]] UE_FORCEINLINE_HINT static TVector2<T> Clamp(const TVector2<T>& V, const TVector2<T>& MinValue, const TVector2<T>& MaxValue);
 
 	/**
 	* Checks for equality with error-tolerant comparison.
@@ -494,7 +494,7 @@ public:
     * @return The length of this vector.
     * @see Size - This function is a synonym for Size()
     */
-	[[nodiscard]] FORCEINLINE T Length() const
+	[[nodiscard]] UE_FORCEINLINE_HINT T Length() const
 	{
 		return Size();
 	}
@@ -513,7 +513,7 @@ public:
 	* @return The squared length of this vector.
     * @see SizeSquared - This function is a synonym for SizeSquared()
 	*/
-	[[nodiscard]] FORCEINLINE T SquaredLength() const
+	[[nodiscard]] UE_FORCEINLINE_HINT T SquaredLength() const
 	{
 		return SizeSquared();
 	}
@@ -621,14 +621,14 @@ public:
 	*
 	* @param A copy of the vector with each component set to +1 or -1
 	*/
-	[[nodiscard]] FORCEINLINE TVector2<T> GetSignVector() const;
+	[[nodiscard]] inline TVector2<T> GetSignVector() const;
 
 	/**
 	* Get a copy of this vector with absolute value of each component.
 	*
 	* @return A copy of this vector with absolute value of each component.
 	*/
-	[[nodiscard]] FORCEINLINE TVector2<T> GetAbs() const;
+	[[nodiscard]] UE_FORCEINLINE_HINT TVector2<T> GetAbs() const;
 
 	/**
 	* Get a textual representation of the vector.
@@ -680,7 +680,7 @@ public:
 		}
 	}
 #else
-	FORCEINLINE void DiagnosticCheckNaN() {}
+	UE_FORCEINLINE_HINT void DiagnosticCheckNaN() {}
 #endif
 
 	/**
@@ -738,7 +738,7 @@ public:
 * @return The hash value from the components
 */
 template<typename T>
-FORCEINLINE uint32 GetTypeHash(const TVector2<T>& Vector)
+UE_FORCEINLINE_HINT uint32 GetTypeHash(const TVector2<T>& Vector)
 {
 	// Note: this assumes there's no padding in TVector2<T> that could contain uncompared data.
 	return FCrc::MemCrc_DEPRECATED(&Vector,sizeof(Vector));
@@ -827,18 +827,18 @@ template<> CORE_API const FVector2d FVector2d::Unit45Deg;
  * @return Result of multiplication.
  */
 template<typename T, typename T2 UE_REQUIRES(std::is_arithmetic_v<T2>)>
-FORCEINLINE TVector2<T> operator*(T2 Scale, const TVector2<T>& V)
+UE_FORCEINLINE_HINT TVector2<T> operator*(T2 Scale, const TVector2<T>& V)
 {
 	return V.operator*(Scale);
 }
 
 template<typename T>
-FORCEINLINE TVector2<T>::TVector2(T InX,T InY)
+UE_FORCEINLINE_HINT TVector2<T>::TVector2(T InX,T InY)
 	:	X(InX), Y(InY)
 { }
 
 template<typename T>
-FORCEINLINE TVector2<T>::TVector2(T InF)
+UE_FORCEINLINE_HINT TVector2<T>::TVector2(T InF)
 	:	X(InF), Y(InF)
 { }
 
@@ -859,25 +859,25 @@ inline TVector2<T>::TVector2(TIntVector2<IntType> V)
 }
 
 template<typename T>
-FORCEINLINE TVector2<T>::TVector2(EForceInit)
+UE_FORCEINLINE_HINT TVector2<T>::TVector2(EForceInit)
 	: X(0), Y(0)
 {
 }
 
 template<typename T>
-FORCEINLINE TVector2<T> TVector2<T>::operator+(const TVector2<T>& V) const
+UE_FORCEINLINE_HINT TVector2<T> TVector2<T>::operator+(const TVector2<T>& V) const
 {
 	return TVector2<T>(X + V.X, Y + V.Y);
 }
 
 template<typename T>
-FORCEINLINE TVector2<T> TVector2<T>::operator-(const TVector2<T>& V) const
+UE_FORCEINLINE_HINT TVector2<T> TVector2<T>::operator-(const TVector2<T>& V) const
 {
 	return TVector2<T>(X - V.X, Y - V.Y);
 }
 
 template<typename T>
-FORCEINLINE TVector2<T> TVector2<T>::operator*(T Scale) const
+UE_FORCEINLINE_HINT TVector2<T> TVector2<T>::operator*(T Scale) const
 {
 	return TVector2<T>(X * Scale, Y * Scale);
 }
@@ -890,127 +890,127 @@ inline TVector2<T> TVector2<T>::operator/(T Scale) const
 }
 
 template<typename T>
-FORCEINLINE TVector2<T> TVector2<T>::operator+(T A) const
+UE_FORCEINLINE_HINT TVector2<T> TVector2<T>::operator+(T A) const
 {
 	return TVector2<T>(X + A, Y + A);
 }
 
 template<typename T>
-FORCEINLINE TVector2<T> TVector2<T>::operator-(T A) const
+UE_FORCEINLINE_HINT TVector2<T> TVector2<T>::operator-(T A) const
 {
 	return TVector2<T>(X - A, Y - A);
 }
 
 template<typename T>
-FORCEINLINE TVector2<T> TVector2<T>::operator*(const TVector2<T>& V) const
+UE_FORCEINLINE_HINT TVector2<T> TVector2<T>::operator*(const TVector2<T>& V) const
 {
 	return TVector2<T>(X * V.X, Y * V.Y);
 }
 
 template<typename T>
-FORCEINLINE TVector2<T> TVector2<T>::operator/(const TVector2<T>& V) const
+UE_FORCEINLINE_HINT TVector2<T> TVector2<T>::operator/(const TVector2<T>& V) const
 {
 	return TVector2<T>(X / V.X, Y / V.Y);
 }
 
 template<typename T>
-FORCEINLINE T TVector2<T>::operator|(const TVector2<T>& V) const
+UE_FORCEINLINE_HINT T TVector2<T>::operator|(const TVector2<T>& V) const
 {
 	return X*V.X + Y*V.Y;
 }
 
 template<typename T>
-FORCEINLINE T TVector2<T>::operator^(const TVector2<T>& V) const
+UE_FORCEINLINE_HINT T TVector2<T>::operator^(const TVector2<T>& V) const
 {
 	return X*V.Y - Y*V.X;
 }
 
 template<typename T>
-FORCEINLINE T TVector2<T>::DotProduct(const TVector2<T>& A, const TVector2<T>& B)
+UE_FORCEINLINE_HINT T TVector2<T>::DotProduct(const TVector2<T>& A, const TVector2<T>& B)
 {
 	return A | B;
 }
 
 template<typename T>
-FORCEINLINE T TVector2<T>::DistSquared(const TVector2<T> &V1, const TVector2<T> &V2)
+UE_FORCEINLINE_HINT T TVector2<T>::DistSquared(const TVector2<T> &V1, const TVector2<T> &V2)
 {
 	return FMath::Square(V2.X-V1.X) + FMath::Square(V2.Y-V1.Y);
 }
 
 template<typename T>
-FORCEINLINE T TVector2<T>::Distance(const TVector2<T>& V1, const TVector2<T>& V2)
+UE_FORCEINLINE_HINT T TVector2<T>::Distance(const TVector2<T>& V1, const TVector2<T>& V2)
 {
 	return FMath::Sqrt(TVector2<T>::DistSquared(V1, V2));
 }
 
 template<typename T>
-FORCEINLINE T TVector2<T>::CrossProduct(const TVector2<T>& A, const TVector2<T>& B)
+UE_FORCEINLINE_HINT T TVector2<T>::CrossProduct(const TVector2<T>& A, const TVector2<T>& B)
 {
 	return A ^ B;
 }
 
 template<typename T>
-FORCEINLINE TVector2<T> TVector2<T>::Max(const TVector2<T>& A, const TVector2<T>& B)
+UE_FORCEINLINE_HINT TVector2<T> TVector2<T>::Max(const TVector2<T>& A, const TVector2<T>& B)
 {
 	return TVector2<T>(FMath::Max(A.X, B.X), FMath::Max(A.Y, B.Y));
 }
 
 template<typename T>
-FORCEINLINE TVector2<T> TVector2<T>::Min(const TVector2<T>& A, const TVector2<T>& B)
+UE_FORCEINLINE_HINT TVector2<T> TVector2<T>::Min(const TVector2<T>& A, const TVector2<T>& B)
 {
 	return TVector2<T>(FMath::Min(A.X, B.X), FMath::Min(A.Y, B.Y));
 }
 
 template<typename T>
-FORCEINLINE TVector2<T> TVector2<T>::Clamp(const TVector2<T>& V, const TVector2<T>& MinValue, const TVector2<T>& MaxValue)
+UE_FORCEINLINE_HINT TVector2<T> TVector2<T>::Clamp(const TVector2<T>& V, const TVector2<T>& MinValue, const TVector2<T>& MaxValue)
 {
 	return TVector2<T>(FMath::Clamp(V.X, MinValue.X, MaxValue.X), FMath::Clamp(V.Y, MinValue.Y, MaxValue.Y));
 }
 
 template<typename T>
-FORCEINLINE bool TVector2<T>::operator==(const TVector2<T>& V) const
+UE_FORCEINLINE_HINT bool TVector2<T>::operator==(const TVector2<T>& V) const
 {
 	return X==V.X && Y==V.Y;
 }
 
 template<typename T>
-FORCEINLINE bool TVector2<T>::operator!=(const TVector2<T>& V) const
+UE_FORCEINLINE_HINT bool TVector2<T>::operator!=(const TVector2<T>& V) const
 {
 	return X!=V.X || Y!=V.Y;
 }
 
 template<typename T>
-FORCEINLINE bool TVector2<T>::ComponentwiseAllLessThan(const TVector2<T>& Other) const
+UE_FORCEINLINE_HINT bool TVector2<T>::ComponentwiseAllLessThan(const TVector2<T>& Other) const
 {
 	return X < Other.X && Y < Other.Y;
 }
 
 template<typename T>
-FORCEINLINE bool TVector2<T>::ComponentwiseAllGreaterThan(const TVector2<T>& Other) const
+UE_FORCEINLINE_HINT bool TVector2<T>::ComponentwiseAllGreaterThan(const TVector2<T>& Other) const
 {
 	return X > Other.X && Y > Other.Y;
 }
 
 template<typename T>
-FORCEINLINE bool TVector2<T>::ComponentwiseAllLessOrEqual(const TVector2<T>& Other) const
+UE_FORCEINLINE_HINT bool TVector2<T>::ComponentwiseAllLessOrEqual(const TVector2<T>& Other) const
 {
 	return X <= Other.X && Y <= Other.Y;
 }
 
 template<typename T>
-FORCEINLINE bool TVector2<T>::ComponentwiseAllGreaterOrEqual(const TVector2<T>& Other) const
+UE_FORCEINLINE_HINT bool TVector2<T>::ComponentwiseAllGreaterOrEqual(const TVector2<T>& Other) const
 {
 	return X >= Other.X && Y >= Other.Y;
 }
 
 template<typename T>
-FORCEINLINE bool TVector2<T>::Equals(const TVector2<T>& V, T Tolerance) const
+UE_FORCEINLINE_HINT bool TVector2<T>::Equals(const TVector2<T>& V, T Tolerance) const
 {
 	return FMath::Abs(X-V.X) <= Tolerance && FMath::Abs(Y-V.Y) <= Tolerance;
 }
 
 template<typename T>
-FORCEINLINE TVector2<T> TVector2<T>::operator-() const
+UE_FORCEINLINE_HINT TVector2<T> TVector2<T>::operator-() const
 {
 	return TVector2<T>(-X, -Y);
 }
@@ -1059,13 +1059,13 @@ inline TVector2<T> TVector2<T>::operator/=(const TVector2<T>& V)
 }
 
 template<typename T>
-FORCEINLINE T& TVector2<T>::operator[](int32 Index)
+UE_FORCEINLINE_HINT T& TVector2<T>::operator[](int32 Index)
 {
 	return Component(Index);
 }
 
 template<typename T>
-FORCEINLINE T TVector2<T>::operator[](int32 Index) const
+UE_FORCEINLINE_HINT T TVector2<T>::operator[](int32 Index) const
 {
 	return Component(Index);
 }
@@ -1078,49 +1078,49 @@ inline void TVector2<T>::Set(T InX, T InY)
 }
 
 template<typename T>
-FORCEINLINE T TVector2<T>::GetMax() const
+UE_FORCEINLINE_HINT T TVector2<T>::GetMax() const
 {
 	return FMath::Max(X,Y);
 }
 
 template<typename T>
-FORCEINLINE T TVector2<T>::GetAbsMax() const
+UE_FORCEINLINE_HINT T TVector2<T>::GetAbsMax() const
 {
 	return FMath::Max(FMath::Abs(X),FMath::Abs(Y));
 }
 
 template<typename T>
-FORCEINLINE T TVector2<T>::GetMin() const
+UE_FORCEINLINE_HINT T TVector2<T>::GetMin() const
 {
 	return FMath::Min(X,Y);
 }
 
 template<typename T>
-FORCEINLINE TVector2<T> TVector2<T>::ComponentMin(const TVector2<T>& Other) const
+UE_FORCEINLINE_HINT TVector2<T> TVector2<T>::ComponentMin(const TVector2<T>& Other) const
 {
     return TVector2<T>(FMath::Min(X, Other.X), FMath::Min(Y, Other.Y));
 }
 
 template<typename T>
-FORCEINLINE TVector2<T> TVector2<T>::ComponentMax(const TVector2<T>& Other) const
+UE_FORCEINLINE_HINT TVector2<T> TVector2<T>::ComponentMax(const TVector2<T>& Other) const
 {
     return TVector2<T>(FMath::Max(X, Other.X), FMath::Max(Y, Other.Y));
 }
 
 template<typename T>
-FORCEINLINE T TVector2<T>::Size() const
+UE_FORCEINLINE_HINT T TVector2<T>::Size() const
 {
 	return FMath::Sqrt(X*X + Y*Y);
 }
 
 template<typename T>
-FORCEINLINE T TVector2<T>::SizeSquared() const
+UE_FORCEINLINE_HINT T TVector2<T>::SizeSquared() const
 {
 	return X*X + Y*Y;
 }
 
 template<typename T>
-FORCEINLINE T TVector2<T>::Dot(const TVector2<T>& V2) const
+UE_FORCEINLINE_HINT T TVector2<T>::Dot(const TVector2<T>& V2) const
 { 
 	return X * V2.X + Y * V2.Y;
 } 
@@ -1189,7 +1189,7 @@ inline bool TVector2<T>::IsNearlyZero(T Tolerance) const
 }
 
 template<typename T>
-FORCEINLINE bool TVector2<T>::IsZero() const
+UE_FORCEINLINE_HINT bool TVector2<T>::IsZero() const
 {
 	return X==0.f && Y==0.f;
 }
@@ -1246,7 +1246,7 @@ inline TVector2<T> TVector2<T>::RoundToVector() const
 }
 
 template<typename T>
-FORCEINLINE TVector2<T> TVector2<T>::ClampAxes(T MinAxisVal, T MaxAxisVal) const
+UE_FORCEINLINE_HINT TVector2<T> TVector2<T>::ClampAxes(T MinAxisVal, T MaxAxisVal) const
 {
 	return TVector2<T>(FMath::Clamp(X, MinAxisVal, MaxAxisVal), FMath::Clamp(Y, MinAxisVal, MaxAxisVal));
 }
@@ -1293,13 +1293,13 @@ inline TVector2<T> TVector2<T>::GetSignVector() const
 }
 
 template<typename T>
-FORCEINLINE TVector2<T> TVector2<T>::GetAbs() const
+UE_FORCEINLINE_HINT TVector2<T> TVector2<T>::GetAbs() const
 {
 	return TVector2<T>(FMath::Abs(X), FMath::Abs(Y));
 }
 
 template<typename T>
-FORCEINLINE FString TVector2<T>::ToString() const
+UE_FORCEINLINE_HINT FString TVector2<T>::ToString() const
 {
 	return FString::Printf(TEXT("X=%3.3f Y=%3.3f"), X, Y);
 }
diff --git a/Engine/Source/Runtime/Core/Public/Math/Vector2DHalf.h b/Engine/Source/Runtime/Core/Public/Math/Vector2DHalf.h
--- a/Engine/Source/Runtime/Core/Public/Math/Vector2DHalf.h
+++ b/Engine/Source/Runtime/Core/Public/Math/Vector2DHalf.h
@@ -21,7 +21,7 @@ struct FVector2DHalf
 public:
 
 	/** Default Constructor (no initialization). */
-	[[nodiscard]] FORCEINLINE FVector2DHalf() { }
+	[[nodiscard]] UE_FORCEINLINE_HINT FVector2DHalf() { }
 
 	/**
 	 * Constructor.
@@ -29,21 +29,21 @@ public:
 	 * InX half float X value
 	 * Iny half float Y value
 	 */
- 	[[nodiscard]] FORCEINLINE FVector2DHalf( const FFloat16& InX,const FFloat16& InY );
+ 	[[nodiscard]] UE_FORCEINLINE_HINT FVector2DHalf( const FFloat16& InX,const FFloat16& InY );
 
 	/** Constructor 
 	 *
 	 * InX float X value
 	 * Iny float Y value
 	 */
-	[[nodiscard]] FORCEINLINE FVector2DHalf( float InX,float InY );
+	[[nodiscard]] UE_FORCEINLINE_HINT FVector2DHalf( float InX,float InY );
 
 	/** Constructor 
 	 *
 	 * Vector2D float vector
 	 */
-	[[nodiscard]] FORCEINLINE FVector2DHalf( const FVector2f& Vector2D );
-	[[nodiscard]] FORCEINLINE FVector2DHalf( const FVector2d& Vector2D );// LWC_TODO: This should probably be explicit
+	[[nodiscard]] UE_FORCEINLINE_HINT FVector2DHalf( const FVector2f& Vector2D );
+	[[nodiscard]] UE_FORCEINLINE_HINT FVector2DHalf( const FVector2d& Vector2D );// LWC_TODO: This should probably be explicit
 
 public:
 
@@ -98,21 +98,21 @@ public:
 /* FVector2DHalf inline functions
  *****************************************************************************/
 
-FORCEINLINE FVector2DHalf::FVector2DHalf( const FFloat16& InX, const FFloat16& InY )
+UE_FORCEINLINE_HINT FVector2DHalf::FVector2DHalf( const FFloat16& InX, const FFloat16& InY )
  	:	X(InX), Y(InY)
 { }
 
 
-FORCEINLINE FVector2DHalf::FVector2DHalf( float InX, float InY )
+UE_FORCEINLINE_HINT FVector2DHalf::FVector2DHalf( float InX, float InY )
 	:	X(InX), Y(InY)
 { }
 
-FORCEINLINE FVector2DHalf::FVector2DHalf( const FVector2f& Vector2D )
+UE_FORCEINLINE_HINT FVector2DHalf::FVector2DHalf( const FVector2f& Vector2D )
 	:	X(Vector2D.X), Y(Vector2D.Y)
 { }
 
 // LWC_TODO: Precision loss.
-FORCEINLINE FVector2DHalf::FVector2DHalf( const FVector2d& Vector2D )
+UE_FORCEINLINE_HINT FVector2DHalf::FVector2DHalf( const FVector2d& Vector2D )
 	:	X((float)Vector2D.X), Y((float)Vector2D.Y)
 { }
 
@@ -133,18 +133,18 @@ inline FVector2DHalf& FVector2DHalf::operator=( const FVector2d& Vector2D )
 	return *this;
 }
 
-FORCEINLINE FString FVector2DHalf::ToString() const
+UE_FORCEINLINE_HINT FString FVector2DHalf::ToString() const
 {
 	return FString::Printf(TEXT("X=%3.3f Y=%3.3f"), (float)X, (float)Y );
 }
 
 
-FORCEINLINE FVector2DHalf::operator FVector2f() const
+UE_FORCEINLINE_HINT FVector2DHalf::operator FVector2f() const
 {
 	return FVector2f((float)X,(float)Y);
 }
 
-FORCEINLINE FVector2DHalf::operator FVector2d() const
+UE_FORCEINLINE_HINT FVector2DHalf::operator FVector2d() const
 {
 	return FVector2d((float)X,(float)Y);
 }
@@ -155,12 +155,12 @@ inline void FVector2DHalf::SetTruncate( float InX, float InY )
 	Y.SetTruncate(InY);
 }
 
-FORCEINLINE void FVector2DHalf::SetTruncate( const FVector2f& Vector2D )
+UE_FORCEINLINE_HINT void FVector2DHalf::SetTruncate( const FVector2f& Vector2D )
 {
 	SetTruncate(Vector2D.X,Vector2D.Y);
 }
 
-FORCEINLINE void FVector2DHalf::SetTruncate( const FVector2d& Vector2D )
+UE_FORCEINLINE_HINT void FVector2DHalf::SetTruncate( const FVector2d& Vector2D )
 {
 	SetTruncate((float)Vector2D.X,(float)Vector2D.Y);
 }
\ No newline at end of file
diff --git a/Engine/Source/Runtime/Core/Public/Math/Vector4.h b/Engine/Source/Runtime/Core/Public/Math/Vector4.h
--- a/Engine/Source/Runtime/Core/Public/Math/Vector4.h
+++ b/Engine/Source/Runtime/Core/Public/Math/Vector4.h
@@ -193,7 +193,7 @@ public:
 	 * @param ComponentIndex The index of the component.
 	 * @return Reference to the desired component.
 	 */
-	[[nodiscard]] FORCEINLINE T& operator[](int32 ComponentIndex);
+	[[nodiscard]] UE_FORCEINLINE_HINT T& operator[](int32 ComponentIndex);
 
 	/**
 	 * Access a specific component of the vector.
@@ -201,7 +201,7 @@ public:
 	 * @param ComponentIndex The index of the component.
 	 * @return Copy of the desired component.
 	 */
-	[[nodiscard]] FORCEINLINE T operator[](int32 ComponentIndex) const;
+	[[nodiscard]] UE_FORCEINLINE_HINT T operator[](int32 ComponentIndex) const;
 
 	// Unary operators.
 
@@ -210,7 +210,7 @@ public:
 	 *
 	 * @return A negated copy of the vector.
 	 */
-	[[nodiscard]] FORCEINLINE TVector4<T> operator-() const;
+	[[nodiscard]] UE_FORCEINLINE_HINT TVector4<T> operator-() const;
 
 	/**
 	 * Gets the result of adding a vector to this.
@@ -218,7 +218,7 @@ public:
 	 * @param V The vector to add.
 	 * @return The result of vector addition.
 	 */
-	[[nodiscard]] FORCEINLINE TVector4<T> operator+(const TVector4<T>& V) const;
+	[[nodiscard]] UE_FORCEINLINE_HINT TVector4<T> operator+(const TVector4<T>& V) const;
 
 	/**
 	 * Adds another vector to this one.
@@ -226,7 +226,7 @@ public:
 	 * @param V The other vector to add.
 	 * @return Copy of the vector after addition.
 	 */
-	FORCEINLINE TVector4<T> operator+=(const TVector4<T>& V);
+	inline TVector4<T> operator+=(const TVector4<T>& V);
 
 	/**
 	 * Gets the result of subtracting a vector from this.
@@ -234,7 +234,7 @@ public:
 	 * @param V The vector to subtract.
 	 * @return The result of vector subtraction.
 	 */
-	[[nodiscard]] FORCEINLINE TVector4<T> operator-(const TVector4<T>& V) const;
+	[[nodiscard]] UE_FORCEINLINE_HINT TVector4<T> operator-(const TVector4<T>& V) const;
 
 	/**
 	 * Subtracts another vector to this one.
@@ -242,7 +242,7 @@ public:
 	 * @param V The other vector to subtract.
 	 * @return Copy of the vector after subtraction.
 	 */
-	FORCEINLINE TVector4<T> operator-=(const TVector4<T>& V);
+	inline TVector4<T> operator-=(const TVector4<T>& V);
 
 	/**
      * Gets the result of adding to each component of the vector.
@@ -251,7 +251,7 @@ public:
      * @return The result of addition.
      */
 	template<typename FArg UE_REQUIRES(std::is_arithmetic_v<FArg>)>
-	[[nodiscard]] FORCEINLINE TVector4<T> operator+(FArg Bias) const
+	[[nodiscard]] UE_FORCEINLINE_HINT TVector4<T> operator+(FArg Bias) const
 	{
 		return TVector4<T>(X + (T)Bias, Y + (T)Bias, Z + (T)Bias, W + (T)Bias);
 	}
@@ -263,7 +263,7 @@ public:
      * @return The result of subtraction.
      */
 	template<typename FArg UE_REQUIRES(std::is_arithmetic_v<FArg>)>
-	[[nodiscard]] FORCEINLINE TVector4<T> operator-(FArg Bias) const
+	[[nodiscard]] UE_FORCEINLINE_HINT TVector4<T> operator-(FArg Bias) const
 	{
 		return TVector4<T>(X - (T)Bias, Y - (T)Bias, Z - (T)Bias, W - (T)Bias);
 	}
@@ -275,7 +275,7 @@ public:
 	 * @return The result of vector scaling.
 	 */
 	template<typename FArg UE_REQUIRES(std::is_arithmetic_v<FArg>)>
-	[[nodiscard]] FORCEINLINE TVector4<T> operator*(FArg Scale) const
+	[[nodiscard]] UE_FORCEINLINE_HINT TVector4<T> operator*(FArg Scale) const
 	{
 		return TVector4(X * Scale, Y * Scale, Z * Scale, W * Scale);
 	}
@@ -445,14 +445,14 @@ public:
 	 * @param Tolerance Minimum squared length of vector for normalization.
 	 * @return A normalized copy of the vector or a zero vector.
 	 */
-	[[nodiscard]] FORCEINLINE TVector4 GetSafeNormal(T Tolerance = UE_SMALL_NUMBER) const;
+	[[nodiscard]] inline TVector4 GetSafeNormal(T Tolerance = UE_SMALL_NUMBER) const;
 
 	/**
 	 * Calculates normalized version of vector without checking if it is non-zero.
 	 *
 	 * @return Normalized version of vector.
 	 */
-	[[nodiscard]] FORCEINLINE TVector4 GetUnsafeNormal3() const;
+	[[nodiscard]] inline TVector4 GetUnsafeNormal3() const;
 
 	/**
 	 * Return the FRotator orientation corresponding to the direction in which the vector points.
@@ -474,7 +474,7 @@ public:
 	 * @return FRotator from the Vector's direction.
 	 * @see ToOrientationRotator()
 	 */
-	[[nodiscard]] FORCEINLINE TRotator<T> Rotation() const
+	[[nodiscard]] UE_FORCEINLINE_HINT TRotator<T> Rotation() const
 	{
 		return ToOrientationRotator();
 	}
@@ -487,7 +487,7 @@ public:
 	 * @param InZ New Z Coordinate.
 	 * @param InW New W Coordinate.
 	 */
-	FORCEINLINE void Set(T InX, T InY, T InZ, T InW);
+	inline void Set(T InX, T InY, T InZ, T InW);
 
 	/** Gets the component-wise min of two vectors. */
 	[[nodiscard]] TVector4<T> ComponentMin(const TVector4<T>& Other) const;
@@ -552,7 +552,7 @@ public:
 		}
 	}
 #else
-	FORCEINLINE void DiagnosticCheckNaN() { }
+	UE_FORCEINLINE_HINT void DiagnosticCheckNaN() { }
 #endif
 
 private:
@@ -597,7 +597,7 @@ public:
 	 *
 	 * @return The hash value from the components
 	 */
-	FORCEINLINE friend uint32 GetTypeHash(const UE::Math::TVector4<T>& Vector)
+	UE_FORCEINLINE_HINT friend uint32 GetTypeHash(const UE::Math::TVector4<T>& Vector)
 	{
 		return FCrc::MemCrc_DEPRECATED(&Vector, sizeof(Vector));
 	}
@@ -669,7 +669,7 @@ inline TVector4<T>::TVector4(EForceInit)
 
 
 template<typename T>
-FORCEINLINE TVector4<T>::TVector4(ENoInit)
+UE_FORCEINLINE_HINT TVector4<T>::TVector4(ENoInit)
 {
 }
 
@@ -686,7 +686,7 @@ inline TVector4<T>::TVector4(TVector2<T> InXY, TVector2<T> InZW)
 
 template<typename T>
 template<typename IntType>
-FORCEINLINE TVector4<T>::TVector4(const TIntVector4<IntType>& InVector)
+UE_FORCEINLINE_HINT TVector4<T>::TVector4(const TIntVector4<IntType>& InVector)
 	: X((T)InVector.X)
 	, Y((T)InVector.Y)
 	, Z((T)InVector.Z)
@@ -695,14 +695,14 @@ FORCEINLINE TVector4<T>::TVector4(const TIntVector4<IntType>& InVector)
 }
 
 template<typename T>
-FORCEINLINE T& TVector4<T>::operator[](int32 ComponentIndex)
+UE_FORCEINLINE_HINT T& TVector4<T>::operator[](int32 ComponentIndex)
 {
 	return Component(ComponentIndex);
 }
 
 
 template<typename T>
-FORCEINLINE T TVector4<T>::operator[](int32 ComponentIndex) const
+UE_FORCEINLINE_HINT T TVector4<T>::operator[](int32 ComponentIndex) const
 {
 	return Component(ComponentIndex);
 }
@@ -721,14 +721,14 @@ inline void TVector4<T>::Set(T InX, T InY, T InZ, T InW)
 
 
 template<typename T>
-FORCEINLINE TVector4<T> TVector4<T>::operator-() const
+UE_FORCEINLINE_HINT TVector4<T> TVector4<T>::operator-() const
 {
 	return TVector4(-X, -Y, -Z, -W);
 }
 
 
 template<typename T>
-FORCEINLINE TVector4<T> TVector4<T>::operator+(const TVector4<T>& V) const
+UE_FORCEINLINE_HINT TVector4<T> TVector4<T>::operator+(const TVector4<T>& V) const
 {
 	return TVector4(X + V.X, Y + V.Y, Z + V.Z, W + V.W);
 }
@@ -744,7 +744,7 @@ inline TVector4<T> TVector4<T>::operator+=(const TVector4<T>& V)
 
 
 template<typename T>
-FORCEINLINE TVector4<T> TVector4<T>::operator-(const TVector4<T>& V) const
+UE_FORCEINLINE_HINT TVector4<T> TVector4<T>::operator-(const TVector4<T>& V) const
 {
 	return TVector4(X - V.X, Y - V.Y, Z - V.Z, W - V.W);
 }
@@ -760,7 +760,7 @@ inline TVector4<T> TVector4<T>::operator-=(const TVector4<T>& V)
 
 
 template<typename T>
-FORCEINLINE TVector4<T> TVector4<T>::operator*(const TVector4<T>& V) const
+UE_FORCEINLINE_HINT TVector4<T> TVector4<T>::operator*(const TVector4<T>& V) const
 {
 	return TVector4(X * V.X, Y * V.Y, Z * V.Z, W * V.W);
 }
@@ -803,28 +803,28 @@ bool TVector4<T>::IsValidIndex(int32 Index) const
 }
 
 template<typename T>
-FORCEINLINE bool TVector4<T>::operator==(const TVector4<T>& V) const
+UE_FORCEINLINE_HINT bool TVector4<T>::operator==(const TVector4<T>& V) const
 {
 	return ((X == V.X) && (Y == V.Y) && (Z == V.Z) && (W == V.W));
 }
 
 
 template<typename T>
-FORCEINLINE bool TVector4<T>::operator!=(const TVector4<T>& V) const
+UE_FORCEINLINE_HINT bool TVector4<T>::operator!=(const TVector4<T>& V) const
 {
 	return ((X != V.X) || (Y != V.Y) || (Z != V.Z) || (W != V.W));
 }
 
 
 template<typename T>
-FORCEINLINE bool TVector4<T>::Equals(const TVector4<T>& V, T Tolerance) const
+UE_FORCEINLINE_HINT bool TVector4<T>::Equals(const TVector4<T>& V, T Tolerance) const
 {
 	return FMath::Abs(X-V.X) <= Tolerance && FMath::Abs(Y-V.Y) <= Tolerance && FMath::Abs(Z-V.Z) <= Tolerance && FMath::Abs(W-V.W) <= Tolerance;
 }
 
 
 template<typename T>
-FORCEINLINE FString TVector4<T>::ToString() const
+UE_FORCEINLINE_HINT FString TVector4<T>::ToString() const
 {
 	return FString::Printf(TEXT("X=%3.3f Y=%3.3f Z=%3.3f W=%3.3f"), X, Y, Z, W);
 }
@@ -867,45 +867,45 @@ inline TVector4<T> TVector4<T>::GetUnsafeNormal3() const
 }
 
 template<typename T>
-FORCEINLINE TVector4<T> TVector4<T>::ComponentMin(const TVector4<T>& Other) const
+UE_FORCEINLINE_HINT TVector4<T> TVector4<T>::ComponentMin(const TVector4<T>& Other) const
 {
     return TVector4<T>(FMath::Min(X, Other.X), FMath::Min(Y, Other.Y), FMath::Min(Z, Other.Z), FMath::Min(W, Other.W));
 }
 
 template<typename T>
-FORCEINLINE TVector4<T> TVector4<T>::ComponentMax(const TVector4<T>& Other) const
+UE_FORCEINLINE_HINT TVector4<T> TVector4<T>::ComponentMax(const TVector4<T>& Other) const
 {
     return TVector4<T>(FMath::Max(X, Other.X), FMath::Max(Y, Other.Y), FMath::Max(Z, Other.Z), FMath::Max(W, Other.W));
 }
 
 template<typename T>
-FORCEINLINE T TVector4<T>::Size3() const
+UE_FORCEINLINE_HINT T TVector4<T>::Size3() const
 {
 	return FMath::Sqrt(X*X + Y*Y + Z*Z);
 }
 
 
 template<typename T>
-FORCEINLINE T TVector4<T>::SizeSquared3() const
+UE_FORCEINLINE_HINT T TVector4<T>::SizeSquared3() const
 {
 	return X*X + Y*Y + Z*Z;
 }
 
 template<typename T>
-FORCEINLINE T TVector4<T>::Size() const
+UE_FORCEINLINE_HINT T TVector4<T>::Size() const
 {
 	return FMath::Sqrt(X*X + Y*Y + Z*Z + W*W);
 }
 
 template<typename T>
-FORCEINLINE T TVector4<T>::SizeSquared() const
+UE_FORCEINLINE_HINT T TVector4<T>::SizeSquared() const
 {
 	return X*X + Y*Y + Z*Z + W*W;
 }
 
 
 template<typename T>
-FORCEINLINE bool TVector4<T>::IsUnit3(T LengthSquaredTolerance) const
+UE_FORCEINLINE_HINT bool TVector4<T>::IsUnit3(T LengthSquaredTolerance) const
 {
 	return FMath::Abs(1.0f - SizeSquared3()) < LengthSquaredTolerance;
 }
@@ -959,7 +959,7 @@ inline TVector4<T> TVector4<T>::operator/=(const TVector4<T>& V)
 
 
 template<typename T>
-FORCEINLINE TVector4<T> TVector4<T>::operator/(const TVector4<T>& V) const
+UE_FORCEINLINE_HINT TVector4<T> TVector4<T>::operator/(const TVector4<T>& V) const
 {
 	return TVector4(X / V.X, Y / V.Y, Z / V.Z, W / V.W);
 }
@@ -1015,7 +1015,7 @@ DECLARE_INTRINSIC_TYPE_LAYOUT(FVector4d);
  */
 
 template<typename T>
-FORCEINLINE T Dot3(const UE::Math::TVector4<T>& V1, const UE::Math::TVector4<T>& V2)
+UE_FORCEINLINE_HINT T Dot3(const UE::Math::TVector4<T>& V1, const UE::Math::TVector4<T>& V2)
 {
 	return V1.X * V2.X + V1.Y * V2.Y + V1.Z * V2.Z;
 }
@@ -1030,13 +1030,13 @@ FORCEINLINE T Dot3(const UE::Math::TVector4<T>& V1, const UE::Math::TVector4<T>&
  */
 
 template<typename T>
-FORCEINLINE T Dot3(const UE::Math::TVector4<T>& V1, const UE::Math::TVector<T>& V2)
+UE_FORCEINLINE_HINT T Dot3(const UE::Math::TVector4<T>& V1, const UE::Math::TVector<T>& V2)
 {
 	return V1.X * V2.X + V1.Y * V2.Y + V1.Z * V2.Z;
 }
 
 template<typename T>
-FORCEINLINE T Dot3(const UE::Math::TVector<T>& V1, const UE::Math::TVector4<T>& V2)
+UE_FORCEINLINE_HINT T Dot3(const UE::Math::TVector<T>& V1, const UE::Math::TVector4<T>& V2)
 {
 	return V1.X * V2.X + V1.Y * V2.Y + V1.Z * V2.Z;
 }
@@ -1049,7 +1049,7 @@ FORCEINLINE T Dot3(const UE::Math::TVector<T>& V1, const UE::Math::TVector4<T>&
  * @return The 4D Dot Product.
  */
 template<typename T>
-FORCEINLINE T Dot4(const UE::Math::TVector4<T>& V1, const UE::Math::TVector4<T>& V2)
+UE_FORCEINLINE_HINT T Dot4(const UE::Math::TVector4<T>& V1, const UE::Math::TVector4<T>& V2)
 {
 	return V1.X * V2.X + V1.Y * V2.Y + V1.Z * V2.Z + V1.W * V2.W;
 }
@@ -1066,14 +1066,14 @@ namespace Math
 	 * @return The result of scaling.
 	 */
 	template<typename T, typename T2 UE_REQUIRES(std::is_arithmetic_v<T2>)>
-	FORCEINLINE TVector4<T> operator*(T2 Scale, const TVector4<T>& V)
+	UE_FORCEINLINE_HINT TVector4<T> operator*(T2 Scale, const TVector4<T>& V)
 	{
 		return V.operator*(Scale);
 	}
 
 
 	template<typename T>
-	FORCEINLINE TVector4<T> TVector4<T>::Reflect3(const TVector4<T>& Normal) const
+	UE_FORCEINLINE_HINT TVector4<T> TVector4<T>::Reflect3(const TVector4<T>& Normal) const
 	{
 		return T(2.0f) * Dot3(*this, Normal) * Normal - *this;
 	}
diff --git a/Engine/Source/Runtime/Core/Public/Math/VectorRegister.h b/Engine/Source/Runtime/Core/Public/Math/VectorRegister.h
--- a/Engine/Source/Runtime/Core/Public/Math/VectorRegister.h
+++ b/Engine/Source/Runtime/Core/Public/Math/VectorRegister.h
@@ -46,14 +46,14 @@ struct alignas(16) PersistentVectorRegister4Double
 {
 	double XYZW[4];
 
-	FORCEINLINE PersistentVectorRegister4Double() = default;
+	UE_FORCEINLINE_HINT PersistentVectorRegister4Double() = default;
 
-	FORCEINLINE PersistentVectorRegister4Double(const VectorRegister4Double& Register)
+	UE_FORCEINLINE_HINT PersistentVectorRegister4Double(const VectorRegister4Double& Register)
 	{
 		VectorStore(Register, XYZW);
 	}
 
-	FORCEINLINE PersistentVectorRegister4Double(const VectorRegister4Float& Register)
+	UE_FORCEINLINE_HINT PersistentVectorRegister4Double(const VectorRegister4Float& Register)
 	{
 		VectorStore(VectorRegister4Double(Register), XYZW);
 	}
@@ -70,7 +70,7 @@ struct alignas(16) PersistentVectorRegister4Double
 		return *this;
 	}
 
-	FORCEINLINE operator VectorRegister4Double() const
+	UE_FORCEINLINE_HINT operator VectorRegister4Double() const
 	{
 		return VectorLoad(XYZW);
 	}
diff --git a/Engine/Source/Runtime/Core/Public/Memory/LinearAllocator.h b/Engine/Source/Runtime/Core/Public/Memory/LinearAllocator.h
--- a/Engine/Source/Runtime/Core/Public/Memory/LinearAllocator.h
+++ b/Engine/Source/Runtime/Core/Public/Memory/LinearAllocator.h
@@ -16,14 +16,14 @@
 struct FLinearAllocator
 {
 	CORE_API FLinearAllocator(SIZE_T ReserveMemorySize);
-	FORCEINLINE ~FLinearAllocator()
+	UE_FORCEINLINE_HINT ~FLinearAllocator()
 	{
 		VirtualMemory.FreeVirtual();
 	}
 
 	CORE_API void* Allocate(SIZE_T Size, uint32 Alignment = 8);
 
-	FORCEINLINE SIZE_T GetAllocatedMemorySize() const
+	UE_FORCEINLINE_HINT SIZE_T GetAllocatedMemorySize() const
 	{
 		return Committed;
 	}
@@ -48,7 +48,7 @@ struct FLinearBlockAllocator
 
 	CORE_API void* Allocate(SIZE_T Size, uint32 Alignment = 8);
 
-	FORCEINLINE SIZE_T GetAllocatedMemorySize() const
+	UE_FORCEINLINE_HINT SIZE_T GetAllocatedMemorySize() const
 	{
 		return TotalAllocated;
 	}
diff --git a/Engine/Source/Runtime/Core/Public/Memory/MemoryArena.h b/Engine/Source/Runtime/Core/Public/Memory/MemoryArena.h
--- a/Engine/Source/Runtime/Core/Public/Memory/MemoryArena.h
+++ b/Engine/Source/Runtime/Core/Public/Memory/MemoryArena.h
@@ -51,10 +51,10 @@ private:
 	uint16	ArenaTag = 0;
 
 public:
-	FORCEINLINE const uint16	GetArenaIndex() const									{ return ArenaTag; }
-	FORCEINLINE void*			GetPointer()											{ return Ptr; }
-	FORCEINLINE const void*		GetPointer() const										{ return Ptr; }
-	FORCEINLINE void			SetPointerAndArena(void* InPtr, uint16 InArenaTag)	{ Ptr = InPtr; ArenaTag = InArenaTag; }
+	UE_FORCEINLINE_HINT const uint16	GetArenaIndex() const									{ return ArenaTag; }
+	UE_FORCEINLINE_HINT void*			GetPointer()											{ return Ptr; }
+	UE_FORCEINLINE_HINT const void*		GetPointer() const										{ return Ptr; }
+	UE_FORCEINLINE_HINT void			SetPointerAndArena(void* InPtr, uint16 InArenaTag)	{ Ptr = InPtr; ArenaTag = InArenaTag; }
 
 	inline FArenaPointer(void* InPtr, uint16 InArenaTag) : Ptr(InPtr), ArenaTag(InArenaTag) {}
 
@@ -68,10 +68,10 @@ private:
 	void*	TaggedPointer = nullptr;
 
 public:
-	FORCEINLINE const uint16	GetArenaIndex() const	{ return uint16(UPTRINT(TaggedPointer) >> ArenaShift); }
-	FORCEINLINE void*			GetPointer()			{ return reinterpret_cast<void*>(UPTRINT(TaggedPointer) & PointerMask); }
-	FORCEINLINE const void*		GetPointer() const		{ return reinterpret_cast<const void*>(UPTRINT(TaggedPointer) & PointerMask); }
-	FORCEINLINE void			SetPointerAndArena(void* InPtr, uint16 InArenaTag) { TaggedPointer = reinterpret_cast<void*>(UPTRINT(InPtr) | (UPTRINT(InArenaTag) << ArenaShift)); }
+	UE_FORCEINLINE_HINT const uint16	GetArenaIndex() const	{ return uint16(UPTRINT(TaggedPointer) >> ArenaShift); }
+	UE_FORCEINLINE_HINT void*			GetPointer()			{ return reinterpret_cast<void*>(UPTRINT(TaggedPointer) & PointerMask); }
+	UE_FORCEINLINE_HINT const void*		GetPointer() const		{ return reinterpret_cast<const void*>(UPTRINT(TaggedPointer) & PointerMask); }
+	UE_FORCEINLINE_HINT void			SetPointerAndArena(void* InPtr, uint16 InArenaTag) { TaggedPointer = reinterpret_cast<void*>(UPTRINT(InPtr) | (UPTRINT(InArenaTag) << ArenaShift)); }
 
 	inline FArenaPointer(void* Ptr, uint16 ArenaIndex) : TaggedPointer(reinterpret_cast<void*>(UPTRINT(Ptr) | (UPTRINT(ArenaIndex) << ArenaShift))) {};
 #endif
diff --git a/Engine/Source/Runtime/Core/Public/Memory/VirtualStackAllocator.h b/Engine/Source/Runtime/Core/Public/Memory/VirtualStackAllocator.h
--- a/Engine/Source/Runtime/Core/Public/Memory/VirtualStackAllocator.h
+++ b/Engine/Source/Runtime/Core/Public/Memory/VirtualStackAllocator.h
@@ -78,7 +78,7 @@ public:
     CORE_API FVirtualStackAllocator(size_t RequestedStackSize, EVirtualStackAllocatorDecommitMode Mode);
     CORE_API ~FVirtualStackAllocator();
 
-	FORCEINLINE FScopedStackAllocatorBookmark CreateScopedBookmark()
+	UE_FORCEINLINE_HINT FScopedStackAllocatorBookmark CreateScopedBookmark()
     {
         // Note that this depends on mandatory Return Value Optimization to ensure the memory is only freed once
         return FScopedStackAllocatorBookmark(NextAllocationStart, this);
diff --git a/Engine/Source/Runtime/Core/Public/Microsoft/COMPointer.h b/Engine/Source/Runtime/Core/Public/Microsoft/COMPointer.h
--- a/Engine/Source/Runtime/Core/Public/Microsoft/COMPointer.h
+++ b/Engine/Source/Runtime/Core/Public/Microsoft/COMPointer.h
@@ -149,7 +149,7 @@ public:
 
 public:
 
-	FORCEINLINE PointerType** operator&()
+	UE_FORCEINLINE_HINT PointerType** operator&()
 	{
 		return &(RawPointer);
 	}
@@ -160,17 +160,17 @@ public:
 		return RawPointer;
 	}
 
-	FORCEINLINE bool operator==(PointerType* const Object) const
+	UE_FORCEINLINE_HINT bool operator==(PointerType* const Object) const
 	{
 		return RawPointer == Object;
 	}
 
-	FORCEINLINE bool operator!=(PointerType* const Object) const
+	UE_FORCEINLINE_HINT bool operator!=(PointerType* const Object) const
 	{
 		return RawPointer != Object;
 	}
 
-	FORCEINLINE operator PointerType*() const
+	UE_FORCEINLINE_HINT operator PointerType*() const
 	{
 		return RawPointer;
 	}
@@ -227,7 +227,7 @@ public:
 	 * @return Pointer to the object, or NULL if not valid.
 	 * @see IsValid
 	 */
-	FORCEINLINE PointerType* Get() const
+	UE_FORCEINLINE_HINT PointerType* Get() const
 	{
 		return RawPointer;
 	}
@@ -238,7 +238,7 @@ public:
 	 * @return true if the pointer is valid, false otherwise.
 	 * @see Get
 	 */
-	FORCEINLINE const bool IsValid() const
+	UE_FORCEINLINE_HINT const bool IsValid() const
 	{
 		return (RawPointer != NULL);
 	}
@@ -258,7 +258,7 @@ private:
 	/** Pointer to the actual object, if any. */
 	PointerType* RawPointer;
 
-	friend FORCEINLINE uint32 GetTypeHash(const TComPtr<T>& InObjectPtr)
+	friend UE_FORCEINLINE_HINT uint32 GetTypeHash(const TComPtr<T>& InObjectPtr)
 	{
 		return GetTypeHash(InObjectPtr.RawPointer);
 	}
diff --git a/Engine/Source/Runtime/Core/Public/Microsoft/MicrosoftPlatformMath.h b/Engine/Source/Runtime/Core/Public/Microsoft/MicrosoftPlatformMath.h
--- a/Engine/Source/Runtime/Core/Public/Microsoft/MicrosoftPlatformMath.h
+++ b/Engine/Source/Runtime/Core/Public/Microsoft/MicrosoftPlatformMath.h
@@ -12,10 +12,10 @@
 struct FMicrosoftPlatformMathBase : public TUnrealPlatformMathSSE4Base<FGenericPlatformMath>
 {
 #if PLATFORM_ENABLE_VECTORINTRINSICS
-	static FORCEINLINE bool IsNaN( float A ) { return _isnan(A) != 0; }
-	static FORCEINLINE bool IsNaN(double A) { return _isnan(A) != 0; }
-	static FORCEINLINE bool IsFinite( float A ) { return _finite(A) != 0; }
-	static FORCEINLINE bool IsFinite(double A) { return _finite(A) != 0; }
+	static UE_FORCEINLINE_HINT bool IsNaN( float A ) { return _isnan(A) != 0; }
+	static UE_FORCEINLINE_HINT bool IsNaN(double A) { return _isnan(A) != 0; }
+	static UE_FORCEINLINE_HINT bool IsFinite( float A ) { return _finite(A) != 0; }
+	static UE_FORCEINLINE_HINT bool IsFinite(double A) { return _finite(A) != 0; }
 
 	#pragma intrinsic(_BitScanReverse)
 
@@ -55,12 +55,12 @@ struct FMicrosoftPlatformMathBase : public TUnrealPlatformMathSSE4Base<FGenericP
 		return 32 - CountLeadingZeros(Arg - 1);
 	}
 
-	static FORCEINLINE uint32 RoundUpToPowerOfTwo(uint32 Arg)
+	static UE_FORCEINLINE_HINT uint32 RoundUpToPowerOfTwo(uint32 Arg)
 	{
 		return 1u << CeilLogTwo(Arg);
 	}
 
-	static FORCEINLINE uint64 RoundUpToPowerOfTwo64(uint64 Arg)
+	static UE_FORCEINLINE_HINT uint64 RoundUpToPowerOfTwo64(uint64 Arg)
 	{
 		return uint64(1) << CeilLogTwo64(Arg);
 	}
@@ -126,7 +126,7 @@ struct FMicrosoftPlatformMathBase : public TUnrealPlatformMathSSE4Base<FGenericP
 #endif
 
 #if PLATFORM_ENABLE_POPCNT_INTRINSIC
-	static FORCEINLINE int32 CountBits(uint64 Bits)
+	static UE_FORCEINLINE_HINT int32 CountBits(uint64 Bits)
 	{
 		return _mm_popcnt_u64(Bits);
 	}
diff --git a/Engine/Source/Runtime/Core/Public/Microsoft/MicrosoftPlatformString.h b/Engine/Source/Runtime/Core/Public/Microsoft/MicrosoftPlatformString.h
--- a/Engine/Source/Runtime/Core/Public/Microsoft/MicrosoftPlatformString.h
+++ b/Engine/Source/Runtime/Core/Public/Microsoft/MicrosoftPlatformString.h
@@ -57,7 +57,7 @@ struct FMicrosoftPlatformString :
 	/**
 	 * Wide character implementation 
 	 **/
-	static FORCEINLINE WIDECHAR* Strcpy(WIDECHAR* Dest, const WIDECHAR* Src)
+	static UE_FORCEINLINE_HINT WIDECHAR* Strcpy(WIDECHAR* Dest, const WIDECHAR* Src)
 	{
 		return (WIDECHAR*)_tcscpy(Dest, Src);
 	}
@@ -65,7 +65,7 @@ struct FMicrosoftPlatformString :
 	// This version was deprecated because Strcpy taking a size field does not match the CRT strcpy, and on some platforms the size field was being ignored.
 	// If the memzeroing causes a performance problem, request a new function StrcpyTruncate.
 	UE_DEPRECATED(5.6, "Use Strncpy instead. Note that Strncpy has a behavior difference from Strcpy: it memzeroes the entire DestCount-sized buffer after the end of string.")
-	static FORCEINLINE WIDECHAR* Strcpy(WIDECHAR* Dest, SIZE_T DestCount, const WIDECHAR* Src)
+	static UE_FORCEINLINE_HINT WIDECHAR* Strcpy(WIDECHAR* Dest, SIZE_T DestCount, const WIDECHAR* Src)
 	{
 		return (WIDECHAR*)_tcscpy(Dest, Src);
 	}
@@ -77,14 +77,14 @@ struct FMicrosoftPlatformString :
 		return Dest;
 	}
 
-	static FORCEINLINE WIDECHAR* Strcat(WIDECHAR* Dest, const WIDECHAR* Src)
+	static UE_FORCEINLINE_HINT WIDECHAR* Strcat(WIDECHAR* Dest, const WIDECHAR* Src)
 	{
 		return (WIDECHAR*)_tcscat(Dest, Src);
 	}
 
 	// This version was deprecated because Strcat taking a size field does not match the CRT strcat, and on some platforms the size field was being ignored.
 	UE_DEPRECATED(5.6, "Use Strncat instead. !!NOTE THAT STRNCAT takes SrcLen rather than DestCount. You must call Strncat(Dest, Src, DestCount - Strlen(Dest) - 1).")
-	static FORCEINLINE WIDECHAR* Strcat(WIDECHAR* Dest, SIZE_T DestCount, const WIDECHAR* Src)
+	static UE_FORCEINLINE_HINT WIDECHAR* Strcat(WIDECHAR* Dest, SIZE_T DestCount, const WIDECHAR* Src)
 	{
 		return (WIDECHAR*)_tcscat(Dest, Src);
 	}
@@ -109,77 +109,77 @@ struct FMicrosoftPlatformString :
 		return Dest;
 	}
 
-	static FORCEINLINE int32 Strcmp( const WIDECHAR* String1, const WIDECHAR* String2 )
+	static UE_FORCEINLINE_HINT int32 Strcmp( const WIDECHAR* String1, const WIDECHAR* String2 )
 	{
 		return _tcscmp(String1, String2);
 	}
 
-	static FORCEINLINE int32 Strncmp( const WIDECHAR* String1, const WIDECHAR* String2, SIZE_T Count )
+	static UE_FORCEINLINE_HINT int32 Strncmp( const WIDECHAR* String1, const WIDECHAR* String2, SIZE_T Count )
 	{
 		return _tcsncmp( String1, String2, Count );
 	}
 
-	static FORCEINLINE int32 Strlen( const WIDECHAR* String )
+	static UE_FORCEINLINE_HINT int32 Strlen( const WIDECHAR* String )
 	{
 		return _tcslen( String );
 	}
 
-	static FORCEINLINE int32 Strnlen( const WIDECHAR* String, SIZE_T StringSize )
+	static UE_FORCEINLINE_HINT int32 Strnlen( const WIDECHAR* String, SIZE_T StringSize )
 	{
 		return _tcsnlen( String, StringSize );
 	}
 
-	static FORCEINLINE const WIDECHAR* Strstr( const WIDECHAR* String, const WIDECHAR* Find)
+	static UE_FORCEINLINE_HINT const WIDECHAR* Strstr( const WIDECHAR* String, const WIDECHAR* Find)
 	{
 		return _tcsstr( String, Find );
 	}
 
-	static FORCEINLINE const WIDECHAR* Strchr( const WIDECHAR* String, WIDECHAR C)
+	static UE_FORCEINLINE_HINT const WIDECHAR* Strchr( const WIDECHAR* String, WIDECHAR C)
 	{
 		return _tcschr( String, C ); 
 	}
 
-	static FORCEINLINE const WIDECHAR* Strrchr( const WIDECHAR* String, WIDECHAR C)
+	static UE_FORCEINLINE_HINT const WIDECHAR* Strrchr( const WIDECHAR* String, WIDECHAR C)
 	{
 		return _tcsrchr( String, C ); 
 	}
 
-	static FORCEINLINE int32 Atoi(const WIDECHAR* String)
+	static UE_FORCEINLINE_HINT int32 Atoi(const WIDECHAR* String)
 	{
 		return _tstoi( String ); 
 	}
 
-	static FORCEINLINE int64 Atoi64(const WIDECHAR* String)
+	static UE_FORCEINLINE_HINT int64 Atoi64(const WIDECHAR* String)
 	{
 		return _tstoi64( String ); 
 	}
 
-	static FORCEINLINE float Atof(const WIDECHAR* String)
+	static UE_FORCEINLINE_HINT float Atof(const WIDECHAR* String)
 	{
 		return (float)_tstof( String );
 	}
 
-	static FORCEINLINE double Atod(const WIDECHAR* String)
+	static UE_FORCEINLINE_HINT double Atod(const WIDECHAR* String)
 	{
 		return _tcstod( String, NULL ); 
 	}
 
-	static FORCEINLINE int32 Strtoi( const WIDECHAR* Start, WIDECHAR** End, int32 Base ) 
+	static UE_FORCEINLINE_HINT int32 Strtoi( const WIDECHAR* Start, WIDECHAR** End, int32 Base ) 
 	{
 		return _tcstoul( Start, End, Base );
 	}
 
-	static FORCEINLINE int64 Strtoi64( const WIDECHAR* Start, WIDECHAR** End, int32 Base ) 
+	static UE_FORCEINLINE_HINT int64 Strtoi64( const WIDECHAR* Start, WIDECHAR** End, int32 Base ) 
 	{
 		return _tcstoi64( Start, End, Base ); 
 	}
 
-	static FORCEINLINE uint64 Strtoui64( const WIDECHAR* Start, WIDECHAR** End, int32 Base ) 
+	static UE_FORCEINLINE_HINT uint64 Strtoui64( const WIDECHAR* Start, WIDECHAR** End, int32 Base ) 
 	{
 		return _tcstoui64( Start, End, Base ); 
 	}
 
-	static FORCEINLINE WIDECHAR* Strtok(WIDECHAR* StrToken, const WIDECHAR* Delim, WIDECHAR** Context)
+	static UE_FORCEINLINE_HINT WIDECHAR* Strtok(WIDECHAR* StrToken, const WIDECHAR* Delim, WIDECHAR** Context)
 	{
 		return _tcstok_s(StrToken, Delim, Context);
 	}
@@ -198,7 +198,7 @@ struct FMicrosoftPlatformString :
 	/** 
 	 * Ansi implementation 
 	 **/
-	static FORCEINLINE ANSICHAR* Strcpy(ANSICHAR* Dest, const ANSICHAR* Src)
+	static UE_FORCEINLINE_HINT ANSICHAR* Strcpy(ANSICHAR* Dest, const ANSICHAR* Src)
 	{
 		return (ANSICHAR*)strcpy(Dest, Src);
 	}
@@ -206,7 +206,7 @@ struct FMicrosoftPlatformString :
 	// This version was deprecated because Strcpy taking a size field does not match the CRT strcpy, and on some platforms the size field was being ignored.
 	// If the memzeroing causes a performance problem, request a new function StrcpyTruncate.
 	UE_DEPRECATED(5.6, "Use Strncpy instead. Note that Strncpy has a behavior difference from Strcpy: it memzeroes the entire DestCount-sized buffer after the end of string.")
-	static FORCEINLINE ANSICHAR* Strcpy(ANSICHAR* Dest, SIZE_T DestCount, const ANSICHAR* Src)
+	static UE_FORCEINLINE_HINT ANSICHAR* Strcpy(ANSICHAR* Dest, SIZE_T DestCount, const ANSICHAR* Src)
 	{
 		return (ANSICHAR*)strcpy(Dest, Src);
 	}
@@ -218,94 +218,94 @@ struct FMicrosoftPlatformString :
 		return Dest;
 	}
 
-	static FORCEINLINE ANSICHAR* Strcat(ANSICHAR* Dest, const ANSICHAR* Src)
+	static UE_FORCEINLINE_HINT ANSICHAR* Strcat(ANSICHAR* Dest, const ANSICHAR* Src)
 	{
 		return (ANSICHAR*)strcat( Dest, Src );
 	}
 
 	// This version was deprecated because Strcat taking a size field does not match the CRT strcat, and on some platforms the size field was being ignored.
 	UE_DEPRECATED(5.6, "Use Strncat instead. !!NOTE THAT STRNCAT takes SrcLen rather than DestSize. You must call Strncat(Dest, Src, DestCount - Strlen(Dest) - 1).")
-	static FORCEINLINE ANSICHAR* Strcat(ANSICHAR* Dest, SIZE_T DestCount, const ANSICHAR* Src)
+	static UE_FORCEINLINE_HINT ANSICHAR* Strcat(ANSICHAR* Dest, SIZE_T DestCount, const ANSICHAR* Src)
 	{
 		return (ANSICHAR*)strcat(Dest, Src);
 	}
 
-	static FORCEINLINE ANSICHAR* Strncat(ANSICHAR* Dest, const ANSICHAR* Src, SIZE_T SrcLen)
+	static UE_FORCEINLINE_HINT ANSICHAR* Strncat(ANSICHAR* Dest, const ANSICHAR* Src, SIZE_T SrcLen)
 	{
 		return (ANSICHAR*)strncat(Dest, Src, SrcLen);
 	}
 
-	static FORCEINLINE int32 Strcmp( const ANSICHAR* String1, const ANSICHAR* String2 )
+	static UE_FORCEINLINE_HINT int32 Strcmp( const ANSICHAR* String1, const ANSICHAR* String2 )
 	{
 		return strcmp(String1, String2);
 	}
 
-	static FORCEINLINE int32 Strncmp( const ANSICHAR* String1, const ANSICHAR* String2, SIZE_T Count )
+	static UE_FORCEINLINE_HINT int32 Strncmp( const ANSICHAR* String1, const ANSICHAR* String2, SIZE_T Count )
 	{
 		return strncmp( String1, String2, Count );
 	}
 	
-	static FORCEINLINE int32 Strlen( const ANSICHAR* String )
+	static UE_FORCEINLINE_HINT int32 Strlen( const ANSICHAR* String )
 	{
 		return strlen( String ); 
 	}
 
-	static FORCEINLINE int32 Strnlen( const ANSICHAR* String, SIZE_T StringSize )
+	static UE_FORCEINLINE_HINT int32 Strnlen( const ANSICHAR* String, SIZE_T StringSize )
 	{
 		return strnlen( String, StringSize );
 	}
 
-	static FORCEINLINE const ANSICHAR* Strstr( const ANSICHAR* String, const ANSICHAR* Find)
+	static UE_FORCEINLINE_HINT const ANSICHAR* Strstr( const ANSICHAR* String, const ANSICHAR* Find)
 	{
 		return strstr(String, Find);
 	}
 
-	static FORCEINLINE const ANSICHAR* Strchr( const ANSICHAR* String, ANSICHAR C)
+	static UE_FORCEINLINE_HINT const ANSICHAR* Strchr( const ANSICHAR* String, ANSICHAR C)
 	{
 		return strchr(String, C);
 	}
 
-	static FORCEINLINE const ANSICHAR* Strrchr( const ANSICHAR* String, ANSICHAR C)
+	static UE_FORCEINLINE_HINT const ANSICHAR* Strrchr( const ANSICHAR* String, ANSICHAR C)
 	{
 		return strrchr(String, C);
 	}
 
-	static FORCEINLINE int32 Atoi(const ANSICHAR* String)
+	static UE_FORCEINLINE_HINT int32 Atoi(const ANSICHAR* String)
 	{
 		return atoi( String ); 
 	}
 
-	static FORCEINLINE int64 Atoi64(const ANSICHAR* String)
+	static UE_FORCEINLINE_HINT int64 Atoi64(const ANSICHAR* String)
 	{
 		return _strtoi64( String, NULL, 10 ); 
 	}
 
-	static FORCEINLINE float Atof(const ANSICHAR* String)
+	static UE_FORCEINLINE_HINT float Atof(const ANSICHAR* String)
 	{
 		return (float)atof( String ); 
 	}
 
-	static FORCEINLINE double Atod(const ANSICHAR* String)
+	static UE_FORCEINLINE_HINT double Atod(const ANSICHAR* String)
 	{
 		return atof( String ); 
 	}
 
-	static FORCEINLINE int32 Strtoi( const ANSICHAR* Start, ANSICHAR** End, int32 Base ) 
+	static UE_FORCEINLINE_HINT int32 Strtoi( const ANSICHAR* Start, ANSICHAR** End, int32 Base ) 
 	{
 		return strtol( Start, End, Base ); 
 	}
 
-	static FORCEINLINE int64 Strtoi64( const ANSICHAR* Start, ANSICHAR** End, int32 Base ) 
+	static UE_FORCEINLINE_HINT int64 Strtoi64( const ANSICHAR* Start, ANSICHAR** End, int32 Base ) 
 	{
 		return _strtoi64( Start, End, Base ); 
 	}
 
-	static FORCEINLINE uint64 Strtoui64( const ANSICHAR* Start, ANSICHAR** End, int32 Base ) 
+	static UE_FORCEINLINE_HINT uint64 Strtoui64( const ANSICHAR* Start, ANSICHAR** End, int32 Base ) 
 	{
 		return _strtoui64( Start, End, Base ); 
 	}
 
-	static FORCEINLINE ANSICHAR* Strtok(ANSICHAR* StrToken, const ANSICHAR* Delim, ANSICHAR** Context)
+	static UE_FORCEINLINE_HINT ANSICHAR* Strtok(ANSICHAR* StrToken, const ANSICHAR* Delim, ANSICHAR** Context)
 	{
 		return strtok_s(StrToken, Delim, Context);
 	}
@@ -320,12 +320,12 @@ struct FMicrosoftPlatformString :
 	 * UCS2CHAR implementation - this is identical to WIDECHAR for Windows platforms
 	 */
 
-	static FORCEINLINE int32 Strlen( const UCS2CHAR* String )
+	static UE_FORCEINLINE_HINT int32 Strlen( const UCS2CHAR* String )
 	{
 		return _tcslen( (const WIDECHAR*)String );
 	}
 
-	static FORCEINLINE int32 Strnlen( const UCS2CHAR* String, SIZE_T StringSize )
+	static UE_FORCEINLINE_HINT int32 Strnlen( const UCS2CHAR* String, SIZE_T StringSize )
 	{
 		return _tcsnlen( (const WIDECHAR*)String, StringSize );
 	}
@@ -333,7 +333,7 @@ struct FMicrosoftPlatformString :
 	/**
 	 * UTF8CHAR implementation.
 	 */
-	static FORCEINLINE UTF8CHAR* Strcpy(UTF8CHAR* Dest, const UTF8CHAR* Src)
+	static UE_FORCEINLINE_HINT UTF8CHAR* Strcpy(UTF8CHAR* Dest, const UTF8CHAR* Src)
 	{
 		return (UTF8CHAR*)Strcpy((ANSICHAR*)Dest, (const ANSICHAR*)Src);
 	}
@@ -348,12 +348,12 @@ struct FMicrosoftPlatformString :
 		PRAGMA_ENABLE_DEPRECATION_WARNINGS
 	}
 
-	static FORCEINLINE UTF8CHAR* Strncpy(UTF8CHAR* Dest, const UTF8CHAR* Src, SIZE_T MaxLen)
+	static UE_FORCEINLINE_HINT UTF8CHAR* Strncpy(UTF8CHAR* Dest, const UTF8CHAR* Src, SIZE_T MaxLen)
 	{
 		return (UTF8CHAR*)Strncpy((ANSICHAR*)Dest, (const ANSICHAR*)Src, MaxLen);
 	}
 
-	static FORCEINLINE UTF8CHAR* Strcat(UTF8CHAR* Dest, const UTF8CHAR* Src)
+	static UE_FORCEINLINE_HINT UTF8CHAR* Strcat(UTF8CHAR* Dest, const UTF8CHAR* Src)
 	{
 		return (UTF8CHAR*)Strcat((ANSICHAR*)Dest, (const ANSICHAR*)Src);
 	}
@@ -367,87 +367,87 @@ struct FMicrosoftPlatformString :
 		PRAGMA_ENABLE_DEPRECATION_WARNINGS
 	}
 
-	static FORCEINLINE UTF8CHAR* Strncat(UTF8CHAR* Dest, const UTF8CHAR* Src, SIZE_T SrcLen)
+	static UE_FORCEINLINE_HINT UTF8CHAR* Strncat(UTF8CHAR* Dest, const UTF8CHAR* Src, SIZE_T SrcLen)
 	{
 		return (UTF8CHAR*)Strncat((ANSICHAR*)Dest, (const ANSICHAR*)Src, SrcLen);
 	}
 
-	static FORCEINLINE int32 Strcmp(const UTF8CHAR* String1, const UTF8CHAR* String2)
+	static UE_FORCEINLINE_HINT int32 Strcmp(const UTF8CHAR* String1, const UTF8CHAR* String2)
 	{
 		return Strcmp((const ANSICHAR*)String1, (const ANSICHAR*)String2);
 	}
 
-	static FORCEINLINE int32 Strncmp(const UTF8CHAR* String1, const UTF8CHAR* String2, SIZE_T Count)
+	static UE_FORCEINLINE_HINT int32 Strncmp(const UTF8CHAR* String1, const UTF8CHAR* String2, SIZE_T Count)
 	{
 		return Strncmp((const ANSICHAR*)String1, (const ANSICHAR*)String2, Count);
 	}
 	
-	static FORCEINLINE int32 Strlen(const UTF8CHAR* String)
+	static UE_FORCEINLINE_HINT int32 Strlen(const UTF8CHAR* String)
 	{
 		return Strlen((const ANSICHAR*)String); 
 	}
 
-	static FORCEINLINE int32 Strnlen(const UTF8CHAR* String, SIZE_T StringSize)
+	static UE_FORCEINLINE_HINT int32 Strnlen(const UTF8CHAR* String, SIZE_T StringSize)
 	{
 		return Strnlen((const ANSICHAR*)String, StringSize);
 	}
 
-	static FORCEINLINE const UTF8CHAR* Strstr(const UTF8CHAR* String, const UTF8CHAR* Find)
+	static UE_FORCEINLINE_HINT const UTF8CHAR* Strstr(const UTF8CHAR* String, const UTF8CHAR* Find)
 	{
 		return (const UTF8CHAR*)Strstr((const ANSICHAR*)String, (const ANSICHAR*)Find);
 	}
 
-	static FORCEINLINE const UTF8CHAR* Strchr(const UTF8CHAR* String, UTF8CHAR C)
+	static UE_FORCEINLINE_HINT const UTF8CHAR* Strchr(const UTF8CHAR* String, UTF8CHAR C)
 	{
 		return (const UTF8CHAR*)Strchr((const ANSICHAR*)String, (ANSICHAR)C);
 	}
 
-	static FORCEINLINE const UTF8CHAR* Strrchr(const UTF8CHAR* String, UTF8CHAR C)
+	static UE_FORCEINLINE_HINT const UTF8CHAR* Strrchr(const UTF8CHAR* String, UTF8CHAR C)
 	{
 		return (const UTF8CHAR*)Strrchr((const ANSICHAR*)String, (ANSICHAR)C);
 	}
 
-	static FORCEINLINE int32 Atoi(const UTF8CHAR* String)
+	static UE_FORCEINLINE_HINT int32 Atoi(const UTF8CHAR* String)
 	{
 		return Atoi((const ANSICHAR*)String);
 	}
 
-	static FORCEINLINE int64 Atoi64(const UTF8CHAR* String)
+	static UE_FORCEINLINE_HINT int64 Atoi64(const UTF8CHAR* String)
 	{
 		return Atoi64((const ANSICHAR*)String);
 	}
 
-	static FORCEINLINE float Atof(const UTF8CHAR* String)
+	static UE_FORCEINLINE_HINT float Atof(const UTF8CHAR* String)
 	{
 		return Atof((const ANSICHAR*)String);
 	}
 
-	static FORCEINLINE double Atod(const UTF8CHAR* String)
+	static UE_FORCEINLINE_HINT double Atod(const UTF8CHAR* String)
 	{
 		return Atod((const ANSICHAR*)String);
 	}
 
-	static FORCEINLINE int32 Strtoi(const UTF8CHAR* Start, UTF8CHAR** End, int32 Base)
+	static UE_FORCEINLINE_HINT int32 Strtoi(const UTF8CHAR* Start, UTF8CHAR** End, int32 Base)
 	{
 		return Strtoi((const ANSICHAR*)Start, (ANSICHAR**)End, Base);
 	}
 
-	static FORCEINLINE int64 Strtoi64(const UTF8CHAR* Start, UTF8CHAR** End, int32 Base)
+	static UE_FORCEINLINE_HINT int64 Strtoi64(const UTF8CHAR* Start, UTF8CHAR** End, int32 Base)
 	{
 		return Strtoi64((const ANSICHAR*)Start, (ANSICHAR**)End, Base);
 	}
 
-	static FORCEINLINE uint64 Strtoui64(const UTF8CHAR* Start, UTF8CHAR** End, int32 Base)
+	static UE_FORCEINLINE_HINT uint64 Strtoui64(const UTF8CHAR* Start, UTF8CHAR** End, int32 Base)
 	{
 		return Strtoui64((const ANSICHAR*)Start, (ANSICHAR**)End, Base);
 	}
 
-	static FORCEINLINE UTF8CHAR* Strtok(UTF8CHAR* StrToken, const UTF8CHAR* Delim, UTF8CHAR** Context)
+	static UE_FORCEINLINE_HINT UTF8CHAR* Strtok(UTF8CHAR* StrToken, const UTF8CHAR* Delim, UTF8CHAR** Context)
 	{
 		return (UTF8CHAR*)Strtok((ANSICHAR*)StrToken, (const ANSICHAR*)Delim, (ANSICHAR**)Context);
 	}
 
-	static FORCEINLINE int32 GetVarArgs(UTF8CHAR* Dest, SIZE_T DestSize, const UTF8CHAR*& Fmt, va_list ArgPtr)
+	static UE_FORCEINLINE_HINT int32 GetVarArgs(UTF8CHAR* Dest, SIZE_T DestSize, const UTF8CHAR*& Fmt, va_list ArgPtr)
 	{
 		return GetVarArgs((ANSICHAR*)Dest, DestSize, *(const ANSICHAR**)&Fmt, ArgPtr);
 	}
diff --git a/Engine/Source/Runtime/Core/Public/Microsoft/MinimalWindowsApi.h b/Engine/Source/Runtime/Core/Public/Microsoft/MinimalWindowsApi.h
--- a/Engine/Source/Runtime/Core/Public/Microsoft/MinimalWindowsApi.h
+++ b/Engine/Source/Runtime/Core/Public/Microsoft/MinimalWindowsApi.h
@@ -128,37 +128,37 @@ namespace Windows
 	MINIMAL_WINDOWS_API BOOLEAN WINAPI TryAcquireSRWLockExclusive(PSRWLOCK SRWLock);
 	MINIMAL_WINDOWS_API BOOLEAN WINAPI TryAcquireSRWLockShared(PSRWLOCK SRWLock);
 
-	FORCEINLINE void WINAPI InitializeSRWLock(SRWLOCK* SRWLock)
+	UE_FORCEINLINE_HINT void WINAPI InitializeSRWLock(SRWLOCK* SRWLock)
 	{
 		InitializeSRWLock((PSRWLOCK)SRWLock);
 	}
 
-	FORCEINLINE BOOLEAN WINAPI TryAcquireSRWLockExclusive(SRWLOCK* SRWLock)
+	UE_FORCEINLINE_HINT BOOLEAN WINAPI TryAcquireSRWLockExclusive(SRWLOCK* SRWLock)
 	{
 		return TryAcquireSRWLockExclusive((PSRWLOCK)SRWLock);
 	}
 
-	FORCEINLINE BOOLEAN WINAPI TryAcquireSRWLockShared(SRWLOCK* SRWLock)
+	UE_FORCEINLINE_HINT BOOLEAN WINAPI TryAcquireSRWLockShared(SRWLOCK* SRWLock)
 	{
 		return TryAcquireSRWLockShared((PSRWLOCK)SRWLock);
 	}
 	
-	FORCEINLINE void WINAPI AcquireSRWLockShared(SRWLOCK* SRWLock)
+	UE_FORCEINLINE_HINT void WINAPI AcquireSRWLockShared(SRWLOCK* SRWLock)
 	{
 		AcquireSRWLockShared((PSRWLOCK)SRWLock);
 	}
 
-	FORCEINLINE void WINAPI ReleaseSRWLockShared(SRWLOCK* SRWLock)
+	UE_FORCEINLINE_HINT void WINAPI ReleaseSRWLockShared(SRWLOCK* SRWLock)
 	{
 		ReleaseSRWLockShared((PSRWLOCK)SRWLock);
 	}
 
-	FORCEINLINE void WINAPI AcquireSRWLockExclusive(SRWLOCK* SRWLock)
+	UE_FORCEINLINE_HINT void WINAPI AcquireSRWLockExclusive(SRWLOCK* SRWLock)
 	{
 		AcquireSRWLockExclusive((PSRWLOCK)SRWLock);
 	}
 
-	FORCEINLINE void WINAPI ReleaseSRWLockExclusive(SRWLOCK* SRWLock)
+	UE_FORCEINLINE_HINT void WINAPI ReleaseSRWLockExclusive(SRWLOCK* SRWLock)
 	{
 		ReleaseSRWLockExclusive((PSRWLOCK)SRWLock);
 	}
@@ -193,62 +193,62 @@ namespace Windows
 	struct OVERLAPPED { void *Opaque[3]; unsigned long Opaque2[2]; };
 	union LARGE_INTEGER { struct { DWORD LowPart; LONG HighPart; }; LONGLONG QuadPart; };
 
-	FORCEINLINE BOOL WINAPI ConnectNamedPipe(HANDLE hNamedPipe, OVERLAPPED* lpOverlapped)
+	UE_FORCEINLINE_HINT BOOL WINAPI ConnectNamedPipe(HANDLE hNamedPipe, OVERLAPPED* lpOverlapped)
 	{
 		return ConnectNamedPipe(hNamedPipe, (LPOVERLAPPED)lpOverlapped);
 	}
 
-	FORCEINLINE BOOL WINAPI GetOverlappedResult(HANDLE hFile, OVERLAPPED* lpOverlapped, LPDWORD lpNumberOfBytesTransferred, BOOL bWait)
+	UE_FORCEINLINE_HINT BOOL WINAPI GetOverlappedResult(HANDLE hFile, OVERLAPPED* lpOverlapped, LPDWORD lpNumberOfBytesTransferred, BOOL bWait)
 	{
 		return GetOverlappedResult(hFile, (LPOVERLAPPED)lpOverlapped, lpNumberOfBytesTransferred, bWait);
 	}
 
-	FORCEINLINE BOOL WINAPI WriteFile(HANDLE hFile, LPCVOID lpBuffer, DWORD nNumberOfBytesToWrite, LPDWORD lpNumberOfBytesWritten, OVERLAPPED* lpOverlapped)
+	UE_FORCEINLINE_HINT BOOL WINAPI WriteFile(HANDLE hFile, LPCVOID lpBuffer, DWORD nNumberOfBytesToWrite, LPDWORD lpNumberOfBytesWritten, OVERLAPPED* lpOverlapped)
 	{
 		return WriteFile(hFile, lpBuffer, nNumberOfBytesToWrite, lpNumberOfBytesWritten, (LPOVERLAPPED)lpOverlapped);
 	}
 
-	FORCEINLINE BOOL WINAPI ReadFile(HANDLE hFile, LPVOID lpBuffer, DWORD nNumberOfBytesToRead, LPDWORD lpNumberOfBytesRead, OVERLAPPED* lpOverlapped)
+	UE_FORCEINLINE_HINT BOOL WINAPI ReadFile(HANDLE hFile, LPVOID lpBuffer, DWORD nNumberOfBytesToRead, LPDWORD lpNumberOfBytesRead, OVERLAPPED* lpOverlapped)
 	{
 		return ReadFile(hFile, lpBuffer, nNumberOfBytesToRead, lpNumberOfBytesRead, (LPOVERLAPPED)lpOverlapped);
 	}
 
-	FORCEINLINE void WINAPI InitializeCriticalSection(CRITICAL_SECTION* lpCriticalSection)
+	UE_FORCEINLINE_HINT void WINAPI InitializeCriticalSection(CRITICAL_SECTION* lpCriticalSection)
 	{
 		InitializeCriticalSection((LPCRITICAL_SECTION)lpCriticalSection);
 	}
 
-	FORCEINLINE BOOL WINAPI InitializeCriticalSectionAndSpinCount(CRITICAL_SECTION* lpCriticalSection, DWORD dwSpinCount)
+	UE_FORCEINLINE_HINT BOOL WINAPI InitializeCriticalSectionAndSpinCount(CRITICAL_SECTION* lpCriticalSection, DWORD dwSpinCount)
 	{
 		return InitializeCriticalSectionAndSpinCount((LPCRITICAL_SECTION)lpCriticalSection, dwSpinCount);
 	}
 
-	FORCEINLINE DWORD WINAPI SetCriticalSectionSpinCount(CRITICAL_SECTION* lpCriticalSection, DWORD dwSpinCount)
+	UE_FORCEINLINE_HINT DWORD WINAPI SetCriticalSectionSpinCount(CRITICAL_SECTION* lpCriticalSection, DWORD dwSpinCount)
 	{
 		return SetCriticalSectionSpinCount((LPCRITICAL_SECTION)lpCriticalSection, dwSpinCount);
 	}
 
-	FORCEINLINE BOOL WINAPI TryEnterCriticalSection(CRITICAL_SECTION* lpCriticalSection)
+	UE_FORCEINLINE_HINT BOOL WINAPI TryEnterCriticalSection(CRITICAL_SECTION* lpCriticalSection)
 	{
 		return TryEnterCriticalSection((LPCRITICAL_SECTION)lpCriticalSection);
 	}
 
-	FORCEINLINE void WINAPI EnterCriticalSection(CRITICAL_SECTION* lpCriticalSection)
+	UE_FORCEINLINE_HINT void WINAPI EnterCriticalSection(CRITICAL_SECTION* lpCriticalSection)
 	{
 		EnterCriticalSection((LPCRITICAL_SECTION)lpCriticalSection);
 	}
 
-	FORCEINLINE void WINAPI LeaveCriticalSection(CRITICAL_SECTION* lpCriticalSection)
+	UE_FORCEINLINE_HINT void WINAPI LeaveCriticalSection(CRITICAL_SECTION* lpCriticalSection)
 	{
 		LeaveCriticalSection((LPCRITICAL_SECTION)lpCriticalSection);
 	}
 
-	FORCEINLINE void WINAPI DeleteCriticalSection(CRITICAL_SECTION* lpCriticalSection)
+	UE_FORCEINLINE_HINT void WINAPI DeleteCriticalSection(CRITICAL_SECTION* lpCriticalSection)
 	{
 		DeleteCriticalSection((LPCRITICAL_SECTION)lpCriticalSection);
 	}
 
-	FORCEINLINE BOOL WINAPI QueryPerformanceCounter(LARGE_INTEGER* Cycles)
+	UE_FORCEINLINE_HINT BOOL WINAPI QueryPerformanceCounter(LARGE_INTEGER* Cycles)
 	{
 		return QueryPerformanceCounter((LPLARGE_INTEGER)Cycles);
 	}
diff --git a/Engine/Source/Runtime/Core/Public/Misc/Attribute.h b/Engine/Source/Runtime/Core/Public/Misc/Attribute.h
--- a/Engine/Source/Runtime/Core/Public/Misc/Attribute.h
+++ b/Engine/Source/Runtime/Core/Public/Misc/Attribute.h
@@ -133,7 +133,7 @@ public:
 	 * Helper function for creating TAttributes from a const member function pointer, accessed through a raw pointer
 	 */
 	template<typename SourceType, typename SourceTypeOrBase, typename... PayloadTypes>
-	[[nodiscard]] FORCEINLINE static TAttribute CreateRaw(const SourceType* InObject, ObjectType (SourceTypeOrBase::*InMethod)(PayloadTypes...) const, typename TDecay<PayloadTypes>::Type... InputPayload)
+	[[nodiscard]] UE_FORCEINLINE_HINT static TAttribute CreateRaw(const SourceType* InObject, ObjectType (SourceTypeOrBase::*InMethod)(PayloadTypes...) const, typename TDecay<PayloadTypes>::Type... InputPayload)
 	{
 		return Create(FGetter::CreateRaw(InObject, InMethod, MoveTemp(InputPayload)...));
 	}
@@ -142,7 +142,7 @@ public:
 	 * Helper function for creating TAttributes from a const member function pointer, accessed through a weak pointer to the shared object
 	 */
 	template<typename SourceType, typename SourceTypeOrBase, typename... PayloadTypes>
-	[[nodiscard]] FORCEINLINE static TAttribute CreateSP(const SourceType* InObject, ObjectType (SourceTypeOrBase::*InMethod)(PayloadTypes...) const, typename TDecay<PayloadTypes>::Type... InputPayload)
+	[[nodiscard]] UE_FORCEINLINE_HINT static TAttribute CreateSP(const SourceType* InObject, ObjectType (SourceTypeOrBase::*InMethod)(PayloadTypes...) const, typename TDecay<PayloadTypes>::Type... InputPayload)
 	{
 		return Create(FGetter::CreateSP(InObject, InMethod, MoveTemp(InputPayload)...));
 	}
@@ -152,7 +152,7 @@ public:
 	 * TAttribute<float> FloatAttribute = TAttribute<float>::CreateLambda([]{ return 10.f; });
 	 */
 	template<typename LambdaType, typename... PayloadTypes>
-	[[nodiscard]] FORCEINLINE static TAttribute CreateLambda(LambdaType&& InCallable, PayloadTypes&&... InputPayload)
+	[[nodiscard]] UE_FORCEINLINE_HINT static TAttribute CreateLambda(LambdaType&& InCallable, PayloadTypes&&... InputPayload)
 	{
 		return Create(FGetter::CreateLambda(Forward<LambdaType>(InCallable), Forward<PayloadTypes>(InputPayload)...));
 	}
@@ -161,12 +161,12 @@ public:
 	 * Helper function for creating TAttributes from a weak shared pointer C++ lambda delegate
 	 */
 	template<typename UserClass, ESPMode Mode, typename LambdaType, typename... PayloadTypes>
-	[[nodiscard]] FORCEINLINE static TAttribute CreateSPLambda(const TSharedRef<UserClass, Mode>& InUserObjectRef, LambdaType&& InCallable, PayloadTypes&&... InputPayload)
+	[[nodiscard]] UE_FORCEINLINE_HINT static TAttribute CreateSPLambda(const TSharedRef<UserClass, Mode>& InUserObjectRef, LambdaType&& InCallable, PayloadTypes&&... InputPayload)
 	{
 		return Create(FGetter::CreateSPLambda(InUserObjectRef, Forward<LambdaType>(InCallable), Forward<PayloadTypes>(InputPayload)...));
 	}
 	template <typename UserClass, typename LambdaType, typename... PayloadTypes>
-	[[nodiscard]] FORCEINLINE static TAttribute CreateSPLambda(UserClass* InUserObject, LambdaType&& InCallable, PayloadTypes&&... InputPayload)
+	[[nodiscard]] UE_FORCEINLINE_HINT static TAttribute CreateSPLambda(UserClass* InUserObject, LambdaType&& InCallable, PayloadTypes&&... InputPayload)
 	{
 		return Create(FGetter::CreateSPLambda(InUserObject, Forward<LambdaType>(InCallable), Forward<PayloadTypes>(InputPayload)...));
 	}
@@ -175,7 +175,7 @@ public:
 	 * Helper function for creating TAttributes from a weak object pointer C++ lambda delegate
 	 */
 	template <typename UserClass, typename LambdaType, typename... PayloadTypes>
-	[[nodiscard]] FORCEINLINE static TAttribute CreateWeakLambda(UserClass* InUserObject, LambdaType&& InCallable, PayloadTypes&&... InputPayload)
+	[[nodiscard]] UE_FORCEINLINE_HINT static TAttribute CreateWeakLambda(UserClass* InUserObject, LambdaType&& InCallable, PayloadTypes&&... InputPayload)
 	{
 		return Create(FGetter::CreateWeakLambda(InUserObject, Forward<LambdaType>(InCallable), Forward<PayloadTypes>(InputPayload)...));
 	}
@@ -196,7 +196,7 @@ public:
 		return Attrib;
 	}
 
-	[[nodiscard]] FORCEINLINE static TAttribute< ObjectType > Create(TFunction<ObjectType(void)>&& InLambda)
+	[[nodiscard]] UE_FORCEINLINE_HINT static TAttribute< ObjectType > Create(TFunction<ObjectType(void)>&& InLambda)
 	{
 		return Create(TAttribute< ObjectType >::FGetter::CreateLambda(MoveTemp(InLambda)));
 	}
@@ -512,7 +512,7 @@ private:
  * Helper function for creating TAttributes from a const member function pointer, accessed through a raw pointer
  */
 template<typename T, typename SourceType, typename SourceTypeOrBase, typename... PayloadTypes>
-[[nodiscard]] FORCEINLINE TAttribute<T> MakeAttributeRaw(const SourceType* InObject, T (SourceTypeOrBase::*InMethod)(PayloadTypes...) const, typename TDecay<PayloadTypes>::Type... InputPayload)
+[[nodiscard]] UE_FORCEINLINE_HINT TAttribute<T> MakeAttributeRaw(const SourceType* InObject, T (SourceTypeOrBase::*InMethod)(PayloadTypes...) const, typename TDecay<PayloadTypes>::Type... InputPayload)
 {
 	return TAttribute<T>::Create(TAttribute<T>::FGetter::CreateRaw(InObject, InMethod, MoveTemp(InputPayload)...));
 }
@@ -521,18 +521,18 @@ template<typename T, typename SourceType, typename SourceTypeOrBase, typename...
  * Helper function for creating TAttributes from a const member function pointer, accessed through a weak pointer to the shared object
  */
 template<typename T, typename SourceType, typename SourceTypeOrBase, typename... PayloadTypes>
-[[nodiscard]] FORCEINLINE TAttribute<T> MakeAttributeSP(const SourceType* InObject, T (SourceTypeOrBase::*InMethod)(PayloadTypes...) const, typename TDecay<PayloadTypes>::Type...  InputPayload)
+[[nodiscard]] UE_FORCEINLINE_HINT TAttribute<T> MakeAttributeSP(const SourceType* InObject, T (SourceTypeOrBase::*InMethod)(PayloadTypes...) const, typename TDecay<PayloadTypes>::Type...  InputPayload)
 {
 	return TAttribute<T>::Create(TAttribute<T>::FGetter::CreateSP(InObject, InMethod, MoveTemp(InputPayload)...));
 }
 template<typename T, typename SourceType, typename SourceTypeOrBase, typename... PayloadTypes>
-[[nodiscard]] FORCEINLINE TAttribute<T> MakeAttributeSP(TSharedRef<SourceType> InObject, T (SourceTypeOrBase::*InMethod)(PayloadTypes...) const, typename TDecay<PayloadTypes>::Type...  InputPayload)
+[[nodiscard]] UE_FORCEINLINE_HINT TAttribute<T> MakeAttributeSP(TSharedRef<SourceType> InObject, T (SourceTypeOrBase::*InMethod)(PayloadTypes...) const, typename TDecay<PayloadTypes>::Type...  InputPayload)
 {
 	return TAttribute<T>::Create(TAttribute<T>::FGetter::CreateSP(InObject, InMethod, MoveTemp(InputPayload)...));
 }
 
 template<typename T, typename SourceType, typename SourceTypeOrBase, typename... PayloadTypes>
-[[nodiscard]] FORCEINLINE TAttribute<T> MakeAttributeUObject(const SourceType* InObject, T (SourceTypeOrBase::*InMethod)(PayloadTypes...) const, typename TDecay<PayloadTypes>::Type... InputPayload)
+[[nodiscard]] UE_FORCEINLINE_HINT TAttribute<T> MakeAttributeUObject(const SourceType* InObject, T (SourceTypeOrBase::*InMethod)(PayloadTypes...) const, typename TDecay<PayloadTypes>::Type... InputPayload)
 {
 	return TAttribute<T>::Create(TAttribute<T>::FGetter::CreateUObject(InObject, InMethod, MoveTemp(InputPayload)...));
 }
diff --git a/Engine/Source/Runtime/Core/Public/Misc/AutomationTest.h b/Engine/Source/Runtime/Core/Public/Misc/AutomationTest.h
--- a/Engine/Source/Runtime/Core/Public/Misc/AutomationTest.h
+++ b/Engine/Source/Runtime/Core/Public/Misc/AutomationTest.h
@@ -712,7 +712,7 @@ struct FAutomationExpectedMessage
 
 };
 
-FORCEINLINE uint32 GetTypeHash(const FAutomationExpectedMessage& Object)
+UE_FORCEINLINE_HINT uint32 GetTypeHash(const FAutomationExpectedMessage& Object)
 {
 	return GetTypeHash(Object.MessagePatternString);
 }
diff --git a/Engine/Source/Runtime/Core/Public/Misc/ByteSwap.h b/Engine/Source/Runtime/Core/Public/Misc/ByteSwap.h
--- a/Engine/Source/Runtime/Core/Public/Misc/ByteSwap.h
+++ b/Engine/Source/Runtime/Core/Public/Misc/ByteSwap.h
@@ -12,12 +12,12 @@
 // Encapsulate Byte swapping generic versions for benchmarking purpose (compare the generic vs intrinsic performance: See ByteSwapTest.cpp).
 namespace Internal
 {
-	static FORCEINLINE uint16 ByteSwapGeneric16(uint16 Value)
+	static UE_FORCEINLINE_HINT uint16 ByteSwapGeneric16(uint16 Value)
 	{
 		return (BYTESWAP_ORDER16_unsigned(Value));
 	}
 
-	static FORCEINLINE uint32 ByteSwapGeneric32(uint32 Value)
+	static UE_FORCEINLINE_HINT uint32 ByteSwapGeneric32(uint32 Value)
 	{
 		return (BYTESWAP_ORDER32_unsigned(Value));
 	}
@@ -56,7 +56,7 @@ inline uint16 BYTESWAP_ORDER16(uint16 Val)
 #endif
 }
 
-FORCEINLINE int16 BYTESWAP_ORDER16(int16 Val)
+UE_FORCEINLINE_HINT int16 BYTESWAP_ORDER16(int16 Val)
 {
 	return static_cast<int16>(BYTESWAP_ORDER16(static_cast<uint16>(Val)));
 }
@@ -70,7 +70,7 @@ inline uint32 BYTESWAP_ORDER32(uint32 Val)
 #endif
 }
 
-FORCEINLINE int32 BYTESWAP_ORDER32(int32 val)
+UE_FORCEINLINE_HINT int32 BYTESWAP_ORDER32(int32 val)
 {
 	return static_cast<int32>(BYTESWAP_ORDER32(static_cast<uint32>(val)));
 }
@@ -84,7 +84,7 @@ inline uint64 BYTESWAP_ORDER64(uint64 Value)
 #endif
 }
 
-FORCEINLINE int64 BYTESWAP_ORDER64(int64 Value)
+UE_FORCEINLINE_HINT int64 BYTESWAP_ORDER64(int64 Value)
 {
 	return static_cast<int64>(BYTESWAP_ORDER64(static_cast<uint64>(Value)));
 }
@@ -103,7 +103,7 @@ inline double BYTESWAP_ORDERD(double val)
 
 #if PLATFORM_TCHAR_IS_UTF8CHAR
 
-FORCEINLINE void BYTESWAP_ORDER_TCHARARRAY(TCHAR* str)
+UE_FORCEINLINE_HINT void BYTESWAP_ORDER_TCHARARRAY(TCHAR* str)
 {
 	static_assert(sizeof(TCHAR) == 1, "Assuming TCHAR is 1 byte wide.");
 
diff --git a/Engine/Source/Runtime/Core/Public/Misc/CString.h b/Engine/Source/Runtime/Core/Public/Misc/CString.h
--- a/Engine/Source/Runtime/Core/Public/Misc/CString.h
+++ b/Engine/Source/Runtime/Core/Public/Misc/CString.h
@@ -173,12 +173,12 @@ struct TCString
 	 * @param Src - source string
 	 * @return destination string
 	 */
-	static FORCEINLINE CharType* Strcpy( CharType* Dest, const CharType* Src );
+	static UE_FORCEINLINE_HINT CharType* Strcpy( CharType* Dest, const CharType* Src );
 
 	// This version was deprecated because Strcpy taking a size field does not match the CRT strcpy, and on some platforms the size field was being ignored.
 	// If the memzeroing causes a performance problem, request a new function StrcpyTruncate.
 	UE_DEPRECATED(5.6, "Use Strncpy instead. Note that Strncpy has a behavior difference from Strcpy: it memzeroes the entire DestCount-sized buffer after the end of string.")
-	static FORCEINLINE CharType* Strcpy(CharType* Dest, SIZE_T DestCount, const CharType* Src);
+	static inline CharType* Strcpy(CharType* Dest, SIZE_T DestCount, const CharType* Src);
 
 	/** 
 	 * Copy a string with length checking. Behavior differs from strncpy in that last character is zeroed. 
@@ -188,7 +188,7 @@ struct TCString
 	 * @param MaxLen - max length of the buffer (including null-terminator)
 	 * @return pointer to resulting string buffer
 	 */
-	static FORCEINLINE CharType* Strncpy( CharType* Dest, const CharType* Src, SIZE_T MaxLen );
+	static inline CharType* Strncpy( CharType* Dest, const CharType* Src, SIZE_T MaxLen );
 
 	/**
 	 * strcpy wrapper
@@ -199,7 +199,7 @@ struct TCString
 	 * @return destination string
 	 */
 	template<SIZE_T DestCount>
-	static FORCEINLINE CharType* Strcpy( CharType (&Dest)[DestCount], const CharType* Src ) 
+	static UE_FORCEINLINE_HINT CharType* Strcpy( CharType (&Dest)[DestCount], const CharType* Src ) 
 	{
 		return Strncpy( Dest, Src, DestCount);
 	}
@@ -212,11 +212,11 @@ struct TCString
 	 * @param Src - source string
 	 * @return destination string
 	 */
-	static FORCEINLINE CharType* Strcat( CharType* Dest, const CharType* Src );
+	static UE_FORCEINLINE_HINT CharType* Strcat( CharType* Dest, const CharType* Src );
 
 	// This version was deprecated because Strcat taking a size field does not match the CRT strcat, and on some platforms the size field was being ignored.
 	UE_DEPRECATED(5.6, "Use StrncatTruncateDest instead.")
-	static FORCEINLINE CharType* Strcat(CharType* Dest, SIZE_T DestSize, const CharType* Src);
+	static inline CharType* Strcat(CharType* Dest, SIZE_T DestSize, const CharType* Src);
 
 	/**
 	 * strcat wrapper
@@ -227,7 +227,7 @@ struct TCString
 	 * @return destination string
 	 */
 	template<SIZE_T DestCount>
-	static FORCEINLINE CharType* Strcat( CharType (&Dest)[DestCount], const CharType* Src ) 
+	static UE_FORCEINLINE_HINT CharType* Strcat( CharType (&Dest)[DestCount], const CharType* Src ) 
 	{ 
 		return StrncatTruncateDest( Dest, DestCount, Src);
 	}
@@ -306,7 +306,7 @@ struct TCString
 	 * @param Destcount - size of Dest in characters
 	 * @return destination string
 	 */
-	static FORCEINLINE CharType* Strupr( CharType* Dest, SIZE_T DestCount );
+	static UE_FORCEINLINE_HINT CharType* Strupr( CharType* Dest, SIZE_T DestCount );
 
 	
 	/**
@@ -317,7 +317,7 @@ struct TCString
 	 * @return destination string
 	 */
 	template<SIZE_T DestCount>
-	static FORCEINLINE CharType* Strupr( CharType (&Dest)[DestCount] ) 
+	static UE_FORCEINLINE_HINT CharType* Strupr( CharType (&Dest)[DestCount] ) 
 	{
 		return Strupr( Dest, DestCount );
 	}
@@ -325,22 +325,22 @@ struct TCString
 	/**
 	 * strcmp wrapper
 	 **/
-	static FORCEINLINE int32 Strcmp( const CharType* String1, const CharType* String2 );
+	static UE_FORCEINLINE_HINT int32 Strcmp( const CharType* String1, const CharType* String2 );
 
 	/**
 	 * strncmp wrapper
 	 */
-	static FORCEINLINE int32 Strncmp( const CharType* String1, const CharType* String2, SIZE_T Count);
+	static UE_FORCEINLINE_HINT int32 Strncmp( const CharType* String1, const CharType* String2, SIZE_T Count);
 
 	/**
 	 * stricmp wrapper
 	 */
-	static FORCEINLINE int32 Stricmp( const CharType* String1, const CharType* String2 );
+	static UE_FORCEINLINE_HINT int32 Stricmp( const CharType* String1, const CharType* String2 );
 	
 	/**
 	 * strnicmp wrapper
 	 */
-	static FORCEINLINE int32 Strnicmp( const CharType* String1, const CharType* String2, SIZE_T Count );
+	static UE_FORCEINLINE_HINT int32 Strnicmp( const CharType* String1, const CharType* String2, SIZE_T Count );
 
 	/**
 	 * Returns a static string that is filled with a variable number of spaces
@@ -441,7 +441,7 @@ struct TCString
 	/**
 	 * strlen wrapper
 	 */
-	static FORCEINLINE int32 Strlen( const CharType* String );
+	static inline int32 Strlen( const CharType* String );
 
 	/**
 	 * Calculate the length of the string up to the given size.
@@ -449,61 +449,61 @@ struct TCString
 	 * @param StringSize The maximum number of characters to read from String.
 	 * @return Length The smaller of StringSize and the number of characters in String before a null character.
 	 */
-	static FORCEINLINE int32 Strnlen( const CharType* String, SIZE_T StringSize );
+	static UE_FORCEINLINE_HINT int32 Strnlen( const CharType* String, SIZE_T StringSize );
 
 	/**
 	 * strstr wrapper
 	 */
-	static FORCEINLINE const CharType* Strstr( const CharType* String, const CharType* Find );
-	static FORCEINLINE CharType* Strstr( CharType* String, const CharType* Find );
+	static UE_FORCEINLINE_HINT const CharType* Strstr( const CharType* String, const CharType* Find );
+	static UE_FORCEINLINE_HINT CharType* Strstr( CharType* String, const CharType* Find );
 
 	/**
 	 * strchr wrapper
 	 */
-	static FORCEINLINE const CharType* Strchr( const CharType* String, CharType c );
-	static FORCEINLINE CharType* Strchr( CharType* String, CharType c );
+	static UE_FORCEINLINE_HINT const CharType* Strchr( const CharType* String, CharType c );
+	static UE_FORCEINLINE_HINT CharType* Strchr( CharType* String, CharType c );
 
 	/**
 	 * strrchr wrapper
 	 */
-	static FORCEINLINE const CharType* Strrchr( const CharType* String, CharType c );
-	static FORCEINLINE CharType* Strrchr( CharType* String, CharType c );
+	static UE_FORCEINLINE_HINT const CharType* Strrchr( const CharType* String, CharType c );
+	static UE_FORCEINLINE_HINT CharType* Strrchr( CharType* String, CharType c );
 
 	/**
 	 * strrstr wrapper
 	 */
-	static FORCEINLINE const CharType* Strrstr( const CharType* String, const CharType* Find );
-	static FORCEINLINE CharType* Strrstr( CharType* String, const CharType* Find );
+	static UE_FORCEINLINE_HINT const CharType* Strrstr( const CharType* String, const CharType* Find );
+	static inline CharType* Strrstr( CharType* String, const CharType* Find );
 
 	/**
 	 * strspn wrapper
 	 */
-	static FORCEINLINE int32 Strspn( const CharType* String, const CharType* Mask );
+	static inline int32 Strspn( const CharType* String, const CharType* Mask );
 
 	/**
 	 * strcspn wrapper
 	 */
-	static FORCEINLINE int32 Strcspn( const CharType* String, const CharType* Mask );
+	static inline int32 Strcspn( const CharType* String, const CharType* Mask );
 
 	/**
 	 * atoi wrapper
 	 */
-	static FORCEINLINE int32 Atoi( const CharType* String );
+	static UE_FORCEINLINE_HINT int32 Atoi( const CharType* String );
 
 	/**
 	 * atoi64 wrapper
 	 */
-	static FORCEINLINE int64 Atoi64( const CharType* String );
+	static UE_FORCEINLINE_HINT int64 Atoi64( const CharType* String );
 	
 	/**
 	 * atof wrapper
 	 */
-	static FORCEINLINE float Atof( const CharType* String );
+	static UE_FORCEINLINE_HINT float Atof( const CharType* String );
 
 	/**
 	 * atod wrapper
 	 */
-	static FORCEINLINE double Atod( const CharType* String );
+	static UE_FORCEINLINE_HINT double Atod( const CharType* String );
 
 	/**
 	 * Converts a string into a boolean value
@@ -536,22 +536,22 @@ struct TCString
 	/**
 	 * strtoi wrapper
 	 */
-	static FORCEINLINE int32 Strtoi( const CharType* Start, CharType** End, int32 Base );
+	static UE_FORCEINLINE_HINT int32 Strtoi( const CharType* Start, CharType** End, int32 Base );
 
 	/**
 	 * strtoi wrapper
 	 */
-	static FORCEINLINE int64 Strtoi64( const CharType* Start, CharType** End, int32 Base );
+	static UE_FORCEINLINE_HINT int64 Strtoi64( const CharType* Start, CharType** End, int32 Base );
 
 	/**
 	 * strtoui wrapper
 	 */
-	static FORCEINLINE uint64 Strtoui64( const CharType* Start, CharType** End, int32 Base );
+	static UE_FORCEINLINE_HINT uint64 Strtoui64( const CharType* Start, CharType** End, int32 Base );
 
 	/**
 	 * strtok wrapper
 	 */
-	static FORCEINLINE CharType* Strtok( CharType* TokenString, const CharType* Delim, CharType** Context );
+	static UE_FORCEINLINE_HINT CharType* Strtok( CharType* TokenString, const CharType* Delim, CharType** Context );
 
 private:
 	static int32 VARARGS SprintfImpl(CharType* Dest, const CharType* Fmt, ...);
@@ -595,7 +595,7 @@ public:
 	 * @param Args - argument list
 	 * @return number of characters written or -1 if truncated
 	 */
-	static FORCEINLINE int32 GetVarArgs( CharType* Dest, SIZE_T DestSize, const CharType*& Fmt, va_list ArgPtr );
+	static UE_FORCEINLINE_HINT int32 GetVarArgs( CharType* Dest, SIZE_T DestSize, const CharType*& Fmt, va_list ArgPtr );
 };
 
 typedef TCString<TCHAR>    FCString;
@@ -974,7 +974,7 @@ const typename TCString<T>::CharType* TCString<T>::Strnstr(const CharType* Str,
 	return nullptr;
 }
 
-template <typename T> FORCEINLINE
+template <typename T> UE_FORCEINLINE_HINT
 typename TCString<T>::CharType* TCString<T>::Strcpy(CharType* Dest, const CharType* Src)
 {
 	return FPlatformString::Strcpy(Dest, Src);
@@ -995,7 +995,7 @@ typename TCString<T>::CharType* TCString<T>::Strncpy( CharType* Dest, const Char
 	return Dest;
 }
 
-template <typename T> FORCEINLINE
+template <typename T> UE_FORCEINLINE_HINT
 typename TCString<T>::CharType* TCString<T>::Strcat(CharType* Dest, const CharType* Src)
 {
 	return FPlatformString::Strcat(Dest, Src);
@@ -1009,31 +1009,31 @@ typename TCString<T>::CharType* TCString<T>::Strcat(CharType* Dest, SIZE_T DestC
 	PRAGMA_ENABLE_DEPRECATION_WARNINGS
 }
 
-template <typename T> FORCEINLINE
+template <typename T> UE_FORCEINLINE_HINT
 typename TCString<T>::CharType* TCString<T>::Strupr(CharType* Dest, SIZE_T DestCount)
 {
 	return FPlatformString::Strupr(Dest, DestCount);
 }
 
-template <typename T> FORCEINLINE
+template <typename T> UE_FORCEINLINE_HINT
 int32 TCString<T>::Strcmp( const CharType* String1, const CharType* String2 )
 {
 	return FPlatformString::Strcmp(String1, String2);
 }
 
-template <typename T> FORCEINLINE
+template <typename T> UE_FORCEINLINE_HINT
 int32 TCString<T>::Strncmp( const CharType* String1, const CharType* String2, SIZE_T Count)
 {
 	return FPlatformString::Strncmp(String1, String2, Count);
 }
 
-template <typename T> FORCEINLINE
+template <typename T> UE_FORCEINLINE_HINT
 int32 TCString<T>::Stricmp( const CharType* String1, const CharType* String2 )
 {
 	return FPlatformString::Stricmp(String1, String2);
 }
 
-template <typename T> FORCEINLINE
+template <typename T> UE_FORCEINLINE_HINT
 int32 TCString<T>::Strnicmp( const CharType* String1, const CharType* String2, SIZE_T Count )
 {
 	return FPlatformString::Strnicmp(String1, String2, Count);
@@ -1057,49 +1057,49 @@ UE_NODEBUG inline int32 TCString<T>::Strlen( const CharType* String )
 	}
 }
 
-template <typename T> FORCEINLINE
+template <typename T> UE_FORCEINLINE_HINT
 int32 TCString<T>::Strnlen( const CharType* String, SIZE_T StringSize ) 
 {
 	return FPlatformString::Strnlen(String, StringSize);
 }
 
-template <typename T> FORCEINLINE
+template <typename T> UE_FORCEINLINE_HINT
 const typename TCString<T>::CharType* TCString<T>::Strstr( const CharType* String, const CharType* Find )
 {
 	return FPlatformString::Strstr(String, Find);
 }
 
-template <typename T> FORCEINLINE
+template <typename T> UE_FORCEINLINE_HINT
 typename TCString<T>::CharType* TCString<T>::Strstr( CharType* String, const CharType* Find )
 {
 	return (CharType*)FPlatformString::Strstr(String, Find);
 }
 
-template <typename T> FORCEINLINE
+template <typename T> UE_FORCEINLINE_HINT
 const typename TCString<T>::CharType* TCString<T>::Strchr( const CharType* String, CharType c ) 
 { 
 	return FPlatformString::Strchr(String, c);
 }
 
-template <typename T> FORCEINLINE
+template <typename T> UE_FORCEINLINE_HINT
 typename TCString<T>::CharType* TCString<T>::Strchr( CharType* String, CharType c ) 
 { 
 	return (CharType*)FPlatformString::Strchr(String, c);
 }
 
-template <typename T> FORCEINLINE
+template <typename T> UE_FORCEINLINE_HINT
 const typename TCString<T>::CharType* TCString<T>::Strrchr( const CharType* String, CharType c )
 { 
 	return FPlatformString::Strrchr(String, c);
 }
 
-template <typename T> FORCEINLINE
+template <typename T> UE_FORCEINLINE_HINT
 typename TCString<T>::CharType* TCString<T>::Strrchr( CharType* String, CharType c )
 { 
 	return (CharType*)FPlatformString::Strrchr(String, c);
 }
 
-template <typename T> FORCEINLINE
+template <typename T> UE_FORCEINLINE_HINT
 const typename TCString<T>::CharType* TCString<T>::Strrstr( const CharType* String, const CharType* Find )
 {
 	return Strrstr((CharType*)String, Find);
@@ -1170,56 +1170,56 @@ int32 TCString<T>::Strcspn( const CharType* String, const CharType* Mask )
 	return UE_PTRDIFF_TO_INT32(StringIt - String);
 }
 
-template <typename T> FORCEINLINE 
+template <typename T> UE_FORCEINLINE_HINT 
 int32 TCString<T>::Atoi( const CharType* String ) 
 {
 	return FPlatformString::Atoi(String);
 }
 
-template <typename T> FORCEINLINE
+template <typename T> UE_FORCEINLINE_HINT
 int64 TCString<T>::Atoi64( const CharType* String )
 { 
 	return FPlatformString::Atoi64(String);
 }
 
-template <typename T> FORCEINLINE
+template <typename T> UE_FORCEINLINE_HINT
 float TCString<T>::Atof( const CharType* String )
 { 
 	return FPlatformString::Atof(String);
 }
 
-template <typename T> FORCEINLINE
+template <typename T> UE_FORCEINLINE_HINT
 double TCString<T>::Atod( const CharType* String )
 { 
 	return FPlatformString::Atod(String);
 }
 
-template <typename T> FORCEINLINE
+template <typename T> UE_FORCEINLINE_HINT
 int32 TCString<T>::Strtoi( const CharType* Start, CharType** End, int32 Base ) 
 { 
 	return FPlatformString::Strtoi(Start, End, Base);
 }
 
-template <typename T> FORCEINLINE
+template <typename T> UE_FORCEINLINE_HINT
 int64 TCString<T>::Strtoi64( const CharType* Start, CharType** End, int32 Base ) 
 { 
 	return FPlatformString::Strtoi64(Start, End, Base);
 }
 
-template <typename T> FORCEINLINE
+template <typename T> UE_FORCEINLINE_HINT
 uint64 TCString<T>::Strtoui64( const CharType* Start, CharType** End, int32 Base ) 
 { 
 	return FPlatformString::Strtoui64(Start, End, Base);
 }
 
 
-template <typename T> FORCEINLINE
+template <typename T> UE_FORCEINLINE_HINT
 typename TCString<T>::CharType* TCString<T>::Strtok( CharType* TokenString, const CharType* Delim, CharType** Context )
 { 
 	return FPlatformString::Strtok(TokenString, Delim, Context);
 }
 
-template <typename T> FORCEINLINE
+template <typename T> UE_FORCEINLINE_HINT
 int32 TCString<T>::GetVarArgs( CharType* Dest, SIZE_T DestSize, const CharType*& Fmt, va_list ArgPtr )
 {
 	return FPlatformString::GetVarArgs(Dest, DestSize, Fmt, ArgPtr);
diff --git a/Engine/Source/Runtime/Core/Public/Misc/Char.h b/Engine/Source/Runtime/Core/Public/Misc/Char.h
--- a/Engine/Source/Runtime/Core/Public/Misc/Char.h
+++ b/Engine/Source/Runtime/Core/Public/Misc/Char.h
@@ -349,7 +349,7 @@ struct TChar : TCharBase<CharType, sizeof(CharType)>
 	*
 	* Mainly needed for subtraction and addition.
 	*/
-	static constexpr FORCEINLINE uint32 ToUnsigned(CharType Char)
+	static constexpr UE_FORCEINLINE_HINT uint32 ToUnsigned(CharType Char)
 	{
 		return (typename TUnsignedIntType<sizeof(CharType)>::Type)Char;
 	}
diff --git a/Engine/Source/Runtime/Core/Public/Misc/ConfigCacheIni.h b/Engine/Source/Runtime/Core/Public/Misc/ConfigCacheIni.h
--- a/Engine/Source/Runtime/Core/Public/Misc/ConfigCacheIni.h
+++ b/Engine/Source/Runtime/Core/Public/Misc/ConfigCacheIni.h
@@ -368,7 +368,7 @@ public:
 #endif
 
 	/** Gets the expanded value (GetValue) without marking it as having been accessed for e.g. writing out to a ConfigFile to disk */
-	FORCEINLINE FString GetValueForWriting() const
+	UE_FORCEINLINE_HINT FString GetValueForWriting() const
 	{
 		return bExpandOnDemand ? ExpandValue(SavedValue) : SavedValue;
 	};
diff --git a/Engine/Source/Runtime/Core/Public/Misc/CoreMiscDefines.h b/Engine/Source/Runtime/Core/Public/Misc/CoreMiscDefines.h
--- a/Engine/Source/Runtime/Core/Public/Misc/CoreMiscDefines.h
+++ b/Engine/Source/Runtime/Core/Public/Misc/CoreMiscDefines.h
@@ -442,13 +442,13 @@ struct TStaticDeprecateExpression
 struct FPlatformUserId
 {
 	/** Sees if this is a valid user */
-	FORCEINLINE bool IsValid() const
+	UE_FORCEINLINE_HINT bool IsValid() const
 	{
 		return InternalId != INDEX_NONE;
 	}
 
 	/** Returns the internal id for debugging/etc */
-	FORCEINLINE int32 GetInternalId() const
+	UE_FORCEINLINE_HINT int32 GetInternalId() const
 	{
 		return InternalId;
 	}
@@ -461,24 +461,24 @@ struct FPlatformUserId
 		return IdToReturn;
 	}
 
-	FORCEINLINE bool operator==(const FPlatformUserId& Other) const
+	UE_FORCEINLINE_HINT bool operator==(const FPlatformUserId& Other) const
 	{
 		return InternalId == Other.InternalId;
 	}
 
-	FORCEINLINE bool operator!=(const FPlatformUserId& Other) const
+	UE_FORCEINLINE_HINT bool operator!=(const FPlatformUserId& Other) const
 	{
 		return InternalId != Other.InternalId;
 	}
 
-	FORCEINLINE friend uint32 GetTypeHash(const FPlatformUserId& UserId)
+	UE_FORCEINLINE_HINT friend uint32 GetTypeHash(const FPlatformUserId& UserId)
 	{
 		return UserId.InternalId;
 	}
 
 	// This should be deprecated when the online code uniformly handles FPlatformUserId */
 	// UE_DEPRECATED(5.x, "Implicit conversion to user index is deprecated, use FPlatformMisc::GetUserIndexForPlatformUser")
-	FORCEINLINE constexpr operator int32() const { return InternalId; }
+	UE_FORCEINLINE_HINT constexpr operator int32() const { return InternalId; }
 
 private:
 	/** Raw id, will be allocated by application layer */
@@ -503,48 +503,48 @@ struct FInputDeviceId
 		return IdToReturn;
 	}
 	
-	FORCEINLINE int32 GetId() const
+	UE_FORCEINLINE_HINT int32 GetId() const
 	{
 		return InternalId;
 	}
 
 	/** Sees if this is a valid input device */
-	FORCEINLINE bool IsValid() const
+	UE_FORCEINLINE_HINT bool IsValid() const
 	{
 		return InternalId >= 0;
 	}
 
-	FORCEINLINE bool operator==(const FInputDeviceId& Other) const
+	UE_FORCEINLINE_HINT bool operator==(const FInputDeviceId& Other) const
 	{
 		return InternalId == Other.InternalId;
 	}
 
-	FORCEINLINE bool operator!=(const FInputDeviceId& Other) const
+	UE_FORCEINLINE_HINT bool operator!=(const FInputDeviceId& Other) const
 	{
 		return InternalId != Other.InternalId;
 	}
 	
-	FORCEINLINE bool operator<(const FInputDeviceId& Other) const
+	UE_FORCEINLINE_HINT bool operator<(const FInputDeviceId& Other) const
 	{
 		return InternalId < Other.InternalId;
 	}
 
-	FORCEINLINE bool operator<=(const FInputDeviceId& Other) const
+	UE_FORCEINLINE_HINT bool operator<=(const FInputDeviceId& Other) const
 	{
 		return InternalId <= Other.InternalId;
 	}
 
-	FORCEINLINE bool operator>(const FInputDeviceId& Other) const
+	UE_FORCEINLINE_HINT bool operator>(const FInputDeviceId& Other) const
 	{
 		return InternalId > Other.InternalId;
 	}
 
-	FORCEINLINE bool operator>=(const FInputDeviceId& Other) const
+	UE_FORCEINLINE_HINT bool operator>=(const FInputDeviceId& Other) const
 	{
 		return InternalId >= Other.InternalId;
 	}
 
-	FORCEINLINE friend uint32 GetTypeHash(const FInputDeviceId& InputId)
+	UE_FORCEINLINE_HINT friend uint32 GetTypeHash(const FInputDeviceId& InputId)
 	{
 		return InputId.InternalId;
 	}
diff --git a/Engine/Source/Runtime/Core/Public/Misc/Crc.h b/Engine/Source/Runtime/Core/Public/Misc/Crc.h
--- a/Engine/Source/Runtime/Core/Public/Misc/Crc.h
+++ b/Engine/Source/Runtime/Core/Public/Misc/Crc.h
@@ -28,7 +28,7 @@ struct FCrc
 	/** Generates CRC hash of the memory area. Prefer Hash/xxhash.h for new code, CRCs are not very good and fairly slow. */
 	typedef uint32 (*MemCrc32Functor)( const void* Data, int32 Length, uint32 CRC );
 	static CORE_API MemCrc32Functor MemCrc32Func;
-	[[nodiscard]] static FORCEINLINE uint32 MemCrc32(const void* Data, int32 Length, uint32 CRC = 0)
+	[[nodiscard]] static UE_FORCEINLINE_HINT uint32 MemCrc32(const void* Data, int32 Length, uint32 CRC = 0)
 	{
 		return MemCrc32Func(Data, Length, CRC);
 	}
@@ -299,7 +299,7 @@ inline uint32 FCrc::Strihash_DEPRECATED(const int32 DataLen, const UTF8CHAR* Dat
  *     Logf(TEXT("Value: %s"), ToCStr(LexToString(Val)));
  * }
  */
-FORCEINLINE const TCHAR* ToCStr(const TCHAR* Ptr)
+UE_FORCEINLINE_HINT const TCHAR* ToCStr(const TCHAR* Ptr)
 {
 	return Ptr;
 }
@@ -318,19 +318,19 @@ struct TStringPointerSetKeyFuncs_DEPRECATED
 
 	static constexpr bool bAllowDuplicateKeys = bInAllowDuplicateKeys;
 
-	static FORCEINLINE KeyInitType GetSetKey(ElementInitType Element)
+	static UE_FORCEINLINE_HINT KeyInitType GetSetKey(ElementInitType Element)
 	{
 		return Element;
 	}
 
 	template <typename ComparableKey>
-	static FORCEINLINE bool Matches(KeyInitType A, const ComparableKey& B)
+	static UE_FORCEINLINE_HINT bool Matches(KeyInitType A, const ComparableKey& B)
 	{
 		return A == B;
 	}
 
 	template <typename ComparableKey = KeyInitType>
-	static FORCEINLINE uint32 GetKeyHash(ComparableKey Key)
+	static UE_FORCEINLINE_HINT uint32 GetKeyHash(ComparableKey Key)
 	{
 		return FCrc::Strihash_DEPRECATED(ToCStr(Key));
 	}
@@ -350,19 +350,19 @@ struct TStringPointerMapKeyFuncs_DEPRECATED
 
 	static constexpr bool bAllowDuplicateKeys = bInAllowDuplicateKeys;
 
-	static FORCEINLINE KeyInitType GetSetKey(ElementInitType Element)
+	static UE_FORCEINLINE_HINT KeyInitType GetSetKey(ElementInitType Element)
 	{
 		return Element.Key;
 	}
 
 	template <typename ComparableKey>
-	static FORCEINLINE bool Matches(KeyInitType A, const ComparableKey& B)
+	static UE_FORCEINLINE_HINT bool Matches(KeyInitType A, const ComparableKey& B)
 	{
 		return A == B;
 	}
 
 	template <typename ComparableKey = KeyInitType>
-	static FORCEINLINE uint32 GetKeyHash(ComparableKey Key)
+	static UE_FORCEINLINE_HINT uint32 GetKeyHash(ComparableKey Key)
 	{
 		return FCrc::Strihash_DEPRECATED(ToCStr(Key));
 	}
