Replace some FORCEINLINE usage in Core/CoreUObject with UE_FORCEINLINE_HINT.

diff --git a/Engine/Source/Runtime/Core/Public/Misc/EngineVersionBase.h b/Engine/Source/Runtime/Core/Public/Misc/EngineVersionBase.h
--- a/Engine/Source/Runtime/Core/Public/Misc/EngineVersionBase.h
+++ b/Engine/Source/Runtime/Core/Public/Misc/EngineVersionBase.h
@@ -39,19 +39,19 @@ public:
 	CORE_API uint32 GetChangelist() const;	
 
 	/** Returns the Major version number corresponding to this version. */
-	FORCEINLINE uint16 GetMajor() const
+	UE_FORCEINLINE_HINT uint16 GetMajor() const
 	{
 		return Major;
 	}
 
 	/** Returns the Minor version number corresponding to this version. */
-	FORCEINLINE uint16 GetMinor() const
+	UE_FORCEINLINE_HINT uint16 GetMinor() const
 	{
 		return Minor;
 	}
 
 	/** Returns the Patch version number corresponding to this version. */
-	FORCEINLINE uint16 GetPatch() const
+	UE_FORCEINLINE_HINT uint16 GetPatch() const
 	{
 		return Patch;
 	}
diff --git a/Engine/Source/Runtime/Core/Public/Misc/EnumRange.h b/Engine/Source/Runtime/Core/Public/Misc/EnumRange.h
--- a/Engine/Source/Runtime/Core/Public/Misc/EnumRange.h
+++ b/Engine/Source/Runtime/Core/Public/Misc/EnumRange.h
@@ -104,7 +104,7 @@ namespace NEnumRangePrivate
 	{
 		using IntType = std::underlying_type_t<EnumType>;
 
-		FORCEINLINE explicit TEnumContiguousIterator(IntType InValue)
+		UE_FORCEINLINE_HINT explicit TEnumContiguousIterator(IntType InValue)
 			: Value(InValue)
 		{
 		}
@@ -115,7 +115,7 @@ namespace NEnumRangePrivate
 			return *this;
 		}
 
-		FORCEINLINE EnumType operator*() const
+		UE_FORCEINLINE_HINT EnumType operator*() const
 		{
 			return (EnumType)Value;
 		}
@@ -123,7 +123,7 @@ namespace NEnumRangePrivate
 	private:
 		IntType Value;
 
-		FORCEINLINE friend bool operator!=(const TEnumContiguousIterator& Lhs, const TEnumContiguousIterator& Rhs)
+		UE_FORCEINLINE_HINT friend bool operator!=(const TEnumContiguousIterator& Lhs, const TEnumContiguousIterator& Rhs)
 		{
 			return Lhs.Value != Rhs.Value;
 		}
@@ -132,7 +132,7 @@ namespace NEnumRangePrivate
 	template <typename EnumType>
 	struct TEnumValueArrayIterator
 	{
-		FORCEINLINE explicit TEnumValueArrayIterator(const EnumType* InPtr)
+		UE_FORCEINLINE_HINT explicit TEnumValueArrayIterator(const EnumType* InPtr)
 			: Ptr(InPtr)
 		{
 		}
@@ -143,7 +143,7 @@ namespace NEnumRangePrivate
 			return *this;
 		}
 
-		FORCEINLINE EnumType operator*() const
+		UE_FORCEINLINE_HINT EnumType operator*() const
 		{
 			return *Ptr;
 		}
@@ -151,7 +151,7 @@ namespace NEnumRangePrivate
 	private:
 		const EnumType* Ptr;
 
-		FORCEINLINE friend bool operator!=(const TEnumValueArrayIterator& Lhs, const TEnumValueArrayIterator& Rhs)
+		UE_FORCEINLINE_HINT friend bool operator!=(const TEnumValueArrayIterator& Lhs, const TEnumValueArrayIterator& Rhs)
 		{
 			return Lhs.Ptr != Rhs.Ptr;
 		}
@@ -194,9 +194,9 @@ namespace NEnumRangePrivate
 	template <typename EnumType>
 	struct TEnumRange_Impl<EnumType, 1>
 	{
-		FORCEINLINE explicit TEnumRange_Impl() = default;
+		UE_FORCEINLINE_HINT explicit TEnumRange_Impl() = default;
 
-		FORCEINLINE explicit TEnumRange_Impl(const EnumType InFirstValue, const EnumType InLastValue) = delete;
+		UE_FORCEINLINE_HINT explicit TEnumRange_Impl(const EnumType InFirstValue, const EnumType InLastValue) = delete;
 
 		TEnumValueArrayIterator<EnumType> begin() const { return TEnumValueArrayIterator<EnumType>(TEnumRangeTraits<EnumType>::template GetPointer<void>(false)); }
 		TEnumValueArrayIterator<EnumType> end  () const { return TEnumValueArrayIterator<EnumType>(TEnumRangeTraits<EnumType>::template GetPointer<void>(true )); }
@@ -210,7 +210,7 @@ namespace UE::EnumFlags::Private
 	{
 		using IntType = std::underlying_type_t<EnumType>;
 
-		FORCEINLINE explicit TIterator(EnumType InFlags)
+		UE_FORCEINLINE_HINT explicit TIterator(EnumType InFlags)
 			: Flags(IntType(InFlags))
 		{
 		}
@@ -231,7 +231,7 @@ namespace UE::EnumFlags::Private
 	private:
 		IntType Flags;
 
-		FORCEINLINE friend bool operator!=(const TIterator& Lhs, const TIterator& Rhs)
+		UE_FORCEINLINE_HINT friend bool operator!=(const TIterator& Lhs, const TIterator& Rhs)
 		{
 			return Lhs.Flags != Rhs.Flags;
 		}
diff --git a/Engine/Source/Runtime/Core/Public/Misc/FeedbackContext.h b/Engine/Source/Runtime/Core/Public/Misc/FeedbackContext.h
--- a/Engine/Source/Runtime/Core/Public/Misc/FeedbackContext.h
+++ b/Engine/Source/Runtime/Core/Public/Misc/FeedbackContext.h
@@ -44,7 +44,7 @@ public:
 	virtual bool ReceivedUserCancel() { return false; }
 	
 	/** Public const access to the current state of the scope stack */
-	FORCEINLINE const FSlowTaskStack& GetScopeStack() const
+	UE_FORCEINLINE_HINT const FSlowTaskStack& GetScopeStack() const
 	{
 		return ScopeStack;
 	}
diff --git a/Engine/Source/Runtime/Core/Public/Misc/FilterCollection.h b/Engine/Source/Runtime/Core/Public/Misc/FilterCollection.h
--- a/Engine/Source/Runtime/Core/Public/Misc/FilterCollection.h
+++ b/Engine/Source/Runtime/Core/Public/Misc/FilterCollection.h
@@ -95,7 +95,7 @@ public:
 	}
 
 	/** Returns the number of Filters in the collection */
-	FORCEINLINE int32 Num() const
+	UE_FORCEINLINE_HINT int32 Num() const
 	{
 		return ChildFilters.Num();
 	}
diff --git a/Engine/Source/Runtime/Core/Public/Misc/FrameTime.h b/Engine/Source/Runtime/Core/Public/Misc/FrameTime.h
--- a/Engine/Source/Runtime/Core/Public/Misc/FrameTime.h
+++ b/Engine/Source/Runtime/Core/Public/Misc/FrameTime.h
@@ -53,7 +53,7 @@ public:
 	/**
 	 * Access this time's frame number
 	 */
-	FORCEINLINE FFrameNumber GetFrame() const
+	UE_FORCEINLINE_HINT FFrameNumber GetFrame() const
 	{
 		return FrameNumber;
 	}
@@ -61,7 +61,7 @@ public:
 	/**
 	 * Access this time's sub frame
 	 */
-	FORCEINLINE float GetSubFrame() const
+	UE_FORCEINLINE_HINT float GetSubFrame() const
 	{
 		return SubFrame;
 	}
@@ -117,37 +117,37 @@ public:
 		return Ar;
 	}
 
-	friend FORCEINLINE_DEBUGGABLE bool operator==(FFrameTime A, FFrameTime B)
+	friend UE_FORCEINLINE_HINT bool operator==(FFrameTime A, FFrameTime B)
 	{
 		return A.FrameNumber == B.FrameNumber && A.SubFrame == B.SubFrame;
 	}
 
 
-	friend FORCEINLINE_DEBUGGABLE bool operator!=(FFrameTime A, FFrameTime B)
+	friend UE_FORCEINLINE_HINT bool operator!=(FFrameTime A, FFrameTime B)
 	{
 		return A.FrameNumber != B.FrameNumber || A.SubFrame != B.SubFrame;
 	}
 
 
-	friend FORCEINLINE_DEBUGGABLE bool operator> (FFrameTime A, FFrameTime B)
+	friend UE_FORCEINLINE_HINT bool operator> (FFrameTime A, FFrameTime B)
 	{
 		return A.FrameNumber >  B.FrameNumber || ( A.FrameNumber == B.FrameNumber && A.SubFrame > B.SubFrame );
 	}
 
 
-	friend FORCEINLINE_DEBUGGABLE bool operator>=(FFrameTime A, FFrameTime B)
+	friend UE_FORCEINLINE_HINT bool operator>=(FFrameTime A, FFrameTime B)
 	{
 		return A.FrameNumber > B.FrameNumber || ( A.FrameNumber == B.FrameNumber && A.SubFrame >= B.SubFrame );
 	}
 
 
-	friend FORCEINLINE_DEBUGGABLE bool operator< (FFrameTime A, FFrameTime B)
+	friend UE_FORCEINLINE_HINT bool operator< (FFrameTime A, FFrameTime B)
 	{
 		return A.FrameNumber <  B.FrameNumber || ( A.FrameNumber == B.FrameNumber && A.SubFrame < B.SubFrame );
 	}
 
 
-	friend FORCEINLINE_DEBUGGABLE bool operator<=(FFrameTime A, FFrameTime B)
+	friend UE_FORCEINLINE_HINT bool operator<=(FFrameTime A, FFrameTime B)
 	{
 		return A.FrameNumber < B.FrameNumber || ( A.FrameNumber == B.FrameNumber && A.SubFrame <= B.SubFrame );
 	}
@@ -229,17 +229,17 @@ public:
 	}
 
 
-	friend FORCEINLINE FFrameTime operator*(FFrameTime A, double Scalar)
+	friend UE_FORCEINLINE_HINT FFrameTime operator*(FFrameTime A, double Scalar)
 	{
 		return FFrameTime::FromDecimal(A.AsDecimal() * Scalar);
 	}
 
-	friend FORCEINLINE FFrameTime operator*(double Scalar, FFrameTime A)
+	friend UE_FORCEINLINE_HINT FFrameTime operator*(double Scalar, FFrameTime A)
 	{
 		return FFrameTime::FromDecimal(A.AsDecimal() * Scalar);
 	}
 
-	friend FORCEINLINE FFrameTime operator/(FFrameTime A, double Scalar)
+	friend UE_FORCEINLINE_HINT FFrameTime operator/(FFrameTime A, double Scalar)
 	{
 		return FFrameTime::FromDecimal(A.AsDecimal() / Scalar);
 	}
@@ -287,25 +287,25 @@ inline FFrameTime& FFrameTime::operator=(FFrameNumber InFrameNumber)
 	return *this;
 }
 
-FORCEINLINE_DEBUGGABLE FFrameNumber FFrameTime::FloorToFrame() const
+UE_FORCEINLINE_HINT FFrameNumber FFrameTime::FloorToFrame() const
 {
 	return FrameNumber;
 }
 
 
-FORCEINLINE_DEBUGGABLE FFrameNumber FFrameTime::CeilToFrame() const
+UE_FORCEINLINE_HINT FFrameNumber FFrameTime::CeilToFrame() const
 {
 	return SubFrame == 0.f ? FrameNumber : FrameNumber+1;
 }
 
 
-FORCEINLINE_DEBUGGABLE FFrameNumber FFrameTime::RoundToFrame() const
+UE_FORCEINLINE_HINT FFrameNumber FFrameTime::RoundToFrame() const
 {
 	return SubFrame < .5f ? FrameNumber : FrameNumber+1;
 }
 
 
-FORCEINLINE_DEBUGGABLE double FFrameTime::AsDecimal() const
+UE_FORCEINLINE_HINT double FFrameTime::AsDecimal() const
 {
 	return double(FrameNumber.Value) + SubFrame;
 }
diff --git a/Engine/Source/Runtime/Core/Public/Misc/HashBuilder.h b/Engine/Source/Runtime/Core/Public/Misc/HashBuilder.h
--- a/Engine/Source/Runtime/Core/Public/Misc/HashBuilder.h
+++ b/Engine/Source/Runtime/Core/Public/Misc/HashBuilder.h
@@ -28,7 +28,7 @@ public:
 	}
 
 	template <typename T>
-	FORCEINLINE typename TEnableIf<TIsPODType<T>::Value, FHashBuilder&>::Type AppendRaw(const T& InData)
+	UE_FORCEINLINE_HINT typename TEnableIf<TIsPODType<T>::Value, FHashBuilder&>::Type AppendRaw(const T& InData)
 	{
 		return AppendRaw(&InData, sizeof(T));
 	}
@@ -79,12 +79,12 @@ public:
 	}
 
 	template <typename T>
-	FORCEINLINE FHashBuilder& operator<<(const T& InData)
+	UE_FORCEINLINE_HINT FHashBuilder& operator<<(const T& InData)
 	{
 		return Append(InData);
 	}
 
-	FORCEINLINE uint32 GetHash() const
+	UE_FORCEINLINE_HINT uint32 GetHash() const
 	{
 		return ~Hash;
 	}
diff --git a/Engine/Source/Runtime/Core/Public/Misc/InlineValue.h b/Engine/Source/Runtime/Core/Public/Misc/InlineValue.h
--- a/Engine/Source/Runtime/Core/Public/Misc/InlineValue.h
+++ b/Engine/Source/Runtime/Core/Public/Misc/InlineValue.h
@@ -137,7 +137,7 @@ public:
 	/**
 	 * Check if this container is wrapping a valid object
 	 */
-	FORCEINLINE bool IsValid() const
+	UE_FORCEINLINE_HINT bool IsValid() const
 	{
 		return bIsValid;
 	}
@@ -156,7 +156,7 @@ public:
 	 * Get the wrapped object, or a user-specified default
 	 * @return The object, or the user-specified default
 	 */
-	FORCEINLINE const BaseType& Get(const BaseType& Default) const
+	UE_FORCEINLINE_HINT const BaseType& Get(const BaseType& Default) const
 	{
 		return bIsValid ? GetValue() : Default;
 	}
@@ -165,13 +165,13 @@ public:
 	 * Get a pointer the wrapped object, or a user-specified default
 	 * @return A pointer to the object, or the user-specified default
 	 */
-	FORCEINLINE BaseType* GetPtr(BaseType* Default = nullptr) const
+	UE_FORCEINLINE_HINT BaseType* GetPtr(BaseType* Default = nullptr) const
 	{
 		return bIsValid ? &GetValue() : Default;
 	}
 
-	FORCEINLINE BaseType&		operator*() const	{ return GetValue(); }
-	FORCEINLINE BaseType*		operator->() const	{ return &GetValue(); }
+	UE_FORCEINLINE_HINT BaseType&		operator*() const	{ return GetValue(); }
+	UE_FORCEINLINE_HINT BaseType*		operator->() const	{ return &GetValue(); }
 
 	/**
 	 * Reserve space for a structure derived from BaseType, of the size and alignment specified .
diff --git a/Engine/Source/Runtime/Core/Public/Misc/LargeWorldRenderPosition.h b/Engine/Source/Runtime/Core/Public/Misc/LargeWorldRenderPosition.h
--- a/Engine/Source/Runtime/Core/Public/Misc/LargeWorldRenderPosition.h
+++ b/Engine/Source/Runtime/Core/Public/Misc/LargeWorldRenderPosition.h
@@ -20,7 +20,7 @@ private:
 	using VectorType = UE::Math::TVector<TScalar>;
 
 public:
-	FORCEINLINE static TScalar GetTileSize()
+	UE_FORCEINLINE_HINT static TScalar GetTileSize()
 	{
 		return static_cast<TScalar>(UE_LWC_RENDER_TILE_SIZE);
 	}
diff --git a/Engine/Source/Runtime/Core/Public/Misc/MTAccessDetector.h b/Engine/Source/Runtime/Core/Public/Misc/MTAccessDetector.h
--- a/Engine/Source/Runtime/Core/Public/Misc/MTAccessDetector.h
+++ b/Engine/Source/Runtime/Core/Public/Misc/MTAccessDetector.h
@@ -267,7 +267,7 @@ public:
 		AccessDetector.AcquireReadAccess();
 	}
 
-	FORCEINLINE ~TScopedReaderAccessDetector()
+	UE_FORCEINLINE_HINT ~TScopedReaderAccessDetector()
 	{
 		AccessDetector.ReleaseReadAccess();
 	}
@@ -276,7 +276,7 @@ private:
 };
 
 template<typename RWAccessDetector>
-FORCEINLINE TScopedReaderAccessDetector<RWAccessDetector> MakeScopedReaderAccessDetector(RWAccessDetector& InAccessDetector)
+UE_FORCEINLINE_HINT TScopedReaderAccessDetector<RWAccessDetector> MakeScopedReaderAccessDetector(RWAccessDetector& InAccessDetector)
 {
 	return TScopedReaderAccessDetector<RWAccessDetector>(InAccessDetector);
 }
@@ -291,7 +291,7 @@ public:
 		AccessDetector.AcquireWriteAccess();
 	}
 
-	FORCEINLINE ~TScopedWriterDetector()
+	UE_FORCEINLINE_HINT ~TScopedWriterDetector()
 	{
 		AccessDetector.ReleaseWriteAccess();
 	}
@@ -300,7 +300,7 @@ private:
 };
 
 template<typename RWAccessDetector>
-FORCEINLINE TScopedWriterDetector<RWAccessDetector> MakeScopedWriterAccessDetector(RWAccessDetector& InAccessDetector)
+UE_FORCEINLINE_HINT TScopedWriterDetector<RWAccessDetector> MakeScopedWriterAccessDetector(RWAccessDetector& InAccessDetector)
 {
 	return TScopedWriterDetector<RWAccessDetector>(InAccessDetector);
 }
@@ -359,12 +359,12 @@ private:
 	// all atomic ops are relaxed to preserve the original memory order, as the detector is compiled out in non-dev builds
 	mutable std::atomic<uint64> State{ 0 }; // it's actually FState, but we need to do `atomic::fetch_add` on it
 
-	FORCEINLINE FState LoadState() const
+	UE_FORCEINLINE_HINT FState LoadState() const
 	{
 		return FState(State.load(std::memory_order_relaxed));
 	}
 
-	FORCEINLINE FState ExchangeState(FState NewState)
+	UE_FORCEINLINE_HINT FState ExchangeState(FState NewState)
 	{
 		return FState(State.exchange(NewState.Value, std::memory_order_relaxed));
 	}
@@ -381,7 +381,7 @@ private:
 		return FState(State.fetch_sub(OneReader.Value, std::memory_order_relaxed));
 	}
 
-	FORCEINLINE friend bool operator==(FState L, FState R)
+	UE_FORCEINLINE_HINT friend bool operator==(FState L, FState R)
 	{
 		return L.Value == R.Value;
 	}
@@ -443,7 +443,7 @@ public:
 
 	struct FDestructionSentinel
 	{
-		FORCEINLINE explicit FDestructionSentinel(EAccessType InAccessType)
+		UE_FORCEINLINE_HINT explicit FDestructionSentinel(EAccessType InAccessType)
 			: AccessType(InAccessType)
 		{
 		}
@@ -501,7 +501,7 @@ private:
 public:
 	FMRSWRecursiveAccessDetector() = default;
 
-	FORCEINLINE FMRSWRecursiveAccessDetector(const FMRSWRecursiveAccessDetector& Other)
+	UE_FORCEINLINE_HINT FMRSWRecursiveAccessDetector(const FMRSWRecursiveAccessDetector& Other)
 		// just default initialization, the copy is not being accessed
 	{
 		CheckOtherThreadWriters(Other.LoadState());
diff --git a/Engine/Source/Runtime/Core/Public/Misc/MTTransactionallySafeAccessDetector.h b/Engine/Source/Runtime/Core/Public/Misc/MTTransactionallySafeAccessDetector.h
--- a/Engine/Source/Runtime/Core/Public/Misc/MTTransactionallySafeAccessDetector.h
+++ b/Engine/Source/Runtime/Core/Public/Misc/MTTransactionallySafeAccessDetector.h
@@ -233,14 +233,14 @@ private:
 		}
 	}
 
-	FORCEINLINE void RegisterOnAbortRelease() const
+	UE_FORCEINLINE_HINT void RegisterOnAbortRelease() const
 	{
 		// We explicitly copy the state here for the case that `this` was stack
 		// allocated and has already died before the on-abort is hit.
 		AutoRTFM::OnAbort([State = this->State] { ReleaseAccess(State); });
 	}
 
-	FORCEINLINE void RegisterOnCommitRelease() const
+	UE_FORCEINLINE_HINT void RegisterOnCommitRelease() const
 	{
 		// We explicitly copy the state here for the case that `this` was stack
 		// allocated and has already died before the on-abort is hit.
diff --git a/Engine/Source/Runtime/Core/Public/Misc/MemStack.h b/Engine/Source/Runtime/Core/Public/Misc/MemStack.h
--- a/Engine/Source/Runtime/Core/Public/Misc/MemStack.h
+++ b/Engine/Source/Runtime/Core/Public/Misc/MemStack.h
@@ -117,7 +117,7 @@ public:
 		FreeChunks(nullptr);
 	}
 
-	FORCEINLINE uint8* PushBytes(size_t AllocSize, size_t Alignment)
+	UE_FORCEINLINE_HINT uint8* PushBytes(size_t AllocSize, size_t Alignment)
 	{
 		return (uint8*)Alloc(AllocSize, FMath::Max(AllocSize >= 16 ? (size_t)16 : (size_t)8, Alignment));
 	}
@@ -157,13 +157,13 @@ public:
 		return Result;
 	}
 
-	FORCEINLINE uint8* GetTop() const
+	UE_FORCEINLINE_HINT uint8* GetTop() const
 	{
 		return Top;
 	}
 
 	/** return true if this stack is empty. */
-	FORCEINLINE bool IsEmpty() const
+	UE_FORCEINLINE_HINT bool IsEmpty() const
 	{
 		return TopChunk == nullptr;
 	}
@@ -173,7 +173,7 @@ public:
 		check(!NumMarks);
 		FreeChunks(nullptr);
 	}
-	FORCEINLINE int32 GetNumMarks()
+	UE_FORCEINLINE_HINT int32 GetNumMarks()
 	{
 		return NumMarks;
 	}
@@ -419,7 +419,7 @@ public:
 		}
 
 		// FContainerAllocatorInterface
-		FORCEINLINE ElementType* GetAllocation() const
+		UE_FORCEINLINE_HINT ElementType* GetAllocation() const
 		{
 			return Data;
 		}
@@ -452,20 +452,20 @@ public:
 				}
 			}
 		}
-		FORCEINLINE SizeType CalculateSlackReserve(SizeType NewMax, SIZE_T NumBytesPerElement) const
+		UE_FORCEINLINE_HINT SizeType CalculateSlackReserve(SizeType NewMax, SIZE_T NumBytesPerElement) const
 		{
 			return DefaultCalculateSlackReserve(NewMax, NumBytesPerElement, false, Alignment);
 		}
-		FORCEINLINE SizeType CalculateSlackShrink(SizeType NewMax, SizeType CurrentMax, SIZE_T NumBytesPerElement) const
+		UE_FORCEINLINE_HINT SizeType CalculateSlackShrink(SizeType NewMax, SizeType CurrentMax, SIZE_T NumBytesPerElement) const
 		{
 			return DefaultCalculateSlackShrink(NewMax, CurrentMax, NumBytesPerElement, false, Alignment);
 		}
-		FORCEINLINE SizeType CalculateSlackGrow(SizeType NewMax, SizeType CurrentMax, SIZE_T NumBytesPerElement) const
+		UE_FORCEINLINE_HINT SizeType CalculateSlackGrow(SizeType NewMax, SizeType CurrentMax, SIZE_T NumBytesPerElement) const
 		{
 			return DefaultCalculateSlackGrow(NewMax, CurrentMax, NumBytesPerElement, false, Alignment);
 		}
 
-		FORCEINLINE SIZE_T GetAllocatedSize(SizeType CurrentMax, SIZE_T NumBytesPerElement) const
+		UE_FORCEINLINE_HINT SIZE_T GetAllocatedSize(SizeType CurrentMax, SIZE_T NumBytesPerElement) const
 		{
 			return CurrentMax * NumBytesPerElement;
 		}
diff --git a/Engine/Source/Runtime/Core/Public/Misc/MemoryReadStream.h b/Engine/Source/Runtime/Core/Public/Misc/MemoryReadStream.h
--- a/Engine/Source/Runtime/Core/Public/Misc/MemoryReadStream.h
+++ b/Engine/Source/Runtime/Core/Public/Misc/MemoryReadStream.h
@@ -22,7 +22,7 @@ public:
 	CORE_API virtual void CopyTo(void* Buffer, int64 InOffset, int64 InSize);
 	virtual void EnsureReadNonBlocking() {};		// Incur any potential wait such that future Read() or CopyTo() calls are wait and thread switch free.
 
-	FORCEINLINE uint32 AddRef() const { return uint32(NumRefs.Increment()); }
+	UE_FORCEINLINE_HINT uint32 AddRef() const { return uint32(NumRefs.Increment()); }
 
 	inline uint32 Release() const
 	{
@@ -34,7 +34,7 @@ public:
 		return uint32(Refs);
 	}
 
-	FORCEINLINE uint32 GetRefCount() const
+	UE_FORCEINLINE_HINT uint32 GetRefCount() const
 	{
 		return uint32(NumRefs.GetValue());
 	}
diff --git a/Engine/Source/Runtime/Core/Public/Misc/MusicalTime.h b/Engine/Source/Runtime/Core/Public/Misc/MusicalTime.h
--- a/Engine/Source/Runtime/Core/Public/Misc/MusicalTime.h
+++ b/Engine/Source/Runtime/Core/Public/Misc/MusicalTime.h
@@ -95,37 +95,37 @@ public:
 		return Ar;
 	}
 
-	friend FORCEINLINE_DEBUGGABLE bool operator==(const FMusicalTime& A, const FMusicalTime& B)
+	friend UE_FORCEINLINE_HINT bool operator==(const FMusicalTime& A, const FMusicalTime& B)
 	{
 		return A.Bar == B.Bar && A.TickInBar == B.TickInBar;
 	}
 
 
-	friend FORCEINLINE_DEBUGGABLE bool operator!=(const FMusicalTime& A, const FMusicalTime& B)
+	friend UE_FORCEINLINE_HINT bool operator!=(const FMusicalTime& A, const FMusicalTime& B)
 	{
 		return A.Bar != B.Bar || A.TickInBar != B.TickInBar;
 	}
 
 
-	friend FORCEINLINE_DEBUGGABLE bool operator> (const FMusicalTime& A, const FMusicalTime& B)
+	friend UE_FORCEINLINE_HINT bool operator> (const FMusicalTime& A, const FMusicalTime& B)
 	{
 		return A.Bar >  B.Bar || ( A.Bar == B.Bar && A.TickInBar > B.TickInBar );
 	}
 
 
-	friend FORCEINLINE_DEBUGGABLE bool operator>=(const FMusicalTime& A, const FMusicalTime& B)
+	friend UE_FORCEINLINE_HINT bool operator>=(const FMusicalTime& A, const FMusicalTime& B)
 	{
 		return A.Bar > B.Bar || ( A.Bar == B.Bar && A.TickInBar >= B.TickInBar );
 	}
 
 
-	friend FORCEINLINE_DEBUGGABLE bool operator< (const FMusicalTime& A, const FMusicalTime& B)
+	friend UE_FORCEINLINE_HINT bool operator< (const FMusicalTime& A, const FMusicalTime& B)
 	{
 		return A.Bar <  B.Bar || ( A.Bar == B.Bar && A.TickInBar < B.TickInBar );
 	}
 
 
-	friend FORCEINLINE_DEBUGGABLE bool operator<=(const FMusicalTime& A, const FMusicalTime& B)
+	friend UE_FORCEINLINE_HINT bool operator<=(const FMusicalTime& A, const FMusicalTime& B)
 	{
 		return A.Bar < B.Bar || ( A.Bar == B.Bar && A.TickInBar <= B.TickInBar );
 	}
diff --git a/Engine/Source/Runtime/Core/Public/Misc/Optional.h b/Engine/Source/Runtime/Core/Public/Misc/Optional.h
--- a/Engine/Source/Runtime/Core/Public/Misc/Optional.h
+++ b/Engine/Source/Runtime/Core/Public/Misc/Optional.h
@@ -70,7 +70,7 @@ namespace UE::Core::Private
 			return FOptional::IsSet(this);
 		}
 
-		FORCEINLINE constexpr void DestroyValue()
+		UE_FORCEINLINE_HINT constexpr void DestroyValue()
 		{
 			// Type is known to be trivially destructible
 		}
@@ -110,7 +110,7 @@ namespace UE::Core::Private
 			return FOptional::IsSet(this);
 		}
 
-		FORCEINLINE constexpr void DestroyValue()
+		UE_FORCEINLINE_HINT constexpr void DestroyValue()
 		{
 			DestructItem(std::addressof(this->TypedValue));
 		}
@@ -433,7 +433,7 @@ public:
 	using Super::IsSet;
 #endif
 
-	[[nodiscard]] FORCEINLINE explicit constexpr operator bool() const
+	[[nodiscard]] UE_FORCEINLINE_HINT explicit constexpr operator bool() const
 	{
 		return IsSet();
 	}
@@ -444,7 +444,7 @@ public:
 		checkf(IsSet(), TEXT("It is an error to call GetValue() on an unset TOptional. Please either check IsSet() or use Get(DefaultValue) instead."));
 		return this->TypedValue;
 	}
-	[[nodiscard]] FORCEINLINE constexpr const OptionalType& GetValue() const
+	[[nodiscard]] UE_FORCEINLINE_HINT constexpr const OptionalType& GetValue() const
 	{
 		return const_cast<TOptional*>(this)->GetValue();
 	}
@@ -453,7 +453,7 @@ public:
 	{
 		return std::addressof(GetValue());
 	}
-	[[nodiscard]] FORCEINLINE constexpr const OptionalType* operator->() const
+	[[nodiscard]] UE_FORCEINLINE_HINT constexpr const OptionalType* operator->() const
 	{
 		return const_cast<TOptional*>(this)->operator->();
 	}
@@ -462,7 +462,7 @@ public:
 	{
 		return GetValue();
 	}
-	[[nodiscard]] FORCEINLINE constexpr const OptionalType& operator*() const
+	[[nodiscard]] UE_FORCEINLINE_HINT constexpr const OptionalType& operator*() const
 	{
 		return const_cast<TOptional*>(this)->operator*();
 	}
@@ -478,7 +478,7 @@ public:
 	{
 		return IsSet() ? std::addressof(this->TypedValue) : nullptr;
 	}
-	[[nodiscard]] FORCEINLINE constexpr const OptionalType* GetPtrOrNull() const
+	[[nodiscard]] UE_FORCEINLINE_HINT constexpr const OptionalType* GetPtrOrNull() const
 	{
 		return const_cast<TOptional*>(this)->GetPtrOrNull();
 	}
@@ -488,7 +488,7 @@ private:
 	 * Destroys the value, must only be called if the value is set, and callers must then mark the value unset or construct a new value in place. 
 	 */
 #if PLATFORM_COMPILER_SUPPORTS_CONSTRAINED_DESTRUCTORS
-	FORCEINLINE constexpr void DestroyValue()
+	UE_FORCEINLINE_HINT constexpr void DestroyValue()
 	{
 		DestructItem(std::addressof(this->TypedValue));
 	}
diff --git a/Engine/Source/Runtime/Core/Public/Misc/OutputDevice.h b/Engine/Source/Runtime/Core/Public/Misc/OutputDevice.h
--- a/Engine/Source/Runtime/Core/Public/Misc/OutputDevice.h
+++ b/Engine/Source/Runtime/Core/Public/Misc/OutputDevice.h
@@ -167,12 +167,12 @@ public:
 	{
 		bSuppressEventTag = bInSuppressEventTag;
 	}
-	FORCEINLINE bool GetSuppressEventTag() const	{	return bSuppressEventTag;	}
+	UE_FORCEINLINE_HINT bool GetSuppressEventTag() const	{	return bSuppressEventTag;	}
 	void SetAutoEmitLineTerminator(bool bInAutoEmitLineTerminator)
 	{
 		bAutoEmitLineTerminator = bInAutoEmitLineTerminator;
 	}
-	FORCEINLINE bool GetAutoEmitLineTerminator() const	{	return bAutoEmitLineTerminator;	}
+	UE_FORCEINLINE_HINT bool GetAutoEmitLineTerminator() const	{	return bAutoEmitLineTerminator;	}
 
 	/** 
 	 * Dumps the contents of this output device's buffer to an archive (supported by output device that have a memory buffer) 
diff --git a/Engine/Source/Runtime/Core/Public/Misc/Paths.h b/Engine/Source/Runtime/Core/Public/Misc/Paths.h
--- a/Engine/Source/Runtime/Core/Public/Misc/Paths.h
+++ b/Engine/Source/Runtime/Core/Public/Misc/Paths.h
@@ -703,7 +703,7 @@ public:
 	static CORE_API const FString& GetRelativePathToRoot();
 
 	template <typename... PathTypes>
-	FORCEINLINE static FString Combine(PathTypes&&... InPaths)
+	UE_FORCEINLINE_HINT static FString Combine(PathTypes&&... InPaths)
 	{
 		return CombineImpl<PathTypes...>(Forward<PathTypes>(InPaths)...);
 	}
diff --git a/Engine/Source/Runtime/Core/Public/Misc/RemoteConfigIni.h b/Engine/Source/Runtime/Core/Public/Misc/RemoteConfigIni.h
--- a/Engine/Source/Runtime/Core/Public/Misc/RemoteConfigIni.h
+++ b/Engine/Source/Runtime/Core/Public/Misc/RemoteConfigIni.h
@@ -64,7 +64,7 @@ public:
 	/** Returns the local IO info object */
 	FRemoteConfigAsyncIOInfo& GetIOInfo();
 
-	FORCEINLINE TStatId GetStatId() const
+	UE_FORCEINLINE_HINT TStatId GetStatId() const
 	{
 		RETURN_QUICK_DECLARE_CYCLE_STAT(FRemoteConfigAsyncWorker, STATGROUP_ThreadPoolAsyncTasks);
 	}
diff --git a/Engine/Source/Runtime/Core/Public/Misc/ReverseIterate.h b/Engine/Source/Runtime/Core/Public/Misc/ReverseIterate.h
--- a/Engine/Source/Runtime/Core/Public/Misc/ReverseIterate.h
+++ b/Engine/Source/Runtime/Core/Public/Misc/ReverseIterate.h
@@ -64,13 +64,13 @@ namespace UE::Core::Private
 		TReverseIterationAdapter(const TReverseIterationAdapter&) = delete;
 		~TReverseIterationAdapter() = default;
 
-		FORCEINLINE constexpr auto begin() const UE_LIFETIMEBOUND
+		UE_FORCEINLINE_HINT constexpr auto begin() const UE_LIFETIMEBOUND
 		{
 			using std::rbegin;
 			return rbegin(Range);
 		}
 
-		FORCEINLINE constexpr auto end() const UE_LIFETIMEBOUND
+		UE_FORCEINLINE_HINT constexpr auto end() const UE_LIFETIMEBOUND
 		{
 			using std::rend;
 			return rend(Range);
diff --git a/Engine/Source/Runtime/Core/Public/Misc/ScopeTryLock.h b/Engine/Source/Runtime/Core/Public/Misc/ScopeTryLock.h
--- a/Engine/Source/Runtime/Core/Public/Misc/ScopeTryLock.h
+++ b/Engine/Source/Runtime/Core/Public/Misc/ScopeTryLock.h
@@ -51,7 +51,7 @@ namespace UE
 			}
 		}
 
-		FORCEINLINE bool IsLocked()
+		UE_FORCEINLINE_HINT bool IsLocked()
 		{
 			return HeldMutex != nullptr;
 		}
diff --git a/Engine/Source/Runtime/Core/Public/Misc/ScopedSlowTask.h b/Engine/Source/Runtime/Core/Public/Misc/ScopedSlowTask.h
--- a/Engine/Source/Runtime/Core/Public/Misc/ScopedSlowTask.h
+++ b/Engine/Source/Runtime/Core/Public/Misc/ScopedSlowTask.h
@@ -44,7 +44,7 @@ struct FScopedSlowTask : FSlowTask
 		Initialize();
 	}
 
-	FORCEINLINE ~FScopedSlowTask()
+	UE_FORCEINLINE_HINT ~FScopedSlowTask()
 	{
 		Destroy();
 	}
diff --git a/Engine/Source/Runtime/Core/Public/Misc/SecureHash.h b/Engine/Source/Runtime/Core/Public/Misc/SecureHash.h
--- a/Engine/Source/Runtime/Core/Public/Misc/SecureHash.h
+++ b/Engine/Source/Runtime/Core/Public/Misc/SecureHash.h
@@ -500,7 +500,7 @@ public:
 		}
 	}
 
-	FORCEINLINE TStatId GetStatId() const
+	UE_FORCEINLINE_HINT TStatId GetStatId() const
 	{
 		RETURN_QUICK_DECLARE_CYCLE_STAT(FAsyncSHAVerify, STATGROUP_ThreadPoolAsyncTasks);
 	}
diff --git a/Engine/Source/Runtime/Core/Public/Misc/TVariantMeta.h b/Engine/Source/Runtime/Core/Public/Misc/TVariantMeta.h
--- a/Engine/Source/Runtime/Core/Public/Misc/TVariantMeta.h
+++ b/Engine/Source/Runtime/Core/Public/Misc/TVariantMeta.h
@@ -332,19 +332,19 @@ namespace Private
 
 	/** Cast a TVariant to its private base */
 	template <typename... Ts>
-	FORCEINLINE TVariantStorage<Ts...>& CastToStorage(TVariant<Ts...>& Variant)
+	UE_FORCEINLINE_HINT TVariantStorage<Ts...>& CastToStorage(TVariant<Ts...>& Variant)
 	{
 		return *(TVariantStorage<Ts...>*)(&Variant);
 	}
 
 	template <typename... Ts>
-	FORCEINLINE TVariantStorage<Ts...>&& CastToStorage(TVariant<Ts...>&& Variant)
+	UE_FORCEINLINE_HINT TVariantStorage<Ts...>&& CastToStorage(TVariant<Ts...>&& Variant)
 	{
 		return (TVariantStorage<Ts...>&&)(*(TVariantStorage<Ts...>*)(&Variant));
 	}
 
 	template <typename... Ts>
-	FORCEINLINE const TVariantStorage<Ts...>& CastToStorage(const TVariant<Ts...>& Variant)
+	UE_FORCEINLINE_HINT const TVariantStorage<Ts...>& CastToStorage(const TVariant<Ts...>& Variant)
 	{
 		return *(const TVariantStorage<Ts...>*)(&Variant);
 	}
diff --git a/Engine/Source/Runtime/Core/Public/Misc/TextFilterUtils.h b/Engine/Source/Runtime/Core/Public/Misc/TextFilterUtils.h
--- a/Engine/Source/Runtime/Core/Public/Misc/TextFilterUtils.h
+++ b/Engine/Source/Runtime/Core/Public/Misc/TextFilterUtils.h
@@ -75,19 +75,19 @@ public:
 	CORE_API bool CompareNumeric(const FTextFilterString& InOther, const ETextFilterComparisonOperation InComparisonOperation) const;
 
 	/** Get the internal uppercase string of this filter string */
-	FORCEINLINE const FString& AsString() const
+	UE_FORCEINLINE_HINT const FString& AsString() const
 	{
 		return InternalString;
 	}
 
 	/** Get the internal uppercase string of this filter string as an FName */
-	FORCEINLINE FName AsName() const
+	UE_FORCEINLINE_HINT FName AsName() const
 	{
 		return FName(*InternalString);
 	}
 
 	/** Is the internal string empty? */
-	FORCEINLINE bool IsEmpty() const
+	UE_FORCEINLINE_HINT bool IsEmpty() const
 	{
 		return InternalString.IsEmpty();
 	}
@@ -169,7 +169,7 @@ namespace TextFilterUtils
 			Init(NameEntry, NumberInternal);
 		}
 
-		FORCEINLINE void Init(const FName& Name)
+		UE_FORCEINLINE_HINT void Init(const FName& Name)
 		{
 			Init(Name.GetDisplayNameEntry(), Name.GetNumber());
 		}
@@ -207,7 +207,7 @@ namespace TextFilterUtils
 			}
 		}
 
-		FORCEINLINE bool IsWide() const
+		UE_FORCEINLINE_HINT bool IsWide() const
 		{
 			return bIsWide;
 		}
@@ -224,7 +224,7 @@ namespace TextFilterUtils
 			return WideName;
 		}
 
-		FORCEINLINE int32 GetMaxBufferLength()
+		UE_FORCEINLINE_HINT int32 GetMaxBufferLength()
 		{
 			return NAME_WITH_NUMBER_SIZE;
 		}
diff --git a/Engine/Source/Runtime/Core/Public/Misc/Timeout.h b/Engine/Source/Runtime/Core/Public/Misc/Timeout.h
--- a/Engine/Source/Runtime/Core/Public/Misc/Timeout.h
+++ b/Engine/Source/Runtime/Core/Public/Misc/Timeout.h
@@ -52,7 +52,7 @@ namespace UE
 		}
 
 		/** Returns true if this was created from Never and does not need to be repeatedly checked */
-		FORCEINLINE bool WillNeverExpire() const
+		UE_FORCEINLINE_HINT bool WillNeverExpire() const
 		{
 			return TimeoutCycles == NeverExpireCycles;
 		}
@@ -64,7 +64,7 @@ namespace UE
 		}
 
 		/** Returns true if this was created from AlwaysExpired and does not need to be repeatedly checked */
-		FORCEINLINE bool IsAlwaysExpired() const
+		UE_FORCEINLINE_HINT bool IsAlwaysExpired() const
 		{
 			return TimeoutCycles == 0;
 		}
diff --git a/Engine/Source/Runtime/Core/Public/Modules/ModuleManager.h b/Engine/Source/Runtime/Core/Public/Modules/ModuleManager.h
--- a/Engine/Source/Runtime/Core/Public/Modules/ModuleManager.h
+++ b/Engine/Source/Runtime/Core/Public/Modules/ModuleManager.h
@@ -354,7 +354,7 @@ public:
 	  * @see GetModuleChecked, LoadModulePtr, LoadModuleChecked
 	  */
 	template<typename TModuleInterface>
-	static FORCEINLINE TModuleInterface* GetModulePtr( const FName ModuleName )
+	static UE_FORCEINLINE_HINT TModuleInterface* GetModulePtr( const FName ModuleName )
 	{
 		return static_cast<TModuleInterface*>(GetModulePtr_Internal(ModuleName));
 	}
@@ -676,12 +676,12 @@ private:
 	CORE_API ModuleInfoPtr FindModule(FName InModuleName);
 	CORE_API ModuleInfoRef FindModuleChecked(FName InModuleName);
 
-	FORCEINLINE TSharedPtr<const FModuleInfo, ESPMode::ThreadSafe> FindModule(FName InModuleName) const
+	UE_FORCEINLINE_HINT TSharedPtr<const FModuleInfo, ESPMode::ThreadSafe> FindModule(FName InModuleName) const
 	{
 		return const_cast<FModuleManager*>(this)->FindModule(InModuleName);
 	}
 
-	FORCEINLINE TSharedRef<const FModuleInfo, ESPMode::ThreadSafe> FindModuleChecked(FName InModuleName) const
+	UE_FORCEINLINE_HINT TSharedRef<const FModuleInfo, ESPMode::ThreadSafe> FindModuleChecked(FName InModuleName) const
 	{
 		return const_cast<FModuleManager*>(this)->FindModuleChecked(InModuleName);
 	}
diff --git a/Engine/Source/Runtime/Core/Public/PixelFormat.h b/Engine/Source/Runtime/Core/Public/PixelFormat.h
--- a/Engine/Source/Runtime/Core/Public/PixelFormat.h
+++ b/Engine/Source/Runtime/Core/Public/PixelFormat.h
@@ -327,7 +327,7 @@ inline bool IsBlockCompressedFormat(EPixelFormat Format)
 	return false;
 }
 
-FORCEINLINE bool IsAnyBlockCompressedPixelFormat(EPixelFormat PixelFormat)
+UE_FORCEINLINE_HINT bool IsAnyBlockCompressedPixelFormat(EPixelFormat PixelFormat)
 {
 	return IsBlockCompressedFormat(PixelFormat) || IsASTCBlockCompressedTextureFormat(PixelFormat) || IsETCBlockCompressedPixelFormat(PixelFormat);
 }
@@ -438,7 +438,7 @@ inline bool IsDXTCBlockCompressedTextureFormat(EPixelFormat PixelFormat)
 	PRAGMA_RESTORE_SWITCH_UNHANDLED_ENUM_CASE_WARNINGS;
 }
 
-FORCEINLINE bool RequiresBlock4Alignment(EPixelFormat PixelFormat)
+UE_FORCEINLINE_HINT bool RequiresBlock4Alignment(EPixelFormat PixelFormat)
 {
 	// BCN DXTC formats require 4x4 alignment, but ASTC/ETC do not
 	return IsDXTCBlockCompressedTextureFormat(PixelFormat);
diff --git a/Engine/Source/Runtime/Core/Public/ProfilingDebugging/CookStats.h b/Engine/Source/Runtime/Core/Public/ProfilingDebugging/CookStats.h
--- a/Engine/Source/Runtime/Core/Public/ProfilingDebugging/CookStats.h
+++ b/Engine/Source/Runtime/Core/Public/ProfilingDebugging/CookStats.h
@@ -83,7 +83,7 @@ public:
 		{
 			return Key < Other.Key;
 		}
-		FORCEINLINE bool operator()(const TKeyValuePair& A, const TKeyValuePair& B) const
+		UE_FORCEINLINE_HINT bool operator()(const TKeyValuePair& A, const TKeyValuePair& B) const
 		{
 			return A.Key < B.Key;
 		}
diff --git a/Engine/Source/Runtime/Core/Public/ProfilingDebugging/Histogram.h b/Engine/Source/Runtime/Core/Public/ProfilingDebugging/Histogram.h
--- a/Engine/Source/Runtime/Core/Public/ProfilingDebugging/Histogram.h
+++ b/Engine/Source/Runtime/Core/Public/ProfilingDebugging/Histogram.h
@@ -42,7 +42,7 @@ struct FHistogram
 	CORE_API void InitFromArray(TArrayView<const double> Thresholds);
 
 	/** Inits histogram with the specified bin boundaries, with the final bucket extending to infinity (e.g., passing in 0,5 creates a [0..5) bucket and a [5..infinity) bucket) */
-	FORCEINLINE void InitFromArray(std::initializer_list<double> Thresholds)
+	UE_FORCEINLINE_HINT void InitFromArray(std::initializer_list<double> Thresholds)
 	{
 		InitFromArray(MakeArrayView(Thresholds));
 	}
@@ -237,11 +237,11 @@ protected:
 		{
 		}
 
-		FORCEINLINE FBin operator-(const FBin& Other) const
+		UE_FORCEINLINE_HINT FBin operator-(const FBin& Other) const
 		{
 			return FBin(MinValue, UpperBound, Sum - Other.Sum, Count - Other.Count);
 		}
-		FORCEINLINE FBin operator+(const FBin& Other) const
+		UE_FORCEINLINE_HINT FBin operator+(const FBin& Other) const
 		{
 			return FBin(MinValue, UpperBound, Sum + Other.Sum, Count + Other.Count);
 		}
diff --git a/Engine/Source/Runtime/Core/Public/ProfilingDebugging/StringsTrace.h b/Engine/Source/Runtime/Core/Public/ProfilingDebugging/StringsTrace.h
--- a/Engine/Source/Runtime/Core/Public/ProfilingDebugging/StringsTrace.h
+++ b/Engine/Source/Runtime/Core/Public/ProfilingDebugging/StringsTrace.h
@@ -77,10 +77,10 @@ public:
 };
 
 #if !UE_TRACE_ENABLED
-FORCEINLINE UE::Trace::FEventRef32 FStringTrace::GetNameRef(const FName& Name) { return UE::Trace::FEventRef32(0,0); }
-FORCEINLINE UE::Trace::FEventRef64 FStringTrace::GetStaticStringRef(const TCHAR* String) { return UE::Trace::FEventRef64(0,0); }
-FORCEINLINE UE::Trace::FEventRef64 FStringTrace::GetStaticStringRef(const ANSICHAR* String) { return UE::Trace::FEventRef64(0,0); }
-FORCEINLINE void FStringTrace::OnConnection() {};
+UE_FORCEINLINE_HINT UE::Trace::FEventRef32 FStringTrace::GetNameRef(const FName& Name) { return UE::Trace::FEventRef32(0,0); }
+UE_FORCEINLINE_HINT UE::Trace::FEventRef64 FStringTrace::GetStaticStringRef(const TCHAR* String) { return UE::Trace::FEventRef64(0,0); }
+UE_FORCEINLINE_HINT UE::Trace::FEventRef64 FStringTrace::GetStaticStringRef(const ANSICHAR* String) { return UE::Trace::FEventRef64(0,0); }
+UE_FORCEINLINE_HINT void FStringTrace::OnConnection() {};
 #endif
 
 #undef UE_API
\ No newline at end of file
diff --git a/Engine/Source/Runtime/Core/Public/Serialization/Archive.h b/Engine/Source/Runtime/Core/Public/Serialization/Archive.h
--- a/Engine/Source/Runtime/Core/Public/Serialization/Archive.h
+++ b/Engine/Source/Runtime/Core/Public/Serialization/Archive.h
@@ -201,19 +201,19 @@ public:
 	}
 
 	/** Returns the global engine serialization version used for this archive. */
-	FORCEINLINE FPackageFileVersion UEVer() const
+	UE_FORCEINLINE_HINT FPackageFileVersion UEVer() const
 	{
 		return ArUEVer;
 	}
 
 	/** Returns the licensee-specific version used for this archive, will be 0 by default. */
-	FORCEINLINE int32 LicenseeUEVer() const
+	UE_FORCEINLINE_HINT int32 LicenseeUEVer() const
 	{
 		return ArLicenseeUEVer;
 	}
 
 	/** Returns the compiled engine version used for this archive. */
-	FORCEINLINE FEngineVersionBase EngineVer() const
+	UE_FORCEINLINE_HINT FEngineVersionBase EngineVer() const
 	{
 		return ArEngineVer;
 	}
@@ -233,19 +233,19 @@ public:
 	CORE_API int32 CustomVer(const struct FGuid& Key) const;
 
 	/** Returns true if this archive is for loading data. */
-	FORCEINLINE bool IsLoading() const
+	UE_FORCEINLINE_HINT bool IsLoading() const
 	{
 		return ArIsLoading;
 	}
 
 	/** Returns true if this archive is loading from a cooked package. */
- 	FORCEINLINE bool IsLoadingFromCookedPackage() const
+ 	UE_FORCEINLINE_HINT bool IsLoadingFromCookedPackage() const
 	{
 		return ArIsLoadingFromCookedPackage;
 	}
 
 	/** Returns true if this archive is for saving data, this can also be a pre-save preparation archive. */
-	FORCEINLINE bool IsSaving() const
+	UE_FORCEINLINE_HINT bool IsSaving() const
 	{
 		return ArIsSaving;
 	}
@@ -267,13 +267,13 @@ public:
 	 * Returns true if this archive serializes to a structured text format. 
 	 * Text format archives should use high level constructs from FStructuredArchive for delimiting data rather than manually seeking through the file.
 	 */
-	FORCEINLINE bool IsTextFormat() const
+	UE_FORCEINLINE_HINT bool IsTextFormat() const
 	{
 		return (ArIsTextFormat && WITH_TEXT_ARCHIVE_SUPPORT);
 	}
 
 	/** Returns true if this archive wants properties to be serialized in binary form instead of safer but slower tagged form. */
-	FORCEINLINE bool WantBinaryPropertySerialization() const
+	UE_FORCEINLINE_HINT bool WantBinaryPropertySerialization() const
 	{
 		return ArWantBinaryPropertySerialization;
 	}
@@ -282,13 +282,13 @@ public:
 	 * Returns true if tagged property serialization should be replaced by faster unversioned serialization.
 	 * This assumes writer and reader share the same property definitions. 
 	 */
-	FORCEINLINE bool UseUnversionedPropertySerialization() const
+	UE_FORCEINLINE_HINT bool UseUnversionedPropertySerialization() const
 	{
 		return ArUseUnversionedPropertySerialization;
 	}
 
 	/** Returns true if this archive wants to always save strings in UTF16 format even if they are ANSI characters. */
-	FORCEINLINE bool IsForcingUnicode() const
+	UE_FORCEINLINE_HINT bool IsForcingUnicode() const
 	{
 		return ArForceUnicode;
 	}
@@ -297,7 +297,7 @@ public:
 	 * Returns true if this archive is saving or loading data destined for persistent storage and should skip transient data.
 	 * This is also true for some intermediate archives for tasks like duplication that are eventually destined for persistent storage.
 	 */
-	FORCEINLINE bool IsPersistent() const
+	UE_FORCEINLINE_HINT bool IsPersistent() const
 	{
 		return ArIsPersistent;
 	}
@@ -315,7 +315,7 @@ public:
 	/**
 	 * Returns true if it is allowed to skip serialization on assets still being compiled to avoid waiting unless strictly necessary. 
 	 */
-	FORCEINLINE bool ShouldSkipCompilingAssets() const
+	UE_FORCEINLINE_HINT bool ShouldSkipCompilingAssets() const
 	{
 		return ArShouldSkipCompilingAssets;
 	}
@@ -328,7 +328,7 @@ public:
 	CORE_API void SetShouldSkipUpdateCustomVersion(bool bShouldSkip);
 
 	/** Returns true if this Archive does not update the custom version data for a given key if has already been set previously in the local custom versions container*/
-	FORCEINLINE bool ShouldSkipUpdateCustomVersion() const
+	UE_FORCEINLINE_HINT bool ShouldSkipUpdateCustomVersion() const
 	{
 		return ArShouldSkipUpdateCustomVersion;
 	}
@@ -347,42 +347,42 @@ public:
 
 	/** Returns true if this Archive is migrating remote object */
 #if UE_WITH_REMOTE_OBJECT_HANDLE
-	FORCEINLINE bool IsMigratingRemoteObjects() const
+	UE_FORCEINLINE_HINT bool IsMigratingRemoteObjects() const
 	{
 		return ArIsMigratingRemoteObjects;
 	}
 #else
-	FORCEINLINE constexpr bool IsMigratingRemoteObjects() const
+	UE_FORCEINLINE_HINT constexpr bool IsMigratingRemoteObjects() const
 	{
 		return false;
 	}
 #endif // UE_WITH_REMOTE_OBJECT_HANDLE
 
 	/** Returns true if this archive contains errors, which means that further serialization is generally not safe. */
-	FORCEINLINE bool IsError() const
+	UE_FORCEINLINE_HINT bool IsError() const
 	{
 		return ArIsError;
 	}
 
-	FORCEINLINE bool GetError() const
+	UE_FORCEINLINE_HINT bool GetError() const
 	{
 		return ArIsError;
 	}
 
 	/** Returns true if this archive contains critical errors that cannot be recovered from. */
-	FORCEINLINE bool IsCriticalError() const
+	UE_FORCEINLINE_HINT bool IsCriticalError() const
 	{
 		return ArIsCriticalError;
 	}
 
 	/** Returns true if this archive contains native or generated code. */
-	FORCEINLINE bool ContainsCode() const
+	UE_FORCEINLINE_HINT bool ContainsCode() const
 	{
 		return ArContainsCode;
 	}
 
 	/** Returns true if this archive contains a ULevel or UWorld object. */
-	FORCEINLINE bool ContainsMap() const
+	UE_FORCEINLINE_HINT bool ContainsMap() const
 	{
 		return ArContainsMap;
 	}
@@ -396,55 +396,55 @@ PRAGMA_ENABLE_DEPRECATION_WARNINGS
 	}
 
 	/** Returns true if this archive should always swap bytes, ignoring endian rules. */
-	FORCEINLINE bool ForceByteSwapping() const
+	UE_FORCEINLINE_HINT bool ForceByteSwapping() const
 	{
 		return ArForceByteSwapping;
 	}
 
 	/** Returns true if this archive is currently serializing class/struct default values. */
-	FORCEINLINE bool IsSerializingDefaults() const
+	UE_FORCEINLINE_HINT bool IsSerializingDefaults() const
 	{
 		return (ArSerializingDefaults > 0) ? true : false;
 	}
 
 	/** Returns true if this archive should ignore archetype references for structs and classes. */
-	FORCEINLINE bool IsIgnoringArchetypeRef() const
+	UE_FORCEINLINE_HINT bool IsIgnoringArchetypeRef() const
 	{
 		return ArIgnoreArchetypeRef;
 	}
 
 	/** Returns true if this archive should handle delta serialization for properties. */
-	FORCEINLINE bool DoDelta() const
+	UE_FORCEINLINE_HINT bool DoDelta() const
 	{
 		return !ArNoDelta;
 	}
 
 	/** Returns true if this archive should perform delta serialization within properties (e.g. TMaps and TSets). */
-	FORCEINLINE bool DoIntraPropertyDelta() const
+	UE_FORCEINLINE_HINT bool DoIntraPropertyDelta() const
 	{
 		return !ArNoIntraPropertyDelta;
 	}
 
 	/** Returns true if this archive should ignore the Outer reference in UObject. */
-	FORCEINLINE bool IsIgnoringOuterRef() const
+	UE_FORCEINLINE_HINT bool IsIgnoringOuterRef() const
 	{
 		return ArIgnoreOuterRef;
 	}
 
 	/** Returns true if this archive should ignore the ClassGeneratedBy reference in UClass. */
-	FORCEINLINE bool IsIgnoringClassGeneratedByRef() const
+	UE_FORCEINLINE_HINT bool IsIgnoringClassGeneratedByRef() const
 	{
 		return ArIgnoreClassGeneratedByRef;
 	}
 
 	/** Returns true if this archive should ignore the Class reference in UObject. */
-	FORCEINLINE bool IsIgnoringClassRef() const
+	UE_FORCEINLINE_HINT bool IsIgnoringClassRef() const
 	{
 		return ArIgnoreClassRef;
 	}
 
 	/** Returns true if this archive sould allow lazy loading of bulk / secondary data. */
-	FORCEINLINE bool IsAllowingLazyLoading() const
+	UE_FORCEINLINE_HINT bool IsAllowingLazyLoading() const
 	{
 		return ArAllowLazyLoading;
 	}
@@ -453,37 +453,37 @@ PRAGMA_ENABLE_DEPRECATION_WARNINGS
 	 * Returns true if this archive is only looking for UObject references.
 	 * This can be false for reference collectors looking for more general references.
 	 */
-	FORCEINLINE bool IsObjectReferenceCollector() const
+	UE_FORCEINLINE_HINT bool IsObjectReferenceCollector() const
 	{
 		return ArIsObjectReferenceCollector;
 	}
 
 	/** Returns true if this archive should modify/search weak object references as well as strong ones. */
-	FORCEINLINE bool IsModifyingWeakAndStrongReferences() const
+	UE_FORCEINLINE_HINT bool IsModifyingWeakAndStrongReferences() const
 	{
 		return ArIsModifyingWeakAndStrongReferences;
 	}
 
 	/** Returns true if this archive is counting memory, normally CountBytes is called to get the size. */
-	FORCEINLINE bool IsCountingMemory() const
+	UE_FORCEINLINE_HINT bool IsCountingMemory() const
 	{
 		return ArIsCountingMemory;
 	}
 
 	/** Returns this archive's property serialization modifier flags. */
-	FORCEINLINE uint32 GetPortFlags() const
+	UE_FORCEINLINE_HINT uint32 GetPortFlags() const
 	{
 		return ArPortFlags;
 	}
 
 	/** Checks to see if any of the passed in property serialization modifier flags are set. */
-	FORCEINLINE bool HasAnyPortFlags(uint32 Flags) const
+	UE_FORCEINLINE_HINT bool HasAnyPortFlags(uint32 Flags) const
 	{
 		return ((ArPortFlags & Flags) != 0);
 	}
 
 	/** Checks to see if all of the passed in property serialization modifier flags are set. */
-	FORCEINLINE bool HasAllPortFlags(uint32 Flags) const
+	UE_FORCEINLINE_HINT bool HasAllPortFlags(uint32 Flags) const
 	{
 		return ((ArPortFlags & Flags) == Flags);
 	}
@@ -499,13 +499,13 @@ PRAGMA_ENABLE_DEPRECATION_WARNINGS
 	}
 
 	/** Returns true if this archive should ignore bulk data. */
-	FORCEINLINE bool ShouldSkipBulkData() const
+	UE_FORCEINLINE_HINT bool ShouldSkipBulkData() const
 	{
 		return ArShouldSkipBulkData;
 	}
 
 	/** Returns the maximum size of data that this archive is allowed to serialize. */
-	FORCEINLINE int64 GetMaxSerializeSize() const
+	UE_FORCEINLINE_HINT int64 GetMaxSerializeSize() const
 	{
 		return ArMaxSerializeSize;
 	}
@@ -628,7 +628,7 @@ PRAGMA_ENABLE_DEPRECATION_WARNINGS
 	/**
 	 * Whether or not this archive is serializing data being sent/received by the netcode
 	 */
-	FORCEINLINE bool IsNetArchive() const
+	UE_FORCEINLINE_HINT bool IsNetArchive() const
 	{
 		return ArIsNetArchive;
 	}
@@ -638,7 +638,7 @@ PRAGMA_ENABLE_DEPRECATION_WARNINGS
 	 *
 	 * @return true if the archive is used for cooking, false otherwise.
 	 */
-	FORCEINLINE bool IsCooking() const
+	UE_FORCEINLINE_HINT bool IsCooking() const
 	{
 		return SavePackageData && SavePackageData->CookContext;
 	}
@@ -685,7 +685,7 @@ PRAGMA_ENABLE_DEPRECATION_WARNINGS
 		return nullptr;
 	}
 
-	FORCEINLINE FArchiveCookContext* GetCookContext()
+	UE_FORCEINLINE_HINT FArchiveCookContext* GetCookContext()
 	{
 		return SavePackageData ? SavePackageData->CookContext : nullptr;
 	}
@@ -696,7 +696,7 @@ PRAGMA_ENABLE_DEPRECATION_WARNINGS
 	 *
 	 * @return Target platform.
 	 */
-	FORCEINLINE const ITargetPlatform* CookingTarget() const
+	UE_FORCEINLINE_HINT const ITargetPlatform* CookingTarget() const
 	{
 		return SavePackageData ? SavePackageData->TargetPlatform : nullptr;
 	}
@@ -745,7 +745,7 @@ PRAGMA_ENABLE_DEPRECATION_WARNINGS
 	 *
 	 * @return Pointer to the property that is currently being serialized
 	 */
-	FORCEINLINE class FProperty* GetSerializedProperty() const
+	UE_FORCEINLINE_HINT class FProperty* GetSerializedProperty() const
 	{
 		return SerializedProperty;
 	}
@@ -760,7 +760,7 @@ PRAGMA_ENABLE_DEPRECATION_WARNINGS
 	 * Get the raw serialized property chain for this archive
 	 * @note Accessing this directly can avoid an array allocation depending on your use-case
 	 */
-	FORCEINLINE const FArchiveSerializedPropertyChain* GetSerializedPropertyChain() const
+	UE_FORCEINLINE_HINT const FArchiveSerializedPropertyChain* GetSerializedPropertyChain() const
 	{
 		return SerializedPropertyChain;
 	}
@@ -807,7 +807,7 @@ public:
 		const uint8* StartFastPathLoadBuffer;
 		const uint8* EndFastPathLoadBuffer;
 		const uint8* OriginalFastPathLoadBuffer;
-		FORCEINLINE FFastPathLoadBuffer()
+		UE_FORCEINLINE_HINT FFastPathLoadBuffer()
 		{
 			Reset();
 		}
@@ -1154,7 +1154,7 @@ public:
 	 *
 	 * @return	Returns the object pointer
 	 */
-	FORCEINLINE T* operator -> () const
+	UE_FORCEINLINE_HINT T* operator -> () const
 	{
 		return Object;
 	}
@@ -1164,7 +1164,7 @@ public:
 	 *
 	 * @return	Returns a reference to the pointer
 	 */
-	FORCEINLINE T*& Get()
+	UE_FORCEINLINE_HINT T*& Get()
 	{
 		return Object;
 	}
@@ -1174,7 +1174,7 @@ public:
 	 *
 	 * @return	Whether or not the pointer is valid
 	 */
-	FORCEINLINE bool IsValid() const
+	UE_FORCEINLINE_HINT bool IsValid() const
 	{
 		return Object != nullptr;
 	}
@@ -1186,7 +1186,7 @@ public:
 	 *
 	 * @return	Whether or not there was an error
 	 */
-	FORCEINLINE bool IsError() const
+	UE_FORCEINLINE_HINT bool IsError() const
 	{
 		return bError;
 	}
@@ -2143,7 +2143,7 @@ public:
 
 #else
 	template<SIZE_T Size>
-	FORCEINLINE bool FastPathLoad(void* InDest)
+	UE_FORCEINLINE_HINT bool FastPathLoad(void* InDest)
 	{
 		return false;
 	}
@@ -2411,7 +2411,7 @@ template <
 	typename EnumType,
 	std::enable_if_t<TIsEnumClass<EnumType>::Value, int> = 0
 >
-FORCEINLINE FArchive& operator<<(FArchive& Ar, EnumType& Value)
+UE_FORCEINLINE_HINT FArchive& operator<<(FArchive& Ar, EnumType& Value)
 {
 	return Ar << (__underlying_type(EnumType)&)Value;
 }
diff --git a/Engine/Source/Runtime/Core/Public/Serialization/ArchiveProxy.h b/Engine/Source/Runtime/Core/Public/Serialization/ArchiveProxy.h
--- a/Engine/Source/Runtime/Core/Public/Serialization/ArchiveProxy.h
+++ b/Engine/Source/Runtime/Core/Public/Serialization/ArchiveProxy.h
@@ -300,13 +300,13 @@ public:
 	}
 #if WITH_EDITORONLY_DATA
 	/** Returns true if the stack of currently serialized properties contains an editor-only property */
-	virtual FORCEINLINE bool IsEditorOnlyPropertyOnTheStack() const override
+	virtual UE_FORCEINLINE_HINT bool IsEditorOnlyPropertyOnTheStack() const override
 	{
 		return InnerArchive.IsEditorOnlyPropertyOnTheStack();
 	}
 #endif
 
-	virtual FORCEINLINE bool IsProxyOf(FArchive* InOther) const override
+	virtual UE_FORCEINLINE_HINT bool IsProxyOf(FArchive* InOther) const override
 	{
 		return InOther == this || InOther == &InnerArchive || InnerArchive.IsProxyOf(InOther);
 	}
diff --git a/Engine/Source/Runtime/Core/Public/Serialization/ArchiveSerializedPropertyChain.h b/Engine/Source/Runtime/Core/Public/Serialization/ArchiveSerializedPropertyChain.h
--- a/Engine/Source/Runtime/Core/Public/Serialization/ArchiveSerializedPropertyChain.h
+++ b/Engine/Source/Runtime/Core/Public/Serialization/ArchiveSerializedPropertyChain.h
@@ -107,7 +107,7 @@ public:
 #endif
 	}
 
-	FORCEINLINE TArray<class FProperty*, TInlineAllocator<8>>::TConstIterator GetRootIterator() const
+	UE_FORCEINLINE_HINT TArray<class FProperty*, TInlineAllocator<8>>::TConstIterator GetRootIterator() const
 	{
 		return SerializedPropertyChain.CreateConstIterator();
 	}
diff --git a/Engine/Source/Runtime/Core/Public/Serialization/BitReader.h b/Engine/Source/Runtime/Core/Public/Serialization/BitReader.h
--- a/Engine/Source/Runtime/Core/Public/Serialization/BitReader.h
+++ b/Engine/Source/Runtime/Core/Public/Serialization/BitReader.h
@@ -140,22 +140,22 @@ PGO_LINK_ENABLE_WARNINGS
 		return Bit;
 	}
 
-	FORCEINLINE_DEBUGGABLE void Serialize( void* Dest, int64 LengthBytes )
+	UE_FORCEINLINE_HINT void Serialize( void* Dest, int64 LengthBytes )
 	{
 		SerializeBits( Dest, LengthBytes*8 );
 	}
 
-	FORCEINLINE_DEBUGGABLE uint8* GetData()
+	UE_FORCEINLINE_HINT uint8* GetData()
 	{
 		return Buffer.GetData();
 	}
 
-	FORCEINLINE_DEBUGGABLE const uint8* GetData() const
+	UE_FORCEINLINE_HINT const uint8* GetData() const
 	{
 		return Buffer.GetData();
 	}
 
-	FORCEINLINE_DEBUGGABLE const TArray<uint8>& GetBuffer() const
+	UE_FORCEINLINE_HINT const TArray<uint8>& GetBuffer() const
 	{
 		return Buffer;
 	}
@@ -166,27 +166,27 @@ PGO_LINK_ENABLE_WARNINGS
 		return &Buffer[(int32)(Pos >> 3)];
 	}
 
-	FORCEINLINE_DEBUGGABLE int64 GetBytesLeft() const
+	UE_FORCEINLINE_HINT int64 GetBytesLeft() const
 	{
 		return ((Num - Pos) + 7) >> 3;
 	}
-	FORCEINLINE_DEBUGGABLE int64 GetBitsLeft() const
+	UE_FORCEINLINE_HINT int64 GetBitsLeft() const
 	{
 		return (Num - Pos);
 	}
-	FORCEINLINE_DEBUGGABLE bool AtEnd()
+	UE_FORCEINLINE_HINT bool AtEnd()
 	{
 		return IsError() || Pos>=Num;
 	}
-	FORCEINLINE_DEBUGGABLE int64 GetNumBytes() const
+	UE_FORCEINLINE_HINT int64 GetNumBytes() const
 	{
 		return (Num+7)>>3;
 	}
-	FORCEINLINE_DEBUGGABLE int64 GetNumBits() const
+	UE_FORCEINLINE_HINT int64 GetNumBits() const
 	{
 		return Num;
 	}
-	FORCEINLINE_DEBUGGABLE int64 GetPosBits() const
+	UE_FORCEINLINE_HINT int64 GetPosBits() const
 	{
 		return Pos;
 	}
@@ -243,7 +243,7 @@ protected:
 
 private:
 
-	FORCEINLINE uint8 Shift(uint8 Cnt)
+	UE_FORCEINLINE_HINT uint8 Shift(uint8 Cnt)
 	{
 		return (uint8)(1<<Cnt);
 	}
@@ -266,12 +266,12 @@ public:
 		: Pos(Reader.Pos)
 	{ }
 
-	FORCEINLINE_DEBUGGABLE int64 GetPos() const
+	UE_FORCEINLINE_HINT int64 GetPos() const
 	{
 		return Pos;
 	}
 
-	FORCEINLINE_DEBUGGABLE void Pop( FBitReader& Reader )
+	UE_FORCEINLINE_HINT void Pop( FBitReader& Reader )
 	{
 		Reader.Pos = Pos;
 	}
diff --git a/Engine/Source/Runtime/Core/Public/Serialization/BitWriter.h b/Engine/Source/Runtime/Core/Public/Serialization/BitWriter.h
--- a/Engine/Source/Runtime/Core/Public/Serialization/BitWriter.h
+++ b/Engine/Source/Runtime/Core/Public/Serialization/BitWriter.h
@@ -110,7 +110,7 @@ public:
 	/**
 	 * Returns the number of bytes written.
 	 */
-	FORCEINLINE int64 GetNumBytes(void) const
+	UE_FORCEINLINE_HINT int64 GetNumBytes(void) const
 	{
 		return (Num+7)>>3;
 	}
@@ -118,7 +118,7 @@ public:
 	/**
 	 * Returns the number of bits written.
 	 */
-	FORCEINLINE int64 GetNumBits(void) const
+	UE_FORCEINLINE_HINT int64 GetNumBits(void) const
 	{
 		return Num;
 	}
@@ -126,7 +126,7 @@ public:
 	/**
 	 * Returns the number of bits the buffer supports.
 	 */
-	FORCEINLINE int64 GetMaxBits(void) const
+	UE_FORCEINLINE_HINT int64 GetMaxBits(void) const
 	{
 		return Max;
 	}
@@ -141,7 +141,7 @@ public:
 	/**
 	 * Sets whether or not this writer intentionally allows overflows (to avoid logspam)
 	 */
-	FORCEINLINE void SetAllowOverflow(bool bInAllow)
+	UE_FORCEINLINE_HINT void SetAllowOverflow(bool bInAllow)
 	{
 		bAllowOverflow = bInAllow;
 	}
@@ -173,7 +173,7 @@ public:
 		return true;
 	}
 
-	FORCEINLINE void SetAllowResize(bool NewResize)
+	UE_FORCEINLINE_HINT void SetAllowResize(bool NewResize)
 	{
 		bAllowResize = NewResize;
 	}
@@ -183,7 +183,7 @@ public:
 	 */
 	CORE_API void Reset() override;
 
-	FORCEINLINE void WriteAlign()
+	UE_FORCEINLINE_HINT void WriteAlign()
 	{
 		Num = ( Num + 7 ) & ( ~0x07 );
 	}
@@ -219,7 +219,7 @@ public:
 		Init(Writer);
 	}
 
-	FORCEINLINE_DEBUGGABLE int64 GetNumBits() const
+	UE_FORCEINLINE_HINT int64 GetNumBits() const
 	{
 		return Num;
 	}
@@ -240,7 +240,7 @@ public:
 	CORE_API void Copy( FBitWriter& Writer, TArray<uint8> &Buffer );
 
 	/** Pops the BitWriter back to the start but doesn't clear what was written. */
-	FORCEINLINE_DEBUGGABLE void PopWithoutClear( FBitWriter& Writer )
+	UE_FORCEINLINE_HINT void PopWithoutClear( FBitWriter& Writer )
 	{
 		Writer.Num = Num;
 	}
diff --git a/Engine/Source/Runtime/Core/Public/Serialization/CustomVersion.h b/Engine/Source/Runtime/Core/Public/Serialization/CustomVersion.h
--- a/Engine/Source/Runtime/Core/Public/Serialization/CustomVersion.h
+++ b/Engine/Source/Runtime/Core/Public/Serialization/CustomVersion.h
@@ -53,7 +53,7 @@ struct FCustomVersion
 	CustomVersionValidatorFunc Validator;
 
 	/** Constructor. */
-	FORCEINLINE FCustomVersion()
+	UE_FORCEINLINE_HINT FCustomVersion()
 	: Validator(nullptr)
 	{
 	}
@@ -69,13 +69,13 @@ struct FCustomVersion
 	}
 
 	/** Equality comparison operator for Key */
-	FORCEINLINE bool operator==(FGuid InKey) const
+	UE_FORCEINLINE_HINT bool operator==(FGuid InKey) const
 	{
 		return Key == InKey;
 	}
 
 	/** Inequality comparison operator for Key */
-	FORCEINLINE bool operator!=(FGuid InKey) const
+	UE_FORCEINLINE_HINT bool operator!=(FGuid InKey) const
 	{
 		return Key != InKey;
 	}
@@ -113,7 +113,7 @@ class FCustomVersionContainer
 
 public:
 	/** Gets available custom versions in this container. */
-	FORCEINLINE const FCustomVersionArray& GetAllVersions() const
+	UE_FORCEINLINE_HINT const FCustomVersionArray& GetAllVersions() const
 	{
 		return Versions;
 	}
diff --git a/Engine/Source/Runtime/Core/Public/Serialization/LargeMemoryData.h b/Engine/Source/Runtime/Core/Public/Serialization/LargeMemoryData.h
--- a/Engine/Source/Runtime/Core/Public/Serialization/LargeMemoryData.h
+++ b/Engine/Source/Runtime/Core/Public/Serialization/LargeMemoryData.h
@@ -24,7 +24,7 @@ public:
 	CORE_API bool Write(void* InData, int64 InOffset, int64 InNum);
 
 	/** Append data at the given offset. */
-	FORCEINLINE void Append(void* InData, int64 InNum)
+	UE_FORCEINLINE_HINT void Append(void* InData, int64 InNum)
 	{
 		Write(InData, NumBytes, InNum);
 	}
@@ -33,19 +33,19 @@ public:
 	CORE_API bool Read(void* OutData, int64 InOffset, int64 InNum) const;
 
 	/** Gets the size of the data written. */
-	FORCEINLINE int64 GetSize() const
+	UE_FORCEINLINE_HINT int64 GetSize() const
 	{
 		return NumBytes;
 	}
 
 	/** Returns the written data.  */
-	FORCEINLINE uint8* GetData()
+	UE_FORCEINLINE_HINT uint8* GetData()
 	{
 		return Data;
 	}
 
 	/** Returns the written data.  */
-	FORCEINLINE const uint8* GetData() const
+	UE_FORCEINLINE_HINT const uint8* GetData() const
 	{
 		return Data;
 	}
diff --git a/Engine/Source/Runtime/Core/Public/Serialization/LargeMemoryWriter.h b/Engine/Source/Runtime/Core/Public/Serialization/LargeMemoryWriter.h
--- a/Engine/Source/Runtime/Core/Public/Serialization/LargeMemoryWriter.h
+++ b/Engine/Source/Runtime/Core/Public/Serialization/LargeMemoryWriter.h
@@ -58,7 +58,7 @@ public:
 	 * Also returns the pointer, so that the caller only needs to call this function to take control
 	 * of the memory.
 	 */
-	FORCEINLINE uint8* ReleaseOwnership()
+	UE_FORCEINLINE_HINT uint8* ReleaseOwnership()
 	{
 		return Data.ReleaseOwnership();
 	}
diff --git a/Engine/Source/Runtime/Core/Public/Serialization/MemoryImage.h b/Engine/Source/Runtime/Core/Public/Serialization/MemoryImage.h
--- a/Engine/Source/Runtime/Core/Public/Serialization/MemoryImage.h
+++ b/Engine/Source/Runtime/Core/Public/Serialization/MemoryImage.h
@@ -613,20 +613,20 @@ public:
 	CORE_API ~FMemoryImageAllocatorBase();
 
 	// FContainerAllocatorInterface
-	FORCEINLINE FScriptContainerElement* GetAllocation() const
+	UE_FORCEINLINE_HINT FScriptContainerElement* GetAllocation() const
 	{
 		return Data.Get();
 	}
 
-	FORCEINLINE SIZE_T GetAllocatedSize(int32 CurrentMax, SIZE_T NumBytesPerElement) const
+	UE_FORCEINLINE_HINT SIZE_T GetAllocatedSize(int32 CurrentMax, SIZE_T NumBytesPerElement) const
 	{
 		return CurrentMax * NumBytesPerElement;
 	}
-	FORCEINLINE bool HasAllocation()
+	UE_FORCEINLINE_HINT bool HasAllocation()
 	{
 		return Data.IsValid();
 	}
-	FORCEINLINE int64 GetFrozenOffsetFromThis() const { return Data.GetFrozenOffsetFromThis(); }
+	UE_FORCEINLINE_HINT int64 GetFrozenOffsetFromThis() const { return Data.GetFrozenOffsetFromThis(); }
 
 	CORE_API void ResizeAllocation(int32 PreviousNumElements, int32 NumElements, SIZE_T NumBytesPerElement, uint32 Alignment);
 	CORE_API void WriteMemoryImage(FMemoryImageWriter& Writer, const FTypeLayoutDesc& TypeDesc, int32 NumAllocatedElements, uint32 Alignment) const;
@@ -653,44 +653,44 @@ public:
 		/** Default constructor. */
 		ForAnyElementType() = default;
 
-		FORCEINLINE SizeType GetInitialCapacity() const
+		UE_FORCEINLINE_HINT SizeType GetInitialCapacity() const
 		{
 			return 0;
 		}
-		FORCEINLINE int32 CalculateSlackReserve(int32 NewMax, int32 NumBytesPerElement) const
+		UE_FORCEINLINE_HINT int32 CalculateSlackReserve(int32 NewMax, int32 NumBytesPerElement) const
 		{
 			return DefaultCalculateSlackReserve(NewMax, NumBytesPerElement, true, Alignment);
 		}
-		FORCEINLINE int32 CalculateSlackReserve(int32 NewMax, int32 NumBytesPerElement, uint32 AlignmentOfElement) const
+		UE_FORCEINLINE_HINT int32 CalculateSlackReserve(int32 NewMax, int32 NumBytesPerElement, uint32 AlignmentOfElement) const
 		{
 			return DefaultCalculateSlackReserve(NewMax, NumBytesPerElement, true, AlignmentOfElement);
 		}
-		FORCEINLINE int32 CalculateSlackShrink(int32 NewMax, int32 CurrentMax, int32 NumBytesPerElement) const
+		UE_FORCEINLINE_HINT int32 CalculateSlackShrink(int32 NewMax, int32 CurrentMax, int32 NumBytesPerElement) const
 		{
 			return DefaultCalculateSlackShrink(NewMax, CurrentMax, NumBytesPerElement, true, Alignment);
 		}
-		FORCEINLINE int32 CalculateSlackShrink(int32 NewMax, int32 CurrentMax, int32 NumBytesPerElement, uint32 AlignmentOfElement) const
+		UE_FORCEINLINE_HINT int32 CalculateSlackShrink(int32 NewMax, int32 CurrentMax, int32 NumBytesPerElement, uint32 AlignmentOfElement) const
 		{
 			return DefaultCalculateSlackShrink(NewMax, CurrentMax, NumBytesPerElement, true, AlignmentOfElement);
 		}
-		FORCEINLINE int32 CalculateSlackGrow(int32 NewMax, int32 CurrentMax, int32 NumBytesPerElement) const
+		UE_FORCEINLINE_HINT int32 CalculateSlackGrow(int32 NewMax, int32 CurrentMax, int32 NumBytesPerElement) const
 		{
 			return DefaultCalculateSlackGrow(NewMax, CurrentMax, NumBytesPerElement, true, Alignment);
 		}
-		FORCEINLINE int32 CalculateSlackGrow(int32 NewMax, int32 CurrentMax, int32 NumBytesPerElement, uint32 AlignmentOfElement) const
+		UE_FORCEINLINE_HINT int32 CalculateSlackGrow(int32 NewMax, int32 CurrentMax, int32 NumBytesPerElement, uint32 AlignmentOfElement) const
 		{
 			return DefaultCalculateSlackGrow(NewMax, CurrentMax, NumBytesPerElement, true, AlignmentOfElement);
 		}
-		FORCEINLINE void ResizeAllocation(int32 CurrentNum, int32 NewMax, SIZE_T NumBytesPerElement)
+		UE_FORCEINLINE_HINT void ResizeAllocation(int32 CurrentNum, int32 NewMax, SIZE_T NumBytesPerElement)
 		{
 			FMemoryImageAllocatorBase::ResizeAllocation(CurrentNum, NewMax, NumBytesPerElement, Alignment);
 		}
-		FORCEINLINE void ResizeAllocation(int32 CurrentNum, int32 NewMax, SIZE_T NumBytesPerElement, uint32 AlignmentOfElement)
+		UE_FORCEINLINE_HINT void ResizeAllocation(int32 CurrentNum, int32 NewMax, SIZE_T NumBytesPerElement, uint32 AlignmentOfElement)
 		{
 			FMemoryImageAllocatorBase::ResizeAllocation(CurrentNum, NewMax, NumBytesPerElement, AlignmentOfElement);
 		}
 
-		FORCEINLINE void WriteMemoryImage(FMemoryImageWriter& Writer, const FTypeLayoutDesc& TypeDesc, int32 NumAllocatedElements) const
+		UE_FORCEINLINE_HINT void WriteMemoryImage(FMemoryImageWriter& Writer, const FTypeLayoutDesc& TypeDesc, int32 NumAllocatedElements) const
 		{
 			FMemoryImageAllocatorBase::WriteMemoryImage(Writer, TypeDesc, NumAllocatedElements, Alignment);
 		}
@@ -701,7 +701,7 @@ public:
 	{
 	public:
 		ForElementType() {}
-		FORCEINLINE ElementType* GetAllocation() const { return (ElementType*)ForAnyElementType::GetAllocation(); }
+		UE_FORCEINLINE_HINT ElementType* GetAllocation() const { return (ElementType*)ForAnyElementType::GetAllocation(); }
 	};
 };
 
@@ -758,7 +758,7 @@ public:
 	FMemoryImageString& operator=(FMemoryImageString&&) = default;
 	FMemoryImageString& operator=(const FMemoryImageString&) = default;
 
-	FORCEINLINE FMemoryImageString(const FString& Other) : Data(Other.GetCharArray()) {}
+	UE_FORCEINLINE_HINT FMemoryImageString(const FString& Other) : Data(Other.GetCharArray()) {}
 
 	template <
 		typename CharType,
@@ -775,20 +775,20 @@ public:
 		}
 	}
 
-	FORCEINLINE operator FString() const
+	UE_FORCEINLINE_HINT operator FString() const
 	{
 		return FString::ConstructFromPtrSize(Data.GetData(), Len());
 	}
 
-	FORCEINLINE const TCHAR* operator*() const
+	UE_FORCEINLINE_HINT const TCHAR* operator*() const
 	{
 		return Data.Num() ? Data.GetData() : TEXT("");
 	}
 
-	FORCEINLINE bool IsEmpty() const { return Data.Num() <= 1; }
-	FORCEINLINE SIZE_T GetAllocatedSize() const { return Data.GetAllocatedSize(); }
+	UE_FORCEINLINE_HINT bool IsEmpty() const { return Data.Num() <= 1; }
+	UE_FORCEINLINE_HINT SIZE_T GetAllocatedSize() const { return Data.GetAllocatedSize(); }
 
-	FORCEINLINE int32 Len() const
+	UE_FORCEINLINE_HINT int32 Len() const
 	{
 		return Data.Num() ? Data.Num() - 1 : 0;
 	}
@@ -832,7 +832,7 @@ public:
 	inline DataType::ElementAllocatorType& GetAllocatorInstance() { return Data.GetAllocatorInstance(); }
 
 	/** Case insensitive string hash function. */
-	friend FORCEINLINE uint32 GetTypeHash(const FMemoryImageString& S)
+	friend UE_FORCEINLINE_HINT uint32 GetTypeHash(const FMemoryImageString& S)
 	{
 		return FCrc::Strihash_DEPRECATED(*S);
 	}
diff --git a/Engine/Source/Runtime/Core/Public/Serialization/MemoryLayout.h b/Engine/Source/Runtime/Core/Public/Serialization/MemoryLayout.h
--- a/Engine/Source/Runtime/Core/Public/Serialization/MemoryLayout.h
+++ b/Engine/Source/Runtime/Core/Public/Serialization/MemoryLayout.h
@@ -84,9 +84,9 @@ namespace EFieldLayoutFlags
 		UseInstanceWithNoProperty = (1u << 3),
 	};
 
-	FORCEINLINE Type MakeFlags(uint32 Flags = None) { return (Type)Flags; }
-	FORCEINLINE Type MakeFlagsEditorOnly(uint32 Flags = None) { return (Type)(WithEditorOnly | Flags); }
-	FORCEINLINE Type MakeFlagsRayTracing(uint32 Flags = None) { return (Type)(WithRayTracing | Flags); }
+	UE_FORCEINLINE_HINT Type MakeFlags(uint32 Flags = None) { return (Type)Flags; }
+	UE_FORCEINLINE_HINT Type MakeFlagsEditorOnly(uint32 Flags = None) { return (Type)(WithEditorOnly | Flags); }
+	UE_FORCEINLINE_HINT Type MakeFlagsRayTracing(uint32 Flags = None) { return (Type)(WithRayTracing | Flags); }
 }
 
 struct FFieldLayoutDesc
@@ -261,7 +261,7 @@ namespace Freeze
 	}
 
 	template<typename T>
-	FORCEINLINE void IntrinsicWriteMemoryImage(FMemoryImageWriter& Writer, const T& Object, const FTypeLayoutDesc& TypeDesc)
+	UE_FORCEINLINE_HINT void IntrinsicWriteMemoryImage(FMemoryImageWriter& Writer, const T& Object, const FTypeLayoutDesc& TypeDesc)
 	{
 		IntrinsicWriteMemoryImage(Writer, &Object, sizeof(T));
 	}
@@ -274,13 +274,13 @@ namespace Freeze
 	}
 
 	template<typename T>
-	FORCEINLINE uint32 IntrinsicAppendHash(const T* DummyObject, const FTypeLayoutDesc& TypeDesc, const FPlatformTypeLayoutParameters& LayoutParams, FSHA1& Hasher)
+	UE_FORCEINLINE_HINT uint32 IntrinsicAppendHash(const T* DummyObject, const FTypeLayoutDesc& TypeDesc, const FPlatformTypeLayoutParameters& LayoutParams, FSHA1& Hasher)
 	{
 		return DefaultAppendHash(TypeDesc, LayoutParams, Hasher);
 	}
 
 	template<typename T>
-	FORCEINLINE uint32 IntrinsicGetTargetAlignment(const T* DummyObject, const FTypeLayoutDesc& TypeDesc, const FPlatformTypeLayoutParameters& LayoutParams)
+	UE_FORCEINLINE_HINT uint32 IntrinsicGetTargetAlignment(const T* DummyObject, const FTypeLayoutDesc& TypeDesc, const FPlatformTypeLayoutParameters& LayoutParams)
 	{
 		return TypeDesc.Alignment;
 	}
@@ -332,13 +332,13 @@ struct TUsePropertyFreezing
 template <typename T, bool bUsePropertyFreezing=TUsePropertyFreezing<T>::Value>
 struct TGetFreezeImageHelper
 {
-	static FORCEINLINE FTypeLayoutDesc::FWriteFrozenMemoryImageFunc* Do() { return &Freeze::DefaultWriteMemoryImage; }
+	static UE_FORCEINLINE_HINT FTypeLayoutDesc::FWriteFrozenMemoryImageFunc* Do() { return &Freeze::DefaultWriteMemoryImage; }
 };
 
 template <typename T, bool bProvidesStaticStruct=TModels_V<CStaticStructProvider, T>>
 struct TGetFreezeImageFieldHelper
 {
-	static FORCEINLINE FFieldLayoutDesc::FWriteFrozenMemoryImageFunc* Do() { return &Freeze::DefaultWriteMemoryImageField; }
+	static UE_FORCEINLINE_HINT FFieldLayoutDesc::FWriteFrozenMemoryImageFunc* Do() { return &Freeze::DefaultWriteMemoryImageField; }
 };
 
 template<typename T, typename BaseType>
@@ -359,7 +359,7 @@ static void InternalInitializeBaseHelper(FTypeLayoutDesc& TypeDesc)
 }
 
 template<typename T>
-static FORCEINLINE void InternalInitializeBasesHelper(FTypeLayoutDesc& TypeDesc) {}
+static UE_FORCEINLINE_HINT void InternalInitializeBasesHelper(FTypeLayoutDesc& TypeDesc) {}
 
 template<typename T, typename BaseType, typename ...RemainingBaseTypes>
 static void InternalInitializeBasesHelper(FTypeLayoutDesc& TypeDesc)
@@ -389,7 +389,7 @@ template<typename T>
 struct TInitializeBaseHelper<T, void>
 {
 	UE_STATIC_ONLY(TInitializeBaseHelper);
-	static FORCEINLINE void Initialize(FTypeLayoutDesc& TypeDesc) {}
+	static UE_FORCEINLINE_HINT void Initialize(FTypeLayoutDesc& TypeDesc) {}
 };
 
 namespace Freeze
@@ -509,7 +509,7 @@ PRAGMA_ENABLE_DEPRECATION_WARNINGS
 	static constexpr int CounterBase = __COUNTER__; \
 	public: using DerivedType = PREPROCESSOR_REMOVE_OPTIONAL_PARENS(T); \
 	static constexpr ETypeLayoutInterface::Type InterfaceType = ETypeLayoutInterface::InInterface; \
-	UE_DECLARE_INTERNAL_LINK_BASE(InternalLinkType) { UE_STATIC_ONLY(InternalLinkType); static FORCEINLINE void Initialize(FTypeLayoutDesc& TypeDesc) {} }
+	UE_DECLARE_INTERNAL_LINK_BASE(InternalLinkType) { UE_STATIC_ONLY(InternalLinkType); static UE_FORCEINLINE_HINT void Initialize(FTypeLayoutDesc& TypeDesc) {} }
 
 #define INTERNAL_DECLARE_INLINE_TYPE_LAYOUT(T, InInterface) \
 	private: static void InternalDestroy(void* Object, const FTypeLayoutDesc&, const FPointerTableBase* PtrTable, bool bIsFrozen) { \
diff --git a/Engine/Source/Runtime/Core/Public/Serialization/StructuredArchive.h b/Engine/Source/Runtime/Core/Public/Serialization/StructuredArchive.h
--- a/Engine/Source/Runtime/Core/Public/Serialization/StructuredArchive.h
+++ b/Engine/Source/Runtime/Core/Public/Serialization/StructuredArchive.h
@@ -93,7 +93,7 @@ public:
 	/**
 	 * Gets the serialization context from the underlying archive.
 	 */
-	FORCEINLINE FArchive& GetUnderlyingArchive() const
+	UE_FORCEINLINE_HINT FArchive& GetUnderlyingArchive() const
 	{
 		return Formatter.GetUnderlyingArchive();
 	}
@@ -101,7 +101,7 @@ public:
 	/**
 	 * Gets the archiving state.
 	 */
-	FORCEINLINE const FArchiveState& GetArchiveState() const
+	UE_FORCEINLINE_HINT const FArchiveState& GetArchiveState() const
 	{
 		return GetUnderlyingArchive().GetArchiveState();
 	}
@@ -209,21 +209,21 @@ private:
 
 #if !WITH_TEXT_ARCHIVE_SUPPORT
 
-	FORCEINLINE FStructuredArchive::FStructuredArchive(FArchiveFormatterType& InFormatter)
+	UE_FORCEINLINE_HINT FStructuredArchive::FStructuredArchive(FArchiveFormatterType& InFormatter)
 		: Formatter(InFormatter)
 	{
 	}
 
-	FORCEINLINE FStructuredArchive::~FStructuredArchive()
+	UE_FORCEINLINE_HINT FStructuredArchive::~FStructuredArchive()
 	{
 	}
 
-	FORCEINLINE FStructuredArchiveSlot FStructuredArchive::Open()
+	UE_FORCEINLINE_HINT FStructuredArchiveSlot FStructuredArchive::Open()
 	{
 		return FStructuredArchiveSlot(FStructuredArchiveSlot::EPrivateToken{}, *this);
 	}
 
-	FORCEINLINE void FStructuredArchive::Close()
+	UE_FORCEINLINE_HINT void FStructuredArchive::Close()
 	{
 	}
 
@@ -231,24 +231,24 @@ private:
 
 namespace UE::StructuredArchive::Private
 {
-	FORCEINLINE FArchive& GetUnderlyingArchiveImpl(FStructuredArchive& StructuredArchive)
+	UE_FORCEINLINE_HINT FArchive& GetUnderlyingArchiveImpl(FStructuredArchive& StructuredArchive)
 	{
 		return StructuredArchive.GetUnderlyingArchive();
 	}
 
-	FORCEINLINE FArchiveState& GetUnderlyingArchiveStateImpl(FStructuredArchive& StructuredArchive)
+	UE_FORCEINLINE_HINT FArchiveState& GetUnderlyingArchiveStateImpl(FStructuredArchive& StructuredArchive)
 	{
 		return StructuredArchive.GetUnderlyingArchive().GetArchiveState();
 	}
 
 #if WITH_TEXT_ARCHIVE_SUPPORT
-	FORCEINLINE FElementId GetCurrentSlotElementIdImpl(FStructuredArchive& StructuredArchive)
+	UE_FORCEINLINE_HINT FElementId GetCurrentSlotElementIdImpl(FStructuredArchive& StructuredArchive)
 	{
 		return StructuredArchive.CurrentSlotElementId;
 	}
 #endif
 
-	FORCEINLINE FArchiveFormatterType& GetFormatterImpl(FStructuredArchive& StructuredArchive)
+	UE_FORCEINLINE_HINT FArchiveFormatterType& GetFormatterImpl(FStructuredArchive& StructuredArchive)
 	{
 		return StructuredArchive.Formatter;
 	}
diff --git a/Engine/Source/Runtime/Core/Public/Serialization/StructuredArchiveChildReader.h b/Engine/Source/Runtime/Core/Public/Serialization/StructuredArchiveChildReader.h
--- a/Engine/Source/Runtime/Core/Public/Serialization/StructuredArchiveChildReader.h
+++ b/Engine/Source/Runtime/Core/Public/Serialization/StructuredArchiveChildReader.h
@@ -38,7 +38,7 @@ public:
 	FStructuredArchiveChildReader& operator=(FStructuredArchiveChildReader&&) = delete;
 	FStructuredArchiveChildReader& operator=(const FStructuredArchiveChildReader&) = delete;
 
-	FORCEINLINE FStructuredArchiveSlot GetRoot() const
+	UE_FORCEINLINE_HINT FStructuredArchiveSlot GetRoot() const
 	{
 		return Root.GetValue();
 	}
@@ -59,7 +59,7 @@ inline FStructuredArchiveChildReader::FStructuredArchiveChildReader(FStructuredA
 	Root.Emplace(StructuredArchive->Open());
 }
 
-FORCEINLINE FStructuredArchiveChildReader::~FStructuredArchiveChildReader()
+UE_FORCEINLINE_HINT FStructuredArchiveChildReader::~FStructuredArchiveChildReader()
 {
 	delete StructuredArchive;
 }
diff --git a/Engine/Source/Runtime/Core/Public/Serialization/StructuredArchiveNameHelpers.h b/Engine/Source/Runtime/Core/Public/Serialization/StructuredArchiveNameHelpers.h
--- a/Engine/Source/Runtime/Core/Public/Serialization/StructuredArchiveNameHelpers.h
+++ b/Engine/Source/Runtime/Core/Public/Serialization/StructuredArchiveNameHelpers.h
@@ -55,19 +55,19 @@ namespace UE::StructuredArchive::Private
 	};
 
 	template <typename T>
-	FORCEINLINE TNamedValue<T> MakeNamedValue(FArchiveFieldName Name, T& Value)
+	UE_FORCEINLINE_HINT TNamedValue<T> MakeNamedValue(FArchiveFieldName Name, T& Value)
 	{
 		return TNamedValue<T>{ Name, Value };
 	}
 
 	template <typename T>
-	FORCEINLINE TNamedAttribute<T> MakeNamedAttribute(FArchiveFieldName Name, T& Value)
+	UE_FORCEINLINE_HINT TNamedAttribute<T> MakeNamedAttribute(FArchiveFieldName Name, T& Value)
 	{
 		return TNamedAttribute<T>{ Name, Value };
 	}
 
 	template <typename T>
-	FORCEINLINE TOptionalNamedAttribute<T> MakeOptionalNamedAttribute(FArchiveFieldName Name, T& Value, const typename TIdentity<T>::Type& Default)
+	UE_FORCEINLINE_HINT TOptionalNamedAttribute<T> MakeOptionalNamedAttribute(FArchiveFieldName Name, T& Value, const typename TIdentity<T>::Type& Default)
 	{
 		return TOptionalNamedAttribute<T>{ Name, Value, Default };
 	}
diff --git a/Engine/Source/Runtime/Core/Public/Serialization/StructuredArchiveSlotBase.h b/Engine/Source/Runtime/Core/Public/Serialization/StructuredArchiveSlotBase.h
--- a/Engine/Source/Runtime/Core/Public/Serialization/StructuredArchiveSlotBase.h
+++ b/Engine/Source/Runtime/Core/Public/Serialization/StructuredArchiveSlotBase.h
@@ -80,7 +80,7 @@ namespace UE::StructuredArchive::Private
 		{
 		}
 #else
-		FORCEINLINE explicit FSlotBase(EPrivateToken, FStructuredArchive& InStructuredArchive)
+		UE_FORCEINLINE_HINT explicit FSlotBase(EPrivateToken, FStructuredArchive& InStructuredArchive)
 			: StructuredArchive(InStructuredArchive)
 		{
 		}
diff --git a/Engine/Source/Runtime/Core/Public/Serialization/StructuredArchiveSlots.h b/Engine/Source/Runtime/Core/Public/Serialization/StructuredArchiveSlots.h
--- a/Engine/Source/Runtime/Core/Public/Serialization/StructuredArchiveSlots.h
+++ b/Engine/Source/Runtime/Core/Public/Serialization/StructuredArchiveSlots.h
@@ -96,13 +96,13 @@ public:
 		typename EnumType,
 		std::enable_if_t<TIsEnumClass<EnumType>::Value, int> = 0
 	>
-	FORCEINLINE void operator<<(EnumType& Value)
+	UE_FORCEINLINE_HINT void operator<<(EnumType& Value)
 	{
 		*this << (__underlying_type(EnumType)&)Value;
 	}
 
 	template <typename T>
-	FORCEINLINE void operator<<(UE::StructuredArchive::Private::TNamedAttribute<T> Item)
+	UE_FORCEINLINE_HINT void operator<<(UE::StructuredArchive::Private::TNamedAttribute<T> Item)
 	{
 		EnterAttribute(Item.Name) << Item.Value;
 	}
@@ -242,14 +242,14 @@ inline void operator<<(FStructuredArchiveSlot Slot, TArray<T>& InArray)
 }
 
 template <>
-FORCEINLINE_DEBUGGABLE void operator<<(FStructuredArchiveSlot Slot, TArray<uint8>& InArray)
+UE_FORCEINLINE_HINT void operator<<(FStructuredArchiveSlot Slot, TArray<uint8>& InArray)
 {
 	Slot.Serialize(InArray);
 }
 
 #if !WITH_TEXT_ARCHIVE_SUPPORT
 	//////////// FStructuredArchiveSlot ////////////
-	FORCEINLINE FStructuredArchiveRecord FStructuredArchiveSlot::EnterRecord()
+	UE_FORCEINLINE_HINT FStructuredArchiveRecord FStructuredArchiveSlot::EnterRecord()
 	{
 		return FStructuredArchiveRecord(FStructuredArchiveRecord::EPrivateToken{}, StructuredArchive);
 	}
@@ -260,7 +260,7 @@ FORCEINLINE_DEBUGGABLE void operator<<(FStructuredArchiveSlot Slot, TArray<uint8
 		return FStructuredArchiveArray(FStructuredArchiveArray::EPrivateToken{}, StructuredArchive);
 	}
 
-	FORCEINLINE FStructuredArchiveStream FStructuredArchiveSlot::EnterStream()
+	UE_FORCEINLINE_HINT FStructuredArchiveStream FStructuredArchiveSlot::EnterStream()
 	{
 		return FStructuredArchiveStream(FStructuredArchiveStream::EPrivateToken{}, StructuredArchive);
 	}
@@ -289,151 +289,151 @@ FORCEINLINE_DEBUGGABLE void operator<<(FStructuredArchiveSlot Slot, TArray<uint8
 		}
 	}
 
-	FORCEINLINE void FStructuredArchiveSlot::operator<< (uint8& Value)
+	UE_FORCEINLINE_HINT void FStructuredArchiveSlot::operator<< (uint8& Value)
 	{
 		UE::StructuredArchive::Private::GetFormatterImpl(StructuredArchive).Serialize(Value);
 	}
 
-	FORCEINLINE void FStructuredArchiveSlot::operator<< (uint16& Value)
+	UE_FORCEINLINE_HINT void FStructuredArchiveSlot::operator<< (uint16& Value)
 	{
 		UE::StructuredArchive::Private::GetFormatterImpl(StructuredArchive).Serialize(Value);
 	}
 
-	FORCEINLINE void FStructuredArchiveSlot::operator<< (uint32& Value)
+	UE_FORCEINLINE_HINT void FStructuredArchiveSlot::operator<< (uint32& Value)
 	{
 		UE::StructuredArchive::Private::GetFormatterImpl(StructuredArchive).Serialize(Value);
 	}
 
-	FORCEINLINE void FStructuredArchiveSlot::operator<< (UTF32CHAR& Value)
+	UE_FORCEINLINE_HINT void FStructuredArchiveSlot::operator<< (UTF32CHAR& Value)
 	{
 		UE::StructuredArchive::Private::GetFormatterImpl(StructuredArchive).Serialize(Value);
 	}
 
-	FORCEINLINE void FStructuredArchiveSlot::operator<< (uint64& Value)
+	UE_FORCEINLINE_HINT void FStructuredArchiveSlot::operator<< (uint64& Value)
 	{
 		UE::StructuredArchive::Private::GetFormatterImpl(StructuredArchive).Serialize(Value);
 	}
 
-	FORCEINLINE void FStructuredArchiveSlot::operator<< (int8& Value)
+	UE_FORCEINLINE_HINT void FStructuredArchiveSlot::operator<< (int8& Value)
 	{
 		UE::StructuredArchive::Private::GetFormatterImpl(StructuredArchive).Serialize(Value);
 	}
 
-	FORCEINLINE void FStructuredArchiveSlot::operator<< (int16& Value)
+	UE_FORCEINLINE_HINT void FStructuredArchiveSlot::operator<< (int16& Value)
 	{
 		UE::StructuredArchive::Private::GetFormatterImpl(StructuredArchive).Serialize(Value);
 	}
 
-	FORCEINLINE void FStructuredArchiveSlot::operator<< (int32& Value)
+	UE_FORCEINLINE_HINT void FStructuredArchiveSlot::operator<< (int32& Value)
 	{
 		UE::StructuredArchive::Private::GetFormatterImpl(StructuredArchive).Serialize(Value);
 	}
 
-	FORCEINLINE void FStructuredArchiveSlot::operator<< (int64& Value)
+	UE_FORCEINLINE_HINT void FStructuredArchiveSlot::operator<< (int64& Value)
 	{
 		UE::StructuredArchive::Private::GetFormatterImpl(StructuredArchive).Serialize(Value);
 	}
 
-	FORCEINLINE void FStructuredArchiveSlot::operator<< (float& Value)
+	UE_FORCEINLINE_HINT void FStructuredArchiveSlot::operator<< (float& Value)
 	{
 		UE::StructuredArchive::Private::GetFormatterImpl(StructuredArchive).Serialize(Value);
 	}
 
-	FORCEINLINE void FStructuredArchiveSlot::operator<< (double& Value)
+	UE_FORCEINLINE_HINT void FStructuredArchiveSlot::operator<< (double& Value)
 	{
 		UE::StructuredArchive::Private::GetFormatterImpl(StructuredArchive).Serialize(Value);
 	}
 
-	FORCEINLINE void FStructuredArchiveSlot::operator<< (bool& Value)
+	UE_FORCEINLINE_HINT void FStructuredArchiveSlot::operator<< (bool& Value)
 	{
 		UE::StructuredArchive::Private::GetFormatterImpl(StructuredArchive).Serialize(Value);
 	}
 
-	FORCEINLINE void FStructuredArchiveSlot::operator<< (FString& Value)
+	UE_FORCEINLINE_HINT void FStructuredArchiveSlot::operator<< (FString& Value)
 	{
 		UE::StructuredArchive::Private::GetFormatterImpl(StructuredArchive).Serialize(Value);
 	}
 
-	FORCEINLINE void FStructuredArchiveSlot::operator<< (FName& Value)
+	UE_FORCEINLINE_HINT void FStructuredArchiveSlot::operator<< (FName& Value)
 	{
 		UE::StructuredArchive::Private::GetFormatterImpl(StructuredArchive).Serialize(Value);
 	}
 
-	FORCEINLINE void FStructuredArchiveSlot::operator<< (UObject*& Value)
+	UE_FORCEINLINE_HINT void FStructuredArchiveSlot::operator<< (UObject*& Value)
 	{
 		UE::StructuredArchive::Private::GetFormatterImpl(StructuredArchive).Serialize(Value);
 	}
 
 #if WITH_VERSE_VM || defined(__INTELLISENSE__)
-	FORCEINLINE void FStructuredArchiveSlot::operator<< (Verse::VCell*& Value)
+	UE_FORCEINLINE_HINT void FStructuredArchiveSlot::operator<< (Verse::VCell*& Value)
 	{
 		UE::StructuredArchive::Private::GetFormatterImpl(StructuredArchive).Serialize(Value);
 	}
 #endif
 
-	FORCEINLINE void FStructuredArchiveSlot::operator<< (FText& Value)
+	UE_FORCEINLINE_HINT void FStructuredArchiveSlot::operator<< (FText& Value)
 	{
 		UE::StructuredArchive::Private::GetFormatterImpl(StructuredArchive).Serialize(Value);
 	}
 
-	FORCEINLINE void FStructuredArchiveSlot::operator<< (FWeakObjectPtr& Value)
+	UE_FORCEINLINE_HINT void FStructuredArchiveSlot::operator<< (FWeakObjectPtr& Value)
 	{
 		UE::StructuredArchive::Private::GetFormatterImpl(StructuredArchive).Serialize(Value);
 	}
 
-	FORCEINLINE void FStructuredArchiveSlot::operator<< (FSoftObjectPath& Value)
+	UE_FORCEINLINE_HINT void FStructuredArchiveSlot::operator<< (FSoftObjectPath& Value)
 	{
 		UE::StructuredArchive::Private::GetFormatterImpl(StructuredArchive).Serialize(Value);
 	}
 
-	FORCEINLINE void FStructuredArchiveSlot::operator<< (FSoftObjectPtr& Value)
+	UE_FORCEINLINE_HINT void FStructuredArchiveSlot::operator<< (FSoftObjectPtr& Value)
 	{
 		UE::StructuredArchive::Private::GetFormatterImpl(StructuredArchive).Serialize(Value);
 	}
 
-	FORCEINLINE void FStructuredArchiveSlot::operator<< (FLazyObjectPtr& Value)
+	UE_FORCEINLINE_HINT void FStructuredArchiveSlot::operator<< (FLazyObjectPtr& Value)
 	{
 		UE::StructuredArchive::Private::GetFormatterImpl(StructuredArchive).Serialize(Value);
 	}
 
-	FORCEINLINE void FStructuredArchiveSlot::operator<< (FObjectPtr& Value)
+	UE_FORCEINLINE_HINT void FStructuredArchiveSlot::operator<< (FObjectPtr& Value)
 	{
 		UE::StructuredArchive::Private::GetFormatterImpl(StructuredArchive).Serialize(Value);
 	}
 
-	FORCEINLINE void FStructuredArchiveSlot::Serialize(TArray<uint8>& Data)
+	UE_FORCEINLINE_HINT void FStructuredArchiveSlot::Serialize(TArray<uint8>& Data)
 	{
 		UE::StructuredArchive::Private::GetFormatterImpl(StructuredArchive).Serialize(Data);
 	}
 
-	FORCEINLINE void FStructuredArchiveSlot::Serialize(void* Data, uint64 DataSize)
+	UE_FORCEINLINE_HINT void FStructuredArchiveSlot::Serialize(void* Data, uint64 DataSize)
 	{
 		UE::StructuredArchive::Private::GetFormatterImpl(StructuredArchive).Serialize(Data, DataSize);
 	}
 
 	//////////// FStructuredArchiveRecord ////////////
 
-	FORCEINLINE FStructuredArchiveSlot FStructuredArchiveRecord::EnterField(FArchiveFieldName Name)
+	UE_FORCEINLINE_HINT FStructuredArchiveSlot FStructuredArchiveRecord::EnterField(FArchiveFieldName Name)
 	{
 		return FStructuredArchiveSlot(FStructuredArchiveSlot::EPrivateToken{}, StructuredArchive);
 	}
 
-	FORCEINLINE FStructuredArchiveRecord FStructuredArchiveRecord::EnterRecord(FArchiveFieldName Name)
+	UE_FORCEINLINE_HINT FStructuredArchiveRecord FStructuredArchiveRecord::EnterRecord(FArchiveFieldName Name)
 	{
 		return EnterField(Name).EnterRecord();
 	}
 
-	FORCEINLINE FStructuredArchiveArray FStructuredArchiveRecord::EnterArray(FArchiveFieldName Name, int32& Num)
+	UE_FORCEINLINE_HINT FStructuredArchiveArray FStructuredArchiveRecord::EnterArray(FArchiveFieldName Name, int32& Num)
 	{
 		return EnterField(Name).EnterArray(Num);
 	}
 
-	FORCEINLINE FStructuredArchiveStream FStructuredArchiveRecord::EnterStream(FArchiveFieldName Name)
+	UE_FORCEINLINE_HINT FStructuredArchiveStream FStructuredArchiveRecord::EnterStream(FArchiveFieldName Name)
 	{
 		return EnterField(Name).EnterStream();
 	}
 
-	FORCEINLINE FStructuredArchiveMap FStructuredArchiveRecord::EnterMap(FArchiveFieldName Name, int32& Num)
+	UE_FORCEINLINE_HINT FStructuredArchiveMap FStructuredArchiveRecord::EnterMap(FArchiveFieldName Name, int32& Num)
 	{
 		return EnterField(Name).EnterMap(Num);
 	}
@@ -452,14 +452,14 @@ FORCEINLINE_DEBUGGABLE void operator<<(FStructuredArchiveSlot Slot, TArray<uint8
 
 	//////////// FStructuredArchiveArray ////////////
 
-	FORCEINLINE FStructuredArchiveSlot FStructuredArchiveArray::EnterElement()
+	UE_FORCEINLINE_HINT FStructuredArchiveSlot FStructuredArchiveArray::EnterElement()
 	{
 		return FStructuredArchiveSlot(FStructuredArchiveSlot::EPrivateToken{}, StructuredArchive);
 	}
 
 	//////////// FStructuredArchiveStream ////////////
 
-	FORCEINLINE FStructuredArchiveSlot FStructuredArchiveStream::EnterElement()
+	UE_FORCEINLINE_HINT FStructuredArchiveSlot FStructuredArchiveStream::EnterElement()
 	{
 		return FStructuredArchiveSlot(FStructuredArchiveSlot::EPrivateToken{}, StructuredArchive);
 	}
diff --git a/Engine/Source/Runtime/Core/Public/Stats/HitchTrackingStatScope.h b/Engine/Source/Runtime/Core/Public/Stats/HitchTrackingStatScope.h
--- a/Engine/Source/Runtime/Core/Public/Stats/HitchTrackingStatScope.h
+++ b/Engine/Source/Runtime/Core/Public/Stats/HitchTrackingStatScope.h
@@ -22,7 +22,7 @@ namespace UE::Stats
 	{
 		const PROFILER_CHAR* StatString;
 	public:
-		FORCEINLINE FHitchTrackingStatScope(const PROFILER_CHAR* InStat)
+		UE_FORCEINLINE_HINT FHitchTrackingStatScope(const PROFILER_CHAR* InStat)
 		{
 			StatString = GHitchDetected ? nullptr : InStat;
 		}
diff --git a/Engine/Source/Runtime/Core/Public/Stats/Stats.h b/Engine/Source/Runtime/Core/Public/Stats/Stats.h
--- a/Engine/Source/Runtime/Core/Public/Stats/Stats.h
+++ b/Engine/Source/Runtime/Core/Public/Stats/Stats.h
@@ -61,27 +61,27 @@ struct FStatGroup_##StatName\
 		CompileTimeEnable = InCompileTimeEnable, \
 		SortByName = InSortByName \
 	}; \
-	static FORCEINLINE const char* GetGroupName() \
+	static UE_FORCEINLINE_HINT const char* GetGroupName() \
 	{ \
 		return #StatName; \
 	} \
-	static FORCEINLINE const char* GetGroupCategory() \
+	static UE_FORCEINLINE_HINT const char* GetGroupCategory() \
 	{ \
 		return #StatCategory; \
 	} \
-	static FORCEINLINE const TCHAR* GetDescription() \
+	static UE_FORCEINLINE_HINT const TCHAR* GetDescription() \
 	{ \
 		return Description; \
 	} \
-	static FORCEINLINE bool IsDefaultEnabled() \
+	static UE_FORCEINLINE_HINT bool IsDefaultEnabled() \
 	{ \
 		return (bool)DefaultEnable; \
 	} \
-	static FORCEINLINE bool IsCompileTimeEnable() \
+	static UE_FORCEINLINE_HINT bool IsCompileTimeEnable() \
 	{ \
 		return (bool)CompileTimeEnable; \
 	} \
-	static FORCEINLINE bool GetSortByName() \
+	static UE_FORCEINLINE_HINT bool GetSortByName() \
 	{ \
 		return (bool)SortByName; \
 	} \
@@ -91,31 +91,31 @@ struct FStatGroup_##StatName\
 struct FStat_##StatName\
 { \
 	typedef FStatGroup_##GroupName TGroup; \
-	static FORCEINLINE const char* GetStatName() \
+	static UE_FORCEINLINE_HINT const char* GetStatName() \
 	{ \
 		return #StatName; \
 	} \
-	static FORCEINLINE const TCHAR* GetDescription() \
+	static UE_FORCEINLINE_HINT const TCHAR* GetDescription() \
 	{ \
 		return Description; \
 	} \
-	static FORCEINLINE EStatDataType::Type GetStatType() \
+	static UE_FORCEINLINE_HINT EStatDataType::Type GetStatType() \
 	{ \
 		return StatType; \
 	} \
-	static FORCEINLINE bool IsClearEveryFrame() \
+	static UE_FORCEINLINE_HINT bool IsClearEveryFrame() \
 	{ \
 		return EnumHasAnyFlags(GetFlags(), EStatFlags::ClearEveryFrame); \
 	} \
-	static FORCEINLINE bool IsCycleStat() \
+	static UE_FORCEINLINE_HINT bool IsCycleStat() \
 	{ \
 		return EnumHasAnyFlags(GetFlags(), EStatFlags::CycleStat); \
 	} \
-	static FORCEINLINE EStatFlags GetFlags() \
+	static UE_FORCEINLINE_HINT EStatFlags GetFlags() \
 	{ \
 		return StatFlags; \
 	} \
-	static FORCEINLINE FPlatformMemory::EMemoryCounterRegion GetMemoryRegion() \
+	static UE_FORCEINLINE_HINT FPlatformMemory::EMemoryCounterRegion GetMemoryRegion() \
 	{ \
 		return MemoryRegion; \
 	} \
diff --git a/Engine/Source/Runtime/Core/Public/Stats/StatsData.h b/Engine/Source/Runtime/Core/Public/Stats/StatsData.h
--- a/Engine/Source/Runtime/Core/Public/Stats/StatsData.h
+++ b/Engine/Source/Runtime/Core/Public/Stats/StatsData.h
@@ -181,7 +181,7 @@ struct FStatCallCountComparer
 template<>
 struct FStatNameComparer<FStatMessage>
 {
-	FORCEINLINE_STATS bool operator()( const FStatMessage& A, const FStatMessage& B ) const
+	UE_FORCEINLINE_HINT bool operator()( const FStatMessage& A, const FStatMessage& B ) const
 	{
 		return A.NameAndInfo.GetRawName().Compare( B.NameAndInfo.GetRawName() ) < 0;
 	}
@@ -343,7 +343,7 @@ protected:
 template<>
 struct FStatNameComparer<FRawStatStackNode>
 {
-	FORCEINLINE_STATS bool operator()( const FRawStatStackNode& A, const FRawStatStackNode& B ) const
+	UE_FORCEINLINE_HINT bool operator()( const FRawStatStackNode& A, const FRawStatStackNode& B ) const
 	{
 		return A.Meta.NameAndInfo.GetRawName().Compare( B.Meta.NameAndInfo.GetRawName() ) < 0;
 	}
diff --git a/Engine/Source/Runtime/Core/Public/Stats/StatsSystemTypes.h b/Engine/Source/Runtime/Core/Public/Stats/StatsSystemTypes.h
--- a/Engine/Source/Runtime/Core/Public/Stats/StatsSystemTypes.h
+++ b/Engine/Source/Runtime/Core/Public/Stats/StatsSystemTypes.h
@@ -75,38 +75,38 @@ struct TStatIdData
 
 struct TStatId
 {
-	FORCEINLINE TStatId()
+	UE_FORCEINLINE_HINT TStatId()
 		: StatIdPtr(&TStatId_NAME_None)
 	{
 	}
-	FORCEINLINE TStatId(TStatIdData const* InStatIdPtr)
+	UE_FORCEINLINE_HINT TStatId(TStatIdData const* InStatIdPtr)
 		: StatIdPtr(InStatIdPtr)
 	{
 	}
-	FORCEINLINE bool IsValidStat() const
+	UE_FORCEINLINE_HINT bool IsValidStat() const
 	{
 		return !IsNone();
 	}
-	FORCEINLINE bool IsNone() const
+	UE_FORCEINLINE_HINT bool IsNone() const
 	{
 		return StatIdPtr->IsNone();
 	}
-	FORCEINLINE TStatIdData const* GetRawPointer() const
+	UE_FORCEINLINE_HINT TStatIdData const* GetRawPointer() const
 	{
 		return StatIdPtr;
 	}
 
-	FORCEINLINE FMinimalName GetMinimalName(EMemoryOrder MemoryOrder) const
+	UE_FORCEINLINE_HINT FMinimalName GetMinimalName(EMemoryOrder MemoryOrder) const
 	{
 		return StatIdPtr->Name.Load(MemoryOrder);
 	}
 
-	FORCEINLINE FName GetName() const
+	UE_FORCEINLINE_HINT FName GetName() const
 	{
 		return MinimalNameToName(StatIdPtr->Name);
 	}
 
-	FORCEINLINE static const TStatIdData& GetStatNone()
+	UE_FORCEINLINE_HINT static const TStatIdData& GetStatNone()
 	{
 		return TStatId_NAME_None;
 	}
@@ -116,7 +116,7 @@ struct TStatId
 	 * StatIdPtr must point to a valid FName pointer.
 	 * @see FStatGroupEnableManager::GetHighPerformanceEnableForStat
 	 */
-	FORCEINLINE const ANSICHAR* GetStatDescriptionANSI() const
+	UE_FORCEINLINE_HINT const ANSICHAR* GetStatDescriptionANSI() const
 	{
 		return StatIdPtr->StatDescriptionAnsi.Get();
 	}
@@ -126,17 +126,17 @@ struct TStatId
 	 * StatIdPtr must point to a valid FName pointer.
 	 * @see FStatGroupEnableManager::GetHighPerformanceEnableForStat
 	 */
-	FORCEINLINE const WIDECHAR* GetStatDescriptionWIDE() const
+	UE_FORCEINLINE_HINT const WIDECHAR* GetStatDescriptionWIDE() const
 	{
 		return StatIdPtr->StatDescriptionWide.Get();
 	}
 
-	FORCEINLINE bool operator==(TStatId Other) const
+	UE_FORCEINLINE_HINT bool operator==(TStatId Other) const
 	{
 		return StatIdPtr == Other.StatIdPtr;
 	}
 
-	FORCEINLINE bool operator!=(TStatId Other) const
+	UE_FORCEINLINE_HINT bool operator!=(TStatId Other) const
 	{
 		return StatIdPtr != Other.StatIdPtr;
 	}
@@ -332,17 +332,17 @@ namespace EStatAllFields
 
 static_assert(EStatAllFields::StartShift > 0, "Too many stat fields.");
 
-FORCEINLINE int64 ToPackedCallCountDuration(uint32 CallCount, uint32 Duration)
+UE_FORCEINLINE_HINT int64 ToPackedCallCountDuration(uint32 CallCount, uint32 Duration)
 {
 	return (int64(CallCount) << 32) | Duration;
 }
 
-FORCEINLINE uint32 FromPackedCallCountDuration_CallCount(int64 Both)
+UE_FORCEINLINE_HINT uint32 FromPackedCallCountDuration_CallCount(int64 Both)
 {
 	return uint32(Both >> 32);
 }
 
-FORCEINLINE uint32 FromPackedCallCountDuration_Duration(int64 Both)
+UE_FORCEINLINE_HINT uint32 FromPackedCallCountDuration_Duration(int64 Both)
 {
 	return uint32(Both & MAX_uint32);
 }
@@ -361,7 +361,7 @@ class FStatNameAndInfo
 public:
 	static CORE_API const char* GpuStatCategory;
 
-	FORCEINLINE_STATS FStatNameAndInfo()
+	UE_FORCEINLINE_HINT FStatNameAndInfo()
 	{
 	}
 
@@ -414,7 +414,7 @@ public:
 	/**
 	 * Internal use, used by the deserializer
 	 */
-	FORCEINLINE_STATS void SetNumberDirect(int32 InNumber)
+	UE_FORCEINLINE_HINT void SetNumberDirect(int32 InNumber)
 	{
 		Number = InNumber;
 	}
@@ -658,7 +658,7 @@ struct FStatMessage
 		NameAndInfo.SetField<EStatOperation>(EStatOperation::SetLongName);
 	}
 
-	explicit FORCEINLINE_STATS FStatMessage(FStatNameAndInfo InStatName)
+	explicit UE_FORCEINLINE_HINT FStatMessage(FStatNameAndInfo InStatName)
 		: NameAndInfo(InStatName)
 	{
 	}
@@ -987,12 +987,12 @@ struct TStatMessage
 		return Value;
 	}
 
-	FORCEINLINE_STATS FName GetShortName() const
+	UE_FORCEINLINE_HINT FName GetShortName() const
 	{
 		return NameAndInfo.GetShortName();
 	}
 
-	FORCEINLINE_STATS FString GetDescription() const
+	UE_FORCEINLINE_HINT FString GetDescription() const
 	{
 		return NameAndInfo.GetDescription();
 	}
@@ -1352,7 +1352,7 @@ public:
 	/** Pseudo-Memory operation. */
 	template<typename TValue>
 	UE_DEPRECATED(5.3, "Use Trace/MemoryInsights and/or LLM for memory profiling.")
-		FORCEINLINE_STATS void AddMemoryMessage(FName InStatName, TValue Value)
+		UE_FORCEINLINE_HINT void AddMemoryMessage(FName InStatName, TValue Value)
 	{
 #if UE_STATS_MEMORY_PROFILER_ENABLED
 		AddStatMessage(FStatMessage(InStatName, EStatOperation::Memory, Value, false));
@@ -1366,24 +1366,24 @@ public:
 	static CORE_API void ExplicitFlush(bool DiscardCallstack = false);
 
 	/** Return true if we are currently collecting data **/
-	static FORCEINLINE_STATS bool IsCollectingData()
+	static UE_FORCEINLINE_HINT bool IsCollectingData()
 	{
 		return bPrimaryEnable;
 	}
-	static FORCEINLINE_STATS bool IsCollectingData(TStatId StatId)
+	static UE_FORCEINLINE_HINT bool IsCollectingData(TStatId StatId)
 	{
 		// we don't test StatId for nullptr here because we assume it is non-null. If it is nullptr, that indicates a problem with higher level code.
 		return !StatId.IsNone() && IsCollectingData();
 	}
 
 	/** Return true if we are currently collecting data **/
-	static FORCEINLINE_STATS bool WillEverCollectData()
+	static UE_FORCEINLINE_HINT bool WillEverCollectData()
 	{
 		return !bPrimaryDisableForever;
 	}
 
 	/** Return true if the threading is ready **/
-	static FORCEINLINE_STATS bool IsThreadingReady()
+	static UE_FORCEINLINE_HINT bool IsThreadingReady()
 	{
 		return FPlatformTLS::IsValidTlsSlot(TlsSlot);
 	}
@@ -1562,7 +1562,7 @@ public:
 		}
 	}
 
-	FORCEINLINE_STATS void Start(TStatId InStatId, bool bAlways = false, UE::FSourceLocation SourceLocation = UE::FSourceLocation::Current())
+	UE_FORCEINLINE_HINT void Start(TStatId InStatId, bool bAlways = false, UE::FSourceLocation SourceLocation = UE::FSourceLocation::Current())
 	{
 		Start(InStatId, EStatFlags::None, bAlways, SourceLocation);
 	}
@@ -1645,7 +1645,7 @@ public:
 		AutoRTFM::PushOnAbortHandler(this, [this]() { this->Stop(); });
 	}
 
-	FORCEINLINE_STATS FScopeCycleCounter(TStatId StatId, bool bAlways = false, UE::FSourceLocation SourceLocation = UE::FSourceLocation::Current())
+	UE_FORCEINLINE_HINT FScopeCycleCounter(TStatId StatId, bool bAlways = false, UE::FSourceLocation SourceLocation = UE::FSourceLocation::Current())
 		: FScopeCycleCounter(StatId, EStatFlags::None, bAlways, SourceLocation)
 	{
 	}
@@ -1664,11 +1664,11 @@ public:
 
 };
 
-FORCEINLINE void StatsPrimaryEnableAdd(int32 Value = 1)
+UE_FORCEINLINE_HINT void StatsPrimaryEnableAdd(int32 Value = 1)
 {
 	FThreadStats::PrimaryEnableAdd(Value);
 }
-FORCEINLINE void StatsPrimaryEnableSubtract(int32 Value = 1)
+UE_FORCEINLINE_HINT void StatsPrimaryEnableSubtract(int32 Value = 1)
 {
 	FThreadStats::PrimaryEnableSubtract(Value);
 }
@@ -1809,7 +1809,7 @@ struct FThreadSafeStaticStatInner : public FThreadSafeStaticStatBase
 		return TStatId(LocalHighPerformanceEnable);
 	}
 
-	FORCEINLINE FName GetStatFName() const
+	UE_FORCEINLINE_HINT FName GetStatFName() const
 	{
 		return GetStatId().GetName();
 	}
@@ -1846,12 +1846,12 @@ struct FDynamicStat : public FThreadSafeStaticStatBase
 		);
 	}
 
-	FORCEINLINE_STATS TStatId GetStatId() const
+	UE_FORCEINLINE_HINT TStatId GetStatId() const
 	{
 		return TStatId(HighPerformanceEnable.Load(EMemoryOrder::Relaxed));
 	}
 
-	FORCEINLINE FName GetStatFName() const
+	UE_FORCEINLINE_HINT FName GetStatFName() const
 	{
 		return GetStatId().GetName();
 	}
@@ -1860,11 +1860,11 @@ struct FDynamicStat : public FThreadSafeStaticStatBase
 template<class TStatData>
 struct FThreadSafeStaticStatInner<TStatData, false>
 {
-	FORCEINLINE TStatId GetStatId()
+	UE_FORCEINLINE_HINT TStatId GetStatId()
 	{
 		return TStatId();
 	}
-	FORCEINLINE FName GetStatFName() const
+	UE_FORCEINLINE_HINT FName GetStatFName() const
 	{
 		return FName();
 	}
diff --git a/Engine/Source/Runtime/Core/Public/Templates/Greater.h b/Engine/Source/Runtime/Core/Public/Templates/Greater.h
--- a/Engine/Source/Runtime/Core/Public/Templates/Greater.h
+++ b/Engine/Source/Runtime/Core/Public/Templates/Greater.h
@@ -13,7 +13,7 @@
 template <typename T = void>
 struct TGreater
 {
-	FORCEINLINE bool operator()(const T& A, const T& B) const
+	UE_FORCEINLINE_HINT bool operator()(const T& A, const T& B) const
 	{
 		return B < A;
 	}
@@ -23,7 +23,7 @@ template <>
 struct TGreater<void>
 {
 	template <typename T, typename U>
-	FORCEINLINE bool operator()(T&& A, U&& B) const
+	UE_FORCEINLINE_HINT bool operator()(T&& A, U&& B) const
 	{
 		return Forward<U>(B) < Forward<T>(A);
 	}
diff --git a/Engine/Source/Runtime/Core/Public/Templates/GuardValueAccessors.h b/Engine/Source/Runtime/Core/Public/Templates/GuardValueAccessors.h
--- a/Engine/Source/Runtime/Core/Public/Templates/GuardValueAccessors.h
+++ b/Engine/Source/Runtime/Core/Public/Templates/GuardValueAccessors.h
@@ -34,7 +34,7 @@ struct TGuardValueAccessors
 	 *
 	 * @return	a const reference to the original data value
 	 */
-	FORCEINLINE const AssignedType& GetOriginalValue() const
+	UE_FORCEINLINE_HINT const AssignedType& GetOriginalValue() const
 	{
 		return OriginalValue;
 	}
