Replace some FORCEINLINE usage in Core/CoreUObject with UE_FORCEINLINE_HINT.

diff --git a/Engine/Source/Runtime/Core/Public/Templates/IdentityFunctor.h b/Engine/Source/Runtime/Core/Public/Templates/IdentityFunctor.h
--- a/Engine/Source/Runtime/Core/Public/Templates/IdentityFunctor.h
+++ b/Engine/Source/Runtime/Core/Public/Templates/IdentityFunctor.h
@@ -10,7 +10,7 @@
 struct FIdentityFunctor
 {
 	template <typename T>
-	FORCEINLINE T&& operator()(T&& Val) const
+	UE_FORCEINLINE_HINT T&& operator()(T&& Val) const
 	{
 		return (T&&)Val;
 	}
diff --git a/Engine/Source/Runtime/Core/Public/Templates/Invoke.h b/Engine/Source/Runtime/Core/Public/Templates/Invoke.h
--- a/Engine/Source/Runtime/Core/Public/Templates/Invoke.h
+++ b/Engine/Source/Runtime/Core/Public/Templates/Invoke.h
@@ -40,7 +40,7 @@ namespace UE::Core::Private
  * See: http://en.cppreference.com/w/cpp/utility/functional/invoke
  */
 template <typename FuncType, typename... ArgTypes>
-FORCEINLINE auto Invoke(FuncType&& Func, ArgTypes&&... Args)
+UE_FORCEINLINE_HINT auto Invoke(FuncType&& Func, ArgTypes&&... Args)
 	-> decltype(((FuncType&&)Func)((ArgTypes&&)Args...))
 {
 	return ((FuncType&&)Func)((ArgTypes&&)Args...);
diff --git a/Engine/Source/Runtime/Core/Public/Templates/Less.h b/Engine/Source/Runtime/Core/Public/Templates/Less.h
--- a/Engine/Source/Runtime/Core/Public/Templates/Less.h
+++ b/Engine/Source/Runtime/Core/Public/Templates/Less.h
@@ -14,7 +14,7 @@
 template <typename T /*= void */>
 struct TLess
 {
-	FORCEINLINE bool operator()(const T& A, const T& B) const
+	UE_FORCEINLINE_HINT bool operator()(const T& A, const T& B) const
 	{
 		return A < B;
 	}
@@ -24,7 +24,7 @@ template <>
 struct TLess<void>
 {
 	template <typename T, typename U>
-	FORCEINLINE bool operator()(T&& A, U&& B) const
+	UE_FORCEINLINE_HINT bool operator()(T&& A, U&& B) const
 	{
 		return Forward<T>(A) < Forward<U>(B);
 	}
diff --git a/Engine/Source/Runtime/Core/Public/Templates/NonNullPointer.h b/Engine/Source/Runtime/Core/Public/Templates/NonNullPointer.h
--- a/Engine/Source/Runtime/Core/Public/Templates/NonNullPointer.h
+++ b/Engine/Source/Runtime/Core/Public/Templates/NonNullPointer.h
@@ -26,7 +26,7 @@ public:
 	/**
 	 * Hack that can be used under extraordinary circumstances
 	 */
-	FORCEINLINE TNonNullPtr(EDefaultConstructNonNullPtr)
+	UE_FORCEINLINE_HINT TNonNullPtr(EDefaultConstructNonNullPtr)
 		: Object(nullptr)
 	{	
 	}
@@ -34,7 +34,7 @@ public:
 	/**
 	 * nullptr constructor - not allowed.
 	 */
-	FORCEINLINE TNonNullPtr(TYPE_OF_NULLPTR)
+	UE_FORCEINLINE_HINT TNonNullPtr(TYPE_OF_NULLPTR)
 	{
 		// Essentially static_assert(false), but this way prevents GCC/Clang from crying wolf by merely inspecting the function body
 		static_assert(sizeof(ObjectType) == 0, "Tried to initialize TNonNullPtr with a null pointer!");
@@ -57,7 +57,7 @@ public:
 		typename OtherObjectType
 		UE_REQUIRES(std::is_convertible_v<OtherObjectType, ObjectType*>)
 	>
-	FORCEINLINE TNonNullPtr(TNotNull<OtherObjectType> InObject)
+	UE_FORCEINLINE_HINT TNonNullPtr(TNotNull<OtherObjectType> InObject)
 		: Object(InObject)
 	{
 	}
@@ -70,7 +70,7 @@ public:
 		typename OtherObjectType
 		UE_REQUIRES(std::is_convertible_v<OtherObjectType*, ObjectType*>)
 	>
-	FORCEINLINE TNonNullPtr(const TNonNullPtr<OtherObjectType>& Other)
+	UE_FORCEINLINE_HINT TNonNullPtr(const TNonNullPtr<OtherObjectType>& Other)
 		: Object(Other.Get())
 	{
 	}
@@ -126,12 +126,12 @@ public:
 	/**
 	 * Comparison, will also handle default constructed state
 	 */
-	FORCEINLINE bool operator==(const TNonNullPtr& Other) const
+	UE_FORCEINLINE_HINT bool operator==(const TNonNullPtr& Other) const
 	{
 		return Object == Other.Object;
 	}
 #if !PLATFORM_COMPILER_HAS_GENERATED_COMPARISON_OPERATORS
-	FORCEINLINE bool operator!=(const TNonNullPtr& Other) const
+	UE_FORCEINLINE_HINT bool operator!=(const TNonNullPtr& Other) const
 	{
 		return Object != Other.Object;
 	}
@@ -144,7 +144,7 @@ public:
 		typename OtherObjectType
 		UE_REQUIRES(UE_REQUIRES_EXPR(std::declval<ObjectType*>() == std::declval<OtherObjectType*>()))
 	>
-	FORCEINLINE bool operator==(OtherObjectType* Other) const
+	UE_FORCEINLINE_HINT bool operator==(OtherObjectType* Other) const
 	{
 		return Object == Other;
 	}
@@ -152,7 +152,7 @@ public:
 		typename OtherObjectType
 		UE_REQUIRES(UE_REQUIRES_EXPR(std::declval<OtherObjectType*>() == std::declval<ObjectType*>()))
 	>
-	FORCEINLINE friend bool operator==(OtherObjectType* Lhs, const TNonNullPtr& Rhs)
+	UE_FORCEINLINE_HINT friend bool operator==(OtherObjectType* Lhs, const TNonNullPtr& Rhs)
 	{
 		return Lhs == Rhs.Object;
 	}
@@ -161,7 +161,7 @@ public:
 		typename OtherObjectType
 		UE_REQUIRES(UE_REQUIRES_EXPR(std::declval<ObjectType*>() == std::declval<OtherObjectType*>()))
 	>
-	FORCEINLINE bool operator!=(OtherObjectType* Other) const
+	UE_FORCEINLINE_HINT bool operator!=(OtherObjectType* Other) const
 	{
 		return Object != Other;
 	}
@@ -169,7 +169,7 @@ public:
 		typename OtherObjectType
 		UE_REQUIRES(UE_REQUIRES_EXPR(std::declval<OtherObjectType*>() == std::declval<ObjectType*>()))
 	>
-	FORCEINLINE friend bool operator!=(OtherObjectType* Lhs, const TNonNullPtr& Rhs)
+	UE_FORCEINLINE_HINT friend bool operator!=(OtherObjectType* Lhs, const TNonNullPtr& Rhs)
 	{
 		return Lhs != Rhs.Object;
 	}
@@ -216,7 +216,7 @@ public:
 	 * should always be valid but might be in the EDefaultConstructNonNullPtr state 
 	 * during initialization.
 	 */
-	FORCEINLINE bool IsInitialized() const
+	UE_FORCEINLINE_HINT bool IsInitialized() const
 	{
 		return Object != nullptr;
 	}
@@ -231,11 +231,11 @@ public:
 	////////////////////////////////////////////////////
 	constexpr static bool bHasIntrusiveUnsetOptionalState = true;
 	using IntrusiveUnsetOptionalStateType = TNonNullPtr;
-	FORCEINLINE explicit TNonNullPtr(FIntrusiveUnsetOptionalState)
+	UE_FORCEINLINE_HINT explicit TNonNullPtr(FIntrusiveUnsetOptionalState)
 		: Object(nullptr)
 	{
 	}
-	FORCEINLINE bool operator==(FIntrusiveUnsetOptionalState) const
+	UE_FORCEINLINE_HINT bool operator==(FIntrusiveUnsetOptionalState) const
 	{
 		return Object == nullptr;
 	}
@@ -244,7 +244,7 @@ public:
 	//////////////////////////////////////////////////
 
 private:
-	friend FORCEINLINE uint32 GetTypeHash(const TNonNullPtr& InPtr)
+	friend UE_FORCEINLINE_HINT uint32 GetTypeHash(const TNonNullPtr& InPtr)
 	{
 		return PointerHash(InPtr.Object);
 	}
diff --git a/Engine/Source/Runtime/Core/Public/Templates/PimplPtr.h b/Engine/Source/Runtime/Core/Public/Templates/PimplPtr.h
--- a/Engine/Source/Runtime/Core/Public/Templates/PimplPtr.h
+++ b/Engine/Source/Runtime/Core/Public/Templates/PimplPtr.h
@@ -233,8 +233,8 @@ public:
 		}
 	}
 
-	FORCEINLINE bool operator==(TYPE_OF_NULLPTR) { return !IsValid(); }
-	FORCEINLINE bool operator!=(TYPE_OF_NULLPTR) { return  IsValid(); }
+	UE_FORCEINLINE_HINT bool operator==(TYPE_OF_NULLPTR) { return !IsValid(); }
+	UE_FORCEINLINE_HINT bool operator!=(TYPE_OF_NULLPTR) { return  IsValid(); }
 
 private:
 	T* Ptr = nullptr;
@@ -304,8 +304,8 @@ public:
 };
 
 #if !PLATFORM_COMPILER_HAS_GENERATED_COMPARISON_OPERATORS
-template <typename T, EPimplPtrMode Mode> FORCEINLINE bool operator==(TYPE_OF_NULLPTR, const TPimplPtr<T, Mode>& Ptr) { return !Ptr.IsValid(); }
-template <typename T, EPimplPtrMode Mode> FORCEINLINE bool operator!=(TYPE_OF_NULLPTR, const TPimplPtr<T, Mode>& Ptr) { return  Ptr.IsValid(); }
+template <typename T, EPimplPtrMode Mode> UE_FORCEINLINE_HINT bool operator==(TYPE_OF_NULLPTR, const TPimplPtr<T, Mode>& Ptr) { return !Ptr.IsValid(); }
+template <typename T, EPimplPtrMode Mode> UE_FORCEINLINE_HINT bool operator!=(TYPE_OF_NULLPTR, const TPimplPtr<T, Mode>& Ptr) { return  Ptr.IsValid(); }
 #endif
 
 /**
diff --git a/Engine/Source/Runtime/Core/Public/Templates/RefCounting.h b/Engine/Source/Runtime/Core/Public/Templates/RefCounting.h
--- a/Engine/Source/Runtime/Core/Public/Templates/RefCounting.h
+++ b/Engine/Source/Runtime/Core/Public/Templates/RefCounting.h
@@ -455,7 +455,7 @@ class TRefCountPtr
 	using ReferenceType = ReferencedType*;
 
 public:
-	FORCEINLINE TRefCountPtr() = default;
+	UE_FORCEINLINE_HINT TRefCountPtr() = default;
 
 	TRefCountPtr(ReferencedType* InReference, bool bAddRef = true)
 	{
@@ -525,13 +525,13 @@ public:
 		return *this;
 	}
 
-	FORCEINLINE TRefCountPtr& operator=(const TRefCountPtr& InPtr)
+	UE_FORCEINLINE_HINT TRefCountPtr& operator=(const TRefCountPtr& InPtr)
 	{
 		return *this = InPtr.Reference;
 	}
 
 	template<typename CopyReferencedType>
-	FORCEINLINE TRefCountPtr& operator=(const TRefCountPtr<CopyReferencedType>& InPtr)
+	UE_FORCEINLINE_HINT TRefCountPtr& operator=(const TRefCountPtr<CopyReferencedType>& InPtr)
 	{
 		return *this = InPtr.GetReference();
 	}
@@ -565,12 +565,12 @@ public:
 		return *this;
 	}
 
-	FORCEINLINE ReferencedType* operator->() const
+	UE_FORCEINLINE_HINT ReferencedType* operator->() const
 	{
 		return Reference;
 	}
 
-	FORCEINLINE operator ReferenceType() const
+	UE_FORCEINLINE_HINT operator ReferenceType() const
 	{
 		return Reference;
 	}
@@ -581,22 +581,22 @@ public:
 		return &Reference;
 	}
 
-	FORCEINLINE ReferencedType* GetReference() const
+	UE_FORCEINLINE_HINT ReferencedType* GetReference() const
 	{
 		return Reference;
 	}
 
-	FORCEINLINE friend bool IsValidRef(const TRefCountPtr& InReference)
+	UE_FORCEINLINE_HINT friend bool IsValidRef(const TRefCountPtr& InReference)
 	{
 		return InReference.Reference != nullptr;
 	}
 
-	FORCEINLINE bool IsValid() const
+	UE_FORCEINLINE_HINT bool IsValid() const
 	{
 		return Reference != nullptr;
 	}
 
-	FORCEINLINE void SafeRelease()
+	UE_FORCEINLINE_HINT void SafeRelease()
 	{
 		*this = nullptr;
 	}
@@ -636,12 +636,12 @@ private:
 	friend class TRefCountPtr;
 
 public:
-	FORCEINLINE bool operator==(const TRefCountPtr& B) const
+	UE_FORCEINLINE_HINT bool operator==(const TRefCountPtr& B) const
 	{
 		return GetReference() == B.GetReference();
 	}
 
-	FORCEINLINE bool operator==(ReferencedType* B) const
+	UE_FORCEINLINE_HINT bool operator==(ReferencedType* B) const
 	{
 		return GetReference() == B;
 	}
@@ -651,14 +651,14 @@ ALIAS_TEMPLATE_TYPE_LAYOUT(template<typename T>, TRefCountPtr<T>, void*);
 
 #if !PLATFORM_COMPILER_HAS_GENERATED_COMPARISON_OPERATORS
 template<typename ReferencedType>
-FORCEINLINE bool operator==(ReferencedType* A, const TRefCountPtr<ReferencedType>& B)
+UE_FORCEINLINE_HINT bool operator==(ReferencedType* A, const TRefCountPtr<ReferencedType>& B)
 {
 	return A == B.GetReference();
 }
 #endif
 
 template<typename ReferencedType>
-FORCEINLINE uint32 GetTypeHash(const TRefCountPtr<ReferencedType>& InPtr)
+UE_FORCEINLINE_HINT uint32 GetTypeHash(const TRefCountPtr<ReferencedType>& InPtr)
 {
 	return GetTypeHash(InPtr.GetReference());
 }
diff --git a/Engine/Source/Runtime/Core/Public/Templates/ReversePredicate.h b/Engine/Source/Runtime/Core/Public/Templates/ReversePredicate.h
--- a/Engine/Source/Runtime/Core/Public/Templates/ReversePredicate.h
+++ b/Engine/Source/Runtime/Core/Public/Templates/ReversePredicate.h
@@ -21,7 +21,7 @@ public:
 	}
 
 	template <typename T>
-	FORCEINLINE bool operator()( T&& A, T&& B ) const
+	UE_FORCEINLINE_HINT bool operator()( T&& A, T&& B ) const
 	{
 		return Invoke( Predicate, Forward<T>(B), Forward<T>(A) );
 	}
diff --git a/Engine/Source/Runtime/Core/Public/Templates/SharedPointer.h b/Engine/Source/Runtime/Core/Public/Templates/SharedPointer.h
--- a/Engine/Source/Runtime/Core/Public/Templates/SharedPointer.h
+++ b/Engine/Source/Runtime/Core/Public/Templates/SharedPointer.h
@@ -124,7 +124,7 @@
  * @param  InSharedRef  The shared reference to cast
  */
 template< class CastToType, class CastFromType, ESPMode Mode >
-[[nodiscard]] FORCEINLINE TSharedRef< CastToType, Mode > StaticCastSharedRef( TSharedRef< CastFromType, Mode > const& InSharedRef )
+[[nodiscard]] UE_FORCEINLINE_HINT TSharedRef< CastToType, Mode > StaticCastSharedRef( TSharedRef< CastFromType, Mode > const& InSharedRef )
 {
 	return TSharedRef< CastToType, Mode >( InSharedRef, SharedPointerInternals::FStaticCastTag() );
 }
@@ -135,7 +135,7 @@ namespace UE::Core::Private
 	// Needed to work around an Android compiler bug - we need to construct a TSharedRef
 	// from MakeShared without making MakeShared a friend in order to access the private constructor.
 	template <typename ObjectType, ESPMode Mode>
-	FORCEINLINE TSharedRef<ObjectType, Mode> MakeSharedRef(ObjectType* InObject, SharedPointerInternals::TReferenceControllerBase<Mode>* InSharedReferenceCount)
+	UE_FORCEINLINE_HINT TSharedRef<ObjectType, Mode> MakeSharedRef(ObjectType* InObject, SharedPointerInternals::TReferenceControllerBase<Mode>* InSharedReferenceCount)
 	{
 		return TSharedRef<ObjectType, Mode>(InObject, InSharedReferenceCount);
 	}
@@ -451,7 +451,7 @@ public:
 	 *
 	 * @return  Shared pointer to the object
 	 */
-	[[nodiscard]] FORCEINLINE TSharedPtr<ObjectType, Mode> ToSharedPtr() const
+	[[nodiscard]] UE_FORCEINLINE_HINT TSharedPtr<ObjectType, Mode> ToSharedPtr() const
 	{
 		return TSharedPtr<ObjectType, Mode>(*this);
 	}
@@ -461,7 +461,7 @@ public:
 	 *
 	 * @return  Weak pointer to the object
 	 */
-	[[nodiscard]] FORCEINLINE TWeakPtr<ObjectType, Mode> ToWeakPtr() const
+	[[nodiscard]] UE_FORCEINLINE_HINT TWeakPtr<ObjectType, Mode> ToWeakPtr() const
 	{
 		return TWeakPtr<ObjectType, Mode>(*this);
 	}
@@ -508,7 +508,7 @@ public:
 	 *
 	 * @return  Number of shared references to the object (including this reference.)
 	 */
-	[[nodiscard]] FORCEINLINE int32 GetSharedReferenceCount() const
+	[[nodiscard]] UE_FORCEINLINE_HINT int32 GetSharedReferenceCount() const
 	{
 		return SharedReferenceCount.GetSharedReferenceCount();
 	}
@@ -524,7 +524,7 @@ public:
 	 *
 	 * @return  True if there is only one shared reference to the object, and this is it!
 	 */
-	[[nodiscard]] FORCEINLINE bool IsUnique() const
+	[[nodiscard]] UE_FORCEINLINE_HINT bool IsUnique() const
 	{
 		return SharedReferenceCount.IsUnique();
 	}
@@ -626,7 +626,7 @@ private:
 	 *
 	 * @return  True if the shared reference is valid and can be dereferenced
 	 */
-	[[nodiscard]] FORCEINLINE bool IsValid() const
+	[[nodiscard]] UE_FORCEINLINE_HINT bool IsValid() const
 	{
 		return Object != nullptr;
 	}
@@ -1052,7 +1052,7 @@ public:
 	 *
 	 * @return  Weak pointer to the object
 	 */
-	[[nodiscard]] FORCEINLINE TWeakPtr<ObjectType, Mode> ToWeakPtr() const
+	[[nodiscard]] UE_FORCEINLINE_HINT TWeakPtr<ObjectType, Mode> ToWeakPtr() const
 	{
 		return TWeakPtr<ObjectType, Mode>(*this);
 	}
@@ -1062,7 +1062,7 @@ public:
 	 *
 	 * @return  The object owned by this shared pointer, or nullptr
 	 */
-	[[nodiscard]] FORCEINLINE ObjectType* Get() const
+	[[nodiscard]] UE_FORCEINLINE_HINT ObjectType* Get() const
 	{
 		return Object;
 	}
@@ -1072,7 +1072,7 @@ public:
 	 *
 	 * @return  True if the shared pointer is valid and can be dereferenced
 	 */
-	[[nodiscard]] FORCEINLINE explicit operator bool() const
+	[[nodiscard]] UE_FORCEINLINE_HINT explicit operator bool() const
 	{
 		return Object != nullptr;
 	}
@@ -1082,7 +1082,7 @@ public:
 	 *
 	 * @return  True if the shared pointer is valid and can be dereferenced
 	 */
-	[[nodiscard]] FORCEINLINE const bool IsValid() const
+	[[nodiscard]] UE_FORCEINLINE_HINT const bool IsValid() const
 	{
 		return Object != nullptr;
 	}
@@ -1117,7 +1117,7 @@ public:
 	 * Resets this shared pointer, removing a reference to the object.  If there are no other shared
 	 * references to the object then it will be destroyed.
 	 */
-	FORCEINLINE void Reset()
+	UE_FORCEINLINE_HINT void Reset()
 	{
  		*this = TSharedPtr< ObjectType, Mode >();
 	}
@@ -1128,7 +1128,7 @@ public:
 	 *
 	 * @return  Number of shared references to the object (including this reference.)
 	 */
-	[[nodiscard]] FORCEINLINE int32 GetSharedReferenceCount() const
+	[[nodiscard]] UE_FORCEINLINE_HINT int32 GetSharedReferenceCount() const
 	{
 		return SharedReferenceCount.GetSharedReferenceCount();
 	}
@@ -1140,7 +1140,7 @@ public:
 	 *
 	 * @return  True if there is only one shared reference to the object, and this is it!
 	 */
-	[[nodiscard]] FORCEINLINE bool IsUnique() const
+	[[nodiscard]] UE_FORCEINLINE_HINT bool IsUnique() const
 	{
 		return SharedReferenceCount.IsUnique();
 	}
@@ -1509,7 +1509,7 @@ public:
 	 *
 	 * @return  Shared pointer for this object (will only be valid if still referenced!)
 	 */
-	[[nodiscard]] FORCEINLINE TSharedPtr< ObjectType, Mode > Pin() const&
+	[[nodiscard]] UE_FORCEINLINE_HINT TSharedPtr< ObjectType, Mode > Pin() const&
 	{
 		return TSharedPtr< ObjectType, Mode >( *this );
 	}
@@ -1522,7 +1522,7 @@ public:
 	 *
 	 * @return  Shared pointer for this object (will only be valid if still referenced!)
 	 */
-	[[nodiscard]] FORCEINLINE TSharedPtr< ObjectType, Mode > Pin() &&
+	[[nodiscard]] UE_FORCEINLINE_HINT TSharedPtr< ObjectType, Mode > Pin() &&
 	{
 		return TSharedPtr< ObjectType, Mode >( MoveTemp( *this ) );
 	}
@@ -1532,7 +1532,7 @@ public:
 	 *
 	 * @return  True if the weak pointer is valid and a pin operator would have succeeded
 	 */
-	[[nodiscard]] FORCEINLINE bool IsValid() const
+	[[nodiscard]] UE_FORCEINLINE_HINT bool IsValid() const
 	{
 		return Object != nullptr && WeakReferenceCount.IsValid();
 	}
@@ -1541,7 +1541,7 @@ public:
 	 * Resets this weak pointer, removing a weak reference to the object.  If there are no other shared
 	 * or weak references to the object, then the tracking object will be destroyed.
 	 */
-	FORCEINLINE void Reset()
+	UE_FORCEINLINE_HINT void Reset()
 	{
 		*this = TWeakPtr< ObjectType, Mode >();
 	}
@@ -1549,12 +1549,12 @@ public:
 	/**
 	 * Returns true if the object this weak pointer points to is the same as the specified object pointer.
 	 */
-	[[nodiscard]] FORCEINLINE bool HasSameObject( const void* InOtherPtr ) const
+	[[nodiscard]] UE_FORCEINLINE_HINT bool HasSameObject( const void* InOtherPtr ) const
 	{
 		return Pin().Get() == InOtherPtr;
 	}
 
-	FORCEINLINE uint32 GetWeakPtrTypeHash() const
+	UE_FORCEINLINE_HINT uint32 GetWeakPtrTypeHash() const
 	{
 		return ::PointerHash( Object );
 	}
@@ -1777,7 +1777,7 @@ protected:
 	 * @return	Returns this object as a shared pointer
 	 */
 	template< class OtherType >
-	[[nodiscard]] FORCEINLINE static TSharedRef< OtherType, Mode > SharedThis( OtherType* ThisPtr )
+	[[nodiscard]] UE_FORCEINLINE_HINT static TSharedRef< OtherType, Mode > SharedThis( OtherType* ThisPtr )
 	{
 		return StaticCastSharedRef< OtherType >( ThisPtr->AsShared() );
 	}
@@ -1791,7 +1791,7 @@ protected:
 	 * @return	Returns this object as a shared pointer (const)
 	 */
 	template< class OtherType >
-	[[nodiscard]] FORCEINLINE static TSharedRef< OtherType const, Mode > SharedThis( const OtherType* ThisPtr )
+	[[nodiscard]] UE_FORCEINLINE_HINT static TSharedRef< OtherType const, Mode > SharedThis( const OtherType* ThisPtr )
 	{
 		return StaticCastSharedRef< OtherType const >( ThisPtr->AsShared() );
 	}
@@ -1837,7 +1837,7 @@ public:		// @todo: Ideally this would be private, but template sharing problems
 	 *  - Is currently within its constructor (so the shared instance isn't yet available).
 	 *  - Is currently within its destructor (so the shared instance is no longer available).
 	 */
-	[[nodiscard]] FORCEINLINE bool DoesSharedInstanceExist() const
+	[[nodiscard]] UE_FORCEINLINE_HINT bool DoesSharedInstanceExist() const
 	{
 		return WeakThis.IsValid();
 	}
@@ -1859,7 +1859,7 @@ protected:
 	TSharedFromThis( TSharedFromThis const& ) { }
 
 	/** Hidden stub assignment operator */
-	FORCEINLINE TSharedFromThis& operator=( TSharedFromThis const& )
+	UE_FORCEINLINE_HINT TSharedFromThis& operator=( TSharedFromThis const& )
 	{
 		return *this;
 	}
@@ -1883,7 +1883,7 @@ private:
  * @param  InSharedPtr  The shared pointer to cast
  */
 template< class CastToType, class CastFromType, ESPMode Mode >
-[[nodiscard]] FORCEINLINE TSharedPtr< CastToType, Mode > StaticCastSharedPtr(TSharedPtr< CastFromType, Mode > const& InSharedPtr)
+[[nodiscard]] UE_FORCEINLINE_HINT TSharedPtr< CastToType, Mode > StaticCastSharedPtr(TSharedPtr< CastFromType, Mode > const& InSharedPtr)
 {
 	return TSharedPtr< CastToType, Mode >(InSharedPtr, SharedPointerInternals::FStaticCastTag());
 }
@@ -1895,7 +1895,7 @@ template< class CastToType, class CastFromType, ESPMode Mode >
  * @param  InWeakPtr  The weak pointer to cast
  */
 template< class CastToType, class CastFromType, ESPMode Mode >
-[[nodiscard]] FORCEINLINE TWeakPtr< CastToType, Mode > StaticCastWeakPtr(TWeakPtr< CastFromType, Mode > const& InWeakPtr)
+[[nodiscard]] UE_FORCEINLINE_HINT TWeakPtr< CastToType, Mode > StaticCastWeakPtr(TWeakPtr< CastFromType, Mode > const& InWeakPtr)
 {
 	return TWeakPtr< CastToType, Mode >(InWeakPtr, SharedPointerInternals::FStaticCastTag());
 }
@@ -1907,7 +1907,7 @@ template< class CastToType, class CastFromType, ESPMode Mode >
  * @param  InSharedRef  The shared reference to cast
  */
 template< class CastToType, class CastFromType, ESPMode Mode >
-[[nodiscard]] FORCEINLINE TSharedRef< CastToType, Mode > ConstCastSharedRef( TSharedRef< CastFromType, Mode > const& InSharedRef )
+[[nodiscard]] UE_FORCEINLINE_HINT TSharedRef< CastToType, Mode > ConstCastSharedRef( TSharedRef< CastFromType, Mode > const& InSharedRef )
 {
 	return TSharedRef< CastToType, Mode >( InSharedRef, SharedPointerInternals::FConstCastTag() );
 }
@@ -1919,7 +1919,7 @@ template< class CastToType, class CastFromType, ESPMode Mode >
  * @param  InSharedPtr  The shared pointer to cast
  */
 template< class CastToType, class CastFromType, ESPMode Mode >
-[[nodiscard]] FORCEINLINE TSharedPtr< CastToType, Mode > ConstCastSharedPtr( TSharedPtr< CastFromType, Mode > const& InSharedPtr )
+[[nodiscard]] UE_FORCEINLINE_HINT TSharedPtr< CastToType, Mode > ConstCastSharedPtr( TSharedPtr< CastFromType, Mode > const& InSharedPtr )
 {
 	return TSharedPtr< CastToType, Mode >( InSharedPtr, SharedPointerInternals::FConstCastTag() );
 }
@@ -1931,7 +1931,7 @@ template< class CastToType, class CastFromType, ESPMode Mode >
  * @param  InWeakPtr  The weak pointer to cast
  */
 template< class CastToType, class CastFromType, ESPMode Mode >
-[[nodiscard]] FORCEINLINE TWeakPtr< CastToType, Mode > ConstCastWeakPtr(TWeakPtr< CastFromType, Mode > const& InWeakPtr)
+[[nodiscard]] UE_FORCEINLINE_HINT TWeakPtr< CastToType, Mode > ConstCastWeakPtr(TWeakPtr< CastFromType, Mode > const& InWeakPtr)
 {
 	return TWeakPtr< CastToType, Mode >(InWeakPtr, SharedPointerInternals::FConstCastTag());
 }
diff --git a/Engine/Source/Runtime/Core/Public/Templates/SharedPointerInternals.h b/Engine/Source/Runtime/Core/Public/Templates/SharedPointerInternals.h
--- a/Engine/Source/Runtime/Core/Public/Templates/SharedPointerInternals.h
+++ b/Engine/Source/Runtime/Core/Public/Templates/SharedPointerInternals.h
@@ -46,7 +46,7 @@ namespace SharedPointerInternals
 		using RefCountType = std::conditional_t<Mode == ESPMode::ThreadSafe, std::atomic<int32>, int32>;
 
 	public:
-		FORCEINLINE explicit TReferenceControllerBase() = default;
+		UE_FORCEINLINE_HINT explicit TReferenceControllerBase() = default;
 
 		// Number of shared references to this object.  When this count reaches zero, the associated object
 		// will be destroyed (even if there are still weak references!), but not the reference controller.
@@ -97,7 +97,7 @@ namespace SharedPointerInternals
 		}
 
 		/** Checks if there is exactly one reference left to the object. */
-		FORCEINLINE bool IsUnique() const
+		UE_FORCEINLINE_HINT bool IsUnique() const
 		{
 			return 1 == GetSharedReferenceCount();
 		}
@@ -449,7 +449,7 @@ namespace SharedPointerInternals
 	template <typename Type>
 	struct DefaultDeleter
 	{
-		FORCEINLINE void operator()(Type* Object) const
+		UE_FORCEINLINE_HINT void operator()(Type* Object) const
 		{
 			delete Object;
 		}
@@ -486,13 +486,13 @@ namespace SharedPointerInternals
 		ObjectType* Object;
 
 		/** Construct implicitly from a nullptr */
-		FORCEINLINE TRawPtrProxy( TYPE_OF_NULLPTR )
+		UE_FORCEINLINE_HINT TRawPtrProxy( TYPE_OF_NULLPTR )
 			: Object( nullptr )
 		{
 		}
 
 		/** Construct explicitly from an object */
-		explicit FORCEINLINE TRawPtrProxy( ObjectType* InObject )
+		explicit UE_FORCEINLINE_HINT TRawPtrProxy( ObjectType* InObject )
 			: Object( InObject )
 		{
 		}
@@ -536,7 +536,7 @@ namespace SharedPointerInternals
 	public:
 
 		/** Constructor for an empty shared referencer object */
-		FORCEINLINE FSharedReferencer()
+		UE_FORCEINLINE_HINT FSharedReferencer()
 			: ReferenceController( nullptr )
 		{ }
 
@@ -668,7 +668,7 @@ namespace SharedPointerInternals
 		 *
 		 * @return  True if reference is valid
 		 */
-		FORCEINLINE const bool IsValid() const
+		UE_FORCEINLINE_HINT const bool IsValid() const
 		{
 			return ReferenceController != nullptr;
 		}
@@ -678,7 +678,7 @@ namespace SharedPointerInternals
 		 *
 		 * @return  Number of shared references to the object (including this reference.)
 		 */
-		FORCEINLINE const int32 GetSharedReferenceCount() const
+		UE_FORCEINLINE_HINT const int32 GetSharedReferenceCount() const
 		{
 			return ReferenceController != nullptr ? ReferenceController->GetSharedReferenceCount() : 0;
 		}
@@ -689,7 +689,7 @@ namespace SharedPointerInternals
 		 *
 		 * @return  True if there is only one shared reference to the object, and this is it!
 		 */
-		FORCEINLINE const bool IsUnique() const
+		UE_FORCEINLINE_HINT const bool IsUnique() const
 		{
 			return ReferenceController != nullptr && ReferenceController->IsUnique();
 		}
@@ -716,7 +716,7 @@ namespace SharedPointerInternals
 	public:
 
 		/** Default constructor with empty counter */
-		FORCEINLINE FWeakReferencer()
+		UE_FORCEINLINE_HINT FWeakReferencer()
 			: ReferenceController( nullptr )
 		{ }
 
@@ -798,7 +798,7 @@ namespace SharedPointerInternals
 		 *
 		 * @return  True if reference is valid
 		 */
-		FORCEINLINE const bool IsValid() const
+		UE_FORCEINLINE_HINT const bool IsValid() const
 		{
 			return ReferenceController != nullptr && ReferenceController->GetSharedReferenceCount() > 0;
 		}
diff --git a/Engine/Source/Runtime/Core/Public/Templates/Sorting.h b/Engine/Source/Runtime/Core/Public/Templates/Sorting.h
--- a/Engine/Source/Runtime/Core/Public/Templates/Sorting.h
+++ b/Engine/Source/Runtime/Core/Public/Templates/Sorting.h
@@ -20,8 +20,8 @@ struct TDereferenceWrapper
 		: Predicate( InPredicate ) {}
   
 	/** Pass through for non-pointer types */
-	FORCEINLINE bool operator()( T& A, T& B ) { return Predicate( A, B ); } 
-	FORCEINLINE bool operator()( const T& A, const T& B ) const { return Predicate( A, B ); } 
+	UE_FORCEINLINE_HINT bool operator()( T& A, T& B ) { return Predicate( A, B ); } 
+	UE_FORCEINLINE_HINT bool operator()( const T& A, const T& B ) const { return Predicate( A, B ); } 
 };
 /** Partially specialized version of the above class */
 template<typename T, class PREDICATE_CLASS> 
@@ -33,7 +33,7 @@ struct TDereferenceWrapper<T*, PREDICATE_CLASS>
 		: Predicate( InPredicate ) {}
   
 	/** Dereference pointers */
-	FORCEINLINE bool operator()( T* A, T* B ) const 
+	UE_FORCEINLINE_HINT bool operator()( T* A, T* B ) const 
 	{
 		return Predicate( *A, *B ); 
 	} 
@@ -511,7 +511,7 @@ void RadixSort32( ValueType* RESTRICT Dst, ValueType* RESTRICT Src, CountType Nu
 template< typename T >
 struct TRadixSortKeyCastUint32
 {
-	FORCEINLINE uint32 operator()( const T& Value ) const
+	UE_FORCEINLINE_HINT uint32 operator()( const T& Value ) const
 	{
 		return (uint32)Value;
 	}
@@ -682,7 +682,7 @@ void RadixSort64(ValueType* RESTRICT Array, ValueType* RESTRICT Buffer, CountTyp
 template< typename T >
 struct TRadixSortKeyCastUint64
 {
-	FORCEINLINE uint64 operator()(const T& Value) const
+	UE_FORCEINLINE_HINT uint64 operator()(const T& Value) const
 	{
 		return (uint64)Value;
 	}
diff --git a/Engine/Source/Runtime/Core/Public/Templates/Tuple.h b/Engine/Source/Runtime/Core/Public/Templates/Tuple.h
--- a/Engine/Source/Runtime/Core/Public/Templates/Tuple.h
+++ b/Engine/Source/Runtime/Core/Public/Templates/Tuple.h
@@ -129,7 +129,7 @@ namespace UE::Core::Private::Tuple
 		}
 
 		template <typename TupleType>
-		static FORCEINLINE decltype(auto) Get(TupleType&& Tuple)
+		static UE_FORCEINLINE_HINT decltype(auto) Get(TupleType&& Tuple)
 		{
 			return GetImpl(Tuple, Forward<TupleType>(Tuple));
 		}
@@ -182,13 +182,13 @@ namespace UE::Core::Private::Tuple
 	struct TTupleElementGetterByType
 	{
 		template <uint32 DeducedIndex, typename TupleType>
-		static FORCEINLINE decltype(auto) GetImpl(const volatile TTupleBaseElement<Type, DeducedIndex, TupleSize>&, TupleType&& Tuple)
+		static UE_FORCEINLINE_HINT decltype(auto) GetImpl(const volatile TTupleBaseElement<Type, DeducedIndex, TupleSize>&, TupleType&& Tuple)
 		{
 			return TTupleElementGetterByIndex<DeducedIndex, TupleSize>::Get(Forward<TupleType>(Tuple));
 		}
 
 		template <typename TupleType>
-		static FORCEINLINE decltype(auto) Get(TupleType&& Tuple)
+		static UE_FORCEINLINE_HINT decltype(auto) Get(TupleType&& Tuple)
 		{
 			return GetImpl(Tuple, Forward<TupleType>(Tuple));
 		}
@@ -218,7 +218,7 @@ namespace UE::Core::Private::Tuple
 	struct FEqualityHelper
 	{
 		template <typename TupleType>
-		FORCEINLINE static bool Compare(const TupleType& Lhs, const TupleType& Rhs)
+		UE_FORCEINLINE_HINT static bool Compare(const TupleType& Lhs, const TupleType& Rhs)
 		{
 			return Lhs.template Get<ArgToCompare>() == Rhs.template Get<ArgToCompare>() && FEqualityHelper<ArgCount, ArgToCompare + 1>::Compare(Lhs, Rhs);
 		}
@@ -228,7 +228,7 @@ namespace UE::Core::Private::Tuple
 	struct FEqualityHelper<ArgCount, ArgCount>
 	{
 		template <typename TupleType>
-		FORCEINLINE static bool Compare(const TupleType& Lhs, const TupleType& Rhs)
+		UE_FORCEINLINE_HINT static bool Compare(const TupleType& Lhs, const TupleType& Rhs)
 		{
 			return true;
 		}
@@ -238,7 +238,7 @@ namespace UE::Core::Private::Tuple
 	struct TLessThanHelper
 	{
 		template <typename TupleType>
-		FORCEINLINE static bool Do(const TupleType& Lhs, const TupleType& Rhs)
+		UE_FORCEINLINE_HINT static bool Do(const TupleType& Lhs, const TupleType& Rhs)
 		{
 			return Lhs.template Get<ArgToCompare>() < Rhs.template Get<ArgToCompare>() || (!(Rhs.template Get<ArgToCompare>() < Lhs.template Get<ArgToCompare>()) && TLessThanHelper<NumArgs, ArgToCompare + 1>::Do(Lhs, Rhs));
 		}
@@ -248,7 +248,7 @@ namespace UE::Core::Private::Tuple
 	struct TLessThanHelper<NumArgs, ArgToCompare, true>
 	{
 		template <typename TupleType>
-		FORCEINLINE static bool Do(const TupleType& Lhs, const TupleType& Rhs)
+		UE_FORCEINLINE_HINT static bool Do(const TupleType& Lhs, const TupleType& Rhs)
 		{
 			return Lhs.template Get<ArgToCompare>() < Rhs.template Get<ArgToCompare>();
 		}
@@ -258,7 +258,7 @@ namespace UE::Core::Private::Tuple
 	struct TLessThanHelper<NumArgs, NumArgs, false>
 	{
 		template <typename TupleType>
-		FORCEINLINE static bool Do(const TupleType& Lhs, const TupleType& Rhs)
+		UE_FORCEINLINE_HINT static bool Do(const TupleType& Lhs, const TupleType& Rhs)
 		{
 			return false;
 		}
@@ -294,23 +294,23 @@ namespace UE::Core::Private::Tuple
 		TTupleBase& operator=(TTupleBase&& Other) = default;
 		TTupleBase& operator=(const TTupleBase& Other) = default;
 
-		template <uint32 Index> FORCEINLINE decltype(auto) Get()               &  { static_assert(Index < sizeof...(Types), "Invalid index passed to TTuple::Get"); return TTupleElementGetterByIndex<Index, sizeof...(Types)>::Get(static_cast<               TTupleBase& >(*this)); }
-		template <uint32 Index> FORCEINLINE decltype(auto) Get() const         &  { static_assert(Index < sizeof...(Types), "Invalid index passed to TTuple::Get"); return TTupleElementGetterByIndex<Index, sizeof...(Types)>::Get(static_cast<const          TTupleBase& >(*this)); }
-		template <uint32 Index> FORCEINLINE decltype(auto) Get()       volatile&  { static_assert(Index < sizeof...(Types), "Invalid index passed to TTuple::Get"); return TTupleElementGetterByIndex<Index, sizeof...(Types)>::Get(static_cast<      volatile TTupleBase& >(*this)); }
-		template <uint32 Index> FORCEINLINE decltype(auto) Get() const volatile&  { static_assert(Index < sizeof...(Types), "Invalid index passed to TTuple::Get"); return TTupleElementGetterByIndex<Index, sizeof...(Types)>::Get(static_cast<const volatile TTupleBase& >(*this)); }
-		template <uint32 Index> FORCEINLINE decltype(auto) Get()               && { static_assert(Index < sizeof...(Types), "Invalid index passed to TTuple::Get"); return TTupleElementGetterByIndex<Index, sizeof...(Types)>::Get(static_cast<               TTupleBase&&>(*this)); }
-		template <uint32 Index> FORCEINLINE decltype(auto) Get() const         && { static_assert(Index < sizeof...(Types), "Invalid index passed to TTuple::Get"); return TTupleElementGetterByIndex<Index, sizeof...(Types)>::Get(static_cast<const          TTupleBase&&>(*this)); }
-		template <uint32 Index> FORCEINLINE decltype(auto) Get()       volatile&& { static_assert(Index < sizeof...(Types), "Invalid index passed to TTuple::Get"); return TTupleElementGetterByIndex<Index, sizeof...(Types)>::Get(static_cast<      volatile TTupleBase&&>(*this)); }
-		template <uint32 Index> FORCEINLINE decltype(auto) Get() const volatile&& { static_assert(Index < sizeof...(Types), "Invalid index passed to TTuple::Get"); return TTupleElementGetterByIndex<Index, sizeof...(Types)>::Get(static_cast<const volatile TTupleBase&&>(*this)); }
-
-		template <typename T> FORCEINLINE decltype(auto) Get()               &  { static_assert(TTypeCountInParameterPack_V<T, Types...> == 1, "Invalid type passed to TTuple::Get"); return TTupleElementGetterByType<T, sizeof...(Types)>::Get(static_cast<               TTupleBase& >(*this)); }
-		template <typename T> FORCEINLINE decltype(auto) Get() const         &  { static_assert(TTypeCountInParameterPack_V<T, Types...> == 1, "Invalid type passed to TTuple::Get"); return TTupleElementGetterByType<T, sizeof...(Types)>::Get(static_cast<const          TTupleBase& >(*this)); }
-		template <typename T> FORCEINLINE decltype(auto) Get()       volatile&  { static_assert(TTypeCountInParameterPack_V<T, Types...> == 1, "Invalid type passed to TTuple::Get"); return TTupleElementGetterByType<T, sizeof...(Types)>::Get(static_cast<      volatile TTupleBase& >(*this)); }
-		template <typename T> FORCEINLINE decltype(auto) Get() const volatile&  { static_assert(TTypeCountInParameterPack_V<T, Types...> == 1, "Invalid type passed to TTuple::Get"); return TTupleElementGetterByType<T, sizeof...(Types)>::Get(static_cast<const volatile TTupleBase& >(*this)); }
-		template <typename T> FORCEINLINE decltype(auto) Get()               && { static_assert(TTypeCountInParameterPack_V<T, Types...> == 1, "Invalid type passed to TTuple::Get"); return TTupleElementGetterByType<T, sizeof...(Types)>::Get(static_cast<               TTupleBase&&>(*this)); }
-		template <typename T> FORCEINLINE decltype(auto) Get() const         && { static_assert(TTypeCountInParameterPack_V<T, Types...> == 1, "Invalid type passed to TTuple::Get"); return TTupleElementGetterByType<T, sizeof...(Types)>::Get(static_cast<const          TTupleBase&&>(*this)); }
-		template <typename T> FORCEINLINE decltype(auto) Get()       volatile&& { static_assert(TTypeCountInParameterPack_V<T, Types...> == 1, "Invalid type passed to TTuple::Get"); return TTupleElementGetterByType<T, sizeof...(Types)>::Get(static_cast<      volatile TTupleBase&&>(*this)); }
-		template <typename T> FORCEINLINE decltype(auto) Get() const volatile&& { static_assert(TTypeCountInParameterPack_V<T, Types...> == 1, "Invalid type passed to TTuple::Get"); return TTupleElementGetterByType<T, sizeof...(Types)>::Get(static_cast<const volatile TTupleBase&&>(*this)); }
+		template <uint32 Index> UE_FORCEINLINE_HINT decltype(auto) Get()               &  { static_assert(Index < sizeof...(Types), "Invalid index passed to TTuple::Get"); return TTupleElementGetterByIndex<Index, sizeof...(Types)>::Get(static_cast<               TTupleBase& >(*this)); }
+		template <uint32 Index> UE_FORCEINLINE_HINT decltype(auto) Get() const         &  { static_assert(Index < sizeof...(Types), "Invalid index passed to TTuple::Get"); return TTupleElementGetterByIndex<Index, sizeof...(Types)>::Get(static_cast<const          TTupleBase& >(*this)); }
+		template <uint32 Index> UE_FORCEINLINE_HINT decltype(auto) Get()       volatile&  { static_assert(Index < sizeof...(Types), "Invalid index passed to TTuple::Get"); return TTupleElementGetterByIndex<Index, sizeof...(Types)>::Get(static_cast<      volatile TTupleBase& >(*this)); }
+		template <uint32 Index> UE_FORCEINLINE_HINT decltype(auto) Get() const volatile&  { static_assert(Index < sizeof...(Types), "Invalid index passed to TTuple::Get"); return TTupleElementGetterByIndex<Index, sizeof...(Types)>::Get(static_cast<const volatile TTupleBase& >(*this)); }
+		template <uint32 Index> UE_FORCEINLINE_HINT decltype(auto) Get()               && { static_assert(Index < sizeof...(Types), "Invalid index passed to TTuple::Get"); return TTupleElementGetterByIndex<Index, sizeof...(Types)>::Get(static_cast<               TTupleBase&&>(*this)); }
+		template <uint32 Index> UE_FORCEINLINE_HINT decltype(auto) Get() const         && { static_assert(Index < sizeof...(Types), "Invalid index passed to TTuple::Get"); return TTupleElementGetterByIndex<Index, sizeof...(Types)>::Get(static_cast<const          TTupleBase&&>(*this)); }
+		template <uint32 Index> UE_FORCEINLINE_HINT decltype(auto) Get()       volatile&& { static_assert(Index < sizeof...(Types), "Invalid index passed to TTuple::Get"); return TTupleElementGetterByIndex<Index, sizeof...(Types)>::Get(static_cast<      volatile TTupleBase&&>(*this)); }
+		template <uint32 Index> UE_FORCEINLINE_HINT decltype(auto) Get() const volatile&& { static_assert(Index < sizeof...(Types), "Invalid index passed to TTuple::Get"); return TTupleElementGetterByIndex<Index, sizeof...(Types)>::Get(static_cast<const volatile TTupleBase&&>(*this)); }
+
+		template <typename T> UE_FORCEINLINE_HINT decltype(auto) Get()               &  { static_assert(TTypeCountInParameterPack_V<T, Types...> == 1, "Invalid type passed to TTuple::Get"); return TTupleElementGetterByType<T, sizeof...(Types)>::Get(static_cast<               TTupleBase& >(*this)); }
+		template <typename T> UE_FORCEINLINE_HINT decltype(auto) Get() const         &  { static_assert(TTypeCountInParameterPack_V<T, Types...> == 1, "Invalid type passed to TTuple::Get"); return TTupleElementGetterByType<T, sizeof...(Types)>::Get(static_cast<const          TTupleBase& >(*this)); }
+		template <typename T> UE_FORCEINLINE_HINT decltype(auto) Get()       volatile&  { static_assert(TTypeCountInParameterPack_V<T, Types...> == 1, "Invalid type passed to TTuple::Get"); return TTupleElementGetterByType<T, sizeof...(Types)>::Get(static_cast<      volatile TTupleBase& >(*this)); }
+		template <typename T> UE_FORCEINLINE_HINT decltype(auto) Get() const volatile&  { static_assert(TTypeCountInParameterPack_V<T, Types...> == 1, "Invalid type passed to TTuple::Get"); return TTupleElementGetterByType<T, sizeof...(Types)>::Get(static_cast<const volatile TTupleBase& >(*this)); }
+		template <typename T> UE_FORCEINLINE_HINT decltype(auto) Get()               && { static_assert(TTypeCountInParameterPack_V<T, Types...> == 1, "Invalid type passed to TTuple::Get"); return TTupleElementGetterByType<T, sizeof...(Types)>::Get(static_cast<               TTupleBase&&>(*this)); }
+		template <typename T> UE_FORCEINLINE_HINT decltype(auto) Get() const         && { static_assert(TTypeCountInParameterPack_V<T, Types...> == 1, "Invalid type passed to TTuple::Get"); return TTupleElementGetterByType<T, sizeof...(Types)>::Get(static_cast<const          TTupleBase&&>(*this)); }
+		template <typename T> UE_FORCEINLINE_HINT decltype(auto) Get()       volatile&& { static_assert(TTypeCountInParameterPack_V<T, Types...> == 1, "Invalid type passed to TTuple::Get"); return TTupleElementGetterByType<T, sizeof...(Types)>::Get(static_cast<      volatile TTupleBase&&>(*this)); }
+		template <typename T> UE_FORCEINLINE_HINT decltype(auto) Get() const volatile&& { static_assert(TTypeCountInParameterPack_V<T, Types...> == 1, "Invalid type passed to TTuple::Get"); return TTupleElementGetterByType<T, sizeof...(Types)>::Get(static_cast<const volatile TTupleBase&&>(*this)); }
 
 		template <typename FuncType, typename... ArgTypes> decltype(auto) ApplyAfter(FuncType&& Func, ArgTypes&&... Args)               &  { return ::Invoke(Func, Forward<ArgTypes>(Args)..., static_cast<               TTupleBase& >(*this).template Get<Indices>()...); }
 		template <typename FuncType, typename... ArgTypes> decltype(auto) ApplyAfter(FuncType&& Func, ArgTypes&&... Args) const         &  { return ::Invoke(Func, Forward<ArgTypes>(Args)..., static_cast<const          TTupleBase& >(*this).template Get<Indices>()...); }
@@ -345,33 +345,33 @@ namespace UE::Core::Private::Tuple
 			(void)Temp;
 		}
 
-		FORCEINLINE bool operator==(const TTupleBase& Rhs) const
+		UE_FORCEINLINE_HINT bool operator==(const TTupleBase& Rhs) const
 		{
 			// This could be implemented with a fold expression when our compilers support it
 			return FEqualityHelper<sizeof...(Types), 0>::Compare(*this, Rhs);
 		}
 
-		FORCEINLINE bool operator!=(const TTupleBase& Rhs) const
+		UE_FORCEINLINE_HINT bool operator!=(const TTupleBase& Rhs) const
 		{
 			return !(*this == Rhs);
 		}
 
-		FORCEINLINE bool operator<(const TTupleBase& Rhs) const
+		UE_FORCEINLINE_HINT bool operator<(const TTupleBase& Rhs) const
 		{
 			return TLessThanHelper<sizeof...(Types)>::Do(*this, Rhs);
 		}
 
-		FORCEINLINE bool operator<=(const TTupleBase& Rhs) const
+		UE_FORCEINLINE_HINT bool operator<=(const TTupleBase& Rhs) const
 		{
 			return !(Rhs < *this);
 		}
 
-		FORCEINLINE bool operator>(const TTupleBase& Rhs) const
+		UE_FORCEINLINE_HINT bool operator>(const TTupleBase& Rhs) const
 		{
 			return Rhs < *this;
 		}
 
-		FORCEINLINE bool operator>=(const TTupleBase& Rhs) const
+		UE_FORCEINLINE_HINT bool operator>=(const TTupleBase& Rhs) const
 		{
 			return !(*this < Rhs);
 		}
@@ -414,23 +414,23 @@ namespace UE::Core::Private::Tuple
 		TTupleBase& operator=(TTupleBase&& Other) = default;
 		TTupleBase& operator=(const TTupleBase& Other) = default;
 
-		template <uint32 Index> FORCEINLINE decltype(auto) Get()               &  { static_assert(Index < 2, "Invalid index passed to TTuple::Get"); return TTupleElementGetterByIndex<Index, 2>::Get(static_cast<               TTupleBase& >(*this)); }
-		template <uint32 Index> FORCEINLINE decltype(auto) Get() const         &  { static_assert(Index < 2, "Invalid index passed to TTuple::Get"); return TTupleElementGetterByIndex<Index, 2>::Get(static_cast<const          TTupleBase& >(*this)); }
-		template <uint32 Index> FORCEINLINE decltype(auto) Get()       volatile&  { static_assert(Index < 2, "Invalid index passed to TTuple::Get"); return TTupleElementGetterByIndex<Index, 2>::Get(static_cast<      volatile TTupleBase& >(*this)); }
-		template <uint32 Index> FORCEINLINE decltype(auto) Get() const volatile&  { static_assert(Index < 2, "Invalid index passed to TTuple::Get"); return TTupleElementGetterByIndex<Index, 2>::Get(static_cast<const volatile TTupleBase& >(*this)); }
-		template <uint32 Index> FORCEINLINE decltype(auto) Get()               && { static_assert(Index < 2, "Invalid index passed to TTuple::Get"); return TTupleElementGetterByIndex<Index, 2>::Get(static_cast<               TTupleBase&&>(*this)); }
-		template <uint32 Index> FORCEINLINE decltype(auto) Get() const         && { static_assert(Index < 2, "Invalid index passed to TTuple::Get"); return TTupleElementGetterByIndex<Index, 2>::Get(static_cast<const          TTupleBase&&>(*this)); }
-		template <uint32 Index> FORCEINLINE decltype(auto) Get()       volatile&& { static_assert(Index < 2, "Invalid index passed to TTuple::Get"); return TTupleElementGetterByIndex<Index, 2>::Get(static_cast<      volatile TTupleBase&&>(*this)); }
-		template <uint32 Index> FORCEINLINE decltype(auto) Get() const volatile&& { static_assert(Index < 2, "Invalid index passed to TTuple::Get"); return TTupleElementGetterByIndex<Index, 2>::Get(static_cast<const volatile TTupleBase&&>(*this)); }
-
-		template <typename T> FORCEINLINE decltype(auto) Get()               &  { static_assert(TTypeCountInParameterPack_V<T, KeyType, ValueType> == 1, "Invalid type passed to TTuple::Get"); return TTupleElementGetterByType<T, 2>::Get(static_cast<               TTupleBase& >(*this)); }
-		template <typename T> FORCEINLINE decltype(auto) Get() const         &  { static_assert(TTypeCountInParameterPack_V<T, KeyType, ValueType> == 1, "Invalid type passed to TTuple::Get"); return TTupleElementGetterByType<T, 2>::Get(static_cast<const          TTupleBase& >(*this)); }
-		template <typename T> FORCEINLINE decltype(auto) Get()       volatile&  { static_assert(TTypeCountInParameterPack_V<T, KeyType, ValueType> == 1, "Invalid type passed to TTuple::Get"); return TTupleElementGetterByType<T, 2>::Get(static_cast<      volatile TTupleBase& >(*this)); }
-		template <typename T> FORCEINLINE decltype(auto) Get() const volatile&  { static_assert(TTypeCountInParameterPack_V<T, KeyType, ValueType> == 1, "Invalid type passed to TTuple::Get"); return TTupleElementGetterByType<T, 2>::Get(static_cast<const volatile TTupleBase& >(*this)); }
-		template <typename T> FORCEINLINE decltype(auto) Get()               && { static_assert(TTypeCountInParameterPack_V<T, KeyType, ValueType> == 1, "Invalid type passed to TTuple::Get"); return TTupleElementGetterByType<T, 2>::Get(static_cast<               TTupleBase&&>(*this)); }
-		template <typename T> FORCEINLINE decltype(auto) Get() const         && { static_assert(TTypeCountInParameterPack_V<T, KeyType, ValueType> == 1, "Invalid type passed to TTuple::Get"); return TTupleElementGetterByType<T, 2>::Get(static_cast<const          TTupleBase&&>(*this)); }
-		template <typename T> FORCEINLINE decltype(auto) Get()       volatile&& { static_assert(TTypeCountInParameterPack_V<T, KeyType, ValueType> == 1, "Invalid type passed to TTuple::Get"); return TTupleElementGetterByType<T, 2>::Get(static_cast<      volatile TTupleBase&&>(*this)); }
-		template <typename T> FORCEINLINE decltype(auto) Get() const volatile&& { static_assert(TTypeCountInParameterPack_V<T, KeyType, ValueType> == 1, "Invalid type passed to TTuple::Get"); return TTupleElementGetterByType<T, 2>::Get(static_cast<const volatile TTupleBase&&>(*this)); }
+		template <uint32 Index> UE_FORCEINLINE_HINT decltype(auto) Get()               &  { static_assert(Index < 2, "Invalid index passed to TTuple::Get"); return TTupleElementGetterByIndex<Index, 2>::Get(static_cast<               TTupleBase& >(*this)); }
+		template <uint32 Index> UE_FORCEINLINE_HINT decltype(auto) Get() const         &  { static_assert(Index < 2, "Invalid index passed to TTuple::Get"); return TTupleElementGetterByIndex<Index, 2>::Get(static_cast<const          TTupleBase& >(*this)); }
+		template <uint32 Index> UE_FORCEINLINE_HINT decltype(auto) Get()       volatile&  { static_assert(Index < 2, "Invalid index passed to TTuple::Get"); return TTupleElementGetterByIndex<Index, 2>::Get(static_cast<      volatile TTupleBase& >(*this)); }
+		template <uint32 Index> UE_FORCEINLINE_HINT decltype(auto) Get() const volatile&  { static_assert(Index < 2, "Invalid index passed to TTuple::Get"); return TTupleElementGetterByIndex<Index, 2>::Get(static_cast<const volatile TTupleBase& >(*this)); }
+		template <uint32 Index> UE_FORCEINLINE_HINT decltype(auto) Get()               && { static_assert(Index < 2, "Invalid index passed to TTuple::Get"); return TTupleElementGetterByIndex<Index, 2>::Get(static_cast<               TTupleBase&&>(*this)); }
+		template <uint32 Index> UE_FORCEINLINE_HINT decltype(auto) Get() const         && { static_assert(Index < 2, "Invalid index passed to TTuple::Get"); return TTupleElementGetterByIndex<Index, 2>::Get(static_cast<const          TTupleBase&&>(*this)); }
+		template <uint32 Index> UE_FORCEINLINE_HINT decltype(auto) Get()       volatile&& { static_assert(Index < 2, "Invalid index passed to TTuple::Get"); return TTupleElementGetterByIndex<Index, 2>::Get(static_cast<      volatile TTupleBase&&>(*this)); }
+		template <uint32 Index> UE_FORCEINLINE_HINT decltype(auto) Get() const volatile&& { static_assert(Index < 2, "Invalid index passed to TTuple::Get"); return TTupleElementGetterByIndex<Index, 2>::Get(static_cast<const volatile TTupleBase&&>(*this)); }
+
+		template <typename T> UE_FORCEINLINE_HINT decltype(auto) Get()               &  { static_assert(TTypeCountInParameterPack_V<T, KeyType, ValueType> == 1, "Invalid type passed to TTuple::Get"); return TTupleElementGetterByType<T, 2>::Get(static_cast<               TTupleBase& >(*this)); }
+		template <typename T> UE_FORCEINLINE_HINT decltype(auto) Get() const         &  { static_assert(TTypeCountInParameterPack_V<T, KeyType, ValueType> == 1, "Invalid type passed to TTuple::Get"); return TTupleElementGetterByType<T, 2>::Get(static_cast<const          TTupleBase& >(*this)); }
+		template <typename T> UE_FORCEINLINE_HINT decltype(auto) Get()       volatile&  { static_assert(TTypeCountInParameterPack_V<T, KeyType, ValueType> == 1, "Invalid type passed to TTuple::Get"); return TTupleElementGetterByType<T, 2>::Get(static_cast<      volatile TTupleBase& >(*this)); }
+		template <typename T> UE_FORCEINLINE_HINT decltype(auto) Get() const volatile&  { static_assert(TTypeCountInParameterPack_V<T, KeyType, ValueType> == 1, "Invalid type passed to TTuple::Get"); return TTupleElementGetterByType<T, 2>::Get(static_cast<const volatile TTupleBase& >(*this)); }
+		template <typename T> UE_FORCEINLINE_HINT decltype(auto) Get()               && { static_assert(TTypeCountInParameterPack_V<T, KeyType, ValueType> == 1, "Invalid type passed to TTuple::Get"); return TTupleElementGetterByType<T, 2>::Get(static_cast<               TTupleBase&&>(*this)); }
+		template <typename T> UE_FORCEINLINE_HINT decltype(auto) Get() const         && { static_assert(TTypeCountInParameterPack_V<T, KeyType, ValueType> == 1, "Invalid type passed to TTuple::Get"); return TTupleElementGetterByType<T, 2>::Get(static_cast<const          TTupleBase&&>(*this)); }
+		template <typename T> UE_FORCEINLINE_HINT decltype(auto) Get()       volatile&& { static_assert(TTypeCountInParameterPack_V<T, KeyType, ValueType> == 1, "Invalid type passed to TTuple::Get"); return TTupleElementGetterByType<T, 2>::Get(static_cast<      volatile TTupleBase&&>(*this)); }
+		template <typename T> UE_FORCEINLINE_HINT decltype(auto) Get() const volatile&& { static_assert(TTypeCountInParameterPack_V<T, KeyType, ValueType> == 1, "Invalid type passed to TTuple::Get"); return TTupleElementGetterByType<T, 2>::Get(static_cast<const volatile TTupleBase&&>(*this)); }
 
 		template <typename FuncType, typename... ArgTypes> decltype(auto) ApplyAfter(FuncType&& Func, ArgTypes&&... Args)               &  { return ::Invoke(Func, Forward<ArgTypes>(Args)..., static_cast<               TTupleBase& >(*this).template Get<0>(), static_cast<               TTupleBase& >(*this).template Get<1>()); }
 		template <typename FuncType, typename... ArgTypes> decltype(auto) ApplyAfter(FuncType&& Func, ArgTypes&&... Args) const         &  { return ::Invoke(Func, Forward<ArgTypes>(Args)..., static_cast<const          TTupleBase& >(*this).template Get<0>(), static_cast<const          TTupleBase& >(*this).template Get<1>()); }
@@ -463,33 +463,33 @@ namespace UE::Core::Private::Tuple
 			Stream.EnterElement() << Tuple.Value;
 		}
 
-		FORCEINLINE friend bool operator==(const TTupleBase& Lhs, const TTupleBase& Rhs)
+		UE_FORCEINLINE_HINT friend bool operator==(const TTupleBase& Lhs, const TTupleBase& Rhs)
 		{
 			// This could be implemented with a fold expression when our compilers support it
 			return FEqualityHelper<2, 0>::Compare(Lhs, Rhs);
 		}
 
-		FORCEINLINE friend bool operator!=(const TTupleBase& Lhs, const TTupleBase& Rhs)
+		UE_FORCEINLINE_HINT friend bool operator!=(const TTupleBase& Lhs, const TTupleBase& Rhs)
 		{
 			return !(Lhs == Rhs);
 		}
 
-		FORCEINLINE friend bool operator<(const TTupleBase& Lhs, const TTupleBase& Rhs)
+		UE_FORCEINLINE_HINT friend bool operator<(const TTupleBase& Lhs, const TTupleBase& Rhs)
 		{
 			return TLessThanHelper<2>::Do(Lhs, Rhs);
 		}
 
-		FORCEINLINE friend bool operator<=(const TTupleBase& Lhs, const TTupleBase& Rhs)
+		UE_FORCEINLINE_HINT friend bool operator<=(const TTupleBase& Lhs, const TTupleBase& Rhs)
 		{
 			return !(Rhs < Lhs);
 		}
 
-		FORCEINLINE friend bool operator>(const TTupleBase& Lhs, const TTupleBase& Rhs)
+		UE_FORCEINLINE_HINT friend bool operator>(const TTupleBase& Lhs, const TTupleBase& Rhs)
 		{
 			return Rhs < Lhs;
 		}
 
-		FORCEINLINE friend bool operator>=(const TTupleBase& Lhs, const TTupleBase& Rhs)
+		UE_FORCEINLINE_HINT friend bool operator>=(const TTupleBase& Lhs, const TTupleBase& Rhs)
 		{
 			return !(Lhs < Rhs);
 		}
@@ -510,7 +510,7 @@ namespace UE::Core::Private::Tuple
 	}
 
 	template <typename... ElementTypes, typename... Types>
-	FORCEINLINE TTuple<ElementTypes...> MakeTupleImpl(Types&&... Args)
+	UE_FORCEINLINE_HINT TTuple<ElementTypes...> MakeTupleImpl(Types&&... Args)
 	{
 		return TTuple<ElementTypes...>(Forward<Types>(Args)...);
 	}
@@ -537,7 +537,7 @@ namespace UE::Core::Private::Tuple
 		// We need a second function to do the invocation for a particular index, to avoid the pack expansion being
 		// attempted on the indices and tuples simultaneously.
 		template <uint32 Index, typename FuncType, typename... TupleTypes>
-		FORCEINLINE static void InvokeFunc(FuncType&& Func, TupleTypes&&... Tuples)
+		UE_FORCEINLINE_HINT static void InvokeFunc(FuncType&& Func, TupleTypes&&... Tuples)
 		{
 			::Invoke(Forward<FuncType>(Func), Forward<TupleTypes>(Tuples).template Get<Index>()...);
 		}
@@ -633,7 +633,7 @@ namespace UE::Core::Private::Tuple
 	struct TGetTupleHashHelper
 	{
 		template <typename TupleType>
-		FORCEINLINE static uint32 Do(uint32 Hash, const TupleType& Tuple)
+		UE_FORCEINLINE_HINT static uint32 Do(uint32 Hash, const TupleType& Tuple)
 		{
 			return TGetTupleHashHelper<ArgToCombine + 1, ArgCount>::Do(HashCombineFast(Hash, GetTypeHashHelper(Tuple.template Get<ArgToCombine>())), Tuple);
 		}
@@ -643,7 +643,7 @@ namespace UE::Core::Private::Tuple
 	struct TGetTupleHashHelper<ArgIndex, ArgIndex>
 	{
 		template <typename TupleType>
-		FORCEINLINE static uint32 Do(uint32 Hash, const TupleType& Tuple)
+		UE_FORCEINLINE_HINT static uint32 Do(uint32 Hash, const TupleType& Tuple)
 		{
 			return Hash;
 		}
@@ -718,12 +718,12 @@ public:
 };
 
 template <typename... Types>
-FORCEINLINE uint32 GetTypeHash(const TTuple<Types...>& Tuple)
+UE_FORCEINLINE_HINT uint32 GetTypeHash(const TTuple<Types...>& Tuple)
 {
 	return UE::Core::Private::Tuple::TGetTupleHashHelper<1u, sizeof...(Types)>::Do(GetTypeHash(Tuple.template Get<0>()), Tuple);
 }
 
-FORCEINLINE uint32 GetTypeHash(const TTuple<>& Tuple)
+UE_FORCEINLINE_HINT uint32 GetTypeHash(const TTuple<>& Tuple)
 {
 	return 0;
 }
@@ -810,7 +810,7 @@ using TTupleElement = UE::Core::Private::Tuple::TCVTupleElement<Index, const vol
  * }
  */
 template <typename... Types>
-FORCEINLINE TTuple<std::decay_t<Types>...> MakeTuple(Types&&... Args)
+UE_FORCEINLINE_HINT TTuple<std::decay_t<Types>...> MakeTuple(Types&&... Args)
 {
 	return UE::Core::Private::Tuple::MakeTupleImpl<std::decay_t<Types>...>(Forward<Types>(Args)...);
 }
@@ -835,7 +835,7 @@ FORCEINLINE TTuple<std::decay_t<Types>...> MakeTuple(Types&&... Args)
  * }
  */
 template <typename... Types>
-FORCEINLINE TTuple<Types&&...> ForwardAsTuple(Types&&... Args)
+UE_FORCEINLINE_HINT TTuple<Types&&...> ForwardAsTuple(Types&&... Args)
 {
 	return UE::Core::Private::Tuple::MakeTupleImpl<Types&&...>(Forward<Types>(Args)...);
 }
@@ -863,13 +863,13 @@ FORCEINLINE TTuple<Types&&...> ForwardAsTuple(Types&&... Args)
  * }
  */
 template <typename FuncType, typename... Types>
-FORCEINLINE decltype(auto) TransformTuple(TTuple<Types...>&& Tuple, FuncType Func)
+UE_FORCEINLINE_HINT decltype(auto) TransformTuple(TTuple<Types...>&& Tuple, FuncType Func)
 {
 	return UE::Core::Private::Tuple::TTransformTuple_Impl<TMakeIntegerSequence<uint32, sizeof...(Types)>>::Do(MoveTemp(Tuple), MoveTemp(Func));
 }
 
 template <typename FuncType, typename... Types>
-FORCEINLINE decltype(auto) TransformTuple(const TTuple<Types...>& Tuple, FuncType Func)
+UE_FORCEINLINE_HINT decltype(auto) TransformTuple(const TTuple<Types...>& Tuple, FuncType Func)
 {
 	return UE::Core::Private::Tuple::TTransformTuple_Impl<TMakeIntegerSequence<uint32, sizeof...(Types)>>::Do(Tuple, MoveTemp(Func));
 }
@@ -894,7 +894,7 @@ FORCEINLINE decltype(auto) TransformTuple(const TTuple<Types...>& Tuple, FuncTyp
  * }
  */
 template <typename FuncType, typename FirstTupleType, typename... TupleTypes>
-FORCEINLINE void VisitTupleElements(FuncType&& Func, FirstTupleType&& FirstTuple, TupleTypes&&... Tuples)
+UE_FORCEINLINE_HINT void VisitTupleElements(FuncType&& Func, FirstTupleType&& FirstTuple, TupleTypes&&... Tuples)
 {
 	UE::Core::Private::Tuple::TVisitTupleElements_Impl<TMakeIntegerSequence<uint32, TTupleArity<std::decay_t<FirstTupleType>>::Value>>::Do(Forward<FuncType>(Func), Forward<FirstTupleType>(FirstTuple), Forward<TupleTypes>(Tuples)...);
 }
@@ -915,7 +915,7 @@ FORCEINLINE void VisitTupleElements(FuncType&& Func, FirstTupleType&& FirstTuple
  * // Now Ret1, Ret2 and Ret3 contain the unpacked return values.
  */
 template <typename... Types>
-FORCEINLINE TTuple<Types&...> Tie(Types&... Args)
+UE_FORCEINLINE_HINT TTuple<Types&...> Tie(Types&... Args)
 {
 	return TTuple<Types&...>(Args...);
 }
@@ -975,7 +975,7 @@ inline FArchive& operator<<(FArchive& Ar, TTuple<Types...>& Tuple)
 }
 
 template <typename... Types>
-FORCEINLINE void operator<<(FStructuredArchive::FSlot Slot, TTuple<Types...>& Tuple)
+UE_FORCEINLINE_HINT void operator<<(FStructuredArchive::FSlot Slot, TTuple<Types...>& Tuple)
 {
 	TTuple<Types...>::SerializeStructured(Slot, Tuple);
 }
diff --git a/Engine/Source/Runtime/Core/Public/Templates/TypeHash.h b/Engine/Source/Runtime/Core/Public/Templates/TypeHash.h
--- a/Engine/Source/Runtime/Core/Public/Templates/TypeHash.h
+++ b/Engine/Source/Runtime/Core/Public/Templates/TypeHash.h
@@ -212,7 +212,7 @@ template <typename T>
 
 // Use this when inside type that has GetTypeHash() (no in-parameters) implemented. It makes GetTypeHash dispatch in global namespace
 template <typename T>
-[[nodiscard]] FORCEINLINE uint32 GetTypeHashHelper(const T& V) { return static_cast<uint32>(GetTypeHash(V)); }
+[[nodiscard]] UE_FORCEINLINE_HINT uint32 GetTypeHashHelper(const T& V) { return static_cast<uint32>(GetTypeHash(V)); }
 
 // Use this when in a scope that has AppendHash to dispatch AppendHash using unqualified lookup.
 template <typename BuilderType, typename ArgType>
diff --git a/Engine/Source/Runtime/Core/Public/Templates/UniquePtr.h b/Engine/Source/Runtime/Core/Public/Templates/UniquePtr.h
--- a/Engine/Source/Runtime/Core/Public/Templates/UniquePtr.h
+++ b/Engine/Source/Runtime/Core/Public/Templates/UniquePtr.h
@@ -267,7 +267,7 @@ public:
 	/**
 	 * Destructor
 	 */
-	FORCEINLINE ~TUniquePtr()
+	UE_FORCEINLINE_HINT ~TUniquePtr()
 	{
 		GetDeleter()(Ptr);
 	}
@@ -287,7 +287,7 @@ public:
 	 *
 	 * @return true if the TUniquePtr currently owns an object, false otherwise.
 	 */
-	[[nodiscard]] FORCEINLINE explicit operator bool() const
+	[[nodiscard]] UE_FORCEINLINE_HINT explicit operator bool() const
 	{
 		return IsValid();
 	}
@@ -297,7 +297,7 @@ public:
 	 *
 	 * @return A pointer to the object owned by the TUniquePtr.
 	 */
-	[[nodiscard]] FORCEINLINE T* operator->() const
+	[[nodiscard]] UE_FORCEINLINE_HINT T* operator->() const
 	{
 		return Ptr;
 	}
@@ -311,7 +311,7 @@ public:
 		typename DummyT = T
 		UE_REQUIRES(UE_REQUIRES_EXPR(*(DummyT*)nullptr)) // this construct means that operator* is only considered for overload resolution if T is dereferenceable
 	>
-	[[nodiscard]] FORCEINLINE DummyT& operator*() const
+	[[nodiscard]] UE_FORCEINLINE_HINT DummyT& operator*() const
 	{
 		return *Ptr;
 	}
@@ -321,7 +321,7 @@ public:
 	 *
 	 * @return A copy of the pointer to the object owned by the TUniquePtr, or nullptr if no object is being owned.
 	 */
-	[[nodiscard]] FORCEINLINE T* Get() const
+	[[nodiscard]] UE_FORCEINLINE_HINT T* Get() const
 	{
 		return Ptr;
 	}
@@ -359,7 +359,7 @@ public:
 	 *
 	 * @return A reference to the deleter.
 	 */
-	[[nodiscard]] FORCEINLINE Deleter& GetDeleter()
+	[[nodiscard]] UE_FORCEINLINE_HINT Deleter& GetDeleter()
 	{
 		return static_cast<Deleter&>(*this);
 	}
@@ -369,7 +369,7 @@ public:
 	 *
 	 * @return A reference to the deleter.
 	 */
-	[[nodiscard]] FORCEINLINE const Deleter& GetDeleter() const
+	[[nodiscard]] UE_FORCEINLINE_HINT const Deleter& GetDeleter() const
 	{
 		return static_cast<const Deleter&>(*this);
 	}
@@ -388,7 +388,7 @@ public:
 	 * @return true if the two TUniquePtrs are logically substitutable for each other, false otherwise.
 	 */
 	template <typename RhsT>
-	[[nodiscard]] FORCEINLINE bool operator==(const TUniquePtr<RhsT>& Rhs) const
+	[[nodiscard]] UE_FORCEINLINE_HINT bool operator==(const TUniquePtr<RhsT>& Rhs) const
 	{
 		return Get() == Rhs.Get();
 	}
@@ -402,7 +402,7 @@ public:
 	 * @return true if the inner pointer of the left TUniquePtr is less than the inner pointer of the right, false otherwise.
 	 */
 	template <typename RhsT>
-	[[nodiscard]] FORCEINLINE bool operator<(const TUniquePtr<RhsT>& Rhs) const
+	[[nodiscard]] UE_FORCEINLINE_HINT bool operator<(const TUniquePtr<RhsT>& Rhs) const
 	{
 		return Get() < Rhs.Get();
 	}
@@ -416,7 +416,7 @@ public:
 	 * @return true if the inner pointer of the left TUniquePtr is greater than the inner pointer of the right, false otherwise.
 	 */
 	template <typename RhsT>
-	[[nodiscard]] FORCEINLINE bool operator>(const TUniquePtr<RhsT>& Rhs) const
+	[[nodiscard]] UE_FORCEINLINE_HINT bool operator>(const TUniquePtr<RhsT>& Rhs) const
 	{
 		return Get() > Rhs.Get();
 	}
@@ -430,7 +430,7 @@ public:
 	 * @return true if the inner pointer of the left TUniquePtr is less than or equal to the inner pointer of the right, false otherwise.
 	 */
 	template <typename RhsT>
-	[[nodiscard]] FORCEINLINE bool operator<=(const TUniquePtr<RhsT>& Rhs) const
+	[[nodiscard]] UE_FORCEINLINE_HINT bool operator<=(const TUniquePtr<RhsT>& Rhs) const
 	{
 		return Get() <= Rhs.Get();
 	}
@@ -444,7 +444,7 @@ public:
 	 * @return true if the inner pointer of the left TUniquePtr is greater than or equal to the inner pointer of the right, false otherwise.
 	 */
 	template <typename RhsT>
-	[[nodiscard]] FORCEINLINE bool operator>=(const TUniquePtr<RhsT>& Rhs) const
+	[[nodiscard]] UE_FORCEINLINE_HINT bool operator>=(const TUniquePtr<RhsT>& Rhs) const
 	{
 		return Get() >= Rhs.Get();
 	}
@@ -456,7 +456,7 @@ public:
 	 *
 	 * @return true if the TUniquePtr is null, false otherwise.
 	 */
-	[[nodiscard]] FORCEINLINE bool operator==(TYPE_OF_NULLPTR) const
+	[[nodiscard]] UE_FORCEINLINE_HINT bool operator==(TYPE_OF_NULLPTR) const
 	{
 		return !IsValid();
 	}
@@ -471,7 +471,7 @@ public:
 	 * @return false if the two TUniquePtrs are logically substitutable for each other, true otherwise.
 	 */
 	template <typename RhsT>
-	[[nodiscard]] FORCEINLINE bool operator!=(const TUniquePtr<RhsT>& Rhs) const
+	[[nodiscard]] UE_FORCEINLINE_HINT bool operator!=(const TUniquePtr<RhsT>& Rhs) const
 	{
 		return Get() != Rhs.Get();
 	}
@@ -483,7 +483,7 @@ public:
 	 *
 	 * @return true if the TUniquePtr is not null, false otherwise.
 	 */
-	[[nodiscard]] FORCEINLINE bool operator!=(TYPE_OF_NULLPTR) const
+	[[nodiscard]] UE_FORCEINLINE_HINT bool operator!=(TYPE_OF_NULLPTR) const
 	{
 		return IsValid();
 	}
@@ -652,7 +652,7 @@ public:
 	/**
 	 * Destructor
 	 */
-	FORCEINLINE ~TUniquePtr()
+	UE_FORCEINLINE_HINT ~TUniquePtr()
 	{
 		GetDeleter()(Ptr);
 	}
@@ -672,7 +672,7 @@ public:
 	 *
 	 * @return true if the TUniquePtr currently owns an array, false otherwise.
 	 */
-	[[nodiscard]] FORCEINLINE explicit operator bool() const
+	[[nodiscard]] UE_FORCEINLINE_HINT explicit operator bool() const
 	{
 		return IsValid();
 	}
@@ -682,7 +682,7 @@ public:
 	 *
 	 * @return A reference to the object at the specified index by the TUniquePtr.
 	 */
-	[[nodiscard]] FORCEINLINE T& operator[](SIZE_T Index) const
+	[[nodiscard]] UE_FORCEINLINE_HINT T& operator[](SIZE_T Index) const
 	{
 		return Ptr[Index];
 	}
@@ -692,7 +692,7 @@ public:
 	 *
 	 * @return A copy of the pointer to the array owned by the TUniquePtr, or nullptr if no array is being owned.
 	 */
-	[[nodiscard]] FORCEINLINE T* Get() const
+	[[nodiscard]] UE_FORCEINLINE_HINT T* Get() const
 	{
 		return Ptr;
 	}
@@ -739,7 +739,7 @@ public:
 	 *
 	 * @return A reference to the deleter.
 	 */
-	[[nodiscard]] FORCEINLINE Deleter& GetDeleter()
+	[[nodiscard]] UE_FORCEINLINE_HINT Deleter& GetDeleter()
 	{
 		return static_cast<Deleter&>(*this);
 	}
@@ -749,7 +749,7 @@ public:
 	 *
 	 * @return A reference to the deleter.
 	 */
-	[[nodiscard]] FORCEINLINE const Deleter& GetDeleter() const
+	[[nodiscard]] UE_FORCEINLINE_HINT const Deleter& GetDeleter() const
 	{
 		return static_cast<const Deleter&>(*this);
 	}
@@ -767,7 +767,7 @@ public:
 	 * @return true if the two TUniquePtrs are logically substitutable for each other, false otherwise.
 	 */
 	template <typename RhsT>
-	[[nodiscard]] FORCEINLINE bool operator==(const TUniquePtr<RhsT>& Rhs) const
+	[[nodiscard]] UE_FORCEINLINE_HINT bool operator==(const TUniquePtr<RhsT>& Rhs) const
 	{
 		return Get() == Rhs.Get();
 	}
@@ -781,7 +781,7 @@ public:
 	 * @return true if the inner pointer of the left TUniquePtr is less than the inner pointer of the right, false otherwise.
 	 */
 	template <typename RhsT>
-	[[nodiscard]] FORCEINLINE bool operator<(const TUniquePtr<RhsT>& Rhs) const
+	[[nodiscard]] UE_FORCEINLINE_HINT bool operator<(const TUniquePtr<RhsT>& Rhs) const
 	{
 		return Get() < Rhs.Get();
 	}
@@ -795,7 +795,7 @@ public:
 	 * @return true if the inner pointer of the left TUniquePtr is greater than the inner pointer of the right, false otherwise.
 	 */
 	template <typename RhsT>
-	[[nodiscard]] FORCEINLINE bool operator>(const TUniquePtr<RhsT>& Rhs) const
+	[[nodiscard]] UE_FORCEINLINE_HINT bool operator>(const TUniquePtr<RhsT>& Rhs) const
 	{
 		return Get() > Rhs.Get();
 	}
@@ -809,7 +809,7 @@ public:
 	 * @return true if the inner pointer of the left TUniquePtr is less than or equal to the inner pointer of the right, false otherwise.
 	 */
 	template <typename RhsT>
-	[[nodiscard]] FORCEINLINE bool operator<=(const TUniquePtr<RhsT>& Rhs) const
+	[[nodiscard]] UE_FORCEINLINE_HINT bool operator<=(const TUniquePtr<RhsT>& Rhs) const
 	{
 		return Get() <= Rhs.Get();
 	}
@@ -823,7 +823,7 @@ public:
 	 * @return true if the inner pointer of the left TUniquePtr is greater than or equal to the inner pointer of the right, false otherwise.
 	 */
 	template <typename RhsT>
-	[[nodiscard]] FORCEINLINE bool operator>=(const TUniquePtr<RhsT>& Rhs) const
+	[[nodiscard]] UE_FORCEINLINE_HINT bool operator>=(const TUniquePtr<RhsT>& Rhs) const
 	{
 		return Get() >= Rhs.Get();
 	}
@@ -835,7 +835,7 @@ public:
 	 *
 	 * @return true if the TUniquePtr is null, false otherwise.
 	 */
-	[[nodiscard]] FORCEINLINE bool operator==(TYPE_OF_NULLPTR) const
+	[[nodiscard]] UE_FORCEINLINE_HINT bool operator==(TYPE_OF_NULLPTR) const
 	{
 		return !IsValid();
 	}
@@ -850,7 +850,7 @@ public:
 	 * @return false if the two TUniquePtrs are logically substitutable for each other, true otherwise.
 	 */
 	template <typename RhsT>
-	[[nodiscard]] FORCEINLINE bool operator!=(const TUniquePtr<RhsT>& Rhs) const
+	[[nodiscard]] UE_FORCEINLINE_HINT bool operator!=(const TUniquePtr<RhsT>& Rhs) const
 	{
 		return Get() != Rhs.Get();
 	}
@@ -862,7 +862,7 @@ public:
 	 *
 	 * @return true if the TUniquePtr is not null, false otherwise.
 	 */
-	[[nodiscard]] FORCEINLINE bool operator!=(TYPE_OF_NULLPTR) const
+	[[nodiscard]] UE_FORCEINLINE_HINT bool operator!=(TYPE_OF_NULLPTR) const
 	{
 		return IsValid();
 	}
@@ -876,13 +876,13 @@ public:
 
 #if !PLATFORM_COMPILER_HAS_GENERATED_COMPARISON_OPERATORS
 template <typename T>
-[[nodiscard]] FORCEINLINE bool operator==(TYPE_OF_NULLPTR, const TUniquePtr<T>& Rhs)
+[[nodiscard]] UE_FORCEINLINE_HINT bool operator==(TYPE_OF_NULLPTR, const TUniquePtr<T>& Rhs)
 {
 	return !Rhs.IsValid();
 }
 
 template <typename T>
-[[nodiscard]] FORCEINLINE bool operator!=(TYPE_OF_NULLPTR, const TUniquePtr<T>& Rhs)
+[[nodiscard]] UE_FORCEINLINE_HINT bool operator!=(TYPE_OF_NULLPTR, const TUniquePtr<T>& Rhs)
 {
 	return Rhs.IsValid();
 }
@@ -915,7 +915,7 @@ template <
 	typename... TArgs
 	UE_REQUIRES(!std::is_array_v<T>)
 >
-[[nodiscard]] FORCEINLINE TUniquePtr<T> MakeUnique(TArgs&&... Args)
+[[nodiscard]] UE_FORCEINLINE_HINT TUniquePtr<T> MakeUnique(TArgs&&... Args)
 {
 	return TUniquePtr<T>(new T(Forward<TArgs>(Args)...));
 }
@@ -933,7 +933,7 @@ template <
 	typename... TArgs
 	UE_REQUIRES(!std::is_array_v<T>)
 >
-[[nodiscard]] FORCEINLINE TUniquePtr<T> MakeUniqueForOverwrite()
+[[nodiscard]] UE_FORCEINLINE_HINT TUniquePtr<T> MakeUniqueForOverwrite()
 {
 	return TUniquePtr<T>(new T);
 }
diff --git a/Engine/Source/Runtime/Core/Public/Templates/UnrealTemplate.h b/Engine/Source/Runtime/Core/Public/Templates/UnrealTemplate.h
--- a/Engine/Source/Runtime/Core/Public/Templates/UnrealTemplate.h
+++ b/Engine/Source/Runtime/Core/Public/Templates/UnrealTemplate.h
@@ -120,7 +120,7 @@ constexpr int32 GetNum(std::initializer_list<T> List)
  * Returns a non-const reference type as const.
  */
 template <typename T>
-constexpr FORCEINLINE const T& AsConst(T& Ref)
+constexpr UE_FORCEINLINE_HINT const T& AsConst(T& Ref)
 {
 	return Ref;
 }
@@ -136,7 +136,7 @@ void AsConst(const T&& Ref) = delete;
  * This overload is only required until the pointer overloads are removed.
  */
 template <typename T, SIZE_T N>
-constexpr FORCEINLINE const T (&AsConst(T (&Array)[N]))[N]
+constexpr UE_FORCEINLINE_HINT const T (&AsConst(T (&Array)[N]))[N]
 {
 	return Array;
 }
@@ -354,7 +354,7 @@ struct TGuardValue : private FNoncopyable
 	 *
 	 * @return	a const reference to the original data value
 	 */
-	FORCEINLINE const AssignedType& GetOriginalValue() const
+	UE_FORCEINLINE_HINT const AssignedType& GetOriginalValue() const
 	{
 		return OriginalValue;
 	}
@@ -396,7 +396,7 @@ struct TOptionalGuardValue : private FNoncopyable
 	 *
 	 * @return	a const reference to the original data value
 	 */
-	FORCEINLINE const AssignedType& GetOriginalValue() const
+	UE_FORCEINLINE_HINT const AssignedType& GetOriginalValue() const
 	{
 		return OriginalValue;
 	}
@@ -483,7 +483,7 @@ struct TKeyValuePair
 	{
 		return Key < Other.Key;
 	}
-	FORCEINLINE bool operator()( const TKeyValuePair& A, const TKeyValuePair& B ) const
+	UE_FORCEINLINE_HINT bool operator()( const TKeyValuePair& A, const TKeyValuePair& B ) const
 	{
 		return A.Key < B.Key;
 	}
@@ -671,7 +671,7 @@ inline T ReverseBits( T Bits )
  * Generates a bitmask with a given number of bits set.
  */
 template <typename T>
-FORCEINLINE T BitMask( uint32 Count );
+UE_FORCEINLINE_HINT T BitMask( uint32 Count );
 
 template <>
 inline uint64 BitMask<uint64>( uint32 Count )
diff --git a/Engine/Source/Runtime/Core/Public/Templates/UnrealTypeTraits.h b/Engine/Source/Runtime/Core/Public/Templates/UnrealTypeTraits.h
--- a/Engine/Source/Runtime/Core/Public/Templates/UnrealTypeTraits.h
+++ b/Engine/Source/Runtime/Core/Public/Templates/UnrealTypeTraits.h
@@ -97,7 +97,7 @@ template<> \
 struct TFormatSpecifier<type> \
 {  \
 	template <typename CharType = TCHAR> \
-	FORCEINLINE static constexpr decltype(auto) GetFormatSpecifier() \
+	UE_FORCEINLINE_HINT static constexpr decltype(auto) GetFormatSpecifier() \
 	{ \
 		return CHARTEXT(CharType, format); \
 	} \
@@ -200,7 +200,7 @@ struct TNameOf
 template<> \
 struct TNameOf<type> \
 {  \
-	FORCEINLINE static TCHAR const* GetName() \
+	UE_FORCEINLINE_HINT static TCHAR const* GetName() \
 	{ \
 		return TEXT(#type); \
 	} \
diff --git a/Engine/Source/Runtime/Core/Public/UObject/ScriptDelegates.h b/Engine/Source/Runtime/Core/Public/UObject/ScriptDelegates.h
--- a/Engine/Source/Runtime/Core/Public/UObject/ScriptDelegates.h
+++ b/Engine/Source/Runtime/Core/Public/UObject/ScriptDelegates.h
@@ -334,7 +334,7 @@ public:
 		return bResult;
 	}
 
-	FORCEINLINE bool operator!=( const TScriptDelegate& Other ) const
+	UE_FORCEINLINE_HINT bool operator!=( const TScriptDelegate& Other ) const
 	{
 		return !operator==(Other);
 	}
@@ -343,7 +343,7 @@ public:
 		UE_REQUIRES(UE::Core::Private::BackwardCompatibilityCheck<InThreadSafetyMode, OtherThreadSafetyMode>())
 	>
 	/* UE_DEPRECATED(5.3, "Deprecated - remove after TScriptDelegateTraits<FWeakObjectPtr> is removed") */
-	FORCEINLINE bool operator!=(const TScriptDelegate<OtherThreadSafetyMode>& Other) const
+	UE_FORCEINLINE_HINT bool operator!=(const TScriptDelegate<OtherThreadSafetyMode>& Other) const
 	{
 		return !operator==(Other);
 	}
diff --git a/Engine/Source/Runtime/Core/Public/UObject/StrongObjectPtrTemplates.h b/Engine/Source/Runtime/Core/Public/UObject/StrongObjectPtrTemplates.h
--- a/Engine/Source/Runtime/Core/Public/UObject/StrongObjectPtrTemplates.h
+++ b/Engine/Source/Runtime/Core/Public/UObject/StrongObjectPtrTemplates.h
@@ -44,12 +44,12 @@ public:
 		return *this;
 	}
 
-	FORCEINLINE_DEBUGGABLE ~TStrongObjectPtr()
+	UE_FORCEINLINE_HINT ~TStrongObjectPtr()
 	{
 		Reset();
 	}
 
-	[[nodiscard]] FORCEINLINE_DEBUGGABLE TStrongObjectPtr(TYPE_OF_NULLPTR = nullptr)
+	[[nodiscard]] UE_FORCEINLINE_HINT TStrongObjectPtr(TYPE_OF_NULLPTR = nullptr)
 	{
 		static_assert(TPointerIsConvertibleFromTo<ObjectType, const volatile UObject>::Value, "TStrongObjectPtr can only be constructed with UObject types");
 	}
@@ -60,7 +60,7 @@ public:
 		Reset(InObject);
 	}
 
-	[[nodiscard]] FORCEINLINE_DEBUGGABLE TStrongObjectPtr(const TStrongObjectPtr& InOther)
+	[[nodiscard]] UE_FORCEINLINE_HINT TStrongObjectPtr(const TStrongObjectPtr& InOther)
 	{
 		Reset(InOther.Get());
 	}
@@ -70,7 +70,7 @@ public:
 		typename OtherReferencerNameProvider
 		UE_REQUIRES(std::is_convertible_v<OtherObjectType*, ObjectType*>)
 	>
-	[[nodiscard]] FORCEINLINE_DEBUGGABLE TStrongObjectPtr(const TStrongObjectPtr<OtherObjectType, OtherReferencerNameProvider>& InOther)
+	[[nodiscard]] UE_FORCEINLINE_HINT TStrongObjectPtr(const TStrongObjectPtr<OtherObjectType, OtherReferencerNameProvider>& InOther)
 	{
 		Reset(InOther.Get());
 	}
@@ -104,17 +104,17 @@ public:
 		return (ObjectType*)Get();
 	}
 
-	[[nodiscard]] FORCEINLINE_DEBUGGABLE bool IsValid() const
+	[[nodiscard]] UE_FORCEINLINE_HINT bool IsValid() const
 	{
 		return Object != nullptr;
 	}
 
-	[[nodiscard]] FORCEINLINE_DEBUGGABLE explicit operator bool() const
+	[[nodiscard]] UE_FORCEINLINE_HINT explicit operator bool() const
 	{
 		return IsValid();
 	}
 
-	[[nodiscard]] FORCEINLINE_DEBUGGABLE ObjectType* Get() const
+	[[nodiscard]] UE_FORCEINLINE_HINT ObjectType* Get() const
 	{
 		return (ObjectType*)Object;
 	}
@@ -173,7 +173,7 @@ public:
 		}
 	}
 
-	[[nodiscard]] FORCEINLINE_DEBUGGABLE friend uint32 GetTypeHash(const TStrongObjectPtr& InStrongObjectPtr)
+	[[nodiscard]] UE_FORCEINLINE_HINT friend uint32 GetTypeHash(const TStrongObjectPtr& InStrongObjectPtr)
 	{
 		return GetTypeHash(InStrongObjectPtr.Get());
 	}
@@ -183,17 +183,17 @@ private:
 	// This is required because the destructor calls Reset, which need to be fully resolved at declaration.
 	const UObject* Object{ nullptr };
 
-	[[nodiscard]] friend FORCEINLINE bool operator==(const TStrongObjectPtr& InLHS, const TStrongObjectPtr& InRHS)
+	[[nodiscard]] friend UE_FORCEINLINE_HINT bool operator==(const TStrongObjectPtr& InLHS, const TStrongObjectPtr& InRHS)
 	{
 		return InLHS.Get() == InRHS.Get();
 	}
 
-	[[nodiscard]] friend FORCEINLINE bool operator==(const TStrongObjectPtr& InLHS, TYPE_OF_NULLPTR)
+	[[nodiscard]] friend UE_FORCEINLINE_HINT bool operator==(const TStrongObjectPtr& InLHS, TYPE_OF_NULLPTR)
 	{
 		return !InLHS.IsValid();
 	}
 
-	[[nodiscard]] friend FORCEINLINE bool operator==(TYPE_OF_NULLPTR, const TStrongObjectPtr& InRHS)
+	[[nodiscard]] friend UE_FORCEINLINE_HINT bool operator==(TYPE_OF_NULLPTR, const TStrongObjectPtr& InRHS)
 	{
 		return !InRHS.IsValid();
 	}
@@ -203,7 +203,7 @@ private:
 		typename RHSReferencerNameProvider
 		UE_REQUIRES(UE_REQUIRES_EXPR(std::declval<ObjectType*>() == std::declval<RHSObjectType*>()))
 	>
-	[[nodiscard]] friend FORCEINLINE bool operator==(const TStrongObjectPtr& InLHS, const TStrongObjectPtr<RHSObjectType, RHSReferencerNameProvider>& InRHS)
+	[[nodiscard]] friend UE_FORCEINLINE_HINT bool operator==(const TStrongObjectPtr& InLHS, const TStrongObjectPtr<RHSObjectType, RHSReferencerNameProvider>& InRHS)
 	{
 		return InLHS.Get() == InRHS.Get();
 	}
@@ -213,23 +213,23 @@ private:
 		typename LHSReferencerNameProvider
 		UE_REQUIRES(UE_REQUIRES_EXPR(std::declval<LHSObjectType*>() == std::declval<ObjectType*>()))
 	>
-	[[nodiscard]] friend FORCEINLINE bool operator==(const TStrongObjectPtr<LHSObjectType, LHSReferencerNameProvider>& InLHS, const TStrongObjectPtr& InRHS)
+	[[nodiscard]] friend UE_FORCEINLINE_HINT bool operator==(const TStrongObjectPtr<LHSObjectType, LHSReferencerNameProvider>& InLHS, const TStrongObjectPtr& InRHS)
 	{
 		return InLHS.Get() == InRHS.Get();
 	}
 
 #if !PLATFORM_COMPILER_HAS_GENERATED_COMPARISON_OPERATORS
-	[[nodiscard]] friend FORCEINLINE bool operator!=(const TStrongObjectPtr& InLHS, const TStrongObjectPtr& InRHS)
+	[[nodiscard]] friend UE_FORCEINLINE_HINT bool operator!=(const TStrongObjectPtr& InLHS, const TStrongObjectPtr& InRHS)
 	{
 		return InLHS.Get() != InRHS.Get();
 	}
 
-	[[nodiscard]] friend FORCEINLINE bool operator!=(const TStrongObjectPtr& InLHS, TYPE_OF_NULLPTR)
+	[[nodiscard]] friend UE_FORCEINLINE_HINT bool operator!=(const TStrongObjectPtr& InLHS, TYPE_OF_NULLPTR)
 	{
 		return InLHS.IsValid();
 	}
 
-	[[nodiscard]] friend FORCEINLINE bool operator!=(TYPE_OF_NULLPTR, const TStrongObjectPtr& InRHS)
+	[[nodiscard]] friend UE_FORCEINLINE_HINT bool operator!=(TYPE_OF_NULLPTR, const TStrongObjectPtr& InRHS)
 	{
 		return InRHS.IsValid();
 	}
@@ -239,7 +239,7 @@ private:
 		typename RHSReferencerNameProvider
 		UE_REQUIRES(UE_REQUIRES_EXPR(std::declval<ObjectType*>() == std::declval<RHSObjectType*>()))
 	>
-	[[nodiscard]] friend FORCEINLINE bool operator!=(const TStrongObjectPtr& InLHS, const TStrongObjectPtr<RHSObjectType, RHSReferencerNameProvider>& InRHS)
+	[[nodiscard]] friend UE_FORCEINLINE_HINT bool operator!=(const TStrongObjectPtr& InLHS, const TStrongObjectPtr<RHSObjectType, RHSReferencerNameProvider>& InRHS)
 	{
 		return InLHS.Get() != InRHS.Get();
 	}
@@ -249,7 +249,7 @@ private:
 		typename LHSReferencerNameProvider
 		UE_REQUIRES(UE_REQUIRES_EXPR(std::declval<LHSObjectType*>() == std::declval<ObjectType*>()))
 	>
-	[[nodiscard]] friend FORCEINLINE bool operator!=(const TStrongObjectPtr<LHSObjectType, LHSReferencerNameProvider>& InLHS, const TStrongObjectPtr& InRHS)
+	[[nodiscard]] friend UE_FORCEINLINE_HINT bool operator!=(const TStrongObjectPtr<LHSObjectType, LHSReferencerNameProvider>& InLHS, const TStrongObjectPtr& InRHS)
 	{
 		return InLHS.Get() != InRHS.Get();
 	}
diff --git a/Engine/Source/Runtime/Core/Public/Unix/UnixPlatformMath.h b/Engine/Source/Runtime/Core/Public/Unix/UnixPlatformMath.h
--- a/Engine/Source/Runtime/Core/Public/Unix/UnixPlatformMath.h
+++ b/Engine/Source/Runtime/Core/Public/Unix/UnixPlatformMath.h
@@ -18,16 +18,16 @@
 struct FLinuxPlatformMath : public TUnrealPlatformMathSSE4Base<FClangPlatformMath>
 {
 #if PLATFORM_ENABLE_VECTORINTRINSICS
-	static FORCEINLINE bool IsNaN( float A ) { return isnan(A) != 0; }
-	static FORCEINLINE bool IsNaN( double A ) { return isnan(A) != 0; }
-	static FORCEINLINE bool IsFinite( float A ) { return isfinite(A); }
-	static FORCEINLINE bool IsFinite( double A ) { return isfinite(A); }
+	static UE_FORCEINLINE_HINT bool IsNaN( float A ) { return isnan(A) != 0; }
+	static UE_FORCEINLINE_HINT bool IsNaN( double A ) { return isnan(A) != 0; }
+	static UE_FORCEINLINE_HINT bool IsFinite( float A ) { return isfinite(A); }
+	static UE_FORCEINLINE_HINT bool IsFinite( double A ) { return isfinite(A); }
 
 #if PLATFORM_ENABLE_POPCNT_INTRINSIC
 	/**
 	 * Use the SSE instruction to count bits
 	 */
-	static FORCEINLINE int32 CountBits(uint64 Bits)
+	static UE_FORCEINLINE_HINT int32 CountBits(uint64 Bits)
 	{
 		return __builtin_popcountll(Bits);
 	}
diff --git a/Engine/Source/Runtime/Core/Public/Unix/UnixPlatformMisc.h b/Engine/Source/Runtime/Core/Public/Unix/UnixPlatformMisc.h
--- a/Engine/Source/Runtime/Core/Public/Unix/UnixPlatformMisc.h
+++ b/Engine/Source/Runtime/Core/Public/Unix/UnixPlatformMisc.h
@@ -57,7 +57,7 @@ struct FUnixPlatformMisc : public FGenericPlatformMisc
 
 	static CORE_API EAppReturnType::Type MessageBoxExt(EAppMsgType::Type MsgType, const TCHAR* Text, const TCHAR* Caption);
 
-	FORCEINLINE static void MemoryBarrier()
+	UE_FORCEINLINE_HINT static void MemoryBarrier()
 	{
 		__sync_synchronize();
 	}
diff --git a/Engine/Source/Runtime/Core/Public/Unix/UnixPlatformProcess.h b/Engine/Source/Runtime/Core/Public/Unix/UnixPlatformProcess.h
--- a/Engine/Source/Runtime/Core/Public/Unix/UnixPlatformProcess.h
+++ b/Engine/Source/Runtime/Core/Public/Unix/UnixPlatformProcess.h
@@ -36,7 +36,7 @@ struct FProcState
 	~FProcState();
 
 	/** Getter for process id */
-	FORCEINLINE pid_t GetProcessId() const
+	UE_FORCEINLINE_HINT pid_t GetProcessId() const
 	{
 		return ProcessId;
 	}
@@ -134,7 +134,7 @@ struct FProcHandle
 	}
 
 	/** Accessors. */
-	FORCEINLINE pid_t Get() const
+	UE_FORCEINLINE_HINT pid_t Get() const
 	{
 		return ProcInfo ? ProcInfo->GetProcessId() : OpenedPid;
 	}
@@ -147,13 +147,13 @@ struct FProcHandle
 	}
 
 	/** Checks the validity of handle */
-	FORCEINLINE bool IsValid() const
+	UE_FORCEINLINE_HINT bool IsValid() const
 	{
 		return ProcInfo != nullptr || OpenedPid != -1;
 	}
 
 	// the below is not part of FProcHandle API and is specific to Unix implementation
-	FORCEINLINE FProcState* GetProcessInfo() const
+	UE_FORCEINLINE_HINT FProcState* GetProcessInfo() const
 	{
 		return ProcInfo;
 	}
diff --git a/Engine/Source/Runtime/Core/Public/Unix/UnixPlatformTLS.h b/Engine/Source/Runtime/Core/Public/Unix/UnixPlatformTLS.h
--- a/Engine/Source/Runtime/Core/Public/Unix/UnixPlatformTLS.h
+++ b/Engine/Source/Runtime/Core/Public/Unix/UnixPlatformTLS.h
@@ -104,7 +104,7 @@ struct FUnixTLS : public FGenericPlatformTLS
 	 * @param SlotIndex the TLS index to store it in
 	 * @param Value the value to store in the slot
 	 */
-	static FORCEINLINE void SetTlsValue(uint32 SlotIndex,void* Value)
+	static UE_FORCEINLINE_HINT void SetTlsValue(uint32 SlotIndex,void* Value)
 	{
 		pthread_setspecific((pthread_key_t)SlotIndex, Value);
 	}
@@ -114,7 +114,7 @@ struct FUnixTLS : public FGenericPlatformTLS
 	 *
 	 * @return the value stored in the slot
 	 */
-	static FORCEINLINE void* GetTlsValue(uint32 SlotIndex)
+	static UE_FORCEINLINE_HINT void* GetTlsValue(uint32 SlotIndex)
 	{
 		return pthread_getspecific((pthread_key_t)SlotIndex);
 	}
@@ -124,7 +124,7 @@ struct FUnixTLS : public FGenericPlatformTLS
 	 *
 	 * @param SlotIndex the TLS index to store it in
 	 */
-	static FORCEINLINE void FreeTlsSlot(uint32 SlotIndex)
+	static UE_FORCEINLINE_HINT void FreeTlsSlot(uint32 SlotIndex)
 	{
 		pthread_key_delete((pthread_key_t)SlotIndex);
 	}
diff --git a/Engine/Source/Runtime/Core/Public/Windows/WindowsPlatformMisc.h b/Engine/Source/Runtime/Core/Public/Windows/WindowsPlatformMisc.h
--- a/Engine/Source/Runtime/Core/Public/Windows/WindowsPlatformMisc.h
+++ b/Engine/Source/Runtime/Core/Public/Windows/WindowsPlatformMisc.h
@@ -130,7 +130,7 @@ struct FWindowsPlatformMisc
 
 	static CORE_API void AsymmetricThreadFenceHeavy();
 
-	FORCEINLINE static void AsymmetricThreadFenceLight()
+	UE_FORCEINLINE_HINT static void AsymmetricThreadFenceLight()
 	{
 		_ReadWriteBarrier();
 	}
diff --git a/Engine/Source/Runtime/Core/Public/Windows/WindowsPlatformMutex.h b/Engine/Source/Runtime/Core/Public/Windows/WindowsPlatformMutex.h
--- a/Engine/Source/Runtime/Core/Public/Windows/WindowsPlatformMutex.h
+++ b/Engine/Source/Runtime/Core/Public/Windows/WindowsPlatformMutex.h
@@ -28,22 +28,22 @@ public:
 		Windows::InitializeCriticalSectionAndSpinCount(&CriticalSection, 4000);
 	}
 
-	FORCEINLINE ~FWindowsRecursiveMutex()
+	UE_FORCEINLINE_HINT ~FWindowsRecursiveMutex()
 	{
 		Windows::DeleteCriticalSection(&CriticalSection);
 	}
 
-	FORCEINLINE bool TryLock()
+	UE_FORCEINLINE_HINT bool TryLock()
 	{
 		return !!Windows::TryEnterCriticalSection(&CriticalSection);
 	}
 
-	FORCEINLINE void Lock()
+	UE_FORCEINLINE_HINT void Lock()
 	{
 		Windows::EnterCriticalSection(&CriticalSection);
 	}
 
-	FORCEINLINE void Unlock()
+	UE_FORCEINLINE_HINT void Unlock()
 	{
 		Windows::LeaveCriticalSection(&CriticalSection);
 	}
@@ -65,7 +65,7 @@ public:
 	FWindowsSharedMutex(const FWindowsSharedMutex&) = delete;
 	FWindowsSharedMutex& operator=(const FWindowsSharedMutex&) = delete;
 
-	FORCEINLINE FWindowsSharedMutex()
+	UE_FORCEINLINE_HINT FWindowsSharedMutex()
 	{
 		Windows::InitializeSRWLock(&Mutex);
 	}
@@ -75,32 +75,32 @@ public:
 		checkf(!IsLocked(), TEXT("Destroying a lock that is still held!"));
 	}
 
-	FORCEINLINE bool TryLock()
+	UE_FORCEINLINE_HINT bool TryLock()
 	{
 		return !!Windows::TryAcquireSRWLockExclusive(&Mutex);
 	}
 
-	FORCEINLINE void Lock()
+	UE_FORCEINLINE_HINT void Lock()
 	{
 		Windows::AcquireSRWLockExclusive(&Mutex);
 	}
 
-	FORCEINLINE void Unlock()
+	UE_FORCEINLINE_HINT void Unlock()
 	{
 		Windows::ReleaseSRWLockExclusive(&Mutex);
 	}
 
-	FORCEINLINE bool TryLockShared()
+	UE_FORCEINLINE_HINT bool TryLockShared()
 	{
 		return !!Windows::TryAcquireSRWLockShared(&Mutex);
 	}
 
-	FORCEINLINE void LockShared()
+	UE_FORCEINLINE_HINT void LockShared()
 	{
 		Windows::AcquireSRWLockShared(&Mutex);
 	}
 
-	FORCEINLINE void UnlockShared()
+	UE_FORCEINLINE_HINT void UnlockShared()
 	{
 		Windows::ReleaseSRWLockShared(&Mutex);
 	}
diff --git a/Engine/Source/Runtime/Core/Public/Windows/WindowsPlatformProperties.h b/Engine/Source/Runtime/Core/Public/Windows/WindowsPlatformProperties.h
--- a/Engine/Source/Runtime/Core/Public/Windows/WindowsPlatformProperties.h
+++ b/Engine/Source/Runtime/Core/Public/Windows/WindowsPlatformProperties.h
@@ -13,32 +13,32 @@ template<bool HAS_EDITOR_DATA, bool IS_DEDICATED_SERVER, bool IS_CLIENT_ONLY>
 struct FWindowsPlatformProperties
 	: public FGenericPlatformProperties
 {
-	static constexpr FORCEINLINE bool HasEditorOnlyData()
+	static constexpr UE_FORCEINLINE_HINT bool HasEditorOnlyData()
 	{
 		return HAS_EDITOR_DATA;
 	}
 
-	static constexpr FORCEINLINE const char* IniPlatformName()
+	static constexpr UE_FORCEINLINE_HINT const char* IniPlatformName()
 	{
 		return "Windows";
 	}
 
-	static constexpr FORCEINLINE const TCHAR* GetRuntimeSettingsClassName()
+	static constexpr UE_FORCEINLINE_HINT const TCHAR* GetRuntimeSettingsClassName()
 	{
 		return TEXT("/Script/WindowsTargetPlatform.WindowsTargetSettings");
 	}
 
-	static constexpr FORCEINLINE bool IsGameOnly()
+	static constexpr UE_FORCEINLINE_HINT bool IsGameOnly()
 	{
 		return UE_GAME;
 	}
 
-	static constexpr FORCEINLINE bool IsServerOnly()
+	static constexpr UE_FORCEINLINE_HINT bool IsServerOnly()
 	{
 		return IS_DEDICATED_SERVER;
 	}
 
-	static constexpr FORCEINLINE bool IsClientOnly()
+	static constexpr UE_FORCEINLINE_HINT bool IsClientOnly()
 	{
 		return IS_CLIENT_ONLY;
 	}
@@ -63,57 +63,57 @@ struct FWindowsPlatformProperties
 		return "Windows";
 	}
 
-	static constexpr FORCEINLINE bool RequiresCookedData()
+	static constexpr UE_FORCEINLINE_HINT bool RequiresCookedData()
 	{
 		return !HAS_EDITOR_DATA;
 	}
 
-	static constexpr FORCEINLINE bool HasSecurePackageFormat()
+	static constexpr UE_FORCEINLINE_HINT bool HasSecurePackageFormat()
 	{
 		return IS_DEDICATED_SERVER;
 	}
 
-	static constexpr FORCEINLINE bool SupportsMemoryMappedFiles()
+	static constexpr UE_FORCEINLINE_HINT bool SupportsMemoryMappedFiles()
 	{
 		return true;
 	}
 
-	static constexpr FORCEINLINE bool SupportsAudioStreaming()
+	static constexpr UE_FORCEINLINE_HINT bool SupportsAudioStreaming()
 	{
 		return !IsServerOnly();
 	}
 
-	static constexpr FORCEINLINE bool SupportsMeshLODStreaming()
+	static constexpr UE_FORCEINLINE_HINT bool SupportsMeshLODStreaming()
 	{
 		return !IsServerOnly() && !HasEditorOnlyData();
 	}
 
-	static constexpr FORCEINLINE bool SupportsRayTracing()
+	static constexpr UE_FORCEINLINE_HINT bool SupportsRayTracing()
 	{
 		return true;
 	}
 
-	static constexpr FORCEINLINE bool SupportsGrayscaleSRGB()
+	static constexpr UE_FORCEINLINE_HINT bool SupportsGrayscaleSRGB()
 	{
 		return false; // Requires expand from G8 to RGBA
 	}
 
-	static constexpr FORCEINLINE bool SupportsMultipleGameInstances()
+	static constexpr UE_FORCEINLINE_HINT bool SupportsMultipleGameInstances()
 	{
 		return true;
 	}
 
-	static constexpr FORCEINLINE bool SupportsWindowedMode()
+	static constexpr UE_FORCEINLINE_HINT bool SupportsWindowedMode()
 	{
 		return true;
 	}
 	
-	static constexpr FORCEINLINE bool HasFixedResolution()
+	static constexpr UE_FORCEINLINE_HINT bool HasFixedResolution()
 	{
 		return false;
 	}
 
-	static constexpr FORCEINLINE bool SupportsQuit()
+	static constexpr UE_FORCEINLINE_HINT bool SupportsQuit()
 	{
 		return true;
 	}
@@ -138,12 +138,12 @@ struct FWindowsPlatformProperties
 		return 1.0f;
 	}
 
-	static constexpr FORCEINLINE int64 GetMemoryMappingAlignment()
+	static constexpr UE_FORCEINLINE_HINT int64 GetMemoryMappingAlignment()
 	{
 		return 4096;
 	}
 
-	static constexpr FORCEINLINE int GetMaxSupportedVirtualMemoryAlignment()
+	static constexpr UE_FORCEINLINE_HINT int GetMaxSupportedVirtualMemoryAlignment()
 	{
 		return 65536;
 	}
diff --git a/Engine/Source/Runtime/Core/Public/Windows/WindowsPlatformTLS.h b/Engine/Source/Runtime/Core/Public/Windows/WindowsPlatformTLS.h
--- a/Engine/Source/Runtime/Core/Public/Windows/WindowsPlatformTLS.h
+++ b/Engine/Source/Runtime/Core/Public/Windows/WindowsPlatformTLS.h
@@ -39,7 +39,7 @@ struct FWindowsPlatformTLS
 	 *
 	 * @return The thread identifier.
 	 */
-	static FORCEINLINE uint32 GetCurrentThreadId(void)
+	static UE_FORCEINLINE_HINT uint32 GetCurrentThreadId(void)
 	{
 		return Windows::GetCurrentThreadId();
 	}
@@ -52,7 +52,7 @@ struct FWindowsPlatformTLS
 #if WITH_EDITOR
 	CORE_API static uint32 AllocTlsSlot();
 #else
-	FORCEINLINE static uint32 AllocTlsSlot()
+	UE_FORCEINLINE_HINT static uint32 AllocTlsSlot()
 	{
 		return Windows::TlsAlloc();
 }
@@ -67,7 +67,7 @@ struct FWindowsPlatformTLS
 #if WITH_EDITOR
 	CORE_API static void FreeTlsSlot(uint32 SlotIndex);
 #else
-	FORCEINLINE static void FreeTlsSlot(uint32 SlotIndex)
+	UE_FORCEINLINE_HINT static void FreeTlsSlot(uint32 SlotIndex)
 	{
 		Windows::TlsFree(SlotIndex);
 	}
diff --git a/Engine/Source/Runtime/CoreUObject/Public/AssetRegistry/AssetBundleData.h b/Engine/Source/Runtime/CoreUObject/Public/AssetRegistry/AssetBundleData.h
--- a/Engine/Source/Runtime/CoreUObject/Public/AssetRegistry/AssetBundleData.h
+++ b/Engine/Source/Runtime/CoreUObject/Public/AssetRegistry/AssetBundleData.h
@@ -25,9 +25,9 @@ struct FAssetBundleEntry
 	GENERATED_BODY()
 public:
 	/** Declare constructors inline so this can be a header only class */
-	FORCEINLINE FAssetBundleEntry() {}
+	UE_FORCEINLINE_HINT FAssetBundleEntry() {}
 PRAGMA_DISABLE_DEPRECATION_WARNINGS
-	FORCEINLINE ~FAssetBundleEntry() {}
+	UE_FORCEINLINE_HINT ~FAssetBundleEntry() {}
 	// Explicitly declared constructors/operators because we explicitly declared the destructor to hide deprecation warnings.
 	// These can be removed along with the destructor after deprecation.
 	FAssetBundleEntry(const FAssetBundleEntry&) = default;
@@ -84,8 +84,8 @@ struct FAssetBundleData
 public:
 
 	/** Declare constructors inline so this can be a header only class */
-	FORCEINLINE FAssetBundleData() {}
-	FORCEINLINE ~FAssetBundleData() {}
+	UE_FORCEINLINE_HINT FAssetBundleData() {}
+	UE_FORCEINLINE_HINT ~FAssetBundleData() {}
 	FAssetBundleData(const FAssetBundleData&) = default;
 	FAssetBundleData(FAssetBundleData&&) = default;
 	FAssetBundleData& operator=(const FAssetBundleData&) = default;
diff --git a/Engine/Source/Runtime/CoreUObject/Public/AssetRegistry/AssetData.h b/Engine/Source/Runtime/CoreUObject/Public/AssetRegistry/AssetData.h
--- a/Engine/Source/Runtime/CoreUObject/Public/AssetRegistry/AssetData.h
+++ b/Engine/Source/Runtime/CoreUObject/Public/AssetRegistry/AssetData.h
@@ -746,7 +746,7 @@ private:
 	 * Note that this function is force-inlined but defined in AssetData.cpp which is fine as functions will get inlined
 	 * as long as they're defined before they are used for the first time by other functions (SerializeForCacheWithTagsAndBundles in this case)
 	 */
-	FORCEINLINE void SerializeForCacheInternal(FArchive& Ar, FAssetRegistryVersion::Type Version, void (*SerializeTagsAndBundles)(FArchive& , FAssetData&, FAssetRegistryVersion::Type));
+	UE_FORCEINLINE_HINT void SerializeForCacheInternal(FArchive& Ar, FAssetRegistryVersion::Type Version, void (*SerializeTagsAndBundles)(FArchive& , FAssetData&, FAssetRegistryVersion::Type));
 	COREUOBJECT_API void SerializeForCacheWithTagsAndBundles(FArchive& Ar, void (*SerializeTagsAndBundles)(FArchive&, FAssetData&, FAssetRegistryVersion::Type));
 	COREUOBJECT_API void SerializeForCacheOldVersionWithTagsAndBundles(FArchive& Ar, FAssetRegistryVersion::Type Version, void (*SerializeTagsAndBundles)(FArchive&, FAssetData&, FAssetRegistryVersion::Type));
 
@@ -769,7 +769,7 @@ public:
 	/** Helper function that tries to convert short class name to path name */
 	COREUOBJECT_API static FTopLevelAssetPath TryConvertShortClassNameToPathName(FName InClassName, ELogVerbosity::Type FailureMessageVerbosity = ELogVerbosity::Warning);
 
-	friend FORCEINLINE uint32 GetTypeHash(const FAssetData& AssetData)
+	friend UE_FORCEINLINE_HINT uint32 GetTypeHash(const FAssetData& AssetData)
 	{
 		return HashCombine(GetTypeHash(AssetData.PackageName), GetTypeHash(AssetData.AssetName));
 	}
@@ -1066,7 +1066,7 @@ public:
 	}
 
 private:
-	FORCEINLINE void SerializeForCacheInternal(FArchive& Ar, FAssetPackageData& PackageData, FAssetRegistryVersion::Type Version);
+	UE_FORCEINLINE_HINT void SerializeForCacheInternal(FArchive& Ar, FAssetPackageData& PackageData, FAssetRegistryVersion::Type Version);
 
 	enum : uint32
 	{
diff --git a/Engine/Source/Runtime/CoreUObject/Public/Misc/ExclusiveLoadPackageTimeTracker.h b/Engine/Source/Runtime/CoreUObject/Public/Misc/ExclusiveLoadPackageTimeTracker.h
--- a/Engine/Source/Runtime/CoreUObject/Public/Misc/ExclusiveLoadPackageTimeTracker.h
+++ b/Engine/Source/Runtime/CoreUObject/Public/Misc/ExclusiveLoadPackageTimeTracker.h
@@ -79,7 +79,7 @@ public:
 	};
 
 	/** Starts a time for the specified package name. */
-	FORCEINLINE static void PushLoadPackage(FName PackageName)
+	UE_FORCEINLINE_HINT static void PushLoadPackage(FName PackageName)
 	{
 #if WITH_LOADPACKAGE_TIME_TRACKER
 		Get().InternalPushLoadPackage(PackageName);
@@ -87,7 +87,7 @@ public:
 	}
 
 	/** Records a time and stats for the loaded package. */
-	FORCEINLINE static void PopLoadPackage(UPackage* LoadedPackage)
+	UE_FORCEINLINE_HINT static void PopLoadPackage(UPackage* LoadedPackage)
 	{
 #if WITH_LOADPACKAGE_TIME_TRACKER
 		Get().InternalPopLoadPackage(LoadedPackage, nullptr);
@@ -108,7 +108,7 @@ public:
 	}
 
 	/** Records a time and stats for the loaded package. Optionally provide the loaded asset if it is known, otherwise the asset in the package is detected. */
-	FORCEINLINE static void PopPostLoad(UObject* PostLoadObject)
+	UE_FORCEINLINE_HINT static void PopPostLoad(UObject* PostLoadObject)
 	{
 #if WITH_LOADPACKAGE_TIME_TRACKER
 		Get().InternalPopLoadPackage(nullptr, PostLoadObject);
@@ -125,7 +125,7 @@ public:
 	}
 
 	/** Records some time spent in "EndLoad" */
-	FORCEINLINE static void PopEndLoad()
+	UE_FORCEINLINE_HINT static void PopEndLoad()
 	{
 #if WITH_LOADPACKAGE_TIME_TRACKER
 		Get().InternalPopLoadPackage(nullptr, nullptr);
@@ -133,7 +133,7 @@ public:
 	}
 
 	/** Displays the data gathered in various ways. */
-	FORCEINLINE static void DumpReport(const TArray<FString>& Args)
+	UE_FORCEINLINE_HINT static void DumpReport(const TArray<FString>& Args)
 	{
 #if WITH_LOADPACKAGE_TIME_TRACKER
 		Get().InternalDumpReport(Args);
@@ -141,7 +141,7 @@ public:
 	}
 
 	/** Resets the data. */
-	FORCEINLINE static void ResetReport()
+	UE_FORCEINLINE_HINT static void ResetReport()
 	{
 #if WITH_LOADPACKAGE_TIME_TRACKER
 		Get().InternalResetReport();
diff --git a/Engine/Source/Runtime/CoreUObject/Public/Misc/PackageAccessTracking.h b/Engine/Source/Runtime/CoreUObject/Public/Misc/PackageAccessTracking.h
--- a/Engine/Source/Runtime/CoreUObject/Public/Misc/PackageAccessTracking.h
+++ b/Engine/Source/Runtime/CoreUObject/Public/Misc/PackageAccessTracking.h
@@ -88,10 +88,10 @@ namespace PackageAccessTracking_Private
 
 		COREUOBJECT_API void SetPackageName(FName InPackageName);
 
-		FORCEINLINE FName GetPackageName() const { return DirectData.PackageName; }
-		FORCEINLINE FName GetOpName() const { return DirectData.OpName; }
-		FORCEINLINE FPackageAccessRefScope* GetOuter() const { return Outer; }
-		FORCEINLINE const ITargetPlatform* GetTargetPlatform() const { return DirectData.TargetPlatform; }
+		UE_FORCEINLINE_HINT FName GetPackageName() const { return DirectData.PackageName; }
+		UE_FORCEINLINE_HINT FName GetOpName() const { return DirectData.OpName; }
+		UE_FORCEINLINE_HINT FPackageAccessRefScope* GetOuter() const { return Outer; }
+		UE_FORCEINLINE_HINT const ITargetPlatform* GetTargetPlatform() const { return DirectData.TargetPlatform; }
 
 		static COREUOBJECT_API FPackageAccessRefScope* GetCurrentThreadScope();
 		static COREUOBJECT_API FTrackedData* GetCurrentThreadAccumulatedData();
diff --git a/Engine/Source/Runtime/CoreUObject/Public/Serialization/DuplicatedObject.h b/Engine/Source/Runtime/CoreUObject/Public/Serialization/DuplicatedObject.h
--- a/Engine/Source/Runtime/CoreUObject/Public/Serialization/DuplicatedObject.h
+++ b/Engine/Source/Runtime/CoreUObject/Public/Serialization/DuplicatedObject.h
@@ -31,7 +31,7 @@ struct FDuplicatedObject
 	/**
 	 * @return true if this is the default annotation and holds no information about a duplicated object
 	 */
-	FORCEINLINE bool IsDefault()
+	UE_FORCEINLINE_HINT bool IsDefault()
 	{
 		return bIsDefault;
 	}
diff --git a/Engine/Source/Runtime/CoreUObject/Public/Serialization/FindReferencersArchive.h b/Engine/Source/Runtime/CoreUObject/Public/Serialization/FindReferencersArchive.h
--- a/Engine/Source/Runtime/CoreUObject/Public/Serialization/FindReferencersArchive.h
+++ b/Engine/Source/Runtime/CoreUObject/Public/Serialization/FindReferencersArchive.h
@@ -174,12 +174,12 @@ protected:
 		}
 
 	private:
-		FORCEINLINE void CheckFrozen() const
+		UE_FORCEINLINE_HINT void CheckFrozen() const
 		{
 			checkf(bFrozen, TEXT("Container has not been frozen and cannot be searched yet"));
 		}
 
-		FORCEINLINE void CheckUnfrozen() const
+		UE_FORCEINLINE_HINT void CheckUnfrozen() const
 		{
 			checkf(!bFrozen, TEXT("Container has been frozen and cannot be modified anymore"));
 		}
