Markup wrapper functions in Algo with UE_REWRITE and remove redundant inline's.

#rnx

diff --git a/Engine/Source/Runtime/Core/Public/Algo/Accumulate.h b/Engine/Source/Runtime/Core/Public/Algo/Accumulate.h
--- a/Engine/Source/Runtime/Core/Public/Algo/Accumulate.h
+++ b/Engine/Source/Runtime/Core/Public/Algo/Accumulate.h
@@ -11,13 +11,13 @@
 template<typename T = void>
 struct TPlus
 {
-	[[nodiscard]] FORCEINLINE T operator()(const T& A, const T& B) { return A + B; }
+	[[nodiscard]] UE_REWRITE T operator()(const T& A, const T& B) { return A + B; }
 };
 template<>
 struct TPlus<void>
 {
 	template<typename U, typename V>
-	[[nodiscard]] FORCEINLINE auto operator()(U&& A, V&& B) -> decltype(A + B) { return A + B; }
+	[[nodiscard]] UE_REWRITE auto operator()(U&& A, V&& B) -> decltype(A + B) { return A + B; }
 };
 
 
@@ -33,7 +33,7 @@ namespace Algo
 	* @return the result of summing all the elements of Input
 	*/
 	template <typename T, typename A, typename OpT>
-	[[nodiscard]] inline T Accumulate(const A& Input, T Init, OpT Op)
+	[[nodiscard]] T Accumulate(const A& Input, T Init, OpT Op)
 	{
 		T Result = MoveTemp(Init);
 		for (const auto& InputElem : Input)
@@ -52,7 +52,7 @@ namespace Algo
 	 * @return the result of summing all the elements of Input
 	 */
 	template <typename T, typename A>
-	[[nodiscard]] FORCEINLINE T Accumulate(const A& Input, T Init)
+	[[nodiscard]] UE_REWRITE T Accumulate(const A& Input, T Init)
 	{
 		return Accumulate(Input, MoveTemp(Init), TPlus<>());
 	}
@@ -68,7 +68,7 @@ namespace Algo
 	* @return the result of mapping and then summing all the elements of Input
 	*/
 	template <typename T, typename A, typename MapT, typename OpT>
-	[[nodiscard]] inline T TransformAccumulate(const A& Input, MapT MapOp, T Init, OpT Op)
+	[[nodiscard]] T TransformAccumulate(const A& Input, MapT MapOp, T Init, OpT Op)
 	{
 		T Result = MoveTemp(Init);
 		for (const auto& InputElem : Input)
@@ -88,7 +88,7 @@ namespace Algo
 	* @return the result of mapping and then summing all the elements of Input
 	*/
 	template <typename T, typename A, typename MapT>
-	[[nodiscard]] FORCEINLINE T TransformAccumulate(const A& Input, MapT MapOp, T Init)
+	[[nodiscard]] UE_REWRITE T TransformAccumulate(const A& Input, MapT MapOp, T Init)
 	{
 		return TransformAccumulate(Input, MapOp, MoveTemp(Init), TPlus<>());
 	}
diff --git a/Engine/Source/Runtime/Core/Public/Algo/AllOf.h b/Engine/Source/Runtime/Core/Public/Algo/AllOf.h
--- a/Engine/Source/Runtime/Core/Public/Algo/AllOf.h
+++ b/Engine/Source/Runtime/Core/Public/Algo/AllOf.h
@@ -16,7 +16,7 @@ namespace Algo
 	 * @return  true if all the elements are truthy, false otherwise.
 	 */
 	template <typename RangeType>
-	[[nodiscard]] inline bool AllOf(const RangeType& Range)
+	[[nodiscard]] bool AllOf(const RangeType& Range)
 	{
 		for (const auto& Element : Range)
 		{
@@ -38,7 +38,7 @@ namespace Algo
 	 * @return  true if all the projections are truthy, false otherwise.
 	 */
 	template <typename RangeType, typename ProjectionType>
-	[[nodiscard]] inline bool AllOf(const RangeType& Range, ProjectionType Projection)
+	[[nodiscard]] bool AllOf(const RangeType& Range, ProjectionType Projection)
 	{
 		for (const auto& Element : Range)
 		{
diff --git a/Engine/Source/Runtime/Core/Public/Algo/AnyOf.h b/Engine/Source/Runtime/Core/Public/Algo/AnyOf.h
--- a/Engine/Source/Runtime/Core/Public/Algo/AnyOf.h
+++ b/Engine/Source/Runtime/Core/Public/Algo/AnyOf.h
@@ -17,7 +17,7 @@ namespace Algo
 	 * @return  true if at least one element is truthy, false otherwise.
 	 */
 	template <typename RangeType>
-	[[nodiscard]] FORCEINLINE bool AnyOf(const RangeType& Range)
+	[[nodiscard]] UE_REWRITE bool AnyOf(const RangeType& Range)
 	{
 		return !Algo::NoneOf(Range);
 	}
@@ -31,7 +31,7 @@ namespace Algo
 	 * @return  true if at least one projection is truthy, false otherwise.
 	 */
 	template <typename RangeType, typename ProjectionType>
-	[[nodiscard]] FORCEINLINE bool AnyOf(const RangeType& Range, ProjectionType Projection)
+	[[nodiscard]] UE_REWRITE bool AnyOf(const RangeType& Range, ProjectionType Projection)
 	{
 		return !Algo::NoneOf(Range, MoveTemp(Projection));
 	}
@@ -46,7 +46,7 @@ namespace Algo
 	 * @return  true if at least one projection is truthy, false otherwise.
 	 */
 	template <typename RangeType, typename ProjectionType>
-	[[nodiscard]] FORCEINLINE bool AnyOf(const RangeType& Range, ProjectionType Projection, ENoRef NoRef)
+	[[nodiscard]] UE_REWRITE bool AnyOf(const RangeType& Range, ProjectionType Projection, ENoRef NoRef)
 	{
 		return !Algo::NoneOf(Range, MoveTemp(Projection), NoRef);
 	}
diff --git a/Engine/Source/Runtime/Core/Public/Algo/BinarySearch.h b/Engine/Source/Runtime/Core/Public/Algo/BinarySearch.h
--- a/Engine/Source/Runtime/Core/Public/Algo/BinarySearch.h
+++ b/Engine/Source/Runtime/Core/Public/Algo/BinarySearch.h
@@ -20,7 +20,7 @@ namespace AlgoImpl
 	 * @returns Position of the first element >= Value, may be == Num
 	 */
 	template <typename RangeValueType, typename SizeType, typename PredicateValueType, typename ProjectionType, typename SortPredicateType>
-	inline SizeType LowerBoundInternal(RangeValueType* First, const SizeType Num, const PredicateValueType& Value, ProjectionType Projection, SortPredicateType SortPredicate)
+	SizeType LowerBoundInternal(RangeValueType* First, const SizeType Num, const PredicateValueType& Value, ProjectionType Projection, SortPredicateType SortPredicate)
 	{
 		// Current start of sequence to check
 		SizeType Start = 0;
@@ -53,7 +53,7 @@ namespace AlgoImpl
 	 * @returns Position of the first element > Value, may be == Num
 	 */
 	template <typename RangeValueType, typename SizeType, typename PredicateValueType, typename ProjectionType, typename SortPredicateType>
-	inline SizeType UpperBoundInternal(RangeValueType* First, const SizeType Num, const PredicateValueType& Value, ProjectionType Projection, SortPredicateType SortPredicate)
+	SizeType UpperBoundInternal(RangeValueType* First, const SizeType Num, const PredicateValueType& Value, ProjectionType Projection, SortPredicateType SortPredicate)
 	{
 		// Current start of sequence to check
 		SizeType Start = 0;
@@ -89,12 +89,12 @@ namespace Algo
 	 * @returns Position of the first element >= Value, may be position after last element in range
 	 */
 	template <typename RangeType, typename ValueType, typename SortPredicateType>
-	[[nodiscard]] FORCEINLINE auto LowerBound(const RangeType& Range, const ValueType& Value, SortPredicateType SortPredicate) -> decltype(GetNum(Range))
+	[[nodiscard]] UE_REWRITE auto LowerBound(const RangeType& Range, const ValueType& Value, SortPredicateType SortPredicate) -> decltype(GetNum(Range))
 	{
 		return AlgoImpl::LowerBoundInternal(GetData(Range), GetNum(Range), Value, FIdentityFunctor(), SortPredicate);
 	}
 	template <typename RangeType, typename ValueType>
-	[[nodiscard]] FORCEINLINE auto LowerBound(const RangeType& Range, const ValueType& Value) -> decltype(GetNum(Range))
+	[[nodiscard]] UE_REWRITE auto LowerBound(const RangeType& Range, const ValueType& Value) -> decltype(GetNum(Range))
 	{
 		return AlgoImpl::LowerBoundInternal(GetData(Range), GetNum(Range), Value, FIdentityFunctor(), TLess<>());
 	}
@@ -110,12 +110,12 @@ namespace Algo
 	 * @returns Position of the first element >= Value, may be position after last element in range
 	 */
 	template <typename RangeType, typename ValueType, typename ProjectionType, typename SortPredicateType>
-	[[nodiscard]] FORCEINLINE auto LowerBoundBy(const RangeType& Range, const ValueType& Value, ProjectionType Projection, SortPredicateType SortPredicate) -> decltype(GetNum(Range))
+	[[nodiscard]] UE_REWRITE auto LowerBoundBy(const RangeType& Range, const ValueType& Value, ProjectionType Projection, SortPredicateType SortPredicate) -> decltype(GetNum(Range))
 	{
 		return AlgoImpl::LowerBoundInternal(GetData(Range), GetNum(Range), Value, Projection, SortPredicate);
 	}
 	template <typename RangeType, typename ValueType, typename ProjectionType>
-	[[nodiscard]] FORCEINLINE auto LowerBoundBy(const RangeType& Range, const ValueType& Value, ProjectionType Projection) -> decltype(GetNum(Range))
+	[[nodiscard]] UE_REWRITE auto LowerBoundBy(const RangeType& Range, const ValueType& Value, ProjectionType Projection) -> decltype(GetNum(Range))
 	{
 		return AlgoImpl::LowerBoundInternal(GetData(Range), GetNum(Range), Value, Projection, TLess<>());
 	}
@@ -130,12 +130,12 @@ namespace Algo
 	 * @returns Position of the first element > Value, may be past end of range
 	 */
 	template <typename RangeType, typename ValueType, typename SortPredicateType>
-	[[nodiscard]] FORCEINLINE auto UpperBound(const RangeType& Range, const ValueType& Value, SortPredicateType SortPredicate) -> decltype(GetNum(Range))
+	[[nodiscard]] UE_REWRITE auto UpperBound(const RangeType& Range, const ValueType& Value, SortPredicateType SortPredicate) -> decltype(GetNum(Range))
 	{
 		return AlgoImpl::UpperBoundInternal(GetData(Range), GetNum(Range), Value, FIdentityFunctor(), SortPredicate);
 	}
 	template <typename RangeType, typename ValueType>
-	[[nodiscard]] FORCEINLINE auto UpperBound(const RangeType& Range, const ValueType& Value) -> decltype(GetNum(Range))
+	[[nodiscard]] UE_REWRITE auto UpperBound(const RangeType& Range, const ValueType& Value) -> decltype(GetNum(Range))
 	{
 		return AlgoImpl::UpperBoundInternal(GetData(Range), GetNum(Range), Value, FIdentityFunctor(), TLess<>());
 	}
@@ -151,12 +151,12 @@ namespace Algo
 	 * @returns Position of the first element > Value, may be past end of range
 	 */
 	template <typename RangeType, typename ValueType, typename ProjectionType, typename SortPredicateType>
-	[[nodiscard]] FORCEINLINE auto UpperBoundBy(const RangeType& Range, const ValueType& Value, ProjectionType Projection, SortPredicateType SortPredicate) -> decltype(GetNum(Range))
+	[[nodiscard]] UE_REWRITE auto UpperBoundBy(const RangeType& Range, const ValueType& Value, ProjectionType Projection, SortPredicateType SortPredicate) -> decltype(GetNum(Range))
 	{
 		return AlgoImpl::UpperBoundInternal(GetData(Range), GetNum(Range), Value, Projection, SortPredicate);
 	}
 	template <typename RangeType, typename ValueType, typename ProjectionType>
-	[[nodiscard]] FORCEINLINE auto UpperBoundBy(const RangeType& Range, const ValueType& Value, ProjectionType Projection) -> decltype(GetNum(Range))
+	[[nodiscard]] UE_REWRITE auto UpperBoundBy(const RangeType& Range, const ValueType& Value, ProjectionType Projection) -> decltype(GetNum(Range))
 	{
 		return AlgoImpl::UpperBoundInternal(GetData(Range), GetNum(Range), Value, Projection, TLess<>());
 	}
@@ -170,7 +170,7 @@ namespace Algo
 	 * @return Index of found element, or INDEX_NONE
 	 */
 	template <typename RangeType, typename ValueType, typename SortPredicateType>
-	[[nodiscard]] inline auto BinarySearch(const RangeType& Range, const ValueType& Value, SortPredicateType SortPredicate) -> decltype(GetNum(Range))
+	[[nodiscard]] auto BinarySearch(const RangeType& Range, const ValueType& Value, SortPredicateType SortPredicate) -> decltype(GetNum(Range))
 	{
 		auto CheckIndex = LowerBound(Range, Value, SortPredicate);
 		if (CheckIndex < GetNum(Range))
@@ -185,7 +185,7 @@ namespace Algo
 		return INDEX_NONE;
 	}
 	template <typename RangeType, typename ValueType>
-	[[nodiscard]] FORCEINLINE auto BinarySearch(const RangeType& Range, const ValueType& Value)
+	[[nodiscard]] UE_REWRITE auto BinarySearch(const RangeType& Range, const ValueType& Value)
 	{
 		return BinarySearch(Range, Value, TLess<>());
 	}
@@ -200,7 +200,7 @@ namespace Algo
 	 * @return Index of found element, or INDEX_NONE
 	 */
 	template <typename RangeType, typename ValueType, typename ProjectionType, typename SortPredicateType>
-	[[nodiscard]] inline auto BinarySearchBy(const RangeType& Range, const ValueType& Value, ProjectionType Projection, SortPredicateType SortPredicate) -> decltype(GetNum(Range))
+	[[nodiscard]] auto BinarySearchBy(const RangeType& Range, const ValueType& Value, ProjectionType Projection, SortPredicateType SortPredicate) -> decltype(GetNum(Range))
 	{
 		auto CheckIndex = LowerBoundBy(Range, Value, Projection, SortPredicate);
 		if (CheckIndex < GetNum(Range))
@@ -215,7 +215,7 @@ namespace Algo
 		return INDEX_NONE;
 	}
 	template <typename RangeType, typename ValueType, typename ProjectionType>
-	[[nodiscard]] FORCEINLINE auto BinarySearchBy(const RangeType& Range, const ValueType& Value, ProjectionType Projection)
+	[[nodiscard]] UE_REWRITE auto BinarySearchBy(const RangeType& Range, const ValueType& Value, ProjectionType Projection)
 	{
 		return BinarySearchBy(Range, Value, Projection, TLess<>());
 	}
diff --git a/Engine/Source/Runtime/Core/Public/Algo/Contains.h b/Engine/Source/Runtime/Core/Public/Algo/Contains.h
--- a/Engine/Source/Runtime/Core/Public/Algo/Contains.h
+++ b/Engine/Source/Runtime/Core/Public/Algo/Contains.h
@@ -16,7 +16,7 @@ namespace Algo
 	 * @return true if an element was found, false otherwise.
 	 */
 	template <typename RangeType, typename ValueType>
-	[[nodiscard]] FORCEINLINE bool Contains(RangeType&& Range, const ValueType& Value)
+	[[nodiscard]] UE_REWRITE bool Contains(RangeType&& Range, const ValueType& Value)
 	{
 		return !!Algo::Find(Forward<RangeType>(Range), Value);
 	}
@@ -31,7 +31,7 @@ namespace Algo
 	 * @return true if an element was found, false otherwise.
 	 */
 	template <typename RangeType, typename ValueType, typename ProjectionType>
-	[[nodiscard]] FORCEINLINE bool ContainsBy(RangeType&& Range, const ValueType& Value, ProjectionType Proj)
+	[[nodiscard]] UE_REWRITE bool ContainsBy(RangeType&& Range, const ValueType& Value, ProjectionType Proj)
 	{
 		return !!Algo::FindBy(Forward<RangeType>(Range), Value, MoveTemp(Proj));
 	}
diff --git a/Engine/Source/Runtime/Core/Public/Algo/Copy.h b/Engine/Source/Runtime/Core/Public/Algo/Copy.h
--- a/Engine/Source/Runtime/Core/Public/Algo/Copy.h
+++ b/Engine/Source/Runtime/Core/Public/Algo/Copy.h
@@ -15,7 +15,7 @@ namespace Algo
 	 * @param  Predicate  Condition which returns true for elements that should be copied and false for elements that should be skipped
 	 */
 	template <typename InT, typename OutT, typename PredicateT>
-	inline void CopyIf(const InT& Input, OutT& Output, PredicateT Predicate)
+	void CopyIf(const InT& Input, OutT& Output, PredicateT Predicate)
 	{
 		for (const auto& Value : Input)
 		{
@@ -33,7 +33,7 @@ namespace Algo
 	 * @param  Output  Container to hold the output
 	 */
 	template <typename InT, typename OutT>
-	inline void Copy(const InT& Input, OutT& Output)
+	void Copy(const InT& Input, OutT& Output)
 	{
 		for (const auto& Value : Input)
 		{
@@ -49,7 +49,7 @@ namespace Algo
 	 * @param  Output  Container to hold the output
 	 */
 	template <typename InT, typename OutT>
-	inline void Copy(const InT& Input, OutT& Output, ENoRef NoRef)
+	void Copy(const InT& Input, OutT& Output, ENoRef NoRef)
 	{
 		for (const auto Value : Input)
 		{
diff --git a/Engine/Source/Runtime/Core/Public/Algo/Count.h b/Engine/Source/Runtime/Core/Public/Algo/Count.h
--- a/Engine/Source/Runtime/Core/Public/Algo/Count.h
+++ b/Engine/Source/Runtime/Core/Public/Algo/Count.h
@@ -13,7 +13,7 @@ namespace Algo
 	 * @param  InValue  Value to compare against
 	 */
 	template <typename InT, typename ValueT>
-	[[nodiscard]] inline SIZE_T Count(const InT& Input, const ValueT& InValue)
+	[[nodiscard]] SIZE_T Count(const InT& Input, const ValueT& InValue)
 	{
 		SIZE_T Result = 0;
 		for (const auto& Value : Input)
@@ -33,7 +33,7 @@ namespace Algo
 	 * @param  Predicate  Condition which returns true for elements that should be counted and false for elements that should be skipped
 	 */
 	template <typename InT, typename PredicateT>
-	[[nodiscard]] inline SIZE_T CountIf(const InT& Input, PredicateT Predicate)
+	[[nodiscard]] SIZE_T CountIf(const InT& Input, PredicateT Predicate)
 	{
 		SIZE_T Result = 0;
 		for (const auto& Value : Input)
@@ -54,7 +54,7 @@ namespace Algo
 	 * @param  Projection  The projection to apply to the element
 	 */
 	template <typename InT, typename ValueT, typename ProjectionT>
-	[[nodiscard]] inline SIZE_T CountBy(const InT& Input, const ValueT& InValue, ProjectionT Proj)
+	[[nodiscard]] SIZE_T CountBy(const InT& Input, const ValueT& InValue, ProjectionT Proj)
 	{
 		SIZE_T Result = 0;
 		for (const auto& Value : Input)
diff --git a/Engine/Source/Runtime/Core/Public/Algo/Find.h b/Engine/Source/Runtime/Core/Public/Algo/Find.h
--- a/Engine/Source/Runtime/Core/Public/Algo/Find.h
+++ b/Engine/Source/Runtime/Core/Public/Algo/Find.h
@@ -49,7 +49,7 @@ namespace Algo
 	 * @return A pointer to the first element found, or nullptr if none was found.
 	 */
 	template <typename RangeType, typename ValueType>
-	[[nodiscard]] inline auto Find(RangeType&& Range, const ValueType& Value)
+	[[nodiscard]] UE_REWRITE auto Find(RangeType&& Range, const ValueType& Value)
 		-> decltype(AlgoImpl::FindBy(Forward<RangeType>(Range), Value, FIdentityFunctor()))
 	{
 		return AlgoImpl::FindBy(Forward<RangeType>(Range), Value, FIdentityFunctor());
@@ -65,7 +65,7 @@ namespace Algo
 	 * @return A pointer to the first element found, or nullptr if none was found.
 	 */
 	template <typename RangeType, typename ValueType, typename ProjectionType>
-	[[nodiscard]] inline auto FindBy(RangeType&& Range, const ValueType& Value, ProjectionType Proj)
+	[[nodiscard]] UE_REWRITE auto FindBy(RangeType&& Range, const ValueType& Value, ProjectionType Proj)
 		-> decltype(AlgoImpl::FindBy(Forward<RangeType>(Range), Value, MoveTemp(Proj)))
 	{
 		return AlgoImpl::FindBy(Forward<RangeType>(Range), Value, MoveTemp(Proj));
@@ -80,7 +80,7 @@ namespace Algo
 	 * @return A pointer to the first element found, or nullptr if none was found.
 	 */
 	template <typename RangeType, typename PredicateType>
-	[[nodiscard]] inline auto FindByPredicate(RangeType&& Range, PredicateType Pred)
+	[[nodiscard]] UE_REWRITE auto FindByPredicate(RangeType&& Range, PredicateType Pred)
 		-> decltype(AlgoImpl::FindByPredicate(Forward<RangeType>(Range), MoveTemp(Pred)))
 	{
 		return AlgoImpl::FindByPredicate(Forward<RangeType>(Range), MoveTemp(Pred));
diff --git a/Engine/Source/Runtime/Core/Public/Algo/FindLast.h b/Engine/Source/Runtime/Core/Public/Algo/FindLast.h
--- a/Engine/Source/Runtime/Core/Public/Algo/FindLast.h
+++ b/Engine/Source/Runtime/Core/Public/Algo/FindLast.h
@@ -48,7 +48,7 @@ namespace Algo
 	 * @return A pointer to the last element found, or nullptr if none was found.
 	 */
 	template <typename RangeType, typename ValueType>
-	[[nodiscard]] inline auto FindLast(RangeType&& Range, const ValueType& Value)
+	[[nodiscard]] UE_REWRITE auto FindLast(RangeType&& Range, const ValueType& Value)
 		-> decltype(AlgoImpl::FindLastBy(GetData(Range), GetNum(Range), Value, FIdentityFunctor()))
 	{
 		return AlgoImpl::FindLastBy(GetData(Range), GetNum(Range), Value, FIdentityFunctor());
@@ -64,7 +64,7 @@ namespace Algo
 	 * @return A pointer to the last element found, or nullptr if none was found.
 	 */
 	template <typename RangeType, typename ValueType, typename ProjectionType>
-	[[nodiscard]] inline auto FindLastBy(RangeType&& Range, const ValueType& Value, ProjectionType Proj)
+	[[nodiscard]] UE_REWRITE auto FindLastBy(RangeType&& Range, const ValueType& Value, ProjectionType Proj)
 		-> decltype(AlgoImpl::FindLastBy(GetData(Range), GetNum(Range), Value, MoveTemp(Proj)))
 	{
 		return AlgoImpl::FindLastBy(GetData(Range), GetNum(Range), Value, MoveTemp(Proj));
@@ -79,7 +79,7 @@ namespace Algo
 	 * @return A pointer to the last element found, or nullptr if none was found.
 	 */
 	template <typename RangeType, typename PredicateType>
-	[[nodiscard]] inline auto FindLastByPredicate(RangeType&& Range, PredicateType Pred)
+	[[nodiscard]] UE_REWRITE auto FindLastByPredicate(RangeType&& Range, PredicateType Pred)
 		-> decltype(AlgoImpl::FindLastByPredicate(GetData(Range), GetNum(Range), MoveTemp(Pred)))
 	{
 		return AlgoImpl::FindLastByPredicate(GetData(Range), GetNum(Range), MoveTemp(Pred));
diff --git a/Engine/Source/Runtime/Core/Public/Algo/FindSequence.h b/Engine/Source/Runtime/Core/Public/Algo/FindSequence.h
--- a/Engine/Source/Runtime/Core/Public/Algo/FindSequence.h
+++ b/Engine/Source/Runtime/Core/Public/Algo/FindSequence.h
@@ -43,7 +43,7 @@ namespace Algo
 	* @return A pointer to the first occurrence of the "What" sequence in "Where" sequence, or nullptr if not found.
 	*/
 	template<typename RangeWhereType, typename RangeWhatType>
-	[[nodiscard]] inline auto FindSequence(const RangeWhereType& Where, const RangeWhatType& What)
+	[[nodiscard]] auto FindSequence(const RangeWhereType& Where, const RangeWhatType& What)
 		-> decltype( AlgoImpl::FindSequence( GetData(Where), GetData(Where) + GetNum(Where), GetData(What), GetData(What) + GetNum(What)) )
 	{
 		if (GetNum(What) > GetNum(Where))
diff --git a/Engine/Source/Runtime/Core/Public/Algo/FindSortedStringCaseInsensitive.h b/Engine/Source/Runtime/Core/Public/Algo/FindSortedStringCaseInsensitive.h
--- a/Engine/Source/Runtime/Core/Public/Algo/FindSortedStringCaseInsensitive.h
+++ b/Engine/Source/Runtime/Core/Public/Algo/FindSortedStringCaseInsensitive.h
@@ -26,7 +26,7 @@ namespace Algo
 	 * @return The index of the found string in the array, or -1 if the string was not found.
 	 */
 	template <int32 ArraySize>
-	[[nodiscard]] FORCEINLINE int32 FindSortedStringCaseInsensitive(const TCHAR* Str, const TCHAR* const (&SortedArray)[ArraySize])
+	[[nodiscard]] UE_REWRITE int32 FindSortedStringCaseInsensitive(const TCHAR* Str, const TCHAR* const (&SortedArray)[ArraySize])
 	{
 		return FindSortedStringCaseInsensitive(Str, SortedArray, ArraySize);
 	}
diff --git a/Engine/Source/Runtime/Core/Public/Algo/ForEach.h b/Engine/Source/Runtime/Core/Public/Algo/ForEach.h
--- a/Engine/Source/Runtime/Core/Public/Algo/ForEach.h
+++ b/Engine/Source/Runtime/Core/Public/Algo/ForEach.h
@@ -15,7 +15,7 @@ namespace Algo
 	 * @param  Callable   Callable object
 	 */
 	template <typename InT, typename PredicateT, typename CallableT>
-	inline void ForEachIf(InT&& Input, PredicateT Predicate, CallableT Callable)
+	void ForEachIf(InT&& Input, PredicateT Predicate, CallableT Callable)
 	{
 		for (auto& Value : Input)
 		{
@@ -33,7 +33,7 @@ namespace Algo
 	 * @param  Callable   Callable object
 	 */
 	template <typename InT, typename CallableT>
-	inline void ForEach(InT&& Input, CallableT Callable)
+	void ForEach(InT&& Input, CallableT Callable)
 	{
 		for (auto& Value : Input)
 		{
diff --git a/Engine/Source/Runtime/Core/Public/Algo/HeapSort.h b/Engine/Source/Runtime/Core/Public/Algo/HeapSort.h
--- a/Engine/Source/Runtime/Core/Public/Algo/HeapSort.h
+++ b/Engine/Source/Runtime/Core/Public/Algo/HeapSort.h
@@ -15,7 +15,7 @@ namespace Algo
 	 * @param Range		The range to sort.
 	 */
 	template <typename RangeType>
-	FORCEINLINE void HeapSort(RangeType&& Range)
+	UE_REWRITE void HeapSort(RangeType&& Range)
 	{
 		AlgoImpl::HeapSortInternal(GetData(Range), GetNum(Range), FIdentityFunctor(), TLess<>());
 	}
@@ -27,7 +27,7 @@ namespace Algo
 	 * @param Predicate	A binary predicate object used to specify if one element should precede another.
 	 */
 	template <typename RangeType, typename PredicateType>
-	FORCEINLINE void HeapSort(RangeType&& Range, PredicateType Predicate)
+	UE_REWRITE void HeapSort(RangeType&& Range, PredicateType Predicate)
 	{
 		AlgoImpl::HeapSortInternal(GetData(Range), GetNum(Range), FIdentityFunctor(), MoveTemp(Predicate));
 	}
@@ -39,7 +39,7 @@ namespace Algo
 	 * @param Projection	The projection to sort by when applied to the element.
 	 */
 	template <typename RangeType, typename ProjectionType>
-	FORCEINLINE void HeapSortBy(RangeType&& Range, ProjectionType Projection)
+	UE_REWRITE void HeapSortBy(RangeType&& Range, ProjectionType Projection)
 	{
 		AlgoImpl::HeapSortInternal(GetData(Range), GetNum(Range), Projection, TLess<>());
 	}
@@ -52,7 +52,7 @@ namespace Algo
 	 * @param Predicate	A binary predicate object, applied to the projection, used to specify if one element should precede another.
 	 */
 	template <typename RangeType, typename ProjectionType, typename PredicateType>
-	FORCEINLINE void HeapSortBy(RangeType&& Range, ProjectionType Projection, PredicateType Predicate)
+	UE_REWRITE void HeapSortBy(RangeType&& Range, ProjectionType Projection, PredicateType Predicate)
 	{
 		AlgoImpl::HeapSortInternal(GetData(Range), GetNum(Range), MoveTemp(Projection), MoveTemp(Predicate));
 	}
diff --git a/Engine/Source/Runtime/Core/Public/Algo/Heapify.h b/Engine/Source/Runtime/Core/Public/Algo/Heapify.h
--- a/Engine/Source/Runtime/Core/Public/Algo/Heapify.h
+++ b/Engine/Source/Runtime/Core/Public/Algo/Heapify.h
@@ -17,7 +17,7 @@ namespace Algo
 	 * @param Range	The range to heapify.
 	 */
 	template <typename RangeType>
-	FORCEINLINE void Heapify(RangeType&& Range)
+	UE_REWRITE void Heapify(RangeType&& Range)
 	{
 		AlgoImpl::HeapifyInternal(GetData(Range), GetNum(Range), FIdentityFunctor(), TLess<>());
 	}
@@ -29,7 +29,7 @@ namespace Algo
 	 * @param Predicate	A binary predicate object used to specify if one element should precede another.
 	 */
 	template <typename RangeType, typename PredicateType>
-	FORCEINLINE void Heapify(RangeType&& Range, PredicateType Predicate)
+	UE_REWRITE void Heapify(RangeType&& Range, PredicateType Predicate)
 	{
 		AlgoImpl::HeapifyInternal(GetData(Range), GetNum(Range), FIdentityFunctor(), Predicate);
 	}
@@ -42,7 +42,7 @@ namespace Algo
 	 * @param Projection	The projection to apply to the elements.
 	 */
 	template <typename RangeType, typename ProjectionType, typename PredicateType>
-	FORCEINLINE void HeapifyBy(RangeType&& Range, ProjectionType Projection)
+	UE_REWRITE void HeapifyBy(RangeType&& Range, ProjectionType Projection)
 	{
 		AlgoImpl::HeapifyInternal(GetData(Range), GetNum(Range), Projection, TLess<>());
 	}
@@ -55,7 +55,7 @@ namespace Algo
 	 * @param Predicate		A binary predicate object used to specify if one element should precede another.
 	 */
 	template <typename RangeType, typename ProjectionType, typename PredicateType>
-	FORCEINLINE void HeapifyBy(RangeType&& Range, ProjectionType Projection, PredicateType Predicate)
+	UE_REWRITE void HeapifyBy(RangeType&& Range, ProjectionType Projection, PredicateType Predicate)
 	{
 		AlgoImpl::HeapifyInternal(GetData(Range), GetNum(Range), Projection, Predicate);
 	}
diff --git a/Engine/Source/Runtime/Core/Public/Algo/IndexOf.h b/Engine/Source/Runtime/Core/Public/Algo/IndexOf.h
--- a/Engine/Source/Runtime/Core/Public/Algo/IndexOf.h
+++ b/Engine/Source/Runtime/Core/Public/Algo/IndexOf.h
@@ -56,7 +56,7 @@ namespace Algo
 	 * @return The index of the first element found, or (SizeType)-1 if none was found (where SizeType = decltype(GetNum(Range)))
 	 */
 	template <typename RangeType, typename ValueType>
-	[[nodiscard]] inline auto IndexOf(const RangeType& Range, const ValueType& Value)
+	[[nodiscard]] UE_REWRITE auto IndexOf(const RangeType& Range, const ValueType& Value)
 	{
 		static_assert(TIsContiguousContainer<RangeType>::Value, "Must only be used with contiguous containers");
 		return AlgoImpl::IndexOfBy(Range, Value, FIdentityFunctor());
@@ -72,7 +72,7 @@ namespace Algo
 	 * @return The index of the first element found, or (SizeType)-1 if none was found (where SizeType = decltype(GetNum(Range)))
 	 */
 	template <typename RangeType, typename ValueType, typename ProjectionType>
-	[[nodiscard]] inline auto IndexOfBy(const RangeType& Range, const ValueType& Value, ProjectionType Proj)
+	[[nodiscard]] UE_REWRITE auto IndexOfBy(const RangeType& Range, const ValueType& Value, ProjectionType Proj)
 	{
 		static_assert(TIsContiguousContainer<RangeType>::Value, "Must only be used with contiguous containers");
 		return AlgoImpl::IndexOfBy(Range, Value, MoveTemp(Proj));
@@ -87,7 +87,7 @@ namespace Algo
 	 * @return The index of the first element found, or (SizeType)-1 if none was found (where SizeType = decltype(GetNum(Range)))
 	 */
 	template <typename RangeType, typename PredicateType>
-	[[nodiscard]] inline auto IndexOfByPredicate(const RangeType& Range, PredicateType Pred)
+	[[nodiscard]] UE_REWRITE auto IndexOfByPredicate(const RangeType& Range, PredicateType Pred)
 	{
 		static_assert(TIsContiguousContainer<RangeType>::Value, "Must only be used with contiguous containers");
 		return AlgoImpl::IndexOfByPredicate(Range, MoveTemp(Pred));
diff --git a/Engine/Source/Runtime/Core/Public/Algo/IntroSort.h b/Engine/Source/Runtime/Core/Public/Algo/IntroSort.h
--- a/Engine/Source/Runtime/Core/Public/Algo/IntroSort.h
+++ b/Engine/Source/Runtime/Core/Public/Algo/IntroSort.h
@@ -146,7 +146,7 @@ namespace Algo
 	 * @param Range	The range to sort.
 	 */
 	template <typename RangeType>
-	FORCEINLINE void IntroSort(RangeType&& Range)
+	UE_REWRITE void IntroSort(RangeType&& Range)
 	{
 		AlgoImpl::IntroSortInternal(GetData(Range), GetNum(Range), FIdentityFunctor(), TLess<>());
 	}
@@ -158,7 +158,7 @@ namespace Algo
 	 * @param Predicate	A binary predicate object used to specify if one element should precede another.
 	 */
 	template <typename RangeType, typename PredicateType>
-	FORCEINLINE void IntroSort(RangeType&& Range, PredicateType Predicate)
+	UE_REWRITE void IntroSort(RangeType&& Range, PredicateType Predicate)
 	{
 		AlgoImpl::IntroSortInternal(GetData(Range), GetNum(Range), FIdentityFunctor(), MoveTemp(Predicate));
 	}
@@ -170,7 +170,7 @@ namespace Algo
 	 * @param Proj			The projection to sort by when applied to the element.
 	 */
 	template <typename RangeType, typename ProjectionType>
-	inline void IntroSortBy(RangeType&& Range, ProjectionType Proj)
+	UE_REWRITE void IntroSortBy(RangeType&& Range, ProjectionType Proj)
 	{
 // Workaround for a codegen bug that was discovered related to member function pointer projections that are virtual
 #if defined(_MSC_FULL_VER) && _MSC_FULL_VER < 194134123
@@ -188,7 +188,7 @@ namespace Algo
 	 * @param Predicate		A binary predicate object, applied to the projection, used to specify if one element should precede another.
 	 */
 	template <typename RangeType, typename ProjectionType, typename PredicateType>
-	inline void IntroSortBy(RangeType&& Range, ProjectionType Proj, PredicateType Predicate)
+	UE_REWRITE void IntroSortBy(RangeType&& Range, ProjectionType Proj, PredicateType Predicate)
 	{
 // Workaround for a codegen bug that was discovered related to member function pointer projections that are virtual
 #if defined(_MSC_FULL_VER) && _MSC_FULL_VER < 194134123
diff --git a/Engine/Source/Runtime/Core/Public/Algo/IsHeap.h b/Engine/Source/Runtime/Core/Public/Algo/IsHeap.h
--- a/Engine/Source/Runtime/Core/Public/Algo/IsHeap.h
+++ b/Engine/Source/Runtime/Core/Public/Algo/IsHeap.h
@@ -47,7 +47,7 @@ namespace Algo
 	 * @return	returns	true if the range is a min-heap
 	 */
 	template <typename RangeType>
-	[[nodiscard]] FORCEINLINE bool IsHeap(const RangeType& Range)
+	[[nodiscard]] UE_REWRITE bool IsHeap(const RangeType& Range)
 	{
 		return AlgoImpl::IsHeapInternal(GetData(Range), GetNum(Range), FIdentityFunctor(), TLess<>());
 	}
@@ -61,7 +61,7 @@ namespace Algo
 	 * @return	returns		true if the range is a min-heap
 	 */
 	template <typename RangeType, typename PredicateType>
-	[[nodiscard]] FORCEINLINE bool IsHeap(const RangeType& Range, PredicateType Predicate)
+	[[nodiscard]] UE_REWRITE bool IsHeap(const RangeType& Range, PredicateType Predicate)
 	{
 		return AlgoImpl::IsHeapInternal(GetData(Range), GetNum(Range), FIdentityFunctor(), MoveTemp(Predicate));
 	}
@@ -75,7 +75,7 @@ namespace Algo
 	 * @return	returns		true if the range is a min-heap
 	 */
 	template <typename RangeType, typename ProjectionType>
-	[[nodiscard]] FORCEINLINE bool IsHeapBy(const RangeType& Range, ProjectionType Projection)
+	[[nodiscard]] UE_REWRITE bool IsHeapBy(const RangeType& Range, ProjectionType Projection)
 	{
 		return AlgoImpl::IsHeapInternal(GetData(Range), GetNum(Range), MoveTemp(Projection), TLess<>());
 	}
@@ -90,7 +90,7 @@ namespace Algo
 	 * @return	returns		true if the range is a min-heap
 	 */
 	template <typename RangeType, typename ProjectionType, typename PredicateType>
-	[[nodiscard]] FORCEINLINE bool IsHeapBy(const RangeType& Range, ProjectionType Projection, PredicateType Predicate)
+	[[nodiscard]] UE_REWRITE bool IsHeapBy(const RangeType& Range, ProjectionType Projection, PredicateType Predicate)
 	{
 		return AlgoImpl::IsHeapInternal(GetData(Range), GetNum(Range), MoveTemp(Projection), MoveTemp(Predicate));
 	}
diff --git a/Engine/Source/Runtime/Core/Public/Algo/IsSorted.h b/Engine/Source/Runtime/Core/Public/Algo/IsSorted.h
--- a/Engine/Source/Runtime/Core/Public/Algo/IsSorted.h
+++ b/Engine/Source/Runtime/Core/Public/Algo/IsSorted.h
@@ -63,7 +63,7 @@ namespace Algo
 	 * @return true if the range is sorted, false otherwise.
 	 */
 	template <typename RangeType>
-	[[nodiscard]] FORCEINLINE bool IsSorted(const RangeType& Range)
+	[[nodiscard]] UE_REWRITE bool IsSorted(const RangeType& Range)
 	{
 		return AlgoImpl::IsSortedBy(GetData(Range), GetNum(Range), FIdentityFunctor(), TLess<>());
 	}
@@ -77,7 +77,7 @@ namespace Algo
 	 * @return true if the range is sorted, false otherwise.
 	 */
 	template <typename RangeType, typename PredType>
-	[[nodiscard]] FORCEINLINE bool IsSorted(const RangeType& Range, PredType Pred)
+	[[nodiscard]] UE_REWRITE bool IsSorted(const RangeType& Range, PredType Pred)
 	{
 		return AlgoImpl::IsSortedBy(GetData(Range), GetNum(Range), FIdentityFunctor(), MoveTemp(Pred));
 	}
@@ -90,7 +90,7 @@ namespace Algo
 	 * @return true if the range is sorted, false otherwise.
 	 */
 	template <typename RangeType, typename ProjectionType>
-	[[nodiscard]] FORCEINLINE bool IsSortedBy(const RangeType& Range, ProjectionType Projection)
+	[[nodiscard]] UE_REWRITE bool IsSortedBy(const RangeType& Range, ProjectionType Projection)
 	{
 		return AlgoImpl::IsSortedBy(GetData(Range), GetNum(Range), MoveTemp(Projection), TLess<>());
 	}
@@ -104,7 +104,7 @@ namespace Algo
 	 * @return true if the range is sorted, false otherwise.
 	 */
 	template <typename RangeType, typename ProjectionType, typename PredType>
-	[[nodiscard]] FORCEINLINE bool IsSortedBy(const RangeType& Range, ProjectionType Projection, PredType Pred)
+	[[nodiscard]] UE_REWRITE bool IsSortedBy(const RangeType& Range, ProjectionType Projection, PredType Pred)
 	{
 		return AlgoImpl::IsSortedBy(GetData(Range), GetNum(Range), MoveTemp(Projection), MoveTemp(Pred));
 	}
diff --git a/Engine/Source/Runtime/Core/Public/Algo/LegacySort.h b/Engine/Source/Runtime/Core/Public/Algo/LegacySort.h
--- a/Engine/Source/Runtime/Core/Public/Algo/LegacySort.h
+++ b/Engine/Source/Runtime/Core/Public/Algo/LegacySort.h
@@ -117,7 +117,7 @@ namespace Algo
 	 * @param  Range  The range to sort.
 	 */
 	template <typename RangeType>
-	FORCEINLINE void LegacySort(RangeType&& Range)
+	UE_REWRITE void LegacySort(RangeType&& Range)
 	{
 		AlgoImpl::LegacySortInternal(GetData(Range), GetNum(Range), FIdentityFunctor(), TLess<>());
 	}
@@ -130,7 +130,7 @@ namespace Algo
 	 * @param  Predicate  A binary predicate object used to specify if one element should precede another.
 	 */
 	template <typename RangeType, typename PredicateType>
-	FORCEINLINE void LegacySort(RangeType&& Range, PredicateType Pred)
+	UE_REWRITE void LegacySort(RangeType&& Range, PredicateType Pred)
 	{
 		AlgoImpl::LegacySortInternal(GetData(Range), GetNum(Range), FIdentityFunctor(), MoveTemp(Pred));
 	}
@@ -143,7 +143,7 @@ namespace Algo
 	 * @param  Proj   The projection to sort by when applied to the element.
 	 */
 	template <typename RangeType, typename ProjectionType>
-	FORCEINLINE void LegacySortBy(RangeType&& Range, ProjectionType Proj)
+	UE_REWRITE void LegacySortBy(RangeType&& Range, ProjectionType Proj)
 	{
 		AlgoImpl::LegacySortInternal(GetData(Range), GetNum(Range), MoveTemp(Proj), TLess<>());
 	}
@@ -157,7 +157,7 @@ namespace Algo
 	 * @param  Predicate  A binary predicate object, applied to the projection, used to specify if one element should precede another.
 	 */
 	template <typename RangeType, typename ProjectionType, typename PredicateType>
-	FORCEINLINE void LegacySortBy(RangeType&& Range, ProjectionType Proj, PredicateType Pred)
+	UE_REWRITE void LegacySortBy(RangeType&& Range, ProjectionType Proj, PredicateType Pred)
 	{
 		AlgoImpl::LegacySortInternal(GetData(Range), GetNum(Range), MoveTemp(Proj), MoveTemp(Pred));
 	}
diff --git a/Engine/Source/Runtime/Core/Public/Algo/MaxElement.h b/Engine/Source/Runtime/Core/Public/Algo/MaxElement.h
--- a/Engine/Source/Runtime/Core/Public/Algo/MaxElement.h
+++ b/Engine/Source/Runtime/Core/Public/Algo/MaxElement.h
@@ -37,7 +37,7 @@ namespace Algo
 	 * @return A pointer to the maximum element, or nullptr if the range was empty.
 	 */
 	template <typename RangeType>
-	[[nodiscard]] inline auto MaxElement(RangeType&& Range)
+	[[nodiscard]] UE_REWRITE auto MaxElement(RangeType&& Range)
 		-> decltype(AlgoImpl::MaxElementBy(Range, FIdentityFunctor(), TLess<>()))
 	{
 		return AlgoImpl::MaxElementBy(Range, FIdentityFunctor(), TLess<>());
@@ -52,7 +52,7 @@ namespace Algo
 	 * @return A pointer to the maximum element, or nullptr if the range was empty.
 	 */
 	template <typename RangeType, typename ComparatorType>
-	[[nodiscard]] inline auto MaxElement(RangeType&& Range, ComparatorType Comp)
+	[[nodiscard]] UE_REWRITE auto MaxElement(RangeType&& Range, ComparatorType Comp)
 		-> decltype(AlgoImpl::MaxElementBy(Range, FIdentityFunctor(), MoveTemp(Comp)))
 	{
 		return AlgoImpl::MaxElementBy(Range, FIdentityFunctor(), MoveTemp(Comp));
@@ -67,7 +67,7 @@ namespace Algo
 	 * @return A pointer to the maximum element, or nullptr if the range was empty.
 	 */
 	template <typename RangeType, typename ProjectionType>
-	[[nodiscard]] inline auto MaxElementBy(RangeType&& Range, ProjectionType Proj)
+	[[nodiscard]] UE_REWRITE auto MaxElementBy(RangeType&& Range, ProjectionType Proj)
 		-> decltype(AlgoImpl::MaxElementBy(Range, MoveTemp(Proj), TLess<>()))
 	{
 		return AlgoImpl::MaxElementBy(Range, MoveTemp(Proj), TLess<>());
@@ -83,7 +83,7 @@ namespace Algo
 	 * @return A pointer to the maximum element, or nullptr if the range was empty.
 	 */
 	template <typename RangeType, typename ProjectionType, typename ComparatorType>
-	[[nodiscard]] inline auto MaxElementBy(RangeType&& Range, ProjectionType Proj, ComparatorType Comp)
+	[[nodiscard]] UE_REWRITE auto MaxElementBy(RangeType&& Range, ProjectionType Proj, ComparatorType Comp)
 		-> decltype(AlgoImpl::MaxElementBy(Range, MoveTemp(Proj), MoveTemp(Comp)))
 	{
 		return AlgoImpl::MaxElementBy(Range, MoveTemp(Proj), MoveTemp(Comp));
diff --git a/Engine/Source/Runtime/Core/Public/Algo/MinElement.h b/Engine/Source/Runtime/Core/Public/Algo/MinElement.h
--- a/Engine/Source/Runtime/Core/Public/Algo/MinElement.h
+++ b/Engine/Source/Runtime/Core/Public/Algo/MinElement.h
@@ -37,7 +37,7 @@ namespace Algo
 	 * @return A pointer to the minimum element, or nullptr if the range was empty.
 	 */
 	template <typename RangeType>
-	[[nodiscard]] inline auto MinElement(RangeType&& Range)
+	[[nodiscard]] UE_REWRITE auto MinElement(RangeType&& Range)
 		-> decltype(AlgoImpl::MinElementBy(Range, FIdentityFunctor(), TLess<>()))
 	{
 		return AlgoImpl::MinElementBy(Range, FIdentityFunctor(), TLess<>());
@@ -52,7 +52,7 @@ namespace Algo
 	 * @return A pointer to the minimum element, or nullptr if the range was empty.
 	 */
 	template <typename RangeType, typename ComparatorType>
-	[[nodiscard]] inline auto MinElement(RangeType&& Range, ComparatorType Comp)
+	[[nodiscard]] UE_REWRITE auto MinElement(RangeType&& Range, ComparatorType Comp)
 		-> decltype(AlgoImpl::MinElementBy(Range, FIdentityFunctor(), MoveTemp(Comp)))
 	{
 		return AlgoImpl::MinElementBy(Range, FIdentityFunctor(), MoveTemp(Comp));
@@ -67,7 +67,7 @@ namespace Algo
 	 * @return A pointer to the minimum element, or nullptr if the range was empty.
 	 */
 	template <typename RangeType, typename ProjectionType>
-	[[nodiscard]] inline auto MinElementBy(RangeType&& Range, ProjectionType Proj)
+	[[nodiscard]] UE_REWRITE auto MinElementBy(RangeType&& Range, ProjectionType Proj)
 		-> decltype(AlgoImpl::MinElementBy(Range, MoveTemp(Proj), TLess<>()))
 	{
 		return AlgoImpl::MinElementBy(Range, MoveTemp(Proj), TLess<>());
@@ -83,7 +83,7 @@ namespace Algo
 	 * @return A pointer to the minimum element, or nullptr if the range was empty.
 	 */
 	template <typename RangeType, typename ProjectionType, typename ComparatorType>
-	[[nodiscard]] inline auto MinElementBy(RangeType&& Range, ProjectionType Proj, ComparatorType Comp)
+	[[nodiscard]] UE_REWRITE auto MinElementBy(RangeType&& Range, ProjectionType Proj, ComparatorType Comp)
 		-> decltype(AlgoImpl::MinElementBy(Range, MoveTemp(Proj), MoveTemp(Comp)))
 	{
 		return AlgoImpl::MinElementBy(Range, MoveTemp(Proj), MoveTemp(Comp));
diff --git a/Engine/Source/Runtime/Core/Public/Algo/NoneOf.h b/Engine/Source/Runtime/Core/Public/Algo/NoneOf.h
--- a/Engine/Source/Runtime/Core/Public/Algo/NoneOf.h
+++ b/Engine/Source/Runtime/Core/Public/Algo/NoneOf.h
@@ -17,7 +17,7 @@ namespace Algo
 	 * @return  true if no element is truthy, false otherwise.
 	 */
 	template <typename RangeType>
-	[[nodiscard]] inline bool NoneOf(const RangeType& Range)
+	[[nodiscard]] bool NoneOf(const RangeType& Range)
 	{
 		for (const auto& Element : Range)
 		{
@@ -39,7 +39,7 @@ namespace Algo
 	 * @return  true if none of the projections are truthy, false otherwise.
 	 */
 	template <typename RangeType, typename ProjectionType>
-	[[nodiscard]] inline bool NoneOf(const RangeType& Range, ProjectionType Projection)
+	[[nodiscard]] bool NoneOf(const RangeType& Range, ProjectionType Projection)
 	{
 		for (const auto& Element : Range)
 		{
@@ -62,7 +62,7 @@ namespace Algo
 	 * @return  true if none of the projections are truthy, false otherwise.
 	 */
 	template <typename RangeType, typename ProjectionType>
-	[[nodiscard]] inline bool NoneOf(const RangeType& Range, ProjectionType Projection, ENoRef NoRef)
+	[[nodiscard]] bool NoneOf(const RangeType& Range, ProjectionType Projection, ENoRef NoRef)
 	{
 		for (const auto Element : Range)
 		{
diff --git a/Engine/Source/Runtime/Core/Public/Algo/Partition.h b/Engine/Source/Runtime/Core/Public/Algo/Partition.h
--- a/Engine/Source/Runtime/Core/Public/Algo/Partition.h
+++ b/Engine/Source/Runtime/Core/Public/Algo/Partition.h
@@ -55,7 +55,7 @@ namespace Algo
 	 * @return	index of the first element in the second group
 	 */
 	template <typename RangeType, typename UnaryPredicateType>
-	FORCEINLINE auto Partition(RangeType&& Range, UnaryPredicateType Predicate) -> decltype(GetNum(Range))
+	UE_REWRITE auto Partition(RangeType&& Range, UnaryPredicateType Predicate) -> decltype(GetNum(Range))
 	{
 		return Partition(GetData(Range), GetNum(Range), MoveTemp(Predicate));
 	}
diff --git a/Engine/Source/Runtime/Core/Public/Algo/Replace.h b/Engine/Source/Runtime/Core/Public/Algo/Replace.h
--- a/Engine/Source/Runtime/Core/Public/Algo/Replace.h
+++ b/Engine/Source/Runtime/Core/Public/Algo/Replace.h
@@ -15,7 +15,7 @@ namespace Algo
 	 * @param InNew The value to copy in as a replacement.
 	 */
 	template <typename RangeType, typename ValueType>
-	inline void Replace(RangeType&& Range, const ValueType& InOld, const ValueType& InNew)
+	void Replace(RangeType&& Range, const ValueType& InOld, const ValueType& InNew)
 	{
 		for (auto& Elem : Forward<RangeType>(Range))
 		{
@@ -34,7 +34,7 @@ namespace Algo
 	 * @param InNew The value to copy in as a replacement for each element satisfying the predicate.
 	 */
 	template <typename RangeType, typename ValueType, typename PredicateType>
-	inline void ReplaceIf(RangeType&& Range, PredicateType InPred, const ValueType& InNew)
+	void ReplaceIf(RangeType&& Range, PredicateType InPred, const ValueType& InNew)
 	{
 		for (auto& Elem : Forward<RangeType>(Range))
 		{
diff --git a/Engine/Source/Runtime/Core/Public/Algo/Reverse.h b/Engine/Source/Runtime/Core/Public/Algo/Reverse.h
--- a/Engine/Source/Runtime/Core/Public/Algo/Reverse.h
+++ b/Engine/Source/Runtime/Core/Public/Algo/Reverse.h
@@ -8,7 +8,7 @@
 namespace AlgoImpl
 {
 	template <typename T>
-	inline void Reverse(T* Array, int32 ArraySize)
+	void Reverse(T* Array, int32 ArraySize)
 	{
 		for (int32 i = 0, i2 = ArraySize - 1; i < ArraySize / 2 /*rounding down*/; ++i, --i2)
 		{
@@ -25,7 +25,7 @@ namespace Algo
 	 * @param  Array  The array to reverse.
 	 */
 	template <typename T, int32 ArraySize>
-	FORCEINLINE void Reverse(T (&Array)[ArraySize])
+	UE_REWRITE void Reverse(T (&Array)[ArraySize])
 	{
 		return AlgoImpl::Reverse((T*)Array, ArraySize);
 	}
@@ -37,7 +37,7 @@ namespace Algo
 	 * @param  ArraySize  The number of elements in the array.
 	 */
 	template <typename T>
-	FORCEINLINE void Reverse(T* Array, int32 ArraySize)
+	UE_REWRITE void Reverse(T* Array, int32 ArraySize)
 	{
 		return AlgoImpl::Reverse(Array, ArraySize);
 	}
@@ -48,7 +48,7 @@ namespace Algo
 	 * @param  Container  The container to reverse
 	 */
 	template <typename ContainerType>
-	FORCEINLINE void Reverse(ContainerType&& Container)
+	UE_REWRITE void Reverse(ContainerType&& Container)
 	{
 		return AlgoImpl::Reverse(Container.GetData(), Container.Num());
 	}
diff --git a/Engine/Source/Runtime/Core/Public/Algo/Rotate.h b/Engine/Source/Runtime/Core/Public/Algo/Rotate.h
--- a/Engine/Source/Runtime/Core/Public/Algo/Rotate.h
+++ b/Engine/Source/Runtime/Core/Public/Algo/Rotate.h
@@ -56,7 +56,7 @@ namespace Algo
 	 * @return The new index of the element that was previously at the start of the range.
 	 */
 	template <typename RangeType>
-	FORCEINLINE int32 Rotate(RangeType&& Range, int32 Count)
+	UE_REWRITE int32 Rotate(RangeType&& Range, int32 Count)
 	{
 		return AlgoImpl::RotateInternal(GetData(Range), GetNum(Range), Count);
 	}
diff --git a/Engine/Source/Runtime/Core/Public/Algo/SelectRandomWeighted.h b/Engine/Source/Runtime/Core/Public/Algo/SelectRandomWeighted.h
--- a/Engine/Source/Runtime/Core/Public/Algo/SelectRandomWeighted.h
+++ b/Engine/Source/Runtime/Core/Public/Algo/SelectRandomWeighted.h
@@ -58,7 +58,7 @@ namespace Algo
 	 * @param  Proj   The projection to weight the random selection by. Should yield a numeric type.
 	 */
 	template <typename RangeType, typename ProjectionType>
-	inline auto SelectRandomWeightedBy(RangeType&& Range, ProjectionType Proj)
+	UE_REWRITE auto SelectRandomWeightedBy(RangeType&& Range, ProjectionType Proj)
 		-> decltype(AlgoImpl::SelectRandomWeightedBy(Forward<RangeType>(Range), MoveTemp(Proj)))
 	{
 		return AlgoImpl::SelectRandomWeightedBy(Forward<RangeType>(Range), MoveTemp(Proj));
diff --git a/Engine/Source/Runtime/Core/Public/Algo/Sort.h b/Engine/Source/Runtime/Core/Public/Algo/Sort.h
--- a/Engine/Source/Runtime/Core/Public/Algo/Sort.h
+++ b/Engine/Source/Runtime/Core/Public/Algo/Sort.h
@@ -13,7 +13,7 @@ namespace Algo
 	 * @param  Range  The range to sort.
 	 */
 	template <typename RangeType>
-	FORCEINLINE void Sort(RangeType&& Range)
+	UE_REWRITE void Sort(RangeType&& Range)
 	{
 		IntroSort(Forward<RangeType>(Range));
 	}
@@ -25,7 +25,7 @@ namespace Algo
 	 * @param  Predicate  A binary predicate object used to specify if one element should precede another.
 	 */
 	template <typename RangeType, typename PredicateType>
-	FORCEINLINE void Sort(RangeType&& Range, PredicateType Pred)
+	UE_REWRITE void Sort(RangeType&& Range, PredicateType Pred)
 	{
 		IntroSort(Forward<RangeType>(Range), MoveTemp(Pred));
 	}
@@ -37,7 +37,7 @@ namespace Algo
 	 * @param  Proj   The projection to sort by when applied to the element.
 	 */
 	template <typename RangeType, typename ProjectionType>
-	FORCEINLINE void SortBy(RangeType&& Range, ProjectionType Proj)
+	UE_REWRITE void SortBy(RangeType&& Range, ProjectionType Proj)
 	{
 		IntroSortBy(Forward<RangeType>(Range), MoveTemp(Proj));
 	}
@@ -50,7 +50,7 @@ namespace Algo
 	 * @param  Predicate  A binary predicate object, applied to the projection, used to specify if one element should precede another.
 	 */
 	template <typename RangeType, typename ProjectionType, typename PredicateType>
-	FORCEINLINE void SortBy(RangeType&& Range, ProjectionType Proj, PredicateType Pred)
+	UE_REWRITE void SortBy(RangeType&& Range, ProjectionType Proj, PredicateType Pred)
 	{
 		IntroSortBy(Forward<RangeType>(Range), MoveTemp(Proj), MoveTemp(Pred));
 	}
diff --git a/Engine/Source/Runtime/Core/Public/Algo/StableSort.h b/Engine/Source/Runtime/Core/Public/Algo/StableSort.h
--- a/Engine/Source/Runtime/Core/Public/Algo/StableSort.h
+++ b/Engine/Source/Runtime/Core/Public/Algo/StableSort.h
@@ -122,7 +122,7 @@ namespace Algo
 	 * @param  Range  The range to sort.
 	 */
 	template <typename RangeType>
-	FORCEINLINE void StableSort(RangeType&& Range)
+	UE_REWRITE void StableSort(RangeType&& Range)
 	{
 		AlgoImpl::StableSortInternal(GetData(Range), GetNum(Range), FIdentityFunctor(), TLess<>());
 	}
@@ -134,7 +134,7 @@ namespace Algo
 	 * @param  Predicate  A binary predicate object used to specify if one element should precede another.
 	 */
 	template <typename RangeType, typename PredicateType>
-	FORCEINLINE void StableSort(RangeType&& Range, PredicateType Pred)
+	UE_REWRITE void StableSort(RangeType&& Range, PredicateType Pred)
 	{
 		AlgoImpl::StableSortInternal(GetData(Range), GetNum(Range), FIdentityFunctor(), MoveTemp(Pred));
 	}
@@ -146,7 +146,7 @@ namespace Algo
 	 * @param  Proj   The projection to sort by when applied to the element.
 	 */
 	template <typename RangeType, typename ProjectionType>
-	FORCEINLINE void StableSortBy(RangeType&& Range, ProjectionType Proj)
+	UE_REWRITE void StableSortBy(RangeType&& Range, ProjectionType Proj)
 	{
 		AlgoImpl::StableSortInternal(GetData(Range), GetNum(Range), MoveTemp(Proj), TLess<>());
 	}
@@ -159,7 +159,7 @@ namespace Algo
 	 * @param  Predicate  A binary predicate object, applied to the projection, used to specify if one element should precede another.
 	 */
 	template <typename RangeType, typename ProjectionType, typename PredicateType>
-	FORCEINLINE void StableSortBy(RangeType&& Range, ProjectionType Proj, PredicateType Pred)
+	UE_REWRITE void StableSortBy(RangeType&& Range, ProjectionType Proj, PredicateType Pred)
 	{
 		AlgoImpl::StableSortInternal(GetData(Range), GetNum(Range), MoveTemp(Proj), MoveTemp(Pred));
 	}
diff --git a/Engine/Source/Runtime/Core/Public/Algo/TopologicalSort.h b/Engine/Source/Runtime/Core/Public/Algo/TopologicalSort.h
--- a/Engine/Source/Runtime/Core/Public/Algo/TopologicalSort.h
+++ b/Engine/Source/Runtime/Core/Public/Algo/TopologicalSort.h
@@ -28,7 +28,7 @@ namespace Algo
 	 * @return True if succeeded, false if failed due to cycles.
 	 */
 	template <typename RangeType, typename GetElementDependenciesType>
-	FORCEINLINE bool TopologicalSort(RangeType&& UniqueRange, GetElementDependenciesType GetElementDependencies,
+	UE_REWRITE bool TopologicalSort(RangeType&& UniqueRange, GetElementDependenciesType GetElementDependencies,
 		ETopologicalSort Flags = ETopologicalSort::None)
 	{
 		return KahnTopologicalSort(Forward<RangeType>(UniqueRange), Forward<GetElementDependenciesType>(GetElementDependencies), Flags);
diff --git a/Engine/Source/Runtime/Core/Public/Algo/Transform.h b/Engine/Source/Runtime/Core/Public/Algo/Transform.h
--- a/Engine/Source/Runtime/Core/Public/Algo/Transform.h
+++ b/Engine/Source/Runtime/Core/Public/Algo/Transform.h
@@ -44,7 +44,7 @@ namespace Algo
 	 * @param  Trans      Transformation operation
 	 */
 	template <typename InT, typename OutT, typename PredicateT, typename TransformT>
-	inline void TransformIf(const InT& Input, OutT&& Output, PredicateT Predicate, TransformT Trans)
+	void TransformIf(const InT& Input, OutT&& Output, PredicateT Predicate, TransformT Trans)
 	{
 		for (const auto& Value : Input)
 		{
@@ -63,7 +63,7 @@ namespace Algo
 	 * @param  Trans   Transformation operation
 	 */
 	template <typename InT, typename OutT, typename TransformT>
-	inline void Transform(const InT& Input, OutT&& Output, TransformT Trans)
+	void Transform(const InT& Input, OutT&& Output, TransformT Trans)
 	{
 		for (const auto& Value : Input)
 		{
