Added format string validation to LowLevelFatalError

diff --git a/Engine/Source/Runtime/Core/Private/Misc/AssertionMacros.cpp b/Engine/Source/Runtime/Core/Private/Misc/AssertionMacros.cpp
--- a/Engine/Source/Runtime/Core/Private/Misc/AssertionMacros.cpp
+++ b/Engine/Source/Runtime/Core/Private/Misc/AssertionMacros.cpp
@@ -794,7 +794,7 @@ FORCENOINLINE bool FDebug::OptionallyLogFormattedEnsureMessageReturningFalseImpl
 }
 #endif
 
-void UE_COLD UE_DEBUG_SECTION VARARGS LowLevelFatalErrorHandler(const ANSICHAR* File, int32 Line, const TCHAR* Format, ...)
+void UE_COLD UE_DEBUG_SECTION VARARGS UE::Assert::Private::ProcessLowLevelFatalError(const ANSICHAR* File, int32 Line, const TCHAR* Format, ...)
 {
 	va_list Args;
 	va_start(Args, Format);
diff --git a/Engine/Source/Runtime/Core/Public/Misc/AssertionMacros.h b/Engine/Source/Runtime/Core/Public/Misc/AssertionMacros.h
--- a/Engine/Source/Runtime/Core/Public/Misc/AssertionMacros.h
+++ b/Engine/Source/Runtime/Core/Public/Misc/AssertionMacros.h
@@ -6,12 +6,12 @@
 #include "HAL/Platform.h"
 #include "HAL/PlatformMisc.h"
 #include "HAL/PreprocessorHelpers.h"
+#include "Misc/VarArgs.h"
+#include "String/FormatStringSan.h"
 #include "Templates/EnableIf.h"
 #include "Templates/IsArrayOrRefOfTypeByPredicate.h"
 #include "Templates/IsValidVariadicFunctionArg.h"
 #include "Traits/IsCharEncodingCompatibleWith.h"
-#include "Misc/VarArgs.h"
-#include "String/FormatStringSan.h"
 
 #include <atomic>
 
@@ -538,12 +538,20 @@ namespace UEAsserts_Private
 	Low level error macros
 ----------------------------------------------------------------------------*/
 
-/** low level fatal error handler. */
-CORE_API void UE_COLD UE_DEBUG_SECTION VARARGS LowLevelFatalErrorHandler (const ANSICHAR* File, int32 Line, const TCHAR* Format=TEXT(""), ... );
+namespace UE::Assert::Private
+{
+	CORE_API void UE_COLD UE_DEBUG_SECTION VARARGS ProcessLowLevelFatalError(const ANSICHAR* File, int32 Line, const TCHAR* Format, ...);
+} // UE::Assert::Private
+
+template <typename... ArgTypes>
+UE_DEPRECATED(5.7, "Use LowLevelFatalError.")
+UE_REWRITE void LowLevelFatalErrorHandler(const ANSICHAR* File, int32 Line, const TCHAR* Format = TEXT(""), ArgTypes... Args)
+{
+	UE::Assert::Private::ProcessLowLevelFatalError(File, Line, Format, Args...);
+}
 
 #define LowLevelFatalError(Format, ...) \
 	{ \
-		static_assert(TIsArrayOrRefOfTypeByPredicate<decltype(Format), TIsCharEncodingCompatibleWithTCHAR>::Value, "Formatting string must be a TCHAR array."); \
-		LowLevelFatalErrorHandler(__FILE__, __LINE__, (const TCHAR*)Format, ##__VA_ARGS__); \
+		UE_VALIDATE_FORMAT_STRING(Format, ##__VA_ARGS__); \
+		::UE::Assert::Private::ProcessLowLevelFatalError(__FILE__, __LINE__, Format, ##__VA_ARGS__); \
 	}
-
diff --git a/Engine/Source/Runtime/Engine/Public/TextureResource.h b/Engine/Source/Runtime/Engine/Public/TextureResource.h
--- a/Engine/Source/Runtime/Engine/Public/TextureResource.h
+++ b/Engine/Source/Runtime/Engine/Public/TextureResource.h
@@ -63,7 +63,7 @@ struct FTexture2DMipMap
 	{
 		if (InSizeX > 0xFFFF || InSizeY > 0xFFFF || InSizeZ > 0xFFFF)
 		{
-			LowLevelFatalError(TEXT("texture mip size doesn't fit in 16 bits! %ux&ux%u"), InSizeX, InSizeY, InSizeZ);
+			LowLevelFatalError(TEXT("texture mip size doesn't fit in 16 bits! %ux%ux%u"), InSizeX, InSizeY, InSizeZ);
 		}
 	}
 	FTexture2DMipMap(FTexture2DMipMap&&) = default;
