Handle a null file in StaticFailDebugV

A null file is possible when UE_LOG_INCLUDE_SOURCE_LOCATION=0.

#rnx

diff --git a/Engine/Source/Runtime/Core/Private/Misc/AssertionMacros.cpp b/Engine/Source/Runtime/Core/Private/Misc/AssertionMacros.cpp
--- a/Engine/Source/Runtime/Core/Private/Misc/AssertionMacros.cpp
+++ b/Engine/Source/Runtime/Core/Private/Misc/AssertionMacros.cpp
@@ -271,7 +271,7 @@ FORCENOINLINE void StaticFailDebugV(
 	FErrorHistWriter Writer;
 	Writer.Append(Error);
 	Writer.Append(Expression);
-	Writer.Appendf(TEXT(FILE_LINE_DESC_ANSI "\n"), File, Line);
+	Writer.Appendf(TEXT(FILE_LINE_DESC_ANSI "\n"), File ? File : "Unknown", Line);
 	if (Message)
 	{
 		Writer.Append(Message);
