Added missing typename.

diff --git a/Engine/Source/Runtime/Core/Public/Containers/SparseArray.h b/Engine/Source/Runtime/Core/Public/Containers/SparseArray.h
--- a/Engine/Source/Runtime/Core/Public/Containers/SparseArray.h
+++ b/Engine/Source/Runtime/Core/Public/Containers/SparseArray.h
@@ -1275,7 +1275,7 @@ public:
 
 private:
 
-	using FElementOrFreeListLink = SuperType::FElementOrFreeListLink;
+	using FElementOrFreeListLink = typename SuperType::FElementOrFreeListLink;
 
 	/** Extracts the element value from the array's element structure and passes it to the user provided comparison class. */
 	template <typename PREDICATE_CLASS>
