Cleaned up TTuple constraint usage.

diff --git a/Engine/Source/Runtime/Core/Public/Concepts/NotCVRefTo.h b/Engine/Source/Runtime/Core/Public/Concepts/NotCVRefTo.h
new file mode 100644
--- /dev/null
+++ b/Engine/Source/Runtime/Core/Public/Concepts/NotCVRefTo.h
@@ -0,0 +1,26 @@
+// Copyright Epic Games, Inc. All Rights Reserved.
+
+#pragma once
+
+#include "Concepts/SameAs.h"
+#include <type_traits>
+
+namespace UE
+{
+	/**
+	 * Concept which describes convertibility from one type to another.
+	 *
+	 * This is useful for templated constructors to prevent them being instantiated for the class type itself
+	 * so that it will call a move or copy constructor instead:
+	 *
+	 * struct FSomeClass
+	 * {
+	 *     template <CNotCVRefTo<FSomeClass> ArgType>
+	 *     FSomeClass(ArgType&& Arg);
+	 *     {
+	 *     }
+	 * };
+	 */
+	template <typename T, typename SelfType>
+	concept CNotCVRefTo = !UE::CSameAs<std::remove_cv_t<std::remove_reference_t<T>>, SelfType>;
+}
diff --git a/Engine/Source/Runtime/Core/Public/Templates/Tuple.h b/Engine/Source/Runtime/Core/Public/Templates/Tuple.h
--- a/Engine/Source/Runtime/Core/Public/Templates/Tuple.h
+++ b/Engine/Source/Runtime/Core/Public/Templates/Tuple.h
@@ -3,6 +3,7 @@
 #pragma once
 
 #include "CoreTypes.h"
+#include "Concepts/NotCVRefTo.h"
 #include "Templates/UnrealTemplate.h"
 #include "Delegates/IntegerSequence.h"
 #include "Templates/Invoke.h"
@@ -53,10 +54,6 @@ namespace UE::Core::Private::Tuple
 	enum EForwardingConstructor { ForwardingConstructor };
 	enum EOtherTupleConstructor { OtherTupleConstructor };
 
-	// This only exists to have something to expand a parameter pack into, for concept checking
-	template <typename... ArgTypes>
-	void ConceptCheckingHelper(ArgTypes&&...);
-
 	template <typename T, typename... Types>
 	constexpr uint32 TTypeCountInParameterPack_V = 0;
 
@@ -66,17 +63,13 @@ namespace UE::Core::Private::Tuple
 	template <typename T, uint32 Index, uint32 TupleSize>
 	struct TTupleBaseElement
 	{
-		template <
-			typename ArgType
-			UE_REQUIRES(std::is_constructible_v<T, ArgType&&>)
-		>
+		template <typename ArgType>
 		explicit TTupleBaseElement(EForwardingConstructor, ArgType&& Arg)
 			: Value(Forward<ArgType>(Arg))
 		{
 		}
 
 		TTupleBaseElement()
-			requires (std::is_default_constructible_v<T>)
 			: Value()
 		{
 		}
@@ -92,17 +85,13 @@ namespace UE::Core::Private::Tuple
 	template <typename T>
 	struct TTupleBaseElement<T, 0, 2>
 	{
-		template <
-			typename ArgType
-			UE_REQUIRES(std::is_constructible_v<T, ArgType&&>)
-		>
+		template <typename ArgType>
 		explicit TTupleBaseElement(EForwardingConstructor, ArgType&& Arg)
 			: Key(Forward<ArgType>(Arg))
 		{
 		}
 
 		TTupleBaseElement()
-			requires (std::is_default_constructible_v<T>)
 			: Key()
 		{
 		}
@@ -270,19 +259,13 @@ namespace UE::Core::Private::Tuple
 	template <uint32... Indices, typename... Types>
 	struct TTupleBase<TIntegerSequence<uint32, Indices...>, Types...> : TTupleBaseElement<Types, Indices, sizeof...(Types)>...
 	{
-		template <
-			typename... ArgTypes,
-			decltype(ConceptCheckingHelper(TTupleBaseElement<Types, Indices, sizeof...(Types)>(ForwardingConstructor, DeclVal<ArgTypes&&>())...))* = nullptr
-		>
+		template <typename... ArgTypes>
 		explicit TTupleBase(EForwardingConstructor, ArgTypes&&... Args)
 			: TTupleBaseElement<Types, Indices, sizeof...(Types)>(ForwardingConstructor, Forward<ArgTypes>(Args))...
 		{
 		}
 
-		template <
-			typename TupleType,
-			decltype(ConceptCheckingHelper(TTupleBaseElement<Types, Indices, sizeof...(Types)>(ForwardingConstructor, DeclVal<TupleType&&>().template Get<Indices>())...))* = nullptr
-		>
+		template <typename TupleType>
 		explicit TTupleBase(EOtherTupleConstructor, TupleType&& Other)
 			: TTupleBaseElement<Types, Indices, sizeof...(Types)>(ForwardingConstructor, Forward<TupleType>(Other).template Get<Indices>())...
 		{
@@ -386,22 +369,14 @@ namespace UE::Core::Private::Tuple
 
 		using DummyPairIdentifier = void;
 
-		template <
-			typename KeyArgType,
-			typename ValueArgType
-			UE_REQUIRES(std::is_constructible_v<KeyType, KeyArgType&&> && std::is_constructible_v<ValueType, ValueArgType&&>)
-		>
+		template <typename KeyArgType, typename ValueArgType>
 		explicit TTupleBase(EForwardingConstructor, KeyArgType&& KeyArg, ValueArgType&& ValueArg)
 			: Key  (Forward<KeyArgType  >(KeyArg  ))
 			, Value(Forward<ValueArgType>(ValueArg))
 		{
 		}
 
-		template <
-			typename TupleType,
-			typename std::decay_t<TupleType>::DummyPairIdentifier* = nullptr
-			UE_REQUIRES(std::is_constructible_v<KeyType, decltype(DeclVal<TupleType&&>().Get<0>())> && std::is_constructible_v<ValueType, decltype(DeclVal<TupleType&&>().Get<1>())>)
-		>
+		template <typename TupleType>
 		explicit TTupleBase(EOtherTupleConstructor, TupleType&& Other)
 			: Key  (Forward<TupleType>(Other).Get<0>())
 			, Value(Forward<TupleType>(Other).Get<1>())
@@ -496,12 +471,7 @@ namespace UE::Core::Private::Tuple
 	};
 #endif
 
-	template <
-		typename LhsType,
-		typename RhsType,
-		uint32... Indices,
-		decltype(ConceptCheckingHelper((DeclVal<LhsType&>().template Get<Indices>() = DeclVal<RhsType&&>().template Get<Indices>(), 0)...))* = nullptr
-	>
+	template <typename LhsType, typename RhsType, uint32... Indices>
 	static void Assign(LhsType& Lhs, RhsType&& Rhs, TIntegerSequence<uint32, Indices...>)
 	{
 		// This should be implemented with a fold expression when our compilers support it
@@ -648,12 +618,6 @@ namespace UE::Core::Private::Tuple
 			return Hash;
 		}
 	};
-
-	template <typename... Given, typename... Deduced>
-	std::enable_if_t<std::conjunction_v<std::is_constructible<Given, Deduced&&>...>> ConstructibleConceptCheck(Deduced&&...);
-
-	template <typename... Given, typename... Deduced>
-	decltype(ConceptCheckingHelper((DeclVal<Given>() = DeclVal<Deduced&&>(), 0)...)) AssignableConceptCheck(Deduced&&...);
 }
 
 template <typename... Types>
@@ -663,53 +627,43 @@ private:
 	typedef UE::Core::Private::Tuple::TTupleBase<TMakeIntegerSequence<uint32, sizeof...(Types)>, Types...> Super;
 
 public:
-	template <
-		typename... ArgTypes,
-		decltype(UE::Core::Private::Tuple::ConstructibleConceptCheck<Types...>(DeclVal<ArgTypes&&>()...))* = nullptr
-	>
+	template <UE::CNotCVRefTo<TTuple>... ArgTypes>
+		requires (std::is_constructible_v<Types, ArgTypes&&> && ...)
 	explicit(!std::conjunction_v<std::is_convertible<ArgTypes&&, Types>...>) TTuple(ArgTypes&&... Args)
 		: Super(UE::Core::Private::Tuple::ForwardingConstructor, Forward<ArgTypes>(Args)...)
 	{
 	}
 
-	template <
-		typename... OtherTypes,
-		decltype(UE::Core::Private::Tuple::ConstructibleConceptCheck<Types...>(DeclVal<OtherTypes&&>()...))* = nullptr
-	>
+	template <typename... OtherTypes>
+		requires (std::is_constructible_v<Types, OtherTypes&&> && ...)
 	TTuple(TTuple<OtherTypes...>&& Other)
 		: Super(UE::Core::Private::Tuple::OtherTupleConstructor, MoveTemp(Other))
 	{
 	}
 
-	template <
-		typename... OtherTypes,
-		decltype(UE::Core::Private::Tuple::ConstructibleConceptCheck<Types...>(DeclVal<const OtherTypes&>()...))* = nullptr
-	>
+	template <typename... OtherTypes>
+		requires (std::is_constructible_v<Types, const OtherTypes&> && ...)
 	TTuple(const TTuple<OtherTypes...>& Other)
 		: Super(UE::Core::Private::Tuple::OtherTupleConstructor, Other)
 	{
 	}
 
-	TTuple() = default;
+	TTuple() requires (std::is_default_constructible_v<Types> && ...) = default;
 	TTuple(TTuple&&) = default;
 	TTuple(const TTuple&) = default;
 	TTuple& operator=(TTuple&&) = default;
 	TTuple& operator=(const TTuple&) = default;
 
-	template <
-		typename... OtherTypes,
-		decltype(UE::Core::Private::Tuple::AssignableConceptCheck<Types&...>(DeclVal<const OtherTypes&>()...))* = nullptr
-	>
+	template <typename... OtherTypes>
+		requires (std::is_assignable_v<Types, const OtherTypes&> && ...)
 	TTuple& operator=(const TTuple<OtherTypes...>& Other)
 	{
 		UE::Core::Private::Tuple::Assign(*this, Other, TMakeIntegerSequence<uint32, sizeof...(Types)>{});
 		return *this;
 	}
 
-	template <
-		typename... OtherTypes,
-		decltype(UE::Core::Private::Tuple::AssignableConceptCheck<Types&...>(DeclVal<OtherTypes&&>()...))* = nullptr
-	>
+	template <typename... OtherTypes>
+		requires (std::is_assignable_v<Types, OtherTypes&&> && ...)
 	TTuple& operator=(TTuple<OtherTypes...>&& Other)
 	{
 		UE::Core::Private::Tuple::Assign(*this, MoveTemp(Other), TMakeIntegerSequence<uint32, sizeof...(OtherTypes)>{});
@@ -954,11 +908,8 @@ public:
 	using type = typename TTupleElement<N, TTuple<ArgTypes...>>::Type;
 };
 
-template <
-	int N,
-	typename TupleType
-	UE_REQUIRES(TIsTuple_V<std::decay_t<TupleType>>)
->
+template <int N, typename TupleType>
+	requires (TIsTuple_V<std::decay_t<TupleType>>)
 decltype(auto) get(TupleType&& val)
 {
 	return ((TupleType&&)val).template Get<N>();
