Disabled an MSVC tuple workaround for static analysis to see if it's still necessary.

diff --git a/Engine/Source/Runtime/Core/Public/Templates/Tuple.h b/Engine/Source/Runtime/Core/Public/Templates/Tuple.h
--- a/Engine/Source/Runtime/Core/Public/Templates/Tuple.h
+++ b/Engine/Source/Runtime/Core/Public/Templates/Tuple.h
@@ -33,7 +33,7 @@
 // - Key and Value are members of a base class.
 // - Dereferencing is done as part of a compound boolean expression (removing '&& Pair.Value != 15' removes the warning)
 #if (defined(_MSC_VER) && !defined(__clang__)) && USING_CODE_ANALYSIS
-	#define UE_TUPLE_STATIC_ANALYSIS_WORKAROUND 1
+	#define UE_TUPLE_STATIC_ANALYSIS_WORKAROUND 0 // disable this to see if it's still needed
 #else
 	#define UE_TUPLE_STATIC_ANALYSIS_WORKAROUND 0
 #endif
