Fix issue where a showing a messagebox would freeze the application if the message was (accidentally) extremely long.

diff --git a/Engine/Source/Runtime/Core/Private/Windows/WindowsPlatformMisc.cpp b/Engine/Source/Runtime/Core/Private/Windows/WindowsPlatformMisc.cpp
--- a/Engine/Source/Runtime/Core/Private/Windows/WindowsPlatformMisc.cpp
+++ b/Engine/Source/Runtime/Core/Private/Windows/WindowsPlatformMisc.cpp
@@ -1915,7 +1915,15 @@ public:
 	static EAppReturnType::Type Show(HWND ParentWindowHandle, EAppMsgType::Type MsgType, const FString& Text, const FString& Caption)
 	{
 		FWindowsDialog Instance {};
-		Instance.Text = Text;
+		int MaxLength = 1 << 20;
+		if (Text.Len() > MaxLength)
+		{
+			Instance.Text = Text.Left(MaxLength) + TEXT("\r\n[Message truncated]");
+		}
+		else
+		{
+			Instance.Text = Text;
+		}
 		Instance.Caption = Caption;
 
 		switch (MsgType)
