[DirectoryTree]

- Fix bug where bOutExisted was not being properly set when FindOrAdd is called.

#rnx

diff --git a/Engine/Source/Runtime/Core/Public/Containers/DirectoryTree.h b/Engine/Source/Runtime/Core/Public/Containers/DirectoryTree.h
--- a/Engine/Source/Runtime/Core/Public/Containers/DirectoryTree.h
+++ b/Engine/Source/Runtime/Core/Public/Containers/DirectoryTree.h
@@ -475,6 +475,12 @@ ValueType& TDirectoryTree<ValueType>::FindOrAdd(FStringView Path, bool* bOutExis
 	{
 		++NumPaths;
 	}
+
+	if (bOutExisted)
+	{
+		*bOutExisted = bExisted;
+	}
+
 	return Result;
 }
 
