Added UE::CStringViewable and UE::CCompatibleStringViewable for describing types which can be used to construct string views.

diff --git a/Engine/Source/Runtime/Core/Public/Concepts/CompatibleStringViewable.h b/Engine/Source/Runtime/Core/Public/Concepts/CompatibleStringViewable.h
new file mode 100644
--- /dev/null
+++ b/Engine/Source/Runtime/Core/Public/Concepts/CompatibleStringViewable.h
@@ -0,0 +1,25 @@
+// Copyright Epic Games, Inc. All Rights Reserved.
+
+#pragma once
+
+#include "Concepts/CharType.h"
+#include "Concepts/CompatibleCharType.h"
+#include "Concepts/ContiguousRange.h"
+#include "Traits/ElementType.h"
+#include <type_traits>
+
+namespace UE
+{
+	// Concept which describes a type which can be bound to a TStringView of a given char encoding type.
+	template <typename CharRangeType, typename CharEncodingType>
+	concept CCompatibleStringViewable =
+		(
+			std::is_pointer_v<std::decay_t<CharRangeType>> &&
+			UE::CCharType<std::remove_pointer_t<std::decay_t<CharRangeType>>> &&
+			UE::CCompatibleCharType<std::remove_pointer_t<std::decay_t<CharRangeType>>, CharEncodingType>
+		) || (
+			UE::CContiguousRange<CharRangeType> &&
+			UE::CCharType<TElementType_T<CharRangeType>> &&
+			UE::CCompatibleCharType<TElementType_T<CharRangeType>, CharEncodingType>
+		);
+}
diff --git a/Engine/Source/Runtime/Core/Public/Concepts/ContiguousRange.h b/Engine/Source/Runtime/Core/Public/Concepts/ContiguousRange.h
--- a/Engine/Source/Runtime/Core/Public/Concepts/ContiguousRange.h
+++ b/Engine/Source/Runtime/Core/Public/Concepts/ContiguousRange.h
@@ -2,9 +2,8 @@
 
 #pragma once
 
-#include "Concepts/Integral.h"
-#include "Concepts/Pointer.h"
-#include "Templates/UnrealTemplate.h"
+#include "Traits/IsContiguousContainer.h"
+#include <type_traits>
 
 namespace UE
 {
@@ -13,9 +12,6 @@ namespace UE
 	 */
 	template <typename T>
 	concept CContiguousRange =
-		requires (const T& Range)
-		{
-			{ GetData(Range) } -> CPointer;
-			{ GetNum(Range) } -> CIntegral;
-		};
+		!std::is_reference_v<T> && // TIsContiguousContainer wrongly identifies references to containers as containers
+		TIsContiguousContainer_V<T>;
 }
diff --git a/Engine/Source/Runtime/Core/Public/Concepts/StringViewable.h b/Engine/Source/Runtime/Core/Public/Concepts/StringViewable.h
new file mode 100644
--- /dev/null
+++ b/Engine/Source/Runtime/Core/Public/Concepts/StringViewable.h
@@ -0,0 +1,22 @@
+// Copyright Epic Games, Inc. All Rights Reserved.
+
+#pragma once
+
+#include "Concepts/CharType.h"
+#include "Concepts/ContiguousRange.h"
+#include "Traits/ElementType.h"
+#include <type_traits>
+
+namespace UE
+{
+	// Concept which describes a type which can be bound to a TStringView.
+	template <typename CharRangeType>
+	concept CStringViewable =
+		(
+			std::is_pointer_v<std::decay_t<CharRangeType>> &&
+			UE::CCharType<std::remove_pointer_t<std::decay_t<CharRangeType>>>
+		) || (
+			UE::CContiguousRange<CharRangeType> &&
+			UE::CCharType<TElementType_T<CharRangeType>>
+		);
+}
diff --git a/Engine/Source/Runtime/Core/Public/Containers/StringView.h b/Engine/Source/Runtime/Core/Public/Containers/StringView.h
--- a/Engine/Source/Runtime/Core/Public/Containers/StringView.h
+++ b/Engine/Source/Runtime/Core/Public/Containers/StringView.h
@@ -5,6 +5,9 @@
 #include "CoreTypes.h"
 #include "Concepts/CharType.h"
 #include "Concepts/CompatibleCharType.h"
+#include "Concepts/CompatibleStringViewable.h"
+#include "Concepts/NotCVRefTo.h"
+#include "Concepts/StringViewable.h"
 #include "Containers/StringFwd.h" // IWYU pragma: export
 #include "HAL/UnrealMemory.h"
 #include "Math/NumericLimits.h"
@@ -53,14 +56,6 @@ namespace UE::Core::Private
 	}
 } // UE::Core::Private
 
-template <typename CharPtrOrRangeType>
-	requires (!std::is_pointer_v<CharPtrOrRangeType> || TIsCharType_V<std::remove_pointer_t<CharPtrOrRangeType>>)
-[[nodiscard]] constexpr inline auto MakeStringView(const CharPtrOrRangeType& CharPtrOrRange UE_LIFETIMEBOUND) -> decltype(TStringView(CharPtrOrRange));
-
-template <UE::CCharType CharType>
-[[nodiscard]] constexpr inline auto MakeStringView(const CharType* CharPtr UE_LIFETIMEBOUND, int32 Size) -> decltype(TStringView(CharPtr, Size));
-
-
 /**
  * A string view is a non-owning view of a range of characters.
  *
@@ -149,18 +144,9 @@ public:
 	}
 
 	/** Construct a view from a contiguous range of characters, such as FString or TStringBuilder. */
-	template <
-		typename CharRangeType
-		UE_REQUIRES(
-			TAnd<
-				TIsContiguousContainer<CharRangeType>,
-				TIsCharType<TElementType_T<CharRangeType>>,
-				TIsCharEncodingCompatibleWith<TElementType_T<CharRangeType>, CharType>
-			>::Value &&
-			!std::is_array_v<std::remove_reference_t<CharRangeType>> &&
-			!std::is_same_v<CharRangeType, ViewType>
-		)
-	>
+	template <UE::CNotCVRefTo<CharType> CharRangeType>
+		// !std::is_pointer is used here to force the pointer constructor to be used when passed an array or pointer.
+		requires (!std::is_pointer_v<std::decay_t<CharRangeType>> && UE::CCompatibleStringViewable<CharRangeType, CharType>)
 	[[nodiscard]] constexpr inline TStringView(const CharRangeType& InRange UE_LIFETIMEBOUND)
 		: DataPtr((const CharType*)UE::Core::Private::StringViewGetData(InRange))
 		, Size(IntCastChecked<int32>(GetNum(InRange)))
@@ -310,18 +296,27 @@ public:
 	// Comparison
 
 	/**
-	 * Check whether this view is equivalent to a character range.
+	 * Check whether this view is equivalent to another view.
 	 *
-	 * @param Other        A character range that is comparable with the character type of this view.
-	 * @param SearchCase Whether the comparison should ignore case.
+	 * @param OtherView    A string view that is comparable with the character type of this view.
+	 * @param SearchCase   Whether the comparison should ignore case.
+	 * @return 0 is equal, negative if this view is less, positive if this view is greater.
 	 */
-	template <typename OtherRangeType, decltype(MakeStringView(DeclVal<OtherRangeType>()))* = nullptr>
-	[[nodiscard]] inline bool Equals(OtherRangeType&& Other, ESearchCase::Type SearchCase = ESearchCase::CaseSensitive) const
+	template <UE::CCharType OtherCharType>
+	[[nodiscard]] bool Equals(TStringView<OtherCharType> OtherView, ESearchCase::Type SearchCase = ESearchCase::CaseSensitive) const
 	{
-		const auto OtherView = MakeStringView(Forward<OtherRangeType>(Other));
 		return Len() == OtherView.Len() && Compare(OtherView, SearchCase) == 0;
 	}
 
+	/**
+	 * Check whether this view is equivalent to a character range.
+	 *
+	 * @param Other        A character range that is comparable with the character type of this view.
+	 * @param SearchCase Whether the comparison should ignore case.
+	 */
+	template <UE::CStringViewable OtherRangeType>
+	[[nodiscard]] bool Equals(const OtherRangeType& Other, ESearchCase::Type SearchCase = ESearchCase::CaseSensitive) const;
+
 	/**
 	 * Check whether this view is equivalent to a string view.
 	 *
@@ -341,11 +336,8 @@ public:
 	 * @param SearchCase Whether the comparison should ignore case.
 	 * @return 0 is equal, negative if this view is less, positive if this view is greater.
 	 */
-	template <typename OtherRangeType, decltype(MakeStringView(DeclVal<OtherRangeType>()))* = nullptr>
-	[[nodiscard]] inline int32 Compare(OtherRangeType&& Other, ESearchCase::Type SearchCase = ESearchCase::CaseSensitive) const
-	{
-		return Compare(MakeStringView(Forward<OtherRangeType>(Other)), SearchCase);
-	}
+	template <UE::CStringViewable OtherRangeType>
+	[[nodiscard]] int32 Compare(const OtherRangeType& Other, ESearchCase::Type SearchCase = ESearchCase::CaseSensitive) const;
 
 	/**
 	 * Compare this view with a string view.
@@ -790,3 +782,17 @@ template <UE::CCharType CharType>
 {
 	return Lhs.Compare(Rhs, ESearchCase::IgnoreCase) < 0;
 }
+
+template <UE::CCharType CharType>
+template <UE::CStringViewable OtherRangeType>
+[[nodiscard]] inline bool TStringView<CharType>::Equals(const OtherRangeType& Other, ESearchCase::Type SearchCase) const
+{
+	return Equals(MakeStringView(Other), SearchCase);
+}
+
+template <UE::CCharType CharType>
+template <UE::CStringViewable OtherRangeType>
+[[nodiscard]] inline int32 TStringView<CharType>::Compare(const OtherRangeType& Other, ESearchCase::Type SearchCase) const
+{
+	return Compare(MakeStringView(Other), SearchCase);
+}
