[Backout] - CL43725520
#rnx
Original CL Desc
-----------------------------------------------------------------
Outline more methods in TSparseArray and TSet

diff --git a/Engine/Source/Runtime/Core/Public/Containers/Set.h b/Engine/Source/Runtime/Core/Public/Containers/Set.h
--- a/Engine/Source/Runtime/Core/Public/Containers/Set.h
+++ b/Engine/Source/Runtime/Core/Public/Containers/Set.h
@@ -129,39 +129,6 @@ public:
 namespace UE::Core::Private
 {
 	[[noreturn]] CORE_API void OnInvalidSetNum(unsigned long long NewNum);
-
-	template<typename HashType>
-	void CopyHash(HashType& Hash, int32& HashSize, const HashType& Copy, int32 HashSizeCopy)
-	{
-		DestructItems((FSetElementId*)Hash.GetAllocation(), HashSize);
-		Hash.ResizeAllocation(0, HashSizeCopy, sizeof(FSetElementId));
-		ConstructItems<FSetElementId>(Hash.GetAllocation(), (FSetElementId*)Copy.GetAllocation(), HashSizeCopy);
-		HashSize = HashSizeCopy;
-	}
-
-	template<typename HashType>
-	[[nodiscard]] FORCEINLINE FSetElementId& GetTypedHash(HashType& Hash, int32 HashIndex, int32 HashSize)
-	{
-		return ((FSetElementId*)Hash.GetAllocation())[HashIndex & (HashSize - 1)];
-	}
-
-	template<typename HashType>
-	void Rehash(HashType& Hash, int32 HashSize)
-	{
-		// Free the old hash.
-		Hash.ResizeAllocation(0, 0, sizeof(FSetElementId));
-
-		if (HashSize)
-		{
-			// Allocate the new hash.
-			checkSlow(FMath::IsPowerOfTwo(HashSize));
-			Hash.ResizeAllocation(0, HashSize, sizeof(FSetElementId));
-			for (int32 HashIndex = 0; HashIndex < HashSize; ++HashIndex)
-			{
-				GetTypedHash(Hash, HashIndex, HashSize) = FSetElementId();
-			}
-		}
-	}
 }
 
 /**
@@ -257,7 +224,12 @@ public:
 	{
 		if (this != &Copy)
 		{
-			UE::Core::Private::CopyHash(Hash, HashSize, Copy.Hash, Copy.HashSize);
+			int32 CopyHashSize = Copy.HashSize;
+
+			DestructItems((FSetElementId*)Hash.GetAllocation(), HashSize);
+			Hash.ResizeAllocation(0, CopyHashSize, sizeof(FSetElementId));
+			ConstructItems<FSetElementId>(Hash.GetAllocation(), (FSetElementId*)Copy.Hash.GetAllocation(), CopyHashSize);
+			HashSize = CopyHashSize;
 
 			Elements = Copy.Elements;
 		}
@@ -1463,10 +1435,20 @@ private:
 	/** Resizes the hash. */
 	void Rehash()
 	{
-		UE::Core::Private::Rehash(Hash, HashSize);
+		// Free the old hash.
+		Hash.ResizeAllocation(0,0,sizeof(FSetElementId));
 
-		if (HashSize)
+		int32 LocalHashSize = HashSize;
+		if (LocalHashSize)
 		{
+			// Allocate the new hash.
+			checkSlow(FMath::IsPowerOfTwo(HashSize));
+			Hash.ResizeAllocation(0, LocalHashSize, sizeof(FSetElementId));
+			for (int32 HashIndex = 0; HashIndex < LocalHashSize; ++HashIndex)
+			{
+				GetTypedHash(HashIndex) = FSetElementId();
+			}
+
 			// Add the existing elements to the new hash.
 			for(typename ElementArrayType::TConstIterator ElementIt(Elements);ElementIt;++ElementIt)
 			{
