Improved typing on IFunction_OwnedObject::CloneToEmptyStorage.

diff --git a/Engine/Source/Runtime/Core/Public/Templates/Function.h b/Engine/Source/Runtime/Core/Public/Templates/Function.h
--- a/Engine/Source/Runtime/Core/Public/Templates/Function.h
+++ b/Engine/Source/Runtime/Core/Public/Templates/Function.h
@@ -43,6 +43,8 @@ namespace UE::Core::Private::Function
 	template <typename T, bool bUnique, bool bOnHeap>
 	struct TFunction_OwnedObject;
 
+	struct FFunctionStorage;
+
 	template <bool bUnique>
 	struct TFunctionStorage;
 
@@ -54,7 +56,7 @@ namespace UE::Core::Private::Function
 		/**
 		 * Creates a copy of itself into the storage and returns a pointer to the new object within it.
 		 */
-		virtual void* CloneToEmptyStorage(void* Storage) const = 0;
+		virtual void* CloneToEmptyStorage(FFunctionStorage* Storage) const = 0;
 
 		/**
 		 * Returns the address of the object.
@@ -86,7 +88,7 @@ namespace UE::Core::Private::Function
 			return &Obj;
 		}
 
-		void* CloneToEmptyStorage(void* UntypedStorage) const override;
+		void* CloneToEmptyStorage(FFunctionStorage* UntypedStorage) const override;
 
 		virtual void Destroy() override
 		{
@@ -251,7 +253,7 @@ namespace UE::Core::Private::Function
 	};
 
 	template <typename T, bool bUnique, bool bOnHeap>
-	void* TFunction_OwnedObject<T, bUnique, bOnHeap>::CloneToEmptyStorage(void* UntypedStorage) const
+	void* TFunction_OwnedObject<T, bUnique, bOnHeap>::CloneToEmptyStorage(FFunctionStorage* UntypedStorage) const
 	{
 		if constexpr (bUnique)
 		{
@@ -261,7 +263,7 @@ namespace UE::Core::Private::Function
 		}
 		else
 		{
-			TFunctionStorage<false>& Storage = *(TFunctionStorage<false>*)UntypedStorage;
+			TFunctionStorage<false>& Storage = *static_cast<TFunctionStorage<false>*>(UntypedStorage);
 
 			void* NewAlloc;
 #if TFUNCTION_USES_INLINE_STORAGE
