Defaulted the SCOPED_NAMED_EVENT API to use TRACE_CPUPROFILER_EVENT_SCOPE API when ENABLE_NAMED_EVENTS is 0. Also removed old (UE 4.19) deprecation warnings.

diff --git a/Engine/Source/Runtime/Core/Public/HAL/PlatformMisc.h b/Engine/Source/Runtime/Core/Public/HAL/PlatformMisc.h
--- a/Engine/Source/Runtime/Core/Public/HAL/PlatformMisc.h
+++ b/Engine/Source/Runtime/Core/Public/HAL/PlatformMisc.h
@@ -3,9 +3,8 @@
 
 #include "CoreTypes.h"
 #include "GenericPlatform/GenericPlatformMisc.h"
-
 #include COMPILED_PLATFORM_HEADER(PlatformMisc.h)
-
+#include "ProfilingDebugging/CpuProfilerTrace.h"
 
 #ifndef UE_DEBUG_BREAK
 #error UE_DEBUG_BREAK is not defined for this platform
@@ -44,8 +43,6 @@
 
 #if ENABLE_NAMED_EVENTS
 
-#include "ProfilingDebugging/CpuProfilerTrace.h"
-
 class FScopedNamedEventConditional
 {
 public:
@@ -206,40 +203,12 @@ private:
 
 #else
 
-class FScopedNamedEvent
-{
-public:
-	UE_DEPRECATED(4.19, "FScopedNamedEvent is compiled out in shipping builds, use SCOPED_NAMED_EVENT or variant instead to compile correctly for all targets.")
-	FScopedNamedEvent(const struct FColor& Color, const TCHAR* Text)
-	{
-	}
-
-	UE_DEPRECATED(4.19, "FScopedNamedEvent is compiled out in shipping builds, use SCOPED_NAMED_EVENT or variant instead to compile correctly for all targets.")
-	FScopedNamedEvent(const struct FColor& Color, const ANSICHAR* Text)
-	{
-	}
-};
-
-class FScopedNamedEventStatic
-{
-public:
-	UE_DEPRECATED(4.19, "FScopedNamedEventStatic is compiled out in shipping builds, use SCOPED_NAMED_EVENT or variant instead to compile correctly for all targets.")
-	FScopedNamedEventStatic(const struct FColor& Color, const TCHAR* Text)
-	{
-	}
-
-	UE_DEPRECATED(4.19, "FScopedNamedEventStatic is compiled out in shipping builds, use SCOPED_NAMED_EVENT or variant instead to compile correctly for all targets.")
-	FScopedNamedEventStatic(const struct FColor& Color, const ANSICHAR* Text)
-	{
-	}
-};
-
-#define SCOPED_NAMED_EVENT(...)
-#define SCOPED_NAMED_EVENT_FSTRING(...)
-#define SCOPED_NAMED_EVENT_TCHAR(...)
-#define SCOPED_NAMED_EVENT_TEXT(...)
-#define SCOPED_NAMED_EVENT_F(...)
-#define SCOPED_NAMED_EVENT_TCHAR_CONDITIONAL(...)
+#define SCOPED_NAMED_EVENT(Name, Color)									TRACE_CPUPROFILER_EVENT_SCOPE(Name);
+#define SCOPED_NAMED_EVENT_FSTRING(Text, Color)							TRACE_CPUPROFILER_EVENT_SCOPE_TEXT(*Text);
+#define SCOPED_NAMED_EVENT_TCHAR(Text, Color)							TRACE_CPUPROFILER_EVENT_SCOPE_TEXT(Text);
+#define SCOPED_NAMED_EVENT_TEXT(Text, Color)							TRACE_CPUPROFILER_EVENT_SCOPE_STR(Text);
+#define SCOPED_NAMED_EVENT_F(Format, Color, ...)						TRACE_CPUPROFILER_EVENT_SCOPE_TEXT(*FString::Printf(Format, __VA_ARGS__));
+#define SCOPED_NAMED_EVENT_TCHAR_CONDITIONAL(Text, Color, bCondition)	TRACE_CPUPROFILER_EVENT_SCOPE_TEXT_CONDITIONAL(Text, (bCondition));
 
 #endif
 
