Made mutating EnumClassFlags functions usable in constexpr contexts.

diff --git a/Engine/Source/Runtime/Core/Public/Misc/EnumClassFlags.h b/Engine/Source/Runtime/Core/Public/Misc/EnumClassFlags.h
--- a/Engine/Source/Runtime/Core/Public/Misc/EnumClassFlags.h
+++ b/Engine/Source/Runtime/Core/Public/Misc/EnumClassFlags.h
@@ -4,9 +4,9 @@
 
 // Defines all bitwise operators for enum classes so it can be (mostly) used as a regular flags enum
 #define ENUM_CLASS_FLAGS(Enum) \
-	inline           Enum& operator|=(Enum& Lhs, Enum Rhs) { return Lhs = (Enum)((__underlying_type(Enum))Lhs | (__underlying_type(Enum))Rhs); } \
-	inline           Enum& operator&=(Enum& Lhs, Enum Rhs) { return Lhs = (Enum)((__underlying_type(Enum))Lhs & (__underlying_type(Enum))Rhs); } \
-	inline           Enum& operator^=(Enum& Lhs, Enum Rhs) { return Lhs = (Enum)((__underlying_type(Enum))Lhs ^ (__underlying_type(Enum))Rhs); } \
+	inline constexpr Enum& operator|=(Enum& Lhs, Enum Rhs) { return Lhs = (Enum)((__underlying_type(Enum))Lhs | (__underlying_type(Enum))Rhs); } \
+	inline constexpr Enum& operator&=(Enum& Lhs, Enum Rhs) { return Lhs = (Enum)((__underlying_type(Enum))Lhs & (__underlying_type(Enum))Rhs); } \
+	inline constexpr Enum& operator^=(Enum& Lhs, Enum Rhs) { return Lhs = (Enum)((__underlying_type(Enum))Lhs ^ (__underlying_type(Enum))Rhs); } \
 	inline constexpr Enum  operator| (Enum  Lhs, Enum Rhs) { return (Enum)((__underlying_type(Enum))Lhs | (__underlying_type(Enum))Rhs); } \
 	inline constexpr Enum  operator& (Enum  Lhs, Enum Rhs) { return (Enum)((__underlying_type(Enum))Lhs & (__underlying_type(Enum))Rhs); } \
 	inline constexpr Enum  operator^ (Enum  Lhs, Enum Rhs) { return (Enum)((__underlying_type(Enum))Lhs ^ (__underlying_type(Enum))Rhs); } \
@@ -15,9 +15,9 @@
 
 // Friends all bitwise operators for enum classes so the definition can be kept private / protected.
 #define FRIEND_ENUM_CLASS_FLAGS(Enum) \
-	friend           Enum& operator|=(Enum& Lhs, Enum Rhs); \
-	friend           Enum& operator&=(Enum& Lhs, Enum Rhs); \
-	friend           Enum& operator^=(Enum& Lhs, Enum Rhs); \
+	friend constexpr Enum& operator|=(Enum& Lhs, Enum Rhs); \
+	friend constexpr Enum& operator&=(Enum& Lhs, Enum Rhs); \
+	friend constexpr Enum& operator^=(Enum& Lhs, Enum Rhs); \
 	friend constexpr Enum  operator| (Enum  Lhs, Enum Rhs); \
 	friend constexpr Enum  operator& (Enum  Lhs, Enum Rhs); \
 	friend constexpr Enum  operator^ (Enum  Lhs, Enum Rhs); \
@@ -39,14 +39,14 @@ constexpr bool EnumHasAnyFlags(Enum Flags, Enum Contains)
 }
 
 template<typename Enum>
-void EnumAddFlags(Enum& Flags, Enum FlagsToAdd)
+constexpr void EnumAddFlags(Enum& Flags, Enum FlagsToAdd)
 {
 	using UnderlyingType = __underlying_type(Enum);
 	Flags = (Enum)((UnderlyingType)Flags | (UnderlyingType)FlagsToAdd);
 }
 
 template<typename Enum>
-void EnumRemoveFlags(Enum& Flags, Enum FlagsToRemove)
+constexpr void EnumRemoveFlags(Enum& Flags, Enum FlagsToRemove)
 {
 	using UnderlyingType = __underlying_type(Enum);
 	Flags = (Enum)((UnderlyingType)Flags & ~(UnderlyingType)FlagsToRemove);
