Logging: Eliminate static log records in the absence of source location and log tracing

Saves 22 bytes in program size for affected log statements and about 46 bytes in debug size, as measured with MSVC on Win64.

diff --git a/Engine/Source/Runtime/Core/Private/AutoRTFM/AutoRTFMUE.cpp b/Engine/Source/Runtime/Core/Private/AutoRTFM/AutoRTFMUE.cpp
--- a/Engine/Source/Runtime/Core/Private/AutoRTFM/AutoRTFMUE.cpp
+++ b/Engine/Source/Runtime/Core/Private/AutoRTFM/AutoRTFMUE.cpp
@@ -233,23 +233,6 @@ bool IsSeverityActive(autortfm_log_severity Severity)
 	}
 }
 
-ELogVerbosity::Type SeverityToVerbosity(autortfm_log_severity Severity)
-{
-	switch (Severity)
-	{
-		case autortfm_log_verbose:
-			return ELogVerbosity::Verbose;
-		case autortfm_log_info:
-			return ELogVerbosity::Display;
-		case autortfm_log_warn:
-			return ELogVerbosity::Warning;
-		case autortfm_log_error:
-			return ELogVerbosity::Error;
-		case autortfm_log_fatal:
-			return ELogVerbosity::Fatal;
-	}
-}
-
 UE_AUTORTFM_ALWAYS_OPEN_NO_MEMORY_VALIDATION
 TStringConversion<TStringConvert<UTF8CHAR, TCHAR>, 128> FmtToTChar(const char* Format, va_list Args)
 {
@@ -288,21 +271,26 @@ void AutoRTFMLog(const char* File, int Line, void* ProgramCounter, autortfm_log_
 		/* Format */ TEXT("%s"),
 		/* File */ File,
 		/* Line */ Line,
-		/* Verbosity */ SeverityToVerbosity(Severity),
 		/* DynamicData */ &DynamicData
 	);
 
 	switch (Severity)
 	{
 		case autortfm_log_verbose:
+			UE::Logging::Private::BasicLog<ELogVerbosity::Verbose>(Record, &LogAutoRTFM, FmtToTChar(Format, Args).Get());
+			break;
 		case autortfm_log_info:
+			UE::Logging::Private::BasicLog<ELogVerbosity::Display>(Record, &LogAutoRTFM, FmtToTChar(Format, Args).Get());
+			break;
 		case autortfm_log_warn:
+			UE::Logging::Private::BasicLog<ELogVerbosity::Warning>(Record, &LogAutoRTFM, FmtToTChar(Format, Args).Get());
+			break;
 		case autortfm_log_error:
-			UE::Logging::Private::BasicLog(LogAutoRTFM, &Record, FmtToTChar(Format, Args).Get());
+			UE::Logging::Private::BasicLog<ELogVerbosity::Error>(Record, &LogAutoRTFM, FmtToTChar(Format, Args).Get());
 			break;
 		case autortfm_log_fatal:
 			FDebug::DumpStackTraceToLog(TEXT("AutoRTFM backtrace"), ELogVerbosity::Error);
-			UE::Logging::Private::BasicFatalLogWithProgramCounter(LogAutoRTFM, &Record, ProgramCounter, FmtToTChar(Format, Args).Get());
+			UE::Logging::Private::BasicFatalLogWithProgramCounter(Record, &LogAutoRTFM, ProgramCounter, FmtToTChar(Format, Args).Get());
 			break;
 	}
 #endif
@@ -322,22 +310,30 @@ void AutoRTFMLogWithCallstack(void* ProgramCounter, autortfm_log_severity Severi
 		/* Format */ TEXT("%s"),
 		/* File */ "",
 		/* Line */ 0,
-		/* Verbosity */ SeverityToVerbosity(Severity),
 		/* DynamicData */ &DynamicData
 	);
 
 	switch (Severity)
 	{
 		case autortfm_log_verbose:
+			FDebug::DumpStackTraceToLog(TEXT("AutoRTFM backtrace"), ELogVerbosity::Verbose);
+			UE::Logging::Private::BasicLog<ELogVerbosity::Verbose>(Record, &LogAutoRTFM, FmtToTChar(Format, Args).Get());
+			break;
 		case autortfm_log_info:
+			FDebug::DumpStackTraceToLog(TEXT("AutoRTFM backtrace"), ELogVerbosity::Display);
+			UE::Logging::Private::BasicLog<ELogVerbosity::Display>(Record, &LogAutoRTFM, FmtToTChar(Format, Args).Get());
+			break;
 		case autortfm_log_warn:
+			FDebug::DumpStackTraceToLog(TEXT("AutoRTFM backtrace"), ELogVerbosity::Warning);
+			UE::Logging::Private::BasicLog<ELogVerbosity::Warning>(Record, &LogAutoRTFM, FmtToTChar(Format, Args).Get());
+			break;
 		case autortfm_log_error:
-			FDebug::DumpStackTraceToLog(TEXT("AutoRTFM backtrace"), Record.Verbosity);
-			UE::Logging::Private::BasicLog(LogAutoRTFM, &Record, FmtToTChar(Format, Args).Get());
+			FDebug::DumpStackTraceToLog(TEXT("AutoRTFM backtrace"), ELogVerbosity::Error);
+			UE::Logging::Private::BasicLog<ELogVerbosity::Error>(Record, &LogAutoRTFM, FmtToTChar(Format, Args).Get());
 			break;
 		case autortfm_log_fatal:
 			FDebug::DumpStackTraceToLog(TEXT("AutoRTFM backtrace"), ELogVerbosity::Error);
-			UE::Logging::Private::BasicFatalLogWithProgramCounter(LogAutoRTFM, &Record, ProgramCounter, FmtToTChar(Format, Args).Get());
+			UE::Logging::Private::BasicFatalLogWithProgramCounter(Record, &LogAutoRTFM, ProgramCounter, FmtToTChar(Format, Args).Get());
 			break;
 	}
 #endif
diff --git a/Engine/Source/Runtime/Core/Private/Logging/StructuredLog.cpp b/Engine/Source/Runtime/Core/Private/Logging/StructuredLog.cpp
--- a/Engine/Source/Runtime/Core/Private/Logging/StructuredLog.cpp
+++ b/Engine/Source/Runtime/Core/Private/Logging/StructuredLog.cpp
@@ -1209,11 +1209,11 @@ FLogTemplateFieldIterator& FLogTemplateFieldIterator::operator++()
 
 #if LOGTRACE_ENABLED
 template <typename StaticLogRecordType>
-FORCENOINLINE static void EnsureLogMessageSpec(const FLogCategoryBase& Category, const StaticLogRecordType& Log)
+FORCENOINLINE static void EnsureLogMessageSpec(const FLogCategoryBase& Category, const StaticLogRecordType& Log, ELogVerbosity::Type Verbosity)
 {
 	if (!Log.DynamicData.bInitializedTrace.load(std::memory_order_acquire))
 	{
-		FLogTrace::OutputLogMessageSpec(&Log, &Category, Log.Verbosity, Log.File, Log.Line, TEXT("%s"));
+		FLogTrace::OutputLogMessageSpec(&Log, &Category, Verbosity, Log.File, Log.Line, TEXT("%s"));
 		Log.DynamicData.bInitializedTrace.store(true, std::memory_order_release);
 	}
 }
@@ -1455,7 +1455,7 @@ inline static FLogRecord CreateLogRecord(const FLogCategoryBase& Category, const
 #if LOGTRACE_ENABLED
 	if (UE_TRACE_CHANNELEXPR_IS_ENABLED(LogChannel))
 	{
-		EnsureLogMessageSpec(Category, Log);
+		EnsureLogMessageSpec(Category, Log, Log.Verbosity);
 	}
 #endif
 
@@ -1637,31 +1637,31 @@ FORCENOINLINE static FCbObject SerializeBasicLogMessage(const FStaticBasicLogRec
 }
 
 UE_AUTORTFM_ALWAYS_OPEN
-static void BasicLogV(const FLogCategoryBase& Category, const FStaticBasicLogRecord* Log, va_list Args)
+static void BasicLogV(FStaticBasicLogRecordParam Log, const FLogCategoryBase& Category, ELogVerbosity::Type Verbosity, va_list Args)
 {
 #if !NO_LOGGING
 #if LOGTRACE_ENABLED
 	if (UE_TRACE_CHANNELEXPR_IS_ENABLED(LogChannel))
 	{
-		EnsureLogMessageSpec(Category, *Log);
+		EnsureLogMessageSpec(Category, Log, Verbosity);
 	}
 #endif
 
 	if (GConvertBasicLogToLogRecord)
 	{
-		FCbObject Fields = SerializeBasicLogMessage(*Log, Args);
+		FCbObject Fields = SerializeBasicLogMessage(Log, Args);
 
 		FLogRecord Record;
 		Record.SetFormat(GStaticBasicLogFormat);
 		Record.SetTemplate(GetStaticBasicLogTemplate());
 		Record.SetFields(MoveTemp(Fields));
-		Record.SetFile(Log->File);
-		Record.SetLine(Log->Line);
+		Record.SetFile(Log.File);
+		Record.SetLine(Log.Line);
 		Record.SetCategory(Category.GetCategoryName());
-		Record.SetVerbosity(Log->Verbosity);
+		Record.SetVerbosity(Verbosity);
 		Record.SetTime(FLogTime::Now());
 
-		DispatchStaticLogRecord(*Log, Record);
+		DispatchStaticLogRecord(Log, Record);
 	}
 	else
 	{
@@ -1670,52 +1670,61 @@ static void BasicLogV(const FLogCategoryBase& Category, const FStaticBasicLogRec
 		{
 			va_list Args2;
 			va_copy(Args2, Args);
-			BasicLogToTrace(Log, Log->Format, Args2);
+			BasicLogToTrace(&Log, Log.Format, Args2);
 			va_end(Args2);
 		}
 	#endif
-		FMsg::LogV(Log->File, Log->Line, Category.GetCategoryName(), Log->Verbosity, Log->Format, Args);
+		FMsg::LogV(Log.File, Log.Line, Category.GetCategoryName(), Verbosity, Log.Format, Args);
 	}
 #endif
 }
 
-void BasicLog(const FLogCategoryBase& Category, const FStaticBasicLogRecord* Log, ...)
+template <ELogVerbosity::Type Verbosity>
+FORCENOINLINE void BasicLog(FStaticBasicLogRecordParam Log, const FLogCategoryBase* Category, ...)
 {
 #if !NO_LOGGING
 	va_list Args;
-	va_start(Args, Log);
-	BasicLogV(Category, Log, Args);
+	va_start(Args, Category);
+	BasicLogV(Log, *Category, Verbosity, Args);
 	va_end(Args);
 #endif
 }
 
+template CORE_API void BasicLog<ELogVerbosity::NoLogging>(FStaticBasicLogRecordParam Log, const FLogCategoryBase* Category, ...);
+template CORE_API void BasicLog<ELogVerbosity::Error>(FStaticBasicLogRecordParam Log, const FLogCategoryBase* Category, ...);
+template CORE_API void BasicLog<ELogVerbosity::Warning>(FStaticBasicLogRecordParam Log, const FLogCategoryBase* Category, ...);
+template CORE_API void BasicLog<ELogVerbosity::Display>(FStaticBasicLogRecordParam Log, const FLogCategoryBase* Category, ...);
+template CORE_API void BasicLog<ELogVerbosity::Log>(FStaticBasicLogRecordParam Log, const FLogCategoryBase* Category, ...);
+template CORE_API void BasicLog<ELogVerbosity::Verbose>(FStaticBasicLogRecordParam Log, const FLogCategoryBase* Category, ...);
+template CORE_API void BasicLog<ELogVerbosity::VeryVerbose>(FStaticBasicLogRecordParam Log, const FLogCategoryBase* Category, ...);
+
 UE_AUTORTFM_ALWAYS_OPEN
-static void BasicFatalLogV(const FLogCategoryBase& Category, const FStaticBasicLogRecord* Log, void* ProgramCounter, va_list Args)
+static void BasicFatalLogV(FStaticBasicLogRecordParam Log, void* ProgramCounter, va_list Args)
 {
 #if !NO_LOGGING
-	StaticFailDebugV(TEXT("Fatal error:"), Log->File, Log->Line, ProgramCounter, Log->Format, Args);
+	StaticFailDebugV(TEXT("Fatal error:"), Log.File, Log.Line, ProgramCounter, Log.Format, Args);
 
 	UE_DEBUG_BREAK_AND_PROMPT_FOR_REMOTE();
 	FDebug::ProcessFatalError(ProgramCounter);
 #endif
 }
 
-void BasicFatalLog(const FLogCategoryBase& Category, const FStaticBasicLogRecord* Log, ...)
+FORCENOINLINE void BasicFatalLog(FStaticBasicLogRecordParam Log, const FLogCategoryBase* Category, ...)
 {
 #if !NO_LOGGING
 	va_list Args;
-	va_start(Args, Log);
-	BasicFatalLogV(Category, Log, PLATFORM_RETURN_ADDRESS(), Args);
+	va_start(Args, Category);
+	BasicFatalLogV(Log, PLATFORM_RETURN_ADDRESS(), Args);
 	va_end(Args);
 #endif
 }
 
-void BasicFatalLogWithProgramCounter(const FLogCategoryBase& Category, const FStaticBasicLogRecord* Log, void* ProgramCounter, ...)
+FORCENOINLINE void BasicFatalLogWithProgramCounter(FStaticBasicLogRecordParam Log, const FLogCategoryBase* Category, void* ProgramCounter, ...)
 {
 #if !NO_LOGGING
 	va_list Args;
 	va_start(Args, ProgramCounter);
-	BasicFatalLogV(Category, Log, ProgramCounter, Args);
+	BasicFatalLogV(Log, ProgramCounter, Args);
 	va_end(Args);
 #endif
 }
diff --git a/Engine/Source/Runtime/Core/Public/Logging/LogMacros.h b/Engine/Source/Runtime/Core/Public/Logging/LogMacros.h
--- a/Engine/Source/Runtime/Core/Public/Logging/LogMacros.h
+++ b/Engine/Source/Runtime/Core/Public/Logging/LogMacros.h
@@ -24,6 +24,12 @@
 #define UE_LOG_INCLUDE_SOURCE_LOCATION 1
 #endif
 
+#if UE_LOG_INCLUDE_SOURCE_LOCATION || LOGTRACE_ENABLED
+#define UE_PRIVATE_LOG_USE_STATIC_RECORD 1
+#else
+#define UE_PRIVATE_LOG_USE_STATIC_RECORD 0
+#endif
+
 /*----------------------------------------------------------------------------
 	Logging
 ----------------------------------------------------------------------------*/
@@ -125,7 +131,6 @@ struct FStaticBasicLogRecord
 	const TCHAR* Format = nullptr;
 	UE_IF(UE_LOG_INCLUDE_SOURCE_LOCATION,, inline static constexpr) const ANSICHAR* File = nullptr;
 	UE_IF(UE_LOG_INCLUDE_SOURCE_LOCATION,, inline static constexpr) int32 Line = 0;
-	ELogVerbosity::Type Verbosity = ELogVerbosity::Log;
 #if LOGTRACE_ENABLED
 	FStaticBasicLogDynamicData& DynamicData;
 #endif
@@ -135,24 +140,37 @@ struct FStaticBasicLogRecord
 		const TCHAR* InFormat,
 		const ANSICHAR* InFile,
 		int32 InLine,
-		ELogVerbosity::Type InVerbosity,
 		FStaticBasicLogDynamicData* InDynamicData)
 		: Format(InFormat)
 	#if UE_LOG_INCLUDE_SOURCE_LOCATION
 		, File(InFile)
 		, Line(InLine)
 	#endif
-		, Verbosity(InVerbosity)
 	#if LOGTRACE_ENABLED
 		, DynamicData(*InDynamicData)
 	#endif
 	{
 	}
+
+#if !UE_PRIVATE_LOG_USE_STATIC_RECORD
+	// Delete the address-of operator because it is unsafe to pass pointers to records
+	// when they do not have static storage duration.
+	FStaticBasicLogRecord* operator&() const = delete;
+#endif
 };
 
-CORE_API void BasicLog(const FLogCategoryBase& Category, const FStaticBasicLogRecord* Log, ...);
-CORE_API void BasicFatalLog(const FLogCategoryBase& Category, const FStaticBasicLogRecord* Log, ...);
-CORE_API void BasicFatalLogWithProgramCounter(const FLogCategoryBase& Category, const FStaticBasicLogRecord* Log, void* ProgramCounter, ...);
+// Pass FStaticBasicLogRecord by value when the record is not static, because in that
+// configuration it is simply a wrapper for the format string pointer.
+#if UE_PRIVATE_LOG_USE_STATIC_RECORD
+using FStaticBasicLogRecordParam = const FStaticBasicLogRecord&;
+#else
+using FStaticBasicLogRecordParam = FStaticBasicLogRecord;
+#endif
+
+template <ELogVerbosity::Type Verbosity>
+CORE_API void BasicLog(FStaticBasicLogRecordParam Log, const FLogCategoryBase* Category, ...);
+CORE_API void BasicFatalLog(FStaticBasicLogRecordParam Log, const FLogCategoryBase* Category, ...);
+CORE_API void BasicFatalLogWithProgramCounter(FStaticBasicLogRecordParam Log, const FLogCategoryBase* Category, void* ProgramCounter, ...);
 
 } // UE::Logging::Private
 
@@ -287,13 +305,13 @@ CORE_API void BasicFatalLogWithProgramCounter(const FLogCategoryBase& Category,
 		static_assert(TIsArrayOrRefOfTypeByPredicate<decltype(Format), TIsCharEncodingCompatibleWithTCHAR>::Value, "Formatting string must be a TCHAR array."); \
 		UE_VALIDATE_FORMAT_STRING(Format, ##__VA_ARGS__); \
 		UE_IF(LOGTRACE_ENABLED, static ::UE::Logging::Private::FStaticBasicLogDynamicData LOG_Dynamic, ); \
-		static constexpr ::UE::Logging::Private::FStaticBasicLogRecord LOG_Static(Format, __builtin_FILE(), __builtin_LINE(), ::ELogVerbosity::Verbosity, UE_IF(LOGTRACE_ENABLED, &LOG_Dynamic, nullptr)); \
+		UE_IF(UE_PRIVATE_LOG_USE_STATIC_RECORD, static,) constexpr ::UE::Logging::Private::FStaticBasicLogRecord LOG_Static(Format, __builtin_FILE(), __builtin_LINE(), UE_IF(LOGTRACE_ENABLED, &LOG_Dynamic, nullptr)); \
 		static_assert((::ELogVerbosity::Verbosity & ::ELogVerbosity::VerbosityMask) < ::ELogVerbosity::NumVerbosity && ::ELogVerbosity::Verbosity > 0, "Verbosity must be constant and in range."); \
 		if constexpr ((::ELogVerbosity::Verbosity & ELogVerbosity::VerbosityMask) == ::ELogVerbosity::Fatal) \
 		{ \
 			Condition \
 			{ \
-				::UE::Logging::Private::BasicFatalLog(Category, &LOG_Static, ##__VA_ARGS__); \
+				::UE::Logging::Private::BasicFatalLog(LOG_Static, &Category, ##__VA_ARGS__); \
 				CA_ASSUME(false); \
 			} \
 		} \
@@ -305,7 +323,7 @@ CORE_API void BasicFatalLogWithProgramCounter(const FLogCategoryBase& Category,
 				{ \
 					Condition \
 					{ \
-						::UE::Logging::Private::BasicLog(Category, &LOG_Static, ##__VA_ARGS__); \
+						::UE::Logging::Private::BasicLog<ELogVerbosity::Type(::ELogVerbosity::Verbosity & ::ELogVerbosity::VerbosityMask)>(LOG_Static, &Category, ##__VA_ARGS__); \
 					} \
 				} \
 			} \
