Add referenced memory to minidumps produced for fatal errors.

Can be turned off by setting WindowsPlatformCrashContext.MinidumpIncludeReferencedMemory to false.

#rnx

diff --git a/Engine/Source/Runtime/Core/Private/Windows/WindowsPlatformCrashContext.cpp b/Engine/Source/Runtime/Core/Private/Windows/WindowsPlatformCrashContext.cpp
--- a/Engine/Source/Runtime/Core/Private/Windows/WindowsPlatformCrashContext.cpp
+++ b/Engine/Source/Runtime/Core/Private/Windows/WindowsPlatformCrashContext.cpp
@@ -120,6 +120,13 @@ namespace {
 	 * @return Success or failure
 	 */
 
+	static TAutoConsoleVariable<bool> CVarMinidumpIncludeReferencedMemory(
+		TEXT("WindowsPlatformCrashContext.MinidumpIncludeReferencedMemory"),
+		true,
+		TEXT("Include referenced memory in minidumps for non-continuable errors (enable the MiniDumpWithIndirectlyReferencedMemory flag).\n"),
+		ECVF_Default
+	);
+
 	 // #CrashReport: 2014-10-08 Move to FWindowsPlatformCrashContext
 	bool WriteMinidump(HANDLE Process, DWORD ThreadId, FWindowsPlatformCrashContext& InContext, const TCHAR* Path, LPEXCEPTION_POINTERS ExceptionInfo)
 	{
@@ -163,6 +170,14 @@ namespace {
 		{
 			MinidumpType = (MINIDUMP_TYPE)(MiniDumpWithFullMemory | MiniDumpWithFullMemoryInfo | MiniDumpWithHandleData | MiniDumpWithThreadInfo );
 		}
+		else
+		{
+			const bool bIncludeReferencedMemory = CVarMinidumpIncludeReferencedMemory.GetValueOnAnyThread() && !FGenericCrashContext::IsTypeContinuable(InContext.GetType());
+			if (bIncludeReferencedMemory)
+			{
+				MinidumpType = (MINIDUMP_TYPE)(MiniDumpWithThreadInfo | MiniDumpWithIndirectlyReferencedMemory);
+			}
+		}
 
 		const BOOL Result = MiniDumpWriteDump(Process, GetProcessId(Process), FileHandle, MinidumpType, ExceptionInfo ? &DumpExceptionInfo : NULL, &CrashContextStreamInformation, NULL);
 		CloseHandle(FileHandle);
