[Log] Add ensure condition when log category verbosity is misconfigured.

Uses an ensure-condition when a log category sets its runtime verbosity higher than its configured compile-time verbosity.
This was initially submitted as a warning in CL 41611908, but removed to fix a deadlock in CL 43070231.

#rnx

diff --git a/Engine/Source/Runtime/Core/Private/Logging/LogCategory.cpp b/Engine/Source/Runtime/Core/Private/Logging/LogCategory.cpp
--- a/Engine/Source/Runtime/Core/Private/Logging/LogCategory.cpp
+++ b/Engine/Source/Runtime/Core/Private/Logging/LogCategory.cpp
@@ -43,9 +43,7 @@ void FLogCategoryBase::SetVerbosity(ELogVerbosity::Type InNewVerbosity)
 	const ELogVerbosity::Type OldVerbosity = Verbosity;
 	const ELogVerbosity::Type NewVerbosity = (ELogVerbosity::Type)(InNewVerbosity & ELogVerbosity::VerbosityMask);
 	// regularize the verbosity to be at most whatever we were compiled with
-	//@todo-lh - Turn back to the following ensure condition once it's safe to do so:
-	// if (ensureMsgf(NewVerbosity <= CompileTimeVerbosity, TEXT("Log category '%s' tried to set verbosity (%s) higher than its compile-time verbosity level (%s)"), *GetCategoryName().ToString(), ToString(NewVerbosity), ToString(CompileTimeVerbosity)))
-	if (NewVerbosity <= CompileTimeVerbosity)
+	if (ensureMsgf(NewVerbosity <= CompileTimeVerbosity, TEXT("Log category '%s' tried to set verbosity (%s) higher than its compile-time verbosity level (%s)"), *GetCategoryName().ToString(), ToString(NewVerbosity), ToString(CompileTimeVerbosity)))
 	{
 		Verbosity = NewVerbosity;
 	}
diff --git a/Engine/Source/Runtime/Core/Public/Logging/LogMacros.h b/Engine/Source/Runtime/Core/Public/Logging/LogMacros.h
--- a/Engine/Source/Runtime/Core/Public/Logging/LogMacros.h
+++ b/Engine/Source/Runtime/Core/Public/Logging/LogMacros.h
@@ -377,6 +377,7 @@ CORE_API void BasicFatalLogWithProgramCounter(FStaticBasicLogRecordParam Log, co
 	 * @param CompileTimeVerbosity, maximum verbosity to compile into the code
 	**/
 	#define DEFINE_LOG_CATEGORY_STATIC(CategoryName, DefaultVerbosity, CompileTimeVerbosity) \
+		static_assert(ELogVerbosity::DefaultVerbosity <= ELogVerbosity::CompileTimeVerbosity, "Default verbosity of log category cannot be higher than its compile time verbosity"); \
 		static struct FLogCategory##CategoryName : public FLogCategory<ELogVerbosity::DefaultVerbosity, ELogVerbosity::CompileTimeVerbosity> \
 		{ \
 			UE_FORCEINLINE_HINT FLogCategory##CategoryName() : FLogCategory(TEXT(#CategoryName)) {} \
