Change FCrc::StrCrc32 variants over to slicing-by-4, we already have the tables

diff --git a/Engine/Source/Runtime/Core/Public/Misc/Crc.h b/Engine/Source/Runtime/Core/Public/Misc/Crc.h
--- a/Engine/Source/Runtime/Core/Public/Misc/Crc.h
+++ b/Engine/Source/Runtime/Core/Public/Misc/Crc.h
@@ -49,34 +49,17 @@ struct FCrc
 		// want consistency between equivalent strings with different character types.
 		static_assert(sizeof(CharType) <= 4, "StrCrc32 only works with CharType up to 32 bits.");
 
-		if constexpr (sizeof(CharType) != 1)
+		CRC = ~CRC;
+		while (CharType Ch = *Data++)
 		{
-			CRC = ~CRC;
-			while (CharType Ch = *Data++)
-			{
-				CRC = (CRC >> 8) ^ CRCTablesSB8[0][(CRC ^ Ch) & 0xFF];
-				Ch >>= 8;
-				CRC = (CRC >> 8) ^ CRCTablesSB8[0][(CRC ^ Ch) & 0xFF];
-				Ch >>= 8;
-				CRC = (CRC >> 8) ^ CRCTablesSB8[0][(CRC ^ Ch) & 0xFF];
-				Ch >>= 8;
-				CRC = (CRC >> 8) ^ CRCTablesSB8[0][(CRC ^ Ch) & 0xFF];
-			}
-			return ~CRC;
-		}
-		else
-		{
-			/* Special case for when CharType is a byte, which causes warnings when right-shifting by 8 */
-			CRC = ~CRC;
-			while (CharType Ch = *Data++)
-			{
-				CRC = (CRC >> 8) ^ CRCTablesSB8[0][(CRC ^ Ch) & 0xFF];
-				CRC = (CRC >> 8) ^ CRCTablesSB8[0][(CRC     ) & 0xFF];
-				CRC = (CRC >> 8) ^ CRCTablesSB8[0][(CRC     ) & 0xFF];
-				CRC = (CRC >> 8) ^ CRCTablesSB8[0][(CRC     ) & 0xFF];
-			}
-			return ~CRC;
+			uint32 V = CRC ^ Ch;
+			CRC =
+				CRCTablesSB8[3][ V        & 0xFF] ^
+				CRCTablesSB8[2][(V >>  8) & 0xFF] ^
+				CRCTablesSB8[1][(V >> 16) & 0xFF] ^
+				CRCTablesSB8[0][ V >> 24        ];
 		}
+		return ~CRC;
 	}
 
 	/** String CRC. Prefer Hash/Fnv.h for new code, CRCs are not very good and fairly slow. */
@@ -88,36 +71,17 @@ struct FCrc
 		// want consistency between equivalent strings with different character types.
 		static_assert(sizeof(CharType) <= 4, "StrCrc32Len only works with CharType up to 32 bits.");
 
-		if constexpr (sizeof(CharType) != 1)
+		CRC = ~CRC;
+		for (int32 Idx = 0; Idx < Length; ++Idx)
 		{
-			CRC = ~CRC;
-			for (int32 Idx = 0; Idx < Length; ++Idx)
-			{
-				CharType Ch = Data[Idx];
-				CRC = (CRC >> 8) ^ CRCTablesSB8[0][(CRC ^ Ch) & 0xFF];
-				Ch >>= 8;
-				CRC = (CRC >> 8) ^ CRCTablesSB8[0][(CRC ^ Ch) & 0xFF];
-				Ch >>= 8;
-				CRC = (CRC >> 8) ^ CRCTablesSB8[0][(CRC ^ Ch) & 0xFF];
-				Ch >>= 8;
-				CRC = (CRC >> 8) ^ CRCTablesSB8[0][(CRC ^ Ch) & 0xFF];
-			}
-			return ~CRC;
-		}
-		else
-		{
-			/* Special case for when CharType is a byte, which causes warnings when right-shifting by 8 */
-			CRC = ~CRC;
-			for (int32 Idx = 0; Idx < Length; ++Idx)
-			{
-				CharType Ch = Data[Idx];
-				CRC = (CRC >> 8) ^ CRCTablesSB8[0][(CRC ^ Ch) & 0xFF];
-				CRC = (CRC >> 8) ^ CRCTablesSB8[0][(CRC     ) & 0xFF];
-				CRC = (CRC >> 8) ^ CRCTablesSB8[0][(CRC     ) & 0xFF];
-				CRC = (CRC >> 8) ^ CRCTablesSB8[0][(CRC     ) & 0xFF];
-			}
-			return ~CRC;
+			uint32 V = CRC ^ Data[Idx];
+			CRC =
+				CRCTablesSB8[3][ V        & 0xFF] ^
+				CRCTablesSB8[2][(V >>  8) & 0xFF] ^
+				CRCTablesSB8[1][(V >> 16) & 0xFF] ^
+				CRCTablesSB8[0][ V >> 24        ];
 		}
+		return ~CRC;
 	}
 
 	/**
