Upgraded some legacy metaprogramming in delegates and fallout from no longer including IsConst.h.

diff --git a/Engine/Plugins/VirtualProduction/Avalanche/Source/AvalancheTransition/Public/AvaTransitionScene.h b/Engine/Plugins/VirtualProduction/Avalanche/Source/AvalancheTransition/Public/AvaTransitionScene.h
--- a/Engine/Plugins/VirtualProduction/Avalanche/Source/AvalancheTransition/Public/AvaTransitionScene.h
+++ b/Engine/Plugins/VirtualProduction/Avalanche/Source/AvalancheTransition/Public/AvaTransitionScene.h
@@ -116,8 +116,8 @@ struct AVALANCHETRANSITION_API FAvaTransitionScene
 
 	virtual ~FAvaTransitionScene() = default;
 
-	template<typename InStructType
-		UE_REQUIRES(!TIsConst<InStructType>::Value && TModels_V<CStaticStructProvider, InStructType>)>
+	template<typename InStructType>
+		requires (!std::is_const_v<InStructType> && TModels_V<CStaticStructProvider, InStructType>)
 	explicit FAvaTransitionScene(InStructType* InStruct)
 		: DataView(InStructType::StaticStruct(), reinterpret_cast<uint8*>(InStruct))
 	{	
diff --git a/Engine/Source/Runtime/Core/Public/Delegates/DelegateInstancesImpl.h b/Engine/Source/Runtime/Core/Public/Delegates/DelegateInstancesImpl.h
--- a/Engine/Source/Runtime/Core/Public/Delegates/DelegateInstancesImpl.h
+++ b/Engine/Source/Runtime/Core/Public/Delegates/DelegateInstancesImpl.h
@@ -16,7 +16,6 @@
 #include "Delegates/IDelegateInstance.h"
 #include "Delegates/DelegateBase.h"
 #include "Misc/AssertionMacros.h"
-#include "Templates/RemoveReference.h"
 #include "Templates/SharedPointer.h"
 #include "Templates/Tuple.h"
 #include "Templates/UnrealTypeTraits.h"
@@ -154,7 +153,7 @@ public:
 
 	RetValType Execute(ParamTypes... Params) const final
 	{
-		using FParmsWithPayload = TPayload<RetValType(typename TDecay<ParamTypes>::Type..., typename TDecay<VarTypes> ::Type...)>;
+		using FParmsWithPayload = TPayload<RetValType(std::decay_t<ParamTypes>..., std::decay_t<VarTypes>...)>;
 
 		checkSlow(IsSafeToExecute());
 
@@ -169,7 +168,7 @@ public:
 		TStrongObjectPtr<UserClass> PinnedObject = IsInGameThread() ? nullptr : this->UserObjectPtr.Pin();
 		if (UserClass* ActualUserObject = this->UserObjectPtr.Get())
 		{
-			using FParmsWithPayload = TPayload<RetValType(typename TDecay<ParamTypes>::Type..., typename TDecay<VarTypes> ::Type...)>;
+			using FParmsWithPayload = TPayload<RetValType(std::decay_t<ParamTypes>..., std::decay_t<VarTypes>...)>;
 
 			TPlacementNewer<FParmsWithPayload> PayloadAndParams;
 			this->Payload.ApplyAfter(PayloadAndParams, Forward<ParamTypes>(Params)...);
@@ -328,7 +327,7 @@ template <ESPMode SPMode, typename RetValType, typename... ParamTypes, typename
 class TBaseSPLambdaDelegateInstance<SPMode, RetValType(ParamTypes...), UserPolicy, FunctorType, VarTypes...> : public TCommonDelegateInstanceState<RetValType(ParamTypes...), UserPolicy, VarTypes...>
 {
 private:
-	static_assert(std::is_same_v<FunctorType, typename TRemoveReference<FunctorType>::Type>, "FunctorType cannot be a reference");
+	static_assert(std::is_same_v<FunctorType, std::remove_reference_t<FunctorType>>, "FunctorType cannot be a reference");
 
 	using Super = TCommonDelegateInstanceState<RetValType(ParamTypes...), UserPolicy, VarTypes...>;
 	using DelegateBaseType = typename UserPolicy::FDelegateExtras;
@@ -812,7 +811,7 @@ template <typename RetValType, typename... ParamTypes, typename UserPolicy, type
 class TBaseFunctorDelegateInstance<RetValType(ParamTypes...), UserPolicy, FunctorType, VarTypes...> : public TCommonDelegateInstanceState<RetValType(ParamTypes...), UserPolicy, VarTypes...>
 {
 private:
-	static_assert(std::is_same_v<FunctorType, typename TRemoveReference<FunctorType>::Type>, "FunctorType cannot be a reference");
+	static_assert(std::is_same_v<FunctorType, std::remove_reference_t<FunctorType>>, "FunctorType cannot be a reference");
 
 	using Super            = TCommonDelegateInstanceState<RetValType(ParamTypes...), UserPolicy, VarTypes...>;
 	using DelegateBaseType = typename UserPolicy::FDelegateExtras;
@@ -907,7 +906,7 @@ template <typename RetValType, typename... ParamTypes, typename UserPolicy, type
 class TWeakBaseFunctorDelegateInstance<RetValType(ParamTypes...), UserPolicy, FunctorType, VarTypes...> : public TCommonDelegateInstanceState<RetValType(ParamTypes...), UserPolicy, VarTypes...>
 {
 private:
-	static_assert(std::is_same_v<FunctorType, typename TRemoveReference<FunctorType>::Type>, "FunctorType cannot be a reference");
+	static_assert(std::is_same_v<FunctorType, std::remove_reference_t<FunctorType>>, "FunctorType cannot be a reference");
 
 	using Super            = TCommonDelegateInstanceState<RetValType(ParamTypes...), UserPolicy, VarTypes...>;
 	using DelegateBaseType = typename UserPolicy::FDelegateExtras;
@@ -1021,3 +1020,7 @@ private:
 	// const - because the binding doesn't affect the substitutability of a copied delegate.
 	mutable std::remove_const_t<FunctorType> Functor;
 };
+
+#if UE_ENABLE_INCLUDE_ORDER_DEPRECATED_IN_5_7
+#include "Templates/RemoveReference.h"
+#endif
diff --git a/Engine/Source/Runtime/Core/Public/Delegates/DelegateSignatureImpl.inl b/Engine/Source/Runtime/Core/Public/Delegates/DelegateSignatureImpl.inl
--- a/Engine/Source/Runtime/Core/Public/Delegates/DelegateSignatureImpl.inl
+++ b/Engine/Source/Runtime/Core/Public/Delegates/DelegateSignatureImpl.inl
@@ -14,9 +14,8 @@
 #include "UObject/NameTypes.h"
 #include "UObject/ScriptDelegates.h"
 #include "UObject/WeakObjectPtrTemplates.h"
-#include "Templates/IsConst.h"
 #include "Templates/CopyQualifiersFromTo.h"
-#include "Templates/RemoveReference.h"
+#include "Templates/Identity.h"
 #include "Templates/SharedPointer.h"
 #include "Templates/UnrealTemplate.h"
 #include <type_traits>
@@ -126,7 +125,7 @@ public:
 	template<typename FunctorType, typename... VarTypes>
 	inline void BindLambda(FunctorType&& InFunctor, VarTypes&&... Vars)
 	{
-		new (TWriteLockedDelegateAllocation{*this}) TBaseFunctorDelegateInstance<FuncType, UserPolicy, typename TRemoveReference<FunctorType>::Type, std::decay_t<VarTypes>...>(Forward<FunctorType>(InFunctor), Forward<VarTypes>(Vars)...);
+		new (TWriteLockedDelegateAllocation{*this}) TBaseFunctorDelegateInstance<FuncType, UserPolicy, std::remove_reference_t<FunctorType>, std::decay_t<VarTypes>...>(Forward<FunctorType>(InFunctor), Forward<VarTypes>(Vars)...);
 	}
 
 	/**
@@ -136,12 +135,12 @@ public:
 	template<typename UserClass, ESPMode Mode, typename FunctorType, typename... VarTypes>
 	inline void BindSPLambda(const TSharedRef<UserClass, Mode>& InUserObjectRef, FunctorType&& InFunctor, VarTypes&&... Vars)
 	{
-		new (TWriteLockedDelegateAllocation{*this}) TBaseSPLambdaDelegateInstance<Mode, FuncType, UserPolicy, typename TRemoveReference<FunctorType>::Type, std::decay_t<VarTypes>...>(TWeakPtr<const void>(InUserObjectRef), Forward<FunctorType>(InFunctor), Forward<VarTypes>(Vars)...);
+		new (TWriteLockedDelegateAllocation{*this}) TBaseSPLambdaDelegateInstance<Mode, FuncType, UserPolicy, std::remove_reference_t<FunctorType>, std::decay_t<VarTypes>...>(TWeakPtr<const void>(InUserObjectRef), Forward<FunctorType>(InFunctor), Forward<VarTypes>(Vars)...);
 	}
 	template <typename UserClass, typename FunctorType, typename... VarTypes>
 	inline void BindSPLambda(const UserClass* InUserObject, FunctorType&& InFunctor, VarTypes&&... Vars)
 	{
-		new (TWriteLockedDelegateAllocation{*this}) TBaseSPLambdaDelegateInstance<decltype(InUserObject->AsShared())::Mode, FuncType, UserPolicy, typename TRemoveReference<FunctorType>::Type, std::decay_t<VarTypes>...>(TWeakPtr<const void>(InUserObject->AsWeak()), Forward<FunctorType>(InFunctor), Forward<VarTypes>(Vars)...);
+		new (TWriteLockedDelegateAllocation{*this}) TBaseSPLambdaDelegateInstance<decltype(InUserObject->AsShared())::Mode, FuncType, UserPolicy, std::remove_reference_t<FunctorType>, std::decay_t<VarTypes>...>(TWeakPtr<const void>(InUserObject->AsWeak()), Forward<FunctorType>(InFunctor), Forward<VarTypes>(Vars)...);
 	}
 
 	/**
@@ -151,7 +150,7 @@ public:
 	template<typename FunctorType, typename... VarTypes>
 	inline void BindWeakLambda(const UObject* InUserObject, FunctorType&& InFunctor, VarTypes&&... Vars)
 	{
-		new (TWriteLockedDelegateAllocation{*this}) TWeakBaseFunctorDelegateInstance<FuncType, UserPolicy, typename TRemoveReference<FunctorType>::Type, std::decay_t<VarTypes>...>(InUserObject, Forward<FunctorType>(InFunctor), Forward<VarTypes>(Vars)...);
+		new (TWriteLockedDelegateAllocation{*this}) TWeakBaseFunctorDelegateInstance<FuncType, UserPolicy, std::remove_reference_t<FunctorType>, std::decay_t<VarTypes>...>(InUserObject, Forward<FunctorType>(InFunctor), Forward<VarTypes>(Vars)...);
 	}
 
 	/**
@@ -163,7 +162,7 @@ public:
 	template <typename UserClass, typename... VarTypes>
 	inline void BindRaw(UserClass* InUserObject, typename TMemFunPtrType<false, UserClass, RetValType (ParamTypes..., std::decay_t<VarTypes>...)>::Type InFunc, VarTypes&&... Vars)
 	{
-		static_assert(!TIsConst<UserClass>::Value, "Attempting to bind a delegate with a const object pointer and non-const member function.");
+		static_assert(!std::is_const_v<UserClass>, "Attempting to bind a delegate with a const object pointer and non-const member function.");
 
 		new (TWriteLockedDelegateAllocation{*this}) TBaseRawMethodDelegateInstance<false, UserClass, FuncType, UserPolicy, std::decay_t<VarTypes>...>(InUserObject, InFunc, Forward<VarTypes>(Vars)...);
 	}
@@ -179,7 +178,7 @@ public:
 	template <typename UserClass, ESPMode Mode, typename... VarTypes>
 	inline void BindSP(const TSharedRef<UserClass, Mode>& InUserObjectRef, typename TMemFunPtrType<false, UserClass, RetValType (ParamTypes..., std::decay_t<VarTypes>...)>::Type InFunc, VarTypes&&... Vars)
 	{
-		static_assert(!TIsConst<UserClass>::Value, "Attempting to bind a delegate with a const object pointer and non-const member function.");
+		static_assert(!std::is_const_v<UserClass>, "Attempting to bind a delegate with a const object pointer and non-const member function.");
 
 		new (TWriteLockedDelegateAllocation{*this}) TBaseSPMethodDelegateInstance<false, UserClass, Mode, FuncType, UserPolicy, std::decay_t<VarTypes>...>(InUserObjectRef, InFunc, Forward<VarTypes>(Vars)...);
 	}
@@ -198,7 +197,7 @@ public:
 	template <typename UserClass, typename... VarTypes>
 	inline void BindSP(UserClass* InUserObject, typename TMemFunPtrType<false, UserClass, RetValType (ParamTypes..., std::decay_t<VarTypes>...)>::Type InFunc, VarTypes&&... Vars)
 	{
-		static_assert(!TIsConst<UserClass>::Value, "Attempting to bind a delegate with a const object pointer and non-const member function.");
+		static_assert(!std::is_const_v<UserClass>, "Attempting to bind a delegate with a const object pointer and non-const member function.");
 
 		new (TWriteLockedDelegateAllocation{*this}) TBaseSPMethodDelegateInstance<false, UserClass, decltype(InUserObject->AsShared())::Mode, FuncType, UserPolicy, std::decay_t<VarTypes>...>(StaticCastSharedRef<UserClass>(InUserObject->AsShared()), InFunc, Forward<VarTypes>(Vars)...);
 	}
@@ -219,7 +218,7 @@ public:
 	template <typename UserClass, typename... VarTypes>
 	inline void BindThreadSafeSP(const TSharedRef<UserClass, ESPMode::ThreadSafe>& InUserObjectRef, typename TMemFunPtrType<false, UserClass, RetValType (ParamTypes..., std::decay_t<VarTypes>...)>::Type InFunc, VarTypes&&... Vars)
 	{
-		static_assert(!TIsConst<UserClass>::Value, "Attempting to bind a delegate with a const object pointer and non-const member function.");
+		static_assert(!std::is_const_v<UserClass>, "Attempting to bind a delegate with a const object pointer and non-const member function.");
 
 		new (TWriteLockedDelegateAllocation{*this}) TBaseSPMethodDelegateInstance<false, UserClass, ESPMode::ThreadSafe, FuncType, UserPolicy, std::decay_t<VarTypes>...>(InUserObjectRef, InFunc, Forward<VarTypes>(Vars)...);
 	}
@@ -240,7 +239,7 @@ public:
 	template <typename UserClass, typename... VarTypes>
 	inline void BindThreadSafeSP(UserClass* InUserObject, typename TMemFunPtrType<false, UserClass, RetValType (ParamTypes..., std::decay_t<VarTypes>...)>::Type InFunc, VarTypes&&... Vars)
 	{
-		static_assert(!TIsConst<UserClass>::Value, "Attempting to bind a delegate with a const object pointer and non-const member function.");
+		static_assert(!std::is_const_v<UserClass>, "Attempting to bind a delegate with a const object pointer and non-const member function.");
 
 		new (TWriteLockedDelegateAllocation{*this}) TBaseSPMethodDelegateInstance<false, UserClass, ESPMode::ThreadSafe, FuncType, UserPolicy, std::decay_t<VarTypes>...>(StaticCastSharedRef<UserClass>(InUserObject->AsShared()), InFunc, Forward<VarTypes>(Vars)...);
 	}
@@ -276,7 +275,7 @@ public:
 	template <typename UserClass, typename... VarTypes>
 	inline void BindUObject(UserClass* InUserObject, typename TMemFunPtrType<false, UserClass, RetValType (ParamTypes..., std::decay_t<VarTypes>...)>::Type InFunc, VarTypes&&... Vars)
 	{
-		static_assert(!TIsConst<UserClass>::Value, "Attempting to bind a delegate with a const object pointer and non-const member function.");
+		static_assert(!std::is_const_v<UserClass>, "Attempting to bind a delegate with a const object pointer and non-const member function.");
 
 		new (TWriteLockedDelegateAllocation{*this}) TBaseUObjectMethodDelegateInstance<false, UserClass, FuncType, UserPolicy, std::decay_t<VarTypes>...>(InUserObject, InFunc, Forward<VarTypes>(Vars)...);
 	}
@@ -288,7 +287,7 @@ public:
 	template <typename UserClass, typename... VarTypes>
 	inline void BindUObject(TObjectPtr<UserClass> InUserObject, typename TMemFunPtrType<false, UserClass, RetValType (ParamTypes..., std::decay_t<VarTypes>...)>::Type InFunc, VarTypes&&... Vars)
 	{
-		static_assert(!TIsConst<UserClass>::Value, "Attempting to bind a delegate with a const object pointer and non-const member function.");
+		static_assert(!std::is_const_v<UserClass>, "Attempting to bind a delegate with a const object pointer and non-const member function.");
 
 		new (TWriteLockedDelegateAllocation{*this}) TBaseUObjectMethodDelegateInstance<false, UserClass, FuncType, UserPolicy, std::decay_t<VarTypes>...>(ToRawPtr(InUserObject), InFunc, Forward<VarTypes>(Vars)...);
 	}
@@ -363,7 +362,7 @@ public:
 	 * Static: Creates a raw C++ pointer global function delegate
 	 */
 	template <typename... VarTypes>
-	[[nodiscard]] inline static TDelegate<RetValType(ParamTypes...), UserPolicy> CreateStatic(typename TIdentity<RetValType (*)(ParamTypes..., std::decay_t<VarTypes>...)>::Type InFunc, VarTypes&&... Vars)
+	[[nodiscard]] inline static TDelegate<RetValType(ParamTypes...), UserPolicy> CreateStatic(TIdentity_T<RetValType (*)(ParamTypes..., std::decay_t<VarTypes>...)> InFunc, VarTypes&&... Vars)
 	{
 		TDelegate<RetValType(ParamTypes...), UserPolicy> Result;
 		new (TWriteLockedDelegateAllocation{Result}) TBaseStaticDelegateInstance<FuncType, UserPolicy, std::decay_t<VarTypes>...>(InFunc, Forward<VarTypes>(Vars)...);
@@ -378,7 +377,7 @@ public:
 	[[nodiscard]] inline static TDelegate<RetValType(ParamTypes...), UserPolicy> CreateLambda(FunctorType&& InFunctor, VarTypes&&... Vars)
 	{
 		TDelegate<RetValType(ParamTypes...), UserPolicy> Result;
-		new (TWriteLockedDelegateAllocation{Result}) TBaseFunctorDelegateInstance<FuncType, UserPolicy, typename TRemoveReference<FunctorType>::Type, std::decay_t<VarTypes>...>(Forward<FunctorType>(InFunctor), Forward<VarTypes>(Vars)...);
+		new (TWriteLockedDelegateAllocation{Result}) TBaseFunctorDelegateInstance<FuncType, UserPolicy, std::remove_reference_t<FunctorType>, std::decay_t<VarTypes>...>(Forward<FunctorType>(InFunctor), Forward<VarTypes>(Vars)...);
 		return Result;
 	}
 
@@ -390,14 +389,14 @@ public:
 	[[nodiscard]] inline static TDelegate<RetValType(ParamTypes...), UserPolicy> CreateSPLambda(const TSharedRef<UserClass, Mode>& InUserObjectRef, FunctorType&& InFunctor, VarTypes&&... Vars)
 	{
 		TDelegate<RetValType(ParamTypes...), UserPolicy> Result;
-		new (TWriteLockedDelegateAllocation{Result}) TBaseSPLambdaDelegateInstance<Mode, FuncType, UserPolicy, typename TRemoveReference<FunctorType>::Type, std::decay_t<VarTypes>...>(InUserObjectRef, Forward<FunctorType>(InFunctor), Forward<VarTypes>(Vars)...);
+		new (TWriteLockedDelegateAllocation{Result}) TBaseSPLambdaDelegateInstance<Mode, FuncType, UserPolicy, std::remove_reference_t<FunctorType>, std::decay_t<VarTypes>...>(InUserObjectRef, Forward<FunctorType>(InFunctor), Forward<VarTypes>(Vars)...);
 		return Result;
 	}
 	template <typename UserClass, typename FunctorType, typename... VarTypes>
 	[[nodiscard]] inline static TDelegate<RetValType(ParamTypes...), UserPolicy> CreateSPLambda(UserClass* InUserObject, FunctorType&& InFunctor, VarTypes&&... Vars)
 	{
 		TDelegate<RetValType(ParamTypes...), UserPolicy> Result;
-		new (TWriteLockedDelegateAllocation{Result}) TBaseSPLambdaDelegateInstance<decltype(InUserObject->AsWeak())::Mode, FuncType, UserPolicy, typename TRemoveReference<FunctorType>::Type, std::decay_t<VarTypes>...>(TWeakPtr<const void>(InUserObject->AsWeak()), Forward<FunctorType>(InFunctor), Forward<VarTypes>(Vars)...);
+		new (TWriteLockedDelegateAllocation{Result}) TBaseSPLambdaDelegateInstance<decltype(InUserObject->AsWeak())::Mode, FuncType, UserPolicy, std::remove_reference_t<FunctorType>, std::decay_t<VarTypes>...>(TWeakPtr<const void>(InUserObject->AsWeak()), Forward<FunctorType>(InFunctor), Forward<VarTypes>(Vars)...);
 		return Result;
 	}
 
@@ -409,7 +408,7 @@ public:
 	[[nodiscard]] inline static TDelegate<RetValType(ParamTypes...), UserPolicy> CreateWeakLambda(const UObject* InUserObject, FunctorType&& InFunctor, VarTypes&&... Vars)
 	{
 		TDelegate<RetValType(ParamTypes...), UserPolicy> Result;
-		new (TWriteLockedDelegateAllocation{Result}) TWeakBaseFunctorDelegateInstance<FuncType, UserPolicy, typename TRemoveReference<FunctorType>::Type, std::decay_t<VarTypes>...>(InUserObject, Forward<FunctorType>(InFunctor), Forward<VarTypes>(Vars)...);
+		new (TWriteLockedDelegateAllocation{Result}) TWeakBaseFunctorDelegateInstance<FuncType, UserPolicy, std::remove_reference_t<FunctorType>, std::decay_t<VarTypes>...>(InUserObject, Forward<FunctorType>(InFunctor), Forward<VarTypes>(Vars)...);
 		return Result;
 	}
 
@@ -422,7 +421,7 @@ public:
 	template <typename UserClass, typename... VarTypes>
 	[[nodiscard]] inline static TDelegate<RetValType(ParamTypes...), UserPolicy> CreateRaw(UserClass* InUserObject, typename TMemFunPtrType<false, UserClass, RetValType (ParamTypes..., std::decay_t<VarTypes>...)>::Type InFunc, VarTypes&&... Vars)
 	{
-		static_assert(!TIsConst<UserClass>::Value, "Attempting to bind a delegate with a const object pointer and non-const member function.");
+		static_assert(!std::is_const_v<UserClass>, "Attempting to bind a delegate with a const object pointer and non-const member function.");
 
 		TDelegate<RetValType(ParamTypes...), UserPolicy> Result;
 		new (TWriteLockedDelegateAllocation{Result}) TBaseRawMethodDelegateInstance<false, UserClass, FuncType, UserPolicy, std::decay_t<VarTypes>...>(InUserObject, InFunc, Forward<VarTypes>(Vars)...);
@@ -445,7 +444,7 @@ public:
 	template <typename UserClass, ESPMode Mode, typename... VarTypes>
 	[[nodiscard]] inline static TDelegate<RetValType(ParamTypes...), UserPolicy> CreateSP(const TSharedRef<UserClass, Mode>& InUserObjectRef, typename TMemFunPtrType<false, UserClass, RetValType (ParamTypes..., std::decay_t<VarTypes>...)>::Type InFunc, VarTypes&&... Vars)
 	{
-		static_assert(!TIsConst<UserClass>::Value, "Attempting to bind a delegate with a const object pointer and non-const member function.");
+		static_assert(!std::is_const_v<UserClass>, "Attempting to bind a delegate with a const object pointer and non-const member function.");
 
 		TDelegate<RetValType(ParamTypes...), UserPolicy> Result;
 		new (TWriteLockedDelegateAllocation{Result}) TBaseSPMethodDelegateInstance<false, UserClass, Mode, FuncType, UserPolicy, std::decay_t<VarTypes>...>(InUserObjectRef, InFunc, Forward<VarTypes>(Vars)...);
@@ -468,7 +467,7 @@ public:
 	template <typename UserClass, typename... VarTypes>
 	[[nodiscard]] inline static TDelegate<RetValType(ParamTypes...), UserPolicy> CreateSP(UserClass* InUserObject, typename TMemFunPtrType<false, UserClass, RetValType (ParamTypes..., std::decay_t<VarTypes>...)>::Type InFunc, VarTypes&&... Vars)
 	{
-		static_assert(!TIsConst<UserClass>::Value, "Attempting to bind a delegate with a const object pointer and non-const member function.");
+		static_assert(!std::is_const_v<UserClass>, "Attempting to bind a delegate with a const object pointer and non-const member function.");
 
 		TDelegate<RetValType(ParamTypes...), UserPolicy> Result;
 		new (TWriteLockedDelegateAllocation{Result}) TBaseSPMethodDelegateInstance<false, UserClass, decltype(InUserObject->AsShared())::Mode, FuncType, UserPolicy, std::decay_t<VarTypes>...>(StaticCastSharedRef<UserClass>(InUserObject->AsShared()), InFunc, Forward<VarTypes>(Vars)...);
@@ -493,7 +492,7 @@ public:
 	template <typename UserClass, typename... VarTypes>
 	[[nodiscard]] inline static TDelegate<RetValType(ParamTypes...), UserPolicy> CreateThreadSafeSP(const TSharedRef<UserClass, ESPMode::ThreadSafe>& InUserObjectRef, typename TMemFunPtrType<false, UserClass, RetValType (ParamTypes..., std::decay_t<VarTypes>...)>::Type InFunc, VarTypes&&... Vars)
 	{
-		static_assert(!TIsConst<UserClass>::Value, "Attempting to bind a delegate with a const object pointer and non-const member function.");
+		static_assert(!std::is_const_v<UserClass>, "Attempting to bind a delegate with a const object pointer and non-const member function.");
 
 		TDelegate<RetValType(ParamTypes...), UserPolicy> Result;
 		new (TWriteLockedDelegateAllocation{Result}) TBaseSPMethodDelegateInstance<false, UserClass, ESPMode::ThreadSafe, FuncType, UserPolicy, std::decay_t<VarTypes>...>(InUserObjectRef, InFunc, Forward<VarTypes>(Vars)...);
@@ -518,7 +517,7 @@ public:
 	template <typename UserClass, typename... VarTypes>
 	[[nodiscard]] inline static TDelegate<RetValType(ParamTypes...), UserPolicy> CreateThreadSafeSP(UserClass* InUserObject, typename TMemFunPtrType<false, UserClass, RetValType (ParamTypes..., std::decay_t<VarTypes>...)>::Type InFunc, VarTypes&&... Vars)
 	{
-		static_assert(!TIsConst<UserClass>::Value, "Attempting to bind a delegate with a const object pointer and non-const member function.");
+		static_assert(!std::is_const_v<UserClass>, "Attempting to bind a delegate with a const object pointer and non-const member function.");
 
 		TDelegate<RetValType(ParamTypes...), UserPolicy> Result;
 		new (TWriteLockedDelegateAllocation{Result}) TBaseSPMethodDelegateInstance<false, UserClass, ESPMode::ThreadSafe, FuncType, UserPolicy, std::decay_t<VarTypes>...>(StaticCastSharedRef<UserClass>(InUserObject->AsShared()), InFunc, Forward<VarTypes>(Vars)...);
@@ -562,7 +561,7 @@ public:
 	template <typename UserClass, typename... VarTypes>
 	[[nodiscard]] inline static TDelegate<RetValType(ParamTypes...), UserPolicy> CreateUObject(UserClass* InUserObject, typename TMemFunPtrType<false, UserClass, RetValType (ParamTypes..., std::decay_t<VarTypes>...)>::Type InFunc, VarTypes&&... Vars)
 	{
-		static_assert(!TIsConst<UserClass>::Value, "Attempting to bind a delegate with a const object pointer and non-const member function.");
+		static_assert(!std::is_const_v<UserClass>, "Attempting to bind a delegate with a const object pointer and non-const member function.");
 
 		TDelegate<RetValType(ParamTypes...), UserPolicy> Result;
 		new (TWriteLockedDelegateAllocation{Result}) TBaseUObjectMethodDelegateInstance<false, UserClass, FuncType, UserPolicy, std::decay_t<VarTypes>...>(InUserObject, InFunc, Forward<VarTypes>(Vars)...);
@@ -578,7 +577,7 @@ public:
 	template <typename UserClass, typename... VarTypes>
 	[[nodiscard]] inline static TDelegate<RetValType(ParamTypes...), UserPolicy> CreateUObject(TObjectPtr<UserClass> InUserObject, typename TMemFunPtrType<false, UserClass, RetValType (ParamTypes..., std::decay_t<VarTypes>...)>::Type InFunc, VarTypes&&... Vars)
 	{
-		static_assert(!TIsConst<UserClass>::Value, "Attempting to bind a delegate with a const object pointer and non-const member function.");
+		static_assert(!std::is_const_v<UserClass>, "Attempting to bind a delegate with a const object pointer and non-const member function.");
 
 		TDelegate<RetValType(ParamTypes...), UserPolicy> Result;
 		new (TWriteLockedDelegateAllocation{Result}) TBaseUObjectMethodDelegateInstance<false, UserClass, FuncType, UserPolicy, std::decay_t<VarTypes>...>(ToRawPtr(InUserObject), InFunc, Forward<VarTypes>(Vars)...);
@@ -620,13 +619,8 @@ public:
 	 * @return  Returns true if the function was executed
 	 */
 	 // NOTE: Currently only delegates with no return value support ExecuteIfBound()
-	template <
-		// This construct is intended to disable this function when RetValType != void.
-		// DummyRetValType exists to create a substitution which can fail, to achieve SFINAE.
-		typename DummyRetValType = RetValType,
-		std::enable_if_t<std::is_void<DummyRetValType>::value>* = nullptr
-	>
 	inline bool ExecuteIfBound(ParamTypes... Params) const
+		requires (std::is_void_v<RetValType>)
 	{
 		FReadAccessScope ReadScope = GetReadAccessScope();
 
@@ -846,7 +840,7 @@ public:
 	template <typename UserClass, typename... VarTypes>
 	inline FDelegateHandle AddRaw(UserClass* InUserObject, typename TMemFunPtrType<false, UserClass, void (ParamTypes..., std::decay_t<VarTypes>...)>::Type InFunc, VarTypes&&... Vars)
 	{
-		static_assert(!TIsConst<UserClass>::Value, "Attempting to bind a delegate with a const object pointer and non-const member function.");
+		static_assert(!std::is_const_v<UserClass>, "Attempting to bind a delegate with a const object pointer and non-const member function.");
 
 		return Super::AddDelegateInstance(FDelegate::CreateRaw(InUserObject, InFunc, Forward<VarTypes>(Vars)...));
 	}
@@ -867,7 +861,7 @@ public:
 	template <typename UserClass, ESPMode Mode, typename... VarTypes>
 	inline FDelegateHandle AddSP(const TSharedRef<UserClass, Mode>& InUserObjectRef, typename TMemFunPtrType<false, UserClass, void (ParamTypes..., std::decay_t<VarTypes>...)>::Type InFunc, VarTypes&&... Vars)
 	{
-		static_assert(!TIsConst<UserClass>::Value, "Attempting to bind a delegate with a const object pointer and non-const member function.");
+		static_assert(!std::is_const_v<UserClass>, "Attempting to bind a delegate with a const object pointer and non-const member function.");
 
 		return Super::AddDelegateInstance(FDelegate::CreateSP(InUserObjectRef, InFunc, Forward<VarTypes>(Vars)...));
 	}
@@ -888,7 +882,7 @@ public:
 	template <typename UserClass, typename... VarTypes>
 	inline FDelegateHandle AddSP(UserClass* InUserObject, typename TMemFunPtrType<false, UserClass, void (ParamTypes..., std::decay_t<VarTypes>...)>::Type InFunc, VarTypes&&... Vars)
 	{
-		static_assert(!TIsConst<UserClass>::Value, "Attempting to bind a delegate with a const object pointer and non-const member function.");
+		static_assert(!std::is_const_v<UserClass>, "Attempting to bind a delegate with a const object pointer and non-const member function.");
 
 		return Super::AddDelegateInstance(FDelegate::CreateSP(InUserObject, InFunc, Forward<VarTypes>(Vars)...));
 	}
@@ -909,7 +903,7 @@ public:
 	template <typename UserClass, typename... VarTypes>
 	inline FDelegateHandle AddThreadSafeSP(const TSharedRef<UserClass, ESPMode::ThreadSafe>& InUserObjectRef, typename TMemFunPtrType<false, UserClass, void (ParamTypes..., std::decay_t<VarTypes>...)>::Type InFunc, VarTypes&&... Vars)
 	{
-		static_assert(!TIsConst<UserClass>::Value, "Attempting to bind a delegate with a const object pointer and non-const member function.");
+		static_assert(!std::is_const_v<UserClass>, "Attempting to bind a delegate with a const object pointer and non-const member function.");
 
 		return Super::AddDelegateInstance(FDelegate::CreateThreadSafeSP(InUserObjectRef, InFunc, Forward<VarTypes>(Vars)...));
 	}
@@ -932,7 +926,7 @@ public:
 	template <typename UserClass, typename... VarTypes>
 	inline FDelegateHandle AddThreadSafeSP(UserClass* InUserObject, typename TMemFunPtrType<false, UserClass, void (ParamTypes..., std::decay_t<VarTypes>...)>::Type InFunc, VarTypes&&... Vars)
 	{
-		static_assert(!TIsConst<UserClass>::Value, "Attempting to bind a delegate with a const object pointer and non-const member function.");
+		static_assert(!std::is_const_v<UserClass>, "Attempting to bind a delegate with a const object pointer and non-const member function.");
 
 		return Super::AddDelegateInstance(FDelegate::CreateThreadSafeSP(InUserObject, InFunc, Forward<VarTypes>(Vars)...));
 	}
@@ -972,7 +966,7 @@ public:
 	template <typename UserClass, typename... VarTypes>
 	inline FDelegateHandle AddUObject(UserClass* InUserObject, typename TMemFunPtrType<false, UserClass, void (ParamTypes..., std::decay_t<VarTypes>...)>::Type InFunc, VarTypes&&... Vars)
 	{
-		static_assert(!TIsConst<UserClass>::Value, "Attempting to bind a delegate with a const object pointer and non-const member function.");
+		static_assert(!std::is_const_v<UserClass>, "Attempting to bind a delegate with a const object pointer and non-const member function.");
 
 		return Super::AddDelegateInstance(FDelegate::CreateUObject(InUserObject, InFunc, Forward<VarTypes>(Vars)...));
 	}
@@ -984,7 +978,7 @@ public:
 	template <typename UserClass, typename... VarTypes>
 	inline FDelegateHandle AddUObject(TObjectPtr<UserClass> InUserObject, typename TMemFunPtrType<false, UserClass, void (ParamTypes..., std::decay_t<VarTypes>...)>::Type InFunc, VarTypes&&... Vars)
 	{
-		static_assert(!TIsConst<UserClass>::Value, "Attempting to bind a delegate with a const object pointer and non-const member function.");
+		static_assert(!std::is_const_v<UserClass>, "Attempting to bind a delegate with a const object pointer and non-const member function.");
 
 		return Super::AddDelegateInstance(FDelegate::CreateUObject(InUserObject, InFunc, Forward<VarTypes>(Vars)...));
 	}
@@ -1299,3 +1293,8 @@ public:
 #if UE_ENABLE_INCLUDE_ORDER_DEPRECATED_IN_5_5
 #include "Misc/Crc.h"
 #endif
+
+#if UE_ENABLE_INCLUDE_ORDER_DEPRECATED_IN_5_7
+#include "Templates/IsConst.h"
+#include "Templates/RemoveReference.h"
+#endif
diff --git a/Engine/Source/Runtime/Core/Public/Memory/MemoryView.h b/Engine/Source/Runtime/Core/Public/Memory/MemoryView.h
--- a/Engine/Source/Runtime/Core/Public/Memory/MemoryView.h
+++ b/Engine/Source/Runtime/Core/Public/Memory/MemoryView.h
@@ -16,6 +16,7 @@
 #include "Traits/IsContiguousContainer.h"
 
 #include <initializer_list>
+#include <type_traits>
 
 /**
  * A non-owning view of a contiguous region of memory.
@@ -29,7 +30,7 @@ class TMemoryView
 {
 	static_assert(std::is_void_v<DataType>, "DataType must be cv-qualified void");
 
-	using ByteType = std::conditional_t<TIsConst<DataType>::Value, const uint8, uint8>;
+	using ByteType = std::conditional_t<std::is_const_v<DataType>, const uint8, uint8>;
 
 public:
 	/** Construct an empty view. */
@@ -295,7 +296,7 @@ template <
 [[nodiscard]] constexpr inline auto MakeMemoryView(ContainerType&& Container)
 {
 	using ElementType = typename TRemovePointer<decltype(GetData(DeclVal<ContainerType>()))>::Type;
-	constexpr bool bIsConst = TIsConst<ElementType>::Value;
+	constexpr bool bIsConst = std::is_const_v<ElementType>;
 	using DataType = std::conditional_t<bIsConst, const void, void>;
 	return TMemoryView<DataType>(GetData(Container), GetNum(Container) * sizeof(ElementType));
 }
diff --git a/Engine/Source/Runtime/CoreUObject/Public/StructUtils/InstancedStructContainer.h b/Engine/Source/Runtime/CoreUObject/Public/StructUtils/InstancedStructContainer.h
--- a/Engine/Source/Runtime/CoreUObject/Public/StructUtils/InstancedStructContainer.h
+++ b/Engine/Source/Runtime/CoreUObject/Public/StructUtils/InstancedStructContainer.h
@@ -152,7 +152,7 @@ public:
 	template<typename T>
 	struct TIterator
 	{
-		using StructViewType = std::conditional_t<TIsConst<T>::Value, FConstStructView, FStructView>;
+		using StructViewType = std::conditional_t<std::is_const_v<T>, FConstStructView, FStructView>;
 
 		/** @return struct view (or const structview) to the item. */
 		StructViewType operator*()
@@ -194,7 +194,7 @@ public:
 		/** Removes the item pointed by the iterator and adjust the iterator. */
 		void RemoveCurrent()
 		{
-			static_assert(!TIsConst<T>::Value, "Cannot remove on const container.");
+			static_assert(!std::is_const_v<T>, "Cannot remove on const container.");
 			Container.RemoveAt(Index, 1);
 			Index--;
 		}
diff --git a/Engine/Source/Runtime/CoreUObject/Public/StructUtils/SharedStruct.h b/Engine/Source/Runtime/CoreUObject/Public/StructUtils/SharedStruct.h
--- a/Engine/Source/Runtime/CoreUObject/Public/StructUtils/SharedStruct.h
+++ b/Engine/Source/Runtime/CoreUObject/Public/StructUtils/SharedStruct.h
@@ -7,6 +7,7 @@
 #include "StructUtils/StructUtilsTypes.h"
 #include "UObject/ObjectMacros.h"
 #include "UObject/ScriptMacros.h"
+#include <type_traits>
 
 #include "SharedStruct.generated.h"
 
@@ -212,18 +213,18 @@ struct [[nodiscard]] FSharedStruct
 		InitializeAs(TBaseStructure<T>::Get(), nullptr);
 	}
 
-	/** Initializes from templated struct instance. This will create a new instance of the shared struct memory. */
-	template <typename T, UE::StructUtils::EnableIfNotSharedInstancedOrViewStruct<T>* = nullptr>
-	void InitializeAs(const T& Struct)
-	{
-		InitializeAs(TBaseStructure<T>::Get(), reinterpret_cast<const uint8*>(&Struct));
-	}
-
 	/** Initializes from other related struct types. This will create a new instance of the shared struct memory. */
-	template <typename T, UE::StructUtils::EnableIfSharedInstancedOrViewStruct<T>* = nullptr>
+	template <typename T>
 	void InitializeAs(const T& Struct)
 	{
-		InitializeAs(Struct.GetScriptStruct(), Struct.GetMemory());
+		if constexpr (UE::StructUtils::TIsSharedInstancedOrViewStruct_V<T>)
+		{
+			InitializeAs(Struct.GetScriptStruct(), Struct.GetMemory());
+		}
+		else
+		{
+			InitializeAs(TBaseStructure<T>::Get(), reinterpret_cast<const uint8*>(&Struct));
+		}
 	}
 
 	/** Initializes from a struct type and optional data. This will create a new instance of the shared struct memory. */
@@ -349,20 +350,23 @@ struct [[nodiscard]] TSharedStruct
 	explicit TSharedStruct() = default;
 
 	/** Copy constructors */
-	template<typename T = BaseStructT, typename = std::enable_if_t<std::is_base_of_v<BaseStructT, std::decay_t<T>>>>
+	template<typename T = BaseStructT>
+		requires (std::is_base_of_v<BaseStructT, std::decay_t<T>>)
 	TSharedStruct(const TSharedStruct<T>& InOther)
 		: SharedStruct(InOther.SharedStruct)
 	{
 	}
 
-	template<typename T = BaseStructT, typename = std::enable_if_t<std::is_base_of_v<BaseStructT, std::decay_t<T>>>>
+	template<typename T = BaseStructT>
+		requires (std::is_base_of_v<BaseStructT, std::decay_t<T>>)
 	TSharedStruct(TSharedStruct<T>&& InOther)
 		: SharedStruct(MoveTemp(InOther.SharedStruct))
 	{
 	}
 
 	/** Assignment operators */
-	template<typename T = BaseStructT, typename = std::enable_if_t<std::is_base_of_v<BaseStructT, std::decay_t<T>>>>
+	template<typename T = BaseStructT>
+		requires (std::is_base_of_v<BaseStructT, std::decay_t<T>>)
 	TSharedStruct& operator=(const TSharedStruct<T>& InOther)
 	{
 		if (this != &InOther)
@@ -372,7 +376,8 @@ struct [[nodiscard]] TSharedStruct
 		return *this;
 	}
 
-	template<typename T = BaseStructT, typename = std::enable_if_t<std::is_base_of_v<BaseStructT, std::decay_t<T>>>>
+	template<typename T = BaseStructT>
+		requires (std::is_base_of_v<BaseStructT, std::decay_t<T>>)
 	TSharedStruct& operator=(TSharedStruct<T>&& InOther)
 	{
 		if (this != &InOther)
@@ -406,42 +411,46 @@ struct [[nodiscard]] TSharedStruct
 	}
 
 	/** Initializes from a templated struct type. This will create a new instance of the shared struct memory. */
-	template<typename T = BaseStructT, typename = std::enable_if_t<std::is_base_of_v<BaseStructT, std::decay_t<T>>>>
+	template<typename T = BaseStructT>
+		requires (std::is_base_of_v<BaseStructT, std::decay_t<T>>)
 	void Initialize()
 	{
 		SharedStruct.InitializeAs<T>();
 	}
 
 	/** Initializes from templated struct instance. This will create a new instance of the shared struct memory. */
-	template<typename T = BaseStructT,
-		typename = std::enable_if_t<
+	template<typename T = BaseStructT>
+		requires(
 			std::is_base_of_v<BaseStructT, std::decay_t<T>>
 			&& !(std::is_same_v<TStructView<T>, T>
 				|| std::is_same_v<TConstStructView<T>, T>
 				|| std::is_same_v<TSharedStruct<T>, T>
 				|| std::is_same_v<TConstSharedStruct<T>, T>
-				|| std::is_same_v<TInstancedStruct<T>, T>), T>>
+				|| std::is_same_v<TInstancedStruct<T>, T>))
 	void Initialize(const T& Struct)
 	{
 		SharedStruct.InitializeAs(TBaseStructure<T>::Get(), reinterpret_cast<const uint8*>(&Struct));
 	}
 
 	/** Initializes from struct type and emplace args. This will create a new instance of the shared struct memory.*/
-	template<typename T = BaseStructT, typename... TArgs, typename = std::enable_if_t<std::is_base_of_v<BaseStructT, std::decay_t<T>>>>
+	template<typename T = BaseStructT, typename... TArgs>
+		requires (std::is_base_of_v<BaseStructT, std::decay_t<T>>)
 	void Initialize(TArgs&&... InArgs)
 	{
 		SharedStruct.InitializeAs(Forward<TArgs>(InArgs)...);
 	}
 
 	/** Creates a new TSharedStruct. This will create a new instance of the shared struct memory. */
-	template<typename T = BaseStructT, typename = std::enable_if_t<std::is_base_of_v<BaseStructT, std::decay_t<T>>>>
+	template<typename T = BaseStructT>
+		requires (std::is_base_of_v<BaseStructT, std::decay_t<T>>)
 	static TSharedStruct Make()
 	{
 		return FSharedStruct::Make<T>();
 	}
 
 	/** Creates a new TSharedStruct from templated struct instance. This will create a new instance of the shared struct memory. */
-	template<typename T = BaseStructT, typename = std::enable_if_t<std::is_base_of_v<BaseStructT, std::decay_t<T>>>>
+	template<typename T = BaseStructT>
+		requires (std::is_base_of_v<BaseStructT, std::decay_t<T>>)
 	static TSharedStruct Make(const T& Struct)
 	{
 		return FSharedStruct::Make<T>(Struct);
@@ -455,14 +464,16 @@ struct [[nodiscard]] TSharedStruct
 	}
 
 	/** Returns reference to the struct, this getter assumes that all data is valid. */
-	template<typename T = BaseStructT, typename = std::enable_if_t<std::is_base_of_v<BaseStructT, std::decay_t<T>>>>
+	template<typename T = BaseStructT>
+		requires (std::is_base_of_v<BaseStructT, std::decay_t<T>>)
 	T& Get() const
 	{
 		return SharedStruct.Get<T>();
 	}
 
 	/** Returns pointer to the struct, or nullptr if cast is not valid. */
-	template<typename T = BaseStructT, typename = std::enable_if_t<std::is_base_of_v<BaseStructT, std::decay_t<T>>>>
+	template<typename T = BaseStructT>
+		requires (std::is_base_of_v<BaseStructT, std::decay_t<T>>)
 	T* GetPtr() const
 	{
 		return SharedStruct.GetPtr<T>();
@@ -475,13 +486,15 @@ struct [[nodiscard]] TSharedStruct
 	}
 
 	/** Comparison operators. Note: it does not compare the internal structure itself */
-	template<typename T = BaseStructT, typename = std::enable_if_t<std::is_base_of_v<BaseStructT, std::decay_t<T>>>>
+	template<typename T = BaseStructT>
+		requires (std::is_base_of_v<BaseStructT, std::decay_t<T>>)
 	bool operator==(const T& Other) const
 	{
 		return SharedStruct.operator==(Other);
 	}
 
-	template<typename T = BaseStructT, typename = std::enable_if_t<std::is_base_of_v<BaseStructT, std::decay_t<T>>>>
+	template<typename T = BaseStructT>
+		requires (std::is_base_of_v<BaseStructT, std::decay_t<T>>)
 	bool operator!=(const T& Other) const
 	{
 		return SharedStruct.operator!=(Other);
@@ -588,18 +601,18 @@ struct [[nodiscard]] FConstSharedStruct
 		InitializeAs(TBaseStructure<T>::Get(), nullptr);
 	}
 
-	/** Initializes from templated struct instance. This will create a new instance of the shared struct memory. */
-	template <typename T, UE::StructUtils::EnableIfNotSharedInstancedOrViewStruct<T>* = nullptr>
-	void InitializeAs(const T& Struct)
-	{
-		InitializeAs(TBaseStructure<T>::Get(), reinterpret_cast<const uint8*>(&Struct));
-	}
-
 	/** Initializes from other related struct types. This will create a new instance of the shared struct memory. */
-	template <typename T, UE::StructUtils::EnableIfSharedInstancedOrViewStruct<T>* = nullptr>
+	template <typename T>
 	void InitializeAs(const T& Struct)
 	{
-		InitializeAs(Struct.GetScriptStruct(), Struct.GetMemory());
+		if constexpr (UE::StructUtils::TIsSharedInstancedOrViewStruct_V<T>)
+		{
+			InitializeAs(Struct.GetScriptStruct(), Struct.GetMemory());
+		}
+		else
+		{
+			InitializeAs(TBaseStructure<T>::Get(), reinterpret_cast<const uint8*>(&Struct));
+		}
 	}
 
 	/** Initializes from a struct type and optional data. This will create a new instance of the shared struct memory. */
@@ -661,14 +674,16 @@ struct [[nodiscard]] FConstSharedStruct
 
 	/** Returns const reference to the struct, this getter assumes that all data is valid. */
 	template<typename T>
-	constexpr typename TEnableIf<TIsConst<T>::Value, T&>::Type Get() const
+		requires (std::is_const_v<T>)
+	constexpr T& Get() const
 	{
 		return UE::StructUtils::GetStructRef<T>(GetScriptStruct(), GetMemory());
 	}
 
 	/** Returns const pointer to the struct, or nullptr if cast is not valid. */
 	template<typename T>
-	constexpr typename TEnableIf<TIsConst<T>::Value, T*>::Type GetPtr() const
+		requires (std::is_const_v<T>)
+	constexpr T* GetPtr() const
 	{
 		return UE::StructUtils::GetStructPtr<T>(GetScriptStruct(), GetMemory());
 	}
@@ -725,27 +740,36 @@ struct [[nodiscard]] TConstSharedStruct
 {
 	explicit TConstSharedStruct() = default;
 
-	template<typename T = BaseStructT, typename = std::enable_if_t<std::is_base_of_v<BaseStructT, std::decay_t<T>>>>
+	template<typename T = BaseStructT>
+		requires (std::is_base_of_v<BaseStructT, std::decay_t<T>>)
 	TConstSharedStruct(const TConstSharedStruct<T>& Other)
 		: ConstSharedStruct(Other.ConstSharedStruct)
-	{}
+	{
+	}
 
-	template<typename T = BaseStructT, typename = std::enable_if_t<std::is_base_of_v<BaseStructT, std::decay_t<T>>>>
+	template<typename T = BaseStructT>
+		requires (std::is_base_of_v<BaseStructT, std::decay_t<T>>)
 	TConstSharedStruct(const TSharedStruct<T>& ConstSharedStruct)
 		: ConstSharedStruct(ConstSharedStruct.SharedStruct)
-	{}
+	{
+	}
 
-	template<typename T = BaseStructT, typename = std::enable_if_t<std::is_base_of_v<BaseStructT, std::decay_t<T>>>>
+	template<typename T = BaseStructT>
+		requires (std::is_base_of_v<BaseStructT, std::decay_t<T>>)
 	TConstSharedStruct(TConstSharedStruct<T>&& Other)
 		: ConstSharedStruct(MoveTemp(Other.ConstSharedStruct))
-	{}
+	{
+	}
 
-	template<typename T = BaseStructT, typename = std::enable_if_t<std::is_base_of_v<BaseStructT, std::decay_t<T>>>>
+	template<typename T = BaseStructT>
+		requires (std::is_base_of_v<BaseStructT, std::decay_t<T>>)
 	TConstSharedStruct(TSharedStruct<T>&& ConstSharedStruct)
 		: ConstSharedStruct(MoveTemp(ConstSharedStruct.SharedStruct))
-	{}
+	{
+	}
 
-	template<typename T = BaseStructT, typename = std::enable_if_t<std::is_base_of_v<BaseStructT, std::decay_t<T>>>>
+	template<typename T = BaseStructT>
+		requires (std::is_base_of_v<BaseStructT, std::decay_t<T>>)
 	TConstSharedStruct& operator=(const TConstSharedStruct<T>& Other)
 	{
 		if (this == &Other)
@@ -755,7 +779,8 @@ struct [[nodiscard]] TConstSharedStruct
 		return *this;
 	}
 
-	template<typename T = BaseStructT, typename = std::enable_if_t<std::is_base_of_v<BaseStructT, std::decay_t<T>>>>
+	template<typename T = BaseStructT>
+		requires (std::is_base_of_v<BaseStructT, std::decay_t<T>>)
 	TConstSharedStruct& operator=(const TSharedStruct<T>& SharedStruct)
 	{
 		if (this != &SharedStruct)
@@ -784,28 +809,30 @@ struct [[nodiscard]] TConstSharedStruct
 	}
 
 	/** Initializes from templated struct instance. This will create a new instance of the shared struct memory. */
-	template<typename T = BaseStructT,
-		typename = std::enable_if_t<
+	template<typename T = BaseStructT>
+		requires(
 			std::is_base_of_v<BaseStructT, std::decay_t<T>>
 			&& !(std::is_same_v<TStructView<T>, T>
 				|| std::is_same_v<TConstStructView<T>, T>
 				|| std::is_same_v<TSharedStruct<T>, T>
 				|| std::is_same_v<TConstSharedStruct<T>, T>
-				|| std::is_same_v<TInstancedStruct<T>, T>), T>>
+				|| std::is_same_v<TInstancedStruct<T>, T>))
 	void Initialize(const T& Struct)
 	{
 		ConstSharedStruct.InitializeAs<T>(Struct);
 	}
 
 	/** Initializes from struct type and emplace args. This will create a new instance of the shared struct memory. */
-	template<typename T = BaseStructT, typename... TArgs, typename = std::enable_if_t<std::is_base_of_v<BaseStructT, std::decay_t<T>>>>
+	template<typename T = BaseStructT, typename... TArgs>
+		requires (std::is_base_of_v<BaseStructT, std::decay_t<T>>)
 	void Initialize(TArgs&&... InArgs)
 	{
 		ConstSharedStruct.InitializeAs<T>(Forward<TArgs>(InArgs)...);
 	}
 
 	/** Creates a new TSharedStruct from a templated struct type. This will create a new instance of the shared struct memory. */
-	template<typename T = BaseStructT, typename = std::enable_if_t<std::is_base_of_v<BaseStructT, std::decay_t<T>>>>
+	template<typename T = BaseStructT>
+		requires (std::is_base_of_v<BaseStructT, std::decay_t<T>>)
 	static TConstSharedStruct Make()
 	{
 		FConstSharedStruct SharedStruct;
@@ -814,7 +841,8 @@ struct [[nodiscard]] TConstSharedStruct
 	}
 
 	/** Creates a new TSharedStruct from templated struct instance. This will create a new instance of the shared struct memory. */
-	template<typename T = BaseStructT, typename = std::enable_if_t<std::is_base_of_v<BaseStructT, std::decay_t<T>>>>
+	template<typename T = BaseStructT>
+		requires (std::is_base_of_v<BaseStructT, std::decay_t<T>>)
 	static TConstSharedStruct Make(const T& Struct)
 	{
 		FConstSharedStruct SharedStruct;
@@ -823,7 +851,8 @@ struct [[nodiscard]] TConstSharedStruct
 	}
 
 	/** Creates a new TSharedStruct from a struct type and optional data. This will create a new instance of the shared struct memory. */
-	template<typename T = BaseStructT, typename = std::enable_if_t<std::is_base_of_v<BaseStructT, std::decay_t<T>>>>
+	template<typename T = BaseStructT>
+		requires (std::is_base_of_v<BaseStructT, std::decay_t<T>>)
 	static TConstSharedStruct Make(const uint8* InStructMemory = nullptr)
 	{
 		TConstSharedStruct SharedStruct;
@@ -832,7 +861,8 @@ struct [[nodiscard]] TConstSharedStruct
 	}
 
 	/** Creates a new TSharedStruct from the templated type and forward all arguments to constructor. This will create a new instance of the shared struct memory. */
-	template<typename T = BaseStructT, typename... TArgs, typename = std::enable_if_t<std::is_base_of_v<BaseStructT, std::decay_t<T>>>>
+	template<typename T = BaseStructT, typename... TArgs>
+		requires (std::is_base_of_v<BaseStructT, std::decay_t<T>>)
 	static TConstSharedStruct Make(TArgs&&... InArgs)
 	{
 		TConstSharedStruct SharedStruct;
@@ -841,29 +871,33 @@ struct [[nodiscard]] TConstSharedStruct
 	}
 
 	/** Returns const reference to the struct, this getter assumes that all data is valid. */
-	template<typename T = BaseStructT, typename = std::enable_if_t<std::is_base_of_v<BaseStructT, std::decay_t<T>>>>
+	template<typename T = BaseStructT>
+		requires (std::is_base_of_v<BaseStructT, std::decay_t<T>>)
 	constexpr const T& Get() const
 	{
 		return ConstSharedStruct.Get<T>();
 	}
 
 	/** Returns const reference to the struct, this getter assumes that all data is valid. */
-	template<typename T = BaseStructT, typename = std::enable_if_t<std::is_base_of_v<BaseStructT, std::decay_t<T>>>>
-	constexpr typename TEnableIf<TIsConst<T>::Value, T&>::Type Get() const
+	template<typename T = BaseStructT>
+		requires (std::is_const_v<T> && std::is_base_of_v<BaseStructT, std::decay_t<T>>)
+	constexpr T& Get() const
 	{
 		return ConstSharedStruct.Get<T>();
 	}
 
 	/** Returns const pointer to the struct. */
-	template<typename T = BaseStructT, typename = std::enable_if_t<std::is_base_of_v<BaseStructT, std::decay_t<T>>>>
+	template<typename T = BaseStructT>
+		requires (std::is_base_of_v<BaseStructT, std::decay_t<T>>)
 	constexpr const T* GetPtr() const
 	{
 		return ConstSharedStruct.GetPtr<T>();
 	}
 
 	/** Returns const pointer to the struct, or nullptr if cast is not valid. */
-	template<typename T = BaseStructT, typename = std::enable_if_t<std::is_base_of_v<BaseStructT, std::decay_t<T>>>>
-	constexpr typename TEnableIf<TIsConst<T>::Value, T*>::Type GetPtr() const
+	template<typename T = BaseStructT>
+		requires (std::is_const_v<T> && std::is_base_of_v<BaseStructT, std::decay_t<T>>)
+	constexpr T* GetPtr() const
 	{
 		return ConstSharedStruct.GetPtr<T>();
 	}
@@ -875,13 +909,15 @@ struct [[nodiscard]] TConstSharedStruct
 	}
 
 	/** Comparison operators. Note: it does not compare the internal structure itself */
-	template<typename T = BaseStructT, typename = std::enable_if_t<std::is_base_of_v<BaseStructT, std::decay_t<T>>>>
+	template<typename T = BaseStructT>
+		requires (std::is_base_of_v<BaseStructT, std::decay_t<T>>)
 	bool operator==(const T& Other) const
 	{
 		return ConstSharedStruct.operator==(Other);
 	}
 
-	template<typename T = BaseStructT, typename = std::enable_if_t<std::is_base_of_v<BaseStructT, std::decay_t<T>>>>
+	template<typename T = BaseStructT>
+		requires (std::is_base_of_v<BaseStructT, std::decay_t<T>>)
 	bool operator!=(const T& Other) const
 	{
 		return ConstSharedStruct.operator!=(Other);
diff --git a/Engine/Source/Runtime/CoreUObject/Public/StructUtils/StructArrayView.h b/Engine/Source/Runtime/CoreUObject/Public/StructUtils/StructArrayView.h
--- a/Engine/Source/Runtime/CoreUObject/Public/StructUtils/StructArrayView.h
+++ b/Engine/Source/Runtime/CoreUObject/Public/StructUtils/StructArrayView.h
@@ -5,6 +5,7 @@
 #include "UObject/Class.h"
 #include "StructUtils.h"
 #include "StructUtils/StructView.h"
+#include <type_traits>
 
 ///////////////////////////////////////////////////////////////// FStructArrayView /////////////////////////////////////////////////////////////////
 
@@ -21,7 +22,7 @@ struct FStructArrayView
 	template<typename T>
 	explicit FStructArrayView(TArray<T>& InArray)
 		: DataPtr(InArray.GetData())
-		, ScriptStruct(StaticStruct<typename TRemoveReference<T>::Type>())
+		, ScriptStruct(StaticStruct<std::remove_reference_t<T>>())
 		, ElementSize(sizeof(T))
 		, ArrayNum(InArray.Num())
 	{}
@@ -29,7 +30,7 @@ struct FStructArrayView
 	template<typename T>
 	explicit FStructArrayView(TArrayView<T> InArrayView)
 		: DataPtr(InArrayView.GetData())
-		, ScriptStruct(StaticStruct<typename TRemoveReference<T>::Type>())
+		, ScriptStruct(StaticStruct<std::remove_reference_t<T>>())
 		, ElementSize(sizeof(T))
 		, ArrayNum(InArrayView.Num())
 	{}
@@ -400,7 +401,7 @@ struct FConstStructArrayView
 	template<typename T>
 	explicit FConstStructArrayView(TArray<T>& InArray)
 		: DataPtr(InArray.GetData())
-		, ScriptStruct(StaticStruct<typename TRemoveReference<T>::Type>())
+		, ScriptStruct(StaticStruct<std::remove_reference_t<T>>())
 		, ElementSize(sizeof(T))
 		, ArrayNum(InArray.Num())
 	{}
@@ -408,7 +409,7 @@ struct FConstStructArrayView
 	template<typename T>
 	explicit FConstStructArrayView(TArrayView<T> InArrayView)
 		: DataPtr(InArrayView.GetData())
-		, ScriptStruct(StaticStruct<typename TRemoveReference<T>::Type>())
+		, ScriptStruct(StaticStruct<std::remove_reference_t<T>>())
 		, ElementSize(sizeof(T))
 		, ArrayNum(InArrayView.Num())
 	{}
@@ -552,7 +553,8 @@ struct FConstStructArrayView
 	 * @returns const pointer to indexed element.
 	 */
 	template<typename T>
-	constexpr typename TEnableIf<TIsConst<T>::Value, T*>::Type GetPtrAt(const int32 Index) const
+		requires (std::is_const_v<T>)
+	constexpr T* GetPtrAt(const int32 Index) const
 	{
 		return UE::StructUtils::GetStructPtr<T>(ScriptStruct, GetDataAt(Index)); // GetDataAt() calls RangeCheck().
 	}
@@ -565,7 +567,8 @@ struct FConstStructArrayView
 	 * @returns const reference to indexed element.
 	 */
 	template<typename T>
-	constexpr typename TEnableIf<TIsConst<T>::Value, T&>::Type GetAt(const int32 Index) const
+		requires (std::is_const_v<T>)
+	constexpr T& GetAt(const int32 Index) const
 	{
 		return UE::StructUtils::GetStructRef<T>(ScriptStruct, GetDataAt(Index)); // GetDataAt() calls RangeCheck().
 	}
@@ -589,7 +592,8 @@ struct FConstStructArrayView
 	 * @returns const reference to n-th last element from the array.
 	 */
 	template<typename T>
-	constexpr typename TEnableIf<TIsConst<T>::Value, T&>::Type Last(int32 IndexFromTheEnd = 0) const
+		requires (std::is_const_v<T>)
+	constexpr T& Last(int32 IndexFromTheEnd = 0) const
 	{
 		return GetAt<T>(ArrayNum - IndexFromTheEnd - 1);
 	}
diff --git a/Engine/Source/Runtime/CoreUObject/Public/StructUtils/StructUtilsTypes.h b/Engine/Source/Runtime/CoreUObject/Public/StructUtils/StructUtilsTypes.h
--- a/Engine/Source/Runtime/CoreUObject/Public/StructUtils/StructUtilsTypes.h
+++ b/Engine/Source/Runtime/CoreUObject/Public/StructUtils/StructUtilsTypes.h
@@ -53,30 +53,24 @@ namespace UE::StructUtils
 	extern UE_API uint64 GetStructHash64(const FConstSharedStruct& SharedView);
 
 	template <typename T>
-	typename TEnableIf<!TIsDerivedFrom<T, UObject>::IsDerived, UScriptStruct*>::Type GetAsUStruct()
+	auto* GetAsUStruct()
 	{
-		return T::StaticStruct();
-	}
-
-	template <typename T>
-	typename TEnableIf<TIsDerivedFrom<T, UObject>::IsDerived, UClass*>::Type GetAsUStruct()
-	{
-		return T::StaticClass();
+		if constexpr (std::is_base_of_v<UObject, T>)
+		{
+			return T::StaticClass();
+		}
+		else
+		{
+			return T::StaticStruct();
+		}
 	}
 
 	template<typename T>
-	using EnableIfSharedInstancedOrViewStruct = std::enable_if_t<std::is_same<FStructView, T>::value
-		|| std::is_same<FConstStructView, T>::value
-		|| std::is_same<FSharedStruct, T>::value
-		|| std::is_same<FConstSharedStruct, T>::value
-		|| std::is_same<FInstancedStruct, T>::value, T>;
-
-	template<typename T>
-	using EnableIfNotSharedInstancedOrViewStruct = std::enable_if_t<!(std::is_same<FStructView, T>::value
-		|| std::is_same<FConstStructView, T>::value
-		|| std::is_same<FSharedStruct, T>::value
-		|| std::is_same<FConstSharedStruct, T>::value
-		|| std::is_same<FInstancedStruct, T>::value), T>;
+	inline constexpr bool TIsSharedInstancedOrViewStruct_V = std::is_same_v<FStructView, T>
+		|| std::is_same_v<FConstStructView, T>
+		|| std::is_same_v<FSharedStruct, T>
+		|| std::is_same_v<FConstSharedStruct, T>
+		|| std::is_same_v<FInstancedStruct, T>;
 }
 
 /* Predicate useful to find a struct of a specific type in an container */
@@ -84,13 +78,24 @@ struct FStructTypeEqualOperator
 {
 	const UScriptStruct* TypePtr;
 
-	FStructTypeEqualOperator(const UScriptStruct* InTypePtr) : TypePtr(InTypePtr) {}
+	FStructTypeEqualOperator(const UScriptStruct* InTypePtr)
+		: TypePtr(InTypePtr)
+	{
+	}
 
-	template <typename T, UE::StructUtils::EnableIfSharedInstancedOrViewStruct<T>* = nullptr>
-	FStructTypeEqualOperator(const T& Struct) : TypePtr(Struct.GetScriptStruct()) {}
+	template <typename T>
+		requires (UE::StructUtils::TIsSharedInstancedOrViewStruct_V<T>)
+	FStructTypeEqualOperator(const T& Struct)
+		: TypePtr(Struct.GetScriptStruct())
+	{
+	}
 
-	template <typename T, UE::StructUtils::EnableIfSharedInstancedOrViewStruct<T>* = nullptr>
-	bool operator()(const T& Struct) const { return Struct.GetScriptStruct() == TypePtr; }
+	template <typename T>
+		requires (UE::StructUtils::TIsSharedInstancedOrViewStruct_V<T>)
+	bool operator()(const T& Struct) const
+	{
+		return Struct.GetScriptStruct() == TypePtr;
+	}
 };
 
 struct FScriptStructSortOperator
@@ -105,7 +110,8 @@ struct FScriptStructSortOperator
 
 struct FStructTypeSortOperator
 {
-	template <typename T, UE::StructUtils::EnableIfSharedInstancedOrViewStruct<T>* = nullptr>
+	template <typename T>
+		requires (UE::StructUtils::TIsSharedInstancedOrViewStruct_V<T>)
 	bool operator()(const T& A, const T& B) const
 	{
 		const UScriptStruct* AScriptStruct = A.GetScriptStruct();
diff --git a/Engine/Source/Runtime/CoreUObject/Public/StructUtils/StructView.h b/Engine/Source/Runtime/CoreUObject/Public/StructUtils/StructView.h
--- a/Engine/Source/Runtime/CoreUObject/Public/StructUtils/StructView.h
+++ b/Engine/Source/Runtime/CoreUObject/Public/StructUtils/StructView.h
@@ -5,6 +5,7 @@
 #include "StructUtils/InstancedStruct.h"
 #include "StructUtils/SharedStruct.h"
 #include "StructUtils.h"
+#include <type_traits>
 
 ///////////////////////////////////////////////////////////////// FStructView /////////////////////////////////////////////////////////////////
 
@@ -119,35 +120,45 @@ public:
 
 	explicit TStructView() = default;
 
-	template<typename T = BaseStructT, typename = std::enable_if_t<std::is_base_of_v<BaseStructT, std::decay_t<T>>>>
+	template<typename T = BaseStructT>
+		requires (std::is_base_of_v<BaseStructT, std::decay_t<T>>)
 	TStructView(T& InStruct)
 		: FStructView(T::StaticStruct(), reinterpret_cast<uint8*>(&InStruct))
-	{}
+	{
+	}
 
-	template<typename T = BaseStructT, typename = std::enable_if_t<std::is_base_of_v<BaseStructT, std::decay_t<T>>>>
+	template<typename T = BaseStructT>
+		requires (std::is_base_of_v<BaseStructT, std::decay_t<T>>)
 	explicit TStructView(uint8* InStructMemory = nullptr)
 		: FStructView(T::StaticStruct(), InStructMemory)
-	{}
+	{
+	}
 
-	template<typename T = BaseStructT, typename = std::enable_if_t<std::is_base_of_v<BaseStructT, std::decay_t<T>>>>
+	template<typename T = BaseStructT>
+		requires (std::is_base_of_v<BaseStructT, std::decay_t<T>>)
 	TStructView(TInstancedStruct<T>& InstancedStruct)
 		: FStructView(InstancedStruct.GetScriptStruct(), InstancedStruct.GetMutableMemory())
-	{}
+	{
+	}
 
-	template<typename T = BaseStructT, typename = std::enable_if_t<std::is_base_of_v<BaseStructT, std::decay_t<T>>>>
+	template<typename T = BaseStructT>
+		requires (std::is_base_of_v<BaseStructT, std::decay_t<T>>)
 	TStructView(const TSharedStruct<T>& SharedStruct)
 		: FStructView(SharedStruct.GetScriptStruct(), SharedStruct.GetMemory())
-	{}
+	{
+	}
 
 	/** Returns mutable reference to the struct, this getter assumes that all data is valid. */
-	template<typename T = BaseStructT, typename = std::enable_if_t<std::is_base_of_v<BaseStructT, std::decay_t<T>>>>
+	template<typename T = BaseStructT>
+		requires (std::is_base_of_v<BaseStructT, std::decay_t<T>>)
 	T& Get() const
 	{
 		return UE::StructUtils::GetStructRef<T>(ScriptStruct, StructMemory);
 	}
 
 	/** Returns mutable pointer to the struct, or nullptr if cast is not valid. */
-	template<typename T = BaseStructT, typename = std::enable_if_t<std::is_base_of_v<BaseStructT, std::decay_t<T>>>>
+	template<typename T = BaseStructT>
+		requires (std::is_base_of_v<BaseStructT, std::decay_t<T>>)
 	T* GetPtr() const
 	{
 		return UE::StructUtils::GetStructPtr<T, BaseStructT>(ScriptStruct, StructMemory);
@@ -161,13 +172,15 @@ public:
 	}
 
 	/** Comparison operators. Note: it does not compare the internal structure itself */
-	template<typename T = BaseStructT, typename = std::enable_if_t<std::is_base_of_v<BaseStructT, std::decay_t<T>>>>
+	template<typename T = BaseStructT>
+		requires (std::is_base_of_v<BaseStructT, std::decay_t<T>>)
 	bool operator==(const TStructView<T>& Other) const
 	{
 		return ((ScriptStruct == Other.GetScriptStruct()) && (StructMemory == Other.GetMemory()));
 	}
 
-	template<typename T = BaseStructT, typename = std::enable_if_t<std::is_base_of_v<BaseStructT, std::decay_t<T>>>>
+	template<typename T = BaseStructT>
+		requires (std::is_base_of_v<BaseStructT, std::decay_t<T>>)
 	bool operator!=(const TStructView<T>& Other) const
 	{
 		return !operator==(Other);
@@ -230,14 +243,16 @@ public:
 
 	/** Returns const reference to the struct, this getter assumes that all data is valid. */
 	template<typename T>
-	constexpr typename TEnableIf<TIsConst<T>::Value, T&>::Type Get() const
+		requires (std::is_const_v<T>)
+	constexpr T& Get() const
 	{
 		return UE::StructUtils::GetStructRef<T>(ScriptStruct, StructMemory);
 	}
 
 	/** Returns const pointer to the struct, or nullptr if cast is not valid. */
 	template<typename T>
-	constexpr typename TEnableIf<TIsConst<T>::Value, T*>::Type GetPtr() const
+		requires (std::is_const_v<T>)
+	constexpr T* GetPtr() const
 	{
 		return UE::StructUtils::GetStructPtr<T>(ScriptStruct, StructMemory);
 	}
@@ -305,45 +320,59 @@ struct [[nodiscard]] TConstStructView : FConstStructView
 {
 	explicit TConstStructView() = default;
 
-	template<typename T = BaseStructT, typename = std::enable_if_t<std::is_base_of_v<BaseStructT, std::decay_t<T>>>>
+	template<typename T = BaseStructT>
+		requires (std::is_base_of_v<BaseStructT, std::decay_t<T>>)
 	TConstStructView(const T& InStruct)
 		: FConstStructView(T::StaticStruct(), reinterpret_cast<const uint8*>(&InStruct))
-	{}
+	{
+	}
 
-	template<typename T = BaseStructT, typename = std::enable_if_t<std::is_base_of_v<BaseStructT, std::decay_t<T>>>>
+	template<typename T = BaseStructT>
+		requires (std::is_base_of_v<BaseStructT, std::decay_t<T>>)
 	explicit TConstStructView(const uint8* InStructMemory = nullptr)
 		: FConstStructView(T::StaticStruct(), InStructMemory)
-	{}
+	{
+	}
 
-	template<typename T = BaseStructT, typename = std::enable_if_t<std::is_base_of_v<BaseStructT, std::decay_t<T>>>>
+	template<typename T = BaseStructT>
+		requires (std::is_base_of_v<BaseStructT, std::decay_t<T>>)
 	TConstStructView(const TInstancedStruct<T>& InstancedStruct)
 		: FConstStructView(InstancedStruct.GetScriptStruct(), InstancedStruct.GetMemory())
-	{}
+	{
+	}
 
-	template<typename T = BaseStructT, typename = std::enable_if_t<std::is_base_of_v<BaseStructT, std::decay_t<T>>>>
+	template<typename T = BaseStructT>
+		requires (std::is_base_of_v<BaseStructT, std::decay_t<T>>)
 	TConstStructView(const TSharedStruct<T>& SharedStruct)
 		: FConstStructView(SharedStruct.GetScriptStruct(), SharedStruct.GetMemory())
-	{}
+	{
+	}
 
-	template<typename T = BaseStructT, typename = std::enable_if_t<std::is_base_of_v<BaseStructT, std::decay_t<T>>>>
+	template<typename T = BaseStructT>
+		requires (std::is_base_of_v<BaseStructT, std::decay_t<T>>)
 	TConstStructView(const TConstSharedStruct<T>& SharedStruct)
 		: FConstStructView(SharedStruct.GetScriptStruct(), SharedStruct.GetMemory())
-	{}
+	{
+	}
 
-	template<typename T = BaseStructT, typename = std::enable_if_t<std::is_base_of_v<BaseStructT, std::decay_t<T>>>>
+	template<typename T = BaseStructT>
+		requires (std::is_base_of_v<BaseStructT, std::decay_t<T>>)
 	TConstStructView(const TStructView<T> StructView)
 		: FConstStructView(StructView.GetScriptStruct(), StructView.GetMemory())
-	{}
+	{
+	}
 
 	/** Returns const reference to the struct, this getter assumes that all data is valid. */
-	template<typename T = BaseStructT, typename = std::enable_if_t<std::is_base_of_v<BaseStructT, std::decay_t<T>>>>
+	template<typename T = BaseStructT>
+		requires (std::is_base_of_v<BaseStructT, std::decay_t<T>>)
 	constexpr const T& Get() const
 	{
 		return UE::StructUtils::GetStructRef<T>(ScriptStruct, StructMemory);
 	}
 
 	/** Returns const pointer to the struct, or nullptr if cast is not valid. */
-	template<typename T = BaseStructT, typename = std::enable_if_t<std::is_base_of_v<BaseStructT, std::decay_t<T>>>>
+	template<typename T = BaseStructT>
+		requires (std::is_base_of_v<BaseStructT, std::decay_t<T>>)
 	constexpr const T* GetPtr() const
 	{
 		return UE::StructUtils::GetStructPtr<T, BaseStructT>(ScriptStruct, StructMemory);
@@ -357,13 +386,15 @@ struct [[nodiscard]] TConstStructView : FConstStructView
 	}
 
 	/** Comparison operators. Note: it does not compare the internal structure itself */
-	template<typename T = BaseStructT, typename = std::enable_if_t<std::is_base_of_v<BaseStructT, std::decay_t<T>>>>
+	template<typename T = BaseStructT>
+		requires (std::is_base_of_v<BaseStructT, std::decay_t<T>>)
 	bool operator==(const TConstStructView<T>& Other) const
 	{
 		return ((ScriptStruct == Other.GetScriptStruct()) && (StructMemory == Other.GetMemory()));
 	}
 
-	template<typename T = BaseStructT, typename = std::enable_if_t<std::is_base_of_v<BaseStructT, std::decay_t<T>>>>
+	template<typename T = BaseStructT>
+		requires (std::is_base_of_v<BaseStructT, std::decay_t<T>>)
 	bool operator!=(const TConstStructView<T>& Other) const
 	{
 		return !operator==(Other);
diff --git a/Engine/Source/Runtime/Engine/Public/WorldPartition/ActorDescContainerCollection.h b/Engine/Source/Runtime/Engine/Public/WorldPartition/ActorDescContainerCollection.h
--- a/Engine/Source/Runtime/Engine/Public/WorldPartition/ActorDescContainerCollection.h
+++ b/Engine/Source/Runtime/Engine/Public/WorldPartition/ActorDescContainerCollection.h
@@ -4,6 +4,7 @@
 #include "CoreMinimal.h"
 #include "UObject/Object.h"
 #include "WorldPartition/ActorDescContainer.h"
+#include <type_traits>
 
 class FWorldPartitionActorDesc;
 class FReferenceCollector;
@@ -47,17 +48,17 @@ public:
 	const FWorldPartitionActorDesc* GetActorDescByPath(const FSoftObjectPath& ActorPath) const;
 	const FWorldPartitionActorDesc* GetActorDescByName(FName ActorName) const;
 
-	template<typename Dummy = void, typename = typename TEnableIf<!TIsConst<TRemovePointer<ActorDescContPtrType>>::Value, Dummy>::Type>
-	bool RemoveActor(const FGuid& ActorGuid);
+	bool RemoveActor(const FGuid& ActorGuid)
+		requires (!std::is_const_v<std::remove_pointer_t<ActorDescContPtrType>>);
 
-	template<typename Dummy = void, typename = typename TEnableIf<!TIsConst<TRemovePointer<ActorDescContPtrType>>::Value, Dummy>::Type>
-	void OnPackageDeleted(UPackage* Package);
+	void OnPackageDeleted(UPackage* Package)
+		requires (!std::is_const_v<std::remove_pointer_t<ActorDescContPtrType>>);
 
-	template<typename Dummy = void, typename = typename TEnableIf<!TIsConst<TRemovePointer<ActorDescContPtrType>>::Value, Dummy>::Type>
-	bool IsActorDescHandled(const AActor* Actor) const;
+	bool IsActorDescHandled(const AActor* Actor) const
+		requires (!std::is_const_v<std::remove_pointer_t<ActorDescContPtrType>>);
 
-	template<typename Dummy = void, typename = typename TEnableIf<!TIsConst<TRemovePointer<ActorDescContPtrType>>::Value, Dummy>::Type>
-	void LoadAllActors(TArray<FWorldPartitionReference>& OutReferences);
+	void LoadAllActors(TArray<FWorldPartitionReference>& OutReferences)
+		requires (!std::is_const_v<std::remove_pointer_t<ActorDescContPtrType>>);
 
 	DECLARE_EVENT_OneParam(UWorldPartition, FActorDescAddedEvent, FWorldPartitionActorDesc*);
 	FActorDescAddedEvent OnActorDescAddedEvent;
@@ -382,8 +383,8 @@ const FWorldPartitionActorDesc* TActorDescContainerCollection<ActorDescContPtrTy
 }
 
 template<class ActorDescContPtrType>
-template<typename, typename>
 bool TActorDescContainerCollection<ActorDescContPtrType>::RemoveActor(const FGuid& ActorGuid)
+	requires (!std::is_const_v<std::remove_pointer_t<ActorDescContPtrType>>)
 {
 	bool bRemoved = false;
 	ForEachActorDescContainerBreakable([&ActorGuid, &bRemoved](ActorDescContPtrType ActorDescContainer)
@@ -396,8 +397,8 @@ bool TActorDescContainerCollection<ActorDescContPtrType>::RemoveActor(const FGui
 }
 
 template<class ActorDescContPtrType>
-template<typename, typename>
 void TActorDescContainerCollection<ActorDescContPtrType>::OnPackageDeleted(UPackage* Package)
+	requires (!std::is_const_v<std::remove_pointer_t<ActorDescContPtrType>>)
 {
 	ForEachActorDescContainer([Package](ActorDescContPtrType ActorDescContainer)
 	{
@@ -406,8 +407,8 @@ void TActorDescContainerCollection<ActorDescContPtrType>::OnPackageDeleted(UPack
 }
 
 template<class ActorDescContPtrType>
-template<typename, typename>
 bool TActorDescContainerCollection<ActorDescContPtrType>::IsActorDescHandled(const AActor* Actor) const
+	requires (!std::is_const_v<std::remove_pointer_t<ActorDescContPtrType>>)
 {
 	bool bIsHandled = false;
 	ForEachActorDescContainerBreakable([Actor, &bIsHandled](ActorDescContPtrType ActorDescContainer)
@@ -420,8 +421,8 @@ bool TActorDescContainerCollection<ActorDescContPtrType>::IsActorDescHandled(con
 }
 
 template<class ActorDescContPtrType>
-template<typename, typename>
 void TActorDescContainerCollection<ActorDescContPtrType>::LoadAllActors(TArray<FWorldPartitionReference>& OutReferences)
+	requires (!std::is_const_v<std::remove_pointer_t<ActorDescContPtrType>>)
 {
 	ForEachActorDescContainer([&OutReferences](ActorDescContPtrType ActorDescContainer)
 	{
diff --git a/Engine/Source/Runtime/Engine/Public/WorldPartition/ActorDescContainerInstanceCollection.h b/Engine/Source/Runtime/Engine/Public/WorldPartition/ActorDescContainerInstanceCollection.h
--- a/Engine/Source/Runtime/Engine/Public/WorldPartition/ActorDescContainerInstanceCollection.h
+++ b/Engine/Source/Runtime/Engine/Public/WorldPartition/ActorDescContainerInstanceCollection.h
@@ -6,6 +6,7 @@
 #include "UObject/Object.h"
 #include "GameFramework/Actor.h"
 #include "WorldPartition/ActorDescContainerInstance.h"
+#include <type_traits>
 
 class UWorldPartition;
 
@@ -48,14 +49,14 @@ public:
 	ActorDescContPtrType GetActorDescContainerInstance(const FGuid& ActorGuid) const;
 	ActorDescContPtrType FindHandlingContainerInstance(const AActor* Actor) const;
 
-	template<typename Dummy = void, typename = typename TEnableIf<!TIsConst<TRemovePointer<ActorDescContPtrType>>::Value, Dummy>::Type>
-	bool RemoveActor(const FGuid& ActorGuid);
+	bool RemoveActor(const FGuid& ActorGuid)
+		requires (!std::is_const_v<std::remove_pointer_t<ActorDescContPtrType>>);
 
-	template<typename Dummy = void, typename = typename TEnableIf<!TIsConst<TRemovePointer<ActorDescContPtrType>>::Value, Dummy>::Type>
-	void OnPackageDeleted(UPackage* Package);
+	void OnPackageDeleted(UPackage* Package)
+		requires (!std::is_const_v<std::remove_pointer_t<ActorDescContPtrType>>);
 
-	template<typename Dummy = void, typename = typename TEnableIf<!TIsConst<TRemovePointer<ActorDescContPtrType>>::Value, Dummy>::Type>
-	void LoadAllActors(TArray<FWorldPartitionReference>& OutReferences);
+	void LoadAllActors(TArray<FWorldPartitionReference>& OutReferences)
+		requires (!std::is_const_v<std::remove_pointer_t<ActorDescContPtrType>>);
 
 	DECLARE_EVENT_OneParam(UWorldPartition, FActorDescInstanceAddedEvent, FWorldPartitionActorDescInstance*);
 	FActorDescInstanceAddedEvent OnActorDescInstanceAddedEvent;
@@ -417,8 +418,8 @@ ActorDescContPtrType TActorDescContainerInstanceCollection<ActorDescContPtrType>
 }
 
 template<class ActorDescContPtrType>
-template<typename, typename>
 bool TActorDescContainerInstanceCollection<ActorDescContPtrType>::RemoveActor(const FGuid& ActorGuid)
+	requires (!std::is_const_v<std::remove_pointer_t<ActorDescContPtrType>>)
 {
 	bool bRemoved = false;
 	ForEachActorDescContainerInstanceBreakable([&ActorGuid, &bRemoved](ActorDescContPtrType ContainerInstance)
@@ -431,8 +432,8 @@ bool TActorDescContainerInstanceCollection<ActorDescContPtrType>::RemoveActor(co
 }
 
 template<class ActorDescContPtrType>
-template<typename, typename>
 void TActorDescContainerInstanceCollection<ActorDescContPtrType>::OnPackageDeleted(UPackage* Package)
+	requires (!std::is_const_v<std::remove_pointer_t<ActorDescContPtrType>>)
 {
 	ForEachActorDescContainerInstance([Package](ActorDescContPtrType ContainerInstance)
 	{
@@ -441,8 +442,8 @@ void TActorDescContainerInstanceCollection<ActorDescContPtrType>::OnPackageDelet
 }
 
 template<class ActorDescContPtrType>
-template<typename, typename>
 void TActorDescContainerInstanceCollection<ActorDescContPtrType>::LoadAllActors(TArray<FWorldPartitionReference>& OutReferences)
+	requires (!std::is_const_v<std::remove_pointer_t<ActorDescContPtrType>>)
 {
 	ForEachActorDescContainerInstance([&OutReferences](ActorDescContPtrType ContainerInstance)
 	{
diff --git a/Engine/Source/Runtime/MovieScene/Public/EntitySystem/MovieSceneComponentPtr.h b/Engine/Source/Runtime/MovieScene/Public/EntitySystem/MovieSceneComponentPtr.h
--- a/Engine/Source/Runtime/MovieScene/Public/EntitySystem/MovieSceneComponentPtr.h
+++ b/Engine/Source/Runtime/MovieScene/Public/EntitySystem/MovieSceneComponentPtr.h
@@ -3,6 +3,7 @@
 #pragma once
 
 #include "EntitySystem/MovieSceneEntitySystemTypes.h"
+#include <type_traits>
 
 namespace UE
 {
@@ -15,15 +16,17 @@ template<typename T>
 struct TComponentPtr
 {
 	/** Value type is either a T or const T& for read-only types, or T& for writeable types */
-	using ValueType = std::conditional_t<TIsConst<T>::Value, typename TCallTraits<T>::ParamType, T&>;
+	using ValueType = std::conditional_t<std::is_const_v<T>, typename TCallTraits<T>::ParamType, T&>;
 
 	TComponentPtr()
 		: ComponentPtr(nullptr)
-	{}
+	{
+	}
 
 	explicit TComponentPtr(T* InComponentPtr)
 		: ComponentPtr(InComponentPtr)
-	{}
+	{
+	}
 
 	TComponentPtr(const TComponentPtr& RHS) = default;
 	TComponentPtr& operator=(const TComponentPtr& RHS) = default;
diff --git a/Engine/Source/Runtime/SlateCore/Private/FastUpdate/SlateInvalidationWidgetList.cpp b/Engine/Source/Runtime/SlateCore/Private/FastUpdate/SlateInvalidationWidgetList.cpp
--- a/Engine/Source/Runtime/SlateCore/Private/FastUpdate/SlateInvalidationWidgetList.cpp
+++ b/Engine/Source/Runtime/SlateCore/Private/FastUpdate/SlateInvalidationWidgetList.cpp
@@ -764,7 +764,7 @@ namespace InvalidationList
 	template<typename TSlateInvalidationWidgetList, typename Predicate>
 	void ForEachChildren(TSlateInvalidationWidgetList& Self, const typename TSlateInvalidationWidgetList::InvalidationWidgetType& InvalidationWidget, FSlateInvalidationWidgetIndex WidgetIndex, Predicate InPredicate)
 	{
-		using SlateInvalidationWidgetType = std::conditional_t<TIsConst<TSlateInvalidationWidgetList>::Value, const typename TSlateInvalidationWidgetList::InvalidationWidgetType, typename TSlateInvalidationWidgetList::InvalidationWidgetType>;
+		using SlateInvalidationWidgetType = std::conditional_t<std::is_const_v<TSlateInvalidationWidgetList>, const typename TSlateInvalidationWidgetList::InvalidationWidgetType, typename TSlateInvalidationWidgetList::InvalidationWidgetType>;
 		if (InvalidationWidget.LeafMostChildIndex != WidgetIndex)
 		{
 			FSlateInvalidationWidgetIndex CurrentWidgetIndex = Self.IncrementIndex(WidgetIndex);
diff --git a/Engine/Source/Runtime/TypedElementFramework/Public/Elements/Framework/TypedElementQueryBuilder.inl b/Engine/Source/Runtime/TypedElementFramework/Public/Elements/Framework/TypedElementQueryBuilder.inl
--- a/Engine/Source/Runtime/TypedElementFramework/Public/Elements/Framework/TypedElementQueryBuilder.inl
+++ b/Engine/Source/Runtime/TypedElementFramework/Public/Elements/Framework/TypedElementQueryBuilder.inl
@@ -3,7 +3,6 @@
 #include <tuple>
 #include <type_traits>
 #include <utility>
-#include "Templates/IsConst.h"
 #include "Templates/UnrealTypeTraits.h"
 #include "DataStorage/Queries/Description.h"
 #include "DataStorage/Queries/Types.h"
@@ -35,7 +34,7 @@ namespace UE::Editor::DataStorage::Queries
 		constexpr EQueryAccessType GetAccessType()
 		{
 			using BaseType = typename std::remove_reference_t<Type>;
-			if constexpr (TIsConst<BaseType>::Value)
+			if constexpr (std::is_const_v<BaseType>)
 			{
 				return EQueryAccessType::ReadOnly;
 			}
@@ -60,7 +59,7 @@ namespace UE::Editor::DataStorage::Queries
 			{
 				EnumAddFlags(Result, EQueryDependencyFlags::AlwaysRefresh);
 			}
-			if constexpr (TIsConst<BaseType>::Value)
+			if constexpr (std::is_const_v<BaseType>)
 			{
 				EnumAddFlags(Result, EQueryDependencyFlags::ReadOnly);
 			}
@@ -1306,3 +1305,7 @@ e.g. void(ISubqueryContext& Context, RowHandle Row, ColumnType0& ColumnA, const
 		};
 	}
 } // namespace UE::Editor::DataStorage::Queries
+
+#if UE_ENABLE_INCLUDE_ORDER_DEPRECATED_IN_5_7
+#include "Templates/IsConst.h"
+#endif
