Added 0 bytes sized bin to MB3 so all of them will be aliased and take 1 page instead of being served from 16 bytes bin each

diff --git a/Engine/Source/Runtime/Core/Private/HAL/MallocBinned3.cpp b/Engine/Source/Runtime/Core/Private/HAL/MallocBinned3.cpp
--- a/Engine/Source/Runtime/Core/Private/HAL/MallocBinned3.cpp
+++ b/Engine/Source/Runtime/Core/Private/HAL/MallocBinned3.cpp
@@ -114,7 +114,7 @@ struct FFreeBlock
 		, NextFreeBlockIndex(MAX_uint32)
 	{
 		check(InPoolIndex < MAX_uint8 && (InBinSize >> UE_MBC_MIN_SMALL_POOL_ALIGNMENT_SHIFT) <= MAX_uint16);
-		NumFreeBins = InBlockSize / InBinSize;
+		NumFreeBins = InBlockSize / FMath::Max(InBinSize, 16u);
 	}
 
 	FORCEINLINE uint32 GetNumFreeBins() const
@@ -325,7 +325,8 @@ struct FMallocBinned3::Private
 
 				uint32 OutBlockIndex;
 				void* BaseBlockPtr = Allocator.BlockPointerFromContainedPtr(Node, Allocator.SmallPoolTables[InPoolIndex].NumMemoryPagesPerBlock, OutBlockIndex);
-				const uint32 BinIndexWithinBlock = (((uint8*)Node) - ((uint8*)BaseBlockPtr)) / Allocator.SmallPoolTables[InPoolIndex].BinSize;
+				const uint32 BinSizeForDiv = FMath::Max(InBinSize, 16u);
+				const uint32 BinIndexWithinBlock = (((uint8*)Node) - ((uint8*)BaseBlockPtr)) / BinSizeForDiv;
 
 				FPoolInfoSmall* NodePoolBlock = Allocator.SmallPoolTables[InPoolIndex].PoolInfos[OutBlockIndex / Allocator.SmallPoolInfosPerPlatformPage];
 				if (!NodePoolBlock)
@@ -367,7 +368,7 @@ struct FMallocBinned3::Private
 					Allocator.Decommit(InPoolIndex, BaseBlockPtr, BlockSize);
 
 #if UE_MB3_ALLOCATOR_STATS
-					Table.TotalAllocatedBins -= BlockSize / InBinSize;
+					Table.TotalAllocatedBins -= BlockSize / BinSizeForDiv;
 					Table.TotalAllocatedMem -= BlockSize;
 					AllocatedOSSmallPoolMemory.fetch_sub(BlockSize, std::memory_order_relaxed);
 #endif
@@ -410,6 +411,11 @@ FMallocBinned3::FPoolInfoSmall* FMallocBinned3::PushNewPoolToFront(FMallocBinned
 		Table.UnusedAreaOffsetLow = EndOffset;
 	}
 	FFreeBlock* Free = new ((void*)FreePtr) FFreeBlock(Table.BlockSize, InBinSize, InPoolIndex);
+	if (InBinSize == 0)
+	{
+		// special case for 0 bytes sized bin - we want all allocations to be aliased, so we only allocate one page per all of them
+		Free->NumFreeBins = MAX_uint32;
+	}
 #if UE_MB3_ALLOCATOR_STATS
 	AllocatedOSSmallPoolMemory.fetch_add((int64)Table.BlockSize, std::memory_order_relaxed);
 	Table.TotalAllocatedBins += Free->NumFreeBins;
@@ -539,7 +545,7 @@ FMallocBinned3::FMallocBinned3(const FPlatformMemory::FPlatformVirtualMemoryBloc
 
 	verify(FSizeTableEntry::FillSizeTable(OsAllocationGranularity, SizeTable, UE_MB3_BASE_PAGE_SIZE, UE_MB3_MAX_SMALL_POOL_SIZE, UE_MB3_BASE_PAGE_SIZE) == UE_MB3_SMALL_POOL_COUNT);
 	checkf(SizeTable[UE_MB3_SMALL_POOL_COUNT - 1].BinSize == UE_MB3_MAX_SMALL_POOL_SIZE, TEXT("UE_MB3_MAX_SMALL_POOL_SIZE must be equal to the largest bin size"));
-	checkf(sizeof(FFreeBlock) <= SizeTable[0].BinSize, TEXT("Pool header must be able to fit into the smallest bin"));
+	checkf((SizeTable[0].BinSize == 0 && sizeof(FFreeBlock) <= SizeTable[1].BinSize) || sizeof(FFreeBlock) <= SizeTable[0].BinSize, TEXT("Pool header must be able to fit into the smallest bin"));
 
 	SmallPoolInfosPerPlatformPage = OsAllocationGranularity / sizeof(FPoolInfoSmall);
 
diff --git a/Engine/Source/Runtime/Core/Private/HAL/MallocBinnedCommon.cpp b/Engine/Source/Runtime/Core/Private/HAL/MallocBinnedCommon.cpp
--- a/Engine/Source/Runtime/Core/Private/HAL/MallocBinnedCommon.cpp
+++ b/Engine/Source/Runtime/Core/Private/HAL/MallocBinnedCommon.cpp
@@ -20,6 +20,7 @@ CSV_DEFINE_CATEGORY(MallocBinned, false);
 // They must be 16-byte aligned as well.
 static constexpr uint32 BinnedCommonSmallBinSizes4k[] =
 {
+	0, 
 	16, 32, 48, 64, 80, 96, 112, 128, 144, 160, 176, 192, // +16
 	224, 256, 288, 320, // +32
 	368,  // /11 ish
@@ -114,7 +115,7 @@ FSizeTableEntry::FSizeTableEntry(uint32 InBinSize, uint64 PlatformPageSize, uint
 		}
 		break;
 	}
-	check((PlatformPageSize * NumMemoryPagesPerBlock) / BinSize <= MAX_uint32);
+	check(BinSize == 0 || (PlatformPageSize * NumMemoryPagesPerBlock) / BinSize <= MAX_uint32);
 }
 
 static inline void FillTable(FSizeTableEntry* SizeTable, int32& Index, const uint32* BinsList, uint32 BinListSize, uint32 Num4kbPages, uint64 PlatformPageSize, uint32 BasePageSize)
diff --git a/Engine/Source/Runtime/Core/Public/HAL/MallocBinnedCommon.h b/Engine/Source/Runtime/Core/Public/HAL/MallocBinnedCommon.h
--- a/Engine/Source/Runtime/Core/Public/HAL/MallocBinnedCommon.h
+++ b/Engine/Source/Runtime/Core/Public/HAL/MallocBinnedCommon.h
@@ -32,7 +32,7 @@
 #endif
 
 #if !defined(UE_MBC_NUM_LISTED_SMALL_POOLS)
-#	define UE_MBC_NUM_LISTED_SMALL_POOLS				51
+#	define UE_MBC_NUM_LISTED_SMALL_POOLS				52
 #endif
 
 #if !defined(BINNEDCOMMON_USE_SEPARATE_VM_PER_POOL)
