Logging: Fixed handling of SetColor

#rnx

diff --git a/Engine/Source/Runtime/Core/Private/Logging/StructuredLog.cpp b/Engine/Source/Runtime/Core/Private/Logging/StructuredLog.cpp
--- a/Engine/Source/Runtime/Core/Private/Logging/StructuredLog.cpp
+++ b/Engine/Source/Runtime/Core/Private/Logging/StructuredLog.cpp
@@ -1690,13 +1690,13 @@ FORCENOINLINE void BasicLog(FStaticBasicLogRecordParam Log, const FLogCategoryBa
 #endif
 }
 
-template CORE_API void BasicLog<ELogVerbosity::NoLogging>(FStaticBasicLogRecordParam Log, const FLogCategoryBase* Category, ...);
 template CORE_API void BasicLog<ELogVerbosity::Error>(FStaticBasicLogRecordParam Log, const FLogCategoryBase* Category, ...);
 template CORE_API void BasicLog<ELogVerbosity::Warning>(FStaticBasicLogRecordParam Log, const FLogCategoryBase* Category, ...);
 template CORE_API void BasicLog<ELogVerbosity::Display>(FStaticBasicLogRecordParam Log, const FLogCategoryBase* Category, ...);
 template CORE_API void BasicLog<ELogVerbosity::Log>(FStaticBasicLogRecordParam Log, const FLogCategoryBase* Category, ...);
 template CORE_API void BasicLog<ELogVerbosity::Verbose>(FStaticBasicLogRecordParam Log, const FLogCategoryBase* Category, ...);
 template CORE_API void BasicLog<ELogVerbosity::VeryVerbose>(FStaticBasicLogRecordParam Log, const FLogCategoryBase* Category, ...);
+template CORE_API void BasicLog<ELogVerbosity::SetColor>(FStaticBasicLogRecordParam Log, const FLogCategoryBase* Category, ...);
 
 UE_AUTORTFM_ALWAYS_OPEN
 static void BasicFatalLogV(FStaticBasicLogRecordParam Log, void* ProgramCounter, va_list Args)
diff --git a/Engine/Source/Runtime/Core/Public/Logging/LogMacros.h b/Engine/Source/Runtime/Core/Public/Logging/LogMacros.h
--- a/Engine/Source/Runtime/Core/Public/Logging/LogMacros.h
+++ b/Engine/Source/Runtime/Core/Public/Logging/LogMacros.h
@@ -323,7 +323,7 @@ CORE_API void BasicFatalLogWithProgramCounter(FStaticBasicLogRecordParam Log, co
 				{ \
 					Condition \
 					{ \
-						::UE::Logging::Private::BasicLog<ELogVerbosity::Type(::ELogVerbosity::Verbosity & ::ELogVerbosity::VerbosityMask)>(LOG_Static, &Category, ##__VA_ARGS__); \
+						::UE::Logging::Private::BasicLog<::ELogVerbosity::Verbosity>(LOG_Static, &Category, ##__VA_ARGS__); \
 					} \
 				} \
 			} \
