Made TBox's InverseTransformBy and TransformProjectBy both check for box validity before transforming (matching TransformBy)

diff --git a/Engine/Source/Runtime/Core/Public/Math/Box.h b/Engine/Source/Runtime/Core/Public/Math/Box.h
--- a/Engine/Source/Runtime/Core/Public/Math/Box.h
+++ b/Engine/Source/Runtime/Core/Public/Math/Box.h
@@ -936,6 +936,12 @@ void TBox<T>::GetVertices(TVector<T>(&Vertices)[8]) const
 template<typename T>
 TBox<T> TBox<T>::InverseTransformBy(const TTransform<T>& M) const
 {
+	// if we are not valid, return another invalid box.
+	if (!IsValid)
+	{
+		return TBox<T>(ForceInit);
+	}
+
 	TVector<T> Vertices[8];
 	GetVertices(Vertices);
 
@@ -953,6 +959,12 @@ TBox<T> TBox<T>::InverseTransformBy(const TTransform<T>& M) const
 template<typename T>
 TBox<T> TBox<T>::TransformProjectBy(const TMatrix<T>& ProjM) const
 {
+	// if we are not valid, return another invalid box.
+	if (!IsValid)
+	{
+		return TBox<T>(ForceInit);
+	}
+
 	TVector<T> Vertices[8];
 	GetVertices(Vertices);
 
