<saved by Perforce>

#okforversepublic

diff --git a/Engine/Source/Runtime/Core/Public/Containers/Queue.h b/Engine/Source/Runtime/Core/Public/Containers/Queue.h
--- a/Engine/Source/Runtime/Core/Public/Containers/Queue.h
+++ b/Engine/Source/Runtime/Core/Public/Containers/Queue.h
@@ -30,13 +30,15 @@ enum class EQueueMode
  * Template for queues.
  *
  * This template implements an unbounded non-intrusive queue using a lock-free linked
- * list that stores copies of the queued items. The template can operate in two modes:
- * Multiple-producers single-consumer (MPSC) and Single-producer single-consumer (SPSC).
+ * list that stores copies of the queued items. The template can operate in three modes:
+ * Multiple-producers single-consumer (MPSC), Single-producer single-consumer (SPSC), and Single-threaded.
  *
- * The queue is thread-safe in both modes. The Dequeue() method ensures thread-safety by
+ * The queue is thread-safe in both MPSC and SPSC modes. The Dequeue() method ensures thread-safety by
  * writing it in a way that does not depend on possible instruction reordering on the CPU.
  * The Enqueue() method uses an atomic compare-and-swap in multiple-producers scenarios.
  *
+ * The queue is not thread-safe in single-threaded mode, as the name suggests.
+ *
  * @param T The type of items stored in the queue.
  * @param Mode The queue mode (single-producer, single-consumer by default).
  * @todo gmp: Implement node pooling.
