[CIS] Fix warning about identical sub-expressions in static_assert statement for log categories.

#rnx

diff --git a/Engine/Source/Runtime/Core/Public/Logging/LogMacros.h b/Engine/Source/Runtime/Core/Public/Logging/LogMacros.h
--- a/Engine/Source/Runtime/Core/Public/Logging/LogMacros.h
+++ b/Engine/Source/Runtime/Core/Public/Logging/LogMacros.h
@@ -358,9 +358,11 @@ CORE_API void BasicFatalLogWithProgramCounter(FStaticBasicLogRecordParam Log, co
 	 * @param DefaultVerbosity, default run time verbosity
 	 * @param CompileTimeVerbosity, maximum verbosity to compile into the code
 	 **/
+	 //-V:DECLARE_LOG_CATEGORY_EXTERN:501
 	#define DECLARE_LOG_CATEGORY_EXTERN(CategoryName, DefaultVerbosity, CompileTimeVerbosity) \
 		extern struct FLogCategory##CategoryName : public FLogCategory<ELogVerbosity::DefaultVerbosity, ELogVerbosity::CompileTimeVerbosity> \
 		{ \
+			static_assert(ELogVerbosity::DefaultVerbosity <= ELogVerbosity::CompileTimeVerbosity, "Default verbosity of log category cannot be higher than its compile time verbosity"); \
 			UE_FORCEINLINE_HINT FLogCategory##CategoryName() : FLogCategory(TEXT(#CategoryName)) {} \
 		} CategoryName;
 
@@ -376,10 +378,11 @@ CORE_API void BasicFatalLogWithProgramCounter(FStaticBasicLogRecordParam Log, co
 	 * @param DefaultVerbosity, default run time verbosity
 	 * @param CompileTimeVerbosity, maximum verbosity to compile into the code
 	**/
+	//-V:DEFINE_LOG_CATEGORY_STATIC:501
 	#define DEFINE_LOG_CATEGORY_STATIC(CategoryName, DefaultVerbosity, CompileTimeVerbosity) \
-		static_assert(ELogVerbosity::DefaultVerbosity <= ELogVerbosity::CompileTimeVerbosity, "Default verbosity of log category cannot be higher than its compile time verbosity"); \
 		static struct FLogCategory##CategoryName : public FLogCategory<ELogVerbosity::DefaultVerbosity, ELogVerbosity::CompileTimeVerbosity> \
 		{ \
+			static_assert(ELogVerbosity::DefaultVerbosity <= ELogVerbosity::CompileTimeVerbosity, "Default verbosity of log category cannot be higher than its compile time verbosity"); \
 			UE_FORCEINLINE_HINT FLogCategory##CategoryName() : FLogCategory(TEXT(#CategoryName)) {} \
 		} CategoryName;
 
@@ -389,6 +392,7 @@ CORE_API void BasicFatalLogWithProgramCounter(FStaticBasicLogRecordParam Log, co
 	 * @param DefaultVerbosity, default run time verbosity
 	 * @param CompileTimeVerbosity, maximum verbosity to compile into the code
 	**/
+	//-V:DECLARE_LOG_CATEGORY_CLASS:501
 	#define DECLARE_LOG_CATEGORY_CLASS(CategoryName, DefaultVerbosity, CompileTimeVerbosity) \
 		DEFINE_LOG_CATEGORY_STATIC(CategoryName, DefaultVerbosity, CompileTimeVerbosity)
 
