Slate declarative syntax use Bind() for attributes instead of creating temporaries then assigning them.

Need some plumbing in TAttribute to support arguments too

diff --git a/Engine/Source/Runtime/Core/Public/Misc/Attribute.h b/Engine/Source/Runtime/Core/Public/Misc/Attribute.h
--- a/Engine/Source/Runtime/Core/Public/Misc/Attribute.h
+++ b/Engine/Source/Runtime/Core/Public/Misc/Attribute.h
@@ -299,10 +299,10 @@ public:
 	 * @param  InFuncPtr	Function to bind.  The function's structure (return value, arguments, etc) must match IBoundAttributeDelegate's definition.
 	 */
 	template < typename... VarTypes >
-	void BindStatic( TIdentity_T< typename FGetter::template TFuncPtr< VarTypes... > > InFuncPtr, VarTypes... Vars )
+	void BindStatic( TIdentity_T< typename FGetter::template TFuncPtr< std::decay_t<VarTypes>... > > InFuncPtr, VarTypes&&... Vars )
 	{
 		bIsSet = true;
-		Getter.BindStatic( InFuncPtr, Vars... );
+		Getter.BindStatic( InFuncPtr, Forward<VarTypes>(Vars)... );
 	}
 	
 	/**
@@ -313,11 +313,11 @@ public:
 	 * @param  InUserObject  Instance of the class that contains the member function you want to bind.
 	 * @param  InMethodPtr Member function to bind.  The function's structure (return value, arguments, etc) must match IBoundAttributeDelegate's definition.
 	 */
-	template< class SourceType >
-	void BindRaw( SourceType* InUserObject, typename FGetter::template TConstMethodPtr< SourceType > InMethodPtr )
+	template< class SourceType, typename... VarTypes >
+	void BindRaw( SourceType* InUserObject, typename FGetter::template TConstMethodPtr< SourceType, std::decay_t<VarTypes>... > InMethodPtr, VarTypes&&... Vars )
 	{
 		bIsSet = true;
-		Getter.BindRaw( InUserObject, InMethodPtr );
+		Getter.BindRaw( InUserObject, InMethodPtr, Forward<VarTypes>(Vars)... );
 	}
 
 	/**
@@ -328,11 +328,11 @@ public:
 	 * @param  InUserObject  Shared Pointer to the instance of the class that contains the member function you want to bind.  The attribute will only retain a weak pointer to this class.
 	 * @param  InMethodPtr Member function to bind.  The function's structure (return value, arguments, etc) must match IBoundAttributeDelegate's definition.
 	 */
-	template< class SourceType >
-	void Bind( TSharedRef< SourceType > InUserObject, typename FGetter::template TConstMethodPtr< SourceType > InMethodPtr )
+	template< class SourceType, typename... VarTypes >
+	void Bind( TSharedRef< SourceType > InUserObject, typename FGetter::template TConstMethodPtr< SourceType, std::decay_t<VarTypes>... > InMethodPtr, VarTypes&&... Vars )
 	{
 		bIsSet = true;
-		Getter.BindSP( InUserObject, InMethodPtr );
+		Getter.BindSP( InUserObject, InMethodPtr, Forward<VarTypes>(Vars)... );
 	}
 
 	/**
@@ -343,11 +343,11 @@ public:
 	 * @param  InUserObject  Shared Pointer to the instance of the class that contains the member function you want to bind.  The attribute will only retain a weak pointer to this class.
 	 * @param  InMethodPtr Member function to bind.  The function's structure (return value, arguments, etc) must match IBoundAttributeDelegate's definition.
 	 */
-	template< class SourceType >
-	void Bind( SourceType* InUserObject, typename FGetter::template TConstMethodPtr< SourceType > InMethodPtr )
+	template< class SourceType, typename... VarTypes >
+	void Bind( SourceType* InUserObject, typename FGetter::template TConstMethodPtr< SourceType, std::decay_t<VarTypes>... > InMethodPtr, VarTypes&&... Vars )
 	{
 		bIsSet = true;
-		Getter.BindSP( InUserObject, InMethodPtr );
+		Getter.BindSP( InUserObject, InMethodPtr, Forward<VarTypes>(Vars)... );
 	}
 
 	/**
@@ -358,11 +358,11 @@ public:
 	 * @param  InUserObject  Instance of the class that contains the member function you want to bind.
 	 * @param  InMethodPtr Member function to bind.  The function's structure (return value, arguments, etc) must match IBoundAttributeDelegate's definition.
 	 */
-	template< class SourceType >	
-	void BindUObject( SourceType* InUserObject, typename FGetter::template TConstMethodPtr< SourceType > InMethodPtr )
+	template< class SourceType, typename... VarTypes >	
+	void BindUObject( SourceType* InUserObject, typename FGetter::template TConstMethodPtr< SourceType, std::decay_t<VarTypes>... > InMethodPtr, VarTypes&&... Vars)
 	{
 		bIsSet = true;
-		Getter.BindUObject( InUserObject, InMethodPtr );
+		Getter.BindUObject( InUserObject, InMethodPtr, Forward<VarTypes>(Vars)... );
 	}
 
 	/**
@@ -576,4 +576,4 @@ template <typename UserClass, typename LambdaType, typename... PayloadTypes>
 	typedef decltype(InCallable(DeclVal<PayloadTypes>()...)) T;
 
 	return TAttribute<T>::Create(TAttribute<T>::FGetter::CreateWeakLambda(InUserObject, Forward<LambdaType>(InCallable), Forward<PayloadTypes>(InputPayload)...));
-}
\ No newline at end of file
+}
diff --git a/Engine/Source/Runtime/SlateCore/Public/Widgets/DeclarativeSyntaxSupport.h b/Engine/Source/Runtime/SlateCore/Public/Widgets/DeclarativeSyntaxSupport.h
--- a/Engine/Source/Runtime/SlateCore/Public/Widgets/DeclarativeSyntaxSupport.h
+++ b/Engine/Source/Runtime/SlateCore/Public/Widgets/DeclarativeSyntaxSupport.h
@@ -131,7 +131,7 @@ template<typename WidgetType> struct TSlateBaseNamedArgs;
 		template< typename... VarTypes > \
 		WidgetArgsType& AttrName##_Static( TIdentity_T< typename TAttribute< AttrType >::FGetter::template TFuncPtr<VarTypes...> > InFunc, VarTypes... Vars )	\
 		{ \
-			_##AttrName = TAttribute< AttrType >::Create( TAttribute< AttrType >::FGetter::CreateStatic( InFunc, Vars... ) ); \
+			_##AttrName.BindStatic( InFunc, Vars... ); \
 			return static_cast<WidgetArgsType*>(this)->Me(); \
 		} \
 		\
@@ -139,7 +139,7 @@ template<typename WidgetType> struct TSlateBaseNamedArgs;
 		 * technically this works for any functor types, but lambdas are the primary use case */ \
 		WidgetArgsType& AttrName##_Lambda(TFunction< AttrType(void) >&& InFunctor) \
 		{ \
-			_##AttrName = TAttribute< AttrType >::Create(Forward<TFunction< AttrType(void) >>(InFunctor)); \
+			_##AttrName.BindLambda(InFunctor); \
 			return static_cast<WidgetArgsType*>(this)->Me(); \
 		} \
 		\
@@ -147,7 +147,7 @@ template<typename WidgetType> struct TSlateBaseNamedArgs;
 		template< class UserClass, typename... VarTypes >	\
 		WidgetArgsType& AttrName##_Raw( UserClass* InUserObject, typename TAttribute< AttrType >::FGetter::template TConstMethodPtr< UserClass, VarTypes... > InFunc, VarTypes... Vars )	\
 		{ \
-			_##AttrName = TAttribute< AttrType >::Create( TAttribute< AttrType >::FGetter::CreateRaw( InUserObject, InFunc, Vars... ) ); \
+			_##AttrName.BindRaw( InUserObject, InFunc, Vars... ); \
 			return static_cast<WidgetArgsType*>(this)->Me(); \
 		} \
 		\
@@ -155,7 +155,7 @@ template<typename WidgetType> struct TSlateBaseNamedArgs;
 		template< class UserClass, typename... VarTypes >	\
 		WidgetArgsType& AttrName( TSharedRef< UserClass > InUserObjectRef, typename TAttribute< AttrType >::FGetter::template TConstMethodPtr< UserClass, VarTypes... > InFunc, VarTypes... Vars )	\
 		{ \
-			_##AttrName = TAttribute< AttrType >::Create( TAttribute< AttrType >::FGetter::CreateSP( InUserObjectRef, InFunc, Vars... ) ); \
+			_##AttrName.Bind( InUserObjectRef, InFunc, Vars... ); \
 			return static_cast<WidgetArgsType*>(this)->Me(); \
 		} \
 		\
@@ -163,7 +163,7 @@ template<typename WidgetType> struct TSlateBaseNamedArgs;
 		template< class UserClass, typename... VarTypes >	\
 		WidgetArgsType& AttrName( UserClass* InUserObject, typename TAttribute< AttrType >::FGetter::template TConstMethodPtr< UserClass, VarTypes... > InFunc, VarTypes... Vars )	\
 		{ \
-			_##AttrName = TAttribute< AttrType >::Create( TAttribute< AttrType >::FGetter::CreateSP( InUserObject, InFunc, Vars... ) ); \
+			_##AttrName.Bind( InUserObject, InFunc, Vars... ); \
 			return static_cast<WidgetArgsType*>(this)->Me(); \
 		} \
 		\
@@ -171,7 +171,7 @@ template<typename WidgetType> struct TSlateBaseNamedArgs;
 		template< class UserClass, typename... VarTypes >	\
 		WidgetArgsType& AttrName##_UObject( UserClass* InUserObject, typename TAttribute< AttrType >::FGetter::template TConstMethodPtr< UserClass, VarTypes... > InFunc, VarTypes... Vars )	\
 		{ \
-			_##AttrName = TAttribute< AttrType >::Create( TAttribute< AttrType >::FGetter::CreateUObject( InUserObject, InFunc, Vars... ) ); \
+			_##AttrName.BindUObject( InUserObject, InFunc, Vars... ); \
 			return static_cast<WidgetArgsType*>(this)->Me(); \
 		} \
 
