[CIS] Temporarily remove static_assert from DECLARE_LOG_CATEGORY_EXTERN; Partially revert 43829214.

#rnx

diff --git a/Engine/Source/Runtime/Core/Public/Logging/LogMacros.h b/Engine/Source/Runtime/Core/Public/Logging/LogMacros.h
--- a/Engine/Source/Runtime/Core/Public/Logging/LogMacros.h
+++ b/Engine/Source/Runtime/Core/Public/Logging/LogMacros.h
@@ -358,11 +358,9 @@ CORE_API void BasicFatalLogWithProgramCounter(FStaticBasicLogRecordParam Log, co
 	 * @param DefaultVerbosity, default run time verbosity
 	 * @param CompileTimeVerbosity, maximum verbosity to compile into the code
 	 **/
-	 //-V:DECLARE_LOG_CATEGORY_EXTERN:501
 	#define DECLARE_LOG_CATEGORY_EXTERN(CategoryName, DefaultVerbosity, CompileTimeVerbosity) \
 		extern struct FLogCategory##CategoryName : public FLogCategory<ELogVerbosity::DefaultVerbosity, ELogVerbosity::CompileTimeVerbosity> \
 		{ \
-			static_assert(ELogVerbosity::DefaultVerbosity <= ELogVerbosity::CompileTimeVerbosity, "Default verbosity of log category cannot be higher than its compile time verbosity"); \
 			UE_FORCEINLINE_HINT FLogCategory##CategoryName() : FLogCategory(TEXT(#CategoryName)) {} \
 		} CategoryName;
 
