Fix the build breakage caused by robomerge merging badly.

#rnx

diff --git a/Engine/Source/Runtime/Core/Public/HAL/Platform.h b/Engine/Source/Runtime/Core/Public/HAL/Platform.h
--- a/Engine/Source/Runtime/Core/Public/HAL/Platform.h
+++ b/Engine/Source/Runtime/Core/Public/HAL/Platform.h
@@ -696,7 +696,7 @@
 #endif
 #ifndef CALLSITE_FORCEINLINE
 	#define CALLSITE_FORCEINLINE 					/* Force code to be inline at the callsite */
-
+#endif
 #ifndef UE_DEFINE_FORCEINLINE_HINT_TO_INLINE
 	#define UE_DEFINE_FORCEINLINE_HINT_TO_INLINE 0
 #endif
