Generalized IOSAsyncTask to AppleAsyncTask in preparation for using WebKit in the macOS WebBrowser engine plugin

diff --git a/Engine/Source/Runtime/Core/Private/Apple/AppleAsyncTask.cpp b/Engine/Source/Runtime/Core/Private/Apple/AppleAsyncTask.cpp
new file mode 100644
--- /dev/null
+++ b/Engine/Source/Runtime/Core/Private/Apple/AppleAsyncTask.cpp
@@ -0,0 +1,136 @@
+// Copyright Epic Games, Inc. All Rights Reserved.
+
+#include "Apple/AppleAsyncTask.h"
+#include "HAL/PlatformAtomics.h"
+
+@implementation FAppleAsyncTask
+
+@synthesize UserData;
+@synthesize GameThreadCallback;
+
+/** All currently running tasks (which can be created on iOS thread or main thread) */
+NSMutableArray* RunningTasks;
+
+/**
+ * Static class constructor, called before any instances are created
+ */
++ (void)initialize
+{
+	// create the RunningTasks object one time
+	RunningTasks = [[NSMutableArray arrayWithCapacity:4] retain];
+}
+
+/**
+ * Initialize the async task
+ */
+- (id)init
+{
+	self = [super init];
+
+	// add ourself to the list of tasks
+	@synchronized(RunningTasks)
+	{
+		[RunningTasks addObject:self];
+	}
+
+	// return ourself, the constructed object
+	return self;
+}
+
+
++ (void)CreateTaskWithBlock:(bool (^)(void))Block
+{
+	// create a task, and add it to the array
+	FAppleAsyncTask* Task = [[FAppleAsyncTask alloc] init];
+	// set the callback
+	Task.GameThreadCallback = Block;
+	// safely tell the game thread we are ready to go
+	[Task FinishedTask];
+}
+
+- (void)FinishedTask
+{
+	FPlatformAtomics::InterlockedIncrement(&bIsReadyForGameThread);
+}
+
+/**
+ * Check for completion
+ *
+ * @return TRUE if we succeeded (the completion block will have been called)
+ */
+- (bool)CheckForCompletion
+{
+	// handle completion
+	if (bIsReadyForGameThread)
+	{
+		// call the game thread block
+		if (GameThreadCallback)
+		{
+			if (GameThreadCallback())
+			{
+				// only return true if the callback says it's complete
+				return true;
+			}
+		}
+		else
+		{
+			// if there isn't a callback, then just return TRUE to remove the 
+			// task from the queue
+			return true;
+		}
+		
+	}
+
+	// all other cases, we are not complete
+	return false;
+}
+
+/**
+ * Tick all currently running tasks
+ */
++ (void)ProcessAsyncTasks
+{
+	FAppleAsyncTask* CurrentTask = nil;
+
+	// grab one out of the queue to process outside of the lock
+	@synchronized(RunningTasks)
+	{
+		if ([RunningTasks count] > 0)
+		{
+			CurrentTask = [RunningTasks objectAtIndex:0];
+			[RunningTasks removeObjectAtIndex:0];
+		}
+	}
+
+	if (CurrentTask != nil)
+	{
+		if ([CurrentTask CheckForCompletion])
+		{
+			// release the object count
+			[CurrentTask release];
+		}
+		else
+		{
+			// if it's not done, but it back in the list
+			@synchronized(RunningTasks)
+			{
+				[RunningTasks addObject:CurrentTask];
+			}
+		}
+	}
+}
+
+
+/** 
+ * Application destructor
+ */
+- (void)dealloc 
+{
+	[GameThreadCallback release];
+	self.UserData = nil;
+
+	[super dealloc];
+}
+
+
+@end
diff --git a/Engine/Source/Runtime/Core/Private/IOS/IOSAsyncTask.cpp b/Engine/Source/Runtime/Core/Private/IOS/IOSAsyncTask.cpp
--- a/Engine/Source/Runtime/Core/Private/IOS/IOSAsyncTask.cpp
+++ b/Engine/Source/Runtime/Core/Private/IOS/IOSAsyncTask.cpp
@@ -1,136 +1,6 @@
 // Copyright Epic Games, Inc. All Rights Reserved.
 
 #include "IOS/IOSAsyncTask.h"
-#include "HAL/PlatformAtomics.h"
 
 @implementation FIOSAsyncTask
-
-@synthesize UserData;
-@synthesize GameThreadCallback;
-
-/** All currently running tasks (which can be created on iOS thread or main thread) */
-NSMutableArray* RunningTasks;
-
-/**
- * Static class constructor, called before any instances are created
- */
-+ (void)initialize
-{
-	// create the RunningTasks object one time
-	RunningTasks = [[NSMutableArray arrayWithCapacity:4] retain];
-}
-
-/**
- * Initialize the async task
- */
-- (id)init
-{
-	self = [super init];
-
-	// add ourself to the list of tasks
-	@synchronized(RunningTasks)
-	{
-		[RunningTasks addObject:self];
-	}
-
-	// return ourself, the constructed object
-	return self;
-}
-
-
-+ (void)CreateTaskWithBlock:(bool (^)(void))Block
-{
-	// create a task, and add it to the array
-	FIOSAsyncTask* Task = [[FIOSAsyncTask alloc] init];
-	// set the callback
-	Task.GameThreadCallback = Block;
-	// safely tell the game thread we are ready to go
-	[Task FinishedTask];
-}
-
-- (void)FinishedTask
-{
-	FPlatformAtomics::InterlockedIncrement(&bIsReadyForGameThread);
-}
-
-/**
- * Check for completion
- *
- * @return TRUE if we succeeded (the completion block will have been called)
- */
-- (bool)CheckForCompletion
-{
-	// handle completion
-	if (bIsReadyForGameThread)
-	{
-		// call the game thread block
-		if (GameThreadCallback)
-		{
-			if (GameThreadCallback())
-			{
-				// only return true if the callback says it's complete
-				return true;
-			}
-		}
-		else
-		{
-			// if there isn't a callback, then just return TRUE to remove the 
-			// task from the queue
-			return true;
-		}
-		
-	}
-
-	// all other cases, we are not complete
-	return false;
-}
-
-/**
- * Tick all currently running tasks
- */
-+ (void)ProcessAsyncTasks
-{
-	FIOSAsyncTask* CurrentTask = nil;
-
-	// grab one out of the queue to process outside of the lock
-	@synchronized(RunningTasks)
-	{
-		if ([RunningTasks count] > 0)
-		{
-			CurrentTask = [RunningTasks objectAtIndex:0];
-			[RunningTasks removeObjectAtIndex:0];
-		}
-	}
-
-	if (CurrentTask != nil)
-	{
-		if ([CurrentTask CheckForCompletion])
-		{
-			// release the object count
-			[CurrentTask release];
-		}
-		else
-		{
-			// if it's not done, but it back in the list
-			@synchronized(RunningTasks)
-			{
-				[RunningTasks addObject:CurrentTask];
-			}
-		}
-	}
-}
-
-
-/** 
- * Application destructor
- */
-- (void)dealloc 
-{
-	[GameThreadCallback release];
-	self.UserData = nil;
-
-	[super dealloc];
-}
-
-
 @end
diff --git a/Engine/Source/Runtime/Core/Public/Apple/AppleAsyncTask.h b/Engine/Source/Runtime/Core/Public/Apple/AppleAsyncTask.h
new file mode 100644
--- /dev/null
+++ b/Engine/Source/Runtime/Core/Public/Apple/AppleAsyncTask.h
@@ -0,0 +1,48 @@
+// Copyright Epic Games, Inc. All Rights Reserved.
+
+#pragma once
+
+#include "CoreTypes.h"
+
+#include "CoreTypes.h"
+#import <Foundation/Foundation.h>
+
+
+@interface FAppleAsyncTask : NSObject
+{
+@private
+	/** Whether or not the task is ready to have GameThread callback called (set on iOS thread) */
+	int32 bIsReadyForGameThread;
+}
+
+/** Extra data for this async task */
+@property (retain) id UserData;
+
+/**
+ * Code to run on the game thread when the async task completes 
+ * @return true when the task is complete, and it will be pulled from the list of tasks and destroyed
+ */
+@property (copy) bool (^GameThreadCallback)(void);
+
+
+/**
+ * Create an async task object, set the block, and mark it as ready to be processed on main thread
+ * For advanced uses (ie, setting UserData or delaying when it's ready for game thread, use the usual [[alloc] init] pattern
+ *
+ * Note, this doesn't return an AsyncObject so no one is tempted to use it - once it's ready for game thread, it's dangerous to touch it
+ * This is callable on any thread, not just main thread
+ */
++ (void)CreateTaskWithBlock:(bool (^)(void))Block;
+
+/**
+ * Mark that the task is complete on the iOS thread, and now the
+ * GameThread can be fired (the Task is unsafe to use after this call)
+ */
+- (void)FinishedTask;
+
+/**
+ * Tick all currently running tasks
+ */
++ (void)ProcessAsyncTasks;
+
+@end
diff --git a/Engine/Source/Runtime/Core/Public/IOS/IOSAsyncTask.h b/Engine/Source/Runtime/Core/Public/IOS/IOSAsyncTask.h
--- a/Engine/Source/Runtime/Core/Public/IOS/IOSAsyncTask.h
+++ b/Engine/Source/Runtime/Core/Public/IOS/IOSAsyncTask.h
@@ -2,47 +2,10 @@
 
 #pragma once
 
-#include "CoreTypes.h"
+#include "Apple/AppleAsyncTask.h"
 
-#include "CoreTypes.h"
-#import <Foundation/Foundation.h>
-
-
-@interface FIOSAsyncTask : NSObject
+UE_DEPRECATED(5.7, "Use the more general FAppleAsyncTask instead")
+@interface FIOSAsyncTask : FAppleAsyncTask
 {
-@private
-	/** Whether or not the task is ready to have GameThread callback called (set on iOS thread) */
-	int32 bIsReadyForGameThread;
 }
-
-/** Extra data for this async task */
-@property (retain) id UserData;
-
-/**
- * Code to run on the game thread when the async task completes 
- * @return true when the task is complete, and it will be pulled from the list of tasks and destroyed
- */
-@property (copy) bool (^GameThreadCallback)(void);
-
-
-/**
- * Create an async task object, set the block, and mark it as ready to be processed on main thread
- * For advanced uses (ie, setting UserData or delaying when it's ready for game thread, use the usual [[alloc] init] pattern
- *
- * Note, this doesn't return an AsyncObject so no one is tempted to use it - once it's ready for game thread, it's dangerous to touch it
- * This is callable on any thread, not just main thread
- */
-+ (void)CreateTaskWithBlock:(bool (^)(void))Block;
-
-/**
- * Mark that the task is complete on the iOS thread, and now the
- * GameThread can be fired (the Task is unsafe to use after this call)
- */
-- (void)FinishedTask;
-
-/**
- * Tick all currently running tasks
- */
-+ (void)ProcessAsyncTasks;
-
 @end
