Added a check for whether CrashReportClient should generate minidumps on ensures triggered when running from the Editor.

diff --git a/Engine/Source/Runtime/Core/Private/Windows/WindowsPlatformCrashContext.cpp b/Engine/Source/Runtime/Core/Private/Windows/WindowsPlatformCrashContext.cpp
--- a/Engine/Source/Runtime/Core/Private/Windows/WindowsPlatformCrashContext.cpp
+++ b/Engine/Source/Runtime/Core/Private/Windows/WindowsPlatformCrashContext.cpp
@@ -367,7 +367,12 @@ void FWindowsPlatformCrashContext::CopyPlatformSpecificFiles(const TCHAR* Output
 			GConfig->GetBool(TEXT("CrashReportClient"), TEXT("Stall.RecordDump"), bRecordCrashDump, GEngineIni);
 			break;
 		case ECrashContextType::Ensure:
-			GConfig->GetBool(TEXT("CrashReportClient"), TEXT("Ensure.RecordDump"), bRecordCrashDump, GEngineIni);
+#if UE_EDITOR
+			if (!GConfig->GetBool(TEXT("CrashReportClient"), TEXT("Ensure.RecordDump"), bRecordCrashDump, GEditorIni))
+#endif
+			{
+				GConfig->GetBool(TEXT("CrashReportClient"), TEXT("Ensure.RecordDump"), bRecordCrashDump, GEngineIni);
+			}
 			break;
 		case ECrashContextType::AbnormalShutdown:
 		case ECrashContextType::Assert:
