FPaths: Optimize CreateStandardizedPaths by removing some redundant transforms.
Update CollapseRelativeDirectories to allow collapsing all .. diirectories even if some uncollapsible leading .. are present.
Change ConvertRelativePathToFullInternal to take BasePath as StringView to avoid a copy.
#rnx

diff --git a/Engine/Source/Runtime/Core/Private/Misc/PathViews.cpp b/Engine/Source/Runtime/Core/Private/Misc/PathViews.cpp
--- a/Engine/Source/Runtime/Core/Private/Misc/PathViews.cpp
+++ b/Engine/Source/Runtime/Core/Private/Misc/PathViews.cpp
@@ -15,6 +15,7 @@
 #include "Misc/AssertionMacros.h"
 #include "Misc/CString.h"
 #include "Misc/Char.h"
+#include "Misc/PathViewsPrivate.h"
 #include "Misc/StringBuilder.h"
 #include "String/Find.h"
 #include "String/ParseTokens.h"
@@ -610,53 +611,194 @@ static void RemoveAll(FStringBuilderBase& String, FStringView Substring)
 	}
 }
 
-bool FPathViews::CollapseRelativeDirectories(FStringBuilderBase& InOutPath)
+namespace UE4PathViews_Private
+{
+
+void CollapseRemoveAt(FStringBuilderBase& Builder, int32 Index, int32 Count)
+{
+	Builder.RemoveAt(Index, Count);
+}
+
+void CollapseRemoveAt(FString& String, int32 Index, int32 Count)
+{
+	String.RemoveAt(Index, Count, EAllowShrinking::No);
+}
+
+template<typename WritableStringType>
+void CollapseRelativeDirectoriesImplTemplate(WritableStringType& InOutPath,
+	bool bCollapseAllPossible, bool& bOutAllCollapsed)
 {
-	while (InOutPath.Len())
+	constexpr FStringView DotDot(TEXTVIEW(".."));
+	constexpr FStringView SlashDot(TEXTVIEW("/."));
+	constexpr FStringView DotSlash(TEXTVIEW("./"));
+	constexpr FStringView SlashDotDot(TEXTVIEW("/.."));
+
+	bOutAllCollapsed = true;
+	int32 StartOfCollapsibleDirs = 0;
+
+	// If the path has a colon before the first slash, move past it
+	FStringView Path = FStringView(InOutPath);
+	int32 FirstSlashOrColonIndex = UE::String::FindFirstOfAnyChar(Path, { ':', '/' });
+	if (FirstSlashOrColonIndex != INDEX_NONE && Path[FirstSlashOrColonIndex] == ':')
+	{
+		StartOfCollapsibleDirs = FirstSlashOrColonIndex + 1;
+	}
+
+	// We look for extraneous '.' paths by searching for "/." That will not handle a string with
+	// a leading "./" or "D:./"; collapse that now.
+	FStringView PathAfterColon = Path.RightChop(StartOfCollapsibleDirs);
+	if (PathAfterColon.StartsWith(DotSlash, ESearchCase::CaseSensitive))
 	{
-		FStringView Path = InOutPath.ToView();
+		CollapseRemoveAt(InOutPath, StartOfCollapsibleDirs, DotSlash.Len());
+	}
+	// Searching for /. also incorrectly makes the string empty if the string is exactly "/.". Skip over it if so.
+	else if (PathAfterColon == SlashDot)
+	{
+		StartOfCollapsibleDirs += SlashDot.Len();
+	}
+
+	int32 NextScanIndex = StartOfCollapsibleDirs;
 
-		// Consider paths which start with .. or /.. as invalid
-		const FStringView ParentDir = TEXTVIEW("/..");
-		if (Path.StartsWith(TEXTVIEW("..")) || Path.StartsWith(ParentDir))
+	for (;;)
+	{
+		Path = FStringView(InOutPath);
+		// When NextScanIndex == Path.Len we are finished
+		if (NextScanIndex == Path.Len())
 		{
-			return false;
+			break;
 		}
-	
-		int32 Index = UE::String::FindFirst(Path, ParentDir);
 
-		// Ignore folders beginning with dots
-		while (Index != INDEX_NONE && Index + ParentDir.Len() < Path.Len() && Path[Index + ParentDir.Len()] != '/')
+		// If remaining text after StartOfCollapsibleDirs starts with .. mark that we could not convert it
+		// and move past it so we do not try to collapse it in the following .. 
+		if (NextScanIndex == StartOfCollapsibleDirs)
 		{
-			Index = FindNext(Path, Index + ParentDir.Len(), ParentDir);
+			FStringView CollapsibleText = Path.RightChop(StartOfCollapsibleDirs);
+			if (CollapsibleText.StartsWith(DotDot, ESearchCase::CaseSensitive) &&
+				// But don't mark that we could not convert it if it is a directory name that starts with ..,
+				// such as "..whodoesthis/".
+				(CollapsibleText.Len() == DotDot.Len() || CollapsibleText[DotDot.Len()] == '/'))
+			{
+				StartOfCollapsibleDirs += DotDot.Len();
+				bOutAllCollapsed = false;
+				if (!bCollapseAllPossible)
+				{
+					return;
+				}
+				continue;
+			}
 		}
 
-		// If there are no "/.."s left then we're done
+		// If there are no "/."s left then we're done
+		FStringView ScanText = Path.RightChop(NextScanIndex);
+		int32 Index = UE::String::FindFirst(ScanText, SlashDot, ESearchCase::CaseSensitive);
 		if (Index == INDEX_NONE)
 		{
 			break;
 		}
+		Index += NextScanIndex;
 
-		// Find previous directory, stop if we've found a directory that isn't "/./"
-		int32 PreviousSeparatorIndex = Index;
-		while (Path.Left(PreviousSeparatorIndex).FindLastChar(TEXT('/'), /* Out */ PreviousSeparatorIndex) &&
-				Path.Mid(PreviousSeparatorIndex + 1, 2) == TEXTVIEW("./"));
+		// If we have found /./, or a trailing /., remove the spurious /.
+		if (Index + SlashDot.Len() == Path.Len() || Path[Index + SlashDot.Len()] == '/')
+		{
+			CollapseRemoveAt(InOutPath, Index, SlashDot.Len());
+			continue;
+		}
 
-		PreviousSeparatorIndex = FMath::Max(0, PreviousSeparatorIndex);
-		
-		// If we're attempting to remove the drive letter, that's illegal
-		int32 RemoveLen = Index - PreviousSeparatorIndex + ParentDir.Len();
-		if (Contains(Path.Mid(PreviousSeparatorIndex, RemoveLen), TEXT(':')))
+		// If we have found a folder beginning with a dot (/.gitignore/) move past it
+		if (Path[Index + SlashDot.Len()] != '.')
 		{
-			return false;
+			NextScanIndex = Index + SlashDot.Len() + 1;
+			continue;
 		}
 
-		InOutPath.RemoveAt(PreviousSeparatorIndex, RemoveLen);
+		// If we have found a folder beginning with a pair of dots (/..whodoesthis/) move past it
+		if (Index + SlashDotDot.Len() < Path.Len() && Path[Index + SlashDotDot.Len()] != '/')
+		{
+			NextScanIndex = Index + SlashDotDot.Len() + 1;
+			continue;
+		}
+
+		// We have now found a substring starting at Index that is /../ or a trailing /..
+
+		// If our discovered /../ is at StartOfCollapsibleDirs we cannot convert it; move past it
+		if (Index == StartOfCollapsibleDirs)
+		{
+			StartOfCollapsibleDirs = Index + SlashDotDot.Len();
+			NextScanIndex = StartOfCollapsibleDirs;
+			bOutAllCollapsed = false;
+			if (!bCollapseAllPossible)
+			{
+				return;
+			}
+			continue;
+		}
+
+		// If there was no previous slash, handle the case "d:foo/../child" by removing "foo/../"
+		// If there was a previous slash, handle the case "d:/before/foo/../child" by removing "/foo/.."
+		// Find the previous slash before the slash at Index.
+		bool bHasLeadingSlash = true;
+		int32 PreviousSeparatorIndex = -1;
+		int32 RemovalStart = 0;
+		FStringView PreviousText = Path.Left(Index);
+		if (PreviousText.FindLastChar('/', PreviousSeparatorIndex))
+		{
+			RemovalStart = PreviousSeparatorIndex;
+		}
+		else
+		{
+			// If there is no previous slash in the string, then remove all text from start of string, after the
+			// initial colon. We use StartOfCollapsibleDirs as our first index after colon. We can rely on that,
+			// even though other causes move StartOfCollapsibleDirs, because all of those other causes occur only
+			// when there is a slash earlier in the string.
+			RemovalStart = StartOfCollapsibleDirs;
+			bHasLeadingSlash = false;
+		}
+
+		// The first index of InOutPath we are allowed to remove, given that we might have moved past some leading text,
+		// is StartOfCollapsibleDirs. If RemovalStart < StartOfCollapsibleDirs then it includes text we cannot convert.
+		if (RemovalStart < StartOfCollapsibleDirs)
+		{
+			StartOfCollapsibleDirs = Index + SlashDotDot.Len();
+			NextScanIndex = StartOfCollapsibleDirs;
+			bOutAllCollapsed = false;
+			if (!bCollapseAllPossible)
+			{
+				return;
+			}
+			continue;
+		}
+
+		// We now have a trailing "/foo/.." or "/foo/../", or "^foo/.." or "^foo/../" (^ == start of text)
+		// starting at RemovalStart, with the second slash at Index. Remove "/foo/.." or "foo/.." or "^foo/../"
+		int32 RemoveLength = Index - RemovalStart + SlashDotDot.Len();
+		if (!bHasLeadingSlash && RemovalStart + RemoveLength < Path.Len())
+		{
+			// If the removed text starts at beginning of path with no leading slash, also remove the slash after ..
+			++RemoveLength;
+		}
+		CollapseRemoveAt(InOutPath, RemovalStart, RemoveLength);
+		NextScanIndex = RemovalStart;
 	}
+}
 
-	RemoveAll(InOutPath, TEXTVIEW("./"));
+void CollapseRelativeDirectoriesImpl(FStringBuilderBase& InOutPath, bool bCollapseAllPossible, bool& bOutAllCollapsed)
+{
+	CollapseRelativeDirectoriesImplTemplate(InOutPath, bCollapseAllPossible, bOutAllCollapsed);
+}
 
-	return true;
+void CollapseRelativeDirectoriesImpl(FString& InOutPath, bool bCollapseAllPossible, bool& bOutAllCollapsed)
+{
+	CollapseRelativeDirectoriesImplTemplate(InOutPath, bCollapseAllPossible, bOutAllCollapsed);
+	InOutPath.TrimToNullTerminator();
+}
+
+}
+
+bool FPathViews::CollapseRelativeDirectories(FStringBuilderBase& InOutPath, bool bCollapseAllPossible)
+{
+	bool bAllCollapsed;
+	UE4PathViews_Private::CollapseRelativeDirectoriesImpl(InOutPath, bCollapseAllPossible, bAllCollapsed);
+	return bAllCollapsed;
 }
 
 void FPathViews::RemoveDuplicateSlashes(FStringBuilderBase& InOutPath)
diff --git a/Engine/Source/Runtime/Core/Private/Misc/PathViewsPrivate.h b/Engine/Source/Runtime/Core/Private/Misc/PathViewsPrivate.h
new file mode 100644
--- /dev/null
+++ b/Engine/Source/Runtime/Core/Private/Misc/PathViewsPrivate.h
@@ -0,0 +1,16 @@
+// Copyright Epic Games, Inc. All Rights Reserved.
+
+#pragma once
+
+#include "Containers/StringFwd.h"
+#include "CoreTypes.h"
+
+class FString;
+
+namespace UE4PathViews_Private
+{
+
+void CollapseRelativeDirectoriesImpl(FString& InPath, bool bCollapseAllPossible, bool& bOutAllCollapsed);
+void CollapseRelativeDirectoriesImpl(FStringBuilderBase& InPath, bool bCollapseAllPossible, bool& bOutAllCollapsed);
+
+}
\ No newline at end of file
diff --git a/Engine/Source/Runtime/Core/Private/Misc/Paths.cpp b/Engine/Source/Runtime/Core/Private/Misc/Paths.cpp
--- a/Engine/Source/Runtime/Core/Private/Misc/Paths.cpp
+++ b/Engine/Source/Runtime/Core/Private/Misc/Paths.cpp
@@ -17,8 +17,10 @@
 #include "Misc/Guid.h"
 #include "Misc/LazySingleton.h"
 #include "Misc/Parse.h"
+#include "Misc/PathViewsPrivate.h"
 #include "Misc/ScopeLock.h"
 #include "Misc/PathViews.h"
+#include "String/Find.h"
 #include "String/ParseTokens.h"
 #include "UObject/NameTypes.h"
 
@@ -116,12 +118,13 @@ namespace UE4Paths_Private
 		return GetSavedDirSuffix(FPaths::EngineUserDir(), TEXT("-enginesaveddirsuffix="));
 	}
 
-	FString ConvertRelativePathToFullInternal(FString&& BasePath, FString&& InPath)
+	template <typename BasePathAsStringOrStringView>
+	FString ConvertRelativePathToFullInternal(BasePathAsStringOrStringView&& BasePath, FString&& InPath)
 	{
 		FString FullyPathed;
 		if ( FPaths::IsRelative(InPath) )
 		{
-			FullyPathed  = MoveTemp(BasePath);
+			FullyPathed  = MoveTempIfPossible(BasePath);
 			FullyPathed /= MoveTemp(InPath);
 		}
 		else
@@ -1265,86 +1268,11 @@ void FPaths::NormalizeDirectoryName(FString& InPath)
 	FPlatformMisc::NormalizePath(InPath);
 }
 
-bool FPaths::CollapseRelativeDirectories(FString& InPath)
+bool FPaths::CollapseRelativeDirectories(FString& InPath, bool bCollapseAllPossible)
 {
-	const auto& ParentDir = TEXT("/..");
-	const int32 ParentDirLength = UE_ARRAY_COUNT( ParentDir ) - 1; // To avoid hardcoded values
-
-	for (;;)
-	{
-		// An empty path is finished
-		if (InPath.IsEmpty())
-		{
-			break;
-		}
-
-		// Consider empty paths or paths which start with .. or /.. as invalid
-		if (InPath.StartsWith(TEXT(".."), ESearchCase::CaseSensitive) || InPath.StartsWith(ParentDir))
-		{
-			return false;
-		}
-
-		// If there are no "/.."s left then we're done
-		int32 Index = InPath.Find(ParentDir, ESearchCase::CaseSensitive);
-		if (Index == -1)
-		{
-			break;
-		}
-
-		// Ignore folders beginning with dots
-		for (;;)
-		{
-			if (InPath.Len() <= Index + ParentDirLength || InPath[Index + ParentDirLength] == TEXT('/'))
-			{
-				break;
-			}
-
-			Index = InPath.Find(ParentDir, ESearchCase::CaseSensitive, ESearchDir::FromStart, Index + ParentDirLength);
-			if (Index == -1)
-			{
-				break;
-			}
-		}
-
-		if (Index == -1)
-		{
-			break;
-		}
-
-		int32 PreviousSeparatorIndex = Index;
-		for (;;)
-		{
-			// Find the previous slash
-			PreviousSeparatorIndex = FMath::Max(0, InPath.Find( TEXT("/"), ESearchCase::CaseSensitive, ESearchDir::FromEnd, PreviousSeparatorIndex - 1));
-
-			// Stop if we've hit the start of the string
-			if (PreviousSeparatorIndex == 0)
-			{
-				break;
-			}
-
-			// Stop if we've found a directory that isn't "/./"
-			if ((Index - PreviousSeparatorIndex) > 1 && (InPath[PreviousSeparatorIndex + 1] != TEXT('.') || InPath[PreviousSeparatorIndex + 2] != TEXT('/')))
-			{
-				break;
-			}
-		}
-
-		// If we're attempting to remove the drive letter, that's illegal
-		int32 Colon = InPath.Find(TEXT(":"), ESearchCase::CaseSensitive, ESearchDir::FromStart, PreviousSeparatorIndex);
-		if (Colon >= 0 && Colon < Index)
-		{
-			return false;
-		}
-
-		InPath.RemoveAt(PreviousSeparatorIndex, Index - PreviousSeparatorIndex + ParentDirLength, EAllowShrinking::No);
-	}
-
-	InPath.ReplaceInline(TEXT("./"), TEXT(""), ESearchCase::CaseSensitive);
-
-	InPath.TrimToNullTerminator();
-
-	return true;
+	bool bAllCollapsed;
+	UE4PathViews_Private::CollapseRelativeDirectoriesImpl(InPath, bCollapseAllPossible, bAllCollapsed);
+	return bAllCollapsed;
 }
 
 FString FPaths::RemoveDuplicateSlashes(const FString& InPath)
@@ -1425,14 +1353,35 @@ FString FPaths::CreateStandardFilename(const FString& InPath)
 		return InPath;
 	}
 
-	// make an absolute path
-	FString Standardized = FPaths::ConvertRelativePathToFull(WithSlashes);
-
 	// remove duplicate slashes
-	FPaths::RemoveDuplicateSlashes(Standardized);
-	// make it relative to Engine\Binaries\Platform
-	Standardized.ReplaceInline(RootDirectory, *FPaths::GetRelativePathToRoot());
-	return Standardized;
+	FPaths::RemoveDuplicateSlashes(WithSlashes);
+
+	// If the path starts with RelativePathToRoot, it is done, already in standard relative-to-root form.
+	const FString& RelativePathToRoot = FPaths::GetRelativePathToRoot();
+	if (WithSlashes.StartsWith(RelativePathToRoot))
+	{
+		NormalizeFilename(WithSlashes);
+		CollapseRelativeDirectories(WithSlashes, true /* bCollapseAllPossible */);
+		return WithSlashes;
+	}
+
+	// Make it an absolute path and then make that absolute path a subdirectory of the RelativePathToRoot.
+	// This standardizes both incoming absolute paths, and relative paths that are relative to a subdirectory of
+	// root rather than starting with the full relative to root path.
+	// d:\root\Engine\Content                            -> ..\..\..\Engine\Content
+	// ..\..\Content                                     -> ..\..\..\Engine\Content
+	// d:\root\Engine\Binaries\Platform\UnrealEditor.exe -> ..\..\..\Engine\Binaries\Platform\UnrealEditor.exe
+	// UnrealEditor.exe                                  -> ..\..\..\Engine\Binaries\Platform\UnrealEditor.exe
+	FString FullPath = FPaths::ConvertRelativePathToFull(WithSlashes);
+	if (!ensure(FullPath.StartsWith(RootDirectoryView)))
+	{
+		// Unexpected, we should have early exited with bCannotBeStandardized above
+		return InPath;
+	}
+	WithSlashes.Reset();
+	WithSlashes += RelativePathToRoot;
+	WithSlashes += FStringView(FullPath).RightChop(RootDirectoryView.Len());
+	return WithSlashes;
 }
 
 void FPaths::MakeStandardFilename(FString& InPath)
@@ -1511,22 +1460,22 @@ bool FPaths::MakePathRelativeTo( FString& InPath, const TCHAR* InRelativeTo )
 
 FString FPaths::ConvertRelativePathToFull(const FString& InPath)
 {
-	return UE4Paths_Private::ConvertRelativePathToFullInternal(FString(FPlatformProcess::BaseDir()), FString(InPath));
+	return UE4Paths_Private::ConvertRelativePathToFullInternal(FStringView(FPlatformProcess::BaseDir()), FString(InPath));
 }
 
 FString FPaths::ConvertRelativePathToFull(FString&& InPath)
 {
-	return UE4Paths_Private::ConvertRelativePathToFullInternal(FString(FPlatformProcess::BaseDir()), MoveTemp(InPath));
+	return UE4Paths_Private::ConvertRelativePathToFullInternal(FStringView(FPlatformProcess::BaseDir()), MoveTemp(InPath));
 }
 
 FString FPaths::ConvertRelativePathToFull(const FString& BasePath, const FString& InPath)
 {
-	return UE4Paths_Private::ConvertRelativePathToFullInternal(CopyTemp(BasePath), CopyTemp(InPath));
+	return UE4Paths_Private::ConvertRelativePathToFullInternal(BasePath, CopyTemp(InPath));
 }
 
 FString FPaths::ConvertRelativePathToFull(const FString& BasePath, FString&& InPath)
 {
-	return UE4Paths_Private::ConvertRelativePathToFullInternal(CopyTemp(BasePath), MoveTemp(InPath));
+	return UE4Paths_Private::ConvertRelativePathToFullInternal(BasePath, MoveTemp(InPath));
 }
 
 FString FPaths::ConvertRelativePathToFull(FString&& BasePath, const FString& InPath)
diff --git a/Engine/Source/Runtime/Core/Public/Misc/PathViews.h b/Engine/Source/Runtime/Core/Public/Misc/PathViews.h
--- a/Engine/Source/Runtime/Core/Public/Misc/PathViews.h
+++ b/Engine/Source/Runtime/Core/Public/Misc/PathViews.h
@@ -385,7 +385,7 @@ public:
 	static CORE_API void NormalizeDirectoryName(FStringBuilderBase& InOutPath);
 
 	/** Collapses redundant paths like "/./" and "SkipDir/..". FPaths::CollapseRelativeDirectories() equivalent. */
-	static CORE_API bool CollapseRelativeDirectories(FStringBuilderBase& InOutPath);
+	static CORE_API bool CollapseRelativeDirectories(FStringBuilderBase& InOutPath, bool bCollapseAllPossible = false);
 
 	/** Removes duplicate forward slashes, e.g. "a/b//c////f.e" -> "a/b/c/f.e" */
 	static CORE_API void RemoveDuplicateSlashes(FStringBuilderBase& InOutPath);
diff --git a/Engine/Source/Runtime/Core/Public/Misc/Paths.h b/Engine/Source/Runtime/Core/Public/Misc/Paths.h
--- a/Engine/Source/Runtime/Core/Public/Misc/Paths.h
+++ b/Engine/Source/Runtime/Core/Public/Misc/Paths.h
@@ -577,8 +577,12 @@ public:
 	 *	BaseDirectory/SomeDirectory/../SomeOtherDirectory/Filename.ext
 	 * and converts it to:
 	 *	BaseDirectory/SomeOtherDirectory/Filename.ext
+	 * @param bCollapseAllPossible If false, and there are leading .. that are uncollapsible, the function
+	 *        early exits and does not modify .. or . that occur later in the path. If true, the function
+	 *        collapses all . and .. that can be collapsed. Return value is the same in either case.
+	 * @return true if all directories were collapsed, else false.
 	 */
-	static CORE_API bool CollapseRelativeDirectories(FString& InPath);
+	static CORE_API bool CollapseRelativeDirectories(FString& InPath, bool bCollapseAllPossible = false);
 
 	/**
 	 * Removes duplicate slashes in paths.
diff --git a/Engine/Source/Runtime/Core/Tests/Misc/PathTests.h b/Engine/Source/Runtime/Core/Tests/Misc/PathTests.h
--- a/Engine/Source/Runtime/Core/Tests/Misc/PathTests.h
+++ b/Engine/Source/Runtime/Core/Tests/Misc/PathTests.h
@@ -13,87 +13,113 @@
 template<class PathType, class StringType>
 void TestCollapseRelativeDirectories()
 {
-	auto Run = [](const TCHAR* Path, const TCHAR* Expected)
+	auto Run = [](const TCHAR* Path, bool bExpectedAllCollapsed, const TCHAR* Expected)
 	{
 		// Run test
 		StringType Actual;
+		StringType IfPossibleActual;
 		Actual += Path;
-		const bool bValid = PathType::CollapseRelativeDirectories(Actual);
+		IfPossibleActual += Path;
+		const bool bAllCollapsed = PathType::CollapseRelativeDirectories(
+			Actual, false /* bCollapseAllPossible */);
+		const bool bIfPossibleAllCollapsed = PathType::CollapseRelativeDirectories(
+			IfPossibleActual, true /* bCollapseAllPossible */);
 
-		if (Expected)
+		auto IsExpectedText = [Expected](const TCHAR* Actual)
+			{
+				return FCString::Strcmp(Actual, Expected) == 0;
+			};
+		if (bExpectedAllCollapsed)
 		{
 			// If we're looking for a result, make sure it was returned correctly
-			if (!bValid || FCString::Strcmp(*Actual, Expected) != 0)
+			if (!bAllCollapsed || !IsExpectedText(*Actual))
+			{
+				FAIL_CHECK(FString::Printf(TEXT("CollapseRelativeDirectories('%s', false) failed (got (%s, '%s'), expected (%s, '%s'))."),
+					Path, *LexToString(bAllCollapsed), *Actual, *LexToString(bExpectedAllCollapsed), Expected));
+			}
+			if (!bIfPossibleAllCollapsed || !IsExpectedText(*IfPossibleActual))
 			{
-				FAIL_CHECK(FString::Printf(TEXT("Path '%s' failed to collapse correctly (got '%s', expected '%s')."), Path, *Actual, Expected));
+				FAIL_CHECK(FString::Printf(TEXT("CollapseRelativeDirectories('%s', true) failed (got (%s, '%s'), expected (%s, '%s'))."),
+					Path, *LexToString(bIfPossibleAllCollapsed), *IfPossibleActual, *LexToString(bExpectedAllCollapsed), Expected));
 			}
 		}
 		else
 		{
 			// Otherwise, make sure it failed
-			if (bValid)
+			if (bAllCollapsed)
+			{
+				// Some of the string might have been collapsed before finding out in has an incollapsible piece
+				// The only well-defined option for those cases is to leave the original string unaltered, but
+				// enforcing that contract is too expensive, so we don't provide that contract and instead we
+				// only guarantee that the modified path is equivalent, perhaps with some .. and . collapsed.
+				// We don't currently validate that equivalency on our test cases, because it is difficult to calculate.
+				FAIL_CHECK(FString::Printf(TEXT("Path '%s' failed CollapseRelativeDirectories (got (%s, '%s'), expected (%s, '<anyvalue>'))."),
+					Path, *LexToString(bAllCollapsed), *Actual, *LexToString(bExpectedAllCollapsed)));
+			}
+			if (bIfPossibleAllCollapsed || !IsExpectedText(*IfPossibleActual))
 			{
-				FAIL_CHECK(FString::Printf(TEXT("Path '%s' collapsed unexpectedly."), Path));
+				FAIL_CHECK(FString::Printf(TEXT("Path '%s' failed CollapseRelativeDirectoriesIfPossible (got (%s, '%s'), expected (%s, '%s'))."),
+					Path, *LexToString(bIfPossibleAllCollapsed), *IfPossibleActual, *LexToString(bExpectedAllCollapsed), Expected));
 			}
 		}
 	};
 
-	Run(TEXT(".."),                                                   nullptr);
-	Run(TEXT("/.."),                                                  nullptr);
-	Run(TEXT("./"),                                                   TEXT(""));
-	Run(TEXT("./file.txt"),                                           TEXT("file.txt"));
-	Run(TEXT("/."),                                                   TEXT("/."));
-	Run(TEXT("Folder"),                                               TEXT("Folder"));
-	Run(TEXT("/Folder"),                                              TEXT("/Folder"));
-	Run(TEXT("C:/Folder"),                                            TEXT("C:/Folder"));
-	Run(TEXT("C:/Folder/.."),                                         TEXT("C:"));
-	Run(TEXT("C:/Folder/../"),                                        TEXT("C:/"));
-	Run(TEXT("C:/Folder/../file.txt"),                                TEXT("C:/file.txt"));
-	Run(TEXT("Folder/.."),                                            TEXT(""));
-	Run(TEXT("Folder/../"),                                           TEXT("/"));
-	Run(TEXT("Folder/../file.txt"),                                   TEXT("/file.txt"));
-	Run(TEXT("/Folder/.."),                                           TEXT(""));
-	Run(TEXT("/Folder/../"),                                          TEXT("/"));
-	Run(TEXT("/Folder/../file.txt"),                                  TEXT("/file.txt"));
-	Run(TEXT("Folder/../.."),                                         nullptr);
-	Run(TEXT("Folder/../../"),                                        nullptr);
-	Run(TEXT("Folder/../../file.txt"),                                nullptr);
-	Run(TEXT("C:/.."),                                                nullptr);
-	Run(TEXT("C:/."),                                                 TEXT("C:/."));
-	Run(TEXT("C:/./"),                                                TEXT("C:/"));
-	Run(TEXT("C:/./file.txt"),                                        TEXT("C:/file.txt"));
-	Run(TEXT("C:/Folder1/../Folder2"),                                TEXT("C:/Folder2"));
-	Run(TEXT("C:/Folder1/../Folder2/"),                               TEXT("C:/Folder2/"));
-	Run(TEXT("C:/Folder1/../Folder2/file.txt"),                       TEXT("C:/Folder2/file.txt"));
-	Run(TEXT("C:/Folder1/../Folder2/../.."),                          nullptr);
-	Run(TEXT("C:/Folder1/../Folder2/../Folder3"),                     TEXT("C:/Folder3"));
-	Run(TEXT("C:/Folder1/../Folder2/../Folder3/"),                    TEXT("C:/Folder3/"));
-	Run(TEXT("C:/Folder1/../Folder2/../Folder3/file.txt"),            TEXT("C:/Folder3/file.txt"));
-	Run(TEXT("C:/Folder1/Folder2/../../Folder3"),                     TEXT("C:/Folder3"));
-	Run(TEXT("C:/Folder1/Folder2/../../Folder3/"),                    TEXT("C:/Folder3/"));
-	Run(TEXT("C:/Folder1/Folder2/../../Folder3/file.txt"),            TEXT("C:/Folder3/file.txt"));
-	Run(TEXT("C:/Folder1/Folder2/../../Folder3/../Folder4"),          TEXT("C:/Folder4"));
-	Run(TEXT("C:/Folder1/Folder2/../../Folder3/../Folder4/"),         TEXT("C:/Folder4/"));
-	Run(TEXT("C:/Folder1/Folder2/../../Folder3/../Folder4/file.txt"), TEXT("C:/Folder4/file.txt"));
-	Run(TEXT("C:/Folder1/Folder2/../Folder3/../../Folder4"),          TEXT("C:/Folder4"));
-	Run(TEXT("C:/Folder1/Folder2/../Folder3/../../Folder4/"),         TEXT("C:/Folder4/"));
-	Run(TEXT("C:/Folder1/Folder2/../Folder3/../../Folder4/file.txt"), TEXT("C:/Folder4/file.txt"));
-	Run(TEXT("C:/Folder1/Folder2/.././../Folder4"),                   TEXT("C:/Folder4"));
-	Run(TEXT("C:/Folder1/Folder2/.././../Folder4/"),                  TEXT("C:/Folder4/"));
-	Run(TEXT("C:/Folder1/Folder2/.././../Folder4/file.txt"),          TEXT("C:/Folder4/file.txt"));
-	Run(TEXT("C:/A/B/.././../C"),                                     TEXT("C:/C"));
-	Run(TEXT("C:/A/B/.././../C/"),                                    TEXT("C:/C/"));
-	Run(TEXT("C:/A/B/.././../C/file.txt"),                            TEXT("C:/C/file.txt"));
-	Run(TEXT(".svn"),                                                 TEXT(".svn"));
-	Run(TEXT("/.svn"),                                                TEXT("/.svn"));
-	Run(TEXT("./Folder/.svn"),                                        TEXT("Folder/.svn"));
-	Run(TEXT("./.svn/../.svn"),                                       TEXT(".svn"));
-	Run(TEXT(".svn/./.svn/.././../.svn"),                             TEXT("/.svn"));
-	Run(TEXT("C:/Folder1/./Folder2/..Folder3"),						 TEXT("C:/Folder1/Folder2/..Folder3"));
-	Run(TEXT("C:/Folder1/./Folder2/..Folder3/Folder4"),				 TEXT("C:/Folder1/Folder2/..Folder3/Folder4"));
-	Run(TEXT("C:/Folder1/./Folder2/..Folder3/..Folder4"),			 TEXT("C:/Folder1/Folder2/..Folder3/..Folder4"));
-	Run(TEXT("C:/Folder1/./Folder2/..Folder3/Folder4/../Folder5"),	 TEXT("C:/Folder1/Folder2/..Folder3/Folder5"));
-	Run(TEXT("C:/Folder1/..Folder2/Folder3/..Folder4/../Folder5"),	 TEXT("C:/Folder1/..Folder2/Folder3/Folder5"));
+	Run(TEXT(".."),                                                   false, TEXT(".."));
+	Run(TEXT("/.."),                                                  false, TEXT("/.."));
+	Run(TEXT("./"),                                                   true,  TEXT(""));
+	Run(TEXT("./file.txt"),                                           true,  TEXT("file.txt"));
+	Run(TEXT("/."),                                                   true,  TEXT("/."));
+	Run(TEXT("Folder"),                                               true,  TEXT("Folder"));
+	Run(TEXT("/Folder"),                                              true,  TEXT("/Folder"));
+	Run(TEXT("C:/Folder"),                                            true,  TEXT("C:/Folder"));
+	Run(TEXT("C:/Folder/.."),                                         true,  TEXT("C:"));
+	Run(TEXT("C:/Folder/../"),                                        true,  TEXT("C:/"));
+	Run(TEXT("C:/Folder/../file.txt"),                                true,  TEXT("C:/file.txt"));
+	Run(TEXT("Folder/.."),                                            true,  TEXT(""));
+	Run(TEXT("Folder/../"),                                           true,  TEXT(""));
+	Run(TEXT("Folder/../file.txt"),                                   true,  TEXT("file.txt"));
+	Run(TEXT("/Folder/.."),                                           true,  TEXT(""));
+	Run(TEXT("/Folder/../"),                                          true,  TEXT("/"));
+	Run(TEXT("/Folder/../file.txt"),                                  true,  TEXT("/file.txt"));
+	Run(TEXT("Folder/../.."),                                         false, TEXT(".."));
+	Run(TEXT("Folder/../../"),                                        false, TEXT("../"));
+	Run(TEXT("Folder/../../file.txt"),                                false, TEXT("../file.txt"));
+	Run(TEXT("C:/.."),                                                false, TEXT("C:/.."));
+	Run(TEXT("C:/."),                                                 true,  TEXT("C:/."));
+	Run(TEXT("C:/./"),                                                true,  TEXT("C:/"));
+	Run(TEXT("C:/./file.txt"),                                        true,  TEXT("C:/file.txt"));
+	Run(TEXT("C:/Folder1/../Folder2"),                                true,  TEXT("C:/Folder2"));
+	Run(TEXT("C:/Folder1/../Folder2/"),                               true,  TEXT("C:/Folder2/"));
+	Run(TEXT("C:/Folder1/../Folder2/file.txt"),                       true,  TEXT("C:/Folder2/file.txt"));
+	Run(TEXT("C:/Folder1/../Folder2/../.."),                          false, TEXT("C:/.."));
+	Run(TEXT("C:/Folder1/../Folder2/../Folder3"),                     true,  TEXT("C:/Folder3"));
+	Run(TEXT("C:/Folder1/../Folder2/../Folder3/"),                    true,  TEXT("C:/Folder3/"));
+	Run(TEXT("C:/Folder1/../Folder2/../Folder3/file.txt"),            true,  TEXT("C:/Folder3/file.txt"));
+	Run(TEXT("C:/Folder1/Folder2/../../Folder3"),                     true,  TEXT("C:/Folder3"));
+	Run(TEXT("C:/Folder1/Folder2/../../Folder3/"),                    true,  TEXT("C:/Folder3/"));
+	Run(TEXT("C:/Folder1/Folder2/../../Folder3/file.txt"),            true,  TEXT("C:/Folder3/file.txt"));
+	Run(TEXT("C:/Folder1/Folder2/../../Folder3/../Folder4"),          true,  TEXT("C:/Folder4"));
+	Run(TEXT("C:/Folder1/Folder2/../../Folder3/../Folder4/"),         true,  TEXT("C:/Folder4/"));
+	Run(TEXT("C:/Folder1/Folder2/../../Folder3/../Folder4/file.txt"), true,  TEXT("C:/Folder4/file.txt"));
+	Run(TEXT("C:/Folder1/Folder2/../Folder3/../../Folder4"),          true,  TEXT("C:/Folder4"));
+	Run(TEXT("C:/Folder1/Folder2/../Folder3/../../Folder4/"),         true,  TEXT("C:/Folder4/"));
+	Run(TEXT("C:/Folder1/Folder2/../Folder3/../../Folder4/file.txt"), true,  TEXT("C:/Folder4/file.txt"));
+	Run(TEXT("C:/Folder1/Folder2/.././../Folder4"),                   true,  TEXT("C:/Folder4"));
+	Run(TEXT("C:/Folder1/Folder2/.././../Folder4/"),                  true,  TEXT("C:/Folder4/"));
+	Run(TEXT("C:/Folder1/Folder2/.././../Folder4/file.txt"),          true,  TEXT("C:/Folder4/file.txt"));
+	Run(TEXT("C:/A/B/.././../C"),                                     true,  TEXT("C:/C"));
+	Run(TEXT("C:/A/B/.././../C/"),                                    true,  TEXT("C:/C/"));
+	Run(TEXT("C:/A/B/.././../C/file.txt"),                            true,  TEXT("C:/C/file.txt"));
+	Run(TEXT(".svn"),                                                 true,  TEXT(".svn"));
+	Run(TEXT("/.svn"),                                                true,  TEXT("/.svn"));
+	Run(TEXT("./Folder/.svn"),                                        true,  TEXT("Folder/.svn"));
+	Run(TEXT("./.svn/../.svn"),                                       true,  TEXT(".svn"));
+	Run(TEXT(".svn/./.svn/.././../.svn"),                             true,  TEXT(".svn"));
+	Run(TEXT("C:/Folder1/./Folder2/..Folder3"),                       true,  TEXT("C:/Folder1/Folder2/..Folder3"));
+	Run(TEXT("C:/Folder1/./Folder2/..Folder3/Folder4"),               true,  TEXT("C:/Folder1/Folder2/..Folder3/Folder4"));
+	Run(TEXT("C:/Folder1/./Folder2/..Folder3/..Folder4"),             true,  TEXT("C:/Folder1/Folder2/..Folder3/..Folder4"));
+	Run(TEXT("C:/Folder1/./Folder2/..Folder3/Folder4/../Folder5"),    true,  TEXT("C:/Folder1/Folder2/..Folder3/Folder5"));
+	Run(TEXT("C:/Folder1/..Folder2/Folder3/..Folder4/../Folder5"),    true,  TEXT("C:/Folder1/..Folder2/Folder3/Folder5"));
 }
 
 template<class PathType, class StringType>
diff --git a/Engine/Source/Runtime/Core/Tests/Misc/PathsTest.cpp b/Engine/Source/Runtime/Core/Tests/Misc/PathsTest.cpp
--- a/Engine/Source/Runtime/Core/Tests/Misc/PathsTest.cpp
+++ b/Engine/Source/Runtime/Core/Tests/Misc/PathsTest.cpp
@@ -3,15 +3,20 @@
 #if WITH_TESTS
 
 #include "PathTests.h"
+
+#include "Containers/UnrealString.h"
 #include "CoreTypes.h"
+#include "HAL/PlatformProcess.h"
 #include "Misc/AssertionMacros.h"
-#include "Containers/UnrealString.h"
 #include "Misc/Paths.h"
+#include "Misc/PathViews.h"
 #include "Tests/TestHarnessAdapter.h"
 #if WITH_LOW_LEVEL_TESTS
 #include "TestCommon/Expectations.h"
 #endif
 
+void TestCreateStandardFilename();
+
 const FStringView PathTest::BaseDir = TEXTVIEW("/root");
 
 const PathTest::FTestPair PathTest::ExpectedRelativeToAbsolutePaths[10] =
@@ -31,6 +36,7 @@ const PathTest::FTestPair PathTest::ExpectedRelativeToAbsolutePaths[10] =
 TEST_CASE_NAMED(FPathTests, "System::Core::Misc::Paths", "[ApplicationContextMask][SmokeFilter]")
 {
 	TestCollapseRelativeDirectories<FPaths, FString>();
+	TestCreateStandardFilename();
 
 	// Extension texts
 	{
@@ -204,4 +210,138 @@ TEST_CASE_NAMED(FPathTests, "System::Core::Misc::Paths", "[ApplicationContextMas
 	RunValidatePathTest(TEXT("C:\\Path\\Dir^"), false);
 }
 
+void TestCreateStandardFilename()
+{
+	auto NormalizeDirForTest = [](FString& Dir)
+		{
+			while (!Dir.IsEmpty() && FPathViews::IsSeparator(Dir[Dir.Len() - 1]))
+			{
+				Dir.LeftChopInline(1);
+			}
+		};
+	FString RootDirectory = FPaths::ConvertRelativePathToFull(FPaths::RootDir());
+	FString BaseDirectory = FPaths::ConvertRelativePathToFull(FPlatformProcess::BaseDir());
+	FString RootRelative = FPaths::GetRelativePathToRoot();
+	NormalizeDirForTest(RootDirectory);
+	NormalizeDirForTest(BaseDirectory);
+	NormalizeDirForTest(RootRelative);
+
+	FStringView RelPathToBase;
+	if (RootDirectory.IsEmpty() || !FPathViews::TryMakeChildPathRelativeTo(BaseDirectory, RootDirectory, RelPathToBase))
+	{
+		// In the case where the base path (expected <root>/Engine/Binaries/Win64) is not a subdirectory of <root>,
+		// we don't know how to verify the results.
+		// Skip the test.
+		return;
+	}
+	FString RootDirectoryWithSlash = RootDirectory + TEXT("/");
+	FString StandardPathOfBase = FPaths::Combine(RootRelative, RelPathToBase);
+	FStringView StandardPathOfBaseParent = FPathViews::GetPath(StandardPathOfBase);
+	FStringView StandardPathOfBaseParent2 = FPathViews::GetPath(StandardPathOfBaseParent);
+	FStringView StandardPathOfBaseParent3 = FPathViews::GetPath(StandardPathOfBaseParent2);
+	if (StandardPathOfBaseParent.EndsWith(TEXT("..") || StandardPathOfBaseParent.IsEmpty())
+		|| StandardPathOfBaseParent2.EndsWith(TEXT("..") || StandardPathOfBaseParent2.IsEmpty())
+		|| StandardPathOfBaseParent3.EndsWith(TEXT("..") || StandardPathOfBaseParent3.IsEmpty()))
+	{
+		// We expect StandardPathOfBase == ../../../Engine/Binaries/Win64
+		// We want to be able to use this to know that ".." -> "../../../Engine/Binaries" and
+		// similar for ../.. and ../../..
+		// If we don't have a parent directory, or the "parent directory" is a ".." directory,
+		// then we don't know what the value for these directories in standard form should be.
+		// Skip the test.
+		return;
+	}
+	if (RootDirectory.StartsWith(TEXT("X:")))
+	{
+		// We use X: as the root of absolute paths that aren't under root and so standard path will be their
+		// absolute path. If we're running on a workspace that is under the X: drive, skip the test since
+		// it will create different answers for those paths.
+		return;
+	}
+	
+	auto Run = [&RootRelative](const TCHAR* Path, FStringView Expected)
+		{
+			// Run test
+			FString Actual = FPaths::CreateStandardFilename(Path);
+			if (!Expected.Equals(Actual, ESearchCase::CaseSensitive))
+			{
+				FAIL_CHECK(FString::Printf(TEXT("Path '%s' failed CreateStandardFilename (got '%s', expected '%s')."),
+					Path, *Actual, *FString(Expected)));
+			}
+		};
+
+	Run(TEXT(".."), StandardPathOfBaseParent);
+	Run(TEXT("/.."), TEXT("/.."));
+	Run(TEXT("./"), StandardPathOfBase + TEXT("/")); // Keeping the trailing slash is legacy behavior in CollapseSubdirectories
+	Run(TEXT("./file.txt"), FPaths::Combine(StandardPathOfBase, TEXT("file.txt")));
+	Run(TEXT("/."), TEXT("/."));
+	Run(TEXT("Folder"), FPaths::Combine(StandardPathOfBase, TEXT("Folder")));
+	Run(TEXT("/Folder"), TEXT("/Folder"));
+	Run(TEXT("X:/Folder"), TEXT("X:/Folder"));
+	// Current contract is that paths that cannot be standardized do not CollapseRelativeDirectories; they remain unchanged
+	Run(TEXT("X:/Folder/.."), TEXT("X:/Folder/..")); // cannot be standardized
+	Run(TEXT("X:/Folder/../"), TEXT("X:/Folder/../")); // cannot be standardized
+	Run(TEXT("X:/Folder/../file.txt"), TEXT("X:/Folder/../file.txt")); // cannot be standardized
+	Run(TEXT("Folder/.."), StandardPathOfBase);
+	Run(TEXT("Folder/../"), StandardPathOfBase + TEXT("/")); // Keeping the trailing slash is legacy behavior in CollapseSubdirectories
+	Run(TEXT("Folder/../file.txt"), FPaths::Combine(StandardPathOfBase, TEXT("file.txt")));
+	Run(TEXT("/Folder/.."), TEXT("/Folder/..")); // cannot be standardized
+	Run(TEXT("/Folder/../"), TEXT("/Folder/../")); // cannot be standardized
+	Run(TEXT("/Folder/../file.txt"), TEXT("/Folder/../file.txt")); // cannot be standardized
+	Run(TEXT("Folder/../.."), StandardPathOfBaseParent);
+	Run(TEXT("Folder/../../"), FString(StandardPathOfBaseParent) + TEXT("/")); // Keeping the trailing slash is legacy behavior in CollapseSubdirectories
+	Run(TEXT("Folder/../../file.txt"), FPaths::Combine(StandardPathOfBaseParent, TEXT("file.txt")));
+	Run(TEXT("X:/.."), TEXT("X:/.."));
+	Run(TEXT("X:/."), TEXT("X:/.")); // cannot be standardized
+	Run(TEXT("X:/./"), TEXT("X:/./")); // cannot be standardized
+	Run(TEXT("X:/./file.txt"), TEXT("X:/./file.txt")); // cannot be standardized
+	Run(TEXT("X:/Folder1/../Folder2"), TEXT("X:/Folder1/../Folder2")); // cannot be standardized
+	Run(TEXT("Folder1/../Folder2"), FPaths::Combine(StandardPathOfBase, TEXT("Folder2")));
+	Run(TEXT("Folder1/../Folder2/"), FPaths::Combine(StandardPathOfBase, TEXT("Folder2/"))); // Keeping the trailing slash is legacy behavior in CollapseSubdirectories
+	Run(TEXT("Folder1/../Folder2/file.txt"), FPaths::Combine(StandardPathOfBase, TEXT("Folder2/file.txt")));
+	Run(TEXT("Folder1/../Folder2/../.."), StandardPathOfBaseParent);
+	Run(TEXT("Folder1/../Folder2/../Folder3"), FPaths::Combine(StandardPathOfBase, TEXT("Folder3")));
+	Run(TEXT("Folder1/../Folder2/../Folder3/"), FPaths::Combine(StandardPathOfBase, TEXT("Folder3/")));
+	Run(TEXT("Folder1/../Folder2/../Folder3/file.txt"), FPaths::Combine(StandardPathOfBase, TEXT("Folder3/file.txt")));
+	Run(TEXT("Folder1/Folder2/../../Folder3"), FPaths::Combine(StandardPathOfBase, TEXT("Folder3")));
+	Run(TEXT("Folder1/Folder2/../../Folder3/"), FPaths::Combine(StandardPathOfBase, TEXT("Folder3/")));
+	Run(TEXT("Folder1/Folder2/../../Folder3/file.txt"), FPaths::Combine(StandardPathOfBase, TEXT("Folder3/file.txt")));
+	Run(TEXT("Folder1/Folder2/../../Folder3/../Folder4"), FPaths::Combine(StandardPathOfBase, TEXT("Folder4")));
+	Run(TEXT("Folder1/Folder2/../../Folder3/../Folder4/"), FPaths::Combine(StandardPathOfBase, TEXT("Folder4/"))); // Keeping the trailing slash is legacy behavior in CollapseSubdirectories
+	Run(TEXT("Folder1/Folder2/../../Folder3/../Folder4/file.txt"), FPaths::Combine(StandardPathOfBase, TEXT("Folder4/file.txt")));
+	Run(TEXT("Folder1/Folder2/../Folder3/../../Folder4"), FPaths::Combine(StandardPathOfBase, TEXT("Folder4")));
+	Run(TEXT("Folder1/Folder2/../Folder3/../../Folder4/"), FPaths::Combine(StandardPathOfBase, TEXT("Folder4/"))); // Keeping the trailing slash is legacy behavior in CollapseSubdirectories
+	Run(TEXT("Folder1/Folder2/../Folder3/../../Folder4/file.txt"), FPaths::Combine(StandardPathOfBase, TEXT("Folder4/file.txt")));
+	Run(TEXT("Folder1/Folder2/.././../Folder4"), FPaths::Combine(StandardPathOfBase, TEXT("Folder4")));
+	Run(TEXT("Folder1/Folder2/.././../Folder4/"), FPaths::Combine(StandardPathOfBase, TEXT("Folder4/"))); // Keeping the trailing slash is legacy behavior in CollapseSubdirectories
+	Run(TEXT("Folder1/Folder2/.././../Folder4/file.txt"), FPaths::Combine(StandardPathOfBase, TEXT("Folder4/file.txt")));
+	Run(TEXT("A/B/.././../C"), FPaths::Combine(StandardPathOfBase, TEXT("C")));
+	Run(TEXT("A/B/.././../C/"), FPaths::Combine(StandardPathOfBase, TEXT("C/"))); // Keeping the trailing slash is legacy behavior in CollapseSubdirectories
+	Run(TEXT("A/B/.././../C/file.txt"), FPaths::Combine(StandardPathOfBase, TEXT("C/file.txt")));
+	Run(TEXT(".svn"), FPaths::Combine(StandardPathOfBase, TEXT(".svn")));
+	Run(TEXT("/.svn"), TEXT("/.svn"));
+	Run(TEXT("./Folder/.svn"), FPaths::Combine(StandardPathOfBase, TEXT("Folder/.svn")));
+	Run(TEXT("./.svn/../.svn"), FPaths::Combine(StandardPathOfBase, TEXT(".svn")));
+	Run(TEXT(".svn/./.svn/.././../.svn"), FPaths::Combine(StandardPathOfBase, TEXT(".svn")));
+	Run(TEXT("Folder1/./Folder2/..Folder3"), FPaths::Combine(StandardPathOfBase, TEXT("Folder1/Folder2/..Folder3")));
+	Run(TEXT("Folder1/./Folder2/..Folder3/Folder4"), FPaths::Combine(StandardPathOfBase, TEXT("Folder1/Folder2/..Folder3/Folder4")));
+	Run(TEXT("Folder1/./Folder2/..Folder3/..Folder4"), FPaths::Combine(StandardPathOfBase, TEXT("Folder1/Folder2/..Folder3/..Folder4")));
+	Run(TEXT("Folder1/./Folder2/..Folder3/Folder4/../Folder5"), FPaths::Combine(StandardPathOfBase, TEXT("Folder1/Folder2/..Folder3/Folder5")));
+	Run(TEXT("Folder1/..Folder2/Folder3/..Folder4/../Folder5"), FPaths::Combine(StandardPathOfBase, TEXT("Folder1/..Folder2/Folder3/Folder5")));
+	Run(TEXT("../../../Dir/Dir"), FPaths::Combine(StandardPathOfBaseParent3, TEXT("Dir/Dir")));
+	Run(TEXT("../../../Dir/Dir/"), FPaths::Combine(StandardPathOfBaseParent3, TEXT("Dir/Dir/"))); // Keeping the trailing slash is legacy behavior in CollapseSubdirectories
+	Run(TEXT("../../../Dir/Dir/../../Dir/Dir/../../Dir/Dir/File"), FPaths::Combine(StandardPathOfBaseParent3, TEXT("Dir/Dir/File")));
+	Run(TEXT("../../../Dir/Dir/File.ext"), FPaths::Combine(StandardPathOfBaseParent3, TEXT("Dir/Dir/File.ext")));
+	Run(TEXT("../../../Dir/Dir/File"), FPaths::Combine(StandardPathOfBaseParent3, TEXT("Dir/Dir/File")));
+	Run(TEXT("../../../Dir/"), FPaths::Combine(StandardPathOfBaseParent3, TEXT("Dir/"))); // Keeping the trailing slash is legacy behavior in CollapseSubdirectories
+	Run(TEXT("../../../Dir/Binaries/Win64/File.Ext"), FPaths::Combine(StandardPathOfBaseParent3, TEXT("Dir/Binaries/Win64/File.Ext")));
+	Run(TEXT("../../../Dir/Binaries/Win64/../../File.Ext"), FPaths::Combine(StandardPathOfBaseParent3, TEXT("Dir/File.Ext")));
+	Run(TEXT("../../Dir/Dir/File"), FPaths::Combine(StandardPathOfBaseParent2, TEXT("Dir/Dir/File")));
+	// Run(*RootDirectory, RootRelative); // Not currently handled without the slash
+	Run(*RootDirectoryWithSlash, RootRelative + TEXT("/")); // Keeping the trailing slash is legacy behavior in CollapseSubdirectories 
+	Run(*FPaths::Combine(RootDirectory, TEXT("Dir/Dir/File.ext")), FPaths::Combine(RootRelative, TEXT("Dir/Dir/File.ext")));
+	Run(*FPaths::Combine(RootDirectory, TEXT("Engine/Dir/File.ext")), FPaths::Combine(RootRelative, TEXT("Engine/Dir/File.ext")));
+	Run(*(RootDirectory + TEXT("\\Dir\\File")), FPaths::Combine(RootRelative, TEXT("/Dir/File")));
+}
+
 #endif //WITH_TESTS
