FName separate ctor instead of default argument

The code to materialize FNAME_Add as second argument ends up at every call site, and we have a lot of them.

diff --git a/Engine/Source/Runtime/Core/Private/UObject/UnrealNames.cpp b/Engine/Source/Runtime/Core/Private/UObject/UnrealNames.cpp
--- a/Engine/Source/Runtime/Core/Private/UObject/UnrealNames.cpp
+++ b/Engine/Source/Runtime/Core/Private/UObject/UnrealNames.cpp
@@ -3401,6 +3401,18 @@ const FNameEntry* FName::ResolveEntryRecursive(FNameEntryId LookupId)
 }
 
 
+FName::FName(const WIDECHAR* Name)
+	: FName(FNameHelper::MakeDetectNumber(MakeUnconvertedView(Name), FNAME_Add))
+{}
+
+FName::FName(const ANSICHAR* Name)
+	: FName(FNameHelper::MakeDetectNumber(MakeUnconvertedView(Name), FNAME_Add))
+{}
+
+FName::FName(const UTF8CHAR* Name)
+	: FName(FNameHelper::MakeDetectNumber(MakeUnconvertedView(Name), FNAME_Add))
+{}
+
 FName::FName(const WIDECHAR* Name, EFindName FindType)
 	: FName(FNameHelper::MakeDetectNumber(MakeUnconvertedView(Name), FindType))
 {}
diff --git a/Engine/Source/Runtime/Core/Public/UObject/NameTypes.h b/Engine/Source/Runtime/Core/Public/UObject/NameTypes.h
--- a/Engine/Source/Runtime/Core/Public/UObject/NameTypes.h
+++ b/Engine/Source/Runtime/Core/Public/UObject/NameTypes.h
@@ -995,11 +995,15 @@ public:
 	 * When UE_FNAME_OUTLINE_NUMBER is 1, we check for whole name including the number.
 	 *
 	 * @param Name			Value for the string portion of the name
-	 * @param FindType		Action to take (see EFindName)
+	 * @param FindType		Action to take (see EFindName, default is FNAME_Add)
 	 */
-	CORE_API FName(const WIDECHAR* Name, EFindName FindType=FNAME_Add);
-	CORE_API FName(const ANSICHAR* Name, EFindName FindType=FNAME_Add);
-	CORE_API FName(const UTF8CHAR* Name, EFindName FindType=FNAME_Add);
+	CORE_API FName(const WIDECHAR* Name);
+	CORE_API FName(const ANSICHAR* Name);
+	CORE_API FName(const UTF8CHAR* Name);
+
+	CORE_API FName(const WIDECHAR* Name, EFindName FindType);
+	CORE_API FName(const ANSICHAR* Name, EFindName FindType);
+	CORE_API FName(const UTF8CHAR* Name, EFindName FindType);
 
 	/** Create FName from non-null string with known length  */
 	CORE_API FName(int32 Len, const WIDECHAR* Name, EFindName FindType=FNAME_Add);
