Add AutoRTFM LLM support.

diff --git a/Engine/Source/Runtime/Core/Private/AutoRTFM/AutoRTFMUE.cpp b/Engine/Source/Runtime/Core/Private/AutoRTFM/AutoRTFMUE.cpp
--- a/Engine/Source/Runtime/Core/Private/AutoRTFM/AutoRTFMUE.cpp
+++ b/Engine/Source/Runtime/Core/Private/AutoRTFM/AutoRTFMUE.cpp
@@ -9,6 +9,7 @@
 #include "Containers/UnrealString.h"
 #include "GenericPlatform/GenericPlatformCrashContext.h"
 #include "HAL/IConsoleManager.h"
+#include "HAL/LowLevelMemTracker.h"
 #include "HAL/UnrealMemory.h"
 #include "Logging/LogMacros.h"
 #include "Logging/LogVerbosity.h"
@@ -354,22 +355,39 @@ void AutoRTFMEnsureFailure(const char* File, int Line, void* ProgramCounter, con
 
 } // anonymous namespace
 
+LLM_DEFINE_TAG(AutoRTFM);
+
+static void* AutoRTFMAllocate(size_t Size, size_t Alignment)
+{
+	LLM_SCOPE_BYTAG(AutoRTFM);
+	return FMemory::Malloc(Size, Alignment);
+}
+
+static void* AutoRTFMReallocate(void* Pointer, size_t Size, size_t Alignment)
+{
+	LLM_SCOPE_BYTAG(AutoRTFM);
+	return FMemory::Realloc(Pointer, Size, Alignment);
+}
+
+static void* AutoRTFMAllocateZeroed(size_t Size, size_t Alignment)
+{
+	LLM_SCOPE_BYTAG(AutoRTFM);
+	return FMemory::MallocZeroed(Size, Alignment);
+}
+
+static void AutoRTFMFree(void* Pointer)
+{
+	LLM_SCOPE_BYTAG(AutoRTFM);
+	FMemory::Free(Pointer);
+}
+
 CORE_API void AutoRTFM::InitializeForUE()
 {
 	AutoRTFM::ForTheRuntime::FExternAPI ExternAPI;
-	ExternAPI.Allocate = [](size_t Size, size_t Alignment)
-	{
-		return FMemory::Malloc(Size, Alignment);
-	};
-	ExternAPI.Reallocate = [](void* Pointer, size_t Size, size_t Alignment)
-	{
-		return FMemory::Realloc(Pointer, Size, Alignment);
-	};
-	ExternAPI.AllocateZeroed = [](size_t Size, size_t Alignment)
-	{
-		return FMemory::MallocZeroed(Size, Alignment);
-	};
-	ExternAPI.Free = FMemory::Free;
+	ExternAPI.Allocate = AutoRTFMAllocate;
+	ExternAPI.Reallocate = AutoRTFMReallocate;
+	ExternAPI.AllocateZeroed = AutoRTFMAllocateZeroed;
+	ExternAPI.Free = AutoRTFMFree;
 	ExternAPI.EnsureFailure = AutoRTFMEnsureFailure;
 	ExternAPI.Log = AutoRTFMLog;
 	ExternAPI.LogWithCallstack = AutoRTFMLogWithCallstack;
