Added missing header for deprecation macros.

#rnx

diff --git a/Engine/Source/Runtime/Core/Public/Templates/IsConst.h b/Engine/Source/Runtime/Core/Public/Templates/IsConst.h
--- a/Engine/Source/Runtime/Core/Public/Templates/IsConst.h
+++ b/Engine/Source/Runtime/Core/Public/Templates/IsConst.h
@@ -2,6 +2,8 @@
 
 #pragma once
 
+#include "Misc/CoreMiscDefines.h"
+
 // HEADER_UNIT_SKIP - Deprecated
 
 //UE_DEPRECATED_HEADER(5.7, "IsConst.h has been deprecated- please include <type_traits> and use std::is_const instead.")
