[AutoRTFM] Create a transactionally-safe version of the TSpscQueue.

This queue uses the same basic logic and structure as TSpscQueue, but instead of using atomics for synchronization, it uses a mutex. This will give us roughly equivalent performance in non-transactional code, and will properly enforce transactional safety without the overhead of requiring OnCommit for every access.

This fixes a recent regression in Rollbackapalooza, where the async collision trace system was trying to enqueue work in the middle of a transaction.

diff --git a/Engine/Source/Programs/AutoRTFMTests/Private/QueueTests.cpp b/Engine/Source/Programs/AutoRTFMTests/Private/QueueTests.cpp
new file mode 100644
--- /dev/null
+++ b/Engine/Source/Programs/AutoRTFMTests/Private/QueueTests.cpp
@@ -0,0 +1,81 @@
+// Copyright Epic Games, Inc. All Rights Reserved.
+
+#include "AutoRTFMTesting.h"
+#include "Catch2Includes.h"
+#include "Containers/TransactionallySafeSpscQueue.h"
+
+#include <string>
+#include <thread>
+
+// Test loosely based on code from Containers/ConcurrentQueuesTest.cpp.
+template <typename QueueType>
+void TestSpscQueueCorrectness(int Num, bool bTransact)
+{
+	QueueType Queue;
+	uint32 NumProduced = 0;
+	uint32 NumConsumed = 0;
+	std::thread ConsumeTask;
+
+	// Consumer
+	auto ConsumeLoop = [&Queue, &NumConsumed, Num]
+	{
+		while (NumConsumed != Num)
+		{
+			if (Queue.Dequeue().IsSet())
+			{
+				++NumConsumed;
+			}
+		}
+	};
+
+	// Producer
+	auto ProduceLoop = [&Queue, &NumProduced, Num]
+	{
+		do
+		{
+			Queue.Enqueue(typename QueueType::ElementType{});
+		} 
+		while (++NumProduced != Num);
+	};
+	
+	// Run the consumer loop on a helper thread.
+	UE_AUTORTFM_OPEN
+	{
+		ConsumeTask = std::thread{ConsumeLoop};
+	};
+
+	// Run the producer loop within a transaction.
+	if (bTransact)
+	{
+		AutoRTFM::Testing::Commit([&]
+		{
+			ProduceLoop();
+		});
+	}
+	else
+	{
+		ProduceLoop();
+	}
+
+	// Join the consumer thread to allow it to finish.
+	UE_AUTORTFM_OPEN
+	{
+		ConsumeTask.join();
+	};
+
+	// We should have produced and consumed exactly Num items, and the queue should be empty.
+	REQUIRE(NumProduced == Num);
+	REQUIRE(NumConsumed == Num);
+	REQUIRE(!Queue.Dequeue().IsSet());
+}
+
+TEST_CASE("TransactionallySafeSpscQueue")
+{
+	// Non-transactional, trivial and complex types
+	TestSpscQueueCorrectness<TTransactionallySafeSpscQueue<int>>(10000, /*bTransact=*/false);
+	TestSpscQueueCorrectness<TTransactionallySafeSpscQueue<std::string>>(10000, /*bTransact=*/false);
+
+	// Transactional, trivial and complex types
+	TestSpscQueueCorrectness<TTransactionallySafeSpscQueue<int>>(10000, /*bTransact=*/true);
+	TestSpscQueueCorrectness<TTransactionallySafeSpscQueue<std::string>>(10000, /*bTransact=*/true);
+}
diff --git a/Engine/Source/Runtime/Core/Public/Containers/TransactionallySafeSpscQueue.h b/Engine/Source/Runtime/Core/Public/Containers/TransactionallySafeSpscQueue.h
new file mode 100644
--- /dev/null
+++ b/Engine/Source/Runtime/Core/Public/Containers/TransactionallySafeSpscQueue.h
@@ -0,0 +1,192 @@
+// Copyright Epic Games, Inc. All Rights Reserved.
+
+#pragma once
+
+#include "CoreTypes.h"
+#include "Async/TransactionallySafeMutex.h"
+#include "Misc/Optional.h"
+#include "Misc/ScopeLock.h"
+#include "Templates/MemoryOps.h"
+#include "Templates/TypeCompatibleBytes.h"
+#include "Templates/UnrealTemplate.h"
+
+/**
+ * Fast, transactionally-safe single-producer/single-consumer unbounded concurrent queue. 
+ * Doesn't free memory until destruction but recycles consumed items.
+ * Based on TSpscQueue, which is itself based on http://www.1024cores.net/home/lock-free-algorithms/queues/unbounded-spsc-queue
+ * 
+ * The transactionally-safe queue uses a mutex to enforce thread safety instead of atomic 
+ * operations. The difference in performance compared to a TSpscQueue should be negligible 
+ * unless you are CPU-bound on constantly enqueueing and dequeueing objects as fast as 
+ * possible.
+ * 
+ * It is not safe to spin-wait on Dequeue from within an AutoRTFM transaction!
+ * The other thread's Enqueue will be blocked on the mutex, so you will deadlock inside the 
+ * spin-wait. This class works best with the game thread as the producer, and a separate 
+ * helper thread as the consumer.
+ */
+template<typename T, typename AllocatorType = FMemory>
+class TTransactionallySafeSpscQueue final
+{
+public:
+	using ElementType = T;
+
+	UE_NONCOPYABLE(TTransactionallySafeSpscQueue);
+
+	[[nodiscard]] TTransactionallySafeSpscQueue()
+	{
+		FNode* Node = ::new(AllocatorType::Malloc(sizeof(FNode), alignof(FNode))) FNode;
+		Head = First = Tail = TailCopy = Node;
+	}
+
+	~TTransactionallySafeSpscQueue()
+	{
+		// Nobody should have a reference to this class anymore, but we still take the mutex to
+		// guarantee that the queue isn't modified by another thread while a transaction is underway.
+		UE::TScopeLock Lock(Mutex);
+		FNode* Node = First;
+
+		// Delete all sentinel or unoccupied nodes.
+		bool bContinue = false;
+		do
+		{
+			FNode* Next = Node->Next;
+			bContinue = Node != Tail;
+			AllocatorType::Free(Node);
+			Node = Next;
+		} 
+		while (bContinue);
+
+		// Destroy and free all occupied nodes.
+		while (Node != nullptr)
+		{
+			FNode* Next = Node->Next;
+			DestructItem((ElementType*)&Node->Value);
+			AllocatorType::Free(Node);
+			Node = Next;
+		}
+	}
+
+	template <typename... ArgTypes>
+	void Enqueue(ArgTypes&&... Args)
+	{
+		FNode* Node = AllocNode();
+		::new((void*)&Node->Value) ElementType(Forward<ArgTypes>(Args)...);
+
+		UE::TScopeLock Lock(Mutex);
+		Head->Next = Node;
+		Head = Node;
+	}
+
+	// Returns NullOpt if the queue is empty.
+	// Spin-waiting on Dequeue from within an AutoRTFM transaction is likely to deadlock,
+	// as the matching Enqueue will be waiting on Mutex.
+	TOptional<ElementType> Dequeue()
+	{
+		Mutex.Lock();
+		FNode* LocalTail = Tail;
+		FNode* LocalTailNext = LocalTail->Next;
+		Mutex.Unlock();
+
+		if (LocalTailNext == nullptr)
+		{
+			return NullOpt;
+		}
+
+		ElementType* TailNextValue = (ElementType*)&LocalTailNext->Value;
+		TOptional<ElementType> Value{ MoveTemp(*TailNextValue) };
+		DestructItem(TailNextValue);
+
+		Mutex.Lock();
+		Tail = LocalTailNext;
+		Mutex.Unlock();
+
+		return Value;
+	}
+
+	bool Dequeue(ElementType& OutElem)
+	{
+		TOptional<ElementType> LocalElement = Dequeue();
+		if (LocalElement.IsSet())
+		{
+			OutElem = MoveTempIfPossible(LocalElement.GetValue());
+			return true;
+		}
+		
+		return false;
+	}
+
+	[[nodiscard]] bool IsEmpty() const
+	{
+		UE::TScopeLock Lock(Mutex);
+		FNode* LocalTail = Tail;
+		FNode* LocalTailNext = LocalTail->Next;
+
+		return LocalTailNext == nullptr;
+	}
+
+	// As there can be only one consumer, a consumer can safely "peek" the tail of the queue.
+	// Returns a pointer to the tail if the queue is not empty, nullptr otherwise.
+	// There's no overload with TOptional as it doesn't support references.
+	[[nodiscard]] ElementType* Peek() const
+	{
+		UE::TScopeLock Lock(Mutex);
+		FNode* LocalTail = Tail;
+		FNode* LocalTailNext = LocalTail->Next;
+
+		return LocalTailNext
+			? (ElementType*)&LocalTailNext->Value
+			: nullptr;
+	}
+
+private:
+	struct FNode
+	{
+		FNode* Next = nullptr;
+		TTypeCompatibleBytes<ElementType> Value;
+	};
+
+private:
+	FNode* AllocFromCache()
+	{
+		// The mutex must _already be held_ when calling this function.
+		FNode* Node = First;
+		First = First->Next;
+		Node->Next = nullptr;
+		return Node;
+	}
+
+	FNode* AllocNode()
+	{
+		// First, we attempt to allocate a node from internal node cache.
+		{
+			UE::TScopeLock Lock(Mutex);
+			if (First != TailCopy)
+			{
+				return AllocFromCache();
+			}
+
+			TailCopy = Tail;
+			if (First != TailCopy)
+			{
+				return AllocFromCache();
+			}
+		}
+
+		// Our cache is empty; allocate via ::operator new() instead.
+		return ::new(AllocatorType::Malloc(sizeof(FNode), alignof(FNode))) FNode();
+	}
+
+private:
+	// This mutex guards all accesses to this structure or to FNode::Next.
+	mutable UE::FTransactionallySafeMutex Mutex;
+
+	// consumer part 
+	// accessed mainly by consumer, infrequently by producer 
+	FNode* Tail; // tail of the queue 
+	// producer part 
+	// accessed only by producer 
+	FNode* Head; // head of the queue
+	FNode* First; // last unused node (tail of node cache) 
+	FNode* TailCopy; // points to the initially-created tail
+};
\ No newline at end of file
