Made the MSVC implementations of FPlatformMath::CountLeadingZeros and FPlatformMath::CountTrailingZeros constexpr.

diff --git a/Engine/Source/Runtime/Core/Public/Microsoft/MicrosoftPlatformMath.h b/Engine/Source/Runtime/Core/Public/Microsoft/MicrosoftPlatformMath.h
--- a/Engine/Source/Runtime/Core/Public/Microsoft/MicrosoftPlatformMath.h
+++ b/Engine/Source/Runtime/Core/Public/Microsoft/MicrosoftPlatformMath.h
@@ -34,18 +34,63 @@ struct FMicrosoftPlatformMathBase : public TUnrealPlatformMathSSE4Base<FGenericP
 		return BitIndex;
 	}
 
-	static inline uint8 CountLeadingZeros8(uint8 Value)
+	static inline constexpr uint8 CountLeadingZeros8(uint8 Value)
 	{
-		unsigned long BitIndex;
-		_BitScanReverse(&BitIndex, uint32(Value)*2 + 1);
-		return uint8(8 - BitIndex);
+		// return 8 if value was 0
+
+		UE_IF_CONSTEVAL
+		{
+			if (Value == 0)
+			{
+				return 8;
+			}
+
+			uint8 Result = 0;
+			for (;;)
+			{
+				if (Value & 0x80)
+				{
+					return Result;
+				}
+				Value <<= 1;
+				++Result;
+			}
+		}
+		else
+		{
+			unsigned long BitIndex;
+			_BitScanReverse(&BitIndex, uint32(Value)*2 + 1);
+			return uint8(8 - BitIndex);
+		}
 	}
 
-	static inline uint32 CountTrailingZeros(uint32 Value)
+	static inline constexpr uint32 CountTrailingZeros(uint32 Value)
 	{
 		// return 32 if value was 0
-		unsigned long BitIndex;	// 0-based, where the LSB is 0 and MSB is 31
-		return _BitScanForward( &BitIndex, Value ) ? BitIndex : 32;
+
+		UE_IF_CONSTEVAL
+		{
+			if (Value == 0)
+			{
+				return 32;
+			}
+
+			uint32 Result = 0;
+			for (;;)
+			{
+				if (Value & 1)
+				{
+					return Result;
+				}
+				Value >>= 1;
+				++Result;
+			}
+		}
+		else
+		{
+			unsigned long BitIndex;	// 0-based, where the LSB is 0 and MSB is 31
+			return _BitScanForward( &BitIndex, Value ) ? BitIndex : 32;
+		}
 	}
 
 	static inline uint32 CeilLogTwo( uint32 Arg )
@@ -89,28 +134,94 @@ struct FMicrosoftPlatformMathBase : public TUnrealPlatformMathSSE4Base<FGenericP
 		return 64 - CountLeadingZeros64(Arg - 1);
 	}
 
-	static inline uint64 CountLeadingZeros64(uint64 Value)
+	static inline constexpr uint64 CountLeadingZeros64(uint64 Value)
 	{
-		//https://godbolt.org/z/Ejh5G4vPK	
 		// return 64 if value if was 0
-		unsigned long BitIndex;
-		if ( ! _BitScanReverse64(&BitIndex, Value) ) BitIndex = -1;
-        return 63 - BitIndex;
+
+		UE_IF_CONSTEVAL
+		{
+			if (Value == 0)
+			{
+				return 64;
+			}
+
+			uint64 Result = 0;
+			for (;;)
+			{
+				if (Value & 0x8000000000000000)
+				{
+					return Result;
+				}
+				Value <<= 1;
+				++Result;
+			}
+		}
+		else
+		{
+			//https://godbolt.org/z/Ejh5G4vPK
+			unsigned long BitIndex;
+			if ( ! _BitScanReverse64(&BitIndex, Value) ) BitIndex = -1;
+			return 63 - BitIndex;
+		}
 	}
 
-	static inline uint64 CountTrailingZeros64(uint64 Value)
+	static inline constexpr uint64 CountTrailingZeros64(uint64 Value)
 	{
 		// return 64 if Value is 0
-		unsigned long BitIndex;	// 0-based, where the LSB is 0 and MSB is 63
-		return _BitScanForward64( &BitIndex, Value ) ? BitIndex : 64;
+
+		UE_IF_CONSTEVAL
+		{
+			if (Value == 0)
+			{
+				return 64;
+			}
+
+			uint64 Result = 0;
+			for (;;)
+			{
+				if (Value & 1)
+				{
+					return Result;
+				}
+				Value >>= 1;
+				++Result;
+			}
+		}
+		else
+		{
+			unsigned long BitIndex;	// 0-based, where the LSB is 0 and MSB is 63
+			return _BitScanForward64( &BitIndex, Value ) ? BitIndex : 64;
+		}
 	}
 
-	static inline uint32 CountLeadingZeros(uint32 Value)
+	static inline constexpr uint32 CountLeadingZeros(uint32 Value)
 	{
 		// return 32 if value is zero
-		unsigned long BitIndex;
-		_BitScanReverse64(&BitIndex, uint64(Value)*2 + 1);
-		return 32 - BitIndex;
+
+		UE_IF_CONSTEVAL
+		{
+			if (Value == 0)
+			{
+				return 32;
+			}
+
+			uint32 Result = 0;
+			for (;;)
+			{
+				if (Value & 0x80000000)
+				{
+					return Result;
+				}
+				Value <<= 1;
+				++Result;
+			}
+		}
+		else
+		{
+			unsigned long BitIndex;
+			_BitScanReverse64(&BitIndex, uint64(Value)*2 + 1);
+			return 32 - BitIndex;
+		}
 	}
 
 #else // 32-bit
@@ -120,7 +231,7 @@ struct FMicrosoftPlatformMathBase : public TUnrealPlatformMathSSE4Base<FGenericP
 		// return 32 if value is zero
 		unsigned long BitIndex;
 		if ( ! _BitScanReverse(&BitIndex, Value) ) BitIndex = -1;
-        return 31 - BitIndex;
+		return 31 - BitIndex;
 	}
 
 #endif
