FPaths::CreateStandardPath: Downgrade display log to verbose to prevent spam.
#rnx
#incrementalcook

diff --git a/Engine/Source/Runtime/Core/Private/Misc/Paths.cpp b/Engine/Source/Runtime/Core/Private/Misc/Paths.cpp
--- a/Engine/Source/Runtime/Core/Private/Misc/Paths.cpp
+++ b/Engine/Source/Runtime/Core/Private/Misc/Paths.cpp
@@ -1380,7 +1380,7 @@ FString FPaths::CreateStandardFilename(const FString& InPath)
 		WithSlashes += FStringView(FullPath).RightChop(RootDirectoryView.Len());
 		return WithSlashes;
 	}
-	UE_LOG(LogPaths, Display, TEXT("CreateStandardFilename: Unexpected: %s does not start with %.*s."),
+	UE_LOG(LogPaths, Verbose, TEXT("CreateStandardFilename: Unexpected: %s does not start with %.*s."),
 		*FullPath, RootDirectoryView.Len(), RootDirectoryView.GetData());
 	FullPath.ReplaceInline(RootDirectory, *RelativePathToRoot);
 	return FullPath;
