Add comments to FBox to clarify that (here) Extents mean the half-size. Avoids confusion because some other box types in UE use Extents to refer to the full extents.

diff --git a/Engine/Source/Runtime/Core/Public/Math/Box.h b/Engine/Source/Runtime/Core/Public/Math/Box.h
--- a/Engine/Source/Runtime/Core/Public/Math/Box.h
+++ b/Engine/Source/Runtime/Core/Public/Math/Box.h
@@ -400,9 +400,9 @@ public:
 	inline TVector<T> GetClosestPointTo( const TVector<T>& Point ) const;
 
 	/**
-	 * Gets the extents of this box.
+	 * Gets the (half) extents of this box.
 	 *
-	 * @return The box extents.
+	 * @return The box (half) extents.
 	 * @see GetCenter, GetCenterAndExtents, GetSize, GetVolume
 	 */
 	inline TVector<T> GetExtent() const
