remove deprecated Named Chunk Install delegate. All code should now be using Named Chunk Complete delegate.

#rnx

diff --git a/Engine/Source/Runtime/Core/Private/GenericPlatform/GenericPlatformChunkInstall.cpp b/Engine/Source/Runtime/Core/Private/GenericPlatform/GenericPlatformChunkInstall.cpp
--- a/Engine/Source/Runtime/Core/Private/GenericPlatform/GenericPlatformChunkInstall.cpp
+++ b/Engine/Source/Runtime/Core/Private/GenericPlatform/GenericPlatformChunkInstall.cpp
@@ -46,13 +46,6 @@ void FGenericPlatformChunkInstall::DoNamedChunkCompleteCallbacks( const TArrayVi
 
 				NamedChunkCompleteDelegate.Broadcast(Param);
 			}
-
-			PRAGMA_DISABLE_DEPRECATION_WARNINGS
-			if (bIsInstalled && NamedChunkInstallDelegate.IsBound())
-			{
-				NamedChunkInstallDelegate.Broadcast(NamedChunk, bHasSucceeded);
-			}
-			PRAGMA_ENABLE_DEPRECATION_WARNINGS
 		}
 	});
 }
diff --git a/Engine/Source/Runtime/Core/Public/GenericPlatform/GenericPlatformChunkInstall.h b/Engine/Source/Runtime/Core/Public/GenericPlatform/GenericPlatformChunkInstall.h
--- a/Engine/Source/Runtime/Core/Public/GenericPlatform/GenericPlatformChunkInstall.h
+++ b/Engine/Source/Runtime/Core/Public/GenericPlatform/GenericPlatformChunkInstall.h
@@ -104,10 +104,6 @@ public:
 DECLARE_DELEGATE_TwoParams(FPlatformChunkInstallDelegate, uint32, bool);
 DECLARE_MULTICAST_DELEGATE_TwoParams(FPlatformChunkInstallMultiDelegate, uint32, bool);
 
-/** Deprecated delegate called when a Named Chunk either successfully installs or fails to install, bool is success */
-DECLARE_DELEGATE_TwoParams(FPlatformNamedChunkInstallDelegate, FName, bool);
-DECLARE_MULTICAST_DELEGATE_TwoParams(FPlatformNamedChunkInstallMultiDelegate, FName, bool);
-
 /** Delegate called when a Named Chunk either successfully installs or fails to install */
 DECLARE_DELEGATE_OneParam(FPlatformNamedChunkCompleteDelegate, const FNamedChunkCompleteCallbackParam&);
 DECLARE_MULTICAST_DELEGATE_OneParam(FPlatformNamedChunkCompleteMultiDelegate, const FNamedChunkCompleteCallbackParam&);
@@ -297,21 +293,6 @@ public:
 	 */
 	virtual TArray<FName> GetNamedChunksByType(ENamedChunkType NamedChunkType) const = 0;
 
-	/** 
-	 * Request a delegate callback on named chunk install completion or failure. Request may not be respected.
-	 * @param Delegate		The delegate to call when any named chunk is installed or fails to install
-	 * @return				Handle to the bound delegate
-	 */
-	UE_DEPRECATED(5.4, "use AddNamedChunkCompleteDelegate instead")
-	virtual FDelegateHandle AddNamedChunkInstallDelegate( FPlatformNamedChunkInstallDelegate Delegate ) = 0;
-
-	/**
-	 * Remove a delegate callback on named chunk install completion.
-	 * @param Delegate		The delegate to remove.
-	 */
-	UE_DEPRECATED(5.4, "use RemoveNamedChunkCompleteDelegate instead")
-	virtual void RemoveNamedChunkInstallDelegate( FDelegateHandle Delegate ) = 0;
-	
 	/** 
 	 * Request a delegate callback on named chunk install completion or failure. Request may not be respected.
 	 * @param Delegate		The delegate to call when any named chunk is installed or fails to install
@@ -515,16 +496,6 @@ public:
 		return TArray<FName>();
 	}
 
-	virtual FDelegateHandle AddNamedChunkInstallDelegate(FPlatformNamedChunkInstallDelegate Delegate) override
-	{
-		return NamedChunkInstallDelegate.Add(Delegate);
-	}
-
-	virtual void RemoveNamedChunkInstallDelegate(FDelegateHandle Delegate) override
-	{
-		NamedChunkInstallDelegate.Remove(Delegate);
-	}
-
 	virtual void RemoveNamedChunkCompleteDelegate(FDelegateHandle Delegate) override
 	{
 		NamedChunkCompleteDelegate.Remove(Delegate);
@@ -568,7 +539,6 @@ protected:
 
 	/** Delegates called when installation succeeds or fails */
 	FPlatformChunkInstallMultiDelegate InstallDelegate;
-	FPlatformNamedChunkInstallMultiDelegate NamedChunkInstallDelegate;
 	FPlatformNamedChunkCompleteMultiDelegate NamedChunkCompleteDelegate;
 
 	virtual EChunkLocation::Type GetChunkLocation(uint32 ChunkID) override
