Removed obsolete small function optimization macros and Win32 support.

diff --git a/Engine/Source/Runtime/Core/Public/Misc/CoreDefines.h b/Engine/Source/Runtime/Core/Public/Misc/CoreDefines.h
--- a/Engine/Source/Runtime/Core/Public/Misc/CoreDefines.h
+++ b/Engine/Source/Runtime/Core/Public/Misc/CoreDefines.h
@@ -35,31 +35,16 @@ Macros for enabling heap storage instead of inline storage on delegate types.
 Can be overridden by setting to 1 or 0 in the project's .Target.cs files.
 -----------------------------------------------------------------------------*/
 
-#ifdef USE_SMALL_DELEGATES
-	// UE_DEPRECATED(4.22, "USE_SMALL_DELEGATES has been removed")
-	#error USE_SMALL_DELEGATES has been removed - please use NUM_DELEGATE_INLINE_BYTES instead
-#else
-	#ifndef NUM_DELEGATE_INLINE_BYTES
-		#define NUM_DELEGATE_INLINE_BYTES 0
-	#endif
+#ifndef NUM_DELEGATE_INLINE_BYTES
+	#define NUM_DELEGATE_INLINE_BYTES 0
 #endif
 
-#ifdef USE_SMALL_MULTICAST_DELEGATES
-	// UE_DEPRECATED(4.22, "USE_SMALL_MULTICAST_DELEGATES has been removed")
-	#error USE_SMALL_MULTICAST_DELEGATES has been removed - please use NUM_MULTICAST_DELEGATE_INLINE_ENTRIES instead
-#else
-	#ifndef NUM_MULTICAST_DELEGATE_INLINE_ENTRIES
-		#define NUM_MULTICAST_DELEGATE_INLINE_ENTRIES 0
-	#endif
+#ifndef NUM_MULTICAST_DELEGATE_INLINE_ENTRIES
+	#define NUM_MULTICAST_DELEGATE_INLINE_ENTRIES 0
 #endif
 
-#ifdef USE_SMALL_TFUNCTIONS
-	// UE_DEPRECATED(4.22, "USE_SMALL_TFUNCTIONS has been removed")
-	#error USE_SMALL_TFUNCTIONS has been removed - please use NUM_TFUNCTION_INLINE_BYTES instead
-#else
-	#ifndef NUM_TFUNCTION_INLINE_BYTES
-		#define NUM_TFUNCTION_INLINE_BYTES 24
-	#endif
+#ifndef NUM_TFUNCTION_INLINE_BYTES
+	#define NUM_TFUNCTION_INLINE_BYTES 24
 #endif
 
 #ifndef NO_CVARS
diff --git a/Engine/Source/Runtime/Core/Public/Templates/Function.h b/Engine/Source/Runtime/Core/Public/Templates/Function.h
--- a/Engine/Source/Runtime/Core/Public/Templates/Function.h
+++ b/Engine/Source/Runtime/Core/Public/Templates/Function.h
@@ -24,10 +24,7 @@
 	#endif
 #endif
 
-#if defined(_WIN32) && !defined(_WIN64) && (!defined(ALLOW_TFUNCTION_INLINE_ALLOCATORS_ON_WIN32) || !ALLOW_TFUNCTION_INLINE_ALLOCATORS_ON_WIN32)
-	// Don't use inline storage on Win32, because that will affect the alignment of TFunction, and we can't pass extra-aligned types by value on Win32.
-	#define TFUNCTION_USES_INLINE_STORAGE 0
-#elif !defined(NUM_TFUNCTION_INLINE_BYTES) || NUM_TFUNCTION_INLINE_BYTES == 0
+#if !defined(NUM_TFUNCTION_INLINE_BYTES) || NUM_TFUNCTION_INLINE_BYTES == 0
 	#define TFUNCTION_USES_INLINE_STORAGE 0
 #else
 	#define TFUNCTION_USES_INLINE_STORAGE 1
