Fixed constraint on TStringView constructor, broken when adding UE::CNotCVRefTo.

diff --git a/Engine/Source/Runtime/Core/Public/Containers/StringView.h b/Engine/Source/Runtime/Core/Public/Containers/StringView.h
--- a/Engine/Source/Runtime/Core/Public/Containers/StringView.h
+++ b/Engine/Source/Runtime/Core/Public/Containers/StringView.h
@@ -144,7 +144,7 @@ public:
 	}
 
 	/** Construct a view from a contiguous range of characters, such as FString or TStringBuilder. */
-	template <UE::CNotCVRefTo<CharType> CharRangeType>
+	template <UE::CNotCVRefTo<ViewType> CharRangeType>
 		// !std::is_pointer is used here to force the pointer constructor to be used when passed an array or pointer.
 		requires (!std::is_pointer_v<std::decay_t<CharRangeType>> && UE::CCompatibleStringViewable<CharRangeType, CharType>)
 	[[nodiscard]] constexpr inline TStringView(const CharRangeType& InRange UE_LIFETIMEBOUND)
