SecureHash const correctness: FSHA1.GetHash can be const.
#rnx

diff --git a/Engine/Source/Runtime/Core/Private/Misc/SecureHash.cpp b/Engine/Source/Runtime/Core/Private/Misc/SecureHash.cpp
--- a/Engine/Source/Runtime/Core/Private/Misc/SecureHash.cpp
+++ b/Engine/Source/Runtime/Core/Private/Misc/SecureHash.cpp
@@ -1266,7 +1266,7 @@ void FSHA1::Final()
 }
 
 // Get the raw message digest
-void FSHA1::GetHash(uint8 *puDest)
+void FSHA1::GetHash(uint8 *puDest) const
 {
 	FMemory::Memcpy(puDest, m_digest, 20);
 }
diff --git a/Engine/Source/Runtime/Core/Public/Misc/SecureHash.h b/Engine/Source/Runtime/Core/Public/Misc/SecureHash.h
--- a/Engine/Source/Runtime/Core/Public/Misc/SecureHash.h
+++ b/Engine/Source/Runtime/Core/Public/Misc/SecureHash.h
@@ -354,7 +354,7 @@ public:
 	}
 
 	// Report functions: as pre-formatted and raw data
-	CORE_API void GetHash(uint8 *puDest);
+	CORE_API void GetHash(uint8 *puDest) const;
 
 	/**
 	 * Calculate the hash on a single block and return it
