Added FRecursiveWordMutex to support very early and/or very late recursive mutex usage

This mutex is trivially constructible and trivially destructible, which supports it being used before static initialization and after static destruction. It has no dependency on ParkingLot, which makes it safe for use in code that ParkingLot depends on such as GMalloc, LLM, and asserts.

diff --git a/Engine/Source/Runtime/Core/Private/Async/RecursiveWordMutex.cpp b/Engine/Source/Runtime/Core/Private/Async/RecursiveWordMutex.cpp
new file mode 100644
--- /dev/null
+++ b/Engine/Source/Runtime/Core/Private/Async/RecursiveWordMutex.cpp
@@ -0,0 +1,52 @@
+// Copyright Epic Games, Inc. All Rights Reserved.
+
+#include "Async/RecursiveWordMutex.h"
+#include "HAL/PlatformTLS.h"
+
+namespace UE
+{
+
+bool FRecursiveWordMutex::TryLock()
+{
+	const uint32 CurrentThreadId = FPlatformTLS::GetCurrentThreadId();
+	if (ThreadId.load(std::memory_order_relaxed) == CurrentThreadId)
+	{
+		++RecursionCount;
+		return true;
+	}
+	else if (Mutex.TryLock())
+	{
+		ThreadId.store(CurrentThreadId, std::memory_order_relaxed);
+		return true;
+	}
+	return false;
+}
+
+void FRecursiveWordMutex::Lock()
+{
+	const uint32 CurrentThreadId = FPlatformTLS::GetCurrentThreadId();
+	if (ThreadId.load(std::memory_order_relaxed) == CurrentThreadId)
+	{
+		++RecursionCount;
+	}
+	else
+	{
+		Mutex.Lock();
+		ThreadId.store(CurrentThreadId, std::memory_order_relaxed);
+	}
+}
+
+void FRecursiveWordMutex::Unlock()
+{
+	if (RecursionCount > 0)
+	{
+		--RecursionCount;
+	}
+	else
+	{
+		ThreadId.store(0, std::memory_order_relaxed);
+		Mutex.Unlock();
+	}
+}
+
+} // UE
diff --git a/Engine/Source/Runtime/Core/Public/Async/RecursiveWordMutex.h b/Engine/Source/Runtime/Core/Public/Async/RecursiveWordMutex.h
new file mode 100644
--- /dev/null
+++ b/Engine/Source/Runtime/Core/Public/Async/RecursiveWordMutex.h
@@ -0,0 +1,37 @@
+// Copyright Epic Games, Inc. All Rights Reserved.
+
+#pragma once
+
+#include "Async/WordMutex.h"
+#include <atomic>
+
+#define UE_API CORE_API
+
+namespace UE
+{
+
+/**
+ * A recursive mutex that is the size of a pointer and does not depend on ParkingLot.
+ *
+ * Prefer FRecursiveMutex to FRecursiveWordMutex whenever possible.
+ * This mutex is not fair and supports recursive locking.
+ *
+ * This type is valuable when a mutex must be trivially constructible, trivially
+ * destructible, or must be functional before or after static initialization.
+ */
+class FRecursiveWordMutex final
+{
+public:
+	UE_API bool TryLock();
+	UE_API void Lock();
+	UE_API void Unlock();
+
+private:
+	FWordMutex Mutex;
+	uint32 RecursionCount = 0;
+	std::atomic<uint32> ThreadId = 0;
+};
+
+} // UE
+
+#undef UE_API
