Removed unnecessary Abs calls from FPlatformManualResetEvent

diff --git a/Engine/Source/Runtime/Core/Private/Android/AndroidPlatformManualResetEvent.cpp b/Engine/Source/Runtime/Core/Private/Android/AndroidPlatformManualResetEvent.cpp
--- a/Engine/Source/Runtime/Core/Private/Android/AndroidPlatformManualResetEvent.cpp
+++ b/Engine/Source/Runtime/Core/Private/Android/AndroidPlatformManualResetEvent.cpp
@@ -19,7 +19,7 @@ inline static const struct timespec* SecondsToTimeSpec(double Seconds, struct ti
 	{
 		const double TotalSec = Seconds >= 0.0 ? Seconds : 0.0;
 		int32 Sec = int32(TotalSec);
-		int32 NSec = int32(FPlatformMath::Abs((TotalSec - Sec) * 1e9));
+		int32 NSec = int32((TotalSec - Sec) * 1e9);
 		OutTime.tv_sec = Sec;
 		OutTime.tv_nsec = NSec;
 		return &OutTime;
diff --git a/Engine/Source/Runtime/Core/Private/Unix/UnixPlatformManualResetEvent.cpp b/Engine/Source/Runtime/Core/Private/Unix/UnixPlatformManualResetEvent.cpp
--- a/Engine/Source/Runtime/Core/Private/Unix/UnixPlatformManualResetEvent.cpp
+++ b/Engine/Source/Runtime/Core/Private/Unix/UnixPlatformManualResetEvent.cpp
@@ -20,7 +20,7 @@ inline static const struct timespec* SecondsToTimeSpec(double Seconds, struct ti
 	{
 		const double TotalSec = Seconds >= 0.0 ? Seconds : 0.0;
 		int32 Sec = int32(TotalSec);
-		int32 NSec = int32(FPlatformMath::Abs((TotalSec - Sec) * 1e9));
+		int32 NSec = int32((TotalSec - Sec) * 1e9);
 		OutTime.tv_sec = Sec;
 		OutTime.tv_nsec = NSec;
 		return &OutTime;
