Fixed potential race condition in FConfigCacheIni::ForPlatform adding plugins without holding the lock.

diff --git a/Engine/Source/Runtime/Core/Private/Misc/ConfigCacheIni.cpp b/Engine/Source/Runtime/Core/Private/Misc/ConfigCacheIni.cpp
--- a/Engine/Source/Runtime/Core/Private/Misc/ConfigCacheIni.cpp
+++ b/Engine/Source/Runtime/Core/Private/Misc/ConfigCacheIni.cpp
@@ -6772,32 +6772,28 @@ FConfigCacheIni* FConfigCacheIni::ForPlatform(FName PlatformName)
 	FConfigCacheIni* PlatformConfig = nullptr;
 	TArray<FName> PendingModificationPlugins;
 
-	{
-		// protect against other threads clearing the array, or two threads trying to read in a missing platform at the same time
-		FScopeLock Lock(&ConfigForPlatformLock);
-
-		PlatformConfig = ConfigForPlatform.FindRef(PlatformName);
+	// protect against other threads clearing the array, or two threads trying to read in a missing platform at the same time
+	FScopeLock Lock(&ConfigForPlatformLock);
 
-		// read any missing platform configs now, on demand (this will happen when WITH_EDITOR is 0)
-		if (PlatformConfig == nullptr)
-		{
-			double Start = FPlatformTime::Seconds();
+	PlatformConfig = ConfigForPlatform.FindRef(PlatformName);
 
-			PlatformConfig = ConfigForPlatform.Add(PlatformName, new FConfigCacheIni(EConfigCacheType::Temporary, PlatformName, true /* bInGloballyRegistered */));
-			FConfigContext Context = FConfigContext::ReadIntoConfigSystem(PlatformConfig, PlatformName.ToString());
-			InitializeKnownConfigFiles(Context);
+	// read any missing platform configs now, on demand (this will happen when WITH_EDITOR is 0)
+	if (PlatformConfig == nullptr)
+	{
+		double Start = FPlatformTime::Seconds();
 
-			UE_LOG(LogConfig, Display, TEXT("Read in platform %s ini files took %.2f seconds"), *PlatformName.ToString(), FPlatformTime::Seconds() - Start);
-		}
+		PlatformConfig = ConfigForPlatform.Add(PlatformName, new FConfigCacheIni(EConfigCacheType::Temporary, PlatformName, true /* bInGloballyRegistered */));
+		FConfigContext Context = FConfigContext::ReadIntoConfigSystem(PlatformConfig, PlatformName.ToString());
+		InitializeKnownConfigFiles(Context);
 
-		PendingModificationPlugins = MoveTemp(PlatformConfig->PendingModificationPlugins);
-		PlatformConfig->PendingModificationPlugins.Empty();
+		UE_LOG(LogConfig, Display, TEXT("Read in platform %s ini files took %.2f seconds"), *PlatformName.ToString(), FPlatformTime::Seconds() - Start);
 	}
 
-	{
-		// delayed plugin injection
-		PlatformConfig->AddPluginsToBranches(PendingModificationPlugins, nullptr);
-	}
+	PendingModificationPlugins = MoveTemp(PlatformConfig->PendingModificationPlugins);
+	PlatformConfig->PendingModificationPlugins.Empty();
+
+	// delayed plugin injection
+	PlatformConfig->AddPluginsToBranches(PendingModificationPlugins, nullptr);
 
 	return PlatformConfig;
 
