[ChaosVD] Added missing serialization for some properties in FChaosVDParticleDynamicMisc (MLinearEtherDrag, MSleepThresholdMultiplier, MCollisionConstraintFlag).

diff --git a/Engine/Source/Runtime/Core/Public/UObject/FortniteMainBranchObjectVersions.inl b/Engine/Source/Runtime/Core/Public/UObject/FortniteMainBranchObjectVersions.inl
--- a/Engine/Source/Runtime/Core/Public/UObject/FortniteMainBranchObjectVersions.inl
+++ b/Engine/Source/Runtime/Core/Public/UObject/FortniteMainBranchObjectVersions.inl
@@ -657,5 +657,9 @@ FORTNITE_MAIN_VERSION(ChangeDefaultAlphaBlendType, 214)
 
 // Moved Position/Velocity/Projection Iteration Counts from FChaosVDFRigidParticleControlFlags to FChaosVDParticleDynamicMisc
 FORTNITE_MAIN_VERSION(PerParticleIterationCountMovedToDynamicMisc, 215)
+
+// Added missing custom serialization for some properties in the ParticleDynamicMisc structure used by the Chaos Visual Debugger
+FORTNITE_MAIN_VERSION(AddedMissingSerializationForPropertiesInDynamicMisc, 216)
+
 // -----<new versions can be added above this line>-------------------------------------------------
 #undef FORTNITE_MAIN_VERSION
\ No newline at end of file
diff --git a/Engine/Source/Runtime/Experimental/ChaosVisualDebugger/Private/DataWrappers/ChaosVDParticleDataWrapper.cpp b/Engine/Source/Runtime/Experimental/ChaosVisualDebugger/Private/DataWrappers/ChaosVDParticleDataWrapper.cpp
--- a/Engine/Source/Runtime/Experimental/ChaosVisualDebugger/Private/DataWrappers/ChaosVDParticleDataWrapper.cpp
+++ b/Engine/Source/Runtime/Experimental/ChaosVisualDebugger/Private/DataWrappers/ChaosVDParticleDataWrapper.cpp
@@ -159,6 +159,13 @@ bool FChaosVDParticleDynamicMisc::Serialize(FArchive& Ar)
 		Ar << ProjectionSolverIterationCount;
 	}
 
+	if (Ar.CustomVer(FFortniteMainBranchObjectVersion::GUID) >= FFortniteMainBranchObjectVersion::AddedMissingSerializationForPropertiesInDynamicMisc)
+	{
+		Ar << MLinearEtherDrag;
+		Ar << MSleepThresholdMultiplier;
+		Ar << MCollisionConstraintFlag;
+	}
+
 	return true;
 }
 
