Removed dependencies on ParkingLot from Windows stack walking

diff --git a/Engine/Source/Runtime/Core/Private/Windows/WindowsPlatformProcess.cpp b/Engine/Source/Runtime/Core/Private/Windows/WindowsPlatformProcess.cpp
--- a/Engine/Source/Runtime/Core/Private/Windows/WindowsPlatformProcess.cpp
+++ b/Engine/Source/Runtime/Core/Private/Windows/WindowsPlatformProcess.cpp
@@ -205,9 +205,11 @@ void* FWindowsPlatformProcess::GetDllHandle( const TCHAR* FileName )
 {
 	check(FileName);
 
-	// Combine the explicit DLL search directories with the contents of the directory stack 
+	// Combine the explicit DLL search directories with the contents of the directory stack
 	// Note that the search path logic here needs to match the logic found in ResolveImport
-	TArray<FString> SearchPaths;
+	// Use an inline allocator because this can be called from assertions within the
+	// allocator and the array allocations may fail in that situation.
+	TArray<FString, TInlineAllocator<16>> SearchPaths;
 	SearchPaths.Reserve(1 + ((DllDirectoryStack.Num() > 0) ? 1 : 0) + DllDirectories.Num());
 	SearchPaths.Add(FPlatformProcess::GetModulesDirectory());
 	if(DllDirectoryStack.Num() > 0)
@@ -2200,7 +2202,7 @@ static bool ReadLibraryImports(const TCHAR* FileName, TArray<FString>& ImportNam
 	return bResult;
 }
 
-bool FWindowsPlatformProcess::ResolveImport(const FString& Name, const TArray<FString>& SearchPaths, FString& OutFileName)
+bool FWindowsPlatformProcess::ResolveImport(const FString& Name, TConstArrayView<FString> SearchPaths, FString& OutFileName)
 {
 	IPlatformFile& PlatformFile = IPlatformFile::GetPlatformPhysical();
 	auto SearchPathsFunc = [&PlatformFile, &OutFileName, &SearchPaths, &Name](int StartIdx, int EndIdx)
@@ -2274,7 +2276,7 @@ UE_TRACE_EVENT_END()
 
 #endif // CPUPROFILERTRACE_ENABLED
 
-void FWindowsPlatformProcess::ResolveMissingImportsRecursive(const FString& FileName, const TArray<FString>& SearchPaths, TArray<FString>& ImportFileNames, TSet<FString>& VisitedImportNames)
+void FWindowsPlatformProcess::ResolveMissingImportsRecursive(const FString& FileName, TConstArrayView<FString> SearchPaths, TArray<FString>& ImportFileNames, TSet<FString>& VisitedImportNames)
 {
 #if CPUPROFILERTRACE_ENABLED
 	UE_TRACE_LOG_SCOPED_T(Cpu, ResolveMissingImports, CpuChannel)
@@ -2315,7 +2317,7 @@ void FWindowsPlatformProcess::ResolveMissingImportsRecursive(const FString& File
  * @param FileName Path to the library to load
  * @param SearchPaths Search directories to scan for imports
  */
-static void LogImportDiagnostics(const FString& FileName, const TArray<FString>& SearchPaths)
+static void LogImportDiagnostics(const FString& FileName, TConstArrayView<FString> SearchPaths)
 {
 	TArray<FString> ImportNames;
 	if(ReadLibraryImports(*FileName, ImportNames))
@@ -2347,7 +2349,7 @@ UE_TRACE_EVENT_END()
 
 #endif // CPUPROFILERTRACE_ENABLED
 
-void *FWindowsPlatformProcess::LoadLibraryWithSearchPaths(const FString& FileName, const TArray<FString>& SearchPaths)
+void* FWindowsPlatformProcess::LoadLibraryWithSearchPaths(const FString& FileName, TConstArrayView<FString> SearchPaths)
 {
 	UE_SCOPED_IO_ACTIVITY(*WriteToString<256>("Loading Dll ", FileName));
 
diff --git a/Engine/Source/Runtime/Core/Private/Windows/WindowsPlatformStackWalk.cpp b/Engine/Source/Runtime/Core/Private/Windows/WindowsPlatformStackWalk.cpp
--- a/Engine/Source/Runtime/Core/Private/Windows/WindowsPlatformStackWalk.cpp
+++ b/Engine/Source/Runtime/Core/Private/Windows/WindowsPlatformStackWalk.cpp
@@ -1,10 +1,11 @@
 // Copyright Epic Games, Inc. All Rights Reserved.
 
 #include "Windows/WindowsPlatformStackWalk.h"
-#include "Async/RecursiveMutex.h"
+#include "Async/RecursiveWordMutex.h"
 #include "Async/UniqueLock.h"
 #include "HAL/PlatformMemory.h"
 #include "HAL/PlatformMisc.h"
+#include "HAL/PlatformMutex.h"
 #include "Logging/LogMacros.h"
 #include "Math/UnrealMathUtility.h"
 #include "HAL/UnrealMemory.h"
@@ -15,11 +16,9 @@
 #include "Misc/Paths.h"
 #include "Misc/CommandLine.h"
 #include "HAL/PlatformProcess.h"
-#include "HAL/CriticalSection.h"
 #include "CoreGlobals.h"
 #include "Misc/ConfigCacheIni.h"
 #include "Misc/OutputDeviceRedirector.h"
-#include "Misc/ScopeLock.h"
 
 #include "Windows/AllowWindowsPlatformTypes.h"
 THIRD_PARTY_INCLUDES_START
@@ -39,7 +38,7 @@ THIRD_PARTY_INCLUDES_END
 -----------------------------------------------------------------------------*/
 
 //@note Global lock as a tentative fix to narrow down UE-192420
-static UE::FRecursiveMutex GStackWalkingLock;
+static UE::FRecursiveWordMutex GStackWalkingLock;
 
 /** Whether appInitStackWalking() has been called successfully or not. */
 static bool GStackWalkingInitialized = false;
@@ -175,8 +174,8 @@ UE_AUTORTFM_ALWAYS_OPEN void FWindowsPlatformStackWalk::ThreadStackWalkAndDump(A
 	if (ThreadHandle)
 	{
 		// sync with other threads that may try to suspend us as we try to suspend them
-		static FCriticalSection Mutex;
-		FScopeLock Lock(&Mutex);
+		static UE::FPlatformRecursiveMutex Mutex;
+		UE::TUniqueLock Lock(Mutex);
 
 		// Suspend the thread before grabbing its context (possible fix for incomplete callstacks)
 		SuspendThread(ThreadHandle);
diff --git a/Engine/Source/Runtime/Core/Public/Windows/WindowsPlatformProcess.h b/Engine/Source/Runtime/Core/Public/Windows/WindowsPlatformProcess.h
--- a/Engine/Source/Runtime/Core/Public/Windows/WindowsPlatformProcess.h
+++ b/Engine/Source/Runtime/Core/Public/Windows/WindowsPlatformProcess.h
@@ -269,7 +269,7 @@ private:
 	 * @param FileName Path to the library to load
 	 * @param SearchPaths Search directories to scan for imports
 	 */
-	static void* LoadLibraryWithSearchPaths(const FString& FileName, const TArray<FString>& SearchPaths);
+	static void* LoadLibraryWithSearchPaths(const FString& FileName, TConstArrayView<FString> SearchPaths);
 
 	/**
 	 * Resolve an individual import.
@@ -278,7 +278,7 @@ private:
 	 * @param OutFileName On success, receives the path to the imported file
 	 * @return true if an import was found.
 	 */
-	static bool ResolveImport(const FString& Name, const TArray<FString>& SearchPaths, FString& OutFileName);
+	static bool ResolveImport(const FString& Name, TConstArrayView<FString> SearchPaths, FString& OutFileName);
 
 	/**
 	 * Resolve all the imports for the given library, searching through a set of directories.
@@ -288,7 +288,7 @@ private:
 	 * @param ImportFileNames Array which is filled with a list of the resolved imports found in the given search directories
 	 * @param VisitedImportNames Array which stores a list of imports which have been checked
 	 */
-	static void ResolveMissingImportsRecursive(const FString& FileName, const TArray<FString>& SearchPaths, TArray<FString>& ImportFileNames, TSet<FString>& VisitedImportNames);
+	static void ResolveMissingImportsRecursive(const FString& FileName, TConstArrayView<FString> SearchPaths, TArray<FString>& ImportFileNames, TSet<FString>& VisitedImportNames);
 };
 
 typedef FWindowsPlatformProcess FPlatformProcess;
