LLM: Removed dependencies on FCriticalSection and FRWLock

They may be aliases for FRecursiveMutex and FSharedMutex, which we cannot rely on from memory allocators.

diff --git a/Engine/Source/Runtime/Core/Private/HAL/LowLevelMemTracker.cpp b/Engine/Source/Runtime/Core/Private/HAL/LowLevelMemTracker.cpp
--- a/Engine/Source/Runtime/Core/Private/HAL/LowLevelMemTracker.cpp
+++ b/Engine/Source/Runtime/Core/Private/HAL/LowLevelMemTracker.cpp
@@ -5,6 +5,8 @@
 
 #if ENABLE_LOW_LEVEL_MEM_TRACKER
 
+#include "Async/SharedLock.h"
+#include "Async/UniqueLock.h"
 #include "HAL/LowLevelMemStats.h"
 #include "HAL/FileManager.h"
 #include "HAL/IConsoleManager.h"
@@ -1513,7 +1515,7 @@ public:
 		) const
 		{
 #if LLM_ENABLED_FULL_TAGS
-			FReadScopeLock TagDataScopeLock(InLLMRef.TagDataLock);
+			UE::TSharedLock TagDataScopeLock(InLLMRef.TagDataMutex);
 #if !LLM_ALLOW_ASSETS_TAGS && !LLM_ALLOW_UOBJECTCLASSES_TAGS
 			return (*InLLMRef.TagDatas)[Tag];
 #else
@@ -1896,7 +1898,7 @@ void FLowLevelMemTracker::BootstrapInitialise()
 	{
 		return;
 	}
-	FScopeLock BootstrapScopeLock(&BootstrapLock);
+	UE::TUniqueLock BootstrapLock(BootstrapMutex);
 	if (bInitialisedTracking)
 	{
 		return;
@@ -1984,7 +1986,7 @@ void FLowLevelMemTracker::UpdateStatsPerFrame(const TCHAR* LogName)
 		{
 			// UpdateStatsPerFrame is usually called from the game thread, but can sometimes be called from the
 			// async loading thread, so enter a lock for it
-			FScopeLock UpdateScopeLock(&UpdateLock);
+			UE::TUniqueLock UpdateLock(UpdateMutex);
 			if (bFirstTimeUpdating)
 			{
 				// Write the saved overhead value to the stats system; this allows us to see the overhead that is
@@ -1999,7 +2001,7 @@ void FLowLevelMemTracker::UpdateStatsPerFrame(const TCHAR* LogName)
 
 	// UpdateStatsPerFrame is usually called from the game thread, but can sometimes be called from the
 	// async loading thread, so enter a lock for it
-	FScopeLock UpdateScopeLock(&UpdateLock);
+	UE::TUniqueLock UpdateLock(UpdateMutex);
 	BootstrapInitialise();
 
 	if (bFirstTimeUpdating)
@@ -2025,7 +2027,7 @@ void FLowLevelMemTracker::Tick()
 	}
 	// TickInternal is usually called from the game thread, but can sometimes be called from the async loading thread,
 	// so enter a lock for it.
-	FScopeLock UpdateScopeLock(&UpdateLock);
+	UE::TUniqueLock UpdateLock(UpdateMutex);
 	TickInternal();
 }
 
@@ -2128,7 +2130,7 @@ void FLowLevelMemTracker::UpdateTags()
 	bTagAdded = false;
 	bool bNeedsResort = false;
 	{
-		FReadScopeLock ScopeLock(TagDataLock);
+		UE::TSharedLock Lock(TagDataMutex);
 		// Cannot use a ranged-for because FinishConstruct can drop our lock and add elements to TagDatas
 		// Check the valid-index condition on every loop iteration
 		for (int32 TagIndex = 0; TagIndex < TagDatas->Num(); ++TagIndex)
@@ -2151,7 +2153,7 @@ void FLowLevelMemTracker::UpdateTags()
 		}
 
 		FTagDataArray* OldTagDatas;
-		FWriteScopeLock ScopeLock(TagDataLock);
+		UE::TUniqueLock Lock(TagDataMutex);
 		SortTags(OldTagDatas);
 
 		for (int32 TrackerIndex = 0; TrackerIndex < static_cast<int32>(ELLMTracker::Max); TrackerIndex++)
@@ -2168,7 +2170,7 @@ void FLowLevelMemTracker::SortTags(UE::LLMPrivate::FTagDataArray*& OutOldTagData
 {
 	using namespace UE::LLMPrivate;
 
-	// Caller is responsible for holding a WriteLock on TagDataLock
+	// Caller is responsible for holding a UniqueLock on TagDataMutex
 	OutOldTagDatas = TagDatas;
 	TagDatas = Allocator.New<FTagDataArray>();
 	FTagDataArray& LocalTagDatas = *TagDatas;
@@ -2602,7 +2604,7 @@ bool FLowLevelMemTracker::Exec(const TCHAR* Cmd, FOutputDevice& Ar)
 		{
 			CSV_EVENT_GLOBAL(TEXT("LLM_REMEMBER"));
 
-			FScopeLock UpdateScopeLock(&UpdateLock);
+			UE::TUniqueLock UpdateLock(UpdateMutex);
 			// Go through all the tags and REMEMBER the current recorded size values
 			for (int32 TrackerIndex = 0; TrackerIndex < static_cast<int32>(ELLMTracker::Max); TrackerIndex++)
 			{
@@ -2614,7 +2616,7 @@ bool FLowLevelMemTracker::Exec(const TCHAR* Cmd, FOutputDevice& Ar)
 		{
 			CSV_EVENT_GLOBAL(TEXT("LLM_FORGET"));
 
-			FScopeLock UpdateScopeLock(&UpdateLock);
+			UE::TUniqueLock UpdateLock(UpdateMutex);
 			// Go through all tags and make them forget their current size values
 			for (int32 TrackerIndex = 0; TrackerIndex < static_cast<int32>(ELLMTracker::Max); TrackerIndex++)
 			{
@@ -2729,7 +2731,7 @@ void FLowLevelMemTracker::BootstrapTagDatas()
 	// While bootstrapping we are not allowed to construct any FNames because the FName system may not yet have been
 	// constructed. Construct not-fully-initialized TagDatas for the central list of ELLMTags.
 	{
-		FWriteScopeLock ScopeLock(TagDataLock);
+		UE::TUniqueLock Lock(TagDataMutex);
 		bIsBootstrapping = true;
 
 		TagDatas = Allocator.New<FTagDataArray>();
@@ -2818,7 +2820,7 @@ void FLowLevelMemTracker::InitialiseTagDatas_FinishRegister()
 
 	// Record the central list of ELLMTags in TagDataNameMap, and mark that bootstrapping is complete
 	{
-		FWriteScopeLock ScopeLock(TagDataLock);
+		UE::TUniqueLock Lock(TagDataMutex);
 
 #define FINISH_REGISTER(Enum,Str,Stat,Group,ParentTag) \
 		{ \
@@ -2873,7 +2875,7 @@ void FLowLevelMemTracker::ClearTagDatas()
 {
 	using namespace UE::LLMPrivate;
 
-	FWriteScopeLock ScopeLock(TagDataLock);
+	UE::TUniqueLock Lock(TagDataMutex);
 	FLLMTagDeclaration::ClearCreationCallbacks();
 
 	Allocator.Free(TagDataEnumMap, sizeof(FTagData*) * LLM_TAG_COUNT);
@@ -2982,7 +2984,7 @@ UE::LLMPrivate::FTagData& FLowLevelMemTracker::RegisterTagData(FName Name, FName
 		}
 	}
 
-	FWriteScopeLock ScopeLock(TagDataLock);
+	UE::TUniqueLock Lock(TagDataMutex);
 	FTagData*& TagDataForName = TagDataNameMap->FindOrAdd(FTagDataNameKey(Name, TagSet), nullptr);
 	if (TagDataForName)
 	{
@@ -3091,7 +3093,7 @@ void FLowLevelMemTracker::FinishConstruct(UE::LLMPrivate::FTagData* TagData,
 	UE::LLMPrivate::ETagReferenceSource ReferenceSource)
 {
 	using namespace UE::LLMPrivate;
-	// Caller is responsible for holding a ReadLock (NOT a WriteLock) on TagDataLock
+	// Caller is responsible for holding a SharedLock (NOT a UniqueLock) on TagDataMutex
 
 	LLMCheck(TagData);
 	if (TagData->IsFinishConstructed())
@@ -3116,11 +3118,11 @@ void FLowLevelMemTracker::FinishConstruct(UE::LLMPrivate::FTagData* TagData,
 			FTagData** ParentDataPtr = TagDataNameMap->Find(FTagDataNameKey(ParentName, TagData->GetTagSet()));
 			if (!ParentDataPtr)
 			{
-				// We have to drop the ReadLock to call RegisterTagData, which takes a WriteLock.
-				TagDataLock.ReadUnlock();
+				// We have to drop the SharedLock to call RegisterTagData, which takes a UniqueLock.
+				TagDataMutex.UnlockShared();
 				FTagData* ParentTagData = &RegisterTagData(ParentName, NAME_None, NAME_None, NAME_None, NAME_None,
 					false, ELLMTag::CustomName,false, ETagReferenceSource::ImplicitParent, TagData->GetTagSet());
-				TagDataLock.ReadLock();
+				TagDataMutex.LockShared();
 				if (TagData->IsFinishConstructed())
 				{
 					// Another thread got in and finished construction while we were outside of the lock.
@@ -3148,16 +3150,16 @@ void FLowLevelMemTracker::FinishConstruct(UE::LLMPrivate::FTagData* TagData,
 	// possibly have registered. Subscribers must instead read those from GetTrackedTags.
 	if (!TagData->HasEnumTag() || TagData->GetEnumTag() >= ELLMTag::GenericTagCount)
 	{
-		// Leave the critical section while calling the callback, since the callback could be arbitrary
+		// Leave the shared mutex while calling the callback, since the callback could be arbitrary
 		// code that calls back into LLM. The TagData pointer is immutable so we do not have to worry about
 		// it disappearing out from under us.
-		TagDataLock.ReadUnlock();
+		TagDataMutex.UnlockShared();
 		for (const TPair<UE::LLMPrivate::FTagCreationCallback, UPTRINT>& Callback :
 			UE::LLMPrivate::GetTagCreationCallbacks())
 		{
 			Callback.Key(TagData, Callback.Value);
 		}
-		TagDataLock.ReadLock();
+		TagDataMutex.LockShared();
 	}
 }
 
@@ -3172,9 +3174,10 @@ TArray<const UE::LLMPrivate::FTagData*> FLowLevelMemTracker::GetTrackedTags(ELLM
 
 	BootstrapInitialise();
 
-	FReadScopeLock TagDataScopeLock(TagDataLock);
+	UE::TSharedLock Lock(TagDataMutex);
 	
-	TArray<const FTagData*> FoundResults(TagDatas->FilterByPredicate([TagSet](const FTagData* InTagData){
+	TArray<const FTagData*> FoundResults(TagDatas->FilterByPredicate([TagSet](const FTagData* InTagData)
+	{
 		return InTagData != nullptr && InTagData->GetTagSet() == TagSet;
 	}));
 
@@ -3192,7 +3195,7 @@ TArray<const UE::LLMPrivate::FTagData*> FLowLevelMemTracker::GetTrackedTags(ELLM
 	
 	BootstrapInitialise();
 
-	FScopeLock UpdateScopeLock(&UpdateLock); // uses of TagSizes are guarded by the UpdateLock
+	UE::TUniqueLock UpdateLock(UpdateMutex); // uses of TagSizes are guarded by the UpdateLock
 	return GetTracker(Tracker)->GetTagDatas(TagSet);
 }
 
@@ -3207,7 +3210,7 @@ void FLowLevelMemTracker::GetTrackedTagsNamesWithAmount(TMap<FName, uint64>& Tag
 
 	BootstrapInitialise();
 
-	FScopeLock UpdateScopeLock(&UpdateLock);
+	UE::TUniqueLock UpdateLock(UpdateMutex);
 	GetTracker(Tracker)->GetTagsNamesWithAmount(TagsNamesWithAmount, TagSet);
 }
 
@@ -3222,7 +3225,7 @@ void FLowLevelMemTracker::GetTrackedTagsNamesWithAmountFiltered(TMap<FName, uint
 
 	BootstrapInitialise();
 
-	FScopeLock UpdateScopeLock(&UpdateLock);
+	UE::TUniqueLock UpdateLock(UpdateMutex);
 	GetTracker(Tracker)->GetTagsNamesWithAmountFiltered(TagsNamesWithAmount, TagSet, Filters);
 }
 
@@ -3239,7 +3242,7 @@ bool FLowLevelMemTracker::FindTagByName( const TCHAR* Name, uint64& OutTag, ELLM
 
 	if (Name != nullptr)
 	{
-		FReadScopeLock ScopeLock(TagDataLock);
+		UE::TSharedLock Lock(TagDataMutex);
 
 		// Search by Name
 		FName SearchName(Name);
@@ -3391,7 +3394,7 @@ int64 FLowLevelMemTracker::GetTagAmountForTracker(ELLMTracker Tracker, const UE:
 		return 0;
 	}
 
-	FScopeLock UpdateScopeLock(&UpdateLock); // uses of TagSizes are guarded by the UpdateLock
+	UE::TUniqueLock UpdateLock(UpdateMutex); // uses of TagSizes are guarded by the UpdateLock
 	return GetTracker(Tracker)->GetTagAmount(TagData, SizeParams);
 }
 
@@ -3431,7 +3434,7 @@ void FLowLevelMemTracker::SetTagAmountForTracker(ELLMTracker Tracker, ELLMTag Ta
 	BootstrapInitialise();
 	const FTagData* TagData = FindOrAddTagData(Tag);
 
-	FScopeLock UpdateScopeLock(&UpdateLock); // uses of TagSizes are guarded by the UpdateLock
+	UE::TUniqueLock UpdateLock(UpdateMutex); // uses of TagSizes are guarded by the UpdateLock
 	GetTracker(Tracker)->SetTagAmountExternal(TagData, Amount, bAddToTotal);
 }
 
@@ -3446,7 +3449,7 @@ void FLowLevelMemTracker::SetTagAmountForTracker(ELLMTracker Tracker, FName Tag,
 	BootstrapInitialise();
 	const FTagData* TagData = FindOrAddTagData(Tag, TagSet);
 
-	FScopeLock UpdateScopeLock(&UpdateLock); // uses of TagSizes are guarded by the UpdateLock
+	UE::TUniqueLock UpdateLock(UpdateMutex); // uses of TagSizes are guarded by the UpdateLock
 	GetTracker(Tracker)->SetTagAmountExternal(TagData, Amount, bAddToTotal);
 }
 
@@ -3506,7 +3509,7 @@ const UE::LLMPrivate::FTagData* FLowLevelMemTracker::FindOrAddTagData(ELLMTag En
 	LLMCheckf(0 <= Index && Index < LLM_TAG_COUNT, TEXT("Out of range ELLMTag %d"), Index);
 
 	{
-		FReadScopeLock ScopeLock(TagDataLock);
+		UE::TSharedLock Lock(TagDataMutex);
 		FTagData* TagData = TagDataEnumMap[Index];
 		if (TagData)
 		{
@@ -3530,7 +3533,7 @@ const UE::LLMPrivate::FTagData* FLowLevelMemTracker::FindOrAddTagData(ELLMTag En
 	{
 		FTagData* TagData = &RegisterTagData(TagName, NAME_None, NAME_None, NAME_None, NAME_None, true, EnumTag,
 			false, ReferenceSource);
-		FReadScopeLock ScopeLock(TagDataLock);
+		UE::TSharedLock Lock(TagDataMutex);
 		FinishConstruct(TagData, ReferenceSource);
 		return TagData;
 	}
@@ -3548,7 +3551,7 @@ const UE::LLMPrivate::FTagData* FLowLevelMemTracker::FindOrAddTagData(FName TagN
 	using namespace UE::LLMPrivate;
 
 	{
-		FReadScopeLock ScopeLock(TagDataLock);
+		UE::TSharedLock Lock(TagDataMutex);
 		FTagData** TagDataPtr = TagDataNameMap->Find(FTagDataNameKey(TagName, TagSet));
 		if (TagDataPtr)
 		{
@@ -3573,7 +3576,7 @@ const UE::LLMPrivate::FTagData* FLowLevelMemTracker::FindOrAddTagData(FName TagN
 	FTagData* TagData = &RegisterTagData(TagName, NAME_None, NAME_None, StatName, NAME_None, false,
 		ELLMTag::CustomName, bIsStatTag, ReferenceSource, TagSet);
 	{
-		FReadScopeLock ScopeLock(TagDataLock);
+		UE::TSharedLock Lock(TagDataMutex);
 		FinishConstruct(TagData, ReferenceSource);
 		return TagData;
 	}
@@ -3587,7 +3590,7 @@ const UE::LLMPrivate::FTagData* FLowLevelMemTracker::FindTagData(ELLMTag EnumTag
 	int32 Index = static_cast<int32>(EnumTag);
 	LLMCheckf(0 <= Index && Index < LLM_TAG_COUNT, TEXT("Out of range ELLMTag %d"), Index);
 
-	FReadScopeLock ScopeLock(TagDataLock);
+	UE::TSharedLock Lock(TagDataMutex);
 	FTagData* TagData = TagDataEnumMap[Index];
 	if (TagData)
 	{
@@ -3605,7 +3608,7 @@ const UE::LLMPrivate::FTagData* FLowLevelMemTracker::FindTagData(FName TagName,
 {
 	using namespace UE::LLMPrivate;
 
-	FReadScopeLock ScopeLock(TagDataLock);
+	UE::TSharedLock Lock(TagDataMutex);
 
 	FTagData** TagDataPtr = TagDataNameMap->Find(FTagDataNameKey(TagName, TagSet));
 	if (TagDataPtr)
@@ -4289,7 +4292,7 @@ void* FLLMAllocator::Alloc(size_t Size)
 		return nullptr;
 	}
 	int32 BinIndex = GetBinIndex(Size);
-	FScopeLock Lock(&CriticalSection);
+	UE::TUniqueLock Lock(Mutex);
 	if (BinIndex == NumBins)
 	{
 		return AllocPages(Size);
@@ -4304,7 +4307,7 @@ void FLLMAllocator::Free(void* Ptr, size_t Size)
 	if (Ptr != nullptr)
 	{
 		int32 BinIndex = GetBinIndex(Size);
-		FScopeLock Lock(&CriticalSection);
+		UE::TUniqueLock Lock(Mutex);
 		if (BinIndex == NumBins)
 		{
 			FreePages(Ptr, Size);
@@ -4338,7 +4341,7 @@ void* FLLMAllocator::Realloc(void* Ptr, size_t OldSize, size_t NewSize)
 
 int64 FLLMAllocator::GetTotal() const
 {
-	FScopeLock Lock((FCriticalSection*)&CriticalSection);
+	UE::TUniqueLock Lock(const_cast<UE::FPlatformRecursiveMutex&>(Mutex));
 	return Total;
 }
 
@@ -5367,7 +5370,7 @@ bool FLLMTracker::DumpForkedAllocationInfo()
 
 	
 	// We cannot call GetTag to find the Tag from the Allocation when we are within AllocationMap, because GetTag
-	// locks the TagDataLock and TagDataLock has to be entered before entering AllocationMap.LockAll.
+	// locks the TagDataMutex and TagDataMutex has to be entered before entering AllocationMap.LockAll.
 	// So for the TagIdentifier use the CompressedTag rather than calling GetTag.
 #if LLM_ENABLED_FULL_TAGS
 	TMap<int32, FCounts> CountsPerTag;
@@ -5376,7 +5379,7 @@ bool FLLMTracker::DumpForkedAllocationInfo()
 #endif
 	int32 NumTags;
 	{
-		FReadScopeLock TagDataScopeLock(LLMRef.TagDataLock);
+		UE::TSharedLock Lock(LLMRef.TagDataMutex);
 		NumTags = LLMRef.TagDatas->Num();
 	}
 
@@ -5685,7 +5688,7 @@ void FLLMTracker::GetTagsNamesWithAmount(TMap<FName, uint64>& OutTagsNamesWithAm
 void FLLMTracker::GetTagsNamesWithAmountFiltered(TMap<FName, uint64>& OutTagsNamesWithAmount, ELLMTagSet TagSet /* = ELLMTagSet::None */, TArray<FLLMTagSetAllocationFilter>& Filters)
 {
 	// We cannot call GetTag to find the Tag from the Allocation when we are within AllocationMap, because GetTag
-	// locks the TagDataLock and TagDataLock has to be entered before entering AllocationMap.LockAll.
+	// locks the TagDataMutex and TagDataMutex has to be entered before entering AllocationMap.LockAll.
 	// So for the TagIdentifier use the CompressedTag rather than calling GetTag.
 	struct FLocalTagData
 	{
@@ -5694,7 +5697,7 @@ void FLLMTracker::GetTagsNamesWithAmountFiltered(TMap<FName, uint64>& OutTagsNam
 	};
 	TArray<FLocalTagData> TagDataCached;
 	{
-		FReadScopeLock TagDataScopeLock(LLMRef.TagDataLock);
+		UE::TSharedLock Lock(LLMRef.TagDataMutex);
 		const int32 NumTags = LLMRef.TagDatas->Num();
 		TagDataCached.Init(FLocalTagData{}, NumTags);
 		for (FTagData* TagData : (*LLMRef.TagDatas))
diff --git a/Engine/Source/Runtime/Core/Private/HAL/LowLevelMemoryUtils.h b/Engine/Source/Runtime/Core/Private/HAL/LowLevelMemoryUtils.h
--- a/Engine/Source/Runtime/Core/Private/HAL/LowLevelMemoryUtils.h
+++ b/Engine/Source/Runtime/Core/Private/HAL/LowLevelMemoryUtils.h
@@ -7,11 +7,12 @@
 #if ENABLE_LOW_LEVEL_MEM_TRACKER
 
 #include "Algo/BinarySearch.h"
+#include "Async/UniqueLock.h"
 #include "Containers/ContainerAllocationPolicies.h"
 #include "Containers/Map.h"
 #include "Containers/Set.h"
+#include "HAL/PlatformMutex.h"
 #include "Math/NumericLimits.h"
-#include "Misc/ScopeLock.h"
 #include "Templates/Tuple.h"
 #include <type_traits>
 
@@ -233,7 +234,7 @@ public:
 	{
 		for (StripeData& Stripe : MapStripes)
 		{
-			FScopeLock AllocationScopeLock(&Stripe.CriticalSection);
+			UE::TUniqueLock AllocationScopeLock(Stripe.Mutex);
 
 			Stripe.Keys.Clear(true);
 			Stripe.KeyHashes.Clear(true);
@@ -255,7 +256,7 @@ public:
 		SizeType KeyHash = Key.GetHashCode();
 		StripeData& Stripe = MapStripes[GetStripeIndex(KeyHash)];
 
-		FScopeLock AllocationScopeLock(&Stripe.CriticalSection);
+		UE::TUniqueLock AllocationScopeLock(Stripe.Mutex);
 
 		LLMCheck(Stripe.Map);
 
@@ -334,7 +335,7 @@ public:
 		SizeType KeyHash = Key.GetHashCode();
 		const StripeData& Stripe = MapStripes[GetStripeIndex(KeyHash)];
 
-		FScopeLock AllocationScopeLock(&Stripe.CriticalSection);
+		UE::TUniqueLock AllocationScopeLock(Stripe.Mutex);
 
 		SizeType MapIndex = Stripe.GetMapIndex(Key, KeyHash);
 
@@ -355,7 +356,7 @@ public:
 		SizeType KeyHash = Key.GetHashCode();
 		StripeData& Stripe = MapStripes[GetStripeIndex(KeyHash)];
 
-		FScopeLock AllocationScopeLock(&Stripe.CriticalSection);
+		UE::TUniqueLock AllocationScopeLock(Stripe.Mutex);
 
 		LLMCheck(Stripe.Map);
 
@@ -449,7 +450,7 @@ public:
 		SizeType TotalCount = 0;
 		for (StripeData& Stripe : MapStripes)
 		{
-			FScopeLock AllocationScopeLock(&Stripe.CriticalSection);
+			UE::TUniqueLock AllocationScopeLock(Stripe.Mutex);
 			TotalCount += Stripe.Count;
 		}
 		return TotalCount;
@@ -460,7 +461,7 @@ public:
 		SizeType KeyHash = Key.GetHashCode();
 		StripeData& Stripe = MapStripes[GetStripeIndex(KeyHash)];
 
-		FScopeLock AllocationScopeLock(&Stripe.CriticalSection);
+		UE::TUniqueLock AllocationScopeLock(Stripe.Mutex);
 
 		SizeType MapIndex = Stripe.GetMapIndex(Key, KeyHash);
 		return Stripe.IsItemInUse(MapIndex);
@@ -470,7 +471,7 @@ public:
 	{
 		for (StripeData& Stripe : MapStripes)
 		{
-			FScopeLock AllocationScopeLock(&Stripe.CriticalSection);
+			UE::TUniqueLock AllocationScopeLock(Stripe.Mutex);
 
 			Stripe.Keys.Trim();
 			Stripe.KeyHashes.Trim();
@@ -629,13 +630,13 @@ public:
 	void LockAll()
 	{
 		// Locking all stripes is required to create an iterator
-		AllStripesLock.Lock();
+		AllStripesMutex.Lock();
 		LLMCheck(bAllStripesLocked == false);
 		bAllStripesLocked = true;
 
 		for (StripeData& Stripe : MapStripes)
 		{
-			Stripe.CriticalSection.Lock();
+			Stripe.Mutex.Lock();
 		}
 	}
 
@@ -643,12 +644,12 @@ public:
 	{
 		for (StripeData& Stripe : MapStripes)
 		{
-			Stripe.CriticalSection.Unlock();
+			Stripe.Mutex.Unlock();
 		}
 
 		LLMCheck(bAllStripesLocked == true);
 		bAllStripesLocked = false;
-		AllStripesLock.Unlock();
+		AllStripesMutex.Unlock();
 	}
 
 	// data
@@ -772,7 +773,7 @@ private:
 				Index >= StartIndex || Index <= EndIndex;
 		}
 
-		mutable FCriticalSection CriticalSection;
+		mutable UE::FPlatformRecursiveMutex Mutex;
 
 		UE::LLMPrivate::FLLMAllocator* Allocator;
 
@@ -799,7 +800,7 @@ private:
 	// mix function applied to the key), multiple threads will usually end up with their own lock, and not
 	// need to wait on each other.
 	StripeData MapStripes[StripeCount];
-	FCriticalSection AllStripesLock;
+	UE::FPlatformRecursiveMutex AllStripesMutex;
 	std::atomic_bool bAllStripesLocked;
 };
 
diff --git a/Engine/Source/Runtime/Core/Private/Windows/VirtualAllocPageStatus.cpp b/Engine/Source/Runtime/Core/Private/Windows/VirtualAllocPageStatus.cpp
--- a/Engine/Source/Runtime/Core/Private/Windows/VirtualAllocPageStatus.cpp
+++ b/Engine/Source/Runtime/Core/Private/Windows/VirtualAllocPageStatus.cpp
@@ -4,10 +4,10 @@
 
 #if UE_VIRTUALALLOC_PAGE_STATUS_ENABLED 
 
+#include "Async/UniqueLock.h"
 #include "Containers/ArrayView.h"
 #include "HAL/PlatformMemory.h"
 #include "Math/UnrealMathUtility.h"
-#include "Misc/ScopeLock.h"
 #include "Windows/WindowsHWrapper.h"
 
 void FHashMapLinearProbingVAlloc::Init(uint64 PageSize)
@@ -242,7 +242,7 @@ int64 FVirtualAllocPageStatus::MarkChangedAndReturnDeltaSize(void* InStartAddres
 	static_assert(sizeof(uint64) >= sizeof(void*));
 	uint64 StartAddress = reinterpret_cast<uint64>(InStartAddress);
 
-	FScopeLock ScopeLock(&Lock);
+	UE::TUniqueLock Lock(Mutex);
 
 	uint64 StartPage = StartAddress / PageSize;
 	uint64 EndPage = (StartAddress + Size + PageSize - 1) / PageSize;
@@ -358,7 +358,7 @@ void FVirtualAllocPageStatus::AddReservationSize(void* InReservationAddress, SIZ
 	static_assert(sizeof(uint64) >= sizeof(void*));
 	uint64 ReservationAddress = reinterpret_cast<uint64>(InReservationAddress);
 
-	FScopeLock ScopeLock(&Lock);
+	UE::TUniqueLock Lock(Mutex);
 
 	uint64 ReservationPage = ReservationAddress / ReservationAlignment;
 
@@ -372,7 +372,7 @@ SIZE_T FVirtualAllocPageStatus::GetAndRemoveReservationSize(void* InReservationA
 	static_assert(sizeof(uint64) >= sizeof(void*));
 	uint64 ReservationAddress = reinterpret_cast<uint64>(InReservationAddress);
 
-	FScopeLock ScopeLock(&Lock);
+	UE::TUniqueLock Lock(Mutex);
 
 	uint64 ReservationPage = ReservationAddress / ReservationAlignment;
 
diff --git a/Engine/Source/Runtime/Core/Private/Windows/VirtualAllocPageStatus.h b/Engine/Source/Runtime/Core/Private/Windows/VirtualAllocPageStatus.h
--- a/Engine/Source/Runtime/Core/Private/Windows/VirtualAllocPageStatus.h
+++ b/Engine/Source/Runtime/Core/Private/Windows/VirtualAllocPageStatus.h
@@ -8,8 +8,8 @@
 
 #if UE_VIRTUALALLOC_PAGE_STATUS_ENABLED
 
-#include "HAL/CriticalSection.h"
 #include "HAL/Platform.h"
+#include "HAL/PlatformMutex.h"
 #include "Math/NumericLimits.h"
 
 /**
@@ -93,7 +93,7 @@ struct FVirtualAllocPageStatus
 	SIZE_T GetAndRemoveReservationSize(void* InReservationAddress);
 
 private:
-	FCriticalSection Lock;
+	UE::FPlatformRecursiveMutex Mutex;
 	FHashMapLinearProbingVAlloc GroupToPageBitsMap;
 	FHashMapLinearProbingVAlloc PageToReservationSizeMap;
 	uint64 PageSize = 1;
diff --git a/Engine/Source/Runtime/Core/Public/HAL/LowLevelMemTracker.h b/Engine/Source/Runtime/Core/Public/HAL/LowLevelMemTracker.h
--- a/Engine/Source/Runtime/Core/Public/HAL/LowLevelMemTracker.h
+++ b/Engine/Source/Runtime/Core/Public/HAL/LowLevelMemTracker.h
@@ -3,6 +3,7 @@
 #pragma once
 
 #include "CoreTypes.h"
+#include "HAL/PlatformMutex.h"
 #include "LowLevelMemTrackerDefines.h" // LLM_ENABLED_IN_CONFIG
 #include "ProfilingDebugging/TagTrace.h"
 #include "AutoRTFM.h"
@@ -540,7 +541,7 @@ private:
 	void FreePages(void* Ptr, size_t Size);
 	int32 GetBinIndex(size_t Size) const;
 
-	FCriticalSection CriticalSection;
+	UE::FPlatformRecursiveMutex Mutex;
 	LLMAllocFunction PlatformAlloc;
 	LLMFreeFunction PlatformFree;
 	AllocatorPrivate::FBin* Bins;
@@ -879,9 +880,9 @@ private:
 
 	UE::LLMPrivate::FLLMTracker* Trackers[static_cast<int32>(ELLMTracker::Max)];
 
-	mutable FRWLock TagDataLock;
-	FCriticalSection UpdateLock;
-	FCriticalSection BootstrapLock;
+	mutable UE::FPlatformSharedMutex TagDataMutex;
+	UE::FPlatformRecursiveMutex UpdateMutex;
+	UE::FPlatformRecursiveMutex BootstrapMutex;
 
 	uint64 ProgramSize;
 	int64 MemoryUsageCurrentOverhead;
