Add PLATFORM_COMPILER_MSVC as a compliment to PLATFORM_COMPILER_CLANG

#rnx

diff --git a/Engine/Source/Runtime/Core/Public/HAL/Platform.h b/Engine/Source/Runtime/Core/Public/HAL/Platform.h
--- a/Engine/Source/Runtime/Core/Public/HAL/Platform.h
+++ b/Engine/Source/Runtime/Core/Public/HAL/Platform.h
@@ -106,7 +106,7 @@
 #include "GenericPlatform/GenericPlatform.h"
 
 //------------------------------------------------------------------
-// Setup macros for static code analysis
+// Setup macros for compiler detection
 //------------------------------------------------------------------
 #ifndef PLATFORM_COMPILER_CLANG
 #if defined(__clang__)
@@ -116,6 +116,17 @@
 #endif // defined(__clang__)
 #endif
 
+#ifndef PLATFORM_COMPILER_MSVC
+#if defined(_MSC_VER) && !PLATFORM_COMPILER_CLANG
+#define PLATFORM_COMPILER_MSVC			1
+#else
+#define PLATFORM_COMPILER_MSVC			0
+#endif // defined(_MSC_VER) ...
+#endif
+
+//------------------------------------------------------------------
+// Setup macros for static code analysis
+//------------------------------------------------------------------
 #if PLATFORM_COMPILER_CLANG
 	#include "Clang/ClangPlatformCodeAnalysis.h"
 #elif PLATFORM_MICROSOFT
