Use the same VkApplicationInfo for Android Vulkan detection as VulkanRHI initialization so drivers that want to opt out of Vulkan can do so consistently

diff --git a/Engine/Source/Runtime/Core/Private/Android/AndroidPlatformMisc.cpp b/Engine/Source/Runtime/Core/Private/Android/AndroidPlatformMisc.cpp
--- a/Engine/Source/Runtime/Core/Private/Android/AndroidPlatformMisc.cpp
+++ b/Engine/Source/Runtime/Core/Private/Android/AndroidPlatformMisc.cpp
@@ -12,6 +12,7 @@
 #include "Misc/OutputDeviceRedirector.h"
 #include "HAL/IConsoleManager.h"
 #include "Modules/ModuleManager.h"
+#include "BuildSettings.h"
 #include <string.h>
 #include <dlfcn.h>
 #include <sys/statfs.h>
@@ -2049,21 +2050,29 @@ static EDeviceVulkanSupportStatus AttemptVulkanInit(void* VulkanLib)
 		return EDeviceVulkanSupportStatus::NotSupported;
 	}
 
-	// try to create instance to verify driver available
-	VkApplicationInfo App;
-	FMemory::Memzero(App);
-	App.sType = VK_STRUCTURE_TYPE_APPLICATION_INFO;
-	App.pApplicationName = "UE";
-	App.applicationVersion = 0;
-	App.pEngineName = "UE";
-	App.engineVersion = 0;
-	App.apiVersion = UE_VK_API_VERSION;
+	// Try to create instance to verify driver available. Logic here matches that in VulkanRHI.cpp.
+	// Drivers need to detect the app and potentially return VK_ERROR_INCOMPATIBLE_DRIVER to disable Vulkan support for known incompatibility
+	auto* CVarDisableEngineAndAppRegistration = IConsoleManager::Get().FindTConsoleVariableDataInt(TEXT("r.DisableEngineAndAppRegistration"));
+	bool bDisableEngineRegistration = (CVarDisableEngineAndAppRegistration && CVarDisableEngineAndAppRegistration->GetValueOnAnyThread() != 0);
+
+	FString EngineName = FApp::GetEpicProductIdentifier() + FEngineVersion::Current().ToString(EVersionComponent::Minor);
+	FTCHARToUTF8 EngineNameConverter(*EngineName);
+	FTCHARToUTF8 ProjectNameConverter(FApp::GetProjectName());
+
+	VkApplicationInfo AppInfo;
+	FMemory::Memzero(AppInfo);
+	AppInfo.sType = VK_STRUCTURE_TYPE_APPLICATION_INFO;
+	AppInfo.pApplicationName = bDisableEngineRegistration ? nullptr : ProjectNameConverter.Get();
+	AppInfo.applicationVersion = static_cast<uint32>(BuildSettings::GetCurrentChangelist()) | (BuildSettings::IsLicenseeVersion() ? 0x80000000 : 0);
+	AppInfo.pEngineName = bDisableEngineRegistration ? nullptr : EngineNameConverter.Get();
+	AppInfo.engineVersion = FEngineVersion::Current().GetMinor();
+	AppInfo.apiVersion = UE_VK_API_VERSION;
 
 	VkInstanceCreateInfo InstInfo;
 	FMemory::Memzero(InstInfo);
 	InstInfo.sType = VK_STRUCTURE_TYPE_INSTANCE_CREATE_INFO;
 	InstInfo.pNext = nullptr;
-	InstInfo.pApplicationInfo = &App;
+	InstInfo.pApplicationInfo = &AppInfo;
 	InstInfo.enabledExtensionCount = 0;
 	InstInfo.ppEnabledExtensionNames = nullptr;
 
