Added 0 bytes sized bin to MB3 so all of them will be aliased and take 1 page instead of being served from 16 bytes bin each

#was-backout 43860771

diff --git a/Engine/Source/Runtime/Core/Private/HAL/MallocBinned3.cpp b/Engine/Source/Runtime/Core/Private/HAL/MallocBinned3.cpp
--- a/Engine/Source/Runtime/Core/Private/HAL/MallocBinned3.cpp
+++ b/Engine/Source/Runtime/Core/Private/HAL/MallocBinned3.cpp
@@ -150,7 +150,7 @@ struct FFreeBlock
 	uint16 BinSizeShifted;		// Size of the bins that this list points to >> UE_MBC_MIN_SMALL_POOL_ALIGNMENT_SHIFT
 	uint8  PoolIndex;			// Index of this pool
 	uint8  Canary;				// Constant value of 0xe3
-	uint32 NumFreeBins;			// Number of consecutive free bins here, at least 1.
+	uint32 NumFreeBins;			// Number of consecutive free bins here, at least 1 for a newly initialized FFreeBlock
 	uint32 NextFreeBlockIndex;	// Next free block or MAX_uint32
 };
 
@@ -200,7 +200,7 @@ struct FMallocBinned3::FPoolInfoSmall		//This is more like BlockInfoSmall as it
 	{
 		check(HasFreeBin());
 		++Taken;
-		check(Taken != 0);
+		check(BinSize == 0 || Taken != 0);	// This can technically overflow for 0 sized allocations if we allocate more than i16_max and never free any
 		FFreeBlock* Free = (FFreeBlock*)(BlockPtr + BinSize * FirstFreeIndex);
 		void* Result = Free->AllocateBin();
 		if (Free->GetNumFreeBins() == 0)
@@ -214,7 +214,7 @@ struct FMallocBinned3::FPoolInfoSmall		//This is more like BlockInfoSmall as it
 			{
 				FirstFreeIndex = Free->NextFreeBlockIndex;
 				check(uint32(FirstFreeIndex) == Free->NextFreeBlockIndex);
-				check(((FFreeBlock*)(BlockPtr + BinSize * FirstFreeIndex))->GetNumFreeBins());
+				check(BinSize == 0 || ((FFreeBlock*)(BlockPtr + BinSize * FirstFreeIndex))->GetNumFreeBins());
 			}
 		}
 
@@ -308,7 +308,7 @@ struct FMallocBinned3::Private
 
 	static MallocBinnedPrivate::TGlobalRecycler<UE_MB3_SMALL_POOL_COUNT> GGlobalRecycler;
 
-	static void FreeBundles(FMallocBinned3& Allocator, FBundleNode* BundlesToRecycle, uint32 InBinSize, uint32 InPoolIndex)
+	static void FreeBundles(FMallocBinned3& Allocator, FBundleNode* BundlesToRecycle, uint32 InPoolIndex)
 	{
 		FPoolTable& Table = Allocator.SmallPoolTables[InPoolIndex];
 		UE::TUniqueLock Lock(Table.Mutex);
@@ -325,7 +325,9 @@ struct FMallocBinned3::Private
 
 				uint32 OutBlockIndex;
 				void* BaseBlockPtr = Allocator.BlockPointerFromContainedPtr(Node, Allocator.SmallPoolTables[InPoolIndex].NumMemoryPagesPerBlock, OutBlockIndex);
-				const uint32 BinIndexWithinBlock = (((uint8*)Node) - ((uint8*)BaseBlockPtr)) / Allocator.SmallPoolTables[InPoolIndex].BinSize;
+				const uint32 BinSize = Allocator.SmallPoolTables[InPoolIndex].BinSize;
+				const uint32 BinSizeForDiv = FMath::Max(BinSize, 16u);
+				const uint32 BinIndexWithinBlock = (((uint8*)Node) - ((uint8*)BaseBlockPtr)) / BinSizeForDiv;
 
 				FPoolInfoSmall* NodePoolBlock = Allocator.SmallPoolTables[InPoolIndex].PoolInfos[OutBlockIndex / Allocator.SmallPoolInfosPerPlatformPage];
 				if (!NodePoolBlock)
@@ -342,7 +344,7 @@ struct FMallocBinned3::Private
 				FFreeBlock* Free = (FFreeBlock*)Node;
 				Free->NumFreeBins = 1;
 				Free->NextFreeBlockIndex = NodePool->NoFirstFreeIndex ? MAX_uint32 : NodePool->FirstFreeIndex;
-				Free->BinSizeShifted = (InBinSize >> UE_MBC_MIN_SMALL_POOL_ALIGNMENT_SHIFT);
+				Free->BinSizeShifted = (BinSize >> UE_MBC_MIN_SMALL_POOL_ALIGNMENT_SHIFT);
 				Free->Canary = FFreeBlock::CANARY_VALUE;
 				Free->PoolIndex = InPoolIndex;
 				NodePool->FirstFreeIndex = BinIndexWithinBlock;
@@ -352,8 +354,8 @@ struct FMallocBinned3::Private
 				UE_MBC_UPDATE_STATS(--Table.TotalUsedBins);
 
 				// Free this pool.
-				check(NodePool->Taken >= 1);
-				if (--NodePool->Taken == 0)
+				check(BinSize == 0 || NodePool->Taken >= 1);
+				if (--NodePool->Taken == 0 && BinSize != 0)	// we don't want to ever decommit memory for a zero sized bins
 				{
 					NodePool->SetCanary(FPoolInfoSmall::ECanary::Unassigned, true, false);
 					Table.BlocksAllocatedBits.FreeBit(OutBlockIndex);
@@ -367,7 +369,7 @@ struct FMallocBinned3::Private
 					Allocator.Decommit(InPoolIndex, BaseBlockPtr, BlockSize);
 
 #if UE_MB3_ALLOCATOR_STATS
-					Table.TotalAllocatedBins -= BlockSize / InBinSize;
+					Table.TotalAllocatedBins -= BlockSize / BinSizeForDiv;
 					Table.TotalAllocatedMem -= BlockSize;
 					AllocatedOSSmallPoolMemory.fetch_sub(BlockSize, std::memory_order_relaxed);
 #endif
@@ -389,7 +391,7 @@ MallocBinnedPrivate::TGlobalRecycler<UE_MB3_SMALL_POOL_COUNT> FMallocBinned3::Pr
 
 void FMallocBinned3::FreeBundles(FBundleNode* Bundles, uint32 PoolIndex)
 {
-	Private::FreeBundles(*this, Bundles, PoolIndexToBinSize(PoolIndex), PoolIndex);
+	Private::FreeBundles(*this, Bundles, PoolIndex);
 }
 
 FMallocBinned3::FPoolInfoSmall* FMallocBinned3::PushNewPoolToFront(FMallocBinned3::FPoolTable& Table, uint32 InBinSize, uint32 InPoolIndex, uint32& OutBlockIndex)
@@ -539,7 +541,8 @@ FMallocBinned3::FMallocBinned3(const FPlatformMemory::FPlatformVirtualMemoryBloc
 
 	verify(FSizeTableEntry::FillSizeTable(OsAllocationGranularity, SizeTable, UE_MB3_BASE_PAGE_SIZE, UE_MB3_MAX_SMALL_POOL_SIZE, UE_MB3_BASE_PAGE_SIZE) == UE_MB3_SMALL_POOL_COUNT);
 	checkf(SizeTable[UE_MB3_SMALL_POOL_COUNT - 1].BinSize == UE_MB3_MAX_SMALL_POOL_SIZE, TEXT("UE_MB3_MAX_SMALL_POOL_SIZE must be equal to the largest bin size"));
-	checkf(sizeof(FFreeBlock) <= SizeTable[0].BinSize, TEXT("Pool header must be able to fit into the smallest bin"));
+	check(SizeTable[0].BinSize == 0);
+	checkf(sizeof(FFreeBlock) <= SizeTable[1].BinSize, TEXT("Pool header must be able to fit into the smallest bin"));
 
 	SmallPoolInfosPerPlatformPage = OsAllocationGranularity / sizeof(FPoolInfoSmall);
 
@@ -606,6 +609,15 @@ FMallocBinned3::FMallocBinned3(const FPlatformMemory::FPlatformVirtualMemoryBloc
 
 	AllocateHashBuckets();
 
+	// Push a new pool for 0 sized bin to avoid if statements as it requires a special handling
+	uint32 BlockIndex;
+	PushNewPoolToFront(SmallPoolTables[0], 16, 0, BlockIndex);	// Use fake 16 bytes bin size here, so we can avoid div by 0
+	FFreeBlock* FreePtr = (FFreeBlock*)BlockPointerFromIndecies(0, BlockIndex, SmallPoolTables[0].BlockSize);
+	// re-initialize it with proper values for 0 sized bins
+	FreePtr->BinSizeShifted = 0;
+	// This should be 0 to point to itself, causing a circular dependency and so this memory block will never be decommitted
+	FreePtr->NextFreeBlockIndex = 0;
+	
 	MallocBinned3 = this;
 	GFixedMallocLocationPtr = (FMalloc**)(&MallocBinned3);
 }
@@ -689,7 +701,7 @@ void* FMallocBinned3::MallocExternal(SIZE_T Size, uint32 Alignment)
 	{
 		const uint32 PoolIndex = BoundSizeToPoolIndex(Size, MemSizeToPoolIndex);
 		FPerThreadFreeBlockLists* Lists = GMallocBinnedPerThreadCaches ? FPerThreadFreeBlockLists::Get() : nullptr;
-		if (Lists)
+		if (Lists && Size)
 		{
 			if (Lists->ObtainRecycledPartial(PoolIndex, Private::GGlobalRecycler))
 			{
@@ -742,7 +754,7 @@ void* FMallocBinned3::MallocExternal(SIZE_T Size, uint32 Alignment)
 #endif
 		if (GMallocBinnedAllocExtra)
 		{
-			if (Lists)
+			if (Lists && Size)
 			{
 				// prefill the free list with some allocations so we are less likely to hit this slow path with the mutex 
 				for (int32 Index = 0; Index < GMallocBinnedAllocExtra && Pool->HasFreeBin(); Index++)
@@ -914,7 +926,7 @@ void FMallocBinned3::FreeExternal(void* Ptr, uint64 PoolIndex)
 
 		FBundleNode* BundlesToRecycle = nullptr;
 		FPerThreadFreeBlockLists* Lists = GMallocBinnedPerThreadCaches ? FPerThreadFreeBlockLists::Get() : nullptr;
-		if (Lists)
+		if (Lists && PoolIndex)
 		{
 			BundlesToRecycle = Lists->RecycleFullBundle(PoolIndex, Private::GGlobalRecycler);
 			const bool bPushed = Lists->Free(Ptr, PoolIndex, BinSize);
@@ -932,7 +944,7 @@ void FMallocBinned3::FreeExternal(void* Ptr, uint64 PoolIndex)
 		if (BundlesToRecycle)
 		{
 			BundlesToRecycle->NextBundle = nullptr;
-			Private::FreeBundles(*this, BundlesToRecycle, BinSize, PoolIndex);
+			Private::FreeBundles(*this, BundlesToRecycle, PoolIndex);
 #if UE_MB3_ALLOCATOR_STATS
 			if (!Lists)
 			{
diff --git a/Engine/Source/Runtime/Core/Private/HAL/MallocBinnedCommon.cpp b/Engine/Source/Runtime/Core/Private/HAL/MallocBinnedCommon.cpp
--- a/Engine/Source/Runtime/Core/Private/HAL/MallocBinnedCommon.cpp
+++ b/Engine/Source/Runtime/Core/Private/HAL/MallocBinnedCommon.cpp
@@ -20,6 +20,7 @@ CSV_DEFINE_CATEGORY(MallocBinned, false);
 // They must be 16-byte aligned as well.
 static constexpr uint32 BinnedCommonSmallBinSizes4k[] =
 {
+	0, 
 	16, 32, 48, 64, 80, 96, 112, 128, 144, 160, 176, 192, // +16
 	224, 256, 288, 320, // +32
 	368,  // /11 ish
@@ -114,7 +115,7 @@ FSizeTableEntry::FSizeTableEntry(uint32 InBinSize, uint64 PlatformPageSize, uint
 		}
 		break;
 	}
-	check((PlatformPageSize * NumMemoryPagesPerBlock) / BinSize <= MAX_uint32);
+	check(BinSize == 0 || (PlatformPageSize * NumMemoryPagesPerBlock) / BinSize <= MAX_uint32);
 }
 
 static inline void FillTable(FSizeTableEntry* SizeTable, int32& Index, const uint32* BinsList, uint32 BinListSize, uint32 Num4kbPages, uint64 PlatformPageSize, uint32 BasePageSize)
diff --git a/Engine/Source/Runtime/Core/Public/HAL/MallocBinned3.h b/Engine/Source/Runtime/Core/Public/HAL/MallocBinned3.h
--- a/Engine/Source/Runtime/Core/Public/HAL/MallocBinned3.h
+++ b/Engine/Source/Runtime/Core/Public/HAL/MallocBinned3.h
@@ -306,7 +306,7 @@ public:
 	
 		// Only allocate from the small pools if the size is small enough and the alignment isn't crazy large.
 		// With large alignments, we'll waste a lot of memory allocating an entire page, but such alignments are highly unlikely in practice.
-		if ((Size <= UE_MB3_MAX_SMALL_POOL_SIZE) & (Alignment <= UE_MBC_MIN_SMALL_POOL_ALIGNMENT)) // one branch, not two
+		if ((Size != 0) & (Size <= UE_MB3_MAX_SMALL_POOL_SIZE) & (Alignment <= UE_MBC_MIN_SMALL_POOL_ALIGNMENT)) // one branch, not two
 		{
 			FPerThreadFreeBlockLists* Lists = GMallocBinnedPerThreadCaches ? FPerThreadFreeBlockLists::Get() : nullptr;
 			if (Lists)
@@ -334,7 +334,7 @@ public:
 	UE_AUTORTFM_NOAUTORTFM
 	FORCEINLINE virtual void* Realloc(void* Ptr, SIZE_T NewSize, uint32 Alignment) override
 	{
-		if (NewSize <= UE_MB3_MAX_SMALL_POOL_SIZE && Alignment <= UE_MBC_MIN_SMALL_POOL_ALIGNMENT) // one branch, not two
+		if ((NewSize != 0) & (NewSize <= UE_MB3_MAX_SMALL_POOL_SIZE) & (Alignment <= UE_MBC_MIN_SMALL_POOL_ALIGNMENT)) // one branch, not two
 		{
 			FPerThreadFreeBlockLists* Lists = GMallocBinnedPerThreadCaches ? FPerThreadFreeBlockLists::Get() : nullptr;
 
@@ -399,7 +399,7 @@ public:
 	FORCEINLINE virtual void Free(void* Ptr) override
 	{
 		const uint64 PoolIndex = PoolIndexFromPtr(Ptr);
-		if (PoolIndex < UE_MB3_SMALL_POOL_COUNT)
+		if (PoolIndex && PoolIndex < UE_MB3_SMALL_POOL_COUNT)
 		{
 			FPerThreadFreeBlockLists* Lists = GMallocBinnedPerThreadCaches ? FPerThreadFreeBlockLists::Get() : nullptr;
 			if (Lists)
diff --git a/Engine/Source/Runtime/Core/Public/HAL/MallocBinnedCommon.h b/Engine/Source/Runtime/Core/Public/HAL/MallocBinnedCommon.h
--- a/Engine/Source/Runtime/Core/Public/HAL/MallocBinnedCommon.h
+++ b/Engine/Source/Runtime/Core/Public/HAL/MallocBinnedCommon.h
@@ -32,7 +32,7 @@
 #endif
 
 #if !defined(UE_MBC_NUM_LISTED_SMALL_POOLS)
-#	define UE_MBC_NUM_LISTED_SMALL_POOLS				51
+#	define UE_MBC_NUM_LISTED_SMALL_POOLS				52
 #endif
 
 #if !defined(BINNEDCOMMON_USE_SEPARATE_VM_PER_POOL)
