Fix NonUnity build fallout from CL 44058885

diff --git a/Engine/Source/Runtime/Core/Private/Android/AndroidPlatformMisc.cpp b/Engine/Source/Runtime/Core/Private/Android/AndroidPlatformMisc.cpp
--- a/Engine/Source/Runtime/Core/Private/Android/AndroidPlatformMisc.cpp
+++ b/Engine/Source/Runtime/Core/Private/Android/AndroidPlatformMisc.cpp
@@ -10,6 +10,7 @@
 #include "Misc/CoreDelegates.h"
 #include "Misc/FeedbackContext.h"
 #include "Misc/OutputDeviceRedirector.h"
+#include "Misc/EngineVersion.h"
 #include "HAL/IConsoleManager.h"
 #include "Modules/ModuleManager.h"
 #include "BuildSettings.h"
