Speculative fix for crash in FPBDBendingConstraints::Apply when using ISPC on PC.

This crash looks suspiciously like the crash we saw last fall where an optimization to VectorGather was failing on PC (see CL 38926273). This appears to be the equivalent only on VectorScatter. In the same manner we fixed the last issue, disabling the VectorScatter optimization on PC.

diff --git a/Engine/Platforms/Windows/Source/Runtime/Experimental/Chaos/Chaos_Windows.Build.cs b/Engine/Platforms/Windows/Source/Runtime/Experimental/Chaos/Chaos_Windows.Build.cs
--- a/Engine/Platforms/Windows/Source/Runtime/Experimental/Chaos/Chaos_Windows.Build.cs
+++ b/Engine/Platforms/Windows/Source/Runtime/Experimental/Chaos/Chaos_Windows.Build.cs
@@ -11,6 +11,8 @@ namespace UnrealBuildTool.Rules
         {
 			// Disabling the optimized VectorGather call because it is causing crashes on PC.
 			PublicDefinitions.Add("USE_ISPC_OPTIMIZED_VECTORGATHER=0");
+			// Disabling the optimized VectorScatter call because it is causing crashes on PC.
+			PublicDefinitions.Add("USE_ISPC_OPTIMIZED_VECTORSCATTER=0");
 		}
     }
 }
diff --git a/Engine/Source/Runtime/Core/Public/Math/Vector.isph b/Engine/Source/Runtime/Core/Public/Math/Vector.isph
--- a/Engine/Source/Runtime/Core/Public/Math/Vector.isph
+++ b/Engine/Source/Runtime/Core/Public/Math/Vector.isph
@@ -17,6 +17,11 @@
 #define USE_ISPC_OPTIMIZED_VECTORGATHER 1
 #endif
 
+// Opt-out of VectorScatter optimization
+#ifndef USE_ISPC_OPTIMIZED_VECTORSCATTER
+#define USE_ISPC_OPTIMIZED_VECTORSCATTER 1
+#endif
+
 struct FVector2f
 {
 #if EXPLICIT_VECTOR2 == 0
@@ -633,6 +638,7 @@ inline varying FVector3f VectorGather(const uniform FVector3f *varying SrcPtr)
 
 inline void VectorScatter(uniform FVector3d *varying DstPtr, const varying FVector3d &V)
 {
+#if USE_ISPC_OPTIMIZED_VECTORSCATTER
 	if (((1 << TARGET_WIDTH) - 1 ^ lanemask()) == 0)
 	{
 		unmasked
@@ -647,6 +653,7 @@ inline void VectorScatter(uniform FVector3d *varying DstPtr, const varying FVect
 		}
 	}
 	else
+#endif
 	{
 		#pragma ignore warning(perf)
 		*DstPtr = V;
@@ -655,6 +662,7 @@ inline void VectorScatter(uniform FVector3d *varying DstPtr, const varying FVect
 
 inline void VectorScatter(uniform FVector3f *varying DstPtr, const varying FVector3f &V)
 {
+#if USE_ISPC_OPTIMIZED_VECTORSCATTER
 	if (((1 << TARGET_WIDTH) - 1 ^ lanemask()) == 0)
 	{
 		unmasked
@@ -669,6 +677,7 @@ inline void VectorScatter(uniform FVector3f *varying DstPtr, const varying FVect
 		}
 	}
 	else 
+#endif
 	{
 		#pragma ignore warning(perf)
 		*DstPtr = V;
@@ -725,6 +734,7 @@ inline varying FVector4f VectorGather(const uniform FVector4f *varying SrcPtr)
 
 inline void VectorScatter(uniform FVector4d *varying DstPtr, const varying FVector4d &V)
 {
+#if USE_ISPC_OPTIMIZED_VECTORSCATTER
 	if(((1<<TARGET_WIDTH)-1 ^ lanemask()) == 0)
 	{
 		unmasked
@@ -739,6 +749,7 @@ inline void VectorScatter(uniform FVector4d *varying DstPtr, const varying FVect
 		}
 	}
 	else
+#endif
 	{
 		#pragma ignore warning(perf)
 		*DstPtr = V;
@@ -747,6 +758,7 @@ inline void VectorScatter(uniform FVector4d *varying DstPtr, const varying FVect
 
 inline void VectorScatter(uniform FVector4f *varying DstPtr, const varying FVector4f &V)
 {
+#if USE_ISPC_OPTIMIZED_VECTORSCATTER
 	if(((1<<TARGET_WIDTH)-1 ^ lanemask()) == 0)
 	{
 		unmasked
@@ -761,6 +773,7 @@ inline void VectorScatter(uniform FVector4f *varying DstPtr, const varying FVect
 		}
 	}
 	else
+#endif
 	{
 		#pragma ignore warning(perf)
 		*DstPtr = V;
