Suppress GetProcessMemoryInfo() warnings on Wine

Under Wine, GetProcessMemoryInfo() incorrectly reports ERROR_ACCESS_DENIED if the target process is no longer running. This can safely be ignored. This change suppresses the corresponding warning log when this error is encountered on Wine.

#rnx

diff --git a/Engine/Source/Runtime/Core/Private/Windows/WindowsPlatformProcess.cpp b/Engine/Source/Runtime/Core/Private/Windows/WindowsPlatformProcess.cpp
--- a/Engine/Source/Runtime/Core/Private/Windows/WindowsPlatformProcess.cpp
+++ b/Engine/Source/Runtime/Core/Private/Windows/WindowsPlatformProcess.cpp
@@ -2527,7 +2527,14 @@ bool FWindowsPlatformProcess::TryGetMemoryUsage(FProcHandle& ProcessHandle, FPla
 	FPlatformMemory::Memzero(&ProcessMemoryCounters, sizeof(ProcessMemoryCounters));
 	if (!::GetProcessMemoryInfo(ProcessHandle.Get(), &ProcessMemoryCounters, sizeof(ProcessMemoryCounters)))
 	{
-		UE_LOG(LogWindows, Warning, TEXT("Failure in call to GetProcessMemoryInfo (GetLastError=%d)"), ::GetLastError());
+		if (FWindowsPlatformMisc::IsWine() && ::GetLastError() == ERROR_ACCESS_DENIED)
+		{
+			// Under Wine, GetProcessMemoryInfo() incorrectly returns ERROR_ACCESS_DENIED if a process is no longer running.
+		}
+		else
+		{
+			UE_LOG(LogWindows, Warning, TEXT("Failure in call to GetProcessMemoryInfo (GetLastError=%d)"), ::GetLastError());
+		}
 		return false;
 	}
 	OutStats.UsedPhysical = ProcessMemoryCounters.WorkingSetSize;
