Send CPU model ID and microcode revision along with crash reports.

Also log these values during startup.

#rnx

diff --git a/Engine/Source/Runtime/Core/Private/GenericPlatform/GenericPlatformCrashContext.cpp b/Engine/Source/Runtime/Core/Private/GenericPlatform/GenericPlatformCrashContext.cpp
--- a/Engine/Source/Runtime/Core/Private/GenericPlatform/GenericPlatformCrashContext.cpp
+++ b/Engine/Source/Runtime/Core/Private/GenericPlatform/GenericPlatformCrashContext.cpp
@@ -384,6 +384,9 @@ void FGenericCrashContext::Initialize()
 
 	NCached::Set(NCached::Session.CPUVendor, *FPlatformMisc::GetCPUVendor());
 	NCached::Set(NCached::Session.CPUBrand, *FPlatformMisc::GetCPUBrand());
+	NCached::Session.CPUModelID = FPlatformMisc::GetCPUModelID();
+	NCached::Session.CPUMicrocodeRevision = FPlatformMisc::GetCPUMicrocodeRevision();
+
 	NCached::Set(NCached::Session.PrimaryGPUBrand, *FPlatformMisc::GetPrimaryGPUBrand());
 	NCached::Set(NCached::Session.UserName, FPlatformProcess::UserName());
 	NCached::Set(NCached::Session.DefaultLocale, *FPlatformMisc::GetDefaultLocale());
@@ -974,6 +977,8 @@ void FGenericCrashContext::SerializeSessionContext(FString& Buffer)
 
 	AddCrashPropertyInternal(Buffer, TEXT("Misc.CPUVendor"), NCached::Session.CPUVendor);
 	AddCrashPropertyInternal(Buffer, TEXT("Misc.CPUBrand"), NCached::Session.CPUBrand);
+	AddCrashPropertyInternal(Buffer, TEXT("Misc.CPUModelID"), NCached::Session.CPUModelID);
+	AddCrashPropertyInternal(Buffer, TEXT("Misc.CPUMicrocodeRevision"), NCached::Session.CPUMicrocodeRevision);
 	AddCrashPropertyInternal(Buffer, TEXT("Misc.PrimaryGPUBrand"), NCached::Session.PrimaryGPUBrand);
 	AddCrashPropertyInternal(Buffer, TEXT("Misc.OSVersionMajor"), NCached::Session.OsVersion);
 	AddCrashPropertyInternal(Buffer, TEXT("Misc.OSVersionMinor"), NCached::Session.OsSubVersion);
diff --git a/Engine/Source/Runtime/Core/Private/GenericPlatform/GenericPlatformMisc.cpp b/Engine/Source/Runtime/Core/Private/GenericPlatform/GenericPlatformMisc.cpp
--- a/Engine/Source/Runtime/Core/Private/GenericPlatform/GenericPlatformMisc.cpp
+++ b/Engine/Source/Runtime/Core/Private/GenericPlatform/GenericPlatformMisc.cpp
@@ -653,6 +653,18 @@ FString FGenericPlatformMisc::GetCPUBrand()
 	return FString( TEXT( "GenericCPUBrand" ) );
 }
 
+uint32 FGenericPlatformMisc::GetCPUModelID()
+{
+	// Not implemented cross-platform. Each platform may or may not choose to implement this.
+	return 0;
+}
+
+uint32 FGenericPlatformMisc::GetCPUMicrocodeRevision()
+{
+	// Not implemented cross-platform. Each platform may or may not choose to implement this.
+	return 0;
+}
+
 FString FGenericPlatformMisc::GetCPUChipset()
 {
 	return FString( TEXT( "Unknown" ) );
diff --git a/Engine/Source/Runtime/Core/Private/Misc/App.cpp b/Engine/Source/Runtime/Core/Private/Misc/App.cpp
--- a/Engine/Source/Runtime/Core/Private/Misc/App.cpp
+++ b/Engine/Source/Runtime/Core/Private/Misc/App.cpp
@@ -405,6 +405,7 @@ void FApp::PrintStartupLogMessages()
 	FString OSLabel, OSVersion;
 	FPlatformMisc::GetOSVersions(OSLabel, OSVersion);
 	UE_LOG(LogInit, Log, TEXT("OS: %s (%s), CPU: %s, GPU: %s"), *OSLabel, *OSVersion, *FPlatformMisc::GetCPUBrand(), *FPlatformMisc::GetPrimaryGPUBrand());
+	UE_LOG(LogInit, Log, TEXT("CPU Model ID: %x, microcode revision: %x"), FPlatformMisc::GetCPUModelID(), FPlatformMisc::GetCPUMicrocodeRevision());
 
 #if PLATFORM_64BITS
 	UE_LOG(LogInit, Log, TEXT("Compiled (64-bit): %s %s"), BuildSettings::GetBuildDate(), BuildSettings::GetBuildTime());
diff --git a/Engine/Source/Runtime/Core/Private/Unix/UnixPlatformMisc.cpp b/Engine/Source/Runtime/Core/Private/Unix/UnixPlatformMisc.cpp
--- a/Engine/Source/Runtime/Core/Private/Unix/UnixPlatformMisc.cpp
+++ b/Engine/Source/Runtime/Core/Private/Unix/UnixPlatformMisc.cpp
@@ -807,7 +807,11 @@ const TCHAR* FUnixPlatformMisc::GetNullRHIShaderFormat()
 	_XTAG( CPUarchitecture, "CPU architecture" ), \
 	_XTAG( CPUvariant, "CPU variant" ),           \
 	_XTAG( CPUpart, "CPU part" ),                 \
-	_XTAG( CPUrevision, "CPU revision" ),
+	_XTAG( CPUrevision, "CPU revision" ),         \
+	_XTAG( CPUfamily, "cpu family"),              \
+	_XTAG( CPUmodel, "model"),                    \
+	_XTAG( CPUstepping, "stepping"),              \
+	_XTAG( CPUmicrocode, "microcode"),
 
 class FPlatformProcCPUInfo
 {
@@ -816,9 +820,11 @@ public:
 	// True if we were able to parse /proc/cpuinfo and find (at least) CPU Vendor
 	static bool IsValid()          { return (GIsValid == 1); }
 
-	static FString GetCPUBrand()   { ParseCPUInfoFile(); return GCPUBrand; }
-	static FString GetCPUVendor()  { ParseCPUInfoFile(); return GCPUVendor; }
-	static FString GetCPUChipset() { ParseCPUInfoFile(); return GCPUChipset; }
+	static FString GetCPUBrand()			{ ParseCPUInfoFile(); return GCPUBrand; }
+	static FString GetCPUVendor()			{ ParseCPUInfoFile(); return GCPUVendor; }
+	static FString GetCPUChipset()			{ ParseCPUInfoFile(); return GCPUChipset; }
+	static uint32 GetCPUModelID()			{ ParseCPUInfoFile(); return GCPUModelID; }
+	static uint32 GetCPUMicrocodeRevision()	{ ParseCPUInfoFile(); return GCPUMicrocodeRevision; }
 
 private:
 
@@ -846,6 +852,8 @@ private:
 	static FString GCPUVendor;
 	static FString GCPUBrand;
 	static FString GCPUChipset;
+	static uint32 GCPUModelID;
+	static uint32 GCPUMicrocodeRevision;
 };
 
 // Make of the device we are running on eg "Marvell"
@@ -854,6 +862,8 @@ FString FPlatformProcCPUInfo::GCPUVendor;
 FString FPlatformProcCPUInfo::GCPUBrand;
 // Chipset eg Cortex-A72 rev3
 FString FPlatformProcCPUInfo::GCPUChipset;
+uint32 FPlatformProcCPUInfo::GCPUModelID = 0;
+uint32 FPlatformProcCPUInfo::GCPUMicrocodeRevision = 0;
 
 int FPlatformProcCPUInfo::GIsValid = -1;
 
@@ -1105,6 +1115,34 @@ bool FPlatformProcCPUInfo::ParseCPUInfoFile()
 			GCPUBrand = GCPUChipset;
 		}
 
+		// Build a model ID from the family, model and stepping. Works for x64 CPUs.
+		{
+			const FString& strFamily = CPUInfoValues[TokenCPUfamily];
+			const FString& strModel = CPUInfoValues[TokenCPUmodel];
+			const FString& strStepping = CPUInfoValues[TokenCPUstepping];
+			if (!strFamily.IsEmpty() && !strModel.IsEmpty() && !strStepping.IsEmpty())
+			{
+				uint32 FamilyID = FCString::Atoi(*strFamily);
+				uint32 ModelID = FCString::Atoi(*strModel);
+				uint32 SteppingID = FCString::Atoi(*strStepping);
+				// Use the same packing logic as FWindowsPlatformMisc::GetCPUModelID().
+				GCPUModelID = (FamilyID << 12) | (ModelID << 4) | SteppingID;
+			}
+		}
+
+		{
+			const FString& strMicrocode = CPUInfoValues[TokenCPUmicrocode];
+			if (strMicrocode.StartsWith(TEXT("0x")))
+			{
+				TCHAR* endPtr;
+				GCPUMicrocodeRevision = FCString::Strtoi(*strMicrocode + 2, &endPtr, 16);
+				if (*endPtr != 0)
+				{
+					GCPUMicrocodeRevision = 0;
+				}
+			}
+		}
+
 		// Mark success if we found a vendor at least.
 		GIsValid = GCPUVendor.IsEmpty() ? 0 : 1;
 	}
@@ -1231,6 +1269,16 @@ FString FUnixPlatformMisc::GetCPUBrand()
 #endif // PLATFORM_HAS_CPUID
 }
 
+uint32 FUnixPlatformMisc::GetCPUModelID()
+{
+	return FPlatformProcCPUInfo::GetCPUModelID();
+}
+
+uint32 FUnixPlatformMisc::GetCPUMicrocodeRevision()
+{
+	return FPlatformProcCPUInfo::GetCPUMicrocodeRevision();
+}
+
 // From RHIVendorIdToString()
 static const TCHAR* VendorIdToString(uint32 VendorId)
 {
diff --git a/Engine/Source/Runtime/Core/Private/Windows/WindowsPlatformMisc.cpp b/Engine/Source/Runtime/Core/Private/Windows/WindowsPlatformMisc.cpp
--- a/Engine/Source/Runtime/Core/Private/Windows/WindowsPlatformMisc.cpp
+++ b/Engine/Source/Runtime/Core/Private/Windows/WindowsPlatformMisc.cpp
@@ -3187,6 +3187,66 @@ FString FWindowsPlatformMisc::GetCPUBrand()
 	return FCPUIDQueriedData::GetBrand();
 }
 
+uint32 FWindowsPlatformMisc::GetCPUModelID()
+{
+	uint32 CpuID = FCPUIDQueriedData::GetCPUInfo();
+	uint32 SteppingID = CpuID & 0x0F;
+	uint32 ModelID = (CpuID >> 4) & 0x0F;
+	uint32 FamilyID = (CpuID >> 8) & 0x0F;
+
+	if (FamilyID == 0x06 || FamilyID == 0x0F)
+	{
+		uint32 ExtendedModelID = (CpuID >> 16) & 0x0F;
+		ModelID = (ExtendedModelID << 4) + ModelID;
+	}
+
+	if (FamilyID == 0x0F)
+	{
+		uint32 ExtendedFamilyID = (CpuID >> 20) & 0xFF;
+		FamilyID = FamilyID + ExtendedFamilyID;
+	}
+
+	// Repack Family (12 bits), Model (8 bits), Stepping (4 bits) into a single 32-bit integer.
+	return (FamilyID << 12) | (ModelID << 4) | SteppingID;
+}
+
+uint32 FWindowsPlatformMisc::GetCPUMicrocodeRevision()
+{
+	// The microcode revision number is only available in a model-specific register, so we can't read it from userland. However, Windows
+	// seems to store the value in the registry at boot time, so we can try to get it from there.
+
+	HKEY Key = 0;
+	if (RegOpenKeyEx(HKEY_LOCAL_MACHINE, TEXT("HARDWARE\\DESCRIPTION\\System\\CentralProcessor\\0"), 0, KEY_READ, &Key) != ERROR_SUCCESS)
+	{
+		return 0;
+	}
+
+	uint32 MicrocodeVersion = 0;
+
+	uint32 MicrocodeVersionBuf[2];
+	::DWORD Size = 0, ValueType = 0;
+	if (RegQueryValueEx(Key, TEXT("Update Revision"), NULL, &ValueType, NULL, &Size) == ERROR_SUCCESS && ValueType == REG_BINARY && Size <= sizeof(MicrocodeVersionBuf))
+	{
+		if (RegQueryValueEx(Key, TEXT("Update Revision"), NULL, NULL, (LPBYTE)MicrocodeVersionBuf, &Size) == ERROR_SUCCESS)
+		{
+			if (Size == 8)
+			{
+				// On older Windows versions the registry key contains the full 64-bit value of the IA32_BIOS_SIGN_ID MSR. In that case,
+				// the microcode revision number is in the high 32 bits.
+				MicrocodeVersion = MicrocodeVersionBuf[1];
+			}
+			else if (Size == 4)
+			{
+				// In most cases, Windows already shifts the value and stores only the 32-bit microcode revision number.
+				MicrocodeVersion = MicrocodeVersionBuf[0];
+			}
+		}
+	}
+
+	RegCloseKey(Key);
+	return MicrocodeVersion;
+}
+
 #if PLATFORM_CPU_X86_FAMILY
 
 #ifdef _MSC_VER
diff --git a/Engine/Source/Runtime/Core/Public/GenericPlatform/GenericPlatformCrashContext.h b/Engine/Source/Runtime/Core/Public/GenericPlatform/GenericPlatformCrashContext.h
--- a/Engine/Source/Runtime/Core/Public/GenericPlatform/GenericPlatformCrashContext.h
+++ b/Engine/Source/Runtime/Core/Public/GenericPlatform/GenericPlatformCrashContext.h
@@ -173,6 +173,8 @@ struct FSessionContext
 	int32 					LanguageLCID;
 	int32 					NumberOfCores;
 	int32 					NumberOfCoresIncludingHyperthreads;
+	uint32					CPUModelID;
+	uint32					CPUMicrocodeRevision;
 	int32 					SecondsSinceStart;
 	int32 					CrashDumpMode;
 	int32					CrashTrigger;
diff --git a/Engine/Source/Runtime/Core/Public/GenericPlatform/GenericPlatformMisc.h b/Engine/Source/Runtime/Core/Public/GenericPlatform/GenericPlatformMisc.h
--- a/Engine/Source/Runtime/Core/Public/GenericPlatform/GenericPlatformMisc.h
+++ b/Engine/Source/Runtime/Core/Public/GenericPlatform/GenericPlatformMisc.h
@@ -746,6 +746,8 @@ struct FGenericPlatformMisc
 	 * @return	CPU brand string
 	 */
 	static CORE_API FString GetCPUBrand();
+	static CORE_API uint32 GetCPUModelID();
+	static CORE_API uint32 GetCPUMicrocodeRevision();
 
 	/**
 	 * Returns the CPU chipset if known
diff --git a/Engine/Source/Runtime/Core/Public/Unix/UnixPlatformMisc.h b/Engine/Source/Runtime/Core/Public/Unix/UnixPlatformMisc.h
--- a/Engine/Source/Runtime/Core/Public/Unix/UnixPlatformMisc.h
+++ b/Engine/Source/Runtime/Core/Public/Unix/UnixPlatformMisc.h
@@ -80,6 +80,8 @@ struct FUnixPlatformMisc : public FGenericPlatformMisc
 	static CORE_API bool HasCPUIDInstruction();
 	static CORE_API FString GetCPUVendor();
 	static CORE_API FString GetCPUBrand();
+	static CORE_API uint32 GetCPUModelID();
+	static CORE_API uint32 GetCPUMicrocodeRevision();
 
 	/**
 	 * Uses cpuid instruction to get the vendor string
diff --git a/Engine/Source/Runtime/Core/Public/Windows/WindowsPlatformMisc.h b/Engine/Source/Runtime/Core/Public/Windows/WindowsPlatformMisc.h
--- a/Engine/Source/Runtime/Core/Public/Windows/WindowsPlatformMisc.h
+++ b/Engine/Source/Runtime/Core/Public/Windows/WindowsPlatformMisc.h
@@ -253,6 +253,9 @@ struct FWindowsPlatformMisc
 
 	static CORE_API FString GetCPUVendor();
 	static CORE_API FString GetCPUBrand();
+	static CORE_API uint32 GetCPUModelID();
+	static CORE_API uint32 GetCPUMicrocodeRevision();
+
 	static CORE_API FString GetPrimaryGPUBrand();
 	static CORE_API struct FGPUDriverInfo GetGPUDriverInfo(const FString& DeviceDescription, bool bVerbose = true);
 	static CORE_API void GetOSVersions( FString& out_OSVersionLabel, FString& out_OSSubVersionLabel );
