FPaths:CreateStandardFilename.
* Change behavior to normalize paths that cannot be standardized.
* Remove the diagnostic that helped us to find the d:\root\..\foo case that the recent change was not handling correctly.
* Remove some undefined behavior for an edge case with absolute paths and add a warning if it ever occurs.
* Add a minor optimization by combining some operations that enumerate the string.

diff --git a/Engine/Source/Runtime/Core/Private/Misc/Paths.cpp b/Engine/Source/Runtime/Core/Private/Misc/Paths.cpp
--- a/Engine/Source/Runtime/Core/Private/Misc/Paths.cpp
+++ b/Engine/Source/Runtime/Core/Private/Misc/Paths.cpp
@@ -1247,17 +1247,74 @@ bool FPaths::IsRelative(const FString& InPath)
 	return !IsRooted;
 }
 
+void FPaths::PlatformIndependentNormalizePath(FString& Path, bool bRemoveDuplicateSlashes,
+	int32& OutFirstColonIndex, int32& OutFirstSlashIndex)
+{
+	TCHAR* Start = Path.GetCharArray().GetData();
+	TCHAR* ReadHead = Start;
+	TCHAR* WriteHead = Start;
+	TCHAR* End = ReadHead + Path.Len();
+	OutFirstColonIndex = INDEX_NONE;
+	OutFirstSlashIndex = INDEX_NONE;
+
+	while (ReadHead < End)
+	{
+		TCHAR C = *ReadHead++;
+		bool bWrite = true;
+
+		// 1. Convert '\' to '/'.
+		if (C == '\\')
+		{
+			C = '/';
+		}
+
+		// 2. Find FirstColonIndex and FirstSlashIndex for categorization
+		if (C == '/' && OutFirstSlashIndex == INDEX_NONE)
+		{
+			OutFirstSlashIndex = static_cast<int32>(WriteHead - Start);
+		}
+		if (C == ':' && OutFirstColonIndex == INDEX_NONE)
+		{
+			OutFirstColonIndex = static_cast<int32>(WriteHead - Start);
+		}
+
+		// 3. RemoveDuplicateSlashes (note that this behavior should be identical to FPaths::RemoveDuplicateSlashes)
+		if (bRemoveDuplicateSlashes && C == '/' && WriteHead > Start && WriteHead[-1] == '/')
+		{
+			bWrite = false;
+		}
+
+		if (bWrite)
+		{
+			*WriteHead++ = C;
+		}
+	}
+
+	if (WriteHead < End)
+	{
+		*WriteHead++ = TEXT('\0');
+		Path.TrimToNullTerminator();
+	}
+}
+
 void FPaths::NormalizeFilename(FString& InPath)
 {
-	InPath.ReplaceInline(TEXT("\\"), TEXT("/"), ESearchCase::CaseSensitive);
+	int32 FirstColonIndex;
+	int32 FirstSlashIndex;
+	PlatformIndependentNormalizePath(InPath, false /* bRemoveDuplicateSlashes */, FirstColonIndex, FirstSlashIndex);
 
 	FPlatformMisc::NormalizePath(InPath);
 }
 
 void FPaths::NormalizeDirectoryName(FString& InPath)
 {
-	InPath.ReplaceInline(TEXT("\\"), TEXT("/"), ESearchCase::CaseSensitive);
-	if (InPath.EndsWith(TEXT("/"), ESearchCase::CaseSensitive) && !InPath.EndsWith(TEXT("//"), ESearchCase::CaseSensitive) && !InPath.EndsWith(TEXT(":/"), ESearchCase::CaseSensitive))
+	int32 FirstColonIndex;
+	int32 FirstSlashIndex;
+	PlatformIndependentNormalizePath(InPath, false /* bRemoveDuplicateSlashes */, FirstColonIndex, FirstSlashIndex);
+
+	if (InPath.EndsWith(TEXT("/"), ESearchCase::CaseSensitive) 
+		&& !InPath.EndsWith(TEXT("//"), ESearchCase::CaseSensitive)
+		&& !InPath.EndsWith(TEXT(":/"), ESearchCase::CaseSensitive))
 	{
 		// overwrite trailing slash with terminator
 		InPath.GetCharArray()[InPath.Len() - 1] = TEXT('\0');
@@ -1332,58 +1389,77 @@ FString FPaths::CreateStandardFilename(const FString& InPath)
 		return FPaths::CreateStandardFilename(BaseDir);
 	}
 
-	FString WithSlashes = InPath;
-	WithSlashes.ReplaceInline(TEXT("\\"), TEXT("/"), ESearchCase::CaseSensitive);
-	const FStringView SlashesView(WithSlashes);
+	// If we have a windows network path, leave it alone; we do not have a contract for what standardized
+	// form of a network path looks like - do we keep all slashes as windows backslashes? Keep the first two
+	// and convert the rest? Convert all of them including the first two?
+	if (InPath.StartsWith(TEXTVIEW("\\\\")))
+	{
+		return InPath;
+	}
+
+	// Normalize and categorize the path characters before collapsing .. directories and putting it in standard form.
+	FString NormalizedPath = InPath;
+	int32 ColonIndex = -1;
+	int32 SlashIndex = -1;
+	PlatformIndependentNormalizePath(NormalizedPath, true /* bRemoveDuplicateSlashes */, ColonIndex, SlashIndex);
+	bool bHasVolumeSpecifier = ColonIndex > 0 && (SlashIndex < 0 || ColonIndex < SlashIndex);
 
 	const TCHAR* RootDirectory = FPlatformMisc::RootDir();
 	const FStringView RootDirectoryView(RootDirectory);
 
-	// look for paths that cannot be made relative, and are therefore left alone
-	// UNC (windows) network path
-	bool bCannotBeStandardized = InPath.StartsWith(TEXT("\\\\"), ESearchCase::CaseSensitive);
-	// windows drive letter path that doesn't start with base dir
-	bCannotBeStandardized |= ((InPath.Len() > 1) && (InPath[1] == ':') && !SlashesView.StartsWith(RootDirectoryView));
-	// Unix style absolute path that doesn't start with base dir
-	bCannotBeStandardized |= (WithSlashes.GetCharArray()[0] == '/' && !SlashesView.StartsWith(RootDirectoryView));
+	// look for paths that cannot be made relative, and are therefore normalized but not standardized.
+	// 1. UNC (windows) network path would be one of these, but we have an even stricter rule and early exited above.
+	// 2. windows volume path (e.g. 'd:\root2' or 'Storage:\root2') that doesn't start with base dir
+	bool bCannotBeStandardized = bHasVolumeSpecifier && !NormalizedPath.StartsWith(RootDirectoryView);
+	// 3. Unix style absolute path (e.g. `/home`) that doesn't start with base dir
+	bCannotBeStandardized |= SlashIndex == 0 && !NormalizedPath.StartsWith(RootDirectoryView);
 
-	// if it can't be standardized, just return itself
+	// if it can't be standardized, normalize it but don't try to convert it to a relative path from root.
 	if (bCannotBeStandardized)
 	{
-		return InPath;
+		FPlatformMisc::NormalizePath(NormalizedPath);
+		CollapseRelativeDirectories(NormalizedPath, true /* bCollapseAllPossible */);
+		return NormalizedPath;
 	}
 
-	// remove duplicate slashes
-	FPaths::RemoveDuplicateSlashes(WithSlashes);
-
-	// If the path starts with RelativePathToRoot, it is done, already in standard relative-to-root form.
+	// If the path starts with RelativePathToRoot, it is mostly done and just needs to be normalized.
 	const FString& RelativePathToRoot = FPaths::GetRelativePathToRoot();
-	if (WithSlashes.StartsWith(RelativePathToRoot))
+	if (NormalizedPath.StartsWith(RelativePathToRoot))
 	{
-		NormalizeFilename(WithSlashes);
-		CollapseRelativeDirectories(WithSlashes, true /* bCollapseAllPossible */);
-		return WithSlashes;
+		FPlatformMisc::NormalizePath(NormalizedPath);
+		CollapseRelativeDirectories(NormalizedPath, true /* bCollapseAllPossible */);
+		return NormalizedPath;
 	}
 
 	// Make it an absolute path and then make that absolute path a subdirectory of the RelativePathToRoot.
 	// This standardizes both incoming absolute paths, and relative paths that are relative to a subdirectory of
 	// root rather than starting with the full relative to root path.
-	// d:\root\Engine\Content                            -> ..\..\..\Engine\Content
-	// ..\..\Content                                     -> ..\..\..\Engine\Content
-	// d:\root\Engine\Binaries\Platform\UnrealEditor.exe -> ..\..\..\Engine\Binaries\Platform\UnrealEditor.exe
-	// UnrealEditor.exe                                  -> ..\..\..\Engine\Binaries\Platform\UnrealEditor.exe
-	FString FullPath = FPaths::ConvertRelativePathToFull(WithSlashes);
-	if (FullPath.StartsWith(RootDirectoryView))
-	{
-		WithSlashes.Reset();
-		WithSlashes += RelativePathToRoot;
-		WithSlashes += FStringView(FullPath).RightChop(RootDirectoryView.Len());
-		return WithSlashes;
-	}
-	UE_LOG(LogPaths, Verbose, TEXT("CreateStandardFilename: Unexpected: %s does not start with %.*s."),
-		*FullPath, RootDirectoryView.Len(), RootDirectoryView.GetData());
-	FullPath.ReplaceInline(RootDirectory, *RelativePathToRoot);
-	return FullPath;
+	// d:/root/Engine/Content                            -> ../../../Engine/Content
+	// ../../Content                                     -> ../../../Engine/Content
+	// d:/root/Engine/Binaries/Platform/UnrealEditor.exe -> ../../../Engine/Binaries/Platform/UnrealEditor.exe
+	// UnrealEditor.exe                                  -> ../../../Engine/Binaries/Platform/UnrealEditor.exe
+	FString FullPath = FPaths::ConvertRelativePathToFull(NormalizedPath);
+	if (!FullPath.StartsWith(RootDirectoryView))
+	{
+		// Relative paths from the root directory that use .. to escape the root directory can end up here.
+		// D:\root\..\AssetRegistryCache -> D:/AssetRegistryCache
+		// Return their absolute path as the standard path.
+		if (FullPath.Contains(RootDirectoryView))
+		{
+			// Unexpected that the absolute path of the RootDirectory appears at offset > 0 in the absolute path for
+			// InPath. In legacy behavior we used to replace it with RelativePathToRoot, but that would likely be
+			// invalid. Log a warning so that we can find this case where we deviate from legacy behavior if it exists.
+			UE_LOG(LogPaths, Warning,
+				TEXT("CreateStandardFilename: Unexpected: InPath '%s' transformed to '%s' which contains but does not start with '%.*s'."),
+				*InPath, *FullPath, RootDirectoryView.Len(), RootDirectoryView.GetData());
+		}
+		return FullPath;
+	}
+
+	NormalizedPath.Reset();
+	NormalizedPath += RelativePathToRoot;
+	NormalizedPath += FStringView(FullPath).RightChop(RootDirectoryView.Len());
+	return NormalizedPath;
 }
 
 void FPaths::MakeStandardFilename(FString& InPath)
diff --git a/Engine/Source/Runtime/Core/Public/Misc/Paths.h b/Engine/Source/Runtime/Core/Public/Misc/Paths.h
--- a/Engine/Source/Runtime/Core/Public/Misc/Paths.h
+++ b/Engine/Source/Runtime/Core/Public/Misc/Paths.h
@@ -720,6 +720,13 @@ public:
 protected:
 
 	static CORE_API FString CombineInternal(const FStringView* Paths, int32 NumPaths);
+	/**
+	 * Helper function to reduce duplication between NormalizeFilename and other functions that need to normalize
+	 * while also doing other transformations. Does not do all steps of NormalizeFilename, see the comments in the
+	 * function for details of which steps it handles.
+	 */
+	static void PlatformIndependentNormalizePath(FString& Path, bool bRemoveDuplicateSlashes,
+		int32& OutFirstColonIndex, int32& OutFirstSlashIndex);
 
 private:
 	struct FStaticData;
diff --git a/Engine/Source/Runtime/Core/Tests/Misc/PathsTest.cpp b/Engine/Source/Runtime/Core/Tests/Misc/PathsTest.cpp
--- a/Engine/Source/Runtime/Core/Tests/Misc/PathsTest.cpp
+++ b/Engine/Source/Runtime/Core/Tests/Misc/PathsTest.cpp
@@ -279,23 +279,23 @@ void TestCreateStandardFilename()
 	Run(TEXT("/Folder"), TEXT("/Folder"));
 	Run(TEXT("X:/Folder"), TEXT("X:/Folder"));
 	// Current contract is that paths that cannot be standardized do not CollapseRelativeDirectories; they remain unchanged
-	Run(TEXT("X:/Folder/.."), TEXT("X:/Folder/..")); // cannot be standardized
-	Run(TEXT("X:/Folder/../"), TEXT("X:/Folder/../")); // cannot be standardized
-	Run(TEXT("X:/Folder/../file.txt"), TEXT("X:/Folder/../file.txt")); // cannot be standardized
+	Run(TEXT("X:/Folder/.."), TEXT("X:"));
+	Run(TEXT("X:/Folder/../"), TEXT("X:/"));
+	Run(TEXT("X:/Folder/../file.txt"), TEXT("X:/file.txt")); // cannot be standardized
 	Run(TEXT("Folder/.."), StandardPathOfBase);
 	Run(TEXT("Folder/../"), StandardPathOfBase + TEXT("/")); // Keeping the trailing slash is legacy behavior in CollapseSubdirectories
 	Run(TEXT("Folder/../file.txt"), FPaths::Combine(StandardPathOfBase, TEXT("file.txt")));
-	Run(TEXT("/Folder/.."), TEXT("/Folder/..")); // cannot be standardized
-	Run(TEXT("/Folder/../"), TEXT("/Folder/../")); // cannot be standardized
-	Run(TEXT("/Folder/../file.txt"), TEXT("/Folder/../file.txt")); // cannot be standardized
+	Run(TEXT("/Folder/.."), TEXT(""));
+	Run(TEXT("/Folder/../"), TEXT("/"));
+	Run(TEXT("/Folder/../file.txt"), TEXT("/file.txt")); // cannot be standardized
 	Run(TEXT("Folder/../.."), StandardPathOfBaseParent);
 	Run(TEXT("Folder/../../"), FString(StandardPathOfBaseParent) + TEXT("/")); // Keeping the trailing slash is legacy behavior in CollapseSubdirectories
 	Run(TEXT("Folder/../../file.txt"), FPaths::Combine(StandardPathOfBaseParent, TEXT("file.txt")));
 	Run(TEXT("X:/.."), TEXT("X:/.."));
-	Run(TEXT("X:/."), TEXT("X:/.")); // cannot be standardized
-	Run(TEXT("X:/./"), TEXT("X:/./")); // cannot be standardized
-	Run(TEXT("X:/./file.txt"), TEXT("X:/./file.txt")); // cannot be standardized
-	Run(TEXT("X:/Folder1/../Folder2"), TEXT("X:/Folder1/../Folder2")); // cannot be standardized
+	Run(TEXT("X:/."), TEXT("X:/."));
+	Run(TEXT("X:/./"), TEXT("X:/"));
+	Run(TEXT("X:/./file.txt"), TEXT("X:/file.txt"));
+	Run(TEXT("X:/Folder1/../Folder2"), TEXT("X:/Folder2"));
 	Run(TEXT("Folder1/../Folder2"), FPaths::Combine(StandardPathOfBase, TEXT("Folder2")));
 	Run(TEXT("Folder1/../Folder2/"), FPaths::Combine(StandardPathOfBase, TEXT("Folder2/"))); // Keeping the trailing slash is legacy behavior in CollapseSubdirectories
 	Run(TEXT("Folder1/../Folder2/file.txt"), FPaths::Combine(StandardPathOfBase, TEXT("Folder2/file.txt")));
