Fix variable shadowing warning.

#rnx

diff --git a/Engine/Source/Runtime/Core/Private/Unix/UnixPlatformMisc.cpp b/Engine/Source/Runtime/Core/Private/Unix/UnixPlatformMisc.cpp
--- a/Engine/Source/Runtime/Core/Private/Unix/UnixPlatformMisc.cpp
+++ b/Engine/Source/Runtime/Core/Private/Unix/UnixPlatformMisc.cpp
@@ -1118,12 +1118,12 @@ bool FPlatformProcCPUInfo::ParseCPUInfoFile()
 		// Build a model ID from the family, model and stepping. Works for x64 CPUs.
 		{
 			const FString& strFamily = CPUInfoValues[TokenCPUfamily];
-			const FString& strModel = CPUInfoValues[TokenCPUmodel];
+			const FString& strCpuModel = CPUInfoValues[TokenCPUmodel];
 			const FString& strStepping = CPUInfoValues[TokenCPUstepping];
-			if (!strFamily.IsEmpty() && !strModel.IsEmpty() && !strStepping.IsEmpty())
+			if (!strFamily.IsEmpty() && !strCpuModel.IsEmpty() && !strStepping.IsEmpty())
 			{
 				uint32 FamilyID = FCString::Atoi(*strFamily);
-				uint32 ModelID = FCString::Atoi(*strModel);
+				uint32 ModelID = FCString::Atoi(*strCpuModel);
 				uint32 SteppingID = FCString::Atoi(*strStepping);
 				// Use the same packing logic as FWindowsPlatformMisc::GetCPUModelID().
 				GCPUModelID = (FamilyID << 12) | (ModelID << 4) | SteppingID;
