Tags editor: automatically add or checkout tag mapping files

diff --git a/Engine/Source/Developer/AutomationWindow/Private/SAutomationTestTagEditor.h b/Engine/Source/Developer/AutomationWindow/Private/SAutomationTestTagEditor.h
--- a/Engine/Source/Developer/AutomationWindow/Private/SAutomationTestTagEditor.h
+++ b/Engine/Source/Developer/AutomationWindow/Private/SAutomationTestTagEditor.h
@@ -24,6 +24,7 @@
 #include "Styling/AppStyle.h"
 #include "IAutomationControllerModule.h"
 #include "SAutomationTestTagListBox.h"
+#include "SourceControlHelpers.h"
 #include "AutomationWindowStyle.h"
 
 #include "SlateOptMacros.h"
@@ -346,7 +347,11 @@ private:
 							}), false);
 					check(Registered);
 				}
-				FAutomationTestFramework::Get().SaveTestTagMappings(TestNames, TestFilePaths);
+				TArray<FString> TestTagMapFiles = FAutomationTestFramework::Get().SaveTestTagMappings(TestNames, TestFilePaths);
+				if (SourceControlHelpers::IsAvailable())
+				{
+					SourceControlHelpers::CheckOutOrAddFiles(TestTagMapFiles);
+				}
 
 				OnTestTagsSavedEvent.ExecuteIfBound();
 
diff --git a/Engine/Source/Runtime/Core/Private/Misc/AutomationTest.cpp b/Engine/Source/Runtime/Core/Private/Misc/AutomationTest.cpp
--- a/Engine/Source/Runtime/Core/Private/Misc/AutomationTest.cpp
+++ b/Engine/Source/Runtime/Core/Private/Misc/AutomationTest.cpp
@@ -965,7 +965,7 @@ void FAutomationTestFramework::LoadTestTagMappings(const FString& BasePath)
 	}
 }
 
-void FAutomationTestFramework::SaveTestTagMappings(const TArray<FString>& TestFullNames, const TArray<FString>& TestFilePaths) const
+TArray<FString> FAutomationTestFramework::SaveTestTagMappings(const TArray<FString>& TestFullNames, const TArray<FString>& TestFilePaths) const
 {
 	const FString ProjectDirPath = FPaths::ConvertRelativePathToFull(*FPaths::ProjectDir());
 	const FString EngineDirPath = FPaths::ConvertRelativePathToFull(FPaths::EngineDir());
@@ -973,38 +973,41 @@ void FAutomationTestFramework::SaveTestTagMappings(const TArray<FString>& TestFu
 	TArray<FString> RestrictedFolderNames = FPaths::GetRestrictedFolderNames();
 	RestrictedFolderNames.Append({ TEXT("Platforms"), TEXT("Restricted") });
 
-	TArray<FString> RestrictedPathsProject;
-	GetPossibleRestrictedPaths(ProjectDirPath, RestrictedFolderNames, RestrictedPathsProject);
-	TArray<FString> RestrictedPathsEngine;
-	GetPossibleRestrictedPaths(EngineDirPath, RestrictedFolderNames, RestrictedPathsEngine);
+	TArray<FString> AllPathsProject;
+	GetPossibleRestrictedPaths(ProjectDirPath, RestrictedFolderNames, AllPathsProject);
+	TArray<FString> AllPathsEngine;
+	GetPossibleRestrictedPaths(EngineDirPath, RestrictedFolderNames, AllPathsEngine);
 
-	RestrictedPathsProject.Sort([](const FString& A, const FString& B) { return A.Len() > B.Len(); });
-	RestrictedPathsEngine.Sort([](const FString& A, const FString& B) { return A.Len() > B.Len(); });
+	AllPathsProject.Add(ProjectDirPath);
+	AllPathsProject.Sort([](const FString& A, const FString& B) { return A.Len() > B.Len(); });
+
+	AllPathsEngine.Add(EngineDirPath);
+	AllPathsEngine.Sort([](const FString& A, const FString& B) { return A.Len() > B.Len(); });
 
 	FString TestSourcePath;
 	FString TagsSettingsFilePath;
-	TSet<FString> AllTagConfigPaths;
+	TSet<FString> AffectedTagConfigPaths;
 	for (int i = 0; i < TestFilePaths.Num(); i++)
 	{
 		TagsSettingsFilePath.Empty();
 		TestSourcePath = TestFilePaths[i];
-		// Pick longest path with restricted names that is a suffix of the test path
-		for (int j = 0 ; j < RestrictedPathsProject.Num(); j++)
+		// Pick longest project path with restricted names that is a suffix of the test path, or the base project path if not in restricted path
+		for (int j = 0 ; j < AllPathsProject.Num(); j++)
 		{
-			if (FPaths::IsUnderDirectory(TestSourcePath, RestrictedPathsProject[j]))
+			if (FPaths::IsUnderDirectory(TestSourcePath, AllPathsProject[j]))
 			{
-				TagsSettingsFilePath = RestrictedPathsProject[j] / TEXT("Config") / TEXT("Tests") / TEXT("Tags.ini");
+				TagsSettingsFilePath = AllPathsProject[j] / TEXT("Config") / TEXT("Tests") / TEXT("Tags.ini");
 				break;
 			}
 		}
 
 		if (TagsSettingsFilePath.IsEmpty())
 		{
-			for (int j = 0; j < RestrictedPathsEngine.Num(); j++)
+			for (int j = 0; j < AllPathsEngine.Num(); j++)
 			{
-				if (FPaths::IsUnderDirectory(TestSourcePath, RestrictedPathsEngine[j]))
+				if (FPaths::IsUnderDirectory(TestSourcePath, AllPathsEngine[j]))
 				{
-					TagsSettingsFilePath = RestrictedPathsEngine[j] / TEXT("Config") / TEXT("Tests") / TEXT("Tags.ini");
+					TagsSettingsFilePath = AllPathsEngine[j] / TEXT("Config") / TEXT("Tests") / TEXT("Tags.ini");
 					break;
 				}
 			}
@@ -1016,35 +1019,22 @@ void FAutomationTestFramework::SaveTestTagMappings(const TArray<FString>& TestFu
 			continue;
 		}
 
-		if (!IFileManager::Get().FileExists(*TagsSettingsFilePath))
+		TagsSettingsFilePath = FConfigCacheIni::NormalizeConfigIniPath(TagsSettingsFilePath);
+		if (!AffectedTagConfigPaths.Contains(TagsSettingsFilePath))
 		{
-			TUniquePtr<FArchive> Ar = TUniquePtr<FArchive>(IFileManager::Get().CreateFileWriter(*TagsSettingsFilePath));
-			if (!Ar)
-			{
-				UE_LOG(LogAutomationTest, Warning, TEXT("Failed to create test tags mapping file."));
-			}
-			else
-			{
-				UE_LOG(LogAutomationTest, Log, TEXT("Created new test tags mapping file: %s."), *TagsSettingsFilePath);
-
-				const FString TagMappingHeader = TEXT("[TestTagMapping]");
-				Ar->Serialize(TCHAR_TO_ANSI(*TagMappingHeader), TagMappingHeader.Len());
-				Ar->Close();
-			}
-		}
-
-		if (!AllTagConfigPaths.Contains(TagsSettingsFilePath))
-		{
-			AllTagConfigPaths.Add(TagsSettingsFilePath);
+			AffectedTagConfigPaths.Add(TagsSettingsFilePath);
+			GConfig->Add(TagsSettingsFilePath, FConfigFile());
 		}
 		GConfig->SetString(TEXT("TestTagMapping"), *TestFullNames[i], **TestFullNameToTagDataMap.Find(TestFullNames[i]), TagsSettingsFilePath);
 	}
 
-	for (FString ConfigPath : AllTagConfigPaths)
+	for (FString ConfigPath : AffectedTagConfigPaths)
 	{
 		UE_LOG(LogAutomationTest, Log, TEXT("Saving tags to %s."), *ConfigPath);
 		GConfig->Flush(false, ConfigPath);
 	}
+
+	return AffectedTagConfigPaths.Array();
 }
 
 bool FAutomationTestFramework::ShouldTestContent(const FString& Path) const
diff --git a/Engine/Source/Runtime/Core/Public/Misc/AutomationTest.h b/Engine/Source/Runtime/Core/Public/Misc/AutomationTest.h
--- a/Engine/Source/Runtime/Core/Public/Misc/AutomationTest.h
+++ b/Engine/Source/Runtime/Core/Public/Misc/AutomationTest.h
@@ -1184,8 +1184,9 @@ public:
 
 	/**
 	 * Saves runtime edited test tag mappings to Tags.ini
+	 * Returns the list of updated tag mapping files
 	 */
-	CORE_API void SaveTestTagMappings(const TArray<FString>& TestFullNames, const TArray<FString>& TestFilePaths) const;
+	CORE_API TArray<FString> SaveTestTagMappings(const TArray<FString>& TestFullNames, const TArray<FString>& TestFilePaths) const;
 
 	/**
 	 * Whether the testing framework should allow content to be tested or not.  Intended to block developer directories.
