LiveTrace fixes

diff --git a/Engine/Source/Runtime/Core/Private/ProfilingDebugging/TraceAuxiliary.cpp b/Engine/Source/Runtime/Core/Private/ProfilingDebugging/TraceAuxiliary.cpp
--- a/Engine/Source/Runtime/Core/Private/ProfilingDebugging/TraceAuxiliary.cpp
+++ b/Engine/Source/Runtime/Core/Private/ProfilingDebugging/TraceAuxiliary.cpp
@@ -526,7 +526,7 @@ bool FTraceAuxiliaryImpl::ConnectRelay(UPTRINT Handle, UE::Trace::IoWriteFunc Wr
 	}
 
 	const bool bConnected = UE::Trace::RelayTo(Handle, WriteFunc, CloseFunc);
-	if (!bConnected)
+	if (bConnected)
 	{
 		UE_LOG(LogTrace, Display, TEXT("Trace started (connected to relay endpoint)"));
 
