Ensure android stats are reported each frame.
Switch Android's background stat gathering thread to just gathering and caching values.
The stats are reported each frame via the GT.

#tests Local client before/after comparisons.

diff --git a/Engine/Source/Runtime/Core/Private/Android/AndroidStats.cpp b/Engine/Source/Runtime/Core/Private/Android/AndroidStats.cpp
--- a/Engine/Source/Runtime/Core/Private/Android/AndroidStats.cpp
+++ b/Engine/Source/Runtime/Core/Private/Android/AndroidStats.cpp
@@ -120,14 +120,137 @@ static FAutoConsoleVariableRef CVarAndroidHWCPipeStatsEnabled(
 
 static int GThermalStatus = 0;
 static int GTrimMemoryBackgroundLevel = 0;
+static int GTrimMemoryForegroundLevel = 0;
+
 CSV_DEFINE_STAT(AndroidMemory, TrimMemoryBackgroundLevel);
+CSV_DEFINE_STAT(AndroidMemory, TrimMemoryForegroundLevel);
 
 CSV_DEFINE_STAT(AndroidMemory, Mem_RSS);
 CSV_DEFINE_STAT(AndroidMemory, Mem_Swap);
 CSV_DEFINE_STAT(AndroidMemory, Mem_TotalUsed);
 
-static int GTrimMemoryForegroundLevel = 0;
-CSV_DEFINE_STAT(AndroidMemory, TrimMemoryForegroundLevel);
+
+#if CSV_PROFILER || STATS
+namespace UE
+{
+	namespace Android
+	{
+		constexpr uint32 MaxFrequencyGroupStats = 4;
+
+		struct FCPUStatValues
+		{
+			float CPUTemp = 0.0f;
+#if CSV_PROFILER
+			int ThermalStatus = 0;
+			float Thermals5S = 0;
+			int TrimMemoryBackgroundLevel = 0;
+			int TrimMemoryForegroundLevel = 0;
+			FPlatformMemoryStats MemStats;
+#endif
+			struct FFrequencyGroupStats
+			{
+				int CoreCount = 0;
+				int MaxFrequency = 0;
+				int CurrentFrequency = 0.0f;
+				float MaxSingleCoreUtilization = 0.0f;
+			};
+			int NumFrequencyGroups = 0;
+			FFrequencyGroupStats FrequencyGroupStats[UE::Android::MaxFrequencyGroupStats];
+
+#if STATS
+			static inline const FName AndroidFrequencyGroupMaxFreqStats[] = {
+				GET_STATFNAME(STAT_FreqGroup0MaxFrequency),
+				GET_STATFNAME(STAT_FreqGroup1MaxFrequency),
+				GET_STATFNAME(STAT_FreqGroup2MaxFrequency),
+				GET_STATFNAME(STAT_FreqGroup3MaxFrequency),
+			};
+
+			static inline const FName AndroidFrequencyGroupMinFreqStats[] = {
+				GET_STATFNAME(STAT_FreqGroup0MinFrequency),
+				GET_STATFNAME(STAT_FreqGroup1MinFrequency),
+				GET_STATFNAME(STAT_FreqGroup2MinFrequency),
+				GET_STATFNAME(STAT_FreqGroup3MinFrequency),
+			};
+
+			static inline const FName AndroidFrequencyGroupCurrentFreqStats[] = {
+				GET_STATFNAME(STAT_FreqGroup0CurrentFrequency),
+				GET_STATFNAME(STAT_FreqGroup1CurrentFrequency),
+				GET_STATFNAME(STAT_FreqGroup2CurrentFrequency),
+				GET_STATFNAME(STAT_FreqGroup3CurrentFrequency),
+			};
+
+
+			static inline const FName AndroidFrequencyGroupCurrentFreqPercentageStats[] = {
+				GET_STATFNAME(STAT_FreqGroup0CurrentFrequencyPercentage),
+				GET_STATFNAME(STAT_FreqGroup1CurrentFrequencyPercentage),
+				GET_STATFNAME(STAT_FreqGroup2CurrentFrequencyPercentage),
+				GET_STATFNAME(STAT_FreqGroup3CurrentFrequencyPercentage),
+			};
+
+			static inline const FName AndroidFrequencyGroupNumCoresStats[] = {
+				GET_STATFNAME(STAT_FreqGroup0NumCores),
+				GET_STATFNAME(STAT_FreqGroup1NumCores),
+				GET_STATFNAME(STAT_FreqGroup2NumCores),
+				GET_STATFNAME(STAT_FreqGroup3NumCores),
+			};
+
+			static inline const FName AndroidFrequencyGroupMaxCoresUtilizationStats[] = {
+				GET_STATFNAME(STAT_FreqGroup0MaxUtilization),
+				GET_STATFNAME(STAT_FreqGroup1MaxUtilization),
+				GET_STATFNAME(STAT_FreqGroup2MaxUtilization),
+				GET_STATFNAME(STAT_FreqGroup3MaxUtilization),
+			};
+
+#endif
+			void PushStats() const
+			{
+#if CSV_PROFILER
+				CSV_CUSTOM_STAT_DEFINED(CPUTemp, CPUTemp, ECsvCustomStatOp::Set);
+				CSV_CUSTOM_STAT_DEFINED(ThermalStatus, GThermalStatus, ECsvCustomStatOp::Set);
+				CSV_CUSTOM_STAT_DEFINED(ThermalStress, Thermals5S, ECsvCustomStatOp::Set);
+				CSV_CUSTOM_STAT_DEFINED(TrimMemoryForegroundLevel, GTrimMemoryForegroundLevel, ECsvCustomStatOp::Set);
+				CSV_CUSTOM_STAT_DEFINED(TrimMemoryBackgroundLevel, GTrimMemoryBackgroundLevel, ECsvCustomStatOp::Set);
+				CSV_CUSTOM_STAT_DEFINED(Mem_Swap, (int)(MemStats.VMSwap / (1024 * 1024)), ECsvCustomStatOp::Set);
+				CSV_CUSTOM_STAT_DEFINED(Mem_RSS, (int)(MemStats.VMRss / (1024 * 1024)), ECsvCustomStatOp::Set);
+				CSV_CUSTOM_STAT_DEFINED(Mem_TotalUsed, (int)(MemStats.UsedPhysical / (1024 * 1024)), ECsvCustomStatOp::Set);
+				for (int32 FrequencyGroupIndex = 0; FrequencyGroupIndex < NumFrequencyGroups; FrequencyGroupIndex++)
+				{
+					const FFrequencyGroupStats& Stats = FrequencyGroupStats[FrequencyGroupIndex];
+					CSV_CUSTOM_STAT_DEFINED_BY_PTR(GCPUFreqStats[FrequencyGroupIndex], Stats.CurrentFrequency, ECsvCustomStatOp::Set);
+					const float CurrentFrequencyPercentage = ((float)Stats.CurrentFrequency / (float)Stats.MaxFrequency) * 100.0f;
+					CSV_CUSTOM_STAT_DEFINED_BY_PTR(GCPUFreqPercentageStats[FrequencyGroupIndex], CurrentFrequencyPercentage, ECsvCustomStatOp::Set);
+				}
+#endif
+#if STATS
+				static const FName CPUStatName = GET_STATFNAME(STAT_CPUTemp);
+				static const FName ThermalStatus = GET_STATFNAME(STAT_ThermalStatus);
+				SET_FLOAT_STAT_FName(CPUStatName, CPUTemp);
+				SET_DWORD_STAT_FName(ThermalStatus, GThermalStatus);
+
+				for (int32 FrequencyGroupIndex = 0; FrequencyGroupIndex < NumFrequencyGroups; FrequencyGroupIndex++)
+				{
+					const FFrequencyGroupStats& Stats = FrequencyGroupStats[FrequencyGroupIndex];
+					const float CurrentFrequencyPercentage = ((float)Stats.CurrentFrequency / (float)Stats.MaxFrequency) * 100.0f;
+
+					SET_DWORD_STAT_FName(AndroidFrequencyGroupMaxFreqStats[FrequencyGroupIndex], Stats.MaxFrequency);
+					SET_DWORD_STAT_FName(AndroidFrequencyGroupNumCoresStats[FrequencyGroupIndex], Stats.CoreCount);
+					//SET_DWORD_STAT_FName(AndroidFrequencyGroupMinFreqStats[FrequencyGroupIndex], FrequencyGroup.MinFrequency);
+					SET_DWORD_STAT_FName(AndroidFrequencyGroupCurrentFreqStats[FrequencyGroupIndex], Stats.CurrentFrequency);
+					SET_FLOAT_STAT_FName(AndroidFrequencyGroupCurrentFreqPercentageStats[FrequencyGroupIndex], CurrentFrequencyPercentage);
+					SET_FLOAT_STAT_FName(AndroidFrequencyGroupMaxCoresUtilizationStats[FrequencyGroupIndex], Stats.MaxSingleCoreUtilization);
+				}
+#endif
+			}
+
+			static FCPUStatValues& Get()
+			{
+				static FCPUStatValues CSVStatValues;
+				return CSVStatValues;
+			}
+		};
+	}
+}
+#endif // CSV_PROFILER || STATS
 
 void FAndroidStats::OnThermalStatusChanged(int Status)
 {
@@ -199,47 +322,42 @@ void FAndroidStats::UpdateAndroidStats()
 #if CSV_PROFILER || STATS
 	UpdateGPUStats();
 	SCOPED_NAMED_EVENT(UpdateAndroidStats, FColor::Green);
-	
-	if (GAndroidCPUStatsUpdateRate <= 0.0f || GIsStatTaskActive.load())
+
+	if (GAndroidCPUStatsUpdateRate <= 0.0f)
 	{
 		return;
 	}
 
-	GIsStatTaskActive = true;
-	// Run everything in a background task thread so that long system calls won't lock up game thread
-	UE::Tasks::Launch(UE_SOURCE_LOCATION, []()
-	{
-#if CSV_PROFILER
-		SCOPED_NAMED_EVENT(UpdateAndroidStats_CSV, FColor::Green);
-
-		static float CPUTemp = 0.0f;
-		static uint64 LastCollectionTime = FPlatformTime::Cycles64();
-		const uint64 CurrentTime = FPlatformTime::Cycles64();
-		const bool bUpdateStats = ((FPlatformTime::ToSeconds64(CurrentTime - LastCollectionTime) >= GAndroidCPUStatsUpdateRate));
-		static FPlatformMemoryStats MemStats = FAndroidPlatformMemory::GetStats();
-		if (bUpdateStats)
-		{
-			LastCollectionTime = CurrentTime;
-			CPUTemp = FAndroidMisc::GetCPUTemperature();
+	static UE::FMutex CSVStatsLock;
 
-			MemStats = FAndroidPlatformMemory::GetStats();
-		}
+	{
+		/*
+		*   FCPUStatValues contains the state that is pushed each frame to the CSV and stats systems.
+		*	The state is periodically updated via a background thread to avoid potential stalls on the GT.
+		*/
+		UE::TUniqueLock Lock(CSVStatsLock);
+		UE::Android::FCPUStatValues::Get().PushStats();
+	}
 
-		CSV_CUSTOM_STAT_DEFINED(CPUTemp, CPUTemp, ECsvCustomStatOp::Set);
-		CSV_CUSTOM_STAT_DEFINED(ThermalStatus, GThermalStatus, ECsvCustomStatOp::Set);
+	static uint64 LastCollectionTime = FPlatformTime::Cycles64();
+	const uint64 CurrentTime = FPlatformTime::Cycles64();
+	const bool bUpdateStats = ((FPlatformTime::ToSeconds64(CurrentTime - LastCollectionTime) >= GAndroidCPUStatsUpdateRate));
 
-		float Thermals5S = FAndroidPlatformThermal::GetThermalStress(FAndroidPlatformThermal::EForecastPeriod::FIVE_SEC);
-		CSV_CUSTOM_STAT_DEFINED(ThermalStress, Thermals5S, ECsvCustomStatOp::Set);
+	bool bExpected = false;
+	if (!bUpdateStats || !GIsStatTaskActive.compare_exchange_weak(bExpected, true))
+	{
+		return;
+	}
+	LastCollectionTime = CurrentTime;
 
-		CSV_CUSTOM_STAT_DEFINED(TrimMemoryBackgroundLevel, GTrimMemoryBackgroundLevel, ECsvCustomStatOp::Set);
-		CSV_CUSTOM_STAT_DEFINED(Mem_Swap, (int)(MemStats.VMSwap / (1024 * 1024)), ECsvCustomStatOp::Set);
-		CSV_CUSTOM_STAT_DEFINED(Mem_RSS, (int)(MemStats.VMRss / (1024 * 1024)), ECsvCustomStatOp::Set);
-		CSV_CUSTOM_STAT_DEFINED(Mem_TotalUsed, (int)(MemStats.UsedPhysical / (1024 * 1024)), ECsvCustomStatOp::Set);
+	// Run everything in a background task thread so that long system calls won't lock up game thread
+	UE::Tasks::Launch(UE_SOURCE_LOCATION, []()
+	{
+		SCOPED_NAMED_EVENT(UpdateAndroidStats, FColor::Green);
 
-		static const uint32 MaxFrequencyGroupStats = 4;
-		const int32 MaxCoresStatsSupport = 16;
+		constexpr int32 MaxCoresStatsSupport = 16;
 		int32 NumCores = FMath::Min(FAndroidMisc::NumberOfCores(), MaxCoresStatsSupport);
-
+		ensure(FAndroidMisc::NumberOfCores() <= MaxCoresStatsSupport);
 		struct FFrequencyGroup
 		{
 			uint32 MinFrequency;
@@ -295,66 +413,22 @@ void FAndroidStats::UpdateAndroidStats()
 			return 0u;
 		};
 
-		static int32 CurrentFrequencies[MaxFrequencyGroupStats] = { 0,0,0,0 };
-		static float CurrentFrequenciesPercentage[MaxFrequencyGroupStats] = { 0,0,0,0 };
-		for (int32 FrequencyGroupIndex = 0; FrequencyGroupIndex < FrequencyGroups.Num(); FrequencyGroupIndex++)
+		check(FrequencyGroups.Num() <= UE::Android::MaxFrequencyGroupStats);
+#if CSV_PROFILER
 		{
-			if (bUpdateStats)
-			{
-				CurrentFrequencies[FrequencyGroupIndex] = GetFrequencyGroupCurrentFrequency(FrequencyGroupIndex);
-				CurrentFrequenciesPercentage[FrequencyGroupIndex] = ((float)CurrentFrequencies[FrequencyGroupIndex] / (float)FrequencyGroups[FrequencyGroupIndex].MaxFrequency) * 100.0f;
-			}
-			CSV_CUSTOM_STAT_DEFINED_BY_PTR(GCPUFreqStats[FrequencyGroupIndex], CurrentFrequencies[FrequencyGroupIndex], ECsvCustomStatOp::Set);
-			CSV_CUSTOM_STAT_DEFINED_BY_PTR(GCPUFreqPercentageStats[FrequencyGroupIndex], CurrentFrequenciesPercentage[FrequencyGroupIndex], ECsvCustomStatOp::Set);
+			UE::TUniqueLock Lock(CSVStatsLock);
+			UE::Android::FCPUStatValues& CSVStatValues = UE::Android::FCPUStatValues::Get();
+
+			CSVStatValues.CPUTemp = FAndroidMisc::GetCPUTemperature();;
+			CSVStatValues.Thermals5S = FAndroidPlatformThermal::GetThermalStress(FAndroidPlatformThermal::EForecastPeriod::FIVE_SEC);;
+			CSVStatValues.ThermalStatus = GThermalStatus;
+			CSVStatValues.TrimMemoryBackgroundLevel = GTrimMemoryBackgroundLevel;
+			CSVStatValues.TrimMemoryForegroundLevel = GTrimMemoryForegroundLevel;
+			CSVStatValues.MemStats = FAndroidPlatformMemory::GetStats();
 		}
-#endif
+#endif //CSV_PROFILER
 
-#if STATS
-		static const FName AndroidFrequencyGroupMaxFreqStats[] = {
-			GET_STATFNAME(STAT_FreqGroup0MaxFrequency),
-			GET_STATFNAME(STAT_FreqGroup1MaxFrequency),
-			GET_STATFNAME(STAT_FreqGroup2MaxFrequency),
-			GET_STATFNAME(STAT_FreqGroup3MaxFrequency),
-		};
-
-		static const FName AndroidFrequencyGroupMinFreqStats[] = {
-			GET_STATFNAME(STAT_FreqGroup0MinFrequency),
-			GET_STATFNAME(STAT_FreqGroup1MinFrequency),
-			GET_STATFNAME(STAT_FreqGroup2MinFrequency),
-			GET_STATFNAME(STAT_FreqGroup3MinFrequency),
-		};
-
-		static const FName AndroidFrequencyGroupCurrentFreqStats[] = {
-			GET_STATFNAME(STAT_FreqGroup0CurrentFrequency),
-			GET_STATFNAME(STAT_FreqGroup1CurrentFrequency),
-			GET_STATFNAME(STAT_FreqGroup2CurrentFrequency),
-			GET_STATFNAME(STAT_FreqGroup3CurrentFrequency),
-		};
-
-		
-		static const FName AndroidFrequencyGroupCurrentFreqPercentageStats[] = {
-			GET_STATFNAME(STAT_FreqGroup0CurrentFrequencyPercentage),
-			GET_STATFNAME(STAT_FreqGroup1CurrentFrequencyPercentage),
-			GET_STATFNAME(STAT_FreqGroup2CurrentFrequencyPercentage),
-			GET_STATFNAME(STAT_FreqGroup3CurrentFrequencyPercentage),
-		};
-
-		static const FName AndroidFrequencyGroupNumCoresStats[] = {
-			GET_STATFNAME(STAT_FreqGroup0NumCores),
-			GET_STATFNAME(STAT_FreqGroup1NumCores),
-			GET_STATFNAME(STAT_FreqGroup2NumCores),
-			GET_STATFNAME(STAT_FreqGroup3NumCores),
-		};
-
-		static const FName AndroidFrequencyGroupMaxCoresUtilizationStats[] = {
-			GET_STATFNAME(STAT_FreqGroup0MaxUtilization),
-			GET_STATFNAME(STAT_FreqGroup1MaxUtilization),
-			GET_STATFNAME(STAT_FreqGroup2MaxUtilization),
-			GET_STATFNAME(STAT_FreqGroup3MaxUtilization),
-		};
-
-		static float MaxSingleCoreUtilization[MaxFrequencyGroupStats] = { 0.0f };
-		if (bUpdateStats)
+		float MaxSingleCoreUtilization[UE::Android::MaxFrequencyGroupStats] = { 0.0f };
 		{
 			FAndroidMisc::FCPUState& AndroidCPUState = FAndroidMisc::GetCPUState();
 			for (int32 CoreIndex = 0; CoreIndex < NumCores; CoreIndex++)
@@ -368,22 +442,21 @@ void FAndroidStats::UpdateAndroidStats()
 			}
 		}
 
-		for (int32 FrequencyGroupIndex = 0; FrequencyGroupIndex < FrequencyGroups.Num(); FrequencyGroupIndex++)
 		{
-			const FFrequencyGroup& FrequencyGroup = FrequencyGroups[FrequencyGroupIndex];
-			SET_DWORD_STAT_FName(AndroidFrequencyGroupMaxFreqStats[FrequencyGroupIndex], FrequencyGroup.MaxFrequency);
-			SET_DWORD_STAT_FName(AndroidFrequencyGroupNumCoresStats[FrequencyGroupIndex], FrequencyGroup.CoreCount);
-			//SET_DWORD_STAT_FName(AndroidFrequencyGroupMinFreqStats[FrequencyGroupIndex], FrequencyGroup.MinFrequency);
-			SET_DWORD_STAT_FName(AndroidFrequencyGroupCurrentFreqStats[FrequencyGroupIndex], CurrentFrequencies[FrequencyGroupIndex]);
-			SET_FLOAT_STAT_FName(AndroidFrequencyGroupCurrentFreqPercentageStats[FrequencyGroupIndex], CurrentFrequenciesPercentage[FrequencyGroupIndex]);
-			SET_FLOAT_STAT_FName(AndroidFrequencyGroupMaxCoresUtilizationStats[FrequencyGroupIndex], MaxSingleCoreUtilization[FrequencyGroupIndex]);
-		}
+			UE::TUniqueLock Lock(CSVStatsLock);
+			UE::Android::FCPUStatValues& CSVStatValues = UE::Android::FCPUStatValues::Get();
 
-		static const FName CPUStatName = GET_STATFNAME(STAT_CPUTemp);
-		static const FName ThermalStatus = GET_STATFNAME(STAT_ThermalStatus);
-		SET_FLOAT_STAT_FName(CPUStatName, CPUTemp);
-		SET_DWORD_STAT_FName(ThermalStatus, GThermalStatus);
-#endif
+			CSVStatValues.NumFrequencyGroups = FrequencyGroups.Num();
+			for (int32 FrequencyGroupIndex = 0; FrequencyGroupIndex < FrequencyGroups.Num(); FrequencyGroupIndex++)
+			{
+				UE::Android::FCPUStatValues::FFrequencyGroupStats& FrequencyGroupStats = CSVStatValues.FrequencyGroupStats[FrequencyGroupIndex];
+				const FFrequencyGroup& FrequencyGroup = FrequencyGroups[FrequencyGroupIndex];
+				FrequencyGroupStats.CoreCount = FrequencyGroup.CoreCount;
+				FrequencyGroupStats.MaxFrequency = FrequencyGroup.MaxFrequency;
+				FrequencyGroupStats.CurrentFrequency = GetFrequencyGroupCurrentFrequency(FrequencyGroupIndex);
+				FrequencyGroupStats.MaxSingleCoreUtilization = MaxSingleCoreUtilization[FrequencyGroupIndex];
+			}
+		}
 
 		UpdateGPUStats();
 		GIsStatTaskActive = false;
