[PlatformFile]

- Adds EOpenReadFlags::NoBuffering to allow disabling OS caching layers when reading on platforms that support it
- Adds overload for OpenReadNoBuffering allowing for EOpenReadFlags to be provided giving more control to the caller for how files should be opened.
- Allows platform file handle Read() implementations to succeed if more bytes are requested than are available. This is necessary since some platforms require overreading by restricting read request sizes to be aligned to platform volume sector sizes
- Uniformly treats reading a file when the file cursor is already at the end of the file, as a read failure (Read() returns false)

#rnx

diff --git a/Engine/Source/Runtime/Core/Private/Android/AndroidPlatformFile.cpp b/Engine/Source/Runtime/Core/Private/Android/AndroidPlatformFile.cpp
--- a/Engine/Source/Runtime/Core/Private/Android/AndroidPlatformFile.cpp
+++ b/Engine/Source/Runtime/Core/Private/Android/AndroidPlatformFile.cpp
@@ -358,10 +358,15 @@ public:
 #endif
 		check(BytesToRead >= 0);
 		check(Destination);
-		
-		while (BytesToRead > 0)
+		bool bAtEndOfFile = CurrentOffset >= Length;
+		if (bAtEndOfFile)
 		{
+			// Fail if we tried reading when already at the end of the file.
+			return false;
+		}
 
+		do
+		{
 			int64 ThisSize = FMath::Min<int64>(READWRITE_SIZE, BytesToRead);
 			
 			ThisSize = __pread(File->Handle, Destination, ThisSize, CurrentOffset);
@@ -383,14 +388,16 @@ public:
 			else if (ThisSize == 0)
 			{
 				// 0 is EOF
+				bAtEndOfFile = true;
 				break;
 			}
 			CurrentOffset += ThisSize;
 			Destination += ThisSize;
 			BytesToRead -= ThisSize;
-		}
+			bAtEndOfFile = CurrentOffset >= Length;
+		} while (BytesToRead > 0 && !bAtEndOfFile);
 
-		return BytesToRead == 0;
+		return BytesToRead == 0 || bAtEndOfFile;
 	}
 
 	virtual bool ReadAt(uint8* Destination, int64 BytesToRead, int64 Offset) override
diff --git a/Engine/Source/Runtime/Core/Private/Apple/ApplePlatformFile.cpp b/Engine/Source/Runtime/Core/Private/Apple/ApplePlatformFile.cpp
--- a/Engine/Source/Runtime/Core/Private/Apple/ApplePlatformFile.cpp
+++ b/Engine/Source/Runtime/Core/Private/Apple/ApplePlatformFile.cpp
@@ -164,7 +164,7 @@ public:
 		{
 			// Handle virtual file handles (only in read mode, write mode doesn't use the file handle registry)
 			FFileHandleRegistryReadTracker TrackRead(GFileRegistry, *this, true);
-			if (!TrackRead.IsValid())
+			if (!TrackRead.IsValid() || FileOffset >= FileSize)
 			{
 				return false;
 			}
@@ -176,11 +176,18 @@ public:
 			}
 			int64 BytesRead = ReadInternal(Destination, BytesToRead);
 			FileOffset += BytesRead;
-			return BytesRead == BytesToRead;
+			return BytesRead == BytesToRead || FileOffset == FileSize;
 		}
 		else
 		{
-			return ReadInternal(Destination, BytesToRead) == BytesToRead;
+			int64 Pos = Tell();
+			int64 BytesRead = ReadInternal(Destination, BytesToRead);
+			if (BytesRead != BytesToRead)
+			{
+				// If we requested more than was in the file return true if we read all the remaining content
+				return BytesRead > 0 && (Pos + BytesRead) == Size();
+			}
+			return true;
 		}
 	}
 	
diff --git a/Engine/Source/Runtime/Core/Private/IOS/IOSPlatformFile.cpp b/Engine/Source/Runtime/Core/Private/IOS/IOSPlatformFile.cpp
--- a/Engine/Source/Runtime/Core/Private/IOS/IOSPlatformFile.cpp
+++ b/Engine/Source/Runtime/Core/Private/IOS/IOSPlatformFile.cpp
@@ -143,21 +143,24 @@ public:
 		FileHandle = -1;
 	}
 
-	bool InternalRead(uint8* Destination, int64 BytesToRead)
+	int64 InternalRead(uint8* Destination, int64 BytesToRead)
 	{
+		int64 TotalBytesRead = 0;
 		while (BytesToRead)
 		{
 			check(BytesToRead >= 0);
 			int64 ThisSize = FMath::Min<int64>(READWRITE_SIZE, BytesToRead);
 			check(Destination);
-			if (read(FileHandle, Destination, ThisSize) != ThisSize)
+			int64 BytesRead = read(FileHandle, Destination, ThisSize);
+			TotalBytesRead += BytesRead;
+			if (BytesRead != ThisSize)
 			{
-				return false;
+				return TotalBytesRead;
 			}
 			Destination += ThisSize;
 			BytesToRead -= ThisSize;
 		}
-		return true;
+		return TotalBytesRead;
 	}
 
 public:
@@ -169,20 +172,33 @@ public:
  #if MANAGE_FILE_HANDLES_IOS
        if( IsManaged() )
         {
+		   if (FileOffset == FileSize)
+		   {
+			   return false;
+		   }
+
             ActivateSlot();
 			lseek(FileHandle, FileOffset, SEEK_SET);
 			// read into the buffer, and make sure it worked
-			if (InternalRead(Destination, BytesToRead))
+			int64 BytesRead = InternalRead(Destination, BytesToRead);
+			FileOffset += BytesRead;
+			if (BytesRead != BytesToRead)
 			{
-				FileOffset += BytesToRead;
-				return true;
+				return FileOffset == FileSize;
 			}
-			return false;
+			return true;
         }
         else
 #endif
         {
-		    return InternalRead(Destination, BytesToRead);
+			int64 Pos = Tell();
+			int64 BytesRead = InternalRead(Destination, BytesToRead);
+			if (BytesRead != BytesToRead)
+			{
+				// If we requested more than was in the file return true if we read all the remaining content
+				return BytesRead > 0 && (Pos + BytesRead) == Size();
+			}
+			return true;
         }
 	}
 
diff --git a/Engine/Source/Runtime/Core/Private/Unix/UnixPlatformFile.cpp b/Engine/Source/Runtime/Core/Private/Unix/UnixPlatformFile.cpp
--- a/Engine/Source/Runtime/Core/Private/Unix/UnixPlatformFile.cpp
+++ b/Engine/Source/Runtime/Core/Private/Unix/UnixPlatformFile.cpp
@@ -199,7 +199,7 @@ public:
 		{
 			// Handle virtual file handles (only in read mode, write mode doesn't use the file handle registry)
 			FFileHandleRegistryReadTracker TrackRead(GFileRegistry, *this, true);
-			if (!TrackRead.IsValid())
+			if (!TrackRead.IsValid() || FileOffset >= FileSize)
 			{
 				return false;
 			}
@@ -213,12 +213,19 @@ public:
 			}
 			int64 BytesRead = ReadInternal(Destination, BytesToRead);
 			FileOffset += BytesRead;
-			return BytesRead == BytesToRead;
+			return BytesRead == BytesToRead || FileOffset == FileSize;
 		}
 		else // FileOffset is invalid in 'read/write' mode, i.e. not updated by Write(), Seek(), SeekFronEnd(). Read from the current location.
 		{
-			FScopedDiskUtilizationTracker Tracker(BytesToRead, Tell());
-			return ReadInternal(Destination, BytesToRead) == BytesToRead;
+			int64 Pos = Tell();
+			FScopedDiskUtilizationTracker Tracker(BytesToRead, Pos);
+			int64 BytesRead = ReadInternal(Destination, BytesToRead);
+			if (BytesRead != BytesToRead)
+			{
+				// If we requested more than was in the file return true if we read all the remaining content
+				return BytesRead > 0 && (Pos + BytesRead) == Size();
+			}
+			return true;
 		}
 	}
 
diff --git a/Engine/Source/Runtime/Core/Private/Windows/WindowsPlatformFile.cpp b/Engine/Source/Runtime/Core/Private/Windows/WindowsPlatformFile.cpp
--- a/Engine/Source/Runtime/Core/Private/Windows/WindowsPlatformFile.cpp
+++ b/Engine/Source/Runtime/Core/Private/Windows/WindowsPlatformFile.cpp
@@ -769,6 +769,12 @@ public:
 		TRACE_PLATFORMFILE_BEGIN_READ(&OverlappedIO, FileHandle, FilePos, BytesToRead);
 
 		int64 TotalNumRead = 0;
+		bool bAtEndOfFile = FilePos >= FileSize;
+		if (bAtEndOfFile)
+		{
+			return false;
+		}
+
 		do
 		{
 			uint32 BytesToRead32 = (uint32)FMath::Min<int64>(BytesToRead, FWindowsPlatformMisc::IsWine() ? int64(0x7ffff000) : int64(UINT32_MAX));
@@ -798,6 +804,7 @@ public:
 			Destination += NumRead;
 			// Update where we are in the file
 			FilePos += NumRead;
+			bAtEndOfFile = FilePos >= FileSize;
 			UE::WindowsPlatformFile::Private::UpdateOverlappedPos(OverlappedIO, FilePos);
 
 			if (NumRead == 0 && BytesToRead32 != 0)
@@ -811,7 +818,7 @@ public:
 				TRACE_PLATFORMFILE_END_READ(&OverlappedIO, TotalNumRead);
 				return false;
 			}
-		} while (BytesToRead > 0 && FilePos < FileSize);
+		} while (BytesToRead > 0 && !bAtEndOfFile);
 		TRACE_PLATFORMFILE_END_READ(&OverlappedIO, TotalNumRead);
 		return true;
 	}
@@ -1433,30 +1440,33 @@ IAsyncReadFileHandle* FWindowsPlatformFile::OpenAsyncRead(const TCHAR* Filename,
 
 FFileOpenResult FWindowsPlatformFile::OpenRead(const TCHAR* Filename, EOpenReadFlags Flags)
 {
-	const uint32  Access	= GENERIC_READ;
-	const uint32  WinFlags  =	FILE_SHARE_READ | GetFileShareDeleteFlagIfEnabled() |
-								(EnumHasAnyFlags(Flags, EOpenReadFlags::AllowWrite) ? FILE_SHARE_WRITE : 0) |
-								(EnumHasAnyFlags(Flags, EOpenReadFlags::AllowDelete) ? FILE_SHARE_DELETE : 0);
-	const uint32  Create	= OPEN_EXISTING;
+	const uint32 Access				= GENERIC_READ;
+	const uint32 ShareFlags			= FILE_SHARE_READ | GetFileShareDeleteFlagIfEnabled() |
+										(EnumHasAnyFlags(Flags, EOpenReadFlags::AllowWrite) ? FILE_SHARE_WRITE : 0) |
+										(EnumHasAnyFlags(Flags, EOpenReadFlags::AllowDelete) ? FILE_SHARE_DELETE : 0);
+	const uint32 FlagsAndAttributes = (!EnumHasAnyFlags(Flags, EOpenReadFlags::NoBuffering) ? FILE_ATTRIBUTE_NORMAL :
+										FILE_FLAG_NO_BUFFERING |	// Disables OS caching of file data
+										FILE_ATTRIBUTE_TEMPORARY);	// Disables OS caching of file metadata
+	const uint32 Create				= OPEN_EXISTING;
 
 #define USE_OVERLAPPED_IO (!IS_PROGRAM && !WITH_EDITOR)		// Use straightforward synchronous I/O in cooker/editor
 
 	TRACE_PLATFORMFILE_BEGIN_OPEN(Filename);
 #if USE_OVERLAPPED_IO
-	HANDLE Handle    = CreateFileW(*FNormalizedFilename(Filename), Access, WinFlags, NULL, Create, FILE_ATTRIBUTE_NORMAL | FILE_FLAG_OVERLAPPED, NULL);
+	HANDLE Handle    = CreateFileW(*FNormalizedFilename(Filename), Access, ShareFlags, NULL, Create, FlagsAndAttributes | FILE_FLAG_OVERLAPPED, NULL);
 	if (Handle != INVALID_HANDLE_VALUE)
 	{
 		TRACE_PLATFORMFILE_END_OPEN(Handle);
-		FAsyncBufferedFileReaderWindows* FileHandle = new FAsyncBufferedFileReaderWindows(Handle, Access, WinFlags, FILE_FLAG_OVERLAPPED);
+		FAsyncBufferedFileReaderWindows* FileHandle = new FAsyncBufferedFileReaderWindows(Handle, Access, ShareFlags, FlagsAndAttributes | FILE_FLAG_OVERLAPPED);
 		return MakeValue(FileHandle);
 	}
 #else
-	HANDLE Handle = CreateFileW(*FNormalizedFilename(Filename), Access, WinFlags, NULL, Create, FILE_ATTRIBUTE_NORMAL, NULL);
+	HANDLE Handle = CreateFileW(*FNormalizedFilename(Filename), Access, ShareFlags, NULL, Create, FlagsAndAttributes, NULL);
 	if (Handle != INVALID_HANDLE_VALUE)
 	{
 		TRACE_PLATFORMFILE_END_OPEN(Handle);
 
-		FFileHandleWindows* FileHandle = new FFileHandleWindows(Handle, Access, WinFlags, 0);
+		FFileHandleWindows* FileHandle = new FFileHandleWindows(Handle, Access, ShareFlags, FlagsAndAttributes);
 
 		// Some operations can fail during the handle initialization, so we
 		// double check that the handle is valid before returning it
@@ -1489,37 +1499,54 @@ IFileHandle* FWindowsPlatformFile::OpenRead(const TCHAR* Filename, bool bAllowWr
 	return Result.StealValue().Release();
 }
 
-IFileHandle* FWindowsPlatformFile::OpenReadNoBuffering(const TCHAR* Filename, bool bAllowWrite)
+FFileOpenResult FWindowsPlatformFile::OpenReadNoBuffering(const TCHAR* Filename, EOpenReadFlags Flags)
 {
-	uint32  Access = GENERIC_READ;
-	uint32  WinFlags = FILE_SHARE_READ | (bAllowWrite ? FILE_SHARE_WRITE : 0) | GetFileShareDeleteFlagIfEnabled();
-	uint32  Create = OPEN_EXISTING;
+	const uint32 Access				= GENERIC_READ;
+	const uint32 ShareFlags			= FILE_SHARE_READ | GetFileShareDeleteFlagIfEnabled() |
+										(EnumHasAnyFlags(Flags, EOpenReadFlags::AllowWrite) ? FILE_SHARE_WRITE : 0) |
+										(EnumHasAnyFlags(Flags, EOpenReadFlags::AllowDelete) ? FILE_SHARE_DELETE : 0);
+	const uint32 FlagsAndAttributes = FILE_FLAG_OVERLAPPED | 
+										(!EnumHasAnyFlags(Flags, EOpenReadFlags::NoBuffering) ? FILE_ATTRIBUTE_NORMAL :
+										FILE_FLAG_NO_BUFFERING |	// Disables OS caching of file data
+										FILE_ATTRIBUTE_TEMPORARY);	// Disables OS caching of file metadata
+	const uint32 Create				= OPEN_EXISTING;
+
 	TRACE_PLATFORMFILE_BEGIN_OPEN(Filename);
-	HANDLE Handle = CreateFileW(*FNormalizedFilename(Filename), Access, WinFlags, NULL, Create, FILE_ATTRIBUTE_NORMAL | FILE_FLAG_OVERLAPPED, NULL);
+	HANDLE Handle = CreateFileW(*FNormalizedFilename(Filename), Access, ShareFlags, NULL, Create, FlagsAndAttributes, NULL);
 	if (Handle != INVALID_HANDLE_VALUE)
 	{
 		TRACE_PLATFORMFILE_END_OPEN(Handle);
 
-		FFileHandleWindows* FileHandle = new FFileHandleWindows(Handle, Access, WinFlags, FILE_FLAG_OVERLAPPED);
+		FFileHandleWindows* FileHandle = new FFileHandleWindows(Handle, Access, ShareFlags, FlagsAndAttributes);
 
 		// Some operations can fail during the handle initialization, so we
 		// double check that the handle is valid before returning it
 		if (FileHandle->IsValid())
 		{
-			return FileHandle;
+			return MakeValue(FileHandle);
 		}
 		else
 		{
 			delete FileHandle;
-
-			return nullptr;
+			return MakeError(TEXT("Failed to create FFileHandleWindows"));
 		}
 	}
 	else
 	{
 		TRACE_PLATFORMFILE_FAIL_OPEN(Filename);
+		return MakeError(TEXT("Failed to open file for reading"), FPlatformMisc::GetLastError());
+	}
+}
+
+IFileHandle* FWindowsPlatformFile::OpenReadNoBuffering(const TCHAR* Filename, bool bAllowWrite)
+{
+	const EOpenReadFlags Flags = bAllowWrite ? EOpenReadFlags::AllowWrite : EOpenReadFlags::None;
+	FFileOpenResult Result = OpenReadNoBuffering(Filename, Flags);
+	if (Result.HasError())
+	{
 		return nullptr;
 	}
+	return Result.StealValue().Release();
 }
 
 IFileHandle* FWindowsPlatformFile::OpenWrite(const TCHAR* Filename, bool bAppend, bool bAllowRead)
diff --git a/Engine/Source/Runtime/Core/Private/Windows/WindowsPlatformFilePrivate.h b/Engine/Source/Runtime/Core/Private/Windows/WindowsPlatformFilePrivate.h
--- a/Engine/Source/Runtime/Core/Private/Windows/WindowsPlatformFilePrivate.h
+++ b/Engine/Source/Runtime/Core/Private/Windows/WindowsPlatformFilePrivate.h
@@ -66,6 +66,7 @@ public:
 
 	virtual FFileOpenResult OpenRead(const TCHAR* Filename, EOpenReadFlags Flags) override;
 	virtual IFileHandle* OpenRead(const TCHAR* Filename, bool bAllowWrite = false) override;
+	virtual FFileOpenResult OpenReadNoBuffering(const TCHAR* Filename, EOpenReadFlags Flags) override;
 	virtual IFileHandle* OpenReadNoBuffering(const TCHAR* Filename, bool bAllowWrite = false) override;
 	virtual IFileHandle* OpenWrite(const TCHAR* Filename, bool bAppend = false, bool bAllowRead = false) override;
 	virtual FOpenMappedResult OpenMappedEx(const TCHAR* Filename, EOpenReadFlags OpenOptions = EOpenReadFlags::None, int64 MaximumSize = 0) override;
diff --git a/Engine/Source/Runtime/Core/Public/GenericPlatform/GenericPlatformFile.h b/Engine/Source/Runtime/Core/Public/GenericPlatform/GenericPlatformFile.h
--- a/Engine/Source/Runtime/Core/Public/GenericPlatform/GenericPlatformFile.h
+++ b/Engine/Source/Runtime/Core/Public/GenericPlatform/GenericPlatformFile.h
@@ -141,7 +141,10 @@ public:
 	 * Read bytes from the file.
 	 * @param Destination	Buffer to holds the results, should be at least BytesToRead in size.
 	 * @param BytesToRead	Number of bytes to read into the destination.
-	 * @return				true if the operation completed successfully.
+	 * @return				true if the operation completed successfully. Returns false if a read is attempted
+	 *						while the read position is already at the end of the file.
+	 * @remark				If BytesToRead is larger than the bytes remaining in the file from the current read position
+	 *						the Read() will return true if the underlying file read could read all remaining bytes.
 	**/
 	virtual bool		Read(uint8* Destination, int64 BytesToRead) = 0;
 
@@ -495,7 +498,13 @@ public:
 		/** Allow other handles/processes to write to this file. This flag is needed to open files that are currently being written to as well */
 		AllowWrite	= 1 << 0,
 		/** Allow the file to be deleted or renamed while keeping this handle valid for reading. */
-		AllowDelete = 1 << 1
+		AllowDelete = 1 << 1,
+		/** 
+		* Disables buffering performed by the platform operating system, if any. Note, this typically requires read 
+		* offsets, and read sizes to be multiples of the volume's sector size, as well as the buffer being read into 
+		* being aligned to a multiple of the physical sector size. Avoid using if you cannot control these parameters.
+		*/
+		NoBuffering = 1 << 2
 	};
 
 	/**
@@ -519,6 +528,11 @@ public:
 	 */
 	virtual IFileHandle*	OpenRead(const TCHAR* Filename, bool bAllowWrite = false) = 0;
 
+	virtual FFileOpenResult OpenReadNoBuffering(const TCHAR* Filename, EOpenReadFlags Flags)
+	{
+		return OpenRead(Filename, Flags);
+	}
+
 	virtual IFileHandle* OpenReadNoBuffering(const TCHAR* Filename, bool bAllowWrite = false)
 	{
 		return OpenRead(Filename, bAllowWrite);
