[Core] Added conversion utilities between EAxis <-> EAxisList

diff --git a/Engine/Source/Runtime/Core/Public/Math/Axis.h b/Engine/Source/Runtime/Core/Public/Math/Axis.h
--- a/Engine/Source/Runtime/Core/Public/Math/Axis.h
+++ b/Engine/Source/Runtime/Core/Public/Math/Axis.h
@@ -2,7 +2,7 @@
 
 #pragma once
 
-#include "CoreTypes.h"
+#include "HAL/Platform.h"
 
 // Generic axis enum (mirrored for property use in Object.h)
 namespace EAxis
@@ -50,3 +50,72 @@ namespace EAxisList
 		LeftUpForward	= Left | Up | Forward,
 	};
 }
+
+namespace EAxis
+{
+	// Converts the given AxisList value to the corresponding Axis value. Assumes a single axis within the provided AxisList.
+	inline EAxis::Type FromAxisList(const EAxisList::Type InAxisList)
+	{
+		PRAGMA_DISABLE_SWITCH_UNHANDLED_ENUM_CASE_WARNINGS;
+		switch (InAxisList)
+		{
+		case EAxisList::X:
+		case EAxisList::Forward:
+			return EAxis::X;
+
+		case EAxisList::Y:
+		case EAxisList::Left:
+			return EAxis::Y;
+
+		case EAxisList::Z:
+		case EAxisList::Up:
+			return EAxis::Z;
+
+		case EAxisList::None:
+		default:
+			return EAxis::None;
+		}
+		PRAGMA_RESTORE_SWITCH_UNHANDLED_ENUM_CASE_WARNINGS;
+	}
+}
+
+namespace EAxisList
+{
+	// Converts the given Axis to the corresponding AxisList value
+	inline EAxisList::Type FromAxis(const EAxis::Type InAxis, const EAxisList::Type InAxisCoordinateSystem = EAxisList::XYZ)
+	{
+		if (InAxisCoordinateSystem == EAxisList::XYZ)
+		{
+			switch (InAxis)
+			{
+			case EAxis::X:
+				return EAxisList::X;
+			case EAxis::Y:
+				return EAxisList::Y;
+			case EAxis::Z:
+				return EAxisList::Z;
+			case EAxis::None:
+			default:
+				return EAxisList::None;
+			}
+		}
+
+		if (InAxisCoordinateSystem == EAxisList::LeftUpForward)
+		{
+			switch (InAxis)
+			{
+			case EAxis::X:
+				return EAxisList::Forward;
+			case EAxis::Y:
+				return EAxisList::Left;
+			case EAxis::Z:
+				return EAxisList::Up;
+			case EAxis::None:
+			default:
+				return EAxisList::None;
+			}
+		}
+
+		return EAxisList::None;
+	}
+}
