Fix for AndroidPlatfomFile breaking virtual file / obb reading on Android

diff --git a/Engine/Source/Runtime/Core/Private/Android/AndroidPlatformFile.cpp b/Engine/Source/Runtime/Core/Private/Android/AndroidPlatformFile.cpp
--- a/Engine/Source/Runtime/Core/Private/Android/AndroidPlatformFile.cpp
+++ b/Engine/Source/Runtime/Core/Private/Android/AndroidPlatformFile.cpp
@@ -358,7 +358,7 @@ public:
 #endif
 		check(BytesToRead >= 0);
 		check(Destination);
-		bool bAtEndOfFile = CurrentOffset >= Length;
+		bool bAtEndOfFile = CurrentOffset >= Start + Length;
 		if (bAtEndOfFile)
 		{
 			// Fail if we tried reading when already at the end of the file.
@@ -394,7 +394,7 @@ public:
 			CurrentOffset += ThisSize;
 			Destination += ThisSize;
 			BytesToRead -= ThisSize;
-			bAtEndOfFile = CurrentOffset >= Length;
+			bAtEndOfFile = CurrentOffset >= Start + Length;
 		} while (BytesToRead > 0 && !bAtEndOfFile);
 
 		return BytesToRead == 0 || bAtEndOfFile;
