[Applied Usability] Creating new project with Project Location as "\ \" causes unstable project

diff --git a/Engine/Source/Runtime/Core/Private/Windows/WindowsPlatformMisc.cpp b/Engine/Source/Runtime/Core/Private/Windows/WindowsPlatformMisc.cpp
--- a/Engine/Source/Runtime/Core/Private/Windows/WindowsPlatformMisc.cpp
+++ b/Engine/Source/Runtime/Core/Private/Windows/WindowsPlatformMisc.cpp
@@ -2397,6 +2397,12 @@ bool FWindowsPlatformMisc::IsValidAbsolutePathFormat(const FString& Path)
 		bIsValid = false;
 	}
 
+	// Avoid folders with only a space in their name e.g. "D:\TEMP\ " or "D:\TEMP\ \MyFolder"
+	if (bIsValid && Path.Find(TEXT("\\ \\")) != INDEX_NONE || Path.EndsWith(TEXT("\\ ")))
+	{
+		bIsValid = false;
+	}
+
 	// ensure there's no further instances of ':' in the string
 	if ( bIsValid && !(Path.Find(TEXT(":"), ESearchCase::IgnoreCase, ESearchDir::FromStart, 2)==INDEX_NONE) )
 	{
