- Made a new property saving path that can save an array property without needing to save the entire section/class (which was "baking in" property values into, say, DefaultEngine.ini, which would then cause changes to BaseEngine.ini to be stomped on).
- This is gated on the new cvar "ini.UseNewPropertySaving" when saving project settings
- Added ability to directly set properties in a particular static layer.
- Added/changed some APIs to support this:
UObject::SaveConfig now has a version that takes an FSaveConfigContext struct that can be used to save a single property to a specific layer

FConfigCacheIni::SetInSection
FConfigCacheIni::SetInSectionOfStaticLayer
FConfigCacheIni::AddToSectionOfStaticLayer
FConfigCacheIni::AddUniqueToSectionOfStaticLayer
FConfigCacheIni::RemoveKeyFromSectionOfStaticLayer
FConfigCacheIni::RemoveFromSectionOfStaticLayer
FConfigCacheIni::ResetKeyInSectionOfStaticLayer
FConfigCacheIni::SetArrayInSectionOfStaticLayer

FConfigBranch::ReapplyLayers
FConfigBranch::MergeStaticLayersUpTo
FConfigBranch::GetStaticLayer

diff --git a/Engine/Source/Developer/SettingsEditor/Private/Widgets/SSettingsEditor.cpp b/Engine/Source/Developer/SettingsEditor/Private/Widgets/SSettingsEditor.cpp
--- a/Engine/Source/Developer/SettingsEditor/Private/Widgets/SSettingsEditor.cpp
+++ b/Engine/Source/Developer/SettingsEditor/Private/Widgets/SSettingsEditor.cpp
@@ -204,18 +204,25 @@ void SSettingsEditor::NotifyPostChange( const FPropertyChangedEvent& PropertyCha
 					RecordPreferenceChangedAnalytics(Section, PropertyChangedEvent);
 				}
 
-				// Determine if the Property is an Array or Array Element
-				bool bIsArrayOrArrayElement = PropertyThatChanged->GetActiveMemberNode()->GetValue()->IsA(FArrayProperty::StaticClass()) 
-					|| PropertyThatChanged->GetActiveMemberNode()->GetValue()->ArrayDim > 1
-					|| PropertyChangedEvent.Property->GetOwner<FArrayProperty>();
-
-				bool bIsSetOrSetElement = PropertyThatChanged->GetActiveMemberNode()->GetValue()->IsA(FSetProperty::StaticClass())
-					|| PropertyChangedEvent.Property->GetOwner<FSetProperty>();
+				static IConsoleVariable* NewPropertySaving = IConsoleManager::Get().FindConsoleVariable(TEXT("ini.UseNewPropertySaving"));
+				bool bUseFullSectionSave = false;
+				if (NewPropertySaving->GetInt() == 0)
+				{
+					// Determine if the Property is an Array or Array Element
+					bool bIsArrayOrArrayElement = PropertyThatChanged->GetActiveMemberNode()->GetValue()->IsA(FArrayProperty::StaticClass()) 
+						|| PropertyThatChanged->GetActiveMemberNode()->GetValue()->ArrayDim > 1
+						|| PropertyChangedEvent.Property->GetOwner<FArrayProperty>();
+	
+					bool bIsSetOrSetElement = PropertyThatChanged->GetActiveMemberNode()->GetValue()->IsA(FSetProperty::StaticClass())
+						|| PropertyChangedEvent.Property->GetOwner<FSetProperty>();
+	
+					bool bIsMapOrMapElement = PropertyThatChanged->GetActiveMemberNode()->GetValue()->IsA(FMapProperty::StaticClass())
+						|| PropertyChangedEvent.Property->GetOwner<FMapProperty>();
 
-				bool bIsMapOrMapElement = PropertyThatChanged->GetActiveMemberNode()->GetValue()->IsA(FMapProperty::StaticClass())
-					|| PropertyChangedEvent.Property->GetOwner<FMapProperty>();
+					bUseFullSectionSave = bIsArrayOrArrayElement || bIsSetOrSetElement || bIsMapOrMapElement;
+				}
 
-				if (ObjectBeingEdited->GetClass()->HasAnyClassFlags(CLASS_DefaultConfig) && !bIsArrayOrArrayElement && !bIsSetOrSetElement && !bIsMapOrMapElement)
+				if (ObjectBeingEdited->GetClass()->HasAnyClassFlags(CLASS_DefaultConfig) && !bUseFullSectionSave)
 				{
 					if (!Section.IsValid() || Section->NotifySectionOnPropertyModified())
 					{
diff --git a/Engine/Source/Runtime/Core/Private/Misc/ConfigCacheIni.cpp b/Engine/Source/Runtime/Core/Private/Misc/ConfigCacheIni.cpp
--- a/Engine/Source/Runtime/Core/Private/Misc/ConfigCacheIni.cpp
+++ b/Engine/Source/Runtime/Core/Private/Misc/ConfigCacheIni.cpp
@@ -309,6 +309,440 @@
 }
 #endif // UE_BUILD_SHIPPING
 
+
+
+
+/**
+ * Check if an ini file exists, allowing a delegate to determine if it will handle loading it
+ */
+bool DoesConfigFileExistWrapper(const TCHAR* IniFile, const TSet<FString>* IniCacheSet, const TSet<FString>* PrimaryConfigFileCache, const TSet<FString>* SecondaryConfigFileCache)
+{
+	// will any delegates return contents via TSPreLoadConfigFileDelegate()?
+	int32 ResponderCount = 0;
+	FCoreDelegates::TSCountPreLoadConfigFileRespondersDelegate().Broadcast(IniFile, ResponderCount);
+
+	if (ResponderCount > 0)
+	{
+		return true;
+	}
+
+	FString IniFileString = IniFile;
+	if (OverrideFileFromCommandline(IniFileString))
+	{
+		return true;
+	}
+
+	// check staged cache (likely for plugin configs)
+	if (PrimaryConfigFileCache != nullptr || SecondaryConfigFileCache != nullptr)
+	{
+		return (PrimaryConfigFileCache && PrimaryConfigFileCache->Contains(IniFileString)) ||
+			(SecondaryConfigFileCache && SecondaryConfigFileCache->Contains(IniFileString));
+	}
+
+	// So far, testing on cooked consoles, cooked desktop, and the Editor
+	// works fine.
+	// However, there was an issue where INIs wouldn't be found during cooking,
+	// which would pass by silently. Realistically, in most cases this would never
+	// have caused an issue, but using FPlatformProperties::RequiresCookedData
+	// to prevent using the cache in that case ensures full consistency.
+	if (IniCacheSet && FPlatformProperties::RequiresCookedData())
+	{
+		const bool bFileExistsCached = IniCacheSet->Contains(IniFileString);
+
+		// This code can be uncommented if we expect there are INIs that are not being
+		// found in the cache.
+		/**
+		const bool bFileExistsCachedTest = IFileManager::Get().FileSize(IniFile) >= 0;
+		ensureMsgf(
+			bFileExistsCached == bFileExistsCachedTest,
+			TEXT("DoesConfigFileExistWrapper: InCache = %d, InFileSystem = %d, Name = %s, Configs = \n%s"),
+			!!bFileExistsCached,
+			!!bFileExistsCachedTest,
+			IniFile,
+			*FString::Join(*IniCacheSet, TEXT("\n"))
+		);
+		*/
+
+		return bFileExistsCached;
+	}
+
+	// otherwise just look for the normal file to exist
+	UE_LOG(LogConfig, VeryVerbose, TEXT("Looking for a config file without a staged cache for %s"), IniFile);
+	const bool bFileExistsCached = IFileManager::Get().FileExists(IniFile);
+	return bFileExistsCached;
+}
+
+/**
+ * Load ini file, but allowing a delegate to handle the loading instead of the standard file load
+ */
+static bool LoadConfigFileWrapper(const TCHAR* IniFile, FString& Contents, bool bIsOverride = false)
+{
+	// We read the Base.ini and PluginBase.ini files many many times, so cache them
+	static FString BaseIniContents;
+	static FString PluginBaseIniContents;
+
+	const TCHAR* LastSlash = FCString::Strrchr(IniFile, '/');
+	if (LastSlash == nullptr)
+	{
+		LastSlash = FCString::Strrchr(IniFile, '\\');
+	}
+
+	bool bIsBaseIni = LastSlash != nullptr && FCString::Stricmp(LastSlash + 1, TEXT("Base.ini")) == 0;
+	if (bIsBaseIni && BaseIniContents.Len() > 0)
+	{
+		Contents = BaseIniContents;
+		return true;
+	}
+
+	bool bIsPluginBaseIni = LastSlash != nullptr && FCString::Stricmp(LastSlash + 1, TEXT("PluginBase.ini")) == 0;
+	if (bIsPluginBaseIni && PluginBaseIniContents.Len() > 0)
+	{
+		Contents = PluginBaseIniContents;
+		return true;
+	}
+
+	// let other systems load the file instead of the standard load below
+	FCoreDelegates::TSPreLoadConfigFileDelegate().Broadcast(IniFile, Contents);
+
+	// if this loaded any text, we are done, and we won't override the contents with standard ini file data
+	if (Contents.Len())
+	{
+		return true;
+	}
+
+#if ALLOW_INI_OVERRIDE_FROM_COMMANDLINE
+	if (bIsOverride)
+	{
+		// Make sure we bypass the pak layer because if our override is likely under root the pak layer will
+		// just resolve it even if it's an absolute path.
+
+		return FFileHelper::LoadFileToString(Contents, &IPlatformFile::GetPlatformPhysical(), IniFile);
+	}
+#endif
+
+	// note: we don't check if FileOperations are disabled because downloadable content calls this directly (which
+	// needs file ops), and the other caller of this is already checking for disabled file ops
+	// and don't read from the file, if the delegate got anything loaded
+	bool bResult = FFileHelper::LoadFileToString(Contents, IniFile);
+	if (bResult)
+	{
+		if (bIsBaseIni)
+		{
+			BaseIniContents = Contents;
+		}
+		else if (bIsPluginBaseIni)
+		{
+			PluginBaseIniContents = Contents;
+		}
+	}
+	return bResult;
+}
+
+/**
+ * Save an ini file, with delegates also saving the file (its safe to allow both to happen, even tho loading doesn't behave this way)
+ */
+static bool SaveConfigFileWrapper(const TCHAR* IniFile, const FString& Contents)
+{
+	TRACE_CPUPROFILER_EVENT_SCOPE(SaveConfigFileWrapper);
+
+	// let anyone that needs to save it, do so (counting how many did)
+	int32 SavedCount = 0;
+	FCoreDelegates::TSPreSaveConfigFileDelegate().Broadcast(IniFile, Contents, SavedCount);
+
+	// save it even if a delegate did as well
+	bool bLocalWriteSucceeded = false;
+
+	if (FConfigFile::WriteTempFileThenMove())
+	{
+		const FString BaseFilename = FPaths::GetBaseFilename(IniFile);
+		const FString TempFilename = FPaths::CreateTempFilename(*FPaths::ProjectSavedDir(), *BaseFilename.Left(32));
+		bLocalWriteSucceeded = FFileHelper::SaveStringToFile(Contents, *TempFilename, FFileHelper::EEncodingOptions::ForceUTF8WithoutBOM);
+		if (bLocalWriteSucceeded)
+		{
+			if (!IFileManager::Get().Move(IniFile, *TempFilename))
+			{
+				IFileManager::Get().Delete(*TempFilename);
+				bLocalWriteSucceeded = false;
+			}
+		}
+	}
+	else
+	{
+		bLocalWriteSucceeded = FFileHelper::SaveStringToFile(Contents, IniFile, FFileHelper::EEncodingOptions::ForceUTF8WithoutBOM);
+	}
+
+	// success is based on a delegate or file write working (or both)
+	return SavedCount > 0 || bLocalWriteSucceeded;
+}
+
+static bool DeleteConfigFileWrapper(const TCHAR* IniFile)
+{
+	TRACE_CPUPROFILER_EVENT_SCOPE(DeleteConfigFileWrapper);
+
+	bool bDeleted = false;
+	FCoreDelegates::TSPreDeleteConfigFileDelegate().Broadcast(IniFile, bDeleted);
+
+	bDeleted |= IFileManager::Get().Delete(IniFile);
+	return bDeleted;
+}
+
+
+
+/**
+ * Functionality to assist with updating a config file with one property value change.
+ */
+class FSinglePropertyConfigHelper
+{
+public:
+	/**
+	 * We need certain information for the helper to be useful.
+	 *
+	 * @Param InIniFilename - The disk location of the file we wish to edit.
+	 * @Param InSectionName - The section the property belongs to.
+	 * @Param InPropertyName - The name of the property that has been edited.
+	 * @Param InPropertyValue - The new value of the property that has been edited, or unset to remove the property.
+	 */
+	FSinglePropertyConfigHelper(const FString& InIniFilename, const FString& InSectionName, const FString& InPropertyName, const TOptional<FString>& InPropertyValue)
+		: IniFilename(InIniFilename)
+		, SectionName(InSectionName)
+		, PropertyName(InPropertyName)
+		, PropertyValue(InPropertyValue)
+	{
+		// Split the file into the necessary parts.
+		PopulateFileContentHelper();
+	}
+
+	/**
+	 * Setup this helper for writing a property to a static layer section (see ModifyAndWriteLayer)
+	 *
+	 * @Param InIniFilename - The disk location of the file we wish to edit.
+	 * @Param InSectionName - The section the property belongs to.
+	 * @Param InPropertyName - The name of the property that has been edited.
+	 * @Param InLayerSection - The section of a static layer that contains the value(s) to write out (needed to write arrays correctly)
+	 */
+	FSinglePropertyConfigHelper(const FString& InIniFilename, const FString& InSectionName, const FString& InPropertyName, FConfigCommandStreamSection* InLayerSection)
+		: IniFilename(InIniFilename)
+		, SectionName(InSectionName)
+		, PropertyName(InPropertyName)
+		, bUpdateLayer(true)
+		, LayerSection(InLayerSection)
+	{
+		// Split the file into the necessary parts.
+		PopulateFileContentHelper();
+	}
+
+	/**
+	 * Perform the action of updating the config file with the new property value.
+	 */
+	bool UpdateConfigFile()
+	{
+		UpdatePropertyInSection();
+		// Rebuild the file with the updated section.
+
+		FString NewFile = IniFileMakeup.BeforeSection + IniFileMakeup.Section + IniFileMakeup.AfterSection;
+		if (!NewFile.EndsWith(TEXT(LINE_TERMINATOR_ANSI LINE_TERMINATOR_ANSI)))
+		{
+			NewFile.AppendChars(LINE_TERMINATOR, TCString<TCHAR>::Strlen(LINE_TERMINATOR));
+		}
+		return SaveConfigFileWrapper(*IniFilename, NewFile);
+	}
+
+
+private:
+
+	/**
+	 * Clear any trailing whitespace from the end of the output.
+	 */
+	void ClearTrailingWhitespace(FString& InStr)
+	{
+		const FString Endl(LINE_TERMINATOR);
+		while (InStr.EndsWith(LINE_TERMINATOR, ESearchCase::CaseSensitive))
+		{
+			InStr.LeftChopInline(Endl.Len(), EAllowShrinking::No);
+		}
+	}
+
+	/**
+	 * Update the section with the new value for the property.
+	 */
+	void UpdatePropertyInSection()
+	{
+		FString UpdatedSection;
+		if (IniFileMakeup.Section.IsEmpty())
+		{
+			if (PropertyValue.IsSet() || bUpdateLayer)
+			{
+				const FString DecoratedSectionName = FString::Printf(TEXT("[%s]"), *SectionName);
+
+				ClearTrailingWhitespace(IniFileMakeup.BeforeSection);
+				UpdatedSection += LINE_TERMINATOR;
+				UpdatedSection += LINE_TERMINATOR;
+				UpdatedSection += DecoratedSectionName;
+				AppendPropertyLine(UpdatedSection);
+			}
+		}
+		else
+		{
+			FString SectionLine;
+			const TCHAR* Ptr = *IniFileMakeup.Section;
+			bool bUpdatedPropertyOnPass = false;
+			while (Ptr != nullptr && FParse::Line(&Ptr, SectionLine, true))
+			{
+				FString KeyAndEqual = PropertyName + TEXT("=");
+				if (SectionLine.Len() > 0)
+				{
+					if (SectionLine.StartsWith(KeyAndEqual) ||
+					(bUpdateLayer && CommandLookup.Contains(SectionLine[0]) && SectionLine.Mid(1).StartsWith(KeyAndEqual)))
+					{
+						if (!bUpdatedPropertyOnPass)
+						{
+							if (PropertyValue.IsSet() || bUpdateLayer)
+							{
+								AppendPropertyLine(UpdatedSection);
+							}
+						}
+						bUpdatedPropertyOnPass = true;
+					}
+					else
+					{
+						UpdatedSection += SectionLine;
+						UpdatedSection += LINE_TERMINATOR;
+					}
+				}
+				else
+				{
+					UpdatedSection += LINE_TERMINATOR;
+				}
+			}
+
+			// If the property wasnt found in the text of the existing section content,
+			// append it to the end of the section.
+			if (!bUpdatedPropertyOnPass && (bUpdateLayer || PropertyValue.IsSet()))
+			{
+				AppendPropertyLine(UpdatedSection);
+			}
+		}
+
+		UpdatedSection += LINE_TERMINATOR;
+		IniFileMakeup.Section = UpdatedSection;
+
+	}
+
+	/**
+	 * Split the file up into parts:
+	 * -> Before the section we wish to edit, which will remain unaltered,
+	 * ->-> The section we wish to edit, we only seek to edit the single property,
+	 * ->->-> After the section we wish to edit, which will remain unaltered.
+	 */
+	void PopulateFileContentHelper()
+	{
+		FString UnprocessedFileContents;
+		if (LoadConfigFileWrapper(*IniFilename, UnprocessedFileContents))
+		{
+			// Find the section in the file text.
+			const FString DecoratedSectionName = FString::Printf(TEXT("[%s]"), *SectionName);
+
+			const int32 DecoratedSectionNameStartIndex = UnprocessedFileContents.Find(DecoratedSectionName);
+			if (DecoratedSectionNameStartIndex != INDEX_NONE)
+			{
+				// If we found the section, cache off the file text before the section.
+				IniFileMakeup.BeforeSection = UnprocessedFileContents.Left(DecoratedSectionNameStartIndex);
+				UnprocessedFileContents.RemoveAt(0, IniFileMakeup.BeforeSection.Len());
+
+				// For the rest of the file, split it into the section we are editing and the rest of the file after.
+				const TCHAR* Ptr = UnprocessedFileContents.Len() > 0 ? *UnprocessedFileContents : nullptr;
+				FString NextUnprocessedLine;
+				bool bReachedNextSection = false;
+				while (Ptr != nullptr && FParse::Line(&Ptr, NextUnprocessedLine, true))
+				{
+					bReachedNextSection |= (NextUnprocessedLine.StartsWith(TEXT("[")) && NextUnprocessedLine != DecoratedSectionName);
+					if (bReachedNextSection)
+					{
+						IniFileMakeup.AfterSection += NextUnprocessedLine;
+						IniFileMakeup.AfterSection += LINE_TERMINATOR;
+					}
+					else
+					{
+						IniFileMakeup.Section += NextUnprocessedLine;
+						IniFileMakeup.Section += LINE_TERMINATOR;
+					}
+				}
+			}
+			else
+			{
+				IniFileMakeup.BeforeSection = UnprocessedFileContents;
+			}
+		}
+	}
+
+	/**
+	 * Append the property entry to the section
+	 */
+	void AppendPropertyLine(FString& PreText)
+	{
+		check(PropertyValue.IsSet() || bUpdateLayer);
+
+		// Make sure we dont leave much whitespace, and append the property name/value entry
+		ClearTrailingWhitespace(PreText);
+		PreText += LINE_TERMINATOR;
+
+		if (bUpdateLayer)
+		{
+			// get every value matching the key, as an array will have multiple, or we could have a -Foo=1, +Foo=2
+			TArray<FConfigValue*> Values;
+			LayerSection->MultiFindPointer(*PropertyName, Values, true);
+			for (const FConfigValue* Value : Values)
+			{
+				// export each one with a command/value
+				if (Value->ValueType != FConfigValue::EValueType::Set)
+				{
+					const TCHAR* Cmd = CommandLookup.FindKey(Value->ValueType);
+					PreText.AppendChar(*Cmd);
+				}
+				FConfigFile::AppendExportedPropertyLine(PreText, PropertyName, Value->GetSavedValueForWriting());
+			}
+		}
+		else
+		{
+			PreText += FConfigFile::GenerateExportedPropertyLine(PropertyName, PropertyValue.GetValue());
+		}
+	}
+
+
+private:
+	// The disk location of the ini file we seek to edit
+	FString IniFilename;
+
+	// The section in the config file
+	FString SectionName;
+
+	// The name of the property that has been changed
+	FString PropertyName;
+
+	// Setting when updating value(s) in a layer FConfigCommandStream
+	bool bUpdateLayer = false;
+	FConfigCommandStreamSection* LayerSection;
+
+	// The new value, in string format, of the property that has been changed
+	// This will be unset if the property has been removed from the config
+	TOptional<FString> PropertyValue;
+
+	// Helper struct that holds the makeup of the ini file.
+	struct IniFileContent
+	{
+		// The section we wish to edit
+		FString Section;
+
+		// The file contents before the section we are editing
+		FString BeforeSection;
+
+		// The file contents after the section we are editing
+		FString AfterSection;
+	} IniFileMakeup; // Instance of the helper to maintain file structure.
+};
+
+
+
 /*-----------------------------------------------------------------------------
 	FConfigSection
 -----------------------------------------------------------------------------*/
@@ -603,179 +1037,5 @@
 }
 
 
-/**
- * Check if an ini file exists, allowing a delegate to determine if it will handle loading it
- */
-/*static*/ bool DoesConfigFileExistWrapper(const TCHAR* IniFile, const TSet<FString>* IniCacheSet, const TSet<FString>* PrimaryConfigFileCache, const TSet<FString>* SecondaryConfigFileCache)
-{
-	// will any delegates return contents via TSPreLoadConfigFileDelegate()?
-	int32 ResponderCount = 0;
-	FCoreDelegates::TSCountPreLoadConfigFileRespondersDelegate().Broadcast(IniFile, ResponderCount);
-
-	if (ResponderCount > 0)
-	{
-		return true;
-	}
-
-	FString IniFileString = IniFile;
-	if (OverrideFileFromCommandline(IniFileString))
-	{
-		return true;
-	}
-
-	// check staged cache (likely for plugin configs)
-	if (PrimaryConfigFileCache != nullptr || SecondaryConfigFileCache != nullptr)
-	{
-		return (PrimaryConfigFileCache && PrimaryConfigFileCache->Contains(IniFileString)) ||
-			(SecondaryConfigFileCache && SecondaryConfigFileCache->Contains(IniFileString));
-	}
-
-	// So far, testing on cooked consoles, cooked desktop, and the Editor
-	// works fine.
-	// However, there was an issue where INIs wouldn't be found during cooking,
-	// which would pass by silently. Realistically, in most cases this would never
-	// have caused an issue, but using FPlatformProperties::RequiresCookedData
-	// to prevent using the cache in that case ensures full consistency.
-	if (IniCacheSet && FPlatformProperties::RequiresCookedData())
-	{ 
-		const bool bFileExistsCached = IniCacheSet->Contains(IniFileString);
-
-		// This code can be uncommented if we expect there are INIs that are not being
-		// found in the cache.
-		/**
-		const bool bFileExistsCachedTest = IFileManager::Get().FileSize(IniFile) >= 0;
-		ensureMsgf(
-			bFileExistsCached == bFileExistsCachedTest,
-			TEXT("DoesConfigFileExistWrapper: InCache = %d, InFileSystem = %d, Name = %s, Configs = \n%s"),
-			!!bFileExistsCached,
-			!!bFileExistsCachedTest,
-			IniFile,
-			*FString::Join(*IniCacheSet, TEXT("\n"))
-		);
-		*/
-		
-		return bFileExistsCached;
-	}
-
-	// otherwise just look for the normal file to exist
-	UE_LOG(LogConfig, VeryVerbose, TEXT("Looking for a config file without a staged cache for %s"), IniFile);
-	const bool bFileExistsCached = IFileManager::Get().FileExists(IniFile);
-	return bFileExistsCached;
-}
-
-/**
- * Load ini file, but allowing a delegate to handle the loading instead of the standard file load
- */
-static bool LoadConfigFileWrapper(const TCHAR* IniFile, FString& Contents, bool bIsOverride = false)
-{
-	// We read the Base.ini and PluginBase.ini files many many times, so cache them
-	static FString BaseIniContents;
-	static FString PluginBaseIniContents;
-
-	const TCHAR* LastSlash = FCString::Strrchr(IniFile, '/');
-	if (LastSlash == nullptr)
-	{
-		LastSlash = FCString::Strrchr(IniFile, '\\');
-	}
-
-	bool bIsBaseIni = LastSlash != nullptr && FCString::Stricmp(LastSlash + 1, TEXT("Base.ini")) == 0;
-	if (bIsBaseIni && BaseIniContents.Len() > 0)
-	{
-		Contents = BaseIniContents;
-		return true;
-	}
-
-	bool bIsPluginBaseIni = LastSlash != nullptr && FCString::Stricmp(LastSlash + 1, TEXT("PluginBase.ini")) == 0;
-	if (bIsPluginBaseIni && PluginBaseIniContents.Len() > 0)
-	{
-		Contents = PluginBaseIniContents;
-		return true;
-	}
-
-	// let other systems load the file instead of the standard load below
-	FCoreDelegates::TSPreLoadConfigFileDelegate().Broadcast(IniFile, Contents);
-
-	// if this loaded any text, we are done, and we won't override the contents with standard ini file data
-	if (Contents.Len())
-	{
-		return true;
-	}
-
-#if ALLOW_INI_OVERRIDE_FROM_COMMANDLINE
-	if (bIsOverride)
-	{
-		// Make sure we bypass the pak layer because if our override is likely under root the pak layer will
-		// just resolve it even if it's an absolute path.
-
-		return FFileHelper::LoadFileToString(Contents, &IPlatformFile::GetPlatformPhysical(), IniFile);
-	}
-#endif
-
-	// note: we don't check if FileOperations are disabled because downloadable content calls this directly (which
-	// needs file ops), and the other caller of this is already checking for disabled file ops
-	// and don't read from the file, if the delegate got anything loaded
-	bool bResult = FFileHelper::LoadFileToString(Contents, IniFile);
-	if (bResult)
-	{
-		if (bIsBaseIni)
-		{
-			BaseIniContents = Contents;
-		}
-		else if (bIsPluginBaseIni)
-		{
-			PluginBaseIniContents = Contents;
-		}
-	}
-	return bResult;
-}
-
-/**
- * Save an ini file, with delegates also saving the file (its safe to allow both to happen, even tho loading doesn't behave this way)
- */
-static bool SaveConfigFileWrapper(const TCHAR* IniFile, const FString& Contents)
-{
-	TRACE_CPUPROFILER_EVENT_SCOPE(SaveConfigFileWrapper);
-
-	// let anyone that needs to save it, do so (counting how many did)
-	int32 SavedCount = 0;
-	FCoreDelegates::TSPreSaveConfigFileDelegate().Broadcast(IniFile, Contents, SavedCount);
-
-	// save it even if a delegate did as well
-	bool bLocalWriteSucceeded = false;
-
-	if (FConfigFile::WriteTempFileThenMove())
-	{
-		const FString BaseFilename = FPaths::GetBaseFilename(IniFile);
-		const FString TempFilename = FPaths::CreateTempFilename(*FPaths::ProjectSavedDir(), *BaseFilename.Left(32));
-		bLocalWriteSucceeded = FFileHelper::SaveStringToFile(Contents, *TempFilename, FFileHelper::EEncodingOptions::ForceUTF8WithoutBOM);
-		if (bLocalWriteSucceeded)
-		{
-			if (!IFileManager::Get().Move(IniFile, *TempFilename))
-			{
-				IFileManager::Get().Delete(*TempFilename);
-				bLocalWriteSucceeded = false;
-			}
-		}
-	}
-	else
-	{
-		bLocalWriteSucceeded = FFileHelper::SaveStringToFile(Contents, IniFile, FFileHelper::EEncodingOptions::ForceUTF8WithoutBOM);
-	}
-
-	// success is based on a delegate or file write working (or both)
-	return SavedCount > 0 || bLocalWriteSucceeded;
-}
-
-static bool DeleteConfigFileWrapper(const TCHAR* IniFile)
-{
-	TRACE_CPUPROFILER_EVENT_SCOPE(DeleteConfigFileWrapper);
-
-	bool bDeleted = false;
-	FCoreDelegates::TSPreDeleteConfigFileDelegate().Broadcast(IniFile, bDeleted);
-
-	bDeleted |= IFileManager::Get().Delete(IniFile);
-	return bDeleted;
-}
-
 static FConfigCommandStream CalculateDiff(const FConfigFile& First, const FConfigFile& Second, const FString& SingleSection=FString(), const FString& SingleProperty=FString())
 {
@@ -2754,6 +3014,23 @@
 	}
 }
 
+bool FConfigFile::SetInSection(const TCHAR* SectionName, FName Key, const FString& Value)
+{
+	FConfigSection* Section = FindOrAddSectionInternal(SectionName);
+	Section->Remove(Key);
+	Section->Add(Key, FConfigValue(Section, Key, Value));
+	Dirty = true;
+
+	if (ChangeTracker != nullptr)
+	{
+		FConfigCommandStreamSection* Sec = ChangeTracker->FindOrAddSectionInternal(SectionName);
+		Sec->Remove(Key);
+		Sec->Add(Key, FConfigValue(Value, FConfigValue::EValueType::Set));
+	}
+
+	return true;
+}
+
 bool FConfigFile::AddToSection(const TCHAR* SectionName, FName Key, const FString& Value)
 {
 	FConfigSection* Section = FindOrAddSectionInternal(SectionName);
@@ -3555,6 +3832,59 @@
 	UE_LOG(LogConfig, Log, TEXT("Branch '%s' had been unloaded. Reloading on-demand took %.2fms"), *IniName.ToString(), (FPlatformTime::Seconds() - StartTime) * 1000.0f);
 }
 
+void FConfigBranch::ReapplyLayers()
+{
+	// calculate a patch so we don't lose in-memory changes
+	FConfigCommandStream Patch = CalculateDiff(FinalCombinedLayers, InMemoryFile);
 
+	CombinedStaticLayers.Cleanup();
+	FinalCombinedLayers.Cleanup();
+	InMemoryFile.Cleanup();
+
+	for (TPair<FString, FConfigCommandStream>& Pair : StaticLayers)
+	{
+		CombinedStaticLayers.ApplyFile(&Pair.Value);
+	}
+
+	FinalCombinedLayers = CombinedStaticLayers;
+	for (DynamicLayerList::TIterator Node(DynamicLayers.GetHead()); Node; ++Node)
+	{
+		FinalCombinedLayers.ApplyFile(*Node);
+	}
+
+	FinalCombinedLayers.ApplyFile(&CommandLineOverrides);
+	InMemoryFile = FinalCombinedLayers;
+	InMemoryFile.ApplyFile(&Patch);
+}
+
+bool FConfigBranch::MergeStaticLayersUpTo(const FString& LayerSubstring, FConfigFile& OutFile) const
+{
+	for (const TPair<FString, FConfigCommandStream>& Pair : StaticLayers)
+	{
+		if (Pair.Key.Contains(LayerSubstring))
+		{
+			return true;
+		}
+		OutFile.ApplyFile(&Pair.Value);
+	}
+
+	return false;
+}
+
+const FConfigCommandStream* FConfigBranch::GetStaticLayer(const FString& LayerSubstring) const
+{
+	for (const TPair<FString, FConfigCommandStream>& Pair : StaticLayers)
+	{
+		if (Pair.Key.Contains(LayerSubstring))
+		{
+			return &Pair.Value;
+		}
+	}
+
+	return nullptr;
+}
+
+
+
 bool FConfigBranch::RemoveSection(const TCHAR* Section)
 {
@@ -5105,6 +5435,127 @@
 }
 
 
+
+static FConfigCommandStream* FindStaticLayer(FConfigCacheIni* ConfigSystem, const FString& BranchFilename, const FString& LayerKeySubstring, FString& OutFilename, FConfigBranch*& OutBranch)
+{
+	OutBranch = ConfigSystem->FindBranch(*BranchFilename, BranchFilename);
+	if (OutBranch)
+	{
+		for (TPair<FString, FConfigCommandStream>& Pair : OutBranch->StaticLayers)
+		{
+			if (Pair.Key.Contains(LayerKeySubstring))
+			{
+				OutFilename = Pair.Key;
+				return &Pair.Value;
+			}
+		}
+
+		// we didn't find an existing layer, so look to see if we have a matching entry in the hierarchy and insert an empty layer for it
+		for (TPair<int32, FString>& IniFile : OutBranch->Hierarchy)
+		{
+			if (IniFile.Value.Contains(LayerKeySubstring))
+			{
+				int KeyToInsert = IniFile.Key;
+				FString FilenameToInsert = IniFile.Value;
+				bool bInserted = false;
+
+				// now for the painful process of finding where to insert it into the static layers - and since we need it ordered
+				// we have to make an entire new map
+				TMap<FString, FConfigCommandStream> NewLayers;
+				FConfigCommandStream* NewLayer = nullptr;
+				for (TPair<FString, FConfigCommandStream>& Pair : OutBranch->StaticLayers)
+				{
+					// find the entry in the hierarchy that matches this layer filename
+					if (const int32* LayerKey = OutBranch->Hierarchy.FindKey(Pair.Key))
+					{
+						if (!bInserted && *LayerKey > KeyToInsert)
+						{
+							NewLayer = &NewLayers.Add(FilenameToInsert, FConfigCommandStream());
+							bInserted = true;
+						}
+						NewLayers.Add(Pair.Key, Pair.Value);
+					}
+				}
+
+				if (bInserted)
+				{
+					OutBranch->StaticLayers = NewLayers;
+					OutFilename = FilenameToInsert;
+					return NewLayer;
+				}
+			}
+		}
+	}
+	return nullptr;
+}
+
+/**
+ * A wrapper function that finds a layer by name substring and section, calls a passed-in lambda to actually modify the section, writes
+ * changes to disk, then recalculates the in-memory config values
+ */
+static bool ModifyAndWriteLayer(FConfigCacheIni* ConfigSystem, const FString& SectionName, FName Key, const FString& BranchName, const FString& LayerSubstring,
+	bool bWriteFile, TFunction<void(const FConfigBranch*, FConfigCommandStreamSection*)> Modify)
+{
+	// if reapplying ends being slow, we will need to make this an option
+	const bool bReapplyLayers = true;
+
+	FString FilenameToWrite;
+	FConfigBranch* Branch;
+	if (FConfigCommandStream* Layer = FindStaticLayer(ConfigSystem, BranchName, LayerSubstring, FilenameToWrite, Branch))
+	{
+		FConfigCommandStreamSection* Sec = Layer->FindOrAddSectionInternal(SectionName);
+		Modify(Branch, Sec);
+
+		if (bWriteFile)
+		{
+			FSinglePropertyConfigHelper Writer(FilenameToWrite, SectionName, Key.ToString(), Sec);
+			Writer.UpdateConfigFile();
+		}
+
+		if (bReapplyLayers)
+		{
+			Branch->ReapplyLayers();
+		}
+
+		return true;
+	}
+
+	return false;
+}
+
+
+bool FConfigCacheIni::SetInSection(const TCHAR* Section, FName Key, const FString& Value, const FString& Filename)
+{
+	if (FConfigFile* File = Find(*Filename))
+	{
+		return File->SetInSection(Section, Key, Value);
+	}
+	return false;
+}
+
+bool FConfigCacheIni::SetInSectionOfStaticLayer(const TCHAR* Section, FName Key, const FString& Value, const FString& BranchName, const FString& LayerSubstring, bool bWriteFile)
+{
+	return ModifyAndWriteLayer(this, Section, Key, BranchName, LayerSubstring, bWriteFile, 
+		[Value, Key](const FConfigBranch* Branch, FConfigCommandStreamSection* Sec)
+		{
+			// look for an existing value to replace (to maintain order)
+			TArray<FConfigValue*> ExistingValues;
+			Sec->MultiFindPointer(Key, ExistingValues);
+			if (ExistingValues.Num() == 1)
+			{
+				*ExistingValues[0] = FConfigValue(Value, FConfigValue::EValueType::Set);
+			}
+			// otherwise remove any that are there, and add to the end (we don't try to maintain
+			// order in the rare case there is an array and we are replacing with a single)
+			else
+			{
+				Sec->Remove(Key);
+				Sec->Add(Key, FConfigValue(Value, FConfigValue::EValueType::Set));
+			}
+		});
+}
+
+
 bool FConfigCacheIni::AddToSection(const TCHAR* Section, FName Key, const FString& Value, const FString& Filename)
 {
 	if (FConfigFile* File = Find(*Filename))
@@ -5114,6 +5565,15 @@
 	return false;
 }
 
+bool FConfigCacheIni::AddToSectionOfStaticLayer(const TCHAR* Section, FName Key, const FString& Value, const FString& BranchName, const FString& LayerSubstring, bool bWriteFile)
+{
+	return ModifyAndWriteLayer(this, Section, Key, BranchName, LayerSubstring, bWriteFile,
+		[Value, Key](const FConfigBranch* Branch, FConfigCommandStreamSection* Sec)
+		{
+			Sec->Add(Key, FConfigValue(Value, FConfigValue::EValueType::ArrayAdd));
+		});
+}
+
 bool FConfigCacheIni::AddUniqueToSection(const TCHAR* Section, FName Key, const FString& Value, const FString& Filename)
 {
 	if (FConfigFile* File = Find(*Filename))
@@ -5123,6 +5583,15 @@
 	return false;
 }
 
+bool FConfigCacheIni::AddUniqueToSectionOfStaticLayer(const TCHAR* Section, FName Key, const FString& Value, const FString& BranchName, const FString& LayerSubstring, bool bWriteFile)
+{
+	return ModifyAndWriteLayer(this, Section, Key, BranchName, LayerSubstring, bWriteFile,
+		[Value, Key](const FConfigBranch* Branch, FConfigCommandStreamSection* Sec)
+		{
+			Sec->Add(Key, FConfigValue(Value, FConfigValue::EValueType::ArrayAddUnique));
+		});
+}
+
 bool FConfigCacheIni::RemoveKeyFromSection(const TCHAR* Section, FName Key, const FString& Filename)
 {
 	if (FConfigFile* File = Find(*Filename))
@@ -5132,6 +5601,16 @@
 	return false;
 }
 
+bool FConfigCacheIni::RemoveKeyFromSectionOfStaticLayer(const TCHAR* Section, FName Key, const FString& BranchName, const FString& LayerSubstring, bool bWriteFile)
+{
+	return ModifyAndWriteLayer(this, Section, Key, BranchName, LayerSubstring, bWriteFile,
+		[Key](const FConfigBranch* Branch, FConfigCommandStreamSection* Sec)
+		{
+			Sec->Remove(Key);
+			Sec->Add(Key, FConfigValue(TEXT("__Clear__"), FConfigValue::EValueType::Clear));
+		});
+}
+
 bool FConfigCacheIni::RemoveFromSection(const TCHAR* Section, FName Key, const FString& Value, const FString& Filename)
 {
 	if (FConfigFile* File = Find(*Filename))
@@ -5141,6 +5620,16 @@
 	return false;
 }
 
+bool FConfigCacheIni::RemoveFromSectionOfStaticLayer(const TCHAR* Section, FName Key, const FString& Value, const FString& BranchName, const FString& LayerSubstring, bool bWriteFile)
+{
+	return ModifyAndWriteLayer(this, Section, Key, BranchName, LayerSubstring, bWriteFile,
+		[Key, Value](const FConfigBranch* Branch, FConfigCommandStreamSection* Sec)
+		{
+			Sec->Remove(Key);
+			Sec->Add(Key, FConfigValue(Value, FConfigValue::EValueType::Remove));
+		});
+}
+
 bool FConfigCacheIni::ResetKeyInSection(const TCHAR* Section, FName Key, const FString& Filename)
 {
 	if (FConfigFile* File = Find(*Filename))
@@ -5150,6 +5639,77 @@
 	return false;
 }
 
+bool FConfigCacheIni::ResetKeyInSectionOfStaticLayer(const TCHAR* Section, FName Key, const FString& BranchName, const FString& LayerSubstring, bool bWriteFile)
+{
+	return ModifyAndWriteLayer(this, Section, Key, BranchName, LayerSubstring, bWriteFile,
+		[Key](const FConfigBranch* Branch, FConfigCommandStreamSection* Sec)
+		{
+			Sec->Remove(Key);
+		});
+}
+
+bool FConfigCacheIni::SetArrayInSectionOfStaticLayer(const TCHAR* Section, FName Key, const TArray<FString>& Values, bool bClearArray, const FString& BranchName, const FString& LayerSubstring, bool bWriteFile)
+{
+	return ModifyAndWriteLayer(this, Section, Key, BranchName, LayerSubstring, bWriteFile,
+		[Section, Key, Values, bClearArray, LayerSubstring](const FConfigBranch* Branch, FConfigCommandStreamSection* Sec)
+		{
+			// we always start this section over 
+			Sec->Remove(Key);
+			if (bClearArray)
+			{
+				Sec->Add(Key, FConfigValue(TEXT("___Clear___"), FConfigValue::EValueType::Clear));
+				for (const FString& Value : Values)
+				{
+					Sec->Add(Key, FConfigValue(Value, FConfigValue::EValueType::ArrayAddUnique));
+				}
+			}
+			else
+			{
+				// look one layer down, so we can minimize the operations needed to get tot he passed in array
+				FConfigFile ValuesBeforeLayer;
+				Branch->MergeStaticLayersUpTo(LayerSubstring, ValuesBeforeLayer);
+				TArray<const FConfigValue*> ParentValues;
+				if (const FConfigSection* ParentSection = ValuesBeforeLayer.FindSection(Section))
+				{
+					ParentSection->MultiFindPointer(Key, ParentValues, true);
+				}
+
+				TArray<FString> ValuesToAdd;
+				for (const FString& Value : Values)
+				{
+					bool bFound = false;
+					for (const FConfigValue* ParentValue : ParentValues)
+					{
+						if (ParentValue->GetValueForWriting() == Value)
+						{
+							ParentValues.Remove(ParentValue);
+							bFound = true;
+							break;
+						}
+					}
+					if (!bFound)
+					{
+						ValuesToAdd.Add(Value);
+					}
+				}
+
+				// any remaining existing value needs to be removed
+				// @todo: this needs testing with expanion variables - we probably had unexpanded values coming in, in which case we never want to expand - hard to say
+				for (const FConfigValue* ParentValue : ParentValues)
+				{
+					Sec->Remove(Key);
+					Sec->Add(Key, FConfigValue(ParentValue->GetSavedValueForWriting(), FConfigValue::EValueType::Remove));
+				}
+				for (const FString& Value : ValuesToAdd)
+				{
+					Sec->Add(Key, FConfigValue(Value, FConfigValue::EValueType::ArrayAddUnique));
+				}
+			}
+		});
+}
+
+
+
 /**
  * Archive for counting config file memory usage.
  */
@@ -6243,208 +6803,6 @@
 }
 
 
-/**
- * Functionality to assist with updating a config file with one property value change.
- */
-class FSinglePropertyConfigHelper
-{
-public:
-	
-
-	/**
-	 * We need certain information for the helper to be useful.
-	 *
-	 * @Param InIniFilename - The disk location of the file we wish to edit.
-	 * @Param InSectionName - The section the property belongs to.
-	 * @Param InPropertyName - The name of the property that has been edited.
-	 * @Param InPropertyValue - The new value of the property that has been edited, or unset to remove the property.
-	 */
-	FSinglePropertyConfigHelper(const FString& InIniFilename, const FString& InSectionName, const FString& InPropertyName, const TOptional<FString>& InPropertyValue)
-		: IniFilename(InIniFilename)
-		, SectionName(InSectionName)
-		, PropertyName(InPropertyName)
-		, PropertyValue(InPropertyValue)
-	{
-		// Split the file into the necessary parts.
-		PopulateFileContentHelper();
-	}
-
-	/**
-	 * Perform the action of updating the config file with the new property value.
-	 */
-	bool UpdateConfigFile()
-	{
-		UpdatePropertyInSection();
-		// Rebuild the file with the updated section.
-
-		FString NewFile = IniFileMakeup.BeforeSection + IniFileMakeup.Section + IniFileMakeup.AfterSection;
-		if (!NewFile.EndsWith(TEXT(LINE_TERMINATOR_ANSI LINE_TERMINATOR_ANSI)))
-		{
-			NewFile.AppendChars(LINE_TERMINATOR, TCString<TCHAR>::Strlen(LINE_TERMINATOR));
-		}
-		return SaveConfigFileWrapper(*IniFilename, NewFile);
-	}
 
 
-private:
-
-	/**
-	 * Clear any trailing whitespace from the end of the output.
-	 */
-	void ClearTrailingWhitespace(FString& InStr)
-	{
-		const FString Endl(LINE_TERMINATOR);
-		while (InStr.EndsWith(LINE_TERMINATOR, ESearchCase::CaseSensitive))
-		{
-			InStr.LeftChopInline(Endl.Len(), EAllowShrinking::No);
-		}
-	}
-	
-	/**
-	 * Update the section with the new value for the property.
-	 */
-	void UpdatePropertyInSection()
-	{
-		FString UpdatedSection;
-		if (IniFileMakeup.Section.IsEmpty())
-		{
-			if (PropertyValue.IsSet())
-			{
-				const FString DecoratedSectionName = FString::Printf(TEXT("[%s]"), *SectionName);
-
-				ClearTrailingWhitespace(IniFileMakeup.BeforeSection);
-				UpdatedSection += LINE_TERMINATOR;
-				UpdatedSection += LINE_TERMINATOR;
-				UpdatedSection += DecoratedSectionName;
-				AppendPropertyLine(UpdatedSection);
-			}
-		}
-		else
-		{
-			FString SectionLine;
-			const TCHAR* Ptr = *IniFileMakeup.Section;
-			bool bUpdatedPropertyOnPass = false;
-			while (Ptr != nullptr && FParse::Line(&Ptr, SectionLine, true))
-			{
-				if (SectionLine.StartsWith(FString::Printf(TEXT("%s="), *PropertyName)))
-				{
-					if (PropertyValue.IsSet())
-					{
-						UpdatedSection += FConfigFile::GenerateExportedPropertyLine(PropertyName, PropertyValue.GetValue());
-					}
-					bUpdatedPropertyOnPass = true;
-				}
-				else
-				{
-					UpdatedSection += SectionLine;
-					UpdatedSection += LINE_TERMINATOR;
-				}
-			}
-
-			// If the property wasnt found in the text of the existing section content,
-			// append it to the end of the section.
-			if (!bUpdatedPropertyOnPass && PropertyValue.IsSet())
-			{
-				AppendPropertyLine(UpdatedSection);
-			}
-			else
-			{
-				UpdatedSection += LINE_TERMINATOR;
-			}
-		}
-
-		IniFileMakeup.Section = UpdatedSection;
-	}
-	
-	/**
-	 * Split the file up into parts:
-	 * -> Before the section we wish to edit, which will remain unaltered,
-	 * ->-> The section we wish to edit, we only seek to edit the single property,
-	 * ->->-> After the section we wish to edit, which will remain unaltered.
-	 */
-	void PopulateFileContentHelper()
-	{
-		FString UnprocessedFileContents;
-		if (LoadConfigFileWrapper(*IniFilename, UnprocessedFileContents))
-		{
-			// Find the section in the file text.
-			const FString DecoratedSectionName = FString::Printf(TEXT("[%s]"), *SectionName);
-
-			const int32 DecoratedSectionNameStartIndex = UnprocessedFileContents.Find(DecoratedSectionName);
-			if (DecoratedSectionNameStartIndex != INDEX_NONE)
-			{
-				// If we found the section, cache off the file text before the section.
-				IniFileMakeup.BeforeSection = UnprocessedFileContents.Left(DecoratedSectionNameStartIndex);
-				UnprocessedFileContents.RemoveAt(0, IniFileMakeup.BeforeSection.Len());
-
-				// For the rest of the file, split it into the section we are editing and the rest of the file after.
-				const TCHAR* Ptr = UnprocessedFileContents.Len() > 0 ? *UnprocessedFileContents : nullptr;
-				FString NextUnprocessedLine;
-				bool bReachedNextSection = false;
-				while (Ptr != nullptr && FParse::Line(&Ptr, NextUnprocessedLine, true))
-				{
-					bReachedNextSection |= (NextUnprocessedLine.StartsWith(TEXT("[")) && NextUnprocessedLine != DecoratedSectionName);
-					if (bReachedNextSection)
-					{
-						IniFileMakeup.AfterSection += NextUnprocessedLine;
-						IniFileMakeup.AfterSection += LINE_TERMINATOR;
-					}
-					else
-					{
-						IniFileMakeup.Section += NextUnprocessedLine;
-						IniFileMakeup.Section += LINE_TERMINATOR;
-					}
-				}
-			}
-			else
-			{
-				IniFileMakeup.BeforeSection = UnprocessedFileContents;
-			}
-		}
-	}
-	
-	/**
-	 * Append the property entry to the section
-	 */
-	void AppendPropertyLine(FString& PreText)
-	{
-		check(PropertyValue.IsSet());
-		
-		// Make sure we dont leave much whitespace, and append the property name/value entry
-		ClearTrailingWhitespace(PreText);
-		PreText += LINE_TERMINATOR;
-		PreText += FConfigFile::GenerateExportedPropertyLine(PropertyName, PropertyValue.GetValue());
-		PreText += LINE_TERMINATOR;
-	}
-
-
-private:
-	// The disk location of the ini file we seek to edit
-	FString IniFilename;
-
-	// The section in the config file
-	FString SectionName;
-
-	// The name of the property that has been changed
-	FString PropertyName;
-
-	// The new value, in string format, of the property that has been changed
-	// This will be unset if the property has been removed from the config
-	TOptional<FString> PropertyValue;
-
-	// Helper struct that holds the makeup of the ini file.
-	struct IniFileContent
-	{
-		// The section we wish to edit
-		FString Section;
-
-		// The file contents before the section we are editing
-		FString BeforeSection;
-
-		// The file contents after the section we are editing
-		FString AfterSection;
-	} IniFileMakeup; // Instance of the helper to maintain file structure.
-};
-
-
 bool FConfigFile::UpdateSinglePropertyInSection(const TCHAR* DiskFilename, const TCHAR* PropertyName, const TCHAR* SectionName)
diff --git a/Engine/Source/Runtime/Core/Public/Misc/ConfigCacheIni.h b/Engine/Source/Runtime/Core/Public/Misc/ConfigCacheIni.h
--- a/Engine/Source/Runtime/Core/Public/Misc/ConfigCacheIni.h
+++ b/Engine/Source/Runtime/Core/Public/Misc/ConfigCacheIni.h
@@ -918,6 +918,12 @@ public:
 	CORE_API void SetArray(const TCHAR* Section, const TCHAR* Key, const TArray<FString>& Value);
 	
 
+	/**
+	 * Sets the given key/value pair to the Section in the given File. This will replace any existing values (as if ResetKeyInSection had been called)
+	 * @return true if the section was modified
+	 */
+	CORE_API bool SetInSection(const TCHAR* SectionName, FName Key, const FString& Value);
+		
 	/**
 	 * Adds the given key/value pair to the Section. This will always add this pair to the section, even if the pair already exists.
 	 * This is equivalent to the . operator in .ini files
@@ -1168,6 +1174,23 @@ public:
 	CORE_API void SafeUnload();
 	CORE_API void SafeReload();
 
+	/** Recalculate the usable values from all static and dynamic layers. Useful if you update a layer and need to see results in 
+	 * the InMemoryFile (ie what GConfig->GetXXX returns) 
+	 */
+	CORE_API void ReapplyLayers();
+
+	/** 
+	 * Apply all static layers up, but not including, the given layer name substring (so passing "DefaultEngine" will give all values before the DefaultEngine layer)
+	 * This only works if static layers are being tracked (generally in the editor and not client builds - see EBranchReplayMethod::FullReplay)
+	 * @return true if the layer was found, or false if the layer name was not found (which means all static layers will be combined and returned)
+	 */
+	CORE_API bool MergeStaticLayersUpTo(const FString& LayerNameSubstring, FConfigFile& OutFile) const;
+
+	/**
+	 * Find the static layer that first matches the substring (or full) name of the layer filename
+	 */
+	CORE_API const FConfigCommandStream* GetStaticLayer(const FString& LayerSubstring) const;
+
 	/**
 	 * Removes the section completely from all layers of this branch 
 	 * This is destructive! It will not reload on demand (like SafeUnload does). Use this only when you know you will _never_ need
@@ -1676,12 +1699,21 @@ public:
 		const FString&	Filename
 	);
 	
+
+	/**
+	 * Sets the given key/value pair to the Section in the given File. This will replace any existing values (as if ResetKeyInSection had been called)
+	 * @return true if the section was modified
+	 */
+	CORE_API bool SetInSection(const TCHAR* Section, FName Key, const FString& Value, const FString& Filename);
+	CORE_API bool SetInSectionOfStaticLayer(const TCHAR* Section, FName Key, const FString& Value, const FString& BranchName, const FString& LayerSubstring, bool bWriteFile);
+
 	/**
 	 * Adds the given key/value pair to the Section in the given File. This will always add this pair to the section, even if the pair already exists.
 	 * This is equivalent to the . operator in .ini files
 	 * @return true if the section was modified
 	 */
 	CORE_API bool AddToSection(const TCHAR* Section, FName Key, const FString& Value, const FString& Filename);
+	CORE_API bool AddToSectionOfStaticLayer(const TCHAR* Section, FName Key, const FString& Value, const FString& BranchName, const FString& LayerSubstring, bool bWriteFile);
 
 	/**
 	 * Adds the given key/value pair to the Section in the given File, if the pair didn't already exist
@@ -1689,6 +1721,7 @@ public:
 	 * @return true if the section was modified
 	 */
 	CORE_API bool AddUniqueToSection(const TCHAR* Section, FName Key, const FString& Value, const FString& Filename);
+	CORE_API bool AddUniqueToSectionOfStaticLayer(const TCHAR* Section, FName Key, const FString& Value, const FString& BranchName, const FString& LayerSubstring, bool bWriteFile);
 
 	/**
 	 * Removes every entry in the Section in the given File that has Key, no matter what the Value is
@@ -1696,6 +1729,7 @@ public:
 	 * @return true if the section was modified
 	 */
 	CORE_API bool RemoveKeyFromSection(const TCHAR* Section, FName Key, const FString& Filename);
+	CORE_API bool RemoveKeyFromSectionOfStaticLayer(const TCHAR* Section, FName Key, const FString& BranchName, const FString& LayerSubstring, bool bWriteFile);
 
 	/**
 	 * Removes every  entry in the Section in the given File that has the Key/Value pair
@@ -1703,6 +1737,7 @@ public:
 	 * @return true if the section was modified
 	 */
 	CORE_API bool RemoveFromSection(const TCHAR* Section, FName Key, const FString& Value, const FString& Filename);
+	CORE_API bool RemoveFromSectionOfStaticLayer(const TCHAR* Section, FName Key, const FString& Value, const FString& BranchName, const FString& LayerSubstring, bool bWriteFile);
 
 	/**
 	 * Similar to RemvoeKeyFromSection, but if this File's changes are being tracked, then we remove all changes to the
@@ -1712,6 +1747,9 @@ public:
 	 * @return true if the section was modified
 	 */
 	CORE_API bool ResetKeyInSection(const TCHAR* Section, FName Key, const FString& Filename);
+	CORE_API bool ResetKeyInSectionOfStaticLayer(const TCHAR* Section, FName Key, const FString& BranchName, const FString& LayerSubstring, bool bWriteFile);
+
+	CORE_API bool SetArrayInSectionOfStaticLayer(const TCHAR* Section, FName Key, const TArray<FString>& Values, bool bClearArray, const FString& BranchName, const FString& LayerSubstring, bool bWriteFile);
 
 	/**
 	 * Clears out the memory of a branch, but can reload it on demand if needed
diff --git a/Engine/Source/Runtime/CoreUObject/Private/UObject/Obj.cpp b/Engine/Source/Runtime/CoreUObject/Private/UObject/Obj.cpp
--- a/Engine/Source/Runtime/CoreUObject/Private/UObject/Obj.cpp
+++ b/Engine/Source/Runtime/CoreUObject/Private/UObject/Obj.cpp
@@ -94,6 +94,12 @@ namespace
 		GPrematureConfigLoadErrorLevel,
 		TEXT("Determines the level of warning issued if LoadConfig is called on an object that is not yet sufficiently initialized. 0=Display, 1=Warning, 2=Ensure, 3=Check"));
 
+	static int32 GUseNewPropertySaving = 0;
+	static FAutoConsoleVariableRef CVarUseNewPropertySaving(
+		TEXT("ini.UseNewPropertySaving"),
+		GUseNewPropertySaving,
+		TEXT("Determines if UObject property saving should use the Config APIs that can set a property (including arrays) directly in a static layer."));
+
 	// Enable additional (slow) checks in CheckDefaultSubobjects.
 	static constexpr bool bDoCheckSubobjects = false;
 }
@@ -3459,8 +3465,18 @@ void UObject::LoadConfig(UClass* ConfigClass/*=nullptr*/, const TCHAR* InFilenam
 	}
 }
 
-void UObject::SaveConfig(uint64 RequiredPropertyFlags, const TCHAR* InFilename, FConfigCacheIni* Config/*=GConfig*/,
-	bool bAllowCopyToDefaultObject/*=true*/)
+void UObject::SaveConfig(uint64 RequiredPropertyFlags, const TCHAR* InFilename, FConfigCacheIni* Config/*=GConfig*/, bool bAllowCopyToDefaultObject/*=true*/)
+{
+	FSaveConfigContext Context;
+	Context.RequiredPropertyFlags = (EPropertyFlags)RequiredPropertyFlags;
+	Context.Filename = InFilename;
+	Context.ConfigSystem = Config;
+	Context.bAllowCopyToDefaultObject = bAllowCopyToDefaultObject;
+
+	SaveConfig(Context);
+}
+
+void UObject::SaveConfig(const FSaveConfigContext& Context)
 {
 	if (!GetClass()->HasAnyClassFlags(CLASS_Config))
 	{
@@ -3468,7 +3484,9 @@ void UObject::SaveConfig(uint64 RequiredPropertyFlags, const TCHAR* InFilename,
 	}
 
 	// if a filename was specified, save all properties to that file no matter which class they come from
-	const FString Filename = InFilename ? InFilename : GetConfigFilename(this);
+	const FString Filename = Context.Filename.Len() > 0 ? Context.Filename : GetConfigFilename(this);
+	
+	FConfigCacheIni* Config = Context.ConfigSystem;
 
 	// If there is no existing branch for the filename and file operations are enabled, then add a new temporary branch and remove it after the flush.
 	const bool bIsTemporaryBranch = Config->FindBranch(*Filename, Filename) == nullptr && !Config->AreFileOperationsDisabled();
@@ -3477,11 +3495,11 @@ void UObject::SaveConfig(uint64 RequiredPropertyFlags, const TCHAR* InFilename,
 		Config->AddNewBranch(Filename);
 	}
 
-	// Determine whether the file we are writing is a default file config.
-	const bool bIsADefaultIniWrite = Filename == GetDefaultConfigFilename() || Filename == GetGlobalUserConfigFilename();
+	// Determine whether the file we are writing to is a layer in the hierarchy, as opposed to the final Saved file
+	const bool bIsLayerWrite = Filename == GetDefaultConfigFilename() || Filename == GetGlobalUserConfigFilename() || Context.bSaveToLayer;
 
 	const bool bPerObject = UsesPerObjectConfig(this);
-	FString Section;
+	FString Section = Context.OverrideConfigSectionName;
 
 	if (bPerObject == true)
 	{
@@ -3498,24 +3516,33 @@ void UObject::SaveConfig(uint64 RequiredPropertyFlags, const TCHAR* InFilename,
 		}
 
 		//RobM: we need to update this to use GetClass()->GetPathName() after we fix all places that format section names
-		Section = PathNameString + TEXT(" ") + GetClass()->GetName();
+		if (Context.OverrideConfigSectionName.IsEmpty())
+		{
+			Section = PathNameString + TEXT(" ") + GetClass()->GetName();
 
-		OverridePerObjectConfigSection(Section);
+			OverridePerObjectConfigSection(Section);
+		}
 	}
 
 	UObject* CDO = GetClass()->GetDefaultObject();
 
 	// only copy the values to the CDO if this is GConfig and we're not saving the CDO
-	const bool bCopyValues = (bAllowCopyToDefaultObject && this != CDO && Config == GConfig);
+	const bool bCopyValues = (Context.bAllowCopyToDefaultObject && this != CDO && Context.ConfigSystem == GConfig);
 
 	for (FProperty* Property = GetClass()->PropertyLink; Property; Property = Property->PropertyLinkNext)
 	{
+		// we just check names, in case this property is a duplicated but same property as this object's class
+		if (Context.SinglePropertyName != NAME_None && Property->GetFName() != Context.SinglePropertyName)
+		{
+			continue;
+		}
+
 		if (!Property->HasAnyPropertyFlags(CPF_Config))
 		{
 			continue;
 		}
 
-		if ((Property->PropertyFlags & RequiredPropertyFlags) == RequiredPropertyFlags)
+		if ((Property->PropertyFlags & Context.RequiredPropertyFlags) == Context.RequiredPropertyFlags)
 		{
 			UClass* BaseClass = GetClass();
 
@@ -3537,7 +3564,7 @@ void UObject::SaveConfig(uint64 RequiredPropertyFlags, const TCHAR* InFilename,
 			}
 #endif // #if WITH_EDITOR
 
-			if (!bPerObject)
+			if (!bPerObject && Context.OverrideConfigSectionName.IsEmpty())
 			{
 				Section = BaseClass->GetPathName();
 				// allow the class to override the expected section name
@@ -3546,7 +3573,7 @@ void UObject::SaveConfig(uint64 RequiredPropertyFlags, const TCHAR* InFilename,
 
 			// globalconfig properties should always use the owning class's config file
 			// specifying a value for InFilename will override this behavior (as it does with normal properties)
-			const FString& PropFileName = ((Property->PropertyFlags & CPF_GlobalConfig) && InFilename == nullptr) ? Property->GetOwnerClass()->GetConfigName() : Filename;
+			const FString& PropFileName = ((Property->PropertyFlags & CPF_GlobalConfig) && Context.Filename.Len() == 0) ? Property->GetOwnerClass()->GetConfigName() : Filename;
 
 			// Properties that are the same as the parent class' defaults should not be saved to ini
 			// Before modifying any key in the section, first check to see if it is different from the parent.
@@ -3568,49 +3595,86 @@ void UObject::SaveConfig(uint64 RequiredPropertyFlags, const TCHAR* InFilename,
 			FSetProperty* SetProperty = CastField<FSetProperty>(Property);
 			if (Array || SetProperty)
 			{
-				const FConfigSection* Sec = Config->GetSection(*Section, 1, PropFileName);
-				// Default ini's require the array syntax to be applied to the property name
-				FString CompleteKey = FString::Printf(TEXT("%s%s"), bIsADefaultIniWrite ? TEXT("+") : TEXT(""), *Key);
-				if (Sec)
-				{
-					// Delete the old value for the property in the ConfigCache before (conditionally) adding in the new value
-					Config->ResetKeyInSection(*Section, *CompleteKey, PropFileName);
-				}
-
-				if (!bPropDeprecated && (!bShouldCheckIfIdenticalBeforeAdding || !Property->Identical_InContainer(this, SuperClassDefaultObject)))
+				// write a property to a specific layer (with new path, for now the old way of writing to defaults is still below)
+				if (Context.bSaveToLayer)
 				{
+					TArray<FString> ExportedValues;
 					if (Array)
 					{
 						FScriptArrayHelper_InContainer ArrayHelper(Array, this);
-						for (int32 i=0; i<ArrayHelper.Num(); i++)
-						{
-							FString	Buffer;
-							Array->Inner->ExportTextItem_Direct(Buffer, ArrayHelper.GetRawPtr(i), ArrayHelper.GetRawPtr(i), this, PortFlags);
-							Config->AddToSection(*Section, *CompleteKey, *Buffer, PropFileName);
-						}
-						if (ArrayHelper.Num() == 0 && bIsADefaultIniWrite)
+						ExportedValues.AddDefaulted(ArrayHelper.Num());
+						for (int32 i = 0; i < ArrayHelper.Num(); i++)
 						{
-							const FString EmptyKey = FString::Printf(TEXT("!%s"), *Key);
-							Config->AddToSection(*Section, *EmptyKey, TEXT("__ClearArray__"), PropFileName);
+							Array->Inner->ExportTextItem_Direct(ExportedValues[i], ArrayHelper.GetRawPtr(i), ArrayHelper.GetRawPtr(i), this, PortFlags);
 						}
+						Config->SetArrayInSectionOfStaticLayer(*Section, *Key, ExportedValues, false, Property->GetOwnerClass()->GetConfigName(), Filename, true);
 						bProcessedProperty = true;
 					}
 					else if (SetProperty)
 					{
 						FScriptSetHelper_InContainer SetHelper(SetProperty, this);
-
 						// If we have one or fewer elements, fall back to the old export method so we match what we do
 						// on import
 						if (SetHelper.Num() > 1)
 						{
-							for (FScriptSetHelper::FIterator It = SetHelper.CreateIterator(); It; ++It)
+							ExportedValues.AddDefaulted(SetHelper.Num());
+							int Index = 0;
+							for (FScriptSetHelper::FIterator It = SetHelper.CreateIterator(); It; ++It, Index++)
+							{
+								SetProperty->ElementProp->ExportTextItem_Direct(ExportedValues[Index], SetHelper.GetElementPtr(It), SetHelper.GetElementPtr(It), this, PortFlags);
+							}
+							Config->SetArrayInSectionOfStaticLayer(*Section, *Key, ExportedValues, false, Property->GetOwnerClass()->GetConfigName(), Filename, true);
+							bProcessedProperty = true;
+						}
+					}
+				}
+				else
+				{
+
+					const FConfigSection* Sec = Config->GetSection(*Section, 1, PropFileName);
+					// Layer ini's require the array syntax to be applied to the property name
+					FString CompleteKey = FString::Printf(TEXT("%s%s"), bIsLayerWrite ? TEXT("+") : TEXT(""), *Key);
+					if (Sec)
+					{
+						// Delete the old value for the property in the ConfigCache before (conditionally) adding in the new value
+						Config->ResetKeyInSection(*Section, *CompleteKey, PropFileName);
+					}
+
+					if (!bPropDeprecated && (!bShouldCheckIfIdenticalBeforeAdding || !Property->Identical_InContainer(this, SuperClassDefaultObject)))
+					{
+						if (Array)
+						{
+							FScriptArrayHelper_InContainer ArrayHelper(Array, this);
+							for (int32 i = 0; i < ArrayHelper.Num(); i++)
 							{
 								FString	Buffer;
-								SetProperty->ElementProp->ExportTextItem_Direct(Buffer, SetHelper.GetElementPtr(It), SetHelper.GetElementPtr(It), this, PortFlags);
+								Array->Inner->ExportTextItem_Direct(Buffer, ArrayHelper.GetRawPtr(i), ArrayHelper.GetRawPtr(i), this, PortFlags);
 								Config->AddToSection(*Section, *CompleteKey, *Buffer, PropFileName);
 							}
+							if (ArrayHelper.Num() == 0 && bIsLayerWrite)
+							{
+								const FString EmptyKey = FString::Printf(TEXT("!%s"), *Key);
+								Config->AddToSection(*Section, *EmptyKey, TEXT("__ClearArray__"), PropFileName);
+							}
 							bProcessedProperty = true;
 						}
+						else if (SetProperty)
+						{
+							FScriptSetHelper_InContainer SetHelper(SetProperty, this);
+
+							// If we have one or fewer elements, fall back to the old export method so we match what we do
+							// on import
+							if (SetHelper.Num() > 1)
+							{
+								for (FScriptSetHelper::FIterator It = SetHelper.CreateIterator(); It; ++It)
+								{
+									FString	Buffer;
+									SetProperty->ElementProp->ExportTextItem_Direct(Buffer, SetHelper.GetElementPtr(It), SetHelper.GetElementPtr(It), this, PortFlags);
+									Config->AddToSection(*Section, *CompleteKey, *Buffer, PropFileName);
+								}
+								bProcessedProperty = true;
+							}
+						}
 					}
 				}
 			}
@@ -3630,12 +3694,26 @@ void UObject::SaveConfig(uint64 RequiredPropertyFlags, const TCHAR* InFilename,
 					{
 						FString	Value;
 						Property->ExportText_InContainer(Index, Value, this, this, this, PortFlags);
-						Config->SetString(*Section, *Key, *Value, PropFileName);
+						if (Context.bSaveToLayer)
+						{
+							Config->SetInSectionOfStaticLayer(*Section, *Key, *Value, Property->GetOwnerClass()->GetConfigName(), Filename, true);
+						}
+						else
+						{
+							Config->SetString(*Section, *Key, *Value, PropFileName);
+						}
 					}
 					else
 					{
-						// If we are not writing it to config above, we should make sure that this property isn't stagnant in the cache.
-						Config->ResetKeyInSection(*Section, *Key, PropFileName);
+						if (Context.bSaveToLayer)
+						{
+							Config->ResetKeyInSectionOfStaticLayer(*Section, *Key, Property->GetOwnerClass()->GetConfigName(), Filename, true);
+						}
+						else
+						{
+							// If we are not writing it to config above, we should make sure that this property isn't stagnant in the cache.
+							Config->ResetKeyInSection(*Section, *Key, PropFileName);
+						}
 					}
 				}
 			}
@@ -3651,7 +3729,7 @@ void UObject::SaveConfig(uint64 RequiredPropertyFlags, const TCHAR* InFilename,
 	}
 
 	// Only write out the config file if this is GConfig or if we are saving a temporary branch.
-	if (Config == GConfig || bIsTemporaryBranch)
+	if (Context.ConfigSystem == GConfig || bIsTemporaryBranch)
 	{
 		Config->Flush(bIsTemporaryBranch/*bRemoveFromCache*/, *Filename);
 	}
@@ -3782,6 +3860,18 @@ void UObject::UpdateProjectUserConfigFile()
 
 void UObject::UpdateSinglePropertyInConfigFile(const FProperty* InProperty, const FString& InConfigIniName)
 {
+	if (GUseNewPropertySaving > 0)
+	{
+		// run SaveConfig, which has some logic to for diffing against parent class, and now has logic to also
+		// directly update a single config property, including arrays, via FConfigCacheIni functions
+		FSaveConfigContext Context;
+		Context.Filename = InConfigIniName;
+		Context.bSaveToLayer = true;
+		Context.SinglePropertyName = InProperty->GetFName();
+		SaveConfig(Context);
+		return;
+	}
+
 	// Arrays and ini files are a mine field, for now we don't support this.
 	if (!InProperty->IsA(FArrayProperty::StaticClass()))
 	{
diff --git a/Engine/Source/Runtime/CoreUObject/Public/UObject/Object.h b/Engine/Source/Runtime/CoreUObject/Public/UObject/Object.h
--- a/Engine/Source/Runtime/CoreUObject/Public/UObject/Object.h
+++ b/Engine/Source/Runtime/CoreUObject/Public/UObject/Object.h
@@ -1291,7 +1291,20 @@ public:
 	 * @warning Must be safe to call on class-default object
 	 */
 	COREUOBJECT_API void SaveConfig(uint64 RequiredPropertyFlags = CPF_Config, const TCHAR* Filename = nullptr,
-		FConfigCacheIni* Config=GConfig, bool bAllowCopyToDefaultObject=true);
+		FConfigCacheIni* Config=GConfig, bool bAllowCopyToDefaultObject = true);
+
+	struct FSaveConfigContext
+	{
+		EPropertyFlags RequiredPropertyFlags = CPF_Config;
+		FString Filename;
+		FConfigCacheIni* ConfigSystem = GConfig;
+		bool bAllowCopyToDefaultObject = true;
+
+		bool bSaveToLayer = false;
+		FName SinglePropertyName;
+		FString OverrideConfigSectionName;
+	};
+	COREUOBJECT_API void SaveConfig(const FSaveConfigContext& Context);
 
 	/**
 	 * Saves just the section(s) for this class into the default ini file for the class (with just the changes from base)
