[AutoRTFM] Convert PackageIdImpl::Lock to a transactionally-safe lock.

This lock is enabled in EditorTest and can be triggered by the RigVM unit tests.

The editor does not currently use AutoRTFM, so this should be a no-op unless you are running EditorTestAutoRTFM.

diff --git a/Engine/Source/Programs/AutoRTFMTests/Private/UECoreTests.cpp b/Engine/Source/Programs/AutoRTFMTests/Private/UECoreTests.cpp
--- a/Engine/Source/Programs/AutoRTFMTests/Private/UECoreTests.cpp
+++ b/Engine/Source/Programs/AutoRTFMTests/Private/UECoreTests.cpp
@@ -1663,6 +1663,19 @@ TEST_CASE("UECore.FOutputDeviceRedirector")
 
 TEST_CASE("UECore.AsyncLoading")
 {
+	SECTION("DoesPackageExist")
+	{
+		bool Exists = false;
+
+		AutoRTFM::Testing::Commit([&]
+		{
+			FString Name(FString::Printf(TEXT("/Game/AutoRTFMTestPackage%d"), __LINE__));
+			Exists = FPackageName::DoesPackageExist(Name);
+		});
+
+		REQUIRE(!Exists);
+	}
+
 	SECTION("LoadPackageAsync")
 	{
 		int32 RequestId = -1;
diff --git a/Engine/Source/Runtime/Core/Private/IO/PackageId.cpp b/Engine/Source/Runtime/Core/Private/IO/PackageId.cpp
--- a/Engine/Source/Runtime/Core/Private/IO/PackageId.cpp
+++ b/Engine/Source/Runtime/Core/Private/IO/PackageId.cpp
@@ -9,13 +9,16 @@
 #include "Serialization/StructuredArchiveAdapters.h"
 #include "Serialization/StructuredArchiveSlots.h"
 #include "Misc/ScopeRWLock.h"
+#include "Misc/TransactionallySafeRWLock.h"
 
 #if WITH_PACKAGEID_NAME_MAP
 LLM_DEFINE_TAG(PackageId_ReverseMapping);
 
 namespace PackageIdImpl
 {
-	FRWLock Lock;
+	// We use a transactionally-safe lock here because `FPackageId::FromName` can be called by 
+	// `FPackageName::DoesPackageExist`, and this should be safe to check inside of a transaction.
+	FTransactionallySafeRWLock Lock;
 	TMap<uint64, FName> Entries;
 }
 #endif
@@ -34,7 +37,7 @@ FPackageId FPackageId::FromName(const FName& Name)
 
 #if WITH_PACKAGEID_NAME_MAP
 	{
-		FWriteScopeLock ScopeWriteLock(PackageIdImpl::Lock);
+		UE::TWriteScopeLock ScopeWriteLock(PackageIdImpl::Lock);
 		LLM_SCOPE_BYTAG(PackageId_ReverseMapping);
 		FName EntryName = PackageIdImpl::Entries.FindOrAdd(Hash, Name);
 		checkf(EntryName.GetDisplayIndex() == Name.GetDisplayIndex() || EntryName.GetComparisonIndex() == Name.GetComparisonIndex(), TEXT("FPackageId collision: %llu for both %s and %s"), Hash, *Name.ToString(), *EntryName.ToString());
@@ -46,7 +49,7 @@ FPackageId FPackageId::FromName(const FName& Name)
 #if WITH_PACKAGEID_NAME_MAP
 FName FPackageId::GetName() const
 {
-	FReadScopeLock ScopeReadLock(PackageIdImpl::Lock);
+	UE::TReadScopeLock ScopeReadLock(PackageIdImpl::Lock);
 	return PackageIdImpl::Entries.FindRef(Id);
 }
 #endif
