Remove custom signal handling of crashing on Mac games

#rnx

diff --git a/Engine/Source/Runtime/Core/Private/Mac/MacPlatformMisc.cpp b/Engine/Source/Runtime/Core/Private/Mac/MacPlatformMisc.cpp
--- a/Engine/Source/Runtime/Core/Private/Mac/MacPlatformMisc.cpp
+++ b/Engine/Source/Runtime/Core/Private/Mac/MacPlatformMisc.cpp
@@ -2126,8 +2126,19 @@ void FMacPlatformMisc::SetCrashHandler(void (* CrashHandler)(const FGenericCrash
 	SCOPED_AUTORELEASE_POOL;
 	
 	GCrashHandlerPointer = CrashHandler;
+	
+	bool bEnableCustomSignalHandling = true;
+	// For packaged games, only use custom signal handling when CRC is opt-in. Otherwise let macOS process crash without UE interference
+#if !WITH_EDITOR
+	FString ReportClient = FPaths::ConvertRelativePathToFull(FPlatformProcess::GenerateApplicationPath(TEXT("CrashReportClient"), EBuildConfiguration::Development));
+	char CrashReportClient[PATH_MAX+1];
+	FCStringAnsi::Strncpy(CrashReportClient, TCHAR_TO_UTF8(*ReportClient), PATH_MAX + 1);
+	struct stat Stat;
+	stat(CrashReportClient, &Stat);
+	bEnableCustomSignalHandling = (S_ISREG(Stat.st_mode)) && ((Stat.st_mode & S_IXUSR) != 0);	// CRC exists and executable
+#endif
 
-	if (!FMacApplicationInfo::CrashReporter && !GCrashMalloc)
+	if (bEnableCustomSignalHandling && !FMacApplicationInfo::CrashReporter && !GCrashMalloc)
 	{
 		FPlatformProcess::UserHomeDir(); // This may be called from the crash handler, so we call it here to initialize it early.
 										 // This way the crash handler will use a cached string and we will limit allocations and avoid creating an autorelease pool.
