Fix a bug I introduced in CL42548546 that caused DLCs to not be stored in their own oplogs in zenserver.

diff --git a/Engine/Source/Runtime/Core/Private/Misc/App.cpp b/Engine/Source/Runtime/Core/Private/Misc/App.cpp
--- a/Engine/Source/Runtime/Core/Private/Misc/App.cpp
+++ b/Engine/Source/Runtime/Core/Private/Misc/App.cpp
@@ -473,7 +473,7 @@ FString FApp::GetZenStoreProjectId(FStringView SubProject)
 #if PLATFORM_DESKTOP
 	if (FPaths::IsProjectFilePathSet())
 	{
-		return GetZenStoreProjectIdForProject(FPaths::GetProjectFilePath());
+		return GetZenStoreProjectIdForProject(FPaths::GetProjectFilePath(), SubProject);
 	}
 	UE_LOG(LogInit, Fatal, TEXT("GetZenStoreProjectId() called before having a valid project file path"));
 #else
