Add FAndroidMisc::GetVKQualityRecommendation() and SRC_VKQuality

diff --git a/Engine/Plugins/Runtime/AndroidDeviceProfileSelector/Source/AndroidDeviceProfileSelector/Private/AndroidDeviceProfileMatchingRules.h b/Engine/Plugins/Runtime/AndroidDeviceProfileSelector/Source/AndroidDeviceProfileSelector/Private/AndroidDeviceProfileMatchingRules.h
--- a/Engine/Plugins/Runtime/AndroidDeviceProfileSelector/Source/AndroidDeviceProfileSelector/Private/AndroidDeviceProfileMatchingRules.h
+++ b/Engine/Plugins/Runtime/AndroidDeviceProfileSelector/Source/AndroidDeviceProfileSelector/Private/AndroidDeviceProfileMatchingRules.h
@@ -23,6 +23,7 @@ enum ESourceType : int
 	SRC_ConfigRuleVar,
 	SRC_HMDSystemName,
 	SRC_SM5Available,
+	SRC_VKQuality,
 	SRC_MAX,
 };
 
diff --git a/Engine/Plugins/Runtime/AndroidDeviceProfileSelector/Source/AndroidDeviceProfileSelector/Private/AndroidDeviceProfileSelector.cpp b/Engine/Plugins/Runtime/AndroidDeviceProfileSelector/Source/AndroidDeviceProfileSelector/Private/AndroidDeviceProfileSelector.cpp
--- a/Engine/Plugins/Runtime/AndroidDeviceProfileSelector/Source/AndroidDeviceProfileSelector/Private/AndroidDeviceProfileSelector.cpp
+++ b/Engine/Plugins/Runtime/AndroidDeviceProfileSelector/Source/AndroidDeviceProfileSelector/Private/AndroidDeviceProfileSelector.cpp
@@ -51,6 +51,7 @@ void FAndroidDeviceProfileSelector::VerifySelectorParams()
 	check(SelectorProperties.Find(FAndroidProfileSelectorSourceProperties::SRC_HMDSystemName));
 	check(SelectorProperties.Find(FAndroidProfileSelectorSourceProperties::SRC_TotalPhysicalGB));
 	check(SelectorProperties.Find(FAndroidProfileSelectorSourceProperties::SRC_SM5Available));
+	check(SelectorProperties.Find(FAndroidProfileSelectorSourceProperties::SRC_VKQuality));
 }
 
 FString FAndroidDeviceProfileSelector::FindMatchingProfile(const FString& FallbackProfileName)
@@ -72,6 +73,7 @@ FString FAndroidDeviceProfileSelector::FindMatchingProfile(const FString& Fallba
 	FString TotalPhysicalGB		= SelectorProperties.FindChecked(FAndroidProfileSelectorSourceProperties::SRC_TotalPhysicalGB);
 	FString HMDSystemName		= SelectorProperties.FindChecked(FAndroidProfileSelectorSourceProperties::SRC_HMDSystemName);
 	FString SM5Available		= SelectorProperties.FindChecked(FAndroidProfileSelectorSourceProperties::SRC_SM5Available);
+	FString VKQuality			= SelectorProperties.FindChecked(FAndroidProfileSelectorSourceProperties::SRC_VKQuality);
 
 	for (const FProfileMatch& Profile : GetAndroidDeviceProfileMatchingRules()->MatchProfile)
 	{
@@ -129,6 +131,9 @@ FString FAndroidDeviceProfileSelector::FindMatchingProfile(const FString& Fallba
 			case ESourceType::SRC_SM5Available:
 				SourceString = &SM5Available;
 				break;
+			case ESourceType::SRC_VKQuality:
+				SourceString = &VKQuality;
+				break;
 			case ESourceType::SRC_ConfigRuleVar:
 				{
 					// expected Matchstring contents for configrulevar is "configrule_varname|matchstring"
diff --git a/Engine/Plugins/Runtime/AndroidDeviceProfileSelector/Source/AndroidDeviceProfileSelector/Private/AndroidDeviceProfileSelectorModule.cpp b/Engine/Plugins/Runtime/AndroidDeviceProfileSelector/Source/AndroidDeviceProfileSelector/Private/AndroidDeviceProfileSelectorModule.cpp
--- a/Engine/Plugins/Runtime/AndroidDeviceProfileSelector/Source/AndroidDeviceProfileSelector/Private/AndroidDeviceProfileSelectorModule.cpp
+++ b/Engine/Plugins/Runtime/AndroidDeviceProfileSelector/Source/AndroidDeviceProfileSelector/Private/AndroidDeviceProfileSelectorModule.cpp
@@ -112,6 +112,7 @@ void FAndroidDeviceProfileSelectorModule::GetDeviceParametersFromJson(FString& J
 		OutDeviceParameters.Add(FName(TEXT("SRC_TotalPhysicalGB")), AndroidProperties.TotalPhysicalGB);
 		OutDeviceParameters.Add(FName(TEXT("SRC_HMDSystemName")), TEXT(""));
 		OutDeviceParameters.Add(FName(TEXT("SRC_SM5Available")), AndroidProperties.SM5Available ? "true" : "false");
+		OutDeviceParameters.Add(FName(TEXT("SRC_VKQuality")), TEXT(""));
 		OutDeviceParameters.Add(FAndroidProfileSelectorSourceProperties::SRC_ResolutionX, FString::Printf(TEXT("%d"), DeviceSpecs.ResolutionX));
 		OutDeviceParameters.Add(FAndroidProfileSelectorSourceProperties::SRC_ResolutionY, FString::Printf(TEXT("%d"), DeviceSpecs.ResolutionY));
 		OutDeviceParameters.Add(FAndroidProfileSelectorSourceProperties::SRC_InsetsLeft, FString::Printf(TEXT("%f"), DeviceSpecs.InsetsLeft));
diff --git a/Engine/Plugins/Runtime/AndroidDeviceProfileSelector/Source/AndroidDeviceProfileSelector/Public/AndroidDeviceProfileSelector.h b/Engine/Plugins/Runtime/AndroidDeviceProfileSelector/Source/AndroidDeviceProfileSelector/Public/AndroidDeviceProfileSelector.h
--- a/Engine/Plugins/Runtime/AndroidDeviceProfileSelector/Source/AndroidDeviceProfileSelector/Public/AndroidDeviceProfileSelector.h
+++ b/Engine/Plugins/Runtime/AndroidDeviceProfileSelector/Source/AndroidDeviceProfileSelector/Public/AndroidDeviceProfileSelector.h
@@ -29,6 +29,7 @@ namespace FAndroidProfileSelectorSourceProperties
 	static FName SRC_InsetsTop(TEXT("SRC_InsetsTop"));
 	static FName SRC_InsetsRight(TEXT("SRC_InsetsRight"));
 	static FName SRC_InsetsBottom(TEXT("SRC_InsetsBottom"));
+	static FName SRC_VKQuality(TEXT("SRC_VKQuality"));
 };
 
 class FAndroidDeviceProfileSelector
diff --git a/Engine/Plugins/Runtime/AndroidDeviceProfileSelector/Source/AndroidDeviceProfileSelectorRuntime/Private/AndroidDeviceProfileSelectorRuntimeModule.cpp b/Engine/Plugins/Runtime/AndroidDeviceProfileSelector/Source/AndroidDeviceProfileSelectorRuntime/Private/AndroidDeviceProfileSelectorRuntimeModule.cpp
--- a/Engine/Plugins/Runtime/AndroidDeviceProfileSelector/Source/AndroidDeviceProfileSelectorRuntime/Private/AndroidDeviceProfileSelectorRuntimeModule.cpp
+++ b/Engine/Plugins/Runtime/AndroidDeviceProfileSelector/Source/AndroidDeviceProfileSelectorRuntime/Private/AndroidDeviceProfileSelectorRuntimeModule.cpp
@@ -78,6 +78,7 @@ static const TMap<FName,FString>& GetDeviceSelectorParams()
 		AndroidParams.Add(FAndroidProfileSelectorSourceProperties::SRC_HMDSystemName, HMDRequestedProfileName);
 		AndroidParams.Add(FAndroidProfileSelectorSourceProperties::SRC_TotalPhysicalGB, FString::Printf(TEXT("%d"), TotalPhysicalGB));
 		AndroidParams.Add(FAndroidProfileSelectorSourceProperties::SRC_SM5Available, FAndroidMisc::IsDesktopVulkanAvailable() ? TEXT("true") : TEXT("false"));
+		AndroidParams.Add(FAndroidProfileSelectorSourceProperties::SRC_VKQuality, FAndroidMisc::GetVKQualityRecommendation());
 
 #if PLATFORM_ANDROID
 		// allow ConfigRules to override cvars first
diff --git a/Engine/Source/Editor/DeviceProfileServices/Private/CheckAndroidDeviceProfileCommandlet.cpp b/Engine/Source/Editor/DeviceProfileServices/Private/CheckAndroidDeviceProfileCommandlet.cpp
--- a/Engine/Source/Editor/DeviceProfileServices/Private/CheckAndroidDeviceProfileCommandlet.cpp
+++ b/Engine/Source/Editor/DeviceProfileServices/Private/CheckAndroidDeviceProfileCommandlet.cpp
@@ -158,6 +158,7 @@ int32 UCheckAndroidDeviceProfileCommandlet::Main(const FString& RawCommandLine)
 				DeviceParameters.Add(FName(TEXT("SRC_InsetsTop")), Params.FindRef(TEXT("InsetsTop")));
 				DeviceParameters.Add(FName(TEXT("SRC_InsetsRight")), Params.FindRef(TEXT("InsetsRight")));
 				DeviceParameters.Add(FName(TEXT("SRC_InsetsBottom")), Params.FindRef(TEXT("InsetsBottom")));
+				DeviceParameters.Add(FName(TEXT("SRC_VKQuality")), Params.FindRef(TEXT("VKQuality")));
 			}
 			AndroidDeviceProfileSelector->SetSelectorProperties(DeviceParameters);
 			FString ProfileName = ForcedDeviceProfileName.IsEmpty() ? AndroidDeviceProfileSelector->GetDeviceProfileName() : ForcedDeviceProfileName;
diff --git a/Engine/Source/Editor/PIEPreviewDeviceProfileSelector/Private/PIEPreviewDevice.cpp b/Engine/Source/Editor/PIEPreviewDeviceProfileSelector/Private/PIEPreviewDevice.cpp
--- a/Engine/Source/Editor/PIEPreviewDeviceProfileSelector/Private/PIEPreviewDevice.cpp
+++ b/Engine/Source/Editor/PIEPreviewDeviceProfileSelector/Private/PIEPreviewDevice.cpp
@@ -472,6 +472,7 @@ FString FPIEPreviewDevice::GetProfile() const
 				DeviceParameters.Add(FName(TEXT("SRC_HMDSystemName")), AndroidProperties.HMDSystemName);
 				DeviceParameters.Add(FName(TEXT("SRC_TotalPhysicalGB")), AndroidProperties.TotalPhysicalGB);
 				DeviceParameters.Add(FName(TEXT("SRC_SM5Available")), AndroidProperties.SM5Available ? "true" : "false");
+				DeviceParameters.Add(FName(TEXT("SRC_VKQuality")), TEXT(""));
 
 				AndroidDeviceProfileSelector->SetSelectorProperties(DeviceParameters);
 				FString PIEProfileName = AndroidDeviceProfileSelector->GetDeviceProfileName();
diff --git a/Engine/Source/Runtime/Core/Private/Android/AndroidPlatformMisc.cpp b/Engine/Source/Runtime/Core/Private/Android/AndroidPlatformMisc.cpp
--- a/Engine/Source/Runtime/Core/Private/Android/AndroidPlatformMisc.cpp
+++ b/Engine/Source/Runtime/Core/Private/Android/AndroidPlatformMisc.cpp
@@ -64,6 +64,8 @@
 #include <vulkan/vulkan_android.h>
 #include "Android/AndroidPlatformProperties.h"
 
+#include "AndroidVKQuality.h"
+
 #if USE_ANDROID_JNI
 extern AAssetManager * AndroidThunkCpp_GetAssetManager();
 extern int32 GAndroidPackageVersion;
@@ -2031,6 +2033,9 @@ namespace AndroidPlatformMisc
 
 #define UE_VK_API_VERSION	VK_MAKE_VERSION(1, 1, 0)
 
+// Only valid if EDeviceVulkanSupportStatus is EDeviceVulkanSupportStatus::Supported
+VkPhysicalDeviceProperties VulkanPhysicalDeviceProperties;
+
 static EDeviceVulkanSupportStatus AttemptVulkanInit(void* VulkanLib)
 {
 	if (VulkanLib == nullptr)
@@ -2103,8 +2108,7 @@ static EDeviceVulkanSupportStatus AttemptVulkanInit(void* VulkanLib)
 	}
 
 	// Don't care which device - This code is making the assumption that all devices will have same api version.
-	VkPhysicalDeviceProperties DeviceProperties;
-	vkGetPhysicalDeviceProperties(PhysicalDevices[0], &DeviceProperties);
+	vkGetPhysicalDeviceProperties(PhysicalDevices[0], &VulkanPhysicalDeviceProperties);
 
 	//for now we are allowing devices without the timing extension to run with a basic CPU frame pacer.
 #if 0
@@ -2150,7 +2154,7 @@ static EDeviceVulkanSupportStatus AttemptVulkanInit(void* VulkanLib)
 	}
 #endif
 
-	VulkanVersionString = FString::Printf(TEXT("%d.%d.%d"), VK_VERSION_MAJOR(DeviceProperties.apiVersion), VK_VERSION_MINOR(DeviceProperties.apiVersion), VK_VERSION_PATCH(DeviceProperties.apiVersion));
+	VulkanVersionString = FString::Printf(TEXT("%d.%d.%d"), VK_VERSION_MAJOR(VulkanPhysicalDeviceProperties.apiVersion), VK_VERSION_MINOR(VulkanPhysicalDeviceProperties.apiVersion), VK_VERSION_PATCH(VulkanPhysicalDeviceProperties.apiVersion));
 	vkDestroyInstance(Instance, nullptr);
 
 	return EDeviceVulkanSupportStatus::Supported;
@@ -2161,15 +2165,12 @@ static EDeviceVulkanSupportStatus AttemptVulkanInit(void* VulkanLib)
 
 bool FAndroidMisc::HasVulkanDriverSupport()
 {
-// @todo Lumin: this isn't really the best #define to check here - but basically, without JNI and other version checking, we can't safely do it - we'll need
-// non-JNI platforms that support Vulkan to figure out a way to force it (if they want GL + Vulkan support)
 #if !USE_ANDROID_JNI
 	VulkanSupport = EDeviceVulkanSupportStatus::NotSupported;
 	VulkanVersionString = TEXT("0.0.0");
 #else
 	// this version does not check for VulkanRHI or disabled by cvars!
 	if (VulkanSupport == EDeviceVulkanSupportStatus::Uninitialized)
-
 	{
 		// assume no
 		VulkanSupport = EDeviceVulkanSupportStatus::NotSupported;
@@ -2269,7 +2270,6 @@ bool FAndroidMisc::IsVulkanAvailable()
 			GConfig->GetBool(TEXT("/Script/AndroidRuntimeSettings.AndroidRuntimeSettings"), TEXT("bDetectVulkanByDefault"), bDetectVulkanByDefault, GEngineIni);
 			const bool bDetectVulkanCmdLine = FParse::Param(FCommandLine::Get(), TEXT("detectvulkan"));
 
-			// @todo Lumin: Double check all this stuff after merging general android Vulkan SM5 from main
 			const bool bSupportsVulkanSM5 = IsDesktopVulkanAvailable();
 
 			const bool bVulkanDisabledCmdLine = FParse::Param(FCommandLine::Get(), TEXT("GL")) || FParse::Param(FCommandLine::Get(), TEXT("OpenGL"));
@@ -2375,6 +2375,30 @@ FString FAndroidMisc::GetVulkanVersion()
 	return VulkanVersionString;
 }
 
+FString FAndroidMisc::GetVKQualityRecommendation()
+{
+	static bool bCachedResult = false;
+	static FString Result;
+
+#if USE_ANDROID_JNI
+	if (bCachedResult)
+	{
+		return Result;
+	}
+
+	if (VulkanSupport == EDeviceVulkanSupportStatus::Supported)
+	{
+		if (AndroidVKQuality::LoadVkQuality(TCHAR_TO_UTF8(*FAndroidMisc::GetGLVersion()), &AndroidPlatformMisc::VulkanPhysicalDeviceProperties))
+		{
+			Result = AndroidVKQuality::GetVKQualityRecommendation();
+			bCachedResult = true;
+			AndroidVKQuality::UnloadVkQuality();
+		}
+	}
+#endif
+	return Result;
+}
+
 bool FAndroidMisc::IsExternalMemoryAndroidHardwareBufferExtensionLoaded()
 {
 	if (GIsRHIInitialized && ShouldUseVulkan())
diff --git a/Engine/Source/Runtime/Core/Public/Android/AndroidPlatformMisc.h b/Engine/Source/Runtime/Core/Public/Android/AndroidPlatformMisc.h
--- a/Engine/Source/Runtime/Core/Public/Android/AndroidPlatformMisc.h
+++ b/Engine/Source/Runtime/Core/Public/Android/AndroidPlatformMisc.h
@@ -218,6 +218,8 @@ public:
 	static CORE_API bool ShouldUseVulkan();
 	static CORE_API bool ShouldUseDesktopVulkan();
 	static CORE_API FString GetVulkanVersion();
+	static CORE_API FString GetVKQualityRecommendation();
+
 	static CORE_API bool IsExternalMemoryAndroidHardwareBufferExtensionLoaded();
 	typedef TFunction<void(void* NewNativeHandle)> ReInitWindowCallbackType;
 	static CORE_API ReInitWindowCallbackType GetOnReInitWindowCallback();
