[TasksTest]
- Fix potential race since IsCompleted() and Task.Wait() on a task can return before ClearPipe is called and the pipe task count is 0

diff --git a/Engine/Source/Runtime/Core/Tests/Tasks/TasksTest.cpp b/Engine/Source/Runtime/Core/Tests/Tasks/TasksTest.cpp
--- a/Engine/Source/Runtime/Core/Tests/Tasks/TasksTest.cpp
+++ b/Engine/Source/Runtime/Core/Tests/Tasks/TasksTest.cpp
@@ -325,6 +325,7 @@ namespace UE { namespace TasksTests
 				verify(!Task.Wait(FTimespan::FromMilliseconds(100)));
 			}
 			Task.Wait();
+			Pipe.WaitUntilEmpty();
 		}
 
 		{	// test basic functionality of TTask::IsAwaitable()
@@ -475,6 +476,7 @@ namespace UE { namespace TasksTests
 			Tasks::FTask Task1 = Pipe.Launch(UE_SOURCE_LOCATION, [] {});
 			Tasks::FTask Task2 = Pipe.Launch(UE_SOURCE_LOCATION, [] {});
 			Task2.Wait(); // wait for `Task2` completion
+			Pipe.WaitUntilEmpty();
 		}
 
 		{	// an example of thread-safe async interface, kind of a primitive "actor"
@@ -509,11 +511,13 @@ namespace UE { namespace TasksTests
 			FPipe Pipe{ UE_SOURCE_LOCATION };
 			Pipe.Launch(UE_SOURCE_LOCATION, [] {});
 			Pipe.Launch(UE_SOURCE_LOCATION, [] {}).Wait();
+			Pipe.WaitUntilEmpty();
 		}
 
 		{	// launching a piped task with pointer to a function
 			FPipe Pipe{ UE_SOURCE_LOCATION };
 			Pipe.Launch(UE_SOURCE_LOCATION, &DummyFunc).Wait();
+			Pipe.WaitUntilEmpty();
 		}
 
 		{	// launching a piped task with a functor object
@@ -526,6 +530,7 @@ namespace UE { namespace TasksTests
 
 			FPipe Pipe{ UE_SOURCE_LOCATION };
 			Pipe.Launch(UE_SOURCE_LOCATION, FFunctor{}).Wait();
+			Pipe.WaitUntilEmpty();
 		}
 
 		{	// hold the first piped task execution until the next one is piped to test for non-concurrent execution
@@ -541,6 +546,7 @@ namespace UE { namespace TasksTests
 			// we can't just check if `Task1` is completed because pipe gets unblocked and so the next piped task can start execution before the
 			// previous piped task's comlpetion flag is set
 			Pipe.Launch(UE_SOURCE_LOCATION, [&bTask1Done] { check(bTask1Done); }).Wait();
+			Pipe.WaitUntilEmpty();
 		}
 
 		{	// piping another task after the previous one is completed and destroyed
@@ -548,6 +554,7 @@ namespace UE { namespace TasksTests
 
 			Pipe.Launch(UE_SOURCE_LOCATION, [] {}).Wait();
 			Pipe.Launch(UE_SOURCE_LOCATION, [] {}).Wait();
+			Pipe.WaitUntilEmpty();
 		}
 
 		{	// an example of blocking a pipe
@@ -569,6 +576,7 @@ namespace UE { namespace TasksTests
 
 			Event.Trigger(); // unblock
 			Task.Wait();
+			Pipe.WaitUntilEmpty();
 		}
 
 		{	// multiple "inline" tasks piped in one go, check that pipe maintains FIFO order in this case
@@ -581,6 +589,7 @@ namespace UE { namespace TasksTests
 			Block.Trigger();
 			Wait(TArray{ Task1, Task2 });
 			check(bFirstDone && bSecondDone);
+			Pipe.WaitUntilEmpty();
 		}
 
 		UE_BENCHMARK(5, PipeStressTest<200, 100>);
@@ -647,6 +656,7 @@ namespace UE { namespace TasksTests
 		Wait(Tasks);
 
 		check(TasksExecutedNum == TasksNum);
+		Pipe.WaitUntilEmpty();
 	}
 
 	TEST_CASE_NAMED(FTasksPipeWaitUntilEmptyTest, "System::Core::Tasks::Pipe::WaitUntilEmpty", "[ApplicationContextMask][EngineFilte]")
@@ -903,6 +913,8 @@ namespace UE { namespace TasksTests
 
 			Prereq.Trigger();
 			Task1.Wait();
+
+			Pipe.WaitUntilEmpty();
 		}
 
 		{	// a piped task with multiple prerequisites
@@ -921,6 +933,8 @@ namespace UE { namespace TasksTests
 
 			Event.Trigger();
 			Task.Wait();
+
+			Pipe.WaitUntilEmpty();
 		}
 
 		{	// a piped task with multiple prerequisites. with explicit task priority
@@ -939,6 +953,8 @@ namespace UE { namespace TasksTests
 
 			Event.Trigger();
 			Task.Wait();
+
+			Pipe.WaitUntilEmpty();
 		}
 
 		UE_BENCHMARK(5, DependenciesPerfTest<150, 150>);
@@ -1083,6 +1099,7 @@ namespace UE { namespace TasksTests
 			FTask PipedTask1 = Pipe.Launch(UE_SOURCE_LOCATION, [] {});
 			FTask PipedTask2 = Pipe.Launch(UE_SOURCE_LOCATION, [] {});
 			PipedTask2.Wait(); // the pipe is not blocked, deep retraction first executes `PipedTask1`, and then `PipedTask2`
+			Pipe.WaitUntilEmpty();
 		}
 
 		// an example of a deadlock caused by waiting for a piped task from inside execution of another task of the same pipe
@@ -1381,6 +1398,7 @@ namespace UE { namespace TasksTests
 		Pipe.Launch(UE_SOURCE_LOCATION, [] {});
 		FTask FinalTask = Pipe.Launch(UE_SOURCE_LOCATION, [] {});
 		check(FinalTask.Wait(FTimespan::FromMilliseconds(100)));
+		Pipe.WaitUntilEmpty();
 
 	}
 
