Disabled PVS warning clean up

Enabled V530. Return value of 'Foo' function is required to be used.
Enabled V581. Conditional expressions of 'if' statements located next to each other are identical.
Enabled V601. Suspicious implicit type casting.
Enabled V623. Temporary object is created and then destroyed. Consider inspecting the '?:' operator.
Enabled V668. Possible meaningless check for null, as memory was allocated using 'new' operator. Memory allocation will lead to an exception.
Enabled V690. The class implements a copy constructor/operator=, but lacks the operator=/copy constructor.
Enabled V751. Parameter is not used inside function's body.
Enabled V758. Reference was invalidated because of destruction of the temporary object returned by the function.
Enabled V781. Value of a variable is checked after it is used. Possible error in program's logic. Check lines: N1, N2.
Enabled V1008. No more than one iteration of the loop will be performed. Consider inspecting the 'for' operator.
Enabled V1043. A global object variable is declared in the header. Multiple copies of it will be created in all translation units that include this header file.
Enabled V1051. It is possible that an assigned variable should be checked in the next condition. Consider checking for typos.
Remove large number of disabled warnings for specific keywords. Some cases moved to be individual line disabling instead of global, most simply removed.
Added small number of new false positive global exceptions related to newly enabled warnings

diff --git a/Engine/Source/Runtime/Core/Public/Microsoft/MicrosoftPlatformCodeAnalysis.h b/Engine/Source/Runtime/Core/Public/Microsoft/MicrosoftPlatformCodeAnalysis.h
--- a/Engine/Source/Runtime/Core/Public/Microsoft/MicrosoftPlatformCodeAnalysis.h
+++ b/Engine/Source/Runtime/Core/Public/Microsoft/MicrosoftPlatformCodeAnalysis.h
@@ -112,7 +112,7 @@
 	#pragma warning(disable : 6393) // warning C6393: A lookup table of size 365 is not sufficient to handle leap years.
 
 	//PVS-Studio settings:
-	//-V::505,542,581,601,623,668,677,690,704,719,720,730,735,751,758,780,781,1002,1008,1055,1062,1100
+	//-V::505,542,677,704,719,720,730,735,780,1002,1055,1062,1100
 	//-V:TRYCOMPRESSION:519,547
 	//-V:check(:501,547,560,605
 	//-V:checkSlow(:547
@@ -126,104 +126,34 @@
 	//-V:WITH_EDITOR:560
 	//-V:UE_LOG_ACTIVE:560
 	//-V:verify:501
-	//-V:%n:609
 	//-V:UE_BUILD_SHIPPING:501
 	//-V:WITH_EDITOR:501
 	//-V:TestTrueExpr:501
 	//-V:PLATFORM_:517,547
 	//-V:ensureMsgf:562
-	//-V:WindowsMinorVersion:547
-	//-V:Import.XObject:547,560
-	//-V:MotionControllerComponent:547,560
-	//-V:AddUninitialized(sizeof(void*)/:514
-	//-V:TestTrue:678
-	//-V:SetViewTarget:678
-	//-V:Slot:607
-	//-V:RESIDENCY_CHECK_RESULT:607
-	//-V:bHitTesting:581
-	//-V:OptionalType:580 
-	//-V:GetNextNode:681
-	//-V:ConvertToAbsolutePathForExternalAppFor:524
-	//-V:CopySingleValue:524
-	//-V:bTimeLimitReached:560
-	//-V:bRedirectionAllowed:560
-	//-V:NumFailures:560
-	//-V:bAllowInstantToolTips:560
-	//-V:bIsRealTime:560
-	//-V:Position:519
-	//-V:DynamicParameterValue[ParameterIndex]:557
-	//-V:ViewIndex:557
-	//-V:DeviceIndex:557
-	//-V:Interpolation:560
-	//-V:storePortals:560
-	//-V:bDefaultShouldBeMaximized:560
-	//-V:bAllowPerfHUD:560
-	//-V:bUseClientStorage:560
-	//-V:bCalculateThisMapping:560
-	//-V:bDebugSelectedTaskOnly:560
-	//-V:bDebugSelectedTaskOnly:560
-	//-V:bIsPreview:560
-	//-V:bSupportsFastClear:560
-	//-V:bUseAPILibaries:560
-	//-V:bUseCachedBlobs:560
-	//-V:bWireframe:560
-	//-V:Num():560
-	//-V:PLATFORM_MAC:560
-	//-V:Particle->Size.Z:570
-	//-V:ComponentMaskParameter:601
-	//-V:Format(:601
 	//-V:SelectedEmitter:519
 	//-V:MAX_VERTS_PER_POLY:512
 	//-V:127:547
 	//-V:0x7F:547
 	//-V:WARN_COLOR:547
 	//-V:<<:614
-	//-V:FT_LOAD_TARGET_NORMAL:616
 	//-V:OPENGL_PERFORMANCE_DATA_INVALID:564
-	//-V:HLSLCC_VersionMajor:616
-	//-V:bIgnoreFieldReferences:519
-	//-V:CachedQueryInstance:519
-	//-V:MeshContext:519
-	//-V:bAffectedByMarquee:519
-	//-V:CopyCompleteValueFromScriptVM:524
-	//-V:OnStopWatchingPin:524
-	//-V:GetMinChildNodes:524
-	//-V:FromWorldMatrix:524
-	//-V:RemoveSelectedActorsFromSelectedLayer_CanExecute:524
-	//-V:NotifyLevelRemovedFromWorld:524
 	//-V:SPAWN_INIT:595
 	//-V:BEGIN_UPDATE_LOOP:595
 	//-V:OPENGL_PERFORMANCE_DATA_INVALID:560
-	//-V:bSkipTranslationTrack:560
-	//-V:NumSelected>0:581
-	//-V:bTryPerTrackBitwiseCompression:581
-	//-V:DataStripped:581
-	//-V:FromInt:601
 	//-V:UE_CLOG(:501,560
 	//-V:UE_LOG(:501,510,560
 	//-V:UGL_REQUIRED_VOID:501
-	//-V:AnimScriptInstance:595
-	//-V:Driver:595
-	//-V:PSceneAsync->lockWrite:595
-	//-V:Context.World():595
-	//-V:UNIT_LOG:595
 	//-V:ensure(:595
 	//-V:ALLOCATE_VERTEX_DATA_TEMPLATE:501
 	//-V:UGL_REQUIRED:501
 	//-V:DEBUG_LOG_HTTP:523
 	//-V:GIsEditor:560
-	//-V:bHasEditorToken:560
 	//-V:GEventDrivenLoaderEnabled:501
-	//-V:WALK_TO_CHARACTER:519
 	//-V:IMPLEMENT_AI_INSTANT_TEST:773
 	//-V:ENABLE_VERIFY_GL:564
 	//-V:INC_MEMORY_STAT_BY:568
 	//-V:DEC_MEMORY_STAT_BY:568
-	//-V:Key():568
-	//-V:Modify:762
-	//-V:GetTransitionList:762
-	//-V:Execute:768
-	//-V:LAUNCHERSERVICES_SHAREABLEPROJECTPATHS:768
 	//-V:SELECT_STATIC_MESH_VERTEX_TYPE:622
 	//-V:GET_FUNCTION_NAME_CHECKED:685
 	//-V:This(:678
@@ -239,8 +169,12 @@
 	//-V:VertexData:773
 	//-V:Linker:678
 	//-V:self:678
-	//-V:AccumulateParentID:678
 	//-V:FindChar:679
+	//-V:BytesToHex:530
+	//-V:FindOrAdd:530
+	//-V:GetPawn:623
+	//-V:FindComponentByClass:623
+	//-V:Response.Response:1051
 
 	// The following classes retain a reference to data supplied in the constructor by the derived class which can not yet be initialized.
 	//-V:FMemoryWriter(:1050
@@ -252,18 +186,9 @@
 	// Exclude all generated protobuf files
 	//V_EXCLUDE_PATH *.pb.cc
 
-	// warning V530: The return value of function 'HashCombine' is required to be utilized.
-	//-V::530
-
 	// warning V630: Instantiation of TRingBuffer < FString >: The 'Malloc' function is used to allocate memory for an array of objects which are classes containing constructors and destructors.
 	//-V::630
 
-	// warning V1043: A global object variable 'GSmallNumber' is declared in the header. Multiple copies of it will be created in all translation units that include this header file.
-	//-V::1043
-
-	// warning V1051: Consider checking for misprints. It's possible that the 'LayerInfo' should be checked here.
-	//-V::1051
-
 	// Disabling because incorrectly flagging all TStaticArrays
 	// V557: Array overrun is possible
 	//-V::557
diff --git a/Engine/Source/Runtime/Engine/Private/CheatManager.cpp b/Engine/Source/Runtime/Engine/Private/CheatManager.cpp
--- a/Engine/Source/Runtime/Engine/Private/CheatManager.cpp
+++ b/Engine/Source/Runtime/Engine/Private/CheatManager.cpp
@@ -510,7 +510,7 @@ void UCheatManager::ViewSelf()
 	}
 	else
 	{
-		GetOuterAPlayerController()->SetViewTarget(GetOuterAPlayerController());
+		GetOuterAPlayerController()->SetViewTarget(GetOuterAPlayerController()); //-V678
 	}
 	GetOuterAPlayerController()->ClientMessage(LOCTEXT("OwnCamera", "Viewing from own camera").ToString(), TEXT("Event"));
 }
diff --git a/Engine/Source/Runtime/Engine/Private/GameMode.cpp b/Engine/Source/Runtime/Engine/Private/GameMode.cpp
--- a/Engine/Source/Runtime/Engine/Private/GameMode.cpp
+++ b/Engine/Source/Runtime/Engine/Private/GameMode.cpp
@@ -482,7 +482,7 @@ void AGameMode::InitSeamlessTravelPlayer(AController* NewController)
 
 void AGameMode::SetSeamlessTravelViewTarget(APlayerController* PC)
 {
-	PC->SetViewTarget(PC);
+	PC->SetViewTarget(PC); //-V678
 }
 
 void AGameMode::PlayerSwitchedToSpectatorOnly(APlayerController* PC)
diff --git a/Engine/Source/Runtime/Engine/Private/UnrealEngine.cpp b/Engine/Source/Runtime/Engine/Private/UnrealEngine.cpp
--- a/Engine/Source/Runtime/Engine/Private/UnrealEngine.cpp
+++ b/Engine/Source/Runtime/Engine/Private/UnrealEngine.cpp
@@ -15500,7 +15500,7 @@ void UEngine::TickWorldTravel(FWorldContext& Context, float DeltaSeconds)
 
 	if( !Context.World()->NextURL.IsEmpty() )
 	{
-		Context.World()->NextSwitchCountdown -= DeltaSeconds;
+		Context.World()->NextSwitchCountdown -= DeltaSeconds; //-V595
 		if( Context.World()->NextSwitchCountdown <= 0.f )
 		{
 			UE_LOG(LogEngine, Log,  TEXT("Server switch level: %s"), *Context.World()->NextURL );
@@ -15510,7 +15510,7 @@ void UEngine::TickWorldTravel(FWorldContext& Context, float DeltaSeconds)
 			}
 			FString Error;
 			FString NextURL = Context.World()->NextURL;
-			EBrowseReturnVal::Type Ret = Browse( Context, FURL(&Context.LastURL,*NextURL,(ETravelType)Context.World()->NextTravelType), Error );
+			EBrowseReturnVal::Type Ret = Browse( Context, FURL(&Context.LastURL,*NextURL,(ETravelType)Context.World()->NextTravelType), Error ); //-V595
 			if (Ret != EBrowseReturnVal::Success )
 			{
 				UE_LOG(LogLoad, Warning, TEXT("UEngine::TickWorldTravel failed to Handle server travel to URL: %s. Error: %s"), *NextURL, *Error);
@@ -15538,7 +15538,7 @@ void UEngine::TickWorldTravel(FWorldContext& Context, float DeltaSeconds)
 	// Handle client traveling.
 	if( !Context.TravelURL.IsEmpty() )
 	{	
-		AGameModeBase* const GameMode = Context.World()->GetAuthGameMode();
+		AGameModeBase* const GameMode = Context.World()->GetAuthGameMode(); //-V595
 		if (GameMode)
 		{
 			GameMode->StartToLeaveMap();
diff --git a/Engine/Source/Runtime/Experimental/GeometryCollectionEngine/Private/GeometryCollection/GeometryCollectionEngineConversion.cpp b/Engine/Source/Runtime/Experimental/GeometryCollectionEngine/Private/GeometryCollection/GeometryCollectionEngineConversion.cpp
--- a/Engine/Source/Runtime/Experimental/GeometryCollectionEngine/Private/GeometryCollection/GeometryCollectionEngineConversion.cpp
+++ b/Engine/Source/Runtime/Experimental/GeometryCollectionEngine/Private/GeometryCollection/GeometryCollectionEngineConversion.cpp
@@ -1061,7 +1061,7 @@ void FGeometryCollectionEngineConversion::AppendStaticMesh(const UStaticMesh* St
 	GeometryCollectionObject->Materials.Remove(BoneSelectedMaterial);
 	Materials.Remove(BoneSelectedMaterial);
 	
-	AppendStaticMesh(StaticMesh, Materials, StaticMeshTransform, GeometryCollectionObject, ReindexMaterials, bAddInternalMaterials, bSplitComponents, bSetInternalFromMaterialIndex);
+	(void)AppendStaticMesh(StaticMesh, Materials, StaticMeshTransform, GeometryCollectionObject, ReindexMaterials, bAddInternalMaterials, bSplitComponents, bSetInternalFromMaterialIndex);
 }
 
 
@@ -1571,7 +1571,7 @@ bool FGeometryCollectionEngineConversion::AppendGeometryCollectionSourceNoMateri
 		if (const UStaticMesh* SourceStaticMesh = Cast<UStaticMesh>(SourceObject))
 		{
 			bool bLegacyAddInternal = GeometryCollectionSource.bAddInternalMaterials;
-			AppendStaticMesh(
+			(void)AppendStaticMesh(
 				SourceStaticMesh,
 				StartMaterialIndex,
 				GeometryCollectionSource.LocalTransform,
@@ -1585,7 +1585,7 @@ bool FGeometryCollectionEngineConversion::AppendGeometryCollectionSourceNoMateri
 		}
 		else if (const USkeletalMesh* SourceSkeletalMesh = Cast<USkeletalMesh>(SourceObject))
 		{
-			AppendSkeletalMesh(
+			(void)AppendSkeletalMesh(
 				SourceSkeletalMesh,
 				StartMaterialIndex,
 				GeometryCollectionSource.LocalTransform,
diff --git a/Engine/Source/Runtime/Navmesh/Private/DetourTileCache/DetourTileCacheBuilder.cpp b/Engine/Source/Runtime/Navmesh/Private/DetourTileCache/DetourTileCacheBuilder.cpp
--- a/Engine/Source/Runtime/Navmesh/Private/DetourTileCache/DetourTileCacheBuilder.cpp
+++ b/Engine/Source/Runtime/Navmesh/Private/DetourTileCache/DetourTileCacheBuilder.cpp
@@ -1787,10 +1787,10 @@ static void mergePolys(unsigned short* pa, unsigned short* pb, int ea, int eb)
 	int n = 0;
 	// Add pa
 	for (int i = 0; i < na-1; ++i)
-		tmp[n++] = pa[(ea+1+i) % na];
+		tmp[n++] = pa[(ea+1+i) % na]; //-V609
 	// Add pb
 	for (int i = 0; i < nb-1; ++i)
-		tmp[n++] = pb[(eb+1+i) % nb];
+		tmp[n++] = pb[(eb+1+i) % nb]; //-V609
 	
 	memcpy(pa, tmp, sizeof(unsigned short)*MAX_VERTS_PER_POLY);
 }
diff --git a/Engine/Source/Runtime/SlateCore/Private/Fonts/SlateFontRenderer.cpp b/Engine/Source/Runtime/SlateCore/Private/Fonts/SlateFontRenderer.cpp
--- a/Engine/Source/Runtime/SlateCore/Private/Fonts/SlateFontRenderer.cpp
+++ b/Engine/Source/Runtime/SlateCore/Private/Fonts/SlateFontRenderer.cpp
@@ -64,7 +64,7 @@ void AppendGlyphFlags(const FFreeTypeFace& InFace, const FFontData& InFontData,
 		case EFontHinting::Monochrome:	InOutGlyphFlags |= FT_LOAD_TARGET_MONO | FT_LOAD_FORCE_AUTOHINT; break;
 		case EFontHinting::None:		InOutGlyphFlags |= FT_LOAD_NO_AUTOHINT | FT_LOAD_NO_HINTING; break;
 		case EFontHinting::Default:
-		default:						InOutGlyphFlags |= FT_LOAD_TARGET_NORMAL; break;
+		default:						InOutGlyphFlags |= FT_LOAD_TARGET_NORMAL; break; //-V616
 		}
 	}
 	else
