Fix the build failure by making `FUnixPlatformFile::FileExists` always-open.

diff --git a/Engine/Source/Runtime/Core/Public/Unix/UnixPlatformFile.h b/Engine/Source/Runtime/Core/Public/Unix/UnixPlatformFile.h
--- a/Engine/Source/Runtime/Core/Public/Unix/UnixPlatformFile.h
+++ b/Engine/Source/Runtime/Core/Public/Unix/UnixPlatformFile.h
@@ -27,7 +27,7 @@ public:
 	using IPhysicalPlatformFile::IterateDirectory;
 	using IPhysicalPlatformFile::IterateDirectoryStat;
 
-	CORE_API virtual bool FileExists(const TCHAR* Filename) override;
+	AUTORTFM_OPEN CORE_API virtual bool FileExists(const TCHAR* Filename) override;
 	CORE_API virtual int64 FileSize(const TCHAR* Filename) override;
 	CORE_API virtual bool DeleteFile(const TCHAR* Filename) override;
 	CORE_API virtual bool IsReadOnly(const TCHAR* Filename) override;
