Unified error replaces spaces from error codes with underscores.

diff --git a/Engine/Source/Runtime/Core/Private/Experimental/IO/IoStatusError.cpp b/Engine/Source/Runtime/Core/Private/Experimental/IO/IoStatusError.cpp
--- a/Engine/Source/Runtime/Core/Private/Experimental/IO/IoStatusError.cpp
+++ b/Engine/Source/Runtime/Core/Private/Experimental/IO/IoStatusError.cpp
@@ -7,8 +7,16 @@
 #include "Internationalization/Text.h"
 
 
+UE_DEFINE_ERROR_MODULE(IoStore);
+
 namespace UE::UnifiedError::IoStore
 {
+
+	void StripInvalidErrorCodeCharacters(FUtf8String& ErrorCode)
+	{
+		ErrorCode.ReplaceCharInline(UTF8TEXT(' '), UTF8TEXT('_'));
+	}
+
 	class FIoStoreErrorDetails : public UE::UnifiedError::FManditoryErrorDetails, public FRefCountBase
 	{
 	private:
@@ -28,9 +36,10 @@ namespace UE::UnifiedError::IoStore
 
 		const FUtf8String& GetErrorName(const FError& Error) const
 		{
-			if (CachedErrorName.IsEmpty() && (ErrorMessage[0] != TEXT('\0')))
+			if (CachedErrorName.IsEmpty())
 			{
 				CachedErrorName = TCHAR_TO_UTF8(GetIoErrorText(EIoErrorCode(Error.GetErrorCode())));
+				StripInvalidErrorCodeCharacters(CachedErrorName);
 			}
 			return CachedErrorName;
 		}
@@ -42,7 +51,7 @@ namespace UE::UnifiedError::IoStore
 
 		FAnsiString GetModuleIdString(const UE::UnifiedError::FError& Error) const override
 		{
-			return FAnsiString(UE::UnifiedError::IoStore::StaticModuleName);
+			return FAnsiString(UE::UnifiedError::IoStore::GetStaticModuleName());
 		}
 
 		virtual const FText GetErrorFormatString(const FError& Error) const final override
@@ -55,8 +64,8 @@ namespace UE::UnifiedError::IoStore
 			Writer.BeginObject();
 			if (Error.GetModuleId() == UE::UnifiedError::IoStore::StaticModuleId)
 			{
-				Writer.AddString(UTF8TEXTVIEW("ErrorCodeString"), GetIoErrorText(EIoErrorCode(Error.GetErrorCode())));
-				Writer.AddString(UTF8TEXTVIEW("ModuleIdString"), UE::UnifiedError::IoStore::StaticModuleName);
+				Writer.AddString(UTF8TEXTVIEW("ErrorCodeString"), GetErrorName(Error));
+				Writer.AddString(UTF8TEXTVIEW("ModuleIdString"), UE::UnifiedError::IoStore::GetStaticModuleName());
 				Writer.AddInteger(UTF8TEXTVIEW("ErrorCode"), Error.GetErrorCode());
 				Writer.AddInteger(UTF8TEXTVIEW("ModuleId"), Error.GetModuleId());
 			}
diff --git a/Engine/Source/Runtime/Core/Private/Experimental/UnifiedError/UnifiedError.cpp b/Engine/Source/Runtime/Core/Private/Experimental/UnifiedError/UnifiedError.cpp
--- a/Engine/Source/Runtime/Core/Private/Experimental/UnifiedError/UnifiedError.cpp
+++ b/Engine/Source/Runtime/Core/Private/Experimental/UnifiedError/UnifiedError.cpp
@@ -169,7 +169,7 @@ int32 FError::GetModuleId() const
 
 class FAppendFormatStringDetails : public FDynamicErrorDetails
 {
-	DECLARE_FERROR_DETAILS(UnifiedError, FAppendFormatStringDetails);
+	UE_DECLARE_FERROR_DETAILS(UnifiedError, FAppendFormatStringDetails);
 public:
 	FAppendFormatStringDetails() {}
 	FAppendFormatStringDetails(FText&& InAppendFormatString, TRefCountPtr<const IErrorDetails> InInnerErrorDetails = nullptr) : FDynamicErrorDetails(InInnerErrorDetails)
@@ -269,7 +269,7 @@ void FDynamicErrorDetails::SetInnerErrorDetails(TRefCountPtr<const IErrorDetails
 
 // FStaticErrorDetails functions
 
-FStaticErrorDetails::FStaticErrorDetails(const FAnsiStringView& InErrorName, const FAnsiStringView& InModuleName, const FText& InErrorFormatString) : ErrorName(InErrorName), ModuleName(InModuleName), ErrorFormatString(InErrorFormatString)
+FStaticErrorDetails::FStaticErrorDetails(const FAnsiStringView InErrorName, const FAnsiStringView InModuleName, const FText& InErrorFormatString) : ErrorName(InErrorName), ModuleName(InModuleName), ErrorFormatString(InErrorFormatString)
 {
 }
 
@@ -296,6 +296,60 @@ uint32 FErrorDetailsRegistry::RegisterDetails(const FAnsiStringView& ErrorDetail
 	return DetailsId;
 }
 
+// ErrorRegistry functionality
+
+namespace ErrorRegistry
+{
+	class FErrorRegistry
+	{
+	public:
+		FErrorRegistry() {}
+
+	public:
+		uint32 RegisterModule(const FAnsiStringView ModuleName);
+
+		int32 RegisterErrorCode(const FAnsiStringView ErrorName, int32 ModuleId, int32 ErrorCode);
+
+	private:
+		TMap<int32, FAnsiString> ModuleNameMap;
+		TMap<TPair<int32, int32>, FAnsiString> ErrorCodeNameMap;
+	};
+
+	
+	uint32 FErrorRegistry::RegisterModule(const FAnsiStringView ModuleName)
+	{
+		// todo: need to replace this with a stable hashing function
+		uint32 ModuleId = GetTypeHash(ModuleName);
+		checkf(ModuleNameMap.Contains(ModuleId) == false, TEXT("Module %s and %s are trying to register under module id %d"), ANSI_TO_TCHAR(ModuleName.GetData()), ANSI_TO_TCHAR(*ModuleNameMap.FindRef(ModuleId)), ModuleId);
+		ModuleNameMap.Add(ModuleId, ModuleName.GetData());
+		return ModuleId;
+	}
+
+	int32 FErrorRegistry::RegisterErrorCode(const FAnsiStringView ErrorName, int32 ModuleId, int32 ErrorCode)
+	{
+		TPair<int32, int32> CombinedErrorId(ModuleId, ErrorCode);
+		checkf(ErrorCodeNameMap.Contains(CombinedErrorId) == false, TEXT("Error %s and %s are trying to register under same error code moduleid:%d errorcode:%d"), ANSI_TO_TCHAR(ErrorName.GetData()), ANSI_TO_TCHAR(*ErrorCodeNameMap.FindRef(CombinedErrorId)), ModuleId, ErrorCode);
+		ErrorCodeNameMap.Add(CombinedErrorId, ErrorName.GetData());
+		return ErrorCode;
+	}
+
+	FErrorRegistry GErrorRegistry;
+
+	CORE_API uint32 RegisterModule(const FAnsiStringView ModuleName)
+	{
+		return GetTypeHash(ModuleName);
+		//return GErrorRegistry.RegisterModule(ModuleName);
+	}
+
+	CORE_API int32 RegisterErrorCode(const FAnsiStringView ErrorName, int32 ModuleId, int32 ErrorCode)
+	{
+		return ErrorCode;
+		// return GErrorRegistry.RegisterErrorCode(ErrorName, ModuleId, ErrorCode);
+	}
+}
+
+
+
 
 } // namespace UE::Error
 
diff --git a/Engine/Source/Runtime/Core/Public/Experimental/IO/IoStatusError.h b/Engine/Source/Runtime/Core/Public/Experimental/IO/IoStatusError.h
--- a/Engine/Source/Runtime/Core/Public/Experimental/IO/IoStatusError.h
+++ b/Engine/Source/Runtime/Core/Public/Experimental/IO/IoStatusError.h
@@ -6,7 +6,7 @@
 
 class FIoStatus;
 
-DECLARE_ERROR_MODULE(IoStore, 0x15);
+UE_DECLARE_ERROR_MODULE(CORE_API, IoStore);
 
 
 namespace UE::UnifiedError::IoStore
diff --git a/Engine/Source/Runtime/Core/Public/Experimental/UnifiedError/UnifiedError.h b/Engine/Source/Runtime/Core/Public/Experimental/UnifiedError/UnifiedError.h
--- a/Engine/Source/Runtime/Core/Public/Experimental/UnifiedError/UnifiedError.h
+++ b/Engine/Source/Runtime/Core/Public/Experimental/UnifiedError/UnifiedError.h
@@ -42,7 +42,7 @@ public:
 
 
 
-#define DECLARE_ERRORSTRUCT_FEATURES(DetailsNameSpace, DetailsStructName)\
+#define UE_DECLARE_ERRORSTRUCT_FEATURES(DetailsNameSpace, DetailsStructName)\
 	template<>\
 	class TErrorStructFeatures<UE::UnifiedError::DetailsNameSpace::DetailsStructName>\
 	{\
@@ -68,29 +68,64 @@ struct TStringLiteralWithColonsReplaced
 {
 	constexpr TStringLiteralWithColonsReplaced(const char(&InStr)[N])
 	{
-		int outIndex = 0;
+		int OutIndex = 0;
 		int i = 0;
 		for (; i < N; ++i)
 		{
 			if (InStr[i] == '\0')
 			{
-				Str[outIndex] = InStr[i];
+				Str[OutIndex] = InStr[i];
 				break;
 			}
 			else if (InStr[i] == ':' && InStr[i + 1] == ':')
 			{
-				Str[outIndex] = '_';
-				++outIndex;
+				Str[OutIndex] = '_';
+				++OutIndex;
 			}
 			else if (InStr[i] != ':')
 			{
-				Str[outIndex] = InStr[i];
-				++outIndex;
+				Str[OutIndex] = InStr[i];
+				++OutIndex;
 			}
 		}
-		for (; outIndex < N; ++outIndex)
+		for (; OutIndex < N; ++OutIndex)
 		{
-			Str[outIndex] = '\0';
+			Str[OutIndex] = '\0';
+		}
+	}
+
+	char Str[N];
+};
+
+
+template <unsigned int N>
+struct TStringLiteralWithSpacesReplaced
+{
+	constexpr TStringLiteralWithSpacesReplaced(const char(&InStr)[N])
+	{
+		int OutIndex = 0;
+		int i = 0;
+		for (; i < N; ++i)
+		{
+			if (InStr[i] == '\0')
+			{
+				Str[OutIndex] = InStr[i];
+				break;
+			}
+			else if (InStr[i] == ' ')
+			{
+				Str[OutIndex] = '_';
+				++OutIndex;
+			}
+			else 
+			{
+				Str[OutIndex] = InStr[i];
+				++OutIndex;
+			}
+		}
+		for (; OutIndex < N; ++OutIndex)
+		{
+			Str[OutIndex] = '\0';
 		}
 	}
 
@@ -110,7 +145,7 @@ constexpr uint64 CompileTimeHashString(const char* String)
 
 // TODO Daniel: StaticDetailsTypeId is intended to register the creation functions with the FErrorDetailsRegistry for support of serialization / deserialization of error details
 
-#define DECLARE_ERROR_DETAILS_INTERNAL(DetailsNamespace, TypeName)\
+#define UE_DECLARE_ERROR_DETAILS_INTERNAL(DetailsNamespace, TypeName)\
 	inline static constexpr FAnsiStringView StaticGetErrorDetailsTypeName() \
 	{\
 		return #DetailsNamespace"::"#TypeName;\
@@ -138,17 +173,17 @@ constexpr uint64 CompileTimeHashString(const char* String)
 	}
 
 
-#define DECLARE_FERROR_DETAILS_ABSTRACT(DetailsNamespace, TypeName)\
+#define UE_DECLARE_FERROR_DETAILS_ABSTRACT(DetailsNamespace, TypeName)\
 	public:\
-	DECLARE_ERROR_DETAILS_INTERNAL(DetailsNamespace, TypeName);
+	UE_DECLARE_ERROR_DETAILS_INTERNAL(DetailsNamespace, TypeName);
 	
 	// inline const uint32 StaticDetailsTypeId = UE::UnifiedError::FErrorDetailsRegistry::Get().RegisterDetails(StaticGetErrorDetailsTypeName(), nullptr);
 	
 	
 
-#define DECLARE_FERROR_DETAILS(DetailsNamespace, TypeName)\
+#define UE_DECLARE_FERROR_DETAILS(DetailsNamespace, TypeName)\
 	public:\
-	DECLARE_ERROR_DETAILS_INTERNAL(DetailsNamespace, TypeName);\
+	UE_DECLARE_ERROR_DETAILS_INTERNAL(DetailsNamespace, TypeName);\
 	friend IErrorDetails* Create();\
 	inline static IErrorDetails* Create()\
 	{\
@@ -171,6 +206,13 @@ namespace UE::UnifiedError
 		inline constexpr int32 FIRST_CUSTOM_DETAILS_TYPE = 4;
 	}
 
+	namespace ErrorRegistry
+	{
+		CORE_API uint32 RegisterModule(const FAnsiStringView ModuleName);
+
+		CORE_API int32 RegisterErrorCode(const FAnsiStringView ErrorName, int32 ModuleId, int32 ErrorCode);
+	}
+
 
 	class FError;
 
@@ -223,7 +265,7 @@ namespace UE::UnifiedError
 	/// </summary>
 	class FManditoryErrorDetails : public IErrorDetails
 	{
-		DECLARE_FERROR_DETAILS_ABSTRACT(UnifiedError, FManditoryErrorDetails);
+		UE_DECLARE_FERROR_DETAILS_ABSTRACT(UnifiedError, FManditoryErrorDetails);
 	private:
 	public:
 		virtual FAnsiString GetErrorCodeString(const FError& Error) const = 0;
@@ -252,7 +294,7 @@ namespace UE::UnifiedError
 	/// </summary>
 	class FRefCountedErrorDetails : public IErrorDetails, public FRefCountBase
 	{
-		DECLARE_FERROR_DETAILS_ABSTRACT(UnifiedError, FRefCountedErrorDetails);
+		UE_DECLARE_FERROR_DETAILS_ABSTRACT(UnifiedError, FRefCountedErrorDetails);
 
 
 		CORE_API virtual ~FRefCountedErrorDetails();
@@ -268,7 +310,7 @@ namespace UE::UnifiedError
 	/// </summary>
 	class FDynamicErrorDetails : public FRefCountedErrorDetails
 	{
-		DECLARE_FERROR_DETAILS_ABSTRACT(UnifiedError, FDynamicErrorDetails);
+		UE_DECLARE_FERROR_DETAILS_ABSTRACT(UnifiedError, FDynamicErrorDetails);
 	private:
 		TRefCountPtr<const IErrorDetails> InnerErrorDetails;
 	public:
@@ -500,43 +542,6 @@ namespace UE::UnifiedError
 	};
 
 
-	class FErrorRegistry
-	{
-	public:
-		static FErrorRegistry& Get()
-		{
-			static FErrorRegistry StaticRegistry;
-			return StaticRegistry;
-		}
-
-
-	private:
-
-		FErrorRegistry() {}
-
-	public:
-		inline uint32 RegisterModule(const FStringView& ModuleName)
-		{
-			// todo: need to replace this with a stable hashing function
-			uint32 ModuleId = GetTypeHash(ModuleName);
-			checkf(ModuleNameMap.Contains(ModuleId) == false, TEXT("Module %s and %s are trying to register under module id %d"), ModuleName.GetData(), *ModuleNameMap.FindRef(ModuleId), ModuleId);
-			ModuleNameMap.Add(ModuleId, ModuleName.GetData());
-			return ModuleId;
-		}
-
-		inline int32 RegisterErrorCode(const FStringView& ErrorName, int32 ModuleId, int32 ErrorCode)
-		{
-			TPair<int32, int32> CombinedErrorId(ModuleId, ErrorCode);
-			checkf(ErrorCodeNameMap.Contains(CombinedErrorId) == false, TEXT("Error %s and %s are trying to register under same error code moduleid:%d errorcode:%d"), ErrorName.GetData(), *ErrorCodeNameMap.FindRef(CombinedErrorId), ModuleId, ErrorCode);
-			ErrorCodeNameMap.Add(CombinedErrorId, ErrorName.GetData());
-			return ErrorCode;
-		}
-
-	private:
-		TMap<int32, FString> ModuleNameMap;
-		TMap<TPair<int32, int32>, FString> ErrorCodeNameMap;
-	};
-
 
 
 	/// <summary>
@@ -547,14 +552,14 @@ namespace UE::UnifiedError
 	/// </summary>
 	class FStaticErrorDetails : public FManditoryErrorDetails
 	{
-		DECLARE_FERROR_DETAILS_ABSTRACT(UnifiedError, FStaticErrorDetails);
+		UE_DECLARE_FERROR_DETAILS_ABSTRACT(UnifiedError, FStaticErrorDetails);
 	private:
-		const FAnsiStringView& ErrorName;
-		const FAnsiStringView& ModuleName;
+		const FAnsiStringView ErrorName;
+		const FAnsiStringView ModuleName;
 		FText ErrorFormatString;
 	public:
 		// TODO: convert AnsiStringView to UTF8StringView when c++ 20 is supported
-		CORE_API FStaticErrorDetails(const FAnsiStringView& InErrorName, const FAnsiStringView& InModuleName, const FText& InErrorFormatString);
+		CORE_API FStaticErrorDetails(const FAnsiStringView InErrorName, const FAnsiStringView InModuleName, const FText& InErrorFormatString);
 
 		virtual ~FStaticErrorDetails() = default;
 
@@ -650,53 +655,74 @@ CORE_API constexpr uint32 GetTypeIdHash(const TCHAR* TypeName)
 */
 
 
+#define UE_DECLARE_ERROR_MODULE(DeclareApi, ModuleName) \
+	namespace UE::UnifiedError { namespace ModuleName {\
+		DeclareApi FAnsiStringView GetStaticModuleName(); \
+		DeclareApi int32 GetStaticModuleId(); \
+	}}
 
-
-
-#define DECLARE_ERROR_MODULE(ModuleName, ModuleId) \
+#define UE_DEFINE_ERROR_MODULE(ModuleName) \
 	namespace UE::UnifiedError { namespace ModuleName {\
-		inline const int32 StaticModuleId = UE::UnifiedError::FErrorRegistry::Get().RegisterModule(TEXTVIEW(#ModuleName));\
-		static constexpr FAnsiStringView StaticModuleName = ANSITEXTVIEW(#ModuleName);\
+		const int32 StaticModuleId = UE::UnifiedError::ErrorRegistry::RegisterModule(ANSITEXTVIEW(#ModuleName));\
+		int32 GetStaticModuleId() { return StaticModuleId; }\
+		FAnsiStringView GetStaticModuleName() \
+		{ \
+			static const auto Reference = TStringLiteralWithSpacesReplaced(#ModuleName); \
+			return Reference.Str; \
+		}\
 	}}
 
 
-#define DECLARE_ERROR_INTERNAL(ErrorName, ErrorCode, ModuleName, FormatString) \
+#define UE_DECLARE_ERROR_INTERNAL(DeclareApi, ErrorName, ErrorCode, ModuleName, FormatString) \
+	namespace UE::UnifiedError { namespace ModuleName { namespace ErrorName {\
+			DeclareApi FAnsiStringView GetStaticErrorName(); \
+			static const int32 DeclaredErrorCode = ErrorCode; \
+			static const FText DeclaredFormatString = FormatString; \
+			DeclareApi int32 GetErrorCodeId(); \
+			DeclareApi TRefCountPtr<const FStaticErrorDetails> GetStaticErrorDetails(); \
+			DeclareApi const FError& GetStaticError(); \
+	}}}
+
+#define UE_DEFINE_ERROR(ErrorName, ModuleName) \
 	namespace UE::UnifiedError { namespace ModuleName { namespace ErrorName {\
-			static constexpr FAnsiStringView StaticErrorName = ANSITEXTVIEW(#ErrorName);\
-			inline TRefCountPtr<const FStaticErrorDetails> GetStaticErrorDetails() \
+			const int32 ErrorCodeId = UE::UnifiedError::ErrorRegistry::RegisterErrorCode(GetStaticErrorName(), UE::UnifiedError::ModuleName::GetStaticModuleId(), DeclaredErrorCode);\
+			int32 GetErrorCodeId() \
 			{\
-				static FStaticErrorDetails StaticErrorDetails(StaticErrorName, UE::UnifiedError::ModuleName::StaticModuleName, FormatString);\
-				return TRefCountPtr<const FStaticErrorDetails>(&StaticErrorDetails);\
+				return ErrorCodeId;\
 			}\
-			inline const int32 ErrorCodeId = UE::UnifiedError::FErrorRegistry::Get().RegisterErrorCode(TEXTVIEW(#ErrorName), UE::UnifiedError::ModuleName::StaticModuleId, ErrorCode);\
-			UE_FORCEINLINE_HINT int32 GetErrorCodeId() \
+			FAnsiStringView GetStaticErrorName() \
+			{ \
+				static const auto Reference = TStringLiteralWithSpacesReplaced(#ErrorName); \
+				return Reference.Str; \
+			}\
+			TRefCountPtr<const FStaticErrorDetails> GetStaticErrorDetails() \
 			{\
-				return ErrorCodeId;\
+				static FStaticErrorDetails StaticErrorDetails(GetStaticErrorName(), UE::UnifiedError::ModuleName::GetStaticModuleName(), DeclaredFormatString);\
+				return TRefCountPtr<const FStaticErrorDetails>(&StaticErrorDetails);\
 			}\
-	}}}
-
-
-#define DECLARE_ERROR(ErrorName, ErrorCode, ModuleName, FormatString) \
-	DECLARE_ERROR_INTERNAL(ErrorName, ErrorCode, ModuleName, FormatString) \
-	namespace UE::UnifiedError { namespace ModuleName { namespace ErrorName {\
-			UE_FORCEINLINE_HINT FError MakeError() { return FError(UE::UnifiedError::ModuleName::StaticModuleId, GetErrorCodeId(), TRefCountPtr<const UE::UnifiedError::FManditoryErrorDetails>(GetStaticErrorDetails().GetReference())); }\
-			inline const FError& GetStaticError() \
+			const FError& GetStaticError() \
 			{\
 				static FError StaticError = MakeError();\
 				return StaticError;\
 			}\
 	}}}
 
+#define UE_DECLARE_ERROR(DeclareApi, ErrorName, ErrorCode, ModuleName, FormatString) \
+	UE_DECLARE_ERROR_INTERNAL(DeclareApi, ErrorName, ErrorCode, ModuleName, FormatString) \
+	namespace UE::UnifiedError { namespace ModuleName { namespace ErrorName {\
+			inline FError MakeError() { return FError(UE::UnifiedError::ModuleName::GetStaticModuleId(), GetErrorCodeId(), TRefCountPtr<const UE::UnifiedError::FManditoryErrorDetails>(GetStaticErrorDetails().GetReference())); } \
+	}}}
+
 
-#define DECLARE_ERROR_ONEPARAM(ErrorName, ErrorCode, ModuleName, FormatString, ParamOneType, ParamOneName, ParamOneDefault) \
-	DECLARE_ERROR_INTERNAL(ErrorName, ErrorCode, ModuleName, FText::FromString(UTF8TEXT("{"#ModuleName"_F"#ErrorName"}")) ) \
+#define UE_DECLARE_ERROR_ONEPARAM(DeclareApi, ErrorName, ErrorCode, ModuleName, FormatString, ParamOneType, ParamOneName, ParamOneDefault) \
+	UE_DECLARE_ERROR_INTERNAL(DeclareApi, ErrorName, ErrorCode, ModuleName, FText::FromString(UTF8TEXT("{"#ModuleName"_F"#ErrorName"}")) ) \
 	namespace UE::UnifiedError { namespace ModuleName { \
 			struct F##ErrorName \
 			{\
 				ParamOneType ParamOneName;\
 			};\
 	}}\
-	DECLARE_ERRORSTRUCT_FEATURES(ModuleName, F##ErrorName);\
+	UE_DECLARE_ERRORSTRUCT_FEATURES(ModuleName, F##ErrorName);\
 	inline void SerializeForLog(FCbWriter& Writer, const UE::UnifiedError::ModuleName::F##ErrorName& Context)\
 	{\
 		Writer.BeginObject();\
@@ -709,14 +735,9 @@ CORE_API constexpr uint32 GetTypeIdHash(const TCHAR* TypeName)
 	namespace ErrorName { \
 			inline FError MakeError(ParamOneType ParamOneName = ParamOneDefault) \
 			{ \
-				FError Error = FError(UE::UnifiedError::ModuleName::StaticModuleId, UE::UnifiedError::ModuleName::ErrorName::GetErrorCodeId(), UE::UnifiedError::ModuleName::ErrorName::GetStaticErrorDetails().GetReference()); \
+				FError Error = FError(UE::UnifiedError::ModuleName::GetStaticModuleId(), UE::UnifiedError::ModuleName::ErrorName::GetErrorCodeId(), UE::UnifiedError::ModuleName::ErrorName::GetStaticErrorDetails().GetReference()); \
 				F##ErrorName Context = {ParamOneName}; \
 				Error.PushErrorContext(MoveTemp(Context)); \
 				return Error; \
 			} \
-			inline const FError& GetStaticError() \
-			{\
-				static FError StaticError = MakeError();\
-				return StaticError;\
-			}\
 	}}}
diff --git a/Engine/Source/Runtime/Engine/Classes/Engine/Experimental/StreamableManagerError.h b/Engine/Source/Runtime/Engine/Classes/Engine/Experimental/StreamableManagerError.h
--- a/Engine/Source/Runtime/Engine/Classes/Engine/Experimental/StreamableManagerError.h
+++ b/Engine/Source/Runtime/Engine/Classes/Engine/Experimental/StreamableManagerError.h
@@ -5,21 +5,21 @@
 #include "UObject/UObjectGlobals.h"
 
 
-DECLARE_ERROR_MODULE(StreamableManager, 0x10);
+UE_DECLARE_ERROR_MODULE(ENGINE_API, StreamableManager);
 
-DECLARE_ERROR_ONEPARAM(PackageLoadFailed, 1, StreamableManager, NSLOCTEXT("StreamableManager", "PackageLoadError", "Failed to load package {PackageName}"), FString, PackageName, TEXT("Unknown"));
-DECLARE_ERROR_ONEPARAM(PackageLoadCanceled, 2, StreamableManager, NSLOCTEXT("StreamableManager", "PackageLoadCancelled", "Async load canceled {PackageName}"), FString, PackageName, TEXT("Unknown"));
-DECLARE_ERROR(DownloadError, 3, StreamableManager, NSLOCTEXT("StreamableManager", "DownloadError", "Failed to download"));
-DECLARE_ERROR_ONEPARAM(PackageNameInvalid, 4, StreamableManager, NSLOCTEXT("StreamableManager", "PackageNameInvalid", "Found invalid package name {InvalidPackageName}"), FString, InvalidPackageName, TEXT("Unknown"));
+UE_DECLARE_ERROR_ONEPARAM(	ENGINE_API, PackageLoadFailed, 1, StreamableManager, NSLOCTEXT("StreamableManager", "PackageLoadError", "Failed to load package {PackageName}"), FString, PackageName, TEXT("Unknown"));
+UE_DECLARE_ERROR_ONEPARAM(	ENGINE_API, PackageLoadCanceled, 2, StreamableManager, NSLOCTEXT("StreamableManager", "PackageLoadCancelled", "Async load canceled {PackageName}"), FString, PackageName, TEXT("Unknown"));
+UE_DECLARE_ERROR(			ENGINE_API, DownloadError, 3, StreamableManager, NSLOCTEXT("StreamableManager", "DownloadError", "Failed to download"));
+UE_DECLARE_ERROR_ONEPARAM(	ENGINE_API, PackageNameInvalid, 4, StreamableManager, NSLOCTEXT("StreamableManager", "PackageNameInvalid", "Found invalid package name {InvalidPackageName}"), FString, InvalidPackageName, TEXT("Unknown"));
 
-DECLARE_ERROR(IoStoreNotFound, 6, StreamableManager, NSLOCTEXT("StreamableManager", "IoStoreNotFound", "IoStore did not load correctly."));
-DECLARE_ERROR_ONEPARAM(SyncLoadIncomplete, 7, StreamableManager, NSLOCTEXT("StreamableManager", "SyncLoadIncomplete", "Sync load did not complete correctly for {DebugName}."), FString, DebugName, TEXT("Unknown"));
+UE_DECLARE_ERROR(			ENGINE_API, IoStoreNotFound, 6, StreamableManager, NSLOCTEXT("StreamableManager", "IoStoreNotFound", "IoStore did not load correctly."));
+UE_DECLARE_ERROR_ONEPARAM(	ENGINE_API, SyncLoadIncomplete, 7, StreamableManager, NSLOCTEXT("StreamableManager", "SyncLoadIncomplete", "Sync load did not complete correctly for {DebugName}."), FString, DebugName, TEXT("Unknown"));
 
-DECLARE_ERROR(AsyncLoadFailed, 8, StreamableManager, NSLOCTEXT("StreamableManager", "AsyncLoadFailed", "Async load failed"));
-DECLARE_ERROR(AsyncLoadCancelled, 9, StreamableManager, NSLOCTEXT("StreamableManager", "AsyncLoadCancelled", "Async load cancelled"));
-DECLARE_ERROR_ONEPARAM(AsyncLoadUnknownError, 10, StreamableManager, NSLOCTEXT("StreamableManager", "AsyncLoadUnknownError", "Unknown async loading error {AsyncLoadingErrorId}."), int32, AsyncLoadingErrorId, -1);
-DECLARE_ERROR(UnknownError, 11, StreamableManager, NSLOCTEXT("StreamableManager", "UnknownError", "Unknown error occured while streaming asset"));
-DECLARE_ERROR(AsyncLoadNotInstalled, 12, StreamableManager, NSLOCTEXT("StreamableManager", "AsyncLoadNotInstalled", "Async load failed because the package is not installed."));
+UE_DECLARE_ERROR(			ENGINE_API, AsyncLoadFailed, 8, StreamableManager, NSLOCTEXT("StreamableManager", "AsyncLoadFailed", "Async load failed"));
+UE_DECLARE_ERROR(			ENGINE_API, AsyncLoadCancelled, 9, StreamableManager, NSLOCTEXT("StreamableManager", "AsyncLoadCancelled", "Async load cancelled"));
+UE_DECLARE_ERROR_ONEPARAM(	ENGINE_API, AsyncLoadUnknownError, 10, StreamableManager, NSLOCTEXT("StreamableManager", "AsyncLoadUnknownError", "Unknown async loading error {AsyncLoadingErrorId}."), int32, AsyncLoadingErrorId, -1);
+UE_DECLARE_ERROR(			ENGINE_API, UnknownError, 11, StreamableManager, NSLOCTEXT("StreamableManager", "UnknownError", "Unknown error occured while streaming asset"));
+UE_DECLARE_ERROR(			ENGINE_API, AsyncLoadNotInstalled, 12, StreamableManager, NSLOCTEXT("StreamableManager", "AsyncLoadNotInstalled", "Async load failed because the package is not installed."));
 
 namespace UE::UnifiedError::StreamableManager
 {
@@ -35,7 +35,7 @@ namespace UE::UnifiedError::StreamableManager
 
 	
 }
-DECLARE_ERRORSTRUCT_FEATURES(StreamableManager, FStreamableManagerAdditionalContext);
+UE_DECLARE_ERRORSTRUCT_FEATURES(StreamableManager, FStreamableManagerAdditionalContext);
 inline void SerializeForLog(FCbWriter& Writer, const UE::UnifiedError::StreamableManager::FStreamableManagerAdditionalContext& Context)
 {
 	Writer.BeginObject();
diff --git a/Engine/Source/Runtime/Engine/Private/Experimental/StreamableManagerError.cpp b/Engine/Source/Runtime/Engine/Private/Experimental/StreamableManagerError.cpp
--- a/Engine/Source/Runtime/Engine/Private/Experimental/StreamableManagerError.cpp
+++ b/Engine/Source/Runtime/Engine/Private/Experimental/StreamableManagerError.cpp
@@ -2,7 +2,19 @@
 #include "Engine/Experimental/StreamableManagerError.h"
 #include "UObject/UObjectGlobals.h"
 
+UE_DEFINE_ERROR_MODULE(StreamableManager);
 
+UE_DEFINE_ERROR( PackageLoadFailed, StreamableManager);
+UE_DEFINE_ERROR(PackageLoadCanceled, StreamableManager);
+UE_DEFINE_ERROR(DownloadError, StreamableManager);
+UE_DEFINE_ERROR(PackageNameInvalid, StreamableManager);
+UE_DEFINE_ERROR(IoStoreNotFound, StreamableManager);
+UE_DEFINE_ERROR(SyncLoadIncomplete, StreamableManager);
+UE_DEFINE_ERROR(AsyncLoadFailed, StreamableManager);
+UE_DEFINE_ERROR(AsyncLoadCancelled, StreamableManager);
+UE_DEFINE_ERROR(AsyncLoadUnknownError, StreamableManager);
+UE_DEFINE_ERROR(UnknownError, StreamableManager);
+UE_DEFINE_ERROR(AsyncLoadNotInstalled, StreamableManager);
 
 namespace UE::UnifiedError::StreamableManager
 {
