Outline FTransform::Multiply to reduce binary size

diff --git a/Engine/Source/Runtime/Core/Private/Math/TransformVectorized.cpp b/Engine/Source/Runtime/Core/Private/Math/TransformVectorized.cpp
--- a/Engine/Source/Runtime/Core/Private/Math/TransformVectorized.cpp
+++ b/Engine/Source/Runtime/Core/Private/Math/TransformVectorized.cpp
@@ -291,6 +291,59 @@ bool TTransform<T>::DebugEqualMatrix(const TMatrix<T>& Matrix) const
 	return true;
 }
 
+template<typename T>
+void TTransform<T>::Multiply(TTransform<T>* OutTransform, const TTransform<T>* A, const TTransform<T>* B)
+{
+	A->DiagnosticCheckNaN_All();
+	B->DiagnosticCheckNaN_All();
+
+	checkSlow(A->IsRotationNormalized());
+	checkSlow(B->IsRotationNormalized());
+
+	//	When Q = quaternion, S = single scalar scale, and T = translation
+	//	QST(A) = Q(A), S(A), T(A), and QST(B) = Q(B), S(B), T(B)
+
+	//	QST (AxB) 
+
+	// QST(A) = Q(A)*S(A)*P*-Q(A) + T(A)
+	// QST(AxB) = Q(B)*S(B)*QST(A)*-Q(B) + T(B)
+	// QST(AxB) = Q(B)*S(B)*[Q(A)*S(A)*P*-Q(A) + T(A)]*-Q(B) + T(B)
+	// QST(AxB) = Q(B)*S(B)*Q(A)*S(A)*P*-Q(A)*-Q(B) + Q(B)*S(B)*T(A)*-Q(B) + T(B)
+	// QST(AxB) = [Q(B)*Q(A)]*[S(B)*S(A)]*P*-[Q(B)*Q(A)] + Q(B)*S(B)*T(A)*-Q(B) + T(B)
+
+	//	Q(AxB) = Q(B)*Q(A)
+	//	S(AxB) = S(A)*S(B)
+	//	T(AxB) = Q(B)*S(B)*T(A)*-Q(B) + T(B)
+	checkSlow(VectorGetComponent(A->Scale3D, 3) == 0.f);
+	checkSlow(VectorGetComponent(B->Scale3D, 3) == 0.f);
+
+	if (Private_AnyHasNegativeScale(A->Scale3D, B->Scale3D))
+	{
+		// @note, if you have 0 scale with negative, you're going to lose rotation as it can't convert back to quat
+		MultiplyUsingMatrixWithScale(OutTransform, A, B);
+	}
+	else
+	{
+		const TransformVectorRegister QuatA = A->Rotation;
+		const TransformVectorRegister QuatB = B->Rotation;
+		const TransformVectorRegister TranslateA = A->Translation;
+		const TransformVectorRegister TranslateB = B->Translation;
+		const TransformVectorRegister ScaleA = A->Scale3D;
+		const TransformVectorRegister ScaleB = B->Scale3D;
+
+		// RotationResult = B.Rotation * A.Rotation
+		OutTransform->Rotation = VectorQuaternionMultiply2(QuatB, QuatA);
+
+		// TranslateResult = B.Rotate(B.Scale * A.Translation) + B.Translate
+		const TransformVectorRegister ScaledTransA = VectorMultiply(TranslateA, ScaleB);
+		const TransformVectorRegister RotatedTranslate = VectorQuaternionRotateVector(QuatB, ScaledTransA);
+		OutTransform->Translation = VectorAdd(RotatedTranslate, TranslateB);
+
+		// ScaleResult = Scale.B * Scale.A
+		OutTransform->Scale3D = VectorMultiply(ScaleA, ScaleB);
+	}
+}
+
 } // namespace UE::Math
 } // namespace UE
 
diff --git a/Engine/Source/Runtime/Core/Public/Math/TransformVectorized.h b/Engine/Source/Runtime/Core/Public/Math/TransformVectorized.h
--- a/Engine/Source/Runtime/Core/Public/Math/TransformVectorized.h
+++ b/Engine/Source/Runtime/Core/Public/Math/TransformVectorized.h
@@ -806,7 +806,9 @@ public:
 		return Private_TranslationEquals(Other.Translation, Tolerance) && Private_RotationEquals(Other.Rotation, Tolerance);
 	}
 
-	UE_FORCEINLINE_HINT static void Multiply(TTransform<T>* OutTransform, const TTransform<T>* A, const TTransform<T>* B);
+	/** Returns Multiplied Transform of 2 FTransforms **/
+	CORE_API static void Multiply(TTransform<T>* OutTransform, const TTransform<T>* A, const TTransform<T>* B);
+
 	/**
 	 * Sets the components
 	 * @param InRotation The new value for the Rotation component
@@ -1584,59 +1586,6 @@ inline void TTransform<T>::ConstructTransformFromMatrixWithDesiredScale(const UE
 	OutTransform.Translation = VectorLoadFloat3_W0(&Translation);
 }
 
-/** Returns Multiplied Transform of 2 FTransforms **/
-template<typename T>
-inline void TTransform<T>::Multiply(TTransform<T>* OutTransform, const TTransform<T>* A, const TTransform<T>* B)
-{
-	A->DiagnosticCheckNaN_All();
-	B->DiagnosticCheckNaN_All();
-
-	checkSlow(A->IsRotationNormalized());
-	checkSlow(B->IsRotationNormalized());
-
-	//	When Q = quaternion, S = single scalar scale, and T = translation
-	//	QST(A) = Q(A), S(A), T(A), and QST(B) = Q(B), S(B), T(B)
-
-	//	QST (AxB) 
-
-	// QST(A) = Q(A)*S(A)*P*-Q(A) + T(A)
-	// QST(AxB) = Q(B)*S(B)*QST(A)*-Q(B) + T(B)
-	// QST(AxB) = Q(B)*S(B)*[Q(A)*S(A)*P*-Q(A) + T(A)]*-Q(B) + T(B)
-	// QST(AxB) = Q(B)*S(B)*Q(A)*S(A)*P*-Q(A)*-Q(B) + Q(B)*S(B)*T(A)*-Q(B) + T(B)
-	// QST(AxB) = [Q(B)*Q(A)]*[S(B)*S(A)]*P*-[Q(B)*Q(A)] + Q(B)*S(B)*T(A)*-Q(B) + T(B)
-
-	//	Q(AxB) = Q(B)*Q(A)
-	//	S(AxB) = S(A)*S(B)
-	//	T(AxB) = Q(B)*S(B)*T(A)*-Q(B) + T(B)
-	checkSlow(VectorGetComponent(A->Scale3D, 3) == 0.f);
-	checkSlow(VectorGetComponent(B->Scale3D, 3) == 0.f);
-
-	if (Private_AnyHasNegativeScale(A->Scale3D, B->Scale3D))
-	{
-		// @note, if you have 0 scale with negative, you're going to lose rotation as it can't convert back to quat
-		MultiplyUsingMatrixWithScale(OutTransform, A, B);
-	}
-	else
-	{
-		const TransformVectorRegister QuatA = A->Rotation;
-		const TransformVectorRegister QuatB = B->Rotation;
-		const TransformVectorRegister TranslateA = A->Translation;
-		const TransformVectorRegister TranslateB = B->Translation;
-		const TransformVectorRegister ScaleA = A->Scale3D;
-		const TransformVectorRegister ScaleB = B->Scale3D;
-
-		// RotationResult = B.Rotation * A.Rotation
-		OutTransform->Rotation = VectorQuaternionMultiply2(QuatB, QuatA);
-
-		// TranslateResult = B.Rotate(B.Scale * A.Translation) + B.Translate
-		const TransformVectorRegister ScaledTransA = VectorMultiply(TranslateA, ScaleB);
-		const TransformVectorRegister RotatedTranslate = VectorQuaternionRotateVector(QuatB, ScaledTransA);
-		OutTransform->Translation = VectorAdd(RotatedTranslate, TranslateB);
-
-		// ScaleResult = Scale.B * Scale.A
-		OutTransform->Scale3D = VectorMultiply(ScaleA, ScaleB);
-	}
-}
 /** 
  * Apply Scale to this transform
  */
