TSpscQueue iterator
- Allows the single consumer to iterate the contents of the queue without popping from it.
- The single producer may continue to insert items in the queue while the consumer is iterating. These new items may or may not be seen by the consumer, since the consumer might have finished iterating before reaching those new elements.

diff --git a/Engine/Source/Runtime/Core/Public/Containers/SpscQueue.h b/Engine/Source/Runtime/Core/Public/Containers/SpscQueue.h
--- a/Engine/Source/Runtime/Core/Public/Containers/SpscQueue.h
+++ b/Engine/Source/Runtime/Core/Public/Containers/SpscQueue.h
@@ -123,6 +123,51 @@ private:
 		TTypeCompatibleBytes<ElementType> Value;
 	};
 
+public:
+	//
+	// Allows the single consumer to iterate the contents of the queue without popping from it.
+	// 
+	// The single producer may continue to insert items in the queue while the consumer is iterating.
+	// These new items may or may not be seen by the consumer, since the consumer might have finished
+	// iterating before reaching those new elements.
+	//
+	struct FIterator
+	{
+		TSpscQueue::FNode* Current;
+
+		FIterator(const TSpscQueue& Queue)
+			: Current(Queue.Tail.load(std::memory_order_relaxed))
+		{
+			++(*this);
+		}
+
+		FIterator& operator++ ()
+		{
+			Current = Current->Next.load(std::memory_order_acquire);
+			return *this;
+		}
+
+		bool operator== (std::nullptr_t) const
+		{
+			return Current == nullptr;
+		}
+
+		const ElementType& operator* () const
+		{
+			return Current->Value.GetUnchecked();
+		}
+	};
+
+	FIterator begin() const
+	{
+		return FIterator(*this);
+	}
+
+	std::nullptr_t end() const
+	{
+		return nullptr;
+	}
+
 private:
 	FNode* AllocNode()
 	{
