CompositeBuffer/CompressedBuffer get rid of inline consts

They generate init code in every TU that uses them. Also, instantiate FCompositeBuffer dtor explicitly.

diff --git a/Engine/Source/Runtime/Core/Private/Compression/CompressedBuffer.cpp b/Engine/Source/Runtime/Core/Private/Compression/CompressedBuffer.cpp
--- a/Engine/Source/Runtime/Core/Private/Compression/CompressedBuffer.cpp
+++ b/Engine/Source/Runtime/Core/Private/Compression/CompressedBuffer.cpp
@@ -998,6 +998,8 @@ private:
 
 } // UE::CompressedBuffer::Private
 
+const FCompressedBuffer FCompressedBuffer::Null;
+
 FCompressedBuffer FCompressedBuffer::Compress(const FCompositeBuffer& RawData)
 {
 	return Compress(RawData, ECompressedBufferCompressor::Mermaid, ECompressedBufferCompressionLevel::VeryFast);
diff --git a/Engine/Source/Runtime/Core/Private/Memory/CompositeBuffer.cpp b/Engine/Source/Runtime/Core/Private/Memory/CompositeBuffer.cpp
--- a/Engine/Source/Runtime/Core/Private/Memory/CompositeBuffer.cpp
+++ b/Engine/Source/Runtime/Core/Private/Memory/CompositeBuffer.cpp
@@ -11,6 +11,10 @@
 #include "Templates/Function.h"
 #include "Templates/RemoveReference.h"
 
+const FCompositeBuffer FCompositeBuffer::Null;
+
+FCompositeBuffer::~FCompositeBuffer() = default;
+
 void FCompositeBuffer::Reset()
 {
 	Segments.Empty();
diff --git a/Engine/Source/Runtime/Core/Public/Compression/CompressedBuffer.h b/Engine/Source/Runtime/Core/Public/Compression/CompressedBuffer.h
--- a/Engine/Source/Runtime/Core/Public/Compression/CompressedBuffer.h
+++ b/Engine/Source/Runtime/Core/Public/Compression/CompressedBuffer.h
@@ -165,14 +165,12 @@ public:
 	[[nodiscard]] CORE_API FCompositeBuffer DecompressToComposite() const;
 
 	/** A null compressed buffer. */
-	static const FCompressedBuffer Null;
+	CORE_API static const FCompressedBuffer Null;
 
 private:
 	FCompositeBuffer CompressedData;
 };
 
-inline const FCompressedBuffer FCompressedBuffer::Null;
-
 CORE_API FArchive& operator<<(FArchive& Ar, FCompressedBuffer& Buffer);
 
 namespace UE::CompressedBuffer::Private
diff --git a/Engine/Source/Runtime/Core/Public/Memory/CompositeBuffer.h b/Engine/Source/Runtime/Core/Public/Memory/CompositeBuffer.h
--- a/Engine/Source/Runtime/Core/Public/Memory/CompositeBuffer.h
+++ b/Engine/Source/Runtime/Core/Public/Memory/CompositeBuffer.h
@@ -49,6 +49,9 @@ public:
 		}
 	}
 
+	/** Destructor. */
+	CORE_API ~FCompositeBuffer();
+
 	/** Reset this to null. */
 	CORE_API void Reset();
 
@@ -122,7 +125,7 @@ public:
 	[[nodiscard]] CORE_API bool EqualBytes(const FCompositeBuffer& Other) const;
 
 	/** A null composite buffer. */
-	static const FCompositeBuffer Null;
+	CORE_API static const FCompositeBuffer Null;
 
 private:
 	static inline int32 GetBufferCount(const FCompositeBuffer& Buffer) { return Buffer.Segments.Num(); }
@@ -142,5 +145,3 @@ private:
 	TArray<FSharedBuffer, TInlineAllocator<1>> Segments;
 };
 
-inline const FCompositeBuffer FCompositeBuffer::Null;
-
