Stop using InitializeSRWLock because it seems to confuse Application Verifier on Windows if you recreate an SRWLOCK with the same address

diff --git a/Engine/Source/Runtime/Core/Public/Windows/WindowsPlatformMutex.h b/Engine/Source/Runtime/Core/Public/Windows/WindowsPlatformMutex.h
--- a/Engine/Source/Runtime/Core/Public/Windows/WindowsPlatformMutex.h
+++ b/Engine/Source/Runtime/Core/Public/Windows/WindowsPlatformMutex.h
@@ -65,10 +65,7 @@ public:
 	FWindowsSharedMutex(const FWindowsSharedMutex&) = delete;
 	FWindowsSharedMutex& operator=(const FWindowsSharedMutex&) = delete;
 
-	UE_FORCEINLINE_HINT FWindowsSharedMutex()
-	{
-		Windows::InitializeSRWLock(&Mutex);
-	}
+	FWindowsSharedMutex() = default;
 
 	~FWindowsSharedMutex()
 	{
@@ -116,7 +113,7 @@ private:
 		return true;
 	}
 
-	Windows::SRWLOCK Mutex;
+	Windows::SRWLOCK Mutex = { 0 }; // SRWLOCK_INIT
 };
 
 /** A system-wide mutex for Windows. Uses a named mutex. */
