[Platform] Add GetExternalAppMaxPathLength() alongside GetMaxPathLimit() to provide least common path limit across all processes.

Windows has two conditions for a process to support long paths: 1. a global registry entry and 2. a per-process configuration.
All UE processes are long path aware, but that doesn't apply to all external processes.
When UE processes write out files with long paths that are meant to be read by external processes that are not long path aware, the behavior is undefined.
For this reason, this change adds a secondary function that always returns MAX_PATH on Windows and ignores the long path configuration.

#rnx

diff --git a/Engine/Source/Runtime/Core/Private/GenericPlatform/GenericPlatformMisc.cpp b/Engine/Source/Runtime/Core/Private/GenericPlatform/GenericPlatformMisc.cpp
--- a/Engine/Source/Runtime/Core/Private/GenericPlatform/GenericPlatformMisc.cpp
+++ b/Engine/Source/Runtime/Core/Private/GenericPlatform/GenericPlatformMisc.cpp
@@ -1648,6 +1648,11 @@ int32 FGenericPlatformMisc::NumberOfCoresIncludingHyperthreads()
 	return FPlatformMisc::NumberOfCores();
 }
 
+int32 FGenericPlatformMisc::GetExternalAppMaxPathLength()
+{
+	return FPlatformMisc::GetMaxPathLength();
+}
+
 void FGenericPlatformMisc::GetConfiguredCoreLimits(int32 PlatformNumPhysicalCores, int32 PlatformNumLogicalCores,
 	bool& bOutFullyInitialized, int32& OutPhysicalCoreLimit, int32& OutLogicalCoreLimit,
 	bool& bOutSetPhysicalCountToLogicalCount)
diff --git a/Engine/Source/Runtime/Core/Private/Windows/WindowsPlatformMisc.cpp b/Engine/Source/Runtime/Core/Private/Windows/WindowsPlatformMisc.cpp
--- a/Engine/Source/Runtime/Core/Private/Windows/WindowsPlatformMisc.cpp
+++ b/Engine/Source/Runtime/Core/Private/Windows/WindowsPlatformMisc.cpp
@@ -1139,6 +1139,12 @@ int32 FWindowsPlatformMisc::GetMaxPathLength()
 	return LongPathsEnabled.bValue? 32767 : MAX_PATH;
 }
 
+int32 FWindowsPlatformMisc::GetExternalAppMaxPathLength()
+{
+	// On Windows, return the minimum path limit regardless of long-path configuration, since external processes could still be not long-path aware.
+	return MAX_PATH;
+}
+
 void FWindowsPlatformMisc::GetEnvironmentVariable(const TCHAR* VariableName, TCHAR* Result, int32 ResultLength)
 {
 	uint32 Error = ::GetEnvironmentVariableW(VariableName, Result, ResultLength);
diff --git a/Engine/Source/Runtime/Core/Public/GenericPlatform/GenericPlatformMisc.h b/Engine/Source/Runtime/Core/Public/GenericPlatform/GenericPlatformMisc.h
--- a/Engine/Source/Runtime/Core/Public/GenericPlatform/GenericPlatformMisc.h
+++ b/Engine/Source/Runtime/Core/Public/GenericPlatform/GenericPlatformMisc.h
@@ -616,13 +616,22 @@ struct FGenericPlatformMisc
 	static CORE_API void SetEnvironmentVar(const TCHAR* VariableName, const TCHAR* Value);
 
 	/**
-	 * Returns the maximum length of a path
+	 * Returns the maximum length of a path.
+	 * Note that some platforms such as Windows make this a per-process configuration.
 	 */
 	FORCEINLINE static int32 GetMaxPathLength()
 	{
 		return 128;
 	}
 
+	/**
+	 * Returns the maximum length of a path which applies to all processes on the host platform.
+	 * Since some platforms such as Windows make the maximum path length a per-process configuration,
+	 * this function returns the least commmon limit that can be supported by all processes.
+	 * Default implementation returns GetMaxPathLength().
+	 */
+	static CORE_API int32 GetExternalAppMaxPathLength();
+
 	/**
 	 * return the delimiter between paths in the PATH environment variable.
 	 */
diff --git a/Engine/Source/Runtime/Core/Public/Windows/WindowsPlatformMisc.h b/Engine/Source/Runtime/Core/Public/Windows/WindowsPlatformMisc.h
--- a/Engine/Source/Runtime/Core/Public/Windows/WindowsPlatformMisc.h
+++ b/Engine/Source/Runtime/Core/Public/Windows/WindowsPlatformMisc.h
@@ -104,6 +104,7 @@ struct FWindowsPlatformMisc
 	static CORE_API ECrashHandlingType GetCrashHandlingType();
 	static CORE_API ECrashHandlingType SetCrashHandlingType(ECrashHandlingType);
 	static CORE_API int32 GetMaxPathLength();
+	static CORE_API int32 GetExternalAppMaxPathLength();
 
 	UE_DEPRECATED(4.21, "void FPlatformMisc::GetEnvironmentVariable(Name, Result, Length) is deprecated. Use FString FPlatformMisc::GetEnvironmentVariable(Name) instead.")
 	static CORE_API void GetEnvironmentVariable(const TCHAR* VariableName, TCHAR* Result, int32 ResultLength);
