Replaced old constraints from TFunction.

diff --git a/Engine/Source/Runtime/Core/Public/Templates/Function.h b/Engine/Source/Runtime/Core/Public/Templates/Function.h
--- a/Engine/Source/Runtime/Core/Public/Templates/Function.h
+++ b/Engine/Source/Runtime/Core/Public/Templates/Function.h
@@ -3,6 +3,7 @@
 #pragma once
 
 #include "CoreTypes.h"
+#include "Concepts/NotCVRefTo.h"
 #include "Misc/AssertionMacros.h"
 #include "Misc/IntrusiveUnsetOptionalState.h"
 #include "HAL/UnrealMemory.h"
@@ -91,14 +92,14 @@ namespace UE::Core::Private::Function
 		{
 			if constexpr (bOnHeap)
 			{
-			    void* This = this;
-			    this->~TFunction_OwnedObject();
-			    FMemory::Free(This);
+				void* This = this;
+				this->~TFunction_OwnedObject();
+				FMemory::Free(This);
 			}
-            else
-            {
+			else
+			{
 				this->~TFunction_OwnedObject();
-            }
+			}
 		}
 
 		~TFunction_OwnedObject() override
@@ -420,10 +421,7 @@ namespace UE::Core::Private::Function
 			#endif
 		}
 
-		template <
-			typename FunctorType
-			UE_REQUIRES(!std::is_same_v<TFunctionRefBase, std::decay_t<FunctorType>>)
-		>
+		template <UE::CNotCVRefTo<TFunctionRefBase> FunctorType>
 		TFunctionRefBase(FunctorType&& InFunc)
 		{
 			auto* Binding = Storage.Bind(Forward<FunctorType>(InFunc));
@@ -624,13 +622,8 @@ public:
 	/**
 	 * Constructor which binds a TFunctionRef to a callable object.
 	 */
-	template <
-		typename FunctorType
-		UE_REQUIRES(
-			!TIsTFunctionRef<std::decay_t<FunctorType>>::Value &&
-			std::is_invocable_r_v<Ret, std::decay_t<FunctorType>, ParamTypes...>
-		)
-	>
+	template <typename FunctorType>
+		requires (!TIsTFunctionRef<std::decay_t<FunctorType>>::Value && std::is_invocable_r_v<Ret, std::decay_t<FunctorType>, ParamTypes...>)
 	TFunctionRef(FunctorType&& InFunc UE_LIFETIMEBOUND)
 		: Super(Forward<FunctorType>(InFunc))
 	{
@@ -719,13 +712,8 @@ public:
 	/**
 	 * Constructor which binds a TFunction to any function object.
 	 */
-	template <
-		typename FunctorType
-		UE_REQUIRES(
-			!TIsTFunction<std::decay_t<FunctorType>>::Value &&
-			std::is_invocable_r_v<Ret, std::decay_t<FunctorType>, ParamTypes...>
-		)
-	>
+	template <typename FunctorType>
+		requires (!TIsTFunction<std::decay_t<FunctorType>>::Value && std::is_invocable_r_v<Ret, std::decay_t<FunctorType>, ParamTypes...>)
 	TFunction(FunctorType&& InFunc)
 		: Super(static_cast<UE::Core::Private::Function::TFuncPtrTypeIfPossible_T<FunctorType>>(InFunc))
 	{
@@ -839,14 +827,8 @@ public:
 	/**
 	 * Constructor which binds a TFunction to any function object.
 	 */
-	template <
-		typename FunctorType
-		UE_REQUIRES(
-			!TIsTUniqueFunction<std::decay_t<FunctorType>>::Value &&
-			!TIsTFunction      <std::decay_t<FunctorType>>::Value &&
-			std::is_invocable_r_v<Ret, std::decay_t<FunctorType>, ParamTypes...>
-		)
-	>
+	template <typename FunctorType>
+		requires(!TIsTUniqueFunction<std::decay_t<FunctorType>>::Value && !TIsTFunction      <std::decay_t<FunctorType>>::Value && std::is_invocable_r_v<Ret, std::decay_t<FunctorType>, ParamTypes...>)
 	TUniqueFunction(FunctorType&& InFunc)
 		: Super(static_cast<UE::Core::Private::Function::TFuncPtrTypeIfPossible_T<FunctorType>>(InFunc))
 	{
