Removed TFunction natvis visualization, as the upcoming reimplementation will not allow it.

diff --git a/Engine/Source/Runtime/Core/Public/Templates/Function.h b/Engine/Source/Runtime/Core/Public/Templates/Function.h
--- a/Engine/Source/Runtime/Core/Public/Templates/Function.h
+++ b/Engine/Source/Runtime/Core/Public/Templates/Function.h
@@ -17,12 +17,8 @@
 #include <type_traits>
 
 // Disable visualization hack for shipping or test builds.
-#ifndef UE_ENABLE_TFUNCTIONREF_VISUALIZATION
-	#if !(UE_BUILD_SHIPPING || UE_BUILD_TEST)
-		#define UE_ENABLE_TFUNCTIONREF_VISUALIZATION 1
-	#else
-		#define UE_ENABLE_TFUNCTIONREF_VISUALIZATION 0
-	#endif
+#if defined(UE_ENABLE_TFUNCTIONREF_VISUALIZATION) && UE_ENABLE_TFUNCTIONREF_VISUALIZATION
+	COMPILE_WARNING("TFunction visualization has been removed - please undefine UE_ENABLE_TFUNCTIONREF_VISUALIZATION.")
 #endif
 
 #if !defined(NUM_TFUNCTION_INLINE_BYTES) || NUM_TFUNCTION_INLINE_BYTES == 0
@@ -283,26 +279,6 @@ namespace UE::Core::Private::Function
 		}
 	}
 
-	#if UE_ENABLE_TFUNCTIONREF_VISUALIZATION
-		/**
-		 * Helper classes to help debugger visualization.
-		 */
-		struct IDebugHelper
-		{
-			virtual ~IDebugHelper() = 0;
-		};
-
-		inline IDebugHelper::~IDebugHelper()
-		{
-		}
-
-		template <typename T>
-		struct TDebugHelper : IDebugHelper
-		{
-			T* Ptr = nullptr;
-		};
-	#endif
-
 	/**
 	 * A class which is used to instantiate the code needed to call a bound function.
 	 */
@@ -344,13 +320,6 @@ namespace UE::Core::Private::Function
 
 			if (Callable)
 			{
-				#if UE_ENABLE_TFUNCTIONREF_VISUALIZATION
-					// Use Memcpy to copy the other DebugPtrStorage, including vptr (because we don't know the bound type
-					// here), and then reseat the underlying pointer.  Possibly even more evil than the Set code.
-					FMemory::Memcpy(&DebugPtrStorage, &Other.DebugPtrStorage, sizeof(DebugPtrStorage)); //-V598
-					DebugPtrStorage.Ptr = Storage.GetPtr();
-				#endif
-
 				Other.Callable = nullptr;
 			}
 		}
@@ -365,13 +334,6 @@ namespace UE::Core::Private::Function
 
 			if (Callable)
 			{
-				#if UE_ENABLE_TFUNCTIONREF_VISUALIZATION
-					// Use Memcpy to copy the other DebugPtrStorage, including vptr (because we don't know the bound type
-					// here), and then reseat the underlying pointer.  Possibly even more evil than the Set code.
-					FMemory::Memcpy(&DebugPtrStorage, &Other.DebugPtrStorage, sizeof(DebugPtrStorage)); //-V598
-					DebugPtrStorage.Ptr = Storage.GetPtr();
-				#endif
-
 				Other.Callable = nullptr;
 			}
 		}
@@ -391,13 +353,6 @@ namespace UE::Core::Private::Function
 			}
 
 			void* NewPtr = Storage.BindCopy(Other.Storage);
-
-			#if UE_ENABLE_TFUNCTIONREF_VISUALIZATION
-				// Use Memcpy to copy the other DebugPtrStorage, including vptr (because we don't know the bound type
-				// here), and then reseat the underlying pointer.  Possibly even more evil than the Set code.
-				FMemory::Memcpy(&DebugPtrStorage, &Other.DebugPtrStorage, sizeof(DebugPtrStorage)); //-V598
-				DebugPtrStorage.Ptr = NewPtr;
-			#endif
 		}
 
 		TFunctionRefBase(const TFunctionRefBase& Other)
@@ -412,13 +367,6 @@ namespace UE::Core::Private::Function
 			}
 
 			void* NewPtr = Storage.BindCopy(Other.Storage);
-
-			#if UE_ENABLE_TFUNCTIONREF_VISUALIZATION
-				// Use Memcpy to copy the other DebugPtrStorage, including vptr (because we don't know the bound type
-				// here), and then reseat the underlying pointer.  Possibly even more evil than the Set code.
-				FMemory::Memcpy(&DebugPtrStorage, &Other.DebugPtrStorage, sizeof(DebugPtrStorage)); //-V598
-				DebugPtrStorage.Ptr = NewPtr;
-			#endif
 		}
 
 		template <UE::CNotCVRefTo<TFunctionRefBase> FunctorType>
@@ -437,15 +385,6 @@ namespace UE::Core::Private::Function
 			using DecayedFunctorType = typename TRemovePointer<decltype(Binding)>::Type;
 
 			Callable = &TFunctionRefCaller<DecayedFunctorType, Ret, ParamTypes...>::Call;
-
-			#if UE_ENABLE_TFUNCTIONREF_VISUALIZATION
-				// We placement new over the top of the same object each time.  This is illegal,
-				// but it ensures that the vptr is set correctly for the bound type, and so is
-				// visualizable.  We never depend on the state of this object at runtime, so it's
-				// ok.
-				::new ((void*)&DebugPtrStorage) TDebugHelper<DecayedFunctorType>;
-				DebugPtrStorage.Ptr = (void*)Binding;
-			#endif
 		}
 
 		TFunctionRefBase& operator=(TFunctionRefBase&&) = delete;
@@ -506,11 +445,6 @@ namespace UE::Core::Private::Function
 		Ret (*Callable)(void*, ParamTypes&...) = nullptr;
 
 		StorageType Storage;
-
-		#if UE_ENABLE_TFUNCTIONREF_VISUALIZATION
-			// To help debug visualizers
-			TDebugHelper<void> DebugPtrStorage;
-		#endif
 	};
 
 	struct FFunctionRefStoragePolicy
