Move PVS-Studio analysis settings to a new header so it can be shared

#rnx

diff --git a/Engine/Source/Runtime/Core/Public/Analysis/PvsStudioCodeAnalysis.h b/Engine/Source/Runtime/Core/Public/Analysis/PvsStudioCodeAnalysis.h
new file mode 100644
--- /dev/null
+++ b/Engine/Source/Runtime/Core/Public/Analysis/PvsStudioCodeAnalysis.h
@@ -0,0 +1,93 @@
+// Copyright Epic Games, Inc. All Rights Reserved.
+
+#pragma once
+
+#if defined( PVS_STUDIO )
+
+	//PVS-Studio settings:
+	//-V::505,542,677,704,719,720,730,735,780,1002,1055,1062,1100
+	//-V:TRYCOMPRESSION:519,547
+	//-V:check(:501,547,560,605
+	//-V:checkSlow(:547
+	//-V:checkf(:510
+	//-V:checkfSlow(:547
+	//-V:dtAssert(:568
+	//-V:rcAssert(:568
+	//-V:GET_FUNCTION_NAME_CHECKED:521
+	//-V:ENABLE_TEXT_ERROR_CHECKING_RESULTS:560
+	//-V:ENABLE_LOC_TESTING:560,617
+	//-V:WITH_EDITOR:560
+	//-V:UE_LOG_ACTIVE:560
+	//-V:verify:501
+	//-V:UE_BUILD_SHIPPING:501
+	//-V:WITH_EDITOR:501
+	//-V:TestTrueExpr:501
+	//-V:PLATFORM_:517,547
+	//-V:ensureMsgf:562
+	//-V:SelectedEmitter:519
+	//-V:MAX_VERTS_PER_POLY:512
+	//-V:127:547
+	//-V:0x7F:547
+	//-V:WARN_COLOR:547
+	//-V:<<:614
+	//-V:OPENGL_PERFORMANCE_DATA_INVALID:564
+	//-V:SPAWN_INIT:595
+	//-V:BEGIN_UPDATE_LOOP:595
+	//-V:OPENGL_PERFORMANCE_DATA_INVALID:560
+	//-V:UE_CLOG(:501,560
+	//-V:UE_LOG(:501,510,560
+	//-V:UGL_REQUIRED_VOID:501
+	//-V:ensure(:595
+	//-V:ALLOCATE_VERTEX_DATA_TEMPLATE:501
+	//-V:UGL_REQUIRED:501
+	//-V:DEBUG_LOG_HTTP:523
+	//-V:GIsEditor:560
+	//-V:GEventDrivenLoaderEnabled:501
+	//-V:IMPLEMENT_AI_INSTANT_TEST:773
+	//-V:ENABLE_VERIFY_GL:564
+	//-V:INC_MEMORY_STAT_BY:568
+	//-V:DEC_MEMORY_STAT_BY:568
+	//-V:SELECT_STATIC_MESH_VERTEX_TYPE:622
+	//-V:GET_FUNCTION_NAME_CHECKED:685
+	//-V:This(:678
+	//-V:state->error:649
+	//-V:ProjModifiers:616
+	//-V:PERF_DETAILED_PER_CLASS_GC_STATS:686
+	//-V:FMath:656
+	//-V:->*:607
+	//-V:GENERATED_UCLASS_BODY:764
+	//-V:CalcSegmentCostOnPoly:764
+	//-V:DrawLine:764
+	//-V:vrapi_SubmitFrame:641
+	//-V:VertexData:773
+	//-V:Linker:678
+	//-V:self:678
+	//-V:FindChar:679
+	//-V:BytesToHex:530
+	//-V:FindOrAdd:530
+	//-V:GetPawn:623
+	//-V:FindComponentByClass:623
+	//-V:Response.Response:1051
+
+	// The following classes retain a reference to data supplied in the constructor by the derived class which can not yet be initialized.
+	//-V:FMemoryWriter(:1050
+	//-V:FObjectWriter(:1050
+	//-V:FDurationTimer(:1050
+	//-V:FScopedDurationTimer(:1050
+	//-V:FQueryFastData(:1050
+
+	// Exclude all generated protobuf files
+	//V_EXCLUDE_PATH *.pb.cc
+
+	// warning V630: Instantiation of TRingBuffer < FString >: The 'Malloc' function is used to allocate memory for an array of objects which are classes containing constructors and destructors.
+	//-V::630
+
+	// Disabling because incorrectly flagging all TStaticArrays
+	// V557: Array overrun is possible
+	//-V::557
+
+	// Disabling because too many virtuals currently in use in constructors/destructors, need to revist
+	// V1053: Calling the 'foo' virtual function in the constructor/destructor may lead to unexpected result at runtime.
+	//-V::1053
+
+#endif // #if defined( PVS_STUDIO )
diff --git a/Engine/Source/Runtime/Core/Public/Clang/ClangPlatformCodeAnalysis.h b/Engine/Source/Runtime/Core/Public/Clang/ClangPlatformCodeAnalysis.h
--- a/Engine/Source/Runtime/Core/Public/Clang/ClangPlatformCodeAnalysis.h
+++ b/Engine/Source/Runtime/Core/Public/Clang/ClangPlatformCodeAnalysis.h
@@ -9,7 +9,7 @@
 #if defined(__clang__)
 
 // Code analysis features
-#if defined( __clang_analyzer__ )
+#if defined( __clang_analyzer__ ) || defined( PVS_STUDIO )
 	#define USING_CODE_ANALYSIS 1
 #else
 	#define USING_CODE_ANALYSIS 0
@@ -40,6 +40,9 @@
 
 #if USING_CODE_ANALYSIS
 
+	// If enabled, add in settings for PVS-Studio Analysis
+	#include "Analysis/PvsStudioCodeAnalysis.h"
+
 	// A fake function marked with noreturn that acts as a marker for CA_ASSUME to ensure the
 	// static analyzer doesn't take an analysis path that is assumed not to be navigable.
 	void CA_AssumeNoReturn() __attribute__((analyzer_noreturn));
diff --git a/Engine/Source/Runtime/Core/Public/Microsoft/MicrosoftPlatformCodeAnalysis.h b/Engine/Source/Runtime/Core/Public/Microsoft/MicrosoftPlatformCodeAnalysis.h
--- a/Engine/Source/Runtime/Core/Public/Microsoft/MicrosoftPlatformCodeAnalysis.h
+++ b/Engine/Source/Runtime/Core/Public/Microsoft/MicrosoftPlatformCodeAnalysis.h
@@ -35,6 +35,9 @@
 
 #if USING_CODE_ANALYSIS
 
+// If enabled, add in settings for PVS-Studio Analysis
+#include "Analysis/PvsStudioCodeAnalysis.h"
+
 #if !defined(__clang_analyzer__)
 
 	// Input argument
@@ -111,92 +114,6 @@
 	// EX: 365 properties in a UObject will auto gen a table exactly this long.
 	#pragma warning(disable : 6393) // warning C6393: A lookup table of size 365 is not sufficient to handle leap years.
 
-	//PVS-Studio settings:
-	//-V::505,542,677,704,719,720,730,735,780,1002,1055,1062,1100
-	//-V:TRYCOMPRESSION:519,547
-	//-V:check(:501,547,560,605
-	//-V:checkSlow(:547
-	//-V:checkf(:510
-	//-V:checkfSlow(:547
-	//-V:dtAssert(:568
-	//-V:rcAssert(:568
-	//-V:GET_FUNCTION_NAME_CHECKED:521
-	//-V:ENABLE_TEXT_ERROR_CHECKING_RESULTS:560
-	//-V:ENABLE_LOC_TESTING:560,617
-	//-V:WITH_EDITOR:560
-	//-V:UE_LOG_ACTIVE:560
-	//-V:verify:501
-	//-V:UE_BUILD_SHIPPING:501
-	//-V:WITH_EDITOR:501
-	//-V:TestTrueExpr:501
-	//-V:PLATFORM_:517,547
-	//-V:ensureMsgf:562
-	//-V:SelectedEmitter:519
-	//-V:MAX_VERTS_PER_POLY:512
-	//-V:127:547
-	//-V:0x7F:547
-	//-V:WARN_COLOR:547
-	//-V:<<:614
-	//-V:OPENGL_PERFORMANCE_DATA_INVALID:564
-	//-V:SPAWN_INIT:595
-	//-V:BEGIN_UPDATE_LOOP:595
-	//-V:OPENGL_PERFORMANCE_DATA_INVALID:560
-	//-V:UE_CLOG(:501,560
-	//-V:UE_LOG(:501,510,560
-	//-V:UGL_REQUIRED_VOID:501
-	//-V:ensure(:595
-	//-V:ALLOCATE_VERTEX_DATA_TEMPLATE:501
-	//-V:UGL_REQUIRED:501
-	//-V:DEBUG_LOG_HTTP:523
-	//-V:GIsEditor:560
-	//-V:GEventDrivenLoaderEnabled:501
-	//-V:IMPLEMENT_AI_INSTANT_TEST:773
-	//-V:ENABLE_VERIFY_GL:564
-	//-V:INC_MEMORY_STAT_BY:568
-	//-V:DEC_MEMORY_STAT_BY:568
-	//-V:SELECT_STATIC_MESH_VERTEX_TYPE:622
-	//-V:GET_FUNCTION_NAME_CHECKED:685
-	//-V:This(:678
-	//-V:state->error:649
-	//-V:ProjModifiers:616
-	//-V:PERF_DETAILED_PER_CLASS_GC_STATS:686
-	//-V:FMath:656
-	//-V:->*:607
-	//-V:GENERATED_UCLASS_BODY:764
-	//-V:CalcSegmentCostOnPoly:764
-	//-V:DrawLine:764
-	//-V:vrapi_SubmitFrame:641
-	//-V:VertexData:773
-	//-V:Linker:678
-	//-V:self:678
-	//-V:FindChar:679
-	//-V:BytesToHex:530
-	//-V:FindOrAdd:530
-	//-V:GetPawn:623
-	//-V:FindComponentByClass:623
-	//-V:Response.Response:1051
-
-	// The following classes retain a reference to data supplied in the constructor by the derived class which can not yet be initialized.
-	//-V:FMemoryWriter(:1050
-	//-V:FObjectWriter(:1050
-	//-V:FDurationTimer(:1050
-	//-V:FScopedDurationTimer(:1050
-	//-V:FQueryFastData(:1050
-
-	// Exclude all generated protobuf files
-	//V_EXCLUDE_PATH *.pb.cc
-
-	// warning V630: Instantiation of TRingBuffer < FString >: The 'Malloc' function is used to allocate memory for an array of objects which are classes containing constructors and destructors.
-	//-V::630
-
-	// Disabling because incorrectly flagging all TStaticArrays
-	// V557: Array overrun is possible
-	//-V::557
-
-	// Disabling because too many virtuals currently in use in constructors/destructors, need to revist
-	// V1053: Calling the 'foo' virtual function in the constructor/destructor may lead to unexpected result at runtime.
-	//-V::1053
-
 #else // defined(__clang_analyzer__)
 
 // A fake function marked with noreturn that acts as a marker for CA_ASSUME to ensure the
