Changes For Manipulating Constexpr Static Arrays During Compile Time

- Make TStaticArray's storage public so that it is a "Structural Type" and can be passed as a non-type template parameter in C++20
- Add constexpr to the functions necessary to use Algo::Contains during compile time
- Add new constexpr enum flag functions
	- EnumOnlyContainsFlags
    - EnumHasOneFlag
    - EnumHasAnyOneFlag

diff --git a/Engine/Source/Runtime/Core/Public/Algo/Contains.h b/Engine/Source/Runtime/Core/Public/Algo/Contains.h
--- a/Engine/Source/Runtime/Core/Public/Algo/Contains.h
+++ b/Engine/Source/Runtime/Core/Public/Algo/Contains.h
@@ -16,7 +16,7 @@ namespace Algo
 	 * @return true if an element was found, false otherwise.
 	 */
 	template <typename RangeType, typename ValueType>
-	[[nodiscard]] UE_REWRITE bool Contains(RangeType&& Range, const ValueType& Value)
+	[[nodiscard]] UE_REWRITE constexpr bool Contains(RangeType&& Range, const ValueType& Value)
 	{
 		return !!Algo::Find(Forward<RangeType>(Range), Value);
 	}
@@ -31,7 +31,7 @@ namespace Algo
 	 * @return true if an element was found, false otherwise.
 	 */
 	template <typename RangeType, typename ValueType, typename ProjectionType>
-	[[nodiscard]] UE_REWRITE bool ContainsBy(RangeType&& Range, const ValueType& Value, ProjectionType Proj)
+	[[nodiscard]] UE_REWRITE constexpr bool ContainsBy(RangeType&& Range, const ValueType& Value, ProjectionType Proj)
 	{
 		return !!Algo::FindBy(Forward<RangeType>(Range), Value, MoveTemp(Proj));
 	}
diff --git a/Engine/Source/Runtime/Core/Public/Algo/Find.h b/Engine/Source/Runtime/Core/Public/Algo/Find.h
--- a/Engine/Source/Runtime/Core/Public/Algo/Find.h
+++ b/Engine/Source/Runtime/Core/Public/Algo/Find.h
@@ -10,7 +10,7 @@
 namespace AlgoImpl
 {
 	template <typename RangeType, typename ValueType, typename ProjectionType>
-	typename TRangePointerType<typename TRemoveReference<RangeType>::Type>::Type FindBy(RangeType&& Range, const ValueType& Value, ProjectionType Proj)
+	constexpr typename TRangePointerType<typename TRemoveReference<RangeType>::Type>::Type FindBy(RangeType&& Range, const ValueType& Value, ProjectionType Proj)
 	{
 		for (auto&& Elem : Forward<RangeType>(Range))
 		{
@@ -24,7 +24,7 @@ namespace AlgoImpl
 	}
 
 	template <typename RangeType, typename PredicateType>
-	typename TRangePointerType<typename TRemoveReference<RangeType>::Type>::Type FindByPredicate(RangeType&& Range, PredicateType Pred)
+	constexpr typename TRangePointerType<typename TRemoveReference<RangeType>::Type>::Type FindByPredicate(RangeType&& Range, PredicateType Pred)
 	{
 		for (auto&& Elem : Forward<RangeType>(Range))
 		{
@@ -49,7 +49,7 @@ namespace Algo
 	 * @return A pointer to the first element found, or nullptr if none was found.
 	 */
 	template <typename RangeType, typename ValueType>
-	[[nodiscard]] UE_REWRITE auto Find(RangeType&& Range, const ValueType& Value)
+	[[nodiscard]] UE_REWRITE constexpr auto Find(RangeType&& Range, const ValueType& Value)
 		-> decltype(AlgoImpl::FindBy(Forward<RangeType>(Range), Value, FIdentityFunctor()))
 	{
 		return AlgoImpl::FindBy(Forward<RangeType>(Range), Value, FIdentityFunctor());
@@ -65,7 +65,7 @@ namespace Algo
 	 * @return A pointer to the first element found, or nullptr if none was found.
 	 */
 	template <typename RangeType, typename ValueType, typename ProjectionType>
-	[[nodiscard]] UE_REWRITE auto FindBy(RangeType&& Range, const ValueType& Value, ProjectionType Proj)
+	[[nodiscard]] UE_REWRITE constexpr auto FindBy(RangeType&& Range, const ValueType& Value, ProjectionType Proj)
 		-> decltype(AlgoImpl::FindBy(Forward<RangeType>(Range), Value, MoveTemp(Proj)))
 	{
 		return AlgoImpl::FindBy(Forward<RangeType>(Range), Value, MoveTemp(Proj));
@@ -80,7 +80,7 @@ namespace Algo
 	 * @return A pointer to the first element found, or nullptr if none was found.
 	 */
 	template <typename RangeType, typename PredicateType>
-	[[nodiscard]] UE_REWRITE auto FindByPredicate(RangeType&& Range, PredicateType Pred)
+	[[nodiscard]] UE_REWRITE constexpr auto FindByPredicate(RangeType&& Range, PredicateType Pred)
 		-> decltype(AlgoImpl::FindByPredicate(Forward<RangeType>(Range), MoveTemp(Pred)))
 	{
 		return AlgoImpl::FindByPredicate(Forward<RangeType>(Range), MoveTemp(Pred));
diff --git a/Engine/Source/Runtime/Core/Public/Containers/StaticArray.h b/Engine/Source/Runtime/Core/Public/Containers/StaticArray.h
--- a/Engine/Source/Runtime/Core/Public/Containers/StaticArray.h
+++ b/Engine/Source/Runtime/Core/Public/Containers/StaticArray.h
@@ -157,11 +157,13 @@ private:
 		TArrayStorageElementAligned Elements[NumElements];
 	};
 
+public:
+	// Storage is made public so that TStaticArray is a "Structural Type" and therefore can be passed as 
+	// a non-type template parameter like std::array making it a more complete replacement in C++20
+	// It is not intended to be accessed directly. Please use the member functions and operators instead
 	TArrayStorage Storage;
 
 
-public:
-
 	template <typename StorageElementType, bool bReverse = false>
 	struct FRangedForIterator
 	{
diff --git a/Engine/Source/Runtime/Core/Public/Misc/EnumClassFlags.h b/Engine/Source/Runtime/Core/Public/Misc/EnumClassFlags.h
--- a/Engine/Source/Runtime/Core/Public/Misc/EnumClassFlags.h
+++ b/Engine/Source/Runtime/Core/Public/Misc/EnumClassFlags.h
@@ -38,6 +38,55 @@ constexpr bool EnumHasAnyFlags(Enum Flags, Enum Contains)
 	return ((UnderlyingType)Flags & (UnderlyingType)Contains) != 0;
 }
 
+/**
+* Check if the Flags value contains only the flags specified by the Contains argument.
+* It is not required that any of the flags in Contains be set in Flags and Flags may be
+* any combination of the flags specified by Contains (including k=0). But if Flags 
+* contains any flag not in Contains, this function will return false.
+* This function may also be understood as !EnumContainsAnyFlags(Flags, ~Contains).
+*
+* @param Flags    The value to check.
+* @param Contains The flags to check the value against.
+*
+* @return true if Flags consists of only the flags specified by Contains. false otherwise.
+* @note   Returns true if Flags is empty.
+*/
+template <typename Enum>
+constexpr bool EnumOnlyContainsFlags(Enum Flags, Enum Contains)
+{
+	return EnumHasAllFlags(Contains, Flags);
+}
+
+/**
+* Check if Flags has one and only one flag set.
+* This can also be thought of as a check for a power-of-2 value.
+*
+* @param Flags The value to check.
+*
+* @return true if Flags has only one flag set. false otherwise.
+*/
+template <typename Enum>
+constexpr bool EnumHasOneFlag(Enum Flags)
+{
+	using UnderlyingType = __underlying_type(Enum);
+	return ((UnderlyingType)Flags != 0) && (((UnderlyingType)Flags & ((UnderlyingType)Flags - 1)) == 0);
+}
+
+/**
+* Check if Flags has one and only one of the flags specified in OneOfFlags set
+*
+* @param Flags       The value to check.
+* @param OneOfFlags  The flags to check the value for.
+*
+* @return true if Flags has one and only one of the flags in OneOfFlags set. false otherwise.
+*/
+template <typename Enum>
+constexpr bool EnumHasAnyOneFlag(Enum Flags, Enum OneOfFlags)
+{
+	using UnderlyingType = __underlying_type(Enum);
+	return EnumHasOneFlag((Enum)((UnderlyingType)Flags & (UnderlyingType)OneOfFlags));
+}
+
 template<typename Enum>
 constexpr void EnumAddFlags(Enum& Flags, Enum FlagsToAdd)
 {
diff --git a/Engine/Source/Runtime/Core/Public/Templates/IdentityFunctor.h b/Engine/Source/Runtime/Core/Public/Templates/IdentityFunctor.h
--- a/Engine/Source/Runtime/Core/Public/Templates/IdentityFunctor.h
+++ b/Engine/Source/Runtime/Core/Public/Templates/IdentityFunctor.h
@@ -10,7 +10,7 @@
 struct FIdentityFunctor
 {
 	template <typename T>
-	UE_FORCEINLINE_HINT T&& operator()(T&& Val) const
+	UE_FORCEINLINE_HINT constexpr T&& operator()(T&& Val) const
 	{
 		return (T&&)Val;
 	}
diff --git a/Engine/Source/Runtime/Core/Public/Templates/Invoke.h b/Engine/Source/Runtime/Core/Public/Templates/Invoke.h
--- a/Engine/Source/Runtime/Core/Public/Templates/Invoke.h
+++ b/Engine/Source/Runtime/Core/Public/Templates/Invoke.h
@@ -40,14 +40,14 @@ namespace UE::Core::Private
  * See: http://en.cppreference.com/w/cpp/utility/functional/invoke
  */
 template <typename FuncType, typename... ArgTypes>
-UE_FORCEINLINE_HINT auto Invoke(FuncType&& Func, ArgTypes&&... Args)
+UE_FORCEINLINE_HINT constexpr auto Invoke(FuncType&& Func, ArgTypes&&... Args)
 	-> decltype(((FuncType&&)Func)((ArgTypes&&)Args...))
 {
 	return ((FuncType&&)Func)((ArgTypes&&)Args...);
 }
 
 template <typename ReturnType, typename ObjType, typename TargetType>
-inline auto Invoke(ReturnType ObjType::*pdm, TargetType&& Target)
+constexpr auto Invoke(ReturnType ObjType::*pdm, TargetType&& Target)
 	-> decltype(UE::Core::Private::DereferenceIfNecessary<ObjType>((TargetType&&)Target, &Target).*pdm)
 {
 	return UE::Core::Private::DereferenceIfNecessary<ObjType>((TargetType&&)Target, &Target).*pdm;
@@ -59,7 +59,7 @@ template <
 	typename... ArgTypes,
 	typename    ObjType = TMemberFunctionPtrOuter_T<PtrMemFunType>
 >
-inline auto Invoke(PtrMemFunType PtrMemFun, TargetType&& Target, ArgTypes&&... Args)
+constexpr auto Invoke(PtrMemFunType PtrMemFun, TargetType&& Target, ArgTypes&&... Args)
 	-> decltype((UE::Core::Private::DereferenceIfNecessary<ObjType>((TargetType&&)Target, &Target).*PtrMemFun)((ArgTypes&&)Args...))
 {
 	return (UE::Core::Private::DereferenceIfNecessary<ObjType>((TargetType&&)Target, &Target).*PtrMemFun)((ArgTypes&&)Args...);
