Fix merged module loading skipping initialization of the first loaded module in a library

This happened because statically registered modules are assumed to be known before we call LoadModuleWithFailureReason, which is wrong using merged modules.
When using merged modules, we need to give a second chance to this process after we've loaded the merged library so that the module is properly initialized.

#rnx

diff --git a/Engine/Source/Runtime/Core/Private/Modules/ModuleManager.cpp b/Engine/Source/Runtime/Core/Private/Modules/ModuleManager.cpp
--- a/Engine/Source/Runtime/Core/Private/Modules/ModuleManager.cpp
+++ b/Engine/Source/Runtime/Core/Private/Modules/ModuleManager.cpp
@@ -991,11 +991,9 @@ IModuleInterface* FModuleManager::LoadModuleWithFailureReason(const FName InModu
 	// module instead of trying to load it directly.
 	checkf((!ModuleInfo->bWasUnloadedAtShutdown), TEXT("Attempted to load module '%s' that was already unloaded at shutdown.  FModuleManager::LoadModule() was called to load a module that was previously loaded, and was unloaded at shutdown time.  If this assert goes off, you're trying to load a module during the shutdown phase that was already cleaned up.  The easiest way to fix this is to change your code to query for an already-loaded module instead of trying to load it directly."), *InModuleName.ToString());
 
-	// Check if we're statically linked with the module.  Those modules register with the module manager using a static variable,
-	// so hopefully we already know about the name of the module and how to initialize it.
-	ProcessPendingStaticallyLinkedModuleInitializers();
-	const FInitializeStaticallyLinkedModule* ModuleInitializerPtr = StaticallyLinkedModuleInitializers.Find(InModuleName);
-	if (ModuleInitializerPtr != nullptr)
+	// Initialize a module that was registered with the module manager using the static initializer.
+	// In monolithic builds, we should know about this module before we get to load it; using merged modules this can be delayed until the library has been loaded.
+	const auto InitializeModule = [this, InModuleName, InLoadModuleFlags, &ModuleInfo, &LoadedModule, &OutFailureReason, &EnsureModuleFilename](const FInitializeStaticallyLinkedModule* ModuleInitializerPtr)
 	{
 		const FInitializeStaticallyLinkedModule& ModuleInitializer(*ModuleInitializerPtr);
 
@@ -1044,6 +1042,13 @@ IModuleInterface* FModuleManager::LoadModuleWithFailureReason(const FName InModu
 
 			OutFailureReason = EModuleLoadResult::FailedToInitialize;
 		}
+	};
+
+	ProcessPendingStaticallyLinkedModuleInitializers();
+	const FInitializeStaticallyLinkedModule* ModuleInitializerPtr = StaticallyLinkedModuleInitializers.Find(InModuleName);
+	if (ModuleInitializerPtr != nullptr)
+	{
+		InitializeModule(ModuleInitializerPtr);
 	}
 #if IS_MONOLITHIC
 	else
@@ -1104,6 +1109,18 @@ IModuleInterface* FModuleManager::LoadModuleWithFailureReason(const FName InModu
 			
 			if (ModuleInfo->Handle != nullptr)
 			{
+#if UE_MERGED_MODULES
+
+				// We just loaded a library and so we have a lot of pending initializers ready to go, including the one for our module
+				ProcessPendingStaticallyLinkedModuleInitializers();
+				const FInitializeStaticallyLinkedModule* ModuleInitializerPtr = StaticallyLinkedModuleInitializers.Find(InModuleName);
+				if (ModuleInitializerPtr != nullptr)
+				{
+					InitializeModule(ModuleInitializerPtr);
+				}
+
+#endif // UE_MERGED_MODULES
+
 				{
 					// Defer String Table find/load during CDO registration, as it may happen 
 					// before StartupModule has had a chance to load the String Table
