Fix build

#rnx

diff --git a/Engine/Source/Runtime/Core/Private/Modules/ModuleManager.cpp b/Engine/Source/Runtime/Core/Private/Modules/ModuleManager.cpp
--- a/Engine/Source/Runtime/Core/Private/Modules/ModuleManager.cpp
+++ b/Engine/Source/Runtime/Core/Private/Modules/ModuleManager.cpp
@@ -1113,10 +1113,10 @@ IModuleInterface* FModuleManager::LoadModuleWithFailureReason(const FName InModu
 
 				// We just loaded a library and so we have a lot of pending initializers ready to go, including the one for our module
 				ProcessPendingStaticallyLinkedModuleInitializers();
-				const FInitializeStaticallyLinkedModule* ModuleInitializerPtr = StaticallyLinkedModuleInitializers.Find(InModuleName);
-				if (ModuleInitializerPtr != nullptr)
+				const FInitializeStaticallyLinkedModule* MergedModuleInitializerPtr = StaticallyLinkedModuleInitializers.Find(InModuleName);
+				if (MergedModuleInitializerPtr != nullptr)
 				{
-					InitializeModule(ModuleInitializerPtr);
+					InitializeModule(MergedModuleInitializerPtr);
 				}
 
 #endif // UE_MERGED_MODULES
