Add FIoOffsetAndLength::IsValid to make it easier to check if either the offset or length have been set.

#rnx

diff --git a/Engine/Source/Runtime/Core/Internal/IO/IoOffsetLength.h b/Engine/Source/Runtime/Core/Internal/IO/IoOffsetLength.h
--- a/Engine/Source/Runtime/Core/Internal/IO/IoOffsetLength.h
+++ b/Engine/Source/Runtime/Core/Internal/IO/IoOffsetLength.h
@@ -3,6 +3,7 @@
 #pragma once
 
 #include "CoreTypes.h"
+#include "HAL/UnrealMemory.h"
 
 /**
  * Combined offset and length.
@@ -55,6 +56,12 @@ struct FIoOffsetAndLength
 		OffsetAndLength[9] = uint8(Length >> 0);
 	}
 
+	/** Returns true if either the length or offset have been set, otherwise false */
+	inline bool IsValid() const
+	{
+		return !FMemory::MemIsZero(OffsetAndLength, 10);
+	}
+
 private:
 	// We use 5 bytes for offset and size, this is enough to represent
 	// an offset and size of 1PB
