[AutoRTFM] Create FTransactionallySafeSharedMutex.

FTransactionallySafeRWLock is now an alias for FTransactionallySafeSharedMutex when AutoRTFM is enabled, and so the underlying lock primitive is our lightweight FSharedMutex instead of the platform's FRWLock. This CL doesn't add support for any new FSharedMutex APIs, so this is functionally identical to an RWLock with different naming conventions and more efficient internals.

Under the hood, FSharedMutex exposes new APIs like `IsLocked` or `IsLockShared` that may be implemented in followup CLs if a need arises.

diff --git a/Engine/Source/Runtime/Core/Public/Async/TransactionallySafeSharedMutex.h b/Engine/Source/Runtime/Core/Public/Async/TransactionallySafeSharedMutex.h
new file mode 100644
--- /dev/null
+++ b/Engine/Source/Runtime/Core/Public/Async/TransactionallySafeSharedMutex.h
@@ -0,0 +1,203 @@
+// Copyright Epic Games, Inc. All Rights Reserved.
+
+#pragma once
+
+#include "Async/SharedMutex.h"
+#include "AutoRTFM.h"
+
+#if UE_AUTORTFM
+#include "Templates/SharedPointer.h"
+#endif
+
+namespace UE
+{
+
+#if UE_AUTORTFM
+
+// A transactionally-safe shared lock that works in the following novel ways:
+	// - In the open (non-transactional):
+	//   - Take the lock like before. Simple!
+	//   - Free the lock like before too.
+	// - In the closed (transactional):
+	//   - During locking we query `TransactionalLockCount`:
+	//	   - 0 means we haven't taken the lock within our transaction nest and need to acquire the lock.
+	//     - Otherwise we already have the lock (and are preventing non-transactional code seeing any
+	//       modifications we've made while holding the lock), so just bump `TransactionalLockCount`.
+	//     - We also register an on-abort handler to release the lock should we abort (but we need to
+	//       query `TransactionalLockCount` even there because we could be aborting an inner transaction
+	//       and the parent transaction still wants to have the lock held!).
+	//   - During unlocking we defer doing the unlock until the transaction commits.
+	//
+	// Thus with this approach we will hold this lock for the *entirety* of the transactional nest should
+	// we take the lock during the transaction, thus preventing non-transactional code from seeing any
+	// modifications we should make.
+	//
+	// If we are within a transaction, we pessimise our shared lock to a full lock. Note: that it should
+	// potentially be possible to have shared locks work correctly, but serious care will have to be taken to
+	// ensure that we don't have:
+	//   Open Thread     Closed Thread
+	//   -----------     SharedLock
+	//   -----------     SharedUnlock
+	//   Lock            -------------
+	//   Unlock          -------------
+	//   -----------     SharedLock      <- Invalid because the transaction can potentially observe side
+	//                                      effects of the open-threads writes!
+struct FTransactionallySafeSharedMutexDefinition
+{
+	FTransactionallySafeSharedMutexDefinition() : State(MakeShared<FState>())
+	{
+	}
+
+	void LockShared()
+	{
+		if (AutoRTFM::IsTransactional() || AutoRTFM::IsCommittingOrAborting())
+		{
+			// Transactionally pessimise LockShared -> Lock.
+			Lock();
+		}
+		else
+		{
+			State->Mutex.LockShared();
+			ensure(0 == State->TransactionalLockCount);
+		}
+	}
+
+	void UnlockShared()
+	{
+		if (AutoRTFM::IsTransactional() || AutoRTFM::IsCommittingOrAborting())
+		{
+			// Transactionally pessimise UnlockShared -> Unlock.
+			Unlock();
+		}
+		else
+		{
+			ensure(0 == State->TransactionalLockCount);
+			State->Mutex.UnlockShared();
+		}
+	}
+
+	void Lock()
+	{
+		if (AutoRTFM::IsTransactional() || AutoRTFM::IsCommittingOrAborting())
+		{
+			UE_AUTORTFM_OPEN
+			{
+				// The transactional system which can increment TransactionalLockCount
+				// is always single-threaded, thus this is safe to check without atomicity.
+				if (0 == State->TransactionalLockCount)
+				{
+					State->Mutex.Lock();
+				}
+
+				State->TransactionalLockCount += 1;
+			};
+
+			// We explicitly copy the state here for the case that `this` was stack
+			// allocated and has already died before the on-abort is hit.
+			UE_AUTORTFM_ONABORT(State = this->State)
+			{
+				State->Unlock();
+			};
+		}
+		else
+		{
+			State->Mutex.Lock();
+			ensure(0 == State->TransactionalLockCount);
+		}
+	}
+
+	void Unlock()
+	{
+		if (AutoRTFM::IsTransactional() || AutoRTFM::IsCommittingOrAborting())
+		{
+			// We explicitly copy the state here for the case that `this` was stack
+			// allocated and has already died before the on-commit is hit.
+			UE_AUTORTFM_ONCOMMIT(State = this->State)
+			{
+				State->Unlock();
+			};
+		}
+		else
+		{
+			ensure(0 == State->TransactionalLockCount);
+			State->Mutex.Unlock();
+		}
+	}
+
+	bool TryLock()
+	{
+		if (AutoRTFM::IsTransactional() || AutoRTFM::IsCommittingOrAborting())
+		{
+			bool Ret = false;
+
+			UE_AUTORTFM_OPEN
+			{
+				// The transactional system which can increment TransactionalLockCount
+				// is always single-threaded, thus this is safe to check without atomicity.
+				// For TryLock we should only lock when we have a 0 count as no one owns this lock
+				if (0 == State->TransactionalLockCount)
+				{
+					Ret = State->Mutex.TryLock();
+				}
+
+				if (Ret)
+				{
+					State->TransactionalLockCount += 1;
+				}
+			};
+
+			// Only setup the OnAbort if we *did* grab a lock; otherwise, we will not want to do 
+			// anything with the count or lock.
+			if (Ret)
+			{
+				// We explicitly copy the state here for the case that `this` was stack
+				// allocated and has already died before the on-abort is hit.
+				UE_AUTORTFM_ONABORT(State = this->State)
+				{
+					State->Unlock();
+				};
+			}
+
+			return Ret;
+		}
+
+		return State->Mutex.TryLock();
+	}
+
+private:
+	UE_NONCOPYABLE(FTransactionallySafeSharedMutexDefinition)
+
+	struct FState final
+	{
+		::UE::FSharedMutex Mutex;
+		uint32 TransactionalLockCount = 0;
+
+		FState() = default;
+		~FState()
+		{
+			ensure(0 == TransactionalLockCount);
+		}
+
+		void Unlock()
+		{
+			ensure(0 != TransactionalLockCount);
+
+			TransactionalLockCount -= 1;
+
+			if (0 == TransactionalLockCount)
+			{
+				Mutex.Unlock();
+			}
+		}
+	};
+
+	TSharedPtr<FState> State;
+};
+
+using FTransactionallySafeSharedMutex = ::UE::FTransactionallySafeSharedMutexDefinition;
+
+#else
+using FTransactionallySafeSharedMutex = ::UE::FSharedMutex;
+#endif
+
+}  // namespace UE
diff --git a/Engine/Source/Runtime/Core/Public/Misc/TransactionallySafeRWLock.h b/Engine/Source/Runtime/Core/Public/Misc/TransactionallySafeRWLock.h
--- a/Engine/Source/Runtime/Core/Public/Misc/TransactionallySafeRWLock.h
+++ b/Engine/Source/Runtime/Core/Public/Misc/TransactionallySafeRWLock.h
@@ -2,198 +2,51 @@
 
 #pragma once
 
-#include "HAL/CriticalSection.h"
 #include "AutoRTFM.h"
-#include "Templates/SharedPointer.h"
 
 #if UE_AUTORTFM
+#include "Async/TransactionallySafeSharedMutex.h"
+#else
+#include "HAL/CriticalSection.h"
+#endif // UE_AUTORTFM
+
 
-// A transactionally safe lock that works in the following novel ways:
-	// - In the open (non-transactional):
-	//   - Take the lock like before. Simple!
-	//   - Free the lock like before too.
-	// - In the closed (transactional):
-	//   - During locking we query `TransactionalLockCount`:
-	//	   - 0 means we haven't taken the lock within our transaction nest and need to acquire the lock.
-	//     - Otherwise we already have the lock (and are preventing non-transactional code seeing any
-	//       modifications we've made while holding the lock), so just bump `TransactionalLockCount`.
-	//     - We also register an on-abort handler to release the lock should we abort (but we need to
-	//       query `TransactionalLockCount` even there because we could be aborting an inner transaction
-	//       and the parent transaction still wants to have the lock held!).
-	//   - During unlocking we defer doing the unlock until the transaction commits.
-	//
-	// Thus with this approach we will hold this lock for the *entirety* of the transactional nest should
-	// we take the lock during the transaction, thus preventing non-transactional code from seeing any
-	// modifications we should make.
-	//
-	// If we are within a transaction, we pessimise our read-lock to a write-lock. Note: that it should
-	// potentially be possible to have read-locks work correctly, but serious care will have to be taken to
-	// ensure that we don't have:
-	//   Open Thread     Closed Thread
-	//   -----------     ReadLock
-	//   -----------     ReadUnlock
-	//   WriteLock       -------------
-	//   WriteUnlock     -------------
-	//   -----------     ReadLock      <- Invalid because the transaction can potentially observe side
-	//                                    effects of the open-threads writes!
-struct FTransactionallySafeRWLockDefinition
+#if UE_AUTORTFM
+
+// Transactionally-safe RWLocks are implemented in terms of FTransactionallySafeSharedMutex.
+// Unfortunately, the method names differ slightly between RWLock and SharedMutex, so we adapt the names here.
+class FTransactionallySafeRWLock : ::UE::FTransactionallySafeSharedMutex
 {
-	FTransactionallySafeRWLockDefinition() : State(MakeShared<FState>())
-	{
-	}
+private:
+	using Super = ::UE::FTransactionallySafeSharedMutex;
 
+public:
 	void ReadLock()
 	{
-		if (AutoRTFM::IsTransactional() || AutoRTFM::IsCommittingOrAborting())
-		{
-			// Transactionally pessimise ReadLock -> WriteLock.
-			WriteLock();
-		}
-		else
-		{
-			State->Lock.ReadLock();
-			ensure(0 == State->TransactionalLockCount);
-		}
+		Super::LockShared();
 	}
 
 	void ReadUnlock()
 	{
-		if (AutoRTFM::IsTransactional() || AutoRTFM::IsCommittingOrAborting())
-		{
-			// Transactionally pessimise ReadUnlock -> WriteUnlock.
-			WriteUnlock();
-		}
-		else
-		{
-			ensure(0 == State->TransactionalLockCount);
-			State->Lock.ReadUnlock();
-		}
+		Super::UnlockShared();
 	}
 
 	void WriteLock()
 	{
-		if (AutoRTFM::IsTransactional() || AutoRTFM::IsCommittingOrAborting())
-		{
-			UE_AUTORTFM_OPEN
-			{
-				// The transactional system which can increment TransactionalLockCount
-				// is always single-threaded, thus this is safe to check without atomicity.
-				if (0 == State->TransactionalLockCount)
-				{
-					State->Lock.WriteLock();
-				}
-
-				State->TransactionalLockCount += 1;
-			};
-
-			// We explicitly copy the state here for the case that `this` was stack
-			// allocated and has already died before the on-abort is hit.
-			UE_AUTORTFM_ONABORT(State = this->State)
-			{
-				State->WriteUnlock();
-			};
-		}
-		else
-		{
-			State->Lock.WriteLock();
-			ensure(0 == State->TransactionalLockCount);
-		}
+		Super::Lock();
 	}
 
 	void WriteUnlock()
 	{
-		if (AutoRTFM::IsTransactional() || AutoRTFM::IsCommittingOrAborting())
-		{
-			// We explicitly copy the state here for the case that `this` was stack
-			// allocated and has already died before the on-commit is hit.
-			UE_AUTORTFM_ONCOMMIT(State = this->State)
-			{
-				State->WriteUnlock();
-			};
-		}
-		else
-		{
-			ensure(0 == State->TransactionalLockCount);
-			State->Lock.WriteUnlock();
-		}
+		Super::Unlock();
 	}
 
 	bool TryWriteLock()
 	{
-		if (AutoRTFM::IsTransactional() || AutoRTFM::IsCommittingOrAborting())
-		{
-			bool Ret = false;
-
-			UE_AUTORTFM_OPEN
-			{
-				// The transactional system which can increment TransactionalLockCount
-				// is always single-threaded, thus this is safe to check without atomicity.
-				// For TryWriteLock we should only lock when we have a 0 count as no one owns this lock
-				if (0 == State->TransactionalLockCount)
-				{
-					Ret = State->Lock.TryWriteLock();
-				}
-
-				if (Ret)
-				{
-					State->TransactionalLockCount += 1;
-				}
-			};
-
-			// Only setup the OnAbort if we *did* grab a lock, otherwise we will not want to do anything with the count or lock
-			if (Ret)
-			{
-				// We explicitly copy the state here for the case that `this` was stack
-				// allocated and has already died before the on-abort is hit.
-				UE_AUTORTFM_ONABORT(State = this->State)
-				{
-					State->WriteUnlock();
-				};
-			}
-
-			return Ret;
-		}
-
-		return State->Lock.TryWriteLock();
+		return Super::TryLock();
 	}
-
-private:
-	UE_NONCOPYABLE(FTransactionallySafeRWLockDefinition)
-
-	struct FState final
-	{
-		FRWLock Lock;
-		uint32 TransactionalLockCount = 0;
-
-		// Constructor is always open because FRWLock is not transactionally safe.
-		UE_AUTORTFM_ALWAYS_OPEN
-		FState() = default;
-		
-		// Destructor is always open because FRWLock is not transactionally safe.
-		UE_AUTORTFM_ALWAYS_OPEN
-		~FState()
-		{
-			ensure(0 == TransactionalLockCount);
-		}
-
-		void WriteUnlock()
-		{
-			ensure(0 != TransactionalLockCount);
-
-			TransactionalLockCount -= 1;
-
-			if (0 == TransactionalLockCount)
-			{
-				Lock.WriteUnlock();
-			}
-		}
-	};
-
-	TSharedPtr<FState> State;
 };
 
-using FTransactionallySafeRWLock = FTransactionallySafeRWLockDefinition;
-
 #else
-using FTransactionallySafeRWLock = FRWLock;
-#endif
+using FTransactionallySafeRWLock = ::FRWLock;
+#endif
\ No newline at end of file
