Add an IsEngineExitRequested check to SlowTask to prevent ensures that fire when some async tasks aren't completed yet and the engine is shutting down.

diff --git a/Engine/Source/Runtime/Core/Private/Misc/SlowTask.cpp b/Engine/Source/Runtime/Core/Private/Misc/SlowTask.cpp
--- a/Engine/Source/Runtime/Core/Private/Misc/SlowTask.cpp
+++ b/Engine/Source/Runtime/Core/Private/Misc/SlowTask.cpp
@@ -146,7 +146,7 @@ void FSlowTask::Destroy()
 #endif
 
 			FSlowTask* Task = Stack.Last();
-			if (ensureMsgf(Task == this, TEXT("Out-of-order slow task construction/destruction: destroying '%s' but '%s' is at the top of the stack"), *DefaultMessage.ToString(), *Task->DefaultMessage.ToString()))
+			if (!IsEngineExitRequested() && ensureMsgf(Task == this, TEXT("Out-of-order slow task construction/destruction: destroying '%s' but '%s' is at the top of the stack"), *DefaultMessage.ToString(), *Task->DefaultMessage.ToString()))
 			{
 				Stack.Pop(EAllowShrinking::No);
 			}
