[Insights] Fixed platform events tracing to allow "ContextSwitch" and "StackSampling" to be enabled in the same time.

diff --git a/Engine/Source/Runtime/Core/Private/ProfilingDebugging/PlatformEvents.cpp b/Engine/Source/Runtime/Core/Private/ProfilingDebugging/PlatformEvents.cpp
--- a/Engine/Source/Runtime/Core/Private/ProfilingDebugging/PlatformEvents.cpp
+++ b/Engine/Source/Runtime/Core/Private/ProfilingDebugging/PlatformEvents.cpp
@@ -113,7 +113,8 @@ void FPlatformEventsTrace::PostInit()
 	{
 		Enable(FPlatformEventsTrace::EEventType::ContextSwitch);
 	}
-	else if (UE_TRACE_CHANNELEXPR_IS_ENABLED(StackSamplingChannel))
+
+	if (UE_TRACE_CHANNELEXPR_IS_ENABLED(StackSamplingChannel))
 	{
 		Enable(FPlatformEventsTrace::EEventType::StackSampling);
 	}
