Fixed stall detector shutting down on 1st loop when it was expected to shutdown on critical error.

diff --git a/Engine/Source/Runtime/Core/Private/ProfilingDebugging/StallDetector.cpp b/Engine/Source/Runtime/Core/Private/ProfilingDebugging/StallDetector.cpp
--- a/Engine/Source/Runtime/Core/Private/ProfilingDebugging/StallDetector.cpp
+++ b/Engine/Source/Runtime/Core/Private/ProfilingDebugging/StallDetector.cpp
@@ -148,7 +148,7 @@ uint32 UE::FStallDetectorRunnable::Run()
 			FScopeLock ScopeLock(&StallScopesSection);
 
 			// Stop if there has been a crash, do not report crash as stall, Detector->Stats may also be deleted memory
-			if (!GIsCriticalError)
+			if (GIsCriticalError)
 			{
 				break;
 			}
