Add missing #include for PlatformTLS.

Fixes Horde Issue 978566: https://horde.devtools.epicgames.com/job/688258d7ce1c566fd875458a

diff --git a/Engine/Source/Runtime/Core/Public/Containers/StripedMap.h b/Engine/Source/Runtime/Core/Public/Containers/StripedMap.h
--- a/Engine/Source/Runtime/Core/Public/Containers/StripedMap.h
+++ b/Engine/Source/Runtime/Core/Public/Containers/StripedMap.h
@@ -8,6 +8,7 @@
 #include "Async/SharedLock.h"
 #include "Async/SharedMutex.h"
 #include "Async/UniqueLock.h"
+#include "HAL/PlatformTLS.h"
 #include "Misc/TransactionallySafeRWLock.h"
 #include "Misc/ScopeRWLock.h"
 
