Made TArrayView and Invoke constexpr. Fixed UB GetData and deprecated Alignment in TStaticArray

diff --git a/Engine/Plugins/Experimental/NNERuntimeRDG/Source/NNERuntimeRDG/Private/NNERuntimeRDGHlslHelper.h b/Engine/Plugins/Experimental/NNERuntimeRDG/Source/NNERuntimeRDG/Private/NNERuntimeRDGHlslHelper.h
--- a/Engine/Plugins/Experimental/NNERuntimeRDG/Source/NNERuntimeRDG/Private/NNERuntimeRDGHlslHelper.h
+++ b/Engine/Plugins/Experimental/NNERuntimeRDG/Source/NNERuntimeRDG/Private/NNERuntimeRDGHlslHelper.h
@@ -3,6 +3,7 @@
 #pragma once
 
 #include "Containers/StaticArray.h"
+#include "Misc/AlignedElement.h"
 #include "Templates/Function.h"
 
 namespace UE::NNERuntimeRDG::Private { class FTensor; }
@@ -11,7 +12,7 @@ namespace UE::NNERuntimeRDG::Private { class FTensor; }
 
 namespace UE::NNERuntimeRDG::Private::Hlsl
 {
-	using FTensorInfoParam = TStaticArray<FUintVector4, NXRT_TENSORSTRIDEINFO_MAX_NUM_DIMENSIONS, 16U>;
+	using FTensorInfoParam = TStaticArray<UE::Core::TAlignedElement<FUintVector4, 16>, NXRT_TENSORSTRIDEINFO_MAX_NUM_DIMENSIONS>;
 
 	struct FTensorInfoParamArraySpan 
 	{
diff --git a/Engine/Source/Runtime/Core/Public/Containers/ArrayView.h b/Engine/Source/Runtime/Core/Public/Containers/ArrayView.h
--- a/Engine/Source/Runtime/Core/Public/Containers/ArrayView.h
+++ b/Engine/Source/Runtime/Core/Public/Containers/ArrayView.h
@@ -29,7 +29,7 @@ namespace UE::Core::ArrayView::Private
 	// Simply forwards to an unqualified GetData(), but can be called from within TArrayView
 	// where GetData() is already a member and so hides any others.
 	template <typename T>
-	UE_FORCEINLINE_HINT auto GetDataHelper(T&& Arg) -> decltype(GetData(Forward<T>(Arg)))
+	UE_FORCEINLINE_HINT constexpr auto GetDataHelper(T&& Arg) -> decltype(GetData(Forward<T>(Arg)))
 	{
 		return GetData(Forward<T>(Arg));
 	}
@@ -57,7 +57,7 @@ namespace UE::Core::ArrayView::Private
 		static constexpr bool Value = TIsCompatibleElementType_V<std::remove_pointer_t<decltype(GetData(DeclVal<RangeType&>()))>, ElementType>;
 
 		template <typename T>
-		static decltype(auto) GetData(T&& Arg)
+		static constexpr decltype(auto) GetData(T&& Arg)
 		{
 			return UE::Core::ArrayView::Private::GetDataHelper(Forward<T>(Arg));
 		}
@@ -151,7 +151,7 @@ public:
 	/**
 	 * Constructor.
 	 */
-	[[nodiscard]] TArrayView()
+	[[nodiscard]] constexpr TArrayView()
 		: DataPtr(nullptr)
 		, ArrayNum(0)
 	{
@@ -185,7 +185,7 @@ public:
 			!std::is_same_v<CVUnqualifiedOtherRangeType, TArrayView>
 		)
 	>
-	[[nodiscard]] inline TArrayView(OtherRangeType&& Other)
+	[[nodiscard]] constexpr TArrayView(OtherRangeType&& Other)
 		: DataPtr(std::conditional_t<
 						TIsCompatibleRangeType<OtherRangeType>::Value,
 						TIsCompatibleRangeType<OtherRangeType>,
@@ -217,7 +217,7 @@ public:
 			!TIsTArrayView_V<CVUnqualifiedOtherRangeType>
 		)
 	>
-	[[nodiscard]] inline TArrayView(OtherRangeType&& Other UE_LIFETIMEBOUND)
+	[[nodiscard]] constexpr TArrayView(OtherRangeType&& Other UE_LIFETIMEBOUND)
 		: DataPtr(std::conditional_t<
 			TIsCompatibleRangeType<OtherRangeType>::Value,
 			TIsCompatibleRangeType<OtherRangeType>,
@@ -247,7 +247,7 @@ public:
 		typename OtherElementType
 		UE_REQUIRES(UE::Core::ArrayView::Private::TIsCompatibleElementType_V<OtherElementType, ElementType>)
 	>
-	[[nodiscard]] inline TArrayView(OtherElementType* InData UE_LIFETIMEBOUND, SizeType InCount)
+	[[nodiscard]] constexpr TArrayView(OtherElementType* InData UE_LIFETIMEBOUND, SizeType InCount)
 		: DataPtr(InData)
 		, ArrayNum(InCount)
 	{
@@ -259,7 +259,7 @@ public:
 	 *
 	 * The caller is responsible for ensuring that the view does not outlive the initializer list.
 	 */
-	[[nodiscard]] inline TArrayView(std::initializer_list<ElementType> List UE_LIFETIMEBOUND)
+	[[nodiscard]] constexpr TArrayView(std::initializer_list<ElementType> List UE_LIFETIMEBOUND)
 		: DataPtr(UE::Core::ArrayView::Private::GetDataHelper(List))
 		, ArrayNum(GetNum(List))
 	{
@@ -272,12 +272,12 @@ public:
 	constexpr static bool bHasIntrusiveUnsetOptionalState = true;
 	using IntrusiveUnsetOptionalStateType = TArrayView;
 
-	[[nodiscard]] explicit TArrayView(FIntrusiveUnsetOptionalState)
+	[[nodiscard]] explicit constexpr TArrayView(FIntrusiveUnsetOptionalState)
 		: DataPtr(nullptr)
 		, ArrayNum(-1)
 	{
 	}
-	[[nodiscard]] bool UEOpEquals(FIntrusiveUnsetOptionalState) const
+	[[nodiscard]] constexpr bool UEOpEquals(FIntrusiveUnsetOptionalState) const
 	{
 		return ArrayNum == -1;
 	}
@@ -292,7 +292,7 @@ public:
 	 *
 	 * @returns Pointer to first array entry.
 	 */
-	[[nodiscard]] UE_FORCEINLINE_HINT ElementType* GetData() const
+	[[nodiscard]] UE_FORCEINLINE_HINT constexpr ElementType* GetData() const
 	{
 		return DataPtr;
 	}
@@ -318,7 +318,7 @@ public:
 	/**
 	 * Checks array invariants: if array size is greater than or equal to zero.
 	 */
-	UE_FORCEINLINE_HINT void CheckInvariants() const
+	UE_FORCEINLINE_HINT constexpr void CheckInvariants() const
 	{
 		checkSlow(ArrayNum >= 0);
 	}
@@ -328,7 +328,7 @@ public:
 	 *
 	 * @param Index Index to check.
 	 */
-	inline void RangeCheck(SizeType Index) const
+	constexpr void RangeCheck(SizeType Index) const
 	{
 		CheckInvariants();
 
@@ -342,7 +342,7 @@ public:
 	 * @param Index Starting index of the slice.
 	 * @param InNum Length of the slice.
 	 */
-	inline void SliceRangeCheck(SizeType Index, SizeType InNum) const
+	constexpr void SliceRangeCheck(SizeType Index, SizeType InNum) const
 	{
 		checkf(Index >= 0, TEXT("Invalid index (%lld)"), (long long)Index);
 		checkf(InNum >= 0, TEXT("Invalid count (%lld)"), (long long)InNum);
@@ -356,7 +356,7 @@ public:
 	 *
 	 * @returns True if index is valid. False otherwise.
 	 */
-	[[nodiscard]] UE_FORCEINLINE_HINT bool IsValidIndex(SizeType Index) const
+	[[nodiscard]] UE_FORCEINLINE_HINT constexpr bool IsValidIndex(SizeType Index) const
 	{
 		return (Index >= 0) && (Index < ArrayNum);
 	}
@@ -367,7 +367,7 @@ public:
 	 * @returns True if the array view is empty.
 	 * @see Num
 	 */
-	[[nodiscard]] bool IsEmpty() const
+	[[nodiscard]] constexpr bool IsEmpty() const
 	{
 		return ArrayNum == 0;
 	}
@@ -377,13 +377,13 @@ public:
 	 *
 	 * @returns Number of elements in array.
 	 */
-	[[nodiscard]] UE_FORCEINLINE_HINT SizeType Num() const
+	[[nodiscard]] UE_FORCEINLINE_HINT constexpr SizeType Num() const
 	{
 		return ArrayNum;
 	}
 
 	/** @returns Number of bytes used, excluding slack */
-	[[nodiscard]] UE_FORCEINLINE_HINT SIZE_T NumBytes() const
+	[[nodiscard]] UE_FORCEINLINE_HINT constexpr SIZE_T NumBytes() const
 	{
 		return static_cast<SIZE_T>(ArrayNum) * sizeof(ElementType);
 	}
@@ -393,7 +393,7 @@ public:
 	 *
 	 * @returns Reference to indexed element.
 	 */
-	[[nodiscard]] inline ElementType& operator[](SizeType Index) const
+	[[nodiscard]] constexpr ElementType& operator[](SizeType Index) const
 	{
 		RangeCheck(Index);
 		return GetData()[Index];
@@ -407,7 +407,7 @@ public:
 	 *
 	 * @returns Reference to n-th last element from the array.
 	 */
-	[[nodiscard]] inline ElementType& Last(SizeType IndexFromTheEnd = 0) const
+	[[nodiscard]] constexpr ElementType& Last(SizeType IndexFromTheEnd = 0) const
 	{
 		RangeCheck(ArrayNum - IndexFromTheEnd - 1);
 		return GetData()[ArrayNum - IndexFromTheEnd - 1];
@@ -423,40 +423,40 @@ public:
 	 *
 	 * @see Mid
 	 */
-	[[nodiscard]] inline TArrayView Slice(SizeType Index, SizeType InNum) const
+	[[nodiscard]] constexpr TArrayView Slice(SizeType Index, SizeType InNum) const
 	{
 		SliceRangeCheck(Index, InNum);
 		return TArrayView(DataPtr + Index, InNum);
 	}
 
 	/** Returns the left-most part of the view by taking the given number of elements from the left. */
-	[[nodiscard]] inline TArrayView Left(SizeType Count) const
+	[[nodiscard]] constexpr TArrayView Left(SizeType Count) const
 	{
 		return TArrayView(DataPtr, FMath::Clamp(Count, (SizeType)0, ArrayNum));
 	}
 
 	/** Returns the left-most part of the view by chopping the given number of elements from the right. */
-	[[nodiscard]] inline TArrayView LeftChop(SizeType Count) const
+	[[nodiscard]] constexpr TArrayView LeftChop(SizeType Count) const
 	{
 		return TArrayView(DataPtr, FMath::Clamp(ArrayNum - Count, (SizeType)0, ArrayNum));
 	}
 
 	/** Returns the right-most part of the view by taking the given number of elements from the right. */
-	[[nodiscard]] inline TArrayView Right(SizeType Count) const
+	[[nodiscard]] constexpr TArrayView Right(SizeType Count) const
 	{
 		const SizeType OutLen = FMath::Clamp(Count, (SizeType)0, ArrayNum);
 		return TArrayView(DataPtr + ArrayNum - OutLen, OutLen);
 	}
 
 	/** Returns the right-most part of the view by chopping the given number of elements from the left. */
-	[[nodiscard]] inline TArrayView RightChop(SizeType Count) const
+	[[nodiscard]] constexpr TArrayView RightChop(SizeType Count) const
 	{
 		const SizeType OutLen = FMath::Clamp(ArrayNum - Count, (SizeType)0, ArrayNum);
 		return TArrayView(DataPtr + ArrayNum - OutLen, OutLen);
 	}
 
 	/** Returns the middle part of the view by taking up to the given number of elements from the given position. */
-	[[nodiscard]] inline TArrayView Mid(SizeType Index, SizeType Count = TNumericLimits<SizeType>::Max()) const
+	[[nodiscard]] constexpr TArrayView Mid(SizeType Index, SizeType Count = TNumericLimits<SizeType>::Max()) const
 	{
 		ElementType* const CurrentStart  = GetData();
 		const SizeType     CurrentLength = Num();
@@ -477,31 +477,31 @@ public:
 	}
 
 	/** Modifies the view to be the given number of elements from the left. */
-	inline void LeftInline(SizeType CharCount)
+	constexpr void LeftInline(SizeType CharCount)
 	{
 		*this = Left(CharCount);
 	}
 
 	/** Modifies the view by chopping the given number of elements from the right. */
-	inline void LeftChopInline(SizeType CharCount)
+	constexpr void LeftChopInline(SizeType CharCount)
 	{
 		*this = LeftChop(CharCount);
 	}
 
 	/** Modifies the view to be the given number of elements from the right. */
-	inline void RightInline(SizeType CharCount)
+	constexpr void RightInline(SizeType CharCount)
 	{
 		*this = Right(CharCount);
 	}
 
 	/** Modifies the view by chopping the given number of elements from the left. */
-	inline void RightChopInline(SizeType CharCount)
+	constexpr void RightChopInline(SizeType CharCount)
 	{
 		*this = RightChop(CharCount);
 	}
 
 	/** Modifies the view to be the middle part by taking up to the given number of elements from the given position. */
-	inline void MidInline(SizeType Position, SizeType CharCount = TNumericLimits<SizeType>::Max())
+	constexpr void MidInline(SizeType Position, SizeType CharCount = TNumericLimits<SizeType>::Max())
 	{
 		*this = Mid(Position, CharCount);
 	}
@@ -514,7 +514,7 @@ public:
 	 *
 	 * @returns True if found. False otherwise.
 	 */
-	[[nodiscard]] inline bool Find(const ElementType& Item, SizeType& Index) const
+	[[nodiscard]] constexpr bool Find(const ElementType& Item, SizeType& Index) const
 	{
 		Index = this->Find(Item);
 		return Index != static_cast<SizeType>(INDEX_NONE);
@@ -527,7 +527,7 @@ public:
 	 *
 	 * @returns Index of the found element. INDEX_NONE otherwise.
 	 */
-	[[nodiscard]] SizeType Find(const ElementType& Item) const
+	[[nodiscard]] constexpr SizeType Find(const ElementType& Item) const
 	{
 		const ElementType* RESTRICT Start = GetData();
 		for (const ElementType* RESTRICT Data = Start, *RESTRICT DataEnd = Data + ArrayNum; Data != DataEnd; ++Data)
@@ -548,7 +548,7 @@ public:
 	 *
 	 * @returns True if found. False otherwise.
 	 */
-	[[nodiscard]] inline bool FindLast(const ElementType& Item, SizeType& Index) const
+	[[nodiscard]] constexpr bool FindLast(const ElementType& Item, SizeType& Index) const
 	{
 		Index = this->FindLast(Item);
 		return Index != static_cast<SizeType>(INDEX_NONE);
@@ -560,7 +560,7 @@ public:
 	 * @param Item Item to look for.
 	 * @returns Index of the found element. INDEX_NONE otherwise.
 	 */
-	[[nodiscard]] SizeType FindLast(const ElementType& Item) const
+	[[nodiscard]] constexpr SizeType FindLast(const ElementType& Item) const
 	{
 		for (const ElementType* RESTRICT Start = GetData(), *RESTRICT Data = Start + ArrayNum; Data != Start; )
 		{
@@ -582,7 +582,7 @@ public:
 	 * @returns Index of the found element. INDEX_NONE otherwise.
 	 */
 	template <typename Predicate>
-	[[nodiscard]] SizeType FindLastByPredicate(Predicate Pred, SizeType StartIndex) const
+	[[nodiscard]] constexpr SizeType FindLastByPredicate(Predicate Pred, SizeType StartIndex) const
 	{
 		check(StartIndex >= 0 && StartIndex <= this->Num());
 		for (const ElementType* RESTRICT Start = GetData(), *RESTRICT Data = Start + StartIndex; Data != Start; )
@@ -604,7 +604,7 @@ public:
 	* @returns Index of the found element. INDEX_NONE otherwise.
 	*/
 	template <typename Predicate>
-	[[nodiscard]] UE_FORCEINLINE_HINT SizeType FindLastByPredicate(Predicate Pred) const
+	[[nodiscard]] UE_FORCEINLINE_HINT constexpr SizeType FindLastByPredicate(Predicate Pred) const
 	{
 		return FindLastByPredicate(Pred, ArrayNum);
 	}
@@ -619,7 +619,7 @@ public:
 	 *          found.
 	 */
 	template <typename KeyType>
-	[[nodiscard]] SizeType IndexOfByKey(const KeyType& Key) const
+	[[nodiscard]] constexpr SizeType IndexOfByKey(const KeyType& Key) const
 	{
 		const ElementType* RESTRICT Start = GetData();
 		for (const ElementType* RESTRICT Data = Start, *RESTRICT DataEnd = Start + ArrayNum; Data != DataEnd; ++Data)
@@ -641,7 +641,7 @@ public:
 	 *          found.
 	 */
 	template <typename Predicate>
-	[[nodiscard]] SizeType IndexOfByPredicate(Predicate Pred) const
+	[[nodiscard]] constexpr SizeType IndexOfByPredicate(Predicate Pred) const
 	{
 		const ElementType* RESTRICT Start = GetData();
 		for (const ElementType* RESTRICT Data = Start, *RESTRICT DataEnd = Start + ArrayNum; Data != DataEnd; ++Data)
@@ -663,7 +663,7 @@ public:
 	 * @returns Pointer to the first matching element, or nullptr if none is found.
 	 */
 	template <typename KeyType>
-	[[nodiscard]] ElementType* FindByKey(const KeyType& Key) const
+	[[nodiscard]] constexpr ElementType* FindByKey(const KeyType& Key) const
 	{
 		for (ElementType* RESTRICT Data = GetData(), *RESTRICT DataEnd = Data + ArrayNum; Data != DataEnd; ++Data)
 		{
@@ -685,7 +685,7 @@ public:
 	 *         true, or nullptr if none is found.
 	 */
 	template <typename Predicate>
-	[[nodiscard]] ElementType* FindByPredicate(Predicate Pred) const
+	[[nodiscard]] constexpr ElementType* FindByPredicate(Predicate Pred) const
 	{
 		for (ElementType* RESTRICT Data = GetData(), *RESTRICT DataEnd = Data + ArrayNum; Data != DataEnd; ++Data)
 		{
@@ -707,7 +707,7 @@ public:
 	 *          the subset of elements for which the functor returns true.
 	 */
 	template <typename Predicate>
-	[[nodiscard]] TArray<std::remove_const_t<ElementType>> FilterByPredicate(Predicate Pred) const
+	[[nodiscard]] constexpr TArray<std::remove_const_t<ElementType>> FilterByPredicate(Predicate Pred) const
 	{
 		TArray<std::remove_const_t<ElementType>> FilterResults;
 		for (const ElementType* RESTRICT Data = GetData(), *RESTRICT DataEnd = Data + ArrayNum; Data != DataEnd; ++Data)
@@ -726,7 +726,7 @@ public:
 	 * @returns	True if found. False otherwise.
 	 */
 	template <typename ComparisonType>
-	[[nodiscard]] bool Contains(const ComparisonType& Item) const
+	[[nodiscard]] constexpr bool Contains(const ComparisonType& Item) const
 	{
 		for (const ElementType* RESTRICT Data = GetData(), *RESTRICT DataEnd = Data + ArrayNum; Data != DataEnd; ++Data)
 		{
@@ -746,7 +746,7 @@ public:
 	 * @returns	True if found. False otherwise.
 	 */
 	template <typename Predicate>
-	[[nodiscard]] UE_FORCEINLINE_HINT bool ContainsByPredicate(Predicate Pred) const
+	[[nodiscard]] UE_FORCEINLINE_HINT constexpr bool ContainsByPredicate(Predicate Pred) const
 	{
 		return FindByPredicate(Pred) != nullptr;
 	}
@@ -756,10 +756,10 @@ public:
 	 * DO NOT USE DIRECTLY
 	 * STL-like iterators to enable range-based for loop support.
 	 */
-	[[nodiscard]] UE_FORCEINLINE_HINT ElementType*                         begin () const { return GetData(); }
-	[[nodiscard]] UE_FORCEINLINE_HINT ElementType*                         end   () const { return GetData() + Num(); }
-	[[nodiscard]] UE_FORCEINLINE_HINT TReversePointerIterator<ElementType> rbegin() const { return TReversePointerIterator<ElementType>(GetData() + Num()); }
-	[[nodiscard]] UE_FORCEINLINE_HINT TReversePointerIterator<ElementType> rend  () const { return TReversePointerIterator<ElementType>(GetData()); }
+	[[nodiscard]] UE_FORCEINLINE_HINT constexpr ElementType*                         begin () const { return GetData(); }
+	[[nodiscard]] UE_FORCEINLINE_HINT constexpr ElementType*                         end   () const { return GetData() + Num(); }
+	[[nodiscard]] UE_FORCEINLINE_HINT constexpr TReversePointerIterator<ElementType> rbegin() const { return TReversePointerIterator<ElementType>(GetData() + Num()); }
+	[[nodiscard]] UE_FORCEINLINE_HINT constexpr TReversePointerIterator<ElementType> rend  () const { return TReversePointerIterator<ElementType>(GetData()); }
 
 public:
 	/**
@@ -770,7 +770,7 @@ public:
 	 *        If this is not desirable, please use Algo::Sort(MyArray) directly instead.
 	 *        The auto-dereferencing behavior does not occur with smart pointers.
 	 */
-	void Sort()
+	constexpr void Sort()
 	{
 		Algo::Sort(*this, TDereferenceWrapper<ElementType, TLess<>>(TLess<>()));
 	}
@@ -785,7 +785,7 @@ public:
 	 *        The auto-dereferencing behavior does not occur with smart pointers.
 	 */
 	template <class PREDICATE_CLASS>
-	void Sort(const PREDICATE_CLASS& Predicate)
+	constexpr void Sort(const PREDICATE_CLASS& Predicate)
 	{
 		TDereferenceWrapper<ElementType, PREDICATE_CLASS> PredicateWrapper(Predicate);
 		Algo::Sort(*this, PredicateWrapper);
@@ -800,7 +800,7 @@ public:
 	 *        If this is not desirable, please use Algo::StableSort(MyArray) directly instead.
 	 *        The auto-dereferencing behavior does not occur with smart pointers.
 	 */
-	void StableSort()
+	constexpr void StableSort()
 	{
 		Algo::StableSort(*this, TDereferenceWrapper<ElementType, TLess<>>(TLess<>()));
 	}
@@ -817,7 +817,7 @@ public:
 	 *        The auto-dereferencing behavior does not occur with smart pointers.
 	 */
 	template <class PREDICATE_CLASS>
-	void StableSort(const PREDICATE_CLASS& Predicate)
+	constexpr void StableSort(const PREDICATE_CLASS& Predicate)
 	{
 		TDereferenceWrapper<ElementType, PREDICATE_CLASS> PredicateWrapper(Predicate);
 		Algo::StableSort(*this, PredicateWrapper);
@@ -867,7 +867,7 @@ template <
 	typename CVUnqualifiedOtherRangeType = std::remove_cv_t<std::remove_reference_t<OtherRangeType>>
 	UE_REQUIRES(TIsContiguousContainer<CVUnqualifiedOtherRangeType>::Value && TIsTArrayView_V<CVUnqualifiedOtherRangeType>)
 >
-[[nodiscard]] auto MakeArrayView(OtherRangeType&& Other)
+[[nodiscard]] constexpr auto MakeArrayView(OtherRangeType&& Other)
 {
 	return TArrayView<std::remove_pointer_t<decltype(GetData(DeclVal<OtherRangeType&>()))>>(Forward<OtherRangeType>(Other));
 }
@@ -876,19 +876,19 @@ template <
 	typename CVUnqualifiedOtherRangeType = std::remove_cv_t<std::remove_reference_t<OtherRangeType>>
 	UE_REQUIRES(TIsContiguousContainer<CVUnqualifiedOtherRangeType>::Value && !TIsTArrayView_V<CVUnqualifiedOtherRangeType>)
 >
-[[nodiscard]] auto MakeArrayView(OtherRangeType&& Other UE_LIFETIMEBOUND)
+[[nodiscard]] constexpr auto MakeArrayView(OtherRangeType&& Other UE_LIFETIMEBOUND)
 {
 	return TArrayView<std::remove_pointer_t<decltype(GetData(DeclVal<OtherRangeType&>()))>>(Forward<OtherRangeType>(Other));
 }
 
 template<typename ElementType>
-[[nodiscard]] TArrayView<ElementType> MakeArrayView(ElementType* Pointer UE_LIFETIMEBOUND, int32 Size)
+[[nodiscard]] constexpr TArrayView<ElementType> MakeArrayView(ElementType* Pointer UE_LIFETIMEBOUND, int32 Size)
 {
 	return TArrayView<ElementType>(Pointer, Size);
 }
 
 template <typename ElementType>
-[[nodiscard]] TArrayView<const ElementType> MakeArrayView(std::initializer_list<ElementType> List UE_LIFETIMEBOUND)
+[[nodiscard]] constexpr TArrayView<const ElementType> MakeArrayView(std::initializer_list<ElementType> List UE_LIFETIMEBOUND)
 {
 	return TArrayView<const ElementType>(List.begin(), List.size());
 }
@@ -898,7 +898,7 @@ template <
 	typename CVUnqualifiedOtherRangeType = std::remove_cv_t<std::remove_reference_t<OtherRangeType>>
 	UE_REQUIRES(TIsContiguousContainer<CVUnqualifiedOtherRangeType>::Value && TIsTArrayView_V<CVUnqualifiedOtherRangeType>)
 >
-[[nodiscard]] auto MakeConstArrayView(OtherRangeType&& Other)
+[[nodiscard]] constexpr auto MakeConstArrayView(OtherRangeType&& Other)
 {
 	return TArrayView<const std::remove_pointer_t<decltype(GetData(DeclVal<OtherRangeType&>()))>>(Forward<OtherRangeType>(Other));
 }
@@ -907,19 +907,19 @@ template <
 	typename CVUnqualifiedOtherRangeType = std::remove_cv_t<std::remove_reference_t<OtherRangeType>>
 	UE_REQUIRES(TIsContiguousContainer<CVUnqualifiedOtherRangeType>::Value && !TIsTArrayView_V<CVUnqualifiedOtherRangeType>)
 >
-[[nodiscard]] auto MakeConstArrayView(OtherRangeType&& Other UE_LIFETIMEBOUND)
+[[nodiscard]] constexpr auto MakeConstArrayView(OtherRangeType&& Other UE_LIFETIMEBOUND)
 {
 	return TArrayView<const std::remove_pointer_t<decltype(GetData(DeclVal<OtherRangeType&>()))>>(Forward<OtherRangeType>(Other));
 }
 
 template<typename ElementType>
-[[nodiscard]] TArrayView<const ElementType> MakeConstArrayView(const ElementType* Pointer UE_LIFETIMEBOUND, int32 Size)
+[[nodiscard]] constexpr TArrayView<const ElementType> MakeConstArrayView(const ElementType* Pointer UE_LIFETIMEBOUND, int32 Size)
 {
 	return TArrayView<const ElementType>(Pointer, Size);
 }
 
 template <typename ElementType>
-[[nodiscard]] TArrayView<const ElementType> MakeConstArrayView(std::initializer_list<ElementType> List UE_LIFETIMEBOUND)
+[[nodiscard]] constexpr TArrayView<const ElementType> MakeConstArrayView(std::initializer_list<ElementType> List UE_LIFETIMEBOUND)
 {
 	return TArrayView<const ElementType>(List.begin(), List.size());
 }
diff --git a/Engine/Source/Runtime/Core/Public/Containers/StaticArray.h b/Engine/Source/Runtime/Core/Public/Containers/StaticArray.h
--- a/Engine/Source/Runtime/Core/Public/Containers/StaticArray.h
+++ b/Engine/Source/Runtime/Core/Public/Containers/StaticArray.h
@@ -3,6 +3,8 @@
 #pragma once
 
 #include "CoreTypes.h"
+// Deprecated in 5.7. Only needed for backward compatibility.
+#include "Misc/AlignedElement.h"
 #include "Misc/AssertionMacros.h"
 #include "Misc/ReverseIterate.h"
 #include "Templates/UnrealTemplate.h"
@@ -19,28 +21,30 @@ namespace UE::Core::Private
 }
 
 /** An array with a static number of elements. */
-template <typename InElementType, uint32 NumElements, uint32 Alignment = alignof(InElementType)>
-class alignas(Alignment) TStaticArray
+template <typename InElementType, uint32 NumElements, uint32 Alignment = uint32(-1)>
+class TStaticArray
 {
+	// Deprecated in 5.7. Friend only needed for backward compatibility.
+	template <typename, uint32, uint32> friend class TStaticArray;
+
 public:
 	using ElementType = InElementType;
 
+	InElementType Elements[NumElements];
+
 	[[nodiscard]] constexpr TStaticArray() = default;
 
 	// Constructs each element with Args
 	template <typename... ArgTypes>
 	[[nodiscard]] constexpr explicit TStaticArray(EInPlace, ArgTypes&&... Args)
-		: Storage(InPlace, TMakeIntegerSequence<uint32, NumElements>(), Forward<ArgTypes>(Args)...)
+		: TStaticArray{InPlace, TMakeIntegerSequence<uint32, NumElements>(), [&Args...](uint32) { return InElementType(Args...); }}
 	{
 	}
 
 	// Directly initializes the array with the provided values.
-	template <
-		typename... ArgTypes
-		UE_REQUIRES((sizeof...(ArgTypes) > 0 && sizeof...(ArgTypes) <= NumElements) && UE::Core::Private::TCanBeConvertedToFromAll_V<InElementType, ArgTypes...>)
-	>
+	template <typename... ArgTypes> requires((sizeof...(ArgTypes) > 0 && sizeof...(ArgTypes) <= NumElements) && UE::Core::Private::TCanBeConvertedToFromAll_V<InElementType, ArgTypes...>)
 	[[nodiscard]] constexpr TStaticArray(ArgTypes&&... Args)
-		: Storage(PerElement, Forward<ArgTypes>(Args)...)
+		: TStaticArray{PerElement, [&Args] { return InElementType(Forward<ArgTypes>(Args)); }...}
 	{
 	}
 
@@ -50,42 +54,19 @@ public:
 	constexpr TStaticArray& operator=(const TStaticArray& Other) = default;
 
 	// Accessors.
-	[[nodiscard]] inline constexpr InElementType& operator[](uint32 Index)
+	[[nodiscard]] constexpr InElementType& operator[](uint32 Index)
 	{
 		checkSlow(Index < NumElements);
-		return Storage.Elements[Index].Element;
+		return Elements[Index];
 	}
 
-	[[nodiscard]] inline constexpr const InElementType& operator[](uint32 Index) const
+	[[nodiscard]] constexpr const InElementType& operator[](uint32 Index) const
 	{
 		checkSlow(Index < NumElements);
-		return Storage.Elements[Index].Element;
+		return Elements[Index];
 	}
 
-	// Comparisons.
-	[[nodiscard]] constexpr friend bool operator==(const TStaticArray& A,const TStaticArray& B)
-	{
-		for(uint32 ElementIndex = 0;ElementIndex < NumElements;++ElementIndex)
-		{
-			if(!(A[ElementIndex] == B[ElementIndex]))
-			{
-				return false;
-			}
-		}
-		return true;
-	}
-
-	[[nodiscard]] constexpr bool operator!=(const TStaticArray& B) const
-	{
-		for(uint32 ElementIndex = 0;ElementIndex < NumElements;++ElementIndex)
-		{
-			if(!((*this)[ElementIndex] == B[ElementIndex]))
-			{
-				return true;
-			}
-		}
-		return false;
-	}
+	[[nodiscard]] constexpr bool operator==(const TStaticArray&) const = default;
 
 	/**
 	 * Returns true if the array is empty and contains no elements. 
@@ -99,70 +80,100 @@ public:
 	}
 
 	/** The number of elements in the array. */
-	[[nodiscard]] UE_FORCEINLINE_HINT constexpr int32 Num() const
+	[[nodiscard]] UE_REWRITE constexpr int32 Num() const
 	{
 		return NumElements;
 	}
 
 	/** A pointer to the first element of the array */
-	[[nodiscard]] inline constexpr InElementType* GetData()
+	[[nodiscard]] UE_REWRITE constexpr InElementType* GetData()
 	{
-		static_assert((alignof(ElementType) % Alignment) == 0, "GetData() cannot be called on a TStaticArray with non-standard alignment");
-		return &Storage.Elements[0].Element;
+		return Elements;
 	}
-	[[nodiscard]] UE_FORCEINLINE_HINT constexpr const InElementType* GetData() const
+	
+	[[nodiscard]] UE_REWRITE constexpr const InElementType* GetData() const
 	{
-		return const_cast<TStaticArray*>(this)->GetData();
+		return Elements;
 	}
 
+	[[nodiscard]] UE_REWRITE constexpr InElementType*                               begin ()       { return Elements; }
+	[[nodiscard]] UE_REWRITE constexpr const InElementType*                         begin () const { return Elements; }
+	[[nodiscard]] UE_REWRITE constexpr InElementType*                               end   ()       { return Elements + NumElements; }
+	[[nodiscard]] UE_REWRITE constexpr const InElementType*                         end   () const { return Elements + NumElements; }
+	[[nodiscard]] UE_REWRITE constexpr TReversePointerIterator<InElementType>       rbegin()       { return TReversePointerIterator<InElementType>(Elements + NumElements); }
+	[[nodiscard]] UE_REWRITE constexpr TReversePointerIterator<const InElementType> rbegin() const { return TReversePointerIterator<const InElementType>(Elements + NumElements); }
+	[[nodiscard]] UE_REWRITE constexpr TReversePointerIterator<InElementType>       rend  ()       { return TReversePointerIterator<InElementType>(Elements); }
+	[[nodiscard]] UE_REWRITE constexpr TReversePointerIterator<const InElementType> rend  () const { return TReversePointerIterator<const InElementType>(Elements); }
+	
 private:
+	template <uint32... Indices, typename ArgGeneratorType>
+	constexpr explicit TStaticArray(EInPlace, TIntegerSequence<uint32, Indices...>, ArgGeneratorType Generator)
+		: Elements{Generator(Indices)...}
+	{
+	}
 
-	struct alignas(Alignment) TArrayStorageElementAligned
+	template <typename... ArgGeneratorTypes>
+	constexpr explicit TStaticArray(EPerElement, ArgGeneratorTypes... Generator)
+		: Elements{Generator()...}
 	{
-		[[nodiscard]] constexpr TArrayStorageElementAligned() = default;
+	}
+};
 
-		// Index is used to achieve pack expansion in TArrayStorage's first constructor, but is unused here
-		template <typename... ArgTypes>
-		[[nodiscard]] constexpr explicit TArrayStorageElementAligned(EInPlace, uint32 /*Index*/, ArgTypes&&... Args)
-			: Element(Forward<ArgTypes>(Args)...)
-		{
-		}
+template <typename InElementType, uint32 NumElements, uint32 Alignment> requires(Alignment != uint32(-1))
+class TStaticArray<InElementType, NumElements, Alignment>: public TStaticArray<UE::Core::TAlignedElement<InElementType, Alignment>, NumElements>
+{
+	UE_STATIC_ASSERT_WARN(false, "TStaticArray's alignment parameter has been deprecated in 5.7 - you can use TAlignedElement to wrap InElementType.");
+	
+	using TStaticArray<UE::Core::TAlignedElement<InElementType, Alignment>, NumElements>::Elements;
+ 
+public:
+	using ElementType = InElementType;
 
-		InElementType Element;
-	};
+	[[nodiscard]] constexpr TStaticArray() = default;
 
-	struct TArrayStorage
+	// Constructs each element with Args
+	template <typename... ArgTypes>
+	[[nodiscard]] constexpr explicit TStaticArray(EInPlace, ArgTypes&&... Args)
+		: TStaticArray<UE::Core::TAlignedElement<InElementType, Alignment>, NumElements>{InPlace, TMakeIntegerSequence<uint32, NumElements>(), [&Args...](uint32) { return InElementType(Args...); }}
 	{
-		[[nodiscard]] constexpr TArrayStorage() = default;
+	}
 
-		template<uint32... Indices, typename... ArgTypes>
-		[[nodiscard]] constexpr explicit TArrayStorage(EInPlace, TIntegerSequence<uint32, Indices...>, ArgTypes&&... Args)
-			: Elements{ TArrayStorageElementAligned(InPlace, Indices, Args...)... }
-		{
-			// The arguments are deliberately not forwarded arguments here, because we're initializing multiple elements
-			// and don't want an argument to be mutated by the first element's constructor and then that moved-from state
-			// be used to construct the remaining elements.
-			//
-			// This'll mean that it'll be a compile error to use move-only types like TUniquePtr when in-place constructing
-			// TStaticArray elements, which is a natural expectation because that TUniquePtr can only transfer ownership to
-			// a single element.
-		}
+	// Directly initializes the array with the provided values.
+	template <typename... ArgTypes> requires((sizeof...(ArgTypes) > 0 && sizeof...(ArgTypes) <= NumElements) && UE::Core::Private::TCanBeConvertedToFromAll_V<InElementType, ArgTypes...>)
+	[[nodiscard]] constexpr TStaticArray(ArgTypes&&... Args)
+		: TStaticArray<UE::Core::TAlignedElement<InElementType, Alignment>, NumElements>{PerElement, [&Args] { return InElementType(Forward<ArgTypes>(Args)); }...}
+	{
+	}
 
-		template<typename... ArgTypes>
-		[[nodiscard]] constexpr explicit TArrayStorage(EPerElement, ArgTypes&&... Args)
-			: Elements{ TArrayStorageElementAligned(InPlace, 0 /* dummy index */, Forward<ArgTypes>(Args))... }
-		{
-		}
+	[[nodiscard]] constexpr TStaticArray(TStaticArray&& Other) = default;
+	[[nodiscard]] constexpr TStaticArray(const TStaticArray& Other) = default;
+	constexpr TStaticArray& operator=(TStaticArray&& Other) = default;
+	constexpr TStaticArray& operator=(const TStaticArray& Other) = default;
+	
+	constexpr TStaticArray& operator=(TStaticArray<UE::Core::TAlignedElement<InElementType, Alignment>, NumElements>&& Other)
+	{
+		TStaticArray<UE::Core::TAlignedElement<InElementType, Alignment>, NumElements>::operator=(MoveTemp(Other));
+		return *this;
+	}
 
-		TArrayStorageElementAligned Elements[NumElements];
-	};
+	constexpr TStaticArray& operator=(const TStaticArray<UE::Core::TAlignedElement<InElementType, Alignment>, NumElements>& Other)
+	{
+		TStaticArray<UE::Core::TAlignedElement<InElementType, Alignment>, NumElements>::operator=(Other);
+		return *this;
+	}
 
-public:
-	// Storage is made public so that TStaticArray is a "Structural Type" and therefore can be passed as 
-	// a non-type template parameter like std::array making it a more complete replacement in C++20
-	// It is not intended to be accessed directly. Please use the member functions and operators instead
-	TArrayStorage Storage;
+	// Accessors.
+	[[nodiscard]] inline constexpr InElementType& operator[](uint32 Index)
+	{
+		checkSlow(Index < NumElements);
+		return Elements[Index];
+	}
 
+	[[nodiscard]] inline constexpr const InElementType& operator[](uint32 Index) const
+	{
+		checkSlow(Index < NumElements);
+		return Elements[Index];
+	}
 
 	template <typename StorageElementType, bool bReverse = false>
 	struct FRangedForIterator
@@ -172,15 +183,15 @@ public:
 		{
 		}
 
-		[[nodiscard]] constexpr auto& operator*() const
+		[[nodiscard]] constexpr std::conditional_t<std::is_const_v<StorageElementType>, const InElementType, InElementType>& operator*() const
 		{
 			if constexpr (bReverse)
 			{
-				return (Ptr - 1)->Element;
+				return *(Ptr - 1);
 			}
 			else
 			{
-				return Ptr->Element;
+				return *Ptr;
 			}
 		}
 
@@ -206,20 +217,20 @@ public:
 		StorageElementType* Ptr;
 	};
 
-	using RangedForIteratorType             = FRangedForIterator<      TArrayStorageElementAligned>;
-	using RangedForConstIteratorType        = FRangedForIterator<const TArrayStorageElementAligned>;
-	using RangedForReverseIteratorType      = FRangedForIterator<      TArrayStorageElementAligned, true>;
-	using RangedForConstReverseIteratorType = FRangedForIterator<const TArrayStorageElementAligned, true>;
+	using RangedForIteratorType             = FRangedForIterator<      UE::Core::TAlignedElement<InElementType, Alignment>>;
+	using RangedForConstIteratorType        = FRangedForIterator<const UE::Core::TAlignedElement<InElementType, Alignment>>;
+	using RangedForReverseIteratorType      = FRangedForIterator<      UE::Core::TAlignedElement<InElementType, Alignment>, true>;
+	using RangedForConstReverseIteratorType = FRangedForIterator<const UE::Core::TAlignedElement<InElementType, Alignment>, true>;
 
 	/** STL-like iterators to enable range-based for loop support. */
-	[[nodiscard]] UE_FORCEINLINE_HINT RangedForIteratorType             constexpr begin()        { return RangedForIteratorType(Storage.Elements); }
-	[[nodiscard]] UE_FORCEINLINE_HINT RangedForConstIteratorType        constexpr begin()  const { return RangedForConstIteratorType(Storage.Elements); }
-	[[nodiscard]] UE_FORCEINLINE_HINT RangedForIteratorType             constexpr end()          { return RangedForIteratorType(Storage.Elements + NumElements); }
-	[[nodiscard]] UE_FORCEINLINE_HINT RangedForConstIteratorType        constexpr end()    const { return RangedForConstIteratorType(Storage.Elements + NumElements); }
-	[[nodiscard]] UE_FORCEINLINE_HINT RangedForReverseIteratorType      constexpr rbegin()       { return RangedForReverseIteratorType(Storage.Elements + NumElements); }
-	[[nodiscard]] UE_FORCEINLINE_HINT RangedForConstReverseIteratorType constexpr rbegin() const { return RangedForConstReverseIteratorType(Storage.Elements + NumElements); }
-	[[nodiscard]] UE_FORCEINLINE_HINT RangedForReverseIteratorType      constexpr rend()         { return RangedForReverseIteratorType(Storage.Elements); }
-	[[nodiscard]] UE_FORCEINLINE_HINT RangedForConstReverseIteratorType constexpr rend()   const { return RangedForConstReverseIteratorType(Storage.Elements); }
+	[[nodiscard]] UE_FORCEINLINE_HINT RangedForIteratorType             constexpr begin()        { return RangedForIteratorType(Elements); }
+	[[nodiscard]] UE_FORCEINLINE_HINT RangedForConstIteratorType        constexpr begin()  const { return RangedForConstIteratorType(Elements); }
+	[[nodiscard]] UE_FORCEINLINE_HINT RangedForIteratorType             constexpr end()          { return RangedForIteratorType(Elements + NumElements); }
+	[[nodiscard]] UE_FORCEINLINE_HINT RangedForConstIteratorType        constexpr end()    const { return RangedForConstIteratorType(Elements + NumElements); }
+	[[nodiscard]] UE_FORCEINLINE_HINT RangedForReverseIteratorType      constexpr rbegin()       { return RangedForReverseIteratorType(Elements + NumElements); }
+	[[nodiscard]] UE_FORCEINLINE_HINT RangedForConstReverseIteratorType constexpr rbegin() const { return RangedForConstReverseIteratorType(Elements + NumElements); }
+	[[nodiscard]] UE_FORCEINLINE_HINT RangedForReverseIteratorType      constexpr rend()         { return RangedForReverseIteratorType(Elements); }
+	[[nodiscard]] UE_FORCEINLINE_HINT RangedForConstReverseIteratorType constexpr rend()   const { return RangedForConstReverseIteratorType(Elements); }
 };
 
 /** Creates a static array filled with the specified value. */
@@ -237,7 +248,7 @@ template <typename InElementType, uint32 NumElements>
 template <typename ElementType, uint32 NumElements, uint32 Alignment>
 struct TIsContiguousContainer<TStaticArray<ElementType, NumElements, Alignment>>
 {
-	enum { Value = (alignof(ElementType) % Alignment) == 0 };
+	enum { Value = Alignment == uint32(-1) };
 };
 
 /** Serializer. */
diff --git a/Engine/Source/Runtime/Core/Public/Misc/AlignedElement.h b/Engine/Source/Runtime/Core/Public/Misc/AlignedElement.h
new file mode 100644
--- /dev/null
+++ b/Engine/Source/Runtime/Core/Public/Misc/AlignedElement.h
@@ -0,0 +1,61 @@
+// Copyright Epic Games, Inc. All Rights Reserved.
+
+#pragma once
+
+#include "CoreTypes.h"
+#include <type_traits>
+
+namespace UE::Core
+{
+	// Helper class to wrap a type for alignment in containers
+	template <typename T, uint32 Alignment>
+	struct alignas(Alignment) TAlignedElement
+	{
+		static_assert(!std::is_reference_v<T>);
+		
+		T Value;
+
+		UE_REWRITE constexpr TAlignedElement& operator=(T&& Other)
+		{
+			Value = MoveTemp(Other);
+			return *this;
+		}
+
+		UE_REWRITE constexpr TAlignedElement& operator=(const T& Other)
+		{
+			Value = Other;
+			return *this;
+		}
+		
+		UE_REWRITE constexpr operator T&&()&&
+		{
+			return MoveTemp(Value);
+		}
+		
+		UE_REWRITE constexpr operator const T&() const&
+		{
+			return Value;
+		}
+		
+		UE_REWRITE constexpr operator T&()&
+		{
+			return Value;
+		}
+	};
+
+	template <typename T, uint32 Alignment> requires(std::is_class_v<T> && !std::is_final_v<T>)
+	struct alignas(Alignment) TAlignedElement<T, Alignment>: T
+	{
+		UE_REWRITE constexpr TAlignedElement& operator=(T&& Other)
+		{
+			T::operator=(MoveTemp(Other));
+			return *this;
+		}
+		
+		UE_REWRITE constexpr TAlignedElement& operator=(const T& Other)
+		{
+			T::operator=(Other);
+			return *this;
+		}
+	};
+}
diff --git a/Engine/Source/Runtime/Engine/Classes/PhysicsField/PhysicsFieldComponent.h b/Engine/Source/Runtime/Engine/Classes/PhysicsField/PhysicsFieldComponent.h
--- a/Engine/Source/Runtime/Engine/Classes/PhysicsField/PhysicsFieldComponent.h
+++ b/Engine/Source/Runtime/Engine/Classes/PhysicsField/PhysicsFieldComponent.h
@@ -25,7 +25,7 @@ enum class EFieldCommandBuffer : uint8
 struct FPhysicsFieldInfos
 {
 	/** Type of targets offsets */
-	using BufferOffsets = TStaticArray<int32, MAX_PHYSICS_FIELD_TARGETS, 16>;
+	using BufferOffsets = TStaticArray<UE::Core::TAlignedElement<int32, 16>, MAX_PHYSICS_FIELD_TARGETS>;
 
 	/** Size of the datas stored on each voxels*/
 	int32 TargetCount = 1;
@@ -49,7 +49,7 @@ struct FPhysicsFieldInfos
 	BufferOffsets ValidTargets;
 
 	/** Physics Targets bounds */
-	TStaticArray<FIntVector4, MAX_PHYSICS_FIELD_TARGETS, 16> PhysicsBounds;
+	TStaticArray<UE::Core::TAlignedElement<FIntVector4, 16>, MAX_PHYSICS_FIELD_TARGETS> PhysicsBounds;
 
 	/** Clipmap  Center */
 	FVector ClipmapCenter = FVector::ZeroVector;
@@ -82,7 +82,7 @@ struct FPhysicsFieldInfos
 	TArray<FIntVector4> CellsMax;
 
 	/** Min Bounds for each target/clipmap */
-	TStaticArray<int32, MAX_PHYSICS_FIELD_TARGETS, 16> BoundsOffsets;
+	TStaticArray<UE::Core::TAlignedElement<int32, 16>, MAX_PHYSICS_FIELD_TARGETS> BoundsOffsets;
 
 	/** Time in seconds for field evaluation */
 	float TimeSeconds;
@@ -135,7 +135,7 @@ public:
 	FPhysicsFieldResource(const int32 TargetCount, const TArray<EFieldPhysicsType>& TargetTypes,
 		const FPhysicsFieldInfos::BufferOffsets& VectorTargets, const FPhysicsFieldInfos::BufferOffsets& ScalarTargets,
 		const FPhysicsFieldInfos::BufferOffsets& IntegerTargets, const FPhysicsFieldInfos::BufferOffsets& PhysicsTargets,
-		const TStaticArray< FIntVector4, MAX_PHYSICS_FIELD_TARGETS, 16>& PhysicsBounds, const bool bBuildClipmap);
+		const TStaticArray<UE::Core::TAlignedElement<FIntVector4, 16>, MAX_PHYSICS_FIELD_TARGETS>& PhysicsBounds, const bool bBuildClipmap);
 
 	/** Release Field resources. */
 	virtual void ReleaseRHI() override;
diff --git a/Engine/Source/Runtime/Engine/Private/PhysicsField/PhysicsFieldComponent.cpp b/Engine/Source/Runtime/Engine/Private/PhysicsField/PhysicsFieldComponent.cpp
--- a/Engine/Source/Runtime/Engine/Private/PhysicsField/PhysicsFieldComponent.cpp
+++ b/Engine/Source/Runtime/Engine/Private/PhysicsField/PhysicsFieldComponent.cpp
@@ -389,7 +389,7 @@ IMPLEMENT_SHADER_TYPE(, FResetPhysicsFieldClipmapCS, TEXT("/Engine/Private/Physi
 FPhysicsFieldResource::FPhysicsFieldResource(const int32 TargetCount, const TArray<EFieldPhysicsType>& TargetTypes, 
 		const FPhysicsFieldInfos::BufferOffsets& VectorTargets, const FPhysicsFieldInfos::BufferOffsets& ScalarTargets,
 	    const FPhysicsFieldInfos::BufferOffsets& IntegerTargets, const FPhysicsFieldInfos::BufferOffsets& PhysicsTargets,
-		const TStaticArray<FIntVector4, MAX_PHYSICS_FIELD_TARGETS, 16>& PhysicsBounds, const bool bBuildClipmap) : FRenderResource()
+		const TStaticArray<UE::Core::TAlignedElement<FIntVector4, 16>, MAX_PHYSICS_FIELD_TARGETS>& PhysicsBounds, const bool bBuildClipmap) : FRenderResource()
 {
 	FieldInfos.TargetCount = TargetCount;
 	FieldInfos.TargetTypes = TargetTypes;
@@ -616,7 +616,7 @@ void FPhysicsFieldResource::UpdateResource(FRHICommandList& RHICmdList,
 void FPhysicsFieldInstance::InitInstance( const TArray<EFieldPhysicsType>& TargetTypes, const bool bBuildClipmap)
 {
 	FPhysicsFieldInfos::BufferOffsets VectorTargets(InPlace, -1), ScalarTargets(InPlace, -1), IntegerTargets(InPlace, -1), PhysicsTargets(InPlace, -1);
-	TStaticArray<FIntVector4, MAX_PHYSICS_FIELD_TARGETS, 16> PhysicsBounds(InPlace, FIntVector4(-1));
+	TStaticArray<UE::Core::TAlignedElement<FIntVector4, 16>, MAX_PHYSICS_FIELD_TARGETS> PhysicsBounds(InPlace, FIntVector4(-1));
 
 	static const TArray<EFieldPhysicsType> VectorTypes = GetFieldTargetTypes(EFieldOutputType::Field_Output_Vector);
 	static const TArray<EFieldPhysicsType> ScalarTypes = GetFieldTargetTypes(EFieldOutputType::Field_Output_Scalar);
diff --git a/Engine/Source/Runtime/RenderCore/Public/ShaderParameterMacros.h b/Engine/Source/Runtime/RenderCore/Public/ShaderParameterMacros.h
--- a/Engine/Source/Runtime/RenderCore/Public/ShaderParameterMacros.h
+++ b/Engine/Source/Runtime/RenderCore/Public/ShaderParameterMacros.h
@@ -13,6 +13,7 @@
 #include "Algo/Reverse.h"
 #include "Templates/IsArrayOrRefOfTypeByPredicate.h"
 #include "Traits/IsCharEncodingCompatibleWith.h"
+#include "Misc/AlignedElement.h"
 #include "Misc/AssertionMacros.h"
 #include "RHICommandList.h"
 #include "RenderGraphTextureSubresource.h"
@@ -875,7 +876,7 @@ static_assert(sizeof(FRenderTargetBindingSlots) == 256, "FRenderTargetBindingSlo
 
 /** Static array of shader resource shader that is initialized to nullptr. */
 template<typename TElement, uint32 NumElements>
-class alignas(SHADER_PARAMETER_POINTER_ALIGNMENT) TShaderResourceParameterArray : public TStaticArray<TElement, NumElements, SHADER_PARAMETER_POINTER_ALIGNMENT>
+class alignas(SHADER_PARAMETER_POINTER_ALIGNMENT) TShaderResourceParameterArray : public TStaticArray<std::conditional_t<SHADER_PARAMETER_POINTER_ALIGNMENT == alignof(TElement), TElement, UE::Core::TAlignedElement<TElement, SHADER_PARAMETER_POINTER_ALIGNMENT>>, NumElements>
 {
 public:
 	inline TShaderResourceParameterArray()
@@ -1240,7 +1241,7 @@ struct TShaderParameterTypeInfo<T[InNumElements]>
 	static constexpr int32 Alignment = SHADER_PARAMETER_ARRAY_ELEMENT_ALIGNMENT;
 	static constexpr bool bIsStoredInConstantBuffer = TShaderParameterTypeInfo<T>::bIsStoredInConstantBuffer;
 
-	using TAlignedType = TStaticArray<T, InNumElements, Alignment>;
+	using TAlignedType = TStaticArray<std::conditional_t<Alignment == alignof(T), T, UE::Core::TAlignedElement<T, Alignment>>, InNumElements>;
 	using TInstancedType = T;
 
 	static const FShaderParametersMetadata* GetStructMetadata() { return TShaderParameterTypeInfo<T>::GetStructMetadata(); }
@@ -1258,7 +1259,7 @@ struct TShaderParameterTypeInfo<TStaticArray<T,InNumElements,IgnoredAlignment>>
 	static constexpr int32 Alignment = SHADER_PARAMETER_ARRAY_ELEMENT_ALIGNMENT;
 	static constexpr bool bIsStoredInConstantBuffer = TShaderParameterTypeInfo<T>::bIsStoredInConstantBuffer;
 
-	using TAlignedType = TStaticArray<T, InNumElements, Alignment>;
+	using TAlignedType = TStaticArray<std::conditional_t<Alignment == alignof(T), T, UE::Core::TAlignedElement<T, Alignment>>, InNumElements>;
 	using TInstancedType = T;
 
 	static const FShaderParametersMetadata* GetStructMetadata() { return TShaderParameterTypeInfo<T>::GetStructMetadata(); }
diff --git a/Engine/Source/Runtime/Renderer/Private/PostProcess/DiaphragmDOF.cpp b/Engine/Source/Runtime/Renderer/Private/PostProcess/DiaphragmDOF.cpp
--- a/Engine/Source/Runtime/Renderer/Private/PostProcess/DiaphragmDOF.cpp
+++ b/Engine/Source/Runtime/Renderer/Private/PostProcess/DiaphragmDOF.cpp
@@ -2032,7 +2032,7 @@ bool DiaphragmDOF::AddPasses(
 	}
 
 	bool bMatteBoxActive = false;
-	TStaticArray<FUintVector4, FPhysicalCocModel::MaxMatteBoxFlags, SHADER_PARAMETER_ARRAY_ELEMENT_ALIGNMENT> MatteBoxPlanes;
+	TStaticArray<UE::Core::TAlignedElement<FUintVector4, SHADER_PARAMETER_ARRAY_ELEMENT_ALIGNMENT>, FPhysicalCocModel::MaxMatteBoxFlags> MatteBoxPlanes;
 
 	for (int Index = 0; Index < FPhysicalCocModel::MaxMatteBoxFlags; ++Index)
 	{
diff --git a/Engine/Source/Runtime/Renderer/Private/SceneRendering.cpp b/Engine/Source/Runtime/Renderer/Private/SceneRendering.cpp
--- a/Engine/Source/Runtime/Renderer/Private/SceneRendering.cpp
+++ b/Engine/Source/Runtime/Renderer/Private/SceneRendering.cpp
@@ -1244,7 +1244,7 @@ void SetupPhysicsFieldUniformBufferParameters(const FScene* Scene, FEngineShowFl
 	}
 	else
 	{
-		TStaticArray<FIntVector4, MAX_PHYSICS_FIELD_TARGETS, 16> EmptyTargets = {};
+		TStaticArray<UE::Core::TAlignedElement<FIntVector4, 16>, MAX_PHYSICS_FIELD_TARGETS> EmptyTargets = {};
 		ViewUniformShaderParameters.PhysicsFieldClipmapBuffer = GWhiteVertexBufferWithSRV->ShaderResourceViewRHI;
 		ViewUniformShaderParameters.PhysicsFieldClipmapCenter = FVector3f::ZeroVector;
 		ViewUniformShaderParameters.PhysicsFieldClipmapDistance = 1.0;
