Paths: Fix CollapseRelativeDirectories to handle normalizing duplicate slashes that are removed during a .. subdirectory collapse.
#rnx

diff --git a/Engine/Source/Runtime/Core/Private/Misc/PathViews.cpp b/Engine/Source/Runtime/Core/Private/Misc/PathViews.cpp
--- a/Engine/Source/Runtime/Core/Private/Misc/PathViews.cpp
+++ b/Engine/Source/Runtime/Core/Private/Misc/PathViews.cpp
@@ -740,7 +740,47 @@ void CollapseRelativeDirectoriesImplTemplate(WritableStringType& InOutPath,
 		int32 PreviousSeparatorIndex = -1;
 		int32 RemovalStart = 0;
 		FStringView PreviousText = Path.Left(Index);
-		if (PreviousText.FindLastChar('/', PreviousSeparatorIndex))
+
+		bool bFoundPreviousSeparator = PreviousText.FindLastChar('/', PreviousSeparatorIndex);
+
+		if (bFoundPreviousSeparator)
+		{
+			// Handle a special case in unnormalized paths.
+			// CollapseRelativeDirectories can be called on unnormalized paths with spurious duplicate separators,
+			// such as "d:/root/Engine///////../..". We are not required to fixup all duplicate separators, but we
+			// are expected to handle the collapse of .. as if the duplicate separators were removed. If
+			// PreviousSeparatorIndex is a duplicate separator (and not the special case of allowed duplicate
+			// separator at beginning of path), then ignore it and remove it and keep searching for the earlier
+			// duplicate separator before it.
+			if (PreviousSeparatorIndex == Index - 1)
+			{
+				// Move PreviousSeparatorIndex to the beginning of the slashes in the sequence foo////..
+				while (PreviousSeparatorIndex > 0 && PreviousText[PreviousSeparatorIndex - 1] == '/')
+				{
+					--PreviousSeparatorIndex;
+				} 
+
+				// Note that we should be prevented from collpasing the valid occurrence of // at beginning of string
+				// by StartOfCollapsibleDirs being set to point beyond that, but due to legacy behavior we 
+				// allow the collapse of the initial / or //. When we stop allowing that, we need to set
+				// bFoundPreviousSeparator to false here instead of true.
+				if (PreviousSeparatorIndex == 0)
+				{
+					bFoundPreviousSeparator = true;
+				}
+				else
+				{
+					// Now mark for removal all of those slashes, AND the directory foo which comes before them.
+					// To remove the directory foo, find the last / before the new PreviousSeparatorIndex
+					PreviousText = Path.Left(PreviousSeparatorIndex);
+					bFoundPreviousSeparator = PreviousText.FindLastChar('/', PreviousSeparatorIndex);
+					// There is no need to recursively test for new PreviousSeparator == PreviousSeparatorIndex - 1,
+					// because we know it does not. We moved to the start of the sequence of /// and we now have a
+					// non-slash in PreviousSeparatorIndex[-1].
+				}
+			}
+		}
+		if (bFoundPreviousSeparator)
 		{
 			RemovalStart = PreviousSeparatorIndex;
 		}
diff --git a/Engine/Source/Runtime/Core/Tests/Misc/PathTests.h b/Engine/Source/Runtime/Core/Tests/Misc/PathTests.h
--- a/Engine/Source/Runtime/Core/Tests/Misc/PathTests.h
+++ b/Engine/Source/Runtime/Core/Tests/Misc/PathTests.h
@@ -2,6 +2,7 @@
 
 #pragma once
 
+#include "Containers/ArrayView.h"
 #include "Containers/StringView.h"
 #include "Tests/TestHarnessAdapter.h"
 #if WITH_LOW_LEVEL_TESTS
@@ -72,13 +73,13 @@ void TestCollapseRelativeDirectories()
 	Run(TEXT("Folder"),                                               true,  TEXT("Folder"));
 	Run(TEXT("/Folder"),                                              true,  TEXT("/Folder"));
 	Run(TEXT("C:/Folder"),                                            true,  TEXT("C:/Folder"));
-	Run(TEXT("C:/Folder/.."),                                         true,  TEXT("C:"));
+	Run(TEXT("C:/Folder/.."),                                         true,  TEXT("C:")); // removing the leading slash is incorrect but legacy behavior, we should fix eventually
 	Run(TEXT("C:/Folder/../"),                                        true,  TEXT("C:/"));
 	Run(TEXT("C:/Folder/../file.txt"),                                true,  TEXT("C:/file.txt"));
 	Run(TEXT("Folder/.."),                                            true,  TEXT(""));
 	Run(TEXT("Folder/../"),                                           true,  TEXT(""));
 	Run(TEXT("Folder/../file.txt"),                                   true,  TEXT("file.txt"));
-	Run(TEXT("/Folder/.."),                                           true,  TEXT(""));
+	Run(TEXT("/Folder/.."),                                           true,  TEXT("")); // removing the leading slash is incorrect but legacy behavior, we should fix eventually
 	Run(TEXT("/Folder/../"),                                          true,  TEXT("/"));
 	Run(TEXT("/Folder/../file.txt"),                                  true,  TEXT("/file.txt"));
 	Run(TEXT("Folder/../.."),                                         false, TEXT(".."));
@@ -120,6 +121,16 @@ void TestCollapseRelativeDirectories()
 	Run(TEXT("C:/Folder1/./Folder2/..Folder3/..Folder4"),             true,  TEXT("C:/Folder1/Folder2/..Folder3/..Folder4"));
 	Run(TEXT("C:/Folder1/./Folder2/..Folder3/Folder4/../Folder5"),    true,  TEXT("C:/Folder1/Folder2/..Folder3/Folder5"));
 	Run(TEXT("C:/Folder1/..Folder2/Folder3/..Folder4/../Folder5"),    true,  TEXT("C:/Folder1/..Folder2/Folder3/Folder5"));
+	// Handle .. that need to collapse through duplicate separators
+	Run(TEXT("D:/Root/Engine//../.."),                                true,  TEXT("D:")); // removing the leading slash is incorrect but legacy behavior, we should fix eventually
+	Run(TEXT("D:/Root/Engine////////../.."),                          true,  TEXT("D:")); // removing the leading slash is incorrect but legacy behavior, we should fix eventually
+	Run(TEXT("D:/Root/Engine//../../"),                               true,  TEXT("D:/"));
+	Run(TEXT("D:/Root/Engine////////../../"),                         true,  TEXT("D:/"));
+	Run(TEXT("D:/Root//../.."),										  false, TEXT("D:/.."));
+	Run(TEXT("D:/Root////////../.."),                                 false, TEXT("D:/.."));
+	// But don't remove a leading // for network share when collapsing
+	Run(TEXT("//.."),                                                 true, TEXT("")); // removing the leading network path is incorrect but legacy behavior, we should fix eventually
+	Run(TEXT("//Root/../.."),                                         true, TEXT("")); // removing the leading network path is incorrect but legacy behavior, we should fix eventually
 }
 
 template<class PathType, class StringType>
@@ -157,7 +168,7 @@ struct FTestPair
 
 extern const FStringView BaseDir;
 
-extern const FTestPair ExpectedRelativeToAbsolutePaths[10];
+extern TConstArrayView<FTestPair> ExpectedRelativeToAbsolutePaths;
 
 }
 
diff --git a/Engine/Source/Runtime/Core/Tests/Misc/PathsTest.cpp b/Engine/Source/Runtime/Core/Tests/Misc/PathsTest.cpp
--- a/Engine/Source/Runtime/Core/Tests/Misc/PathsTest.cpp
+++ b/Engine/Source/Runtime/Core/Tests/Misc/PathsTest.cpp
@@ -19,7 +19,7 @@ void TestCreateStandardFilename();
 
 const FStringView PathTest::BaseDir = TEXTVIEW("/root");
 
-const PathTest::FTestPair PathTest::ExpectedRelativeToAbsolutePaths[10] =
+static PathTest::FTestPair StaticExpectedRelativeToAbsolutePaths[] =
 {
 	{ TEXTVIEW(""),					TEXTVIEW("/root/") },
 	{ TEXTVIEW("dir"),				TEXTVIEW("/root/dir") },
@@ -31,7 +31,10 @@ const PathTest::FTestPair PathTest::ExpectedRelativeToAbsolutePaths[10] =
 	{ TEXTVIEW("C:\\A\\B"),			TEXTVIEW("C:/A/B") },
 	{ TEXTVIEW("a/b/../c"),			TEXTVIEW("/root/a/c") },
 	{ TEXTVIEW("/a/b/../c"),		TEXTVIEW("/a/c") },
+	{ TEXTVIEW("D:/Root/Engine//../.."), TEXTVIEW("D:") }, // removing the leading slash is incorrect but legacy behavior, we should fix eventually
+	{ TEXTVIEW("D:/Root/Engine//../../"), TEXTVIEW("D:/") },
 };
+TConstArrayView<PathTest::FTestPair> PathTest::ExpectedRelativeToAbsolutePaths = StaticExpectedRelativeToAbsolutePaths;
 
 TEST_CASE_NAMED(FPathTests, "System::Core::Misc::Paths", "[ApplicationContextMask][SmokeFilter]")
 {
