Fix to TStaticArray for certain clang versions of constexpr defaulted functions generating errors in templates

diff --git a/Engine/Source/Runtime/Core/Public/Containers/StaticArray.h b/Engine/Source/Runtime/Core/Public/Containers/StaticArray.h
--- a/Engine/Source/Runtime/Core/Public/Containers/StaticArray.h
+++ b/Engine/Source/Runtime/Core/Public/Containers/StaticArray.h
@@ -66,7 +66,7 @@ public:
 		return Elements[Index];
 	}
 
-	[[nodiscard]] constexpr bool operator==(const TStaticArray&) const = default;
+	[[nodiscard]] bool operator==(const TStaticArray&) const = default;
 
 	/**
 	 * Returns true if the array is empty and contains no elements. 
