Fixed a deadlock in TIntrusiveMutex

FWakeState::bHasWaitingThreads indicates that there MAY be waiting threads, and is not a guarantee. It cannot be used as an assumption that there is a waiting thread that will take the lock and later unlock it to wake the other threads.

#rnx

diff --git a/Engine/Source/Runtime/Core/Public/Async/IntrusiveMutex.h b/Engine/Source/Runtime/Core/Public/Async/IntrusiveMutex.h
--- a/Engine/Source/Runtime/Core/Public/Async/IntrusiveMutex.h
+++ b/Engine/Source/Runtime/Core/Public/Async/IntrusiveMutex.h
@@ -196,8 +196,10 @@ public:
 	{
 		ParkingLot::WakeOne(GetWaitAddress(State), [&State](ParkingLot::FWakeState WakeState) -> uint64
 		{
-			if (!WakeState.bHasWaitingThreads)
+			if (!WakeState.bDidWake)
 			{
+				// Keep the flag until no thread wakes, otherwise shared locks may win before
+				// an exclusive lock has a chance.
 				State.fetch_and(~MayHaveWaitingLockFlag, std::memory_order_relaxed);
 			}
 			return 0;
@@ -209,8 +211,10 @@ public:
 		bool bDidWake = false;
 		ParkingLot::WakeOne(GetWaitAddress(State), [&State, &bDidWake](ParkingLot::FWakeState WakeState) -> uint64
 		{
-			if (!WakeState.bHasWaitingThreads)
+			if (!WakeState.bDidWake)
 			{
+				// Keep the flag until no thread wakes, otherwise shared locks may win before
+				// an exclusive lock has a chance.
 				State.fetch_and(~MayHaveWaitingLockFlag, std::memory_order_relaxed);
 			}
 			bDidWake = WakeState.bDidWake;
