Fix FNamespacedIoStatusBuilderBuilder not working properly in shipping builds

diff --git a/Engine/Source/Runtime/Core/Private/IO/IoStatus.cpp b/Engine/Source/Runtime/Core/Private/IO/IoStatus.cpp
--- a/Engine/Source/Runtime/Core/Private/IO/IoStatus.cpp
+++ b/Engine/Source/Runtime/Core/Private/IO/IoStatus.cpp
@@ -97,12 +97,14 @@ FIoStatusBuilder::FIoStatusBuilder(EIoErrorCode InStatusCode)
 {
 }
 
-FIoStatusBuilder::FIoStatusBuilder(const TCHAR* NamespaceStr, EIoErrorCode StatusCode, const UE::FSourceLocation& SourceLoc)
+FIoStatusBuilder::FIoStatusBuilder(
+	const TCHAR* NamespaceStr, EIoErrorCode StatusCode,
+	const uint_least32_t Line, const uint_least32_t Column)
 	: FIoStatusBuilder(StatusCode)
 {
 	Message
 		<< TEXT('<') << NamespaceStr << TEXT('-')
-		<< SourceLoc.GetLine() << TEXT(',') << SourceLoc.GetColumn() << TEXT("> ");
+		<< Line << TEXT(',') << Column << TEXT(">");
 }
 
 FIoStatusBuilder::~FIoStatusBuilder()
diff --git a/Engine/Source/Runtime/Core/Public/IO/IoStatus.h b/Engine/Source/Runtime/Core/Public/IO/IoStatus.h
--- a/Engine/Source/Runtime/Core/Public/IO/IoStatus.h
+++ b/Engine/Source/Runtime/Core/Public/IO/IoStatus.h
@@ -5,7 +5,6 @@
 #include "Containers/StringView.h"
 #include "Containers/UnrealString.h"
 #include "Misc/StringBuilder.h"
-#include "Misc/SourceLocation.h"
 
 template <typename CharType> class TStringBuilderBase;
 
@@ -144,7 +143,9 @@ protected:
 public:
 	CORE_API explicit	FIoStatusBuilder(EIoErrorCode StatusCode);
 
-	CORE_API			FIoStatusBuilder(const TCHAR* NamespaceStr, EIoErrorCode StatusCode, const UE::FSourceLocation& SourceLoc = UE::FSourceLocation::Current());
+	CORE_API			FIoStatusBuilder(const TCHAR* NamespaceStr, EIoErrorCode StatusCode, 
+							const uint_least32_t Line = __builtin_LINE(),
+							const uint_least32_t Column = __builtin_COLUMN());
 
 	CORE_API			~FIoStatusBuilder();
 
@@ -172,9 +173,11 @@ public:
 	{
 	}
 
-	FIoStatusBuilder operator()(EIoErrorCode StatusCode, const UE::FSourceLocation& SourceLoc = UE::FSourceLocation::Current()) const
+	FIoStatusBuilder operator()(EIoErrorCode StatusCode, 
+		const uint_least32_t Line = __builtin_LINE(),
+		const uint_least32_t Column = __builtin_COLUMN()) const
 	{
-		return FIoStatusBuilder(NamespaceStr, StatusCode, SourceLoc);
+		return FIoStatusBuilder(NamespaceStr, StatusCode, Line, Column);
 	}
 };
 
