- Added "config lint" console command, which can find issues with config files that devs may want to clean up (same value as a previous layer, array entries with both + and -, no value in BaseFoo.in)

diff --git a/Engine/Source/Runtime/Core/Private/Misc/ConfigCacheIni.cpp b/Engine/Source/Runtime/Core/Private/Misc/ConfigCacheIni.cpp
--- a/Engine/Source/Runtime/Core/Private/Misc/ConfigCacheIni.cpp
+++ b/Engine/Source/Runtime/Core/Private/Misc/ConfigCacheIni.cpp
@@ -3871,6 +3871,20 @@ bool FConfigBranch::MergeStaticLayersUpTo(const FString& LayerSubstring, FConfig
 	return false;
 }
 
+bool FConfigBranch::MergeStaticLayersUpToAndIncluding(const FString& LayerSubstring, FConfigFile& OutFile) const
+{
+	for (const TPair<FString, FConfigCommandStream>& Pair : StaticLayers)
+	{
+		OutFile.ApplyFile(&Pair.Value);
+		if (Pair.Key.Contains(LayerSubstring))
+		{
+			return true;
+		}
+	}
+
+	return false;
+}
+
 const FConfigCommandStream* FConfigBranch::GetStaticLayer(const FString& LayerSubstring) const
 {
 	for (const TPair<FString, FConfigCommandStream>& Pair : StaticLayers)
@@ -7615,6 +7629,162 @@ class FIniExec : public FSelfRegisteringExec
 			delete CSV;
 		}
 
+		if (FParse::Command(&Cmd, TEXT("Lint")))
+		{
+			FString BranchName;
+
+			if (!FParse::Token(Cmd, BranchName, true))
+			{
+				Ar.Logf(TEXT("Usage: config lint <branch> [platform] [-SameAsPrevious] [-MissingBase] [-ArrayDuplicate]\nIf no - options are specified, all issues will be shown"));
+			}
+			else
+			{
+				FString PlatformNameString;
+				FName PlatformName = FPlatformProperties::IniPlatformName();
+				if (FParse::Token(Cmd, PlatformNameString, true))
+				{
+					if (!PlatformNameString.StartsWith("-"))
+					{
+						PlatformName = *PlatformNameString;
+					}
+				}
+				bool bShowSameAsPrevious = FParse::Param(Cmd, TEXT("-SameAsPrevious"));
+				bool bShowMissingDefault = FParse::Param(Cmd, TEXT("-MissingBase"));
+				bool bShowArrayDuplicates = FParse::Param(Cmd, TEXT("-ArrayDuplicate"));
+				// if none specified, do all
+				if (!bShowSameAsPrevious && !bShowMissingDefault && !bShowArrayDuplicates)
+				{
+					bShowSameAsPrevious = bShowMissingDefault = bShowArrayDuplicates = true;
+				}
+
+
+				FConfigCacheIni ConfigSystem(EConfigCacheType::Temporary, PlatformName);
+				FConfigContext Context = FConfigContext::ReadIntoConfigSystem(&ConfigSystem, PlatformNameString);
+				Context.Load(*BranchName);
+
+				FConfigBranch* Branch = ConfigSystem.FindBranch(*BranchName, BranchName);
+				if (!Branch)
+				{
+					Ar.Logf(TEXT("Unknown branch %s"), *BranchName);
+					return true;
+				}
+
+				auto LogHeaders = [&Ar]( bool& bShownFileHeader, bool& bShownSectionHeader, const FString& Filename, const FString& SectionName)
+					{
+						bool bIsFirstSection = false;
+						if (!bShownFileHeader)
+						{
+							Ar.Logf(TEXT("-------------------------------------------------------------"));
+							Ar.Logf(TEXT("%s:"), *Filename);
+							bShownFileHeader = true;
+							bIsFirstSection = true;
+						}
+						if (!bShownSectionHeader)
+						{
+							if (!bIsFirstSection)
+							{
+								Ar.Logf(TEXT("  "));
+							}
+							Ar.Logf(TEXT("  [%s]"), *SectionName);
+							bShownSectionHeader = true;
+						}
+					};
+
+				FConfigFile BaseLayer;
+				FConfigFile PreviousLayer;
+				// now walk the hierarchy, diffing things
+				bool bTrueBaseFound = false;
+				bool bBranchBaseFound = false;
+				for (const TPair<int32, FString>& Pair : Branch->Hierarchy)
+				{
+					if (!bTrueBaseFound)
+					{
+						// prime the layers with the first, mainly empty layer
+						Branch->MergeStaticLayersUpToAndIncluding(Pair.Value, PreviousLayer);
+						bTrueBaseFound = true;
+						continue;
+					}
+
+					FConfigFile CurrentCombinedLayer;
+					Branch->MergeStaticLayersUpToAndIncluding(Pair.Value, CurrentCombinedLayer);
+					if (!bBranchBaseFound)
+					{
+						BaseLayer = CurrentCombinedLayer;
+						PreviousLayer = CurrentCombinedLayer;
+
+						// don't do much here, no need to lint against the true base
+						bBranchBaseFound = true;
+						continue;
+					}
+
+					
+					const FConfigCommandStream* CurrentLayer = Branch->GetStaticLayer(Pair.Value);
+					if (CurrentLayer == nullptr)
+					{
+						// most layers won't exist
+						continue;
+					}
+
+					bool bShownFileHeader = false;
+					// now walk over everthing, and make sure they don't match previous layer, and that they were in the base layer to be initialized correctly
+					for (const TPair<FString, FConfigCommandStreamSection>& SectionPair : *CurrentLayer)
+					{
+						bool bShownSectionHeader = false;
+						TMap<FString, uint8> AddRemoveValues;
+						const FConfigSection* PreviousSection = PreviousLayer.FindSection(SectionPair.Key);
+
+						for (const TPair<FName, FConfigValue>& ValuePair : SectionPair.Value)
+						{
+							if (ValuePair.Value.ValueType == FConfigValue::EValueType::Set)
+							{
+								const FConfigValue* PreviousValue = PreviousSection ? PreviousSection->Find(ValuePair.Key) : nullptr;
+								if (PreviousValue == nullptr)
+								{
+									if (bShowMissingDefault)
+									{
+										LogHeaders(bShownFileHeader, bShownSectionHeader, Pair.Value, SectionPair.Key);
+										Ar.Logf(TEXT("  %s -- No Default in Base"), *ValuePair.Key.ToString());
+									}
+								}
+								else if (PreviousValue->GetSavedValueForWriting() == ValuePair.Value.GetSavedValueForWriting())
+								{
+									if (bShowSameAsPrevious)
+									{
+										LogHeaders(bShownFileHeader, bShownSectionHeader, Pair.Value, SectionPair.Key);
+										Ar.Logf(TEXT("  %s -- Identical To Previous"), *ValuePair.Key.ToString());
+									}
+								}
+							}
+							else if (ValuePair.Value.ValueType == FConfigValue::EValueType::ArrayAddUnique)
+							{
+								uint8& Flags = AddRemoveValues.FindOrAdd(ValuePair.Value.GetSavedValueForWriting());
+								Flags |= 1;
+							}
+							else if (ValuePair.Value.ValueType == FConfigValue::EValueType::Remove)
+							{
+								uint8& Flags = AddRemoveValues.FindOrAdd(ValuePair.Value.GetSavedValueForWriting());
+								Flags |= 2;
+							}
+						}
+
+						if (bShowArrayDuplicates)
+						{
+							for (const TPair<FString, uint8>& AddRemovePair : AddRemoveValues)
+							{
+								if (AddRemovePair.Value == 3)
+								{
+									LogHeaders(bShownFileHeader, bShownSectionHeader, Pair.Value, SectionPair.Key);
+									Ar.Logf(TEXT("  %s -- Has Add and Remove"), *AddRemovePair.Key);
+								}
+							}
+						}
+					};
+
+					PreviousLayer = CurrentCombinedLayer;
+				}
+			}
+		}
+
 		return true;
 	}
 	
diff --git a/Engine/Source/Runtime/Core/Public/Misc/ConfigCacheIni.h b/Engine/Source/Runtime/Core/Public/Misc/ConfigCacheIni.h
--- a/Engine/Source/Runtime/Core/Public/Misc/ConfigCacheIni.h
+++ b/Engine/Source/Runtime/Core/Public/Misc/ConfigCacheIni.h
@@ -1185,6 +1185,11 @@ public:
 	 * @return true if the layer was found, or false if the layer name was not found (which means all static layers will be combined and returned)
 	 */
 	CORE_API bool MergeStaticLayersUpTo(const FString& LayerNameSubstring, FConfigFile& OutFile) const;
+	
+	/**
+	 * Same as MergeStaticLayersUpTo, but includes the passed in layer
+	 */
+	CORE_API bool MergeStaticLayersUpToAndIncluding(const FString& LayerNameSubstring, FConfigFile& OutFile) const;
 
 	/**
 	 * Find the static layer that first matches the substring (or full) name of the layer filename
