[TaskGraph]

Fixes problem where FTaskEvent::Trigger() was not thread safe and could attempt to schedule the same job more than once leading to asserts firing in TryUnlock()

- Adds new TryLaunchSafe() method which will track if a task has been launched
- Adds now atomic bool TaskLaunchedSafely to track when a task is launched using the Safe api.
- Moved FTaskEvent::Trigger to use TryLaunchSafe() and moves memory management logic into TryLaunchSafe

#rnx

diff --git a/Engine/Source/Runtime/Core/Public/Tasks/Task.h b/Engine/Source/Runtime/Core/Public/Tasks/Task.h
--- a/Engine/Source/Runtime/Core/Public/Tasks/Task.h
+++ b/Engine/Source/Runtime/Core/Public/Tasks/Task.h
@@ -251,10 +251,7 @@ namespace UE::Tasks
 		{
 			if (!IsCompleted()) // event can be triggered multiple times
 			{
-				// An event is not "in the system" until it's triggered, and should be kept alive only by external references. Once it's triggered it's in the system 
-				// and can outlive external references, so we need to keep it alive by holding an internal reference. It will be released when the event is signalled
-				Pimpl->AddRef();
-				Pimpl->TryLaunch(sizeof(*Pimpl));
+				Pimpl->Trigger(sizeof(*Pimpl));
 			}
 		}
 	};
diff --git a/Engine/Source/Runtime/Core/Public/Tasks/TaskPrivate.h b/Engine/Source/Runtime/Core/Public/Tasks/TaskPrivate.h
--- a/Engine/Source/Runtime/Core/Public/Tasks/TaskPrivate.h
+++ b/Engine/Source/Runtime/Core/Public/Tasks/TaskPrivate.h
@@ -385,6 +385,25 @@ public:
 				return Pipe;
 			}
 
+			// Tries to schedule task execution but has additional tracking to ensure a task is not scheduled more than once. 
+			// This method should only be used when scheduling more than once is allowed/practical (e.g. FTaskEvent).
+			// Returns false if the task has been triggered before or has incomplete dependencies (prerequisites or is blocked by a pipe). 
+			// In the latter case, the task will be automatically scheduled when all dependencies are completed.
+			bool Trigger(uint64 TaskSize)
+			{
+				if (TaskTriggered.exchange(true, std::memory_order_relaxed))
+				{
+					// Task is already launched
+					return false;
+				}
+
+				// Now that we know we have triggered the task (it might not yet be launched/running based on dependencies) it's now in the
+				// task system and can outlive external references, so we need to keep it alive by holding an internal reference. 
+				// It will be released when completed.
+				AddRef();
+				return TryLaunch(TaskSize);
+			}
+
 			// Tries to schedule task execution. Returns false if the task has incomplete dependencies (prerequisites or is blocked by a pipe). 
 			// In this case the task will be automatically scheduled when all dependencies are completed.
 			bool TryLaunch(uint64 TaskSize)
@@ -699,6 +718,9 @@ public:
 
 			EExtendedTaskPriority ExtendedPriority; // internal priorities, if any
 
+			// Note: Only ever set when the task is launched from Trigger()
+			std::atomic<bool> TaskTriggered = false;
+
 			std::atomic<uint32> ExecutingThreadId = FThread::InvalidThreadId;
 
 #if UE_TASK_TRACE_ENABLED
diff --git a/Engine/Source/Runtime/Core/Tests/Tasks/TasksTest.cpp b/Engine/Source/Runtime/Core/Tests/Tasks/TasksTest.cpp
--- a/Engine/Source/Runtime/Core/Tests/Tasks/TasksTest.cpp
+++ b/Engine/Source/Runtime/Core/Tests/Tasks/TasksTest.cpp
@@ -815,6 +815,21 @@ namespace UE { namespace TasksTests
 		Wait(BranchTasks);
 	}
 
+	TEST_CASE_NAMED(FTasksEventTest, "System::Core::Tasks::TriggerTaskEventConcurrently", "[ApplicationContextMask][EngineFilter]")
+	{
+		// Regression test to ensure we do not assert when triggering an event from more than one thread at a time.
+		for (int i = 0; i < 100000; ++i)
+		{
+			FTaskEvent Event{ UE_SOURCE_LOCATION };
+			FTask DeferredTask{ Launch(UE_SOURCE_LOCATION, [] {}, Event) };
+			ParallelFor(64, [Event, DeferredTask] (int32) mutable
+			{
+				Event.Trigger();
+				DeferredTask.Wait();
+			});
+		}
+	}
+
 	TEST_CASE_NAMED(FTasksDependenciesTest, "System::Core::Tasks::Dependencies", "[ApplicationContextMask][EngineFilter]")
 	{
 		{	// a task is not executed until its prerequisite (FTaskEvent) is completed
