Mac - Use proc_pid_rusage on Mac for reporting physical memory use

#tests Local testing by generating crash and checking memory

diff --git a/Engine/Source/Runtime/Core/Private/Apple/ApplePlatformMemory.cpp b/Engine/Source/Runtime/Core/Private/Apple/ApplePlatformMemory.cpp
--- a/Engine/Source/Runtime/Core/Private/Apple/ApplePlatformMemory.cpp
+++ b/Engine/Source/Runtime/Core/Private/Apple/ApplePlatformMemory.cpp
@@ -25,6 +25,7 @@
 #include <sys/mman.h>
 #include <sys/mount.h>
 #include <objc/runtime.h>
+#include <libproc.h>
 #if PLATFORM_IOS
 #include "IOS/IOSPlatformMisc.h"
 #include <os/proc.h>
@@ -439,7 +440,6 @@ FPlatformMemoryStats FApplePlatformMemory::GetStats()
 	// considers "free" pages a waste of resources.
 	MemoryStats.AvailableVirtual = (Stats.free_count + Stats.inactive_count) * MemoryConstants.PageSize;
 	
-#if PLATFORM_IOS
 	rusage_info_current rusage_payload;
 		 
 	int ret = proc_pid_rusage(getpid(),
@@ -448,14 +448,7 @@ FPlatformMemoryStats FApplePlatformMemory::GetStats()
 	NSCAssert(ret == 0, @"Could not get rusage: %i.", errno); // Look up in `man errno`.
 			
 	MemoryStats.UsedPhysical = rusage_payload.ri_phys_footprint;
-#else
-	// Just get memory information for the process and report the working set instead
-	mach_task_basic_info_data_t TaskInfo;
-	mach_msg_type_number_t TaskInfoCount = MACH_TASK_BASIC_INFO_COUNT;
-	task_info( mach_task_self(), MACH_TASK_BASIC_INFO, (task_info_t)&TaskInfo, &TaskInfoCount );
 
-    MemoryStats.UsedPhysical = TaskInfo.resident_size;
-#endif
 	if(MemoryStats.UsedPhysical > MemoryStats.PeakUsedPhysical)
 	{
 		MemoryStats.PeakUsedPhysical = MemoryStats.UsedPhysical;
