ApplePlatformMemory - Fix compile error on iOS

#tests local build on iOS & Mac

diff --git a/Engine/Source/Runtime/Core/Private/Apple/ApplePlatformMemory.cpp b/Engine/Source/Runtime/Core/Private/Apple/ApplePlatformMemory.cpp
--- a/Engine/Source/Runtime/Core/Private/Apple/ApplePlatformMemory.cpp
+++ b/Engine/Source/Runtime/Core/Private/Apple/ApplePlatformMemory.cpp
@@ -25,7 +25,9 @@
 #include <sys/mman.h>
 #include <sys/mount.h>
 #include <objc/runtime.h>
+#if PLATFORM_MAC
 #include <libproc.h>
+#endif
 #if PLATFORM_IOS
 #include "IOS/IOSPlatformMisc.h"
 #include <os/proc.h>
