Fix the double free caused by us mixing creation of `TStatId`'s in the open/closed.

All the actual heroics were done by [REDACTED_PERSON], this is just the easy bit fixing it cause his PC decided to yeet off into the sun.

#rnx

diff --git a/Engine/Source/Programs/AutoRTFMTests/AutoRTFMTestsWithSTATS.Target.cs b/Engine/Source/Programs/AutoRTFMTests/AutoRTFMTestsWithSTATS.Target.cs
new file mode 100644
--- /dev/null
+++ b/Engine/Source/Programs/AutoRTFMTests/AutoRTFMTestsWithSTATS.Target.cs
@@ -0,0 +1,17 @@
+// Copyright Epic Games, Inc. All Rights Reserved.
+
+using EpicGames.Core;
+using UnrealBuildTool;
+using System.Collections.Generic;
+using UnrealBuildBase;
+
+[SupportedPlatforms(UnrealPlatformClass.All)]
+public class AutoRTFMTestsWithSTATSTarget : AutoRTFMTestsTarget
+{
+	public AutoRTFMTestsWithSTATSTarget(TargetInfo Target) : base(Target)
+	{
+		GlobalDefinitions.Remove("ENABLE_STATNAMEDEVENTS=1");
+		GlobalDefinitions.Remove("ENABLE_STATNAMEDEVENTS_UOBJECT=1");
+		GlobalDefinitions.Add("FORCE_USE_STATS=1");
+	}
+}
diff --git a/Engine/Source/Programs/AutoRTFMTests/Private/AutoRTFMTests.cpp b/Engine/Source/Programs/AutoRTFMTests/Private/AutoRTFMTests.cpp
--- a/Engine/Source/Programs/AutoRTFMTests/Private/AutoRTFMTests.cpp
+++ b/Engine/Source/Programs/AutoRTFMTests/Private/AutoRTFMTests.cpp
@@ -144,7 +144,7 @@ int main(int ArgC, const char* ArgV[])
 	}
 #endif
 
-	return Result;
+	return 0 == FDebug::GetNumEnsureFailures() ? 0 : -1;
 }
 
 bool CheckNoMemoryLeaks()
diff --git a/Engine/Source/Programs/AutoRTFMTests/Private/UECoreTests.cpp b/Engine/Source/Programs/AutoRTFMTests/Private/UECoreTests.cpp
--- a/Engine/Source/Programs/AutoRTFMTests/Private/UECoreTests.cpp
+++ b/Engine/Source/Programs/AutoRTFMTests/Private/UECoreTests.cpp
@@ -807,6 +807,7 @@ TEST_CASE("UECore.FTextCache")
 	}
 }
 
+#if ENABLE_STATNAMEDEVENTS_UOBJECT
 TEST_CASE("UECore.FUObjectItem")
 {
 	SECTION("CreateStatID First In Open")
@@ -956,6 +957,69 @@ TEST_CASE("UECore.FUObjectItem")
 		}
 	}
 }
+#endif
+
+#if STATS
+DECLARE_STATS_GROUP(TEXT("AutoRTFM Testing"), STATGROUP_AutoRTFMTesting, STATCAT_Hidden);
+
+TEST_CASE("TStatId")
+{
+	SECTION("QUICK_SCOPE_CYCLE_COUNTER after new statid bug")
+	{
+		AutoRTFM::Testing::Commit([&]
+		{
+			FString Name = "Test"; Name.AppendInt(__LINE__);
+			FDynamicStats::CreateStatId<FStatGroup_STATGROUP_AutoRTFMTesting>(Name);
+
+			// This gets initialized in the open, which caused a bug previously with the `CreateStatId` above.
+			QUICK_SCOPE_CYCLE_COUNTER(Stat_MyThing);
+		});
+	}
+
+	SECTION("FString Construction")
+	{
+		AutoRTFM::Testing::Commit([&]
+		{
+			{
+				FString Name = "Test"; Name.AppendInt(__LINE__);
+				FDynamicStats::CreateStatId<FStatGroup_STATGROUP_AutoRTFMTesting>(Name);
+			}
+
+			{
+				FString Name = "Test"; Name.AppendInt(__LINE__);
+				FDynamicStats::CreateStatIdInt64<FStatGroup_STATGROUP_AutoRTFMTesting>(Name);
+			}
+
+		
+			{
+				FString Name = "Test"; Name.AppendInt(__LINE__);
+				FDynamicStats::CreateStatIdDouble<FStatGroup_STATGROUP_AutoRTFMTesting>(Name);
+			}
+
+			{
+				FString Name = "Test"; Name.AppendInt(__LINE__);
+				FDynamicStats::CreateMemoryStatId<FStatGroup_STATGROUP_AutoRTFMTesting>(Name);
+			}
+		});
+	}
+
+	SECTION("FName Construction")
+	{
+		AutoRTFM::Testing::Commit([&]
+		{
+			{
+				FString Name = "Test"; Name.AppendInt(__LINE__);
+				FDynamicStats::CreateStatId<FStatGroup_STATGROUP_AutoRTFMTesting>(FName(Name));
+			}
+
+			{
+				FString Name = "Test"; Name.AppendInt(__LINE__);
+				FDynamicStats::CreateMemoryStatId<FStatGroup_STATGROUP_AutoRTFMTesting>(FName(Name));
+			}
+		});
+	}
+}
+#endif
 
 TEST_CASE("UECore.TScopeLock_TransactionallySafeCriticalSection")
 {
diff --git a/Engine/Source/Runtime/ApplicationCore/Private/Windows/WindowsPlatformSplash.cpp b/Engine/Source/Runtime/ApplicationCore/Private/Windows/WindowsPlatformSplash.cpp
--- a/Engine/Source/Runtime/ApplicationCore/Private/Windows/WindowsPlatformSplash.cpp
+++ b/Engine/Source/Runtime/ApplicationCore/Private/Windows/WindowsPlatformSplash.cpp
@@ -1,6 +1,7 @@
 // Copyright Epic Games, Inc. All Rights Reserved.
 
 #include "Windows/WindowsPlatformSplash.h"
+#include "AutoRTFM/AutoRTFM.h"
 #include "HAL/PlatformMemory.h"
 #include "HAL/PlatformMisc.h"
 #include "HAL/PlatformApplicationMisc.h"
@@ -721,6 +722,7 @@ void FWindowsPlatformSplash::Show()
 			DWORD ThreadID = 0;
 			GSplashScreenThread = CreateThread(NULL, 128 * 1024, (LPTHREAD_START_ROUTINE)StartSplashScreenThread, (LPVOID)NULL, STACK_SIZE_PARAM_IS_A_RESERVATION, &ThreadID);
 #if	STATS
+			AutoRTFM::UnreachableIfClosed("TODO: Could remove if we marked `FWindowsPlatformSplash` as `AUTORTFM_DISABLE`!");
 			FStartupMessages::Get().AddThreadMetadata( FName( "SplashScreenThread" ), ThreadID );
 #endif // STATS
 		}
diff --git a/Engine/Source/Runtime/Core/Private/HAL/ThreadingBase.cpp b/Engine/Source/Runtime/Core/Private/HAL/ThreadingBase.cpp
--- a/Engine/Source/Runtime/Core/Private/HAL/ThreadingBase.cpp
+++ b/Engine/Source/Runtime/Core/Private/HAL/ThreadingBase.cpp
@@ -894,6 +894,7 @@ void FRunnableThread::SetupCreatedThread(FRunnableThread*& NewThread, class FRun
 void FRunnableThread::PostCreate(EThreadPriority InThreadPriority)
 {
 #if	STATS
+	AutoRTFM::UnreachableIfClosed("TODO: Can remove this if we mark `FRunnableThread` as `AUTORTFM_DISABLE`!");
 	FStartupMessages::Get().AddThreadMetadata( FName( *GetThreadName() ), GetThreadID() );
 #endif // STATS
 }
diff --git a/Engine/Source/Runtime/Core/Private/Stats/StatsData.cpp b/Engine/Source/Runtime/Core/Private/Stats/StatsData.cpp
--- a/Engine/Source/Runtime/Core/Private/Stats/StatsData.cpp
+++ b/Engine/Source/Runtime/Core/Private/Stats/StatsData.cpp
@@ -1715,8 +1715,12 @@ FName FStatsThreadState::GetStatThreadName( const FStatPacket& Packet ) const
 
 			static const FName NAME_UnknownThread = TEXT( "UnknownThread" );
 			NewThreadName = FName( *FStatsUtils::BuildUniqueThreadName( Packet.ThreadId ) );
-			// This is an unknown thread, but still we need the metadata in the system.
-			FStartupMessages::Get().AddThreadMetadata( NAME_UnknownThread, Packet.ThreadId );
+
+			UE_AUTORTFM_OPEN
+			{
+				// This is an unknown thread, but still we need the metadata in the system.
+				FStartupMessages::Get().AddThreadMetadata( NAME_UnknownThread, Packet.ThreadId );
+			};
 		}
 		ThreadName = NewThreadName;
 	}
diff --git a/Engine/Source/Runtime/Core/Private/Stats/StatsSystemTypes.cpp b/Engine/Source/Runtime/Core/Private/Stats/StatsSystemTypes.cpp
--- a/Engine/Source/Runtime/Core/Private/Stats/StatsSystemTypes.cpp
+++ b/Engine/Source/Runtime/Core/Private/Stats/StatsSystemTypes.cpp
@@ -1290,6 +1290,8 @@ void FThreadStats::ExplicitFlush(bool DiscardCallstack)
 
 void FThreadStats::StartThread()
 {
+	AutoRTFM::UnreachableIfClosed("TODO: Can remove if `FEngineLoop` is marked `AUTORTFM_DISABLE`!");
+
 	FThreadStats::FrameDataIsIncomplete(); // make this non-zero
 	check(IsInGameThread());
 	check(!IsThreadingReady());
diff --git a/Engine/Source/Runtime/Core/Public/Stats/DynamicStats.h b/Engine/Source/Runtime/Core/Public/Stats/DynamicStats.h
--- a/Engine/Source/Runtime/Core/Public/Stats/DynamicStats.h
+++ b/Engine/Source/Runtime/Core/Public/Stats/DynamicStats.h
@@ -1,6 +1,7 @@
 // Copyright Epic Games, Inc. All Rights Reserved.
 #pragma once
 
+#include "AutoRTFM/AutoRTFM.h"
 #include "Misc/Build.h"
 #include "Stats/LightweightStats.h"
 #include "Stats/StatsSystemTypes.h"
@@ -63,7 +64,7 @@ struct FDynamicStats
 	static TStatId CreateMemoryStatId(const FString& StatNameOrDescription, FPlatformMemory::EMemoryCounterRegion MemRegion = FPlatformMemory::MCR_Physical)
 	{
 #if	STATS
-		return CreateMemoryStatId<TStatGroup>(FName(*StatNameOrDescription), MemRegion);
+		return CreateMemoryStatIdInternal<TStatGroup>(FName(*StatNameOrDescription), MemRegion);
 #else
 		return TStatId();
 #endif // STATS
@@ -73,41 +74,70 @@ struct FDynamicStats
 	static TStatId CreateMemoryStatId(const FName StatNameOrDescription, FPlatformMemory::EMemoryCounterRegion MemRegion = FPlatformMemory::MCR_Physical)
 	{
 #if	STATS
-		FStartupMessages::Get().AddMetadata(StatNameOrDescription, *StatNameOrDescription.ToString(),
+		return CreateMemoryStatIdInternal<TStatGroup>(StatNameOrDescription, MemRegion);
+#else
+		return TStatId();
+#endif // STATS
+	}
+
+#if	STATS
+private: // private since this can only be declared if STATS is defined, due to EStatDataType in signature
+	template< typename TStatGroup >
+	static TStatId CreateStatIdInternal(const FName StatNameOrDescription, EStatDataType::Type Type, bool IsTimer, bool bClearEveryFrame = true)
+	{
+		return AutoRTFM::Open([&]
+		{
+			// #jira SOL-8290: can remove this `UnreachableIfClosed` once resolved.
+			AutoRTFM::UnreachableIfClosed();
+			
+			return CreateStatIdInternalInTheOpen<TStatGroup>(StatNameOrDescription, Type, IsTimer, bClearEveryFrame);
+		});
+	}
+
+	template< typename TStatGroup >
+	AUTORTFM_DISABLE static TStatId CreateStatIdInternalInTheOpen(const FName StatNameOrDescription, EStatDataType::Type Type, bool IsTimer, bool bClearEveryFrame)
+	{
+		FStartupMessages::Get().AddMetadata(StatNameOrDescription, nullptr,
 			TStatGroup::GetGroupName(),
 			TStatGroup::GetGroupCategory(),
 			TStatGroup::GetDescription(),
-			false, EStatDataType::ST_int64, false, false, MemRegion);
+			bClearEveryFrame, Type, IsTimer, false);
 
 		TStatId StatID = IStatGroupEnableManager::Get().GetHighPerformanceEnableForStat(StatNameOrDescription,
 			TStatGroup::GetGroupName(),
 			TStatGroup::GetGroupCategory(),
 			TStatGroup::DefaultEnable,
-			false, EStatDataType::ST_int64, *StatNameOrDescription.ToString(), false, false, MemRegion);
+			bClearEveryFrame, Type, nullptr, IsTimer, false);
 
 		return StatID;
-#else
-		return TStatId();
-#endif // STATS
 	}
 
-#if	STATS
-private: // private since this can only be declared if STATS is defined, due to EStatDataType in signature
 	template< typename TStatGroup >
-	static TStatId CreateStatIdInternal(const FName StatNameOrDescription, EStatDataType::Type Type, bool IsTimer, bool bClearEveryFrame = true)
+	static TStatId CreateMemoryStatIdInternal(const FName StatNameOrDescription, FPlatformMemory::EMemoryCounterRegion MemRegion = FPlatformMemory::MCR_Physical)
 	{
+		return AutoRTFM::Open([&]
+		{
+			// #jira SOL-8290: can remove this `UnreachableIfClosed` once resolved.
+			AutoRTFM::UnreachableIfClosed();
 
-		FStartupMessages::Get().AddMetadata(StatNameOrDescription, nullptr,
+			return CreateMemoryStatIdInternalInTheOpen<TStatGroup>(StatNameOrDescription, MemRegion);
+		});
+	}
+
+	template< typename TStatGroup >
+	AUTORTFM_DISABLE static TStatId CreateMemoryStatIdInternalInTheOpen(const FName StatNameOrDescription, FPlatformMemory::EMemoryCounterRegion MemRegion = FPlatformMemory::MCR_Physical)
+	{
+		FStartupMessages::Get().AddMetadata(StatNameOrDescription, *StatNameOrDescription.ToString(),
 			TStatGroup::GetGroupName(),
 			TStatGroup::GetGroupCategory(),
 			TStatGroup::GetDescription(),
-			bClearEveryFrame, Type, IsTimer, false);
+			false, EStatDataType::ST_int64, false, false, MemRegion);
 
 		TStatId StatID = IStatGroupEnableManager::Get().GetHighPerformanceEnableForStat(StatNameOrDescription,
 			TStatGroup::GetGroupName(),
 			TStatGroup::GetGroupCategory(),
 			TStatGroup::DefaultEnable,
-			bClearEveryFrame, Type, nullptr, IsTimer, false);
+			false, EStatDataType::ST_int64, *StatNameOrDescription.ToString(), false, false, MemRegion);
 
 		return StatID;
 	}
diff --git a/Engine/Source/Runtime/Core/Public/Stats/LightweightStats.h b/Engine/Source/Runtime/Core/Public/Stats/LightweightStats.h
--- a/Engine/Source/Runtime/Core/Public/Stats/LightweightStats.h
+++ b/Engine/Source/Runtime/Core/Public/Stats/LightweightStats.h
@@ -49,6 +49,11 @@ struct TStatId
 	{
 		return StatString != Other.StatString;
 	}
+
+	static void AutoRTFMAssignFromOpenToClosed(TStatId& Closed, const TStatId& Open)
+	{
+		Closed = Open;
+	}
 };
 
 namespace UE::Stats::Private
@@ -407,7 +412,13 @@ namespace UE::Stats::Private
 
 #elif !STATS // UE_USE_LIGHTWEIGHT_STATS
 
-struct TStatId {};
+struct TStatId
+{
+	static void AutoRTFMAssignFromOpenToClosed(TStatId& Closed, const TStatId& Open)
+	{
+		Closed = Open;
+	}
+};
 
 class FScopeCycleCounter
 {
diff --git a/Engine/Source/Runtime/Core/Public/Stats/StatsSystemTypes.h b/Engine/Source/Runtime/Core/Public/Stats/StatsSystemTypes.h
--- a/Engine/Source/Runtime/Core/Public/Stats/StatsSystemTypes.h
+++ b/Engine/Source/Runtime/Core/Public/Stats/StatsSystemTypes.h
@@ -146,6 +146,11 @@ struct TStatId
 		return GetTypeHash(StatId.StatIdPtr);
 	}
 
+	static void AutoRTFMAssignFromOpenToClosed(TStatId& Closed, const TStatId& Open)
+	{
+		Closed = Open;
+	}
+
 private:
 	/** NAME_None. */
 	CORE_API static TStatIdData TStatId_NAME_None;
@@ -1704,14 +1709,14 @@ class FStartupMessages
 	friend class FStatsThread;
 
 	TArray64<FStatMessage> DelayedMessages;
-	FTransactionallySafeCriticalSection CriticalSection;
+	FCriticalSection CriticalSection;
 
 public:
 	/** Adds a thread metadata. */
-	CORE_API void AddThreadMetadata(const FName InThreadFName, uint32 InThreadID);
+	AUTORTFM_DISABLE CORE_API void AddThreadMetadata(const FName InThreadFName, uint32 InThreadID);
 
 	/** Adds a regular metadata. */
-	CORE_API void AddMetadata(FName InStatName, const TCHAR* InStatDesc, const char* InGroupName, const char* InGroupCategory, const TCHAR* InGroupDesc, bool bShouldClearEveryFrame, EStatDataType::Type InStatType, bool bCycleStat, bool bSortByName, FPlatformMemory::EMemoryCounterRegion InMemoryRegion = FPlatformMemory::MCR_Invalid);
+	AUTORTFM_DISABLE CORE_API void AddMetadata(FName InStatName, const TCHAR* InStatDesc, const char* InGroupName, const char* InGroupCategory, const TCHAR* InGroupDesc, bool bShouldClearEveryFrame, EStatDataType::Type InStatType, bool bCycleStat, bool bSortByName, FPlatformMemory::EMemoryCounterRegion InMemoryRegion = FPlatformMemory::MCR_Invalid);
 
 	/** Access the singleton. */
 	CORE_API static FStartupMessages& Get();
@@ -1781,7 +1786,7 @@ struct FThreadSafeStaticStatBase
 {
 protected:
 	mutable TAtomic<const TStatIdData*> HighPerformanceEnable { nullptr };
-	CORE_API const TStatIdData* DoSetup(FName InStatName, const TCHAR* InStatDesc, const char* InGroupName, const char* InGroupCategory, const TCHAR* InGroupDesc, bool bDefaultEnable, bool bShouldClearEveryFrame, EStatDataType::Type InStatType, bool bCycleStat, bool bSortByName, FPlatformMemory::EMemoryCounterRegion InMemoryRegion) const;
+	AUTORTFM_OPEN CORE_API const TStatIdData* DoSetup(FName InStatName, const TCHAR* InStatDesc, const char* InGroupName, const char* InGroupCategory, const TCHAR* InGroupDesc, bool bDefaultEnable, bool bShouldClearEveryFrame, EStatDataType::Type InStatType, bool bCycleStat, bool bSortByName, FPlatformMemory::EMemoryCounterRegion InMemoryRegion) const;
 };
 
 template<class TStatData, bool TCompiledIn>
