- Fixed an issue with removing array entries when saving with the ini.UseNewPropertySaving=1

diff --git a/Engine/Source/Runtime/Core/Private/Misc/ConfigCacheIni.cpp b/Engine/Source/Runtime/Core/Private/Misc/ConfigCacheIni.cpp
--- a/Engine/Source/Runtime/Core/Private/Misc/ConfigCacheIni.cpp
+++ b/Engine/Source/Runtime/Core/Private/Misc/ConfigCacheIni.cpp
@@ -5711,7 +5711,6 @@ bool FConfigCacheIni::SetArrayInSectionOfStaticLayer(const TCHAR* Section, FName
 				// @todo: this needs testing with expanion variables - we probably had unexpanded values coming in, in which case we never want to expand - hard to say
 				for (const FConfigValue* ParentValue : ParentValues)
 				{
-					Sec->Remove(Key);
 					Sec->Add(Key, FConfigValue(ParentValue->GetSavedValueForWriting(), FConfigValue::EValueType::Remove));
 				}
 				for (const FString& Value : ValuesToAdd)
