Log large memory allocations and their framepointers if the allocation size exceeds the specified MB threshold. Disabled by default. Code is also opted out by default.

#rnx

diff --git a/Engine/Source/Runtime/Core/Private/HAL/MallocBinned3.cpp b/Engine/Source/Runtime/Core/Private/HAL/MallocBinned3.cpp
--- a/Engine/Source/Runtime/Core/Private/HAL/MallocBinned3.cpp
+++ b/Engine/Source/Runtime/Core/Private/HAL/MallocBinned3.cpp
@@ -57,6 +57,78 @@ uint16 FMallocBinned3::SmallBinSizesShifted[UE_MB3_SMALL_POOL_COUNT + 1] = { 0 }
 	uint8* FMallocBinned3::PoolBaseVMPtr[UE_MB3_SMALL_POOL_COUNT] = { nullptr };
 #endif
 
+#ifndef UE_MB3_LOG_LARGE_ALLOCATION
+#define UE_MB3_LOG_LARGE_ALLOCATION 0
+#endif
+
+#if UE_MB3_LOG_LARGE_ALLOCATION
+
+#include "HAL/PlatformStackWalk.h"
+
+	static int32 GLogLargeAllocationSizeInMB = 0;
+	static FAutoConsoleVariableRef CVarLogLargeAllocationSizeInMB(
+		TEXT("Memory.LogLargeAllocationInMB"),
+		GLogLargeAllocationSizeInMB,
+		TEXT("Log large memory allocations and their callstacks if the allocation size exceeds the specified MB threshold (default is 0, meaning no allocations are logged)."));
+
+	static void LogLargeAllocation(SIZE_T Size)
+	{
+		if (GLogLargeAllocationSizeInMB <= 0)
+		{
+			return;
+		}
+
+		int32 AllocationSizeInMB = int32(Size >> 20);
+
+		if (AllocationSizeInMB < GLogLargeAllocationSizeInMB)
+		{
+			return;
+		}
+
+		static bool bInFunction = false;
+		if (bInFunction)
+		{
+			return;
+		}
+		
+		TGuardValue<bool> GuardInFunction(bInFunction, true);
+
+		// Grab callstack
+		static constexpr int32 MaxCallStackDepth = 16;
+		static constexpr int32 CallStackEntriesToSkipCount = 2;
+
+		uint64 FullCallStack[MaxCallStackDepth + CallStackEntriesToSkipCount] = { 0 };
+		FPlatformStackWalk::CaptureStackBackTrace(&FullCallStack[0], MaxCallStackDepth + CallStackEntriesToSkipCount);
+
+		// Skip first n entries as they are inside the allocator.
+		uint64* CallStack = &FullCallStack[CallStackEntriesToSkipCount];
+
+		UE_LOG(LogMemory, Log, TEXT("Allocating %d MB - 0x%08llx 0x%08llx 0x%08llx 0x%08llx 0x%08llx 0x%08llx 0x%08llx 0x%08llx 0x%08llx 0x%08llx 0x%08llx 0x%08llx 0x%08llx 0x%08llx 0x%08llx 0x%08llx"),
+			   AllocationSizeInMB,
+			   CallStack[0],
+			   CallStack[1],
+			   CallStack[2],
+			   CallStack[3],
+			   CallStack[4],
+			   CallStack[5],
+			   CallStack[6],
+			   CallStack[7],
+			   CallStack[8],
+			   CallStack[9],
+			   CallStack[10],
+			   CallStack[11],
+			   CallStack[12],
+			   CallStack[13],
+			   CallStack[14],
+			   CallStack[15]);
+	}
+
+#else // UE_MB3_LOG_LARGE_ALLOCATION
+
+	static void LogLargeAllocation(SIZE_T Size) {}
+
+#endif // UE_MB3_LOG_LARGE_ALLOCATION
+
 FMallocBinned3* FMallocBinned3::MallocBinned3 = nullptr;
 // Mapping of sizes to small table indices
 uint8 FMallocBinned3::MemSizeToPoolIndex[1 + (UE_MB3_MAX_SMALL_POOL_SIZE >> UE_MBC_MIN_SMALL_POOL_ALIGNMENT_SHIFT)] = { 0 };
@@ -779,12 +851,13 @@ void* FMallocBinned3::MallocExternal(SIZE_T Size, uint32 Alignment)
 
 	NOALLOC_SCOPE_CYCLE_COUNTER(STAT_FMallocBinned3_MallocExternalLarge);
 
+	LogLargeAllocation(AlignedSize);
+
 #if UE_MB3_USE_CACHED_PAGE_ALLOCATOR_FOR_LARGE_ALLOCS
 	UE::TUniqueLock Lock(Mutex);
 	void* Result = GetCachedOSPageAllocator().Allocate(AlignedSize);
 	check(IsAligned(Result, Alignment));
 #else
-
 	FPlatformMemory::FPlatformVirtualMemoryBlock Block = FPlatformMemory::FPlatformVirtualMemoryBlock::AllocateVirtual(AlignedSize, Alignment);
 	Block.Commit(0, AlignedSize);
 	void* Result = Block.GetVirtualPointer();
