Fixed mac process reaping to be non blocking and not leaking

#tests Checked that pids properly got waited on

diff --git a/Engine/Source/Runtime/Core/Private/Mac/MacPlatformProcess.cpp b/Engine/Source/Runtime/Core/Private/Mac/MacPlatformProcess.cpp
--- a/Engine/Source/Runtime/Core/Private/Mac/MacPlatformProcess.cpp
+++ b/Engine/Source/Runtime/Core/Private/Mac/MacPlatformProcess.cpp
@@ -560,8 +560,7 @@ FProcHandle FMacPlatformProcess::CreateProcInternal(const TCHAR* URL, const TCHA
 		*OutProcessID = ChildPid;
 	}
 
-	// [RCL] 2015-03-11 @FIXME: is bLaunchDetached usable when determining whether we're in 'fire and forget' mode? This doesn't exactly match what bLaunchDetached is used for.
-	return FProcHandle(new FProcState(ChildPid, bLaunchDetached));
+	return FProcHandle(new FProcState(ChildPid));
 }
 
 /*
@@ -586,119 +585,136 @@ FProcHandle FMacPlatformProcess::OpenProcess(uint32 ProcessID)
 	return FProcHandle(bIsRunning ? Pid : -1);
 }
 
-/**
- * This class exists as an imperfect workaround to allow both "fire and forget" children and children about whose return code we actually care.
- * (maybe we could fork and daemonize ourselves for the first case instead?)
- */
-struct FChildWaiterThread : public FRunnable
+namespace UE::Mac
 {
-	/** Global table of all waiter threads */
-	static TArray<FChildWaiterThread *>		ChildWaiterThreadsArray;
-
-	/** Lock guarding the acess to child waiter threads */
-	static FCriticalSection					ChildWaiterThreadsArrayGuard;
-
-	/** Pid of child to wait for */
-	int ChildPid;
-
-	FChildWaiterThread(pid_t InChildPid)
-		:	ChildPid(InChildPid)
-	{
-		// add ourselves to thread array
-		ChildWaiterThreadsArrayGuard.Lock();
-		ChildWaiterThreadsArray.Add(this);
-		ChildWaiterThreadsArrayGuard.Unlock();
+	namespace
+	{
+		bool bProcessReaperStarted = false;
+		// Using semaphore_t as it's async-signal-safe
+		semaphore_t ProcessReaperSemaphore;
+		std::atomic_bool bProcessReaperSignaled = false;
+		// Using pthread_mutex_t so it's not dynamic initialized
+		pthread_mutex_t ProcessesToReapMutex = PTHREAD_MUTEX_INITIALIZER;
+		// [[clang::no_destroy]] required here to not potentially crash on exit
+		[[clang::no_destroy]] TSet<pid_t> ProcessesToReap;
+		
+		void ReleaseProcessReaper()
+		{
+			if (!bProcessReaperSignaled.exchange(true, std::memory_order_seq_cst))
+			{
+				verify(semaphore_signal(ProcessReaperSemaphore) == KERN_SUCCESS);
+			}
+		}
+		
+		void AcquireProcessReaper()
+		{
+			verify(semaphore_wait(ProcessReaperSemaphore) == KERN_SUCCESS);
+			bProcessReaperSignaled.store(false, std::memory_order_seq_cst);
+		}
 	}
 
-	virtual ~FChildWaiterThread()
+	void StartProcessReaper()
 	{
-		// remove
-		ChildWaiterThreadsArrayGuard.Lock();
-		ChildWaiterThreadsArray.RemoveSingle(this);
-		ChildWaiterThreadsArrayGuard.Unlock();
-	}
+		if (bProcessReaperStarted)
+		{
+			return;
+		}
+		
+		bProcessReaperStarted = true;
+		
+		verify(semaphore_create(mach_task_self(), &ProcessReaperSemaphore, SYNC_POLICY_FIFO, 0) == KERN_SUCCESS);
 
-	virtual uint32 Run()
-	{
-		for(;;)	// infinite loop in case we get EINTR and have to repeat
+		// Using pthread_t so it runs forever
+		pthread_t Thread;
+		verify(pthread_create(&Thread, nullptr, [](void*) -> void*
 		{
-			siginfo_t SignalInfo;
-			if (waitid(P_PID, ChildPid, &SignalInfo, WEXITED))
+			auto Get = []() -> pid_t
+			{
+				siginfo_t SigInfo;
+				SigInfo.si_pid = 0;
+				while (waitid(P_ALL, 0, &SigInfo, WEXITED | WNOHANG | WNOWAIT) && errno == EINTR);
+				return SigInfo.si_pid;
+			};
+			
+			auto Reap = [](pid_t PId) -> bool
 			{
-				if (errno != EINTR)
+				do
 				{
-					int ErrNo = errno;
-					UE_LOG(LogHAL, Fatal, TEXT("FChildWaiterThread::Run(): waitid for pid %d failed (errno=%d, %s)"), 
-							 static_cast< int32 >(ChildPid), ErrNo, UTF8_TO_TCHAR(strerror(ErrNo)));
-					break;	// exit the loop if for some reason Fatal log (above) returns
+					if (pid_t Result = waitpid(PId, nullptr, WNOHANG); Result >= 0)
+					{
+						return Result;
+					}
 				}
-			}
-			else
+				while (errno == EINTR);
+				return false;
+			};
+			
+			for (;;)
 			{
-				check(SignalInfo.si_pid == ChildPid);
-				break;
-			}
-		}
+				AcquireProcessReaper();
+				
+				while (pid_t PId = Get())
+				{
+					verify(pthread_mutex_lock(&ProcessesToReapMutex) == 0);
+					ON_SCOPE_EXIT
+					{
+						verify(pthread_mutex_unlock(&ProcessesToReapMutex) == 0);
+					};
+
+					if (!ProcessesToReap.Remove(PId))
+					{
+						for (TSet<pid_t>::TIterator Iterator = ProcessesToReap; Iterator; ++Iterator)
+						{
+							if (Reap(*Iterator))
+							{
+								Iterator.RemoveCurrent();
+							}
+						}
+						break;
+					}
 
-		return 0;
+					verify(Reap(PId));
+				}
+			}
+		}, nullptr) == 0);
+		
+		struct sigaction Action;
+		Action.sa_handler = [](int)
+		{
+			ReleaseProcessReaper();
+		};
+		sigemptyset(&Action.sa_mask);
+		Action.sa_flags = SA_NOCLDSTOP | SA_RESTART;
+		verify(sigaction(SIGCHLD, &Action, nullptr) == 0);
 	}
 
-	virtual void Exit()
+	void AddToProcessReaper(pid_t PId)
 	{
-		// unregister from the array
-		delete this;
-	}
-};
+		verify(pthread_mutex_lock(&ProcessesToReapMutex) == 0);
+		ON_SCOPE_EXIT
+		{
+			verify(pthread_mutex_unlock(&ProcessesToReapMutex) == 0);
+		};
 
-/** See FChildWaiterThread */
-TArray<FChildWaiterThread *> FChildWaiterThread::ChildWaiterThreadsArray;
-/** See FChildWaiterThread */
-FCriticalSection FChildWaiterThread::ChildWaiterThreadsArrayGuard;
+		StartProcessReaper();
+		
+		ProcessesToReap.Add(PId);
+		ReleaseProcessReaper();
+	}
+}
 
 /** Initialization constructor. */
-FProcState::FProcState(pid_t InProcessId, bool bInFireAndForget)
-	:	ProcessId(InProcessId)
-	,	bIsRunning(true)  // assume it is
-	,	bHasBeenWaitedFor(false)
-	,	ReturnCode(-1)
-	,	bFireAndForget(bInFireAndForget)
+FProcState::FProcState(pid_t InProcessId)
+:	ProcessId{InProcessId}
 {
+	check(InProcessId);
 }
 
 FProcState::~FProcState()
 {
-	if (!bFireAndForget)
+	if (IsRunning())
 	{
-		// If not in 'fire and forget' mode, try to catch the common problems that leave zombies:
-		// - We don't want to close the handle of a running process as with our current scheme this will certainly leak a zombie.
-		// - Nor we want to leave the handle unwait()ed for.
-		
-		if (bIsRunning)
-		{
-			// Warn the users before going into what may be a very long block
-			UE_LOG(LogHAL, Log, TEXT("Closing a process handle while the process (pid=%d) is still running - we will block until it exits to prevent a zombie"),
-				GetProcessId()
-			);
-		}
-		else if (!bHasBeenWaitedFor)	// if child is not running, but has not been waited for, still communicate a problem, but we shouldn't be blocked for long in this case.
-		{
-			UE_LOG(LogHAL, Log, TEXT("Closing a process handle of a process (pid=%d) that has not been wait()ed for - will wait() now to reap a zombie"),
-				GetProcessId()
-			);
-		}
-
-		Wait();	// will exit immediately if everything is Ok
-	}
-	else if (IsRunning())
-	{
-		// warn about leaking a thread ;/
-		UE_LOG(LogHAL, Log, TEXT("Process (pid=%d) is still running - we will reap it in a waiter thread, but the thread handle is going to be leaked."),
-				 GetProcessId()
-			);
-
-		FChildWaiterThread * WaiterRunnable = new FChildWaiterThread(GetProcessId());
-		// [RCL] 2015-03-11 @FIXME: do not leak
-		FRunnableThread * WaiterThread = FRunnableThread::Create(WaiterRunnable, *FString::Printf(TEXT("waitpid(%d)"), GetProcessId()), 32768 /* needs just a small stack */, TPri_BelowNormal);
+		UE::Mac::AddToProcessReaper(ProcessId);
 	}
 }
 
@@ -782,7 +798,7 @@ void FProcState::Wait()
 	for(;;)	// infinite loop in case we get EINTR and have to repeat
 	{
 		siginfo_t SignalInfo;
-		if (waitid(P_PID, GetProcessId(), &SignalInfo, WEXITED))
+		if (waitid(P_PID, GetProcessId(), &SignalInfo, WEXITED | WNOWAIT))
 		{
 			if (errno != EINTR)
 			{
@@ -795,6 +811,7 @@ void FProcState::Wait()
 		else
 		{
 			check(SignalInfo.si_pid == GetProcessId());
+			UE::Mac::AddToProcessReaper(ProcessId);
 
 			ReturnCode = (SignalInfo.si_code == CLD_EXITED) ? SignalInfo.si_status : -1;
 			bHasBeenWaitedFor = true;
diff --git a/Engine/Source/Runtime/Core/Public/Mac/MacPlatformProcess.h b/Engine/Source/Runtime/Core/Public/Mac/MacPlatformProcess.h
--- a/Engine/Source/Runtime/Core/Public/Mac/MacPlatformProcess.h
+++ b/Engine/Source/Runtime/Core/Public/Mac/MacPlatformProcess.h
@@ -12,18 +12,11 @@
 /** Wrapper around Unix pid_t. Should not be copyable as changes in the process state won't be properly propagated to all copies. */
 struct FProcState
 {
-	/** Default constructor. */
-	inline FProcState()
-		:	ProcessId(0)
-		,	bIsRunning(false)
-		,	bHasBeenWaitedFor(false)
-		,	ReturnCode(-1)
-		,	bFireAndForget(false)
-	{
-	}
+	FProcState(const FProcState&) = delete;
+	FProcState& operator=(const FProcState&) = delete;
 
 	/** Initialization constructor. */
-	explicit FProcState(pid_t InProcessId, bool bInFireAndForget);
+	explicit FProcState(pid_t InProcessId);
 
 	/** Destructor. */
 	~FProcState();
@@ -58,42 +51,17 @@ struct FProcState
 
 protected:  // the below is not a public API!
 
-	// FProcState should not be copyable since it breeds problems (e.g. one copy could have wait()ed, but another would not know it)
-
-	/** Copy constructor - should not be publicly accessible */
-	inline FProcState(const FProcState& Other)
-		:	ProcessId(Other.ProcessId)
-		,	bIsRunning(Other.bIsRunning)  // assume it is
-		,	bHasBeenWaitedFor(Other.bHasBeenWaitedFor)
-		,	ReturnCode(Other.ReturnCode)
-		,	bFireAndForget(Other.bFireAndForget)
-	{
-		checkf(false, TEXT("FProcState should not be copied"));
-	}
-
-	/** Assignment operator - should not be publicly accessible */
-	inline FProcState& operator=(const FProcState& Other)
-	{
-		checkf(false, TEXT("FProcState should not be copied"));
-		return *this;
-	}
-
-	// -------------------------
-
 	/** Process id */
 	pid_t	ProcessId;
 
 	/** Whether the process has finished or not (cached) */
-	bool	bIsRunning;
+	bool	bIsRunning = true;
 
 	/** Whether the process's return code has been collected */
-	bool	bHasBeenWaitedFor;
+	bool	bHasBeenWaitedFor = false;
 
 	/** Return code of the process (if negative, means that process did not finish gracefully but was killed/crashed*/
-	int32	ReturnCode;
-
-	/** Whether this child is fire-and-forget */
-	bool	bFireAndForget;
+	int32	ReturnCode = -1;
 };
 
 /** FProcHandle can be copied (and thus passed by value). */
