Disabling libGPUCounters  as they are causing bad GPU timing data to emerge.

diff --git a/Engine/Source/Runtime/Core/Private/Android/AndroidStats.cpp b/Engine/Source/Runtime/Core/Private/Android/AndroidStats.cpp
--- a/Engine/Source/Runtime/Core/Private/Android/AndroidStats.cpp
+++ b/Engine/Source/Runtime/Core/Private/Android/AndroidStats.cpp
@@ -11,7 +11,8 @@
 DEFINE_LOG_CATEGORY_STATIC(LogAndroidStats, Log, Log);
 
 #if !UE_BUILD_SHIPPING && PLATFORM_ANDROID_ARM64
-#	define HWCPIPE_SUPPORTED 1
+//#	define HWCPIPE_SUPPORTED 1 Disabling for now as we are seeing multiple missing frames since 36.00, starting since CL 42524925
+#	define HWCPIPE_SUPPORTED 0
 #else
 #	define HWCPIPE_SUPPORTED 0
 #endif
