Added better error messages when TLess, TGreater and TEqualTo are used without the corresponding operator being defined, and made it clear that they shouldn't be specialized.
Added an explicitly deleted operator< for FName to clarify why it's not defined.

diff --git a/Engine/Source/Runtime/Core/Public/Templates/EqualTo.h b/Engine/Source/Runtime/Core/Public/Templates/EqualTo.h
--- a/Engine/Source/Runtime/Core/Public/Templates/EqualTo.h
+++ b/Engine/Source/Runtime/Core/Public/Templates/EqualTo.h
@@ -2,21 +2,30 @@
 
 #pragma once
 
+#include "CoreTypes.h"
+#include "Concepts/SameAs.h"
 #include "UnrealTemplate.h"
 
 /**
- * Binary predicate class for performing equality comparisons.
- *
- * Uses operator== when available.
+ * Binary predicate class for performing equality comparisons.  Assumes == operator is defined for the template type.
  *
  * See: https://en.cppreference.com/w/cpp/utility/functional/equal_to
+ *
+ * TEqualTo should not be specialized - an appropriate operator== should be added.
  */
 template <typename T = void>
 struct TEqualTo
 {
-	constexpr auto operator()(const T& Lhs, const T& Rhs) const -> decltype(Lhs == Rhs)
+	UE_REWRITE constexpr bool operator()(const T& Lhs, const T& Rhs) const
 	{
-		return Lhs == Rhs;
+		if constexpr (requires { { Lhs == Rhs } -> UE::CSameAs<bool>; })
+		{
+			return Lhs == Rhs;
+		}
+		else
+		{
+			static_assert(sizeof(T) == 0, "Trying to use TEqualTo<T> where T doesn't have an appropriate operator== overload. Please add bool operator==(T, T), do not specialize TEqualTo.");
+		}
 	}
 };
 
@@ -24,8 +33,15 @@ template <>
 struct TEqualTo<void>
 {
 	template <typename T, typename U>
-	constexpr auto operator()(T&& Lhs, U&& Rhs) const -> decltype(Forward<T>(Lhs) == Forward<U>(Rhs))
+	UE_REWRITE constexpr bool operator()(T&& Lhs, U&& Rhs) const
 	{
-		return Forward<T>(Lhs) == Forward<U>(Rhs);
+		if constexpr (requires { { Forward<T>(Lhs) == Forward<U>(Rhs) } -> UE::CSameAs<bool>; })
+		{
+			return Forward<T>(Lhs) == Forward<U>(Rhs);
+		}
+		else
+		{
+			static_assert(sizeof(T) == 0, "Trying to use TEqualTo<void> with types without an appropriate operator== overload. Please add bool operator==(T, U), do not specialize TEqualTo.");
+		}
 	}
 };
diff --git a/Engine/Source/Runtime/Core/Public/Templates/Greater.h b/Engine/Source/Runtime/Core/Public/Templates/Greater.h
--- a/Engine/Source/Runtime/Core/Public/Templates/Greater.h
+++ b/Engine/Source/Runtime/Core/Public/Templates/Greater.h
@@ -3,19 +3,30 @@
 #pragma once
 
 #include "CoreTypes.h"
+#include "Concepts/SameAs.h"
 #include "UnrealTemplate.h"
 
 /**
  * Binary predicate class for sorting elements in reverse order.  Assumes < operator is defined for the template type.
+ * Forward declaration exists in ContainersFwd.h.
  *
  * See: http://en.cppreference.com/w/cpp/utility/functional/greater
+ *
+ * TGreater should not be specialized - an appropriate operator< should be added.
  */
 template <typename T = void>
 struct TGreater
 {
-	UE_FORCEINLINE_HINT bool operator()(const T& A, const T& B) const
+	UE_REWRITE constexpr bool operator()(const T& A, const T& B) const
 	{
-		return B < A;
+		if constexpr (requires { { B < A } -> UE::CSameAs<bool>; })
+		{
+			return B < A;
+		}
+		else
+		{
+			static_assert(sizeof(T) == 0, "Trying to use TGreater<T> where T doesn't have an appropriate operator< overload. Please add bool operator<(T, T), do not specialize TGreater.");
+		}
 	}
 };
 
@@ -23,8 +34,15 @@ template <>
 struct TGreater<void>
 {
 	template <typename T, typename U>
-	UE_FORCEINLINE_HINT bool operator()(T&& A, U&& B) const
+	UE_REWRITE constexpr bool operator()(T&& A, U&& B) const
 	{
-		return Forward<U>(B) < Forward<T>(A);
+		if constexpr (requires { { Forward<U>(B) < Forward<T>(A) } -> UE::CSameAs<bool>; })
+		{
+			return Forward<U>(B) < Forward<T>(A);
+		}
+		else
+		{
+			static_assert(sizeof(T) == 0, "Trying to use TGreater<void> with types without an appropriate operator< overload. Please add bool operator<(T, U), do not specialize TGreater.");
+		}
 	}
 };
diff --git a/Engine/Source/Runtime/Core/Public/Templates/Less.h b/Engine/Source/Runtime/Core/Public/Templates/Less.h
--- a/Engine/Source/Runtime/Core/Public/Templates/Less.h
+++ b/Engine/Source/Runtime/Core/Public/Templates/Less.h
@@ -3,20 +3,30 @@
 #pragma once
 
 #include "CoreTypes.h"
+#include "Concepts/SameAs.h"
 #include "UnrealTemplate.h"
 
 /**
  * Binary predicate class for sorting elements in order.  Assumes < operator is defined for the template type.
- * Forward declaration exists in ContainersFwd.h
+ * Forward declaration exists in ContainersFwd.h.
  *
  * See: http://en.cppreference.com/w/cpp/utility/functional/less
+ *
+ * TLess should not be specialized - an appropriate operator< should be added.
  */
 template <typename T /*= void */>
 struct TLess
 {
-	UE_FORCEINLINE_HINT bool operator()(const T& A, const T& B) const
+	UE_REWRITE constexpr bool operator()(const T& A, const T& B) const
 	{
-		return A < B;
+		if constexpr (requires { { A < B } -> UE::CSameAs<bool>; })
+		{
+			return A < B;
+		}
+		else
+		{
+			static_assert(sizeof(T) == 0, "Trying to use TLess<T> where T doesn't have an appropriate operator< overload. Please add bool operator<(T, T), do not specialize TLess.");
+		}
 	}
 };
 
@@ -24,8 +34,15 @@ template <>
 struct TLess<void>
 {
 	template <typename T, typename U>
-	UE_FORCEINLINE_HINT bool operator()(T&& A, U&& B) const
+	UE_REWRITE constexpr bool operator()(T&& A, U&& B) const
 	{
-		return Forward<T>(A) < Forward<U>(B);
+		if constexpr (requires { { Forward<T>(A) < Forward<U>(B) } -> UE::CSameAs<bool>; })
+		{
+			return Forward<T>(A) < Forward<U>(B);
+		}
+		else
+		{
+			static_assert(sizeof(T) == 0, "Trying to use TLess<void> with types without an appropriate operator< overload. Please add bool operator<(T, U), do not specialize TLess.");
+		}
 	}
 };
diff --git a/Engine/Source/Runtime/Core/Public/UObject/NameTypes.h b/Engine/Source/Runtime/Core/Public/UObject/NameTypes.h
--- a/Engine/Source/Runtime/Core/Public/UObject/NameTypes.h
+++ b/Engine/Source/Runtime/Core/Public/UObject/NameTypes.h
@@ -1751,6 +1751,10 @@ public:
 	CORE_API bool UEOpEquals(const FLazyName& Rhs) const;
 };
 
+// Ordering operators for FName are deliberately disabled - users should use FNameFastLess or FNameLexicalLess to select the desired behavior.
+bool operator<(FName, FName) = delete;
+bool operator>(FName, FName) = delete;
+
 /**
  * Serialization util that optimizes WITH_CASE_PRESERVING_NAME-loading by reducing comparison id lookups
  *
