UnifiedError change plumbing for FormatString

The previous version caused a copy to get instantiated and initialized in every TU.
Routing it through an inline function is a bit of a hack but avoids API changes.

diff --git a/Engine/Source/Runtime/Core/Public/Experimental/UnifiedError/UnifiedError.h b/Engine/Source/Runtime/Core/Public/Experimental/UnifiedError/UnifiedError.h
--- a/Engine/Source/Runtime/Core/Public/Experimental/UnifiedError/UnifiedError.h
+++ b/Engine/Source/Runtime/Core/Public/Experimental/UnifiedError/UnifiedError.h
@@ -676,16 +676,18 @@ CORE_API constexpr uint32 GetTypeIdHash(const TCHAR* TypeName)
 #define UE_DECLARE_ERROR_INTERNAL(DeclareApi, ErrorName, ErrorCode, ModuleName, FormatString) \
 	namespace UE::UnifiedError { namespace ModuleName { namespace ErrorName {\
 			DeclareApi FAnsiStringView GetStaticErrorName(); \
-			static const int32 DeclaredErrorCode = ErrorCode; \
-			static const FText DeclaredFormatString = FormatString; \
+			static constexpr int32 DeclaredErrorCode = ErrorCode; \
+			extern DeclareApi const FText DeclaredFormatString; \
 			DeclareApi int32 GetErrorCodeId(); \
 			DeclareApi TRefCountPtr<const FStaticErrorDetails> GetStaticErrorDetails(); \
 			DeclareApi const FError& GetStaticError(); \
+			static inline FText GetFormatString() { return FormatString; } \
 	}}}
 
 #define UE_DEFINE_ERROR(ErrorName, ModuleName) \
 	namespace UE::UnifiedError { namespace ModuleName { namespace ErrorName {\
 			const int32 ErrorCodeId = UE::UnifiedError::ErrorRegistry::RegisterErrorCode(GetStaticErrorName(), UE::UnifiedError::ModuleName::GetStaticModuleId(), DeclaredErrorCode);\
+			const FText DeclaredFormatString = GetFormatString(); \
 			int32 GetErrorCodeId() \
 			{\
 				return ErrorCodeId;\
