Remove stale comment

diff --git a/Engine/Source/Runtime/Core/Public/Misc/CoreDelegates.h b/Engine/Source/Runtime/Core/Public/Misc/CoreDelegates.h
--- a/Engine/Source/Runtime/Core/Public/Misc/CoreDelegates.h
+++ b/Engine/Source/Runtime/Core/Public/Misc/CoreDelegates.h
@@ -93,7 +93,6 @@ struct FMountPaksExArgs
 	const TCHAR* PakFilePath = nullptr;
 	int32 Order = INDEX_NONE;
 	FPakMountOptions MountOptions;
-	// TODO: @Daniel.Lamb - this should use FError
 	TValueOrError<FMountResult, UE::UnifiedError::FError> Result = MakeValue();
 };
 }
