Added [[nodiscard]] to TLess, TGreater and TEqualTo.

diff --git a/Engine/Source/Runtime/Core/Public/Templates/EqualTo.h b/Engine/Source/Runtime/Core/Public/Templates/EqualTo.h
--- a/Engine/Source/Runtime/Core/Public/Templates/EqualTo.h
+++ b/Engine/Source/Runtime/Core/Public/Templates/EqualTo.h
@@ -16,7 +16,7 @@
 template <typename T = void>
 struct TEqualTo
 {
-	UE_REWRITE constexpr bool operator()(const T& Lhs, const T& Rhs) const
+	[[nodiscard]] UE_REWRITE constexpr bool operator()(const T& Lhs, const T& Rhs) const
 	{
 		if constexpr (requires { { Lhs == Rhs } -> UE::CSameAs<bool>; })
 		{
@@ -33,7 +33,7 @@ template <>
 struct TEqualTo<void>
 {
 	template <typename T, typename U>
-	UE_REWRITE constexpr bool operator()(T&& Lhs, U&& Rhs) const
+	[[nodiscard]] UE_REWRITE constexpr bool operator()(T&& Lhs, U&& Rhs) const
 	{
 		if constexpr (requires { { Forward<T>(Lhs) == Forward<U>(Rhs) } -> UE::CSameAs<bool>; })
 		{
diff --git a/Engine/Source/Runtime/Core/Public/Templates/Greater.h b/Engine/Source/Runtime/Core/Public/Templates/Greater.h
--- a/Engine/Source/Runtime/Core/Public/Templates/Greater.h
+++ b/Engine/Source/Runtime/Core/Public/Templates/Greater.h
@@ -17,7 +17,7 @@
 template <typename T = void>
 struct TGreater
 {
-	UE_REWRITE constexpr bool operator()(const T& A, const T& B) const
+	[[nodiscard]] UE_REWRITE constexpr bool operator()(const T& A, const T& B) const
 	{
 		if constexpr (requires { { B < A } -> UE::CSameAs<bool>; })
 		{
@@ -34,7 +34,7 @@ template <>
 struct TGreater<void>
 {
 	template <typename T, typename U>
-	UE_REWRITE constexpr bool operator()(T&& A, U&& B) const
+	[[nodiscard]] UE_REWRITE constexpr bool operator()(T&& A, U&& B) const
 	{
 		if constexpr (requires { { Forward<U>(B) < Forward<T>(A) } -> UE::CSameAs<bool>; })
 		{
diff --git a/Engine/Source/Runtime/Core/Public/Templates/Less.h b/Engine/Source/Runtime/Core/Public/Templates/Less.h
--- a/Engine/Source/Runtime/Core/Public/Templates/Less.h
+++ b/Engine/Source/Runtime/Core/Public/Templates/Less.h
@@ -17,7 +17,7 @@
 template <typename T /*= void */>
 struct TLess
 {
-	UE_REWRITE constexpr bool operator()(const T& A, const T& B) const
+	[[nodiscard]] UE_REWRITE constexpr bool operator()(const T& A, const T& B) const
 	{
 		if constexpr (requires { { A < B } -> UE::CSameAs<bool>; })
 		{
@@ -34,7 +34,7 @@ template <>
 struct TLess<void>
 {
 	template <typename T, typename U>
-	UE_REWRITE constexpr bool operator()(T&& A, U&& B) const
+	[[nodiscard]] UE_REWRITE constexpr bool operator()(T&& A, U&& B) const
 	{
 		if constexpr (requires { { Forward<T>(A) < Forward<U>(B) } -> UE::CSameAs<bool>; })
 		{
