Added more detail as to why FNames have no ordering operators and what to use instead, and added the missing ones.

diff --git a/Engine/Source/Runtime/Core/Public/UObject/NameTypes.h b/Engine/Source/Runtime/Core/Public/UObject/NameTypes.h
--- a/Engine/Source/Runtime/Core/Public/UObject/NameTypes.h
+++ b/Engine/Source/Runtime/Core/Public/UObject/NameTypes.h
@@ -1751,9 +1751,20 @@ public:
 	CORE_API bool UEOpEquals(const FLazyName& Rhs) const;
 };
 
-// Ordering operators for FName are deliberately disabled - users should use FNameFastLess or FNameLexicalLess to select the desired behavior.
+// Ordering operators for FName are deliberately disabled as there is no correct default.
+// Users should use the following to select the desired behavior:
+//
+// Lexicographical comparison:
+// - Expression: NameA.LexicalLess(NameB)
+// - Predicate: FNameLexicalLess
+//
+// Fast comparison:
+// - Expression: NameA.FastLess(NameB)
+// - Predicate: FNameFastLess
 bool operator<(FName, FName) = delete;
 bool operator>(FName, FName) = delete;
+bool operator<=(FName, FName) = delete;
+bool operator>=(FName, FName) = delete;
 
 /**
  * Serialization util that optimizes WITH_CASE_PRESERVING_NAME-loading by reducing comparison id lookups
