Fixed some copied+pasted constraints in UEOps.h.

diff --git a/Engine/Source/Runtime/Core/Public/Misc/UEOps.h b/Engine/Source/Runtime/Core/Public/Misc/UEOps.h
--- a/Engine/Source/Runtime/Core/Public/Misc/UEOps.h
+++ b/Engine/Source/Runtime/Core/Public/Misc/UEOps.h
@@ -141,13 +141,13 @@ template <typename LhsType, typename RhsType>
 }
 
 template <typename LhsType, typename RhsType>
-[[nodiscard]] UE_REWRITE constexpr auto operator>(const LhsType& Lhs, const RhsType& Rhs) -> decltype(UE::Core::Private::UEOpEquals(Rhs, Lhs))
+[[nodiscard]] UE_REWRITE constexpr auto operator>(const LhsType& Lhs, const RhsType& Rhs) -> decltype(UE::Core::Private::UEOpLessThan(Rhs, Lhs))
 {
 	return Rhs.UEOpLessThan(Lhs);
 }
 
 template <typename LhsType, typename RhsType>
-[[nodiscard]] UE_REWRITE constexpr auto operator>(const LhsType& Lhs, const RhsType& Rhs) -> decltype(UE::Core::Private::DisableUEOpEquals(Rhs, Lhs), UE::Core::Private::UEOpEquals(Lhs, Rhs))
+[[nodiscard]] UE_REWRITE constexpr auto operator>(const LhsType& Lhs, const RhsType& Rhs) -> decltype(UE::Core::Private::DisableUEOpLessThan(Rhs, Lhs), UE::Core::Private::UEOpLessThan(Lhs, Rhs))
 {
 	return Lhs.UEOpLessThan(Rhs);
 }
@@ -165,13 +165,13 @@ template <typename LhsType, typename RhsType>
 }
 
 template <typename LhsType, typename RhsType>
-[[nodiscard]] UE_REWRITE constexpr auto operator<=(const LhsType& Lhs, const RhsType& Rhs) -> decltype(UE::Core::Private::UEOpEquals(Rhs, Lhs))
+[[nodiscard]] UE_REWRITE constexpr auto operator<=(const LhsType& Lhs, const RhsType& Rhs) -> decltype(UE::Core::Private::UEOpLessThan(Rhs, Lhs))
 {
 	return !Rhs.UEOpLessThan(Lhs);
 }
 
 template <typename LhsType, typename RhsType>
-[[nodiscard]] UE_REWRITE constexpr auto operator<=(const LhsType& Lhs, const RhsType& Rhs) -> decltype(UE::Core::Private::DisableUEOpEquals(Rhs, Lhs), UE::Core::Private::UEOpEquals(Lhs, Rhs))
+[[nodiscard]] UE_REWRITE constexpr auto operator<=(const LhsType& Lhs, const RhsType& Rhs) -> decltype(UE::Core::Private::DisableUEOpLessThan(Rhs, Lhs), UE::Core::Private::UEOpLessThan(Lhs, Rhs))
 {
 	return !Lhs.UEOpLessThan(Rhs);
 }
