Use resident size of file mappings as ProgramSize on Android

diff --git a/Engine/Source/Runtime/Core/Private/Android/AndroidPlatformMemory.cpp b/Engine/Source/Runtime/Core/Private/Android/AndroidPlatformMemory.cpp
--- a/Engine/Source/Runtime/Core/Private/Android/AndroidPlatformMemory.cpp
+++ b/Engine/Source/Runtime/Core/Private/Android/AndroidPlatformMemory.cpp
@@ -12,16 +12,20 @@
 #include <jni.h>
 #include <sys/sysinfo.h>
 #include <sys/mman.h>
+#if !UE_BUILD_SHIPPING
+#include <linux/prctl.h> 
+#include <sys/prctl.h>
+#endif
 
 #include "Android/AndroidPlatformCrashContext.h"
 #include "Async/TaskGraphInterfaces.h"
 #include "Async/Async.h"
+#include "Tasks/Task.h"
 
 #define JNI_CURRENT_VERSION JNI_VERSION_1_6
 extern JavaVM* GJavaVM;
 extern jobject GGameActivityThis;
 
-
 static int32 GAndroidAddSwapToTotalPhysical = 1;
 static FAutoConsoleVariableRef CVarAddSwapToTotalPhysical(
 	TEXT("android.AddSwapToTotalPhysical"),
@@ -57,6 +61,21 @@ static int64 GetNativeHeapAllocatedSize()
 /** Controls growth of pools - see PooledVirtualMemoryAllocator.cpp */
 extern float GVMAPoolScale;
 
+#if !UE_BUILD_SHIPPING
+#ifndef PR_SET_VMA
+#define PR_SET_VMA 0x53564d41
+#endif
+#ifndef PR_SET_VMA_ANON_NAME
+#define PR_SET_VMA_ANON_NAME 0
+#endif
+static inline void SetMmapName(const void* Ptr, const size_t Size, const char* Name)
+{
+	prctl(PR_SET_VMA, PR_SET_VMA_ANON_NAME, Ptr, Size, Name);
+}
+#else
+static inline void SetMmapName(const void*, const size_t, const char*) {}
+#endif
+
 void FAndroidPlatformMemory::Init()
 {
 	// Only allow this method to be called once
@@ -269,12 +288,17 @@ FPlatformMemoryStats FAndroidPlatformMemory::GetStats()
 				MemoryStats.VMRss = MemoryStats.UsedPhysical;
 				++FieldsSetSuccessfully;
 			}
+			else if (strstr(Line, "RssFile:") == Line)
+			{
+				MemoryStats.RssFile = AndroidPlatformMemory::GetBytesFromStatusLine(Line);
+				++FieldsSetSuccessfully;
+			}
 			else if (strstr(Line, "VmSwap:") == Line)
 			{
 				MemoryStats.VMSwap = AndroidPlatformMemory::GetBytesFromStatusLine(Line);
 				++FieldsSetSuccessfully;
 			}
-		} while (FieldsSetSuccessfully < 5);
+		} while (FieldsSetSuccessfully < 6);
 
 		fclose(ProcMemStats);
 	}
@@ -439,10 +463,6 @@ FMalloc* FAndroidPlatformMemory::BaseAllocator()
 	}
 
 	FPlatformMemoryStats Stats = FAndroidPlatformMemory::GetStats();
-#if ENABLE_LOW_LEVEL_MEM_TRACKER
-	// make sure LLM is using UsedPhysical for program size, instead of Available-Free
-	FLowLevelMemTracker::Get().SetProgramSize(Stats.UsedPhysical);
-#endif
 	FPlatformMemory::ProgramSize = Stats.UsedPhysical;
 
 #if FORCE_ANSI_ALLOCATOR
@@ -511,6 +531,7 @@ void* FAndroidPlatformMemory::BinnedAllocFromOS(SIZE_T Size)
 
 			SIZE_T SizeToAlloc = FMath::Max<SIZE_T>(MinAllocSize, Align(Size, MemoryConstants.PageSize) + MemoryConstants.BinnedPageSize);
 			uint8* UnalignedBase = (uint8*)mmap(nullptr, SizeToAlloc, PROT_READ | PROT_WRITE, MAP_PRIVATE | MAP_ANON, -1, 0);
+			SetMmapName(UnalignedBase, SizeToAlloc, "UE");
 			End = UnalignedBase + SizeToAlloc;
 			Base = Align(UnalignedBase, MemoryConstants.BinnedPageSize);
 
@@ -548,6 +569,7 @@ void* FAndroidPlatformMemory::BinnedAllocFromOS(SIZE_T Size)
 	else
 	{
 		Ptr = mmap(nullptr, Size, PROT_READ | PROT_WRITE, MAP_PRIVATE | MAP_ANON, -1, 0);
+		SetMmapName(Ptr, Size, "UE");
 	}
 
 	LLM_IF_ENABLED(FLowLevelMemTracker::Get().OnLowLevelAlloc(ELLMTracker::Platform, Ptr, Size));
@@ -592,6 +614,7 @@ FAndroidPlatformMemory::FPlatformVirtualMemoryBlock FAndroidPlatformMemory::FPla
 		FPlatformMemory::OnOutOfMemory(Result.GetActualSize(), InAlignment);
 	}
 	check(Result.Ptr && IsAligned(Result.Ptr, Alignment));
+	SetMmapName(Result.Ptr, Result.GetActualSize(), "UE");
 	return Result;
 }
 
@@ -644,6 +667,7 @@ int64 LLMMallocTotal = 0;
 void* LLMAlloc(size_t Size)
 {
 	void* Ptr = mmap(nullptr, Size, PROT_READ | PROT_WRITE, MAP_PRIVATE | MAP_ANON, -1, 0);
+	SetMmapName(Ptr, Size, "UE");
 
 	LLMMallocTotal += Size;
 
@@ -675,6 +699,137 @@ bool FAndroidPlatformMemory::GetLLMAllocFunctions(void*(*&OutAllocFunction)(size
 #endif
 }
 
+#if !UE_BUILD_SHIPPING
+
+static std::atomic<bool> GAndroidMemoryProbeTickRun = {};
+static FTSTicker::FDelegateHandle GAndroidMemoryProbeTicker = {};
+static UE::Tasks::TTask<void> GAndroidMemoryProbeTask = {};
+
+static float GAndroidMemoryProbeInterval = 10.0f;
+static FAutoConsoleVariableRef CVarAndroidMemoryProbeInterval(
+	TEXT("android.MemoryProbeInterval"),
+	GAndroidMemoryProbeInterval,
+	TEXT("/proc/self/smaps polling interval in seconds, default 10 seconds."),
+	ECVF_Default
+);
+
+void FAndroidPlatformMemory::AndroidMemoryProbeStart()
+{
+	if (GAndroidMemoryProbeInterval <= 0.0f)
+	{
+		return;
+	}
+
+	if (GAndroidMemoryProbeTicker.IsValid())
+	{
+		return;
+	}
+
+	GAndroidMemoryProbeTickRun = true;
+
+	GAndroidMemoryProbeTicker = FTSTicker::GetCoreTicker().AddTicker(FTickerDelegate::CreateLambda([](float)
+	{
+		if (!GAndroidMemoryProbeTask.IsValid() || GAndroidMemoryProbeTask.IsCompleted())
+		{
+			GAndroidMemoryProbeTask = UE::Tasks::Launch(
+				TEXT("AndroidMemoryProbeTick"),
+				AndroidMemoryProbeTick,
+				LowLevelTasks::ETaskPriority::BackgroundLow
+			);
+		}
+		return true;
+	}), GAndroidMemoryProbeInterval);
+}
+
+void FAndroidPlatformMemory::AndroidMemoryProbeStop()
+{
+	if (!GAndroidMemoryProbeTicker.IsValid())
+	{
+		return;
+	}
+
+	GAndroidMemoryProbeTickRun = false;
+
+	if (GAndroidMemoryProbeTask.IsValid() && !GAndroidMemoryProbeTask.IsCompleted() && GAndroidMemoryProbeTask.IsAwaitable())
+	{
+		GAndroidMemoryProbeTask.Wait();
+	}
+
+	GAndroidMemoryProbeTask = {};
+
+	FTSTicker::GetCoreTicker().RemoveTicker(GAndroidMemoryProbeTicker);
+	GAndroidMemoryProbeTicker = {};
+}
+
+void FAndroidPlatformMemory::AndroidMemoryProbeTick()
+{
+	if (!GAndroidMemoryProbeTickRun)
+	{
+		return;
+	}
+
+	TRACE_CPUPROFILER_EVENT_SCOPE(AndroidMemoryProbeTick2);
+
+	FILE* Smaps = fopen("/proc/self/smaps", "r");
+	if (!Smaps)
+	{
+		return;
+	}
+
+	char LineBuffer[1024] = { 0 };
+
+	bool bIsProgramSize = false;
+	uint64 ProgramSize = 0;
+
+	while (GAndroidMemoryProbeTickRun)
+	{
+		char* LinePtr = fgets(LineBuffer, UE_ARRAY_COUNT(LineBuffer), Smaps);
+		if (!LinePtr)
+		{
+			break; // eof or an error
+		}
+		FAnsiStringView Line(LinePtr);
+
+		if (Line.Contains("-"))
+		{
+			// parsing region entry
+			// 7fffcccb4000-7fffcccd8000 rw-p 00000000 00:00 0                          [stack]
+			bIsProgramSize =
+				FAnsiStringView(Line).EndsWith("/libUnreal.so\n") ||     // our own .so in case if it's unpacked
+				FAnsiStringView(Line).EndsWith("/base.apk\n") ||         // our own .so in case if it's stored with-in .apk
+				FAnsiStringView(Line).EndsWith("[anon:.bss]\n") ||
+				FAnsiStringView(Line).EndsWith("[anon:linker_alloc]\n");
+		}
+		else if (bIsProgramSize && Line.StartsWith("Rss:"))
+		{
+			ProgramSize += AndroidPlatformMemory::GetBytesFromStatusLine(LinePtr);
+		}
+	}
+
+	fclose(Smaps);
+
+#if ENABLE_LOW_LEVEL_MEM_TRACKER
+	FLowLevelMemTracker::Get().SetProgramSize(ProgramSize);
+#endif
+	FPlatformMemory::ProgramSize = ProgramSize;
+}
+
+#else
+
+void FAndroidPlatformMemory::AndroidMemoryProbeStart()
+{
+}
+
+void FAndroidPlatformMemory::AndroidMemoryProbeStop()
+{
+}
+
+void FAndroidPlatformMemory::AndroidMemoryProbeTick()
+{
+}
+
+#endif
+
 #if USING_HW_ADDRESS_SANITISER && PLATFORM_USED_NDK_VERSION_INTEGER >= 26
 
 // Using libc++_static on NDK r26b with HWAsan enabled is not officially supported
diff --git a/Engine/Source/Runtime/Core/Private/Android/AndroidPlatformMisc.cpp b/Engine/Source/Runtime/Core/Private/Android/AndroidPlatformMisc.cpp
--- a/Engine/Source/Runtime/Core/Private/Android/AndroidPlatformMisc.cpp
+++ b/Engine/Source/Runtime/Core/Private/Android/AndroidPlatformMisc.cpp
@@ -625,12 +625,16 @@ void FAndroidMisc::PlatformInit()
 	InitCpuThermalSensor();
 
 	AndroidScudoMemoryTrace::Init();
+
+	FAndroidPlatformMemory::AndroidMemoryProbeStart();
 }
 
 extern void AndroidThunkCpp_DismissSplashScreen();
 
 void FAndroidMisc::PlatformTearDown()
 {
+	FAndroidPlatformMemory::AndroidMemoryProbeStop();
+
 #if (STATS || ENABLE_STATNAMEDEVENTS)
 	StopTraceMarkers();
 #endif
diff --git a/Engine/Source/Runtime/Core/Private/HAL/LowLevelMemTracker.cpp b/Engine/Source/Runtime/Core/Private/HAL/LowLevelMemTracker.cpp
--- a/Engine/Source/Runtime/Core/Private/HAL/LowLevelMemTracker.cpp
+++ b/Engine/Source/Runtime/Core/Private/HAL/LowLevelMemTracker.cpp
@@ -2067,6 +2067,8 @@ void FLowLevelMemTracker::TickInternal()
 	// On some Android devices total used mem (VmRss + VmSwap) does not include GPU memory allocations
 	// and so it's going to be lower than TrackedTotal
 	const int64 PlatformProcessMemory = FMath::Max((int64)PlatformStats.UsedPhysical, TrackedTotal);
+
+	PlatformTracker.SetTagAmountInUpdate(FindOrAddTagData(ELLMTag::PlatformMMIO), (int64)PlatformStats.RssFile, false);
 #else
 	const int64 PlatformProcessMemory = static_cast<int64>(PlatformStats.UsedPhysical);
 #endif
@@ -2269,6 +2271,8 @@ void FLowLevelMemTracker::PublishDataPerFrame(const TCHAR* LogName)
 
 void FLowLevelMemTracker::InitialiseProgramSize()
 {
+	// On Android we report ProgramSize via AndroidMemoryProbeTick
+#if !PLATFORM_ANDROID
 	if (!ProgramSize)
 	{
 		FPlatformMemoryStats Stats = FPlatformMemory::GetStats();
@@ -2277,6 +2281,7 @@ void FLowLevelMemTracker::InitialiseProgramSize()
 		GetTracker(ELLMTracker::Platform)->TrackMemory(ELLMTag::ProgramSizePlatform, ProgramSize, ELLMAllocType::System);
 		GetTracker(ELLMTracker::Default)->TrackMemory(ELLMTag::ProgramSize, ProgramSize, ELLMAllocType::System);
 	}
+#endif
 }
 
 void FLowLevelMemTracker::SetProgramSize(uint64 InProgramSize)
diff --git a/Engine/Source/Runtime/Core/Public/Android/AndroidPlatformMemory.h b/Engine/Source/Runtime/Core/Public/Android/AndroidPlatformMemory.h
--- a/Engine/Source/Runtime/Core/Public/Android/AndroidPlatformMemory.h
+++ b/Engine/Source/Runtime/Core/Public/Android/AndroidPlatformMemory.h
@@ -16,6 +16,7 @@ struct FPlatformMemoryStats : public FGenericPlatformMemoryStats
 	FGenericPlatformMemoryStats::EMemoryPressureStatus GetMemoryPressureStatus() const;
 	uint64 VMSwap = 0;
 	uint64 VMRss = 0;
+	uint64 RssFile = 0;
 };
 
 /**
@@ -100,6 +101,10 @@ struct FAndroidPlatformMemory : public FGenericPlatformMemory
 	// called when OS (via JNI) reports memory trouble, triggers MemoryWarningHandler callback on game thread if set.
 	enum class EOSMemoryStatusCategory { OSTrim };
 	static CORE_API void UpdateOSMemoryStatus(EOSMemoryStatusCategory OSMemoryStatusCategory, int value);
+
+	static void AndroidMemoryProbeStart();
+	static void AndroidMemoryProbeStop();
+	static void AndroidMemoryProbeTick();
 };
 
 typedef FAndroidPlatformMemory FPlatformMemory;
