Moved LogLargeAllocation from MallocBinned3 to MallocBinnedCommon so that MallocBinned2 can also log large allocations. Opted out by default

diff --git a/Engine/Source/Runtime/Core/Private/HAL/MallocBinned2.cpp b/Engine/Source/Runtime/Core/Private/HAL/MallocBinned2.cpp
--- a/Engine/Source/Runtime/Core/Private/HAL/MallocBinned2.cpp
+++ b/Engine/Source/Runtime/Core/Private/HAL/MallocBinned2.cpp
@@ -630,6 +630,7 @@ void* FMallocBinned2::MallocExternalLarge(SIZE_T Size, uint32 Alignment)
 		UE::TUniqueLock Lock(ExternalAllocMutex);
 
 		// Use OS for non-pooled allocations.
+		LogLargeAllocation(AlignedSize);
 		Result = CachedOSPageAllocator.Allocate(AlignedSize, 0, &ExternalAllocMutex);
 		if (!Result)
 		{
diff --git a/Engine/Source/Runtime/Core/Private/HAL/MallocBinned3.cpp b/Engine/Source/Runtime/Core/Private/HAL/MallocBinned3.cpp
--- a/Engine/Source/Runtime/Core/Private/HAL/MallocBinned3.cpp
+++ b/Engine/Source/Runtime/Core/Private/HAL/MallocBinned3.cpp
@@ -57,78 +57,6 @@ uint16 FMallocBinned3::SmallBinSizesShifted[UE_MB3_SMALL_POOL_COUNT + 1] = { 0 }
 	uint8* FMallocBinned3::PoolBaseVMPtr[UE_MB3_SMALL_POOL_COUNT] = { nullptr };
 #endif
 
-#ifndef UE_MB3_LOG_LARGE_ALLOCATION
-#define UE_MB3_LOG_LARGE_ALLOCATION 0
-#endif
-
-#if UE_MB3_LOG_LARGE_ALLOCATION
-
-#include "HAL/PlatformStackWalk.h"
-
-	static int32 GLogLargeAllocationSizeInMB = 0;
-	static FAutoConsoleVariableRef CVarLogLargeAllocationSizeInMB(
-		TEXT("Memory.LogLargeAllocationInMB"),
-		GLogLargeAllocationSizeInMB,
-		TEXT("Log large memory allocations and their callstacks if the allocation size exceeds the specified MB threshold (default is 0, meaning no allocations are logged)."));
-
-	static void LogLargeAllocation(SIZE_T Size)
-	{
-		if (GLogLargeAllocationSizeInMB <= 0)
-		{
-			return;
-		}
-
-		int32 AllocationSizeInMB = int32(Size >> 20);
-
-		if (AllocationSizeInMB < GLogLargeAllocationSizeInMB)
-		{
-			return;
-		}
-
-		static bool bInFunction = false;
-		if (bInFunction)
-		{
-			return;
-		}
-		
-		TGuardValue<bool> GuardInFunction(bInFunction, true);
-
-		// Grab callstack
-		static constexpr int32 MaxCallStackDepth = 16;
-		static constexpr int32 CallStackEntriesToSkipCount = 2;
-
-		uint64 FullCallStack[MaxCallStackDepth + CallStackEntriesToSkipCount] = { 0 };
-		FPlatformStackWalk::CaptureStackBackTrace(&FullCallStack[0], MaxCallStackDepth + CallStackEntriesToSkipCount);
-
-		// Skip first n entries as they are inside the allocator.
-		uint64* CallStack = &FullCallStack[CallStackEntriesToSkipCount];
-
-		UE_LOG(LogMemory, Log, TEXT("Allocating %d MB - 0x%08llx 0x%08llx 0x%08llx 0x%08llx 0x%08llx 0x%08llx 0x%08llx 0x%08llx 0x%08llx 0x%08llx 0x%08llx 0x%08llx 0x%08llx 0x%08llx 0x%08llx 0x%08llx"),
-			   AllocationSizeInMB,
-			   CallStack[0],
-			   CallStack[1],
-			   CallStack[2],
-			   CallStack[3],
-			   CallStack[4],
-			   CallStack[5],
-			   CallStack[6],
-			   CallStack[7],
-			   CallStack[8],
-			   CallStack[9],
-			   CallStack[10],
-			   CallStack[11],
-			   CallStack[12],
-			   CallStack[13],
-			   CallStack[14],
-			   CallStack[15]);
-	}
-
-#else // UE_MB3_LOG_LARGE_ALLOCATION
-
-	static void LogLargeAllocation(SIZE_T Size) {}
-
-#endif // UE_MB3_LOG_LARGE_ALLOCATION
-
 FMallocBinned3* FMallocBinned3::MallocBinned3 = nullptr;
 // Mapping of sizes to small table indices
 uint8 FMallocBinned3::MemSizeToPoolIndex[1 + (UE_MB3_MAX_SMALL_POOL_SIZE >> UE_MBC_MIN_SMALL_POOL_ALIGNMENT_SHIFT)] = { 0 };
diff --git a/Engine/Source/Runtime/Core/Private/HAL/MallocBinnedCommon.cpp b/Engine/Source/Runtime/Core/Private/HAL/MallocBinnedCommon.cpp
--- a/Engine/Source/Runtime/Core/Private/HAL/MallocBinnedCommon.cpp
+++ b/Engine/Source/Runtime/Core/Private/HAL/MallocBinnedCommon.cpp
@@ -11,6 +11,10 @@
 #include "Templates/UnrealTemplate.h"
 #include "ProfilingDebugging/CsvProfiler.h"
 
+#if UE_MBC_LOG_LARGE_ALLOCATION
+	#include "HAL/PlatformStackWalk.h"
+#endif
+
 PRAGMA_DISABLE_UNSAFE_TYPECAST_WARNINGS
 
 CSV_DEFINE_CATEGORY(MallocBinned, false);
@@ -618,4 +622,65 @@ void FMallocBinnedCommonBase::GetAllocatorStatsInternal(FGenericMemoryStats& Out
 }
 #endif	//~UE_MBC_ALLOCATOR_STATS
 
+#if UE_MBC_LOG_LARGE_ALLOCATION
+
+static int32 GLogLargeAllocationSizeInMB = 0;
+static FAutoConsoleVariableRef CVarLogLargeAllocationSizeInMB(
+	TEXT("MallocBinned.LogLargeAllocationInMB"),
+	GLogLargeAllocationSizeInMB,
+	TEXT("Log large memory allocations and their callstacks if the allocation size exceeds the specified MB threshold (default is 0, meaning no allocations are logged)."));
+
+void FMallocBinnedCommonBase::LogLargeAllocation(SIZE_T Size) const
+{
+	if (GLogLargeAllocationSizeInMB <= 0)
+	{
+		return;
+	}
+
+	const int32 AllocationSizeInMB = int32((Size / 1024) / 1024);
+
+	if (AllocationSizeInMB < GLogLargeAllocationSizeInMB)
+	{
+		return;
+	}
+
+	static bool bInFunction = false;
+	if (bInFunction)
+	{
+		return;
+	}
+
+	TGuardValue<bool> GuardInFunction(bInFunction, true);
+
+	// Grab callstack
+	static constexpr int32 MaxCallStackDepth = 16;
+	static constexpr int32 CallStackEntriesToSkipCount = 2;
+
+	uint64 FullCallStack[MaxCallStackDepth + CallStackEntriesToSkipCount] = { 0 };
+	FPlatformStackWalk::CaptureStackBackTrace(&FullCallStack[0], MaxCallStackDepth + CallStackEntriesToSkipCount);
+
+	// Skip first n entries as they are inside the allocator.
+	uint64* CallStack = &FullCallStack[CallStackEntriesToSkipCount];
+
+	UE_LOG(LogMemory, Log, TEXT("Allocating %d MB - 0x%08llx 0x%08llx 0x%08llx 0x%08llx 0x%08llx 0x%08llx 0x%08llx 0x%08llx 0x%08llx 0x%08llx 0x%08llx 0x%08llx 0x%08llx 0x%08llx 0x%08llx 0x%08llx"),
+		   AllocationSizeInMB,
+		   CallStack[0],
+		   CallStack[1],
+		   CallStack[2],
+		   CallStack[3],
+		   CallStack[4],
+		   CallStack[5],
+		   CallStack[6],
+		   CallStack[7],
+		   CallStack[8],
+		   CallStack[9],
+		   CallStack[10],
+		   CallStack[11],
+		   CallStack[12],
+		   CallStack[13],
+		   CallStack[14],
+		   CallStack[15]);
+}
+#endif
+
 PRAGMA_RESTORE_UNSAFE_TYPECAST_WARNINGS
diff --git a/Engine/Source/Runtime/Core/Public/HAL/MallocBinnedCommon.h b/Engine/Source/Runtime/Core/Public/HAL/MallocBinnedCommon.h
--- a/Engine/Source/Runtime/Core/Public/HAL/MallocBinnedCommon.h
+++ b/Engine/Source/Runtime/Core/Public/HAL/MallocBinnedCommon.h
@@ -93,6 +93,10 @@
 #	define UE_MBC_UPDATE_STATS(x)
 #endif
 
+#ifndef UE_MBC_LOG_LARGE_ALLOCATION
+#	define UE_MBC_LOG_LARGE_ALLOCATION 0
+#endif
+
 CSV_DECLARE_CATEGORY_EXTERN(MallocBinned);
 
 extern CORE_API float GMallocBinnedFlushThreadCacheMaxWaitTime;
@@ -412,6 +416,12 @@ protected:
 		// this is expected not to return
 		FPlatformMemory::OnOutOfMemory(Size, Alignment);
 	}
+
+#if UE_MBC_LOG_LARGE_ALLOCATION
+		void LogLargeAllocation(SIZE_T Size) const;
+#else
+		void LogLargeAllocation(SIZE_T Size) const {};
+#endif
 };
 
 template <class AllocType, int NumSmallPools, int MaxSmallPoolSize>
