Logging: Reduced memory use by FOutputDeviceFile

- Fixed a leak of an extra FAsyncWriter when creation of the writer contains logging.
- Reduced stack size from the default to 32 KiB. Saves 992 KiB in the default configuration.
- Moved InitialBufferSize out of the header to make it easier to tweak if further memory savings are desired.

diff --git a/Engine/Source/Runtime/Core/Private/Misc/OutputDeviceFile.cpp b/Engine/Source/Runtime/Core/Private/Misc/OutputDeviceFile.cpp
--- a/Engine/Source/Runtime/Core/Private/Misc/OutputDeviceFile.cpp
+++ b/Engine/Source/Runtime/Core/Private/Misc/OutputDeviceFile.cpp
@@ -164,6 +164,7 @@ FAsyncWriter::FAsyncWriter(FArchive& InAr, FAsyncWriter::EThreadNameOption NameO
 	, BufferEndPos(0)
 	, LastArchiveFlushTime(0.0)
 {
+	constexpr int32 InitialBufferSize = 128 * 1024;
 	Buffer.AddUninitialized(InitialBufferSize);
 
 	float CommandLineInterval = 0.0;
@@ -190,7 +191,8 @@ FAsyncWriter::FAsyncWriter(FArchive& InAr, FAsyncWriter::EThreadNameOption NameO
 				WriterName = FString::Printf(TEXT("FAsyncWriter_%llu"), ++AsyncWriterCount);
 			}
 
-			FPlatformAtomics::InterlockedExchangePtr((void**)&Thread, FForkProcessHelper::CreateForkableThread(this, *WriterName, 0, TPri_BelowNormal));
+			int32 StackSize = 32 * 1024;
+			FPlatformAtomics::InterlockedExchangePtr((void**)&Thread, FForkProcessHelper::CreateForkableThread(this, *WriterName, StackSize, TPri_BelowNormal));
 		}
 	}
 }
@@ -340,8 +342,9 @@ FOutputDeviceFile::FOutputDeviceFile( const TCHAR* InFilename, bool bInDisableBa
 , OnFileOpenedFn(MoveTemp(FileOpenedCallback))
 , AppendIfExists(bInAppendIfExists)
 , Dead(false)
-, CategoryInclusionInternal(nullptr)
+, bCreatingWriter(false)
 , bDisableBackup(bInDisableBackup)
+, CategoryInclusionInternal(nullptr)
 {
 	if( InFilename )
 	{
@@ -463,6 +466,13 @@ bool FOutputDeviceFile::CreateWriter(uint32 MaxAttempts)
 		return true;
 	}
 
+	// Ignore re-entrant calls to CreateWriter because that will create and leak a writer.
+	if (bCreatingWriter)
+	{
+		return true;
+	}
+	TGuardValue<bool> SerializingScope(bCreatingWriter, true);
+
 	// Make log filename.
 	if (!Filename[0])
 	{
diff --git a/Engine/Source/Runtime/Core/Public/Misc/OutputDeviceFile.h b/Engine/Source/Runtime/Core/Public/Misc/OutputDeviceFile.h
--- a/Engine/Source/Runtime/Core/Public/Misc/OutputDeviceFile.h
+++ b/Engine/Source/Runtime/Core/Public/Misc/OutputDeviceFile.h
@@ -29,11 +29,6 @@ class FRunnableThread;
 */
 class FAsyncWriter : public FRunnable, public FSingleThreadRunnable, public FArchive
 {
-	enum EConstants
-	{
-		InitialBufferSize = 128 * 1024
-	};
-
 	/** Thread to run the worker FRunnable on. Serializes the ring buffer to disk. */
 	volatile FRunnableThread* Thread;
 	/** Stops this thread */
@@ -186,15 +181,16 @@ private:
 	TCHAR Filename[1024];
 	bool AppendIfExists;
 	bool Dead;
+	bool bCreatingWriter;
+
+	/** If true, existing files will not be backed up */
+	bool bDisableBackup;
 
 	/** Internal data for category inclusion. Must be declared inside CPP file as it uses a TSet<FName> */
 	struct FCategoryInclusionInternal;
 
 	TUniquePtr<FCategoryInclusionInternal> CategoryInclusionInternal;
 
-	/** If true, existing files will not be backed up */
-	bool		bDisableBackup;
-
 	CORE_API void WriteRaw(const TCHAR* C);
 
 	/** Creates the async writer and its archive. Returns true if successful.  */
