Bug: tag ini config gets overwritten instead of merging new tag settings
Fix: Use FConfigFile directly and make sure to make ini file writable

#rnx

diff --git a/Engine/Source/Runtime/Core/Private/Misc/AutomationTest.cpp b/Engine/Source/Runtime/Core/Private/Misc/AutomationTest.cpp
--- a/Engine/Source/Runtime/Core/Private/Misc/AutomationTest.cpp
+++ b/Engine/Source/Runtime/Core/Private/Misc/AutomationTest.cpp
@@ -7,6 +7,7 @@
 #include "AutoRTFM.h"
 #include "Algo/Copy.h"
 #include "HAL/FileManager.h"
+#include "HAL/PlatformFileManager.h"
 #include "HAL/PlatformStackWalk.h"
 #include "HAL/ThreadHeartBeat.h"
 #include "Internationalization/Internationalization.h"
@@ -986,7 +987,7 @@ TArray<FString> FAutomationTestFramework::SaveTestTagMappings(const TArray<FStri
 
 	FString TestSourcePath;
 	FString TagsSettingsFilePath;
-	TSet<FString> AffectedTagConfigPaths;
+	TMap<FString, FConfigFile> AffectedTagConfigPaths;
 	for (int i = 0; i < TestFilePaths.Num(); i++)
 	{
 		TagsSettingsFilePath.Empty();
@@ -1022,19 +1023,25 @@ TArray<FString> FAutomationTestFramework::SaveTestTagMappings(const TArray<FStri
 		TagsSettingsFilePath = FConfigCacheIni::NormalizeConfigIniPath(TagsSettingsFilePath);
 		if (!AffectedTagConfigPaths.Contains(TagsSettingsFilePath))
 		{
-			AffectedTagConfigPaths.Add(TagsSettingsFilePath);
-			GConfig->Add(TagsSettingsFilePath, FConfigFile());
+			FPlatformFileManager::Get().GetPlatformFile().SetReadOnly(*TagsSettingsFilePath, false);
+			AffectedTagConfigPaths.Add(TagsSettingsFilePath, FConfigFile());
 		}
-		GConfig->SetString(TEXT("TestTagMapping"), *TestFullNames[i], **TestFullNameToTagDataMap.Find(TestFullNames[i]), TagsSettingsFilePath);
+		AffectedTagConfigPaths[TagsSettingsFilePath].SetInSection(TEXT("TestTagMapping"), *TestFullNames[i], **TestFullNameToTagDataMap.Find(TestFullNames[i]));
 	}
 
-	for (FString ConfigPath : AffectedTagConfigPaths)
+	TSet<FString> AllConfigFiles;
+	AffectedTagConfigPaths.GetKeys(AllConfigFiles);
+	for (FString ConfigPath : AllConfigFiles)
 	{
 		UE_LOG(LogAutomationTest, Log, TEXT("Saving tags to %s."), *ConfigPath);
-		GConfig->Flush(false, ConfigPath);
+		if (IFileManager::Get().FileExists(*ConfigPath))
+		{
+			AffectedTagConfigPaths[ConfigPath].Combine(ConfigPath);
+		}
+		AffectedTagConfigPaths[ConfigPath].Write(ConfigPath);
 	}
 
-	return AffectedTagConfigPaths.Array();
+	return AllConfigFiles.Array();
 }
 
 bool FAutomationTestFramework::ShouldTestContent(const FString& Path) const
