[Insights] Added tracing for the platform event settings (sampling interval).

diff --git a/Engine/Source/Runtime/Core/Private/ProfilingDebugging/Microsoft/EventTracingForWindows.cpp b/Engine/Source/Runtime/Core/Private/ProfilingDebugging/Microsoft/EventTracingForWindows.cpp
--- a/Engine/Source/Runtime/Core/Private/ProfilingDebugging/Microsoft/EventTracingForWindows.cpp
+++ b/Engine/Source/Runtime/Core/Private/ProfilingDebugging/Microsoft/EventTracingForWindows.cpp
@@ -117,7 +117,9 @@ void FPlatformEvents::TraceEventCallback(PEVENT_RECORD Event)
 		// Recent updates of ETW changes the size of the payload somewhat.
 		constexpr uint32 UserDataSizeOld = sizeof(FContextSwitchEvent) - 4;
 		constexpr uint32 UserDataSizeNew = sizeof(FContextSwitchEvent);
-		if (Payload == nullptr || (Event->UserDataLength != UserDataSizeOld && Event->UserDataLength != UserDataSizeNew))
+		if (Payload == nullptr ||
+			(Event->UserDataLength != UserDataSizeOld &&
+			 Event->UserDataLength != UserDataSizeNew))
 		{
 			return;
 		}
@@ -176,15 +178,23 @@ void FPlatformEvents::TraceEventCallback(PEVENT_RECORD Event)
 	else if (Opcode == FStackWalkEvent::Type)
 	{
 		const FStackWalkEvent* Payload = (const FStackWalkEvent*)Event->UserData;
-		if (Payload == nullptr || Payload->ProcessId != CurrentProcessId || Event->UserDataLength < sizeof(FStackWalkEvent))
+		if (Payload == nullptr ||
+			Payload->ProcessId != CurrentProcessId ||
+			Event->UserDataLength < sizeof(FStackWalkEvent))
 		{
 			return;
 		}
 
 		uint32 Count = (Event->UserDataLength - sizeof(FStackWalkEvent)) / sizeof(uint64);
 
-		// sometimes stack walk from ETW gives a call-stack that is fully in kernel address space
-		// so if all addresses have a high bit set, then ignore this call stack
+		// Sometimes we get call stack with only one or two entries, which obviously is not very useful.
+		//if (Count <= 2)
+		//{
+		//	return;
+		//}
+
+		// Sometimes stack walk from ETW gives a call-stack that is fully in kernel address space.
+		// So, if all addresses have a high bit set, then ignore this call stack.
 		bool bValid = false;
 		for (uint32 Index = 0; Index < Count; Index++)
 		{
@@ -195,9 +205,7 @@ void FPlatformEvents::TraceEventCallback(PEVENT_RECORD Event)
 				break;
 			}
 		}
-
-		// sometimes we get call stack with only one or two entries, which obviously is not very useful
-		if (bValid && Count > 2)
+		if (bValid)
 		{
 			FPlatformEventsTrace::OutputStackSample(Payload->TimeStamp, Payload->ThreadId, Payload->Stack, Count);
 		}
@@ -537,6 +545,8 @@ void FPlatformEvents::StopETW()
 
 void FPlatformEventsTrace::Init(uint32 SamplingIntervalUsec)
 {
+	FPlatformEventsTrace::OutputSettings(SamplingIntervalUsec);
+
 	GPlatformEvents = new FPlatformEvents(SamplingIntervalUsec);
 }
 
diff --git a/Engine/Source/Runtime/Core/Private/ProfilingDebugging/PlatformEvents.cpp b/Engine/Source/Runtime/Core/Private/ProfilingDebugging/PlatformEvents.cpp
--- a/Engine/Source/Runtime/Core/Private/ProfilingDebugging/PlatformEvents.cpp
+++ b/Engine/Source/Runtime/Core/Private/ProfilingDebugging/PlatformEvents.cpp
@@ -10,35 +10,49 @@ UE_TRACE_CHANNEL_DEFINE(StackSamplingChannel)
 
 /////////////////////////////////////////////////////////////////////
 
-// represents time interval when thread was running on specific core
+// Represents the settings for Context Switches and Stack Sampling
+UE_TRACE_EVENT_BEGIN(PlatformEvent, Settings, NoSync | Important)
+	UE_TRACE_EVENT_FIELD(uint32, SamplingInterval) // [microseconds]
+UE_TRACE_EVENT_END()
+
+// Represents time interval when thread was running on specific core
 UE_TRACE_EVENT_BEGIN(PlatformEvent, ContextSwitch, NoSync)
-UE_TRACE_EVENT_FIELD(uint64, StartTime)
-UE_TRACE_EVENT_FIELD(uint64, EndTime)
-UE_TRACE_EVENT_FIELD(uint32, ThreadId)
-UE_TRACE_EVENT_FIELD(uint8, CoreNumber)
+	UE_TRACE_EVENT_FIELD(uint64, StartTime)
+	UE_TRACE_EVENT_FIELD(uint64, EndTime)
+	UE_TRACE_EVENT_FIELD(uint32, ThreadId) // system thread id
+	UE_TRACE_EVENT_FIELD(uint8, CoreNumber)
 UE_TRACE_EVENT_END()
 
-// represents call stack addresses in stack sampling
+// Represents call stack addresses in Stack Sampling
 UE_TRACE_EVENT_BEGIN(PlatformEvent, StackSample, NoSync)
-UE_TRACE_EVENT_FIELD(uint64, Time)
-UE_TRACE_EVENT_FIELD(uint32, ThreadId)
-UE_TRACE_EVENT_FIELD(uint64[], Addresses)
+	UE_TRACE_EVENT_FIELD(uint64, Time)
+	UE_TRACE_EVENT_FIELD(uint32, ThreadId) // system thread id
+	UE_TRACE_EVENT_FIELD(uint64[], Addresses)
 UE_TRACE_EVENT_END()
 
-// thread information for context switches
-// NOTE: in some cases name of process can be empty, for example when there are no
-// privileges to query it or there are no name for process thread belongs to.
-// Depending on platform actual name can be absolute path to process executable.
+// Thread information for Context Switches and Stack Sampling
+// NOTE: In some cases, name of process can be empty, for example when there are no
+// privileges to query it or there is no name for the process thread belongs to.
+// Depending on platform, actual name can be absolute path to process executable.
 UE_TRACE_EVENT_BEGIN(PlatformEvent, ThreadName, NoSync)
-UE_TRACE_EVENT_FIELD(uint32, ThreadId)
-UE_TRACE_EVENT_FIELD(uint32, ProcessId)
-UE_TRACE_EVENT_FIELD(UE::Trace::WideString, Name)
+	UE_TRACE_EVENT_FIELD(uint32, ThreadId) // system thread id
+	UE_TRACE_EVENT_FIELD(uint32, ProcessId)
+	UE_TRACE_EVENT_FIELD(UE::Trace::WideString, Name)
 UE_TRACE_EVENT_END()
 
 /////////////////////////////////////////////////////////////////////
 
 #if PLATFORM_SUPPORTS_PLATFORM_EVENTS
 
+void FPlatformEventsTrace::OutputSettings(uint32 SamplingIntervalUsec)
+{
+	bool bShouldTrace = UE_TRACE_CHANNELEXPR_IS_ENABLED(StackSamplingChannel) || UE_TRACE_CHANNELEXPR_IS_ENABLED(ContextSwitchChannel);
+	UE_TRACE_LOG(PlatformEvent, Settings, bShouldTrace)
+		<< Settings.SamplingInterval(SamplingIntervalUsec);
+}
+
+/////////////////////////////////////////////////////////////////////
+
 void FPlatformEventsTrace::OutputContextSwitch(uint64 StartTime, uint64 EndTime, uint32 ThreadId, uint8 CoreNumber)
 {
 	UE_TRACE_LOG(PlatformEvent, ContextSwitch, ContextSwitchChannel)
@@ -62,7 +76,8 @@ void FPlatformEventsTrace::OutputStackSample(uint64 Time, uint32 ThreadId, const
 
 void FPlatformEventsTrace::OutputThreadName(uint32 ThreadId, uint32 ProcessId, const TCHAR* Name, uint32 NameLen)
 {
-	UE_TRACE_LOG(PlatformEvent, ThreadName, ContextSwitchChannel)
+	bool bShouldTrace = UE_TRACE_CHANNELEXPR_IS_ENABLED(StackSamplingChannel) || UE_TRACE_CHANNELEXPR_IS_ENABLED(ContextSwitchChannel);
+	UE_TRACE_LOG(PlatformEvent, ThreadName, bShouldTrace)
 		<< ThreadName.ThreadId(ThreadId)
 		<< ThreadName.ProcessId(ProcessId)
 		<< ThreadName.Name(Name, NameLen);
diff --git a/Engine/Source/Runtime/Core/Public/ProfilingDebugging/PlatformEvents.h b/Engine/Source/Runtime/Core/Public/ProfilingDebugging/PlatformEvents.h
--- a/Engine/Source/Runtime/Core/Public/ProfilingDebugging/PlatformEvents.h
+++ b/Engine/Source/Runtime/Core/Public/ProfilingDebugging/PlatformEvents.h
@@ -40,6 +40,7 @@ public:
 	CORE_API static void Stop();
 	CORE_API static void OnTraceChannelUpdated(const FString& ChannelName, bool bIsEnabled);
 
+	CORE_API static void OutputSettings(uint32 SamplingIntervalUsec);
 	CORE_API static void OutputContextSwitch(uint64 StartTime, uint64 EndTime, uint32 ThreadId, uint8 CoreNumber);
 	CORE_API static void OutputStackSample(uint64 Time, uint32 ThreadId, const uint64* Addresses, uint32 AddressCount);
 	CORE_API static void OutputThreadName(uint32 ThreadId, uint32 ProcessId, const TCHAR* Name, uint32 NameLen);
