Replace redundant C++17 code in UEOps.h.
Fixed some broken heterogeneous ordering operators by adding UEOpGreaterThan.
Added UEOpGreaterThan to TString, FStringOutputDevice and TStringView.

diff --git a/Engine/Source/Runtime/Core/Public/Containers/StringView.h b/Engine/Source/Runtime/Core/Public/Containers/StringView.h
--- a/Engine/Source/Runtime/Core/Public/Containers/StringView.h
+++ b/Engine/Source/Runtime/Core/Public/Containers/StringView.h
@@ -430,6 +430,7 @@ private:
 	static bool PrivateEquals(TStringView Lhs, const CharType* Rhs);
 	static bool PrivateEquals(TStringView Lhs, TStringView Rhs);
 	static bool PrivateLess(TStringView Lhs, TStringView Rhs);
+	static bool PrivateGreater(TStringView Lhs, TStringView Rhs);
 
 	template <UE::CCharType OtherCharType>
 	inline bool PrivateEquals(const OtherCharType* Other, ESearchCase::Type SearchCase) const;
@@ -464,7 +465,7 @@ public:
 		typename CharRangeType
 		UE_REQUIRES(std::is_convertible_v<CharRangeType&&, TStringView>)
 	>
-	[[nodiscard]] inline auto UEOpEquals(CharRangeType&& Rhs) const
+	[[nodiscard]] UE_REWRITE auto UEOpEquals(CharRangeType&& Rhs) const
 	{
 		return TStringView::PrivateEquals(*this, ImplicitConv<TStringView>(Forward<CharRangeType>(Rhs)));
 	}
@@ -473,12 +474,21 @@ public:
 		typename CharRangeType
 		UE_REQUIRES(std::is_convertible_v<CharRangeType&&, TStringView>)
 	>
-	[[nodiscard]] inline auto UEOpLessThan(CharRangeType&& Rhs)
+	[[nodiscard]] UE_REWRITE auto UEOpLessThan(CharRangeType&& Rhs)
 	{
 		return TStringView::PrivateLess(*this, ImplicitConv<TStringView>(Forward<CharRangeType>(Rhs)));
 	}
 
-	[[nodiscard]] inline bool UEOpEquals(const CharType* Rhs) const
+	template <
+		typename CharRangeType
+		UE_REQUIRES(std::is_convertible_v<CharRangeType&&, TStringView>)
+	>
+	[[nodiscard]] UE_REWRITE auto UEOpGreaterThan(CharRangeType&& Rhs)
+	{
+		return TStringView::PrivateGreater(*this, ImplicitConv<TStringView>(Forward<CharRangeType>(Rhs)));
+	}
+
+	[[nodiscard]] UE_REWRITE bool UEOpEquals(const CharType* Rhs) const
 	{
 		return TStringView::PrivateEquals(*this, Rhs);
 	}
@@ -783,6 +793,12 @@ template <UE::CCharType CharType>
 	return Lhs.Compare(Rhs, ESearchCase::IgnoreCase) < 0;
 }
 
+template <UE::CCharType CharType>
+[[nodiscard]] inline bool TStringView<CharType>::PrivateGreater(TStringView Lhs, TStringView Rhs)
+{
+	return Lhs.Compare(Rhs, ESearchCase::IgnoreCase) > 0;
+}
+
 template <UE::CCharType CharType>
 template <UE::CStringViewable OtherRangeType>
 [[nodiscard]] inline bool TStringView<CharType>::Equals(const OtherRangeType& Other, ESearchCase::Type SearchCase) const
diff --git a/Engine/Source/Runtime/Core/Public/Containers/UnrealString.h.inl b/Engine/Source/Runtime/Core/Public/Containers/UnrealString.h.inl
--- a/Engine/Source/Runtime/Core/Public/Containers/UnrealString.h.inl
+++ b/Engine/Source/Runtime/Core/Public/Containers/UnrealString.h.inl
@@ -905,6 +905,20 @@ public:
 		return FPlatformString::Stricmp(**this, Rhs) < 0;
 	}
 
+	/**
+	 * Lexicographically test whether the left string is > the right string
+	 *
+	 * @param Lhs String to compare against.
+	 * @param Rhs String to compare against.
+	 * @return true if the left string is lexicographically > the right string, otherwise false
+	 * @note case insensitive
+	 */
+	template <typename CharType>
+	[[nodiscard]] UE_FORCEINLINE_HINT bool UEOpGreaterThan(const CharType* Rhs) const
+	{
+		return FPlatformString::Stricmp(**this, Rhs) > 0;
+	}
+
 	/**
 	 * Lexicographically test whether the left string is == the right string
 	 *
diff --git a/Engine/Source/Runtime/Core/Public/Misc/StringOutputDevice.h b/Engine/Source/Runtime/Core/Public/Misc/StringOutputDevice.h
--- a/Engine/Source/Runtime/Core/Public/Misc/StringOutputDevice.h
+++ b/Engine/Source/Runtime/Core/Public/Misc/StringOutputDevice.h
@@ -45,15 +45,20 @@ public:
 		return FString::operator+=(Other);
 	}
 
-	bool UEOpEquals(const FString& Rhs) const
+	UE_REWRITE bool UEOpEquals(const FString& Rhs) const
 	{
 		return FString::UEOpEquals(Rhs);
 	}
 
-	bool UEOpLessThan(const FString& Rhs) const
+	UE_REWRITE bool UEOpLessThan(const FString& Rhs) const
 	{
 		return FString::UEOpLessThan(Rhs);
 	}
+
+	UE_REWRITE bool UEOpGreaterThan(const FString& Rhs) const
+	{
+		return Rhs.UEOpLessThan(*this);
+	}
 };
 
 template <>
diff --git a/Engine/Source/Runtime/Core/Public/Misc/UEOps.h b/Engine/Source/Runtime/Core/Public/Misc/UEOps.h
--- a/Engine/Source/Runtime/Core/Public/Misc/UEOps.h
+++ b/Engine/Source/Runtime/Core/Public/Misc/UEOps.h
@@ -2,7 +2,8 @@
 
 #pragma once
 
-#include "Templates/Requires.h"
+#include "Concepts/SameAs.h"
+#include "HAL/Platform.h"
 
 namespace UE::Core::Private
 {
@@ -11,55 +12,22 @@ namespace UE::Core::Private
 	struct FIncomplete;
 
 	template <typename LhsType, typename RhsType>
-	auto EnableUEOpEqualsImpl(const LhsType& Lhs, const RhsType& Rhs, int) -> decltype(Lhs.UEOpEquals(Rhs));
-	template <typename LhsType, typename RhsType>
-	void EnableUEOpEqualsImpl(const LhsType& Lhs, const RhsType& Rhs, ...) = delete;
-
-	template <typename LhsType, typename RhsType>
-	auto DisableUEOpEqualsImpl(const LhsType& Lhs, const RhsType& Rhs, int) -> decltype(Lhs.UEOpEquals(Rhs)) = delete;
-	template <typename LhsType, typename RhsType>
-	void DisableUEOpEqualsImpl(const LhsType& Lhs, const RhsType& Rhs, ...);
+	concept CWithUEOpEquals = requires(const LhsType & Lhs, const RhsType & Rhs)
+	{
+		{ Lhs.UEOpEquals(Rhs) } -> UE::CSameAs<bool>;
+	};
 
 	template <typename LhsType, typename RhsType>
-	auto EnableUEOpLessThanImpl(const LhsType& Lhs, const RhsType& Rhs, int) -> decltype(Lhs.UEOpLessThan(Rhs));
-	template <typename LhsType, typename RhsType>
-	void EnableUEOpLessThanImpl(const LhsType& Lhs, const RhsType& Rhs, ...) = delete;
+	concept CWithUEOpLessThan = requires(const LhsType & Lhs, const RhsType & Rhs)
+	{
+		{ Lhs.UEOpLessThan(Rhs) } -> UE::CSameAs<bool>;
+	};
 
 	template <typename LhsType, typename RhsType>
-	auto DisableUEOpLessThanImpl(const LhsType& Lhs, const RhsType& Rhs, int) -> decltype(Lhs.UEOpLessThan(Rhs)) = delete;
-	template <typename LhsType, typename RhsType>
-	void DisableUEOpLessThanImpl(const LhsType& Lhs, const RhsType& Rhs, ...);
-
-	// SFINAE-compatible concept checks for the existence of various UEOp member functions.
-	//
-	// We can't use metaprogramming classes like TModels or std::enable_if because they end up being instantiated with incomplete types and
-	// end up with incorrect overload resolution later in the compilation.
-	//
-	// Use:
-	// decltype(UE::Core::Private::UEOpEquals(A, B)) should be present in the signature of the function to enable it for comparisons between
-	// the two operands.
-	// This will cause SFINAE failure (and thus 'enable' the overload) if A.UEOpEquals(B) is a valid expression. If the function is enabled,
-	// the return type of this function will be the same as the return type as A.UEOpEquals(B).
-	//
-	// decltype(UE::Core::Private::DisableUEOpEquals(A, B)) will cause SFINAE failure (and thus 'disable' the overload) if A.UEOpEquals(B) is
-	// a valid expression.
-	// The typical use is in conjunction with decltype(UE::Core::Private::UEOpEquals(B, A)) to cause mutually exclusive overloads with a
-	// decltype(UE::Core::Private::UEOpEquals(A, B)) overload to prevent ambiguity when A and B are comparable in both directions, e.g.
-	// comparing TSharedPtr<X> with TSharedPtr<Y>.
-	//
-	// UE::Core::Private::UEOpLessThan and UE::Core::Private::DisableUEOpLessThan work similarly, but for UEOpLessThan.
-
-	template <typename LhsType, typename RhsType>
-	auto UEOpEquals(const LhsType& Lhs, const RhsType& Rhs) -> decltype(::UE::Core::Private::EnableUEOpEqualsImpl(Lhs, Rhs, 0));
-
-	template <typename LhsType, typename RhsType>
-	auto DisableUEOpEquals(const LhsType& Lhs, const RhsType& Rhs) -> decltype(::UE::Core::Private::DisableUEOpEqualsImpl(Lhs, Rhs, 0));
-
-	template <typename LhsType, typename RhsType>
-	auto UEOpLessThan(const LhsType& Lhs, const RhsType& Rhs) -> decltype(::UE::Core::Private::EnableUEOpLessThanImpl(Lhs, Rhs, 0));
-
-	template <typename LhsType, typename RhsType>
-	auto DisableUEOpLessThan(const LhsType& Lhs, const RhsType& Rhs) -> decltype(::UE::Core::Private::DisableUEOpLessThanImpl(Lhs, Rhs, 0));
+	concept CWithUEOpGreaterThan = requires(const LhsType & Lhs, const RhsType & Rhs)
+	{
+		{ Lhs.UEOpGreaterThan(Rhs) } -> UE::CSameAs<bool>;
+	};
 }
 
 /**
@@ -67,111 +35,97 @@ namespace UE::Core::Private
  * These ensure that the operators and their variants are defined correctly based on a single customization
  * point and should result in faster overload resolution.
  *
+ * bool FMyType::UEOpEquals(const FMyType&) const;
  * bool FMyType::UEOpEquals(const OtherType&) const;
  *   Should be defined to opt into operator== and operator!=.
  *
+ * bool FMyType::UEOpLessThan(const FMyType&) const;
  * bool FMyType::UEOpLessThan(const OtherType&) const;
+ * bool FMyType::UEOpGreaterThan(const OtherType&) const;
  *   Should be defined to opt into operator<, operator<=, operator> and operator>=.
+ *
+ * bool FMyType::UEOpGreaterThan(const FMyType&) const;
+ *   Will never be called if FMyType::UEOpLessThan is defined, so is redundant.
  */
 
 template <typename LhsType, typename RhsType>
-[[nodiscard]] UE_REWRITE constexpr auto operator==(const LhsType& Lhs, const RhsType& Rhs) -> decltype(UE::Core::Private::UEOpEquals(Lhs, Rhs))
+	requires UE::Core::Private::CWithUEOpEquals<LhsType, RhsType> || UE::Core::Private::CWithUEOpEquals<RhsType, LhsType>
+[[nodiscard]] UE_REWRITE constexpr bool operator==(const LhsType& Lhs, const RhsType& Rhs)
 {
-	return Lhs.UEOpEquals(Rhs);
-}
-
-template <typename LhsType, typename RhsType>
-[[nodiscard]] UE_REWRITE constexpr auto operator==(const LhsType& Lhs, const RhsType& Rhs) -> decltype(UE::Core::Private::DisableUEOpEquals(Lhs, Rhs), UE::Core::Private::UEOpEquals(Rhs, Lhs))
-{
-	return Rhs.UEOpEquals(Lhs);
+	if constexpr (UE::Core::Private::CWithUEOpEquals<LhsType, RhsType>)
+	{
+		return Lhs.UEOpEquals(Rhs);
+	}
+	else
+	{
+		return Rhs.UEOpEquals(Lhs);
+	}
 }
 
 // This overload exists to support A == NULL, which is common in comparison operators, but NULL is not perfectly-forwardable.
 template <typename LhsType>
-[[nodiscard]] UE_REWRITE constexpr auto operator==(const LhsType& Lhs, UE::Core::Private::FIncomplete* Rhs) -> decltype(UE::Core::Private::UEOpEquals(Lhs, nullptr))
+	requires UE::Core::Private::CWithUEOpEquals<LhsType, decltype(nullptr)>
+[[nodiscard]] UE_REWRITE constexpr bool operator==(const LhsType& Lhs, UE::Core::Private::FIncomplete* Rhs)
 {
 	return Lhs.UEOpEquals(nullptr);
 }
 
-// This overload exists to support NULL == A, which is common in comparison operators, but NULL is not perfectly-forwardable.
-template <typename RhsType>
-[[nodiscard]] UE_REWRITE constexpr auto operator==(UE::Core::Private::FIncomplete* Lhs, const RhsType& Rhs) -> decltype(UE::Core::Private::UEOpEquals(Rhs, nullptr))
-{
-	return Rhs.UEOpEquals(nullptr);
-}
-
-#if !PLATFORM_COMPILER_HAS_GENERATED_COMPARISON_OPERATORS
 template <typename LhsType, typename RhsType>
-[[nodiscard]] UE_REWRITE constexpr auto operator!=(const LhsType& Lhs, const RhsType& Rhs) -> decltype(UE::Core::Private::UEOpEquals(Lhs, Rhs))
+	requires UE::Core::Private::CWithUEOpLessThan<LhsType, RhsType> || UE::Core::Private::CWithUEOpGreaterThan<RhsType, LhsType>
+[[nodiscard]] UE_REWRITE constexpr bool operator<(const LhsType& Lhs, const RhsType& Rhs)
 {
-	return !Lhs.UEOpEquals(Rhs);
+	if constexpr (UE::Core::Private::CWithUEOpLessThan<LhsType, RhsType>)
+	{
+		return Lhs.UEOpLessThan(Rhs);
+	}
+	else
+	{
+		return Rhs.UEOpGreaterThan(Lhs);
+	}
 }
 
 template <typename LhsType, typename RhsType>
-[[nodiscard]] UE_REWRITE constexpr auto operator!=(const LhsType& Lhs, const RhsType& Rhs) -> decltype(UE::Core::Private::DisableUEOpEquals(Lhs, Rhs), UE::Core::Private::UEOpEquals(Rhs, Lhs))
+	requires UE::Core::Private::CWithUEOpLessThan<RhsType, LhsType> || UE::Core::Private::CWithUEOpGreaterThan<LhsType, RhsType>
+[[nodiscard]] UE_REWRITE constexpr bool operator>(const LhsType& Lhs, const RhsType& Rhs)
 {
-	return !Rhs.UEOpEquals(Lhs);
-}
-
-// This overload exists to support A != NULL, which is common in comparison operators, but NULL is not perfectly-forwardable.
-template <typename LhsType>
-[[nodiscard]] UE_REWRITE constexpr auto operator!=(const LhsType& Lhs, UE::Core::Private::FIncomplete* Rhs) -> decltype(UE::Core::Private::UEOpEquals(Lhs, nullptr))
-{
-	return !Lhs.UEOpEquals(nullptr);
-}
-
-// This overload exists to support NULL != A, which is common in comparison operators, but NULL is not perfectly-forwardable.
-template <typename RhsType>
-[[nodiscard]] UE_REWRITE constexpr auto operator!=(UE::Core::Private::FIncomplete* Lhs, const RhsType& Rhs) -> decltype(UE::Core::Private::UEOpEquals(Rhs, nullptr))
-{
-	return !Rhs.UEOpEquals(nullptr);
-}
-#endif
-
-template <typename LhsType, typename RhsType>
-[[nodiscard]] UE_REWRITE constexpr auto operator<(const LhsType& Lhs, const RhsType& Rhs) -> decltype(UE::Core::Private::UEOpLessThan(Lhs, Rhs))
-{
-	return Lhs.UEOpLessThan(Rhs);
-}
-
-template <typename LhsType, typename RhsType>
-[[nodiscard]] UE_REWRITE constexpr auto operator<(const LhsType& Lhs, const RhsType& Rhs) -> decltype(UE::Core::Private::DisableUEOpLessThan(Lhs, Rhs), UE::Core::Private::UEOpLessThan(Rhs, Lhs))
-{
-	return Rhs.UEOpLessThan(Lhs);
-}
-
-template <typename LhsType, typename RhsType>
-[[nodiscard]] UE_REWRITE constexpr auto operator>(const LhsType& Lhs, const RhsType& Rhs) -> decltype(UE::Core::Private::UEOpLessThan(Rhs, Lhs))
-{
-	return Rhs.UEOpLessThan(Lhs);
-}
-
-template <typename LhsType, typename RhsType>
-[[nodiscard]] UE_REWRITE constexpr auto operator>(const LhsType& Lhs, const RhsType& Rhs) -> decltype(UE::Core::Private::DisableUEOpLessThan(Rhs, Lhs), UE::Core::Private::UEOpLessThan(Lhs, Rhs))
-{
-	return Lhs.UEOpLessThan(Rhs);
-}
-
-template <typename LhsType, typename RhsType>
-[[nodiscard]] UE_REWRITE constexpr auto operator>=(const LhsType& Lhs, const RhsType& Rhs) -> decltype(UE::Core::Private::UEOpLessThan(Lhs, Rhs))
-{
-	return !Lhs.UEOpLessThan(Rhs);
+	if constexpr (UE::Core::Private::CWithUEOpLessThan<RhsType, LhsType>)
+	{
+		return Rhs.UEOpLessThan(Lhs);
+	}
+	else
+	{
+		return Lhs.UEOpGreaterThan(Lhs);
+	}
 }
 
 template <typename LhsType, typename RhsType>
-[[nodiscard]] UE_REWRITE constexpr auto operator>=(const LhsType& Lhs, const RhsType& Rhs) -> decltype(UE::Core::Private::DisableUEOpLessThan(Lhs, Rhs), UE::Core::Private::UEOpLessThan(Rhs, Lhs))
+	requires UE::Core::Private::CWithUEOpLessThan<LhsType, RhsType> || UE::Core::Private::CWithUEOpGreaterThan<RhsType, LhsType>
+[[nodiscard]] UE_REWRITE constexpr bool operator>=(const LhsType& Lhs, const RhsType& Rhs)
 {
-	return !Rhs.UEOpLessThan(Lhs);
+	if constexpr (UE::Core::Private::CWithUEOpLessThan<LhsType, RhsType>)
+	{
+		return !Lhs.UEOpLessThan(Rhs);
+	}
+	else
+	{
+		return !Rhs.UEOpGreaterThan(Lhs);
+	}
 }
 
 template <typename LhsType, typename RhsType>
-[[nodiscard]] UE_REWRITE constexpr auto operator<=(const LhsType& Lhs, const RhsType& Rhs) -> decltype(UE::Core::Private::UEOpLessThan(Rhs, Lhs))
+	requires UE::Core::Private::CWithUEOpLessThan<RhsType, LhsType> || UE::Core::Private::CWithUEOpGreaterThan<LhsType, RhsType>
+[[nodiscard]] UE_REWRITE constexpr bool operator<=(const LhsType& Lhs, const RhsType& Rhs)
 {
-	return !Rhs.UEOpLessThan(Lhs);
+	if constexpr (UE::Core::Private::CWithUEOpLessThan<RhsType, LhsType>)
+	{
+		return !Rhs.UEOpLessThan(Lhs);
+	}
+	else
+	{
+		return !Lhs.UEOpGreaterThan(Rhs);
+	}
 }
 
-template <typename LhsType, typename RhsType>
-[[nodiscard]] UE_REWRITE constexpr auto operator<=(const LhsType& Lhs, const RhsType& Rhs) -> decltype(UE::Core::Private::DisableUEOpLessThan(Rhs, Lhs), UE::Core::Private::UEOpLessThan(Lhs, Rhs))
-{
-	return !Lhs.UEOpLessThan(Rhs);
-}
+#if UE_ENABLE_INCLUDE_ORDER_DEPRECATED_IN_5_7
+#include "Templates/Requires.h"
+#endif
