Fixed operator> logic.

diff --git a/Engine/Source/Runtime/Core/Public/Misc/UEOps.h b/Engine/Source/Runtime/Core/Public/Misc/UEOps.h
--- a/Engine/Source/Runtime/Core/Public/Misc/UEOps.h
+++ b/Engine/Source/Runtime/Core/Public/Misc/UEOps.h
@@ -94,7 +94,7 @@ template <typename LhsType, typename RhsType>
 	}
 	else
 	{
-		return Lhs.UEOpGreaterThan(Lhs);
+		return Lhs.UEOpGreaterThan(Rhs);
 	}
 }
 
