[MicrosoftRunnableThread]
- During thread creation, setup thread priority before the thread is allowed to execute to avoid race conditions

diff --git a/Engine/Source/Runtime/Core/Private/Microsoft/MicrosoftRunnableThread.h b/Engine/Source/Runtime/Core/Private/Microsoft/MicrosoftRunnableThread.h
--- a/Engine/Source/Runtime/Core/Private/Microsoft/MicrosoftRunnableThread.h
+++ b/Engine/Source/Runtime/Core/Private/Microsoft/MicrosoftRunnableThread.h
@@ -173,13 +173,15 @@ protected:
 		}
 		else
 		{
+			// Set the priority before resuming the thread, since the thread itself could end up manipulating its own priority
+			// and cause a race condition.
+			ThreadPriority = TPri_Normal; // Set back to default in case any SetThreadPrio() impls compare against current value to reduce syscalls
+			SetThreadPriority(InThreadPri);
+
 			ResumeThread(Thread);
 
 			// Let the thread start up
 			ThreadInitSyncEvent->Wait(INFINITE);
-
-			ThreadPriority = TPri_Normal; // Set back to default in case any SetThreadPrio() impls compare against current value to reduce syscalls
-			SetThreadPriority(InThreadPri);
 		}
 
 		// Cleanup the sync event
