Fixing issues when using UpdateProcThreadAttribute. Values set via UpdateProcThreadAttribute must be live until the attribute is deleted.

diff --git a/Engine/Source/Runtime/Core/Private/Windows/WindowsPlatformProcess.cpp b/Engine/Source/Runtime/Core/Private/Windows/WindowsPlatformProcess.cpp
--- a/Engine/Source/Runtime/Core/Private/Windows/WindowsPlatformProcess.cpp
+++ b/Engine/Source/Runtime/Core/Private/Windows/WindowsPlatformProcess.cpp
@@ -974,6 +974,9 @@ bool FWindowsPlatformProcess::ExecProcess(const TCHAR* URL, const TCHAR* Params,
 
 	HANDLE hStdOutRead = NULL;
 	HANDLE hStdErrRead = NULL;
+
+	// InheritHandles data must persist at least as long as AttributeList. see: https://learn.microsoft.com/en-us/windows/win32/api/processthreadsapi/nf-processthreadsapi-updateprocthreadattribute
+	HANDLE InheritHandles[2] = {};
 	TArray<uint8> AttributeList;
 
 	if(OutStdOut != nullptr || OutStdErr != nullptr)
@@ -996,7 +999,9 @@ bool FWindowsPlatformProcess::ExecProcess(const TCHAR* URL, const TCHAR* Params,
 			verify(InitializeProcThreadAttributeList(StartupInfoEx.lpAttributeList, 1, 0, &BufferSize));
 		}
 
-		HANDLE InheritHandles[2] = { StartupInfoEx.StartupInfo.hStdOutput, StartupInfoEx.StartupInfo.hStdError };
+		InheritHandles[0] = StartupInfoEx.StartupInfo.hStdOutput;
+		InheritHandles[1] = StartupInfoEx.StartupInfo.hStdError;
+
 		verify(UpdateProcThreadAttribute(StartupInfoEx.lpAttributeList, 0, PROC_THREAD_ATTRIBUTE_HANDLE_LIST, InheritHandles, sizeof(InheritHandles), NULL, NULL));
 	}
 
