Make FWindowsPlatformProcess::CreateProc and FMacPlatformProcess::CreateProc specify the handles to inherit if bInheritHandles is false. Also make Windows and Mac be consistent when bLaunchDetached is true

#DisabledValidations nonUnityAndStaticAnalysis

diff --git a/Engine/Plugins/Experimental/LearningAgents/Source/LearningTraining/Private/LearningTrainer.cpp b/Engine/Plugins/Experimental/LearningAgents/Source/LearningTraining/Private/LearningTrainer.cpp
--- a/Engine/Plugins/Experimental/LearningAgents/Source/LearningTraining/Private/LearningTrainer.cpp
+++ b/Engine/Plugins/Experimental/LearningAgents/Source/LearningTraining/Private/LearningTrainer.cpp
@@ -32,7 +32,7 @@ namespace UE::Learning
 			return false;
 		}
 
-		ProcessHandle = FPlatformProcess::CreateProc(*Path, *Params, false, bHideWindow, bHideWindow, nullptr, 0, *FPaths::RootDir(), WritePipe, ReadPipe);
+		ProcessHandle = FPlatformProcess::CreateProc(*Path, *Params, false, bHideWindow, bHideWindow, nullptr, 0, *FPaths::RootDir(), WritePipe);
 		bIsLaunched = true;
 		return true;
 	}
diff --git a/Engine/Plugins/Experimental/Web/WebAPI/Source/WebAPILiquidJS/Private/WebAPILiquidJSProcess.cpp b/Engine/Plugins/Experimental/Web/WebAPI/Source/WebAPILiquidJS/Private/WebAPILiquidJSProcess.cpp
--- a/Engine/Plugins/Experimental/Web/WebAPI/Source/WebAPILiquidJS/Private/WebAPILiquidJSProcess.cpp
+++ b/Engine/Plugins/Experimental/Web/WebAPI/Source/WebAPILiquidJS/Private/WebAPILiquidJSProcess.cpp
@@ -156,8 +156,7 @@ uint32 FWebAPILiquidJSProcess::Run()
 		nullptr,		/* OutProcessID */
 		0,				/* PriorityModifier */
 		*Root,			/* OptionalWorkingDirectory */
-		WritePipe,		/* PipeWriteChild */
-		ReadPipe		/* PipeReadChild */
+		WritePipe		/* PipeWriteChild */
 	);
 
 	if (!Process.IsValid())
diff --git a/Engine/Plugins/MovieScene/MovieRenderPipeline/Source/MovieRenderPipelineCore/Private/Graph/Nodes/MovieGraphCommandLineEncoderNode.cpp b/Engine/Plugins/MovieScene/MovieRenderPipeline/Source/MovieRenderPipelineCore/Private/Graph/Nodes/MovieGraphCommandLineEncoderNode.cpp
--- a/Engine/Plugins/MovieScene/MovieRenderPipeline/Source/MovieRenderPipelineCore/Private/Graph/Nodes/MovieGraphCommandLineEncoderNode.cpp
+++ b/Engine/Plugins/MovieScene/MovieRenderPipeline/Source/MovieRenderPipelineCore/Private/Graph/Nodes/MovieGraphCommandLineEncoderNode.cpp
@@ -594,7 +594,7 @@ void UMovieGraphCommandLineEncoderNode::LaunchEncoder(const FEncoderParams& InPa
 	constexpr bool bLaunchHidden = true;
 	constexpr bool bLaunchReallyHidden = bLaunchHidden;
 	const FString ExecutableArg = FString::Format(TEXT("{Executable}"), InParams.NamedArguments);
-	const FProcHandle ProcessHandle = FPlatformProcess::CreateProc(*ExecutableArg, *CommandToRun, bLaunchDetached, bLaunchHidden, bLaunchReallyHidden, nullptr, 0, nullptr, PipeWrite, PipeRead);
+	const FProcHandle ProcessHandle = FPlatformProcess::CreateProc(*ExecutableArg, *CommandToRun, bLaunchDetached, bLaunchHidden, bLaunchReallyHidden, nullptr, 0, nullptr, PipeWrite);
 	if (ProcessHandle.IsValid())
 	{
 		FActiveJob& NewJob = ActiveEncodeJobs.AddDefaulted_GetRef();
diff --git a/Engine/Plugins/MovieScene/MovieRenderPipeline/Source/MovieRenderPipelineCore/Private/MoviePipelineCommandLineEncoder.cpp b/Engine/Plugins/MovieScene/MovieRenderPipeline/Source/MovieRenderPipelineCore/Private/MoviePipelineCommandLineEncoder.cpp
--- a/Engine/Plugins/MovieScene/MovieRenderPipeline/Source/MovieRenderPipelineCore/Private/MoviePipelineCommandLineEncoder.cpp
+++ b/Engine/Plugins/MovieScene/MovieRenderPipeline/Source/MovieRenderPipelineCore/Private/MoviePipelineCommandLineEncoder.cpp
@@ -326,7 +326,7 @@ void UMoviePipelineCommandLineEncoder::LaunchEncoder(const FEncoderParams& InPar
 	verify(FPlatformProcess::CreatePipe(PipeRead, PipeWrite));
 
 	FString ExecutableArg = FString::Format(TEXT("{Executable}"), InParams.NamedArguments);
-	FProcHandle ProcessHandle = FPlatformProcess::CreateProc(*ExecutableArg, *CommandLineArgs, bLaunchDetached, bLaunchHidden, bLaunchReallyHidden, nullptr, 0, nullptr, PipeWrite, PipeRead);
+	FProcHandle ProcessHandle = FPlatformProcess::CreateProc(*ExecutableArg, *CommandLineArgs, bLaunchDetached, bLaunchHidden, bLaunchReallyHidden, nullptr, 0, nullptr, PipeWrite);
 	if (ProcessHandle.IsValid())
 	{
 		FActiveJob& NewJob = ActiveEncodeJobs.AddDefaulted_GetRef();
diff --git a/Engine/Plugins/MovieScene/MovieRenderPipeline/Source/MovieRenderPipelineCore/Private/MoviePipelineCommandLineEncoderSettings.cpp b/Engine/Plugins/MovieScene/MovieRenderPipeline/Source/MovieRenderPipelineCore/Private/MoviePipelineCommandLineEncoderSettings.cpp
--- a/Engine/Plugins/MovieScene/MovieRenderPipeline/Source/MovieRenderPipelineCore/Private/MoviePipelineCommandLineEncoderSettings.cpp
+++ b/Engine/Plugins/MovieScene/MovieRenderPipeline/Source/MovieRenderPipelineCore/Private/MoviePipelineCommandLineEncoderSettings.cpp
@@ -36,7 +36,7 @@ static void PrintAvailableCodecs()
 
 	verify(FPlatformProcess::CreatePipe(PipeRead, PipeWrite));
 
-	FProcHandle ProcessHandle = FPlatformProcess::CreateProc(*Settings->ExecutablePath, *CommandLineArgs, bLaunchDetached, bLaunchHidden, bLaunchReallyHidden, nullptr, 0, nullptr, PipeWrite, PipeRead);
+	FProcHandle ProcessHandle = FPlatformProcess::CreateProc(*Settings->ExecutablePath, *CommandLineArgs, bLaunchDetached, bLaunchHidden, bLaunchReallyHidden, nullptr, 0, nullptr, PipeWrite);
 	if (ProcessHandle.IsValid())
 	{
 		FPlatformProcess::Sleep(0.01);
diff --git a/Engine/Plugins/Runtime/nDisplay/Source/DisplayClusterFillDerivedDataCache/Private/DisplayClusterFillDerivedDataCacheWorker.cpp b/Engine/Plugins/Runtime/nDisplay/Source/DisplayClusterFillDerivedDataCache/Private/DisplayClusterFillDerivedDataCacheWorker.cpp
--- a/Engine/Plugins/Runtime/nDisplay/Source/DisplayClusterFillDerivedDataCache/Private/DisplayClusterFillDerivedDataCacheWorker.cpp
+++ b/Engine/Plugins/Runtime/nDisplay/Source/DisplayClusterFillDerivedDataCache/Private/DisplayClusterFillDerivedDataCacheWorker.cpp
@@ -55,7 +55,7 @@ FDisplayClusterFillDerivedDataCacheWorker::FDisplayClusterFillDerivedDataCacheWo
 	const bool bLaunchReallyHidden = false;
 	ProcessHandle = FPlatformProcess::CreateProc(
 		*CurrentExecutableName, *Arguments, bLaunchDetached, bLaunchHidden, bLaunchReallyHidden, &ProcessID,
-		0, nullptr, WritePipe, ReadPipe);
+		0, nullptr, WritePipe);
 	
 	Thread = TUniquePtr<FRunnableThread>(FRunnableThread::Create(this, TEXT("FDisplayClusterFillDerivedDataCacheWorker")));
 }
diff --git a/Engine/Plugins/VirtualProduction/RemoteControlWebInterface/Source/RemoteControlWebInterface/Private/RCWebInterfaceProcess.cpp b/Engine/Plugins/VirtualProduction/RemoteControlWebInterface/Source/RemoteControlWebInterface/Private/RCWebInterfaceProcess.cpp
--- a/Engine/Plugins/VirtualProduction/RemoteControlWebInterface/Source/RemoteControlWebInterface/Private/RCWebInterfaceProcess.cpp
+++ b/Engine/Plugins/VirtualProduction/RemoteControlWebInterface/Source/RemoteControlWebInterface/Private/RCWebInterfaceProcess.cpp
@@ -176,8 +176,7 @@ uint32 FRemoteControlWebInterfaceProcess::Run()
 		nullptr,		/* OutProcessID */
 		0,				/* PriorityModifier */
 		*Root,			/* OptionalWorkingDirectory */
-		WritePipe,		/* PipeWriteChild */
-		ReadPipe		/* PipeReadChild */
+		WritePipe		/* PipeWriteChild */
 	);
 
 	if (!Process.IsValid())
diff --git a/Engine/Source/Developer/Android/AndroidDeviceDetection/Private/AndroidDeviceDetectionModule.cpp b/Engine/Source/Developer/Android/AndroidDeviceDetection/Private/AndroidDeviceDetectionModule.cpp
--- a/Engine/Source/Developer/Android/AndroidDeviceDetection/Private/AndroidDeviceDetectionModule.cpp
+++ b/Engine/Source/Developer/Android/AndroidDeviceDetection/Private/AndroidDeviceDetectionModule.cpp
@@ -44,154 +44,6 @@
 #endif
 #include "String/ParseLines.h"
 
-// FIX in 5.7 by removing this and calling FPlatformProcess::CreateProc directly once this code is moved there
-#if PLATFORM_WINDOWS
-#include "Windows/AllowWindowsPlatformTypes.h"
-static FProcHandle CreateProc(const TCHAR* URL, const TCHAR* Parms, bool bLaunchDetached, bool bLaunchHidden, bool bLaunchReallyHidden, uint32* OutProcessID, int32 PriorityModifier, const TCHAR* OptionalWorkingDirectory, void* PipeWriteChild)
-{
-	auto HandleError = [&URL, &Parms, &OutProcessID]
-	{
-		DWORD ErrorCode = GetLastError();
-
-		TCHAR ErrorMessage[512];
-		FWindowsPlatformMisc::GetSystemErrorMessage(ErrorMessage, 512, ErrorCode);
-
-		UE_LOG(LogWindows, Warning, TEXT("CreateProc failed: %s (0x%08x)"), ErrorMessage, ErrorCode);
-		if (ErrorCode == ERROR_NOT_ENOUGH_MEMORY || ErrorCode == ERROR_OUTOFMEMORY)
-		{
-			// These errors are common enough that we want some available memory information
-			FPlatformMemoryStats Stats = FPlatformMemory::GetStats();
-			UE_LOG(LogWindows, Warning, TEXT("Mem used: %.2f MB, OS Free %.2f MB"), Stats.UsedPhysical / 1048576.0f, Stats.AvailablePhysical / 1048576.0f);
-		}
-		UE_LOG(LogWindows, Warning, TEXT("URL: %s %s"), URL, Parms);
-		if (OutProcessID != nullptr)
-		{
-			*OutProcessID = 0;
-		}
-	};
-	
-	//UE_LOG(LogWindows, Log,  TEXT("CreateProc %s %s"), URL, Parms );
-
-	// initialize process creation flags
-	uint32 CreateFlags = NORMAL_PRIORITY_CLASS;
-	if (PriorityModifier < 0)
-	{
-		CreateFlags = (PriorityModifier == -1) ? BELOW_NORMAL_PRIORITY_CLASS : IDLE_PRIORITY_CLASS;
-	}
-	else if (PriorityModifier > 0)
-	{
-		CreateFlags = (PriorityModifier == 1) ? ABOVE_NORMAL_PRIORITY_CLASS : HIGH_PRIORITY_CLASS;
-	}
-
-	if (bLaunchDetached)
-	{
-		CreateFlags |= DETACHED_PROCESS;
-	}
-
-	// initialize window flags
-	uint32 dwFlags = 0;
-	uint16 ShowWindowFlags = SW_HIDE;
-	if (bLaunchReallyHidden)
-	{
-		dwFlags = STARTF_USESHOWWINDOW;
-	}
-	else if (bLaunchHidden)
-	{
-		dwFlags = STARTF_USESHOWWINDOW;
-		ShowWindowFlags = SW_SHOWMINNOACTIVE;
-	}
-
-	TArray<uint8> AttributeList;
-	{
-		HANDLE Handles[1];
-		uint32 HandleCount = 0;
-		
-		if (PipeWriteChild)
-		{
-			Handles[HandleCount++] = HANDLE(PipeWriteChild);
-		}
-		
-		if (HandleCount)
-		{
-			SIZE_T BufferSize = 0;
-			verify(!InitializeProcThreadAttributeList(nullptr, 1, 0, &BufferSize));
-			check(GetLastError() == ERROR_INSUFFICIENT_BUFFER);
-			AttributeList.SetNum(BufferSize);
-
-			if (!InitializeProcThreadAttributeList(LPPROC_THREAD_ATTRIBUTE_LIST(AttributeList.GetData()), 1, 0, &BufferSize))
-			{
-				HandleError();
-				return {};
-			}
-			if (!UpdateProcThreadAttribute(LPPROC_THREAD_ATTRIBUTE_LIST(AttributeList.GetData()), 0, PROC_THREAD_ATTRIBUTE_HANDLE_LIST, Handles, HandleCount * sizeof(HANDLE), nullptr, nullptr))
-			{
-				HandleError();
-				DeleteProcThreadAttributeList(LPPROC_THREAD_ATTRIBUTE_LIST(AttributeList.GetData()));
-				return {};
-			}
-
-			CreateFlags |= EXTENDED_STARTUPINFO_PRESENT;
-			dwFlags |= STARTF_USESTDHANDLES;
-		}
-	}
-
-	// initialize startup info
-	STARTUPINFOEX StartupInfoEx
-	{
-		{
-			sizeof(STARTUPINFOEX),
-			NULL, NULL, NULL,
-			(DWORD)CW_USEDEFAULT,
-			(DWORD)CW_USEDEFAULT,
-			(DWORD)CW_USEDEFAULT,
-			(DWORD)CW_USEDEFAULT,
-			(DWORD)0, (DWORD)0, (DWORD)0,
-			(DWORD)dwFlags,
-			ShowWindowFlags,
-			0, NULL,
-			NULL,
-			HANDLE(PipeWriteChild),
-			HANDLE(PipeWriteChild)
-		},
-		LPPROC_THREAD_ATTRIBUTE_LIST(AttributeList.GetData())
-	};
-
-	// create the child process
-	FString CommandLine = FString::Printf(TEXT("\"%s\" %s"), URL, Parms);
-	PROCESS_INFORMATION ProcInfo;
-
-	if (!CreateProcess(NULL, CommandLine.GetCharArray().GetData(), nullptr, nullptr, !AttributeList.IsEmpty(), (DWORD)CreateFlags, NULL, OptionalWorkingDirectory, &StartupInfoEx.StartupInfo, &ProcInfo))
-	{
-		HandleError();
-		if (!AttributeList.IsEmpty())
-		{
-			DeleteProcThreadAttributeList(LPPROC_THREAD_ATTRIBUTE_LIST(AttributeList.GetData()));
-		}
-		return {};
-	}
-
-	if (!AttributeList.IsEmpty())
-	{
-		DeleteProcThreadAttributeList(LPPROC_THREAD_ATTRIBUTE_LIST(AttributeList.GetData()));
-	}
-	
-	if (OutProcessID != nullptr)
-	{
-		*OutProcessID = ProcInfo.dwProcessId;
-	}
-
-	::CloseHandle(ProcInfo.hThread);
-
-	return FProcHandle(ProcInfo.hProcess);
-}
-#include "Windows/HideWindowsPlatformTypes.h"
-#else
-static FProcHandle CreateProc(const TCHAR* URL, const TCHAR* Parms, bool bLaunchDetached, bool bLaunchHidden, bool bLaunchReallyHidden, uint32* OutProcessID, int32 PriorityModifier, const TCHAR* OptionalWorkingDirectory, void* PipeWriteChild)
-{
-	return FPlatformProcess::CreateProc(URL, Parms, bLaunchDetached, bLaunchHidden, bLaunchReallyHidden, OutProcessID, PriorityModifier, OptionalWorkingDirectory, PipeWriteChild);
-}
-#endif
-
 #define LOCTEXT_NAMESPACE "FAndroidDeviceDetectionModule" 
 
 DEFINE_LOG_CATEGORY_STATIC(AndroidDeviceDetectionLog, Log, All);
@@ -232,7 +84,7 @@ public:
 
 	virtual void Stop()
 	{
-		bStopRequested.store(true, std::memory_order_seq_cst);
+		bStopRequested.store(true, std::memory_order_relaxed);
 		StopEvent.Notify();
 	}
 
@@ -297,52 +149,47 @@ private:
 			return false;
 		}
 
-		void* ReadPipe;
-		void* WritePipe;
-		if (!FPlatformProcess::CreatePipe(ReadPipe, WritePipe))
+		UE::HAL::FInputPipe InputPipe = UE::HAL::NewPipe;
+		if (!InputPipe)
 		{
 			return false;
 		}
-
-		int32 ReturnCode = -1;
 		
+		UE::HAL::FProcess Process{{.Uri = *ADBPath, .Arguments = *CommandLine, .bDetached = true, .bHidden = true, .StdOut = InputPipe}};
+		if (!Process)
+		{
+			return false;
+		}
+		
+		int32 ExitCode = -1;
+		
+		for (UE::FMonotonicTimePoint MaxTime = UE::FMonotonicTimePoint::Now() + UE::FMonotonicTimeSpan::FromSeconds(10.0);;)
 		{
-			FProcHandle Process = CreateProc(*ADBPath, *CommandLine, true, true, false, nullptr, 0, nullptr, WritePipe);
-			if (Process.IsValid())
+			bool bIsRunning = Process.IsRunning();
+			
+			if (OutStdOut)
 			{
-				for (UE::FMonotonicTimePoint MaxTime = UE::FMonotonicTimePoint::Now() + UE::FMonotonicTimeSpan::FromSeconds(10.0);;)
-				{
-					bool bIsRunning = FPlatformProcess::IsProcRunning(Process);
-					
-					if (OutStdOut)
-					{
-						*OutStdOut += FPlatformProcess::ReadPipe(ReadPipe);
-					}
-
-					if (!bIsRunning)
-					{
-						verify(FPlatformProcess::GetProcReturnCode(Process, &ReturnCode));
-						break;
-					}
+				*OutStdOut += InputPipe.Read();
+			}
 
-					FPlatformProcess::Sleep(0.0);
+			if (!bIsRunning)
+			{
+				ExitCode = *Process.GetExitCode();
+				break;
+			}
 
-					if (bStopRequested.load(std::memory_order_relaxed) || UE::FMonotonicTimePoint::Now() >= MaxTime)
-					{
-						FPlatformProcess::TerminateProc(Process);
-						break;
-					}
-				}
+			FPlatformProcess::YieldThread();
 
-				FPlatformProcess::CloseProc(Process);
+			if (bStopRequested.load(std::memory_order_relaxed) || UE::FMonotonicTimePoint::Now() >= MaxTime)
+			{
+				Process.Kill();
+				break;
 			}
 		}
-		
-		FPlatformProcess::ClosePipe(ReadPipe, WritePipe);
-		
-		if (ReturnCode != 0)
+
+		if (ExitCode != 0)
 		{
-			FPlatformMisc::LowLevelOutputDebugStringf(TEXT("The Android SDK command '%s' failed to run. Return code: %d\n"), *CommandLine, ReturnCode);
+			FPlatformMisc::LowLevelOutputDebugStringf(TEXT("The Android SDK command '%s' failed to run. Exit Code: %d\n"), *CommandLine, ExitCode);
 			return false;
 		}
 
diff --git a/Engine/Source/Developer/Android/AndroidTargetPlatformControls/Private/AndroidTargetDevice.inl b/Engine/Source/Developer/Android/AndroidTargetPlatformControls/Private/AndroidTargetDevice.inl
--- a/Engine/Source/Developer/Android/AndroidTargetPlatformControls/Private/AndroidTargetDevice.inl
+++ b/Engine/Source/Developer/Android/AndroidTargetPlatformControls/Private/AndroidTargetDevice.inl
@@ -48,18 +48,7 @@ inline bool FAndroidTargetDevice::PowerOn()
 
 	if (FString AndroidHomeDirectory = FPlatformMisc::GetEnvironmentVariable(TEXT("ANDROID_HOME")); !AndroidHomeDirectory.IsEmpty())
 	{
-#if PLATFORM_WINDOWS
-		bool bLaunchDetached = false;
-#else
-		bool bLaunchDetached = true;
-#endif
-		FProcHandle Process = FPlatformProcess::CreateProc(*FPaths::Combine(AndroidHomeDirectory, TEXT("emulator"), TEXT("emulator")), *(TEXT("-avd ") + DeviceId.RightChop(4)), bLaunchDetached, false, true, nullptr, 0, nullptr, nullptr);
-
-		if (Process.IsValid())
-		{
-			FPlatformProcess::CloseProc(Process);
-			return true;
-		}
+		return bool{UE::HAL::FProcess{{.Uri = *FPaths::Combine(AndroidHomeDirectory, TEXT("emulator"), TEXT("emulator")), .Arguments = *(TEXT("-avd ") + DeviceId.RightChop(4)), .bDetached = true, .bHidden = true}}};
 	}
 
 	return false;
diff --git a/Engine/Source/Developer/DesktopPlatform/Private/DesktopPlatformBase.cpp b/Engine/Source/Developer/DesktopPlatform/Private/DesktopPlatformBase.cpp
--- a/Engine/Source/Developer/DesktopPlatform/Private/DesktopPlatformBase.cpp
+++ b/Engine/Source/Developer/DesktopPlatform/Private/DesktopPlatformBase.cpp
@@ -740,7 +740,7 @@ FProcHandle FDesktopPlatformBase::InvokeUnrealBuildToolAsync(const FString& InCm
 	const bool bLaunchHidden = true;
 	const bool bLaunchReallyHidden = bLaunchHidden;
 
-	FProcHandle ProcHandle = FPlatformProcess::CreateProc(*ExecutableFileName, *CmdLineParams, bLaunchDetached, bLaunchHidden, bLaunchReallyHidden, NULL, 0, NULL, OutWritePipe, OutReadPipe);
+	FProcHandle ProcHandle = FPlatformProcess::CreateProc(*ExecutableFileName, *CmdLineParams, bLaunchDetached, bLaunchHidden, bLaunchReallyHidden, NULL, 0, NULL, OutWritePipe);
 	if (!ProcHandle.IsValid())
 	{
 		UE_LOG(LogDesktopPlatform, Warning, TEXT("Failed to launch UnrealBuildTool (exe path is '%s')"), *ExecutableFileName);
@@ -1269,7 +1269,7 @@ FProcHandle FDesktopPlatformBase::InvokeOidcTokenToolAsync(const FString& InArgu
 	const bool bLaunchHidden = true;
 	const bool bLaunchReallyHidden = bLaunchHidden;
 
-	FProcHandle ProcHandle = FPlatformProcess::CreateProc(*ExecutableFileName, *CmdLineParams, bLaunchDetached, bLaunchHidden, bLaunchReallyHidden, NULL, 0, FPlatformProcess::BaseDir(), OutWritePipe, OutReadPipe);
+	FProcHandle ProcHandle = FPlatformProcess::CreateProc(*ExecutableFileName, *CmdLineParams, bLaunchDetached, bLaunchHidden, bLaunchReallyHidden, NULL, 0, FPlatformProcess::BaseDir(), OutWritePipe);
 	if (!ProcHandle.IsValid())
 	{
 		UE_LOG(LogDesktopPlatform, Warning, TEXT("Failed to launch OidcToken (exe path is '%s')"), *ExecutableFileName);
@@ -1966,10 +1966,11 @@ bool FDesktopPlatformBase::BuildUnrealBuildTool(const FString& RootDir, FOutputD
 
 	verify(FPlatformProcess::CreatePipe(PipeRead, PipeWrite));
 
-	FProcHandle ProcHandle = FPlatformProcess::CreateProc(*CompilerExecutableFilename, *CmdLineParams, bLaunchDetached, bLaunchHidden, bLaunchReallyHidden, NULL, 0, NULL, PipeWrite, PipeRead);
+	FProcHandle ProcHandle = FPlatformProcess::CreateProc(*CompilerExecutableFilename, *CmdLineParams, bLaunchDetached, bLaunchHidden, bLaunchReallyHidden, NULL, 0, NULL, PipeWrite);
 	if (!ProcHandle.IsValid())
 	{
 		Ar.Log(TEXT("Failed to start process."));
+		FPlatformProcess::ClosePipe(PipeRead, PipeWrite);
 		return false;
 	}
 	FPlatformProcess::WaitForProc(ProcHandle);
diff --git a/Engine/Source/Developer/LauncherServices/Private/Launcher/LauncherUATTask.h b/Engine/Source/Developer/LauncherServices/Private/Launcher/LauncherUATTask.h
--- a/Engine/Source/Developer/LauncherServices/Private/Launcher/LauncherUATTask.h
+++ b/Engine/Source/Developer/LauncherServices/Private/Launcher/LauncherUATTask.h
@@ -92,7 +92,7 @@ protected:
 
 
 		// launch UAT and monitor its progress
-		ProcessHandle = FPlatformProcess::CreateProc(*(ExecutablePath / Executable), *UATCommandLine, false, true, true, NULL, 0, *ExecutablePath, WritePipe, ReadPipe);
+		ProcessHandle = FPlatformProcess::CreateProc(*(ExecutablePath / Executable), *UATCommandLine, false, true, true, NULL, 0, *ExecutablePath, WritePipe);
 
 		while (FPlatformProcess::IsProcRunning(ProcessHandle) && !EndTextFound)
 		{
diff --git a/Engine/Source/Editor/CSVtoSVG/Private/CSVtoSVGModule.cpp b/Engine/Source/Editor/CSVtoSVG/Private/CSVtoSVGModule.cpp
--- a/Engine/Source/Editor/CSVtoSVG/Private/CSVtoSVGModule.cpp
+++ b/Engine/Source/Editor/CSVtoSVG/Private/CSVtoSVGModule.cpp
@@ -89,7 +89,7 @@ FFilePath GenerateSVG(const UCSVtoSVGArugments& Arguments, const TArray<FString>
 		return FFilePath();
 	}
 
-	FProcHandle ProcHandle = FPlatformProcess::CreateProc(*CSVtoSVGExe, *CmdLine, true, false, false, nullptr, 0, nullptr, PipeWrite, PipeRead);
+	FProcHandle ProcHandle = FPlatformProcess::CreateProc(*CSVtoSVGExe, *CmdLine, true, false, false, nullptr, 0, nullptr, PipeWrite);
 	if (ProcHandle.IsValid())
 	{
 		FString CSVtoCSVOutput;
@@ -115,10 +115,11 @@ FFilePath GenerateSVG(const UCSVtoSVGArugments& Arguments, const TArray<FString>
 			}
 		}
 
-		FPlatformProcess::ClosePipe(PipeRead, PipeWrite);
 		FPlatformProcess::CloseProc(ProcHandle);
 	}
 
+	FPlatformProcess::ClosePipe(PipeRead, PipeWrite);
+
 	return FFilePath();
 }
 
diff --git a/Engine/Source/Programs/CrashReportClient/Private/Windows/CrashReportClientMainWindows.cpp b/Engine/Source/Programs/CrashReportClient/Private/Windows/CrashReportClientMainWindows.cpp
--- a/Engine/Source/Programs/CrashReportClient/Private/Windows/CrashReportClientMainWindows.cpp
+++ b/Engine/Source/Programs/CrashReportClient/Private/Windows/CrashReportClientMainWindows.cpp
@@ -31,6 +31,7 @@ void CopyDiagnosticFilesToClipboard(TConstArrayView<FString> Files)
 		GlobalMem = GlobalAlloc( GMEM_MOVEABLE, RequiredSize );
 		check(GlobalMem);
 		uint8* Data = (uint8*) GlobalLock( GlobalMem );
+		check(Data);
 		DROPFILES* Drop = (DROPFILES*)Data;
 		Drop->pFiles = sizeof(DROPFILES);
 		Drop->fWide = 1;
@@ -134,29 +135,30 @@ int WINAPI WinMain(_In_ HINSTANCE hInInstance, _In_opt_ HINSTANCE hPrevInstance,
 			GetModuleFileName(NULL, RespawnExePathname, MAX_PATH);
 			FString RespawnExeArguments(Arguments);
 			RespawnExeArguments.Append(" -RespawnedInstance");
-			uint32 RespawnPid = 0;
+			UE::HAL::EProcessId RespawnPid;
 
 			// Respawn itself to sever the process grouping with the Editor. If the user kills the Editor process group in task manager,
 			// CRC will not die at the same time, will be able to capture the Editor exit code and send the MTBF report to correctly
 			// identify the Editor 'AbnormalShutdown' as 'Killed' instead.
-			FProcHandle Handle = FPlatformProcess::CreateProc(
-				RespawnExePathname,
-				*RespawnExeArguments,
-				true, false, false,
-				&RespawnPid, 0,
-				nullptr,
-				reinterpret_cast<void*>(ChildPipe), // Ensure the child process inherit this pipe handle that was previously inherited from its parent.
-				nullptr);
-
-			if (Handle.IsValid())
+			UE::HAL::FProcess Process
+			{
+				{
+					.Uri = RespawnExePathname,
+					.Arguments = *RespawnExeArguments,
+					.bDetached = true,
+					.bInheritHandles = true
+				},
+				RespawnPid
+			};
+
+			if (Process)
 			{
 				FString PidPathname = FString::Printf(TEXT("%sue-crc-pid-%d"), FPlatformProcess::UserTempDir(), MonitoredEditorPid);
 				if (TUniquePtr<FArchive> Ar = TUniquePtr<FArchive>(IFileManager::Get().CreateFileWriter(*PidPathname, FILEWRITE_EvenIfReadOnly)))
 				{
-					*Ar << RespawnPid;
+					std::underlying_type_t<UE::HAL::EProcessId> Pid = std::underlying_type_t<UE::HAL::EProcessId>(RespawnPid);
+					*Ar << Pid;
 				}
-				
-				FPlatformProcess::CloseProc(Handle);
 			}
 		}
 		RequestEngineExit(TEXT("Respawn instance."));
diff --git a/Engine/Source/Programs/GPUReshape/Source/GPUReshapeBootstrapper/Private/GPUReshapeBootstrapper.cpp b/Engine/Source/Programs/GPUReshape/Source/GPUReshapeBootstrapper/Private/GPUReshapeBootstrapper.cpp
--- a/Engine/Source/Programs/GPUReshape/Source/GPUReshapeBootstrapper/Private/GPUReshapeBootstrapper.cpp
+++ b/Engine/Source/Programs/GPUReshape/Source/GPUReshapeBootstrapper/Private/GPUReshapeBootstrapper.cpp
@@ -117,8 +117,7 @@ INT32_MAIN_INT32_ARGC_TCHAR_ARGV()
 		nullptr,
 		0,
 		nullptr,
-		PipeWrite,
-		PipeRead
+		PipeWrite
 	);
 	
 	if (!Handle.IsValid())
diff --git a/Engine/Source/Programs/UGSCore/Perforce.cpp b/Engine/Source/Programs/UGSCore/Perforce.cpp
--- a/Engine/Source/Programs/UGSCore/Perforce.cpp
+++ b/Engine/Source/Programs/UGSCore/Perforce.cpp
@@ -100,7 +100,7 @@ struct FPerforceExe
 		const bool bWriting = !WriteChildText.IsEmpty();
 		if (bWriting)
 		{
-			FPlatformProcess::CreatePipe(ChildReadPipe, ParentWritePipe);
+			FPlatformProcess::CreatePipe(ChildReadPipe, ParentWritePipe, true);
 		}
 
 		FProcHandle P4Proc = FPlatformProcess::CreateProc(*GetPerforceExe(), *CommandLine, false, true, true, nullptr, 0, nullptr, ChildWritePipe, ChildReadPipe);
@@ -125,6 +125,8 @@ struct FPerforceExe
 			if (AbortEvent->Wait(FTimespan::Zero()))
 			{
 				FPlatformProcess::TerminateProc(P4Proc);
+				FPlatformProcess::ClosePipe(ParentReadPipe, ChildWritePipe);
+				FPlatformProcess::CloseProc(P4Proc);
 				return -1;
 			}
 		}
@@ -137,6 +139,9 @@ struct FPerforceExe
 		int ExitCode = -1;
 
 		const bool bGotReturnCode = FPlatformProcess::GetProcReturnCode(P4Proc, &ExitCode);
+		
+		FPlatformProcess::CloseProc(P4Proc);
+		
 		if (bGotReturnCode)
 		{
 			return ExitCode;
diff --git a/Engine/Source/Programs/UGSCore/Utility.cpp b/Engine/Source/Programs/UGSCore/Utility.cpp
--- a/Engine/Source/Programs/UGSCore/Utility.cpp
+++ b/Engine/Source/Programs/UGSCore/Utility.cpp
@@ -118,7 +118,7 @@ int FUtility::ExecuteProcess(const TCHAR* FileName, const TCHAR* CommandLine, co
 	void* WritePipe = nullptr;
 	FPlatformProcess::CreatePipe(ReadPipe, WritePipe);
 
-	FProcHandle Proc = FPlatformProcess::CreateProc(FileName, CommandLine, false, true, true, &ProcId, 0, nullptr, WritePipe, ReadPipe);
+	FProcHandle Proc = FPlatformProcess::CreateProc(FileName, CommandLine, false, true, true, &ProcId, 0, nullptr, WritePipe);
 
 	FString Output;
 	FString LatestOutput = FPlatformProcess::ReadPipe(ReadPipe);
@@ -130,6 +130,8 @@ int FUtility::ExecuteProcess(const TCHAR* FileName, const TCHAR* CommandLine, co
 		if (AbortEvent->Wait(FTimespan::Zero()))
 		{
 			FPlatformProcess::TerminateProc(Proc);
+			FPlatformProcess::ClosePipe(ReadPipe, WritePipe);
+			FPlatformProcess::CloseProc(Proc);
 			return -1;
 		}
 	}
@@ -142,6 +144,8 @@ int FUtility::ExecuteProcess(const TCHAR* FileName, const TCHAR* CommandLine, co
 	int ExitCode = -1;
 	bool GotReturnCode = FPlatformProcess::GetProcReturnCode(Proc, &ExitCode);
 
+	FPlatformProcess::CloseProc(Proc);
+	
 	if (GotReturnCode)
 	{
 		return ExitCode;
@@ -157,7 +161,7 @@ int FUtility::ExecuteProcess(const TCHAR* FileName, const TCHAR* CommandLine, co
 	void* WritePipe = nullptr;
 	FPlatformProcess::CreatePipe(ReadPipe, WritePipe);
 
-	FProcHandle Proc = FPlatformProcess::CreateProc(FileName, CommandLine, false, true, true, &ProcId, 0, nullptr, WritePipe, ReadPipe);
+	FProcHandle Proc = FPlatformProcess::CreateProc(FileName, CommandLine, false, true, true, &ProcId, 0, nullptr, WritePipe);
 
 	FString LatestOutput = FPlatformProcess::ReadPipe(ReadPipe);
 	while (FPlatformProcess::IsProcRunning(Proc) || !LatestOutput.IsEmpty())
@@ -189,6 +193,8 @@ int FUtility::ExecuteProcess(const TCHAR* FileName, const TCHAR* CommandLine, co
 		if (AbortEvent->Wait(FTimespan::Zero()))
 		{
 			FPlatformProcess::TerminateProc(Proc);
+			FPlatformProcess::ClosePipe(ReadPipe, WritePipe);
+			FPlatformProcess::CloseProc(Proc);
 			return -1;
 		}
 	}
@@ -198,6 +204,8 @@ int FUtility::ExecuteProcess(const TCHAR* FileName, const TCHAR* CommandLine, co
 	int ExitCode = -1;
 	bool GotReturnCode = FPlatformProcess::GetProcReturnCode(Proc, &ExitCode);
 
+	FPlatformProcess::CloseProc(Proc);
+
 	if (GotReturnCode)
 	{
 		return ExitCode;
diff --git a/Engine/Source/Runtime/Core/Private/GenericPlatform/GenericPlatformProcess.cpp b/Engine/Source/Runtime/Core/Private/GenericPlatform/GenericPlatformProcess.cpp
--- a/Engine/Source/Runtime/Core/Private/GenericPlatform/GenericPlatformProcess.cpp
+++ b/Engine/Source/Runtime/Core/Private/GenericPlatform/GenericPlatformProcess.cpp
@@ -288,6 +288,12 @@ FProcHandle FGenericPlatformProcess::CreateProc( const TCHAR* URL, const TCHAR*
 	return FProcHandle();
 }
 
+TTuple<FProcHandle, UE::HAL::EProcessId> FGenericPlatformProcess::CreateProc(const UE::HAL::FProcessStartInfo& StartInfo)
+{
+	UE_LOG(LogHAL, Fatal, TEXT("FGenericPlatformProcess::CreateProc not implemented on this platform"));
+	return {};
+}
+
 FProcHandle FGenericPlatformProcess::OpenProcess(uint32 ProcessID)
 {
 	UE_LOG(LogHAL, Fatal, TEXT("FGenericPlatformProcess::OpenProcess not implemented on this platform"));
diff --git a/Engine/Source/Runtime/Core/Private/Mac/MacPlatformProcess.cpp b/Engine/Source/Runtime/Core/Private/Mac/MacPlatformProcess.cpp
--- a/Engine/Source/Runtime/Core/Private/Mac/MacPlatformProcess.cpp
+++ b/Engine/Source/Runtime/Core/Private/Mac/MacPlatformProcess.cpp
@@ -205,7 +205,7 @@ bool FMacPlatformProcess::ExecProcess( const TCHAR* URL, const TCHAR* Params, in
 	const bool bLaunchHidden = false;
 	const bool bLaunchReallyHidden = bLaunchHidden;
 
-	FProcHandle ProcHandle = FPlatformProcess::CreateProcInternal(*ExecutableFileName, *CmdLineParams, bLaunchDetached, bLaunchHidden, bLaunchReallyHidden, NULL, 0, OptionalWorkingDirectory, PipeStdOutWrite, PipeStdErrWrite, nullptr);
+	FProcHandle ProcHandle = FPlatformProcess::CreateProcInternal(*ExecutableFileName, *CmdLineParams, bLaunchDetached, bLaunchHidden, bLaunchReallyHidden, NULL, 0, OptionalWorkingDirectory, PipeStdOutWrite, nullptr, PipeStdErrWrite, false);
 	if (ProcHandle.IsValid())
 	{
 		while (FPlatformProcess::IsProcRunning(ProcHandle))
@@ -288,18 +288,29 @@ bool FMacPlatformProcess::ExecProcess( const TCHAR* URL, const TCHAR* Params, in
 	return bInvoked;
 }
 
-FProcHandle FMacPlatformProcess::CreateProc( const TCHAR* URL, const TCHAR* Parms, bool bLaunchDetached, bool bLaunchHidden, bool bLaunchReallyHidden, uint32* OutProcessID, int32 PriorityModifier, const TCHAR* OptionalWorkingDirectory, void* PipeWriteChild, void* PipeReadChild)
+FProcHandle FMacPlatformProcess::CreateProc(const TCHAR* URL, const TCHAR* Parms, bool bLaunchDetached, bool bLaunchHidden, bool bLaunchReallyHidden, uint32* OutProcessID, int32 PriorityModifier, const TCHAR* OptionalWorkingDirectory, void* PipeWriteChild, void* PipeReadChild)
 {
-	return CreateProcInternal(URL, Parms, bLaunchDetached, bLaunchHidden, bLaunchReallyHidden, OutProcessID, PriorityModifier, OptionalWorkingDirectory, PipeWriteChild, PipeWriteChild, PipeReadChild);
+	return CreateProcInternal(URL, Parms, bLaunchDetached, bLaunchHidden, bLaunchReallyHidden, OutProcessID, PriorityModifier, OptionalWorkingDirectory, PipeWriteChild, PipeReadChild, PipeWriteChild, false);
 }
 
-FProcHandle FMacPlatformProcess::CreateProc( const TCHAR* URL, const TCHAR* Parms, bool bLaunchDetached, bool bLaunchHidden, bool bLaunchReallyHidden, uint32* OutProcessID, int32 PriorityModifier, const TCHAR* OptionalWorkingDirectory, void* PipeWriteChild, void* PipeReadChild, void* PipeStdErrChild)
+FProcHandle FMacPlatformProcess::CreateProc(const TCHAR* URL, const TCHAR* Parms, bool bLaunchDetached, bool bLaunchHidden, bool bLaunchReallyHidden, uint32* OutProcessID, int32 PriorityModifier, const TCHAR* OptionalWorkingDirectory, void* PipeWriteChild, void* PipeReadChild, void* PipeStdErrChild)
 {
-	return CreateProcInternal(URL, Parms, bLaunchDetached, bLaunchHidden, bLaunchReallyHidden, OutProcessID, PriorityModifier, OptionalWorkingDirectory, PipeWriteChild, PipeStdErrChild, PipeReadChild);
+	return CreateProcInternal(URL, Parms, bLaunchDetached, bLaunchHidden, bLaunchReallyHidden, OutProcessID, PriorityModifier, OptionalWorkingDirectory, PipeWriteChild, PipeReadChild, PipeStdErrChild, false);
 }
 
-FProcHandle FMacPlatformProcess::CreateProcInternal(const TCHAR* URL, const TCHAR* Parms, bool bLaunchDetached, bool bLaunchHidden, bool bLaunchReallyHidden, uint32* OutProcessID, int32 PriorityModifier, const TCHAR* OptionalWorkingDirectory, void* PipeStdOutChild, void* PipeStdErrChild, void *PipeStdInChild)
+TTuple<FProcHandle, UE::HAL::EProcessId> FMacPlatformProcess::CreateProc(const UE::HAL::FProcessStartInfo& StartInfo)
 {
+	uint32 Id;
+	FProcHandle Handle = CreateProcInternal(StartInfo.Uri, StartInfo.Arguments, StartInfo.bDetached, false, StartInfo.bHidden, &Id, StartInfo.PriorityModifier, StartInfo.WorkingDirectory, StartInfo.StdOut, StartInfo.StdIn, StartInfo.StdErr, StartInfo.bInheritHandles);
+	return {Handle, UE::HAL::EProcessId{Id}};
+}
+
+FProcHandle FMacPlatformProcess::CreateProcInternal(const TCHAR* URL, const TCHAR* Parms, bool bLaunchDetached, bool bLaunchHidden, bool bLaunchReallyHidden, uint32* OutProcessID, int32 PriorityModifier, const TCHAR* OptionalWorkingDirectory, void* PipeStdOutChild, void* PipeStdInChild, void* PipeStdErrChild, bool bInheritHandles)
+{
+	check(URL);
+	check(URL[0] != '\"');
+	check(Parms);
+
 	SCOPED_AUTORELEASE_POOL;
 
 	// @TODO bLaunchHidden bLaunchReallyHidden are not handled
@@ -465,7 +476,7 @@ FProcHandle FMacPlatformProcess::CreateProcInternal(const TCHAR* URL, const TCHA
 
 	posix_spawnattr_t SpawnAttr;
 	posix_spawnattr_init(&SpawnAttr);
-	short int SpawnFlags = 0;
+	short int SpawnFlags = !bInheritHandles ? POSIX_SPAWN_CLOEXEC_DEFAULT : 0;
 
 	// Makes spawned processes have its own unique group id so we can kill the entire group without killing the parent
 	SpawnFlags |= POSIX_SPAWN_SETPGROUP;
@@ -513,16 +524,49 @@ FProcHandle FMacPlatformProcess::CreateProcInternal(const TCHAR* URL, const TCHA
 	{
 		posix_spawn_file_actions_adddup2(&FileActions, [(NSFileHandle*)PipeStdOutChild fileDescriptor], STDOUT_FILENO);
 	}
+	else if (!bInheritHandles)
+	{
+		if (bLaunchDetached)
+		{
+			posix_spawn_file_actions_addopen(&FileActions, STDOUT_FILENO, "/dev/null", O_RDWR, 0);
+		}
+		else
+		{
+			posix_spawn_file_actions_addinherit_np(&FileActions, STDOUT_FILENO);
+		}
+	}
 
 	if (PipeStdErrChild)
 	{
 		posix_spawn_file_actions_adddup2(&FileActions, [(NSFileHandle*)PipeStdErrChild fileDescriptor], STDERR_FILENO);
 	}
+	else if (!bInheritHandles)
+	{
+		if (bLaunchDetached)
+		{
+			posix_spawn_file_actions_addopen(&FileActions, STDERR_FILENO, "/dev/null", O_RDWR, 0);
+		}
+		else
+		{
+			posix_spawn_file_actions_addinherit_np(&FileActions, STDERR_FILENO);
+		}
+	}
 
 	if (PipeStdInChild)
 	{
 		posix_spawn_file_actions_adddup2(&FileActions, [(NSFileHandle*)PipeStdInChild fileDescriptor], STDIN_FILENO);
 	}
+	else if (!bInheritHandles)
+	{
+		if (bLaunchDetached)
+		{
+			posix_spawn_file_actions_addopen(&FileActions, STDIN_FILENO, "/dev/null", O_RDONLY, 0);
+		}
+		else
+		{
+			posix_spawn_file_actions_addinherit_np(&FileActions, STDIN_FILENO);
+		}
+	}
 
 	if (OptionalWorkingDirectory)
 	{
@@ -1279,8 +1323,15 @@ bool FMacPlatformProcess::CreatePipe( void*& ReadPipe, void*& WritePipe, bool bW
 	int pipefd[2];
 	pipe(pipefd);
 
-	// The read pipe should be non-blocking, but the write pipe SHOULD block
-	fcntl(pipefd[0], F_SETFL, O_NONBLOCK);
+	if (bWritePipeLocal)
+	{
+		fcntl(pipefd[1], F_SETFL, O_CLOEXEC);
+	}
+	else
+	{
+		// The read pipe should be non-blocking, but the write pipe SHOULD block
+		fcntl(pipefd[0], F_SETFL, O_CLOEXEC | O_NONBLOCK);
+	}
 
 	// create an NSFileHandle from the descriptor
 	ReadPipe = [[NSFileHandle alloc] initWithFileDescriptor: pipefd[0]];
diff --git a/Engine/Source/Runtime/Core/Private/Misc/FeedbackContextMarkup.cpp b/Engine/Source/Runtime/Core/Private/Misc/FeedbackContextMarkup.cpp
--- a/Engine/Source/Runtime/Core/Private/Misc/FeedbackContextMarkup.cpp
+++ b/Engine/Source/Runtime/Core/Private/Misc/FeedbackContextMarkup.cpp
@@ -54,7 +54,7 @@ bool FFeedbackContextMarkup::PipeProcessOutput(const FText& Description, const F
 	Warn->BeginSlowTask(Description, true, true);
 
 	// Create the process
-	FProcHandle ProcessHandle = FPlatformProcess::CreateProc(*URL, *Params, false, true, true, NULL, 0, NULL, PipeWrite, PipeRead);
+	FProcHandle ProcessHandle = FPlatformProcess::CreateProc(*URL, *Params, false, true, true, NULL, 0, NULL, PipeWrite);
 	if(ProcessHandle.IsValid())
 	{
 		FString BufferedText;
@@ -99,8 +99,7 @@ bool FFeedbackContextMarkup::PipeProcessOutput(const FText& Description, const F
 	Warn->EndSlowTask();
 
 	// Close the pipes
-	FPlatformProcess::ClosePipe(0, PipeRead);
-	FPlatformProcess::ClosePipe(0, PipeWrite);
+	FPlatformProcess::ClosePipe(PipeRead, PipeWrite);
 	return bRes;
 }
 
diff --git a/Engine/Source/Runtime/Core/Private/Misc/InteractiveProcess.cpp b/Engine/Source/Runtime/Core/Private/Misc/InteractiveProcess.cpp
--- a/Engine/Source/Runtime/Core/Private/Misc/InteractiveProcess.cpp
+++ b/Engine/Source/Runtime/Core/Private/Misc/InteractiveProcess.cpp
@@ -24,7 +24,7 @@ static FORCEINLINE bool CreatePipeWrite(void*& ReadPipe, void*& WritePipe)
 
 	return true;
 #else
-	return FPlatformProcess::CreatePipe(ReadPipe, WritePipe);
+	return FPlatformProcess::CreatePipe(ReadPipe, WritePipe, true);
 #endif // PLATFORM_WINDOWS
 }
 
diff --git a/Engine/Source/Runtime/Core/Private/Unix/UnixPlatformProcess.cpp b/Engine/Source/Runtime/Core/Private/Unix/UnixPlatformProcess.cpp
--- a/Engine/Source/Runtime/Core/Private/Unix/UnixPlatformProcess.cpp
+++ b/Engine/Source/Runtime/Core/Private/Unix/UnixPlatformProcess.cpp
@@ -818,7 +818,19 @@ FProcHandle FUnixPlatformProcess::CreateProc(const TCHAR* URL, const TCHAR* Parm
 {
 	// CreateProc used to only have a single "write" pipe argument, which Windows and Mac would pipe both stdout and stderr into.
 	// On Unix though, only stdout was piped to it, and stderr wasn't available at all, so we'll preserve that behaviour in this overload for compatibility with existing code
-	return CreateProc(URL, Parms, bLaunchDetached, bLaunchHidden, bLaunchReallyHidden, OutProcessID, PriorityModifier, OptionalWorkingDirectory, PipeWriteChild, PipeReadChild, PipeWriteChild);
+	return CreateProcInternal(URL, Parms, bLaunchDetached, bLaunchHidden, bLaunchReallyHidden, OutProcessID, PriorityModifier, OptionalWorkingDirectory, PipeWriteChild, PipeReadChild, PipeWriteChild, false);
+}
+
+FProcHandle FUnixPlatformProcess::CreateProc(const TCHAR* URL, const TCHAR* Parms, bool bLaunchDetached, bool bLaunchHidden, bool bLaunchReallyHidden, uint32* OutProcessID, int32 PriorityModifier, const TCHAR* OptionalWorkingDirectory, void* PipeWriteChild, void* PipeReadChild, void* PipeStdErrChild)
+{
+	return CreateProcInternal(URL, Parms, bLaunchDetached, bLaunchHidden, bLaunchReallyHidden, OutProcessID, PriorityModifier, OptionalWorkingDirectory, PipeWriteChild, PipeReadChild, PipeStdErrChild, false);
+}
+
+TTuple<FProcHandle, UE::HAL::EProcessId> FUnixPlatformProcess::CreateProc(const UE::HAL::FProcessStartInfo& StartInfo)
+{
+	uint32 Id;
+	FProcHandle Handle = CreateProcInternal(StartInfo.Uri, StartInfo.Arguments, StartInfo.bDetached, false, StartInfo.bHidden, &Id, StartInfo.PriorityModifier, StartInfo.WorkingDirectory, StartInfo.StdOut, StartInfo.StdIn, StartInfo.StdErr, StartInfo.bInheritHandles);
+	return {Handle, UE::HAL::EProcessId{Id}};
 }
 
 static char* MallocedUtf8FromString(const FString& Str)
@@ -896,7 +908,7 @@ static bool ParseCmdLineToken(const TCHAR* token, bool& OutIsInString, bool& Out
 	return false;
 }
 
-FProcHandle FUnixPlatformProcess::CreateProc(const TCHAR* URL, const TCHAR* Parms, bool bLaunchDetached, bool bLaunchHidden, bool bLaunchReallyHidden, uint32* OutProcessID, int32 PriorityModifier, const TCHAR* OptionalWorkingDirectory, void* PipeWriteChild, void* PipeReadChild, void* PipeStdErrChild)
+FProcHandle FUnixPlatformProcess::CreateProcInternal(const TCHAR* URL, const TCHAR* Parms, bool bLaunchDetached, bool bLaunchHidden, bool bLaunchReallyHidden, uint32* OutProcessID, int32 PriorityModifier, const TCHAR* OptionalWorkingDirectory, void* PipeWriteChild, void* PipeReadChild, void* PipeStdErrChild, bool bInheritHandles)
 {
 	// @TODO bLaunchHidden bLaunchReallyHidden are not handled
 
diff --git a/Engine/Source/Runtime/Core/Private/Windows/WindowsPlatformCrashContext.cpp b/Engine/Source/Runtime/Core/Private/Windows/WindowsPlatformCrashContext.cpp
--- a/Engine/Source/Runtime/Core/Private/Windows/WindowsPlatformCrashContext.cpp
+++ b/Engine/Source/Runtime/Core/Private/Windows/WindowsPlatformCrashContext.cpp
@@ -594,32 +594,39 @@ FProcHandle LaunchCrashReportClient(void** OutWritePipe, void** OutReadPipe, uin
 	// Launch the crash reporter if the client exists
 	if (CreateCrashReportClientPath(CrashReporterClientPath, CR_CLIENT_MAX_PATH_LEN))
 	{
-		Handle = FPlatformProcess::CreateProc(
-			CrashReporterClientPath,
-			CrashReporterClientArgs,
-			true, false, false,
-			OutCrashReportClientProcessId, 0,
-			nullptr,
-			PipeChildInRead, //Pass this to allow inherit handles in child proc
-			nullptr);
-
-		DWORD pid = GetProcessId(Handle.Get());
-		UE_LOG(LogWindows, Log, TEXT("Started CrashReportClient (pid=%d)"), pid);
+		UE::HAL::EProcessId ProcessId;
+		UE::HAL::FProcess Process
+		{
+			{
+				.Uri = CrashReporterClientPath,
+				.Arguments = CrashReporterClientArgs,
+				.bDetached = true,
+				.bInheritHandles = true
+			},
+			ProcessId
+		};
+		
+		if (OutCrashReportClientProcessId)
+		{
+			*OutCrashReportClientProcessId = std::underlying_type_t<UE::HAL::EProcessId>(ProcessId);
+		}
+		
+		UE_LOG(LogWindows, Log, TEXT("Started CrashReportClient (pid=%d)"), std::underlying_type_t<UE::HAL::EProcessId>(ProcessId));
 
 		// The CRC instance launched above will respanwn itself to sever the link with the Editor process group. This way, if the user kills the Editor
 		// process group in TaskManager, CRC will not die at the same moment and will be able to capture the Editor exit code and send the session summary.
-		if (Handle.IsValid())
+		if (Process)
 		{
 			if (FEngineBuildSettings::IsSourceDistribution())
 			{
 				// Don't wait longer than n seconds for CRC to respawn. This is a workaround for people that did not recompile CrashReportClientEditor after updating/recompiling the Engine/Editor.
 				// This is for people that had compiled CRC prior 4.25.0 (13410773) and recompiled the Editor only after updating around 4.25.1.
-				::WaitForSingleObject(Handle.Get(), 3000);
+				::WaitForSingleObject(Process.NativeHandle(), 3000);
 			}
 			else // Distributed binaries (CRC is expected to be prebuilt).
 			{
 				// Wait for the intermediate CRC to respawn itself and exit.
-				::WaitForSingleObject(Handle.Get(), INFINITE);
+				::WaitForSingleObject(Process.NativeHandle(), INFINITE);
 			}
 			
 			// The respanwned CRC process writes its own PID to a file named by this process PID (by parsing the -MONITOR={PID} arguments)
@@ -633,8 +640,7 @@ FProcHandle LaunchCrashReportClient(void** OutWritePipe, void** OutReadPipe, uin
 			// The file is not required anymore.
 			IFileManager::Get().Delete(*PidFilePathname, /*RequireExist*/false, /*EvenReadOnly*/true);
 
-			// Close the handle before reassigning it.
-			FPlatformProcess::CloseProc(Handle);
+			Process = nullptr;
 
 			// Acquire a handle on the final CRC instance responsible to handle the crash/reports, but forbid this process from terminating it in case we try to terminate it by accident (like a stomped handle that would terminate the wrong process)
 			// Use the minimum required access rights to avoid creating a SecuritySandbox bypass.
diff --git a/Engine/Source/Runtime/Core/Private/Windows/WindowsPlatformProcess.cpp b/Engine/Source/Runtime/Core/Private/Windows/WindowsPlatformProcess.cpp
--- a/Engine/Source/Runtime/Core/Private/Windows/WindowsPlatformProcess.cpp
+++ b/Engine/Source/Runtime/Core/Private/Windows/WindowsPlatformProcess.cpp
@@ -464,104 +464,231 @@ void FWindowsPlatformProcess::LaunchURLInternal(const TCHAR* URL, const TCHAR* P
 
 FProcHandle FWindowsPlatformProcess::CreateProc(const TCHAR* URL, const TCHAR* Parms, bool bLaunchDetached, bool bLaunchHidden, bool bLaunchReallyHidden, uint32* OutProcessID, int32 PriorityModifier, const TCHAR* OptionalWorkingDirectory, void* PipeWriteChild, void* PipeReadChild)
 {
+	if (PipeWriteChild)
+	{
+		DWORD Flags;
+		ensure(GetHandleInformation(HANDLE(PipeWriteChild), &Flags));
+		if (!(Flags & HANDLE_FLAG_INHERIT))
+		{
+			UE_LOG(LogWindows, Error, TEXT("PipeWriteChild passed to CreateProc is not inheritable."));
+			PipeWriteChild = nullptr;
+		}
+	}
+	if (PipeReadChild)
+	{
+		DWORD Flags;
+		ensure(GetHandleInformation(HANDLE(PipeReadChild), &Flags));
+		if (!(Flags & HANDLE_FLAG_INHERIT))
+		{
+			UE_LOG(LogWindows, Error, TEXT("PipeReadChild passed to CreateProc is not inheritable."));
+			PipeReadChild = nullptr;
+		}
+	}
 	// CreateProc used to only have a single "write" pipe argument that both stdout and stderr would be piped into on Windows,
 	// so for this overload we'll preserve that behaviour for compatibility with existing code
-	return CreateProc(URL, Parms, bLaunchDetached, bLaunchHidden, bLaunchReallyHidden, OutProcessID, PriorityModifier, OptionalWorkingDirectory, PipeWriteChild, PipeReadChild, PipeWriteChild);
+	return CreateProcInternal(URL, Parms, bLaunchDetached, bLaunchHidden, bLaunchReallyHidden, OutProcessID, PriorityModifier, OptionalWorkingDirectory, PipeWriteChild, PipeReadChild, PipeWriteChild, false);
 }
 
 FProcHandle FWindowsPlatformProcess::CreateProc(const TCHAR* URL, const TCHAR* Parms, bool bLaunchDetached, bool bLaunchHidden, bool bLaunchReallyHidden, uint32* OutProcessID, int32 PriorityModifier, const TCHAR* OptionalWorkingDirectory, void* PipeWriteChild, void* PipeReadChild, void* PipeStdErrChild)
 {
-	//UE_LOG(LogWindows, Log,  TEXT("CreateProc %s %s"), URL, Parms );
+	if (PipeWriteChild)
+	{
+		DWORD Flags;
+		ensure(GetHandleInformation(HANDLE(PipeWriteChild), &Flags));
+		if (!(Flags & HANDLE_FLAG_INHERIT))
+		{
+			UE_LOG(LogWindows, Error, TEXT("PipeWriteChild passed to CreateProc is not inheritable."));
+			PipeWriteChild = nullptr;
+		}
+	}
+	if (PipeReadChild)
+	{
+		DWORD Flags;
+		ensure(GetHandleInformation(HANDLE(PipeReadChild), &Flags));
+		if (!(Flags & HANDLE_FLAG_INHERIT))
+		{
+			UE_LOG(LogWindows, Error, TEXT("PipeReadChild passed to CreateProc is not inheritable."));
+			PipeReadChild = nullptr;
+		}
+	}
+	if (PipeStdErrChild)
+	{
+		DWORD Flags;
+		ensure(GetHandleInformation(HANDLE(PipeStdErrChild), &Flags));
+		if (!(Flags & HANDLE_FLAG_INHERIT))
+		{
+			UE_LOG(LogWindows, Error, TEXT("PipeStdErrChild passed to CreateProc is not inheritable."));
+			PipeStdErrChild = nullptr;
+		}
+	}
+	return CreateProcInternal(URL, Parms, bLaunchDetached, bLaunchHidden, bLaunchReallyHidden, OutProcessID, PriorityModifier, OptionalWorkingDirectory, PipeWriteChild, PipeReadChild, PipeStdErrChild, false);
+}
+
+TTuple<FProcHandle, UE::HAL::EProcessId> FWindowsPlatformProcess::CreateProc(const UE::HAL::FProcessStartInfo& StartInfo)
+{
+	uint32 Id;
+	FProcHandle Handle = CreateProcInternal(StartInfo.Uri, StartInfo.Arguments, StartInfo.bDetached, false, StartInfo.bHidden, &Id, StartInfo.PriorityModifier, StartInfo.WorkingDirectory, StartInfo.StdOut, StartInfo.StdIn, StartInfo.StdErr, StartInfo.bInheritHandles);
+	return {Handle, UE::HAL::EProcessId{Id}};
+}
 
-	// initialize process creation flags
-	uint32 CreateFlags = NORMAL_PRIORITY_CLASS;
+#include "Windows/AllowWindowsPlatformTypes.h"
+FProcHandle FWindowsPlatformProcess::CreateProcInternal(const TCHAR* URL, const TCHAR* Parms, bool bLaunchDetached, bool bLaunchHidden, bool bLaunchReallyHidden, uint32* OutProcessID, int32 PriorityModifier, const TCHAR* OptionalWorkingDirectory, void* PipeWriteChild, void* PipeReadChild, void* PipeStdErrChild, bool bInheritHandles)
+{
+	check(URL);
+	check(URL[0] != '\"');
+	check(Parms);
+	
+	auto HandleError = [&URL, &Parms, &OutProcessID]
+	{
+#if !UE_BUILD_SHIPPING
+		DWORD Error = GetLastError();
+		TCHAR ErrorMessage[512];
+		FWindowsPlatformMisc::GetSystemErrorMessage(ErrorMessage, 512, Error);
+
+		UE_LOG(LogWindows, Warning, TEXT("CreateProc failed: %s (0x%08x)"), ErrorMessage, Error);
+		if (Error == ERROR_NOT_ENOUGH_MEMORY || Error == ERROR_OUTOFMEMORY)
+		{
+			// These errors are common enough that we want some available memory information
+			FPlatformMemoryStats Stats = FPlatformMemory::GetStats();
+			UE_LOG(LogWindows, Warning, TEXT("Mem used: %.2f MB, OS Free %.2f MB"), Stats.UsedPhysical / 1048576.0f, Stats.AvailablePhysical / 1048576.0f);
+		}
+		UE_LOG(LogWindows, Warning, TEXT("URL: %s %s"), URL, Parms);
+		
+#endif
+		if (OutProcessID)
+		{
+			*OutProcessID = 0;
+		}
+	};
+	
+	// CREATE_NEW_PROCESS_GROUP for consistency with Unix-like platforms
+	DWORD CreationFlags = CREATE_NEW_PROCESS_GROUP;
+	
 	if (PriorityModifier < 0)
 	{
-		CreateFlags = (PriorityModifier == -1) ? BELOW_NORMAL_PRIORITY_CLASS : IDLE_PRIORITY_CLASS;
+		CreationFlags |= PriorityModifier == -1 ? BELOW_NORMAL_PRIORITY_CLASS : IDLE_PRIORITY_CLASS;
 	}
 	else if (PriorityModifier > 0)
 	{
-		CreateFlags = (PriorityModifier == 1) ? ABOVE_NORMAL_PRIORITY_CLASS : HIGH_PRIORITY_CLASS;
+		CreationFlags |= PriorityModifier == 1 ? ABOVE_NORMAL_PRIORITY_CLASS : HIGH_PRIORITY_CLASS;
+	}
+	else
+	{
+		CreationFlags |= NORMAL_PRIORITY_CLASS;
 	}
 
 	if (bLaunchDetached)
 	{
-		CreateFlags |= DETACHED_PROCESS;
+		// CREATE_NO_WINDOW because DETACHED_PROCESS can cause console windows to appear if the child spawns a process
+		// https://ikriv.com/dev/cpp/ConsoleProxy/flags has the best explanation
+		CreationFlags |= bLaunchReallyHidden ? CREATE_NO_WINDOW : DETACHED_PROCESS;
 	}
 
-	// initialize window flags
-	uint32 dwFlags = 0;
-	uint16 ShowWindowFlags = SW_HIDE;
+	DWORD Flags;
+	WORD ShowWindow;
 	if (bLaunchReallyHidden)
 	{
-		dwFlags = STARTF_USESHOWWINDOW;
+		Flags = STARTF_USESHOWWINDOW;
+		ShowWindow = SW_HIDE;
 	}
 	else if (bLaunchHidden)
 	{
-		dwFlags = STARTF_USESHOWWINDOW;
-		ShowWindowFlags = SW_SHOWMINNOACTIVE;
+		Flags = STARTF_USESHOWWINDOW;
+		ShowWindow = SW_SHOWMINNOACTIVE;
 	}
-
-	if (PipeWriteChild != nullptr || PipeReadChild != nullptr || PipeStdErrChild != nullptr)
+	else
 	{
-		dwFlags |= STARTF_USESTDHANDLES;
+		Flags = 0;
+		ShowWindow = 0;
 	}
 
-	// initialize startup info
-	STARTUPINFO StartupInfo = {
-		sizeof(STARTUPINFO),
-		NULL, NULL, NULL,
-		(::DWORD)CW_USEDEFAULT,
-		(::DWORD)CW_USEDEFAULT,
-		(::DWORD)CW_USEDEFAULT,
-		(::DWORD)CW_USEDEFAULT,
-		(::DWORD)0, (::DWORD)0, (::DWORD)0,
-		(::DWORD)dwFlags,
-		ShowWindowFlags,
-		0, NULL,
-		HANDLE(PipeReadChild),
-		HANDLE(PipeWriteChild),
-		HANDLE(PipeStdErrChild)
+	// Required to exist until DeleteProcThreadAttributeList is called
+	HANDLE Handles[3];
+	LPPROC_THREAD_ATTRIBUTE_LIST AttributeList = nullptr;
+	ON_SCOPE_EXIT
+	{
+		if (AttributeList)
+		{
+			DeleteProcThreadAttributeList(AttributeList);
+		}
 	};
 
-	bool bInheritHandles = (dwFlags & STARTF_USESTDHANDLES) != 0;
-
-	// create the child process
-	FString CommandLine = FString::Printf(TEXT("\"%s\" %s"), URL, Parms);
-	PROCESS_INFORMATION ProcInfo;
-
-	if (!CreateProcess(NULL, CommandLine.GetCharArray().GetData(), nullptr, nullptr, bInheritHandles, (::DWORD)CreateFlags, NULL, OptionalWorkingDirectory, &StartupInfo, &ProcInfo))
+	if (!bInheritHandles)
 	{
-		DWORD ErrorCode = GetLastError();
-
-		TCHAR ErrorMessage[512];
-		FWindowsPlatformMisc::GetSystemErrorMessage(ErrorMessage, 512, ErrorCode);
-
-		UE_LOG(LogWindows, Warning, TEXT("CreateProc failed: %s (0x%08x)"), ErrorMessage, ErrorCode);
-		if (ErrorCode == ERROR_NOT_ENOUGH_MEMORY || ErrorCode == ERROR_OUTOFMEMORY)
+		uint32 HandleCount = 0;
+		if (PipeWriteChild)
 		{
-			// These errors are common enough that we want some available memory information
-			FPlatformMemoryStats Stats = FPlatformMemory::GetStats();
-			UE_LOG(LogWindows, Warning, TEXT("Mem used: %.2f MB, OS Free %.2f MB"), Stats.UsedPhysical / 1048576.0f, Stats.AvailablePhysical / 1048576.0f);
+			Handles[HandleCount++] = HANDLE(PipeWriteChild);
 		}
-		UE_LOG(LogWindows, Warning, TEXT("URL: %s %s"), URL, Parms);
-		if (OutProcessID != nullptr)
+		if (PipeReadChild && PipeReadChild != PipeWriteChild)
 		{
-			*OutProcessID = 0;
+			Handles[HandleCount++] = HANDLE(PipeReadChild);
+		}
+		if (PipeStdErrChild && PipeStdErrChild != PipeWriteChild && PipeStdErrChild != PipeReadChild)
+		{
+			Handles[HandleCount++] = HANDLE(PipeStdErrChild);
 		}
+		
+		if (HandleCount)
+		{
+			SIZE_T Size = 0;
+			verify(!InitializeProcThreadAttributeList(nullptr, 1, 0, &Size));
+			check(GetLastError() == ERROR_INSUFFICIENT_BUFFER);
+			check(Size > 0);
+			check(Size < 1024);
+			// Doing what .NET does here for efficiency
+			void* Buffer = _alloca(Size);
+
+			if (!InitializeProcThreadAttributeList(LPPROC_THREAD_ATTRIBUTE_LIST(Buffer), 1, 0, &Size))
+			{
+				HandleError();
+				return {};
+			}
 
-		return FProcHandle();
+			AttributeList = LPPROC_THREAD_ATTRIBUTE_LIST(Buffer);
+
+			if (!UpdateProcThreadAttribute(AttributeList, 0, PROC_THREAD_ATTRIBUTE_HANDLE_LIST, Handles, HandleCount * sizeof(HANDLE), nullptr, nullptr))
+			{
+				HandleError();
+				return {};
+			}
+
+			CreationFlags |= EXTENDED_STARTUPINFO_PRESENT;
+			Flags |= STARTF_USESTDHANDLES;
+		}
 	}
 
-	if (OutProcessID != nullptr)
+	FString CommandLine = FString::Printf(TEXT("\"%s\" %s"), URL, Parms);
+	STARTUPINFOEX StartupInfoEx
 	{
-		*OutProcessID = ProcInfo.dwProcessId;
+		{
+			.cb = sizeof(STARTUPINFOEX),
+			.dwFlags = Flags,
+			.wShowWindow = ShowWindow,
+			.hStdInput = HANDLE(PipeReadChild),
+			.hStdOutput = HANDLE(PipeWriteChild),
+			.hStdError = HANDLE(PipeStdErrChild)
+		},
+		AttributeList
+	};
+	PROCESS_INFORMATION ProcessInfomation;
+
+	if (!CreateProcess(nullptr, CommandLine.GetCharArray().GetData(), nullptr, nullptr, bInheritHandles || AttributeList, CreationFlags, nullptr, OptionalWorkingDirectory, &StartupInfoEx.StartupInfo, &ProcessInfomation))
+	{
+		HandleError();
+		return {};
 	}
 
-	::CloseHandle( ProcInfo.hThread );
+	CloseHandle(ProcessInfomation.hThread);
 
-	return FProcHandle(ProcInfo.hProcess);
+	if (OutProcessID)
+	{
+		*OutProcessID = ProcessInfomation.dwProcessId;
+	}
+
+	return FProcHandle{ProcessInfomation.hProcess};
 }
+#include "Windows/HideWindowsPlatformTypes.h"
 
 bool FWindowsPlatformProcess::SetProcPriority(FProcHandle& InProcHandle, int32 PriorityModifier)
 {
@@ -974,9 +1101,8 @@ bool FWindowsPlatformProcess::ExecProcess(const TCHAR* URL, const TCHAR* Params,
 
 	HANDLE hStdOutRead = NULL;
 	HANDLE hStdErrRead = NULL;
-
-	// InheritHandles data must persist at least as long as AttributeList. see: https://learn.microsoft.com/en-us/windows/win32/api/processthreadsapi/nf-processthreadsapi-updateprocthreadattribute
-	HANDLE InheritHandles[2] = {};
+	// Required to exist until DeleteProcThreadAttributeList is called
+	HANDLE InheritHandles[2];
 	TArray<uint8> AttributeList;
 
 	if(OutStdOut != nullptr || OutStdErr != nullptr)
@@ -1001,7 +1127,6 @@ bool FWindowsPlatformProcess::ExecProcess(const TCHAR* URL, const TCHAR* Params,
 
 		InheritHandles[0] = StartupInfoEx.StartupInfo.hStdOutput;
 		InheritHandles[1] = StartupInfoEx.StartupInfo.hStdError;
-
 		verify(UpdateProcThreadAttribute(StartupInfoEx.lpAttributeList, 0, PROC_THREAD_ATTRIBUTE_HANDLE_LIST, InheritHandles, sizeof(InheritHandles), NULL, NULL));
 	}
 
@@ -1017,10 +1142,13 @@ bool FWindowsPlatformProcess::ExecProcess(const TCHAR* URL, const TCHAR* Params,
 		CommandLine = FString::Printf(TEXT("%s %s"), URL, Params);
 	}
 
+	// CREATE_NEW_PROCESS_GROUP for consistency with Unix-like platforms
+	// CREATE_NO_WINDOW because DETACHED_PROCESS can cause console windows to appear if the child spawns a process
+	// https://ikriv.com/dev/cpp/ConsoleProxy/flags has the best explanation
+	uint32 CreateFlags = NORMAL_PRIORITY_CLASS | CREATE_NEW_PROCESS_GROUP | CREATE_NO_WINDOW;
 	// We only want to add the EXTENDED_STARTUPINFO_PRESENT flag if StartupInfoEx.lpAttributeList is actually setup.
 	// If StartupInfoEx.lpAttributeList is NULL when the EXTENDED_STARTUPINFO_PRESENT flag is used, then CreateProcess causes an Access Violation crash on some Win32 configurations.
 	// This is specific to when the call is redirected to APIHook_CreateProcessW in AcLayers.dll rather than the standard CreateProcessW implementation in kernel32.dll.
-	uint32 CreateFlags = NORMAL_PRIORITY_CLASS | DETACHED_PROCESS;
 	if (StartupInfoEx.lpAttributeList != NULL)
 	{
 		CreateFlags |= EXTENDED_STARTUPINFO_PRESENT;
diff --git a/Engine/Source/Runtime/Core/Public/Android/AndroidPlatformProcess.h b/Engine/Source/Runtime/Core/Public/Android/AndroidPlatformProcess.h
--- a/Engine/Source/Runtime/Core/Public/Android/AndroidPlatformProcess.h
+++ b/Engine/Source/Runtime/Core/Public/Android/AndroidPlatformProcess.h
@@ -13,7 +13,7 @@ struct FProcHandle : public TProcHandle<void*, nullptr>
 {
 public:
 	/** Default constructor. */
-	UE_FORCEINLINE_HINT FProcHandle()
+	UE_FORCEINLINE_HINT constexpr FProcHandle()
 		: TProcHandle()
 	{}
 
diff --git a/Engine/Source/Runtime/Core/Public/GenericPlatform/GenericPlatformProcess.h b/Engine/Source/Runtime/Core/Public/GenericPlatform/GenericPlatformProcess.h
--- a/Engine/Source/Runtime/Core/Public/GenericPlatform/GenericPlatformProcess.h
+++ b/Engine/Source/Runtime/Core/Public/GenericPlatform/GenericPlatformProcess.h
@@ -10,6 +10,7 @@
 #include "HAL/PlatformMisc.h"
 #include "Misc/EnumClassFlags.h"
 #include "Templates/Function.h"
+#include "Templates/Tuple.h"
 
 class FEvent;
 
@@ -90,7 +91,7 @@ struct TProcHandle
 public:
 
 	/** Default constructor. */
-	UE_FORCEINLINE_HINT TProcHandle()
+	UE_FORCEINLINE_HINT constexpr TProcHandle()
 		: Handle( InvalidHandleValue )
 	{ }
 
@@ -169,6 +170,45 @@ namespace UE::Profiling
 	};
 }
 
+namespace UE::HAL
+{
+	struct FReadHandle
+	{
+		void* Value = nullptr;
+		
+		operator void*() const
+		{
+			return Value;
+		}
+	};
+	
+	struct FWriteHandle
+	{
+		void* Value = nullptr;
+		
+		operator void*() const
+		{
+			return Value;
+		}
+	};
+	
+	struct FProcessStartInfo
+	{
+		const TCHAR* Uri;
+		const TCHAR* Arguments = TEXT("");
+		const TCHAR* WorkingDirectory = nullptr;
+		int32 PriorityModifier = 0;
+		bool bDetached = false;
+		bool bHidden = false;
+		bool bInheritHandles = false;
+		FReadHandle StdIn;
+		FWriteHandle StdOut;
+		FWriteHandle StdErr = StdOut;
+	};
+	
+	enum struct EProcessId: uint32 {};
+}
+
 /**
 * Generic implementation for most platforms, these tend to be unused and unimplemented
 **/
@@ -535,6 +575,13 @@ struct FGenericPlatformProcess
 	 */
 	static CORE_API FProcHandle CreateProc( const TCHAR* URL, const TCHAR* Parms, bool bLaunchDetached, bool bLaunchHidden, bool bLaunchReallyHidden, uint32* OutProcessID, int32 PriorityModifier, const TCHAR* OptionalWorkingDirectory, void* PipeWriteChild, void* PipeReadChild, void* PipeStdErrChild);
 
+	/**
+	 * Creates a new process using StartInfo. The new process runs the specified executable file in the security context of the calling process.
+	 * @param StartInfo				The starting parameters 
+	 * @return	The process Handle and Id for use in other process functions
+	 */
+	static CORE_API TTuple<FProcHandle, UE::HAL::EProcessId> CreateProc(const UE::HAL::FProcessStartInfo& StartInfo);
+	
 	/**
 	 * Opens an existing process. 
 	 *
diff --git a/Engine/Source/Runtime/Core/Public/HAL/PlatformProcess.h b/Engine/Source/Runtime/Core/Public/HAL/PlatformProcess.h
--- a/Engine/Source/Runtime/Core/Public/HAL/PlatformProcess.h
+++ b/Engine/Source/Runtime/Core/Public/HAL/PlatformProcess.h
@@ -5,3 +5,232 @@
 #include "CoreTypes.h"
 
 #include COMPILED_PLATFORM_HEADER(PlatformProcess.h) // IWYU pragma: export
+
+namespace UE::HAL
+{
+	struct [[nodiscard]] FPipe
+	{
+		constexpr FPipe() = default;
+		FPipe(const FPipe&) = delete;
+		FPipe& operator=(const FPipe&) = delete;
+		
+		constexpr FPipe(std::nullptr_t)
+		{
+		}
+		
+		FPipe(FPipe&& Other)
+			: ReadHandle{std::exchange(Other.ReadHandle, {})}
+			, WriteHandle{std::exchange(Other.WriteHandle, {})}
+		{
+		}
+		
+		~FPipe()
+		{
+			FPlatformProcess::ClosePipe(ReadHandle, WriteHandle);
+		}
+		
+		FPipe& operator=(std::nullptr_t)
+		{
+			FPlatformProcess::ClosePipe(std::exchange(ReadHandle, {}), std::exchange(WriteHandle, {}));
+			return *this;
+		}
+		
+		FPipe& operator=(FPipe&& Other)
+		{
+			Swap(ReadHandle, Other.ReadHandle);
+			Swap(WriteHandle, Other.WriteHandle);
+			return *this;
+		}
+		
+		bool operator==(std::nullptr_t) const
+		{
+			return !ReadHandle;
+		}
+		
+		explicit operator bool() const
+		{
+			return bool(ReadHandle);
+		}
+
+	protected:
+		FReadHandle ReadHandle;
+		FWriteHandle WriteHandle;
+		
+		FPipe(TTuple<FReadHandle, FWriteHandle> Pipe)
+			: ReadHandle{Pipe.Get<FReadHandle>()}
+			, WriteHandle{Pipe.Get<FWriteHandle>()}
+		{
+		}
+		
+		explicit FPipe(bool bWritePipeLocal)
+			: FPipe{[&bWritePipeLocal]
+			{
+				void* ReadPipe;
+				void* WritePipe;
+				if (!FPlatformProcess::CreatePipe(ReadPipe, WritePipe, bWritePipeLocal))
+				{
+					ReadPipe = nullptr;
+					WritePipe = nullptr;
+				}
+				return TTuple<FReadHandle, FWriteHandle>{ReadPipe, WritePipe};
+			}()}
+		{
+		}
+	};
+	
+	inline constexpr struct FNewPipe
+	{
+	}
+	NewPipe;
+	
+	struct [[nodiscard]] FInputPipe: FPipe
+	{
+		using FPipe::FPipe;
+		
+		FInputPipe(FNewPipe)
+			: FPipe{false}
+		{
+		}
+
+		void* NativeHandle() const
+		{
+			return ReadHandle;
+		}
+
+		FString Read()
+		{
+			check(*this);
+			return FPlatformProcess::ReadPipe(ReadHandle);
+		}
+		
+		operator FWriteHandle()&
+		{
+			return WriteHandle;
+		}
+	};
+	
+	struct [[nodiscard]] FOutputPipe: FPipe
+	{
+		using FPipe::FPipe;
+		
+		FOutputPipe(FNewPipe)
+			: FPipe{true}
+		{
+		}
+		
+		void* NativeHandle() const
+		{
+			return WriteHandle;
+		}
+		
+		bool Write(const FString& Message)
+		{
+			check(*this);
+			return FPlatformProcess::WritePipe(WriteHandle, Message);
+		}
+		
+		operator FReadHandle()&
+		{
+			return ReadHandle;
+		}
+	};
+	
+	struct [[nodiscard]] FProcess
+	{
+		constexpr FProcess() = default;
+		FProcess(const FProcess&) = delete;
+		FProcess& operator=(const FProcess&) = delete;
+		
+		constexpr FProcess(std::nullptr_t)
+		{
+		}
+		
+		FProcess(EProcessId ProcessId)
+			: Handle{FPlatformProcess::OpenProcess(std::underlying_type_t<EProcessId>(ProcessId))}
+		{
+		}
+
+		FProcess(const FProcessStartInfo& StartInfo)
+			: Handle{FPlatformProcess::CreateProc(StartInfo).Get<FProcHandle>()}
+		{
+		}
+		
+		FProcess(const FProcessStartInfo& StartInfo, EProcessId& OutProcessId)
+			: Handle{[&StartInfo, &OutProcessId]
+			{
+				TTuple<FProcHandle, EProcessId> Process = FPlatformProcess::CreateProc(StartInfo);
+				OutProcessId = Process.Get<EProcessId>();
+				return Process.Get<FProcHandle>();
+			}()}
+		{
+		}
+		
+		FProcess(FProcess&& Other)
+			: Handle{Other.Handle}
+		{
+			Other.Handle.Reset();
+		}
+		
+		~FProcess()
+		{
+			FPlatformProcess::CloseProc(Handle);
+		}
+		
+		FProcess& operator=(std::nullptr_t)
+		{
+			FPlatformProcess::CloseProc(Handle);
+			return *this;
+		}
+		
+		FProcess& operator=(FProcess&& Other)
+		{
+			Swap(Handle, Other.Handle);
+			return *this;
+		}
+		
+		decltype(auto) NativeHandle() const
+		{
+			return Handle.Get();
+		}
+		
+		bool IsRunning()
+		{
+			check(*this);
+			return FPlatformProcess::IsProcRunning(Handle);
+		}
+		
+		TOptional<int32> GetExitCode()
+		{
+			check(*this);
+			int32 ExitCode;
+			return FPlatformProcess::GetProcReturnCode(Handle, &ExitCode) ? TOptional{ExitCode} : NullOpt;
+		}
+		
+		FProcess& WaitForExit()
+		{
+			check(*this);
+			FPlatformProcess::WaitForProc(Handle);
+			return *this;
+		}
+		
+		FProcess& Kill(bool bEntireTree = false)
+		{
+			check(*this);
+			FPlatformProcess::TerminateProc(Handle, bEntireTree);
+			return *this;
+		}
+
+		bool operator==(std::nullptr_t) const
+		{
+			return !Handle.IsValid();
+		}
+
+		explicit operator bool() const
+		{
+			return Handle.IsValid();
+		}
+
+	private:
+		FProcHandle Handle;
+	};
+}
diff --git a/Engine/Source/Runtime/Core/Public/IOS/IOSPlatformProcess.h b/Engine/Source/Runtime/Core/Public/IOS/IOSPlatformProcess.h
--- a/Engine/Source/Runtime/Core/Public/IOS/IOSPlatformProcess.h
+++ b/Engine/Source/Runtime/Core/Public/IOS/IOSPlatformProcess.h
@@ -13,7 +13,7 @@ struct FProcHandle : public TProcHandle<void*, nullptr>
 {
 public:
 	/** Default constructor. */
-	UE_FORCEINLINE_HINT FProcHandle()
+	UE_FORCEINLINE_HINT constexpr FProcHandle()
 		: TProcHandle()
 	{}
 
diff --git a/Engine/Source/Runtime/Core/Public/Mac/MacPlatformProcess.h b/Engine/Source/Runtime/Core/Public/Mac/MacPlatformProcess.h
--- a/Engine/Source/Runtime/Core/Public/Mac/MacPlatformProcess.h
+++ b/Engine/Source/Runtime/Core/Public/Mac/MacPlatformProcess.h
@@ -77,7 +77,7 @@ struct FProcHandle
 	  */
 	pid_t				OpenedPid;
 
-	FProcHandle()
+	constexpr FProcHandle()
 	:	ProcInfo(nullptr), OpenedPid(-1)
 	{
 	}
@@ -206,8 +206,9 @@ struct CORE_API FMacPlatformProcess : public FApplePlatformProcess
 	static bool CanLaunchURL(const TCHAR* URL);
 	static void LaunchURL(const TCHAR* URL, const TCHAR* Parms, FString* Error);
 	static FString GetGameBundleId();
-	static FProcHandle CreateProc( const TCHAR* URL, const TCHAR* Parms, bool bLaunchDetached, bool bLaunchHidden, bool bLaunchReallyHidden, uint32* OutProcessID, int32 PriorityModifier, const TCHAR* OptionalWorkingDirectory, void* PipeWriteChild, void* PipeReadChild = nullptr );
-	static FProcHandle CreateProc( const TCHAR* URL, const TCHAR* Parms, bool bLaunchDetached, bool bLaunchHidden, bool bLaunchReallyHidden, uint32* OutProcessID, int32 PriorityModifier, const TCHAR* OptionalWorkingDirectory, void* PipeWriteChild, void* PipeReadChild, void* PipeStdErrChild );
+	static FProcHandle CreateProc(const TCHAR* URL, const TCHAR* Parms, bool bLaunchDetached, bool bLaunchHidden, bool bLaunchReallyHidden, uint32* OutProcessID, int32 PriorityModifier, const TCHAR* OptionalWorkingDirectory, void* PipeWriteChild, void* PipeReadChild = nullptr);
+	static FProcHandle CreateProc(const TCHAR* URL, const TCHAR* Parms, bool bLaunchDetached, bool bLaunchHidden, bool bLaunchReallyHidden, uint32* OutProcessID, int32 PriorityModifier, const TCHAR* OptionalWorkingDirectory, void* PipeWriteChild, void* PipeReadChild, void* PipeStdErrChild);
+	static TTuple<FProcHandle, UE::HAL::EProcessId> CreateProc(const UE::HAL::FProcessStartInfo& StartInfo);
 	static FProcHandle OpenProcess(uint32 ProcessID);
 	static bool IsProcRunning( FProcHandle & ProcessHandle );
 	static void WaitForProc( FProcHandle & ProcessHandle );
@@ -236,7 +237,7 @@ struct CORE_API FMacPlatformProcess : public FApplePlatformProcess
 
 private:
 
-	static FProcHandle CreateProcInternal(const TCHAR* URL, const TCHAR* Parms, bool bLaunchDetached, bool bLaunchHidden, bool bLaunchReallyHidden, uint32* OutProcessID, int32 PriorityModifier, const TCHAR* OptionalWorkingDirectory, void* PipeStdOutChild, void* PipeStdErrChild, void *PipeStdInChild);
+	static FProcHandle CreateProcInternal(const TCHAR* URL, const TCHAR* Parms, bool bLaunchDetached, bool bLaunchHidden, bool bLaunchReallyHidden, uint32* OutProcessID, int32 PriorityModifier, const TCHAR* OptionalWorkingDirectory, void* PipeStdOutChild, void* PipeStdInChild, void* PipeStdErrChild, bool bInheritHandles);
 };
 
 typedef FMacPlatformProcess FPlatformProcess;
diff --git a/Engine/Source/Runtime/Core/Public/Unix/UnixPlatformProcess.h b/Engine/Source/Runtime/Core/Public/Unix/UnixPlatformProcess.h
--- a/Engine/Source/Runtime/Core/Public/Unix/UnixPlatformProcess.h
+++ b/Engine/Source/Runtime/Core/Public/Unix/UnixPlatformProcess.h
@@ -118,7 +118,7 @@ struct FProcHandle
 	  */
 	pid_t				OpenedPid;
 
-	FProcHandle()
+	constexpr FProcHandle()
 	:	ProcInfo(nullptr), OpenedPid(-1)
 	{
 	}
@@ -294,7 +294,8 @@ struct FUnixPlatformProcess : public FGenericPlatformProcess
 	static CORE_API bool CanLaunchURL(const TCHAR* URL);
 	static CORE_API void LaunchURL(const TCHAR* URL, const TCHAR* Parms, FString* Error);
 	static CORE_API FProcHandle CreateProc(const TCHAR* URL, const TCHAR* Parms, bool bLaunchDetached, bool bLaunchHidden, bool bLaunchReallyHidden, uint32* OutProcessID, int32 PriorityModifier, const TCHAR* OptionalWorkingDirectory, void* PipeWriteChild, void* PipeReadChild = nullptr);
-	static CORE_API FProcHandle CreateProc(const TCHAR* URL, const TCHAR* Parms, bool bLaunchDetached, bool bLaunchHidden, bool bLaunchReallyHidden, uint32* OutProcessID, int32 PriorityModifier, const TCHAR* OptionalWorkingDirectory, void* PipeWriteChild, void* PipeReadChild , void* PipeStdErrChild);
+	static CORE_API FProcHandle CreateProc(const TCHAR* URL, const TCHAR* Parms, bool bLaunchDetached, bool bLaunchHidden, bool bLaunchReallyHidden, uint32* OutProcessID, int32 PriorityModifier, const TCHAR* OptionalWorkingDirectory, void* PipeWriteChild, void* PipeReadChild, void* PipeStdErrChild);
+	static CORE_API TTuple<FProcHandle, UE::HAL::EProcessId> CreateProc(const UE::HAL::FProcessStartInfo& StartInfo);
 	static CORE_API FProcHandle OpenProcess(uint32 ProcessID);
 	static CORE_API bool IsProcRunning( FProcHandle & ProcessHandle );
 	static CORE_API void WaitForProc( FProcHandle & ProcessHandle );
@@ -328,4 +329,7 @@ struct FUnixPlatformProcess : public FGenericPlatformProcess
 	 * Like other directory functions, cannot return nullptr!
 	 */
 	static CORE_API const TCHAR* UserHomeDir();
+
+private:
+	static FProcHandle CreateProcInternal(const TCHAR* URL, const TCHAR* Parms, bool bLaunchDetached, bool bLaunchHidden, bool bLaunchReallyHidden, uint32* OutProcessID, int32 PriorityModifier, const TCHAR* OptionalWorkingDirectory, void* PipeWriteChild, void* PipeReadChild, void* PipeStdErrChild, bool bInheritHandles);
 };
diff --git a/Engine/Source/Runtime/Core/Public/Windows/WindowsPlatformProcess.h b/Engine/Source/Runtime/Core/Public/Windows/WindowsPlatformProcess.h
--- a/Engine/Source/Runtime/Core/Public/Windows/WindowsPlatformProcess.h
+++ b/Engine/Source/Runtime/Core/Public/Windows/WindowsPlatformProcess.h
@@ -17,7 +17,7 @@ struct FProcHandle : public TProcHandle<Windows::HANDLE, nullptr>
 {
 public:
 	/** Default constructor. */
-	UE_FORCEINLINE_HINT FProcHandle()
+	UE_FORCEINLINE_HINT constexpr FProcHandle()
 		: TProcHandle()
 	{}
 
@@ -164,8 +164,9 @@ public:
 	static CORE_API const TCHAR* GetModuleExtension();
 	static CORE_API const TCHAR* GetBinariesSubdirectory();
 	static CORE_API const FString GetModulesDirectory();
-	static CORE_API FProcHandle CreateProc( const TCHAR* URL, const TCHAR* Parms, bool bLaunchDetached, bool bLaunchHidden, bool bLaunchReallyHidden, uint32* OutProcessID, int32 PriorityModifier, const TCHAR* OptionalWorkingDirectory, void* PipeWriteChild, void* PipeReadChild = nullptr);
-	static CORE_API FProcHandle CreateProc( const TCHAR* URL, const TCHAR* Parms, bool bLaunchDetached, bool bLaunchHidden, bool bLaunchReallyHidden, uint32* OutProcessID, int32 PriorityModifier, const TCHAR* OptionalWorkingDirectory, void* PipeWriteChild, void* PipeReadChild, void* PipeStdErrChild);
+	static CORE_API FProcHandle CreateProc(const TCHAR* URL, const TCHAR* Parms, bool bLaunchDetached, bool bLaunchHidden, bool bLaunchReallyHidden, uint32* OutProcessID, int32 PriorityModifier, const TCHAR* OptionalWorkingDirectory, void* PipeWriteChild, void* PipeReadChild = nullptr);
+	static CORE_API FProcHandle CreateProc(const TCHAR* URL, const TCHAR* Parms, bool bLaunchDetached, bool bLaunchHidden, bool bLaunchReallyHidden, uint32* OutProcessID, int32 PriorityModifier, const TCHAR* OptionalWorkingDirectory, void* PipeWriteChild, void* PipeReadChild, void* PipeStdErrChild);
+	static CORE_API TTuple<FProcHandle, UE::HAL::EProcessId> CreateProc(const UE::HAL::FProcessStartInfo& StartInfo);
 	static CORE_API bool SetProcPriority(FProcHandle & InProcHandle, int32 PriorityModifier);
 	static CORE_API FProcHandle OpenProcess(uint32 ProcessID);
 	static CORE_API bool IsProcRunning( FProcHandle & ProcessHandle );
@@ -244,6 +245,8 @@ protected:
 
 private:
 
+	static FProcHandle CreateProcInternal(const TCHAR* URL, const TCHAR* Parms, bool bLaunchDetached, bool bLaunchHidden, bool bLaunchReallyHidden, uint32* OutProcessID, int32 PriorityModifier, const TCHAR* OptionalWorkingDirectory, void* PipeWriteChild, void* PipeReadChild, void* PipeStdErrChild, bool bInheritHandles);
+	
 	static void LaunchURLInternal(const TCHAR* URL, const TCHAR* Parms, FString* Error);
 
 	/**
