Disabled V530 on FProcess WaitForExit and Kill due to a PVS false postive

diff --git a/Engine/Source/Runtime/Core/Public/HAL/PlatformProcess.h b/Engine/Source/Runtime/Core/Public/HAL/PlatformProcess.h
--- a/Engine/Source/Runtime/Core/Public/HAL/PlatformProcess.h
+++ b/Engine/Source/Runtime/Core/Public/HAL/PlatformProcess.h
@@ -206,6 +206,8 @@ namespace UE::HAL
 			return FPlatformProcess::GetProcReturnCode(Handle, &ExitCode) ? TOptional{ExitCode} : NullOpt;
 		}
 		
+		// Bug in PVS to flag references
+		//-V:WaitForExit:530
 		FProcess& WaitForExit()
 		{
 			check(*this);
@@ -213,6 +215,8 @@ namespace UE::HAL
 			return *this;
 		}
 		
+		// Bug in PVS to flag references
+		//-V:Kill:530
 		FProcess& Kill(bool bEntireTree = false)
 		{
 			check(*this);
