[Core] Add delegates for the unmounting of Pak files

#tests local build

diff --git a/Engine/Source/Runtime/Core/Private/Misc/CoreDelegates.cpp b/Engine/Source/Runtime/Core/Private/Misc/CoreDelegates.cpp
--- a/Engine/Source/Runtime/Core/Private/Misc/CoreDelegates.cpp
+++ b/Engine/Source/Runtime/Core/Private/Misc/CoreDelegates.cpp
@@ -42,6 +42,18 @@ TTSMulticastDelegate<void(const IPakFile&)>& FCoreDelegates::GetOnPakFileMounted
 	return Singleton;
 }
 
+TTSMulticastDelegate<void(const IPakFile&)>& FCoreDelegates::GetOnPakFileUnmounting()
+{
+	static TTSMulticastDelegate<void(const IPakFile&)> Singleton;
+	return Singleton;
+}
+
+TTSMulticastDelegate<void(const IPakFile&)>& FCoreDelegates::GetOnPakFileUnmounted()
+{
+	static TTSMulticastDelegate<void(const IPakFile&)> Singleton;
+	return Singleton;
+}
+
 TMulticastDelegate<void(const FString&)> FCoreDelegates::NewFileAddedDelegate;
 TMulticastDelegate<void()> FCoreDelegates::NoPakFilesMountedDelegate;
 
diff --git a/Engine/Source/Runtime/Core/Public/Misc/CoreDelegates.h b/Engine/Source/Runtime/Core/Public/Misc/CoreDelegates.h
--- a/Engine/Source/Runtime/Core/Public/Misc/CoreDelegates.h
+++ b/Engine/Source/Runtime/Core/Public/Misc/CoreDelegates.h
@@ -142,6 +142,12 @@ public:
 	// After a pakfile is mounted this is called
 	static CORE_API TTSMulticastDelegate<void(const IPakFile&)>& GetOnPakFileMounted2();
 
+	// Before a pakfile is unmounted this is called
+	static CORE_API TTSMulticastDelegate<void(const IPakFile&)>& GetOnPakFileUnmounting();
+
+	// After a pakfile is unmounted this is called
+	static CORE_API TTSMulticastDelegate<void(const IPakFile&)>& GetOnPakFileUnmounted();
+
 	// After a file is added this is called
 	static CORE_API TMulticastDelegate<void(const FString&)> NewFileAddedDelegate;
 
diff --git a/Engine/Source/Runtime/PakFile/Private/IPlatformFilePak.cpp b/Engine/Source/Runtime/PakFile/Private/IPlatformFilePak.cpp
--- a/Engine/Source/Runtime/PakFile/Private/IPlatformFilePak.cpp
+++ b/Engine/Source/Runtime/PakFile/Private/IPlatformFilePak.cpp
@@ -6199,6 +6199,7 @@ bool FPakPlatformFile::Unmount(const TCHAR* InPakFilename)
 
 	if (UnmountedPak)
 	{
+		FCoreDelegates::GetOnPakFileUnmounting().Broadcast(*UnmountedPak);
 		RemoveCachedPakSignaturesFile(*UnmountedPak->GetFilename());
 	}
 
@@ -6240,6 +6241,11 @@ bool FPakPlatformFile::Unmount(const TCHAR* InPakFilename)
 			UnmountedPak->SetIsMounted(false);
 		}
 	}
+
+	if (UnmountedPak)
+	{
+		FCoreDelegates::GetOnPakFileUnmounted().Broadcast(*UnmountedPak);
+	}
 	return UnmountedPak.IsValid() || bRemovedContainerFile;
 }
 
