Added [[nodiscard]] to mutex functions

diff --git a/Engine/Source/Runtime/Core/Public/Async/ExternalMutex.h b/Engine/Source/Runtime/Core/Public/Async/ExternalMutex.h
--- a/Engine/Source/Runtime/Core/Public/Async/ExternalMutex.h
+++ b/Engine/Source/Runtime/Core/Public/Async/ExternalMutex.h
@@ -33,12 +33,12 @@ public:
 	{
 	}
 
-	inline bool IsLocked() const
+	[[nodiscard]] inline bool IsLocked() const
 	{
 		return TIntrusiveMutex<ParamsType>::IsLocked(State);
 	}
 
-	inline bool TryLock()
+	[[nodiscard]] inline bool TryLock()
 	{
 		return TIntrusiveMutex<ParamsType>::TryLock(State);
 	}
diff --git a/Engine/Source/Runtime/Core/Public/Async/IntrusiveMutex.h b/Engine/Source/Runtime/Core/Public/Async/IntrusiveMutex.h
--- a/Engine/Source/Runtime/Core/Public/Async/IntrusiveMutex.h
+++ b/Engine/Source/Runtime/Core/Public/Async/IntrusiveMutex.h
@@ -110,12 +110,12 @@ class TIntrusiveMutex
 	}
 
 public:
-	FORCEINLINE static bool IsLocked(const std::atomic<StateType>& State)
+	[[nodiscard]] FORCEINLINE static bool IsLocked(const std::atomic<StateType>& State)
 	{
 		return !!(State.load(std::memory_order_relaxed) & IsLockedFlag);
 	}
 
-	FORCEINLINE static bool TryLock(std::atomic<StateType>& State)
+	[[nodiscard]] FORCEINLINE static bool TryLock(std::atomic<StateType>& State)
 	{
 		StateType Expected = State.load(std::memory_order_relaxed);
 		return !(Expected & IsLockedMask) &&
diff --git a/Engine/Source/Runtime/Core/Public/Async/Mutex.h b/Engine/Source/Runtime/Core/Public/Async/Mutex.h
--- a/Engine/Source/Runtime/Core/Public/Async/Mutex.h
+++ b/Engine/Source/Runtime/Core/Public/Async/Mutex.h
@@ -28,12 +28,12 @@ public:
 	FMutex(const FMutex&) = delete;
 	FMutex& operator=(const FMutex&) = delete;
 
-	inline bool IsLocked() const
+	[[nodiscard]] inline bool IsLocked() const
 	{
 		return !!(State.load(std::memory_order_relaxed) & IsLockedFlag);
 	}
 
-	inline bool TryLock()
+	[[nodiscard]] inline bool TryLock()
 	{
 		uint8 Expected = State.load(std::memory_order_relaxed);
 		return !(Expected & IsLockedFlag) &&
diff --git a/Engine/Source/Runtime/Core/Public/Async/RecursiveMutex.h b/Engine/Source/Runtime/Core/Public/Async/RecursiveMutex.h
--- a/Engine/Source/Runtime/Core/Public/Async/RecursiveMutex.h
+++ b/Engine/Source/Runtime/Core/Public/Async/RecursiveMutex.h
@@ -23,12 +23,12 @@ public:
 	FRecursiveMutex(const FRecursiveMutex&) = delete;
 	FRecursiveMutex& operator=(const FRecursiveMutex&) = delete;
 
-	inline bool IsLocked() const
+	[[nodiscard]] inline bool IsLocked() const
 	{
 		return !!(State.load(std::memory_order_relaxed) & LockCountMask);
 	}
 
-	UE_API bool TryLock();
+	[[nodiscard]] UE_API bool TryLock();
 	UE_API void Lock();
 	UE_API void Unlock();
 
diff --git a/Engine/Source/Runtime/Core/Public/Async/RecursiveWordMutex.h b/Engine/Source/Runtime/Core/Public/Async/RecursiveWordMutex.h
--- a/Engine/Source/Runtime/Core/Public/Async/RecursiveWordMutex.h
+++ b/Engine/Source/Runtime/Core/Public/Async/RecursiveWordMutex.h
@@ -22,7 +22,7 @@ namespace UE
 class FRecursiveWordMutex final
 {
 public:
-	UE_API bool TryLock();
+	[[nodiscard]] UE_API bool TryLock();
 	UE_API void Lock();
 	UE_API void Unlock();
 
diff --git a/Engine/Source/Runtime/Core/Public/Async/SharedMutex.h b/Engine/Source/Runtime/Core/Public/Async/SharedMutex.h
--- a/Engine/Source/Runtime/Core/Public/Async/SharedMutex.h
+++ b/Engine/Source/Runtime/Core/Public/Async/SharedMutex.h
@@ -26,12 +26,12 @@ public:
 	FSharedMutex(const FSharedMutex&) = delete;
 	FSharedMutex& operator=(const FSharedMutex&) = delete;
 
-	inline bool IsLocked() const
+	[[nodiscard]] inline bool IsLocked() const
 	{
 		return !!(State.load(std::memory_order_relaxed) & IsLockedFlag);
 	}
 
-	inline bool TryLock()
+	[[nodiscard]] inline bool TryLock()
 	{
 		uint32 Expected = State.load(std::memory_order_relaxed);
 		return !(Expected & (IsLockedFlag | SharedLockCountMask)) &&
@@ -61,12 +61,12 @@ public:
 		WakeWaitingThreads(LastState);
 	}
 
-	inline bool IsLockShared() const
+	[[nodiscard]] inline bool IsLockShared() const
 	{
 		return !!(State.load(std::memory_order_relaxed) & SharedLockCountMask);
 	}
 
-	inline bool TryLockShared()
+	[[nodiscard]] inline bool TryLockShared()
 	{
 		uint32 Expected = State.load(std::memory_order_relaxed);
 		while (LIKELY(!(Expected & (IsLockedFlag | MayHaveWaitingLockFlag))))
diff --git a/Engine/Source/Runtime/Core/Public/Async/SharedRecursiveMutex.h b/Engine/Source/Runtime/Core/Public/Async/SharedRecursiveMutex.h
--- a/Engine/Source/Runtime/Core/Public/Async/SharedRecursiveMutex.h
+++ b/Engine/Source/Runtime/Core/Public/Async/SharedRecursiveMutex.h
@@ -15,7 +15,7 @@ namespace UE::Core::Private
 
 struct FSharedRecursiveMutexLink
 {
-	static bool Owns(const FSharedRecursiveMutex* Mutex);
+	[[nodiscard]] static bool Owns(const FSharedRecursiveMutex* Mutex);
 	void Push(const FSharedRecursiveMutex* Mutex);
 	void Pop();
 
@@ -44,22 +44,22 @@ public:
 	FSharedRecursiveMutex(const FSharedRecursiveMutex&) = delete;
 	FSharedRecursiveMutex& operator=(const FSharedRecursiveMutex&) = delete;
 
-	inline bool IsLocked() const
+	[[nodiscard]] inline bool IsLocked() const
 	{
 		return !!(State.load(std::memory_order_relaxed) & LockCountMask);
 	}
 
-	UE_API bool TryLock();
+	[[nodiscard]] UE_API bool TryLock();
 	UE_API void Lock();
 	UE_API void Unlock();
 
-	inline bool IsLockShared() const
+	[[nodiscard]] inline bool IsLockShared() const
 	{
 		return !!(State.load(std::memory_order_relaxed) & SharedLockCountMask);
 	}
 
 	// Use TSharedLock or TDynamicSharedLock to acquire a shared lock.
-	UE_API bool TryLockShared(Core::Private::FSharedRecursiveMutexLink& Link);
+	[[nodiscard]] UE_API bool TryLockShared(Core::Private::FSharedRecursiveMutexLink& Link);
 	UE_API void LockShared(Core::Private::FSharedRecursiveMutexLink& Link);
 	UE_API void UnlockShared(Core::Private::FSharedRecursiveMutexLink& Link);
 
@@ -165,7 +165,7 @@ public:
 		}
 	}
 
-	bool TryLock()
+	[[nodiscard]] bool TryLock()
 	{
 		check(!bLocked);
 		check(Mutex);
@@ -188,7 +188,7 @@ public:
 		Mutex->UnlockShared(Link);
 	}
 
-	inline bool OwnsLock() const
+	[[nodiscard]] inline bool OwnsLock() const
 	{
 		return bLocked;
 	}
diff --git a/Engine/Source/Runtime/Core/Public/Async/TransactionallySafeMutex.h b/Engine/Source/Runtime/Core/Public/Async/TransactionallySafeMutex.h
--- a/Engine/Source/Runtime/Core/Public/Async/TransactionallySafeMutex.h
+++ b/Engine/Source/Runtime/Core/Public/Async/TransactionallySafeMutex.h
@@ -81,7 +81,7 @@ struct FTransactionallySafeMutex final
 		}
 	}
 
-	inline bool TryLock()
+	[[nodiscard]] inline bool TryLock()
 	{
 		if (AutoRTFM::IsTransactional() || AutoRTFM::IsCommittingOrAborting())
 		{
@@ -154,7 +154,7 @@ struct FTransactionallySafeMutex final
 	 *   until the transaction fully succeeds or rolls back. Because of this, IsLocked may
 	 *   return true even after the mutex has been fully unlocked.
 	 */
-	bool IsLocked() const
+	[[nodiscard]] bool IsLocked() const
 	{
 		return State->MutexIsLocked();
 	}
diff --git a/Engine/Source/Runtime/Core/Public/Async/TransactionallySafeRecursiveMutex.h b/Engine/Source/Runtime/Core/Public/Async/TransactionallySafeRecursiveMutex.h
--- a/Engine/Source/Runtime/Core/Public/Async/TransactionallySafeRecursiveMutex.h
+++ b/Engine/Source/Runtime/Core/Public/Async/TransactionallySafeRecursiveMutex.h
@@ -99,7 +99,7 @@ struct FTransactionallySafeRecursiveMutexDefinition final
 		}
 	}
 
-	bool TryLock()
+	[[nodiscard]] bool TryLock()
 	{
 		if (AutoRTFM::IsTransactional() || AutoRTFM::IsCommittingOrAborting())
 		{
diff --git a/Engine/Source/Runtime/Core/Public/Async/TransactionallySafeSharedMutex.h b/Engine/Source/Runtime/Core/Public/Async/TransactionallySafeSharedMutex.h
--- a/Engine/Source/Runtime/Core/Public/Async/TransactionallySafeSharedMutex.h
+++ b/Engine/Source/Runtime/Core/Public/Async/TransactionallySafeSharedMutex.h
@@ -124,7 +124,7 @@ struct FTransactionallySafeSharedMutexDefinition
 		}
 	}
 
-	bool TryLock()
+	[[nodiscard]] bool TryLock()
 	{
 		if (AutoRTFM::IsTransactional() || AutoRTFM::IsCommittingOrAborting())
 		{
diff --git a/Engine/Source/Runtime/Core/Public/Async/WordMutex.h b/Engine/Source/Runtime/Core/Public/Async/WordMutex.h
--- a/Engine/Source/Runtime/Core/Public/Async/WordMutex.h
+++ b/Engine/Source/Runtime/Core/Public/Async/WordMutex.h
@@ -25,7 +25,7 @@ public:
 	FWordMutex(const FWordMutex&) = delete;
 	FWordMutex& operator=(const FWordMutex&) = delete;
 
-	inline bool TryLock()
+	[[nodiscard]] inline bool TryLock()
 	{
 		UPTRINT Expected = 0;
 		return State.compare_exchange_strong(Expected, IsLockedFlag, std::memory_order_acquire, std::memory_order_relaxed);
