FUnitConversion::Convert : Suppress compilation warnings for the case where T=float

Previously, this would generate "warning C4244: 'return': conversion from 'double' to 'T', possible loss of data" for T=float when built from a module that set `CppCompileWarningSettings.UnsafeTypeCastWarningLevel` to warning | error.

diff --git a/Engine/Source/Runtime/Core/Public/Math/UnitConversion.inl b/Engine/Source/Runtime/Core/Public/Math/UnitConversion.inl
--- a/Engine/Source/Runtime/Core/Public/Math/UnitConversion.inl
+++ b/Engine/Source/Runtime/Core/Public/Math/UnitConversion.inl
@@ -28,8 +28,6 @@ namespace UnitConversion
 	CORE_API double AngularSpeedUnificationFactor(EUnit From);
 	/** Find the common quantization factor for the specified acceleration unit. Quantizes to m/s2. */
 	CORE_API double AccelerationUnificationFactor(EUnit From);
-	/** Find the common quantization factor for the specified temperature unit. Quantizes to Kelvin. */
-	CORE_API double TemperatureUnificationFactor(EUnit From);
 	/** Find the common quantization factor for the specified mass unit. Quantizes to Grams. */
 	CORE_API double MassUnificationFactor(EUnit From);
 	/** Find the common quantization factor for the specified density unit. Quantizes to GramsPerCubicCentimeter. */
@@ -88,22 +86,22 @@ T FUnitConversion::Convert(T InValue, EUnit From, EUnit To)
 
 	switch(FUnitConversion::GetUnitType(From))
 	{
-		case EUnitType::Distance:			return InValue * DistanceUnificationFactor(From)			* (1.0 / DistanceUnificationFactor(To));
-		case EUnitType::Angle:				return InValue * AngleUnificationFactor(From) 				* (1.0 / AngleUnificationFactor(To));
-		case EUnitType::Speed:				return InValue * SpeedUnificationFactor(From) 				* (1.0 / SpeedUnificationFactor(To));
-		case EUnitType::AngularSpeed:		return InValue * AngularSpeedUnificationFactor(From)		* (1.0 / AngularSpeedUnificationFactor(To));
-		case EUnitType::Acceleration:		return InValue * AccelerationUnificationFactor(From)		* (1.0 / AccelerationUnificationFactor(To));
-		case EUnitType::Mass:				return InValue * MassUnificationFactor(From) 				* (1.0 / MassUnificationFactor(To));
-		case EUnitType::Density:			return InValue * DensityUnificationFactor(From)				* (1.0 / DensityUnificationFactor(To));
-		case EUnitType::Force:				return InValue * ForceUnificationFactor(From) 				* (1.0 / ForceUnificationFactor(To));
-		case EUnitType::Torque:				return InValue * TorqueUnificationFactor(From)				* (1.0 / TorqueUnificationFactor(To));
-		case EUnitType::PositionalImpulse:	return InValue * PositionalImpulseUnificationFactor(From)	* (1.0 / PositionalImpulseUnificationFactor(To));
-		case EUnitType::Frequency:			return InValue * FrequencyUnificationFactor(From) 			* (1.0 / FrequencyUnificationFactor(To));
-		case EUnitType::DataSize:			return InValue * DataSizeUnificationFactor(From) 			* (1.0 / DataSizeUnificationFactor(To));
+		case EUnitType::Distance:			return static_cast<T>(InValue * DistanceUnificationFactor(From)				* (1.0 / DistanceUnificationFactor(To)));
+		case EUnitType::Angle:				return static_cast<T>(InValue * AngleUnificationFactor(From) 				* (1.0 / AngleUnificationFactor(To)));
+		case EUnitType::Speed:				return static_cast<T>(InValue * SpeedUnificationFactor(From) 				* (1.0 / SpeedUnificationFactor(To)));
+		case EUnitType::AngularSpeed:		return static_cast<T>(InValue * AngularSpeedUnificationFactor(From)			* (1.0 / AngularSpeedUnificationFactor(To)));
+		case EUnitType::Acceleration:		return static_cast<T>(InValue * AccelerationUnificationFactor(From)			* (1.0 / AccelerationUnificationFactor(To)));
+		case EUnitType::Mass:				return static_cast<T>(InValue * MassUnificationFactor(From) 				* (1.0 / MassUnificationFactor(To)));
+		case EUnitType::Density:			return static_cast<T>(InValue * DensityUnificationFactor(From)				* (1.0 / DensityUnificationFactor(To)));
+		case EUnitType::Force:				return static_cast<T>(InValue * ForceUnificationFactor(From) 				* (1.0 / ForceUnificationFactor(To)));
+		case EUnitType::Torque:				return static_cast<T>(InValue * TorqueUnificationFactor(From)				* (1.0 / TorqueUnificationFactor(To)));
+		case EUnitType::PositionalImpulse:	return static_cast<T>(InValue * PositionalImpulseUnificationFactor(From)	* (1.0 / PositionalImpulseUnificationFactor(To)));
+		case EUnitType::Frequency:			return static_cast<T>(InValue * FrequencyUnificationFactor(From) 			* (1.0 / FrequencyUnificationFactor(To)));
+		case EUnitType::DataSize:			return static_cast<T>(InValue * DataSizeUnificationFactor(From) 			* (1.0 / DataSizeUnificationFactor(To)));
 		case EUnitType::LuminousFlux:		return InValue;	
-		case EUnitType::Time:				return InValue * TimeUnificationFactor(From) 				* (1.0 / TimeUnificationFactor(To));
-		case EUnitType::Multipliers:		return InValue * MultiplierUnificationFactor(From) 			* (1.0 / MultiplierUnificationFactor(To));
-		case EUnitType::Stress:				return InValue * StressUnificationFactor(From)				* (1.0 / StressUnificationFactor(To));
+		case EUnitType::Time:				return static_cast<T>(InValue * TimeUnificationFactor(From) 				* (1.0 / TimeUnificationFactor(To)));
+		case EUnitType::Multipliers:		return static_cast<T>(InValue * MultiplierUnificationFactor(From) 			* (1.0 / MultiplierUnificationFactor(To)));
+		case EUnitType::Stress:				return static_cast<T>(InValue * StressUnificationFactor(From)				* (1.0 / StressUnificationFactor(To)));
 		// Temperature conversion is not just a simple multiplication, so needs special treatment
 		case EUnitType::Temperature:
 		{
@@ -111,24 +109,24 @@ T FUnitConversion::Convert(T InValue, EUnit From, EUnit To)
 			// Put it into kelvin
 			if (From == EUnit::Celsius)
 			{
-				NewValue = NewValue + 273.15f;
+				NewValue = NewValue + 273.15;
 			}
 			else if (From == EUnit::Farenheit)
 			{
-				NewValue = (NewValue + 459.67f) * 5.f/9.f;
+				NewValue = (NewValue + 459.67) * 5.0/9.0;
 			}
 			// And out again
 			if (To == EUnit::Celsius)
 			{
-				return NewValue - 273.15f;
+				return static_cast<T>(NewValue - 273.15);
 			}
 			else if (To == EUnit::Farenheit)
 			{
-				return NewValue * 9.f/5.f - 459.67f;
+				return static_cast<T>(NewValue * 9.0/5.0 - 459.67);
 			}
 			else
 			{
-				return NewValue;
+				return static_cast<T>(NewValue);
 			}
 		}
 
