Fix `FCsvProfiler::SetMetadataInternal` to work under transactions.

diff --git a/Engine/Source/Programs/AutoRTFMTests/Private/CsvProfilerTests.cpp b/Engine/Source/Programs/AutoRTFMTests/Private/CsvProfilerTests.cpp
--- a/Engine/Source/Programs/AutoRTFMTests/Private/CsvProfilerTests.cpp
+++ b/Engine/Source/Programs/AutoRTFMTests/Private/CsvProfilerTests.cpp
@@ -137,4 +137,9 @@ TEST_CASE("CsvProfiler")
 	{
 		FCsvProfiler::RecordCustomStat("Wowwee", CSV_CATEGORY_INDEX(MyDoubleCategory), static_cast<double>(42), ECsvCustomStatOp::Accumulate);
 	});
+
+	AutoRTFM::Testing::Commit([&]
+	{
+		FCsvProfiler::SetMetadata(TEXT("SomeKey"), TEXT("SomeValue"));
+	});
 }
diff --git a/Engine/Source/Runtime/Core/Private/ProfilingDebugging/CsvProfiler.cpp b/Engine/Source/Runtime/Core/Private/ProfilingDebugging/CsvProfiler.cpp
--- a/Engine/Source/Runtime/Core/Private/ProfilingDebugging/CsvProfiler.cpp
+++ b/Engine/Source/Runtime/Core/Private/ProfilingDebugging/CsvProfiler.cpp
@@ -325,7 +325,7 @@ struct FCsvPersistentCustomStats
 {
 	void RecordStats()
 	{
-		FScopeLock Lock(&Cs);
+		UE::TScopeLock Lock(Cs);
 		for (FCsvPersistentCustomStatBase* BaseStat : Stats)
 		{
 			switch(BaseStat->GetStatType())
@@ -349,7 +349,7 @@ struct FCsvPersistentCustomStats
 	TCsvPersistentCustomStat<T>* GetOrCreatePersistentCustomStat(FName Name, int32 CategoryIndex, bool bResetEachFrame)
 	{
 		LLM_SCOPE(ELLMTag::CsvProfiler);
-		FScopeLock Lock(&Cs);
+		UE::TScopeLock Lock(Cs);
 		FCsvUniqueStatID Id(Name, CategoryIndex);
 		FCsvPersistentCustomStatBase** FindStat = StatLookup.Find(Id.Hash);
 		if (FindStat)
@@ -767,7 +767,7 @@ public:
 
 	FString GetCategoryNameByIndex(int32 Index) const
 	{
-		FScopeLock Lock(&CS);
+		UE::TScopeLock Lock(CS);
 		return CategoryNames[Index];
 	}
 
@@ -778,7 +778,7 @@ public:
 
 	int32 GetCategoryIndex(const FString& CategoryName) const
 	{
-		FScopeLock Lock(&CS);
+		UE::TScopeLock Lock(CS);
 		const int32* CategoryIndex = CategoryNameToIndex.Find(CategoryName.ToLower());
 		if (CategoryIndex)
 		{
@@ -846,7 +846,7 @@ public:
 		{
 			bEnableByDefault = false;
 		}
-		FScopeLock Lock(&CS);
+		UE::TScopeLock Lock(CS);
 		{
 			Index = GetCategoryIndex(CategoryName);
 			if (ensureMsgf(Index == -1, TEXT("CSV stat category already declared: %s. Note: Categories are not case sensitive"), *CategoryName))
@@ -2404,7 +2404,7 @@ public:
 	static FORCENOINLINE FCsvProfilerThreadData* CreateTLSData()
 	{
 		QUICK_SCOPE_CYCLE_COUNTER(CSVProfiler_ThreadData_CreateTLSData);
-		FScopeLock Lock(&Singleton.TlsCS);
+		UE::TScopeLock Lock(Singleton.TlsCS);
 
 		FSharedPtr ProfilerThreadPtr = MakeShareable(new FCsvProfilerThreadData());
 		FPlatformTLS::SetTlsValue(Singleton.TlsSlot, ProfilerThreadPtr.Get());
@@ -2438,7 +2438,7 @@ public:
 	static inline void GetTlsInstances(TArray<FSharedPtr>& OutTlsInstances)
 	{
 		QUICK_SCOPE_CYCLE_COUNTER(CSVProfiler_ThreadData_GetTlsInstances);
-		FScopeLock Lock(&Singleton.TlsCS);
+		UE::TScopeLock Lock(Singleton.TlsCS);
 		OutTlsInstances.Empty(Singleton.TlsInstances.Num() + 1);
 
 		for (int32 Index = Singleton.TlsInstances.Num() - 1; Index >= 0; --Index)
@@ -2479,7 +2479,7 @@ public:
 
 		// Clean up dead entries in the thread data array.
 		// This will remove both the current instance, and any others that have expired.
-		FScopeLock Lock(&Singleton.TlsCS);
+		UE::TScopeLock Lock(Singleton.TlsCS);
 		for (auto Iter = Singleton.TlsInstances.CreateIterator(); Iter; ++Iter)
 		{
 			if (!Iter->IsValid())
@@ -3189,7 +3189,7 @@ public:
 			return;
 		}
 		CSV_SCOPED_TIMING_STAT_EXCLUSIVE(CsvProfiler_Stall);
-		FScopeLock Lock(&StallGameThreadCs);
+		UE::TScopeLock Lock(StallGameThreadCs);
 	}
 
 private:
@@ -4048,7 +4048,7 @@ TSharedFuture<FString> FCsvProfiler::EndCapture(FGraphEventRef EventToSignal)
 	// Copy the metadata array for the next FinalizeCsvFile
 	TMap<FString, FString> CopyMetadataMap;
 	{
-		FScopeLock Lock(&MetadataCS);
+		UE::TScopeLock Lock(MetadataCS);
 		CopyMetadataMap = MetadataMap;
 		// Merge the metadata
 		CopyMetadataMap.Append(MoveTemp(NonPersistentMetadataMap));
@@ -4390,7 +4390,7 @@ void FCsvProfiler::SetNonPersistentMetadata(const TCHAR* Key, const TCHAR* Value
 TMap<FString, FString> FCsvProfiler::GetMetadataMapCopy()
 {
 	LLM_SCOPE(ELLMTag::CsvProfiler);
-	FScopeLock Lock(&MetadataCS);
+	UE::TScopeLock Lock(MetadataCS);
 	TMap<FString, FString> MetadataMapCopy = MetadataMap;
 	MetadataMapCopy.Append(NonPersistentMetadataMap);
 	return MetadataMapCopy;
@@ -4406,7 +4406,7 @@ void FCsvProfiler::SetMetadataInternal(const TCHAR* Key, const TCHAR* Value, boo
 	TMap<FString, FString>& CurrentMetadataMap = PersistenceType == EMetadataPersistenceType::Persistent ? MetadataMap : NonPersistentMetadataMap;
 	if (Value == nullptr)
 	{
-		FScopeLock Lock(&MetadataCS);
+		UE::TScopeLock Lock(MetadataCS);
 		if (CurrentMetadataMap.Contains(KeyLower))
 		{
 			UE_LOG(LogCsvProfiler, Display, TEXT("Metadata unset : %s"), *KeyLower);
@@ -4426,7 +4426,7 @@ void FCsvProfiler::SetMetadataInternal(const TCHAR* Key, const TCHAR* Value, boo
 			}
 		}
 		// Only log if the metadata changed, to prevent logspam 
-		FScopeLock Lock(&MetadataCS);
+		UE::TScopeLock Lock(MetadataCS);
 		if (!CurrentMetadataMap.Contains(KeyLower) || CurrentMetadataMap[KeyLower] != ValueStr)
 		{
 			UE_LOG(LogCsvProfiler, Display, TEXT("Metadata set : %s=\"%s\""), *KeyLower, *ValueStr);
diff --git a/Engine/Source/Runtime/Core/Public/ProfilingDebugging/CsvProfiler.h b/Engine/Source/Runtime/Core/Public/ProfilingDebugging/CsvProfiler.h
--- a/Engine/Source/Runtime/Core/Public/ProfilingDebugging/CsvProfiler.h
+++ b/Engine/Source/Runtime/Core/Public/ProfilingDebugging/CsvProfiler.h
@@ -16,12 +16,12 @@
 #include "Containers/UnrealString.h"
 #include "CoreTypes.h"
 #include "Delegates/Delegate.h"
-#include "HAL/CriticalSection.h"
 #include "HAL/PlatformCrt.h"
 #include "HAL/PreprocessorHelpers.h"
 #include "HAL/ThreadSafeCounter.h"
 #include "Misc/EnumClassFlags.h"
 #include "Misc/Guid.h"
+#include "Misc/TransactionallySafeCriticalSection.h"
 #include "ProfilingDebugging/CsvProfilerConfig.h"
 #include "ProfilingDebugging/CsvProfilerTrace.h"
 #include "ProfilingDebugging/MiscTrace.h"
@@ -558,7 +558,7 @@ private:
 	TMap<FString, FString> MetadataMap;
 	TMap<FString, FString> NonPersistentMetadataMap;
 	TQueue<TMap<FString, FString>> MetadataQueue;
-	FCriticalSection MetadataCS;
+	FTransactionallySafeCriticalSection MetadataCS;
 
 	class FCsvStreamWriter* CsvWriter;
 
