Fix unwind info parsing for loaded dll files

Instead of assuming that unwind function table is at beginning of .pdata section use actual virtual address from exeception data directory.

#rnx

diff --git a/Engine/Source/Runtime/Core/Private/ProfilingDebugging/Microsoft/WindowsCallstackTrace.cpp b/Engine/Source/Runtime/Core/Private/ProfilingDebugging/Microsoft/WindowsCallstackTrace.cpp
--- a/Engine/Source/Runtime/Core/Private/ProfilingDebugging/Microsoft/WindowsCallstackTrace.cpp
+++ b/Engine/Source/Runtime/Core/Private/ProfilingDebugging/Microsoft/WindowsCallstackTrace.cpp
@@ -132,6 +132,8 @@ enum
 	UWOP_SET_FPREG			= 3,	// 1 node
 	UWOP_SAVE_NONVOL		= 4,	// 2 nodes
 	UWOP_SAVE_NONVOL_FAR	= 5,	// 3 nodes
+	UWOP_EPILOG				= 6,	// 2 nodes
+	UWOP_SPARE_CODE			= 7,	// 3 nodes
 	UWOP_SAVE_XMM128		= 8,	// 2 nodes
 	UWOP_SAVE_XMM128_FAR	= 9,	// 3 nodes
 	UWOP_PUSH_MACHFRAME		= 10,	// 1 node
@@ -280,48 +282,53 @@ void FBacktracer::AddModule(UPTRINT ModuleBase, const TCHAR* Name)
 		}
 	}
 
-	uint32 NumSections = FileHeader->NumberOfSections;
-	const auto* Sections = (IMAGE_SECTION_HEADER*)(UPTRINT(&(NtHeader->OptionalHeader)) + FileHeader->SizeOfOptionalHeader);
+	const auto* DataDirectories = NtHeader->OptionalHeader.DataDirectory;
+	const auto* ExceptionDirectory = &DataDirectories[IMAGE_DIRECTORY_ENTRY_EXCEPTION];
 
-	// Find ".pdata" section
-	UPTRINT PdataBase = 0;
-	UPTRINT PdataEnd = 0;
-	for (uint32 i = 0; i < NumSections; ++i)
-	{
-		const IMAGE_SECTION_HEADER* Section = Sections + i;
-		if (*(uint64*)(Section->Name) == 0x61'74'61'64'70'2eull) // Sections names are eight bytes and zero padded. This constant is '.pdata'
-		{
-			PdataBase = ModuleBase + Section->VirtualAddress;
-			PdataEnd = PdataBase + Section->SizeOfRawData;
-			break;
-		}
-	}
-
-	if (PdataBase == 0)
+	if (UPTRINT(&(NtHeader->OptionalHeader)) + FileHeader->SizeOfOptionalHeader < UPTRINT(ExceptionDirectory + 1))
 	{
+		// Make sure ExceptionDirectory is in PE headers
 		return;
 	}
 
-	// Count the number of functions. The assumption here is that if we have got this far then there is at least one function
-	uint32 NumFunctions = uint32(PdataEnd - PdataBase) / sizeof(RUNTIME_FUNCTION);
-	if (NumFunctions == 0)
+	if (ExceptionDirectory->VirtualAddress == 0 || ExceptionDirectory->Size == 0)
 	{
 		return;
 	}
 
-	const auto* FunctionTables = (RUNTIME_FUNCTION*)PdataBase;
-	do
+	uint32 NumSections = FileHeader->NumberOfSections;
+	const auto* Sections = (IMAGE_SECTION_HEADER*)(UPTRINT(&(NtHeader->OptionalHeader)) + FileHeader->SizeOfOptionalHeader);
+
+	const RUNTIME_FUNCTION* FunctionTables = nullptr;
+	uint32 NumFunctions = 0;
+
+	for (uint32 i = 0; i < NumSections; ++i)
 	{
-		const RUNTIME_FUNCTION* Function = FunctionTables + NumFunctions - 1;
-		if (uint32(Function->BeginAddress) < uint32(Function->EndAddress))
+		const IMAGE_SECTION_HEADER* Section = Sections + i;
+		uint64 SectionAddress = Section->VirtualAddress;
+		uint64 SectionSize = Section->Misc.VirtualSize;
+
+		// Find section where RUNTIME_FUNCTION table starts in
+		if (ExceptionDirectory->VirtualAddress >= SectionAddress && ExceptionDirectory->VirtualAddress < SectionAddress + SectionSize)
 		{
-			break;
+			// And make sure RUNTIME_FUNCTION table fully fits into section
+			if ((uint64)ExceptionDirectory->VirtualAddress + ExceptionDirectory->Size <= SectionAddress + SectionSize)
+			{
+				FunctionTables = (RUNTIME_FUNCTION*)(ModuleBase + ExceptionDirectory->VirtualAddress);
+				NumFunctions = uint32(ExceptionDirectory->Size) / sizeof(RUNTIME_FUNCTION);
+				break;
+			}
+
+			// Otherwise do not use it
+			return;
 		}
+	}
 
-		--NumFunctions;
+	if (!FunctionTables || NumFunctions == 0)
+	{
+		return;
 	}
-	while (NumFunctions != 0);
-		
+
 	// Allocate some space for the module's function-to-frame-size table
 	auto* OutTable = (FFunction*)Malloc->Malloc(sizeof(FFunction) * NumFunctions);
 	FFunction* OutTableCursor = OutTable;
@@ -421,6 +428,8 @@ void FBacktracer::AddModule(UPTRINT ModuleBase, const TCHAR* Name)
 			case UWOP_SAVE_NONVOL_FAR:	Code += 3; break;
 			case UWOP_SAVE_XMM128:		Code += 2; break;
 			case UWOP_SAVE_XMM128_FAR:	Code += 3; break;
+			case UWOP_EPILOG:			Code += 2; break;
+			case UWOP_SPARE_CODE:		Code += 3; break;
 
 			default:
 #if BACKTRACE_DBGLVL >= 2
