Introduce FChange::GetSize and return its result in FTransaction::DataSize.

The point of this is two-fold:
- SUndoHistory panel will start showing the undo stack size of changes shown in the undo panel. This also makes it easier for developers to check whether their FChange is within the memory bounds they expect.
- Over memory budget handling: UTransBuffer::BeginInternal will start considering custom changes as part of the undo stack size.

diff --git a/Engine/Source/Editor/UnrealEd/Private/EditorTransaction.cpp b/Engine/Source/Editor/UnrealEd/Private/EditorTransaction.cpp
--- a/Engine/Source/Editor/UnrealEd/Private/EditorTransaction.cpp
+++ b/Engine/Source/Editor/UnrealEd/Private/EditorTransaction.cpp
@@ -1086,7 +1086,9 @@ SIZE_T FTransaction::DataSize() const
 	SIZE_T Result=0;
 	for( int32 i=0; i<Records.Num(); i++ )
 	{
-		Result += Records[i].SerializedObject.SerializedData.Num();
+		const FObjectRecord& Record = Records[i];
+		const SIZE_T CustomChangeSize = Record.CustomChange ? Record.CustomChange->GetSize() : 0;
+		Result += Record.SerializedObject.SerializedData.Num() + CustomChangeSize;
 	}
 	return Result;
 }
diff --git a/Engine/Source/Runtime/Core/Public/Misc/Change.h b/Engine/Source/Runtime/Core/Public/Misc/Change.h
--- a/Engine/Source/Runtime/Core/Public/Misc/Change.h
+++ b/Engine/Source/Runtime/Core/Public/Misc/Change.h
@@ -62,6 +62,25 @@ public:
 	/** Prints this change to the log, including sub-changes if there are any.  For compound changes, there might be multiple lines.  You should not need to override this function. */
 	CORE_API virtual void PrintToLog( FFeedbackContext& FeedbackContext, const int32 IndentLevel = 0 );
 
+	/**
+	 * Gets the approximate number of bytes this change takes up in memory.
+	 *
+	 * Used to
+	 *	- determine the size of transactions based on a custom change, e.g. when trimming the oldest transactions when the undo stack size is exceeded.
+	 *	- displaying the stack size, e.g. in the Undo Panel.
+	 *
+	 * A typical implementation returns the size of any dynamically allocated memory, such GetAllocatedSize called on its containers (TArray, TSet, TMap).
+	 * You should consider adding sizeof(FMyChange), if it's not insignificant. Remember: it's enough to return the approx. allocation size; it's
+	 * unreasonable to expect every FChange implementation to return the exact amount, which also is not needed for the purpose of this function.
+	 * Keep in mind, that for nested containers you need to call GetAllocatedSize recursively, e.g. for a "TArray<TArray<T>> Array" you'd need to
+	 * return "Array.GetAllocatedSize() + Algo::Accumulate(Array, 0, [](const TArray<T>& A){ return A.GetAllocatedSize(); })".
+	 *
+	 * This function is introduced in 5.7: too many FChange implementations exist so this function will need to be implemented gradually.
+	 * 
+	 * @return The approximate number of bytes this change takes up in memory.
+	 */
+	virtual SIZE_T GetSize() const { return 0; }
+
 	/** Virtual destructor */
 	virtual ~FChange()
 	{
