Add support for repeating csvExecCmds. Frame indices can be prefixed with an r to denote that the commands should be repeated every N frames.
e.g: -csvExecCmds="r:500:profilegpu" runs a GPU profile every 500 frames.
#DisabledValidations CheckCopyright, CheckBalancedMacros, FoundationTestsBuild, FoundationTestsRun

diff --git a/Engine/Source/Runtime/Core/Private/ProfilingDebugging/CsvProfiler.cpp b/Engine/Source/Runtime/Core/Private/ProfilingDebugging/CsvProfiler.cpp
--- a/Engine/Source/Runtime/Core/Private/ProfilingDebugging/CsvProfiler.cpp
+++ b/Engine/Source/Runtime/Core/Private/ProfilingDebugging/CsvProfiler.cpp
@@ -444,8 +444,8 @@ void CsvEndNamedEvent()
 }
 #endif //CSV_PROFILER_SUPPORT_NAMED_EVENTS
 
-
 static TMap<uint32, TArray<FString>>* GCsvFrameExecCmds = NULL;
+static TMap<uint32, TArray<FString>>* GCsvFrameModuloExecCmds = NULL;
 
 struct FEventExecCmds
 {
@@ -907,7 +907,11 @@ void FCsvProfiler::GetFrameExecCommands(TArray<FString>& OutFrameCommands)
 {
 	check(IsInGameThread());
 	OutFrameCommands.Empty();
-	if (GCsvProfilerIsCapturing && GCsvFrameExecCmds)
+	if (!GCsvProfilerIsCapturing)
+	{
+		return;
+	}
+	if (GCsvFrameExecCmds)
 	{
 		TArray<FString>* FrameCommands = GCsvFrameExecCmds->Find(CaptureFrameNumber);
 		if (FrameCommands)
@@ -915,7 +919,18 @@ void FCsvProfiler::GetFrameExecCommands(TArray<FString>& OutFrameCommands)
 			OutFrameCommands = *FrameCommands;
 		}
 	}
-	if (GCsvProfilerIsCapturing && GCsvEventExecCmds)
+	if (GCsvFrameModuloExecCmds)
+	{
+		for (auto It(GCsvFrameModuloExecCmds->CreateConstIterator()); It; ++It)
+		{
+			uint32 FrameModulo = It.Key();
+			if ((CaptureFrameNumber % FrameModulo) == 0)
+			{
+				OutFrameCommands = It.Value();
+			}
+		}
+	}
+	if (GCsvEventExecCmds)
 	{
 		for (FEventExecCmds& EventPair : *GCsvEventExecCmds)
 		{
@@ -4688,6 +4703,7 @@ void FCsvProfiler::Init()
 	if (FParse::Value(FCommandLine::Get(), TEXT("-csvExecCmds="), CsvExecCommandsStr, false))
 	{
 		GCsvFrameExecCmds = new TMap<uint32, TArray<FString>>();
+		GCsvFrameModuloExecCmds = new TMap<uint32, TArray<FString>>();
 
 		TArray<FString> CsvExecCommandsList;
 		if (CsvExecCommandsStr.ParseIntoArray(CsvExecCommandsList, TEXT(","), true) > 0)
@@ -4697,15 +4713,33 @@ void FCsvProfiler::Init()
 				int32 ColonIndex = -1;
 				if (FrameAndCommand.FindChar(TEXT(':'), ColonIndex))
 				{
-					FString FrameStr = FrameAndCommand.Mid(0,ColonIndex);
+					FString FrameStr = FrameAndCommand.Mid(0,ColonIndex).ToLower();
 					FString CommandStr = FrameAndCommand.Mid(ColonIndex+1);
+
+					// Modulo frames are prefixed with "r". The command will be executed every n frames
+					bool bModuloCommand = FrameStr.StartsWith(TEXT("r"));
+					if (bModuloCommand)
+					{
+						FrameStr = FrameAndCommand.Mid(1, ColonIndex);
+					}
+
+					TMap<uint32, TArray<FString>>* ExecCmdList = bModuloCommand ? GCsvFrameModuloExecCmds : GCsvFrameExecCmds;
 					uint32 Frame = FCString::Atoi(*FrameStr);
-					if (!GCsvFrameExecCmds->Find(Frame))
+					if (!ExecCmdList->Find(Frame))
+					{
+						ExecCmdList->Add(Frame, TArray<FString>());
+					}
+					(*ExecCmdList)[Frame].Add(CommandStr);
+
+					if ( bModuloCommand )
 					{
-						GCsvFrameExecCmds->Add(Frame, TArray<FString>());
+						UE_LOG(LogCsvProfiler, Display, TEXT("Added CsvExecCommand (repeating) - Modulo: %d : %s"), Frame, *CommandStr);
 					}
-					(*GCsvFrameExecCmds)[Frame].Add(CommandStr);
-					UE_LOG(LogCsvProfiler, Display, TEXT("Added CsvExecCommand - frame %d : %s"), Frame, *CommandStr);
+					else
+					{
+						UE_LOG(LogCsvProfiler, Display, TEXT("Added CsvExecCommand - frame %d : %s"), Frame, *CommandStr);
+					}
+
 				}
 			}
 		}
