Fixed incorrect handling of spurious wakes in MicrosoftPlatformManualResetEvent. Needs to return only if the address was successfully waited on AND the wait has been signaled.

diff --git a/Engine/Source/Runtime/Core/Private/Microsoft/MicrosoftPlatformManualResetEvent.cpp b/Engine/Source/Runtime/Core/Private/Microsoft/MicrosoftPlatformManualResetEvent.cpp
--- a/Engine/Source/Runtime/Core/Private/Microsoft/MicrosoftPlatformManualResetEvent.cpp
+++ b/Engine/Source/Runtime/Core/Private/Microsoft/MicrosoftPlatformManualResetEvent.cpp
@@ -35,9 +35,18 @@ FORCENOINLINE bool FMicrosoftManualResetEvent::WaitForSlow(FMonotonicTimeSpan Wa
 	LowLevelTasks::FOversubscriptionScope _;
 
 	const DWORD WaitMs = WaitTime.IsInfinity() ? INFINITE : DWORD(FPlatformMath::CeilToInt64(WaitTime.ToMilliseconds()));
-	if (LIKELY(WaitOnAddress(&bWait, &bLocalWait, sizeof(bool), WaitMs) || GetLastError() == ERROR_TIMEOUT))
+	if (LIKELY(WaitOnAddress(&bWait, &bLocalWait, sizeof(bool), WaitMs)))
 	{
-		return !bWait.load(std::memory_order_acquire);
+		const bool bWaitCompleted = !bWait.load(std::memory_order_acquire);
+		if (bWaitCompleted)
+		{
+			return true;
+		}
+	}
+	else if (GetLastError() == ERROR_TIMEOUT)
+	{
+		const bool bWaitCompleted = !bWait.load(std::memory_order_acquire);
+		return bWaitCompleted;
 	}
 
 	// Handle a spurious wake by waiting until the wait time has elapsed one more time because WaitUntilSlow
