Append repeating csvExecCmds so they don't suppress existing ones that occur on the same frame
#tests tested on console
#DisabledValidations CheckCopyright, CheckBalancedMacros, FoundationTestsBuild, FoundationTestsRun

diff --git a/Engine/Source/Runtime/Core/Private/ProfilingDebugging/CsvProfiler.cpp b/Engine/Source/Runtime/Core/Private/ProfilingDebugging/CsvProfiler.cpp
--- a/Engine/Source/Runtime/Core/Private/ProfilingDebugging/CsvProfiler.cpp
+++ b/Engine/Source/Runtime/Core/Private/ProfilingDebugging/CsvProfiler.cpp
@@ -926,7 +926,7 @@ void FCsvProfiler::GetFrameExecCommands(TArray<FString>& OutFrameCommands)
 			uint32 FrameModulo = It.Key();
 			if ((CaptureFrameNumber % FrameModulo) == 0)
 			{
-				OutFrameCommands = It.Value();
+				OutFrameCommands.Append(It.Value());
 			}
 		}
 	}
