- Added UE_DEPRECATED_ERROR which can give a deprecation message, but will also fail to compile. It is useful when could cause dataloss if a  licesnsee were to ignore the deprecation message, with also giving a custom message on how to fix the issue. WIll be used when moving properties into EditirOnlyData objects to tell users they must replace usage with accessors.
- Added a couple of helpers for preformatted messages for moving properties

diff --git a/Engine/Source/Runtime/Core/Public/Misc/CoreMiscDefines.h b/Engine/Source/Runtime/Core/Public/Misc/CoreMiscDefines.h
--- a/Engine/Source/Runtime/Core/Public/Misc/CoreMiscDefines.h
+++ b/Engine/Source/Runtime/Core/Public/Misc/CoreMiscDefines.h
@@ -299,6 +299,30 @@ enum EPerElement {PerElement};
 	#define UE_DEPRECATED(Version, Message) [[deprecated(Message " - Please update your code to the new API before upgrading to the next release, otherwise your project will no longer compile.")]]
 
 /**
+ * An alternative use of the deprecation feature is to note an API breaking change that could cause data loss or similar, so we use this when 
+ * tagging a property that is a compile error to use (for instance by changing its type to FDeprecationCompileError type, which is incompatible with 
+ * anything), but it can still show the deprecation message telling users how to fix their code. This should only be used when a nice, clean 
+ * deprecation is not possible.
+ * YOU MUST FOLLOW THIS WITH SOMETHING THAT CAUSES AN ERROR - for instance, change the variable type to FDeprecationCompileError so any use will fail,
+ * or use one of the below helpers
+ */
+	struct FDeprecationCompileError {};
+	#define UE_DEPRECATED_ERROR(Version, Message) [[deprecated(Message " - ** This is a compile error so as to avoid data loss, or similar. ** You will need to update your code to get it to compile correctly.")]]
+	
+/** A version of the above that declares a compilation-error type with a message indicating it has moved and you 
+ * need to use accessor functions instead.
+ */
+	#define UE_DEPRECATED_ERROR_WITH_MOVED_MEMBER_MESSAGE(Version, Member, Message) \
+		[[deprecated(Message " ** This is a compile error so as to avoid data loss, or similar. ** You will need to update your code to get it to compile correctly.")]] \
+		FDeprecationCompileError Member;
+	
+/** A version of the above that prints out a premade message with member name */
+	#define UE_DEPRECATED_ERROR_WITH_MOVED_MEMBER(Version, Member) UE_DEPRECATED_ERROR_WITH_MOVED_MEMBER_MESSAGE(Version, Member, "The member " #Member " has moved. Use the Get/Set" #Member " accessors instead.")
+
+/** This version prints out a message for the accessor without the b(SetFoo, not SetbFoo), but as such you need to pass in the bool Member without the b */
+	#define UE_DEPRECATED_ERROR_WITH_MOVED_BOOL_MEMBER(Version, Member) UE_DEPRECATED_ERROR_WITH_MOVED_MEMBER_MESSAGE(Version, b##Member, "The member b" #Member " has moved. Use the Get/Set" #Member " accessors instead.")
+
+ /**
  * Macro for marking up engine-internal code, functions and types. This informs the users that there is no guarantee that the API won't change or be removed in 
  *  future releases and that they should use at their risk. 
  *  Using an API marked as such in a non-engine module will trigger a compilation warning (this relies on the c++ [[deprecated]] attribute internally) that can be suppressed by either : 
