Asset Insights trace control UI

#DisabledValidations nonUnityAndStaticAnalysis

diff --git a/Engine/Source/Runtime/Core/Private/HAL/LowLevelMemTracker.cpp b/Engine/Source/Runtime/Core/Private/HAL/LowLevelMemTracker.cpp
--- a/Engine/Source/Runtime/Core/Private/HAL/LowLevelMemTracker.cpp
+++ b/Engine/Source/Runtime/Core/Private/HAL/LowLevelMemTracker.cpp
@@ -928,7 +928,7 @@ uint8 LlmGetActiveTag()
 
 #endif  // UE_ENABLE_ARRAY_SLACK_TRACKING
 
-UE_TRACE_CHANNEL(MemTagChannel, "Memory overview", true)
+UE_TRACE_CHANNEL(MemTagChannel, "Memory overview")
 
 UE_TRACE_EVENT_BEGIN(LLM, TagsSpec, NoSync|Important)
 	UE_TRACE_EVENT_FIELD(const void*, TagId)
