Fix module cleanup accounting only for loaded modules, when static initialization means a lot more might have active UObjects

Thix fixes ghost modules that would mysteriously never unload, and keep GC ownership through UClass/CDO UObjects that are created at static initialization, while the modules were never technically loaded.

#rnx

diff --git a/Engine/Source/Runtime/Core/Private/Modules/ModuleManager.cpp b/Engine/Source/Runtime/Core/Private/Modules/ModuleManager.cpp
--- a/Engine/Source/Runtime/Core/Private/Modules/ModuleManager.cpp
+++ b/Engine/Source/Runtime/Core/Private/Modules/ModuleManager.cpp
@@ -618,13 +618,15 @@ void FModuleManager::InternalFreeLibrary(FName ModuleName, void* Handle)
 
 		TArray<FName> UnloadedModules;
 
-		// Remove statically linked initializer info, or the next module load for that library will crash on a dangling pointer
-		for (const TPair<FName, ModuleInfoRef>& OtherModuleNameAndInfo : Modules)
+		// Identify modules to unload by their statically linked identifier & their library path,
+		// then remove statically linked initializer info - if it exists it means the module needs cleanup
+		for (const TPair<FName, FInitializeStaticallyLinkedModule>& OtherModuleNameAndInfo : StaticallyLinkedModuleInitializers)
 		{
-			const ModuleInfoRef& OtherModuleInfo = OtherModuleNameAndInfo.Value;
-			const FString OtherModuleLibraryName = FPaths::ConvertRelativePathToFull(OtherModuleInfo->Filename);
+			TMap<FName, FString> OtherModulesPathMap;
+			FindModulePaths(*OtherModuleNameAndInfo.Key.ToString(), OtherModulesPathMap);
+			const FString* OtherModuleLibraryNamePtr = OtherModulesPathMap.Find(OtherModuleNameAndInfo.Key);
 
-			if (OtherModuleLibraryName == LibraryToUnload)
+			if (OtherModuleLibraryNamePtr && *OtherModuleLibraryNamePtr == LibraryToUnload)
 			{
 				UE_LOG(LogModuleManager, Verbose, TEXT("InternalFreeLibrary: cleaning up for module '%s'"),
 					*OtherModuleNameAndInfo.Key.ToString());
@@ -1856,6 +1858,9 @@ void FModuleManager::ProcessPendingStaticallyLinkedModuleInitializers() const
 
 	for (TPair<FLazyName, FInitializeStaticallyLinkedModule>& Module : PendingStaticallyLinkedModuleInitializers)
 	{
+		UE_LOG(LogModuleManager, VeryVerbose, TEXT("ProcessPendingStaticallyLinkedModuleInitializers: '%s'"),
+			*FName(Module.Key).ToString());
+
 		FName NameKey = FName(Module.Key);
 		checkf(!StaticallyLinkedModuleInitializers.Contains(NameKey), TEXT("Duplicate module '%s' registered"), *NameKey.ToString());
 		StaticallyLinkedModuleInitializers.Add(NameKey, Module.Value);
