Fix crash on merge module unload

Building the FName array from const& isn't legal because the source variable is getting removed before use

#rnx

diff --git a/Engine/Source/Runtime/Core/Private/Modules/ModuleManager.cpp b/Engine/Source/Runtime/Core/Private/Modules/ModuleManager.cpp
--- a/Engine/Source/Runtime/Core/Private/Modules/ModuleManager.cpp
+++ b/Engine/Source/Runtime/Core/Private/Modules/ModuleManager.cpp
@@ -620,7 +620,7 @@ void FModuleManager::InternalFreeLibrary(FName ModuleName, void* Handle)
 
 		// Identify modules to unload by their statically linked identifier & their library path,
 		// then remove statically linked initializer info - if it exists it means the module needs cleanup
-		for (const TPair<FName, FInitializeStaticallyLinkedModule>& OtherModuleNameAndInfo : StaticallyLinkedModuleInitializers)
+		for (TPair<FName, FInitializeStaticallyLinkedModule> OtherModuleNameAndInfo : StaticallyLinkedModuleInitializers)
 		{
 			TMap<FName, FString> OtherModulesPathMap;
 			FindModulePaths(*OtherModuleNameAndInfo.Key.ToString(), OtherModulesPathMap);
