Implent RtlGetDeviceFamilyInfoEnum for detecting if a device is considered Handhled
https://learn.microsoft.com/en-us/gaming/gdk/docs/gdk-dev/pc-dev/handheld/handheld-guidance

Adding an additional hardcoded list of APUs that we also know are to be considered as Handheld devices.

diff --git a/Engine/Source/Programs/UnrealBuildTool/Platform/Windows/UEBuildWindows.cs b/Engine/Source/Programs/UnrealBuildTool/Platform/Windows/UEBuildWindows.cs
--- a/Engine/Source/Programs/UnrealBuildTool/Platform/Windows/UEBuildWindows.cs
+++ b/Engine/Source/Programs/UnrealBuildTool/Platform/Windows/UEBuildWindows.cs
@@ -1962,6 +1962,7 @@ namespace UnrealBuildTool
 			LinkEnvironment.SystemLibraries.Add("dbghelp.lib");
 			LinkEnvironment.SystemLibraries.Add("comctl32.lib");
 			LinkEnvironment.SystemLibraries.Add("Winmm.lib");
+			LinkEnvironment.SystemLibraries.Add("ntdll.lib");
 			LinkEnvironment.SystemLibraries.Add("kernel32.lib");
 			LinkEnvironment.SystemLibraries.Add("user32.lib");
 			LinkEnvironment.SystemLibraries.Add("gdi32.lib");
diff --git a/Engine/Source/Runtime/Core/Private/Windows/WindowsPlatformMisc.cpp b/Engine/Source/Runtime/Core/Private/Windows/WindowsPlatformMisc.cpp
--- a/Engine/Source/Runtime/Core/Private/Windows/WindowsPlatformMisc.cpp
+++ b/Engine/Source/Runtime/Core/Private/Windows/WindowsPlatformMisc.cpp
@@ -966,6 +966,9 @@ void FWindowsPlatformMisc::PlatformInit()
 	// Timer resolution.
 	UE_LOG(LogInit, Log, TEXT("High frequency timer resolution =%f MHz"), 0.000001 / FPlatformTime::GetSecondsPerCycle() );
 
+	// If the Device has been identified as a Handheld device
+	UE_LOG(LogInit, Log, TEXT("Handheld Device: %d"), FWindowsPlatformMisc::IsDeviceHandheld() ? 1 : 0);
+
 	// Register on the game thread.
 	FWindowsPlatformStackWalk::RegisterOnModulesChanged();
 
@@ -4282,6 +4285,53 @@ FText FWindowsPlatformMisc::GetFileManagerName()
 	return NSLOCTEXT("WindowsPlatform", "FileManagerName", "Explorer");
 }
 
+bool FWindowsPlatformMisc::IsDeviceHandheld()
+{
+#ifndef DEVICEFAMILYDEVICEFORM_GAMING_HANDHELD
+	#define DEVICEFAMILYDEVICEFORM_GAMING_HANDHELD    0x0000002E 
+#endif
+
+	static const TSet<FString> HandheldAPU_CPUBrandSet =
+	{
+
+		TEXT("AMD Custom APU 0405"),		//Steam Deck LCD
+		TEXT("AMD Custom APU 0932"),		//Steam Deck OLED
+		TEXT("AMD Ryzen Z1"),
+		TEXT("AMD Ryzen Z1 Extreme"),
+		TEXT("AMD Ryzen Z2 A"),				//ROG Xbox Ally
+		TEXT("AMD Ryzen AI Z2 Extreme"),	//ROG Xbox Ally X
+		TEXT("AMD Ryzen Z2 Go")
+	};
+
+	static struct FWindowsHandheldCheck
+	{
+		bool DeviceHandheld = false;
+
+		FWindowsHandheldCheck()
+		{
+			// https://learn.microsoft.com/en-us/gaming/gdk/docs/gdk-dev/pc-dev/handheld/handheld-guidance
+			DWORD DeviceForm;
+			::RtlGetDeviceFamilyInfoEnum(nullptr, nullptr, &DeviceForm);
+
+			if (DeviceForm == DEVICEFAMILYDEVICEFORM_GAMING_HANDHELD)
+			{
+				DeviceHandheld = true;
+			}
+
+			if (!DeviceHandheld)
+			{
+				// Some devices are not correctly reporting that they are an Handheld device, check our hard-coded table of Handheld CPU Brand
+				if (HandheldAPU_CPUBrandSet.Contains(FWindowsPlatformMisc::GetCPUBrand()))
+				{
+					DeviceHandheld = true;
+				}
+			}
+		}
+	} HandheldCheck;
+
+	return HandheldCheck.DeviceHandheld;
+}
+
 static bool IsRunningOnBatteryInternal(SYSTEM_POWER_STATUS Status)
 {
 	BYTE BatteryPresentFlags = 4  // Critical-the battery capacity is at less than five percent.
diff --git a/Engine/Source/Runtime/Core/Public/Windows/WindowsPlatformMisc.h b/Engine/Source/Runtime/Core/Public/Windows/WindowsPlatformMisc.h
--- a/Engine/Source/Runtime/Core/Public/Windows/WindowsPlatformMisc.h
+++ b/Engine/Source/Runtime/Core/Public/Windows/WindowsPlatformMisc.h
@@ -332,6 +332,12 @@ struct FWindowsPlatformMisc
 		return true;
 	}
 
+	/**
+	 * Returns whether the device is considered an Handheld device. It's not a dynamic property, the status is cached at startup.
+	 */
+
+	static CORE_API bool IsDeviceHandheld();
+	
 	/**
 	 * Returns whether the platform is running on battery power or not. For Windows laptops this is true regardless of charging state.
 	 */
