Platform: Delete code for supporting a 4-byte TCHAR

diff --git a/Engine/Source/Runtime/Core/Private/Android/AndroidPlatformString.cpp b/Engine/Source/Runtime/Core/Private/Android/AndroidPlatformString.cpp
deleted file mode 100644
diff --git a/Engine/Source/Runtime/Core/Public/Android/AndroidPlatformString.h b/Engine/Source/Runtime/Core/Public/Android/AndroidPlatformString.h
--- a/Engine/Source/Runtime/Core/Public/Android/AndroidPlatformString.h
+++ b/Engine/Source/Runtime/Core/Public/Android/AndroidPlatformString.h
@@ -1,777 +1,15 @@
 // Copyright Epic Games, Inc. All Rights Reserved.
 
-
-/*=============================================================================================
-	AndroidString.h: Android platform string classes
-==============================================================================================*/
-
 #pragma once
 
 #include "Misc/Build.h"
-#include "HAL/Platform.h"
-
-/**
- * Android string implementation
- **/
-
-#if PLATFORM_TCHAR_IS_CHAR16
-
-// By default we now use 2-byte strings on Android.
-
+#include "Misc/CoreMiscDefines.h"
 #include "GenericPlatform/GenericWidePlatformString.h"
 
+UE_DEPRECATED(5.7, "FAndroidPlatformString has been deprecated, please use FPlatformString instead.")
 typedef FGenericWidePlatformString FAndroidPlatformString;
 
-#else
-
-// NOTE: This legacy 4-byte implementation will probably be removed in the future.
-
-// @todo android: probably rewrite/simplify most of this.  currently it converts all wide chars
-// to ANSI.  The Android NDK appears to have wchar_t support, but many of the functions appear
-// to just be stubs to the non-wide versions.  This doesn't work for obvious reasons.
-
-#include <alloca.h>
-#include <cstdlib>
-#include "Misc/Char.h"
-#include "GenericPlatform/GenericPlatformMemory.h"
-#include "GenericPlatform/GenericPlatformString.h"
-
-struct FAndroidPlatformString : public FGenericPlatformString
-{
-	using FGenericPlatformString::Stricmp;
-	using FGenericPlatformString::Strncmp;
-	using FGenericPlatformString::Strnicmp;
-
-	template <typename CharType>
-	static inline CharType* Strupr(CharType* Dest, SIZE_T DestCount)
-	{
-		for (CharType* Char = Dest; *Char && DestCount > 0; Char++, DestCount--)
-		{
-			*Char = TChar<CharType>::ToUpper(*Char);
-		}
-		return Dest;
-	}
-
-	/** 
-	 * Widechar implementation 
-	 **/
-	static inline WIDECHAR* Strcpy(WIDECHAR* Dest, const WIDECHAR* Src)
-	{
-		if (!Dest || !Src)
-		{
-			return Dest;
-		}
-
-		int Pos = 0;
-		while (Src[Pos])
-		{
-			Dest[Pos] = Src[Pos];
-			++Pos;
-		}
-
-		Dest[Pos] = 0;
-
-		return Dest;
-	}
-
-	// This version was deprecated because Strcpy taking a size field does not match the CRT strcpy, and on some platforms the size field was being ignored.
-	// If the memzeroing causes a performance problem, request a new function StrcpyTruncate.
-	UE_DEPRECATED(5.6, "Use Strncpy instead. Note that Strncpy has a behavior difference from Strcpy: it memzeroes the entire DestCount-sized buffer after the end of string.")
-	static UE_FORCEINLINE_HINT WIDECHAR* Strcpy(WIDECHAR* Dest, SIZE_T DestCount, const WIDECHAR* Src)
-	{
-		return Strcpy(Dest, Src);
-	}
-
-	static inline WIDECHAR* Strncpy(WIDECHAR* Dest, const WIDECHAR* Src, SIZE_T MaxLen)
-	{
-		if (!Dest || !Src)
-		{
-			return Dest;
-		}
-
-		int Pos = 0;
-		while ((Pos < MaxLen) && Src[Pos])
-		{
-			Dest[Pos] = Src[Pos];
-			++Pos;
-		}
-
-		while (Pos < MaxLen)
-		{
-			Dest[Pos] = 0;
-			++Pos;
-		}
-
-		Dest[MaxLen-1]=0;
-		return Dest;
-	}
-
-	static inline WIDECHAR* Strcat(WIDECHAR* Dest, const WIDECHAR* Src)
-	{
-		if (!Dest || !Src)
-		{
-			return Dest;
-		}
-		if (!*Src)
-		{
-			return Dest;
-		}
-
-		WIDECHAR* NewDest = Dest + Strlen(Dest);
-		do
-		{
-			*NewDest++ = *Src++;
-		} while (*Src);
-		*NewDest = 0;
-		return Dest;
-	}
-
-	// This version was deprecated because Strcat taking a size field does not match the CRT strcat, and on some platforms the size field was being ignored.
-	UE_DEPRECATED(5.6, "Use Strncat instead. !!NOTE THAT STRNCAT takes SrcLen rather than DestCount. You must call Strncat(Dest, Src, DestCount - Strlen(Dest) - 1).")
-	static UE_FORCEINLINE_HINT WIDECHAR* Strcat(WIDECHAR* Dest, SIZE_T DestCount, const WIDECHAR* Src)
-	{
-		return Strcat(Dest, Src);
-	}
-
-	static inline WIDECHAR* Strncat(WIDECHAR* Dest, const WIDECHAR* Src, SIZE_T SrcLen)
-	{
-		if (!Dest || !Src)
-		{
-			return Dest;
-		}
-		if (!SrcLen || !*Src)
-		{
-			return Dest;
-		}
-		WIDECHAR* NewDest = Dest + Strlen(Dest);
-		SIZE_T AppendedCount = 0;
-		do
-		{
-			*NewDest++ = *Src++;
-			++AppendedCount;
-		} while (AppendedCount < SrcLen && *Src);
-		*NewDest = 0;
-		return Dest;
-	}
-
-	static inline int32 Strcmp( const WIDECHAR* String1, const WIDECHAR* String2 )
-	{
-		for (; *String1 || *String2; String1++, String2++)
-		{
-			if (*String1 != *String2)
-			{
-				return *String1 - *String2;
-			}
-		}
-		return 0;
-	}
-
-	static inline int32 Strlen( const WIDECHAR* String )
-	{
-		if (!String)
-			return 0;
-
-		int Len = 0;
-		while (String[Len])
-		{
-			++Len;
-		}
-
-		return Len;
-	}
-
-	static inline int32 Strnlen( const WIDECHAR* String, SIZE_T StringSize )
-	{
-		int Len = 0;
-		while (StringSize-- > 0 && *String++)
-		{
-			++Len;
-		}
-		return Len;
-	}
-
-	static inline void CopyWideToAnsi(ANSICHAR* Dest, const WIDECHAR* Src)
-	{
-		if (!Src || !Dest)
-			return;
-
-		int Pos = 0;
-		while (Src[Pos])
-		{
-			if (Src[Pos] <= 255)
-			{
-				Dest[Pos] = static_cast<ANSICHAR>(Src[Pos]);
-			}
-			else
-			{
-				Dest[Pos] = '?';
-			}
-
-			++Pos;
-		}
-
-		Dest[Pos] = 0;
-	}
-
-	static inline void CopyAnsiToWide(WIDECHAR* Dest, const ANSICHAR* Src)
-	{
-		if (!Src || !Dest)
-			return;
-
-		int Pos = 0;
-		while (Src[Pos])
-		{
-			Dest[Pos] = Src[Pos];
-			++Pos;
-		}
-
-		Dest[Pos] = 0;
-	}
-
-	static inline const WIDECHAR* Strstr( const WIDECHAR* String, const WIDECHAR* Find)
-	{
-		WIDECHAR FindChar = *Find;
-
-		// Always find an empty string
-		if (!FindChar)
-			return String;
-
-		++Find;
-		size_t MemCmpLen = wcslen(Find);
-		const WIDECHAR* FoundChar;
-		for (;;)
-		{
-			FoundChar = wcschr(String, FindChar);
-			if (!FoundChar)
-			{
-				//	No more instances of FindChar in String, Find does not exist in String
-				break;
-			}
-
-			// Found FindChar, now set String to one character after FoundChar...
-			// We can now compare characters after FoundChar to Find, which is pointing at original Find[1], to see if the rest of the characters match
-			String = FoundChar + 1;
-			if (!wmemcmp(String, Find, MemCmpLen))
-			{
-				// Strings match, return pointer to beginning of Find instance in String
-				return FoundChar;
-			}
-
-			// No match, String is already ready to loop again to find the next instance of FindChar
-		}
-
-		return nullptr;
-	}
-
-	static inline const WIDECHAR* Strchr( const WIDECHAR* String, WIDECHAR C)
-	{
-		if (!String)
-			return NULL;
-
-		int Pos = 0;
-		while (String[Pos])
-		{
-			if(String[Pos] == C)
-			{
-				return &(String[Pos]);
-			}
-
-			++Pos;
-		}
-
-		if (C == 0)
-		{
-			return &(String[Pos]);
-		}
-
-		return NULL;
-	}
-
-	static inline const WIDECHAR* Strrchr( const WIDECHAR* String, WIDECHAR C)
-	{
-		if (!String)
-			return NULL;
-
-		const WIDECHAR* Last = NULL;
-
-		int Pos = 0;
-		while (String[Pos])
-		{
-			if(String[Pos] == C)
-			{
-				Last = &(String[Pos]);
-			}
-
-			++Pos;
-		}
-
-		if (C == 0)
-		{
-			Last = &(String[Pos]);
-		}
-
-		return Last;
-	}
-
-	static inline int32 Atoi(const WIDECHAR* String)
-	{
-		int StringLen = Strlen(String);
-		ANSICHAR* AnsiString = (ANSICHAR*)FMemory_Alloca(StringLen+1);
-		CopyWideToAnsi(AnsiString, String);
-
-		return atoi(AnsiString);
-	}
-
-	static inline int64 Atoi64(const WIDECHAR* String)
-	{
-		int StringLen = Strlen(String);
-		ANSICHAR* AnsiString = (ANSICHAR*)FMemory_Alloca(StringLen+1);
-		CopyWideToAnsi(AnsiString, String);
-
-		return strtoll(AnsiString, NULL, 10);
-	}
-
-	static inline float Atof(const WIDECHAR* String)
-	{
-		int StringLen = Strlen(String);
-		ANSICHAR* AnsiString = (ANSICHAR*)FMemory_Alloca(StringLen+1);
-		CopyWideToAnsi(AnsiString, String);
-
-		return (float)atof(AnsiString);
-	}
-
-	static inline double Atod(const WIDECHAR* String)
-	{
-		int StringLen = Strlen(String);
-		ANSICHAR* AnsiString = (ANSICHAR*)FMemory_Alloca(StringLen+1);
-		CopyWideToAnsi(AnsiString, String);
-
-		return atof(AnsiString);
-	}
-
-	static inline int32 Strtoi( const WIDECHAR* Start, WIDECHAR** End, int32 Base ) 
-	{
-		int StartLen = Strlen(Start);
-		ANSICHAR* AnsiStart = (ANSICHAR*)FMemory_Alloca(StartLen+1);
-		CopyWideToAnsi(AnsiStart, Start);
-
-		ANSICHAR* AnsiEnd = NULL;
-
-		int32 Res = strtol(AnsiStart, &AnsiEnd, Base);
-
-		if (End)
-		{
-			if (AnsiEnd == NULL)
-			{
-				*End = NULL;
-			}
-			else
-			{
-				*End = (WIDECHAR*)(Start + (AnsiEnd - AnsiStart));
-			}
-		}
-
-		return Res;
-	}
-
-	static inline int64 Strtoi64( const WIDECHAR* Start, WIDECHAR** End, int32 Base ) 
-	{
-		int StartLen = Strlen(Start);
-		ANSICHAR* AnsiStart = (ANSICHAR*)FMemory_Alloca(StartLen+1);
-		CopyWideToAnsi(AnsiStart, Start);
-
-		ANSICHAR* AnsiEnd = NULL;
-
-		uint64 Res = strtoll(AnsiStart, &AnsiEnd, Base);
-
-		if (End)
-		{
-			if (AnsiEnd == NULL)
-			{
-				*End = NULL;
-			}
-			else
-			{
-				*End = (WIDECHAR*)(Start + (AnsiEnd - AnsiStart));
-			}
-		}
-
-		return Res;
-	}
-
-	static inline uint64 Strtoui64( const WIDECHAR* Start, WIDECHAR** End, int32 Base ) 
-	{
-		int StartLen = Strlen(Start);
-		ANSICHAR* AnsiStart = (ANSICHAR*)FMemory_Alloca(StartLen+1);
-		CopyWideToAnsi(AnsiStart, Start);
-
-		ANSICHAR* AnsiEnd = NULL;
-
-		uint64 Res = strtoull(AnsiStart, &AnsiEnd, Base);
-
-		if (End)
-		{
-			if (AnsiEnd == NULL)
-			{
-				*End = NULL;
-			}
-			else
-			{
-				*End = (WIDECHAR*)(Start + (AnsiEnd - AnsiStart));
-			}
-		}
-
-		return Res;
-	}
-
-	static inline WIDECHAR* Strtok(WIDECHAR* StrToken, const WIDECHAR* Delim, WIDECHAR** Context)
-	{
-		int StrTokenLen = Strlen(StrToken);
-		int DelimLen = Strlen(Delim);
-		ANSICHAR* AnsiStrToken = (ANSICHAR*)FMemory_Alloca(StrTokenLen+1);
-		ANSICHAR* AnsiDelim = (ANSICHAR*)FMemory_Alloca(DelimLen+1);
-		CopyWideToAnsi(AnsiStrToken, StrToken);
-		CopyWideToAnsi(AnsiDelim, Delim);
-
-		ANSICHAR* Pos = strtok(AnsiStrToken, AnsiDelim);
-
-		if (!Pos)
-		{
-			return NULL;
-		}
-
-		return StrToken + (Pos - AnsiStrToken);
-	}
-
-	static inline int32 GetVarArgs( WIDECHAR* Dest, SIZE_T DestSize, const WIDECHAR*& Fmt, va_list ArgPtr )
-	{
-#if PLATFORM_USE_LS_SPEC_FOR_WIDECHAR
-		// fix up the Fmt string, as fast as possible, without using an FString
-		const WIDECHAR* OldFormat = Fmt;
-		WIDECHAR* NewFormat = (WIDECHAR*)FMemory_Alloca((Strlen(Fmt) * 2 + 1) * sizeof(WIDECHAR));
-		
-		int32 NewIndex = 0;
-
-		for (; *OldFormat != 0; NewIndex++, OldFormat++)
-		{
-			// fix up %s -> %ls
-			if (OldFormat[0] == LITERAL(WIDECHAR, '%'))
-			{
-				NewFormat[NewIndex++] = *OldFormat++;
-
-				if (*OldFormat == LITERAL(WIDECHAR, '%'))
-				{
-					NewFormat[NewIndex] = *OldFormat;
-				}
-				else
-				{
-					const WIDECHAR* NextChar = OldFormat;
-					
-					while(*NextChar != 0 && !FChar::IsAlpha(*NextChar)) 
-					{ 
-						NewFormat[NewIndex++] = *NextChar;
-						++NextChar; 
-					};
-
-					if (*NextChar == LITERAL(WIDECHAR, 's'))
-					{
-						NewFormat[NewIndex++] = LITERAL(WIDECHAR, 'l');
-						NewFormat[NewIndex] = *NextChar;
-					}
-					else if (*NextChar == LITERAL(WIDECHAR, 'S'))
-					{
-						NewFormat[NewIndex] = LITERAL(WIDECHAR, 's');
-					}
-					else
-					{
-						NewFormat[NewIndex] = *NextChar;
-					}
-
-					OldFormat = NextChar;
-				}
-			}
-			else
-			{
-				NewFormat[NewIndex] = *OldFormat;
-			}
-		}
-		NewFormat[NewIndex] = 0;
-#endif // PLATFORM_USE_LS_SPEC_FOR_WIDECHAR
-		int32 Result = vswprintf( Dest, DestSize, NewFormat, ArgPtr);
-		return Result;
-	}
-
-	/** 
-	 * Ansi implementation 
-	 **/
-	static UE_FORCEINLINE_HINT ANSICHAR* Strcpy(ANSICHAR* Dest, const ANSICHAR* Src)
-	{
-		return strcpy( Dest, Src );
-	}
-
-	// This version was deprecated because Strcpy taking a size field does not match the CRT strcpy, and on some platforms the size field was being ignored.
-	// If the memzeroing causes a performance problem, request a new function StrcpyTruncate.
-	UE_DEPRECATED(5.6, "Use Strncpy instead. Note that Strncpy has a behavior difference from Strcpy: it memzeroes the entire DestCount-sized buffer after the end of string.")
-	static UE_FORCEINLINE_HINT ANSICHAR* Strcpy(ANSICHAR* Dest, SIZE_T DestCount, const ANSICHAR* Src)
-	{
-		return Strcpy(Dest, Src);
-	}
-
-	static inline ANSICHAR* Strncpy(ANSICHAR* Dest, const ANSICHAR* Src, SIZE_T MaxLen)
-	{
-		::strncpy(Dest, Src, MaxLen);
-		Dest[MaxLen-1]=0;
-		return Dest;
-	}
-
-	static UE_FORCEINLINE_HINT ANSICHAR* Strcat(ANSICHAR* Dest, const ANSICHAR* Src)
-	{
-		return strcat( Dest, Src );
-	}
-
-	// This version was deprecated because Strcat taking a size field does not match the CRT strcat, and on some platforms the size field was being ignored.
-	UE_DEPRECATED(5.6, "Use Strncat instead. !!NOTE THAT STRNCAT takes SrcLen rather than DestCount. You must call Strncat(Dest, Src, DestCount - Strlen(Dest) - 1).")
-	static UE_FORCEINLINE_HINT ANSICHAR* Strcat(ANSICHAR* Dest, SIZE_T DestCount, const ANSICHAR* Src)
-	{
-		return Strcat(Dest, Src);
-	}
-
-	static UE_FORCEINLINE_HINT ANSICHAR* Strncat(ANSICHAR* Dest, const ANSICHAR* Src, SIZE_T SrcLen)
-	{
-		return strncat(Dest, Src, SrcLen);
-	}
-
-	static UE_FORCEINLINE_HINT int32 Strcmp( const ANSICHAR* String1, const ANSICHAR* String2 )
-	{
-		return strcmp(String1, String2);
-	}
-
-	static UE_FORCEINLINE_HINT int32 Strncmp( const ANSICHAR* String1, const ANSICHAR* String2, SIZE_T Count )
-	{
-		return strncmp( String1, String2, Count );
-	}
-
-	static UE_FORCEINLINE_HINT int32 Strlen( const ANSICHAR* String )
-	{
-		return strlen( String ); 
-	}
-
-	static UE_FORCEINLINE_HINT int32 Strnlen( const ANSICHAR* String, SIZE_T StringSize )
-	{
-		return strnlen( String, StringSize );
-	}
-
-	static UE_FORCEINLINE_HINT const ANSICHAR* Strstr( const ANSICHAR* String, const ANSICHAR* Find)
-	{
-		return strstr(String, Find);
-	}
-
-	static UE_FORCEINLINE_HINT const ANSICHAR* Strchr( const ANSICHAR* String, ANSICHAR C)
-	{
-		return strchr(String, C);
-	}
-
-	static UE_FORCEINLINE_HINT const ANSICHAR* Strrchr( const ANSICHAR* String, ANSICHAR C)
-	{
-		return strrchr(String, C);
-	}
-
-	static UE_FORCEINLINE_HINT int32 Atoi(const ANSICHAR* String)
-	{
-		return atoi( String ); 
-	}
-
-	static UE_FORCEINLINE_HINT int64 Atoi64(const ANSICHAR* String)
-	{
-		return strtoll( String, NULL, 10 );
-	}
-
-	static UE_FORCEINLINE_HINT float Atof(const ANSICHAR* String)
-	{
-		return (float)atof( String ); 
-	}
-
-	static UE_FORCEINLINE_HINT double Atod(const ANSICHAR* String)
-	{
-		return atof( String ); 
-	}
-
-	static UE_FORCEINLINE_HINT int32 Strtoi( const ANSICHAR* Start, ANSICHAR** End, int32 Base ) 
-	{
-		return strtol( Start, End, Base ); 
-	}
-
-	static UE_FORCEINLINE_HINT int64 Strtoi64( const ANSICHAR* Start, ANSICHAR** End, int32 Base ) 
-	{
-		return strtoll(Start, End, Base);;
-	}
-
-	static UE_FORCEINLINE_HINT uint64 Strtoui64( const ANSICHAR* Start, ANSICHAR** End, int32 Base ) 
-	{
-		return strtoull(Start, End, Base);;
-	}
-
-	static UE_FORCEINLINE_HINT ANSICHAR* Strtok(ANSICHAR* StrToken, const ANSICHAR* Delim, ANSICHAR** Context)
-	{
-		return strtok(StrToken, Delim);
-	}
-
-	static inline int32 GetVarArgs( ANSICHAR* Dest, SIZE_T DestSize, const ANSICHAR*& Fmt, va_list ArgPtr )
-	{
-		int32 Result = vsnprintf(Dest, DestSize, Fmt, ArgPtr);
-		return (Result != -1 && Result < (int32)DestSize) ? Result : -1;
-	}
-
-	/**
-	 * UTF8CHAR implementation.
-	 */
-	static UE_FORCEINLINE_HINT UTF8CHAR* Strcpy(UTF8CHAR* Dest, const UTF8CHAR* Src)
-	{
-		return (UTF8CHAR*)Strcpy((ANSICHAR*)Dest, (const ANSICHAR*)Src);
-	}
-	// This version was deprecated because Strcpy taking a size field does not match the CRT strcpy, and on some platforms the size field was being ignored.
-	// If the memzeroing causes a performance problem, request a new function StrcpyTruncate.
-	UE_DEPRECATED(5.6, "Use Strncpy instead. Note that Strncpy has a behavior difference from Strcpy: it memzeroes the entire DestCount-sized buffer after the end of string.")
-	static inline UTF8CHAR* Strcpy(UTF8CHAR* Dest, SIZE_T DestCount, const UTF8CHAR* Src)
-	{
-		PRAGMA_DISABLE_DEPRECATION_WARNINGS
-		return (UTF8CHAR*)Strcpy((ANSICHAR*)Dest, DestCount, (const ANSICHAR*)Src);
-		PRAGMA_ENABLE_DEPRECATION_WARNINGS
-	}
-
-	static UE_FORCEINLINE_HINT UTF8CHAR* Strncpy(UTF8CHAR* Dest, const UTF8CHAR* Src, SIZE_T MaxLen)
-	{
-		return (UTF8CHAR*)Strncpy((ANSICHAR*)Dest, (const ANSICHAR*)Src, MaxLen);
-	}
-
-	static UE_FORCEINLINE_HINT UTF8CHAR* Strcat(UTF8CHAR* Dest, const UTF8CHAR* Src)
-	{
-		return (UTF8CHAR*)Strcat((ANSICHAR*)Dest, (const ANSICHAR*)Src);
-	}
-
-	// This version was deprecated because Strcat taking a size field does not match the CRT strcat, and on some platforms the size field was being ignored.
-	UE_DEPRECATED(5.6, "Use Strncat instead. !!NOTE THAT STRNCAT takes SrcLen rather than DestCount. You must call Strncat(Dest, Src, DestCount - Strlen(Dest) - 1).")
-	static inline UTF8CHAR* Strcat(UTF8CHAR* Dest, SIZE_T DestCount, const UTF8CHAR* Src)
-	{
-		PRAGMA_DISABLE_DEPRECATION_WARNINGS
-		return (UTF8CHAR*)Strcat((ANSICHAR*)Dest, DestCount, (const ANSICHAR*)Src);
-		PRAGMA_ENABLE_DEPRECATION_WARNINGS
-	}
-
-	static UE_FORCEINLINE_HINT UTF8CHAR* Strncat(UTF8CHAR* Dest, const UTF8CHAR* Src, SIZE_T SrcLen)
-	{
-		return (UTF8CHAR*)Strncat((ANSICHAR*)Dest, (const ANSICHAR*)Src, SrcLen);
-	}
-
-	static UE_FORCEINLINE_HINT int32 Strcmp(const UTF8CHAR* String1, const UTF8CHAR* String2)
-	{
-		return Strcmp((const ANSICHAR*)String1, (const ANSICHAR*)String2);
-	}
-
-	static UE_FORCEINLINE_HINT int32 Strncmp(const UTF8CHAR* String1, const UTF8CHAR* String2, SIZE_T Count)
-	{
-		return Strncmp((const ANSICHAR*)String1, (const ANSICHAR*)String2, Count);
-	}
-
-	static UE_FORCEINLINE_HINT int32 Strlen(const UTF8CHAR* String)
-	{
-		return Strlen((const ANSICHAR*)String);
-	}
-
-	static UE_FORCEINLINE_HINT int32 Strnlen( const UTF8CHAR* String, SIZE_T StringSize )
-	{
-		return Strnlen((const ANSICHAR*)String, StringSize);
-	}
-
-	static UE_FORCEINLINE_HINT const UTF8CHAR* Strstr(const UTF8CHAR* String, const UTF8CHAR* Find)
-	{
-		return (const UTF8CHAR*)Strstr((const ANSICHAR*)String, (const ANSICHAR*)Find);
-	}
-
-	static UE_FORCEINLINE_HINT const UTF8CHAR* Strchr(const UTF8CHAR* String, UTF8CHAR C)
-	{
-		return (const UTF8CHAR*)Strchr((const ANSICHAR*)String, (ANSICHAR)C);
-	}
-
-	static UE_FORCEINLINE_HINT const UTF8CHAR* Strrchr(const UTF8CHAR* String, UTF8CHAR C)
-	{
-		return (const UTF8CHAR*)Strrchr((const ANSICHAR*)String, (ANSICHAR)C);
-	}
-
-	static UE_FORCEINLINE_HINT int32 Atoi(const UTF8CHAR* String)
-	{
-		return Atoi((const ANSICHAR*)String);
-	}
-
-	static UE_FORCEINLINE_HINT int64 Atoi64(const UTF8CHAR* String)
-	{
-		return Atoi64((const ANSICHAR*)String);
-	}
-
-	static UE_FORCEINLINE_HINT float Atof(const UTF8CHAR* String)
-	{
-		return Atof((const ANSICHAR*)String);
-	}
-
-	static UE_FORCEINLINE_HINT double Atod(const UTF8CHAR* String)
-	{
-		return Atod((const ANSICHAR*)String);
-	}
-
-	static UE_FORCEINLINE_HINT int32 Strtoi(const UTF8CHAR* Start, UTF8CHAR** End, int32 Base)
-	{
-		return Strtoi((const ANSICHAR*)Start, (ANSICHAR**)End, Base);
-	}
-
-	static UE_FORCEINLINE_HINT int64 Strtoi64(const UTF8CHAR* Start, UTF8CHAR** End, int32 Base)
-	{
-		return Strtoi64((const ANSICHAR*)Start, (ANSICHAR**)End, Base);
-	}
-
-	static UE_FORCEINLINE_HINT uint64 Strtoui64(const UTF8CHAR* Start, UTF8CHAR** End, int32 Base)
-	{
-		return Strtoui64((const ANSICHAR*)Start, (ANSICHAR**)End, Base);
-	}
-
-	static UE_FORCEINLINE_HINT UTF8CHAR* Strtok(UTF8CHAR* StrToken, const UTF8CHAR* Delim, UTF8CHAR** Context)
-	{
-		return (UTF8CHAR*)Strtok((ANSICHAR*)StrToken, (const ANSICHAR*)Delim, (ANSICHAR**)Context);
-	}
-
-	static UE_FORCEINLINE_HINT int32 GetVarArgs(UTF8CHAR* Dest, SIZE_T DestSize, const UTF8CHAR*& Fmt, va_list ArgPtr)
-	{
-		return GetVarArgs((ANSICHAR*)Dest, DestSize, *(const ANSICHAR**)&Fmt, ArgPtr);
-	}
-
-	/** 
-	 * UCS2 implementation 
-	 **/
-
-	static inline int32 Strlen( const UCS2CHAR* String )
-	{
-		int32 Result = 0;
-		while (*String++)
-		{
-			++Result;
-		}
-
-		return Result;
-	}
-
-	static inline int32 Strnlen( const UCS2CHAR* String, SIZE_T StringSize )
-	{
-		int32 Result = 0;
-		while (StringSize-- > 0 && *String++)
-		{
-			++Result;
-		}
-
-		return Result;
-	}
-};
-#endif
-
-typedef FAndroidPlatformString FPlatformString;
+typedef FGenericWidePlatformString FPlatformString;
 
 // Format specifiers to be able to print values of these types correctly, for example when using UE_LOG.
 // SIZE_T format specifier
diff --git a/Engine/Source/Runtime/Core/Public/Apple/ApplePlatformString.h b/Engine/Source/Runtime/Core/Public/Apple/ApplePlatformString.h
--- a/Engine/Source/Runtime/Core/Public/Apple/ApplePlatformString.h
+++ b/Engine/Source/Runtime/Core/Public/Apple/ApplePlatformString.h
@@ -1,17 +1,9 @@
 // Copyright Epic Games, Inc. All Rights Reserved.
 
-/*=============================================================================================
-	ApplePlatformString.h: Apple platform string classes, mostly implemented with ANSI C++
-==============================================================================================*/
-
 #pragma once
 
 #include "CoreTypes.h"
-#if PLATFORM_TCHAR_IS_CHAR16
 #include "GenericPlatform/GenericWidePlatformString.h"
-#else
-#include "GenericPlatform/StandardPlatformString.h"
-#endif
 #include "Misc/AssertionMacros.h"
 #include "Misc/OutputDevice.h"
 
@@ -21,6 +13,8 @@
 #include "IOS/IOSSystemIncludes.h"
 #endif
 
+static_assert(PLATFORM_TCHAR_IS_CHAR16, "TCHAR must be CHAR16.");
+
 #ifdef __OBJC__
 #import <Foundation/NSString.h>
 
@@ -43,14 +37,9 @@ class FString;
 
 
 /**
-* 2-byte replacement s string implementation
-**/
-struct FApplePlatformString
-#if PLATFORM_TCHAR_IS_CHAR16
-	: public FGenericWidePlatformString
-#else
-	: public FStandardPlatformString
-#endif
+ * 2-byte string implementation
+ */
+struct FApplePlatformString : public FGenericWidePlatformString
 {
 	// These should be replaced with equivalent Convert and ConvertedLength functions when we properly support variable-length encodings.
 /*	static void WideCharToMultiByte(const TCHAR *Source, uint32 LengthWM1, ANSICHAR *Dest, uint32 LengthA)
diff --git a/Engine/Source/Runtime/Core/Public/Mac/MacPlatform.h b/Engine/Source/Runtime/Core/Public/Mac/MacPlatform.h
--- a/Engine/Source/Runtime/Core/Public/Mac/MacPlatform.h
+++ b/Engine/Source/Runtime/Core/Public/Mac/MacPlatform.h
@@ -1,14 +1,10 @@
 // Copyright Epic Games, Inc. All Rights Reserved.
 
-/*================================================================================
-	MacPlatform.h: Setup for the Mac platform
-==================================================================================*/
-
 #pragma once
 
 #include "Apple/ApplePlatform.h"
 
-#define PLATFORM_MAC_USE_CHAR16 1
+#define PLATFORM_MAC_USE_CHAR16 1 UE_DEPRECATED_MACRO(5.7, "PLATFORM_MAC_USE_CHAR16 has been deprecated and should be replaced with 1.")
 
 /**
 * Mac specific types
@@ -18,12 +14,8 @@ struct FMacPlatformTypes : public FGenericPlatformTypes
 	typedef unsigned int		DWORD;
 	typedef size_t				SIZE_T;
 	typedef decltype(NULL)		TYPE_OF_NULL;
-#if PLATFORM_MAC_USE_CHAR16
 	typedef char16_t			WIDECHAR;
 	typedef WIDECHAR			TCHAR;
-#else
-	typedef char16_t			CHAR16;
-#endif
 };
 
 typedef FMacPlatformTypes FPlatformTypes;
@@ -51,11 +43,7 @@ typedef FMacPlatformTypes FPlatformTypes;
 //#define PLATFORM_USE_LS_SPEC_FOR_WIDECHAR				1
 #define PLATFORM_COMPILER_DISTINGUISHES_INT_AND_LONG	1
 #define PLATFORM_WCHAR_IS_4_BYTES						1
-#if PLATFORM_MAC_USE_CHAR16
-#	define PLATFORM_TCHAR_IS_CHAR16						1
-#else
-#	define PLATFORM_TCHAR_IS_4_BYTES					1
-#endif
+#define PLATFORM_TCHAR_IS_CHAR16						1
 #define PLATFORM_HAS_BSD_SOCKET_FEATURE_IOCTL			1
 #define PLATFORM_HAS_BSD_SOCKET_FEATURE_POLL			1
 //#define PLATFORM_USE_PTHREADS							1
diff --git a/Engine/Source/Runtime/Core/Public/Unix/UnixPlatform.h b/Engine/Source/Runtime/Core/Public/Unix/UnixPlatform.h
--- a/Engine/Source/Runtime/Core/Public/Unix/UnixPlatform.h
+++ b/Engine/Source/Runtime/Core/Public/Unix/UnixPlatform.h
@@ -14,7 +14,7 @@
 
 struct FGenericPlatformTypes;
 
-#define PLATFORM_UNIX_USE_CHAR16 1
+#define PLATFORM_UNIX_USE_CHAR16 1 UE_DEPRECATED_MACRO(5.7, "PLATFORM_UNIX_USE_CHAR16 has been deprecated and should be replaced with 1.")
 
 /**
 * Unix specific types
@@ -24,10 +24,8 @@ struct FUnixPlatformTypes : public FGenericPlatformTypes
 	typedef unsigned int		DWORD;
 	typedef __SIZE_TYPE__		SIZE_T;
 	typedef decltype(__null)	TYPE_OF_NULL;
-#if PLATFORM_UNIX_USE_CHAR16
 	typedef char16_t			WIDECHAR;
 	typedef WIDECHAR			TCHAR;
-#endif
 };
 
 typedef FUnixPlatformTypes FPlatformTypes;
@@ -46,11 +44,7 @@ typedef FUnixPlatformTypes FPlatformTypes;
 #define PLATFORM_SUPPORTS_PRAGMA_PACK					1
 #define PLATFORM_USE_LS_SPEC_FOR_WIDECHAR				1
 #define PLATFORM_WCHAR_IS_4_BYTES						1
-#if PLATFORM_UNIX_USE_CHAR16
 #define PLATFORM_TCHAR_IS_CHAR16						1
-#else
-#define PLATFORM_TCHAR_IS_4_BYTES						1
-#endif
 #define PLATFORM_HAS_BSD_TIME							1
 #define PLATFORM_HAS_BSD_THREAD_CPUTIME					1
 #define PLATFORM_USE_PTHREADS							1
diff --git a/Engine/Source/Runtime/Core/Public/Unix/UnixPlatformString.h b/Engine/Source/Runtime/Core/Public/Unix/UnixPlatformString.h
--- a/Engine/Source/Runtime/Core/Public/Unix/UnixPlatformString.h
+++ b/Engine/Source/Runtime/Core/Public/Unix/UnixPlatformString.h
@@ -1,23 +1,14 @@
 // Copyright Epic Games, Inc. All Rights Reserved.
 
-/*=============================================================================================
-	UnixPlatformString.h: Unix platform string classes, mostly implemented with ANSI C++
-==============================================================================================*/
-
 #pragma once
 
 #include "CoreTypes.h"
-#include "GenericPlatform/StandardPlatformString.h"
 #include "GenericPlatform/GenericWidePlatformString.h"
+
 /**
-* Unix string implementation
-**/
-struct FUnixPlatformString : public 
-#if PLATFORM_TCHAR_IS_CHAR16
-	FGenericWidePlatformString
-#else
-	FStandardPlatformString
-#endif
+ * Unix string implementation
+ */
+struct FUnixPlatformString : public FGenericWidePlatformString
 {
 	template<typename CHAR>
 	static inline int32 Strlen(const CHAR* String)
@@ -95,3 +86,7 @@ typedef FUnixPlatformString FPlatformString;
 #define UINT64_x_FMT "llx"
 // uint64 format specifier for uppercase hexadecimal output
 #define UINT64_X_FMT "llX"
+
+#if UE_ENABLE_INCLUDE_ORDER_DEPRECATED_IN_5_7
+#include "GenericPlatform/StandardPlatformString.h"
+#endif
