Paths: Fix another incorrect case of CreateStandardFilename. Paths that "start with" RelativePathToRoot but add on another .. are outside of the root, and their StandardFilename is their absolute path rather than their RelativePathFromBase.
#rnx

diff --git a/Engine/Source/Runtime/Core/Private/Misc/Paths.cpp b/Engine/Source/Runtime/Core/Private/Misc/Paths.cpp
--- a/Engine/Source/Runtime/Core/Private/Misc/Paths.cpp
+++ b/Engine/Source/Runtime/Core/Private/Misc/Paths.cpp
@@ -1428,7 +1428,20 @@ FString FPaths::CreateStandardFilename(const FString& InPath)
 	{
 		FPlatformMisc::NormalizePath(NormalizedPath);
 		CollapseRelativeDirectories(NormalizedPath, true /* bCollapseAllPossible */);
-		return NormalizedPath;
+
+		// Paths outside of the root path define their standardized path == absolute path.
+		// A path might "start with" RelativePathToRoot, but still be outside the root path because
+		// it adds on another "..". Check for that here before returning NormalizedPath based on the
+		// assumption that it is a child path of root and already normalized.
+		FStringView Suffix = FStringView(NormalizedPath).RightChop(RelativePathToRoot.Len());
+		if (Suffix.Len() > 0 && Suffix[0] == '/')
+		{
+			Suffix = Suffix.RightChop(1);
+		}
+		if (!Suffix.StartsWith(TEXT("..")))
+		{
+			return NormalizedPath;
+		}
 	}
 
 	// Make it an absolute path and then make that absolute path a subdirectory of the RelativePathToRoot.
diff --git a/Engine/Source/Runtime/Core/Tests/Misc/PathsTest.cpp b/Engine/Source/Runtime/Core/Tests/Misc/PathsTest.cpp
--- a/Engine/Source/Runtime/Core/Tests/Misc/PathsTest.cpp
+++ b/Engine/Source/Runtime/Core/Tests/Misc/PathsTest.cpp
@@ -261,6 +261,14 @@ void TestCreateStandardFilename()
 		// it will create different answers for those paths.
 		return;
 	}
+	FString RootDirectoryParent = FPaths::GetPath(RootDirectory);
+	bool bAllowEscapeRootDirectoryTests = !RootDirectoryParent.IsEmpty()
+		&& RootDirectory == FPaths::Combine(RootDirectoryParent, FPaths::GetBaseFilename(RootDirectory));
+	FString RelativePathToRootDirectoryParent;
+	if (bAllowEscapeRootDirectoryTests)
+	{
+		RelativePathToRootDirectoryParent = FPaths::Combine(RootRelative, TEXT(".."));
+	}
 	
 	auto Run = [&RootRelative](const TCHAR* Path, FStringView Expected)
 		{
@@ -345,6 +353,16 @@ void TestCreateStandardFilename()
 	Run(*FPaths::Combine(RootDirectory, TEXT("Dir/Dir/File.ext")), FPaths::Combine(RootRelative, TEXT("Dir/Dir/File.ext")));
 	Run(*FPaths::Combine(RootDirectory, TEXT("Engine/Dir/File.ext")), FPaths::Combine(RootRelative, TEXT("Engine/Dir/File.ext")));
 	Run(*(RootDirectory + TEXT("\\Dir\\File")), FPaths::Combine(RootRelative, TEXT("/Dir/File")));
+	if (bAllowEscapeRootDirectoryTests)
+	{
+		// Relative paths that escape the root directory (e.g. specify the root directory's parent or a child other than the root directory
+		// of the root directory's parent should be written as absolute paths
+		// With Base D:\Root\Engine\Binaries\Win64
+		// ../../../.. -> D:\ rather than ../../../..
+		// ../../../../Peer -> D:\Peer rather than ../../../../Peer
+		Run(*RelativePathToRootDirectoryParent, RootDirectoryParent);
+		Run(*FPaths::Combine(RelativePathToRootDirectoryParent, TEXT("Peer")), FPaths::Combine(RootDirectoryParent, TEXT("Peer")));
+	}
 }
 
 #endif //WITH_TESTS
