[Misc] Add a console command to reset ensure state so they fire again

diff --git a/Engine/Source/Runtime/Core/Private/Misc/AssertionMacros.cpp b/Engine/Source/Runtime/Core/Private/Misc/AssertionMacros.cpp
--- a/Engine/Source/Runtime/Core/Private/Misc/AssertionMacros.cpp
+++ b/Engine/Source/Runtime/Core/Private/Misc/AssertionMacros.cpp
@@ -138,6 +138,14 @@ FAutoConsoleVariableRef CVarEnsureAlwaysEnabled(
 	TEXT("Set to false to turn ensureAlways into regular ensure"),
 	ECVF_Default);
 
+uint8 GEnsureResetState = 1;
+FAutoConsoleCommand CVarResetEnsureState(
+	TEXT("core.ResetEnsureState"),
+	TEXT("Reset all ensures so they will fire again"),
+	FConsoleCommandDelegate::CreateLambda([]()
+		{
+			++GEnsureResetState;
+		}));
 
 CORE_API void (*GPrintScriptCallStackFn)() = nullptr;
 
@@ -861,11 +869,11 @@ bool UE_COLD UE_DEBUG_SECTION VARARGS CheckVerifyImpl(const ANSICHAR* File, int3
 }
 
 UE_AUTORTFM_ALWAYS_OPEN
-bool UE_COLD UE_DEBUG_SECTION UE::Assert::Private::ExecCheckImplInternal(std::atomic<bool>& bExecuted, bool bAlways, const ANSICHAR* File, int32 Line, const ANSICHAR* Expr)
+bool UE_COLD UE_DEBUG_SECTION UE::Assert::Private::ExecCheckImplInternal(std::atomic<uint8>& bExecuted, bool bAlways, const ANSICHAR* File, int32 Line, const ANSICHAR* Expr)
 {
-	if (((bAlways && GEnsureAlwaysEnabled) || !bExecuted.load(std::memory_order_relaxed)) && FPlatformMisc::IsEnsureAllowed())
+	if (((bAlways && GEnsureAlwaysEnabled) || bExecuted.load(std::memory_order_relaxed) != GEnsureResetState) && FPlatformMisc::IsEnsureAllowed())
 	{
-		if (bExecuted.exchange(true, std::memory_order_release) && !bAlways)
+		if (bExecuted.exchange(GEnsureResetState, std::memory_order_release) == GEnsureResetState && !bAlways)
 		{
 			return false;
 		}
@@ -878,10 +886,16 @@ bool UE_COLD UE_DEBUG_SECTION UE::Assert::Private::ExecCheckImplInternal(std::at
 	return false;
 }
 
+// Check done separately so we don't evaluate the var args if we're not expecting to process the ensure
+bool UE_COLD UE_DEBUG_SECTION UE::Assert::Private::CheckEnsureFailed(bool bAlways, const std::atomic<uint8>& bExecuted)
+{
+	return ((bAlways && GEnsureAlwaysEnabled) || bExecuted.load(std::memory_order_relaxed) != GEnsureResetState) && FPlatformMisc::IsEnsureAllowed();
+}
+
 UE_AUTORTFM_ALWAYS_OPEN
-bool UE_COLD UE_DEBUG_SECTION EnsureFailedV(std::atomic<bool>& bExecuted, const UE::Assert::Private::FStaticEnsureRecord* Ensure, void* ProgramCounter, va_list Args)
+bool UE_COLD UE_DEBUG_SECTION EnsureFailedV(std::atomic<uint8>& bExecuted, const UE::Assert::Private::FStaticEnsureRecord* Ensure, void* ProgramCounter, va_list Args)
 {
-	if (bExecuted.exchange(true, std::memory_order_release) && !(Ensure->bAlways && GEnsureAlwaysEnabled))
+	if (bExecuted.exchange(GEnsureResetState, std::memory_order_release) == GEnsureResetState && !(Ensure->bAlways && GEnsureAlwaysEnabled))
 	{
 		return false;
 	}
@@ -890,7 +904,7 @@ bool UE_COLD UE_DEBUG_SECTION EnsureFailedV(std::atomic<bool>& bExecuted, const
 	return bResult;
 }
 
-bool UE_COLD UE_DEBUG_SECTION VARARGS UE::Assert::Private::EnsureFailed(std::atomic<bool>& bExecuted, const FStaticEnsureRecord* Ensure, ...)
+bool UE_COLD UE_DEBUG_SECTION VARARGS UE::Assert::Private::EnsureFailed(std::atomic<uint8>& bExecuted, const FStaticEnsureRecord* Ensure, ...)
 {
 	va_list Args;
 	va_start(Args, Ensure);
diff --git a/Engine/Source/Runtime/Core/Public/Misc/AssertionMacros.h b/Engine/Source/Runtime/Core/Public/Misc/AssertionMacros.h
--- a/Engine/Source/Runtime/Core/Public/Misc/AssertionMacros.h
+++ b/Engine/Source/Runtime/Core/Public/Misc/AssertionMacros.h
@@ -361,7 +361,6 @@ RetType UE_COLD UE_DEBUG_SECTION DispatchCheckVerify (InnerType&& Inner, ArgType
  */
 
 #if DO_ENSURE && !USING_CODE_ANALYSIS // The Visual Studio 2013 analyzer doesn't understand these complex conditionals
-
 	namespace UE::Assert::Private
 	{
 	
@@ -390,13 +389,14 @@ RetType UE_COLD UE_DEBUG_SECTION DispatchCheckVerify (InnerType&& Inner, ArgType
 		}
 	};
 
-	CORE_API bool UE_COLD UE_DEBUG_SECTION VARARGS EnsureFailed(std::atomic<bool>& bExecuted, const FStaticEnsureRecord* Ensure, ...);
-	CORE_API bool UE_COLD UE_DEBUG_SECTION ExecCheckImplInternal(std::atomic<bool>& bExecuted, bool bAlways, const ANSICHAR* File, int32 Line, const ANSICHAR* Expr);
+	CORE_API bool UE_COLD UE_DEBUG_SECTION VARARGS CheckEnsureFailed(bool bAlways, const std::atomic<uint8>& bExecuted);
+	CORE_API bool UE_COLD UE_DEBUG_SECTION VARARGS EnsureFailed(std::atomic<uint8>& bExecuted, const FStaticEnsureRecord* Ensure, ...);
+	CORE_API bool UE_COLD UE_DEBUG_SECTION ExecCheckImplInternal(std::atomic<uint8>& bExecuted, bool bAlways, const ANSICHAR* File, int32 Line, const ANSICHAR* Expr);
 
 	} // UE::Assert::Private
 
 	template <uint64 Uid>
-	std::atomic<bool> bGEnsureHasExecuted = false;
+	std::atomic<uint8> bGEnsureHasExecuted = false;
 
 	UE_NODEBUG constexpr uint64 FileLineHashForEnsure(const char* Filename, uint32 Line)
 	{
@@ -443,8 +443,9 @@ RetType UE_COLD UE_DEBUG_SECTION DispatchCheckVerify (InnerType&& Inner, ArgType
 			[&]() UE_COLD UE_DEBUG_SECTION \
 			{ \
 				UE_VALIDATE_FORMAT_STRING(InFormat, ##__VA_ARGS__); \
+				std::atomic<uint8>& bEnsureHasExecuted = ::bGEnsureHasExecuted < FileLineHashForEnsure(__FILE__, __LINE__) >; \
 				static constexpr ::UE::Assert::Private::FStaticEnsureRecord ENSURE_Static(InFormat, #InExpression, __builtin_FILE(), __builtin_LINE(), Always); \
-				if ((Always || !::bGEnsureHasExecuted<FileLineHashForEnsure(__FILE__, __LINE__)>.load(std::memory_order_relaxed)) && FPlatformMisc::IsEnsureAllowed() && ::UE::Assert::Private::EnsureFailed(::bGEnsureHasExecuted<FileLineHashForEnsure(__FILE__, __LINE__)>, &ENSURE_Static, ##__VA_ARGS__)) \
+				if (::UE::Assert::Private::CheckEnsureFailed(Always, bEnsureHasExecuted) && ::UE::Assert::Private::EnsureFailed(bEnsureHasExecuted, &ENSURE_Static, ##__VA_ARGS__)) \
 				{ \
 					PLATFORM_BREAK(); \
 				} \
