Linux Content Worker: Increase the crash handling timeout on Wine

#rnx

diff --git a/Engine/Source/Runtime/Core/Private/Windows/WindowsPlatformCrashContext.cpp b/Engine/Source/Runtime/Core/Private/Windows/WindowsPlatformCrashContext.cpp
--- a/Engine/Source/Runtime/Core/Private/Windows/WindowsPlatformCrashContext.cpp
+++ b/Engine/Source/Runtime/Core/Private/Windows/WindowsPlatformCrashContext.cpp
@@ -190,7 +190,7 @@ namespace {
 	FORCEINLINE double GetCrashTimeoutSeconds()
 	{
 		// 60s is plenty of time to generate a crash report under Windows, but generating minidumps under Wine is much slower
-		double TimeoutSeconds = FWindowsPlatformMisc::IsWine() ? 120.0 : DefaultCrashHandlingTimeoutSecs;
+		double TimeoutSeconds = FWindowsPlatformMisc::IsWine() ? 300.0 : DefaultCrashHandlingTimeoutSecs;
 		if (GConfig)
 		{
 			// If available override with configurable value. Negative values are interpreted as infinite wait (not generally recommended)
