[Core] Add support for move constructor/assignment for TDoubleLinkedList

* Respect Rule of 5
* Copy is still deleted as there is manual memory management

#rnx

diff --git a/Engine/Source/Runtime/Core/Public/Containers/List.h b/Engine/Source/Runtime/Core/Public/Containers/List.h
--- a/Engine/Source/Runtime/Core/Public/Containers/List.h
+++ b/Engine/Source/Runtime/Core/Public/Containers/List.h
@@ -504,6 +504,27 @@ public:
 		Empty();
 	}
 
+	// Can't copy
+	TDoubleLinkedList(const TDoubleLinkedList&) = delete;
+	TDoubleLinkedList& operator=(const TDoubleLinkedList&) = delete;
+
+	/** Move Constructor/Assignment */
+	[[nodiscard]] TDoubleLinkedList(TDoubleLinkedList&& Other)
+		: TDoubleLinkedList()
+	{
+		Swap(HeadNode, Other.HeadNode);
+		Swap(TailNode, Other.TailNode);
+		Swap(ListSize, Other.ListSize);
+	}
+	
+	TDoubleLinkedList& operator=(TDoubleLinkedList&& Other)
+	{
+		Swap(HeadNode, Other.HeadNode);
+		Swap(TailNode, Other.TailNode);
+		Swap(ListSize, Other.ListSize);
+		return *this;
+	}
+
 	// Adding/Removing methods
 
 	/**
@@ -784,9 +805,6 @@ private:
 	TDoubleLinkedListNode* TailNode;
 	int32 ListSize;
 
-	TDoubleLinkedList(const TDoubleLinkedList&);
-	TDoubleLinkedList& operator=(const TDoubleLinkedList&);
-
 	[[nodiscard]] friend TIterator      begin(      TDoubleLinkedList& List) { return TIterator     (List.GetHead()); }
 	[[nodiscard]] friend TConstIterator begin(const TDoubleLinkedList& List) { return TConstIterator(List.GetHead()); }
 	[[nodiscard]] friend TIterator      end  (      TDoubleLinkedList& List) { return TIterator     (nullptr); }
