[Scheduler]
- Fix potential deadlock if a foreground task is scheduling a background one

diff --git a/Engine/Source/Runtime/Core/Private/Async/Fundamental/Scheduler.cpp b/Engine/Source/Runtime/Core/Private/Async/Fundamental/Scheduler.cpp
--- a/Engine/Source/Runtime/Core/Private/Async/Fundamental/Scheduler.cpp
+++ b/Engine/Source/Runtime/Core/Private/Async/Fundamental/Scheduler.cpp
@@ -540,9 +540,15 @@ namespace LowLevelTasks
 			if ((bIsBackgroundTask && !bIsBackgroundWorker) || bIsStandbyWorker)
 			{
 				QueuePreference = EQueuePreference::GlobalQueuePreference;
-			}
 
-			bWakeUpWorker |= bIsStandbyWorker || CachedLocalQueue == nullptr;
+				// Always wake up a worker if we're scheduling a background task from a foreground thread
+				// since foreground threads are not allowed to process them.
+				bWakeUpWorker = true;
+			}
+			else
+			{
+				bWakeUpWorker |= CachedLocalQueue == nullptr;
+			}
 
 			// Always force local queue usage when launching from the game thread to minimize cost.
 			if (CachedLocalQueue && CachedLocalQueue == GameThreadLocalQueue.Get())
diff --git a/Engine/Source/Runtime/Core/Tests/Tasks/TasksTest.cpp b/Engine/Source/Runtime/Core/Tests/Tasks/TasksTest.cpp
--- a/Engine/Source/Runtime/Core/Tests/Tasks/TasksTest.cpp
+++ b/Engine/Source/Runtime/Core/Tests/Tasks/TasksTest.cpp
@@ -1824,6 +1824,32 @@ namespace UE { namespace TasksTests
 		// Restore original worker count
 		LowLevelTasks::FScheduler::Get().RestartWorkers();
 	}
+	
+	// This test will force the scheduling of a background task from the normal task which exercise some code in the scheduler
+	// to send the task to the global queue and always wake up a worker since background tasks can't run on normal threads.
+	TEST_CASE_NAMED(FTaskGraphBackgroundWithNormalAsPrereq, "System::Core::Async::Tasks::BackgroundWithNormalAsPrereq", "[.][ApplicationContextMask][EngineFilter]")
+	{
+		for (int32 Index = 0; Index < 100; ++Index)
+		{
+			UE::Tasks::FTask Task =
+				UE::Tasks::Launch(
+					TEXT("Task"),
+					[] { },
+					UE::Tasks::Launch(TEXT("Task"), [] { }), // Launch Normal task first and use it as prereq for the background one
+					LowLevelTasks::ETaskPriority::BackgroundNormal
+				);
+
+			double StartTime = FPlatformTime::Seconds();
+			while (!Task.IsCompleted())
+			{
+				if (FPlatformTime::Seconds() - StartTime > 10.0f)
+				{
+					FAIL_CHECK(TEXT("Test is likely deadlocked, aborting."));
+					return;
+				}
+			}
+		}
+	}
 
 #if PLATFORM_SUPPORTS_ASYMMETRIC_FENCES
 	TEST_CASE_NAMED(FAsymmetricThreadFence, "System::Core::Async::AsymmetricThreadFence", "[.][ApplicationContextMask][EngineFilter]")
