Update iOS LaunchURL() to match other platforms to not block.  Add a new method, LaunchSyncURL() to block till there's a result, up to the Timeout specified.
Also enable VisionOS support as this version of "openURL:options:completionHandler:" supports visionOS 1.0+

diff --git a/Engine/Source/Runtime/Core/Private/IOS/IOSPlatformProcess.cpp b/Engine/Source/Runtime/Core/Private/IOS/IOSPlatformProcess.cpp
--- a/Engine/Source/Runtime/Core/Private/IOS/IOSPlatformProcess.cpp
+++ b/Engine/Source/Runtime/Core/Private/IOS/IOSPlatformProcess.cpp
@@ -41,9 +41,6 @@ FRunnableThread* FIOSPlatformProcess::CreateRunnableThread()
 
 void FIOSPlatformProcess::LaunchURL( const TCHAR* URL, const TCHAR* Parms, FString* Error )
 {
-#if PLATFORM_VISIONOS
-	*Error = TEXT("LaunchURL is not supported on VisionOS");
-#else
 	UE_LOG(LogIOS, Log,  TEXT("LaunchURL %s %s"), URL, Parms?Parms:TEXT("") );
 
 	if (FCoreDelegates::ShouldLaunchUrl.IsBound() && !FCoreDelegates::ShouldLaunchUrl.Execute(URL))
@@ -55,6 +52,61 @@ void FIOSPlatformProcess::LaunchURL( const TCHAR* URL, const TCHAR* Parms, FStri
 		return;
 	}
 
+	FString SchemeName;
+	bool bHasSchemeName = FParse::SchemeNameFromURI(URL, SchemeName);
+		
+	NSString* CFUrl = (NSString*)FPlatformString::TCHARToCFString( URL );
+	if(!bHasSchemeName)
+	{
+		CFUrl = [@"http://" stringByAppendingString:CFUrl];
+	}
+	[[UIApplication sharedApplication] openURL: [NSURL URLWithString: CFUrl] options:@{} completionHandler:
+		^(BOOL Success)
+		{
+			if (!Success)
+			{
+				UE_LOG(LogIOS, Warning, TEXT("LaunchURL of %s was NOT successful"), URL);
+			}
+		}
+	];
+	
+	if (Error != nullptr)
+	{
+		*Error = TEXT("");
+	}
+}
+
+/**
+* Opens the specified URL and waits for it to open, up to the Timeout (s).  If there's an error with the opening
+* that's found within the timeout period, or a timeout occurs, Error is set to "Unable to open url"
+*
+* @param URL The url to open
+* @param Parms Not currently used
+* @param Timeout The length of time to wait (in seconds) before it times out
+* @param Error Either blank or with an error message if fails to open or times out
+**/
+void FIOSPlatformProcess::LaunchSyncURL( const TCHAR* URL, const TCHAR* Parms, float Timeout, FString* Error )
+{
+	UE_LOG(LogIOS, Log, TEXT("LauncSynchURL %s %s"), URL, Parms?Parms:TEXT("") );
+
+	if (FCoreDelegates::ShouldLaunchUrl.IsBound() && !FCoreDelegates::ShouldLaunchUrl.Execute(URL))
+	{
+		if (Error)
+		{
+			*Error = TEXT("LauncSynchURL cancelled by delegate");
+		}
+		return;
+	}
+
+	if (Timeout <= 0)
+	{
+		if (Error)
+		{
+			*Error = TEXT("LauncSynchURL invalid Timeout specified");
+		}
+		return;
+	}
+
 	dispatch_semaphore_t Semaphore = dispatch_semaphore_create(0);
 	__block BOOL Result;
 	FString SchemeName;
@@ -73,8 +125,8 @@ void FIOSPlatformProcess::LaunchURL( const TCHAR* URL, const TCHAR* Parms, FStri
 		}
 	];
 	
-	// wait for a result, but timeout after 1s
-	if (dispatch_semaphore_wait(Semaphore, dispatch_time(DISPATCH_TIME_NOW, NSEC_PER_SEC)))
+	// wait for a result, up to "Timeout" seconds
+	if (dispatch_semaphore_wait(Semaphore, dispatch_time(DISPATCH_TIME_NOW, (uint64_t)(Timeout*(double)NSEC_PER_SEC))))
 	{
 		// We timed out, error unknown
 		Result = TRUE;
@@ -85,7 +137,6 @@ void FIOSPlatformProcess::LaunchURL( const TCHAR* URL, const TCHAR* Parms, FStri
 	{
 		*Error = Result ? TEXT("") : TEXT("Unable to open url");
 	}
-#endif
 }
 
 bool FIOSPlatformProcess::CanLaunchURL(const TCHAR* URL)
diff --git a/Engine/Source/Runtime/Core/Public/IOS/IOSPlatformProcess.h b/Engine/Source/Runtime/Core/Public/IOS/IOSPlatformProcess.h
--- a/Engine/Source/Runtime/Core/Public/IOS/IOSPlatformProcess.h
+++ b/Engine/Source/Runtime/Core/Public/IOS/IOSPlatformProcess.h
@@ -32,6 +32,7 @@ struct CORE_API FIOSPlatformProcess : public FApplePlatformProcess
 	static const TCHAR* BaseDir();
 	static FRunnableThread* CreateRunnableThread();
 	static void LaunchURL( const TCHAR* URL, const TCHAR* Parms, FString* Error );
+	static void LaunchSyncURL( const TCHAR* URL, const TCHAR* Parms, float Timeout, FString* Error );
 	static bool CanLaunchURL(const TCHAR* URL);
 	static FString GetGameBundleId();
 	static void SetupGameThread();
