Remove deprecated code methods.

#rnx

diff --git a/Engine/Source/Runtime/Core/Public/Virtualization/VirtualizationSystem.h b/Engine/Source/Runtime/Core/Public/Virtualization/VirtualizationSystem.h
--- a/Engine/Source/Runtime/Core/Public/Virtualization/VirtualizationSystem.h
+++ b/Engine/Source/Runtime/Core/Public/Virtualization/VirtualizationSystem.h
@@ -771,16 +771,6 @@ public:
 	 */
 	virtual FVirtualizationResult TryVirtualizePackages(TConstArrayView<FString> PackagePaths, EVirtualizationOptions Options) = 0;
 
-	UE_DEPRECATED(5.2, "Use The other overload that takes a EVirtualizationOptions parameter")
-	virtual ERehydrationResult TryRehydratePackages(TConstArrayView<FString> PackagePaths, TArray<FText>& OutErrors)
-	{
-		FRehydrationResult Result = TryRehydratePackages(PackagePaths, ERehydrationOptions::None);
-
-		OutErrors = MoveTemp(Result.Errors);
-
-		return Result.WasSuccessful() ? ERehydrationResult::Success : ERehydrationResult::Failed;
-	}
-
 	/**
 	 * Runs the rehydration process on a set of packages. This involves downloading virtualized payloads and placing them back in the trailer of
 	 * the given packages.
@@ -816,21 +806,6 @@ public:
 	virtual void DumpStats() const = 0;
 
 	using GetPayloadActivityInfoFuncRef = TFunctionRef<void(const FString& DebugName, const FString& ConfigName, const FPayloadActivityInfo& PayloadInfo)>;
-
-	UE_DEPRECATED(5.5, "Use IVirtualizationSystem::GetBackendStatistics instead")
-	virtual void GetPayloadActivityInfo(GetPayloadActivityInfoFuncRef Func) const
-	{
-		for (const FBackendStats& Stats : GetBackendStatistics())
-		{
-			Func(Stats.DebugName, Stats.ConfigName, Stats.PayloadActivity);
-		}
-	}
-
-	UE_DEPRECATED(5.5, "Use IVirtualizationSystem::GetSystemStatistics instead")
-	virtual FPayloadActivityInfo GetAccumualtedPayloadActivityInfo() const
-	{
-		return GetSystemStatistics();
-	}
 	
 	/** Access profiling info relating to payload activity per backend. Stats will only be collected if ENABLE_COOK_STATS is enabled.*/
 	virtual TArray<FBackendStats> GetBackendStatistics() const = 0;
