Added missing ARM Yield intrinsic for MSVC

diff --git a/Engine/Source/Runtime/Core/Public/GenericPlatform/GenericPlatformProcess.h b/Engine/Source/Runtime/Core/Public/GenericPlatform/GenericPlatformProcess.h
--- a/Engine/Source/Runtime/Core/Public/GenericPlatform/GenericPlatformProcess.h
+++ b/Engine/Source/Runtime/Core/Public/GenericPlatform/GenericPlatformProcess.h
@@ -904,7 +904,7 @@ struct FGenericPlatformProcess
 #if PLATFORM_USE_SSE2_FOR_THREAD_YIELD
 		_mm_pause();
 #elif PLATFORM_CPU_ARM_FAMILY
-#	if _MSC_VER
+#	if defined(_MSC_VER)
 		__isb(0);
 #	else
 		__asm__ __volatile__("isb");
@@ -966,7 +966,11 @@ struct FGenericPlatformProcess
 		for (uint64 i = 0; i < Cycles; i++)
 		{
 #if PLATFORM_CPU_ARM_FAMILY
+#	if defined(_MSC_VER)
+			__yield();
+#	else
 			__builtin_arm_yield();
+#	endif
 #else
 #			error "Not implemented on this platform"
 #endif
