RemoteObjects: Expose GetRemoteId() on TWeakObjectPtr

diff --git a/Engine/Source/Runtime/Core/Public/UObject/WeakObjectPtrTemplates.h b/Engine/Source/Runtime/Core/Public/UObject/WeakObjectPtrTemplates.h
--- a/Engine/Source/Runtime/Core/Public/UObject/WeakObjectPtrTemplates.h
+++ b/Engine/Source/Runtime/Core/Public/UObject/WeakObjectPtrTemplates.h
@@ -293,6 +293,11 @@ public:
 	{
 		return WeakPtr.HasSameObject(Other);
 	}
+
+	FORCEINLINE auto GetRemoteId() const
+	{
+		return WeakPtr.GetRemoteId();
+	}
 #endif // UE_WITH_REMOTE_OBJECT_HANDLE
 
 	[[nodiscard]] FORCEINLINE bool IsRemote() const
