Core changes to support constinit initialization
Many functions and constructors simply marked constexpr
Containers and allocators gain explicit consteval constructors to handle cases where allocators require full initialization in a constinit context. (e.g. inline allocator)

diff --git a/Engine/Source/Runtime/Core/Public/Algo/MaxElement.h b/Engine/Source/Runtime/Core/Public/Algo/MaxElement.h
--- a/Engine/Source/Runtime/Core/Public/Algo/MaxElement.h
+++ b/Engine/Source/Runtime/Core/Public/Algo/MaxElement.h
@@ -11,7 +11,7 @@
 namespace AlgoImpl
 {
 	template <typename RangeType, typename ProjectionType, typename PredicateType>
-	typename TRangePointerType<RangeType>::Type MaxElementBy(RangeType& Range, ProjectionType Proj, PredicateType Pred)
+	constexpr typename TRangePointerType<RangeType>::Type MaxElementBy(RangeType& Range, ProjectionType Proj, PredicateType Pred)
 	{
 		typename TRangePointerType<RangeType>::Type Result = nullptr;
 
@@ -37,7 +37,7 @@ namespace Algo
 	 * @return A pointer to the maximum element, or nullptr if the range was empty.
 	 */
 	template <typename RangeType>
-	[[nodiscard]] UE_REWRITE auto MaxElement(RangeType&& Range)
+	[[nodiscard]] UE_REWRITE constexpr auto MaxElement(RangeType&& Range)
 		-> decltype(AlgoImpl::MaxElementBy(Range, FIdentityFunctor(), TLess<>()))
 	{
 		return AlgoImpl::MaxElementBy(Range, FIdentityFunctor(), TLess<>());
diff --git a/Engine/Source/Runtime/Core/Public/Containers/Array.h b/Engine/Source/Runtime/Core/Public/Containers/Array.h
--- a/Engine/Source/Runtime/Core/Public/Containers/Array.h
+++ b/Engine/Source/Runtime/Core/Public/Containers/Array.h
@@ -680,6 +680,14 @@ public:
 	{
 	}
 
+	/** Explicit consteval constructor for allocators that require zero-initialization of members for constant initialization */
+	[[nodiscard]] explicit consteval TArray(EConstEval)
+		: AllocatorInstance(ConstEval)
+		, ArrayNum(0)
+		, ArrayMax(AllocatorInstance.GetInitialCapacity())
+	{
+	}
+
 	/**
 	 * Constructor from a raw array of elements.
 	 *
diff --git a/Engine/Source/Runtime/Core/Public/Containers/ArrayView.h b/Engine/Source/Runtime/Core/Public/Containers/ArrayView.h
--- a/Engine/Source/Runtime/Core/Public/Containers/ArrayView.h
+++ b/Engine/Source/Runtime/Core/Public/Containers/ArrayView.h
@@ -144,7 +144,7 @@ public:
 	static_assert(std::is_signed_v<SizeType>, "TArrayView only supports signed index types");
 
 	// Defaulted object behavior - we want compiler-generated functions rather than going through the generic range constructor.
-	[[nodiscard]] TArrayView(const TArrayView&) = default;
+	[[nodiscard]] constexpr TArrayView(const TArrayView&) = default;
 	TArrayView& operator=(const TArrayView&) = default;
 	~TArrayView() = default;
 
diff --git a/Engine/Source/Runtime/Core/Public/Containers/BitArray.h b/Engine/Source/Runtime/Core/Public/Containers/BitArray.h
--- a/Engine/Source/Runtime/Core/Public/Containers/BitArray.h
+++ b/Engine/Source/Runtime/Core/Public/Containers/BitArray.h
@@ -377,6 +377,13 @@ public:
 		// ClearPartialSlackBits is already satisfied since final word does not exist when NumBits == 0
 	}
 
+	[[nodiscard]] explicit consteval TBitArray(EConstEval)
+	:	AllocatorInstance(ConstEval)
+	, 	NumBits(0)
+	,	MaxBits(AllocatorInstance.GetInitialCapacity() * NumBitsPerDWORD)
+	{
+	}
+
 	/**
 	 * Minimal initialization constructor.
 	 * @param Value - The value to initial the bits to.
diff --git a/Engine/Source/Runtime/Core/Public/Containers/ContainerAllocationPolicies.h b/Engine/Source/Runtime/Core/Public/Containers/ContainerAllocationPolicies.h
--- a/Engine/Source/Runtime/Core/Public/Containers/ContainerAllocationPolicies.h
+++ b/Engine/Source/Runtime/Core/Public/Containers/ContainerAllocationPolicies.h
@@ -642,6 +642,11 @@ public:
 			: Data(nullptr)
 		{}
 
+		/** Explicitly consteval constructor. */
+		explicit consteval ForAnyElementType(EConstEval)
+			: Data(nullptr)
+		{}
+
 		/**
 		 * Moves the state of another allocator into this one.  The allocator can be different.
 		 *
@@ -781,7 +786,7 @@ public:
 			return !!Data;
 		}
 
-		SizeType GetInitialCapacity() const
+		constexpr SizeType GetInitialCapacity() const
 		{
 			return 0;
 		}
@@ -814,6 +819,10 @@ public:
 	public:
 		/** Default constructor. */
 		ForElementType() = default;
+		explicit consteval ForElementType(EConstEval) 
+			: ForAnyElementType(ConstEval) 
+		{
+		}
 
 		UE_FORCEINLINE_HINT ElementType* GetAllocation() const
 		{
@@ -873,6 +882,12 @@ public:
 		{
 		}
 
+		explicit consteval ForElementType(EConstEval)
+			: InlineData() // Force value initialization
+			, SecondaryData(ConstEval)
+		{
+		}
+
 		/**
 		 * Moves the state of another allocator into this one.
 		 * Assumes that the allocator is currently empty, i.e. memory may be allocated but any existing elements have already been destructed (if necessary).
@@ -974,7 +989,7 @@ public:
 			return SecondaryData.HasAllocation();
 		}
 
-		SizeType GetInitialCapacity() const
+		constexpr SizeType GetInitialCapacity() const
 		{
 			return NumInlineElements;
 		}
@@ -993,8 +1008,8 @@ public:
 		ForElementType(const ForElementType&);
 		ForElementType& operator=(const ForElementType&);
 
-		/** The data is stored in this array if less than NumInlineElements is needed. */
-		TTypeCompatibleBytes<ElementType> InlineData[NumInlineElements];
+		/** The data is stored in this array if less than NumInlineElements is needed. Uninitialized by default. */
+		alignas(ElementType) uint8 InlineData[sizeof(ElementType) * NumInlineElements];
 
 		/** The data is allocated through the indirect allocation policy if more than NumInlineElements is needed. */
 		typename SecondaryAllocator::template ForElementType<ElementType> SecondaryData;
diff --git a/Engine/Source/Runtime/Core/Public/Containers/EnumAsByte.h b/Engine/Source/Runtime/Core/Public/Containers/EnumAsByte.h
--- a/Engine/Source/Runtime/Core/Public/Containers/EnumAsByte.h
+++ b/Engine/Source/Runtime/Core/Public/Containers/EnumAsByte.h
@@ -25,16 +25,16 @@ class TEnumAsByte
 public:
 	using EnumType = InEnumType;
 
-	[[nodiscard]] TEnumAsByte() = default;
-	[[nodiscard]] TEnumAsByte(const TEnumAsByte&) = default;
-	TEnumAsByte& operator=(const TEnumAsByte&) = default;
+	[[nodiscard]] constexpr TEnumAsByte() = default;
+	[[nodiscard]] constexpr TEnumAsByte(const TEnumAsByte&) = default;
+	constexpr TEnumAsByte& operator=(const TEnumAsByte&) = default;
 
 	/**
 	 * Constructor, initialize to the enum value.
 	 *
 	 * @param InValue value to construct with.
 	 */
-	[[nodiscard]] UE_FORCEINLINE_HINT TEnumAsByte( EnumType InValue )
+	[[nodiscard]] UE_FORCEINLINE_HINT constexpr TEnumAsByte( EnumType InValue )
 		: Value(static_cast<uint8>(InValue))
 	{
 	}
@@ -44,7 +44,7 @@ public:
 	 *
 	 * @param InValue value to construct with.
 	 */
-	[[nodiscard]] explicit UE_FORCEINLINE_HINT TEnumAsByte( int32 InValue )
+	[[nodiscard]] explicit UE_FORCEINLINE_HINT constexpr TEnumAsByte( int32 InValue )
 		: Value(static_cast<uint8>(InValue))
 	{
 	}
@@ -54,7 +54,7 @@ public:
 	 *
 	 * @param InValue value to construct with.
 	 */
-	[[nodiscard]] explicit UE_FORCEINLINE_HINT TEnumAsByte( uint8 InValue )
+	[[nodiscard]] explicit UE_FORCEINLINE_HINT constexpr TEnumAsByte( uint8 InValue )
 		: Value(InValue)
 	{
 	}
@@ -66,7 +66,7 @@ public:
 	 * @param InValue The value to compare with.
 	 * @return true if the two values are equal, false otherwise.
 	 */
-	[[nodiscard]] bool operator==( EnumType InValue ) const
+	[[nodiscard]] constexpr bool operator==( EnumType InValue ) const
 	{
 		return static_cast<EnumType>(Value) == InValue;
 	}
@@ -77,13 +77,13 @@ public:
 	 * @param InValue The value to compare with.
 	 * @return true if the two values are equal, false otherwise.
 	 */
-	[[nodiscard]] bool operator==(TEnumAsByte InValue) const
+	[[nodiscard]] constexpr bool operator==(TEnumAsByte InValue) const
 	{
 		return Value == InValue.Value;
 	}
 
 	/** Implicit conversion to EnumType. */
-	[[nodiscard]] operator EnumType() const
+	[[nodiscard]] constexpr operator EnumType() const
 	{
 		return (EnumType)Value;
 	}
@@ -95,7 +95,7 @@ public:
 	 *
 	 * @return The enumeration value.
 	 */
-	[[nodiscard]] EnumType GetValue() const
+	[[nodiscard]] constexpr EnumType GetValue() const
 	{
 		return (EnumType)Value;
 	}
@@ -105,7 +105,7 @@ public:
 	 *
 	 * @return The enumeration value.
 	 */
-	[[nodiscard]] uint8 GetIntValue() const
+	[[nodiscard]] constexpr uint8 GetIntValue() const
 	{
 		return Value;
 	}
@@ -117,7 +117,7 @@ private:
 };
 
 template<class T>
-[[nodiscard]] UE_FORCEINLINE_HINT uint32 GetTypeHash(const TEnumAsByte<T>& Enum)
+[[nodiscard]] UE_FORCEINLINE_HINT constexpr uint32 GetTypeHash(const TEnumAsByte<T>& Enum)
 {
 	return GetTypeHash((uint8)Enum.GetValue());
 }
diff --git a/Engine/Source/Runtime/Core/Public/Containers/Map.h b/Engine/Source/Runtime/Core/Public/Containers/Map.h
--- a/Engine/Source/Runtime/Core/Public/Containers/Map.h
+++ b/Engine/Source/Runtime/Core/Public/Containers/Map.h
@@ -142,6 +142,10 @@ public:
 
 protected:
 	[[nodiscard]] TMapBase() = default;
+	[[nodiscard]] explicit consteval TMapBase(EConstEval)
+	: Pairs(ConstEval)
+	{
+	} 
 	[[nodiscard]] TMapBase(TMapBase&&) = default;
 	[[nodiscard]] TMapBase(const TMapBase&) = default;
 	TMapBase& operator=(TMapBase&&) = default;
@@ -1149,6 +1153,10 @@ protected:
 	typedef TMapBase<KeyType, ValueType, SetAllocator, KeyFuncs> Super;
 
 	[[nodiscard]] TSortableMapBase() = default;
+	[[nodiscard]] explicit consteval TSortableMapBase(EConstEval)
+	: Super(ConstEval)
+	{
+	}
 	[[nodiscard]] TSortableMapBase(TSortableMapBase&&) = default;
 	[[nodiscard]] TSortableMapBase(const TSortableMapBase&) = default;
 	TSortableMapBase& operator=(TSortableMapBase&&) = default;
@@ -1322,6 +1330,11 @@ public:
 	TMap& operator=(TMap&&) = default;
 	TMap& operator=(const TMap&) = default;
 
+	[[nodiscard]] explicit consteval TMap(EConstEval)
+	: Super(ConstEval)
+	{
+	}
+
 	/** Constructor for moving elements from a TMap with a different SetAllocator */
 	template<typename OtherSetAllocator>
 	[[nodiscard]] TMap(TMap<KeyType, ValueType, OtherSetAllocator, KeyFuncs>&& Other)
@@ -1882,7 +1895,7 @@ class TScriptMap
 	using DerivedType = std::conditional_t<std::is_void_v<InDerivedType>, TScriptMap, InDerivedType>;
 
 public:
-	[[nodiscard]] static FScriptMapLayout GetScriptLayout(int32 KeySize, int32 KeyAlignment, int32 ValueSize, int32 ValueAlignment)
+	[[nodiscard]] static constexpr FScriptMapLayout GetScriptLayout(int32 KeySize, int32 KeyAlignment, int32 ValueSize, int32 ValueAlignment)
 	{
 		FScriptMapLayout Result;
 
diff --git a/Engine/Source/Runtime/Core/Public/Containers/Set.h b/Engine/Source/Runtime/Core/Public/Containers/Set.h
--- a/Engine/Source/Runtime/Core/Public/Containers/Set.h
+++ b/Engine/Source/Runtime/Core/Public/Containers/Set.h
@@ -210,6 +210,12 @@ public:
 	/** Initialization constructor. */
 	[[nodiscard]] UE_FORCEINLINE_HINT TSet() = default;
 
+	[[nodiscard]] explicit consteval TSet(EConstEval)
+	: Elements(ConstEval)
+	, Hash(ConstEval)
+	{
+	}
+
 	/** Copy constructor. */
 	[[nodiscard]] UE_FORCEINLINE_HINT TSet(const TSet& Copy)
 	{
@@ -1784,7 +1790,7 @@ class TScriptSet
 	using DerivedType = std::conditional_t<std::is_void_v<InDerivedType>, TScriptSet, InDerivedType>;
 
 public:
-	[[nodiscard]] static FScriptSetLayout GetScriptLayout(int32 ElementSize, int32 ElementAlignment)
+	[[nodiscard]] static constexpr FScriptSetLayout GetScriptLayout(int32 ElementSize, int32 ElementAlignment)
 	{
 		FScriptSetLayout Result;
 
diff --git a/Engine/Source/Runtime/Core/Public/Containers/SparseArray.h b/Engine/Source/Runtime/Core/Public/Containers/SparseArray.h
--- a/Engine/Source/Runtime/Core/Public/Containers/SparseArray.h
+++ b/Engine/Source/Runtime/Core/Public/Containers/SparseArray.h
@@ -76,6 +76,15 @@ protected:
 	{
 	}
 
+	/** Explicitly consteval constructor. */
+	[[nodiscard]] explicit consteval TSparseArrayBase(EConstEval)
+		: Data(ConstEval)
+		, AllocationFlags(ConstEval)
+		, FirstFreeIndex(-1)
+		, NumFreeIndices(0)
+	{
+	}
+
 	/** Copy constructor. */
 	[[nodiscard]] TSparseArrayBase(const TSparseArrayBase& InCopy)
 		: FirstFreeIndex(-1)
@@ -853,6 +862,11 @@ public:
 	/** Default constructor. */
 	[[nodiscard]] TSparseArray() = default;
 
+	[[nodiscard]] explicit consteval TSparseArray(EConstEval)
+		: SuperType(ConstEval)
+	{
+	}
+
 	/** Move constructor. */
 	[[nodiscard]] TSparseArray(TSparseArray&& InCopy)
 	{
@@ -1449,7 +1463,7 @@ class TScriptSparseArray
 	using DerivedType = std::conditional_t<std::is_void_v<InDerivedType>, TScriptSparseArray, InDerivedType>;
 
 public:
-	[[nodiscard]] static FScriptSparseArrayLayout GetScriptLayout(int32 ElementSize, int32 ElementAlignment)
+	[[nodiscard]] static constexpr FScriptSparseArrayLayout GetScriptLayout(int32 ElementSize, int32 ElementAlignment)
 	{
 		FScriptSparseArrayLayout Result;
 		Result.Alignment     = FMath::Max(ElementAlignment, (int32)alignof(FFreeListLink));
diff --git a/Engine/Source/Runtime/Core/Public/Containers/UnrealString.h.inl b/Engine/Source/Runtime/Core/Public/Containers/UnrealString.h.inl
--- a/Engine/Source/Runtime/Core/Public/Containers/UnrealString.h.inl
+++ b/Engine/Source/Runtime/Core/Public/Containers/UnrealString.h.inl
@@ -70,6 +70,10 @@ private:
 
 public:
 	[[nodiscard]] UE_STRING_CLASS() = default;
+	[[nodiscard]] explicit consteval UE_STRING_CLASS(EConstEval) 
+		: Data(ConstEval)
+	{
+	}
 	[[nodiscard]] UE_STRING_CLASS(UE_STRING_CLASS&&) = default;
 	[[nodiscard]] UE_STRING_CLASS(const UE_STRING_CLASS&) = default;
 	UE_STRING_CLASS& operator=(UE_STRING_CLASS&&) = default;
diff --git a/Engine/Source/Runtime/Core/Public/IO/IoHash.h b/Engine/Source/Runtime/Core/Public/IO/IoHash.h
--- a/Engine/Source/Runtime/Core/Public/IO/IoHash.h
+++ b/Engine/Source/Runtime/Core/Public/IO/IoHash.h
@@ -35,7 +35,7 @@ public:
 	using ByteArray = uint8[20];
 
 	/** Construct a zero hash. */
-	FIoHash() = default;
+	constexpr FIoHash() = default;
 
 	/** Construct a hash from an array of 20 bytes. */
 	inline explicit FIoHash(const ByteArray& Hash);
diff --git a/Engine/Source/Runtime/Core/Public/IO/PackageId.h b/Engine/Source/Runtime/Core/Public/IO/PackageId.h
--- a/Engine/Source/Runtime/Core/Public/IO/PackageId.h
+++ b/Engine/Source/Runtime/Core/Public/IO/PackageId.h
@@ -23,7 +23,7 @@ class FPackageId
 	inline explicit FPackageId(uint64 InId): Id(InId) {}
 
 public:
-	FPackageId() = default;
+	constexpr FPackageId() = default;
 
 	CORE_API static FPackageId FromName(const FName& Name);
 
diff --git a/Engine/Source/Runtime/Core/Public/Misc/AssertionMacros.h b/Engine/Source/Runtime/Core/Public/Misc/AssertionMacros.h
--- a/Engine/Source/Runtime/Core/Public/Misc/AssertionMacros.h
+++ b/Engine/Source/Runtime/Core/Public/Misc/AssertionMacros.h
@@ -481,7 +481,7 @@ namespace UEAsserts_Private
 
 // Returns FName(TEXT("EnumeratorName")), while statically verifying that the enumerator exists in the enum
 #define GET_ENUMERATOR_NAME_CHECKED(EnumName, EnumeratorName) \
-	((void)sizeof(UEAsserts_Private::GetMemberNameCheckedJunk(EnumName::EnumeratorName)), FName(TEXT(#EnumeratorName)))
+	(FName(GET_ENUMERATOR_NAME_STRING_VIEW_CHECKED(EnumName, EnumeratorName)))
 
 #define GET_ENUMERATOR_NAME_STRING_CHECKED(EnumName, EnumeratorName) \
 	((void)sizeof(UEAsserts_Private::GetMemberNameCheckedJunk(EnumName::EnumeratorName)), TEXT(#EnumeratorName))
@@ -491,7 +491,7 @@ namespace UEAsserts_Private
 
 // Returns FName(TEXT("MemberName")), while statically verifying that the member exists in ClassName
 #define GET_MEMBER_NAME_CHECKED(ClassName, MemberName) \
-	((void)sizeof(UEAsserts_Private::GetMemberNameCheckedJunk(((ClassName*)0)->MemberName)), FName(TEXT(#MemberName)))
+	(FName(GET_MEMBER_NAME_STRING_VIEW_CHECKED(ClassName, MemberName)))
 
 #define GET_MEMBER_NAME_STRING_CHECKED(ClassName, MemberName) \
 	((void)sizeof(UEAsserts_Private::GetMemberNameCheckedJunk(((ClassName*)0)->MemberName)), TEXT(#MemberName))
@@ -501,7 +501,7 @@ namespace UEAsserts_Private
 
 // Returns FName(TEXT("FunctionName")), while statically verifying that the function exists in ClassName
 #define GET_FUNCTION_NAME_CHECKED(ClassName, FunctionName) \
-	((void)sizeof(&ClassName::FunctionName), FName(TEXT(#FunctionName)))
+	(FName(GET_FUNCTION_NAME_STRING_VIEW_CHECKED(ClassName, FunctionName)))
 
 #define GET_FUNCTION_NAME_STRING_CHECKED(ClassName, FunctionName) \
 	((void)sizeof(&ClassName::FunctionName), TEXT(#FunctionName))
@@ -512,16 +512,16 @@ namespace UEAsserts_Private
 // Returns FName(TEXT("FunctionName")), while statically verifying that the function exists in ClassName
 // Handles overloaded functions by specifying the argument of the overload to use
 #define GET_FUNCTION_NAME_CHECKED_OneParam(ClassName, FunctionName, ArgType) \
-	((void)sizeof((std::declval<ClassName&>().FunctionName(std::declval<PREPROCESSOR_REMOVE_OPTIONAL_PARENS(ArgType)>()), (int)0)), FName(TEXT(#FunctionName)))
+	(FName(GET_FUNCTION_NAME_STRING_CHECKED_OneParam(ClassName, FunctionName, ArgType)))
 
 #define GET_FUNCTION_NAME_CHECKED_TwoParams(ClassName, FunctionName, ArgType1, ArgType2) \
-	((void)sizeof((std::declval<ClassName&>().FunctionName(std::declval<PREPROCESSOR_REMOVE_OPTIONAL_PARENS(ArgType1)>(), std::declval<PREPROCESSOR_REMOVE_OPTIONAL_PARENS(ArgType2)>()), (int)0)), FName(TEXT(#FunctionName)))
+	(FName(GET_FUNCTION_NAME_STRING_CHECKED_TwoParams(ClassName, FunctionName, ArgType1, ArgType2)))
 
 #define GET_FUNCTION_NAME_CHECKED_ThreeParams(ClassName, FunctionName, ArgType1, ArgType2, ArgType3) \
-	((void)sizeof((std::declval<ClassName&>().FunctionName(std::declval<PREPROCESSOR_REMOVE_OPTIONAL_PARENS(ArgType1)>(), std::declval<PREPROCESSOR_REMOVE_OPTIONAL_PARENS(ArgType2)>(), std::declval<PREPROCESSOR_REMOVE_OPTIONAL_PARENS(ArgType3)>()), (int)0)), FName(TEXT(#FunctionName)))
+	(FName(GET_FUNCTION_NAME_STRING_CHECKED_ThreeParams(ClassName, FunctionName, ArgType1, ArgType2, ArgType3)))
 
 #define GET_FUNCTION_NAME_CHECKED_FourParams(ClassName, FunctionName, ArgType1, ArgType2, ArgType3, ArgType4) \
-	((void)sizeof((std::declval<ClassName&>().FunctionName(std::declval<PREPROCESSOR_REMOVE_OPTIONAL_PARENS(ArgType1)>(), std::declval<PREPROCESSOR_REMOVE_OPTIONAL_PARENS(ArgType2)>(), std::declval<PREPROCESSOR_REMOVE_OPTIONAL_PARENS(ArgType3)>(), std::declval<PREPROCESSOR_REMOVE_OPTIONAL_PARENS(ArgType4)>()), (int)0)), FName(TEXT(#FunctionName)))
+	(FName(GET_FUNCTION_NAME_STRING_CHECKED_FourParams(ClassName, FunctionName, ArgType1, ArgType2, ArgType3, ArgType4)))
 
 #define GET_FUNCTION_NAME_STRING_CHECKED_OneParam(ClassName, FunctionName, ArgType) \
 	((void)sizeof((std::declval<ClassName&>().FunctionName(std::declval<PREPROCESSOR_REMOVE_OPTIONAL_PARENS(ArgType)>()), (int)0)), TEXT(#FunctionName))
diff --git a/Engine/Source/Runtime/Core/Public/Misc/CoreMiscDefines.h b/Engine/Source/Runtime/Core/Public/Misc/CoreMiscDefines.h
--- a/Engine/Source/Runtime/Core/Public/Misc/CoreMiscDefines.h
+++ b/Engine/Source/Runtime/Core/Public/Misc/CoreMiscDefines.h
@@ -156,6 +156,9 @@ enum EForceInit
 	ForceInitToZero
 };
 enum ENoInit {NoInit};
+// Used to add an explicitly consteval constructor when the default constructor for a type cannot
+// be made constexpr (e.g. we are attempting to avoid zero-initialization when constructing at runtime)
+enum EConstEval {ConstEval};
 enum EInPlace {InPlace};
 enum EPerElement {PerElement};
 
diff --git a/Engine/Source/Runtime/Core/Public/Misc/StructBuilder.h b/Engine/Source/Runtime/Core/Public/Misc/StructBuilder.h
--- a/Engine/Source/Runtime/Core/Public/Misc/StructBuilder.h
+++ b/Engine/Source/Runtime/Core/Public/Misc/StructBuilder.h
@@ -13,7 +13,7 @@
 class FStructBuilder
 {
 public:
-	FStructBuilder()
+	constexpr FStructBuilder()
 		: EndOfLastMember(0)
 		, Alignment      (0)
 	{
@@ -27,7 +27,7 @@ public:
 	 *
 	 * @return  The offset of the member from the start of the struct.
 	 */
-	int32 AddMember(int32 MemberSize, int32 MemberAlignment)
+	constexpr int32 AddMember(int32 MemberSize, int32 MemberAlignment)
 	{
 		int32 Result    = Align(EndOfLastMember, MemberAlignment);
 		EndOfLastMember = Result + MemberSize;
@@ -40,7 +40,7 @@ public:
 	 *
 	 * @return  The size of the struct.
 	 */
-	int32 GetSize() const
+	constexpr int32 GetSize() const
 	{
 		return Align(EndOfLastMember, Alignment);
 	}
@@ -50,7 +50,7 @@ public:
 	 *
 	 * @return  The alignment of the struct.
 	 */
-	int32 GetAlignment() const
+	constexpr int32 GetAlignment() const
 	{
 		return Alignment;
 	}
diff --git a/Engine/Source/Runtime/Core/Public/Templates/Tuple.h b/Engine/Source/Runtime/Core/Public/Templates/Tuple.h
--- a/Engine/Source/Runtime/Core/Public/Templates/Tuple.h
+++ b/Engine/Source/Runtime/Core/Public/Templates/Tuple.h
@@ -47,7 +47,7 @@ template <typename KeyType, typename ValueType>
 using TPair = TTuple<KeyType, ValueType>;
 
 template <typename... Types>
-TTuple<std::decay_t<Types>...> MakeTuple(Types&&... Args);
+constexpr TTuple<std::decay_t<Types>...> MakeTuple(Types&&... Args);
 
 namespace UE::Core::Private::Tuple
 {
@@ -64,12 +64,12 @@ namespace UE::Core::Private::Tuple
 	struct TTupleBaseElement
 	{
 		template <typename ArgType>
-		explicit TTupleBaseElement(EForwardingConstructor, ArgType&& Arg)
+		constexpr explicit TTupleBaseElement(EForwardingConstructor, ArgType&& Arg)
 			: Value(Forward<ArgType>(Arg))
 		{
 		}
 
-		TTupleBaseElement()
+		constexpr TTupleBaseElement()
 			: Value()
 		{
 		}
@@ -86,12 +86,12 @@ namespace UE::Core::Private::Tuple
 	struct TTupleBaseElement<T, 0, 2>
 	{
 		template <typename ArgType>
-		explicit TTupleBaseElement(EForwardingConstructor, ArgType&& Arg)
+		constexpr explicit TTupleBaseElement(EForwardingConstructor, ArgType&& Arg)
 			: Key(Forward<ArgType>(Arg))
 		{
 		}
 
-		TTupleBaseElement()
+		constexpr TTupleBaseElement()
 			: Key()
 		{
 		}
@@ -260,7 +260,7 @@ namespace UE::Core::Private::Tuple
 	struct TTupleBase<TIntegerSequence<uint32, Indices...>, Types...> : TTupleBaseElement<Types, Indices, sizeof...(Types)>...
 	{
 		template <typename... ArgTypes>
-		explicit TTupleBase(EForwardingConstructor, ArgTypes&&... Args)
+		constexpr explicit TTupleBase(EForwardingConstructor, ArgTypes&&... Args)
 			: TTupleBaseElement<Types, Indices, sizeof...(Types)>(ForwardingConstructor, Forward<ArgTypes>(Args))...
 		{
 		}
@@ -271,7 +271,7 @@ namespace UE::Core::Private::Tuple
 		{
 		}
 
-		TTupleBase() = default;
+		constexpr TTupleBase() = default;
 		TTupleBase(TTupleBase&& Other) = default;
 		TTupleBase(const TTupleBase& Other) = default;
 		TTupleBase& operator=(TTupleBase&& Other) = default;
@@ -480,7 +480,7 @@ namespace UE::Core::Private::Tuple
 	}
 
 	template <typename... ElementTypes, typename... Types>
-	UE_FORCEINLINE_HINT TTuple<ElementTypes...> MakeTupleImpl(Types&&... Args)
+	UE_FORCEINLINE_HINT constexpr TTuple<ElementTypes...> MakeTupleImpl(Types&&... Args)
 	{
 		return TTuple<ElementTypes...>(Forward<Types>(Args)...);
 	}
@@ -629,7 +629,7 @@ private:
 public:
 	template <UE::CNotCVRefTo<TTuple>... ArgTypes>
 		requires (std::is_constructible_v<Types, ArgTypes&&> && ...)
-	explicit(!std::conjunction_v<std::is_convertible<ArgTypes&&, Types>...>) TTuple(ArgTypes&&... Args)
+	constexpr explicit(!std::conjunction_v<std::is_convertible<ArgTypes&&, Types>...>) TTuple(ArgTypes&&... Args)
 		: Super(UE::Core::Private::Tuple::ForwardingConstructor, Forward<ArgTypes>(Args)...)
 	{
 	}
@@ -648,7 +648,7 @@ public:
 	{
 	}
 
-	TTuple() requires (std::is_default_constructible_v<Types> && ...) = default;
+	constexpr TTuple() requires (std::is_default_constructible_v<Types> && ...) = default;
 	TTuple(TTuple&&) = default;
 	TTuple(const TTuple&) = default;
 	TTuple& operator=(TTuple&&) = default;
@@ -764,7 +764,7 @@ using TTupleElement = UE::Core::Private::Tuple::TCVTupleElement<Index, const vol
  * }
  */
 template <typename... Types>
-UE_FORCEINLINE_HINT TTuple<std::decay_t<Types>...> MakeTuple(Types&&... Args)
+UE_FORCEINLINE_HINT constexpr TTuple<std::decay_t<Types>...> MakeTuple(Types&&... Args)
 {
 	return UE::Core::Private::Tuple::MakeTupleImpl<std::decay_t<Types>...>(Forward<Types>(Args)...);
 }
diff --git a/Engine/Source/Runtime/Core/Public/Templates/UniquePtr.h b/Engine/Source/Runtime/Core/Public/Templates/UniquePtr.h
--- a/Engine/Source/Runtime/Core/Public/Templates/UniquePtr.h
+++ b/Engine/Source/Runtime/Core/Public/Templates/UniquePtr.h
@@ -120,7 +120,7 @@ public:
 	/**
 	 * Default constructor - initializes the TUniquePtr to null.
 	 */
-	inline TUniquePtr()
+	inline constexpr TUniquePtr()
 		: Deleter()
 		, Ptr    (nullptr)
 	{
@@ -511,7 +511,7 @@ public:
 	/**
 	 * Default constructor - initializes the TUniquePtr to null.
 	 */
-	inline TUniquePtr()
+	inline constexpr TUniquePtr()
 		: Deleter()
 		, Ptr    (nullptr)
 	{
@@ -565,7 +565,7 @@ public:
 	/**
 	 * nullptr constructor - initializes the TUniquePtr to null.
 	 */
-	inline TUniquePtr(TYPE_OF_NULLPTR)
+	inline constexpr TUniquePtr(TYPE_OF_NULLPTR)
 		: Deleter()
 		, Ptr    (nullptr)
 	{
diff --git a/Engine/Source/Runtime/Core/Public/Templates/UnrealTemplate.h b/Engine/Source/Runtime/Core/Public/Templates/UnrealTemplate.h
--- a/Engine/Source/Runtime/Core/Public/Templates/UnrealTemplate.h
+++ b/Engine/Source/Runtime/Core/Public/Templates/UnrealTemplate.h
@@ -740,7 +740,7 @@ T&& DeclVal();
  * @return The object converted to the specified type.
  */
 template <typename T>
-UE_REWRITE T ImplicitConv(typename TIdentity<T>::Type Obj)
+UE_REWRITE constexpr T ImplicitConv(typename TIdentity<T>::Type Obj)
 {
 	return Obj;
 }
diff --git a/Engine/Source/Runtime/Core/Public/UObject/NameTypes.h b/Engine/Source/Runtime/Core/Public/UObject/NameTypes.h
--- a/Engine/Source/Runtime/Core/Public/UObject/NameTypes.h
+++ b/Engine/Source/Runtime/Core/Public/UObject/NameTypes.h
@@ -59,8 +59,8 @@ struct FNameEntryId
 	using IntrusiveUnsetOptionalStateType = FNameEntryId;
 	
 	// Default initialize to be equal to NAME_None
-	FNameEntryId() : Value(0) {}
-	FNameEntryId(ENoInit) {}
+	constexpr FNameEntryId() : Value(0) {}
+	constexpr FNameEntryId(ENoInit) {}
 	explicit FNameEntryId(FIntrusiveUnsetOptionalState) : Value(~0u) {}
 
 	bool IsNone() const
@@ -954,7 +954,7 @@ public:
 	/**
 	 * Default constructor, initialized to None
 	 */
-	FORCEINLINE FName()
+	FORCEINLINE constexpr FName()
 #if !UE_FNAME_OUTLINE_NUMBER
 		: Number(NAME_NO_NUMBER_INTERNAL)
 #endif //!UE_FNAME_OUTLINE_NUMBER
@@ -964,7 +964,7 @@ public:
 	/**
 	 * Scary no init constructor, used for something obscure in UObjectBase
 	 */
-	explicit FName(ENoInit)
+	explicit constexpr FName(ENoInit)
 		: ComparisonIndex(NoInit)
 #if WITH_CASE_PRESERVING_NAME
 		, DisplayIndex(NoInit)
diff --git a/Engine/Source/Runtime/CoreUObject/Public/Misc/PackagePath.h b/Engine/Source/Runtime/CoreUObject/Public/Misc/PackagePath.h
--- a/Engine/Source/Runtime/CoreUObject/Public/Misc/PackagePath.h
+++ b/Engine/Source/Runtime/CoreUObject/Public/Misc/PackagePath.h
@@ -91,7 +91,7 @@ public:
 	// Public API
 
 	/** Construct an empty PackagePath */
-	FPackagePath() = default;
+	constexpr FPackagePath() = default;
 
 	// Copy/Move constructors/assignment
 
