Fix for a case where UnifiedErrorCodes still have spaces in them.

diff --git a/Engine/Source/Runtime/Core/Private/Experimental/IO/IoStatusError.cpp b/Engine/Source/Runtime/Core/Private/Experimental/IO/IoStatusError.cpp
--- a/Engine/Source/Runtime/Core/Private/Experimental/IO/IoStatusError.cpp
+++ b/Engine/Source/Runtime/Core/Private/Experimental/IO/IoStatusError.cpp
@@ -46,7 +46,7 @@ namespace UE::UnifiedError::IoStore
 
 		FAnsiString GetErrorCodeString(const UE::UnifiedError::FError& Error) const override
 		{
-			return TCHAR_TO_ANSI(GetIoErrorText(EIoErrorCode(Error.GetErrorCode())));
+			return FAnsiString(*GetErrorName(Error));
 		}
 
 		FAnsiString GetModuleIdString(const UE::UnifiedError::FError& Error) const override
