Defaulting GetDiskTotalAndFreeSpace to true for iOS

#rnx

diff --git a/Engine/Source/Runtime/Core/Public/IOS/IOSPlatformMisc.h b/Engine/Source/Runtime/Core/Public/IOS/IOSPlatformMisc.h
--- a/Engine/Source/Runtime/Core/Public/IOS/IOSPlatformMisc.h
+++ b/Engine/Source/Runtime/Core/Public/IOS/IOSPlatformMisc.h
@@ -122,7 +122,7 @@ struct CORE_API FIOSPlatformMisc : public FApplePlatformMisc
 	static FString GetDeviceId();
 	static FString GetOSVersion();
 	static FString GetUniqueAdvertisingId();
-	static bool GetDiskTotalAndFreeSpace(const FString& InPath, uint64& TotalNumberOfBytes, uint64& NumberOfFreeBytes, bool IncludeFreeable = false);
+	static bool GetDiskTotalAndFreeSpace(const FString& InPath, uint64& TotalNumberOfBytes, uint64& NumberOfFreeBytes, bool IncludeFreeable = true);
 	
 	static void RequestStoreReview();
 
