MiMalloc comment fix

diff --git a/Engine/Source/Runtime/Core/Public/HAL/MallocMimalloc.h b/Engine/Source/Runtime/Core/Public/HAL/MallocMimalloc.h
--- a/Engine/Source/Runtime/Core/Public/HAL/MallocMimalloc.h
+++ b/Engine/Source/Runtime/Core/Public/HAL/MallocMimalloc.h
@@ -18,7 +18,7 @@
 #if MIMALLOC_ENABLED
 
 /**
- * TBB 64-bit scalable memory allocator.
+ * MiMalloc 64-bit scalable memory allocator.
  */
 class FMallocMimalloc final
 	: public FMalloc
