Don't increment PlatformMMIO in mapped files on Android as it's already set globally via stats in 44445310

#rnx

diff --git a/Engine/Source/Runtime/Core/Private/Android/AndroidPlatformFile.cpp b/Engine/Source/Runtime/Core/Private/Android/AndroidPlatformFile.cpp
--- a/Engine/Source/Runtime/Core/Private/Android/AndroidPlatformFile.cpp
+++ b/Engine/Source/Runtime/Core/Private/Android/AndroidPlatformFile.cpp
@@ -1154,7 +1154,6 @@ public:
 
 	virtual IMappedFileRegion* MapRegion(int64 Offset = 0, int64 BytesToMap = MAX_int64, FFileMappingFlags Flags = EMappedFileFlags::ENone) override
 	{
-		LLM_PLATFORM_SCOPE(ELLMTag::PlatformMMIO);
 		const int64 CurrentFileSize = GetCurrentFileSize();
 		check(Offset < CurrentFileSize); // don't map zero bytes and don't map off the end of the file
 		BytesToMap = FMath::Min<int64>(BytesToMap, CurrentFileSize - Offset);
@@ -1185,7 +1184,6 @@ public:
 			UE_LOG(LogAndroidFile, Warning, TEXT("Failed to map memory %s, error is %d"), *Filename, errno);
 			return nullptr;
 		}
-		LLM_IF_ENABLED(FLowLevelMemTracker::Get().OnLowLevelAlloc(ELLMTracker::Platform, AlignedMapPtr, AlignedSize));
 
 		// create a mapping for this range
 		const uint8* MapPtr = AlignedMapPtr + Offset - AlignedOffset;
@@ -1196,11 +1194,9 @@ public:
 
 	void UnMap(const FAndroidMappedFileRegion* Region)
 	{
-		LLM_PLATFORM_SCOPE(ELLMTag::PlatformMMIO);
 		check(NumOutstandingRegions > 0);
 		NumOutstandingRegions--;
 
-		LLM_IF_ENABLED(FLowLevelMemTracker::Get().OnLowLevelFree(ELLMTracker::Platform, Region->AlignedPtr));
 		const int Res = munmap(const_cast<uint8*>(Region->AlignedPtr), Region->AlignedSize);
 #if LOG_ANDROID_FILE
 		FPlatformMisc::LowLevelOutputDebugStringf(TEXT("Failed to unmap region from %s, errno=%s"), *Filename, UTF8_TO_TCHAR(strerror(errno)));
