[Algo]

Specifies to use the global ::Invoke implementation rather than whatever is found during lookup rules since this will become ambiguous if any other SomeOtherNamespace::Invoke exists when using TopologicalSort from within that namespace.

#rnx

diff --git a/Engine/Source/Runtime/Core/Public/Algo/KahnTopologicalSort.h b/Engine/Source/Runtime/Core/Public/Algo/KahnTopologicalSort.h
--- a/Engine/Source/Runtime/Core/Public/Algo/KahnTopologicalSort.h
+++ b/Engine/Source/Runtime/Core/Public/Algo/KahnTopologicalSort.h
@@ -190,7 +190,7 @@ namespace AlgoImpl
 		Handle = 0;
 		for (const ElementType& Element : UniqueRange)
 		{
-			const auto& DependenciesInput = Invoke(GetElementDependencies, Element);
+			const auto& DependenciesInput = ::Invoke(GetElementDependencies, Element);
 			TArray<FKahnHandle>& UniqueElementDependencies = Dependencies[Handle];
 
 			for (const ElementType& Dependency : DependenciesInput)
