LowLevelTests: Expose CheckScope and EnsureScope as a public API in Functional Tests.
#rnx

diff --git a/Engine/Source/Runtime/Core/Public/Tests/CheckScope.h b/Engine/Source/Runtime/Core/Public/Tests/CheckScope.h
--- a/Engine/Source/Runtime/Core/Public/Tests/CheckScope.h
+++ b/Engine/Source/Runtime/Core/Public/Tests/CheckScope.h
@@ -8,14 +8,14 @@ class FCheckScopeOutputDeviceError;
 ///@brief Scope that captures failed `check` calls
 struct FCheckScope
 {
-	FCheckScope();
+	CORE_API FCheckScope();
 
 	///@brief Only catches failed `check` calls that contain Msg
 	///@param Msg string to look for in the `check` message
-	FCheckScope(const ANSICHAR* Msg);
-	~FCheckScope();
+	CORE_API explicit FCheckScope(const ANSICHAR* Msg);
+	CORE_API ~FCheckScope();
 
-	int GetCount();
+	CORE_API int GetCount();
 private:
 	FCheckScopeOutputDeviceError* DeviceError;
 	bool bIgnoreDebugger;
diff --git a/Engine/Source/Runtime/Core/Public/Tests/EnsureScope.h b/Engine/Source/Runtime/Core/Public/Tests/EnsureScope.h
--- a/Engine/Source/Runtime/Core/Public/Tests/EnsureScope.h
+++ b/Engine/Source/Runtime/Core/Public/Tests/EnsureScope.h
@@ -8,18 +8,21 @@ struct FEnsureHandlerArgs;
 ///@brief Scope that captures failed `ensure` calls
 struct FEnsureScope
 {
-	FEnsureScope();
+	CORE_API FEnsureScope();
 
 	///@brief Only captures failed `ensure` calls with the exact message 
 	///@param Msg string to compare against failed `ensure` message
-	FEnsureScope(const ANSICHAR* Msg);
+	CORE_API explicit FEnsureScope(const ANSICHAR* Msg);
 
 	///@brief Captures failed `ensure` calls that return true from callback
 	///@param EnsureFunc lambda called for each failed `ensure` return true to handle the error
-	FEnsureScope(TFunction<bool(const FEnsureHandlerArgs&)> EnsureFunc);
-	~FEnsureScope();
+	CORE_API explicit FEnsureScope(TFunction<bool(const FEnsureHandlerArgs&)> EnsureFunc);
+	CORE_API ~FEnsureScope();
 
-	int GetCount() { return Count; }
+	int GetCount()
+	{
+		return Count;
+	}
 
 private:
 	TFunction<bool(const FEnsureHandlerArgs&)> OldHandler;
