Fixed missing destructor call in TSparseArray::EmplaceAt when element at Index is already valid.

#tests replayrun, PIE, manual verification

diff --git a/Engine/Source/Runtime/Core/Public/Containers/SparseArray.h b/Engine/Source/Runtime/Core/Public/Containers/SparseArray.h
--- a/Engine/Source/Runtime/Core/Public/Containers/SparseArray.h
+++ b/Engine/Source/Runtime/Core/Public/Containers/SparseArray.h
@@ -625,8 +625,11 @@ public:
 		}
 		else
 		{
+			FElementOrFreeListLink* Elem = (FElementOrFreeListLink*)Data.GetData() + Index;
+			((ElementType&)Elem->ElementData).~ElementType();
+
 			Allocation.Index = Index;
-			Allocation.Pointer = &((FElementOrFreeListLink*)Data.GetData())[Allocation.Index].ElementData;
+			Allocation.Pointer = &Elem->ElementData;
 		}
 
 		new(Allocation) ElementType(Forward<ArgsType>(Args)...);
