Delegates STATIC_FUNCTION_FNAME binary size optimizations

Names are C++ identifiers so always ASCII, no point in using TEXT
Make string initializers static const so they don't get initialized on the stack with code

diff --git a/Engine/Source/Programs/AutoRTFMTests/Private/UECoreTests.cpp b/Engine/Source/Programs/AutoRTFMTests/Private/UECoreTests.cpp
--- a/Engine/Source/Programs/AutoRTFMTests/Private/UECoreTests.cpp
+++ b/Engine/Source/Programs/AutoRTFMTests/Private/UECoreTests.cpp
@@ -312,7 +312,7 @@ TEST_CASE("UECore.STATIC_FUNCTION_FNAME")
 	{
 		AutoRTFM::Testing::Abort([&]()
 		{
-			Name = STATIC_FUNCTION_FNAME(TEXT("WOWWEE"));
+			Name = STATIC_FUNCTION_FNAME("WOWWEE");
 			AutoRTFM::AbortTransaction();
 		});
 
@@ -323,7 +323,7 @@ TEST_CASE("UECore.STATIC_FUNCTION_FNAME")
 	{
 		AutoRTFM::Testing::Commit([&]()
 		{
-			Name = STATIC_FUNCTION_FNAME(TEXT("WOWWEE"));
+			Name = STATIC_FUNCTION_FNAME("WOWWEE");
 		});
 	}
 }
diff --git a/Engine/Source/Runtime/Core/Public/Delegates/Delegate.h b/Engine/Source/Runtime/Core/Public/Delegates/Delegate.h
--- a/Engine/Source/Runtime/Core/Public/Delegates/Delegate.h
+++ b/Engine/Source/Runtime/Core/Public/Delegates/Delegate.h
@@ -371,21 +371,15 @@ class DynamicMulticastDelegateClassName : public TBaseDynamicMulticastDelegate<F
 		{
 			static FName Get()
 			{
+				static const T Str[sizeof...(Chars) + 1] = { Chars..., 0 };
 				FName* Result = nullptr;
 				UE_AUTORTFM_OPEN
 				{
-					static FName StaticResult = Create();
+					static FName StaticResult{Str};
 					Result = &StaticResult;
 				};
 				return *Result;
 			}
-
-		private:
-			static FName Create()
-			{
-				const T Str[sizeof...(Chars) + 1] = { Chars..., 0 };
-				return Str;
-			}
 		};
 
 		/**
@@ -428,7 +422,7 @@ class DynamicMulticastDelegateClassName : public TBaseDynamicMulticastDelegate<F
  * @param	UserObject		UObject instance
  * @param	FuncName		Function pointer to member UFUNCTION, usually in form &UClassName::FunctionName
  */
-#define BindDynamic( UserObject, FuncName ) __Internal_BindDynamic( UserObject, FuncName, STATIC_FUNCTION_FNAME( TEXT( #FuncName ) ) )
+#define BindDynamic( UserObject, FuncName ) __Internal_BindDynamic( UserObject, FuncName, STATIC_FUNCTION_FNAME( #FuncName ) )
 
 /**
  * Helper macro to bind a UObject instance and a member UFUNCTION to a dynamic multi-cast delegate.
@@ -436,7 +430,7 @@ class DynamicMulticastDelegateClassName : public TBaseDynamicMulticastDelegate<F
  * @param	UserObject		UObject instance
  * @param	FuncName		Function pointer to member UFUNCTION, usually in form &UClassName::FunctionName
  */
-#define AddDynamic( UserObject, FuncName ) __Internal_AddDynamic( UserObject, FuncName, STATIC_FUNCTION_FNAME( TEXT( #FuncName ) ) )
+#define AddDynamic( UserObject, FuncName ) __Internal_AddDynamic( UserObject, FuncName, STATIC_FUNCTION_FNAME( #FuncName ) )
 
 /**
  * Helper macro to bind a UObject instance and a member UFUNCTION to a dynamic multi-cast delegate, but only if it hasn't been bound before.
@@ -444,7 +438,7 @@ class DynamicMulticastDelegateClassName : public TBaseDynamicMulticastDelegate<F
  * @param	UserObject		UObject instance
  * @param	FuncName		Function pointer to member UFUNCTION, usually in form &UClassName::FunctionName
  */
-#define AddUniqueDynamic( UserObject, FuncName ) __Internal_AddUniqueDynamic( UserObject, FuncName, STATIC_FUNCTION_FNAME( TEXT( #FuncName ) ) )
+#define AddUniqueDynamic( UserObject, FuncName ) __Internal_AddUniqueDynamic( UserObject, FuncName, STATIC_FUNCTION_FNAME( #FuncName ) )
 
 /**
  * Helper macro to unbind a UObject instance and a member UFUNCTION from this multi-cast delegate.
@@ -452,7 +446,7 @@ class DynamicMulticastDelegateClassName : public TBaseDynamicMulticastDelegate<F
  * @param	UserObject		UObject instance
  * @param	FuncName		Function pointer to member UFUNCTION, usually in form &UClassName::FunctionName
  */
-#define RemoveDynamic( UserObject, FuncName ) __Internal_RemoveDynamic( UserObject, FuncName, STATIC_FUNCTION_FNAME( TEXT( #FuncName ) ) )
+#define RemoveDynamic( UserObject, FuncName ) __Internal_RemoveDynamic( UserObject, FuncName, STATIC_FUNCTION_FNAME( #FuncName ) )
 
  /**
   * Helper macro to tests if a UObject instance and a member UFUNCTION are already bound to this multi-cast delegate.
@@ -460,7 +454,7 @@ class DynamicMulticastDelegateClassName : public TBaseDynamicMulticastDelegate<F
   * @param	UserObject		UObject instance
   * @param	FuncName		Function pointer to member UFUNCTION, usually in form &UClassName::FunctionName
   */
-#define IsAlreadyBound( UserObject, FuncName ) __Internal_IsAlreadyBound( UserObject, FuncName, STATIC_FUNCTION_FNAME( TEXT( #FuncName ) ) )
+#define IsAlreadyBound( UserObject, FuncName ) __Internal_IsAlreadyBound( UserObject, FuncName, STATIC_FUNCTION_FNAME( #FuncName ) )
 
 
 namespace UE::Delegates::Private
@@ -472,12 +466,12 @@ namespace UE::Delegates::Private
 	 * @param  InMacroFunctionName  The string containing the member function name.
 	 * @return An FName of the member function name.
 	 */
-	inline FName GetTrimmedMemberFunctionName(const TCHAR* InMacroFunctionName)
+	inline FName GetTrimmedMemberFunctionName(const ANSICHAR* InMacroFunctionName)
 	{
 		// We strip off the class prefix and just return the function name by itself.
 		check(InMacroFunctionName);
-		const TCHAR* Result = FCString::Strrstr( InMacroFunctionName, TEXT( "::" ) );
-		checkf(Result && Result[2] != (TCHAR)'0', TEXT("'%s' does not look like a member function"), InMacroFunctionName);
+		const ANSICHAR* Result = FCStringAnsi::Strrstr( InMacroFunctionName,  "::" );
+		checkf(Result && Result[2] != '0', TEXT("'%hs' does not look like a member function"), InMacroFunctionName);
 		return FName(Result + 2);
 	}
 }
