Performance: Fix external profilers such as PIX or Superluminal  failing to show event markers.

Change 39988953 added support for specifying an external profiler via DefaultEngine.ini and optimized the performance of external events when an external profiler was activated but not currently capturing. This introduced a bug as external profilers with external tool initiated captures were no longer receiving events.

This change fixes this by:
- Introducing an interface to explicitly enable or disable sending events.
- Events default to enabled, restoring the previous default.
- Updating individual external profilers implementations to request event generation appropriate for their scoped event implementation and methods of initiating a capture.

diff --git a/Engine/Source/Runtime/Core/Private/ProfilingDebugging/ExternalProfiler.cpp b/Engine/Source/Runtime/Core/Private/ProfilingDebugging/ExternalProfiler.cpp
--- a/Engine/Source/Runtime/Core/Private/ProfilingDebugging/ExternalProfiler.cpp
+++ b/Engine/Source/Runtime/Core/Private/ProfilingDebugging/ExternalProfiler.cpp
@@ -48,6 +48,8 @@ FExternalProfiler* FActiveExternalProfilerBase::ActiveProfiler = nullptr;
 
 std::atomic<bool> FActiveExternalProfilerBase::bIsRecording = false;
 
+std::atomic<bool> FActiveExternalProfilerBase::bEnableScopedEvents = true;
+
 static FExternalProfiler* FindProfilerByName(TArrayView<FExternalProfiler*> Profilers, FStringView Name)
 {
 	FExternalProfiler** Found = Algo::FindByPredicate(
@@ -144,6 +146,29 @@ bool FActiveExternalProfilerBase::IsActiveProfilerRecording()
 	return bIsRecording.load(std::memory_order_relaxed);
 }
 
+void FActiveExternalProfilerBase::EnableScopedEvents(bool bEnable)
+{
+	bool bPreviousEnable = bEnableScopedEvents.exchange(bEnable, std::memory_order_relaxed);
+
+	if (bPreviousEnable == bEnable)
+	{
+		return;
+	}
+
+	FExternalProfiler* Profiler = GetActiveProfiler();
+	if (!Profiler)
+	{
+		return;
+	}
+	
+	Profiler->OnEnableScopedEventsChanged(bEnable);
+}
+
+bool FActiveExternalProfilerBase::AreScopedEventsEnabled()
+{
+	return bEnableScopedEvents.load(std::memory_order_relaxed);
+}
+
 void FScopedExternalProfilerBase::StartScopedTimer( const bool bWantPause )
 {
 	bWasRecording = FActiveExternalProfilerBase::SetActiveProfilerRecording(!bWantPause);
diff --git a/Engine/Source/Runtime/Core/Public/ProfilingDebugging/ExternalProfiler.h b/Engine/Source/Runtime/Core/Public/ProfilingDebugging/ExternalProfiler.h
--- a/Engine/Source/Runtime/Core/Public/ProfilingDebugging/ExternalProfiler.h
+++ b/Engine/Source/Runtime/Core/Public/ProfilingDebugging/ExternalProfiler.h
@@ -62,6 +62,9 @@ public:
 	/** @return Returns the name to use for any profiler registered as a modular feature usable by this system */
 	static CORE_API FName GetFeatureName();
 
+	/** Notification that scoped event transmission status changed. */
+	virtual void OnEnableScopedEventsChanged(bool bEnabled) {};
+
 	/** Starts a scoped event specific to the profiler. */
 	virtual void StartScopedEvent(const struct FColor& Color, const TCHAR* Text) {};
 
@@ -110,6 +113,35 @@ public:
 	static CORE_API bool SetActiveProfilerRecording(bool bRecording);
 	static CORE_API bool IsActiveProfilerRecording();
 
+	/**
+	 * Static: Enables or disables calls to the active profiler's StartScopedEvent and StopScopedEvent functions.
+	 *
+	 * This function calls OnEnableScopedEventsChanged on the active external profiler if the enabled status
+	 * has changed. 
+	 *  
+	 * External profilers may support initiating captures either from an external tool or via code.
+	 * The performance cost of emitting scoped events varies by profiler. To minimize overhead, including 
+	 * calls to stubbed virtual functions, EnableScopedEvent can be used by a profiler implementation to
+	 * control event emission.
+	 * 
+	 * Example scenarios:
+	 * 
+	 * 1. A profiler that doesn't implement scoped events can disable scoped events to avoid the overhead
+	 *    of calling stubbed virtual functions.
+	 * 2. A profiler that implements scoped events and initiates captures externally can leverage 
+	 *    the default behavior of emitting events.
+	 * 3. A profiler with high overhead scoped events may disable scoped events to avoid overhead and 
+	 *    require explicit user action to enable them.
+	 * 4. A profiler that uses code-initiated captures may enable scoped events when capture starts
+	 *    and disables events when the capture ends.
+	 */
+	static CORE_API void EnableScopedEvents(bool bEnable);
+
+	/**
+	 * Returns true if scoped events are enabled by the active profiler.
+	 */
+	static CORE_API bool AreScopedEventsEnabled();
+
 private:
 	/** Static: True if we've tried to initialize a profiler already */
 	static CORE_API bool bDidInitialize;
@@ -119,6 +151,9 @@ private:
 
 	/** Static: Whether the active profiler is recording events, to our knowledge. */
 	static std::atomic<bool> bIsRecording;
+
+	/** Static: Whether StartScopedEvent and StartScopedEvent should be called. Defaults to true. */
+	static std::atomic<bool> bEnableScopedEvents;
 };
 
 /**
@@ -208,7 +243,7 @@ public:
 	/** Starts a scoped event specific to the profiler. */
 	inline static void StartScopedEvent(const struct FColor& Color, const TCHAR* Text)
 	{
-		if (FActiveExternalProfilerBase::IsActiveProfilerRecording())
+		if (FActiveExternalProfilerBase::AreScopedEventsEnabled())
 		{
 			if (FExternalProfiler* Profiler = FActiveExternalProfilerBase::GetActiveProfiler())
 			{
@@ -220,7 +255,7 @@ public:
 	/** Starts a scoped event specific to the profiler. */
 	inline static void StartScopedEvent(const struct FColor& Color, const ANSICHAR* Text)
 	{
-		if (FActiveExternalProfilerBase::IsActiveProfilerRecording())
+		if (FActiveExternalProfilerBase::AreScopedEventsEnabled())
 		{
 			if (FExternalProfiler* Profiler = FActiveExternalProfilerBase::GetActiveProfiler())
 			{
@@ -232,7 +267,7 @@ public:
 	/** Ends a scoped event specific to the profiler. */
 	inline static void EndScopedEvent()
 	{
-		if (FActiveExternalProfilerBase::IsActiveProfilerRecording())
+		if (FActiveExternalProfilerBase::AreScopedEventsEnabled())
 		{
 			if (FExternalProfiler* Profiler = FActiveExternalProfilerBase::GetActiveProfiler())
 			{
