Core: FileOpenLog: Added a missing implementation of the UsesCache() member.

Without this override, FLoggingAsyncReadFileHandle declared itself as being backed by cache regardless of whether the underlying filesystem supported it or not.

This change is being submitted mostly for consistency. Currently, we only have one platform which returns 'false' from UsesCache, and that's about to change anyway. So, there should be no change in behavior, especially given that FLoggingAsyncReadFileHandle is a special-purpose implementation used only during development, and the fact that UsesCache() is only called by the legacy async streamer.

diff --git a/Engine/Source/Runtime/Core/Public/HAL/IPlatformFileOpenLogWrapper.h b/Engine/Source/Runtime/Core/Public/HAL/IPlatformFileOpenLogWrapper.h
--- a/Engine/Source/Runtime/Core/Public/HAL/IPlatformFileOpenLogWrapper.h
+++ b/Engine/Source/Runtime/Core/Public/HAL/IPlatformFileOpenLogWrapper.h
@@ -48,6 +48,11 @@ public:
 		return ActualRequest->SizeRequest(CompleteCallback);
 	}
 	virtual IAsyncReadRequest* ReadRequest(int64 Offset, int64 BytesToRead, EAsyncIOPriorityAndFlags PriorityAndFlags = AIOP_Normal, FAsyncFileCallBack* CompleteCallback = nullptr, uint8* UserSuppliedMemory = nullptr) override;
+	
+	virtual bool UsesCache() override
+	{
+		return ActualRequest->UsesCache();
+	}
 };
 
 
