Fix errors in comments in the FActiveExternalProfilerBase API.

diff --git a/Engine/Source/Runtime/Core/Public/ProfilingDebugging/ExternalProfiler.h b/Engine/Source/Runtime/Core/Public/ProfilingDebugging/ExternalProfiler.h
--- a/Engine/Source/Runtime/Core/Public/ProfilingDebugging/ExternalProfiler.h
+++ b/Engine/Source/Runtime/Core/Public/ProfilingDebugging/ExternalProfiler.h
@@ -114,14 +114,14 @@ public:
 	static CORE_API bool IsActiveProfilerRecording();
 
 	/**
-	 * Static: Enables or disables calls to the active profiler's StartScopedEvent and StopScopedEvent functions.
+	 * Static: Enables or disables calls to the active profiler's StartScopedEvent and EndScopedEvent functions.
 	 *
 	 * This function calls OnEnableScopedEventsChanged on the active external profiler if the enabled status
 	 * has changed. 
 	 *  
 	 * External profilers may support initiating captures either from an external tool or via code.
 	 * The performance cost of emitting scoped events varies by profiler. To minimize overhead, including 
-	 * calls to stubbed virtual functions, EnableScopedEvent can be used by a profiler implementation to
+	 * calls to stubbed virtual functions, EnableScopedEvents can be used by a profiler implementation to
 	 * control event emission.
 	 * 
 	 * Example scenarios:
@@ -152,7 +152,7 @@ private:
 	/** Static: Whether the active profiler is recording events, to our knowledge. */
 	static std::atomic<bool> bIsRecording;
 
-	/** Static: Whether StartScopedEvent and StartScopedEvent should be called. Defaults to true. */
+	/** Static: Whether StartScopedEvent and EndScopedEvent should be called. Defaults to true. */
 	static std::atomic<bool> bEnableScopedEvents;
 };
 
