ModuleManager static registration code size refactor

Make registration function static and handle the singleton resolve and creating the delegate;
make module instance creation funcs static to facilitate this.

diff --git a/Engine/Source/Runtime/Core/Private/Modules/ModuleManager.cpp b/Engine/Source/Runtime/Core/Private/Modules/ModuleManager.cpp
--- a/Engine/Source/Runtime/Core/Private/Modules/ModuleManager.cpp
+++ b/Engine/Source/Runtime/Core/Private/Modules/ModuleManager.cpp
@@ -1979,6 +1979,11 @@ const TCHAR *FModuleManager::GetUBTConfiguration()
 	return LexToString(FApp::GetBuildConfiguration());
 }
 
+void FModuleManager::RegisterStaticallyLinkedModule(const FLazyName InModuleName, FInitializeStaticallyLinkedModuleRaw* InInitializer)
+{
+	Get().PendingStaticallyLinkedModuleInitializers.Emplace(InModuleName, FInitializeStaticallyLinkedModule::CreateStatic(InInitializer));
+}
+
 void FModuleManager::StartProcessingNewlyLoadedObjects()
 {
 	// Only supposed to be called once
diff --git a/Engine/Source/Runtime/Core/Public/Modules/ModuleManager.h b/Engine/Source/Runtime/Core/Public/Modules/ModuleManager.h
--- a/Engine/Source/Runtime/Core/Public/Modules/ModuleManager.h
+++ b/Engine/Source/Runtime/Core/Public/Modules/ModuleManager.h
@@ -431,16 +431,16 @@ public:
 	/** Delegate that's used by the module manager to initialize a registered module that we statically linked with (monolithic only) */
 	DECLARE_DELEGATE_RetVal( IModuleInterface*, FInitializeStaticallyLinkedModule )
 
+	/** Raw function type used during module registration. */
+	typedef IModuleInterface *FInitializeStaticallyLinkedModuleRaw();
+
 	/**
 	 * Registers an initializer for a module that is statically linked.
 	 *
 	 * @param InModuleName The name of this module.
 	 * @param InInitializerDelegate The delegate that will be called to initialize an instance of this module.
 	 */
-	void RegisterStaticallyLinkedModule( const FLazyName InModuleName, const FInitializeStaticallyLinkedModule& InInitializerDelegate )
-	{
-		PendingStaticallyLinkedModuleInitializers.Emplace( InModuleName, InInitializerDelegate );
-	}
+	static CORE_API void RegisterStaticallyLinkedModule( const FLazyName InModuleName, FInitializeStaticallyLinkedModuleRaw* InInitializer );
 
 	/**
 	 * Called by the engine at startup to let the Module Manager know that it's now
@@ -803,14 +803,9 @@ public:
 	 */
 	FStaticallyLinkedModuleRegistrant( FLazyName InModuleName )
 	{
-		// Create a delegate to our InitializeModule method
-		FModuleManager::FInitializeStaticallyLinkedModule InitializerDelegate = FModuleManager::FInitializeStaticallyLinkedModule::CreateRaw(
-				this, &FStaticallyLinkedModuleRegistrant<ModuleClass>::InitializeModule );
-
-		// Register this module
-		FModuleManager::Get().RegisterStaticallyLinkedModule(
+		FModuleManager::RegisterStaticallyLinkedModule(
 			InModuleName,			// Module name
-			InitializerDelegate );	// Initializer delegate
+			&InitializeModule );	// Initializer
 	}
 	
 	/**
@@ -821,7 +816,7 @@ public:
 	 *
 	 * @return A pointer to a new instance of the module.
 	 */
-	IModuleInterface* InitializeModule( )
+	static IModuleInterface* InitializeModule( )
 	{
 		return new ModuleClass();
 	}
