Fix issues regarding tag editing: config file update, source control integration, tag filter additive with other filters
Edit test tags button for functional tests, opens automation window, displays deprecation message
#rnx

diff --git a/Engine/Build/InstalledEngineFilters.xml b/Engine/Build/InstalledEngineFilters.xml
--- a/Engine/Build/InstalledEngineFilters.xml
+++ b/Engine/Build/InstalledEngineFilters.xml
@@ -260,6 +260,10 @@
 		Engine/.../WinArm64/...
 		Engine/.../Win64/arm64/...
 		Engine/.../overlay-arm64-uwp/...
+
+		<!-- Exclude Test Configs -->
+		Engine/Config/Tests/...
+		Engine/.../Config/Tests/...
 	</Property>
 
 	<!-- Filters for installed DDC -->
diff --git a/Engine/Source/Developer/AutomationController/Public/AutomationControllerSettings.h b/Engine/Source/Developer/AutomationController/Public/AutomationControllerSettings.h
--- a/Engine/Source/Developer/AutomationController/Public/AutomationControllerSettings.h
+++ b/Engine/Source/Developer/AutomationController/Public/AutomationControllerSettings.h
@@ -270,6 +270,12 @@ public:
 	bool bResetTelemetryStorageOnNewSession;
 
 	UE_API virtual void PostInitProperties() override;
+
+	/** List of pre-selected tests to filter by when requesting tests */
+	TArray<FString> InitialSelectedTestsNames;
+
+	/** Prepare to open test tag editor when a list of InitialSelectedTestsNames is provided */
+	bool bOpenTagEditorWhenInitialSelectedTestsNames;
 };
 
 #undef UE_API
diff --git a/Engine/Source/Developer/AutomationWindow/Private/SAutomationTestTagEditor.h b/Engine/Source/Developer/AutomationWindow/Private/SAutomationTestTagEditor.h
--- a/Engine/Source/Developer/AutomationWindow/Private/SAutomationTestTagEditor.h
+++ b/Engine/Source/Developer/AutomationWindow/Private/SAutomationTestTagEditor.h
@@ -3,6 +3,9 @@
 #pragma once
 
 #include "CoreMinimal.h"
+#include "HAL/FileManager.h"
+#include "ISourceControlProvider.h"
+#include "ISourceControlModule.h"
 #include "Misc/CoreDelegates.h"
 #include "Misc/MessageDialog.h"
 #include "Misc/PackageName.h"
@@ -348,7 +351,7 @@ private:
 			{
 				Changes = Changes.Append(TEXT("Removing tags: ")).Append(FString::Join(MarkedForDeleteTags, TEXT(","))).Append(TEXT("."));
 			}
-			FText ChangesText = FText::FromString(*Changes);
+			FText ChangesText = FText::FromString(Changes);
 			if (EAppReturnType::Yes == FMessageDialog::Open(EAppMsgType::YesNo, FText::Format(LOCTEXT("TagsChangedSaveMsg", "Confirm tag changes to apply.\n{0}"), ChangesText)))
 			{
 				TArray<FString> Tags;
@@ -387,11 +390,25 @@ private:
 							}), false);
 					check(Registered);
 				}
-				TArray<FString> TestTagMapFiles = FAutomationTestFramework::Get().SaveTestTagMappings(TestNames, TestFilePaths);
-				if (SourceControlHelpers::IsAvailable())
-				{
-					SourceControlHelpers::CheckOutOrAddFiles(TestTagMapFiles);
-				}
+
+				FAutomationTestFramework::Get().SaveTestTagMappings(TestNames, TestFilePaths, FBeforeTagMappingConfigSaved::CreateLambda(
+						[](const FString& ConfigPath)
+						{
+							if (SourceControlHelpers::IsAvailable())
+							{
+								if (IFileManager::Get().FileExists(*ConfigPath))
+								{
+									SourceControlHelpers::SyncFile(ConfigPath);
+								}
+							}
+						}),
+						FAfterTagMappingConfigSaved::CreateLambda([](const FString& ConfigPath)
+						{
+							if (SourceControlHelpers::IsAvailable())
+							{
+								SourceControlHelpers::CheckOutOrAddFile(ConfigPath);
+							}
+						}));
 
 				OnTestTagsSavedEvent.ExecuteIfBound();
 
diff --git a/Engine/Source/Developer/AutomationWindow/Private/SAutomationTestTagFilter.h b/Engine/Source/Developer/AutomationWindow/Private/SAutomationTestTagFilter.h
--- a/Engine/Source/Developer/AutomationWindow/Private/SAutomationTestTagFilter.h
+++ b/Engine/Source/Developer/AutomationWindow/Private/SAutomationTestTagFilter.h
@@ -47,8 +47,9 @@ public:
 	 *
 	 * @param InArgs The declaration data for this widget.
 	 */
-	void Construct(const FArguments& InArgs)
+	void Construct(const FArguments& InArgs, TSharedPtr<FAutomationGroupFilter> InAutomationTagFilter)
 	{
+		AutomationTagFilter = InAutomationTagFilter;
 		ChildSlot
 			[
 				SNew(SBox)
@@ -136,40 +137,9 @@ public:
 
 	void ApplyTagFilter(const FString& TagFilter)
 	{
-		TSharedPtr <AutomationFilterCollection> AutomationFilters = MakeShareable(new AutomationFilterCollection());
-		FAutomationGroupFilter* FilterTags = new FAutomationGroupFilter();
 		TArray<FAutomatedTestTagFilter> TagList;
 		TagList.Add(FAutomatedTestTagFilter(TagFilter));
-		FilterTags->SetTagFilter(TagList);
-		AutomationFilters->Add(MakeShareable(FilterTags));
-
-		FModuleManager::GetModuleChecked<IAutomationControllerModule>("AutomationController").GetAutomationController()->SetFilter(AutomationFilters);
-	}
-
-	static void ShowInWindow(const TSharedRef<const SWidget> InParentWindow, const FOnTestFilterAppliedEvent& OnTestFilterAppliedEvent)
-	{
-		TSharedPtr<SAutomationTestTagFilter> TestTagFilter;
-		TSharedRef<SWidget> TagsFilterForm = SNew(SBox)
-			.WidthOverride(640)
-			[
-				SAssignNew(TestTagFilter, SAutomationTestTagFilter)
-			];
-
-		TSharedRef<SWindow> Window = SNew(SWindow)
-			.Title(LOCTEXT("TestTagsFilter", "Filter Tags"))
-			.SizingRule(ESizingRule::Autosized)
-			.SupportsMinimize(false)
-			.SupportsMaximize(false)
-			.AutoCenter(EAutoCenter::PrimaryWorkArea);
-
-		Window->SetContent
-		(
-			TagsFilterForm
-		);
-
-		TestTagFilter->SetOnTestFilterAppliedEvent(OnTestFilterAppliedEvent);
-
-		FSlateApplication::Get().AddModalWindow(Window, InParentWindow);
+		AutomationTagFilter->SetTagFilter(TagList);
 	}
 
 	void SetOnTestFilterAppliedEvent(const FOnTestFilterAppliedEvent& OnTestFilterAppliedEvent)
@@ -186,6 +156,7 @@ public:
 	}
 
 protected:
+	TSharedPtr<FAutomationGroupFilter> AutomationTagFilter;
 	FOnTestFilterAppliedEvent OnTestFilterApplied;
 	TSharedPtr<SEditableTextBox> FilterText;
 	TSharedPtr<SLayeredImage> FilterButtonImage;
diff --git a/Engine/Source/Developer/AutomationWindow/Private/SAutomationWindow.cpp b/Engine/Source/Developer/AutomationWindow/Private/SAutomationWindow.cpp
--- a/Engine/Source/Developer/AutomationWindow/Private/SAutomationWindow.cpp
+++ b/Engine/Source/Developer/AutomationWindow/Private/SAutomationWindow.cpp
@@ -201,10 +201,12 @@ void SAutomationWindow::Construct( const FArguments& InArgs, const IAutomationCo
 	AutomationTextFilter = MakeShareable( new AutomationReportTextFilter( AutomationReportTextFilter::FItemToStringArray::CreateSP( this, &SAutomationWindow::PopulateReportSearchStrings ) ) );
 	AutomationGeneralFilter = MakeShareable( new FAutomationFilter() ); 
 	AutomationGroupFilter = MakeShareable( new FAutomationGroupFilter() );
+	AutomationTagFilter = MakeShareable(new FAutomationGroupFilter());
 	AutomationFilters = MakeShareable( new AutomationFilterCollection() );
 	AutomationFilters->Add( AutomationTextFilter );
 	AutomationFilters->Add( AutomationGeneralFilter );
 	AutomationFilters->Add( AutomationGroupFilter );
+	AutomationFilters->Add( AutomationTagFilter );
 
 	bIsRequestingTests = false;
 
@@ -688,6 +690,24 @@ void SAutomationWindow::Construct( const FArguments& InArgs, const IAutomationCo
 }
 END_SLATE_FUNCTION_BUILD_OPTIMIZATION
 
+void SAutomationWindow::Tick(const FGeometry& AllottedGeometry, const double InCurrentTime, const float InDeltaTime)
+{
+	SCompoundWidget::Tick(AllottedGeometry, InCurrentTime, InDeltaTime);
+
+	UAutomationControllerSettings* Settings = UAutomationControllerSettings::StaticClass()->GetDefaultObject<UAutomationControllerSettings>();
+	if (AutomationController->GetFilteredReports().Num() > 0 && Settings && Settings->bOpenTagEditorWhenInitialSelectedTestsNames)
+	{
+		Settings->bOpenTagEditorWhenInitialSelectedTestsNames = false;
+
+		AutomationController->SetEnabledTests(Settings->InitialSelectedTestsNames);
+
+		SAutomationTestTagEditor::ShowInWindow(this->AsShared(), AutomationController->GetEnabledReports(), FOnTestTagsSavedEvent());
+
+		AutomationController->SetEnabledTests(TArray<FString>());
+		Settings->InitialSelectedTestsNames.Empty();
+	}
+}
+
 
 void SAutomationWindow::HandleResultDisplayTypeStateChanged( ECheckBoxState NewRadioState, EAutomationGrapicalDisplayType::Type NewDisplayType)
 {
@@ -1008,13 +1028,10 @@ TSharedRef< SWidget > SAutomationWindow::MakeAutomationWindowToolBar( const TSha
 			]
 		];
 
-	TagFilter = SNew(SAutomationTestTagFilter);
+	TagFilter = SNew(SAutomationTestTagFilter, AutomationTagFilter);
 	TagFilter->SetOnTestFilterAppliedEvent(FOnTestFilterAppliedEvent::CreateLambda([this]()
 		{
-			AutomationController->SetEnabledTests(TArray<FString>());
-			TestTable->RequestTreeRefresh();
-			TestTable->ClearExpandedItems();
-			TestTable->RebuildList();
+			OnRefreshTestCallback();
 		}));
 
 	FToolBarBuilder ToolbarBuilder( InCommandList, FMultiBoxCustomization::None );
diff --git a/Engine/Source/Developer/AutomationWindow/Private/SAutomationWindow.h b/Engine/Source/Developer/AutomationWindow/Private/SAutomationWindow.h
--- a/Engine/Source/Developer/AutomationWindow/Private/SAutomationWindow.h
+++ b/Engine/Source/Developer/AutomationWindow/Private/SAutomationWindow.h
@@ -120,6 +120,9 @@ public:
 	// Expanded the given item if its name is in the array of strings given.
 	static void ExpandItemsInList(TSharedPtr<STreeView<TSharedPtr<IAutomationReport>>> InTestTable, TSharedPtr<IAutomationReport> InReport, TArray<FString> ItemsToExpand);
 
+	/** Override base class Tick */
+	virtual void Tick(const FGeometry& AllottedGeometry, const double InCurrentTime, const float InDeltaTime) override;
+
 protected:
 
 	/**
@@ -627,6 +630,9 @@ private:
 	/** The automation group filter - filters tests based on selected group filter. */
 	TSharedPtr< FAutomationGroupFilter > AutomationGroupFilter;
 
+	/** The automation tag filter - filters tests based on tag expression. */
+	TSharedPtr< FAutomationGroupFilter > AutomationTagFilter;
+
 	/** The automation filter collection - contains the automation filters. */
 	TSharedPtr< AutomationFilterCollection > AutomationFilters;
 
diff --git a/Engine/Source/Developer/FunctionalTesting/Classes/FunctionalTest.h b/Engine/Source/Developer/FunctionalTesting/Classes/FunctionalTest.h
--- a/Engine/Source/Developer/FunctionalTesting/Classes/FunctionalTest.h
+++ b/Engine/Source/Developer/FunctionalTesting/Classes/FunctionalTest.h
@@ -273,14 +273,28 @@ public:
 	/**
 	 * Tags describing this test separated by square brackets, such as '[dog]' or '[cat]' or '[Graphics][prio0][unstable]'.
 	 * Tags can be used to run subsets of tests, or to categorize data in test reports.
+	 * Deprecating, please use "Edit Test Tags" instead.
 	 */
 	UPROPERTY(EditAnywhere, BlueprintReadOnly, Category = "Functional Testing")
 	FString TestTags;
 
+#if WITH_EDITOR
+	/**
+	 * Edit tags button that will open the Automation tab in the Session Frontend window then show the tag editor for this test.  
+	 */
+	UFUNCTION(CallInEditor, Category="Functional Testing", DisplayName = "Edit Test Tags", meta=(DisplayAfter="TestTags"))
+	UE_API void EditTags();
+#endif
+
+	/** Return functional test beautified name */
+	FString GetBeautifiedName();
+
 private:
 	UPROPERTY()
 	TObjectPtr<UBillboardComponent> SpriteComponent;
 
+	FString PreviousTestTags;
+
 protected:
 	/**
 	 * Allows a test to be disabled.  If a test is disabled, it will not appear in the set of
@@ -727,6 +741,7 @@ public:
 	virtual FString GetReproString() const { return GetFName().ToString(); }
 
 #if WITH_EDITOR
+	UE_API virtual void PostLoad() override;
 	UE_API void PostEditChangeProperty(struct FPropertyChangedEvent& PropertyChangedEvent) override;
 	UE_API virtual void GetAssetRegistryTags(FAssetRegistryTagsContext Context) const override;
 	UE_DEPRECATED(5.4, "Implement the version that takes FAssetRegistryTagsContext instead.")
diff --git a/Engine/Source/Developer/FunctionalTesting/FunctionalTesting.Build.cs b/Engine/Source/Developer/FunctionalTesting/FunctionalTesting.Build.cs
--- a/Engine/Source/Developer/FunctionalTesting/FunctionalTesting.Build.cs
+++ b/Engine/Source/Developer/FunctionalTesting/FunctionalTesting.Build.cs
@@ -33,7 +33,8 @@ public class FunctionalTesting : ModuleRules
 					"SourceControl",
 					"EditorFramework",
 					"UnrealEd",
-					"LevelEditor"
+					"LevelEditor",
+					"SessionFrontend"
 				}
 			);
 
diff --git a/Engine/Source/Developer/FunctionalTesting/Private/FunctionalTest.cpp b/Engine/Source/Developer/FunctionalTesting/Private/FunctionalTest.cpp
--- a/Engine/Source/Developer/FunctionalTesting/Private/FunctionalTest.cpp
+++ b/Engine/Source/Developer/FunctionalTesting/Private/FunctionalTest.cpp
@@ -33,6 +33,12 @@
 #include "FunctionalTestBase.h"
 #include "UnrealClient.h"
 
+#if WITH_EDITOR
+#include "AutomationControllerSettings.h"
+#include "ISessionFrontendModule.h"
+#include "Misc/MessageDialog.h"
+#endif
+
 #include UE_INLINE_GENERATED_CPP_BY_NAME(FunctionalTest)
 
 DECLARE_CYCLE_STAT(TEXT("FunctionalTest - RunTest"), STAT_FunctionalTest_RunTest, STATGROUP_FunctionalTest);
@@ -240,6 +246,22 @@ AFunctionalTest::AFunctionalTest( const FObjectInitializer& ObjectInitializer )
 #endif
 }
 
+#if WITH_EDITOR
+void AFunctionalTest::EditTags()
+{
+	UAutomationControllerSettings* Settings = UAutomationControllerSettings::StaticClass()->GetDefaultObject<UAutomationControllerSettings>();
+	if (Settings)
+	{
+		Settings->InitialSelectedTestsNames.Empty();
+		Settings->InitialSelectedTestsNames.Add(GetBeautifiedName());
+		Settings->bOpenTagEditorWhenInitialSelectedTestsNames = true;
+	}
+
+	ISessionFrontendModule& SessionFrontend = FModuleManager::LoadModuleChecked<ISessionFrontendModule>("SessionFrontend");
+	SessionFrontend.InvokeSessionFrontend(FName("AutomationPanel"));
+}
+#endif
+
 void AFunctionalTest::OnConstruction(const FTransform& Transform)
 {
 	Super::OnConstruction(Transform);
@@ -631,6 +653,15 @@ void AFunctionalTest::RegisterAutoDestroyActor(AActor* ActorToAutoDestroy)
 
 #if WITH_EDITOR
 
+void AFunctionalTest::PostLoad()
+{
+	Super::PostLoad();
+
+	SetFlags(RF_Transactional);
+
+	PreviousTestTags = TestTags;
+}
+
 void AFunctionalTest::PostEditChangeProperty( struct FPropertyChangedEvent& PropertyChangedEvent)
 {
 	static const FName NAME_FunctionalTesting = FName(TEXT("Functional Testing"));
@@ -659,30 +690,55 @@ void AFunctionalTest::PostEditChangeProperty( struct FPropertyChangedEvent& Prop
 					TimesUpResult = EFunctionalTestResult::Failed;
 				}
 			}
-			else if (PropertyChangedEvent.Property->GetFName() == NAME_TestTags)
+			else if (PropertyChangedEvent.Property->GetFName() == NAME_TestTags && !TestTags.IsEmpty())
 			{
-				if (const ULevel* Level = GetLevel())
+				bool bContinueEdit = true;
+				// this property will be deprected
+
+				if (EAppReturnType::Yes == FMessageDialog::Open(EAppMsgType::YesNo, FText::FromString(TEXT("Editing asset tags are being deprecated and \"Edit Test Tags\" should be used instead. Continue?"))))
 				{
-					if (const UPackage* Package = Level->GetPackage())
-					{
-						// re-register existing tags for the loaded test
-						FSoftObjectPath CurrentMapPath = FSoftObjectPath(this->GetLevel()->GetPackage());
-						FAssetData MapAsset = IAssetRegistry::Get()->GetAssetByObjectPath(CurrentMapPath);
-						const FString PartialSuiteName = MapPackageToAutomationPath(MapAsset);
-						const FString FullBeautifiedName(PartialSuiteName + TEXT(".") + *GetActorLabel());
-						const FString FullTestName(TEXT("Project.Functional Tests." + FullBeautifiedName));
+					PreviousTestTags = TestTags;
+				}
+				else
+				{
+					bContinueEdit = false;
+					TestTags = PreviousTestTags;
+				}
 
+				if (bContinueEdit)
+				{
+					const FString FullTestName = GetBeautifiedName();
+					if (!FullTestName.IsEmpty())
+					{
 						FAutomationTestFramework& TestFramework = FAutomationTestFramework::Get();
 						TestFramework.UnregisterAutomationTestTags(FullTestName);
 						TestFramework.RegisterAutomationTestTags(FullTestName, TestTags);
 					}
 				}
-				
 			}
 		}
 	}
 }
 
+FString AFunctionalTest::GetBeautifiedName()
+{
+	if (const ULevel* Level = GetLevel())
+	{
+		if (const UPackage* Package = Level->GetPackage())
+		{
+			FSoftObjectPath CurrentMapPath = FSoftObjectPath(this->GetLevel()->GetPackage());
+			FAssetData MapAsset = IAssetRegistry::Get()->GetAssetByObjectPath(CurrentMapPath);
+			const FString PartialSuiteName = MapPackageToAutomationPath(MapAsset);
+			const FString FullBeautifiedName(PartialSuiteName + TEXT(".") + *GetActorLabel());
+			const FString FullTestName(TEXT("Project.Functional Tests." + FullBeautifiedName));
+
+			return FullTestName;
+		}
+	}
+
+	return TEXT("");
+}
+
 void AFunctionalTest::GetAssetRegistryTags(TArray<FAssetRegistryTag>& OutTags) const
 {
 	PRAGMA_DISABLE_DEPRECATION_WARNINGS;
diff --git a/Engine/Source/Runtime/Core/Private/Misc/AutomationTest.cpp b/Engine/Source/Runtime/Core/Private/Misc/AutomationTest.cpp
--- a/Engine/Source/Runtime/Core/Private/Misc/AutomationTest.cpp
+++ b/Engine/Source/Runtime/Core/Private/Misc/AutomationTest.cpp
@@ -926,20 +926,23 @@ void FAutomationTestFramework::LoadTestTagMappings()
 
 	if (!bTagMappingFilesLoaded)
 	{
-		LoadTestTagMappings(FPaths::ConvertRelativePathToFull(FPaths::ProjectDir()));
-		LoadTestTagMappings(FPaths::ConvertRelativePathToFull(FPaths::EngineDir()));
+		LoadTestTagMappings(FPaths::ConvertRelativePathToFull(FPaths::ProjectDir()), false);
+		LoadTestTagMappings(FPaths::ConvertRelativePathToFull(FPaths::EngineDir()), true);
 		
 		bTagMappingFilesLoaded = true;
 	}
 }
 
-void FAutomationTestFramework::LoadTestTagMappings(const FString& BasePath)
+void FAutomationTestFramework::LoadTestTagMappings(const FString& BasePath, bool RestrictedPathsOnly)
 {
 	TArray<FString> RestrictedFolderNames = FPaths::GetRestrictedFolderNames();
 	RestrictedFolderNames.Append({ TEXT("Platforms"), TEXT("Restricted") });
 
 	TArray<FString> AllPaths;
-	AllPaths.Add(BasePath);
+	if (!RestrictedPathsOnly)
+	{
+		AllPaths.Add(BasePath);
+	}
 
 	GetPossibleRestrictedPaths(BasePath, RestrictedFolderNames, AllPaths);
 
@@ -953,20 +956,23 @@ void FAutomationTestFramework::LoadTestTagMappings(const FString& BasePath)
 		if (IFileManager::Get().FileExists(*TagsSettingsFilePath))
 		{
 			UE_LOG(LogAutomationTest, Log, TEXT("Loading tags from %s."), *TagsSettingsFilePath);
-			GConfig->LoadFile(TagsSettingsFilePath);
-
-			GConfig->GetSection(TEXT("TestTagMapping"), TestTagMappings, TagsSettingsFilePath);
-			FString TestName, TestTags;
-			for (FString MapLine : TestTagMappings)
+			FConfigFile TagConfig;
+			TagConfig.Read(TagsSettingsFilePath);
+
+			const FConfigSection* TagSection = TagConfig.FindSection(TEXT("TestTagMapping"));
+			TArray<FName> TestNames;
+			TagSection->GenerateKeyArray(TestNames);
+			FString TestTags;
+			for (FName TestName : TestNames)
 			{
-				MapLine.Split(TEXT("="), &TestName, &TestTags);
-				RegisterAutomationTestTags(TestName, TestTags, false);
+				TagSection->GetString(*TestName.ToString(), TestTags);
+				RegisterAutomationTestTags(TestName.ToString(),  TestTags, false);
 			}
 		}
 	}
 }
 
-TArray<FString> FAutomationTestFramework::SaveTestTagMappings(const TArray<FString>& TestFullNames, const TArray<FString>& TestFilePaths) const
+void FAutomationTestFramework::SaveTestTagMappings(const TArray<FString>& TestFullNames, const TArray<FString>& TestFilePaths, const FBeforeTagMappingConfigSaved& BeforeConfigSaved, const FAfterTagMappingConfigSaved& AfterConfigSaved) const
 {
 	const FString ProjectDirPath = FPaths::ConvertRelativePathToFull(*FPaths::ProjectDir());
 	const FString EngineDirPath = FPaths::ConvertRelativePathToFull(FPaths::EngineDir());
@@ -975,15 +981,16 @@ TArray<FString> FAutomationTestFramework::SaveTestTagMappings(const TArray<FStri
 	RestrictedFolderNames.Append({ TEXT("Platforms"), TEXT("Restricted") });
 
 	TArray<FString> AllPathsProject;
+	// Search for both restricted and non-restricted paths for tests living inside a project
 	GetPossibleRestrictedPaths(ProjectDirPath, RestrictedFolderNames, AllPathsProject);
-	TArray<FString> AllPathsEngine;
-	GetPossibleRestrictedPaths(EngineDirPath, RestrictedFolderNames, AllPathsEngine);
+	TArray<FString> AllRestrictedPathsEngine;
+	// Search strictly for restricted paths for tests living under Engine/
+	GetPossibleRestrictedPaths(EngineDirPath, RestrictedFolderNames, AllRestrictedPathsEngine);
 
 	AllPathsProject.Add(ProjectDirPath);
 	AllPathsProject.Sort([](const FString& A, const FString& B) { return A.Len() > B.Len(); });
 
-	AllPathsEngine.Add(EngineDirPath);
-	AllPathsEngine.Sort([](const FString& A, const FString& B) { return A.Len() > B.Len(); });
+	AllRestrictedPathsEngine.Sort([](const FString& A, const FString& B) { return A.Len() > B.Len(); });
 
 	FString TestSourcePath;
 	FString TagsSettingsFilePath;
@@ -1004,44 +1011,56 @@ TArray<FString> FAutomationTestFramework::SaveTestTagMappings(const TArray<FStri
 
 		if (TagsSettingsFilePath.IsEmpty())
 		{
-			for (int j = 0; j < AllPathsEngine.Num(); j++)
+			for (int j = 0; j < AllRestrictedPathsEngine.Num(); j++)
 			{
-				if (FPaths::IsUnderDirectory(TestSourcePath, AllPathsEngine[j]))
+				if (FPaths::IsUnderDirectory(TestSourcePath, AllRestrictedPathsEngine[j]))
 				{
-					TagsSettingsFilePath = AllPathsEngine[j] / TEXT("Config") / TEXT("Tests") / TEXT("Tags.ini");
+					TagsSettingsFilePath = AllRestrictedPathsEngine[j] / TEXT("Config") / TEXT("Tests") / TEXT("Tags.ini");
 					break;
 				}
 			}
 		}
 
+		// Tests that aren't in any restricted path under Engine/ will end up in NotForLicensees
+		if (TagsSettingsFilePath.IsEmpty() && FPaths::IsUnderDirectory(TestSourcePath, EngineDirPath))
+		{
+			TagsSettingsFilePath = EngineDirPath / TEXT("Restricted") / TEXT("NotForLicensees") / TEXT("Config") / TEXT("Tests") / TEXT("Tags.ini");
+		}
+
 		if (TagsSettingsFilePath.IsEmpty())
 		{
-			UE_LOG(LogAutomationTest, Error, TEXT("Test neither in project path or engine path, cannot be mapped. Path: %s"), *TestSourcePath);
+			UE_LOG(LogAutomationTest, Error, TEXT("Test neither in project path or a restricted engine path, cannot be mapped. Path: %s"), *TestSourcePath);
 			continue;
 		}
 
-		TagsSettingsFilePath = FConfigCacheIni::NormalizeConfigIniPath(TagsSettingsFilePath);
 		if (!AffectedTagConfigPaths.Contains(TagsSettingsFilePath))
 		{
-			FPlatformFileManager::Get().GetPlatformFile().SetReadOnly(*TagsSettingsFilePath, false);
 			AffectedTagConfigPaths.Add(TagsSettingsFilePath, FConfigFile());
+
+			// Removal of readonly flag for Combine and Write to work properly
+			FPlatformFileManager::Get().GetPlatformFile().SetReadOnly(*TagsSettingsFilePath, false);
+
+			BeforeConfigSaved.ExecuteIfBound(TagsSettingsFilePath);
+
+			// If file exists, import lines first before updating with SetInSection
+			if (IFileManager::Get().FileExists(*TagsSettingsFilePath))
+			{
+				AffectedTagConfigPaths[TagsSettingsFilePath].Combine(TagsSettingsFilePath);
+			}
 		}
-		AffectedTagConfigPaths[TagsSettingsFilePath].SetInSection(TEXT("TestTagMapping"), *TestFullNames[i], **TestFullNameToTagDataMap.Find(TestFullNames[i]));
+
+		AffectedTagConfigPaths[TagsSettingsFilePath].SetInSection(TEXT("TestTagMapping"), FName(TestFullNames[i]), *TestFullNameToTagDataMap.Find(TestFullNames[i]));
 	}
 
-	TSet<FString> AllConfigFiles;
-	AffectedTagConfigPaths.GetKeys(AllConfigFiles);
-	for (FString ConfigPath : AllConfigFiles)
+	TSet<FString> AllConfigsPaths;
+	AffectedTagConfigPaths.GetKeys(AllConfigsPaths);
+	for (const FString& ConfigPath : AllConfigsPaths)
 	{
-		UE_LOG(LogAutomationTest, Log, TEXT("Saving tags to %s."), *ConfigPath);
-		if (IFileManager::Get().FileExists(*ConfigPath))
-		{
-			AffectedTagConfigPaths[ConfigPath].Combine(ConfigPath);
-		}
+		UE_LOG(LogAutomationTestFramework, Log, TEXT("Saving tags to %s."), *ConfigPath);
 		AffectedTagConfigPaths[ConfigPath].Write(ConfigPath);
-	}
+		AfterConfigSaved.ExecuteIfBound(TagsSettingsFilePath);
 
-	return AllConfigFiles.Array();
+	}
 }
 
 bool FAutomationTestFramework::ShouldTestContent(const FString& Path) const
diff --git a/Engine/Source/Runtime/Core/Public/Misc/AutomationTest.h b/Engine/Source/Runtime/Core/Public/Misc/AutomationTest.h
--- a/Engine/Source/Runtime/Core/Public/Misc/AutomationTest.h
+++ b/Engine/Source/Runtime/Core/Public/Misc/AutomationTest.h
@@ -933,6 +933,11 @@ DECLARE_MULTICAST_DELEGATE_OneParam(FOnTestEvent, FAutomationTestBase*);
 
 DECLARE_MULTICAST_DELEGATE_OneParam(FOnTestSectionEvent, const FString& /*Section*/);
 
+DECLARE_DELEGATE_OneParam(FBeforeTagMappingConfigSaved, const FString& /*ConfigPath*/);
+
+DECLARE_DELEGATE_OneParam(FAfterTagMappingConfigSaved, const FString& /*ConfigPath*/);
+
+
 /** Class representing the main framework for running automation tests */
 class FAutomationTestFramework
 {
@@ -1180,13 +1185,13 @@ public:
 	/**
 	 * Reads test tag mappings under a certain base path
 	 */
-	CORE_API void LoadTestTagMappings(const FString& BasePath);
+	CORE_API void LoadTestTagMappings(const FString& BasePath, const bool RestrictedPathsOnly);
+
 
 	/**
-	 * Saves runtime edited test tag mappings to Tags.ini
-	 * Returns the list of updated tag mapping files
+	 * Return the list of test tag config files affected by editing as a map of config path and FConfigFile object
 	 */
-	CORE_API TArray<FString> SaveTestTagMappings(const TArray<FString>& TestFullNames, const TArray<FString>& TestFilePaths) const;
+	CORE_API void SaveTestTagMappings(const TArray<FString>& TestFullNames, const TArray<FString>& TestFilePaths, const FBeforeTagMappingConfigSaved& BeforeConfigSaved, const FAfterTagMappingConfigSaved& AfterConfigSaved) const;
 
 	/**
 	 * Whether the testing framework should allow content to be tested or not.  Intended to block developer directories.
