Fixed FWindowsPlatformMisc::SetEnvironmentVar to unset the variable if it is null or "", just like on Mac and Linux, and .NET 8

diff --git a/Engine/Source/Runtime/Core/Private/Windows/WindowsPlatformMisc.cpp b/Engine/Source/Runtime/Core/Private/Windows/WindowsPlatformMisc.cpp
--- a/Engine/Source/Runtime/Core/Private/Windows/WindowsPlatformMisc.cpp
+++ b/Engine/Source/Runtime/Core/Private/Windows/WindowsPlatformMisc.cpp
@@ -1180,6 +1180,11 @@ FString FWindowsPlatformMisc::GetEnvironmentVariable(const TCHAR* VariableName)
 
 void FWindowsPlatformMisc::SetEnvironmentVar(const TCHAR* VariableName, const TCHAR* Value)
 {
+	if (Value != nullptr && Value[0] == TEXT('\0'))
+	{
+		Value = nullptr;
+	}
+	
 	uint32 Error = ::SetEnvironmentVariable(VariableName, Value);
 	if (Error == 0)
 	{
