Output global allocator stats to GLog when title goes OOM

diff --git a/Engine/Source/Runtime/Core/Private/GenericPlatform/GenericPlatformMemory.cpp b/Engine/Source/Runtime/Core/Private/GenericPlatform/GenericPlatformMemory.cpp
--- a/Engine/Source/Runtime/Core/Private/GenericPlatform/GenericPlatformMemory.cpp
+++ b/Engine/Source/Runtime/Core/Private/GenericPlatform/GenericPlatformMemory.cpp
@@ -20,6 +20,7 @@
 #include "Misc/ConfigCacheIni.h"
 #include "Misc/CommandLine.h"
 #include "Misc/MessageDialog.h"
+#include "Misc/OutputDeviceRedirector.h"
 #include "Templates/UnrealTemplate.h"
 #include "HAL/IConsoleManager.h"
 #include "ProfilingDebugging/MemoryTrace.h"
@@ -216,6 +217,11 @@ void FGenericPlatformMemory::OnOutOfMemory(uint64 Size, uint32 Alignment)
 {
 	auto HandleOOM = [&]()
 	{
+		if (GLog)
+		{
+			GLog->Panic();
+		}
+			
 		// Update memory stats before we enter the crash handler.
 		OOMAllocationSize = Size;
 		OOMAllocationAlignment = Alignment;
@@ -251,9 +257,9 @@ void FGenericPlatformMemory::OnOutOfMemory(uint64 Size, uint32 Alignment)
 			(uint64)PlatformMemoryStats.UsedVirtual, double(PlatformMemoryStats.UsedVirtual) / (1024 * 1024 * 1024),
 			(uint64)PlatformMemoryStats.PeakUsedVirtual, double(PlatformMemoryStats.PeakUsedVirtual) / (1024 * 1024 * 1024)
 		);
-		if (GWarn)
+		if (BackupOOMMemoryPool && GLog)
 		{
-			UE::Private::GMalloc->DumpAllocatorStats(*GWarn);
+			UE::Private::GMalloc->DumpAllocatorStats(*GLog);
 		}
 
 		// let any registered handlers go
