Fix compilation of Android development builds under NDK 28+. ThermalStatus variable was being shadowed when STATS and CSVPROFILER was enabled.

diff --git a/Engine/Source/Runtime/Core/Private/Android/AndroidStats.cpp b/Engine/Source/Runtime/Core/Private/Android/AndroidStats.cpp
--- a/Engine/Source/Runtime/Core/Private/Android/AndroidStats.cpp
+++ b/Engine/Source/Runtime/Core/Private/Android/AndroidStats.cpp
@@ -224,9 +224,9 @@ namespace UE
 #endif
 #if STATS
 				static const FName CPUStatName = GET_STATFNAME(STAT_CPUTemp);
-				static const FName ThermalStatus = GET_STATFNAME(STAT_ThermalStatus);
+				static const FName ThermalStatusStatName = GET_STATFNAME(STAT_ThermalStatus);
 				SET_FLOAT_STAT_FName(CPUStatName, CPUTemp);
-				SET_DWORD_STAT_FName(ThermalStatus, GThermalStatus);
+				SET_DWORD_STAT_FName(ThermalStatusStatName, GThermalStatus);
 
 				for (int32 FrequencyGroupIndex = 0; FrequencyGroupIndex < NumFrequencyGroups; FrequencyGroupIndex++)
 				{
