Clarifying comments around getting disk space.

diff --git a/Engine/Source/Runtime/Core/Private/IOS/IOSPlatformMisc.cpp b/Engine/Source/Runtime/Core/Private/IOS/IOSPlatformMisc.cpp
--- a/Engine/Source/Runtime/Core/Private/IOS/IOSPlatformMisc.cpp
+++ b/Engine/Source/Runtime/Core/Private/IOS/IOSPlatformMisc.cpp
@@ -1022,7 +1022,7 @@ FString FIOSPlatformMisc::GetOSVersion()
 	return FString([[UIDevice currentDevice] systemVersion]);
 }
 
-bool FIOSPlatformMisc::GetDiskTotalAndFreeSpace(const FString& InPath, uint64& TotalNumberOfBytes, uint64& NumberOfFreeBytes, bool IncludeFreeable)
+bool FIOSPlatformMisc::GetDiskTotalAndFreeSpace(const FString& InPath, uint64& OutTotalNumberOfBytes, uint64& OutNumberOfFreeBytes, bool bInIncludeFreeable)
 {
     SCOPED_AUTORELEASE_POOL;
 
@@ -1030,24 +1030,28 @@ bool FIOSPlatformMisc::GetDiskTotalAndFreeSpace(const FString& InPath, uint64& T
 #if !PLATFORM_TVOS
     bool GetValueSuccess = false;
     
+	// NSURLVolumeAvailableCapacityForImportantUsageKey
+	// This key is more reliable because of how much free space can be "on demand" on iOS - on devices we've seen as much as 100gb
+	// difference between the "ImportantUsage" and "Capacity" keys. Capacity is almost certainly not what you want. However, because "ImportantUsage" involves
+	// accumulating all the data that could be freed if needed, it can be very slow (upwards of 100ms).
     NSNumber *FreeBytes = nil;
     NSURL *URL = [NSURL fileURLWithPath : NSHomeDirectory()];
-    GetValueSuccess = [URL getResourceValue : &FreeBytes forKey : (IncludeFreeable ? NSURLVolumeAvailableCapacityForImportantUsageKey : NSURLVolumeAvailableCapacityKey) error : nil];
+    GetValueSuccess = [URL getResourceValue : &FreeBytes forKey : (bInIncludeFreeable ? NSURLVolumeAvailableCapacityForImportantUsageKey : NSURLVolumeAvailableCapacityKey) error : nil];
     if (FreeBytes)
     {
-        NumberOfFreeBytes = [FreeBytes longLongValue];
+        OutNumberOfFreeBytes = [FreeBytes longLongValue];
     }
     
     NSNumber *TotalBytes = nil;
     GetValueSuccess = GetValueSuccess &&[URL getResourceValue : &TotalBytes forKey : NSURLVolumeTotalCapacityKey error : nil];
     if (TotalBytes)
     {
-        TotalNumberOfBytes = [TotalBytes longLongValue];
+        OutTotalNumberOfBytes = [TotalBytes longLongValue];
     }
     
     if (GetValueSuccess
-        && (NumberOfFreeBytes > 0)
-        && (TotalNumberOfBytes > 0))
+        && (OutNumberOfFreeBytes > 0)
+        && (OutTotalNumberOfBytes > 0))
     {
         return true;
     }
@@ -1058,12 +1062,15 @@ bool FIOSPlatformMisc::GetDiskTotalAndFreeSpace(const FString& InPath, uint64& T
         NSDictionary<NSFileAttributeKey, id>* FSStat = [[NSFileManager defaultManager] attributesOfFileSystemForPath:NSHomeDirectory() error : nil];
         if (FSStat)
         {
-            NumberOfFreeBytes = [[FSStat objectForKey : NSFileSystemFreeSize] longLongValue];
-            TotalNumberOfBytes = [[FSStat objectForKey : NSFileSystemSize] longLongValue];
+            OutNumberOfFreeBytes = [[FSStat objectForKey : NSFileSystemFreeSize] longLongValue];
+            OutTotalNumberOfBytes = [[FSStat objectForKey : NSFileSystemSize] longLongValue];
 
             return true;
         }
 
+		// This should almost never happen to make sure to leave a breadcrumb if it does since
+		// callers are probably not checking results.
+		UE_LOG(LogInit, Warning, TEXT("Call to GetDiskTotalAndFreeSpace failed."));
         return false;
     }
 }
diff --git a/Engine/Source/Runtime/Core/Public/GenericPlatform/GenericPlatformMisc.h b/Engine/Source/Runtime/Core/Public/GenericPlatform/GenericPlatformMisc.h
--- a/Engine/Source/Runtime/Core/Public/GenericPlatform/GenericPlatformMisc.h
+++ b/Engine/Source/Runtime/Core/Public/GenericPlatform/GenericPlatformMisc.h
@@ -789,7 +789,13 @@ struct FGenericPlatformMisc
 	 */
 	static CORE_API FString GetOSVersion();
 
-	/** Retrieves information about the total number of bytes and number of free bytes for the specified disk path. */
+	/** 
+	 * Retrieves information about the total number of bytes and number of free bytes for the specified disk path. 
+	 *
+	 * Note that on modern file systems the concept of "free space" is a bit more loosely defined than one might historically expect due 
+	 * to features like on-demand freeing of caches and automatic content deduplication. As a result, this function can be
+	 * very slow (10s of ms), and the free space value should not be used for exact measurements.
+	*/
 	static CORE_API bool GetDiskTotalAndFreeSpace( const FString& InPath, uint64& TotalNumberOfBytes, uint64& NumberOfFreeBytes );
 
 	/**
diff --git a/Engine/Source/Runtime/Core/Public/IOS/IOSPlatformMisc.h b/Engine/Source/Runtime/Core/Public/IOS/IOSPlatformMisc.h
--- a/Engine/Source/Runtime/Core/Public/IOS/IOSPlatformMisc.h
+++ b/Engine/Source/Runtime/Core/Public/IOS/IOSPlatformMisc.h
@@ -122,7 +122,10 @@ struct CORE_API FIOSPlatformMisc : public FApplePlatformMisc
 	static FString GetDeviceId();
 	static FString GetOSVersion();
 	static FString GetUniqueAdvertisingId();
-	static bool GetDiskTotalAndFreeSpace(const FString& InPath, uint64& TotalNumberOfBytes, uint64& NumberOfFreeBytes, bool IncludeFreeable = true);
+
+	// bInIncludeFreeable determines whether to use the "Important" iOS free space value or the actual free space (Capacity). 
+	// The "Important" value is almost certainly what you want - but it's slow (upwards of 100ms).
+	static bool GetDiskTotalAndFreeSpace(const FString& InPath, uint64& OutTotalNumberOfBytes, uint64& OutNumberOfFreeBytes, bool bInIncludeFreeable = true);
 	
 	static void RequestStoreReview();
 
