[AutoRTFM] Fix memory corruption/deadlocking when a mutex's acquisition and release spans open and closed.

Previously, if we attempted to lock in the open and unlock in the closed, we'd trigger the memory validator.

If we attempted to lock in the closed and unlock in the open, we'd leave the mutex in an invalid state and future lock attempts would deadlock.

Now, the mutex detects if it is transactionally open and goes into the closed before locking and unlocking. This case should be very rare, but it's worth getting this right instead of deadlocking or corrupting memory.

diff --git a/Engine/Source/Programs/AutoRTFMTests/Private/MutexTests.cpp b/Engine/Source/Programs/AutoRTFMTests/Private/MutexTests.cpp
--- a/Engine/Source/Programs/AutoRTFMTests/Private/MutexTests.cpp
+++ b/Engine/Source/Programs/AutoRTFMTests/Private/MutexTests.cpp
@@ -837,18 +837,64 @@ TEST_CASE("TransactionallySafeMutex.Locked Mutex In Destructed Object")
 
 TEST_CASE("TransactionallySafeMutex.Closed Lock Then Open Unlock")
 {
-	// TODO(SOL-8282): Mutex.Unlock() causes the memory validator to fail.
-
-/*
 	UE::FTransactionallySafeMutex Mutex;
 	AutoRTFM::Testing::Commit([&]
 	{
 		Mutex.Lock();
+		REQUIRE(Mutex.IsLocked());
 
 		AutoRTFM::Open([&]
 		{
 			Mutex.Unlock();
 		});
 	});
-*/
-} 
+	REQUIRE(!Mutex.IsLocked());
+}
+
+TEST_CASE("TransactionallySafeMutex.Closed TryLock Then Open Unlock")
+{
+	UE::FTransactionallySafeMutex Mutex;
+	AutoRTFM::Testing::Commit([&]
+	{
+		REQUIRE(Mutex.TryLock());
+		REQUIRE(Mutex.IsLocked());
+
+		AutoRTFM::Open([&]
+		{
+			Mutex.Unlock();
+		});
+	});
+	REQUIRE(!Mutex.IsLocked());
+}
+
+TEST_CASE("TransactionallySafeMutex.Open Lock Then Closed Unlock")
+{
+	UE::FTransactionallySafeMutex Mutex;
+	AutoRTFM::Testing::Commit([&]
+	{
+		AutoRTFM::Open([&]
+		{
+			Mutex.Lock();
+		});
+		REQUIRE(Mutex.IsLocked());
+
+		Mutex.Unlock();
+	});
+	REQUIRE(!Mutex.IsLocked());
+}
+
+TEST_CASE("TransactionallySafeMutex.Open TryLock Then Closed Unlock")
+{
+	UE::FTransactionallySafeMutex Mutex;
+	AutoRTFM::Testing::Commit([&]
+	{
+		AutoRTFM::Open([&]
+		{
+			REQUIRE(Mutex.TryLock());
+		});
+		REQUIRE(Mutex.IsLocked());
+
+		Mutex.Unlock();
+	});
+	REQUIRE(!Mutex.IsLocked());
+}
diff --git a/Engine/Source/Runtime/Core/Public/Async/TransactionallySafeMutex.h b/Engine/Source/Runtime/Core/Public/Async/TransactionallySafeMutex.h
--- a/Engine/Source/Runtime/Core/Public/Async/TransactionallySafeMutex.h
+++ b/Engine/Source/Runtime/Core/Public/Async/TransactionallySafeMutex.h
@@ -51,28 +51,13 @@ struct FTransactionallySafeMutex final
 
 	void Lock()
 	{
-		if (AutoRTFM::IsTransactional() || AutoRTFM::IsCommittingOrAborting())
+		if (AutoRTFM::IsClosed() || AutoRTFM::IsCommittingOrAborting())
 		{
-			ensure(State->TransactionalLockCount >= -1 && State->TransactionalLockCount <= 0);
-
-			State->TransactionalLockCount++;
-
-			if (State->bTransactionallyLocked)
-			{
-				// Underlying mutex has already been locked for the duration of the transaction.
-				return;
-			}
-
-			if (State->TransactionalLockCount == 0)
-			{
-				// The mutex was locked before the transaction and unlocked inside the transaction.
-				return;
-			}
-
-			// First time the mutex has been locked during the transaction.
-			// Mark the transactionally locked flag and register on-commit and on-abort handlers.
-			State->MutexLock();
-			SetTransactionallyLockedAndRegisterHandlers();
+			LockTransactionally();
+		}
+		else if (AutoRTFM::IsTransactional())
+		{
+			InTheClosed<&FTransactionallySafeMutex::LockTransactionally>();
 		}
 		else
 		{
@@ -83,42 +68,13 @@ struct FTransactionallySafeMutex final
 
 	[[nodiscard]] inline bool TryLock()
 	{
-		if (AutoRTFM::IsTransactional() || AutoRTFM::IsCommittingOrAborting())
+		if (AutoRTFM::IsClosed() || AutoRTFM::IsCommittingOrAborting())
 		{
-			ensure(State->TransactionalLockCount >= -1 && State->TransactionalLockCount <= 1);
-
-			if (State->TransactionalLockCount > 0)
-			{
-				return false; // Attempting double-lock within transaction.
-			}
-
-			if (State->bTransactionallyLocked)
-			{
-				// Underlying mutex has already been locked for the duration of the transaction.
-				State->TransactionalLockCount++;
-				return true;
-			}
-
-			if (State->TransactionalLockCount < 0)
-			{
-				// The mutex was locked before the transaction and unlocked inside the transaction.
-				// As unlock is deferred, bump the lock counter and return true to emulate the lock.
-				State->TransactionalLockCount++;
-				return true;
-			}
-			
-			// Attempt to lock the underlying mutex.
-			if (!State->MutexTryLock())
-			{
-				return false;
-			}
-
-			// First time the mutex has been locked during the transaction.
-			// Increment the lock count, mark the transactionally locked flag and register 
-			// on-commit and on-abort handlers.
-			State->TransactionalLockCount++;
-			SetTransactionallyLockedAndRegisterHandlers();
-			return true;
+			return TryLockTransactionally();
+		}
+		else if (AutoRTFM::IsTransactional())
+		{
+			return InTheClosed<&FTransactionallySafeMutex::TryLockTransactionally>();
 		}
 		else
 		{
@@ -129,14 +85,13 @@ struct FTransactionallySafeMutex final
 
 	void Unlock()
 	{
-		if (AutoRTFM::IsTransactional() || AutoRTFM::IsCommittingOrAborting())
+		if (AutoRTFM::IsClosed() || AutoRTFM::IsCommittingOrAborting())
 		{
-			ensure(State->TransactionalLockCount >= 0 && State->TransactionalLockCount <= 1);
-
-			State->TransactionalLockCount--;
-			// Use an on-commit handler to unlock the underlying mutex if not balanced with a lock
-			// before the transaction commits.
-			MaybeRegisterCommitHandler();
+			UnlockTransactionally();
+		}
+		else if (AutoRTFM::IsTransactional())
+		{
+			InTheClosed<&FTransactionallySafeMutex::UnlockTransactionally>();
 		}
 		else
 		{
@@ -162,6 +117,99 @@ struct FTransactionallySafeMutex final
 private:
 	UE_NONCOPYABLE(FTransactionallySafeMutex)
 
+	template <bool (FTransactionallySafeMutex::*Functor)()>
+	bool InTheClosed()
+	{
+		bool Result;
+		AutoRTFM::EContextStatus Status = AutoRTFM::Close([this, &Result]
+		{
+			Result = (this->*Functor)();
+		});
+		check(Status == AutoRTFM::EContextStatus::OnTrack);
+		return Result;
+	}
+
+	bool TryLockTransactionally()
+	{
+		check(AutoRTFM::IsClosed() || AutoRTFM::IsCommittingOrAborting());
+		ensure(State->TransactionalLockCount >= -1 && State->TransactionalLockCount <= 1);
+
+		if (State->TransactionalLockCount > 0)
+		{
+			return false; // Attempting double-lock within transaction.
+		}
+
+		if (State->bTransactionallyLocked)
+		{
+			// Underlying mutex has already been locked for the duration of the transaction.
+			State->TransactionalLockCount++;
+			return true;
+		}
+
+		if (State->TransactionalLockCount < 0)
+		{
+			// The mutex was locked before the transaction and unlocked inside the transaction.
+			// As unlock is deferred, bump the lock counter and return true to emulate the lock.
+			State->TransactionalLockCount++;
+			return true;
+		}
+
+		// Attempt to lock the underlying mutex.
+		if (!State->MutexTryLock())
+		{
+			return false;
+		}
+
+		// First time the mutex has been locked during the transaction.
+		// Increment the lock count, mark the transactionally locked flag and register 
+		// on-commit and on-abort handlers.
+		State->TransactionalLockCount++;
+		SetTransactionallyLockedAndRegisterHandlers();
+		return true;
+	};
+
+	// Acquires the mutex in a transactionally-safe way. Assumes we are either in the closed or
+	// are mid-commit/mid-abort.
+	bool LockTransactionally()
+	{
+		check(AutoRTFM::IsClosed() || AutoRTFM::IsCommittingOrAborting());
+		ensure(State->TransactionalLockCount >= -1 && State->TransactionalLockCount <= 0);
+
+		State->TransactionalLockCount++;
+
+		if (State->bTransactionallyLocked)
+		{
+			// Underlying mutex has already been locked for the duration of the transaction.
+			return true;
+		}
+
+		if (State->TransactionalLockCount == 0)
+		{
+			// The mutex was locked before the transaction and unlocked inside the transaction.
+			return true;
+		}
+
+		// First time the mutex has been locked during the transaction.
+		// Mark the transactionally locked flag and register on-commit and on-abort handlers.
+		State->MutexLock();
+		SetTransactionallyLockedAndRegisterHandlers();
+		return true;
+	}
+
+	// Releases the mutex in a transactionally-safe way. Assumes we are either in the closed or
+	// are mid-commit/mid-abort.
+	bool UnlockTransactionally()
+	{
+		check(AutoRTFM::IsClosed() || AutoRTFM::IsCommittingOrAborting());
+		ensure(State->TransactionalLockCount >= 0 && State->TransactionalLockCount <= 1);
+
+		State->TransactionalLockCount--;
+		// Use an on-commit handler to unlock the underlying mutex if not balanced with a lock
+		// before the transaction commits.
+		MaybeRegisterCommitHandler();
+		return true;
+	}
+
 	// Called when the underlying mutex is locked for the first time within a transaction nest.
 	// Sets the bTransactionallyLocked flag, calls MaybeRegisterCommitHandler() and registers an
 	// on-abort handler to unlock the underlying mutex.
