Allowing Frame Start and End Misc Trace events to work with minimal trace support.

This was done by upgrading the frame channel and trace events to the trace minimal API.

Misc Trace is still compiled out by default in shipping builds, therefore to be able to use Trace the frame start and end events, MISCTRACE_ENABLED needs to be defined to 1 in the desired target.

diff --git a/Engine/Source/Runtime/Core/Private/ProfilingDebugging/MiscTrace.cpp b/Engine/Source/Runtime/Core/Private/ProfilingDebugging/MiscTrace.cpp
--- a/Engine/Source/Runtime/Core/Private/ProfilingDebugging/MiscTrace.cpp
+++ b/Engine/Source/Runtime/Core/Private/ProfilingDebugging/MiscTrace.cpp
@@ -1,14 +1,12 @@
 // Copyright Epic Games, Inc. All Rights Reserved.
 #include "ProfilingDebugging/MiscTrace.h"
 
-#if MISCTRACE_ENABLED
-
 #include "Trace/Trace.inl"
 #include "Misc/CString.h"
 #include "HAL/PlatformTLS.h"
 #include "HAL/PlatformTime.h"
 
-UE_TRACE_CHANNEL(FrameChannel)
+UE_TRACE_MINIMAL_CHANNEL(FrameChannel)
 UE_TRACE_CHANNEL_DEFINE(BookmarkChannel)
 UE_TRACE_CHANNEL(RegionChannel)
 UE_TRACE_CHANNEL(ScreenshotChannel)
@@ -49,15 +47,15 @@ UE_TRACE_EVENT_BEGIN(Misc, RegionEndWithId)
 	UE_TRACE_EVENT_FIELD(uint64, RegionId)
 UE_TRACE_EVENT_END()
 
-UE_TRACE_EVENT_BEGIN(Misc, BeginFrame)
-	UE_TRACE_EVENT_FIELD(uint64, Cycle)
-	UE_TRACE_EVENT_FIELD(uint8, FrameType)
-UE_TRACE_EVENT_END()
+UE_TRACE_MINIMAL_EVENT_BEGIN(Misc, BeginFrame)
+	UE_TRACE_MINIMAL_EVENT_FIELD(uint64, Cycle)
+	UE_TRACE_MINIMAL_EVENT_FIELD(uint8, FrameType)
+UE_TRACE_MINIMAL_EVENT_END()
 
-UE_TRACE_EVENT_BEGIN(Misc, EndFrame)
-	UE_TRACE_EVENT_FIELD(uint64, Cycle)
-	UE_TRACE_EVENT_FIELD(uint8, FrameType)
-UE_TRACE_EVENT_END()
+UE_TRACE_MINIMAL_EVENT_BEGIN(Misc, EndFrame)
+	UE_TRACE_MINIMAL_EVENT_FIELD(uint64, Cycle)
+	UE_TRACE_MINIMAL_EVENT_FIELD(uint8, FrameType)
+UE_TRACE_MINIMAL_EVENT_END()
 
 UE_TRACE_EVENT_BEGIN(Misc, ScreenshotHeader)
 	UE_TRACE_EVENT_FIELD(uint32, Id)
@@ -137,26 +135,26 @@ void FMiscTrace::OutputBookmarkInternalCycles(uint64 Cycles, const void* Bookmar
 
 void FMiscTrace::OutputBeginFrame(ETraceFrameType FrameType)
 {
-	if (!UE_TRACE_CHANNELEXPR_IS_ENABLED(FrameChannel))
+	if (!UE_TRACE_MINIMAL_CHANNELEXPR_IS_ENABLED(FrameChannel))
 	{
 		return;
 	}
 
 	uint64 Cycle = FPlatformTime::Cycles64();
-	UE_TRACE_LOG(Misc, BeginFrame, FrameChannel)
+	UE_TRACE_MINIMAL_LOG(Misc, BeginFrame, FrameChannel)
 		<< BeginFrame.Cycle(Cycle)
 		<< BeginFrame.FrameType((uint8)FrameType);
 }
 
 void FMiscTrace::OutputEndFrame(ETraceFrameType FrameType)
 {
-	if (!UE_TRACE_CHANNELEXPR_IS_ENABLED(FrameChannel))
+	if (!UE_TRACE_MINIMAL_CHANNELEXPR_IS_ENABLED(FrameChannel))
 	{
 		return;
 	}
 
 	uint64 Cycle = FPlatformTime::Cycles64();
-	UE_TRACE_LOG(Misc, EndFrame, FrameChannel)
+	UE_TRACE_MINIMAL_LOG(Misc, EndFrame, FrameChannel)
 		<< EndFrame.Cycle(Cycle)
 		<< EndFrame.FrameType((uint8)FrameType);
 }
@@ -211,4 +209,3 @@ bool FMiscTrace::ShouldTraceRegion()
 {
 	return UE_TRACE_CHANNELEXPR_IS_ENABLED(RegionChannel);
 }
-#endif
diff --git a/Engine/Source/Runtime/Core/Public/ProfilingDebugging/MiscTrace.h b/Engine/Source/Runtime/Core/Public/ProfilingDebugging/MiscTrace.h
--- a/Engine/Source/Runtime/Core/Public/ProfilingDebugging/MiscTrace.h
+++ b/Engine/Source/Runtime/Core/Public/ProfilingDebugging/MiscTrace.h
@@ -18,6 +18,10 @@
 #endif
 #endif
 
+#ifndef FRAME_TRACE_ENABLED
+#define FRAME_TRACE_ENABLED MISCTRACE_ENABLED
+#endif
+
 UE_TRACE_CHANNEL_EXTERN(BookmarkChannel, CORE_API)
 
 enum ETraceFrameType
@@ -47,8 +51,6 @@ struct FTraceUtils
 	}
 };
 
-#if MISCTRACE_ENABLED
-
 class FName;
 
 struct FMiscTrace
@@ -109,6 +111,8 @@ private:
 	CORE_API static void OutputBookmarkInternalCycles(uint64 Cycles, const void* BookmarkPoint, uint32 CallstackId, uint16 EncodedFormatArgsSize, uint8* EncodedFormatArgs);
 };
 
+#if MISCTRACE_ENABLED
+
 #define TRACE_BOOKMARK(Format, ...) \
 if (UE_TRACE_CHANNELEXPR_IS_ENABLED(BookmarkChannel)) \
 { \
@@ -172,12 +176,6 @@ FMiscTrace::OutputBeginRegionWithId(RegionName __VA_OPT__(,) __VA_ARGS__);
 #define TRACE_END_REGION_WITH_ID(RegionId) \
 	FMiscTrace::OutputEndRegionWithId(RegionId);
 
-#define TRACE_BEGIN_FRAME(FrameType) \
-	FMiscTrace::OutputBeginFrame(FrameType);
-
-#define TRACE_END_FRAME(FrameType) \
-	FMiscTrace::OutputEndFrame(FrameType);
-
 #define TRACE_SCREENSHOT(Name, Cycle, Width, Height, Data) \
 	FMiscTrace::OutputScreenshot(Name, Cycle, Width, Height, Data);
 
@@ -198,11 +196,24 @@ FMiscTrace::OutputBeginRegionWithId(RegionName __VA_OPT__(,) __VA_ARGS__);
 #define TRACE_BEGIN_REGION_WITH_ID(...) 0;
 #define TRACE_END_REGION(...)
 #define TRACE_END_REGION_WITH_ID(...)
-#define TRACE_BEGIN_FRAME(...)
-#define TRACE_END_FRAME(...)
 #define TRACE_SCREENSHOT(...)
 #define SHOULD_TRACE_SCREENSHOT(...) false
 #define SHOULD_TRACE_BOOKMARK(...) false
 #define SHOULD_TRACE_REGION(...) false
 
-#endif
+#endif // MISCTRACE_ENABLED
+
+#if FRAME_TRACE_ENABLED 
+
+#define TRACE_BEGIN_FRAME(FrameType) \
+FMiscTrace::OutputBeginFrame(FrameType);
+
+#define TRACE_END_FRAME(FrameType) \
+FMiscTrace::OutputEndFrame(FrameType);
+
+#else
+
+#define TRACE_BEGIN_FRAME(...)
+#define TRACE_END_FRAME(...)
+
+#endif // FRAME_TRACE_ENABLED
