Disable stall detector when running under tsan.

[RN] none

diff --git a/Engine/Source/Runtime/Core/Public/ProfilingDebugging/StallDetector.h b/Engine/Source/Runtime/Core/Public/ProfilingDebugging/StallDetector.h
--- a/Engine/Source/Runtime/Core/Public/ProfilingDebugging/StallDetector.h
+++ b/Engine/Source/Runtime/Core/Public/ProfilingDebugging/StallDetector.h
@@ -18,7 +18,7 @@
 #include <atomic>
 
 #ifndef STALL_DETECTOR
- #if WITH_EDITOR && ( PLATFORM_WINDOWS || PLATFORM_UNIX ) && !UE_BUILD_SHIPPING && COUNTERSTRACE_ENABLED
+ #if WITH_EDITOR && ( PLATFORM_WINDOWS || PLATFORM_UNIX ) && !UE_BUILD_SHIPPING && COUNTERSTRACE_ENABLED && !USING_THREAD_SANITISER
   #define STALL_DETECTOR 1
  #else
   #define STALL_DETECTOR 0
