Added bIsUGC flag to LocMeta to control which GFPs are allowed to replace existing localization data

This allows non-UGC GFPs to still load their base LocRes data even if there's a FPolyglotTextData hotfix applied to that ID, as is required for bIsMinimalPatch to work correctly

diff --git a/Engine/Source/Developer/Localization/Private/UserGeneratedContentLocalization.cpp b/Engine/Source/Developer/Localization/Private/UserGeneratedContentLocalization.cpp
--- a/Engine/Source/Developer/Localization/Private/UserGeneratedContentLocalization.cpp
+++ b/Engine/Source/Developer/Localization/Private/UserGeneratedContentLocalization.cpp
@@ -722,6 +722,7 @@ bool CompileLocalizationTarget(const FString& LocalizationTargetDirectory, const
 		FTextLocalizationMetaDataResource LocMeta;
 		if (FTextLocalizationResourceGenerator::GenerateLocMeta(LocTextHelper, LocResName, LocMeta))
 		{
+			LocMeta.bIsUGC = true;
 			if (!LocMeta.SaveToFile(LocalizationTargetDirectory / LocMetaName))
 			{
 				UE_LOG(LogUGCLocalization, Error, TEXT("Failed to save LocMeta file for '%s'"), *LocTextHelper.GetTargetName());
diff --git a/Engine/Source/Runtime/Core/Private/Internationalization/TextLocalizationManager.cpp b/Engine/Source/Runtime/Core/Private/Internationalization/TextLocalizationManager.cpp
--- a/Engine/Source/Runtime/Core/Private/Internationalization/TextLocalizationManager.cpp
+++ b/Engine/Source/Runtime/Core/Private/Internationalization/TextLocalizationManager.cpp
@@ -1262,12 +1262,33 @@ void FTextLocalizationManager::HandleLocalizationTargetsMounted(TArrayView<const
 	LocLoadFlags |= (FApp::IsGame() ? ELocalizationLoadFlags::Game : ELocalizationLoadFlags::None);
 	LocLoadFlags |= ELocalizationLoadFlags::Engine;
 	LocLoadFlags |= ELocalizationLoadFlags::Native;
-	// We don't allow dynamically loaded additional data to replace any existing data, as additional localization is lower priority when performing a full update, so this flag emulates that behavior
-	LocLoadFlags |= ELocalizationLoadFlags::SkipExisting;
 
 	const TArray<FString> PrioritizedCultureNames = FInternationalization::Get().GetPrioritizedCultureNames(FInternationalization::Get().GetCurrentLanguage()->GetName());
 
-	LoadLocalizationTargetsForPrioritizedCultures_Async(LocalizationTargetPaths, PrioritizedCultureNames, LocLoadFlags, MoveTemp(Notification));
+	QueueAsyncTask([AvailableTextSources = LocalizedTextSources, LocalizationTargetPaths = TArray<FString>(LocalizationTargetPaths), PrioritizedCultureNames = TArray<FString>(PrioritizedCultureNames), LocLoadFlags, Notification = MoveTemp(Notification)]()
+	{
+		for (const FString& LocalizationTargetPath : LocalizationTargetPaths)
+		{
+			if (!IFileManager::Get().DirectoryExists(*LocalizationTargetPath))
+			{
+				continue;
+			}
+
+			FTextLocalizationMetaDataResource LocMetaResource;
+			{
+				const FString LocMetaFilename = FPaths::GetBaseFilename(LocalizationTargetPath) + TEXT(".locmeta");
+				LocMetaResource.LoadFromFile(LocalizationTargetPath / LocMetaFilename);
+			}
+
+			// We don't allow dynamically loaded UGC localization data to replace any existing data
+			const ELocalizationLoadFlags FinalLocLoadFlags = LocLoadFlags | (LocMetaResource.bIsUGC ? ELocalizationLoadFlags::SkipExisting : ELocalizationLoadFlags::None);
+			FTextLocalizationManager::Get().LoadLocalizationTargetsForPrioritizedCultures_Sync(AvailableTextSources, MakeArrayView(&LocalizationTargetPath, 1), PrioritizedCultureNames, FinalLocLoadFlags);
+		}
+		if (Notification)
+		{
+			Notification();
+		}
+	});
 }
 
 void FTextLocalizationManager::HandleLocalizationTargetsUnmounted(TArrayView<const FString> LocalizationTargetPaths, TUniqueFunction<void()>&& Notification)
diff --git a/Engine/Source/Runtime/Core/Private/Internationalization/TextLocalizationResource.cpp b/Engine/Source/Runtime/Core/Private/Internationalization/TextLocalizationResource.cpp
--- a/Engine/Source/Runtime/Core/Private/Internationalization/TextLocalizationResource.cpp
+++ b/Engine/Source/Runtime/Core/Private/Internationalization/TextLocalizationResource.cpp
@@ -107,6 +107,15 @@ bool FTextLocalizationMetaDataResource::LoadFromArchive(FArchive& Archive, const
 		CompiledCultures.Reset();
 	}
 
+	if (VersionNumber >= FTextLocalizationResourceVersion::ELocMetaVersion::AddedIsUGC)
+	{
+		Archive << bIsUGC;
+	}
+	else
+	{
+		bIsUGC = false;
+	}
+
 	return true;
 }
 
@@ -142,6 +151,9 @@ bool FTextLocalizationMetaDataResource::SaveToArchive(FArchive& Archive, const F
 
 		// Added by version: AddedCompiledCultures
 		Archive << CompiledCultures;
+
+		// Added by version: AddedIsUGC
+		Archive << bIsUGC;
 	}
 
 	return true;
diff --git a/Engine/Source/Runtime/Core/Public/Internationalization/TextLocalizationResource.h b/Engine/Source/Runtime/Core/Public/Internationalization/TextLocalizationResource.h
--- a/Engine/Source/Runtime/Core/Public/Internationalization/TextLocalizationResource.h
+++ b/Engine/Source/Runtime/Core/Public/Internationalization/TextLocalizationResource.h
@@ -35,6 +35,9 @@ public:
 	/** Name of all cultures with compiled LocRes files for the localization target (@note may be empty for older LocMeta files). */
 	TArray<FString> CompiledCultures;
 
+	/** True if this LocMeta was generated from UGC localization data (@see UserGeneratedContentLocalization) */
+	bool bIsUGC = false;
+
 	/** Load the given LocMeta file into this resource. */
 	CORE_API bool LoadFromFile(const FString& FilePath);
 
diff --git a/Engine/Source/Runtime/Core/Public/Internationalization/TextLocalizationResourceVersion.h b/Engine/Source/Runtime/Core/Public/Internationalization/TextLocalizationResourceVersion.h
--- a/Engine/Source/Runtime/Core/Public/Internationalization/TextLocalizationResourceVersion.h
+++ b/Engine/Source/Runtime/Core/Public/Internationalization/TextLocalizationResourceVersion.h
@@ -28,6 +28,8 @@ struct FTextLocalizationResourceVersion
 		Initial = 0,
 		/** Added complete list of cultures compiled for the localization target. */
 		AddedCompiledCultures,
+		/** Added bIsUGC flag */
+		AddedIsUGC,
 
 		LatestPlusOne,
 		Latest = LatestPlusOne - 1
