- Added the beginnings of FPosixOSPlatformProcess, which Linux subclasses from, and Mac pulls in as needed (mainly due to IOS not being in the PosixOS PlatformGroup)
- Moved IsFirstInstance from UnixPlatformProcess to PosixOS, and set Mac up to use it

diff --git a/Engine/Source/Runtime/Core/Private/Mac/MacPlatformMisc.cpp b/Engine/Source/Runtime/Core/Private/Mac/MacPlatformMisc.cpp
--- a/Engine/Source/Runtime/Core/Private/Mac/MacPlatformMisc.cpp
+++ b/Engine/Source/Runtime/Core/Private/Mac/MacPlatformMisc.cpp
@@ -638,6 +638,7 @@ void FMacPlatformMisc::PlatformTearDown()
 		
 		[GMacAppInfo.StdErrPipe release];
 	}
+	FPlatformProcess::CeaseBeingFirstInstance();
 }
 
 void FMacPlatformMisc::SetEnvironmentVar(const TCHAR* InVariableName, const TCHAR* Value)
diff --git a/Engine/Source/Runtime/Core/Private/PosixOS/PosixOSPlatformProcess.cpp b/Engine/Source/Runtime/Core/Private/PosixOS/PosixOSPlatformProcess.cpp
new file mode 100644
--- /dev/null
+++ b/Engine/Source/Runtime/Core/Private/PosixOS/PosixOSPlatformProcess.cpp
@@ -0,0 +1,59 @@
+// Copyright Epic Games, Inc. All Rights Reserved.
+
+#include "PosixOS/PosixOSPlatformProcess.h"
+#include <sys/fcntl.h>
+#include <sys/file.h>
+
+
+static int GFileLockDescriptor = -1;
+
+bool FPosixOSPlatformProcess::IsFirstInstance()
+{
+	// set default return if we are unable to access lock file.
+	static bool bIsFirstInstance = false;
+	static bool bNeverFirst = FParse::Param(FCommandLine::Get(), TEXT("neverfirst"));
+
+	if (!bIsFirstInstance && !bNeverFirst)	// once we determined that we're first, this can never change until we exit; otherwise, we re-check each time
+	{
+		// create the file if it doesn't exist
+		if (GFileLockDescriptor == -1)
+		{
+			FString LockFileName(TEXT("/tmp/"));
+			FString ExecPath(FPlatformProcess::ExecutableName());
+			ExecPath.ReplaceInline(TEXT("/"), TEXT("-"), ESearchCase::CaseSensitive);
+			// [RCL] 2015-09-20: can run out of filename limits (256 bytes) due to a long path, be conservative and assume 4-char UTF-8 name like e.g. Japanese
+			ExecPath.RightInline(80, EAllowShrinking::No);
+
+			LockFileName += ExecPath;
+
+			GFileLockDescriptor = open(TCHAR_TO_UTF8(*LockFileName), O_RDWR | O_CREAT, 0666);
+		}
+
+		if (GFileLockDescriptor != -1)
+		{
+			if (flock(GFileLockDescriptor, LOCK_EX | LOCK_NB) == 0)
+			{
+				// lock file successfully locked by this process - no more checking if we're first!
+				bIsFirstInstance = true;
+			}
+			else
+			{
+				// we were unable to lock file. so some other process beat us to lock file.
+				bIsFirstInstance = false;
+			}
+		}
+	}
+
+	return bIsFirstInstance;
+}
+
+void FPosixOSPlatformProcess::CeaseBeingFirstInstance()
+{
+	if (GFileLockDescriptor != -1)
+	{
+		// may fail if we didn't have the lock
+		flock(GFileLockDescriptor, LOCK_UN | LOCK_NB);
+		close(GFileLockDescriptor);
+		GFileLockDescriptor = -1;
+	}
+}
diff --git a/Engine/Source/Runtime/Core/Private/Unix/UnixPlatformProcess.cpp b/Engine/Source/Runtime/Core/Private/Unix/UnixPlatformProcess.cpp
--- a/Engine/Source/Runtime/Core/Private/Unix/UnixPlatformProcess.cpp
+++ b/Engine/Source/Runtime/Core/Private/Unix/UnixPlatformProcess.cpp
@@ -2274,59 +2274,6 @@ FString FUnixPlatformProcess::FProcEnumInfo::GetName() const
 	return FPaths::GetCleanFilename(GetFullPath());
 }
 
-static int GFileLockDescriptor = -1;
-
-bool FUnixPlatformProcess::IsFirstInstance()
-{
-	// set default return if we are unable to access lock file.
-	static bool bIsFirstInstance = false;
-	static bool bNeverFirst = FParse::Param(FCommandLine::Get(), TEXT("neverfirst"));
-
-	if (!bIsFirstInstance && !bNeverFirst)	// once we determined that we're first, this can never change until we exit; otherwise, we re-check each time
-	{
-		// create the file if it doesn't exist
-		if (GFileLockDescriptor == -1)
-		{
-			FString LockFileName(TEXT("/tmp/"));
-			FString ExecPath(FPlatformProcess::ExecutableName());
-			ExecPath.ReplaceInline(TEXT("/"), TEXT("-"), ESearchCase::CaseSensitive);
-			// [RCL] 2015-09-20: can run out of filename limits (256 bytes) due to a long path, be conservative and assume 4-char UTF-8 name like e.g. Japanese
-			ExecPath.RightInline(80, EAllowShrinking::No);
-
-			LockFileName += ExecPath;
-
-			GFileLockDescriptor = open(TCHAR_TO_UTF8(*LockFileName), O_RDWR | O_CREAT, 0666);
-		}
-
-		if (GFileLockDescriptor != -1)
-		{
-			if (flock(GFileLockDescriptor, LOCK_EX | LOCK_NB) == 0)
-			{
-				// lock file successfully locked by this process - no more checking if we're first!
-				bIsFirstInstance = true;
-			}
-			else
-			{
-				// we were unable to lock file. so some other process beat us to lock file.
-				bIsFirstInstance = false;
-			}
-		}
-	}
-
-	return bIsFirstInstance;
-}
-
-void FUnixPlatformProcess::CeaseBeingFirstInstance()
-{
-	if (GFileLockDescriptor != -1)
-	{
-		// may fail if we didn't have the lock
-		flock(GFileLockDescriptor, LOCK_UN | LOCK_NB);
-		close(GFileLockDescriptor);
-		GFileLockDescriptor = -1;
-	}
-}
-
 void FUnixPlatformProcess::OnChildEndFramePostFork()
 {
 	FCoreDelegates::OnEndFrame.Remove(OnEndFrameHandle);
@@ -2401,4 +2348,4 @@ void FUnixPlatformProcess::SetThreadPriority(EThreadPriority NewPriority)
 	int32 NiceValue = FUnixPlatformProcess::TranslateThreadPriority(NewPriority);
 
 	FUnixPlatformProcess::SetThreadNiceValue(ThreadId, NiceValue);
-}
\ No newline at end of file
+}
diff --git a/Engine/Source/Runtime/Core/Public/Mac/MacPlatformProcess.h b/Engine/Source/Runtime/Core/Public/Mac/MacPlatformProcess.h
--- a/Engine/Source/Runtime/Core/Public/Mac/MacPlatformProcess.h
+++ b/Engine/Source/Runtime/Core/Public/Mac/MacPlatformProcess.h
@@ -7,6 +7,7 @@
 #pragma once
 
 #include "Apple/ApplePlatformProcess.h"
+#include "PosixOS/PosixOSPlatformProcess.h"
 #include <sys/sysctl.h>
 
 /** Wrapper around Unix pid_t. Should not be copyable as changes in the process state won't be properly propagated to all copies. */
@@ -230,6 +231,17 @@ struct CORE_API FMacPlatformProcess : public FApplePlatformProcess
 	static bool WritePipe(void* WritePipe, const uint8* Data, const int32 DataLength, int32* OutDataLength = nullptr);
 	static bool ExecProcess(const TCHAR* URL, const TCHAR* Params, int32* OutReturnCode, FString* OutStdOut, FString* OutStdErr, const TCHAR* OptionalWorkingDirectory = NULL, bool bShouldEndWithParentProcess = false);
 	static void SetThreadAffinityMask(uint64 AffinityMask);
+	
+	/** Mixin some PosixOS functions */
+	static inline bool IsFirstInstance()
+	{
+		return FPosixOSPlatformProcess::IsFirstInstance();
+	}
+	static void CeaseBeingFirstInstance()
+	{
+		FPosixOSPlatformProcess::CeaseBeingFirstInstance();
+	}
+	
 
 	// Mac specific
 	static const TCHAR* UserPreferencesDir();
diff --git a/Engine/Source/Runtime/Core/Public/PosixOS/PosixOSPlatformProcess.h b/Engine/Source/Runtime/Core/Public/PosixOS/PosixOSPlatformProcess.h
new file mode 100644
--- /dev/null
+++ b/Engine/Source/Runtime/Core/Public/PosixOS/PosixOSPlatformProcess.h
@@ -0,0 +1,18 @@
+// Copyright Epic Games, Inc. All Rights Reserved.
+
+
+#pragma once
+
+#include "GenericPlatform/GenericPlatformProcess.h"
+
+
+// This class can be used as a "mixin" or subclassed from - for instance
+// Linux can fully use (subclass from) this but Mac may not have all Posix 
+// functionality, and it subclasses from Apple, which is shared with IOS 
+// (which isn't in the PosixOS platform group)
+struct FPosixOSPlatformProcess : public FGenericPlatformProcess
+{
+	/** FPlatformProcess API */
+	static CORE_API bool IsFirstInstance();
+	static CORE_API void CeaseBeingFirstInstance();
+};
diff --git a/Engine/Source/Runtime/Core/Public/Unix/UnixPlatformProcess.h b/Engine/Source/Runtime/Core/Public/Unix/UnixPlatformProcess.h
--- a/Engine/Source/Runtime/Core/Public/Unix/UnixPlatformProcess.h
+++ b/Engine/Source/Runtime/Core/Public/Unix/UnixPlatformProcess.h
@@ -9,7 +9,7 @@
 #include "Containers/Array.h"
 #include "Containers/UnrealString.h"
 #include "CoreTypes.h"
-#include "GenericPlatform/GenericPlatformProcess.h" // IWYU pragma: export
+#include "PosixOS/PosixOSPlatformProcess.h" // IWYU pragma: export
 #include "HAL/PlatformMisc.h"
 #include "Misc/AssertionMacros.h"
 #include "Unix/UnixSystemIncludes.h" // IWYU pragma: export
@@ -205,7 +205,7 @@ protected:
 /**
  * Unix implementation of the Process OS functions
  */
-struct FUnixPlatformProcess : public FGenericPlatformProcess
+struct FUnixPlatformProcess : public FPosixOSPlatformProcess
 {
 	struct FProcEnumInfo;
 
@@ -311,17 +311,12 @@ struct FUnixPlatformProcess : public FGenericPlatformProcess
 	static CORE_API bool ExecProcess(const TCHAR* URL, const TCHAR* Params, int32* OutReturnCode, FString* OutStdOut, FString* OutStdErr, const TCHAR* OptionalWorkingDirectory = NULL, bool bShouldEndWithParentProcess = false);
 	static CORE_API void ExploreFolder( const TCHAR* FilePath );
 	static CORE_API bool LaunchFileInDefaultExternalApplication( const TCHAR* FileName, const TCHAR* Parms = NULL, ELaunchVerb::Type Verb = ELaunchVerb::Open, bool bPromptToOpenOnFailure = true);
-	static CORE_API bool IsFirstInstance();
 	static CORE_API void OnChildEndFramePostFork();
 	static CORE_API int32 TranslateThreadPriority(EThreadPriority Priority);
 	static CORE_API void SetThreadNiceValue( uint32_t ThreadId, int32 NiceValue );
 	static CORE_API void SetThreadPriority( EThreadPriority NewPriority );
 	static CORE_API bool GetPerFrameProcessorUsage(uint32 ProcessId, float& ProcessUsageFraction, float& IdleUsageFraction);
 
-	/**
-	 * @brief Releases locks that we held for IsFirstInstance check
-	 */
-	static CORE_API void CeaseBeingFirstInstance();
 
 	/**
 	 * @brief Returns user home directory (i.e. $HOME).
