[Mobile] Resetting our support for HWCPIPE to use the new GPU counters.

#tests Tested  to ensure GPU stats are reporting correctly.
#rnx

diff --git a/Engine/Source/Runtime/Core/Private/Android/AndroidStats.cpp b/Engine/Source/Runtime/Core/Private/Android/AndroidStats.cpp
--- a/Engine/Source/Runtime/Core/Private/Android/AndroidStats.cpp
+++ b/Engine/Source/Runtime/Core/Private/Android/AndroidStats.cpp
@@ -11,8 +11,9 @@
 DEFINE_LOG_CATEGORY_STATIC(LogAndroidStats, Log, Log);
 
 #if !UE_BUILD_SHIPPING && PLATFORM_ANDROID_ARM64
-//#	define HWCPIPE_SUPPORTED 1 Disabling for now as we are seeing multiple missing frames since 36.00, starting since CL 42524925
-#	define HWCPIPE_SUPPORTED 0
+#	define HWCPIPE_SUPPORTED 1 
+//#	define HWCPIPE_SUPPORTED 0 No longer disabling for now; we were seeing multiple missing frames since 36.00, starting since CL 42524925
+//                             Keep this in place in-case we need to fix this again at a later time.
 #else
 #	define HWCPIPE_SUPPORTED 0
 #endif
