Changed platform I/O dispatcher from using a fixed set to a dynamic set of priority queues.
Also added secondary support thread to better mimic the existing I/O dispatcher implementation.

diff --git a/Engine/Source/Runtime/Core/Internal/IO/GenericPlatformIoDispatcher.cpp b/Engine/Source/Runtime/Core/Internal/IO/GenericPlatformIoDispatcher.cpp
--- a/Engine/Source/Runtime/Core/Internal/IO/GenericPlatformIoDispatcher.cpp
+++ b/Engine/Source/Runtime/Core/Internal/IO/GenericPlatformIoDispatcher.cpp
@@ -112,10 +112,10 @@ FIoStatus FGenericPlatformIoDispatcher::OnInitialize()
 	UE_LOG(LogPlatformIoDispatcher, Log, TEXT("Initializing, Platform='%s', ReadSize=%dKB, ReadMemory=%dMB, FileCache=%dMB, RequestSorting=%s, MaxConcurrentDecodeJobs=%d"),
 		TEXT("Generic"), GIoDispatcherBufferSizeKB, GIoDispatcherBufferMemoryMB, GIoDispatcherCacheSizeMB, bSortByOffset ? TEXT("ByOffset") : TEXT("BySeqNo"), GIoDispatcherDecompressionWorkerCount);
 
-	ChunkBlockMemoryPool.Initialize(GIoDispatcherDecompressionWorkerCount, DefaultCompressionBlockSize);
+	ChunkBlockMemoryPool.Initialize(GIoDispatcherDecompressionWorkerCount + 1, DefaultCompressionBlockSize);
 	FileBlockMemoryPool.Initialize(FileBlockSize, GIoDispatcherBufferMemoryMB << 20, GIoDispatcherBufferAlignment); 
 	FileBlockCache.Initialize(FileBlockSize, uint64(GIoDispatcherCacheSizeMB) << 20);
-	FileBlockRequestQueue.SortByOffset(bSortByOffset);
+	IoQueue.SortByOffset(bSortByOffset);
 
 	return FIoStatus::Ok;
 }
@@ -173,22 +173,26 @@ void FGenericPlatformIoDispatcher::OnWakeUp()
 
 bool FGenericPlatformIoDispatcher::ProcessFileIo(bool& bOutWaitingForMemory)
 {
-	bOutWaitingForMemory	= false;
-	FIoFileBlockRequest* FileBlockRequest = nullptr;
+	bOutWaitingForMemory					= false;
+	FIoFileBlockRequest* FileBlockRequest	= nullptr;
 	{
-		UE::TUniqueLock Lock(FileBlockMutex);
-		if (FIoFileBlockRequest *NextPending = FileBlockRequestQueue.Peek(); NextPending != nullptr)
+		UE::TUniqueLock QueueLock(IoQueue);
+		if (FIoFileBlockRequest* NextPending = IoQueue.Peek(); NextPending != nullptr)
 		{
 			if (NextPending->ErrorCode == EIoErrorCode::Cancelled)
 			{
-				FileBlockRequest = FileBlockRequestQueue.Dequeue();
+				FileBlockRequest = IoQueue.Dequeue();
 			}
 			else
 			{
-				if (bOutWaitingForMemory = FileBlockMemoryPool.IsEmpty(); bOutWaitingForMemory == false)
+				NextPending->Buffer = FileBlockMemoryPool.Alloc(NextPending->BufferHandle);
+				if (NextPending->BufferHandle.IsValid())
 				{
-					FileBlockRequest			= FileBlockRequestQueue.Dequeue();
-					FileBlockRequest->Buffer	= FileBlockMemoryPool.Alloc(FileBlockRequest->BufferHandle);
+					FileBlockRequest = IoQueue.Dequeue();
+				}
+				else
+				{
+					bOutWaitingForMemory = true;
 				}
 			}
 		}
@@ -234,7 +238,7 @@ bool FGenericPlatformIoDispatcher::ProcessFileIo(bool& bOutWaitingForMemory)
 		}
 	}
 
-	ProcessCompletedFileBlock(*FileBlockRequest);
+	EnqueueCompletedFileBlock(*FileBlockRequest);
 
 	return true;
 }
diff --git a/Engine/Source/Runtime/Core/Internal/IO/PlatformIoDispatcherBase.cpp b/Engine/Source/Runtime/Core/Internal/IO/PlatformIoDispatcherBase.cpp
--- a/Engine/Source/Runtime/Core/Internal/IO/PlatformIoDispatcherBase.cpp
+++ b/Engine/Source/Runtime/Core/Internal/IO/PlatformIoDispatcherBase.cpp
@@ -3,7 +3,9 @@
 #include "PlatformIoDispatcherBase.h"
 
 #include "Algo/BinarySearch.h"
+#include "Algo/IsSorted.h"
 #include "Async/UniqueLock.h"
+#include "HAL/Event.h"
 #include "HAL/LowLevelMemTracker.h"
 #include "IO/IoChunkEncoding.h"
 #include "IO/IoStatus.h"
@@ -28,19 +30,12 @@ bool IsIoErrorAnyOf(EIoErrorCode ErrorCode, Args... Expected)
 uint32 FIoFileBlockRequest::NextSeqNo = 0;
 
 ////////////////////////////////////////////////////////////////////////////////
-FIoFileBlockRequestQueue::FIoFileBlockRequestQueue(FPlatformIoDispatcherStats& InStats)
+FIoQueue::FIoQueue(FPlatformIoDispatcherStats& InStats)
 	: Stats(InStats)
 {
-	if (bSortByOffset)
-	{
-		for (int32 Idx = static_cast<int32>(EIoFileReadPriority::Count) - 1; Idx >= 0; --Idx)
-		{
-			PrioQueues[Idx].ByOffset.Reserve(128);
-		}
-	}
 }
 
-void FIoFileBlockRequestQueue::Enqueue(FIoFileBlockRequestList&& FileBlockRequests)
+void FIoQueue::Enqueue(FIoFileBlockRequestList&& FileBlockRequests)
 {
 	while (FIoFileBlockRequest* Request = FileBlockRequests.PopHead())
 	{
@@ -48,13 +43,7 @@ void FIoFileBlockRequestQueue::Enqueue(FIoFileBlockRequestList&& FileBlockReques
 		Request->QueueStatus = FIoFileBlockRequest::EQueueStatus::Enqueued;
 		if (bSortByOffset)
 		{
-			const EIoFileReadPriority Priority	= IoFileReadPriorityFromDispatcherPriority(Request->Priority);
-			FPrioQueue& PrioQueue				= PrioQueues[static_cast<int32>(Priority)];
-			const int32 Idx						= Algo::UpperBoundBy(PrioQueue.ByOffset, ToFileOffset(Request), ToFileOffset, FileOffsetLess);
-
-			PrioQueue.ByOffset.Insert(Request, Idx);
-			PrioQueue.BySequence.AddTail(Request);
-			PrioQueue.PeekIndex = INDEX_NONE;
+			AddToPrioQueue(*Request);
 		}
 		else
 		{
@@ -65,7 +54,7 @@ void FIoFileBlockRequestQueue::Enqueue(FIoFileBlockRequestList&& FileBlockReques
 	}
 }
 
-FIoFileBlockRequest* FIoFileBlockRequestQueue::Dequeue()
+FIoFileBlockRequest* FIoQueue::Dequeue()
 {
 	const bool bDequeue = true;
 	FIoFileBlockRequest* Request = GetCancelled(bDequeue);
@@ -90,7 +79,7 @@ FIoFileBlockRequest* FIoFileBlockRequestQueue::Dequeue()
 	return Request;
 }
 
-FIoFileBlockRequest* FIoFileBlockRequestQueue::Peek()
+FIoFileBlockRequest* FIoQueue::Peek()
 {
 	const bool bDequeue = false;
 	FIoFileBlockRequest* Request = GetCancelled(bDequeue);
@@ -109,24 +98,21 @@ FIoFileBlockRequest* FIoFileBlockRequestQueue::Peek()
 	return Request;
 }
 
-void FIoFileBlockRequestQueue::Reprioritize()
+void FIoQueue::Reprioritize()
 {
 	TRACE_CPUPROFILER_EVENT_SCOPE(RequestQueueUpdateOrder);
 
 	if (bSortByOffset)
 	{
 		TArray<FIoFileBlockRequest*> RequestsToReprioritize;
-		for (int32 QueueIdx = 0; FPrioQueue& PrioQueue : PrioQueues)
+		for (FPrioQueue& PrioQueue : PrioQueues)
 		{
-			PrioQueue.PeekIndex						= INDEX_NONE;
-			const EIoFileReadPriority QueuePriority = static_cast<EIoFileReadPriority>(QueueIdx++);
+			PrioQueue.PeekIndex = INDEX_NONE;
 
 			for (int32 RequestIdx = PrioQueue.ByOffset.Num() - 1; RequestIdx >= 0; --RequestIdx)
 			{
-				FIoFileBlockRequest* Request		= PrioQueue.ByOffset[RequestIdx];
-				const EIoFileReadPriority Priority	= IoFileReadPriorityFromDispatcherPriority(Request->Priority);
-
-				if (Priority != QueuePriority)
+				FIoFileBlockRequest* Request = PrioQueue.ByOffset[RequestIdx];
+				if (Request->Priority != PrioQueue.Priority)
 				{
 					RequestsToReprioritize.Add(Request);
 					PrioQueue.ByOffset.RemoveAt(RequestIdx, EAllowShrinking::No);
@@ -135,15 +121,11 @@ void FIoFileBlockRequestQueue::Reprioritize()
 			}
 		}
 
+		PrioQueues.RemoveAll([](FPrioQueue& PrioQueue) { return PrioQueue.ByOffset.IsEmpty(); });
+
 		for (FIoFileBlockRequest* Request : RequestsToReprioritize)
 		{
-			const EIoFileReadPriority Priority	= IoFileReadPriorityFromDispatcherPriority(Request->Priority);
-			FPrioQueue& PrioQueue				= PrioQueues[static_cast<int32>(Priority)];
-			const int32 Idx						= Algo::UpperBoundBy(PrioQueue.ByOffset, ToFileOffset(Request), ToFileOffset, FileOffsetLess);
-
-			PrioQueue.ByOffset.Insert(Request, Idx);
-			PrioQueue.BySequence.AddTail(Request);
-			PrioQueue.PeekIndex = INDEX_NONE;
+			AddToPrioQueue(*Request);
 		}
 	}
 	else if (Heap.IsEmpty() == false)
@@ -152,21 +134,43 @@ void FIoFileBlockRequestQueue::Reprioritize()
 	}
 }
 
-FIoFileBlockRequest* FIoFileBlockRequestQueue::GetByOffset(bool bDequeue)
+void FIoQueue::AddToPrioQueue(FIoFileBlockRequest& Request)
+{
+	const int32 QueueIndex = Algo::LowerBoundBy(PrioQueues, Request.Priority, QueueToPriority, TLess<int32>());
+	if (!PrioQueues.IsValidIndex(QueueIndex) || PrioQueues[QueueIndex].Priority != Request.Priority)
+	{
+		PrioQueues.Insert(FPrioQueue(Request.Priority), QueueIndex);
+	}
+
+	check(PrioQueues.IsValidIndex(QueueIndex));
+	check(Algo::IsSortedBy(PrioQueues, QueueToPriority, TLess<int32>()));
+	FPrioQueue& PrioQueue	= PrioQueues[QueueIndex];
+	const int32 Idx			= Algo::UpperBoundBy(PrioQueue.ByOffset, ToFileOffset(&Request), ToFileOffset, FileOffsetLess);
+
+	check(PrioQueue.Priority == Request.Priority);
+	PrioQueue.ByOffset.Insert(&Request, Idx);
+	PrioQueue.BySequence.AddTail(&Request);
+	PrioQueue.PeekIndex = INDEX_NONE;
+}
+
+FIoFileBlockRequest* FIoQueue::GetByOffset(bool bDequeue)
 {
-	for (int32 Idx = static_cast<int32>(EIoFileReadPriority::Count) - 1; Idx >= 0; --Idx)
+	FIoFileBlockRequest* Request = nullptr;
+	if (!PrioQueues.IsEmpty()) 
 	{
-		FPrioQueue& PrioQueue = PrioQueues[Idx];
-		if (PrioQueue.BySequence.IsEmpty() == false)
+		FPrioQueue& PrioQueue	= PrioQueues.Last();
+		Request					= GetByOffset(PrioQueue, bDequeue);
+
+		if (PrioQueue.ByOffset.IsEmpty())
 		{
-			return GetByOffset(PrioQueue, static_cast<EIoFileReadPriority>(Idx), bDequeue);
+			PrioQueues.Pop();
 		}
 	}
 
-	return nullptr;
+	return Request;
 }
 
-FIoFileBlockRequest* FIoFileBlockRequestQueue::GetByOffset(FPrioQueue& PrioQueue, EIoFileReadPriority QueuePriority, bool bDequeue)
+FIoFileBlockRequest* FIoQueue::GetByOffset(FPrioQueue& PrioQueue, bool bDequeue)
 {
 	check(!PrioQueue.ByOffset.IsEmpty() && !PrioQueue.BySequence.IsEmpty());
 
@@ -185,11 +189,10 @@ FIoFileBlockRequest* FIoFileBlockRequestQueue::GetByOffset(FPrioQueue& PrioQueue
 			bHeadRequestTooOld = (FPlatformTime::Cycles64() - PrioQueue.BySequence.PeekHead()->TimeCreated) >= ThresholdCycles;
 		}
 
-		const EIoFileReadPriority LastFileOffsetPriority = IoFileReadPriorityFromDispatcherPriority(LastFileOffset.Priority);
 		const bool bChooseByOffset = 
 				LastFileOffset.FileHandle.IsValid()
 			&&	!bHeadRequestTooOld 
-			&&  (GIoDispatcherMaintainSortingOnPriorityChange || LastFileOffsetPriority == QueuePriority);
+			&&  (GIoDispatcherMaintainSortingOnPriorityChange || LastFileOffset.Priority == PrioQueue.Priority);
 		if (bChooseByOffset)
 		{
 			RequestIndex = Algo::LowerBoundBy(PrioQueue.ByOffset, LastFileOffset, ToFileOffset, FileOffsetLess);
@@ -215,10 +218,9 @@ FIoFileBlockRequest* FIoFileBlockRequestQueue::GetByOffset(FPrioQueue& PrioQueue
 	if (bDequeue)
 	{
 		PrioQueue.ByOffset.RemoveAt(RequestIndex, EAllowShrinking::No);
-		const bool bRemoved = PrioQueue.BySequence.Remove(Request);
+		PrioQueue.BySequence.Remove(Request);
 		PrioQueue.PeekIndex = INDEX_NONE;
 		LastFileOffset		= ToFileOffset(Request);
-		check(bRemoved);
 	}
 	else
 	{
@@ -228,7 +230,7 @@ FIoFileBlockRequest* FIoFileBlockRequestQueue::GetByOffset(FPrioQueue& PrioQueue
 	return Request;
 }
 
-FIoFileBlockRequest* FIoFileBlockRequestQueue::GetCancelled(bool bDequeue)
+FIoFileBlockRequest* FIoQueue::GetCancelled(bool bDequeue)
 {
 	if (bReprioritizeCancelled)
 	{
@@ -236,7 +238,7 @@ FIoFileBlockRequest* FIoFileBlockRequestQueue::GetCancelled(bool bDequeue)
 		TRACE_CPUPROFILER_EVENT_SCOPE(RequestQueueRemoveCancelled);
 		if (bSortByOffset)
 		{
-			for (int32 QueueIdx = 0; FPrioQueue& PrioQueue : PrioQueues)
+			for (FPrioQueue& PrioQueue : PrioQueues)
 			{
 				PrioQueue.PeekIndex = INDEX_NONE;
 				for (int32 RequestIdx = PrioQueue.ByOffset.Num() - 1; RequestIdx >= 0; --RequestIdx)
@@ -250,6 +252,8 @@ FIoFileBlockRequest* FIoFileBlockRequestQueue::GetCancelled(bool bDequeue)
 					}
 				}
 			}
+
+			PrioQueues.RemoveAll([](FPrioQueue& PrioQueue) { return PrioQueue.ByOffset.IsEmpty(); });
 		}
 		else
 		{
@@ -317,11 +321,15 @@ void FIoFileBlockMemoryPool::Initialize(uint32 FileBlockSize, uint32 FilePoolSiz
 
 void* FIoFileBlockMemoryPool::Alloc(FIoBufferHandle& OutHandle)
 {
+	TUniqueLock Lock(Mutex);
+
 	if (FMemoryBlock* Block = FreeBlock)
 	{
-		FreeBlock	= Block->Next;
-		Block->Next	= nullptr;
-		Block->RefCount.store(1, std::memory_order_relaxed);
+		check(Block->RefCount == 0);
+
+		FreeBlock		= Block->Next;
+		Block->Next		= nullptr;
+		Block->RefCount = 1;
 		OutHandle	= FIoBufferHandle(Block->Index);
 
 		Stats.OnFileBlockMemoryAllocated(BlockSize);
@@ -332,38 +340,39 @@ void* FIoFileBlockMemoryPool::Alloc(FIoBufferHandle& OutHandle)
 	return nullptr;
 }
 
-bool FIoFileBlockMemoryPool::IsEmpty() const
+bool FIoFileBlockMemoryPool::Free(FIoBufferHandle& Handle)
 {
-	return FreeBlock == nullptr;
-}
+	bool bFreed = false;
+	TUniqueLock Lock(Mutex);
 
-void FIoFileBlockMemoryPool::Free(FIoBufferHandle& Handle)
-{
 	if (Handle.IsValid() == false)
 	{
-		return;
+		return bFreed;
 	}
 
 	FMemoryBlock& Block = MemoryBlocks[Handle.Value()];
 	check(Block.Next == nullptr);
 	check(Block.RefCount > 0);
 
-	if (1 == Block.RefCount.fetch_sub(1, std::memory_order_relaxed))
+	if (--Block.RefCount == 0)
 	{
 		Block.Next	= FreeBlock;
 		FreeBlock	= &Block;
+		bFreed		= true;
 
 		Stats.OnFileBlockMemoryFreed(BlockSize);
 	}
 
 	Handle = FIoBufferHandle();
+	return bFreed;
 }
 
 void FIoFileBlockMemoryPool::AddRef(FIoBufferHandle Handle)
 {
+	TUniqueLock Lock(Mutex);
 	check(Handle.IsValid());
 	FMemoryBlock& Block = MemoryBlocks[Handle.Value()];
-	Block.RefCount.fetch_add(1, std::memory_order_relaxed);
+	Block.RefCount++;
 }
 
 ////////////////////////////////////////////////////////////////////////////////
@@ -504,7 +513,7 @@ bool FIoFileBlockCache::Get(FIoFileBlockRequest& FileBlockRequest)
 
 void FIoFileBlockCache::Put(FIoFileBlockRequest& FileBlockRequest)
 {
-	bool bIsCacheableBlock = CacheMemory != nullptr;// && Block->BytesUsed < Block->Size;
+	bool bIsCacheableBlock = CacheMemory != nullptr && FileBlockRequest.BytesUsed < FileBlockRequest.Size;
 	if (!bIsCacheableBlock)
 	{
 		return;
@@ -602,10 +611,10 @@ bool FPlatformIoDispatcherRequestMgr::TryCancelReadRequest(FIoPlatformReadReques
 				}
 
 				// Cancel the file block request if all encoded block requests has been cancelled
-				bool bCancelFileBlock		= true;
-				bool bCancelEncodedBlock	= true;
+				bool bCancelFileBlock = true;
 				for (FIoEncodedBlockRequest* EncodedBlockRequest : FileBlockRequest->EncodedBlockRequests)
 				{
+					bool bCancelEncodedBlock = true;
 					for (FIoEncodedBlockRequest::FScatterTarget& ScatterTarget : EncodedBlockRequest->ScatterTargets)
 					{
 						if (ScatterTarget.Request == &ReadRequest)
@@ -735,6 +744,11 @@ void FPlatformIoDispatcherRequestMgr::AddFileBlockRequest(FIoPlatformReadRequest
 	AddToLink(&ReadRequest.ScatterGather.FileBlockLink, FileBlockRequest.BlockKey);
 }
 
+void FPlatformIoDispatcherRequestMgr::AddFileBlockRequest(FIoPlatformReadRequest& ReadRequest, FIoBlockKey FileBlockKey)
+{
+	AddToLink(&ReadRequest.ScatterGather.FileBlockLink, FileBlockKey);
+}
+
 FIoEncodedBlockRequest& FPlatformIoDispatcherRequestMgr::GetOrCreateEncodedBlockRequest(FIoBlockKey BlockKey, bool& bCreated)
 {
 	check(BlockKey.IsValid());
@@ -758,6 +772,10 @@ void FPlatformIoDispatcherRequestMgr::Remove(FIoEncodedBlockRequest& Request)
 	if (Request.ErrorCode != EIoErrorCode::Cancelled)
 	{
 		EncodedBlockLookup.Remove(Request.BlockKey);
+		if (EncodedBlockLookup.IsEmpty())
+		{
+			EncodedBlockLookup.Empty(128);
+		}
 	}
 }
 
@@ -767,6 +785,10 @@ void FPlatformIoDispatcherRequestMgr::Destroy(FIoFileBlockRequest& Request)
 	if (Request.ErrorCode != EIoErrorCode::Cancelled)
 	{
 		FileBlockLookup.Remove(Request.BlockKey);
+		if (FileBlockLookup.IsEmpty())
+		{
+			FileBlockLookup.Empty(128);
+		}
 	}
 	FileBlockAllocator.Destroy(&Request);
 }
@@ -820,10 +842,55 @@ void FPlatformIoDispatcherRequestMgr::AddToLink(FIoFileBlockLink* Link, FIoBlock
 	}
 }
 
+////////////////////////////////////////////////////////////////////////////////
+class FIoServiceThread final
+	: public FRunnable
+{
+public:
+	typedef bool(FPlatformIoDispatcherBase::* FThreadEntry)();
+
+	FIoServiceThread(FPlatformIoDispatcherBase* InDispatcher, FThreadEntry InThreadEntry, const TCHAR* InName, EThreadPriority InPriority)
+		: Dispatcher(InDispatcher)
+		, ThreadEntry(InThreadEntry)
+		, Thread(FRunnableThread::Create(this, InName, 0, InPriority))
+	{ }
+
+	virtual void	Stop() override	{ bStopRequested.store(true, std::memory_order_seq_cst); WakeUpEvent->Trigger(); }
+	void			WakeUp() const	{ WakeUpEvent->Trigger(); }
+	void			Join() const	{ Thread->WaitForCompletion(); }
+
+private:
+	virtual uint32	Run() override;
+
+	FPlatformIoDispatcherBase*		Dispatcher;
+	FThreadEntry					ThreadEntry;
+	FEventRef						WakeUpEvent;
+	std::atomic_bool				bStopRequested{false};
+	TUniquePtr<FRunnableThread>		Thread;
+};
+
+////////////////////////////////////////////////////////////////////////////////
+uint32 FIoServiceThread::Run() 
+{
+	for (;;)
+	{
+		while((Dispatcher->*ThreadEntry)())
+		{ }
+
+		if (bStopRequested.load(std::memory_order_relaxed))
+		{
+			break;
+		}
+		WakeUpEvent->Wait();
+	}
+
+	return 0;
+}
+
 ////////////////////////////////////////////////////////////////////////////////
 FPlatformIoDispatcherBase::FPlatformIoDispatcherBase(FPlatformIoDispatcherCreateParams&& Params)
 	: FileBlockMemoryPool(Stats)
-	, FileBlockRequestQueue(Stats)
+	, IoQueue(Stats)
 	, FileBlockCache(Stats)
 {
 	bMultithreaded = Params.bMultithreaded;
@@ -833,6 +900,11 @@ FPlatformIoDispatcherBase::~FPlatformIoDispatcherBase()
 {
 	check(bStopRequested);
 	check(Thread.IsValid() == false);
+	if (SupportThread.IsValid())
+	{
+		SupportThread->Stop();
+		SupportThread->Join();
+	}
 }
 
 bool FPlatformIoDispatcherBase::Init()
@@ -848,6 +920,10 @@ uint32 FPlatformIoDispatcherBase::Run()
 void FPlatformIoDispatcherBase::Stop()
 {
 	bStopRequested = true;
+	if (SupportThread.IsValid())
+	{
+		SupportThread->Stop();
+	}
 }
 
 FIoStatus FPlatformIoDispatcherBase::Initialize()
@@ -855,7 +931,8 @@ FIoStatus FPlatformIoDispatcherBase::Initialize()
 	FIoStatus Status = OnInitialize();
 	if (Status.IsOk() && bMultithreaded)
 	{
-		Thread.Reset(FRunnableThread::Create(this, TEXT("IoService"), 0, TPri_AboveNormal));
+		Thread.Reset(FRunnableThread::Create(this, TEXT("IoService #1"), 0, TPri_AboveNormal));
+		SupportThread = MakeUnique<FIoServiceThread>(this, &FPlatformIoDispatcherBase::ProcessBlocks, TEXT("IoService #2"), TPri_BelowNormal);
 	}
 
 	return Status;
@@ -912,13 +989,13 @@ FIoFileReadRequest FPlatformIoDispatcherBase::ScatterGather(FIoScatterGatherRequ
 				{
 					const FIoBlockKey FileBlockKey			= FIoBlockKey(FileInfo.FileId, IntCastChecked<uint32>(FileBlockIndex));
 					FIoFileBlockRequest* FileBlockRequest	= RequestMgr.GetFileBlockRequest(FileBlockKey);
-					check(FileBlockRequest != nullptr);
-					RequestMgr.AddFileBlockRequest(ReadRequest, *FileBlockRequest);
-					if (Params.Priority > FileBlockRequest->Priority)
+
+					if (FileBlockRequest != nullptr && Params.Priority > FileBlockRequest->Priority)
 					{
 						FileBlockRequest->Priority	= Params.Priority;
 						bReprioritizeReadQueue 		= true;
 					}
+					RequestMgr.AddFileBlockRequest(ReadRequest, FileBlockKey);
 				}
 
 				continue;
@@ -954,6 +1031,15 @@ FIoFileReadRequest FPlatformIoDispatcherBase::ScatterGather(FIoScatterGatherRequ
 
 					CreatedFileBlockRequests.AddTail(&FileBlockRequest);
 				}
+				else if (Params.Priority > FileBlockRequest.Priority)
+				{
+					FileBlockRequest.Priority	= Params.Priority;
+					bReprioritizeReadQueue		= true;
+				}
+
+				FileBlockRequest.BytesUsed += 
+					uint32(FMath::Min(BlockRequest.FileOffset + Align(BlockRequest.BlockCompressedSize, FAES::AESBlockSize), FileBlockRequest.FileOffset + FileBlockRequest.FileSize) -
+						   FMath::Max(BlockRequest.FileOffset, FileBlockRequest.FileOffset));
 
 				check(BlockRequest.RemainingFileBlocks < MAX_uint8);
 				check(BlockRequest.FileBlockCount < MAX_uint8);
@@ -963,12 +1049,12 @@ FIoFileReadRequest FPlatformIoDispatcherBase::ScatterGather(FIoScatterGatherRequ
 		}
 	}
 	{
-		TUniqueLock Lock(FileBlockMutex);
+		TUniqueLock QueueLock(IoQueue);
 		if (bReprioritizeReadQueue)
 		{
-			FileBlockRequestQueue.Reprioritize();
+			IoQueue.Reprioritize();
 		}
-		FileBlockRequestQueue.Enqueue(MoveTemp(CreatedFileBlockRequests));
+		IoQueue.Enqueue(MoveTemp(CreatedFileBlockRequests));
 	}
 	OnWakeUp();
 
@@ -1003,7 +1089,7 @@ void FPlatformIoDispatcherBase::CancelRequest(FIoFileReadRequest Request)
 	FIoFileReadRequestCompleted OnCompleted; 
 	{
 		TUniqueLock MgrLock(RequestMgr);
-		TUniqueLock QueueLock(FileBlockMutex);
+		TUniqueLock QueueLock(IoQueue);
 
 		bool bAnyBlockCancelled = false;
 		if (RequestMgr.TryCancelReadRequest(ReadRequest, bAnyBlockCancelled))
@@ -1012,7 +1098,7 @@ void FPlatformIoDispatcherBase::CancelRequest(FIoFileReadRequest Request)
 		}
 		if (bAnyBlockCancelled)
 		{
-			FileBlockRequestQueue.ReprioritizeCancelled();
+			IoQueue.ReprioritizeCancelled();
 		}
 	}
 
@@ -1036,10 +1122,10 @@ void FPlatformIoDispatcherBase::CancelAllRequests(FIoFileHandle FileHandle)
 
 	{
 		TUniqueLock MgrLock(RequestMgr);
-		TUniqueLock QueueLock(FileBlockMutex);
+		TUniqueLock QueueLock(IoQueue);
 		if (RequestMgr.TryCancelAllReadRequests(FileHandle))
 		{
-			FileBlockRequestQueue.ReprioritizeCancelled();
+			IoQueue.ReprioritizeCancelled();
 		}
 	}
 }
@@ -1056,23 +1142,23 @@ void FPlatformIoDispatcherBase::UpdatePriority(FIoFileReadRequest Request, int32
 	bool bReprioritize = false;
 	{
 		TUniqueLock MgrLock(RequestMgr);
+		TUniqueLock QueueLock(IoQueue);
 
 		FTempArray<FIoFileBlockRequest*> FileBlockRequests;
 		RequestMgr.GetFileBlockRequests(ReadRequest, FileBlockRequests);
 
 		for (FIoFileBlockRequest* FileBlockRequest : FileBlockRequests)
 		{
-			if (NewPriority != FileBlockRequest->Priority)
+			if (NewPriority > FileBlockRequest->Priority)
 			{
 				FileBlockRequest->Priority	= NewPriority;
 				bReprioritize				= true;
 			}
 		}
-	}
-	if (bReprioritize)
-	{
-		TUniqueLock QueueLock(FileBlockMutex);
-		FileBlockRequestQueue.Reprioritize();
+		if (bReprioritize)
+		{
+			IoQueue.Reprioritize();
+		}
 	}
 }
 
@@ -1101,172 +1187,138 @@ uint32 FPlatformIoDispatcherBase::GetNextFileId()
 	return NextFileId.fetch_add(1);
 }
 
-void FPlatformIoDispatcherBase::ProcessCompletedFileBlock(FIoFileBlockRequest& FileBlockRequest)
+void FPlatformIoDispatcherBase::EnqueueCompletedFileBlock(FIoFileBlockRequest& FileBlockRequest)
+{
+	if (bMultithreaded)
+	{
+		{
+			TUniqueLock Lock(FileBlockMutex);
+			FileBlocksToComplete.AddTail(&FileBlockRequest);
+		}
+		SupportThread->WakeUp();
+	}
+	else
+	{
+		FileBlocksToComplete.AddTail(&FileBlockRequest);
+		while (ProcessBlocks())
+		{ }
+	}
+}
+
+void FPlatformIoDispatcherBase::CompleteFileBlock(FIoFileBlockRequest& FileBlockRequest)
 {
-	TRACE_CPUPROFILER_EVENT_SCOPE(ProcessCompletedFileBlock);
+	TRACE_CPUPROFILER_EVENT_SCOPE(CompleteeFileBlock);
 	check(IsIoErrorAnyOf(FileBlockRequest.ErrorCode, EIoErrorCode::Ok, EIoErrorCode::Cancelled, EIoErrorCode::ReadError));
 
+	const EIoErrorCode							FileBlockErrorCode = FileBlockRequest.ErrorCode;
 	const uint64								FileBlockOffset = FileBlockRequest.FileOffset;
 	FIoBufferHandle								FileBlockBufferHandle = FileBlockRequest.BufferHandle;
 	FMutableMemoryView							FileBlockBufferView(FileBlockRequest.Buffer, FileBlockRequest.Size);
 	FIoFileBlockRequest::FEncodedBlocksArray	EncodedBlockRequests;
+
+	Stats.OnFileBlockCompleted(FileBlockRequest);
 	{
 		TUniqueLock Lock(RequestMgr);
-		for (FIoEncodedBlockRequest* EncodedBlockRequest : FileBlockRequest.EncodedBlockRequests)
-		{
-			check(IsIoErrorAnyOf(EncodedBlockRequest->ErrorCode, EIoErrorCode::Unknown, EIoErrorCode::Cancelled, EIoErrorCode::ReadError));
-			RequestMgr.Remove(*EncodedBlockRequest);
-			if (FileBlockRequest.ErrorCode != EIoErrorCode::Ok && EncodedBlockRequest->ErrorCode == EIoErrorCode::Unknown)
-			{
-				EncodedBlockRequest->ErrorCode = FileBlockRequest.ErrorCode;
-			}
-		}
 		EncodedBlockRequests = MoveTemp(FileBlockRequest.EncodedBlockRequests);
-		Stats.OnFileBlockCompleted(FileBlockRequest);
 		RequestMgr.Destroy(FileBlockRequest);
 	}
 
 	FIoEncodedBlockRequestList	CancelledOrFailedBlocks;
-	FIoEncodedBlockRequestList	BlocksToDecode;
 	for (FIoEncodedBlockRequest* EncodedBlockRequest : EncodedBlockRequests)
 	{
+		check(IsIoErrorAnyOf(EncodedBlockRequest->ErrorCode, EIoErrorCode::Unknown, EIoErrorCode::Cancelled, EIoErrorCode::ReadError));
 		check(EncodedBlockRequest->RemainingFileBlocks > 0);
-		if (EncodedBlockRequest->ErrorCode != EIoErrorCode::Unknown)
+
+		if (FileBlockErrorCode != EIoErrorCode::Ok && EncodedBlockRequest->ErrorCode == EIoErrorCode::Unknown)
 		{
-			check(IsIoErrorAnyOf(EncodedBlockRequest->ErrorCode, EIoErrorCode::Cancelled, EIoErrorCode::ReadError));
-			if (--EncodedBlockRequest->RemainingFileBlocks == 0)
-			{
-				CancelledOrFailedBlocks.AddTail(EncodedBlockRequest);
-			}
-			continue;
+			EncodedBlockRequest->ErrorCode = FileBlockErrorCode;
 		}
 
-		const uint64 EncodedBlockFileSize	= Align(EncodedBlockRequest->BlockCompressedSize, FAES::AESBlockSize);
-		const int64 OffsetInFileBlock		= EncodedBlockRequest->FileOffset - FileBlockOffset;
-
-		// Is the encoded block crossing file block boundaries
-		if (EncodedBlockRequest->FileBlockCount > 1)
+		if (EncodedBlockRequest->ErrorCode == EIoErrorCode::Unknown)
 		{
-			// When crossing file blocks we need to allocate a temp buffer to hold the data
-			if (EncodedBlockRequest->EncodedData == nullptr)
+			const uint64 EncodedBlockFileSize	= Align(EncodedBlockRequest->BlockCompressedSize, FAES::AESBlockSize);
+			const int64 OffsetInFileBlock		= EncodedBlockRequest->FileOffset - FileBlockOffset;
+
+			// Is the encoded block crossing file block boundaries
+			if (EncodedBlockRequest->FileBlockCount > 1)
 			{
-				EncodedBlockRequest->EncodedData = FMemory::Malloc(EncodedBlockFileSize);
-			}
+				// When crossing file blocks we need to allocate a temp buffer to hold the data
+				if (EncodedBlockRequest->EncodedData == nullptr)
+				{
+					EncodedBlockRequest->EncodedData = FMemory::Malloc(EncodedBlockFileSize);
+				}
 
-			FMemoryView FileBlock = FileBlockBufferView;
-			FMutableMemoryView EncodedBlock(EncodedBlockRequest->EncodedData, EncodedBlockFileSize);
+				FMemoryView FileBlock = FileBlockBufferView;
+				FMutableMemoryView EncodedBlock(EncodedBlockRequest->EncodedData, EncodedBlockFileSize);
 
-			if (OffsetInFileBlock < 0)
-			{
-				// Second or third or ... 
-				EncodedBlock.RightChopInline(FMath::Abs(OffsetInFileBlock));
-				FileBlock.LeftInline(EncodedBlock.GetSize());
+				if (OffsetInFileBlock < 0)
+				{
+					// Second or third or ... 
+					EncodedBlock.RightChopInline(FMath::Abs(OffsetInFileBlock));
+					FileBlock.LeftInline(EncodedBlock.GetSize());
+				}
+				else
+				{
+					// First file block
+					FileBlock.MidInline(OffsetInFileBlock, EncodedBlock.GetSize());
+				}
+				EncodedBlock.CopyFrom(FileBlock);
 			}
 			else
 			{
-				// First file block
-				FileBlock.MidInline(OffsetInFileBlock, EncodedBlock.GetSize());
+				// When the encoded block is completely wihtin the file block we use and keep a reference to the file block memory
+				check(OffsetInFileBlock >= 0);
+				EncodedBlockRequest->FileBufferHandle	= FileBlockBufferHandle; 
+				EncodedBlockRequest->EncodedData		= FileBlockBufferView.Mid(OffsetInFileBlock, EncodedBlockFileSize).GetData();
+				FileBlockMemoryPool.AddRef(EncodedBlockRequest->FileBufferHandle);
 			}
-			EncodedBlock.CopyFrom(FileBlock);
-		}
-		else
-		{
-			// When the encoded block is completely wihtin the file block we use and keep a reference to the file block memory
-			check(OffsetInFileBlock >= 0);
-			EncodedBlockRequest->FileBufferHandle	= FileBlockBufferHandle; 
-			EncodedBlockRequest->EncodedData		= FileBlockBufferView.Mid(OffsetInFileBlock, EncodedBlockFileSize).GetData();
-			FileBlockMemoryPool.AddRef(EncodedBlockRequest->FileBufferHandle);
 		}
 
 		if (--EncodedBlockRequest->RemainingFileBlocks == 0)
 		{
-			BlocksToDecode.AddTail(EncodedBlockRequest);
+			TUniqueLock Lock(RequestMgr);
+			RequestMgr.Remove(*EncodedBlockRequest);
+			if (EncodedBlockRequest->ErrorCode == EIoErrorCode::Unknown)
+			{
+				BlocksToDecode.AddTail(EncodedBlockRequest);
+				Stats.OnDecodeRequestEnqueued(*EncodedBlockRequest);
+			}
+			else
+			{
+				check(IsIoErrorAnyOf(EncodedBlockRequest->ErrorCode, EIoErrorCode::Cancelled, EIoErrorCode::ReadError));
+				CancelledOrFailedBlocks.AddTail(EncodedBlockRequest);
+			}
 		}
 	}
 
-	EnqueueBlocksToDecode(MoveTemp(BlocksToDecode));
-
+	if (FileBlockMemoryPool.Free(FileBlockBufferHandle))
 	{
-		TUniqueLock Lock(FileBlockMutex);
-		FileBlockMemoryPool.Free(FileBlockBufferHandle);
+		OnWakeUp();
 	}
 
 	if (CancelledOrFailedBlocks.IsEmpty() == false)
 	{
-		while (FIoEncodedBlockRequest* ToComplete = CancelledOrFailedBlocks.PopHead())
+		FIoPlatformReadRequestList CompletedReadRequests;
+		for (FIoEncodedBlockRequest& ToComplete : CancelledOrFailedBlocks)
 		{
-			for (FIoEncodedBlockRequest::FScatterTarget& ScatterTarget : ToComplete->ScatterTargets)
+			for (FIoEncodedBlockRequest::FScatterTarget& ScatterTarget : ToComplete.ScatterTargets)
 			{
-				ScatterTarget.Request->ErrorCode = ToComplete->ErrorCode;
+				ScatterTarget.Request->ErrorCode = ToComplete.ErrorCode;
 			}
-			CompleteEncodedBlockRequest(*ToComplete);
+			CompleteDecodedBlock(ToComplete, CompletedReadRequests);
 		}
-	}
-}
 
-void FPlatformIoDispatcherBase::EnqueueBlocksToDecode(FIoEncodedBlockRequestList BlocksToDecode)
-{
-	if (bMultithreaded == false)
-	{
-		while (FIoEncodedBlockRequest* BlockRequest = BlocksToDecode.PopHead())
-		{
-			// Allocate memory
-			check(BlockRequest->ErrorCode == EIoErrorCode::Unknown);
-			check(BlockRequest->BufferHandle.IsValid() == false);
-			BlockRequest->DecodedData = ChunkBlockMemoryPool.Alloc(BlockRequest->BlockUncompressedSize, BlockRequest->BufferHandle);
-
-			// Decode the block
-			Stats.OnDecodeRequestEnqueued(*BlockRequest);
-			const FIoChunkBlockDecodeResult Result = FIoChunkEncoding::DecodeBlock(
-				FIoChunkBlockDecodingParams 
-				{
-					.EncryptionKey		= BlockRequest->EncryptionKey,
-					.BlockHash			= BlockRequest->BlockHash,
-					.CompressionFormat	= BlockRequest->CompressionMethod,
-				},
-				FMutableMemoryView(BlockRequest->EncodedData, BlockRequest->BlockCompressedSize),
-				FMutableMemoryView(BlockRequest->DecodedData, BlockRequest->BlockUncompressedSize));
-
-			BlockRequest->DecodedData	= Result.DecodedBlock.GetData();
-			BlockRequest->ErrorCode		= Result.ErrorCode;
-			Stats.OnDecodeRequestCompleted(*BlockRequest);
-
-			// Copy the data from the block to the destination buffer(s) 
-			ScatterDecodedBlock(*BlockRequest);
-			ChunkBlockMemoryPool.Free(BlockRequest->BufferHandle);
-			FileBlockMemoryPool.Free(BlockRequest->FileBufferHandle);
-			BlockRequest->DecodedData = nullptr;
-
-			// Trigger completion callbacks and destroy the request
-			CompleteEncodedBlockRequest(*BlockRequest);
-		}
-	}
-	else
-	{
-		FIoEncodedBlockRequestList BlocksToEnqueue;
 		{
-			TUniqueLock Lock(FileBlockMutex);
-			while (FIoEncodedBlockRequest* BlockRequest = BlocksToDecode.PopHead())
+			TUniqueLock Lock(RequestMgr);
+			while (FIoPlatformReadRequest* ReadRequest = CompletedReadRequests.PopHead())
 			{
-				check(BlockRequest->ErrorCode == EIoErrorCode::Unknown);
-				check(BlockRequest->BufferHandle.IsValid() == false);
-				BlockRequest->DecodedData = ChunkBlockMemoryPool.Alloc(BlockRequest->BlockUncompressedSize, BlockRequest->BufferHandle);
-				if (BlockRequest->BufferHandle.IsValid())
-				{
-					check(BlockRequest->DecodedData != nullptr);
-					BlocksToEnqueue.AddTail(BlockRequest);
-				}
-				else
-				{
-					BlocksReadyForDecoding.AddTail(BlockRequest);
-				}
-				Stats.OnDecodeRequestEnqueued(*BlockRequest);
+				RequestMgr.Destroy(*ReadRequest);
 			}
-		}
 
-		while (FIoEncodedBlockRequest* BlockRequest = BlocksToEnqueue.PopHead())
-		{
-			EnqueueBlockToDecode(*BlockRequest);
+			while (FIoEncodedBlockRequest* ToComplete = CancelledOrFailedBlocks.PopHead())
+			{
+				RequestMgr.Destroy(*ToComplete);
+			}
 		}
 	}
 }
@@ -1306,93 +1358,30 @@ void FPlatformIoDispatcherBase::ProcessDecodedBlock(FIoEncodedBlockRequest& Enco
 
 	ScatterDecodedBlock(EncodedBlockRequest); 
 
-	// Release memory so we can start the next file read requests
-	FIoEncodedBlockRequest* NextBlockToDecode = nullptr;
+	// Free the file block memory and wake up dispatcher to start the next file read request
+	if (FileBlockMemoryPool.Free(EncodedBlockRequest.FileBufferHandle))
 	{
-		TUniqueLock Lock(FileBlockMutex);
-		FileBlockMemoryPool.Free(EncodedBlockRequest.FileBufferHandle);
+		OnWakeUp();
+	}
 
-		if (NextBlockToDecode = BlocksReadyForDecoding.PopHead(); NextBlockToDecode != nullptr)
-		{
-			NextBlockToDecode->DecodedData		= ChunkBlockMemoryPool.Realloc(EncodedBlockRequest.BufferHandle, NextBlockToDecode->BlockUncompressedSize);
-			NextBlockToDecode->BufferHandle		= EncodedBlockRequest.BufferHandle;
-			EncodedBlockRequest.BufferHandle	= FIoBufferHandle();
-			
-		}
-		else
-		{
-			ChunkBlockMemoryPool.Free(EncodedBlockRequest.BufferHandle);
-		}
-		EncodedBlockRequest.DecodedData = nullptr;
+	{
+		TUniqueLock Lock(EncodedBlockMutex);
+		BlocksToComplete.AddTail(&EncodedBlockRequest);
 		Stats.OnDecodeRequestCompleted(EncodedBlockRequest);
 	}
 
-	// Wakeup dispatcher to start the next file read request
-	OnWakeUp();
-
-	CompleteEncodedBlockRequest(EncodedBlockRequest);
-
-	// Instead of recursively calling the chunk decoder we enqueue more blocks by setting the next request parameter 
-	if (NextBlockToDecode != nullptr)
+	// Wakeup support thread to start next decoding request
+	if (SupportThread.IsValid())
 	{
-		auto OnBlockDecoded = [this, NextBlockToDecode](FIoChunkBlockDecodeResult&& DecodeResult, FIoChunkBlockDecodeRequest& Next)
-		{
-			ProcessDecodedBlock(*NextBlockToDecode, MoveTemp(DecodeResult), Next);
-		};
-		NextDecodeRequest = FIoChunkBlockDecodeRequest
-		{
-			.Params = FIoChunkBlockDecodingParams
-			{
-				.EncryptionKey		= NextBlockToDecode->EncryptionKey,
-				.BlockHash			= NextBlockToDecode->BlockHash,
-				.CompressionFormat	= NextBlockToDecode->CompressionMethod
-			},
-			.EncodedBlock			= FMutableMemoryView(NextBlockToDecode->EncodedData, NextBlockToDecode->BlockCompressedSize),
-			.DecodedBlock			= FMutableMemoryView(NextBlockToDecode->DecodedData, NextBlockToDecode->BlockUncompressedSize),
-			.OnDecoded				= MoveTemp(OnBlockDecoded) 
-		};
+		SupportThread->WakeUp();
 	}
 }
 
 void FPlatformIoDispatcherBase::ScatterDecodedBlock(FIoEncodedBlockRequest& EncodedBlockRequest)
 {
-	LLM_SCOPE_BYNAME(TEXT("FileSystem/FileIoStore"));
 	TRACE_CPUPROFILER_EVENT_SCOPE(IoDispatcherScatter);
 	check(EncodedBlockRequest.ErrorCode != EIoErrorCode::Unknown);
 
-	{
-		// TODO: Better to do this when dequeued to avoid mutex?
-		TUniqueLock Lock(ScatterMutex);
-		for (FIoEncodedBlockRequest::FScatterTarget& ScatterTarget : EncodedBlockRequest.ScatterTargets)
-		{
-			FIoPlatformReadRequest& Request = *ScatterTarget.Request;
-			const EIoErrorCode ErrorCode	= Request.ErrorCode.load(std::memory_order_relaxed);
-
-			Stats.OnBytesScattered(ScatterTarget.SizeInBlock);
-			if (EncodedBlockRequest.ErrorCode != EIoErrorCode::Ok)
-			{
-				Request.ErrorCode.store(EncodedBlockRequest.ErrorCode, std::memory_order_relaxed);
-				Request.FailedBlockId = EncodedBlockRequest.BlockKey.GetBlockId();
-			}
-			else
-			{
-				EIoErrorCode Expected = EIoErrorCode::Ok;
-				Request.ErrorCode.compare_exchange_weak(Expected, EncodedBlockRequest.ErrorCode);
-			}
-	
-			// If the size in block is zero the request was cancelled
-			if (ScatterTarget.SizeInBlock > 0)
-			{
-				check(Request.IsScatterGather());
-				if (Request.Dst.GetSize() == 0)
-				{
-					UE::FInheritedContextScope InheritedContextScope = Request.RestoreInheritedContext();
-					Request.Dst = FIoBuffer(ScatterTarget.Request->DstSize);
-				}
-			}
-		}
-	}
-
 	FMemoryView DecodedBlock(EncodedBlockRequest.DecodedData, EncodedBlockRequest.BlockUncompressedSize);
 	for (FIoEncodedBlockRequest::FScatterTarget& ScatterTarget : EncodedBlockRequest.ScatterTargets)
 	{
@@ -1407,21 +1396,32 @@ void FPlatformIoDispatcherBase::ScatterDecodedBlock(FIoEncodedBlockRequest& Enco
 			FMutableMemoryView	Dst = Request.Dst.GetMutableView().RightChop(ScatterTarget.OffsetInDst);
 			FMemoryView			Src = DecodedBlock.Mid(ScatterTarget.OffsetInBlock, ScatterTarget.SizeInBlock);
 			Dst.CopyFrom(Src);
+			Stats.OnBytesScattered(ScatterTarget.SizeInBlock);
 		}
 	}
 }
 
-void FPlatformIoDispatcherBase::CompleteEncodedBlockRequest(FIoEncodedBlockRequest& EncodedBlockRequest)
+void FPlatformIoDispatcherBase::CompleteDecodedBlock(FIoEncodedBlockRequest& EncodedBlockRequest, FIoPlatformReadRequestList& OutCompletedReadRequests)
 {
-	FIoPlatformReadRequestList CompletedReadRequests;
+	TRACE_CPUPROFILER_EVENT_SCOPE(CompleteDecodedBlock);
 	for (FIoEncodedBlockRequest::FScatterTarget& ScatterTarget : EncodedBlockRequest.ScatterTargets)
 	{
 		FIoPlatformReadRequest& Request = *ScatterTarget.Request;
 		check(Request.IsScatterGather());
 		check(Request.ScatterGather.RemainingBlocks > 0);
-		check(Request.ErrorCode != EIoErrorCode::Unknown);
+		
+		if (EncodedBlockRequest.ErrorCode != EIoErrorCode::Ok)
+		{
+			Request.ErrorCode.store(EncodedBlockRequest.ErrorCode, std::memory_order_relaxed);
+			Request.FailedBlockId = EncodedBlockRequest.BlockKey.GetBlockId();
+		}
+		else
+		{
+			EIoErrorCode Expected = EIoErrorCode::Ok;
+			Request.ErrorCode.compare_exchange_weak(Expected, EncodedBlockRequest.ErrorCode);
+		}
 
-		if (1 == Request.ScatterGather.RemainingBlocks.fetch_sub(1, std::memory_order_seq_cst))
+		if (1 == Request.ScatterGather.RemainingBlocks.fetch_sub(1, std::memory_order_relaxed))
 		{
 			FIoFileReadRequestCompleted OnCompleted = MoveTemp(Request.OnCompleted);
 			if (OnCompleted)
@@ -1433,19 +1433,183 @@ void FPlatformIoDispatcherBase::CompleteEncodedBlockRequest(FIoEncodedBlockReque
 					.ErrorCode		= Request.ErrorCode.load(std::memory_order_relaxed),
 				});
 			}
-			CompletedReadRequests.AddTail(&Request);
+			OutCompletedReadRequests.AddTail(&Request);
+		}
+	}
+}
+
+bool FPlatformIoDispatcherBase::ProcessCompletedFileBlocks()
+{
+	FIoFileBlockRequestList ToComplete;
+	{
+		UE::TUniqueLock Lock(FileBlockMutex);
+		ToComplete				= MoveTemp(FileBlocksToComplete);
+		FileBlocksToComplete	= FIoFileBlockRequestList();
+	}
+
+	if (ToComplete.IsEmpty())
+	{
+		return false;
+	}
+
+	while (FIoFileBlockRequest* FileBlock = ToComplete.PopHead())
+	{
+		CompleteFileBlock(*FileBlock);
+	}
+
+	return true;
+}
+
+bool FPlatformIoDispatcherBase::ProcessBlocksToDecode()
+{
+	bool bDidSomething = false;
+	FIoPlatformReadRequestList	ReadRequestsToDestroy;
+	FIoEncodedBlockRequestList	EncodedBlockRequestsToDestroy;
+
+	while (FIoEncodedBlockRequest* NextEncodedBlockRequest = BlocksToDecode.PeekHead())
+	{
+		check(NextEncodedBlockRequest->ErrorCode == EIoErrorCode::Unknown);
+		check(NextEncodedBlockRequest->BufferHandle.IsValid() == false);
+
+		// Try allocate temporary decoding memory for the encoded block
+		NextEncodedBlockRequest->DecodedData = ChunkBlockMemoryPool.Alloc(NextEncodedBlockRequest->BlockUncompressedSize, NextEncodedBlockRequest->BufferHandle);
+		if (NextEncodedBlockRequest->BufferHandle.IsValid() == false)
+		{
+			break;
+		}
+
+		FIoEncodedBlockRequest* EncodedBlockRequest = BlocksToDecode.PopHead();
+		bDidSomething = true;
+
+		// Allocate target memory
+		{
+			LLM_SCOPE_BYNAME(TEXT("FileSystem/FileIoStore"));
+
+			for (FIoEncodedBlockRequest::FScatterTarget& ScatterTarget : EncodedBlockRequest->ScatterTargets)
+			{
+				FIoPlatformReadRequest& Request = *ScatterTarget.Request;
+				// If the size in block is zero the request was cancelled
+				if (ScatterTarget.SizeInBlock > 0)
+				{
+					check(Request.IsScatterGather());
+					if (Request.Dst.GetSize() == 0)
+					{
+						UE::FInheritedContextScope InheritedContextScope = Request.RestoreInheritedContext();
+						Request.Dst = FIoBuffer(ScatterTarget.Request->DstSize);
+					}
+				}
+			}
+		}
+
+		const bool bDecodeAndScatterAsync = bMultithreaded &&
+			(!EncodedBlockRequest->CompressionMethod.IsNone() ||
+			 !EncodedBlockRequest->EncryptionKey.IsEmpty() ||
+			 !EncodedBlockRequest->BlockHash.IsEmpty());
+
+		if (bDecodeAndScatterAsync)
+		{
+			EnqueueBlockToDecode(*EncodedBlockRequest);
+		}
+		else
+		{
+			{
+				TRACE_CPUPROFILER_EVENT_SCOPE(IoDispatcherDecodeBlock);
+				const FIoChunkBlockDecodeResult Result = FIoChunkEncoding::DecodeBlock(
+					FIoChunkBlockDecodingParams 
+					{
+						.EncryptionKey		= EncodedBlockRequest->EncryptionKey,
+						.BlockHash			= EncodedBlockRequest->BlockHash,
+						.CompressionFormat	= EncodedBlockRequest->CompressionMethod,
+					},
+					FMutableMemoryView(EncodedBlockRequest->EncodedData, EncodedBlockRequest->BlockCompressedSize),
+					FMutableMemoryView(EncodedBlockRequest->DecodedData, EncodedBlockRequest->BlockUncompressedSize));
+
+				EncodedBlockRequest->DecodedData	= Result.DecodedBlock.GetData();
+				EncodedBlockRequest->ErrorCode		= Result.ErrorCode;
+			}
+
+			ScatterDecodedBlock(*EncodedBlockRequest);
+			CompleteDecodedBlock(*EncodedBlockRequest, ReadRequestsToDestroy);
+			EncodedBlockRequestsToDestroy.AddTail(EncodedBlockRequest);
+
+			// Free the file block memory and wake up the dispatcher to start the next file read request
+			if (FileBlockMemoryPool.Free(EncodedBlockRequest->FileBufferHandle))
+			{
+				OnWakeUp();
+			}
+
+			ChunkBlockMemoryPool.Free(EncodedBlockRequest->BufferHandle);
+			EncodedBlockRequest->DecodedData = nullptr;
+			Stats.OnDecodeRequestCompleted(*EncodedBlockRequest);
 		}
 	}
 
 	// Destroy the block request and completed read request(s)
+	{
+		TUniqueLock Lock(RequestMgr);
+		while (FIoPlatformReadRequest* ToDestroy = ReadRequestsToDestroy.PopHead())
+		{
+			RequestMgr.Destroy(*ToDestroy);
+		}
+
+		while (FIoEncodedBlockRequest* ToDestroy = EncodedBlockRequestsToDestroy.PopHead())
+		{
+			RequestMgr.Destroy(*ToDestroy);
+		}
+	}
+
+	return bDidSomething;
+}
+
+bool FPlatformIoDispatcherBase::ProcessBlocksToComplete()
+{
+	FIoEncodedBlockRequestList ToComplete;
+	{
+		TUniqueLock Lock(EncodedBlockMutex);
+		ToComplete			= MoveTemp(BlocksToComplete);
+		BlocksToComplete	= FIoEncodedBlockRequestList();
+	}
+
+	if (ToComplete.IsEmpty())
+	{
+		return false;
+	}
+
+	FIoPlatformReadRequestList CompletedReadRequests;
+	for (FIoEncodedBlockRequest& EncodedBlockRequest : ToComplete)
+	{
+		ChunkBlockMemoryPool.Free(EncodedBlockRequest.BufferHandle);
+		EncodedBlockRequest.DecodedData = nullptr;
+		CompleteDecodedBlock(EncodedBlockRequest, CompletedReadRequests);
+	}
+
 	{
 		TUniqueLock Lock(RequestMgr);
 		while (FIoPlatformReadRequest* ReadRequest = CompletedReadRequests.PopHead())
 		{
 			RequestMgr.Destroy(*ReadRequest);
 		}
-		RequestMgr.Destroy(EncodedBlockRequest);
+
+		while (FIoEncodedBlockRequest* EncodedBlockRequest = ToComplete.PopHead())
+		{
+			RequestMgr.Destroy(*EncodedBlockRequest);
+		}
 	}
+
+	return true;
+}
+
+bool FPlatformIoDispatcherBase::ProcessBlocks()
+{
+	TRACE_CPUPROFILER_EVENT_SCOPE(ProcessIO);
+
+	bool bDidSomething = false;
+
+	bDidSomething |= ProcessCompletedFileBlocks();
+	bDidSomething |= ProcessBlocksToComplete();
+	bDidSomething |= ProcessBlocksToDecode();
+
+	return bDidSomething;
 }
 
 } // namespace UE
diff --git a/Engine/Source/Runtime/Core/Internal/IO/PlatformIoDispatcherBase.h b/Engine/Source/Runtime/Core/Internal/IO/PlatformIoDispatcherBase.h
--- a/Engine/Source/Runtime/Core/Internal/IO/PlatformIoDispatcherBase.h
+++ b/Engine/Source/Runtime/Core/Internal/IO/PlatformIoDispatcherBase.h
@@ -11,7 +11,6 @@
 #include "IO/IoBuffer.h"
 #include "IO/IoChunkEncoding.h"
 #include "IO/IoContainers.h"
-#include "IO/IoDispatcherPriority.h"
 #include "IO/PlatformIoDispatcher.h"
 #include "IO/PlatformIoDispatcherStats.h"
 #include "IO/IoStatus.h"
@@ -24,36 +23,6 @@
 namespace UE
 {
 
-////////////////////////////////////////////////////////////////////////////////
-enum class EIoFileReadPriority : uint8
-{
-	Min,
-	Low,
-	Medium,
-	High,
-	Max,
-	Count
-};
-
-////////////////////////////////////////////////////////////////////////////////
-static inline EIoFileReadPriority IoFileReadPriorityFromDispatcherPriority(int32 Priority)
-{
-	if (Priority < IoDispatcherPriority_Low)
-	{
-		return EIoFileReadPriority::Min;
-	}
-	if (Priority < IoDispatcherPriority_Medium)
-	{
-		return EIoFileReadPriority::Low;
-	}
-	if (Priority < IoDispatcherPriority_High)
-	{
-		return EIoFileReadPriority::Medium;
-	}
-
-	return Priority > IoDispatcherPriority_High ? EIoFileReadPriority::Max : EIoFileReadPriority::High;
-}
-
 ////////////////////////////////////////////////////////////////////////////////
 template<typename T, SIZE_T Size = 8>
 using FTempArray = TArray<T, TInlineAllocator<Size>>;
@@ -257,6 +226,7 @@ struct FIoFileBlockRequest
 
 	FEncodedBlocksArray				EncodedBlockRequests;
 	FIoFileBlockRequest*			Next = nullptr;
+	FIoFileBlockRequest*			Prev = nullptr;
 	FIoPlatformReadRequest*			DirectReadRequest = nullptr;
 	void*							Buffer = nullptr;
 	FIoFileHandle					FileHandle;
@@ -264,6 +234,7 @@ struct FIoFileBlockRequest
 	uint64							FileOffset = 0;
 	int64							FileSize = 0;
 	uint64							Size = 0;
+	uint64							BytesUsed = 0;
 	uint64							TimeCreated = FPlatformTime::Cycles64();
 	uint32							SeqNo = NextSeqNo++;
 	FIoBufferHandle					BufferHandle;
@@ -284,8 +255,10 @@ struct FIoPlatformFileInfo
 };
 
 ////////////////////////////////////////////////////////////////////////////////
-class FIoFileBlockRequestQueue
+class FIoQueue
 {
+	using FIoFileBlockRequestDoubleLinkedList = TIntrusiveTwoWayList<FIoFileBlockRequest>;
+
 	struct FFileOffset
 	{
 		FIoFileHandle	FileHandle;
@@ -295,9 +268,17 @@ class FIoFileBlockRequestQueue
 
 	struct FPrioQueue
 	{
-		TArray<FIoFileBlockRequest*>	ByOffset;
-		FIoFileBlockRequestList			BySequence;
-		int32							PeekIndex = INDEX_NONE;
+		explicit	FPrioQueue(int32 InPriority) : Priority(InPriority) { }
+					FPrioQueue(const FPrioQueue&) = delete;
+					FPrioQueue(FPrioQueue&&) = default; 
+
+		FPrioQueue& operator=(const FPrioQueue&) = delete;
+		FPrioQueue& operator=(FPrioQueue&&) = default; 
+
+		TArray<FIoFileBlockRequest*>		ByOffset;
+		FIoFileBlockRequestDoubleLinkedList	BySequence;
+		int32								PeekIndex = INDEX_NONE;
+		int32								Priority; 
 	};
 
 	static FFileOffset ToFileOffset(const FIoFileBlockRequest* Request)
@@ -310,6 +291,8 @@ class FIoFileBlockRequestQueue
 		};
 	}
 
+	static int32 QueueToPriority(const FPrioQueue& Queue) { return Queue.Priority; }
+
 	static bool FileOffsetLess(const FFileOffset& LHS, const FFileOffset& RHS)
 	{
 		if (LHS.FileHandle.Value() == RHS.FileHandle.Value())
@@ -329,7 +312,7 @@ class FIoFileBlockRequestQueue
 	}
 
 public:
-							FIoFileBlockRequestQueue(FPlatformIoDispatcherStats& InStats);
+							FIoQueue(FPlatformIoDispatcherStats& InStats);
 	void					Enqueue(FIoFileBlockRequestList&& Requests);
 	FIoFileBlockRequest*	Dequeue();
 	FIoFileBlockRequest*	Peek();
@@ -337,16 +320,21 @@ public:
 	void					ReprioritizeCancelled() { bReprioritizeCancelled = true; }
 	void					SortByOffset(bool bValue) { bSortByOffset = bValue; }
 
+	inline void				Lock() { Mutex.Lock(); }
+	inline void				Unlock() { Mutex.Unlock(); }
+
 private:
+	void					AddToPrioQueue(FIoFileBlockRequest& Request);
 	FIoFileBlockRequest*	GetByOffset(bool bDequeue);
-	FIoFileBlockRequest*	GetByOffset(FPrioQueue& PrioQueue, EIoFileReadPriority QueuePriority, bool bDequeue);
+	FIoFileBlockRequest*	GetByOffset(FPrioQueue& PrioQueue, bool bDequeue);
 	FIoFileBlockRequest*	GetCancelled(bool bDequeue);
 
 	FPlatformIoDispatcherStats&		Stats;
 	TArray<FIoFileBlockRequest*>	Heap;
 	TArray<FIoFileBlockRequest*>	CancelledHeap;
-	FPrioQueue						PrioQueues[static_cast<int32>(EIoFileReadPriority::Count)];
+	TArray<FPrioQueue>				PrioQueues;
 	FFileOffset						LastFileOffset;
+	FMutex							Mutex;
 	bool							bSortByOffset = false;
 	bool							bReprioritizeCancelled = false;
 };
@@ -358,9 +346,8 @@ public:
 					FIoFileBlockMemoryPool(FPlatformIoDispatcherStats& InStats);
 	void			Initialize(uint32 FileBlockSize, uint32 PoolSize, uint32 Alignment = 4096);
 	void*			Alloc(FIoBufferHandle& OutHandle);
-	void			Free(FIoBufferHandle& Handle);
+	bool			Free(FIoBufferHandle& Handle);
 	void			AddRef(FIoBufferHandle Handle);
-	bool			IsEmpty() const;
 
 private:
 
@@ -377,7 +364,7 @@ private:
 		FMemoryBlock*		Next		= nullptr;
 		void*				Memory		= nullptr;
 		int32				Index		= INDEX_NONE;
-		std::atomic_int32_t	RefCount{0};
+		int32				RefCount	= 0;
 	};
 
 	FPlatformIoDispatcherStats&			Stats;
@@ -385,6 +372,7 @@ private:
 	TArray<FMemoryBlock>				MemoryBlocks;
 	FMemoryBlock*						FreeBlock = nullptr;
 	uint32								BlockSize;
+	FMutex								Mutex;
 };
 
 ////////////////////////////////////////////////////////////////////////////////
@@ -464,6 +452,7 @@ public:
 	FIoFileBlockRequest*				GetFileBlockRequest(FIoBlockKey BlockKey);
 	void								GetFileBlockRequests(FIoPlatformReadRequest& ReadRequest, FTempArray<FIoFileBlockRequest*>& OutRequests);
 	void								AddFileBlockRequest(FIoPlatformReadRequest& ReadRequest, FIoFileBlockRequest& FileBlockRequest);
+	void								AddFileBlockRequest(FIoPlatformReadRequest& ReadRequest, FIoBlockKey FileBlockKey);
 	FIoEncodedBlockRequest&				GetOrCreateEncodedBlockRequest(FIoBlockKey BlockKey, bool& bCreated);
 	void								Remove(FIoEncodedBlockRequest& Request);
 	void								Destroy(FIoFileBlockRequest& Request);
@@ -487,6 +476,9 @@ private:
 	UE::FMutex							Mutex;
 };
 
+// Forward declare thread helper
+class FIoServiceThread;
+
 ////////////////////////////////////////////////////////////////////////////////
 /// This is a base class for platform specific I/O dispatcher implementations.
 /// There are two kinds of I/O requests when reading data from disk. Direct-read requests
@@ -529,23 +521,32 @@ protected:
 	virtual FIoPlatformFileInfo			GetPlatformFileInfo(FIoFileHandle FileHandle) = 0;
 
 	uint32								GetNextFileId();
-	void								ProcessCompletedFileBlock(FIoFileBlockRequest& FileBlockRequest);
-	void								EnqueueBlocksToDecode(FIoEncodedBlockRequestList BlocksToDecode);
+
+	void								EnqueueCompletedFileBlock(FIoFileBlockRequest& FileBlockRequest);
+	void								CompleteFileBlock(FIoFileBlockRequest& FileBlockRequest);
 	void								EnqueueBlockToDecode(FIoEncodedBlockRequest& EncodedBlockRequest);
 	void								ProcessDecodedBlock(FIoEncodedBlockRequest& EncodedBlockRequest, FIoChunkBlockDecodeResult&& DecodeResult, FIoChunkBlockDecodeRequest& NextDecodeRequest);
 	void								ScatterDecodedBlock(FIoEncodedBlockRequest& EncodedBlockRequest);
-	void								CompleteEncodedBlockRequest(FIoEncodedBlockRequest& EncodedBlockRequest);
+	void								CompleteDecodedBlock(FIoEncodedBlockRequest& EncodedBlockRequest, FIoPlatformReadRequestList& OutCompletedReadRequests);
+
+	bool								ProcessCompletedFileBlocks();
+	bool								ProcessBlocksToDecode();
+	bool								ProcessBlocksToComplete();
+	bool								ProcessBlocks();
 
 	FPlatformIoDispatcherStats			Stats;
 	TUniquePtr<FRunnableThread>			Thread;
+	TUniquePtr<FIoServiceThread>		SupportThread;
 	FPlatformIoDispatcherRequestMgr		RequestMgr;
 	FIoChunkBlockMemoryPool				ChunkBlockMemoryPool;
 	FIoFileBlockMemoryPool				FileBlockMemoryPool;
-	FIoFileBlockRequestQueue			FileBlockRequestQueue;
+	FIoQueue							IoQueue;
 	FIoFileBlockCache					FileBlockCache;
-	FIoEncodedBlockRequestList			BlocksReadyForDecoding;
+	FIoFileBlockRequestList				FileBlocksToComplete;
+	FIoEncodedBlockRequestList			BlocksToDecode;
+	FIoEncodedBlockRequestList			BlocksToComplete;
 	UE::FMutex							FileBlockMutex;
-	UE::FMutex							ScatterMutex;
+	UE::FMutex							EncodedBlockMutex;
 	uint32								FileBlockSize = 256 << 10;
 	bool								bMultithreaded = false;
 	std::atomic_bool					bStopRequested{false};
diff --git a/Engine/Source/Runtime/Core/Tests/IO/PlatformIoDispatcherTests.cpp b/Engine/Source/Runtime/Core/Tests/IO/PlatformIoDispatcherTests.cpp
--- a/Engine/Source/Runtime/Core/Tests/IO/PlatformIoDispatcherTests.cpp
+++ b/Engine/Source/Runtime/Core/Tests/IO/PlatformIoDispatcherTests.cpp
@@ -44,7 +44,7 @@ TEST_CASE("Core::IO::Platform::FileReadQueue", "[Core][IO]")
 		// Arrange
 		const int32					ExpectedReadCount = 20;
 		FPlatformIoDispatcherStats	Stats;
-		FIoFileBlockRequestQueue	Queue(Stats);
+		FIoQueue					Queue(Stats);
 		TArray<FIoFileBlockRequest> Enqueued;
 
 		Enqueued.SetNum(ExpectedReadCount);
@@ -72,7 +72,7 @@ TEST_CASE("Core::IO::Platform::FileReadQueue", "[Core][IO]")
 		// Arrange
 		const int32						ExpectedReadCount = 5;
 		FPlatformIoDispatcherStats		Stats;
-		FIoFileBlockRequestQueue		Queue(Stats);
+		FIoQueue						Queue(Stats);
 		TArray<FIoFileBlockRequest> 	Enqueued;
 		FIoFileHandle					FileHandle(1);
 
@@ -109,7 +109,7 @@ TEST_CASE("Core::IO::Platform::FileReadQueue", "[Core][IO]")
 		// Arrange
 		const int32						ExpectedReadCount = 5;
 		FPlatformIoDispatcherStats		Stats;
-		FIoFileBlockRequestQueue		Queue(Stats);
+		FIoQueue						Queue(Stats);
 		TArray<FIoFileBlockRequest> 	Enqueued;
 		FIoFileHandle					FileHandle(1);
 
@@ -149,7 +149,7 @@ TEST_CASE("Core::IO::Platform::FileReadQueue", "[Core][IO]")
 		// Arrange
 		const int32						ExpectedReadCount = 20;
 		FPlatformIoDispatcherStats		Stats;
-		FIoFileBlockRequestQueue		Queue(Stats);
+		FIoQueue						Queue(Stats);
 		TArray<FIoFileBlockRequest> 	Enqueued;
 		FIoFileHandle					FileHandle(1);
 		uint64							FileBlockOffset = 0; 
