Extended IPlatformFile interface with open write overload that returns a file handle or an error code.

diff --git a/Engine/Source/Runtime/Core/Private/GenericPlatform/GenericPlatformFile.cpp b/Engine/Source/Runtime/Core/Private/GenericPlatform/GenericPlatformFile.cpp
--- a/Engine/Source/Runtime/Core/Private/GenericPlatform/GenericPlatformFile.cpp
+++ b/Engine/Source/Runtime/Core/Private/GenericPlatform/GenericPlatformFile.cpp
@@ -553,6 +553,20 @@ FOpenMappedResult IPlatformFile::OpenMappedEx(const TCHAR* Filename, EOpenReadFl
 	return MakeError(TEXT("OpenMappedEx is not implemented on this platform"));
 }
 
+FFileOpenResult IPlatformFile::OpenWrite(const TCHAR* Filename, EOpenWriteFlags Flags)
+{
+	const bool bAppend = EnumHasAnyFlags(Flags, EOpenWriteFlags::Append);
+	const bool bAllowRead = EnumHasAnyFlags(Flags, EOpenWriteFlags::AllowRead);
+
+	if (IFileHandle* Handle = OpenWrite(Filename, bAppend, bAllowRead); Handle != nullptr)
+	{
+		return MakeValue(Handle);
+	}
+
+	const uint32 LastError = FPlatformMisc::GetLastError();
+	return MakeError(FString::Printf(TEXT("Failed to open file '%s' for writing"), Filename != nullptr ? Filename : TEXT("<Null>")), int32(LastError));
+}
+
 DEFINE_STAT(STAT_AsyncFileMemory);
 DEFINE_STAT(STAT_AsyncFileHandles);
 DEFINE_STAT(STAT_AsyncFileRequests);
diff --git a/Engine/Source/Runtime/Core/Public/GenericPlatform/GenericPlatformFile.h b/Engine/Source/Runtime/Core/Public/GenericPlatform/GenericPlatformFile.h
--- a/Engine/Source/Runtime/Core/Public/GenericPlatform/GenericPlatformFile.h
+++ b/Engine/Source/Runtime/Core/Public/GenericPlatform/GenericPlatformFile.h
@@ -511,7 +511,7 @@ public:
 	 * Open a file handle for reading.
 	 * 
 	 * @param Filename	The file to be opened
-	 * @param Flags		Allows specialization of the open operation, @see EReadFlags
+	 * @param Flags		Allows specialization of the open operation, @see EOpenReadFlags
 	 * 
 	 * @return	The return value will either contain the valid file handle or an error message. @see FFileOpenResult
 	 */
@@ -538,8 +538,28 @@ public:
 		return OpenRead(Filename, bAllowWrite);
 	}
 
+	/** Flags to be used when opening a file for writing. */
+	enum class EOpenWriteFlags : uint8
+	{
+		/** No additional flags. */
+		None		= 0,
+		/** Allow other handles/processes to read from this file. */
+		AllowRead	= 1 << 0,
+		/** Append to file. */ 
+		Append		= 1 << 1,
+	};
+
 	/** Attempt to open a file for writing. If successful will return a non-nullptr pointer. Close the file by deleting the handle. **/
-	virtual IFileHandle*	OpenWrite(const TCHAR* Filename, bool bAppend = false, bool bAllowRead = false) = 0;
+	virtual IFileHandle* OpenWrite(const TCHAR* Filename, bool bAppend = false, bool bAllowRead = false) = 0;
+	/**
+	 * Open a file handle for writing. 
+	 * 
+	 * @param Filename	The file to be opened
+	 * @param Flags		Allows specialization of the open operation, @see EOpenWriteFlags
+	 * 
+	 * @return	The return value will either contain the valid file handle or an error message. @see FFileOpenResult
+	 */
+	CORE_API virtual FFileOpenResult OpenWrite(const TCHAR* Filename, EOpenWriteFlags Flags);
 
 	/** Return true if the directory exists. **/
 	virtual bool		DirectoryExists(const TCHAR* Directory) = 0;
