- Added IProcessSentinel - a class to know specifically when a process is running (with WIndows and PosixOS implementations), with optional contents that can be read
- Wrote FGenericProcessSentinel, but it is not currently hooked up in the non-Windows/Posix case because it does not 100% guarantee the state of the sentinel is 100% accurate - it's more like 95-99% accurate (it depends on ProcessId not having been recycled). There is commented out code that could make it 100% for platforms that guarantee a file delete will fail if another platform has the file open for write)

diff --git a/Engine/Source/Runtime/Core/Private/GenericPlatform/GenericPlatformProcess.cpp b/Engine/Source/Runtime/Core/Private/GenericPlatform/GenericPlatformProcess.cpp
--- a/Engine/Source/Runtime/Core/Private/GenericPlatform/GenericPlatformProcess.cpp
+++ b/Engine/Source/Runtime/Core/Private/GenericPlatform/GenericPlatformProcess.cpp
@@ -9,9 +9,12 @@
 #include "Logging/LogMacros.h"
 #include "CoreGlobals.h"
 #include "HAL/FileManager.h"
+#include "HAL/PlatformFileManager.h"
 #include "Misc/Parse.h"
 #include "Misc/SingleThreadEvent.h"
 #include "Misc/CommandLine.h"
+#include "Misc/CoreDelegates.h"
+#include "Misc/FileHelper.h"
 #include "Misc/Paths.h"
 #include "Stats/Stats.h"
 #include "Stats/ThreadIdleStats.h"
@@ -37,6 +40,17 @@
 DEFINE_STAT(STAT_Sleep);
 DEFINE_STAT(STAT_EventWait);
 
+bool IProcessSentinel::DoesSentinelFileExist(const TCHAR* Name, FString* OptionalOutContents)
+{
+	return FPlatformProcess::CreateProcessSentinelObject()->DoesSentinelFileExistImpl(Name, OptionalOutContents);
+}
+
+FString IProcessSentinel::MakeFilename(const TCHAR* Name)
+{
+	return FPaths::ProjectIntermediateDir() + Name + TEXT(".sentinel");
+}
+
+
 void* FGenericPlatformProcess::GetDllHandle( const TCHAR* Filename )
 {
 	UE_LOG(LogHAL, Fatal, TEXT("FPlatformProcess::GetDllHandle not implemented on this platform"));
@@ -697,6 +711,13 @@ PRAGMA_DISABLE_DEPRECATION_WARNINGS
 PRAGMA_ENABLE_DEPRECATION_WARNINGS
 }
 
+void FGenericPlatformProcess::CeaseBeingFirstInstance()
+{
+PRAGMA_DISABLE_DEPRECATION_WARNINGS
+	GIsFirstInstance = false;
+PRAGMA_ENABLE_DEPRECATION_WARNINGS
+}
+
 UE::FPlatformSystemWideMutexNotImplemented::FPlatformSystemWideMutexNotImplemented(const FString& Name, FTimespan Timeout)
 {
 	UE_LOG(LogHAL, Fatal, TEXT("FPlatformSystemWideMutex not implemented on this platform."));
@@ -727,3 +748,139 @@ void FGenericPlatformProcess::ModifyThreadAssignmentForUObjectReferenceCollector
 }
 
 
+class FGenericProcessSentinel : public IProcessSentinel
+{
+public:
+
+	FGenericProcessSentinel()
+	{
+		FCoreDelegates::OnHandleSystemError.AddRaw(this, &FGenericProcessSentinel::OnCrash);
+		FCoreDelegates::OnShutdownAfterError.AddRaw(this, &FGenericProcessSentinel::OnShutdown);
+	}
+	
+	virtual ~FGenericProcessSentinel()
+	{
+		Cleanup();
+	}
+	
+	/**
+	 * Create the sentinel file
+	 */
+	virtual bool CreateSentinelFile(const TCHAR* Name, const FString& Contents) override
+	{
+		SavedFilename = MakeFilename(Name);
+
+		if (IFileManager::Get().FileExists(*SavedFilename))
+		{
+			if (!IFileManager::Get().Delete(*SavedFilename))
+			{
+				// attempt to lock it, if we can't, then another instance is running
+				UE_LOG(LogHAL, Error, TEXT("Another instance is running, unable to delete Sentinel file %s"), *SavedFilename);
+				return false;
+			}
+		}
+
+		// open the file for write so other processes can't
+		IFileHandle* Handle = FPlatformFileManager::Get().GetPlatformFile().OpenWrite(*SavedFilename);
+		if (Handle)
+		{
+			char PidString[16];
+			int Len = FCStringAnsi::Sprintf(PidString, "%d\n", FPlatformProcess::GetCurrentProcessId());
+			Handle->Write((uint8*)PidString, Len);
+			
+			FTCHARToUTF8 Convert(*Contents);
+			Handle->Write((uint8*)Convert.Get(), Convert.Length());
+			Handle->Flush(true);
+
+			return true;
+		}
+		return false;
+	}
+
+private:
+		
+	/** 
+	 * The actual implementation to determine if the sentinel file is exists and is valid (due to a process dying,
+	 * the file may actually exist on disk, but isn't a valid file anymore, so a simple FileExists may not be good enough
+	 */
+	virtual bool DoesSentinelFileExistImpl(const TCHAR* Name, FString* OptionalOutContents=nullptr) override
+	{
+		FString Filename = MakeFilename(Name);
+		FString LocalContents;
+		FString* Contents = (OptionalOutContents == nullptr) ? &LocalContents : OptionalOutContents;
+
+		/////////////////////////////////////////////////////////////////////////////////////////////////////////////
+		// ALTERNATE, MORE SECURE MODE
+		// If the platform can depend on Delete failing while another process as the file open for Write, then we
+		// can just attempt to delete it - if it gets deleted, then we know another process doesn't have it open 
+		// for write, which means the process died unexpectedly, and we can clean it up like it wasn't there
+		// However, not all platforms will block deleting if the file is open for write, so it's not a generic
+		// solution. 
+		// if (!IFileManager::Get().FileExists(*SavedFilename) || IFileManager::Get().Delete(*SavedFilename))
+		// {
+		//	return false;
+		// }
+		/////////////////////////////////////////////////////////////////////////////////////////////////////////////
+
+		// get the PID from the first line
+		if (FFileHelper::LoadFileToString(*Contents, *Filename))
+		{
+			int LineBreak = Contents->Find(TEXT("\n"));
+			if (LineBreak > 0)
+			{
+				FString PidString = Contents->Mid(0, LineBreak);
+				int Pid = FCString::Atoi(*PidString);
+				*Contents = Contents->Mid(LineBreak + 1);
+				
+				// is pid running
+				FProcHandle Proc = FPlatformProcess::OpenProcess(Pid);
+				if (FPlatformProcess::IsProcRunning(Proc))
+				{
+					return true;
+				}
+				
+				// if not, delete the file to clean up
+				IFileManager::Get().Delete(*SavedFilename);
+			}
+		}
+
+		return false;
+	}
+
+	/** Callback functions on crash and shutdown after error */
+	void OnCrash()
+	{
+		Cleanup();
+	}
+	void OnShutdown()
+	{
+		Cleanup();
+	}
+
+	void Cleanup()
+	{
+		if (FileHandle.IsValid())
+		{
+			FileHandle.Reset();
+		}
+		IFileManager::Get().Delete(*SavedFilename);
+	}
+
+	FString SavedFilename;
+	TUniquePtr<FArchive> FileHandle;
+};
+
+
+TSharedPtr<IProcessSentinel> FGenericPlatformProcess::CreateProcessSentinelObject()
+{
+	UE_LOG(LogHAL, Fatal, TEXT("This platform doesn't currently support process sentinels, because the generic one can't guarantee that the file will go away on abnormal termination (ie, killing the process, as opposed to crashing normally). If this is needed, look at FGenericProcessSentinel, it's a good starting point."));
+	return TSharedPtr<IProcessSentinel>(new FGenericProcessSentinel());
+}
+
+TSharedPtr<IProcessSentinel> FGenericPlatformProcess::CreateProcessSentinel(const TCHAR* Name, const FString& Contents)	
+{
+	TSharedPtr<IProcessSentinel> Sentinel = FPlatformProcess::CreateProcessSentinelObject();
+	
+	Sentinel->CreateSentinelFile(Name, Contents);
+	return Sentinel;
+}
diff --git a/Engine/Source/Runtime/Core/Private/PosixOS/PosixOSPlatformProcess.cpp b/Engine/Source/Runtime/Core/Private/PosixOS/PosixOSPlatformProcess.cpp
--- a/Engine/Source/Runtime/Core/Private/PosixOS/PosixOSPlatformProcess.cpp
+++ b/Engine/Source/Runtime/Core/Private/PosixOS/PosixOSPlatformProcess.cpp
@@ -61,3 +61,93 @@ void FPosixOSPlatformProcess::CeaseBeingFirstInstance()
 		GFileLockDescriptor = -1;
 	}
 }
+
+
+class FPosixOSProcessSentinel : public IProcessSentinel
+{
+public:
+	
+	FPosixOSProcessSentinel()
+		: fd(-1)
+	{
+		
+	}
+	
+	virtual ~FPosixOSProcessSentinel()
+	{
+		close(fd);
+		unlink(TCHAR_TO_UTF8(*SavedFilename));
+	}
+	
+	virtual bool CreateSentinelFile(const TCHAR* Name, const FString& Contents) override
+	{
+		SavedFilename = MakeFilename(Name);
+
+		fd = open(TCHAR_TO_UTF8(*SavedFilename), O_RDWR | O_CREAT, 0666);
+		if (flock(fd, LOCK_EX | LOCK_NB) != 0)
+		{
+			close(fd);
+			UE_LOG(LogHAL, Error, TEXT("Another instance is running, unable to create sentinel file %s. Will not create or remove sentinel in this process"), *SavedFilename);
+			fd = -1;
+			return false;
+		}
+		
+		// write the contents
+		FTCHARToUTF8 Convert(*Contents);
+		ftruncate(fd, 0);
+		write(fd, Convert.Get(), Convert.Length());
+		
+		UE_LOG(LogHAL, Display, TEXT("Created sentinel %s"), *SavedFilename);				
+		return true;
+	}
+	
+	
+	
+protected:
+	virtual bool DoesSentinelFileExistImpl(const TCHAR* Name, FString* OptionalOutContents) override
+	{
+		FString Filename = MakeFilename(Name);
+
+		// if the file doesn't exist at all, then early out
+		if (!IFileManager::Get().FileExists(*Filename))
+		{
+			return false;
+		}
+		
+		// the file exists, so open a handle to it (not expected to fail)
+		int TempFd = open(TCHAR_TO_UTF8(*Filename), O_RDWR, 0666);
+		if (TempFd == -1)
+		{
+			UE_LOG(LogHAL, Error, TEXT("Didn't expect open to fail %s"), *Filename);		
+			return false;
+		}
+		// attempt to get an exclusive lock - if is able to, if it can that means another process
+		// had created the file but then died before deleting it. 
+		if (flock(TempFd, LOCK_EX | LOCK_NB) == 0)
+		{
+			close(TempFd);
+			// the file shouldn't exist, so clean it up and pretend it never existed
+			IFileManager::Get().Delete(*Filename);
+			return false;
+		}
+		close(TempFd);
+		
+		// get contents if desired, can use normal file ops now
+		if (OptionalOutContents)
+		{
+			FFileHelper::LoadFileToString(*OptionalOutContents, *Filename);
+		}
+		
+		return true;	
+	}
+	
+	// the opened and flock'd file descriptor for this sentinel file
+	int fd;
+	FString SavedFilename;
+};
+
+
+TSharedPtr<IProcessSentinel> FPosixOSPlatformProcess::CreateProcessSentinelObject()
+{
+	return TSharedPtr<IProcessSentinel>(new FPosixOSProcessSentinel());
+}
diff --git a/Engine/Source/Runtime/Core/Private/Windows/WindowsPlatformProcess.cpp b/Engine/Source/Runtime/Core/Private/Windows/WindowsPlatformProcess.cpp
--- a/Engine/Source/Runtime/Core/Private/Windows/WindowsPlatformProcess.cpp
+++ b/Engine/Source/Runtime/Core/Private/Windows/WindowsPlatformProcess.cpp
@@ -2779,6 +2779,86 @@ void FWindowsPlatformProcess::SetThreadPriority(EThreadPriority NewPriority)
 	::SetThreadPriority(::GetCurrentThread(), FRunnableThreadWin::TranslateThreadPriority(NewPriority));
 }
 
+class FWindowsProcessSentinel : public IProcessSentinel
+{
+public:
+	virtual ~FWindowsProcessSentinel()
+	{
+		// this will also delete the file
+		CloseHandle(FileHandle);
+	}
+
+	virtual bool CreateSentinelFile(const TCHAR* Name, const FString& Contents) override
+	{
+		FString Filename = MakeFilename(Name);
+
+		// if the file exists, another process is running, so can't create it
+		if (IFileManager::Get().FileExists(*Filename))
+		{
+			return false;
+		}
+		
+		// create the file in a way that will delete itself when the process closes
+		IFileManager::Get().MakeDirectory(*FPaths::GetPath(Filename), true);
+		DWORD Sharing = FILE_SHARE_READ | FILE_SHARE_DELETE;
+		DWORD Flags = FILE_FLAG_DELETE_ON_CLOSE;
+		FileHandle = CreateFile(*Filename, GENERIC_WRITE, Sharing, nullptr, CREATE_ALWAYS, Flags, nullptr);
+		
+		// write and flush the contents
+		FTCHARToUTF8 Convert(*Contents);
+		DWORD BytesWritten = 0;
+		WriteFile(FileHandle, Convert.Get(), Convert.Length(), &BytesWritten, nullptr);
+		FlushFileBuffers(FileHandle);
+		return FileHandle != nullptr;
+	}
+	
+	virtual bool DoesSentinelFileExistImpl(const TCHAR* Name, FString* OptionalOutContents) override
+	{
+		FString Filename = MakeFilename(Name);
+
+		// must also open with FILE_SHARE_DELETE for the above FILE_FLAG_DELETE_ON_CLOSE
+		DWORD Sharing = FILE_SHARE_DELETE | FILE_SHARE_WRITE;
+		DWORD Flags = 0;
+		HANDLE Handle = CreateFile(*Filename, GENERIC_READ, Sharing, nullptr, OPEN_EXISTING, Flags, nullptr);
+
+		// if this fails to open, it's because the file didn't exist - if it existed with a process running, it would return a value
+		if (Handle == INVALID_HANDLE_VALUE)
+		{
+			if (IFileManager::Get().FileExists(*Filename))
+			{
+				UE_LOG(LogWindows, Error, TEXT("Didn't expect the file to exist of the Handle is invalid!"));
+			}
+
+			return false;
+		}
+
+		if (OptionalOutContents != nullptr)
+		{
+			// get the filesize, without using GetFileSize which was always returning 0 (possibly due ot the FILE_FLAG_DELETE_ON_CLOSE)
+			DWORD FileSize = SetFilePointer(Handle, 0, nullptr, FILE_END);
+			SetFilePointer(Handle, 0, nullptr, FILE_BEGIN);
+
+			// read in the contents
+			ANSICHAR* Buf = (ANSICHAR*)FMemory::Malloc(FileSize + 1);
+			if (ReadFile(Handle, Buf, FileSize, nullptr, nullptr))
+			{
+				Buf[FileSize] = 0;
+				*OptionalOutContents = FString(Buf);
+			}
+			FMemory::Free(Buf);
+
+		}
+		return true;
+	}
+	
+	HANDLE FileHandle;
+};
+
+TSharedPtr<IProcessSentinel> FWindowsPlatformProcess::CreateProcessSentinelObject()
+{
+	return TSharedPtr<IProcessSentinel>(new FWindowsProcessSentinel());
+}
+
 PRAGMA_RESTORE_UNSAFE_TYPECAST_WARNINGS
 
 #include "Windows/HideWindowsPlatformTypes.h"
diff --git a/Engine/Source/Runtime/Core/Public/GenericPlatform/GenericPlatformProcess.h b/Engine/Source/Runtime/Core/Public/GenericPlatform/GenericPlatformProcess.h
--- a/Engine/Source/Runtime/Core/Public/GenericPlatform/GenericPlatformProcess.h
+++ b/Engine/Source/Runtime/Core/Public/GenericPlatform/GenericPlatformProcess.h
@@ -11,6 +11,7 @@
 #include "Misc/EnumClassFlags.h"
 #include "Templates/Function.h"
 #include "Templates/Tuple.h"
+#include "Templates/SharedPointer.h"
 
 class FEvent;
 
@@ -209,6 +210,38 @@ namespace UE::HAL
 	enum struct EProcessId: uint32 {};
 }
 
+
+/**
+  * A helper class that creates and checks for existence of a sentinel file that can be tied to process lifetime. 
+  * NOTE: The generic version will not handle the case where a process just dies without crashing - the file
+  * will not be deleted. The three desktop platforms have subclasses that will make sure that 
+  * DoesSentinelFileExist will return false if the process that created it is not running. DO NOT assume 
+  * FileExists() type calls would be accurate (and also there's not a public method to retrieve the path to 
+  * discourage such a check anyway).
+ */
+class IProcessSentinel
+{
+public:
+	/** Query for existence of a sentinel file (usually from another process than created it) */	
+	static bool CORE_API DoesSentinelFileExist(const TCHAR* Name, FString* OptionalOutContents=nullptr);
+	
+	/**
+	 * Create the sentinel file, if possible 
+	 * @return false if the file could not be created (because another process has it locked)
+	 */
+	virtual bool CreateSentinelFile(const TCHAR* Name, const FString& Contents) = 0;
+
+protected:	
+	/** Helper to create the path for the sentinel file. Defaults to ProcessIntermediate() dir */
+	virtual FString CORE_API MakeFilename(const TCHAR* Name);
+	
+	/** 
+	 * The actual implementation to determine if the sentinel file is exists and is valid (due to a process dying,
+	 * the file may actually exist on disk, but isn't a valid file anymore, so a simple FileExists may not be good enough
+	 */
+	virtual bool DoesSentinelFileExistImpl(const TCHAR* Name, FString* OptionalOutContents=nullptr) = 0;
+};
+
 /**
 * Generic implementation for most platforms, these tend to be unused and unimplemented
 **/
@@ -870,6 +903,11 @@ struct FGenericPlatformProcess
 	 */
 	static CORE_API bool IsFirstInstance();
 
+	/**
+	 * Stop reporting that we are the first instance
+	 */
+	static CORE_API void CeaseBeingFirstInstance();
+
 	/**
 	 * Tears down allocated process resources.
 	 */
@@ -895,6 +933,16 @@ struct FGenericPlatformProcess
 	 */
 	static CORE_API void ModifyThreadAssignmentForUObjectReferenceCollector( int32& NumThreads, int32& NumBackgroundThreads, ENamedThreads::Type& NormalThreadName, ENamedThreads::Type& BackgroundThreadName );
 
+	/** 
+	 * Create a GenericProcessSentinel, or a subclass, object
+	 */
+	static CORE_API TSharedPtr<IProcessSentinel> CreateProcessSentinelObject();
+
+	/** 
+	 * Create a GenericProcessSentinel, or a subclass, object, and then uses it to create the actual sentinel file on disk
+	 */
+	static CORE_API TSharedPtr<IProcessSentinel> CreateProcessSentinel(const TCHAR* Name, const FString& Contents);	
+
 	/**
 	 * Tells the processor to pause for implementation-specific amount of time. Is used for spin-loops to improve the speed at 
 	 * which the code detects the release of the lock and power-consumption.
diff --git a/Engine/Source/Runtime/Core/Public/Mac/MacPlatformProcess.h b/Engine/Source/Runtime/Core/Public/Mac/MacPlatformProcess.h
--- a/Engine/Source/Runtime/Core/Public/Mac/MacPlatformProcess.h
+++ b/Engine/Source/Runtime/Core/Public/Mac/MacPlatformProcess.h
@@ -241,6 +241,10 @@ struct CORE_API FMacPlatformProcess : public FApplePlatformProcess
 	{
 		FPosixOSPlatformProcess::CeaseBeingFirstInstance();
 	}
+	static TSharedPtr<IProcessSentinel> CreateProcessSentinelObject()
+	{
+		return FPosixOSPlatformProcess::CreateProcessSentinelObject();
+	}
 	
 
 	// Mac specific
diff --git a/Engine/Source/Runtime/Core/Public/PosixOS/PosixOSPlatformProcess.h b/Engine/Source/Runtime/Core/Public/PosixOS/PosixOSPlatformProcess.h
--- a/Engine/Source/Runtime/Core/Public/PosixOS/PosixOSPlatformProcess.h
+++ b/Engine/Source/Runtime/Core/Public/PosixOS/PosixOSPlatformProcess.h
@@ -15,4 +15,5 @@ struct FPosixOSPlatformProcess : public FGenericPlatformProcess
 	/** FPlatformProcess API */
 	static CORE_API bool IsFirstInstance();
 	static CORE_API void CeaseBeingFirstInstance();
+	static CORE_API TSharedPtr<IProcessSentinel> CreateProcessSentinelObject();
 };
diff --git a/Engine/Source/Runtime/Core/Public/Windows/WindowsPlatformProcess.h b/Engine/Source/Runtime/Core/Public/Windows/WindowsPlatformProcess.h
--- a/Engine/Source/Runtime/Core/Public/Windows/WindowsPlatformProcess.h
+++ b/Engine/Source/Runtime/Core/Public/Windows/WindowsPlatformProcess.h
@@ -209,6 +209,7 @@ public:
 	static CORE_API void SetupAudioThread();
 	static CORE_API void TeardownAudioThread();
 	static CORE_API bool IsFirstInstance();
+	static CORE_API TSharedPtr<IProcessSentinel> CreateProcessSentinelObject();
 
 	/**
 	 * @brief Releases locks that we held for IsFirstInstance check
