Rename the internal non virtual overload of OpenWrite to OpenWriteInternal.

#rnx

diff --git a/Engine/Source/Runtime/Core/Private/Android/AndroidPlatformFile.cpp b/Engine/Source/Runtime/Core/Private/Android/AndroidPlatformFile.cpp
--- a/Engine/Source/Runtime/Core/Private/Android/AndroidPlatformFile.cpp
+++ b/Engine/Source/Runtime/Core/Private/Android/AndroidPlatformFile.cpp
@@ -1977,10 +1977,10 @@ public:
 	virtual IFileHandle* OpenWrite(const TCHAR* Filename, bool bAppend, bool bAllowRead) override
 	{
 		const bool bForceAllowLocal = false;
-		return OpenWrite(Filename, bAppend, bAllowRead, bForceAllowLocal);
+		return OpenWriteInternal(Filename, bAppend, bAllowRead, bForceAllowLocal);
 	}
 
-	IFileHandle* OpenWrite(const TCHAR* Filename, bool bAppend, bool bAllowRead, bool bForceAllowLocal)
+	IFileHandle* OpenWriteInternal(const TCHAR* Filename, bool bAppend, bool bAllowRead, bool bForceAllowLocal)
 	{
 #if LOG_ANDROID_FILE
 		FPlatformMisc::LowLevelOutputDebugStringf(TEXT("FAndroidPlatformFile::OpenWrite('%s')"), Filename);
