Fixed pipe and task timeout calculations that generated bad code

diff --git a/Engine/Source/Runtime/Core/Private/Tasks/Pipe.cpp b/Engine/Source/Runtime/Core/Private/Tasks/Pipe.cpp
--- a/Engine/Source/Runtime/Core/Private/Tasks/Pipe.cpp
+++ b/Engine/Source/Runtime/Core/Private/Tasks/Pipe.cpp
@@ -84,7 +84,7 @@ namespace UE::Tasks
 				return true;
 			}
 
-			if (!EmptyEventRef->WaitFor(Token, UE::FMonotonicTimeSpan::FromMilliseconds(Timeout.GetRemainingRoundedUpMilliseconds())))
+			if (!EmptyEventRef->WaitFor(Token, Timeout.WillNeverExpire() ? UE::FMonotonicTimeSpan::Infinity() : UE::FMonotonicTimeSpan::FromSeconds(Timeout.GetRemainingSeconds())))
 			{
 				break;
 			}
diff --git a/Engine/Source/Runtime/Core/Private/Tasks/TaskPrivate.cpp b/Engine/Source/Runtime/Core/Private/Tasks/TaskPrivate.cpp
--- a/Engine/Source/Runtime/Core/Private/Tasks/TaskPrivate.cpp
+++ b/Engine/Source/Runtime/Core/Private/Tasks/TaskPrivate.cpp
@@ -256,7 +256,7 @@ namespace UE::Tasks
 
 					// Always flush events before entering a wait to make sure there's nothing missing in Unreal Insights that could prevent us understanding what's going on.
 					TRACE_CPUPROFILER_EVENT_FLUSH();
-					StateChangeEvent.WaitFor(Token, UE::FMonotonicTimeSpan::FromMilliseconds(Timeout.GetRemainingRoundedUpMilliseconds()));
+					StateChangeEvent.WaitFor(Token, Timeout.WillNeverExpire() ? UE::FMonotonicTimeSpan::Infinity() : UE::FMonotonicTimeSpan::FromSeconds(Timeout.GetRemainingSeconds()));
 				}
 
 				// Once the state of the task has changed (either closed or scheduled), it's time to do another round of retraction to help if possible.
