Clamp the wait time in FPlatformManualResetEvent for platforms with lower-precision wait time

diff --git a/Engine/Source/Runtime/Core/Private/Microsoft/MicrosoftPlatformManualResetEvent.cpp b/Engine/Source/Runtime/Core/Private/Microsoft/MicrosoftPlatformManualResetEvent.cpp
--- a/Engine/Source/Runtime/Core/Private/Microsoft/MicrosoftPlatformManualResetEvent.cpp
+++ b/Engine/Source/Runtime/Core/Private/Microsoft/MicrosoftPlatformManualResetEvent.cpp
@@ -34,19 +34,15 @@ FORCENOINLINE bool FMicrosoftManualResetEvent::WaitForSlow(FMonotonicTimeSpan Wa
 
 	LowLevelTasks::FOversubscriptionScope _;
 
-	const DWORD WaitMs = WaitTime.IsInfinity() ? INFINITE : DWORD(FPlatformMath::CeilToInt64(WaitTime.ToMilliseconds()));
-	if (LIKELY(WaitOnAddress(&bWait, &bLocalWait, sizeof(bool), WaitMs)))
-	{
-		const bool bWaitCompleted = !bWait.load(std::memory_order_acquire);
-		if (bWaitCompleted)
-		{
-			return true;
-		}
-	}
-	else if (GetLastError() == ERROR_TIMEOUT)
+	// Clamp to INFINITE. Test against INFINITE - 1 because of the ceiling operation.
+	const double RawWaitMs = WaitTime.ToMilliseconds();
+	const DWORD WaitMs = RawWaitMs > double(INFINITE - 1) ? INFINITE : DWORD(FPlatformMath::CeilToInt64(RawWaitMs));
+
+	const bool bTimedOut = !WaitOnAddress(&bWait, &bLocalWait, sizeof(bool), WaitMs) && GetLastError() == ERROR_TIMEOUT;
+	bLocalWait = bWait.load(std::memory_order_acquire);
+	if (LIKELY(!bLocalWait || bTimedOut))
 	{
-		const bool bWaitCompleted = !bWait.load(std::memory_order_acquire);
-		return bWaitCompleted;
+		return !bLocalWait;
 	}
 
 	// Handle a spurious wake by waiting until the wait time has elapsed one more time because WaitUntilSlow
