Fix material's shading models being incorrect when cooking projects with Substrate materials and hidden conversion.

When cooking, the function CacheResourceShadersForRendering() is not called, causing RebuildShadingModelField() to never being called. Because of this, the ShadingModel and ShadingModels properties are incorrect. This later causes shader code generation to have missing conversion blocs.

diff --git a/Engine/Source/Runtime/Core/Private/UObject/FortniteMainBranchObjectVersion.cpp b/Engine/Source/Runtime/Core/Private/UObject/FortniteMainBranchObjectVersion.cpp
--- a/Engine/Source/Runtime/Core/Private/UObject/FortniteMainBranchObjectVersion.cpp
+++ b/Engine/Source/Runtime/Core/Private/UObject/FortniteMainBranchObjectVersion.cpp
@@ -14,6 +14,6 @@ TMap<FGuid, FGuid> FFortniteMainBranchObjectVersion::GetSystemGuids()
 	SystemGuids.Add(DevGuids.Niagara_LatestScriptCompileVersion, FGuid("783FF3700BB641C799B2D43552FBF91E"));
 	SystemGuids.Add(DevGuids.SkeletalMeshDerivedDataVersion, FGuid("C998F2F3B2884A02B7E290B8C25835F0"));
 	SystemGuids.Add(DevGuids.STATICMESH_DERIVEDDATA_VER, FGuid("3ABF32149F9D4D838750368AB95573BE"));
-	SystemGuids.Add(DevGuids.MaterialTranslationDDCVersion, FGuid("C6B481BEACC04585B05635E3EC491F99"));
+	SystemGuids.Add(DevGuids.MaterialTranslationDDCVersion, FGuid("36B83A3B63AD4B19ABB959FFB26A9F24"));
 	return SystemGuids;
 }
diff --git a/Engine/Source/Runtime/Engine/Private/Materials/Material.cpp b/Engine/Source/Runtime/Engine/Private/Materials/Material.cpp
--- a/Engine/Source/Runtime/Engine/Private/Materials/Material.cpp
+++ b/Engine/Source/Runtime/Engine/Private/Materials/Material.cpp
@@ -3535,10 +3535,12 @@ bool UMaterial::ConvertMaterialToSubstrateMaterial(bool bAllowEmptyMaterialUpdat
 	}
 	else if (bRequestedByUser && EditorOnly->FrontMaterial.IsConnected())
 	{
+		RebuildShadingModelField();
 		return false; // We do not convert if the front material pin has is already connected.
 	}
 	else if(Substrate::IsHiddenMaterialAssetConversionEnabled() && !bRequestedByUser)
 	{
+		RebuildShadingModelField();
 		return false; // This is to work with ShouldAllowSubstrateReset. So that a material authored with IsHiddenMaterialAssetConversionEnabled using the root node inputs can then be converted to substrate.
 	}
 
