Unity compile issue with IFileManager

#rnx

diff --git a/Engine/Source/Runtime/Core/Private/PosixOS/PosixOSPlatformProcess.cpp b/Engine/Source/Runtime/Core/Private/PosixOS/PosixOSPlatformProcess.cpp
--- a/Engine/Source/Runtime/Core/Private/PosixOS/PosixOSPlatformProcess.cpp
+++ b/Engine/Source/Runtime/Core/Private/PosixOS/PosixOSPlatformProcess.cpp
@@ -2,8 +2,10 @@
 
 #include "PosixOS/PosixOSPlatformProcess.h"
 #include "Containers/StringConv.h"
+#include "HAL/FileManager.h"
 #include "HAL/PlatformProcess.h"
 #include "Misc/CommandLine.h"
+#include "Misc/FileHelper.h"
 #include "Misc/Parse.h"
 #include <sys/fcntl.h>
 #include <sys/file.h>
