Fixed potential infinite loop in the serialization of packed ints when loading malformed values.

[RN] Fixed potential infinite loop in the serialization of packed ints when loading malformed values.

diff --git a/Engine/Source/Runtime/Core/Private/Serialization/Archive.cpp b/Engine/Source/Runtime/Core/Private/Serialization/Archive.cpp
--- a/Engine/Source/Runtime/Core/Private/Serialization/Archive.cpp
+++ b/Engine/Source/Runtime/Core/Private/Serialization/Archive.cpp
@@ -1327,22 +1327,35 @@ void FArchive::SerializeIntPacked(uint32& Value)
 {
 	if (IsLoading())
 	{
-		Value = 0;
+		// Set the byte cap to 10 in order to handle serialized 64-bit integers whose type changed to 32-bit.
+		constexpr size_t MaxBytes = 10;
+
+		uint32 TempValue = 0;
 		uint8 cnt = 0;
 		uint8 more = 1;
 		while(more)
 		{
+			if (!ensureMsgf(cnt < MaxBytes, TEXT("SerializeIntPacked tried to load more than %zu packed bytes for a 32-bit value"), MaxBytes))
+			{
+				SetError();
+				return;
+			}
+
 			uint8 NextByte;
 			Serialize(&NextByte, 1);			// Read next byte
 
 			more = NextByte & 1;				// Check 1 bit to see if theres more after this
 			NextByte = NextByte >> 1;			// Shift to get actual 7 bit value
-			Value += NextByte << (7 * cnt++);	// Add to total value
+			TempValue += (uint32)NextByte << (7 * cnt++);	// Add to total value
 		}
+
+		Value = TempValue;
 	}
 	else
 	{
-		uint8 PackedBytes[5];
+		constexpr size_t MaxBytes = 5;
+
+		uint8 PackedBytes[MaxBytes];
 		int32 PackedByteCount = 0;
 		uint32 Remaining = Value;
 		while(true)
@@ -1367,24 +1380,34 @@ void FArchive::SerializeIntPacked(uint32& Value)
 
 void FArchive::SerializeIntPacked64(uint64& Value)
 {
+	constexpr size_t MaxBytes = 10;
+
 	if (IsLoading())
 	{
-		Value = 0;
+		uint64 TempValue = 0;
 		uint8 cnt = 0;
 		uint8 more = 1;
 		while (more)
 		{
+			if (!ensureMsgf(cnt < MaxBytes, TEXT("SerializeIntPacked64 tried to load more than %zu packed bytes for a 64-bit value"), MaxBytes))
+			{
+				SetError();
+				return;
+			}
+
 			uint8 NextByte;
 			Serialize(&NextByte, 1);			// Read next byte
 
 			more = NextByte & 1;				// Check 1 bit to see if theres more after this
 			NextByte = NextByte >> 1;			// Shift to get actual 7 bit value
-			Value += (uint64)NextByte << (7 * cnt++);	// Add to total value
+			TempValue += (uint64)NextByte << (7 * cnt++);	// Add to total value
 		}
+
+		Value = TempValue;
 	}
 	else
 	{
-		uint8 PackedBytes[10];
+		uint8 PackedBytes[MaxBytes];
 		int32 PackedByteCount = 0;
 		uint64 Remaining = Value;
 		while (true)
