Add warnings to GIsPlayInEditorWorld that GWorld may be null when it is true.

diff --git a/Engine/Source/Runtime/Core/Private/Misc/CoreGlobals.cpp b/Engine/Source/Runtime/Core/Private/Misc/CoreGlobals.cpp
--- a/Engine/Source/Runtime/Core/Private/Misc/CoreGlobals.cpp
+++ b/Engine/Source/Runtime/Core/Private/Misc/CoreGlobals.cpp
@@ -390,7 +390,11 @@ const TCHAR* LexToString(ELoaderType Type)
 
 /** Whether the cooker is currently loading a package or not												*/
 bool					GIsCookerLoadingPackage			= false;
-/** Whether GWorld points to the play in editor world														*/
+/**
+ * Whether GWorld should point to the play in editor world.
+ * @warning There are times when GIsPlayInEditorWorld == true and GWorld == nullptr! Notably during LoadMap 
+ * during PIE, which can happen when PIE initiates a Browse to another level.
+ */
 bool					GIsPlayInEditorWorld			= false;
 /** Whether or not PIE was attempting to play from PlayerStart												*/
 bool					GIsPIEUsingPlayerStart			= false;
diff --git a/Engine/Source/Runtime/Core/Public/CoreGlobals.h b/Engine/Source/Runtime/Core/Public/CoreGlobals.h
--- a/Engine/Source/Runtime/Core/Public/CoreGlobals.h
+++ b/Engine/Source/Runtime/Core/Public/CoreGlobals.h
@@ -502,7 +502,11 @@ namespace UE
 /** Whether the cooker is currently loading a package or not */
 extern CORE_API bool GIsCookerLoadingPackage;
 
-/** Whether GWorld points to the play in editor world */
+/**
+ * Whether GWorld should point to the play in editor world.
+ * @warning There are times when GIsPlayInEditorWorld == true and GWorld == nullptr! Notably during LoadMap 
+ * during PIE, which can happen when PIE initiates a Browse to another level.
+ */
 extern CORE_API bool GIsPlayInEditorWorld;
 
 namespace UE
