Comment update suggested by BugHawk

diff --git a/Engine/Source/Runtime/Core/Public/Containers/EnumAsByte.h b/Engine/Source/Runtime/Core/Public/Containers/EnumAsByte.h
--- a/Engine/Source/Runtime/Core/Public/Containers/EnumAsByte.h
+++ b/Engine/Source/Runtime/Core/Public/Containers/EnumAsByte.h
@@ -50,7 +50,7 @@ public:
 	}
 
 	/**
-	 * Constructor, initialize to the int32 value.
+	 * Constructor, initialize to the uint8 value.
 	 *
 	 * @param InValue value to construct with.
 	 */
diff --git a/Engine/Source/Runtime/Core/Public/Misc/AssertionMacros.h b/Engine/Source/Runtime/Core/Public/Misc/AssertionMacros.h
--- a/Engine/Source/Runtime/Core/Public/Misc/AssertionMacros.h
+++ b/Engine/Source/Runtime/Core/Public/Misc/AssertionMacros.h
@@ -479,7 +479,7 @@ namespace UEAsserts_Private
 	bool GetMemberNameCheckedJunk(R(*)(Args...));
 }
 
-// Returns FName(TEXT("EnumeratorName")), while statically verifying that the enumerator exists in the enum
+// Returns FName(TEXTVIEW("EnumeratorName")), while statically verifying that the enumerator exists in the enum
 #define GET_ENUMERATOR_NAME_CHECKED(EnumName, EnumeratorName) \
 	(FName(GET_ENUMERATOR_NAME_STRING_VIEW_CHECKED(EnumName, EnumeratorName)))
 
@@ -489,7 +489,7 @@ namespace UEAsserts_Private
 #define GET_ENUMERATOR_NAME_STRING_VIEW_CHECKED(EnumName, EnumeratorName) \
 	((void)sizeof(UEAsserts_Private::GetMemberNameCheckedJunk(EnumName::EnumeratorName)), TEXTVIEW(#EnumeratorName))
 
-// Returns FName(TEXT("MemberName")), while statically verifying that the member exists in ClassName
+// Returns FName(TEXTVIEW("MemberName")), while statically verifying that the member exists in ClassName
 #define GET_MEMBER_NAME_CHECKED(ClassName, MemberName) \
 	(FName(GET_MEMBER_NAME_STRING_VIEW_CHECKED(ClassName, MemberName)))
 
@@ -499,7 +499,7 @@ namespace UEAsserts_Private
 #define GET_MEMBER_NAME_STRING_VIEW_CHECKED(ClassName, MemberName) \
 	((void)sizeof(UEAsserts_Private::GetMemberNameCheckedJunk(((ClassName*)0)->MemberName)), TEXTVIEW(#MemberName))
 
-// Returns FName(TEXT("FunctionName")), while statically verifying that the function exists in ClassName
+// Returns FName(TEXTVIEW("FunctionName")), while statically verifying that the function exists in ClassName
 #define GET_FUNCTION_NAME_CHECKED(ClassName, FunctionName) \
 	(FName(GET_FUNCTION_NAME_STRING_VIEW_CHECKED(ClassName, FunctionName)))
 
