Added MallocBinned lightweigtht bin allocation callstack tracker. Track allocation callstack when a specific bin memory usage goes over a specific threshold in MB. Disabled and opted out by default (scoped with UE_MBC_LIGHTWEIGHT_BIN_CALLSTACK_TRACKER define)
Added MB2 bin stats to DumpStats function

#rnx

diff --git a/Engine/Source/Runtime/Core/Private/HAL/MallocBinned2.cpp b/Engine/Source/Runtime/Core/Private/HAL/MallocBinned2.cpp
--- a/Engine/Source/Runtime/Core/Private/HAL/MallocBinned2.cpp
+++ b/Engine/Source/Runtime/Core/Private/HAL/MallocBinned2.cpp
@@ -452,6 +452,8 @@ FMallocBinned2::~FMallocBinned2()
 
 void FMallocBinned2::OnMallocInitialized()
 {
+	TMallocBinnedCommon::OnMallocInitialized();
+
 #if UE_USE_VERYLARGEPAGEALLOCATOR
 	FCoreDelegates::GetLowLevelAllocatorMemoryTrimDelegate().AddLambda([this]()
 		{
@@ -466,7 +468,6 @@ void FMallocBinned2::OnMallocInitialized()
 			CachedOSPageAllocator.Refresh();
 		}
 	);
-
 #endif
 }
 
@@ -583,6 +584,10 @@ void* FMallocBinned2::MallocExternalSmall(SIZE_T Size, uint32 Alignment)
 	AllocatedSmallPoolMemory.fetch_add(Table.BinSize, std::memory_order_relaxed);
 #endif
 
+#if UE_MBC_LIGHTWEIGHT_BIN_CALLSTACK_TRACKER
+	TrackBinAllocation(Table.BinSize, Table.TotalAllocatedMem);
+#endif
+
 	if (GMallocBinnedAllocExtra)
 	{
 		if (Lists)
@@ -929,6 +934,14 @@ void FMallocBinned2::DumpAllocatorStats(class FOutputDevice& Ar)
 			) / (1024.0f * 1024.0f));
 	Ar.Logf(TEXT("Cached free OS pages: %fmb"), ((double)OSPageAllocatorCachedFreeSize) / (1024.0f * 1024.0f));
 
+	for (int32 i = 0; i < UE_MB2_SMALL_POOL_COUNT; i++)
+	{
+		const float Fragmentation = 1.0f - (float)SmallPoolTables[i].TotalUsedBins / (float)SmallPoolTables[i].TotalAllocatedBins;
+		const float TotalMem = (float)SmallPoolTables[i].TotalAllocatedMem / 1024.0f / 1024.0f;
+		Ar.Logf(TEXT("Bin %6d Fragmentation %d %%, Wasted Mem %.2f MB, Total Allocated Mem %.2f MB"),
+				PoolIndexToBinSize(i), int(Fragmentation * 100.0f), TotalMem * Fragmentation, TotalMem);
+	}
+
 #if !PLATFORM_UNIX && !PLATFORM_ANDROID
 	// Doesn't implemented
 #else
diff --git a/Engine/Source/Runtime/Core/Private/HAL/MallocBinned3.cpp b/Engine/Source/Runtime/Core/Private/HAL/MallocBinned3.cpp
--- a/Engine/Source/Runtime/Core/Private/HAL/MallocBinned3.cpp
+++ b/Engine/Source/Runtime/Core/Private/HAL/MallocBinned3.cpp
@@ -740,6 +740,11 @@ void* FMallocBinned3::MallocExternal(SIZE_T Size, uint32 Alignment)
 		Table.HeadEndAlloc(Size);
 		AllocatedSmallPoolMemory.fetch_add(Table.BinSize, std::memory_order_relaxed);
 #endif
+
+#if UE_MBC_LIGHTWEIGHT_BIN_CALLSTACK_TRACKER
+		TrackBinAllocation(Table.BinSize, Table.TotalAllocatedMem);
+#endif
+
 		if (GMallocBinnedAllocExtra)
 		{
 			if (Lists)
@@ -1077,7 +1082,7 @@ void FMallocBinned3::DumpAllocatorStats(class FOutputDevice& Ar)
 	{
 		const float Fragmentation = 1.0f - (float)SmallPoolTables[i].TotalUsedBins / (float)SmallPoolTables[i].TotalAllocatedBins;
 		const float TotalMem = (float)SmallPoolTables[i].TotalAllocatedMem / 1024.0f / 1024.0f;
-		Ar.Logf(TEXT("Bin %6d Fragrmentation %d %%, Wasted Mem %.2f MB, Total Allocated Mem %.2f MB"), 
+		Ar.Logf(TEXT("Bin %6d Fragmentation %d %%, Wasted Mem %.2f MB, Total Allocated Mem %.2f MB"), 
 			PoolIndexToBinSize(i), int(Fragmentation * 100.0f), TotalMem * Fragmentation, TotalMem);
 	}
 
diff --git a/Engine/Source/Runtime/Core/Private/HAL/MallocBinnedCommon.cpp b/Engine/Source/Runtime/Core/Private/HAL/MallocBinnedCommon.cpp
--- a/Engine/Source/Runtime/Core/Private/HAL/MallocBinnedCommon.cpp
+++ b/Engine/Source/Runtime/Core/Private/HAL/MallocBinnedCommon.cpp
@@ -11,8 +11,9 @@
 #include "Templates/UnrealTemplate.h"
 #include "ProfilingDebugging/CsvProfiler.h"
 
-#if UE_MBC_LOG_LARGE_ALLOCATION
+#if UE_MBC_LOG_LARGE_ALLOCATION || UE_MBC_LIGHTWEIGHT_BIN_CALLSTACK_TRACKER
 	#include "HAL/PlatformStackWalk.h"
+	#include "Misc/CoreDelegates.h"
 #endif
 
 PRAGMA_DISABLE_UNSAFE_TYPECAST_WARNINGS
@@ -683,4 +684,143 @@ void FMallocBinnedCommonBase::LogLargeAllocation(SIZE_T Size) const
 }
 #endif
 
+#if UE_MBC_LIGHTWEIGHT_BIN_CALLSTACK_TRACKER
+static int32 GMBCBinSizeToTrack = 0;
+static FAutoConsoleVariableRef CVarGMBCBinSizeToTrack(
+	TEXT("MallocBinned.LightWeightBinTracker.BinSize"),
+	GMBCBinSizeToTrack,
+	TEXT("BinSize to track (default is 0, meaning no allocations are tracked)."));
+
+static int32 GMBCTotalBinSizeInMBThresholdToTrack = 0;
+static FAutoConsoleVariableRef CVarGMBCTotalBinSizeInMBThresholdToTrack(
+	TEXT("MallocBinned.LightWeightBinTracker.TotalBinSizeInMBThreshold"),
+	GMBCTotalBinSizeInMBThresholdToTrack,
+	TEXT("Track allocation callstack for bin (specified with MallocBinned.LightWeightTracker.BinSize) when its usage size is over a certain threshold in MB (default is 0, meaning no allocations are tracked)."));
+
+struct FBinCallStackTracker
+{
+	void Track(uint32 BinSize, uint32 TotalAllocatedMem);
+	void ForceCrash();
+
+	static constexpr const uint32 NumFramePointers = 16;
+	static constexpr const uint32 NumCallStacks = 512;
+
+	uint64 AllocationCallStacks[NumCallStacks][NumFramePointers];
+	int32 AllocationCallStacksFrame[NumCallStacks] = { -1 };
+
+	uint32 CallStackIndex = 0;
+	std::atomic<int32> CallStackFrame = 0;
+
+	FCriticalSection Mutex;
+} GMBCBinCallStackTracker;
+
+void FMallocBinnedCommonBase::TrackBinAllocation(uint32 BinSize, uint32 TotalAllocatedMem)
+{
+	GMBCBinCallStackTracker.Track(BinSize, TotalAllocatedMem);
+}
+
+void FBinCallStackTracker::Track(uint32 BinSize, uint32 TotalAllocatedMem)
+{
+	if (GMBCBinSizeToTrack <= 0 || GMBCTotalBinSizeInMBThresholdToTrack <= 0)
+	{
+		return;
+	}
+	
+	int32 TotalAllocatedMemInMb = (TotalAllocatedMem / 1024) / 1024;
+	if (GMBCBinSizeToTrack != BinSize || TotalAllocatedMemInMb < GMBCTotalBinSizeInMBThresholdToTrack)
+	{
+		return;
+	}
+
+	static bool bInFunction = false;
+	if (bInFunction)
+	{
+		return;
+	}
+
+	TGuardValue<bool> GuardInFunction(bInFunction, true);
+	FScopeLock Lock(&Mutex);
+
+	// Grab callstack
+	static constexpr int32 CallStackEntriesToSkipCount = 2;
+	static constexpr int32 MaxCallStackDepth = NumFramePointers + CallStackEntriesToSkipCount;
+
+	uint64 FullCallStack[MaxCallStackDepth] = { 0 };
+	FPlatformStackWalk::CaptureStackBackTrace(&FullCallStack[0], MaxCallStackDepth);
+
+	uint64* CallStackSrc = &FullCallStack[CallStackEntriesToSkipCount];
+	uint64* CallStackDst = &AllocationCallStacks[CallStackIndex][0];
+	FMemory::Memcpy(CallStackDst, CallStackSrc, NumFramePointers * sizeof(uint64));
+
+	AllocationCallStacksFrame[CallStackIndex] = CallStackFrame.load(std::memory_order_relaxed);
+
+	CallStackIndex++;
+	CallStackIndex = CallStackIndex % NumCallStacks;
+}
+
+// Intentionally disabling optimization and pgo profile to ensure breadcrumbing data are copied in the stack
+UE_DISABLE_OPTIMIZATION_SHIP
+UE_NO_PROFILE_ATTRIBUTE
+void FBinCallStackTracker::ForceCrash()
+{
+	// No need to force crash if tracker is disabled.
+	if (GMBCBinSizeToTrack <= 0 || GMBCTotalBinSizeInMBThresholdToTrack <= 0)
+	{
+		return;
+	}
+	
+	FScopeLock Lock(&Mutex);
+
+	// Allocate on stack and copy latest 32 callstack data so we could get them from a minidump
+	static constexpr int NumLocalCallStacks = 32;
+	static_assert(NumLocalCallStacks <= NumCallStacks);
+
+	uint64 StackAllocationCallStacks[NumLocalCallStacks][NumFramePointers];
+	int32 StackAllocationCallStacksFrame[NumLocalCallStacks];
+
+	for (int32 Index = 0; Index < NumLocalCallStacks; Index++)
+	{
+		CallStackIndex--;
+		CallStackIndex = CallStackIndex % NumCallStacks;
+
+		uint64* CallStackSrc = &AllocationCallStacks[CallStackIndex][0];
+		for (int32 FramePointerIndex = 0; FramePointerIndex < NumFramePointers; FramePointerIndex++)
+		{
+			StackAllocationCallStacks[Index][FramePointerIndex] = CallStackSrc[FramePointerIndex];
+		}
+		
+		StackAllocationCallStacksFrame[Index] = AllocationCallStacksFrame[CallStackIndex];
+	}
+
+	UE_LOG(LogMemory, Fatal, TEXT("OOM Force crash requested by MallocBinned Light Weight CallStack tracker, Frame %d"), CallStackFrame.load(std::memory_order_relaxed));
+}
+// Enable back optimization
+UE_ENABLE_OPTIMIZATION_SHIP
+#endif
+
+void FMallocBinnedCommonBase::OnMallocInitialized()
+{
+#if UE_MBC_LIGHTWEIGHT_BIN_CALLSTACK_TRACKER
+	// When OnMallocInitialized is called, global static OnBeginFrame is not yet allocated on some platforms. 
+	// As a workaround OnBeginFrame delegate is added via TSConfigReadyForUse which is called after global static initialization
+	FCoreDelegates::TSConfigReadyForUse().AddLambda([this]()
+		{
+			UE_CALL_ONCE([&]()
+			{
+				FCoreDelegates::OnBeginFrame.AddLambda([this]()
+					{
+						GMBCBinCallStackTracker.CallStackFrame++;
+						GMBCBinCallStackTracker.CallStackFrame &= 0xFFFF; // Wrap it around 65536 to get it more human readable
+					});
+			});
+		});
+
+
+	FCoreDelegates::GetOutOfMemoryDelegate().AddLambda([this]()
+		{
+			GMBCBinCallStackTracker.ForceCrash();
+		});
+#endif
+}
+
 PRAGMA_RESTORE_UNSAFE_TYPECAST_WARNINGS
diff --git a/Engine/Source/Runtime/Core/Public/HAL/MallocBinnedCommon.h b/Engine/Source/Runtime/Core/Public/HAL/MallocBinnedCommon.h
--- a/Engine/Source/Runtime/Core/Public/HAL/MallocBinnedCommon.h
+++ b/Engine/Source/Runtime/Core/Public/HAL/MallocBinnedCommon.h
@@ -97,6 +97,14 @@
 #	define UE_MBC_LOG_LARGE_ALLOCATION 0
 #endif
 
+#ifndef UE_MBC_LIGHTWEIGHT_BIN_CALLSTACK_TRACKER
+#	define UE_MBC_LIGHTWEIGHT_BIN_CALLSTACK_TRACKER		0
+#endif
+
+#if UE_MBC_LIGHTWEIGHT_BIN_CALLSTACK_TRACKER && !UE_MBC_ALLOCATOR_STATS
+#	error "MB lightweight bin callstack tracker needs UE_MBC_ALLOCATOR_STATS to be enabled."
+#endif
+
 CSV_DECLARE_CATEGORY_EXTERN(MallocBinned);
 
 extern CORE_API float GMallocBinnedFlushThreadCacheMaxWaitTime;
@@ -188,6 +196,8 @@ public:
 
 	static_assert(sizeof(FBundleNode) <= UE_MBC_MIN_SMALL_POOL_ALIGNMENT, "Bundle nodes must fit into the smallest block size");
 
+	CORE_API virtual void OnMallocInitialized() override; 
+
 protected:
 	struct FPtrToPoolMapping
 	{
@@ -418,9 +428,13 @@ protected:
 	}
 
 #if UE_MBC_LOG_LARGE_ALLOCATION
-		void LogLargeAllocation(SIZE_T Size) const;
+	void LogLargeAllocation(SIZE_T Size) const;
 #else
-		void LogLargeAllocation(SIZE_T Size) const {};
+	void LogLargeAllocation(SIZE_T Size) const {};
+#endif
+
+#if UE_MBC_LIGHTWEIGHT_BIN_CALLSTACK_TRACKER
+	void TrackBinAllocation(uint32 BinSize, uint32 TotalAllocatedMem);
 #endif
 };
 
