Fix potential division by zero when dumping Mallocbinned stats

#rnx

diff --git a/Engine/Source/Runtime/Core/Private/HAL/MallocBinned2.cpp b/Engine/Source/Runtime/Core/Private/HAL/MallocBinned2.cpp
--- a/Engine/Source/Runtime/Core/Private/HAL/MallocBinned2.cpp
+++ b/Engine/Source/Runtime/Core/Private/HAL/MallocBinned2.cpp
@@ -936,7 +936,7 @@ void FMallocBinned2::DumpAllocatorStats(class FOutputDevice& Ar)
 
 	for (int32 i = 0; i < UE_MB2_SMALL_POOL_COUNT; i++)
 	{
-		const float Fragmentation = 1.0f - (float)SmallPoolTables[i].TotalUsedBins / (float)SmallPoolTables[i].TotalAllocatedBins;
+		const float Fragmentation = (SmallPoolTables[i].TotalAllocatedBins > 0) ? 1.0f - (float)SmallPoolTables[i].TotalUsedBins / (float)SmallPoolTables[i].TotalAllocatedBins : 0;
 		const float TotalMem = (float)SmallPoolTables[i].TotalAllocatedMem / 1024.0f / 1024.0f;
 		Ar.Logf(TEXT("Bin %6d Fragmentation %d %%, Wasted Mem %.2f MB, Total Allocated Mem %.2f MB"),
 				PoolIndexToBinSize(i), int(Fragmentation * 100.0f), TotalMem * Fragmentation, TotalMem);
diff --git a/Engine/Source/Runtime/Core/Private/HAL/MallocBinned3.cpp b/Engine/Source/Runtime/Core/Private/HAL/MallocBinned3.cpp
--- a/Engine/Source/Runtime/Core/Private/HAL/MallocBinned3.cpp
+++ b/Engine/Source/Runtime/Core/Private/HAL/MallocBinned3.cpp
@@ -1080,7 +1080,7 @@ void FMallocBinned3::DumpAllocatorStats(class FOutputDevice& Ar)
 
 	for (int32 i = 0; i < UE_MB3_SMALL_POOL_COUNT; i++)
 	{
-		const float Fragmentation = 1.0f - (float)SmallPoolTables[i].TotalUsedBins / (float)SmallPoolTables[i].TotalAllocatedBins;
+		const float Fragmentation = (SmallPoolTables[i].TotalAllocatedBins > 0) ? 1.0f - (float)SmallPoolTables[i].TotalUsedBins / (float)SmallPoolTables[i].TotalAllocatedBins : 0;
 		const float TotalMem = (float)SmallPoolTables[i].TotalAllocatedMem / 1024.0f / 1024.0f;
 		Ar.Logf(TEXT("Bin %6d Fragmentation %d %%, Wasted Mem %.2f MB, Total Allocated Mem %.2f MB"), 
 			PoolIndexToBinSize(i), int(Fragmentation * 100.0f), TotalMem * Fragmentation, TotalMem);
