Enforcing the IO Error Code String Array to be the same size as the IO Error Code enum

#rnx

diff --git a/Engine/Source/Runtime/Core/Private/IO/IoDispatcher.cpp b/Engine/Source/Runtime/Core/Private/IO/IoDispatcher.cpp
--- a/Engine/Source/Runtime/Core/Private/IO/IoDispatcher.cpp
+++ b/Engine/Source/Runtime/Core/Private/IO/IoDispatcher.cpp
@@ -59,8 +59,10 @@ static const TCHAR* const GetIoErrorText_ErrorCodeTextArray[] =
 	TEXT("Disabled"),
 	TEXT("Not Installed"),
 	TEXT("Pending Host Group"),
-	TEXT("Timeout")
+	TEXT("Timeout"),
+	TEXT("Delete Error")
 };
+static_assert(std::size(GetIoErrorText_ErrorCodeTextArray) == static_cast<uint32>(EIoErrorCode::Last), "GetIoErrorText_ErrorCodeTextArray is missing entries to match EIoErrorCode. Please update the array.");
 
 const TCHAR* const* GetIoErrorText_ErrorCodeText = GetIoErrorText_ErrorCodeTextArray;
 
