Extended FIoStatus with system error code and changed fixed size error message to shared string reducing the size from 132 to 16 bytes.

#rnx

diff --git a/Engine/Source/Runtime/Core/Private/IO/IoStatus.cpp b/Engine/Source/Runtime/Core/Private/IO/IoStatus.cpp
--- a/Engine/Source/Runtime/Core/Private/IO/IoStatus.cpp
+++ b/Engine/Source/Runtime/Core/Private/IO/IoStatus.cpp
@@ -39,50 +39,69 @@ FIoStatus::~FIoStatus()
 {
 }
 
-FIoStatus::FIoStatus(EIoErrorCode Code)
-:	ErrorCode(Code)
+FIoStatus::FIoStatus(EIoErrorCode InErrorCode)
+	: ErrorCode(InErrorCode)
 {
-	ErrorMessage[0] = TEXT('\0');
 }
 
-FIoStatus::FIoStatus(EIoErrorCode Code, const FStringView& InErrorMessage)
-: ErrorCode(Code)
+FIoStatus::FIoStatus(EIoErrorCode InErrorCode, uint32 InSystemErrorCode)
+	: ErrorCode(InErrorCode)
+	, SystemErrorCode(InSystemErrorCode)
+{
+}
+
+FIoStatus::FIoStatus(EIoErrorCode InErrorCode, const FStringView& InErrorMessage)
+	: ErrorMessage(InErrorMessage)
+	, ErrorCode(InErrorCode)
+{
+}
+
+FIoStatus::FIoStatus(EIoErrorCode InErrorCode, uint32 InSystemErrorCode, const FStringView& InErrorMessage)
+	: ErrorMessage(InErrorMessage)
+	, ErrorCode(InErrorCode)
+	, SystemErrorCode(InSystemErrorCode)
 {
-	const int32 ErrorMessageLength = FMath::Min(MaxErrorMessageLength - 1, InErrorMessage.Len());
-	FPlatformString::Convert(ErrorMessage, ErrorMessageLength, InErrorMessage.GetData(), ErrorMessageLength);
-	ErrorMessage[ErrorMessageLength] = TEXT('\0');
 }
 
 FIoStatus& FIoStatus::operator=(const FIoStatus& Other)
 {
+	ErrorMessage = Other.ErrorMessage;
 	ErrorCode = Other.ErrorCode;
-	FMemory::Memcpy(ErrorMessage, Other.ErrorMessage, MaxErrorMessageLength * sizeof(TCHAR));
+	SystemErrorCode	= Other.SystemErrorCode;
 
 	return *this;
 }
 
 FIoStatus& FIoStatus::operator=(const EIoErrorCode InErrorCode)
 {
+	ErrorMessage.Reset();
 	ErrorCode = InErrorCode;
-	ErrorMessage[0] = TEXT('\0');
+	SystemErrorCode = 0;
 
 	return *this;
 }
 
 bool FIoStatus::operator==(const FIoStatus& Other) const
 {
-	return ErrorCode == Other.ErrorCode &&
-		FPlatformString::Stricmp(ErrorMessage, Other.ErrorMessage) == 0;
+	return ErrorMessage == Other.ErrorMessage && ErrorCode == Other.ErrorCode && SystemErrorCode == Other.SystemErrorCode;
 }
 
 FString FIoStatus::ToString() const
 {
-	if (ErrorMessage[0] == TEXT('\0'))
+	if (ErrorMessage.IsEmpty())
 	{
-		return FString::Format(TEXT("({0})"), { GetIoErrorText(ErrorCode) });
+		if (LIKELY(SystemErrorCode == 0))
+		{
+			return FString::Format(TEXT("({0})"), { GetIoErrorText(ErrorCode) });
+		}
+		return FString::Format(TEXT("({0}, SystemErrorCode={1})"), { GetIoErrorText(ErrorCode), SystemErrorCode });
 	}
 
-	return FString::Format(TEXT("{0} ({1})"), { ErrorMessage, GetIoErrorText(ErrorCode) });
+	if (LIKELY(SystemErrorCode == 0))
+	{
+		return FString::Format(TEXT("{0} ({1})"), { *ErrorMessage, GetIoErrorText(ErrorCode) });
+	}
+	return FString::Format(TEXT("{0} ({1}, SystemErrorCode={2})"), { *ErrorMessage, GetIoErrorText(ErrorCode), SystemErrorCode });
 }
 
 void StatusOrCrash(const FIoStatus& Status)
diff --git a/Engine/Source/Runtime/Core/Public/IO/IoStatus.h b/Engine/Source/Runtime/Core/Public/IO/IoStatus.h
--- a/Engine/Source/Runtime/Core/Public/IO/IoStatus.h
+++ b/Engine/Source/Runtime/Core/Public/IO/IoStatus.h
@@ -3,6 +3,7 @@
 #pragma once
 
 #include "Containers/StringView.h"
+#include "Containers/SharedString.h"
 #include "Containers/UnrealString.h"
 #include "Misc/StringBuilder.h"
 
@@ -47,7 +48,7 @@ CORE_API const TCHAR* GetIoErrorText(EIoErrorCode ErrorCode);
 ///////////////////////////////////////////////////////////////////////////////
 
 /**
- * I/O status with error code and message.
+ * I/O status with error code, error message and an optional system error code.
  */
 class FIoStatus
 {
@@ -55,10 +56,12 @@ public:
 	CORE_API			FIoStatus();
 	CORE_API			~FIoStatus();
 
-	CORE_API			FIoStatus(EIoErrorCode Code, const FStringView& InErrorMessage);
-	CORE_API			FIoStatus(EIoErrorCode Code);
+	CORE_API			FIoStatus(EIoErrorCode ErrorCode, const FStringView& ErrorMessage);
+	CORE_API			FIoStatus(EIoErrorCode ErrorCode, uint32 SystemErrorCode, const FStringView& ErrorMessage);
+	CORE_API			FIoStatus(EIoErrorCode ErrorCode, uint32 SystemErrorCode);
+	CORE_API			FIoStatus(EIoErrorCode ErrorCode);
 	CORE_API FIoStatus&	operator=(const FIoStatus& Other);
-	CORE_API FIoStatus&	operator=(const EIoErrorCode InErrorCode);
+	CORE_API FIoStatus&	operator=(const EIoErrorCode ErrorCode);
 
 	CORE_API bool		operator==(const FIoStatus& Other) const;
 			 bool		operator!=(const FIoStatus& Other) const { return !operator==(Other); }
@@ -66,19 +69,18 @@ public:
 	inline bool			IsOk() const { return ErrorCode == EIoErrorCode::Ok; }
 	inline bool			IsCompleted() const { return ErrorCode != EIoErrorCode::Unknown; }
 	inline EIoErrorCode	GetErrorCode() const { return ErrorCode; }
+	inline uint32		GetSystemErrorCode() const { return SystemErrorCode; }
+	const TCHAR*		GetErrorMessage() const { return ErrorMessage.IsEmpty() ? nullptr : *ErrorMessage; }
 	CORE_API FString	ToString() const;
 
 	CORE_API static const FIoStatus Ok;
 	CORE_API static const FIoStatus Unknown;
 	CORE_API static const FIoStatus Invalid;
 
-	static constexpr int32 MaxErrorMessageLength = 128;
-	using FErrorMessage = TCHAR[MaxErrorMessageLength];
-
-	const FErrorMessage& GetErrorMessage() const { return ErrorMessage; }
 private:
-	EIoErrorCode	ErrorCode = EIoErrorCode::Ok;
-	FErrorMessage	ErrorMessage;
+	UE::FSharedString	ErrorMessage;
+	EIoErrorCode		ErrorCode = EIoErrorCode::Ok;
+	uint32				SystemErrorCode = 0;
 
 	friend class FIoStatusBuilder;
 };
@@ -117,10 +119,10 @@ public:
 	template<typename U>
 	TIoStatusOr<T>& operator=(const TIoStatusOr<U>& Other);
 
-	const FIoStatus& Status() const;
+	const FIoStatus&	Status() const;
 	bool				IsOk() const;
 
-	const T& ValueOrDie();
+	const T&			ValueOrDie();
 	T					ConsumeValueOrDie();
 
 	void				Reset();
@@ -139,7 +141,7 @@ class FIoStatusBuilder
 {
 protected:
 	EIoErrorCode		StatusCode;
-	TStringBuilder<FIoStatus::MaxErrorMessageLength> Message;
+	TStringBuilder<256> Message;
 public:
 	CORE_API explicit	FIoStatusBuilder(EIoErrorCode StatusCode);
 
