Reduce log spam for ios low memory warnings

#rnx

diff --git a/Engine/Source/Runtime/Core/Private/IOS/IOSPlatformMisc.cpp b/Engine/Source/Runtime/Core/Private/IOS/IOSPlatformMisc.cpp
--- a/Engine/Source/Runtime/Core/Private/IOS/IOSPlatformMisc.cpp
+++ b/Engine/Source/Runtime/Core/Private/IOS/IOSPlatformMisc.cpp
@@ -972,10 +972,24 @@ bool FIOSPlatformMisc::HasMemoryWarningHandler()
 
 void FIOSPlatformMisc::HandleLowMemoryWarning()
 {
+#if UE_BUILD_SHIPPING || UE_BUILD_TEST
+	static double LastPrintTime = 0;
+	const double CurrentTime = FPlatformTime::Seconds();
+	const double CheckInterval = 0.05; // Wait at least 50ms before re-printing information
+
+	if (CurrentTime >= LastPrintTime + CheckInterval)
+	{
+		LastPrintTime = CurrentTime;
+#endif
+
 	UE_LOG(LogInit, Log, TEXT("Low Memory Warning Triggered"));
 	UE_LOG(LogInit, Log, TEXT("Free Memory at Startup: %d MB"), GStartupFreeMemoryMB);
 	UE_LOG(LogInit, Log, TEXT("Free Memory Now       : %d MB"), GetFreeMemoryMB());
 
+#if UE_BUILD_SHIPPING || UE_BUILD_TEST
+	}
+#endif
+
 	if(GMemoryWarningHandler != NULL)
 	{
 		FGenericMemoryWarningContext Context;
