- Close the handle after querying for a sentinel, was leaking handles before

diff --git a/Engine/Source/Runtime/Core/Private/Windows/WindowsPlatformProcess.cpp b/Engine/Source/Runtime/Core/Private/Windows/WindowsPlatformProcess.cpp
--- a/Engine/Source/Runtime/Core/Private/Windows/WindowsPlatformProcess.cpp
+++ b/Engine/Source/Runtime/Core/Private/Windows/WindowsPlatformProcess.cpp
@@ -2826,7 +2826,7 @@ public:
 		{
 			if (IFileManager::Get().FileExists(*Filename))
 			{
-				UE_LOG(LogWindows, Error, TEXT("Didn't expect the file to exist of the Handle is invalid!"));
+				UE_LOG(LogWindows, Error, TEXT("WindowsSentinel: Didn't expect the file to exist if the Handle is invalid!"));
 			}
 
 			return false;
@@ -2848,6 +2848,8 @@ public:
 			FMemory::Free(Buf);
 
 		}
+
+		CloseHandle(Handle);
 		return true;
 	}
 	
