Extended FIoStatusBuilder to accept SystemErrorCode argument.

#rnx

diff --git a/Engine/Source/Runtime/Core/Private/IO/IoStatus.cpp b/Engine/Source/Runtime/Core/Private/IO/IoStatus.cpp
--- a/Engine/Source/Runtime/Core/Private/IO/IoStatus.cpp
+++ b/Engine/Source/Runtime/Core/Private/IO/IoStatus.cpp
@@ -112,7 +112,14 @@ void StatusOrCrash(const FIoStatus& Status)
 //////////////////////////////////////////////////////////////////////////
 
 FIoStatusBuilder::FIoStatusBuilder(EIoErrorCode InStatusCode)
-:	StatusCode(InStatusCode)
+	: StatusCode(InStatusCode)
+	, SystemErrorCode(0)
+{
+}
+
+FIoStatusBuilder::FIoStatusBuilder(EIoErrorCode InStatusCode, uint32 InSystemErrorCode)
+	: StatusCode(InStatusCode)
+	, SystemErrorCode(InSystemErrorCode)
 {
 }
 
@@ -132,7 +139,7 @@ FIoStatusBuilder::~FIoStatusBuilder()
 
 FIoStatusBuilder::operator FIoStatus()
 {
-	return FIoStatus(StatusCode, Message);
+	return FIoStatus(StatusCode, SystemErrorCode);
 }
 
 FIoStatusBuilder& FIoStatusBuilder::operator<<(FStringView String)
diff --git a/Engine/Source/Runtime/Core/Public/IO/IoStatus.h b/Engine/Source/Runtime/Core/Public/IO/IoStatus.h
--- a/Engine/Source/Runtime/Core/Public/IO/IoStatus.h
+++ b/Engine/Source/Runtime/Core/Public/IO/IoStatus.h
@@ -141,9 +141,11 @@ class FIoStatusBuilder
 {
 protected:
 	EIoErrorCode		StatusCode;
+	uint32				SystemErrorCode;
 	TStringBuilder<256> Message;
 public:
 	CORE_API explicit	FIoStatusBuilder(EIoErrorCode StatusCode);
+	CORE_API explicit	FIoStatusBuilder(EIoErrorCode StatusCode, uint32 SystemErrorCode);
 
 	CORE_API			FIoStatusBuilder(const TCHAR* NamespaceStr, EIoErrorCode StatusCode, 
 							const uint_least32_t Line = __builtin_LINE(),
