Fixed bug with missing arguments in conversion operators in FIoStatusBuilder.

#rnx

diff --git a/Engine/Source/Runtime/Core/Private/IO/IoStatus.cpp b/Engine/Source/Runtime/Core/Private/IO/IoStatus.cpp
--- a/Engine/Source/Runtime/Core/Private/IO/IoStatus.cpp
+++ b/Engine/Source/Runtime/Core/Private/IO/IoStatus.cpp
@@ -139,7 +139,7 @@ FIoStatusBuilder::~FIoStatusBuilder()
 
 FIoStatusBuilder::operator FIoStatus()
 {
-	return FIoStatus(StatusCode, SystemErrorCode);
+	return FIoStatus(StatusCode, SystemErrorCode, Message);
 }
 
 FIoStatusBuilder& FIoStatusBuilder::operator<<(FStringView String)
diff --git a/Engine/Source/Runtime/Core/Public/IO/IoStatus.h b/Engine/Source/Runtime/Core/Public/IO/IoStatus.h
--- a/Engine/Source/Runtime/Core/Public/IO/IoStatus.h
+++ b/Engine/Source/Runtime/Core/Public/IO/IoStatus.h
@@ -158,7 +158,7 @@ public:
 	template<typename T>
 	operator TIoStatusOr<T>()
 	{
-		return FIoStatus(StatusCode, Message);
+		return FIoStatus(StatusCode, SystemErrorCode, Message);
 	}
 
 	CORE_API FIoStatusBuilder& operator<<(FStringView String);
