Separating out core ddc versions in order to facilitate p4 locking.

diff --git a/Engine/Source/Runtime/Core/Private/UObject/FortniteMainBranchObjectVersion.cpp b/Engine/Source/Runtime/Core/Private/UObject/FortniteMainBranchObjectVersion.cpp
--- a/Engine/Source/Runtime/Core/Private/UObject/FortniteMainBranchObjectVersion.cpp
+++ b/Engine/Source/Runtime/Core/Private/UObject/FortniteMainBranchObjectVersion.cpp
@@ -1,6 +1,8 @@
 // Copyright Epic Games, Inc. All Rights Reserved.
 #include "UObject/FortniteMainBranchObjectVersion.h"
 
+#include "FortniteMainBranchObjectVersion_Restricted.inl"
+
 TMap<FGuid, FGuid> FFortniteMainBranchObjectVersion::GetSystemGuids()
 {
 	TMap<FGuid, FGuid> SystemGuids;
@@ -9,11 +11,15 @@ TMap<FGuid, FGuid> FFortniteMainBranchObjectVersion::GetSystemGuids()
 	SystemGuids.Add(DevGuids.GLOBALSHADERMAP_DERIVEDDATA_VER, FGuid("2ACD1C7798984B98BECABF4D32E4AA20"));
 	SystemGuids.Add(DevGuids.LANDSCAPE_MOBILE_COOK_VERSION, FGuid("32D02EF867C74B71A0D4E0FA41392732"));
 	SystemGuids.Add(DevGuids.MATERIALSHADERMAP_DERIVEDDATA_VER, FGuid("F051877974B44D409435F321BDB39DBB"));
-	SystemGuids.Add(DevGuids.NANITE_DERIVEDDATA_VER, FGuid("ACF98184978742A592537DA223A1E6BE"));
 	SystemGuids.Add(DevGuids.NIAGARASHADERMAP_DERIVEDDATA_VER, FGuid("6360A977062842A29ED30E3A7ACB0E64"));
 	SystemGuids.Add(DevGuids.Niagara_LatestScriptCompileVersion, FGuid("783FF3700BB641C799B2D43552FBF91E"));
-	SystemGuids.Add(DevGuids.SkeletalMeshDerivedDataVersion, FGuid("C998F2F3B2884A02B7E290B8C25835F0"));
-	SystemGuids.Add(DevGuids.STATICMESH_DERIVEDDATA_VER, FGuid("3ABF32149F9D4D838750368AB95573BE"));
 	SystemGuids.Add(DevGuids.MaterialTranslationDDCVersion, FGuid("9D0538021BDB47DAA63FFC09DDAAC00E"));
+
+	// These GUIDs are in AppendRestrictedFortniteMainStreamObjectVersionGuids as they should rarely be changed.
+	//SystemGuids.Add(DevGuids.NANITE_DERIVEDDATA_VER, FGuid("ACF98184978742A592537DA223A1E6BE"));
+	//SystemGuids.Add(DevGuids.SkeletalMeshDerivedDataVersion, FGuid("C998F2F3B2884A02B7E290B8C25835F0"));
+	//SystemGuids.Add(DevGuids.STATICMESH_DERIVEDDATA_VER, FGuid("3ABF32149F9D4D838750368AB95573BE"));
+		
+	AppendRestrictedFortniteMainStreamObjectVersionGuids(SystemGuids);
 	return SystemGuids;
 }
diff --git a/Engine/Source/Runtime/Core/Private/UObject/FortniteMainBranchObjectVersion_Restricted.inl b/Engine/Source/Runtime/Core/Private/UObject/FortniteMainBranchObjectVersion_Restricted.inl
new file mode 100644
--- /dev/null
+++ b/Engine/Source/Runtime/Core/Private/UObject/FortniteMainBranchObjectVersion_Restricted.inl
@@ -0,0 +1,11 @@
+// Copyright Epic Games, Inc. All Rights Reserved.
+
+static void AppendRestrictedFortniteMainStreamObjectVersionGuids(TMap<FGuid, FGuid>& SystemGuids)
+{
+	// Guids added here are expected to only be bumped after careful consideration of the
+	// patch effects.
+	const FDevSystemGuids& DevGuids = FDevSystemGuids::Get();
+	SystemGuids.Add(DevGuids.NANITE_DERIVEDDATA_VER, FGuid("ACF98184978742A592537DA223A1E6BE"));
+	SystemGuids.Add(DevGuids.SkeletalMeshDerivedDataVersion, FGuid("C998F2F3B2884A02B7E290B8C25835F0"));
+	SystemGuids.Add(DevGuids.STATICMESH_DERIVEDDATA_VER, FGuid("3ABF32149F9D4D838750368AB95573BE"));
+}
diff --git a/Engine/Source/Runtime/Core/Private/UObject/UE5MainStreamObjectVersion.cpp b/Engine/Source/Runtime/Core/Private/UObject/UE5MainStreamObjectVersion.cpp
--- a/Engine/Source/Runtime/Core/Private/UObject/UE5MainStreamObjectVersion.cpp
+++ b/Engine/Source/Runtime/Core/Private/UObject/UE5MainStreamObjectVersion.cpp
@@ -3,6 +3,7 @@
 #include "Containers/Map.h"
 #include "Misc/Guid.h"
 
+#include "UE5MainStreamObjectVersion_Restricted.inl"
 
 TMap<FGuid, FGuid> FUE5MainStreamObjectVersion::GetSystemGuids()
 {
@@ -12,15 +13,19 @@ TMap<FGuid, FGuid> FUE5MainStreamObjectVersion::GetSystemGuids()
 	SystemGuids.Add(DevGuids.GLOBALSHADERMAP_DERIVEDDATA_VER, FGuid("FEAB692FBBDD4113857BCD32C2B644BC"));
 	SystemGuids.Add(DevGuids.LANDSCAPE_MOBILE_COOK_VERSION, FGuid("71000000000000000000000000000035"));
 	SystemGuids.Add(DevGuids.MATERIALSHADERMAP_DERIVEDDATA_VER, FGuid("D28F2FDCCAC4499B8F8E87728235A46C"));
-	SystemGuids.Add(DevGuids.NANITE_DERIVEDDATA_VER, FGuid("49D2EF76967340EFB22F2F398E4DDDDD"));
 	SystemGuids.Add(DevGuids.NIAGARASHADERMAP_DERIVEDDATA_VER, FGuid("8A37C45D24F2423CBE5F8F371DE33575"));
 	SystemGuids.Add(DevGuids.Niagara_LatestScriptCompileVersion, FGuid("9C8812406D824DE29BBC6F303964BB7E"));
-	SystemGuids.Add(DevGuids.SkeletalMeshDerivedDataVersion, FGuid("CFE0AF3B4F48465E8DE6A6BD81881F18"));
-	SystemGuids.Add(DevGuids.STATICMESH_DERIVEDDATA_VER, FGuid("98F7D79A5811013825E75BBCC41ED3E9"));
 	SystemGuids.Add(DevGuids.POSESEARCHDB_DERIVEDDATA_VER, FGuid("4E595C2AC5E947D6BA9ABC874353E5BC"));
 	SystemGuids.Add(DevGuids.GROOM_BINDING_DERIVED_DATA_VERSION, FGuid("156678D4F9084D7CAE98FADC6AB93573"));
 	SystemGuids.Add(DevGuids.GROOM_DERIVED_DATA_VERSION, FGuid("00B8207248CC44B88F7DEC8F328A5D8B"));
 	SystemGuids.Add(DevGuids.MaterialTranslationDDCVersion, FGuid("5E6E21D1AC364AACA1921EE46BF10B4C"));
 
+	// These GUIDS are in restricted as they should rarely be changed.
+	//SystemGuids.Add(DevGuids.NANITE_DERIVEDDATA_VER, FGuid("49D2EF76967340EFB22F2F398E4DDDDD"));
+	//SystemGuids.Add(DevGuids.SkeletalMeshDerivedDataVersion, FGuid("CFE0AF3B4F48465E8DE6A6BD81881F18"));
+	//SystemGuids.Add(DevGuids.STATICMESH_DERIVEDDATA_VER, FGuid("98F7D79A5811013825E75BBCC41ED3E9"));
+		
+	AppendRestrictedUE5MainStreamObjectVersionGuids(SystemGuids);
+
 	return SystemGuids;
 }
diff --git a/Engine/Source/Runtime/Core/Private/UObject/UE5MainStreamObjectVersion_Restricted.inl b/Engine/Source/Runtime/Core/Private/UObject/UE5MainStreamObjectVersion_Restricted.inl
new file mode 100644
--- /dev/null
+++ b/Engine/Source/Runtime/Core/Private/UObject/UE5MainStreamObjectVersion_Restricted.inl
@@ -0,0 +1,11 @@
+// Copyright Epic Games, Inc. All Rights Reserved.
+
+static void AppendRestrictedUE5MainStreamObjectVersionGuids(TMap<FGuid, FGuid>& SystemGuids)
+{
+	// Guids added here are expected to only be bumped after careful consideration of the
+	// patch effects.
+	const FDevSystemGuids& DevGuids = FDevSystemGuids::Get();
+	SystemGuids.Add(DevGuids.NANITE_DERIVEDDATA_VER, FGuid("49D2EF76967340EFB22F2F398E4DDDDD"));
+	SystemGuids.Add(DevGuids.SkeletalMeshDerivedDataVersion, FGuid("CFE0AF3B4F48465E8DE6A6BD81881F18"));
+	SystemGuids.Add(DevGuids.STATICMESH_DERIVEDDATA_VER, FGuid("98F7D79A5811013825E75BBCC41ED3E9"));
+}
