Load tags using line by line read and custom regex instead of config object
Default config reader doesn't read correctly when there's spaces in the test name

#rnx

diff --git a/Engine/Source/Runtime/Core/Private/Misc/AutomationTest.cpp b/Engine/Source/Runtime/Core/Private/Misc/AutomationTest.cpp
--- a/Engine/Source/Runtime/Core/Private/Misc/AutomationTest.cpp
+++ b/Engine/Source/Runtime/Core/Private/Misc/AutomationTest.cpp
@@ -17,6 +17,7 @@
 #include "Misc/App.h"
 #include "Misc/CommandLine.h"
 #include "Misc/ConfigCacheIni.h"
+#include "Misc/FileHelper.h"
 #include "Misc/OutputDeviceRedirector.h"
 #include "Misc/Paths.h"
 #include "Misc/ScopedSlowTask.h"
@@ -956,17 +957,31 @@ void FAutomationTestFramework::LoadTestTagMappings(const FString& BasePath, bool
 		if (IFileManager::Get().FileExists(*TagsSettingsFilePath))
 		{
 			UE_LOG(LogAutomationTest, Log, TEXT("Loading tags from %s."), *TagsSettingsFilePath);
-			FConfigFile TagConfig;
-			TagConfig.Read(TagsSettingsFilePath);
-
-			const FConfigSection* TagSection = TagConfig.FindSection(TEXT("TestTagMapping"));
-			TArray<FName> TestNames;
-			TagSection->GenerateKeyArray(TestNames);
-			FString TestTags;
-			for (FName TestName : TestNames)
+			TArray<FString> Lines;
+			if (FFileHelper::LoadFileToStringArray(Lines, *TagsSettingsFilePath))
 			{
-				TagSection->GetString(*TestName.ToString(), TestTags);
-				RegisterAutomationTestTags(TestName.ToString(),  TestTags, false);
+				const FRegexPattern FileNameToTagPattern(TEXT("(.*)=(\\[.+\\])+"), ERegexPatternFlags::CaseInsensitive);
+				for (FString Line : Lines)
+				{
+					if (Line.StartsWith(TEXT(";")))
+					{
+						continue;
+					}
+					FRegexMatcher Matcher(FileNameToTagPattern, Line);
+					if (Matcher.FindNext())
+					{
+						FString TestName = Matcher.GetCaptureGroup(1);
+						FString TestTags = Matcher.GetCaptureGroup(2);
+						if (!TestName.IsEmpty() && !TestTags.IsEmpty())
+						{
+							RegisterAutomationTestTags(TestName, TestTags, false);
+						}
+					}
+				}
+			}
+			else
+			{
+				UE_LOG(LogAutomationTest, Error, TEXT("Could not load mapping file %s"), *TagsSettingsFilePath);
 			}
 		}
 	}
