COMPILE_WARNING reports as error when WarningsAsErrors is being used

diff --git a/Engine/Source/Runtime/Core/Public/HAL/Platform.h b/Engine/Source/Runtime/Core/Public/HAL/Platform.h
--- a/Engine/Source/Runtime/Core/Public/HAL/Platform.h
+++ b/Engine/Source/Runtime/Core/Public/HAL/Platform.h
@@ -942,12 +942,20 @@
 #ifdef _MSC_VER
 	#define MSC_FORMAT_DIAGNOSTIC_HELPER_2(x) #x
 	#define MSC_FORMAT_DIAGNOSTIC_HELPER(x) MSC_FORMAT_DIAGNOSTIC_HELPER_2(x)
-	#define COMPILE_WARNING(x) __pragma(message(__FILE__ "(" MSC_FORMAT_DIAGNOSTIC_HELPER(__LINE__) "): warning: " x))
 	#define COMPILE_ERROR(x) __pragma(message(__FILE__ "(" MSC_FORMAT_DIAGNOSTIC_HELPER(__LINE__) "): error: " x))
+	#if UE_WARNINGS_AS_ERRORS
+		#define COMPILE_WARNING(x) COMPILE_ERROR(x)
+	#else
+		#define COMPILE_WARNING(x) __pragma(message(__FILE__ "(" MSC_FORMAT_DIAGNOSTIC_HELPER(__LINE__) "): warning: " x))
+	#endif
 #else
 	#define GCC_DIAGNOSTIC_HELPER(x) _Pragma(#x)
-	#define COMPILE_WARNING(x) GCC_DIAGNOSTIC_HELPER(GCC warning x)
 	#define COMPILE_ERROR(x) GCC_DIAGNOSTIC_HELPER(GCC error x)
+	#if UE_WARNINGS_AS_ERRORS
+		#define COMPILE_WARNING(x) COMPILE_ERROR(x)
+	#else
+		#define COMPILE_WARNING(x) GCC_DIAGNOSTIC_HELPER(GCC warning x)
+	#endif
 #endif
 
 // Tells the compiler to put the decorated function in a certain section (aka. segment) of the executable.
