Fix a nearly 10 year old comment to appease the analyzer
#rnx

diff --git a/Engine/Source/Runtime/Core/Public/HAL/Platform.h b/Engine/Source/Runtime/Core/Public/HAL/Platform.h
--- a/Engine/Source/Runtime/Core/Public/HAL/Platform.h
+++ b/Engine/Source/Runtime/Core/Public/HAL/Platform.h
@@ -938,7 +938,7 @@
 	#define PLATFORM_CACHE_LINE_SIZE	64
 #endif
 
-// Compile-time warnings and errors. Use these as "#pragma COMPILER_WARNING("XYZ")". GCC does not expand macro parameters to _Pragma, so we can't wrap the #pragma part.
+// Compile-time warnings and errors. Use these as "#pragma COMPILE_WARNING("XYZ")". GCC does not expand macro parameters to _Pragma, so we can't wrap the #pragma part.
 #ifdef _MSC_VER
 	#define MSC_FORMAT_DIAGNOSTIC_HELPER_2(x) #x
 	#define MSC_FORMAT_DIAGNOSTIC_HELPER(x) MSC_FORMAT_DIAGNOSTIC_HELPER_2(x)
