Fix non-unity when compiling FoundationTests

#rnx

diff --git a/Engine/Source/Developer/Zen/Tests/ZenTest.cpp b/Engine/Source/Developer/Zen/Tests/ZenTest.cpp
--- a/Engine/Source/Developer/Zen/Tests/ZenTest.cpp
+++ b/Engine/Source/Developer/Zen/Tests/ZenTest.cpp
@@ -9,6 +9,7 @@
 #include "Experimental/ZenServerInterface.h"
 #include "HAL/FileManager.h"
 #include "Misc/Paths.h"
+#include "Misc/ScopeExit.h"
 
 namespace UE
 {
diff --git a/Engine/Source/Runtime/Core/Tests/Algo/UniqueTest.cpp b/Engine/Source/Runtime/Core/Tests/Algo/UniqueTest.cpp
--- a/Engine/Source/Runtime/Core/Tests/Algo/UniqueTest.cpp
+++ b/Engine/Source/Runtime/Core/Tests/Algo/UniqueTest.cpp
@@ -4,7 +4,9 @@
 
 #include "Algo/Unique.h"
 #include "Containers/Array.h"
-
+#include "Containers/ArrayView.h"
+#include "Containers/UnrealString.h"
+#include "Templates/Tuple.h"
 #include "Tests/TestHarnessAdapter.h"
 
 TEST_CASE_NAMED(FUniqueTest, "System::Core::Algo::Unique", "[ApplicationContextMask][SmokeFilter]")
diff --git a/Engine/Source/Runtime/Core/Tests/HAL/PlatformAtomicsTest.cpp b/Engine/Source/Runtime/Core/Tests/HAL/PlatformAtomicsTest.cpp
--- a/Engine/Source/Runtime/Core/Tests/HAL/PlatformAtomicsTest.cpp
+++ b/Engine/Source/Runtime/Core/Tests/HAL/PlatformAtomicsTest.cpp
@@ -2,7 +2,9 @@
 
 #include "CoreTypes.h"
 #include "Misc/AssertionMacros.h"
+#include "Templates/Function.h"
 #include "Tests/TestHarnessAdapter.h"
+#include "Containers/UnrealString.h"
 
 #if WITH_TESTS
 
diff --git a/Engine/Source/Runtime/Core/Tests/Misc/AtomicTest.cpp b/Engine/Source/Runtime/Core/Tests/Misc/AtomicTest.cpp
--- a/Engine/Source/Runtime/Core/Tests/Misc/AtomicTest.cpp
+++ b/Engine/Source/Runtime/Core/Tests/Misc/AtomicTest.cpp
@@ -1,10 +1,13 @@
 // Copyright Epic Games, Inc. All Rights Reserved.
 
 #include "CoreTypes.h"
+#include "Math/UnrealMathUtility.h"
 #include "Misc/AssertionMacros.h"
 #include "Templates/Atomic.h"
 #include "Templates/EnableIf.h"
+#include "Templates/UnrealTemplate.h"
 #include "Traits/IntType.h"
+#include "Containers/UnrealString.h"
 
 #if WITH_TESTS
 
diff --git a/Engine/Source/Runtime/Core/Tests/Misc/OutputDeviceRedirectorTest.cpp b/Engine/Source/Runtime/Core/Tests/Misc/OutputDeviceRedirectorTest.cpp
--- a/Engine/Source/Runtime/Core/Tests/Misc/OutputDeviceRedirectorTest.cpp
+++ b/Engine/Source/Runtime/Core/Tests/Misc/OutputDeviceRedirectorTest.cpp
@@ -3,7 +3,7 @@
 #include "Misc/OutputDeviceRedirector.h"
 
 #if WITH_TESTS
-
+#include "CoreGlobals.h"
 #include "Tests/TestHarnessAdapter.h"
 
 TEST_CASE_NAMED(FOutputDeviceRedirectorTest, "Core::Logging::OutputDeviceRedirector", "[Core][EngineFilter]")
diff --git a/Engine/Source/Runtime/Core/Tests/Misc/PlatformStringConvertTests.cpp b/Engine/Source/Runtime/Core/Tests/Misc/PlatformStringConvertTests.cpp
--- a/Engine/Source/Runtime/Core/Tests/Misc/PlatformStringConvertTests.cpp
+++ b/Engine/Source/Runtime/Core/Tests/Misc/PlatformStringConvertTests.cpp
@@ -1,7 +1,10 @@
 // Copyright Epic Games, Inc. All Rights Reserved.
 
 #include "CoreTypes.h"
+#include "HAL/PlatformString.h"
+#include "Templates/UniquePtr.h"
 #include "Tests/TestHarnessAdapter.h"
+#include "Containers/UnrealString.h"
 #include <type_traits>
 
 #if WITH_TESTS
diff --git a/Engine/Source/Runtime/Core/Tests/Misc/TVariantTest.cpp b/Engine/Source/Runtime/Core/Tests/Misc/TVariantTest.cpp
--- a/Engine/Source/Runtime/Core/Tests/Misc/TVariantTest.cpp
+++ b/Engine/Source/Runtime/Core/Tests/Misc/TVariantTest.cpp
@@ -3,12 +3,13 @@
 #if WITH_TESTS
 
 #include "CoreTypes.h"
-#include "Templates/SharedPointer.h"
-#include "Templates/UniquePtr.h"
-#include "Tests/TestHarnessAdapter.h"
+#include "Math/UnrealMath.h"
 #include "Misc/TVariant.h"
 #include "Serialization/MemoryReader.h"
 #include "Serialization/MemoryWriter.h"
+#include "Templates/SharedPointer.h"
+#include "Templates/UniquePtr.h"
+#include "Tests/TestHarnessAdapter.h"
 
 namespace UE::Core::Private
 {
diff --git a/Engine/Source/Runtime/CoreUObject/Tests/ObjectPinningTest.cpp b/Engine/Source/Runtime/CoreUObject/Tests/ObjectPinningTest.cpp
--- a/Engine/Source/Runtime/CoreUObject/Tests/ObjectPinningTest.cpp
+++ b/Engine/Source/Runtime/CoreUObject/Tests/ObjectPinningTest.cpp
@@ -6,6 +6,8 @@
 #include "TestMacros/Assertions.h"
 #include "Tests/Benchmark.h"
 #include "Tasks/Task.h"
+#include "UObject/Package.h"
+#include "UObject/UObjectGlobals.h"
 
 namespace UE::CoreObject::Private::Tests
 {
