Allow custom crash handling for GPU crashes

#12648

diff --git a/Engine/Source/Runtime/Core/Private/Windows/WindowsPlatformCrashContext.cpp b/Engine/Source/Runtime/Core/Private/Windows/WindowsPlatformCrashContext.cpp
--- a/Engine/Source/Runtime/Core/Private/Windows/WindowsPlatformCrashContext.cpp
+++ b/Engine/Source/Runtime/Core/Private/Windows/WindowsPlatformCrashContext.cpp
@@ -1972,7 +1972,7 @@ FORCENOINLINE void ReportGPUCrash(const TCHAR* ErrorMessage, void* ProgramCounte
 		ULONG_PTR Arguments[] = { (ULONG_PTR)&Info };
 		::RaiseException(GPUCrashExceptionCode, 0, UE_ARRAY_COUNT(Arguments), Arguments);
 	}
-	__except (ReportCrash(GetExceptionInformation()))
+	__except (FPlatformMisc::GetCrashHandlingType() == ECrashHandlingType::Default ? ReportCrash(GetExceptionInformation()) : EXCEPTION_CONTINUE_SEARCH)	
 	{
 		FPlatformMisc::RequestExit(false, TEXT("WindowsPlatformCrashContext.ReportGPUCrash"));
 	}
