PR #13686: Export `FAdditionalCrashContextStack::ExecuteProviders` to allow third-party crash handling systems to make use of additional crash context data

diff --git a/Engine/Source/Runtime/Core/Public/GenericPlatform/GenericPlatformCrashContext.h b/Engine/Source/Runtime/Core/Public/GenericPlatform/GenericPlatformCrashContext.h
--- a/Engine/Source/Runtime/Core/Public/GenericPlatform/GenericPlatformCrashContext.h
+++ b/Engine/Source/Runtime/Core/Public/GenericPlatform/GenericPlatformCrashContext.h
@@ -298,6 +298,8 @@ struct FSharedCrashContext
  */
 struct FCrashContextExtendedWriter
 {
+	virtual ~FCrashContextExtendedWriter() = default;
+
 	/** Adds a named buffer to the report. Intended for larger payloads. */
 	virtual void AddBuffer(const TCHAR* Identifier, const uint8* Data, uint32 DataSize) = 0;
 
@@ -919,7 +921,7 @@ struct FAdditionalCrashContextStack
 	CORE_API static void PushProvider(struct FScopedAdditionalCrashContextProvider* Provider);
 	CORE_API static void PopProvider();
 
-	static void ExecuteProviders(FCrashContextExtendedWriter& Writer);
+	CORE_API static void ExecuteProviders(FCrashContextExtendedWriter& Writer);
 
 private:
 	enum { MaxStackDepth = 16 };
