Channel callbacks

* Add callback pre channel state change. Allows instrumentations to deny state changes channels. For user feedback, instrumentations are encouraged to provide a clear text reason for the denial.
* Add callback post channel state change. Allows instrumentation to perform actions when a channel is enabled or disabled, such as sending initial state, resetting tracked state or activating tracking.
* Updated TraceAux to expose denial messages in certain situations.
* Added result message type in TraceControl in order to communicate failure to set channel states.

diff --git a/Engine/Source/Runtime/Core/Private/ProfilingDebugging/TraceAuxiliary.cpp b/Engine/Source/Runtime/Core/Private/ProfilingDebugging/TraceAuxiliary.cpp
--- a/Engine/Source/Runtime/Core/Private/ProfilingDebugging/TraceAuxiliary.cpp
+++ b/Engine/Source/Runtime/Core/Private/ProfilingDebugging/TraceAuxiliary.cpp
@@ -173,9 +173,9 @@ public:
 	void ResetCommandlineChannels();
 	bool HasCommandlineChannels() const { return !CommandlineChannels.IsEmpty(); }
 	void EnableChannels(const TCHAR* ChannelList, const FTraceAuxiliary::FLogCategoryAlias& LogCategory);
-	void EnableChannels(TConstArrayView<uint32> ChannelIds);
+	void EnableChannels(TConstArrayView<uint32> ChannelIds, TMap<uint32, FString>* OutReason);
 	void DisableChannels(const TCHAR* ChannelList, const FTraceAuxiliary::FLogCategoryAlias& LogCategory);
-	void DisableChannels(TConstArrayView<uint32> ChannelIds);
+	void DisableChannels(TConstArrayView<uint32> ChannelIds, TMap<uint32, FString>* OutReason);
 	bool Connect(FTraceAuxiliary::EConnectionType Type, const TCHAR* Parameter, const FTraceAuxiliary::FLogCategoryAlias& LogCategory, uint16 SendFlags);
 	bool ConnectRelay(UPTRINT Handle, UE::Trace::IoWriteFunc WriteFunc, UE::Trace::IoCloseFunc CloseFunc, uint16 SendFlags);
 	bool Stop();
@@ -208,16 +208,16 @@ private:
 	};
 
 	void AddCommandlineChannel(const TCHAR* Name, const FTraceAuxiliary::FLogCategoryAlias& LogCategory);
-	template <class T> void ForEachChannel(const TCHAR* ChannelList, bool bResolvePresets, const FTraceAuxiliary::FLogCategoryAlias& LogCategory, T Callable);
+	template<typename T, typename... Ts> void ForEachChannel(const TCHAR* ChannelList, bool bResolvePresets, T&& Callable, Ts&&... Args);
 	static uint32 HashChannelName(const TCHAR* Name);
 	bool SendToHost(const TCHAR* Host, const FTraceAuxiliary::FLogCategoryAlias& LogCategory, uint16 SendFlags);
 	bool WriteToFile(const TCHAR* Path, const FTraceAuxiliary::FLogCategoryAlias& LogCategory, uint16 SendFlags);
 	bool FinalizeFilePath(const TCHAR* InPath, FString& OutPath, const FTraceAuxiliary::FLogCategoryAlias& LogCategory);
 
 	template<typename ChannelType>
-	bool EnableChannel(ChannelType Channel, const FTraceAuxiliary::FLogCategoryAlias& LogCategory);
+	bool EnableChannel(ChannelType Channel, const FTraceAuxiliary::FLogCategoryAlias& LogCategory, FString* OutReason = nullptr);
 	template<typename ChannelType>
-	bool DisableChannel(ChannelType Channel, const FTraceAuxiliary::FLogCategoryAlias& LogCategory);
+	bool DisableChannel(ChannelType Channel, const FTraceAuxiliary::FLogCategoryAlias& LogCategory, FString* OutReason = nullptr);
 
 	typedef TMap<uint32, FChannelEntry, TInlineSetAllocator<128>> ChannelSet;
 	ChannelSet CommandlineChannels;
@@ -271,7 +271,7 @@ void FTraceAuxiliaryImpl::AddCommandlineChannel(const TCHAR* Name, const FTraceA
 ////////////////////////////////////////////////////////////////////////////////
 void FTraceAuxiliaryImpl::AddCommandlineChannels(const TCHAR* ChannelList)
 {
-	ForEachChannel(ChannelList, true, LogTrace, &FTraceAuxiliaryImpl::AddCommandlineChannel);
+	ForEachChannel(ChannelList, true, &FTraceAuxiliaryImpl::AddCommandlineChannel, LogTrace);
 }
 
 ////////////////////////////////////////////////////////////////////////////////
@@ -285,16 +285,21 @@ void FTraceAuxiliaryImpl::EnableChannels(const TCHAR* ChannelList, const FTraceA
 {
 	if (ChannelList)
 	{
-		ForEachChannel(ChannelList, true, LogCategory, &FTraceAuxiliaryImpl::EnableChannel<const TCHAR*>);
+		ForEachChannel(ChannelList, true, &FTraceAuxiliaryImpl::EnableChannel<const TCHAR*>, LogCategory, nullptr);
 	}
 }
 
 ////////////////////////////////////////////////////////////////////////////////
-void FTraceAuxiliaryImpl::EnableChannels(TConstArrayView<uint32> ChannelIds)
+void FTraceAuxiliaryImpl::EnableChannels(TConstArrayView<uint32> ChannelIds, TMap<uint32, FString>* OutReason)
 {
 	for (const auto ChannelId : ChannelIds)
 	{
-		EnableChannel(ChannelId, LogTrace);
+		FString DenyReason;
+		EnableChannel(ChannelId, LogTrace, &DenyReason);
+		if (OutReason && !DenyReason.IsEmpty())
+		{
+			OutReason->Add(ChannelId, DenyReason);
+		}
 	}
 }
 
@@ -304,32 +309,43 @@ void FTraceAuxiliaryImpl::DisableChannels(const TCHAR* ChannelList, const FTrace
 {
 	if (ChannelList)
 	{
-		ForEachChannel(ChannelList, true, LogCategory, &FTraceAuxiliaryImpl::DisableChannel<const TCHAR*>);
+		ForEachChannel(ChannelList, true, &FTraceAuxiliaryImpl::DisableChannel<const TCHAR*>, LogCategory, nullptr);
 	}
 	else
 	{
 		// Disable all channels.
 		TStringBuilder<128> EnabledChannels;
 		GetActiveChannelsString(EnabledChannels);
-		ForEachChannel(EnabledChannels.ToString(), true, LogCategory, &FTraceAuxiliaryImpl::DisableChannel<const TCHAR*>);
+		ForEachChannel(EnabledChannels.ToString(), true, &FTraceAuxiliaryImpl::DisableChannel<const TCHAR*>, LogCategory, nullptr);
 	}
 }
 
 ////////////////////////////////////////////////////////////////////////////////
-void FTraceAuxiliaryImpl::DisableChannels(TConstArrayView<uint32> ChannelIds)
+void FTraceAuxiliaryImpl::DisableChannels(TConstArrayView<uint32> ChannelIds, TMap<uint32, FString>* OutReason)
 {
 	for (const auto ChannelId : ChannelIds)
 	{
-		DisableChannel(ChannelId, LogTrace);
+		FString DenyReason;
+		DisableChannel(ChannelId, LogTrace, &DenyReason);
+		if (OutReason && !DenyReason.IsEmpty())
+		{
+			OutReason->Add(ChannelId, DenyReason);
+		}
 	}
 }
 
 ////////////////////////////////////////////////////////////////////////////////
-template<typename T>
-void FTraceAuxiliaryImpl::ForEachChannel(const TCHAR* ChannelList, bool bResolvePresets, const FTraceAuxiliary::FLogCategoryAlias& LogCategory, T Callable)
+
+/**
+ * Parses channel tokens from a "chan1, chan2, ...." string and expands presets
+ * to an output array.
+ * @param ChannelList Input string
+ * @param bResolvePreset If presets should be resolved further
+ * @param OutTokens Output tokens. Lifetime of view are the same as the input views.
+ */
+void ParseTokens(FStringView ChannelList, bool bResolvePresets, TArray<FStringView>& OutTokens, bool bToOwnedStrings, TArray<FString>& OutOwnedTokens)
 {
-	check(ChannelList);
-	UE::String::ParseTokens(ChannelList, TEXT(","), [this, bResolvePresets, Callable, &LogCategory] (const FStringView& Token)
+	UE::String::ParseTokens(ChannelList, TEXT(","), [&] (const FStringView& Token)
 	{
 		TCHAR Name[80];
 		const size_t ChannelNameSize = Token.CopyString(Name, UE_ARRAY_COUNT(Name) - 1);
@@ -337,35 +353,61 @@ void FTraceAuxiliaryImpl::ForEachChannel(const TCHAR* ChannelList, bool bResolve
 
 		if (bResolvePresets)
 		{
-			FString Value;
+			FString ConfigChannelPreset;
 			// Check against hard coded presets
 			if (FCString::Stricmp(Name, GDefaultChannels.Name) == 0)
 			{
-				ForEachChannel(GDefaultChannels.ChannelList, false, LogCategory, Callable);
+				ParseTokens(GDefaultChannels.ChannelList, false, OutTokens, false, OutOwnedTokens);
+				return;
 			}
 			else if (FCString::Stricmp(Name,GMemoryChannels.Name) == 0)
 			{
-				ForEachChannel(GMemoryChannels.ChannelList, false, LogCategory, Callable);
+				ParseTokens(GMemoryChannels.ChannelList, false, OutTokens, false, OutOwnedTokens);
+				return;
 			}
 			else if (FCString::Stricmp(Name, GMemoryLightChannels.Name) == 0)
 			{
-				ForEachChannel(GMemoryLightChannels.ChannelList, false, LogCategory, Callable);
+				ParseTokens(GMemoryLightChannels.ChannelList, false, OutTokens, false, OutOwnedTokens);
+				return;
 			}
 			// Check against data driven presets (if available)
-			else if (GConfig && GConfig->GetString(TEXT("Trace.ChannelPresets"), Name, Value, GEngineIni))
+			else if (GConfig && GConfig->GetString(TEXT("Trace.ChannelPresets"), Name, ConfigChannelPreset, GEngineIni))
 			{
-				ForEachChannel(*Value, false, LogCategory, Callable);
+				// Since ConfigChannelPreset is temporary string we need to output
+				// those to the owned tokens
+				ParseTokens(*ConfigChannelPreset, false, OutTokens, true, OutOwnedTokens);
 				return;
 			}
 		}
-
-		Invoke(Callable, this, Name, LogCategory);
+		bToOwnedStrings ? OutOwnedTokens.Emplace(Token) : OutTokens.Add(Token);
 	});
 }
 
+////////////////////////////////////////////////////////////////////////////////
+template<typename T, typename... Ts>
+void FTraceAuxiliaryImpl::ForEachChannel(const TCHAR* ChannelList, bool bResolvePresets, T&& Callable, Ts&&... Args)
+{
+	check(ChannelList);
+	TArray<FStringView> ExpandedChannels;
+	TArray<FString> ExpandedChannelStrings;
+	ParseTokens(ChannelList, bResolvePresets, ExpandedChannels, false, ExpandedChannelStrings);
+	for (const FStringView& Channel : ExpandedChannels)
+	{
+		TCHAR Name[80];
+		const size_t ChannelNameSize = Channel.CopyString(Name, UE_ARRAY_COUNT(Name) - 1);
+		Name[ChannelNameSize] = '\0';
+
+		Invoke(Callable, this, Name, std::forward<Ts>(Args)...);
+	}
+	for (const FString& Channel : ExpandedChannelStrings)
+	{
+		Invoke(Callable, this, *Channel, std::forward<Ts>(Args)...);
+	}
+}
+
 ////////////////////////////////////////////////////////////////////////////////
 template <typename ChannelType>
-bool FTraceAuxiliaryImpl::EnableChannel(ChannelType ChannelIdentifier, const FTraceAuxiliary::FLogCategoryAlias& LogCategory)
+bool FTraceAuxiliaryImpl::EnableChannel(ChannelType ChannelIdentifier, const FTraceAuxiliary::FLogCategoryAlias& LogCategory, FString* OutReason)
 {
 	// Channel names have been provided by the user and may not exist yet. As
 	// we want to maintain bActive accurately (channels toggles are reference
@@ -373,6 +415,11 @@ bool FTraceAuxiliaryImpl::EnableChannel(ChannelType ChannelIdentifier, const FTr
 	UE::Trace::FChannel* Channel = UE::Trace::FindChannel(ChannelIdentifier);
 	if (!Channel)
 	{
+		if (OutReason)
+		{
+			FString ChannelStr = LexToString(ChannelIdentifier);
+			OutReason->Appendf(TEXT("Unknown channel %s"), *ChannelStr);
+		}
 		return false;
 	}
 
@@ -398,21 +445,44 @@ bool FTraceAuxiliaryImpl::EnableChannel(ChannelType ChannelIdentifier, const FTr
 		return Channel->IsEnabled();
 	}
 
-	const bool bIsEnabled = Channel->Toggle(true);
-	FPlatformEventsTrace::OnTraceChannelUpdated(ChannelName.ToString(), bIsEnabled);
-
+	const TCHAR* DenyReason = nullptr;
+	const bool bIsEnabled = Channel->Toggle(true, &DenyReason);
+	if (DenyReason)
+	{
+		if (OutReason)
+		{
+			OutReason->Append(DenyReason);
+		}
+		else
+		{
+			UE_LOG_REF(
+				LogCategory, 
+				Warning, 
+				TEXT("Cannot enable channel '%s': %s"), 
+				ChannelName.ToString(), 
+				DenyReason
+			);
+		}
+	}
+	
+	UE_CLOG(bIsEnabled, LogTrace, Verbose, TEXT("Enabled channel '%s'."), ChannelName.ToString());
 	return bIsEnabled;
 }
 
 ////////////////////////////////////////////////////////////////////////////////
 template <typename ChannelType>
-bool FTraceAuxiliaryImpl::DisableChannel(ChannelType ChannelIdentifier, const FTraceAuxiliary::FLogCategoryAlias& LogCategory)
+bool FTraceAuxiliaryImpl::DisableChannel(ChannelType ChannelIdentifier, const FTraceAuxiliary::FLogCategoryAlias& LogCategory, FString* OutReason)
 {
 	// Channel names have been provided by the user and may not exist yet. As
 	// we want to maintain bActive accurately we will first check Trace knows of the channel.
 	UE::Trace::FChannel* Channel = UE::Trace::FindChannel(ChannelIdentifier);
 	if (!Channel)
 	{
+		if (OutReason)
+		{
+			FString ChannelStr = LexToString(ChannelIdentifier);
+			OutReason->Appendf(TEXT("Unknown channel %s"), *ChannelStr);
+		}
 		return false;
 	}
 
@@ -438,9 +508,26 @@ bool FTraceAuxiliaryImpl::DisableChannel(ChannelType ChannelIdentifier, const FT
 		return Channel->IsEnabled();
 	}
 
-	const bool bIsEnabled = Channel->Toggle(false);
-	FPlatformEventsTrace::OnTraceChannelUpdated(ChannelName.ToString(), bIsEnabled);
+	const TCHAR* DenyReason = nullptr;
+	const bool bIsEnabled = Channel->Toggle(false, &DenyReason);
+	if (DenyReason)
+	{
+		if (OutReason)
+		{
+			OutReason->Append(DenyReason);
+		}
+		else
+		{
+			UE_LOG_REF(
+				LogCategory,
+				Warning,
+				TEXT("Cannot disable channel '%s': %s"),
+				ChannelName.ToString(), DenyReason
+			);
+		}
+	}
 
+	UE_CLOG(!bIsEnabled, LogTrace, Verbose, TEXT("Disabled channel '%s'."), ChannelName.ToString());
 	return bIsEnabled;
 }
 
@@ -592,7 +679,7 @@ void FTraceAuxiliaryImpl::FreezeReadOnlyChannels()
 void FTraceAuxiliaryImpl::ResumeChannels()
 {
 	// Enable channels from the "paused" preset.
-	ForEachChannel(*PausedPreset, false, LogTrace, &FTraceAuxiliaryImpl::EnableChannel<const TCHAR*>);
+	ForEachChannel(*PausedPreset, false, &FTraceAuxiliaryImpl::EnableChannel<const TCHAR*>, LogTrace, nullptr);
 
 	PausedPreset.Empty();
 }
@@ -608,7 +695,7 @@ void FTraceAuxiliaryImpl::PauseChannels()
 	PausedPreset = EnabledChannels.ToString();
 
 	// Disable all "paused" channels.
-	ForEachChannel(*PausedPreset, true, LogTrace, &FTraceAuxiliaryImpl::DisableChannel<const TCHAR*>);
+	ForEachChannel(*PausedPreset, true, &FTraceAuxiliaryImpl::DisableChannel<const TCHAR*>, LogTrace, nullptr);
 }
 
 ////////////////////////////////////////////////////////////////////////////////
@@ -2045,17 +2132,17 @@ void FTraceAuxiliary::EnableCommandlineChannels()
 #endif
 }
 
-void FTraceAuxiliary::EnableChannels(TConstArrayView<uint32> ChannelIds)
+void FTraceAuxiliary::EnableChannels(TConstArrayView<uint32> ChannelIds, TMap<uint32, FString>* OutErrors)
 {
 #if UE_TRACE_ENABLED
-	GTraceAuxiliary.EnableChannels(ChannelIds);
+	GTraceAuxiliary.EnableChannels(ChannelIds, OutErrors);
 #endif
 }
 
-void FTraceAuxiliary::DisableChannels(TConstArrayView<uint32> ChannelIds)
+void FTraceAuxiliary::DisableChannels(TConstArrayView<uint32> ChannelIds, TMap<uint32, FString>* OutErrors)
 {
 #if UE_TRACE_ENABLED
-	GTraceAuxiliary.DisableChannels(ChannelIds);
+	GTraceAuxiliary.DisableChannels(ChannelIds, OutErrors);
 #endif
 }
 
diff --git a/Engine/Source/Runtime/Core/Public/ProfilingDebugging/TraceAuxiliary.h b/Engine/Source/Runtime/Core/Public/ProfilingDebugging/TraceAuxiliary.h
--- a/Engine/Source/Runtime/Core/Public/ProfilingDebugging/TraceAuxiliary.h
+++ b/Engine/Source/Runtime/Core/Public/ProfilingDebugging/TraceAuxiliary.h
@@ -214,16 +214,16 @@ public:
 	/**
 	 * Enable channels to emit events belonging to this category.
 	 * @note Note that presets cannot be used when specifying channel ids.
-	 * @param ChannelIds List of channels (or a preset) to enable. 
+	 * @param ChannelIds List of channels to enable. 
 	 */
-	static CORE_API void EnableChannels(TConstArrayView<uint32> ChannelIds);
+	static CORE_API void EnableChannels(TConstArrayView<uint32> ChannelIds, TMap<uint32, FString>* OutErrors = nullptr);
 
 	/**
 	 * Disable channels to mute events belonging to this category.
 	 * @note Note that presets cannot be used when specifying channel ids.
-	 * @param ChannelIds List of channels (or a preset) to disable. 
+	 * @param ChannelIds List of channels to disable. 
 	 */
-	static CORE_API void DisableChannels(TConstArrayView<uint32> ChannelIds);
+	static CORE_API void DisableChannels(TConstArrayView<uint32> ChannelIds, TMap<uint32, FString>* OutErrors = nullptr);
 
 	/**
 	 * Disable channels to stop recording traces with them.
diff --git a/Engine/Source/Runtime/Engine/Private/TraceService.cpp b/Engine/Source/Runtime/Engine/Private/TraceService.cpp
--- a/Engine/Source/Runtime/Engine/Private/TraceService.cpp
+++ b/Engine/Source/Runtime/Engine/Private/TraceService.cpp
@@ -88,8 +88,15 @@ void FTraceService::FillTraceStatusMessage(FTraceControlStatus* Message)
 
 void FTraceService::OnChannelSet(const FTraceControlChannelsSet& Message, const TSharedRef<IMessageContext>& Context)
 {
-	FTraceAuxiliary::EnableChannels(Message.ChannelIdsToEnable);
-	FTraceAuxiliary::DisableChannels(Message.ChannelIdsToDisable);
+	TMap<uint32, FString> Errors;
+	FTraceAuxiliary::EnableChannels(Message.ChannelIdsToEnable, &Errors);
+	FTraceAuxiliary::DisableChannels(Message.ChannelIdsToDisable, &Errors);
+	if (!Errors.IsEmpty())
+	{
+		auto Response = MessageEndpoint->MakeMessage<FTraceControlChannelsSetError>();
+		Response->Errors = MoveTemp(Errors);
+		MessageEndpoint->Send(Response, Context->GetSender());
+	}
 }
 
 void FTraceService::OnStop(const FTraceControlStop& Message, const TSharedRef<IMessageContext>& Context)
diff --git a/Engine/Source/Runtime/EngineMessages/Public/TraceControlMessages.h b/Engine/Source/Runtime/EngineMessages/Public/TraceControlMessages.h
--- a/Engine/Source/Runtime/EngineMessages/Public/TraceControlMessages.h
+++ b/Engine/Source/Runtime/EngineMessages/Public/TraceControlMessages.h
@@ -184,6 +184,18 @@ struct FTraceControlChannelsSet
 	TArray<uint32> ChannelIdsToDisable;
 };
 
+/**
+ * Error messages related to channel enable/disable commands
+ */
+USTRUCT()
+struct FTraceControlChannelsSetError
+{
+	GENERATED_USTRUCT_BODY()
+
+	UPROPERTY(EditAnywhere, Category="Message")
+	TMap<uint32, FString> Errors;
+};
+
 /**
  * Message sent when first connecting to a session. Expected to return
  * FTraceControlPong message.
diff --git a/Engine/Source/Runtime/TraceLog/Private/Trace/Channel.cpp b/Engine/Source/Runtime/TraceLog/Private/Trace/Channel.cpp
--- a/Engine/Source/Runtime/TraceLog/Private/Trace/Channel.cpp
+++ b/Engine/Source/Runtime/TraceLog/Private/Trace/Channel.cpp
@@ -1,11 +1,13 @@
 // Copyright Epic Games, Inc. All Rights Reserved.
 
 #include "Trace/Detail/Channel.h"
+
 #include "Math/UnrealMathUtility.h"
-#include "Trace/Trace.h"
-#include "Trace/Trace.inl"
 #include "Trace/Detail/Atomic.h"
 #include "Trace/Detail/Channel.h"
+#include "Trace/Lock.h"
+#include "Trace/Trace.h"
+#include "Trace/Trace.inl"
 
 #include <ctype.h>
 
@@ -182,12 +184,18 @@ void FChannel::Initialize()
 	// Now we can disable all channels.
 	// Channels specified on the command line (using -trace=<channels> argument)
 	// will be further re-enabled after this call.
-	ToggleAll(false);
+	ToggleAllInternal(false, false);
 	GChannelsInitialized = true;
 }
 
 ///////////////////////////////////////////////////////////////////////////////
 void FChannel::ToggleAll(bool bEnabled)
+{
+	ToggleAllInternal(bEnabled, true);
+}
+
+///////////////////////////////////////////////////////////////////////////////
+void FChannel::ToggleAllInternal(bool bEnabled, bool bIssueCallback)
 {
 	using namespace Private;
 
@@ -200,7 +208,7 @@ void FChannel::ToggleAll(bool bEnabled)
 	{
 		for (; Channel != nullptr; Channel = (FChannel*)(Channel->Next))
 		{
-			Channel->Toggle(bEnabled);
+			Channel->ToggleInternal(bEnabled, bIssueCallback);
 		}
 	}
 }
@@ -303,16 +311,40 @@ void FChannel::EnumerateChannels(ChannelIterCallback Func, void* User)
 		}
 	}
 }
+	
+///////////////////////////////////////////////////////////////////////////////
+bool FChannel::Toggle(bool bEnabled, const TCHAR** OutReason)
+{
+	return ToggleInternal(bEnabled, true, OutReason);
+}
 
 ///////////////////////////////////////////////////////////////////////////////
-bool FChannel::Toggle(bool bEnabled)
+bool FChannel::ToggleInternal(bool bEnabled, bool bIssueCallback, const TCHAR** OutReason)
 {
 	using namespace Private;
-	AtomicStoreRelaxed(&Enabled, bEnabled ? 1 : -1);
 
-	UE_TRACE_MINIMAL_LOG(Trace, ChannelToggle, TraceLogChannel)
-		<< ChannelToggle.Id(Name.Hash)
-		<< ChannelToggle.IsEnabled(IsEnabled());
+	Private::FLock ScopedLock(&Lock);
+	if (ScopedLock.TimedOut())
+	{
+		static const TCHAR* TimedOutError = TEXT("Timed out waiting for lock");
+		if (OutReason)
+		{
+			*OutReason = TimedOutError;
+		}
+		return IsEnabled();
+	}
+
+	bool bCanToggle = bIssueCallback ? OnToggle(bEnabled, OutReason) : true;
+	if (bCanToggle)
+	{
+		AtomicStoreRelaxed(&Enabled, bEnabled ? 1 : -1);
+
+		UE_TRACE_MINIMAL_LOG(Trace, ChannelToggle, TraceLogChannel)
+			<< ChannelToggle.Id(Name.Hash)
+			<< ChannelToggle.IsEnabled(IsEnabled());
+
+		OnToggled(bEnabled);
+	}
 
 	return IsEnabled();
 }
@@ -322,7 +354,7 @@ bool FChannel::Toggle(const ANSICHAR* ChannelName, bool bEnabled)
 {
 	if (FChannel* Channel = FChannel::FindChannel(ChannelName))
 	{
-		return Channel->Toggle(bEnabled);
+		return Channel->ToggleInternal(bEnabled, true);
 	}
 	return false;
 }
diff --git a/Engine/Source/Runtime/TraceLog/Private/Trace/Lock.h b/Engine/Source/Runtime/TraceLog/Private/Trace/Lock.h
new file mode 100644
--- /dev/null
+++ b/Engine/Source/Runtime/TraceLog/Private/Trace/Lock.h
@@ -0,0 +1,58 @@
+// Copyright Epic Games, Inc. All Rights Reserved.
+
+#pragma once
+
+#include "Trace/Config.h"
+#include "Trace/Detail/Atomic.h"
+#include "Platform.h"
+
+#if TRACE_PRIVATE_MINIMAL_ENABLED
+
+namespace UE {
+namespace Trace {
+namespace Private {
+
+struct FLock 
+{
+	FLock(volatile uint8* InLock)
+		: Lock(InLock)
+	{
+		CyclesPerSecond = TimeGetFrequency();
+		StartSeconds = GetTime();
+
+		while (!AtomicCompareExchangeAcquire(Lock, uint8(1u), uint8(0u)))
+		{
+			ThreadSleep(0);
+
+			if (TimedOut())
+			{
+				break;
+			}
+		}
+	}
+
+	~FLock()
+	{
+		AtomicExchangeRelease(Lock, uint8(0u));
+	}
+
+	double GetTime()
+	{
+		return static_cast<double>(TimeGetTimestamp()) / static_cast<double>(CyclesPerSecond);
+	}
+
+	bool TimedOut()
+	{
+		const double WaitTime = GetTime() - StartSeconds;
+		return WaitTime > MaxWaitSeconds;
+	}
+
+	volatile uint8* Lock;
+	uint64 CyclesPerSecond;
+	double StartSeconds;
+	inline const static double MaxWaitSeconds = 1.0;
+};
+
+}}}
+
+#endif
diff --git a/Engine/Source/Runtime/TraceLog/Public/Trace/Detail/Channel.h b/Engine/Source/Runtime/TraceLog/Public/Trace/Detail/Channel.h
--- a/Engine/Source/Runtime/TraceLog/Public/Trace/Detail/Channel.h
+++ b/Engine/Source/Runtime/TraceLog/Public/Trace/Detail/Channel.h
@@ -34,44 +34,84 @@ class FChannel
 public:
 	struct Iter
 	{
-						~Iter();
+		~Iter();
 		const FChannel*	GetNext();
-		void*			Inner[3];
+		void*Inner[3];
 	};
 
 	struct InitArgs
 	{
-		const ANSICHAR*		Desc;		// User facing description string
-		bool				bReadOnly;	// If set, channel cannot be changed during a run, only set through command line.
+		/**
+		 * User facing description string.
+		 */
+		const ANSICHAR* Desc;
+		/**
+		 * If set, channel cannot be changed during a run, only set through command line.
+		 */
+		bool bReadOnly;
 	};
 
-	TRACELOG_API void	Setup(const ANSICHAR* InChannelName, const InitArgs& Args);
-	TRACELOG_API static void Initialize();
-	static Iter			ReadNew();
-	void				Announce() const;
-	static bool			Toggle(const ANSICHAR* ChannelName, bool bEnabled);
-	static void			ToggleAll(bool bEnabled);
-	static void			PanicDisableAll(); // Disabled channels wont be logged with UE_TRACE_LOG
+	/**
+	 * Initializes a channel. Used by channel macros, do no use directly.
+	 */
+	TRACELOG_API void Setup(const ANSICHAR* InChannelName, const InitArgs& Args);
+
+	/**
+	 * Allows channels to act pre state change deny the change from happening.
+	 *
+	 * @param bNewState Desired new state for channel
+	 * @param OutReason Deny reason. Set only when returning false.
+	 * @return True when new state is acceptable, false otherwise.
+	 */
+	TRACELOG_API virtual bool OnToggle(bool bNewState, const TCHAR** OutReason) { return true; }
+
+	/**
+	 * Allows channels to act post state change. At this point it is possible to emit
+	 * events on this channel.
+	 *
+	 * @param bNewState New state for the channel
+	 */
+	TRACELOG_API virtual void OnToggled(bool bNewState) {}
+
+	/**
+	 * Toggles the channel state (on/off). If the channel could not be toggled an optional
+	 * deny reason can be accessed.
+	 *
+	 * @param bEnabled Requested state
+	 * @param OutReason Optional string pointer. Will contain deny reason on failure.
+	 * @return New channel state
+	 */
+	TRACELOG_API bool	Toggle(bool bEnabled, const TCHAR** OutReason);
+
+	static void Initialize();
+	static Iter	ReadNew();
+	void Announce() const;
+	static bool Toggle(const ANSICHAR* ChannelName, bool bEnabled);
+	static void ToggleAll(bool bEnabled);
+	static void PanicDisableAll(); // Disabled channels wont be logged with UE_TRACE_LOG
 	static FChannel* FindChannel(const ANSICHAR* ChannelName);
 	static FChannel* FindChannel(FChannelId ChannelId);
-	static void			EnumerateChannels(ChannelIterCallback Func, void* User);
-	TRACELOG_API bool	Toggle(bool bEnabled);
-	bool	IsEnabled() const;
-	bool	IsReadOnly() const { return Args.bReadOnly; };
-	uint32	GetName(const ANSICHAR** OutName) const;
-	explicit			operator bool () const;
-	bool				operator | (const FChannel& Rhs) const;
+	static void EnumerateChannels(ChannelIterCallback Func, void* User);
+	bool IsEnabled() const;
+	bool IsReadOnly() const { return Args.bReadOnly; };
+	uint32 GetName(const ANSICHAR** OutName) const;
+	explicit operator bool () const;
+	bool operator | (const FChannel& Rhs) const;
 
 private:
-	FChannel*			Next;
+	bool ToggleInternal(bool bEnabled, bool bIssueCallback, const TCHAR** OutReason = nullptr);
+	static void ToggleAllInternal(bool bEnabled, bool bIssueCallback);
+	
+	FChannel* Next;
 	struct
 	{
-		const ANSICHAR*	Ptr;
-		uint32			Len;
-		uint32			Hash;
-	}					Name;
-	volatile int32		Enabled;
-	InitArgs			Args;
+		const ANSICHAR* Ptr;
+		uint32 Len;
+		uint32 Hash;
+	} Name;
+	volatile int32 Enabled;
+	InitArgs Args;
+	volatile uint8 Lock;
 };
 
 inline uint32 FChannel::GetName(const ANSICHAR** OutName) const
diff --git a/Engine/Source/Runtime/TraceLog/Public/Trace/Detail/Trace.h b/Engine/Source/Runtime/TraceLog/Public/Trace/Detail/Trace.h
--- a/Engine/Source/Runtime/TraceLog/Public/Trace/Detail/Trace.h
+++ b/Engine/Source/Runtime/TraceLog/Public/Trace/Detail/Trace.h
@@ -30,9 +30,10 @@ class FChannel;
 
 #define TRACE_PRIVATE_CHANNEL_DEFAULT_ARGS false, "None"
 
-#define TRACE_PRIVATE_CHANNEL_DECLARE(LinkageType, ChannelName) \
-	static UE::Trace::FChannel ChannelName##Object; \
-	LinkageType UE::Trace::FChannel& ChannelName = ChannelName##Object;
+#define TRACE_PRIVATE_CHANNEL_DECLARE(LinkageType, ChannelName, ChannelClass) \
+	static_assert(std::is_base_of_v<UE::Trace::FChannel, ChannelClass>);\
+	static ChannelClass ChannelName##Object; \
+	LinkageType ChannelClass& ChannelName = ChannelName##Object;
 
 #define TRACE_PRIVATE_CHANNEL_IMPL(ChannelName, ...) \
 	struct F##ChannelName##Registrator \
@@ -45,16 +46,27 @@ class FChannel;
 	static F##ChannelName##Registrator ChannelName##Reg = F##ChannelName##Registrator();
 
 #define TRACE_PRIVATE_CHANNEL(ChannelName, ...) \
-	TRACE_PRIVATE_CHANNEL_DECLARE(static, ChannelName) \
+	TRACE_PRIVATE_CHANNEL_DECLARE(static, ChannelName, UE::Trace::FChannel) \
+	TRACE_PRIVATE_CHANNEL_IMPL(ChannelName, ##__VA_ARGS__)
+
+#define TRACE_PRIVATE_CHANNEL_CUSTOM(ChannelName, ChannelClass, ...) \
+	TRACE_PRIVATE_CHANNEL_DECLARE(static, ChannelName, ChannelClass) \
 	TRACE_PRIVATE_CHANNEL_IMPL(ChannelName, ##__VA_ARGS__)
 
 #define TRACE_PRIVATE_CHANNEL_DEFINE(ChannelName, ...) \
-	TRACE_PRIVATE_CHANNEL_DECLARE(, ChannelName) \
+	TRACE_PRIVATE_CHANNEL_DECLARE(, ChannelName, UE::Trace::FChannel) \
+	TRACE_PRIVATE_CHANNEL_IMPL(ChannelName, ##__VA_ARGS__)
+
+#define TRACE_PRIVATE_CHANNEL_CUSTOM_DEFINE(ChannelName, ChannelClass, ...) \
+	TRACE_PRIVATE_CHANNEL_DECLARE(, ChannelName, ChannelClass) \
 	TRACE_PRIVATE_CHANNEL_IMPL(ChannelName, ##__VA_ARGS__)
 
 #define TRACE_PRIVATE_CHANNEL_EXTERN(ChannelName, ...) \
 	__VA_ARGS__ extern UE::Trace::FChannel& ChannelName;
 
+#define TRACE_PRIVATE_CHANNEL_CUSTOM_EXTERN(ChannelName, ChannelClass, ...) \
+	__VA_ARGS__ extern class ChannelClass& ChannelName;
+
 #define TRACE_PRIVATE_CHANNELEXPR_IS_ENABLED(ChannelsExpr) \
 	bool(ChannelsExpr)
 
@@ -214,10 +226,16 @@ class FChannel;
 
 #define TRACE_PRIVATE_NOP_CHANNEL(ChannelName, ...)
 
+#define TRACE_PRIVATE_NOP_CHANNEL_CUSTOM(ChannelName, ChannelClass, ...)
+
 #define TRACE_PRIVATE_NOP_CHANNEL_EXTERN(ChannelName, ...)
 
+#define TRACE_PRIVATE_NOP_CHANNEL_CUSTOM_EXTERN(ChannelName, ChannelClass, ...)
+
 #define TRACE_PRIVATE_NOP_CHANNEL_DEFINE(ChannelName, ...)
 
+#define TRACE_PRIVATE_NOP_CHANNEL_CUSTOM_DEFINE(ChannelName, ChannelClass, ...)
+
 #define TRACE_PRIVATE_NOP_CHANNELEXPR_IS_ENABLED(ChannelsExpr) \
 	false
 
diff --git a/Engine/Source/Runtime/TraceLog/Public/Trace/Trace.h b/Engine/Source/Runtime/TraceLog/Public/Trace/Trace.h
--- a/Engine/Source/Runtime/TraceLog/Public/Trace/Trace.h
+++ b/Engine/Source/Runtime/TraceLog/Public/Trace/Trace.h
@@ -444,8 +444,11 @@ UE_TRACE_API FChannel* FindChannel(FChannelId ChannelId) UE_TRACE_IMPL(nullptr);
 #define UE_TRACE_GET_DEFINITION_TYPE_ID(LoggerName, EventName)									TRACE_IMPL(GET_DEFINITION_TYPE_ID, LoggerName, EventName)
 #define UE_TRACE_LOG_DEFINITION(LoggerName, EventName, Id, ChannelsExpr, ...)					TRACE_IMPL(LOG_DEFINITION, LoggerName, EventName, Id, ChannelsExpr, ##__VA_ARGS__)
 #define UE_TRACE_CHANNEL(ChannelName, ...)														TRACE_IMPL(CHANNEL, ChannelName, ##__VA_ARGS__)
+#define UE_TRACE_CHANNEL_CUSTOM(ChannelName, ChannelClass, ...)									TRACE_IMPL(CHANNEL_CUSTOM, ChannelName, ChannelClass, ##__VA_ARGS__)
 #define UE_TRACE_CHANNEL_EXTERN(ChannelName, ...)												TRACE_IMPL(CHANNEL_EXTERN, ChannelName, ##__VA_ARGS__)
+#define UE_TRACE_CHANNEL_CUSTOM_EXTERN(ChannelName, ChannelClass, ...)							TRACE_IMPL(CHANNEL_CUSTOM_EXTERN, ChannelName, ChannelClass, ##__VA_ARGS__)
 #define UE_TRACE_CHANNEL_DEFINE(ChannelName, ...)												TRACE_IMPL(CHANNEL_DEFINE, ChannelName, ##__VA_ARGS__)
+#define UE_TRACE_CHANNEL_CUSTOM_DEFINE(ChannelName, ChannelClass, ...)							TRACE_IMPL(CHANNEL_CUSTOM_DEFINE, ChannelName, ChannelClass, ##__VA_ARGS__)
 #define UE_TRACE_CHANNELEXPR_IS_ENABLED(ChannelsExpr)											TRACE_IMPL(CHANNELEXPR_IS_ENABLED, ChannelsExpr)
 
 
