Outline allocation function calls to reduce binary size and enabled fixed gmalloc class on supported platforms\configs

diff --git a/Engine/Source/Runtime/Core/Private/HAL/MallocBinned2.cpp b/Engine/Source/Runtime/Core/Private/HAL/MallocBinned2.cpp
--- a/Engine/Source/Runtime/Core/Private/HAL/MallocBinned2.cpp
+++ b/Engine/Source/Runtime/Core/Private/HAL/MallocBinned2.cpp
@@ -233,9 +233,10 @@ struct FMallocBinned2::Private
 {
 	static MallocBinnedPrivate::TGlobalRecycler<UE_MB2_SMALL_POOL_COUNT> GGlobalRecycler;
 
-	static void FreeBundles(FMallocBinned2& Allocator, FBundleNode* BundlesToRecycle, uint32 InBinSize, uint32 InPoolIndex)
+	static void FreeBundles(FMallocBinned2& Allocator, FBundleNode* BundlesToRecycle, uint32 InPoolIndex)
 	{
 		FPoolTable& Table = Allocator.SmallPoolTables[InPoolIndex];
+		const uint32 BinSize = Table.BinSize;
 
 		UE::TUniqueLock Lock(Table.Mutex);
 
@@ -270,7 +271,7 @@ struct FMallocBinned2::Private
 				FFreeBlock* Free = (FFreeBlock*)Node;
 				Free->NumFreeBins = 1;
 				Free->NextFreeBlock = NodePool->FirstFreeBlock;
-				Free->BinSize = InBinSize;
+				Free->BinSize = BinSize;
 				Free->CanaryAndForkState = Allocator.CurrentCanary;
 				Free->PoolIndex = InPoolIndex;
 				NodePool->FirstFreeBlock = Free;
@@ -292,8 +293,8 @@ struct FMallocBinned2::Private
 					}
 
 #if UE_MB2_ALLOCATOR_STATS
-					int NumBinsInBlock = GMB2PageSize / InBinSize;
-					if (NumBinsInBlock * InBinSize + sizeof(FFreeBlock) > GMB2PageSize)
+					int NumBinsInBlock = GMB2PageSize / BinSize;
+					if (NumBinsInBlock * BinSize + sizeof(FFreeBlock) > GMB2PageSize)
 					{
 						NumBinsInBlock--;
 					}
@@ -486,7 +487,7 @@ void FMallocBinned2::OnPreFork()
 		while (FBundleNode* Node = FMallocBinned2::Private::GGlobalRecycler.PopBundle(PoolIndex))
 		{
 			Node->NextBundle = nullptr; // We need to override this because it's in a union with Count which is not needed for the freeing work
-			FMallocBinned2::Private::FreeBundles(*this, Node, SmallPoolTables[PoolIndex].BinSize, PoolIndex);
+			FMallocBinned2::Private::FreeBundles(*this, Node, PoolIndex);
 		}
 	}
 
@@ -528,8 +529,16 @@ bool FMallocBinned2::IsInternallyThreadSafe() const
 	return true;
 }
 
-void* FMallocBinned2::MallocSelect(SIZE_T Size, uint32 Alignment, bool bUseSmallPool)
+FORCEINLINE static bool UseSmallAlloc(SIZE_T Size, uint32 Alignment)
 {
+	return ((Size <= UE_MB2_MAX_SMALL_POOL_SIZE) & (Alignment <= UE_MBC_MIN_SMALL_POOL_ALIGNMENT)); // one branch, not two
+}
+
+void* FMallocBinned2::Malloc(SIZE_T Size, uint32 Alignment)
+{
+	// Only allocate from the small pools if the size is small enough and the alignment isn't crazy large.
+	// With large alignments, we'll waste a lot of memory allocating an entire page, but such alignments are highly unlikely in practice.
+	bool bUseSmallPool = UseSmallAlloc(Size, Alignment);
 	if (!bUseSmallPool)
 	{
 		bUseSmallPool = PromoteToLargerBin(Size, Alignment, *this);
@@ -661,11 +670,11 @@ void* FMallocBinned2::MallocExternalLarge(SIZE_T Size, uint32 Alignment)
 	return Result;
 }
 
-void* FMallocBinned2::ReallocExternal(void* Ptr, SIZE_T NewSize, uint32 Alignment)
+void* FMallocBinned2::Realloc(void* Ptr, SIZE_T NewSize, uint32 Alignment)
 {
 	if (NewSize == 0)
 	{
-		FMallocBinned2::FreeExternal(Ptr);
+		FMallocBinned2::Free(Ptr);
 		return nullptr;
 	}
 
@@ -692,14 +701,14 @@ void* FMallocBinned2::ReallocExternal(void* Ptr, SIZE_T NewSize, uint32 Alignmen
 		}
 
 		// Reallocate and copy the data across
-		void* Result = FMallocBinned2::MallocSelect(NewSize, Alignment);
+		void* Result = FMallocBinned2::Malloc(NewSize, Alignment);
 		FMemory::Memcpy(Result, Ptr, FMath::Min<SIZE_T>(NewSize, BinSize));
-		FMallocBinned2::FreeExternal(Ptr);
+		FMallocBinned2::Free(Ptr);
 		return Result;
 	}
 	if (!Ptr)
 	{
-		void* Result = FMallocBinned2::MallocSelect(NewSize, Alignment);
+		void* Result = FMallocBinned2::Malloc(NewSize, Alignment);
 		return Result;
 	}
 
@@ -712,7 +721,7 @@ void* FMallocBinned2::ReallocExternal(void* Ptr, SIZE_T NewSize, uint32 Alignmen
 
 	const UPTRINT PoolOsBytes = Pool->GetOsAllocatedBytes();
 	const SIZE_T PoolOSRequestedBytes = Pool->GetOSRequestedBytes();
-	checkf(PoolOSRequestedBytes <= PoolOsBytes, TEXT("FMallocBinned2::ReallocExternal %d %d"), int32(PoolOSRequestedBytes), int32(PoolOsBytes));
+	checkf(PoolOSRequestedBytes <= PoolOsBytes, TEXT("FMallocBinned2::Realloc %d %d"), int32(PoolOSRequestedBytes), int32(PoolOsBytes));
 	bool bUseSmallMalloc = UseSmallAlloc(NewSize, Alignment);
 	if (!bUseSmallMalloc)
 	{
@@ -739,7 +748,7 @@ void* FMallocBinned2::ReallocExternal(void* Ptr, SIZE_T NewSize, uint32 Alignmen
 		}
 
 		FMemory::Memcpy(Result, Ptr, FMath::Min(NewSize, PoolOSRequestedBytes));
-		FMallocBinned2::FreeExternal(Ptr);
+		FMallocBinned2::Free(Ptr);
 		return Result;
 	}
 
@@ -759,7 +768,7 @@ void* FMallocBinned2::ReallocExternal(void* Ptr, SIZE_T NewSize, uint32 Alignmen
 	return Ptr;
 }
 
-void FMallocBinned2::FreeExternal(void* Ptr)
+void FMallocBinned2::Free(void* Ptr)
 {
 	NOALLOC_SCOPE_CYCLE_COUNTER(STAT_FMallocBinned2_FreeExternal);
 
@@ -797,7 +806,7 @@ void FMallocBinned2::FreeExternal(void* Ptr)
 		if (BundlesToRecycle)
 		{
 			BundlesToRecycle->NextBundle = nullptr;
-			Private::FreeBundles(*this, BundlesToRecycle, BinSize, PoolIndex);
+			Private::FreeBundles(*this, BundlesToRecycle, PoolIndex);
 #if UE_MB2_ALLOCATOR_STATS
 			if (!Lists)
 			{
@@ -823,7 +832,7 @@ void FMallocBinned2::FreeExternal(void* Ptr)
 		AllocatedLargePoolMemoryWAlignment.fetch_sub((int64)PoolOsBytes, std::memory_order_relaxed);
 #endif
 
-		checkf(PoolOSRequestedBytes <= PoolOsBytes, TEXT("FMallocBinned2::FreeExternal %d %d"), int32(PoolOSRequestedBytes), int32(PoolOsBytes));
+		checkf(PoolOSRequestedBytes <= PoolOsBytes, TEXT("FMallocBinned2::Free %d %d"), int32(PoolOSRequestedBytes), int32(PoolOsBytes));
 		Pool->SetCanary(FPoolInfo::ECanary::Unassigned, true, false);
 		// Free an OS allocation.
 		CachedOSPageAllocator.Free(Ptr, PoolOsBytes, &ExternalAllocMutex, FPerThreadFreeBlockLists::Get() != nullptr && GBinned2MoveOSFreesOffTimeCriticalThreads != 0);
@@ -875,7 +884,7 @@ const TCHAR* FMallocBinned2::GetDescriptiveName()
 
 void FMallocBinned2::FreeBundles(FBundleNode* Bundles, uint32 PoolIndex)
 {
-	Private::FreeBundles(*this, Bundles, PoolIndexToBinSize(PoolIndex), PoolIndex);
+	Private::FreeBundles(*this, Bundles, PoolIndex);
 }
 
 void FMallocBinned2::Trim(bool bTrimThreadCaches)
@@ -952,13 +961,6 @@ void FMallocBinned2::DumpAllocatorStats(class FOutputDevice& Ar)
 	Ar.Logf(TEXT("Allocator Stats for binned2 are not in this build set UE_MB2_ALLOCATOR_STATS 1 in MallocBinned2.cpp"));
 #endif
 }
-#if !UE_MB2_INLINE
-	#if PLATFORM_USES_FIXED_GMalloc_CLASS && !FORCE_ANSI_ALLOCATOR && USE_MALLOC_BINNED2
-		//#define FMEMORY_INLINE_FUNCTION_DECORATOR  FORCEINLINE
-		#define FMEMORY_INLINE_GMalloc (FMallocBinned2::MallocBinned2)
-		#include "FMemory.inl"
-	#endif
-#endif
 
 void FMallocBinned2::UpdateStats()
 {
diff --git a/Engine/Source/Runtime/Core/Private/HAL/MallocBinned3.cpp b/Engine/Source/Runtime/Core/Private/HAL/MallocBinned3.cpp
--- a/Engine/Source/Runtime/Core/Private/HAL/MallocBinned3.cpp
+++ b/Engine/Source/Runtime/Core/Private/HAL/MallocBinned3.cpp
@@ -308,7 +308,7 @@ struct FMallocBinned3::Private
 
 	static MallocBinnedPrivate::TGlobalRecycler<UE_MB3_SMALL_POOL_COUNT> GGlobalRecycler;
 
-	static void FreeBundles(FMallocBinned3& Allocator, FBundleNode* BundlesToRecycle, uint32 InBinSize, uint32 InPoolIndex)
+	static void FreeBundles(FMallocBinned3& Allocator, FBundleNode* BundlesToRecycle, uint32 InPoolIndex)
 	{
 		FPoolTable& Table = Allocator.SmallPoolTables[InPoolIndex];
 		UE::TUniqueLock Lock(Table.Mutex);
@@ -324,8 +324,9 @@ struct FMallocBinned3::Private
 				FBundleNode* NextNode = Node->NextNodeInCurrentBundle;
 
 				uint32 OutBlockIndex;
+				const uint32 BinSize = Allocator.SmallPoolTables[InPoolIndex].BinSize;
 				void* BaseBlockPtr = Allocator.BlockPointerFromContainedPtr(Node, Allocator.SmallPoolTables[InPoolIndex].NumMemoryPagesPerBlock, OutBlockIndex);
-				const uint32 BinIndexWithinBlock = (((uint8*)Node) - ((uint8*)BaseBlockPtr)) / Allocator.SmallPoolTables[InPoolIndex].BinSize;
+				const uint32 BinIndexWithinBlock = (((uint8*)Node) - ((uint8*)BaseBlockPtr)) / BinSize;
 
 				FPoolInfoSmall* NodePoolBlock = Allocator.SmallPoolTables[InPoolIndex].PoolInfos[OutBlockIndex / Allocator.SmallPoolInfosPerPlatformPage];
 				if (!NodePoolBlock)
@@ -342,7 +343,7 @@ struct FMallocBinned3::Private
 				FFreeBlock* Free = (FFreeBlock*)Node;
 				Free->NumFreeBins = 1;
 				Free->NextFreeBlockIndex = NodePool->NoFirstFreeIndex ? MAX_uint32 : NodePool->FirstFreeIndex;
-				Free->BinSizeShifted = (InBinSize >> UE_MBC_MIN_SMALL_POOL_ALIGNMENT_SHIFT);
+				Free->BinSizeShifted = (BinSize >> UE_MBC_MIN_SMALL_POOL_ALIGNMENT_SHIFT);
 				Free->Canary = FFreeBlock::CANARY_VALUE;
 				Free->PoolIndex = InPoolIndex;
 				NodePool->FirstFreeIndex = BinIndexWithinBlock;
@@ -367,7 +368,7 @@ struct FMallocBinned3::Private
 					Allocator.Decommit(InPoolIndex, BaseBlockPtr, BlockSize);
 
 #if UE_MB3_ALLOCATOR_STATS
-					Table.TotalAllocatedBins -= BlockSize / InBinSize;
+					Table.TotalAllocatedBins -= BlockSize / BinSize;
 					Table.TotalAllocatedMem -= BlockSize;
 					AllocatedOSSmallPoolMemory.fetch_sub(BlockSize, std::memory_order_relaxed);
 #endif
@@ -389,7 +390,7 @@ MallocBinnedPrivate::TGlobalRecycler<UE_MB3_SMALL_POOL_COUNT> FMallocBinned3::Pr
 
 void FMallocBinned3::FreeBundles(FBundleNode* Bundles, uint32 PoolIndex)
 {
-	Private::FreeBundles(*this, Bundles, PoolIndexToBinSize(PoolIndex), PoolIndex);
+	Private::FreeBundles(*this, Bundles, PoolIndex);
 }
 
 FMallocBinned3::FPoolInfoSmall* FMallocBinned3::PushNewPoolToFront(FMallocBinned3::FPoolTable& Table, uint32 InBinSize, uint32 InPoolIndex, uint32& OutBlockIndex)
@@ -667,7 +668,7 @@ bool FMallocBinned3::IsInternallyThreadSafe() const
 	return true;
 }
 
-void* FMallocBinned3::MallocExternal(SIZE_T Size, uint32 Alignment)
+void* FMallocBinned3::Malloc(SIZE_T Size, uint32 Alignment)
 {
 	static_assert(DEFAULT_ALIGNMENT <= UE_MBC_MIN_SMALL_POOL_ALIGNMENT, "DEFAULT_ALIGNMENT is assumed to be zero"); // used below
 
@@ -675,9 +676,9 @@ void* FMallocBinned3::MallocExternal(SIZE_T Size, uint32 Alignment)
 	//            Larger alignments can waste a lot of memory allocating an entire page, so some smaller alignments are 
 	//			  handled in the fallback path if less than a predefined max small pool alignment. 
 
-	bool UsePools = (Size <= UE_MB3_MAX_SMALL_POOL_SIZE) && (Alignment <= UE_MBC_MIN_SMALL_POOL_ALIGNMENT);
+	bool UsePools = (Size <= UE_MB3_MAX_SMALL_POOL_SIZE);
 	
-	if (!UsePools)
+	if (UNLIKELY(UsePools && (Alignment > UE_MBC_MIN_SMALL_POOL_ALIGNMENT)))
 	{
 		// check if allocations that require alignment larger than UE_MBC_MIN_SMALL_POOL_ALIGNMENT can be promoted to a bin with a natural alignment that matches
 		// i.e. 16 bytes allocation with 128 bytes alignment can be promoted to 128 bytes bin
@@ -697,7 +698,7 @@ void* FMallocBinned3::MallocExternal(SIZE_T Size, uint32 Alignment)
 				{
 #if UE_MB3_ALLOCATOR_STATS
 					SmallPoolTables[PoolIndex].HeadEndAlloc(Size);
-					const uint32 BinSize = PoolIndexToBinSize(PoolIndex);
+					const uint32 BinSize = SmallPoolTables[PoolIndex].BinSize;
 					Lists->AllocatedMemory += BinSize;
 #endif
 					return Result;
@@ -723,11 +724,11 @@ void* FMallocBinned3::MallocExternal(SIZE_T Size, uint32 Alignment)
 			{
 				if ((PoolIndex + 1) < UE_MB3_SMALL_POOL_COUNT)
 				{
-					return MallocExternal(SmallPoolTables[PoolIndex + 1].BinSize, Alignment);
+					return FMallocBinned3::Malloc(SmallPoolTables[PoolIndex + 1].BinSize, Alignment);
 				}
 				else
 				{
-					return MallocExternal(UE_MB3_MAX_SMALL_POOL_SIZE + 1, Alignment);
+					return FMallocBinned3::Malloc(UE_MB3_MAX_SMALL_POOL_SIZE + 1, Alignment);
 				}
 			}
 		}
@@ -836,7 +837,7 @@ void* FMallocBinned3::MallocExternal(SIZE_T Size, uint32 Alignment)
 	return Result;
 }
 
-void* FMallocBinned3::ReallocExternal(void* Ptr, SIZE_T NewSize, uint32 Alignment)
+void* FMallocBinned3::Realloc(void* Ptr, SIZE_T NewSize, uint32 Alignment)
 {
 	const uint64 PoolIndex = PoolIndexFromPtr(Ptr);
 	if (NewSize == 0)
@@ -864,14 +865,14 @@ void* FMallocBinned3::ReallocExternal(void* Ptr, SIZE_T NewSize, uint32 Alignmen
 		}
 
 		// Reallocate and copy the data across
-		void* Result = FMallocBinned3::MallocExternal(NewSize, Alignment);
+		void* Result = FMallocBinned3::Malloc(NewSize, Alignment);
 		FMemory::Memcpy(Result, Ptr, FMath::Min<SIZE_T>(NewSize, BinSize));
 		FMallocBinned3::FreeExternal(Ptr, PoolIndex);
 		return Result;
 	}
 	if (!Ptr)
 	{
-		void* Result = FMallocBinned3::MallocExternal(NewSize, Alignment);
+		void* Result = FMallocBinned3::Malloc(NewSize, Alignment);
 		return Result;
 	}
 
@@ -894,7 +895,7 @@ void* FMallocBinned3::ReallocExternal(void* Ptr, SIZE_T NewSize, uint32 Alignmen
 	{
 		ExternalAllocMutex.Unlock();
 		// Grow or shrink.
-		void* Result = FMallocBinned3::MallocExternal(NewSize, Alignment);
+		void* Result = FMallocBinned3::Malloc(NewSize, Alignment);
 		SIZE_T CopySize = FMath::Min<SIZE_T>(NewSize, PoolOSRequestedBytes);
 		FMemory::Memcpy(Result, Ptr, CopySize);
 		FMallocBinned3::FreeExternal(Ptr, PoolIndex);
@@ -938,7 +939,7 @@ void FMallocBinned3::FreeExternal(void* Ptr, uint64 PoolIndex)
 		if (BundlesToRecycle)
 		{
 			BundlesToRecycle->NextBundle = nullptr;
-			Private::FreeBundles(*this, BundlesToRecycle, BinSize, PoolIndex);
+			Private::FreeBundles(*this, BundlesToRecycle, PoolIndex);
 #if UE_MB3_ALLOCATOR_STATS
 			if (!Lists)
 			{
@@ -1120,14 +1121,6 @@ void FMallocBinned3::DumpAllocatorStats(class FOutputDevice& Ar)
 #endif
 }
 
-#if !UE_MB3_INLINE
-	#if PLATFORM_USES_FIXED_GMalloc_CLASS && !FORCE_ANSI_ALLOCATOR && USE_MALLOC_BINNED3
-		//#define FMEMORY_INLINE_FUNCTION_DECORATOR  FORCEINLINE
-		#define FMEMORY_INLINE_GMalloc (FMallocBinned3::MallocBinned3)
-		#include "FMemory.inl"
-	#endif
-#endif
-
 #endif //~PLATFORM_HAS_FPlatformVirtualMemoryBlock
 
 PRAGMA_RESTORE_UNSAFE_TYPECAST_WARNINGS
diff --git a/Engine/Source/Runtime/Core/Private/HAL/MallocBinnedCommon.cpp b/Engine/Source/Runtime/Core/Private/HAL/MallocBinnedCommon.cpp
--- a/Engine/Source/Runtime/Core/Private/HAL/MallocBinnedCommon.cpp
+++ b/Engine/Source/Runtime/Core/Private/HAL/MallocBinnedCommon.cpp
@@ -10,6 +10,7 @@
 #include "Templates/AlignmentTemplates.h"
 #include "Templates/UnrealTemplate.h"
 #include "ProfilingDebugging/CsvProfiler.h"
+#include "Logging/LogMacros.h"
 
 #if UE_MBC_LOG_LARGE_ALLOCATION || UE_MBC_LIGHTWEIGHT_BIN_CALLSTACK_TRACKER
 	#include "HAL/PlatformStackWalk.h"
@@ -623,6 +624,12 @@ void FMallocBinnedCommonBase::GetAllocatorStatsInternal(FGenericMemoryStats& Out
 }
 #endif	//~UE_MBC_ALLOCATOR_STATS
 
+void FMallocBinnedCommonBase::UnrecognizedPointerFatalError(void* Ptr)
+{
+	UE_LOG(LogMemory, Fatal, TEXT("FMallocBinnedCommon attempted to GetAllocationSizeExternal on an unrecognized pointer %p"), Ptr);
+	FPlatformProcess::SleepInfinite(); // Unreachable
+}
+
 #if UE_MBC_LOG_LARGE_ALLOCATION
 
 static int32 GLogLargeAllocationSizeInMB = 0;
diff --git a/Engine/Source/Runtime/Core/Private/HAL/UnrealMemory.cpp b/Engine/Source/Runtime/Core/Private/HAL/UnrealMemory.cpp
--- a/Engine/Source/Runtime/Core/Private/HAL/UnrealMemory.cpp
+++ b/Engine/Source/Runtime/Core/Private/HAL/UnrealMemory.cpp
@@ -23,6 +23,7 @@
 #include "HAL/MallocDoubleFreeFinder.h"
 #include "HAL/MallocFrameProfiler.h"
 #include "HAL/MallocStomp2.h"
+#include "HAL/Memory.h"
 
 #if MALLOC_GT_HOOKS
 
@@ -89,7 +90,7 @@ public:
 	/**
 	* Malloc
 	*/
-	virtual void* Malloc(SIZE_T Size, uint32 Alignment) override
+	 void* Malloc(SIZE_T Size, uint32 Alignment) 
 	{
 		return UsedMalloc->Malloc(Size, Alignment);
 	}
@@ -671,7 +672,32 @@ void FUseSystemMallocForNew::operator delete[](void* Ptr)
 	FMemory::SystemFree(Ptr);
 }
 
+void* FMemory::Malloc(SIZE_T Count, uint32 Alignment)
+{
+	return FMemory_MallocInline(Count, Alignment);
+}
 
-#if !INLINE_FMEMORY_OPERATION && !PLATFORM_USES_FIXED_GMalloc_CLASS
-#	include "HAL/FMemory.inl"
-#endif
+void* FMemory::Realloc(void* Original, SIZE_T Count, uint32 Alignment)
+{
+	return FMemory_ReallocInline(Original, Count, Alignment);
+}
+
+void FMemory::Free(void* Original)
+{
+	return FMemory_FreeInline(Original);
+}
+
+SIZE_T FMemory::GetAllocSize(void* Original)
+{
+	return FMemory_GetAllocSizeInline(Original);
+}
+
+void* FMemory::MallocZeroed(SIZE_T Count, uint32 Alignment)
+{
+	return FMemory_MallocZeroedInline(Count, Alignment);
+}
+
+SIZE_T FMemory::QuantizeSize(SIZE_T Count, uint32 Alignment)
+{
+	return FMemory_QuantizeSizeInline(Count, Alignment);
+}
\ No newline at end of file
diff --git a/Engine/Source/Runtime/Core/Public/Experimental/ConcurrentLinearAllocator.h b/Engine/Source/Runtime/Core/Public/Experimental/ConcurrentLinearAllocator.h
--- a/Engine/Source/Runtime/Core/Public/Experimental/ConcurrentLinearAllocator.h
+++ b/Engine/Source/Runtime/Core/Public/Experimental/ConcurrentLinearAllocator.h
@@ -5,6 +5,7 @@
 #include <atomic>
 #include <type_traits>
 
+#include "HAL/Memory.h"
 #include "HAL/UnrealMemory.h"
 #include "HAL/LowLevelMemTracker.h"
 #include "HAL/MallocBinnedCommon.h"
@@ -41,18 +42,18 @@ struct FAlignedAllocator
 
 	inline static void* Malloc(SIZE_T Size, uint32 Alignment)
 	{
-		if (UNLIKELY(!UE::Private::GMalloc))
+		if (UNLIKELY(!FMEMORY_INLINE_GMalloc))
 		{
 			// There is no public function to create GMalloc and this will do it for us.
 			FMemory::Free(FMemory::Malloc(0));
-			check(UE::Private::GMalloc);
+			check(FMEMORY_INLINE_GMalloc);
 		}
-		return UE::Private::GMalloc->Malloc(Size, Alignment);
+		return FMEMORY_INLINE_GMalloc->Malloc(Size, Alignment);
 	}
 
 	UE_FORCEINLINE_HINT static void Free(void* Pointer, SIZE_T Size)
 	{
-		UE::Private::GMalloc->Free(Pointer);
+		FMEMORY_INLINE_GMalloc->Free(Pointer);
 	}
 };
 
diff --git a/Engine/Source/Runtime/Core/Public/HAL/FMemory.inl b/Engine/Source/Runtime/Core/Public/HAL/FMemory.inl
--- a/Engine/Source/Runtime/Core/Public/HAL/FMemory.inl
+++ b/Engine/Source/Runtime/Core/Public/HAL/FMemory.inl
@@ -1,11 +1,7 @@
 // Copyright Epic Games, Inc. All Rights Reserved.
 
-#if !defined(FMEMORY_INLINE_FUNCTION_DECORATOR)
-	#define FMEMORY_INLINE_FUNCTION_DECORATOR 
-#endif
-
 #if !defined(FMEMORY_INLINE_GMalloc)
-	#define FMEMORY_INLINE_GMalloc UE::Private::GMalloc
+#	error "FMEMORY_INLINE_GMalloc should be defined before including this file. Possibly FMemory.inl is included directly instead of including Memory.h"
 #endif
 
 #include "HAL/LowLevelMemTracker.h"
@@ -14,7 +10,11 @@
 struct FMemory;
 struct FScopedMallocTimer;
 
-FMEMORY_INLINE_FUNCTION_DECORATOR void* FMemory::Malloc(SIZE_T Count, uint32 Alignment)
+void FMemory_FreeInline(void* Original);
+SIZE_T FMemory_GetAllocSizeInline(void* Original);
+
+
+FORCEINLINE void* FMemory_MallocInline(SIZE_T Count, uint32 Alignment)
 {
 	// AutoRTFM: For non-transactional code, all of these calls optimize away and the
 	// behavior is the same as it always has been.
@@ -25,9 +25,9 @@ FMEMORY_INLINE_FUNCTION_DECORATOR void* FMemory::Malloc(SIZE_T Count, uint32 Ali
 	void* Ptr = AutoRTFM::Open([Count, Alignment]
 	{
 		void* Alloc = nullptr;
-		if (!FMEMORY_INLINE_GMalloc)
+		if (UNLIKELY(!FMEMORY_INLINE_GMalloc))
 		{
-			Alloc = MallocExternal(Count, Alignment);
+			Alloc = FMemory::MallocExternal(Count, Alignment);
 		}
 		else
 		{
@@ -47,13 +47,13 @@ FMEMORY_INLINE_FUNCTION_DECORATOR void* FMemory::Malloc(SIZE_T Count, uint32 Ali
 	{
 		// Disable the code analysis warning that complains that Free is being passed
 		// a pointer that may be null. Free explicitly handles this case already.
-		Free(Ptr); //-V575
+		FMemory::Free(Ptr); //-V575
 	});
 
 	return AutoRTFM::DidAllocate(Ptr, Count);
 }
 
-FMEMORY_INLINE_FUNCTION_DECORATOR void* FMemory::Realloc(void* Original, SIZE_T Count, uint32 Alignment)
+FORCEINLINE void* FMemory_ReallocInline(void* Original, SIZE_T Count, uint32 Alignment)
 {
 	if (AutoRTFM::IsClosed())
 	{
@@ -80,7 +80,7 @@ FMEMORY_INLINE_FUNCTION_DECORATOR void* FMemory::Realloc(void* Original, SIZE_T
 		// that case and skip the `Malloc` call entirely.
 		if (Count > 0)
 		{
-			Ptr = Malloc(Count, Alignment);
+			Ptr = FMemory::Malloc(Count, Alignment);
 
 			if (!Ptr)
 			{
@@ -92,13 +92,13 @@ FMEMORY_INLINE_FUNCTION_DECORATOR void* FMemory::Realloc(void* Original, SIZE_T
 		{
 			if (Ptr)
 			{
-				SIZE_T OriginalCount = GetAllocSize(Original);
+				SIZE_T OriginalCount = FMemory::GetAllocSize(Original);
 				SIZE_T CopyCount = FGenericPlatformMath::Min(Count, OriginalCount); // handle the case where the new size is smaller
 
-				Memcpy(Ptr, Original, CopyCount);
+				FMemory::Memcpy(Ptr, Original, CopyCount);
 			}
 			
-			Free(Original);
+			FMemory::Free(Original);
 		}
 
 		return Ptr;
@@ -110,9 +110,9 @@ FMEMORY_INLINE_FUNCTION_DECORATOR void* FMemory::Realloc(void* Original, SIZE_T
 	LLM_IF_ENABLED(if (Original != nullptr) FLowLevelMemTracker::Get().OnLowLevelFree(ELLMTracker::Default, Original, ELLMAllocType::FMalloc));
 
 	void* Ptr;
-	if (!FMEMORY_INLINE_GMalloc)
+	if (UNLIKELY(!FMEMORY_INLINE_GMalloc))
 	{
-		Ptr = ReallocExternal(Original, Count, Alignment);
+		Ptr = FMemory::ReallocExternal(Original, Count, Alignment);
 	}
 	else
 	{
@@ -136,7 +136,7 @@ FMEMORY_INLINE_FUNCTION_DECORATOR void* FMemory::Realloc(void* Original, SIZE_T
 	return Ptr;
 }
 
-FMEMORY_INLINE_FUNCTION_DECORATOR void FMemory::Free(void* Original)
+FORCEINLINE void FMemory_FreeInline(void* Original)
 {
 	if (!Original)
 	{
@@ -152,9 +152,9 @@ FMEMORY_INLINE_FUNCTION_DECORATOR void FMemory::Free(void* Original)
 		// optional tracking of every allocation
 		LLM_IF_ENABLED(FLowLevelMemTracker::Get().OnLowLevelFree(ELLMTracker::Default, Original, ELLMAllocType::FMalloc));
 
-		if (!FMEMORY_INLINE_GMalloc)
+		if (UNLIKELY(!FMEMORY_INLINE_GMalloc))
 		{
-			FreeExternal(Original);
+			FMemory::FreeExternal(Original);
 			return;
 		}
 		DoGamethreadHook(2);
@@ -165,14 +165,14 @@ FMEMORY_INLINE_FUNCTION_DECORATOR void FMemory::Free(void* Original)
 	};
 }
 
-FMEMORY_INLINE_FUNCTION_DECORATOR SIZE_T FMemory::GetAllocSize(void* Original)
+FORCEINLINE SIZE_T FMemory_GetAllocSizeInline(void* Original)
 {
 	SIZE_T Result;
 	UE_AUTORTFM_OPEN
 	{
-		if (!FMEMORY_INLINE_GMalloc)
+		if (UNLIKELY(!FMEMORY_INLINE_GMalloc))
 		{
-			Result = GetAllocSizeExternal(Original);
+			Result = FMemory::GetAllocSizeExternal(Original);
 		}
 		else
 		{
@@ -193,7 +193,7 @@ FMEMORY_INLINE_FUNCTION_DECORATOR SIZE_T FMemory::GetAllocSize(void* Original)
 	return Result;
 }
 
-FMEMORY_INLINE_FUNCTION_DECORATOR void* FMemory::MallocZeroed(SIZE_T Count, uint32 Alignment)
+FORCEINLINE void* FMemory_MallocZeroedInline(SIZE_T Count, uint32 Alignment)
 {
 	void* Ptr = nullptr; // Silence bogus static analysis warnings.
 
@@ -204,40 +204,40 @@ FMEMORY_INLINE_FUNCTION_DECORATOR void* FMemory::MallocZeroed(SIZE_T Count, uint
 	// This is because allocators are already transactional - malloc can be rolled back by
 	// calling free.
 	UE_AUTORTFM_OPEN
+	{
+		if (UNLIKELY(!FMEMORY_INLINE_GMalloc))
 		{
-			if (!FMEMORY_INLINE_GMalloc)
-			{
-				Ptr = MallocZeroedExternal(Count, Alignment);
-			}
-			else
-			{
-				DoGamethreadHook(0);
-				FScopedMallocTimer Timer(0);
-				Ptr = FMEMORY_INLINE_GMalloc->MallocZeroed(Count, Alignment);
-			}
-	// optional tracking of every allocation
-	LLM_IF_ENABLED(FLowLevelMemTracker::Get().OnLowLevelAlloc(ELLMTracker::Default, Ptr, Count, ELLMTag::Untagged, ELLMAllocType::FMalloc));
-		};
+			Ptr = FMemory::MallocZeroedExternal(Count, Alignment);
+		}
+		else
+		{
+			DoGamethreadHook(0);
+			FScopedMallocTimer Timer(0);
+			Ptr = FMEMORY_INLINE_GMalloc->MallocZeroed(Count, Alignment);
+		}
+		// optional tracking of every allocation
+		LLM_IF_ENABLED(FLowLevelMemTracker::Get().OnLowLevelAlloc(ELLMTracker::Default, Ptr, Count, ELLMTag::Untagged, ELLMAllocType::FMalloc));
+	};
 
 	// AutoRTFM: This is a no-op for non-transactional code.
 	// For transactional code, this defers a call to Free if the transaction aborts,
 	// so that rolling back this allocation will end up freeing the memory.
 	AutoRTFM::OnAbort([Ptr]
-		{
-			// Disable the code analysis warning that complains that Free is being passed
-			// a pointer that may be null. Free explicitly handles this case already.
-			Free(Ptr); //-V575
-		});
+	{
+		// Disable the code analysis warning that complains that Free is being passed
+		// a pointer that may be null. Free explicitly handles this case already.
+		FMemory::Free(Ptr); //-V575
+	});
 
 	return AutoRTFM::DidAllocate(Ptr, Count);
 }
 
-FMEMORY_INLINE_FUNCTION_DECORATOR SIZE_T FMemory::QuantizeSize(SIZE_T Count, uint32 Alignment)
+FORCEINLINE SIZE_T FMemory_QuantizeSizeInline(SIZE_T Count, uint32 Alignment)
 {
 	SIZE_T Result;
 	UE_AUTORTFM_OPEN
 	{
-		if (!FMEMORY_INLINE_GMalloc)
+		if (UNLIKELY(!FMEMORY_INLINE_GMalloc))
 		{
 			Result = Count;
 		}
@@ -248,5 +248,4 @@ FMEMORY_INLINE_FUNCTION_DECORATOR SIZE_T FMemory::QuantizeSize(SIZE_T Count, uin
 	};
 
 	return Result;
-}
-
+}
\ No newline at end of file
diff --git a/Engine/Source/Runtime/Core/Public/HAL/LowLevelMemTracker.h b/Engine/Source/Runtime/Core/Public/HAL/LowLevelMemTracker.h
--- a/Engine/Source/Runtime/Core/Public/HAL/LowLevelMemTracker.h
+++ b/Engine/Source/Runtime/Core/Public/HAL/LowLevelMemTracker.h
@@ -9,17 +9,15 @@
 #include "AutoRTFM.h"
 
 #ifndef PLATFORM_SUPPORTS_LLM
-#define PLATFORM_SUPPORTS_LLM 1
+	#define PLATFORM_SUPPORTS_LLM 1
 #endif
 
-// LLM is currently incompatible with PLATFORM_USES_FIXED_GMalloc_CLASS, because LLM is activated way too early
-// This is not a problem, because fixed GMalloc is only used in Test/Shipping builds
-#define LLM_ENABLED_ON_PLATFORM (PLATFORM_SUPPORTS_LLM && !PLATFORM_USES_FIXED_GMalloc_CLASS)
+#define LLM_ENABLED_ON_PLATFORM UE_DEPRECATED_MACRO(5.7, "Use PLATFORM_SUPPORTS_LLM instead") (PLATFORM_SUPPORTS_LLM)
 
 #ifdef ENABLE_LOW_LEVEL_MEM_TRACKER
-#error ENABLE_LOW_LEVEL_MEM_TRACKER is now a derived define that should not be defined separately. Define LLM_ENABLED_IN_CONFIG (build environment only) or LLM_ENABLED_ON_PLATFORM (build environment or c++ header) instead.
+	#error ENABLE_LOW_LEVEL_MEM_TRACKER is now a derived define that should not be defined separately. Define LLM_ENABLED_IN_CONFIG (build environment only) or PLATFORM_SUPPORTS_LLM (build environment or c++ header) instead.
 #endif
-#define ENABLE_LOW_LEVEL_MEM_TRACKER (LLM_ENABLED_IN_CONFIG && LLM_ENABLED_ON_PLATFORM)
+#define ENABLE_LOW_LEVEL_MEM_TRACKER (LLM_ENABLED_IN_CONFIG && PLATFORM_SUPPORTS_LLM)
 
 #if ENABLE_LOW_LEVEL_MEM_TRACKER
 
diff --git a/Engine/Source/Runtime/Core/Public/HAL/MallocBinned2.h b/Engine/Source/Runtime/Core/Public/HAL/MallocBinned2.h
--- a/Engine/Source/Runtime/Core/Public/HAL/MallocBinned2.h
+++ b/Engine/Source/Runtime/Core/Public/HAL/MallocBinned2.h
@@ -237,226 +237,21 @@ private:
 	}
 
 public:
-	CORE_API FMallocBinned2();
-	CORE_API virtual ~FMallocBinned2();
+	FMallocBinned2();
+	virtual ~FMallocBinned2();
 
 	// FMalloc interface.
-	CORE_API virtual bool IsInternallyThreadSafe() const override;
+	virtual bool IsInternallyThreadSafe() const override;
 
 	UE_AUTORTFM_NOAUTORTFM
-	FORCEINLINE virtual void* Malloc(SIZE_T Size, uint32 Alignment) override
-	{
-#if UE_MB2_ALLOCATOR_STATS_VALIDATION
-		UE::TUniqueLock Lock(ValidationCriticalSection);
-		++RecursionCounter;
-		void *Result = MallocInline(Size, Alignment);
-		if (!IsOSAllocation(Result))
-		{
-			SIZE_T OutSize;
-			ensure(GetAllocationSize(Result, OutSize));
-			AllocatedSmallPoolMemoryValidation += OutSize;
-			if (RecursionCounter == 1)
-			{
-				check(GetTotalAllocatedSmallPoolMemory() == AllocatedSmallPoolMemoryValidation);
-				if (GetTotalAllocatedSmallPoolMemory() != AllocatedSmallPoolMemoryValidation)
-				{
-					UE_DEBUG_BREAK();
-				}
-			}
-		}
-		--RecursionCounter;
-		return Result;
-#else
-		return MallocInline(Size, Alignment);
-#endif
-	}
-
-protected:
-	FORCEINLINE void* MallocInline(SIZE_T Size, uint32 Alignment)
-	{
-		// Only allocate from the small pools if the size is small enough and the alignment isn't crazy large.
-		// With large alignments, we'll waste a lot of memory allocating an entire page, but such alignments are highly unlikely in practice.
-		const bool bUseSmallPool = UseSmallAlloc(Size, Alignment);
-		if (bUseSmallPool)
-		{
-			FPerThreadFreeBlockLists* Lists = GMallocBinnedPerThreadCaches ? FPerThreadFreeBlockLists::Get() : nullptr;
-			if (Lists)
-			{
-				const uint32 PoolIndex = BoundSizeToPoolIndex(Size, MemSizeToPoolIndex);
-				if (void* Result = Lists->Malloc(PoolIndex))
-				{
-#if UE_MB2_ALLOCATOR_STATS
-					const uint32 BinSize = PoolIndexToBinSize(PoolIndex);
-					Lists->AllocatedMemory += BinSize;
-#endif
-					return Result;
-				}
-			}
-		}
-
-		return MallocSelect(Size, Alignment, bUseSmallPool);
-	}
+	virtual void* Malloc(SIZE_T Size, uint32 Alignment) override;
 
-	FORCEINLINE static bool UseSmallAlloc(SIZE_T Size, uint32 Alignment)
-	{
-		return ((Size <= UE_MB2_MAX_SMALL_POOL_SIZE) & (Alignment <= UE_MBC_MIN_SMALL_POOL_ALIGNMENT)); // one branch, not two
-	}
-
-	CORE_API void* MallocSelect(SIZE_T Size, uint32 Alignment, bool bUseSmallPool);
-	FORCEINLINE void* MallocSelect(SIZE_T Size, uint32 Alignment)
-	{
-		return MallocSelect(Size, Alignment, UseSmallAlloc(Size, Alignment));
-	}
-
-public:
 	UE_AUTORTFM_NOAUTORTFM
-	FORCEINLINE virtual void* Realloc(void* Ptr, SIZE_T NewSize, uint32 Alignment) override
-	{
-#if UE_MB2_ALLOCATOR_STATS_VALIDATION
-		const bool bOldIsOsAllocation = IsOSAllocation(Ptr);
-		SIZE_T OldSize;
-		if (!bOldIsOsAllocation)
-		{
-			ensure(GetAllocationSize(Ptr, OldSize));
-		}
+	virtual void* Realloc(void* Ptr, SIZE_T NewSize, uint32 Alignment) override;
 
-		UE::TUniqueLock Lock(ValidationCriticalSection);
-		++RecursionCounter;
-		void *Result = ReallocInline(Ptr, NewSize, Alignment);
-		if (!bOldIsOsAllocation)
-		{
-			AllocatedSmallPoolMemoryValidation -= OldSize;
-		}
-		if (!IsOSAllocation(Result))
-		{
-			SIZE_T OutSize;
-			ensure(GetAllocationSize(Result, OutSize));
-			AllocatedSmallPoolMemoryValidation += OutSize;
-		}
-		if (RecursionCounter == 1)
-		{
-			check(GetTotalAllocatedSmallPoolMemory() == AllocatedSmallPoolMemoryValidation);
-			if (GetTotalAllocatedSmallPoolMemory() != AllocatedSmallPoolMemoryValidation)
-			{
-				UE_DEBUG_BREAK();
-			}
-		}
-		--RecursionCounter;
-		return Result;
-#else
-		return ReallocInline(Ptr, NewSize, Alignment);
-#endif
-	}
-
-protected:
-	FORCEINLINE void* ReallocInline(void* Ptr, SIZE_T NewSize, uint32 Alignment) 
-	{
-		if (NewSize <= UE_MB2_MAX_SMALL_POOL_SIZE && Alignment <= UE_MBC_MIN_SMALL_POOL_ALIGNMENT) // one branch, not two
-		{
-			FPerThreadFreeBlockLists* Lists = GMallocBinnedPerThreadCaches ? FPerThreadFreeBlockLists::Get() : nullptr;
-			if (Lists && (!Ptr || !IsOSAllocation(Ptr)))
-			{
-				uint32 BinSize = 0;
-				uint32 PoolIndex = 0;
-
-				bool bCanFree = true; // the nullptr is always "freeable"
-				if (Ptr)
-				{
-					// Reallocate to a smaller/bigger pool if necessary
-					FFreeBlock* Free = GetPoolHeaderFromPointer(Ptr);
-					BinSize = Free->BinSize;
-					PoolIndex = Free->PoolIndex;
-					// If canary is invalid we will assert in ReallocExternal. Otherwise it's the pre-fork canary and we will allocate new memory without touching this allocation.
-					bCanFree = Free->CanaryAndForkState == CurrentCanary;
-					if (NewSize && bCanFree && NewSize <= BinSize && (PoolIndex == 0 || NewSize > PoolIndexToBinSize(PoolIndex - 1)))
-					{
-						return Ptr;
-					}
-					bCanFree = bCanFree && Lists->CanFree(PoolIndex, BinSize);
-				}
-				if (bCanFree)
-				{
-					const uint32 NewPoolIndex = BoundSizeToPoolIndex(NewSize, MemSizeToPoolIndex);
-					const uint32 NewBinSize = PoolIndexToBinSize(NewPoolIndex);
-					void* Result = NewSize ? Lists->Malloc(NewPoolIndex) : nullptr;
-#if UE_MB2_ALLOCATOR_STATS
-					if (Result)
-					{
-						Lists->AllocatedMemory += NewBinSize;
-					}
-#endif
-					if (Result || !NewSize)
-					{
-						if (Result && Ptr)
-						{
-							FMemory::Memcpy(Result, Ptr, FPlatformMath::Min<SIZE_T>(NewSize, BinSize));
-						}
-						if (Ptr)
-						{
-							const bool bDidPush = Lists->Free(Ptr, PoolIndex, BinSize);
-							checkSlow(bDidPush);
-							UE_MBC_UPDATE_STATS(Lists->AllocatedMemory -= BinSize);
-						}
-
-						return Result;
-					}
-				}
-			}
-		}
-		void* Result = ReallocExternal(Ptr, NewSize, Alignment);
-		return Result;
-	}
-
-public:
 	UE_AUTORTFM_NOAUTORTFM
-	FORCEINLINE virtual void Free(void* Ptr) override
-	{
-#if UE_MB2_ALLOCATOR_STATS_VALIDATION
-		UE::TUniqueLock Lock(ValidationCriticalSection);
-		++RecursionCounter;
-		if (!IsOSAllocation(Ptr))
-		{
-			SIZE_T OutSize;
-			ensure(GetAllocationSize(Ptr, OutSize));
-			AllocatedSmallPoolMemoryValidation -= OutSize;
-		}
-		FreeInline(Ptr);
-		if (RecursionCounter == 1)
-		{
-			check(GetTotalAllocatedSmallPoolMemory() == AllocatedSmallPoolMemoryValidation);
-			if (GetTotalAllocatedSmallPoolMemory() != AllocatedSmallPoolMemoryValidation)
-			{
-				UE_DEBUG_BREAK();
-			}
-		}
-		--RecursionCounter;
-#else
-		FreeInline(Ptr);
-#endif
-	}
-
-protected:
-	FORCEINLINE void FreeInline(void* Ptr)
-	{
-		if (!IsOSAllocation(Ptr))
-		{
-			FPerThreadFreeBlockLists* Lists = GMallocBinnedPerThreadCaches ? FPerThreadFreeBlockLists::Get() : nullptr;
-			if (Lists)
-			{
-				FFreeBlock* BasePtr = GetPoolHeaderFromPointer(Ptr);
-				const int32 BinSize = BasePtr->BinSize;
-				// If canary is invalid we will assert in FreeExternal. Otherwise it's the pre-fork canary and we will turn this free into a no-op.
-				if (BasePtr->CanaryAndForkState == CurrentCanary && Lists->Free(Ptr, BasePtr->PoolIndex, BinSize))
-				{
-					UE_MBC_UPDATE_STATS(Lists->AllocatedMemory -= BinSize);
-					return;
-				}
-			}
-		}
-		FreeExternal(Ptr);
-	}
+	virtual void Free(void* Ptr) override;
 
-public:
 	inline bool GetSmallAllocationSize(void* Ptr, SIZE_T& SizeOut) const
 	{
 		if (!IsOSAllocation(Ptr))
@@ -483,13 +278,13 @@ public:
 		return QuantizeSizeCommon(Count, Alignment, *this);
 	}
 
-	CORE_API virtual bool ValidateHeap() override;
-	CORE_API virtual void Trim(bool bTrimThreadCaches) override;
-	CORE_API virtual const TCHAR* GetDescriptiveName() override;
-	CORE_API virtual void UpdateStats() override;
-	CORE_API virtual void OnMallocInitialized() override;
-	CORE_API virtual void OnPreFork() override;
-	CORE_API virtual void OnPostFork() override;
+	virtual bool ValidateHeap() override;
+	virtual void Trim(bool bTrimThreadCaches) override;
+	virtual const TCHAR* GetDescriptiveName() override;
+	virtual void UpdateStats() override;
+	virtual void OnMallocInitialized() override;
+	virtual void OnPreFork() override;
+	virtual void OnPostFork() override;
 	virtual uint64 GetImmediatelyFreeableCachedMemorySize() const override
 	{
 		return CachedOSPageAllocator.GetCachedImmediatelyFreeable();
@@ -500,12 +295,10 @@ public:
 	}
 	// End FMalloc interface.
 
-	CORE_API void* MallocExternalSmall(SIZE_T Size, uint32 Alignment);
-	CORE_API void* MallocExternalLarge(SIZE_T Size, uint32 Alignment);
-	CORE_API void* ReallocExternal(void* Ptr, SIZE_T NewSize, uint32 Alignment);
-	CORE_API void FreeExternal(void *Ptr);
+	void* MallocExternalSmall(SIZE_T Size, uint32 Alignment);
+	void* MallocExternalLarge(SIZE_T Size, uint32 Alignment);
 
-	CORE_API void CanaryFail(const FFreeBlock* Block) const;
+	void CanaryFail(const FFreeBlock* Block) const;
 	inline void CanaryTest(const FFreeBlock* Block) const
 	{
 #if BINNED2_FORK_SUPPORT
@@ -520,12 +313,12 @@ public:
 	}
 
 	/** Dumps current allocator stats to the log. */
-	CORE_API virtual void DumpAllocatorStats(class FOutputDevice& Ar) override;
+	virtual void DumpAllocatorStats(class FOutputDevice& Ar) override;
 	
-	static CORE_API uint16 SmallBinSizes[UE_MB2_SMALL_POOL_COUNT];
-	static CORE_API FMallocBinned2* MallocBinned2;
+	static uint16 SmallBinSizes[UE_MB2_SMALL_POOL_COUNT];
+	static FMallocBinned2* MallocBinned2;
 	// Mapping of sizes to small table indices
-	static CORE_API uint8 MemSizeToPoolIndex[1 + (UE_MB2_MAX_SMALL_POOL_SIZE >> UE_MBC_MIN_SMALL_POOL_ALIGNMENT_SHIFT)];
+	static uint8 MemSizeToPoolIndex[1 + (UE_MB2_MAX_SMALL_POOL_SIZE >> UE_MBC_MIN_SMALL_POOL_ALIGNMENT_SHIFT)];
 
 	static void* AllocateMetaDataMemory(SIZE_T Size);
 	static void FreeMetaDataMemory(void* Ptr, SIZE_T Size);
@@ -540,15 +333,6 @@ public:
 	void FlushCurrentThreadCacheInternal(bool bNewEpochOnly = false);
 };
 
-#define UE_MB2_INLINE (1)
-#if UE_MB2_INLINE // during development, it helps with iteration time to not include these here, but rather in the .cpp
-	#if PLATFORM_USES_FIXED_GMalloc_CLASS && !FORCE_ANSI_ALLOCATOR && USE_MALLOC_BINNED2
-		#define FMEMORY_INLINE_FUNCTION_DECORATOR  FORCEINLINE
-		#define FMEMORY_INLINE_GMalloc (FMallocBinned2::MallocBinned2)
-		#include "FMemory.inl" // IWYU pragma: export
-	#endif
-#endif
-
 #if UE_ENABLE_INCLUDE_ORDER_DEPRECATED_IN_5_7
 #	include "HAL/CriticalSection.h"
 #	if UE_MB2_ALLOCATOR_STATS_VALIDATION
diff --git a/Engine/Source/Runtime/Core/Public/HAL/MallocBinned3.h b/Engine/Source/Runtime/Core/Public/HAL/MallocBinned3.h
--- a/Engine/Source/Runtime/Core/Public/HAL/MallocBinned3.h
+++ b/Engine/Source/Runtime/Core/Public/HAL/MallocBinned3.h
@@ -292,130 +292,23 @@ private:
 
 public:
 
-	CORE_API FMallocBinned3(const FPlatformMemory::FPlatformVirtualMemoryBlock* ExternalMemoryBlock = nullptr);
+	FMallocBinned3(const FPlatformMemory::FPlatformVirtualMemoryBlock* ExternalMemoryBlock = nullptr);
 
-	CORE_API virtual ~FMallocBinned3();
+	virtual ~FMallocBinned3();
 
 	// FMalloc interface.
-	CORE_API virtual bool IsInternallyThreadSafe() const override;
+	virtual bool IsInternallyThreadSafe() const override;
 
 	UE_AUTORTFM_NOAUTORTFM
-	FORCEINLINE virtual void* Malloc(SIZE_T Size, uint32 Alignment) override
-	{
-		void* Result = nullptr;
-	
-		// Only allocate from the small pools if the size is small enough and the alignment isn't crazy large.
-		// With large alignments, we'll waste a lot of memory allocating an entire page, but such alignments are highly unlikely in practice.
-		if ((Size <= UE_MB3_MAX_SMALL_POOL_SIZE) & (Alignment <= UE_MBC_MIN_SMALL_POOL_ALIGNMENT)) // one branch, not two
-		{
-			FPerThreadFreeBlockLists* Lists = GMallocBinnedPerThreadCaches ? FPerThreadFreeBlockLists::Get() : nullptr;
-			if (Lists)
-			{
-				const uint32 PoolIndex = BoundSizeToPoolIndex(Size, MemSizeToPoolIndex);
-				Result = Lists->Malloc(PoolIndex);
-#if UE_MB3_ALLOCATOR_STATS
-				if (Result)
-				{
-					SmallPoolTables[PoolIndex].HeadEndAlloc(Size);
-					const uint32 BinSize = PoolIndexToBinSize(PoolIndex);
-					Lists->AllocatedMemory += BinSize;
-				}
-#endif
-			}
-		}
-		if (Result == nullptr)
-		{
-			Result = MallocExternal(Size, Alignment);
-		}
-
-		return Result;
-	}
+	virtual void* Malloc(SIZE_T Size, uint32 Alignment) override;
 
 	UE_AUTORTFM_NOAUTORTFM
-	FORCEINLINE virtual void* Realloc(void* Ptr, SIZE_T NewSize, uint32 Alignment) override
-	{
-		if (NewSize <= UE_MB3_MAX_SMALL_POOL_SIZE && Alignment <= UE_MBC_MIN_SMALL_POOL_ALIGNMENT) // one branch, not two
-		{
-			FPerThreadFreeBlockLists* Lists = GMallocBinnedPerThreadCaches ? FPerThreadFreeBlockLists::Get() : nullptr;
-
-			const uint64 PoolIndex = PoolIndexFromPtr(Ptr);
-			if ((!!Lists) & ((!Ptr) | (PoolIndex < UE_MB3_SMALL_POOL_COUNT)))
-			{
-				uint32 BinSize = 0;
-
-				bool bCanFree = true; // the nullptr is always "freeable"
-				if (Ptr)
-				{
-					// Reallocate to a smaller/bigger pool if necessary
-					BinSize = PoolIndexToBinSize(PoolIndex);
-					if ((!!NewSize) & (NewSize <= BinSize) & ((!PoolIndex) | (NewSize > PoolIndexToBinSize(static_cast<uint32>(PoolIndex - 1)))))
-					{
-#if UE_MB3_ALLOCATOR_STATS
-						SmallPoolTables[PoolIndex].HeadEndAlloc(NewSize);
-						SmallPoolTables[PoolIndex].HeadEndFree();
-#endif
-						return Ptr;
-					}
-					bCanFree = Lists->CanFree(PoolIndex, BinSize);
-				}
-				if (bCanFree)
-				{
-					const uint32 NewPoolIndex = BoundSizeToPoolIndex(NewSize, MemSizeToPoolIndex);
-					const uint32 NewBinSize = PoolIndexToBinSize(NewPoolIndex);
-					void* Result = NewSize ? Lists->Malloc(NewPoolIndex) : nullptr;
-#if UE_MB3_ALLOCATOR_STATS
-					if (Result)
-					{
-						SmallPoolTables[NewPoolIndex].HeadEndAlloc(NewSize);
-						Lists->AllocatedMemory += NewBinSize;
-					}
-#endif
-					if (Result || !NewSize)
-					{
-						if (Result && Ptr)
-						{
-							FMemory::Memcpy(Result, Ptr, FPlatformMath::Min<SIZE_T>(NewSize, BinSize));
-						}
-						if (Ptr)
-						{
-							const bool bDidPush = Lists->Free(Ptr, PoolIndex, BinSize);
-							checkSlow(bDidPush);
-#if UE_MB3_ALLOCATOR_STATS
-							SmallPoolTables[PoolIndex].HeadEndFree();
-							Lists->AllocatedMemory -= BinSize;
-#endif
-						}
-
-						return Result;
-					}
-				}
-			}
-		}
-		void* Result = ReallocExternal(Ptr, NewSize, Alignment);
-		return Result;
-	}
+	virtual void* Realloc(void* Ptr, SIZE_T NewSize, uint32 Alignment) override;
 
 	UE_AUTORTFM_NOAUTORTFM
 	FORCEINLINE virtual void Free(void* Ptr) override
 	{
-		const uint64 PoolIndex = PoolIndexFromPtr(Ptr);
-		if (PoolIndex < UE_MB3_SMALL_POOL_COUNT)
-		{
-			FPerThreadFreeBlockLists* Lists = GMallocBinnedPerThreadCaches ? FPerThreadFreeBlockLists::Get() : nullptr;
-			if (Lists)
-			{
-				const int32 BinSize = PoolIndexToBinSize(PoolIndex);
-				if (Lists->Free(Ptr, PoolIndex, BinSize))	//TODO: move BinSize into the FPerThreadFreeBlockLists list itself to avoid doing a lookup here
-				{
-#if UE_MB3_ALLOCATOR_STATS
-					SmallPoolTables[PoolIndex].HeadEndFree();
-					Lists->AllocatedMemory -= BinSize;
-#endif
-					return;
-				}
-			}
-		}
-		FreeExternal(Ptr, PoolIndex);
+		FreeExternal(Ptr, PoolIndexFromPtr(Ptr));
 	}
 
 	inline bool GetSmallAllocationSize(void* Ptr, SIZE_T& SizeOut) const
@@ -444,20 +337,18 @@ public:
 		return QuantizeSizeCommon(Count, Alignment, *this);
 	}
 
-	CORE_API virtual bool ValidateHeap() override;
-	CORE_API virtual void Trim(bool bTrimThreadCaches) override;
-	CORE_API virtual const TCHAR* GetDescriptiveName() override;
+	virtual bool ValidateHeap() override;
+	virtual void Trim(bool bTrimThreadCaches) override;
+	virtual const TCHAR* GetDescriptiveName() override;
 	/** Dumps current allocator stats to the log. */
-	CORE_API virtual void DumpAllocatorStats(class FOutputDevice& Ar) override;
-	CORE_API virtual void UpdateStats() override
+	virtual void DumpAllocatorStats(class FOutputDevice& Ar) override;
+	virtual void UpdateStats() override
 	{
 		UpdateStatsCommon(*this);
 	}
 	// End FMalloc interface.
 
-	CORE_API void* MallocExternal(SIZE_T Size, uint32 Alignment);
-	CORE_API void* ReallocExternal(void* Ptr, SIZE_T NewSize, uint32 Alignment);
-	CORE_API void FreeExternal(void *Ptr, uint64 PoolIndex);
+	void FreeExternal(void* Ptr, uint64 PoolIndex);
 	
 	// +1 enables PoolIndexToBinSize(~0u / -1) dummy access that helps avoid PoolIndex == 0 branching in Realloc(),
 	// see ((!PoolIndex) | (NewSize > PoolIndexToBinSize(static_cast<uint32>(PoolIndex - 1)))
@@ -494,15 +385,6 @@ public:
 
 PRAGMA_RESTORE_UNSAFE_TYPECAST_WARNINGS
 
-#define UE_MB3_INLINE (1)
-#if UE_MB3_INLINE // during development, it helps with iteration time to not include these here, but rather in the .cpp
-	#if PLATFORM_USES_FIXED_GMalloc_CLASS && !FORCE_ANSI_ALLOCATOR && 0/*USE_MALLOC_BINNED3*/
-		#define FMEMORY_INLINE_FUNCTION_DECORATOR  FORCEINLINE
-		#define FMEMORY_INLINE_GMalloc (FMallocBinned3::MallocBinned3)
-		#include "FMemory.inl" // IWYU pragma: export
-	#endif
-#endif
-
 #if UE_ENABLE_INCLUDE_ORDER_DEPRECATED_IN_5_7
 #	include "HAL/CriticalSection.h"
 #endif
diff --git a/Engine/Source/Runtime/Core/Public/HAL/MallocBinnedCommon.h b/Engine/Source/Runtime/Core/Public/HAL/MallocBinnedCommon.h
--- a/Engine/Source/Runtime/Core/Public/HAL/MallocBinnedCommon.h
+++ b/Engine/Source/Runtime/Core/Public/HAL/MallocBinnedCommon.h
@@ -3,7 +3,6 @@
 #pragma once
 
 #include "CoreTypes.h"
-#include "CoreGlobals.h"
 #include <atomic>
 #include "HAL/MemoryBase.h"
 #include "HAL/PlatformMutex.h"
@@ -12,7 +11,7 @@
 #include "Async/WordMutex.h"
 #include "Misc/ScopeLock.h"
 #include "Templates/AlignmentTemplates.h"
-#include "Logging/LogMacros.h"
+#include "Templates/MemoryOps.h"
 #include "ProfilingDebugging/CsvProfiler.h"
 
 
@@ -421,12 +420,14 @@ protected:
 	void GetAllocatorStatsInternal(FGenericMemoryStats& OutStats, int64 TotalAllocatedSmallPoolMemory);
 #endif
 
-	static CA_NO_RETURN void OutOfMemory(uint64 Size, uint32 Alignment = 0)
+	[[noreturn]] static void OutOfMemory(uint64 Size, uint32 Alignment = 0)
 	{
 		// this is expected not to return
 		FPlatformMemory::OnOutOfMemory(Size, Alignment);
 	}
 
+	[[noreturn]] void UnrecognizedPointerFatalError(void* Ptr);
+
 #if UE_MBC_LOG_LARGE_ALLOCATION
 	void LogLargeAllocation(SIZE_T Size) const;
 #else
@@ -867,7 +868,8 @@ protected:
 		return PoolIndex;
 	}
 
-	bool PromoteToLargerBin(SIZE_T& Size, uint32& Alignment, const AllocType& Alloc) const
+	// force no inline, so it will not bloat fast code path since this is unlikely to happen
+	FORCENOINLINE bool PromoteToLargerBin(SIZE_T& Size, uint32& Alignment, const AllocType& Alloc) const
 	{
 		// try to promote our allocation request to a larger bin with a matching natural alignment
 		// if requested alignment is larger than UE_MBC_MIN_SMALL_POOL_ALIGNMENT but smaller than UE_MBC_MAX_SMALL_POOL_ALIGNMENT
@@ -915,7 +917,7 @@ protected:
 
 		if (!Pool)
 		{
-			UE_LOG(LogMemory, Fatal, TEXT("FMallocBinned Attempt to GetAllocationSizeExternal an unrecognized pointer %p"), Ptr);
+			UnrecognizedPointerFatalError(Ptr);
 		}
 		const SIZE_T PoolOsBytes = Pool->GetOsAllocatedBytes();
 		const SIZE_T PoolOSRequestedBytes = Pool->GetOSRequestedBytes();
diff --git a/Engine/Source/Runtime/Core/Public/HAL/Memory.h b/Engine/Source/Runtime/Core/Public/HAL/Memory.h
new file mode 100644
--- /dev/null
+++ b/Engine/Source/Runtime/Core/Public/HAL/Memory.h
@@ -0,0 +1,19 @@
+// Copyright Epic Games, Inc. All Rights Reserved.
+
+#pragma once
+
+#if PLATFORM_USES_FIXED_GMalloc_CLASS
+#	if USE_MALLOC_BINNED2
+#		include "HAL/MallocBinned2.h" // HEADER_UNIT_IGNORE - Causes circular includes
+#		define FMEMORY_INLINE_GMalloc (FMallocBinned2::MallocBinned2)
+#	elif USE_MALLOC_BINNED3
+#		include "HAL/MallocBinned3.h" // HEADER_UNIT_IGNORE - Causes circular includes
+#		define FMEMORY_INLINE_GMalloc (FMallocBinned3::MallocBinned3)
+#	endif
+#endif
+
+#if !defined(FMEMORY_INLINE_GMalloc)
+#	define FMEMORY_INLINE_GMalloc UE::Private::GMalloc
+#endif
+
+#include "HAL/FMemory.inl"
\ No newline at end of file
diff --git a/Engine/Source/Runtime/Core/Public/HAL/UnrealMemory.h b/Engine/Source/Runtime/Core/Public/HAL/UnrealMemory.h
--- a/Engine/Source/Runtime/Core/Public/HAL/UnrealMemory.h
+++ b/Engine/Source/Runtime/Core/Public/HAL/UnrealMemory.h
@@ -272,8 +272,6 @@ struct FMemory
 	*/
 	static CORE_API void ExplicitInit(FMalloc& Allocator);
 
-private:
-	static CORE_API void GCreateMalloc();
 	// These versions are called either at startup or in the event of a crash
 	static CORE_API void* MallocExternal(SIZE_T Count, uint32 Alignment = DEFAULT_ALIGNMENT);
 	static CORE_API void* ReallocExternal(void* Original, SIZE_T Count, uint32 Alignment = DEFAULT_ALIGNMENT);
@@ -281,21 +279,7 @@ private:
 	static CORE_API SIZE_T GetAllocSizeExternal(void* Original);
 	static CORE_API void* MallocZeroedExternal(SIZE_T Count, uint32 Alignment = DEFAULT_ALIGNMENT);
 	static CORE_API SIZE_T QuantizeSizeExternal(SIZE_T Count, uint32 Alignment = DEFAULT_ALIGNMENT);
-};
-
-#define INLINE_FMEMORY_OPERATION (0) // untested, but should work. Inlines FMemory::Malloc, etc
 
-#if INLINE_FMEMORY_OPERATION
-	#if PLATFORM_USES_FIXED_GMalloc_CLASS
-		#error "PLATFORM_USES_FIXED_GMalloc_CLASS and INLINE_FMEMORY_OPERATION are not compatible. PLATFORM_USES_FIXED_GMalloc_CLASS is inlined below."
-	#endif
-
-	#define FMEMORY_INLINE_FUNCTION_DECORATOR FORCEINLINE
-// IWYU pragma: begin_export
-	#include "FMemory.inl" // HEADER_UNIT_IGNORE - Causes circular includes
-// IWYU pragma: end_export
-#endif
-
-#if PLATFORM_USES_FIXED_GMalloc_CLASS && !FORCE_ANSI_ALLOCATOR && USE_MALLOC_BINNED2
-	#include "HAL/MallocBinned2.h" // HEADER_UNIT_IGNORE - Causes circular includes
-#endif
+private:
+	static CORE_API void GCreateMalloc();
+};
\ No newline at end of file
diff --git a/Engine/Source/Runtime/Core/Public/Modules/Boilerplate/ModuleBoilerplate.h b/Engine/Source/Runtime/Core/Public/Modules/Boilerplate/ModuleBoilerplate.h
--- a/Engine/Source/Runtime/Core/Public/Modules/Boilerplate/ModuleBoilerplate.h
+++ b/Engine/Source/Runtime/Core/Public/Modules/Boilerplate/ModuleBoilerplate.h
@@ -5,6 +5,7 @@
 #include "CoreTypes.h"
 #include "HAL/UnrealMemory.h"
 #include "UObject/NameTypes.h"
+#include "HAL/Memory.h"
 // Include the full definition of FVisualizerDebuggingState so visualizers can see the full type information
 //   Without this include (and just a fwd decl), visualizers in some modules will not be able to resolve
 //   FVisualizerDebuggingState::Ptrs or FVisualizerDebuggingState::GuidString
@@ -43,30 +44,30 @@
 #if !FORCE_ANSI_ALLOCATOR && !defined(__clang_analyzer__)
 	static_assert(__STDCPP_DEFAULT_NEW_ALIGNMENT__ <= 16, "Expecting 16-byte default operator new alignment - alignments > 16 may have bloat");
 	#define REPLACEMENT_OPERATOR_NEW_AND_DELETE \
-		OPERATOR_NEW_MSVC_PRAGMA void* operator new  ( size_t Size                                                    ) OPERATOR_NEW_THROW_SPEC      { return FMemory::Malloc( Size ? Size : 1, __STDCPP_DEFAULT_NEW_ALIGNMENT__ ); } \
-		OPERATOR_NEW_MSVC_PRAGMA void* operator new[]( size_t Size                                                    ) OPERATOR_NEW_THROW_SPEC      { return FMemory::Malloc( Size ? Size : 1, __STDCPP_DEFAULT_NEW_ALIGNMENT__ ); } \
-		OPERATOR_NEW_MSVC_PRAGMA void* operator new  ( size_t Size,                             const std::nothrow_t& ) OPERATOR_NEW_NOTHROW_SPEC    { return FMemory::Malloc( Size ? Size : 1, __STDCPP_DEFAULT_NEW_ALIGNMENT__ ); } \
-		OPERATOR_NEW_MSVC_PRAGMA void* operator new[]( size_t Size,                             const std::nothrow_t& ) OPERATOR_NEW_NOTHROW_SPEC    { return FMemory::Malloc( Size ? Size : 1, __STDCPP_DEFAULT_NEW_ALIGNMENT__ ); } \
-		OPERATOR_NEW_MSVC_PRAGMA void* operator new  ( size_t Size, std::align_val_t Alignment                        ) OPERATOR_NEW_THROW_SPEC      { return FMemory::Malloc( Size ? Size : 1, (std::size_t)Alignment ); } \
-		OPERATOR_NEW_MSVC_PRAGMA void* operator new[]( size_t Size, std::align_val_t Alignment                        ) OPERATOR_NEW_THROW_SPEC      { return FMemory::Malloc( Size ? Size : 1, (std::size_t)Alignment ); } \
-		OPERATOR_NEW_MSVC_PRAGMA void* operator new  ( size_t Size, std::align_val_t Alignment, const std::nothrow_t& ) OPERATOR_NEW_NOTHROW_SPEC    { return FMemory::Malloc( Size ? Size : 1, (std::size_t)Alignment ); } \
-		OPERATOR_NEW_MSVC_PRAGMA void* operator new[]( size_t Size, std::align_val_t Alignment, const std::nothrow_t& ) OPERATOR_NEW_NOTHROW_SPEC    { return FMemory::Malloc( Size ? Size : 1, (std::size_t)Alignment ); } \
-		void operator delete  ( void* Ptr                                                                             ) OPERATOR_DELETE_THROW_SPEC   { FMemory::Free( Ptr ); } \
-		void operator delete[]( void* Ptr                                                                             ) OPERATOR_DELETE_THROW_SPEC   { FMemory::Free( Ptr ); } \
-		void operator delete  ( void* Ptr,                                                      const std::nothrow_t& ) OPERATOR_DELETE_NOTHROW_SPEC { FMemory::Free( Ptr ); } \
-		void operator delete[]( void* Ptr,                                                      const std::nothrow_t& ) OPERATOR_DELETE_NOTHROW_SPEC { FMemory::Free( Ptr ); } \
-		void operator delete  ( void* Ptr,             size_t Size                                                    ) OPERATOR_DELETE_THROW_SPEC   { FMemory::Free( Ptr ); } \
-		void operator delete[]( void* Ptr,             size_t Size                                                    ) OPERATOR_DELETE_THROW_SPEC   { FMemory::Free( Ptr ); } \
-		void operator delete  ( void* Ptr,             size_t Size,                             const std::nothrow_t& ) OPERATOR_DELETE_NOTHROW_SPEC { FMemory::Free( Ptr ); } \
-		void operator delete[]( void* Ptr,             size_t Size,                             const std::nothrow_t& ) OPERATOR_DELETE_NOTHROW_SPEC { FMemory::Free( Ptr ); } \
-		void operator delete  ( void* Ptr,                          std::align_val_t Alignment                        ) OPERATOR_DELETE_THROW_SPEC   { FMemory::Free( Ptr ); } \
-		void operator delete[]( void* Ptr,                          std::align_val_t Alignment                        ) OPERATOR_DELETE_THROW_SPEC   { FMemory::Free( Ptr ); } \
-		void operator delete  ( void* Ptr,                          std::align_val_t Alignment, const std::nothrow_t& ) OPERATOR_DELETE_NOTHROW_SPEC { FMemory::Free( Ptr ); } \
-		void operator delete[]( void* Ptr,                          std::align_val_t Alignment, const std::nothrow_t& ) OPERATOR_DELETE_NOTHROW_SPEC { FMemory::Free( Ptr ); } \
-		void operator delete  ( void* Ptr,             size_t Size, std::align_val_t Alignment                        ) OPERATOR_DELETE_THROW_SPEC   { FMemory::Free( Ptr ); } \
-		void operator delete[]( void* Ptr,             size_t Size, std::align_val_t Alignment                        ) OPERATOR_DELETE_THROW_SPEC   { FMemory::Free( Ptr ); } \
-		void operator delete  ( void* Ptr,             size_t Size, std::align_val_t Alignment, const std::nothrow_t& ) OPERATOR_DELETE_NOTHROW_SPEC { FMemory::Free( Ptr ); } \
-		void operator delete[]( void* Ptr,             size_t Size, std::align_val_t Alignment, const std::nothrow_t& ) OPERATOR_DELETE_NOTHROW_SPEC { FMemory::Free( Ptr ); }
+		FORCENOINLINE OPERATOR_NEW_MSVC_PRAGMA void* operator new  ( size_t Size                                                    ) OPERATOR_NEW_THROW_SPEC      { return FMemory_MallocInline( Size ? Size : 1, __STDCPP_DEFAULT_NEW_ALIGNMENT__ ); } \
+		FORCENOINLINE OPERATOR_NEW_MSVC_PRAGMA void* operator new[]( size_t Size                                                    ) OPERATOR_NEW_THROW_SPEC      { return FMemory_MallocInline( Size ? Size : 1, __STDCPP_DEFAULT_NEW_ALIGNMENT__ ); } \
+		FORCENOINLINE OPERATOR_NEW_MSVC_PRAGMA void* operator new  ( size_t Size,                             const std::nothrow_t& ) OPERATOR_NEW_NOTHROW_SPEC    { return FMemory_MallocInline( Size ? Size : 1, __STDCPP_DEFAULT_NEW_ALIGNMENT__ ); } \
+		FORCENOINLINE OPERATOR_NEW_MSVC_PRAGMA void* operator new[]( size_t Size,                             const std::nothrow_t& ) OPERATOR_NEW_NOTHROW_SPEC    { return FMemory_MallocInline( Size ? Size : 1, __STDCPP_DEFAULT_NEW_ALIGNMENT__ ); } \
+		FORCENOINLINE OPERATOR_NEW_MSVC_PRAGMA void* operator new  ( size_t Size, std::align_val_t Alignment                        ) OPERATOR_NEW_THROW_SPEC      { return FMemory_MallocInline( Size ? Size : 1, (std::size_t)Alignment ); } \
+		FORCENOINLINE OPERATOR_NEW_MSVC_PRAGMA void* operator new[]( size_t Size, std::align_val_t Alignment                        ) OPERATOR_NEW_THROW_SPEC      { return FMemory_MallocInline( Size ? Size : 1, (std::size_t)Alignment ); } \
+		FORCENOINLINE OPERATOR_NEW_MSVC_PRAGMA void* operator new  ( size_t Size, std::align_val_t Alignment, const std::nothrow_t& ) OPERATOR_NEW_NOTHROW_SPEC    { return FMemory_MallocInline( Size ? Size : 1, (std::size_t)Alignment ); } \
+		FORCENOINLINE OPERATOR_NEW_MSVC_PRAGMA void* operator new[]( size_t Size, std::align_val_t Alignment, const std::nothrow_t& ) OPERATOR_NEW_NOTHROW_SPEC    { return FMemory_MallocInline( Size ? Size : 1, (std::size_t)Alignment ); } \
+		FORCENOINLINE void operator delete  ( void* Ptr                                                                             ) OPERATOR_DELETE_THROW_SPEC   { FMemory_FreeInline( Ptr ); } \
+		FORCENOINLINE void operator delete[]( void* Ptr                                                                             ) OPERATOR_DELETE_THROW_SPEC   { FMemory_FreeInline( Ptr ); } \
+		FORCENOINLINE void operator delete  ( void* Ptr,                                                      const std::nothrow_t& ) OPERATOR_DELETE_NOTHROW_SPEC { FMemory_FreeInline( Ptr ); } \
+		FORCENOINLINE void operator delete[]( void* Ptr,                                                      const std::nothrow_t& ) OPERATOR_DELETE_NOTHROW_SPEC { FMemory_FreeInline( Ptr ); } \
+		FORCENOINLINE void operator delete  ( void* Ptr,             size_t Size                                                    ) OPERATOR_DELETE_THROW_SPEC   { FMemory_FreeInline( Ptr ); } \
+		FORCENOINLINE void operator delete[]( void* Ptr,             size_t Size                                                    ) OPERATOR_DELETE_THROW_SPEC   { FMemory_FreeInline( Ptr ); } \
+		FORCENOINLINE void operator delete  ( void* Ptr,             size_t Size,                             const std::nothrow_t& ) OPERATOR_DELETE_NOTHROW_SPEC { FMemory_FreeInline( Ptr ); } \
+		FORCENOINLINE void operator delete[]( void* Ptr,             size_t Size,                             const std::nothrow_t& ) OPERATOR_DELETE_NOTHROW_SPEC { FMemory_FreeInline( Ptr ); } \
+		FORCENOINLINE void operator delete  ( void* Ptr,                          std::align_val_t Alignment                        ) OPERATOR_DELETE_THROW_SPEC   { FMemory_FreeInline( Ptr ); } \
+		FORCENOINLINE void operator delete[]( void* Ptr,                          std::align_val_t Alignment                        ) OPERATOR_DELETE_THROW_SPEC   { FMemory_FreeInline( Ptr ); } \
+		FORCENOINLINE void operator delete  ( void* Ptr,                          std::align_val_t Alignment, const std::nothrow_t& ) OPERATOR_DELETE_NOTHROW_SPEC { FMemory_FreeInline( Ptr ); } \
+		FORCENOINLINE void operator delete[]( void* Ptr,                          std::align_val_t Alignment, const std::nothrow_t& ) OPERATOR_DELETE_NOTHROW_SPEC { FMemory_FreeInline( Ptr ); } \
+		FORCENOINLINE void operator delete  ( void* Ptr,             size_t Size, std::align_val_t Alignment                        ) OPERATOR_DELETE_THROW_SPEC   { FMemory_FreeInline( Ptr ); } \
+		FORCENOINLINE void operator delete[]( void* Ptr,             size_t Size, std::align_val_t Alignment                        ) OPERATOR_DELETE_THROW_SPEC   { FMemory_FreeInline( Ptr ); } \
+		FORCENOINLINE void operator delete  ( void* Ptr,             size_t Size, std::align_val_t Alignment, const std::nothrow_t& ) OPERATOR_DELETE_NOTHROW_SPEC { FMemory_FreeInline( Ptr ); } \
+		FORCENOINLINE void operator delete[]( void* Ptr,             size_t Size, std::align_val_t Alignment, const std::nothrow_t& ) OPERATOR_DELETE_NOTHROW_SPEC { FMemory_FreeInline( Ptr ); }
 #else
 	#define REPLACEMENT_OPERATOR_NEW_AND_DELETE
 #endif
