Mac - Include Mac MainThread in CrashContext.runtime-xml for CrashReporter

#tests Local test on Mac & Preflight

diff --git a/Engine/Source/Runtime/Core/Private/HAL/ThreadingBase.cpp b/Engine/Source/Runtime/Core/Private/HAL/ThreadingBase.cpp
--- a/Engine/Source/Runtime/Core/Private/HAL/ThreadingBase.cpp
+++ b/Engine/Source/Runtime/Core/Private/HAL/ThreadingBase.cpp
@@ -21,6 +21,10 @@
 #include "Windows/WindowsHWrapper.h"
 #endif
 
+#if PLATFORM_MAC
+#include "Mac/CocoaThread.h"
+#endif
+
 #include <atomic>
 
 #ifndef IS_RUNNING_GAMETHREAD_ON_EXTERNAL_THREAD
@@ -608,6 +612,9 @@ void FThreadManager::GetAllThreadStackBackTraces(TArray<FThreadStackBackTrace>&
 	const int32 NumThreads = Threads.Num() + 1;
 
 	StackTraces.Empty(NumThreads);
+#if PLATFORM_MAC && MAC_SEPARATE_GAME_THREAD
+	GetAllThreadStackBackTraces_ProcessSingle(CurThreadId, GMacMainThreadId, TEXT("MacMainThread"), StackTraces.AddDefaulted_GetRef());
+#endif
 	GetAllThreadStackBackTraces_ProcessSingle(CurThreadId, GGameThreadId, TEXT("GameThread"), StackTraces.AddDefaulted_GetRef());
 
 	ForEachThread(
@@ -627,6 +634,17 @@ void FThreadManager::ForEachThreadStackBackTrace(TFunctionRef<bool(uint32 Thread
 
 	FScopeLock Lock(&ThreadsCritical);
 	bIsThreadListDirty = false;
+	
+#if PLATFORM_MAC && MAC_SEPARATE_GAME_THREAD
+	{
+		const TConstArrayView<uint64> ProgramCounterOutputArrayView = ThreadStackBackTraces_PerformStackWalk(CurThreadId, GMacMainThreadId, ProgramCounterBuffer);
+		const bool bContinue = Func(GMacMainThreadId, TEXT("MacMainThread"), ProgramCounterOutputArrayView);
+		if (!CheckThreadListSafeToContinueIteration() || !bContinue)
+		{
+			return;
+		}
+	}
+#endif
 
 	{
 		const TConstArrayView<uint64> ProgramCounterOutputArrayView = ThreadStackBackTraces_PerformStackWalk(CurThreadId, GGameThreadId, ProgramCounterBuffer);
diff --git a/Engine/Source/Runtime/Core/Private/Mac/CocoaThread.cpp b/Engine/Source/Runtime/Core/Private/Mac/CocoaThread.cpp
--- a/Engine/Source/Runtime/Core/Private/Mac/CocoaThread.cpp
+++ b/Engine/Source/Runtime/Core/Private/Mac/CocoaThread.cpp
@@ -6,8 +6,6 @@
 #include "Templates/RefCounting.h"
 #include "Misc/OutputDeviceRedirector.h"
 
-#define MAC_SEPARATE_GAME_THREAD 1 // Separate the main & game threads so that we better handle the interaction between the Cocoa's event delegates and UE's event polling.
-
 // this is the size of the game thread stack, it must be a multiple of 4k
 #if (UE_BUILD_DEBUG)
 	#define GAME_THREAD_STACK_SIZE 64 * 1024 * 1024
@@ -23,7 +21,7 @@ NSString* UnrealCloseEventMode = @"UnrealCloseEventMode";
 NSString* UnrealIMEEventMode = @"UnrealIMEEventMode";
 
 static FCocoaGameThread* GCocoaGameThread = nil;
-static uint32 GMainThreadId = 0;
+uint32 GMacMainThreadId = 0;
 
 static TAutoConsoleVariable<bool> CVarMacCocoaThreadIgnoreWaitMode(
 	TEXT("Mac.IgnoreCocoaThreadWaitMode"),
@@ -382,7 +380,7 @@ FCocoaRunLoopSource* FCocoaRunLoopSource::GameRunLoopSource = nullptr;
 	
 	// We have exited the game thread, so any UE code running now should treat the Main thread
 	// as the game thread, so we don't crash in static destructors.
-	GGameThreadId = GMainThreadId;
+	GGameThreadId = GMacMainThreadId;
 	GCocoaGameThread = nil;
 	
 	// Tell the main thread we are OK to quit, but don't wait for it.
@@ -493,7 +491,7 @@ void RunGameThread(id Target, SEL Selector)
 	// @todo: Proper support for sudden termination (OS termination of application without any events, notifications or signals) - which presently can assert, crash or corrupt. At present we can't deal with this and we need to ensure that we don't permit the system to kill us without warning by disabling & enabling support around operations which must be committed to disk atomically.
 	[[NSProcessInfo processInfo] disableSuddenTermination];
 	
-	GMainThreadId = FPlatformTLS::GetCurrentThreadId();
+	GMacMainThreadId = FPlatformTLS::GetCurrentThreadId();
 	
 #if MAC_SEPARATE_GAME_THREAD
 	// Register main run loop source
diff --git a/Engine/Source/Runtime/Core/Public/Mac/CocoaThread.h b/Engine/Source/Runtime/Core/Public/Mac/CocoaThread.h
--- a/Engine/Source/Runtime/Core/Public/Mac/CocoaThread.h
+++ b/Engine/Source/Runtime/Core/Public/Mac/CocoaThread.h
@@ -3,6 +3,8 @@
 
 #include "CoreMinimal.h"
 
+#define MAC_SEPARATE_GAME_THREAD 1 // Separate the main & game threads so that we better handle the interaction between the Cocoa's event delegates and UE's event polling.
+
 /* Custom run-loop modes for Unreal that process only certain kinds of events to simulate Windows event ordering. */
 CORE_API extern NSString* UnrealNilEventMode; /* Process only mandatory events */
 CORE_API extern NSString* UnrealShowEventMode; /* Process only show window events */
@@ -11,6 +13,11 @@ CORE_API extern NSString* UnrealFullscreenEventMode; /* Process only fullscreen
 CORE_API extern NSString* UnrealCloseEventMode; /* Process only close window events */
 CORE_API extern NSString* UnrealIMEEventMode; /* Process only input method events */
 
+#if MAC_SEPARATE_GAME_THREAD
+/** Thread ID of the Mac MainThread */
+extern CORE_API uint32 GMacMainThreadId;
+#endif
+
 @interface NSThread (FCocoaThread)
 + (NSThread*) gameThread; // Returns the main game thread, or nil if has yet to be constructed.
 + (bool) isGameThread; // True if the current thread is the main game thread, else false.
