Fixed compilation due to a missing CoreGlobals.h include

diff --git a/Engine/Source/Runtime/Core/Public/HAL/Memory.h b/Engine/Source/Runtime/Core/Public/HAL/Memory.h
--- a/Engine/Source/Runtime/Core/Public/HAL/Memory.h
+++ b/Engine/Source/Runtime/Core/Public/HAL/Memory.h
@@ -13,6 +13,7 @@
 #endif
 
 #if !defined(FMEMORY_INLINE_GMalloc)
+#	include "CoreGlobals.h"
 #	define FMEMORY_INLINE_GMalloc UE::Private::GMalloc
 #endif
 
