[IOS] Notification about network path characteristic

FIOSPlatformMisc::OnNetworkConnectionCharacteristicsChanged could be used to know about network characteristic changes, e.g. knowing if new connection is constrained or expensive.

[RN]

diff --git a/Engine/Source/Runtime/Core/Private/IOS/IOSPlatformMisc.cpp b/Engine/Source/Runtime/Core/Private/IOS/IOSPlatformMisc.cpp
--- a/Engine/Source/Runtime/Core/Private/IOS/IOSPlatformMisc.cpp
+++ b/Engine/Source/Runtime/Core/Private/IOS/IOSPlatformMisc.cpp
@@ -8,6 +8,7 @@
 #include "Apple/ApplePlatformMisc.h"
 #include "Async/TaskGraphInterfaces.h"
 #include "Apple/ApplePlatformCrashContext.h"
+#include "Delegates/Delegate.h"
 #include "HAL/ExceptionHandling.h"
 #include "HAL/FileManager.h"
 #include "HAL/PlatformOutputDevices.h"
@@ -93,9 +94,30 @@ namespace UE::FIOSPlatformMisc
 {
 	// Used to monitor network state. Replaces the deprecated SCNetworkReachability API.
 	// Note: defined here to not force objcisims into the header which is included on non-objc platforms
-	nw_path_monitor_t PathMonitor = nullptr;
-	nw_path_t CachedNetworkPath = nullptr;
-	ENetworkConnectionType LastReportedNetworkType = ENetworkConnectionType::Unknown;
+	static nw_path_monitor_t PathMonitor = nullptr;
+	static nw_path_t CachedNetworkPath = nullptr;
+	static ENetworkConnectionType LastReportedNetworkType = ENetworkConnectionType::Unknown;
+
+	static TMulticastDelegate<void(::FIOSPlatformMisc::FNetworkConnectionCharacteristics), FDefaultTSDelegateUserPolicy> OnNetworkConnectionCharacteristicsChanged;
+
+	// This is called from a whatever task graph thread picks.
+	static void NotifyAboutNetworkConnectionCharacteristicsChange(nw_path_t InPath)
+	{
+		::FIOSPlatformMisc::FNetworkConnectionCharacteristics Details;
+		
+		Details.bSupportsDNS   = nw_path_has_dns(InPath);
+		Details.bSupportsIPv4  = nw_path_has_ipv4(InPath);
+		Details.bSupportsIPv6  = nw_path_has_ipv6(InPath);
+		Details.bIsConstrained = nw_path_is_constrained(InPath);
+		Details.bIsExpensive   = nw_path_is_expensive(InPath);
+		
+		// Path's reference count was increased to guarantee it lives
+		// til this moment of time. Refer to nw_path_monitor_set_update_handler
+		// callback below. However, we just used it, so:
+		nw_release(InPath);
+
+		OnNetworkConnectionCharacteristicsChanged.Broadcast(Details);
+	}
 
 	static ENetworkConnectionType CheckNetworkConnectionType()
 	{
@@ -159,7 +181,7 @@ namespace UE::FIOSPlatformMisc
 			}
 		}
 	}
-}
+} // UE::FIOSPlatformMisc
 
 void FIOSPlatformMisc::PlatformInit()
 {
@@ -251,7 +273,10 @@ void FIOSPlatformMisc::PlatformInit()
 		UE::FIOSPlatformMisc::PathMonitor, 
 		^(nw_path_t InPath)
 		{
-			UE_LOG(LogInit, Log, TEXT("Network path changed"));
+			bool bIsConstrained = nw_path_is_constrained(InPath);
+			bool bIsExpensive   = nw_path_is_expensive(InPath);
+
+			UE_LOG(LogInit, Log, TEXT("Network path changed (IsConstrained=%d, IsExpensive=%d)"), bIsConstrained, bIsExpensive);
 			
 			// cache the path we get
 			if (UE::FIOSPlatformMisc::PathMonitor != nullptr)
@@ -268,8 +293,17 @@ void FIOSPlatformMisc::PlatformInit()
 					UE::FIOSPlatformMisc::LastReportedNetworkType = ConnectionType;
 					FFunctionGraphTask::CreateAndDispatchWhenReady([ConnectionType]() {
 						FCoreDelegates::OnNetworkConnectionChanged.Broadcast(ConnectionType);
-					}, TStatId(), NULL, ENamedThreads::GameThread);
+					}, TStatId(), nullptr, ENamedThreads::GameThread);
 				}
+				
+				// Adding a reference, so path survives to the task body.
+				nw_retain(InPath);
+				FFunctionGraphTask::CreateAndDispatchWhenReady(
+				// nw_path_t is actually a pointer, so captured by value.
+				[InPath] ()
+				{
+					UE::FIOSPlatformMisc::NotifyAboutNetworkConnectionCharacteristicsChange(InPath);
+				}, TStatId());
 			}
 		}
 	);
@@ -278,6 +312,11 @@ void FIOSPlatformMisc::PlatformInit()
 	nw_path_monitor_start(UE::FIOSPlatformMisc::PathMonitor);
 }
 
+TMulticastDelegateRegistration<void(FIOSPlatformMisc::FNetworkConnectionCharacteristics), FDefaultTSDelegateUserPolicy>& FIOSPlatformMisc::OnNetworkConnectionCharacteristicsChanged()
+{
+	return UE::FIOSPlatformMisc::OnNetworkConnectionCharacteristicsChanged;
+}
+
 // Defines the PlatformFeatures module name for iOS, used by PlatformFeatures.h.
 const TCHAR* FIOSPlatformMisc::GetPlatformFeaturesModuleName()
 {
diff --git a/Engine/Source/Runtime/Core/Public/IOS/IOSPlatformMisc.h b/Engine/Source/Runtime/Core/Public/IOS/IOSPlatformMisc.h
--- a/Engine/Source/Runtime/Core/Public/IOS/IOSPlatformMisc.h
+++ b/Engine/Source/Runtime/Core/Public/IOS/IOSPlatformMisc.h
@@ -12,6 +12,11 @@
 template <typename FuncType>
 class TFunction;
 
+struct FDefaultTSDelegateUserPolicy;
+
+template <typename FuncType, typename UserPolicy>
+class TMulticastDelegateRegistration;
+
 /**
 * iOS implementation of the misc OS functions
 **/
@@ -236,6 +241,25 @@ struct CORE_API FIOSPlatformMisc : public FApplePlatformMisc
 	}
 
 	static bool RequestDeviceCheckToken(TFunction<void(const TArray<uint8>&)> QuerySucceededFunc, TFunction<void(const FString&, const FString&)> QueryFailedFunc);
+	
+	// Information about the properties of the network that a connection uses.
+	struct FNetworkConnectionCharacteristics
+	{
+		// Indndicating whether the used network interface has a DNS server configured.
+		bool bSupportsDNS   : 1 = false;
+		// Indicating whether the used network interface can route IPv4 traffic.
+		bool bSupportsIPv4  : 1 = false;
+		// Indicating whether the used network interface can route IPv6 traffic.
+		bool bSupportsIPv6  : 1 = false;
+		// Indicating whether the used network interface is in Low Data Mode.
+		// It is applicable to both WiFi & cellular connections.
+		bool bIsConstrained : 1 = false;
+		// Indicating whether the used network interface is considered
+		// expensive, such as cellular or a personal hotspot.
+		bool bIsExpensive   : 1 = false;
+	};
+	// Thread-safe even that is fired from a task-thread.
+	static TMulticastDelegateRegistration<void(FNetworkConnectionCharacteristics), FDefaultTSDelegateUserPolicy>& OnNetworkConnectionCharacteristicsChanged();
     
 	FORCEINLINE static void ChooseHDRDeviceAndColorGamut(uint32 DeviceId, uint32 DisplayNitLevel, EDisplayOutputFormat& OutputDevice, EDisplayColorGamut& ColorGamut)
 	{
