Replaced USE_512MB_MAX_MEMORY_PER_BLOCK_SIZE with UE_MB3_MAX_MEMORY_PER_POOL_SIZE_MB in MB3 to allow better configuration of the allocator

diff --git a/Engine/Source/Runtime/Core/Private/HAL/MallocBinned3.cpp b/Engine/Source/Runtime/Core/Private/HAL/MallocBinned3.cpp
--- a/Engine/Source/Runtime/Core/Private/HAL/MallocBinned3.cpp
+++ b/Engine/Source/Runtime/Core/Private/HAL/MallocBinned3.cpp
@@ -149,7 +149,7 @@ struct FFreeBlock
 
 	uint16 BinSizeShifted;		// Size of the bins that this list points to >> UE_MBC_MIN_SMALL_POOL_ALIGNMENT_SHIFT
 	uint8  PoolIndex;			// Index of this pool
-	uint8  Canary;				// Constant value of 0xe3
+	uint8  Canary;				// Constant value of 0xe7
 	uint32 NumFreeBins;			// Number of consecutive free bins here, at least 1.
 	uint32 NextFreeBlockIndex;	// Next free block or MAX_uint32
 };
diff --git a/Engine/Source/Runtime/Core/Public/HAL/MallocBinned3.h b/Engine/Source/Runtime/Core/Public/HAL/MallocBinned3.h
--- a/Engine/Source/Runtime/Core/Public/HAL/MallocBinned3.h
+++ b/Engine/Source/Runtime/Core/Public/HAL/MallocBinned3.h
@@ -22,9 +22,13 @@
 #define UE_MB3_USE_CACHED_PAGE_ALLOCATOR_FOR_LARGE_ALLOCS (0)
 
 #if PLATFORM_IOS
-#	define USE_512MB_MAX_MEMORY_PER_BLOCK_SIZE		1
-#elif !defined(USE_512MB_MAX_MEMORY_PER_BLOCK_SIZE)
-#	define USE_512MB_MAX_MEMORY_PER_BLOCK_SIZE		0
+#	define UE_MB3_MAX_MEMORY_PER_POOL_SIZE_MB		512
+#endif
+
+#if defined(USE_512MB_MAX_MEMORY_PER_BLOCK_SIZE)
+#	warning "USE_512MB_MAX_MEMORY_PER_BLOCK_SIZE is deprecated. Please use UE_MB3_MAX_MEMORY_PER_POOL_SIZE_MB=%MB% instead"
+#else
+#	define UE_MB3_MAX_MEMORY_PER_POOL_SIZE_MB		1024
 #endif
 
 #define UE_MB3_BASE_PAGE_SIZE						4096			// Minimum "page size" for binned3
@@ -39,11 +43,7 @@
 #endif
 #define UE_MB3_SMALL_POOL_COUNT						(UE_MBC_NUM_LISTED_SMALL_POOLS + (UE_MB3_MAX_SMALL_POOL_SIZE - UE_MBC_MAX_LISTED_SMALL_POOL_SIZE) / UE_MB3_BASE_PAGE_SIZE)
 
-#if USE_512MB_MAX_MEMORY_PER_BLOCK_SIZE
-#	define UE_MB3_MAX_MEMORY_PER_POOL_SIZE_SHIFT	(29) // maximum of 512MB per pool
-#else
-#	define UE_MB3_MAX_MEMORY_PER_POOL_SIZE_SHIFT	(30) // maximum of 1GB per pool
-#endif
+#define UE_MB3_MAX_MEMORY_PER_POOL_SIZE_SHIFT		(FMath::CountTrailingZeros(FMath::RoundUpToPowerOfTwo(UE_MB3_MAX_MEMORY_PER_POOL_SIZE_MB * 1024 * 1024)))
 
 #define UE_MB3_MAX_MEMORY_PER_POOL_SIZE				(1ull << UE_MB3_MAX_MEMORY_PER_POOL_SIZE_SHIFT) 
 
