Fixed UE_MB3_MAX_MEMORY_PER_POOL_SIZE_MB macro redefinition error

diff --git a/Engine/Source/Runtime/Core/Public/HAL/MallocBinned3.h b/Engine/Source/Runtime/Core/Public/HAL/MallocBinned3.h
--- a/Engine/Source/Runtime/Core/Public/HAL/MallocBinned3.h
+++ b/Engine/Source/Runtime/Core/Public/HAL/MallocBinned3.h
@@ -23,7 +23,7 @@
 
 #if PLATFORM_IOS
 #	define UE_MB3_MAX_MEMORY_PER_POOL_SIZE_MB		512
-#else
+#elif !defined(UE_MB3_MAX_MEMORY_PER_POOL_SIZE_MB)
 #	define UE_MB3_MAX_MEMORY_PER_POOL_SIZE_MB		1024
 #endif
 
