Add private scope struct to disable running as a cook commandlet.
To be used when using the commandline parameter of:
-RunAsCookCommandlet
so that not all parts of the commandlet can run as a cook commandlet. This is to work around deep complexity with world partition which can simulate the generated assets while not running a cook commandlet.

diff --git a/Engine/Source/Runtime/Core/Private/Misc/CoreGlobals.cpp b/Engine/Source/Runtime/Core/Private/Misc/CoreGlobals.cpp
--- a/Engine/Source/Runtime/Core/Private/Misc/CoreGlobals.cpp
+++ b/Engine/Source/Runtime/Core/Private/Misc/CoreGlobals.cpp
@@ -1010,6 +1010,20 @@ void SetMultiprocessId(int32 MultiprocessId)
 #endif
 }
 
+FScopeDisableRunningCookCommandlet::FScopeDisableRunningCookCommandlet()
+{
+#if WITH_EDITOR
+	bPrevious = PRIVATE_GIsRunningCookCommandlet;
+	PRIVATE_GIsRunningCookCommandlet = false;
+#endif
+}
+
+FScopeDisableRunningCookCommandlet::~FScopeDisableRunningCookCommandlet()
+{
+#if WITH_EDITOR
+	PRIVATE_GIsRunningCookCommandlet = bPrevious;
+#endif
+}
 }
 
 bool(*IsGCLockedFuncPtr)() = nullptr;
diff --git a/Engine/Source/Runtime/Core/Public/CoreGlobals.h b/Engine/Source/Runtime/Core/Public/CoreGlobals.h
--- a/Engine/Source/Runtime/Core/Public/CoreGlobals.h
+++ b/Engine/Source/Runtime/Core/Public/CoreGlobals.h
@@ -339,6 +339,13 @@ CORE_API int32 GetMultiprocessId();
 namespace UE::Private
 {
 CORE_API void SetMultiprocessId(int32 MultiprocessId);
+struct FScopeDisableRunningCookCommandlet
+{
+	CORE_API FScopeDisableRunningCookCommandlet();
+	CORE_API ~FScopeDisableRunningCookCommandlet();
+private:
+	bool bPrevious = false;
+};
 }
 
 class FIsDuplicatingClassForReinstancing
