Fix Scalability in Preview Platforms with Scalability and Material Menus

Added SetWithCurrentPriorityAndTag function that replaces the value of the current highest priority and tag.
Added SETBY_ERROR(0x01000000) that is the priority we set if we try to override a ECVF_SetByConstructor with these functions and log an error.
Added to SetWithCurrentPriorityAndTag and SetWithCurrentPriority a MaxLevel optional parameter that will determine the MaxLevel of Priority the function will stomp over.
Added -replacecurrent to SetCVar command that will call SetWithCurrentPriorityAndTag.

diff --git a/Engine/Source/Editor/LevelEditor/Private/LevelEditorActions.cpp b/Engine/Source/Editor/LevelEditor/Private/LevelEditorActions.cpp
--- a/Engine/Source/Editor/LevelEditor/Private/LevelEditorActions.cpp
+++ b/Engine/Source/Editor/LevelEditor/Private/LevelEditorActions.cpp
@@ -766,7 +766,8 @@ void FLevelEditorActionCallbacks::AttachToSocketSelection(const FName SocketName
 void FLevelEditorActionCallbacks::SetMaterialQualityLevel( EMaterialQualityLevel::Type NewQualityLevel )
 {
 	static IConsoleVariable* MaterialQualityLevelVar = IConsoleManager::Get().FindConsoleVariable(TEXT("r.MaterialQualityLevel"));
-	MaterialQualityLevelVar->Set(NewQualityLevel, ECVF_SetByScalability);
+	EConsoleVariableFlags MaxPriority = GEditor && GEditor->IsFeatureLevelPreviewActive() ? ECVF_SetByPreview: ECVF_SetByScalability;
+	MaterialQualityLevelVar->ReplaceCurrentPriorityAndTag(NewQualityLevel, MaxPriority);
 
 	GUnrealEd->OnSceneMaterialsModified();
 
diff --git a/Engine/Source/Runtime/Core/Private/HAL/ConsoleManager.cpp b/Engine/Source/Runtime/Core/Private/HAL/ConsoleManager.cpp
--- a/Engine/Source/Runtime/Core/Private/HAL/ConsoleManager.cpp
+++ b/Engine/Source/Runtime/Core/Private/HAL/ConsoleManager.cpp
@@ -304,14 +304,9 @@ public:
 	}
 
 	virtual void OnChanged(EConsoleVariableFlags SetBy, bool bForce)
-	{
-		// we don't want any of this if SetOnly is used
-		if (SetBy & ECVF_Set_SetOnly_Unsafe)
-		{
-			return;
-		}
-		
+	{	
 		// SetBy can include set flags. Discard them here
+		bool bIsSetOnly = SetBy & ECVF_Set_SetOnly_Unsafe;
 		SetBy = EConsoleVariableFlags(SetBy & ~ECVF_SetFlagMask);
 
 		// you have to specify a SetBy e.g. ECVF_SetByCommandline
@@ -324,6 +319,12 @@ public:
 
 		Flags = (EConsoleVariableFlags)(((uint32)Flags & ECVF_FlagMask) | SetBy);
 
+		// we don't want to call the callback if SetOnly is used
+		if (bIsSetOnly)
+		{
+			return;
+		}
+
 		OnChangedCallback.Broadcast(this);
 	}
 
@@ -756,7 +757,34 @@ public:
 	TSortedMap<int, TArray<FTaggedHistoryData>> History;
 	
 	bool bHasTaggedArrayData = false;
-	
+
+	void FindHighestPriorityAndTag(
+		EConsoleVariableFlags MaxPriority,
+		EConsoleVariableFlags& OutFlags, 
+		FName& OutTag)
+	{
+		OutTag = NAME_None;
+		OutFlags = SETBY_ERROR;
+
+		for (typename TSortedMap<int, TArray<FTaggedHistoryData>>::TReverseIterator It(History); It; ++It)
+		{
+			if (It.Value().Num() > 0)
+			{
+				if ((EConsoleVariableFlags)It.Key() <= MaxPriority)
+				{
+					OutTag = It.Value().Last().Key;
+					OutFlags = (EConsoleVariableFlags)(It.Key());
+					// Replacing a value that is ECVF_SetByConstructor is not allowed.
+					if (OutFlags == ECVF_SetByConstructor)
+					{
+						OutFlags = SETBY_ERROR;
+					}
+					return;
+				}
+			}
+		}
+	}
+
 	/**
 	 * When we change a vlaue, we track the change at each priority level, so we can unset them later. 
 	 * @return false if this changed a value in place and it is not the final level (in array type priorities). 
@@ -796,17 +824,15 @@ public:
 		}
 		else
 		{
-			// untagged array types must only have one entry with no tag, but we push it to the end
-			if (Tag == NAME_None)
+			// Tag's in array types must be distinct.
+			// If the Tag already exists, we remove it and add it at the end.
+
+			for (int32 Index = 0; Index < ValueArray.Num(); Index++)
 			{
-				
-				for (int32 Index = 0; Index < ValueArray.Num(); Index++)
+				if (ValueArray[Index].Key == Tag)
 				{
-					if (ValueArray[Index].Key == Tag)
-					{
-						ValueArray.RemoveAt(Index);
-						break;
-					}
+					ValueArray.RemoveAt(Index);
+					break;
 				}
 			}
 
@@ -933,6 +959,53 @@ public:
 		}
 	}
 	
+	virtual FResolvedContext ResolveContext(const FSetContext& Context) override
+	{
+		FResolvedContext Resolved;
+		Resolved.Flags = Context.Flags;
+		Resolved.Tag = Context.Tag;
+
+		EConsoleVariableFlags HighestSetBy = ECVF_SetByConstructor;
+		FName HighestTag = NAME_None;
+
+		if ((Context.PriorityMode == FSetContext::EPriorityMode::ReplaceCurrent) ||
+			(Context.TagMode == FSetContext::ETagMode::ReplaceCurrent))
+		{
+			if (PriorityHistory == nullptr)
+			{
+				PriorityHistory = new FConsoleVariableHistory<T>();
+				PriorityHistory->Track(*GetString(), ECVF_SetByConstructor, NAME_None);
+			}
+
+			PriorityHistory->FindHighestPriorityAndTag(
+				Context.MaxPriority,
+				HighestSetBy,
+				HighestTag);
+		}
+		
+		if (Context.PriorityMode == FSetContext::EPriorityMode::ReplaceCurrent)
+		{
+			EConsoleVariableFlags ExtraFlags = (EConsoleVariableFlags)(Context.Flags & ~ECVF_SetByMask);
+			Resolved.Flags = (EConsoleVariableFlags)(HighestSetBy | ExtraFlags);
+			if (HighestSetBy == SETBY_ERROR)
+			{
+				UE_LOG(LogConsoleManager, Display, TEXT("Trying to Replace Cvar=%s Set By Constructor Value. This is not allowed."), *IConsoleManager::Get().FindConsoleObjectName(this));
+			}
+		}
+
+		if (Context.TagMode == FSetContext::ETagMode::ReplaceCurrent)
+		{
+			Resolved.Tag = HighestTag;
+		}
+		
+		if ((Context.TagMode == FSetContext::ETagMode::Replace))
+		{
+			Resolved.Flags = (EConsoleVariableFlags)(Resolved.Flags | ECVF_Set_ReplaceExistingTag);
+		}
+
+		return Resolved;
+	}
+
 	/**
 	 * Subclasses implement this to simply set TypedValue as the current value for the CVar. It must not call any callbacks
 	 */
@@ -1255,18 +1328,18 @@ public:
 		delete this; 
 	}
 	
-	virtual void Set(const TCHAR* InValue, EConsoleVariableFlags SetBy, FName Tag)
+	virtual void Set(const TCHAR* Value, const FConsoleVariableBase::FResolvedContext& Context) override
 	{
 		// if it was already set in the history (this returns false if nothing needs to be done), then skip the Set
-		if (this->TrackHistory(InValue, SetBy, Tag) == false)
+		if (this->TrackHistory(Value, Context.Flags, Context.Tag) == false)
 		{
 			return;
 		}
 
-		if (CanChange(SetBy))
+		if (CanChange(Context.Flags))
 		{
-			UE::ConsoleManager::Private::GetValueFromString<T>(Data.ShadowedValue[0], InValue);
-			OnChanged(SetBy, false);
+			UE::ConsoleManager::Private::GetValueFromString<T>(Data.ShadowedValue[0], Value);
+			OnChanged(Context.Flags, false);
 		}
 	}
 
@@ -1362,12 +1435,12 @@ template<> TConsoleVariableData<float>* FConsoleVariable<float>::AsVariableFloat
 
 // specialization for FString
 
-template<> void FConsoleVariable<FString>::Set(const TCHAR* InValue, EConsoleVariableFlags SetBy, FName Tag)
+template<> void FConsoleVariable<FString>::Set(const TCHAR* Value, const FConsoleVariableBase::FResolvedContext& Context)
 {
-	if (PreprocessSet(InValue, SetBy, Tag))
+	if (PreprocessSet(Value, Context.Flags, Context.Tag))
 	{
-		Data.ShadowedValue[0] = InValue;
-		OnChanged(SetBy);
+		Data.ShadowedValue[0] = Value;
+		OnChanged(Context.Flags);
 	}
 }
 
@@ -1539,17 +1612,17 @@ public:
 		delete this; 
 	}
 
-	virtual void Set(const TCHAR* InValue, EConsoleVariableFlags SetBy, FName Tag)
+	virtual void Set(const TCHAR* Value, const FConsoleVariableBase::FResolvedContext& Context) override
 	{
-		if (this->TrackHistory(InValue, SetBy, Tag) == false)
+		if (this->TrackHistory(Value, Context.Flags, Context.Tag) == false)
 		{
 			return;
 		}
 		
-		if(CanChange(SetBy))
+		if (CanChange(Context.Flags))
 		{
-			UE::ConsoleManager::Private::GetValueFromString<T>(MainValue, InValue);
-			OnChanged(SetBy);
+			UE::ConsoleManager::Private::GetValueFromString<T>(MainValue, Value);
+			OnChanged(Context.Flags);
 		}
 	}
 
@@ -1627,12 +1700,12 @@ public:
 	{
 		delete this;
 	}
-	virtual void Set(const TCHAR* InValue, EConsoleVariableFlags SetBy, FName Tag)
+	virtual void Set(const TCHAR* Value, const FConsoleVariableBase::FResolvedContext& Context) override
 	{
-		if (PreprocessSet(InValue, SetBy, Tag))
+		if (PreprocessSet(Value, Context.Flags, Context.Tag))
 		{
-			MainValue = InValue;
-			OnChanged(SetBy);
+			MainValue = Value;
+			OnChanged(Context.Flags);
 		}
 	}
 	virtual bool GetBool() const
@@ -1707,14 +1780,14 @@ public:
 	{
 		delete this;
 	}
-	virtual void Set(const TCHAR* InValue, EConsoleVariableFlags SetBy, FName Tag)
+	virtual void Set(const TCHAR* Value, const FConsoleVariableBase::FResolvedContext& Context) override
 	{
-		TrackHistory(InValue, SetBy, Tag);
-		
-		if (CanChange(SetBy))
+		TrackHistory(Value, Context.Flags, Context.Tag);
+
+		if (CanChange(Context.Flags))
 		{
-			MainValue = FName(InValue);
-			OnChanged(SetBy);
+			MainValue = FName(Value);
+			OnChanged(Context.Flags);
 		}
 	}
 	virtual bool GetBool() const
@@ -1783,9 +1856,9 @@ public:
 	{
 		delete this; 
 	} 
-	virtual void Set(const TCHAR* InValue, EConsoleVariableFlags SetBy, FName Tag)
+	virtual void Set(const TCHAR* InValue, const FConsoleVariableBase::FResolvedContext& Context) override
 	{
-		if (PreprocessSet(InValue, SetBy, Tag))
+		if (PreprocessSet(InValue, Context.Flags, Context.Tag))
 		{
 			int32 Value = FCString::Atoi(InValue);
 
@@ -1794,7 +1867,7 @@ public:
 			FMath::SetBoolInBitField(Force0MaskPtr, BitNumber, Value == 0);
 			FMath::SetBoolInBitField(Force1MaskPtr, BitNumber, Value == 1);
 
-			OnChanged(SetBy, false);
+			OnChanged(Context.Flags, false);
 		}
 	}
 	virtual bool GetBool() const
@@ -2706,7 +2779,7 @@ static void SetUnsetCVar(const TMap<FString, IConsoleObject*>& ConsoleObjects, c
 	{
 		if (bSet)
 		{
-			Ar.Logf(TEXT("Usage: SetCVar [Platform@]CVarName Value [-setby=Priority] [-tag=SomeTag]"));
+			Ar.Logf(TEXT("Usage: SetCVar [Platform@]CVarName Value [-setby=Priority] [-tag=SomeTag] [-replacecurrent]"));
 		}
 		else
 		{
@@ -2764,26 +2837,33 @@ static void SetUnsetCVar(const TMap<FString, IConsoleObject*>& ConsoleObjects, c
 	{
 		Value = FParse::Token(Params, false);
 	}
-	
-	EConsoleVariableFlags SetBy = ECVF_SetByConsole;
-	FName Tag = NAME_None;
-	FString Str;
-	if (FParse::Value(Params, TEXT("-setby="), Str))
-	{
-		SetBy = GetConsoleVariableSetByValue(*Str);
-	}
-	if (FParse::Value(Params, TEXT("-tag="), Str))
-	{
-		Tag = *Str;
-	}
 
-	if (bSet)
+	if (bSet && FParse::Param(Params, TEXT("-replacecurrent")))
 	{
-		CVar->Set(*Value, SetBy, Tag);
+		CVar->ReplaceCurrentPriorityAndTag(*Value);
 	}
 	else
 	{
-		CVar->Unset(SetBy, Tag);
+		EConsoleVariableFlags SetBy = ECVF_SetByConsole;
+		FName Tag = NAME_None;
+		FString Str;
+		if (FParse::Value(Params, TEXT("-setby="), Str))
+		{
+			SetBy = GetConsoleVariableSetByValue(*Str);
+		}
+		if (FParse::Value(Params, TEXT("-tag="), Str))
+		{
+			Tag = *Str;
+		}
+
+		if (bSet)
+		{
+			CVar->Set(*Value, SetBy, Tag);
+		}
+		else
+		{
+			CVar->Unset(SetBy, Tag);
+		}
 	}
 #else
 	Ar.Logf(TEXT("Unable to lookup a CVar value on another platform in this build"));
@@ -4824,12 +4904,23 @@ public:
 		}
 	}
 
+	virtual FConsoleVariableBase::FResolvedContext ResolveContext(const FConsoleVariableBase::FSetContext& Context) override
+	{
+		if (BindForUsage())
+		{
+			return RealVariable->ResolveContext(Context);
+		}
+		FResolvedContext Resolved;
+		Resolved.Flags = Context.Flags;
+		Resolved.Tag = Context.Tag;
+		return Resolved;
+	}
 
-	virtual void Set(const TCHAR* InValue, EConsoleVariableFlags SetBy, FName Tag) override
+	virtual void Set(const TCHAR* Value, const FConsoleVariableBase::FResolvedContext& Context) override
 	{
 		if (BindForUsage())
 		{
-			RealVariable->Set(InValue, SetBy, Tag);
+			RealVariable->Set(Value, Context);
 		}
 	}
 
diff --git a/Engine/Source/Runtime/Core/Public/HAL/IConsoleManager.h b/Engine/Source/Runtime/Core/Public/HAL/IConsoleManager.h
--- a/Engine/Source/Runtime/Core/Public/HAL/IConsoleManager.h
+++ b/Engine/Source/Runtime/Core/Public/HAL/IConsoleManager.h
@@ -22,6 +22,8 @@
 	#define cvarCheckCode(...)
 #endif
 
+#define SETBY_ERROR ((EConsoleVariableFlags)0x01000000)
+
 template <class T> class TConsoleVariableData;
 
 /**
@@ -126,7 +128,7 @@ enum EConsoleVariableFlags
 	// Use to set a cvar without calling all cvar sinks. Much faster, but potentially unsafe. Use only if you know the particular cvar/setting does not require a sink call
 	ECVF_Set_NoSinkCall_Unsafe =	0x00010000,
 	
-	// Similar to ECVF_Set_NoSinkCall_Unsafe, but this will do nothing but set the cvar, and push updates to render thread, no changing flags, no sinks, etc
+	// Similar to ECVF_Set_NoSinkCall_Unsafe, but this will do nothing but set the cvar and flags, and push updates to render thread, no sinks, etc
 	ECVF_Set_SetOnly_Unsafe = 		0x00020000,
 
 	// When this is set, this will stomp on existing tagged array type values in place, instead of adding to the end of the array
@@ -141,34 +143,36 @@ enum EConsoleVariableFlags
 
 	// Lowest priority (default after console variable creation)
 	ECVF_SetByConstructor =			0x00000000,
+	// Reserved value for SETBY_ERROR. Please don't use.
+	// SETBY_ERROR =				0x01000000 
 	// Set by scalability groups from Scalability.ini (lower priority than game settings so it's easier to override partially). THIS IS AN ARRAY TYPE
-	ECVF_SetByScalability =			0x01000000,
+	ECVF_SetByScalability =			0x02000000,
 	// Default priority for engine-level game user settings, platform-specific settings will override this
-	ECVF_SetByGameSetting =			0x02000000,
+	ECVF_SetByGameSetting =			0x03000000,
 	// Set by project settings UI or specific sections in ini file (higher priority than game setting to allow enforcing some settings for this project)
-	ECVF_SetByProjectSetting =		0x03000000,
+	ECVF_SetByProjectSetting =		0x04000000,
 	// Used by the [ConsoleVariables] section of Engine.ini as well as FSystemSettings
-	ECVF_SetBySystemSettingsIni =	0x04000000,
+	ECVF_SetBySystemSettingsIni =	0x05000000,
 	// Dyanmically loaded/unloaded plugins, used with the History concept to restore cvars on plugin unload. THIS IS AN ARRAY TYPE meaning multiple vales with this SetBy are stored in the history.
-	ECVF_SetByPluginLowPriority =	0x05000000,
+	ECVF_SetByPluginLowPriority =	0x06000000,
 	// Per device settings using the DeviceProfiles.ini hierarchy (e.g. specific iOS device, higher priority than per project to do device specific settings). THIS IS AN ARRAY TYPE
-	ECVF_SetByDeviceProfile =		0x06000000,
+	ECVF_SetByDeviceProfile =		0x07000000,
 	// Dyanmically loaded/unloaded plugins, used with the History concept to restore cvars on plugin unload. THIS IS AN ARRAY TYPE meaning multiple vales with this SetBy are stored in the history.
-	ECVF_SetByPluginHighPriority =	0x07000000,
+	ECVF_SetByPluginHighPriority =	0x08000000,
 	// User settable game overrides, used for GameUserSettings fields that need to override device specific settings
-	ECVF_SetByGameOverride =		0x08000000,
+	ECVF_SetByGameOverride =		0x09000000,
 	// Set by local consolevariables.ini, mostly used for testing multiple projects
-	ECVF_SetByConsoleVariablesIni = 0x09000000,
+	ECVF_SetByConsoleVariablesIni = 0x0A000000,
 	// Set by hotfix. THIS IS AN ARRAY TYPE meaning multiple vales with this SetBy are stored in the history
-	ECVF_SetByHotfix =				0x0A000000,
+	ECVF_SetByHotfix =				0x0B000000,
 	// Set for previewing in editor. THIS IS AN ARRAY TYPE meaning multiple vales with this SetBy are stored in the history
-	ECVF_SetByPreview =				0x0B000000,
+	ECVF_SetByPreview =				0x0C000000,
 	// Used by some command line parameters, others use the Console priority instead
-	ECVF_SetByCommandline =			0x0C000000,
+	ECVF_SetByCommandline =			0x0D000000,
 	// Used for high priority temporary debugging or operation modes
-	ECVF_SetByCode =				0x0D000000,
+	ECVF_SetByCode =				0x0E000000,
 	// Highest priority used via editor UI or or game/editor interactive console
-	ECVF_SetByConsole =				0x0E000000,
+	ECVF_SetByConsole =				0x0F000000,
 
 
 	// ------------------------------------------------
@@ -554,12 +558,60 @@ class IConsoleVariable : public IConsoleObject
 {
 public:
 
+	struct FResolvedContext
+	{
+		EConsoleVariableFlags Flags;
+		FName Tag;
+	};
+
+	struct FSetContext
+	{
+		FSetContext()
+		{
+		}
+
+		FSetContext(EConsoleVariableFlags InFlags, FName InTag)
+			: Flags(InFlags)
+			, Tag(InTag)
+		{
+		}
+
+
+		enum class EPriorityMode : uint8
+		{
+			Set,
+			// the following are limited by Max priority if its set
+			ReplaceCurrent,
+		};
+
+		enum class ETagMode : uint8
+		{
+			// Sets the Tag at the end of the array. Remove the Tag if it already exists.
+			Set,
+			// Replace the Tag in-place if it already exists. Sets the Tag at the end if it doesn't exist.
+			Replace,
+			// Replaces the current highest Tag.
+			ReplaceCurrent,
+		};
+
+
+		EConsoleVariableFlags Flags = ECVF_SetByCode;
+		EPriorityMode PriorityMode = EPriorityMode::Set;
+
+		FName Tag = NAME_None;
+		ETagMode TagMode = ETagMode::Set;
+
+		EConsoleVariableFlags MaxPriority = ECVF_SetByConsole;
+	};
+
+	virtual FResolvedContext ResolveContext(const FSetContext& Context) = 0;
+
 	/**
 	 * Set the internal value from the specified string. 
-	 * @param SetBy anything in ECVF_LastSetMask e.g. ECVF_SetByScalability
-	 * @param Tag optional Tag to set with the value - only useful when UE_ALLOW_CVAR_HISTORY is set, and when setting in an ARRAY type SetBy
+	 * @param InValue Internal Value to Set
+	 * @param Context The ResolvedContext of the FSetContext
 	 **/
-	virtual void Set(const TCHAR* InValue, EConsoleVariableFlags SetBy = ECVF_SetByCode, FName Tag = NAME_None) = 0;
+	virtual void Set(const TCHAR* InValue, const FResolvedContext& Context) = 0;
 
 	/**
 	 * Unsets the value at a certain SetBy priority (this is only useful when UE_ALLOW_CVAR_HISTORY is set). The value of the CVar
@@ -664,47 +716,69 @@ public:
 
 	// convenience methods
 
+	/** Set the internal value from the SetContext. */
+	void Set(const TCHAR* InValue, const FSetContext& Context)
+	{
+		Set(InValue, this->ResolveContext(Context));
+	}
+
 	/** Set the internal value from the specified bool. */
-	void Set(bool InValue, EConsoleVariableFlags SetBy = ECVF_SetByCode, FName Tag=NAME_None)
+	void Set(bool InValue, const FSetContext& Context)
 	{
 		// NOTE: Bool needs to use 1 and 0 here rather than true/false, as this may be a int32 or something
 		// and eventually this code calls, TTypeFromString<T>::FromString which won't handle the true/false,
 		// but 1 and 0 will work for whatever.
 		// inefficient but no common code path
-		Set(InValue ? TEXT("1") : TEXT("0"), SetBy, Tag);
+		Set(InValue ? TEXT("1") : TEXT("0"), Context);
 	}
 	/** Set the internal value from the specified int. */
-	void Set(int32 InValue, EConsoleVariableFlags SetBy = ECVF_SetByCode, FName Tag=NAME_None)
+	void Set(int32 InValue, const FSetContext& Context)
 	{
 		// inefficient but no common code path
-		Set(*FString::Printf(TEXT("%d"), InValue), SetBy, Tag);
+		Set(*FString::Printf(TEXT("%d"), InValue), Context);
 	}
 	/** Set the internal value from the specified float. */
-	void Set(float InValue, EConsoleVariableFlags SetBy = ECVF_SetByCode, FName Tag=NAME_None)
+	void Set(float InValue, const FSetContext& Context)
 	{
 		// inefficient but no common code path
-		Set(*FString::Printf(TEXT("%g"), InValue), SetBy, Tag);
+		Set(*FString::Printf(TEXT("%g"), InValue), Context);
 	}
 
-	void SetWithCurrentPriority(bool InValue, FName Tag=NAME_None)
+
+	template <typename T>
+	void Set(T Value, EConsoleVariableFlags Flags = ECVF_SetByCode, FName Tag = NAME_None)
 	{
-		EConsoleVariableFlags CurFlags = (EConsoleVariableFlags)(GetFlags() & ECVF_SetByMask);
-		Set(InValue, CurFlags, Tag);
+		Set(Value, FSetContext(Flags, Tag));
 	}
-	void SetWithCurrentPriority(int32 InValue, FName Tag=NAME_None)
-	{
-		EConsoleVariableFlags CurFlags = (EConsoleVariableFlags)(GetFlags() & ECVF_SetByMask);
-		Set(InValue, CurFlags, Tag);
-	}
-	void SetWithCurrentPriority(float InValue, FName Tag=NAME_None)
+
+	template <typename T>
+	void SetWithCurrentPriority(
+		T Value, 
+		FName Tag = NAME_None, 
+		EConsoleVariableFlags MaxPriority = ECVF_SetByConsole)
 	{
+		FSetContext Context;
+		Context.Tag = Tag;
+		Context.PriorityMode = FSetContext::EPriorityMode::ReplaceCurrent;
+		Context.MaxPriority = MaxPriority;
 		EConsoleVariableFlags CurFlags = (EConsoleVariableFlags)(GetFlags() & ECVF_SetByMask);
-		Set(InValue, CurFlags, Tag);
+		Context.Flags = CurFlags;
+		Set(Value, Context);
 	}
-	void SetWithCurrentPriority(const TCHAR* InValue, FName Tag=NAME_None)
+
+	template<typename T>
+	void ReplaceCurrentPriorityAndTag(
+		T InValue, 
+		EConsoleVariableFlags MaxPriority = ECVF_SetByConsole
+		)
 	{
+		FSetContext Context;
+		Context.TagMode = FSetContext::ETagMode::ReplaceCurrent;
+		Context.PriorityMode = FSetContext::EPriorityMode::ReplaceCurrent;
+		Context.MaxPriority = MaxPriority;
 		EConsoleVariableFlags CurFlags = (EConsoleVariableFlags)(GetFlags() & ECVF_SetByMask);
-		Set(InValue, CurFlags, Tag);
+		Context.Flags = CurFlags;
+		Set(InValue, Context);
 	}
 
 
diff --git a/Engine/Source/Runtime/Engine/Private/Scalability.cpp b/Engine/Source/Runtime/Engine/Private/Scalability.cpp
--- a/Engine/Source/Runtime/Engine/Private/Scalability.cpp
+++ b/Engine/Source/Runtime/Engine/Private/Scalability.cpp
@@ -12,6 +12,7 @@
 #include "Interfaces/IProjectManager.h"
 
 #if WITH_EDITOR
+#include "Editor.h"
 #include "DeviceProfiles/DeviceProfile.h"
 #include "DeviceProfiles/DeviceProfileManager.h"
 #endif
@@ -429,7 +430,8 @@ static void SetGroupQualityLevel(const TCHAR* InGroupName, int32 InQualityLevel,
 					{
 						if (IConsoleVariable* CVar = IConsoleManager::Get().FindConsoleVariable(*CVarName))
 						{
-							CVar->SetWithCurrentPriority(*CVarValue);
+							EConsoleVariableFlags MaxPriority = ECVF_SetByPreview;
+							CVar->ReplaceCurrentPriorityAndTag(*CVarValue, MaxPriority);
 						}
 					}
 				}
@@ -486,7 +488,7 @@ FQualityLevels GetDefaultQualityLevels()
 
 void ResetQualityLevelsToDefault()
 {
-	SetQualityLevels(GetDefaultQualityLevels(), false);
+	SetQualityLevels(GetDefaultQualityLevels());
 }
 #endif
 
@@ -639,7 +641,7 @@ void InitScalabilitySystem()
 	bScalabilityShaderPlatformHasBeenChanged = false;
 #endif
 	// Set defaults
-	SetQualityLevels(FQualityLevels(), true);
+	SetQualityLevels(FQualityLevels());
 	GScalabilityBackupQualityLevels = FQualityLevels();
 	GScalabilityQualityLevelsOverride = FQualityLevels();
 	GScalabilityUsingTemporaryQualityLevels = false;
@@ -891,6 +893,13 @@ template<typename T> void SetQualityLevelCVar(TAutoConsoleVariable<T>& CVarToEdi
 	{
 		CVarToEdit.AsVariable()->SetWithCurrentPriority(DesiredValue);
 	}
+#if WITH_EDITOR
+	else if (GEditor && GEditor->IsFeatureLevelPreviewActive())
+	{
+		EConsoleVariableFlags MaxPriority = ECVF_SetByPreview;
+		CVarToEdit.AsVariable()->ReplaceCurrentPriorityAndTag(DesiredValue, MaxPriority);
+	}
+#endif
 	else
 	{
 		T PreviousValue;
@@ -919,10 +928,6 @@ void SetQualityLevels(const FQualityLevels& QualityLevels, bool bForce/* = false
 	ClampedLevels.SetShadingQuality(QualityLevels.ShadingQuality);
 	ClampedLevels.SetLandscapeQuality(QualityLevels.LandscapeQuality);
 
-#if WITH_EDITOR
-	bForce = bForce || bScalabilityShaderPlatformHasBeenChanged;
-#endif
-
 	if (GScalabilityUsingTemporaryQualityLevels && !bForce)
 	{
 		// When temporary scalability is active, non-temporary sets are
