Optimization to avoid reading thread heartbeat settings from config every frame.

Saves 35us per frame on Switch.

#tests Win64 Test Client, replayrun

diff --git a/Engine/Plugins/Online/OnlineFramework/Source/Hotfix/Private/OnlineHotfixManager.cpp b/Engine/Plugins/Online/OnlineFramework/Source/Hotfix/Private/OnlineHotfixManager.cpp
--- a/Engine/Plugins/Online/OnlineFramework/Source/Hotfix/Private/OnlineHotfixManager.cpp
+++ b/Engine/Plugins/Online/OnlineFramework/Source/Hotfix/Private/OnlineHotfixManager.cpp
@@ -7,6 +7,7 @@
 #include "UObject/Package.h"
 
 #include "HAL/ConsoleManager.h"
+#include "HAL/ThreadHeartBeat.h"
 
 #include "Misc/PackageName.h"
 #include "Misc/EngineVersion.h"
@@ -852,6 +853,9 @@ void UOnlineHotfixManager::TriggerHotfixComplete(EHotfixResult HotfixResult)
 		PatchAssetsFromIniFiles();
 	}
 
+	FThreadHeartBeat::Get().InitSettings();
+	FGameThreadHitchHeartBeat::Get().InitSettings();
+
 	TriggerOnHotfixCompleteDelegates(HotfixResult);
 	if (HotfixResult == EHotfixResult::Failed)
 	{
diff --git a/Engine/Source/Runtime/Core/Private/HAL/ThreadHeartBeat.cpp b/Engine/Source/Runtime/Core/Private/HAL/ThreadHeartBeat.cpp
--- a/Engine/Source/Runtime/Core/Private/HAL/ThreadHeartBeat.cpp
+++ b/Engine/Source/Runtime/Core/Private/HAL/ThreadHeartBeat.cpp
@@ -85,7 +85,7 @@ FThreadHeartBeat::FThreadHeartBeat()
 	// platforms that implement frame-present based detection on the first call the PresentFrame().
 	PresentHeartBeat.SuspendedCount = 1;
 
-	InitSettings();
+	InitSettingsInternal();
 
 	const bool bAllowThreadHeartBeat = FPlatformMisc::AllowThreadHeartBeat() && (ConfigHangDuration > 0.0 || ConfigPresentDuration > 0.0);
 
@@ -349,6 +349,16 @@ void FThreadHeartBeat::Start()
 }
 
 void FThreadHeartBeat::InitSettings()
+{
+	const uint32 ThreadId = FPlatformTLS::GetCurrentThreadId();
+	if (ThreadId == GGameThreadId && GConfig)
+	{
+		FScopeLock HeartBeatLock(&HeartBeatCritical);
+		InitSettingsInternal();
+	}
+}
+
+void FThreadHeartBeat::InitSettingsInternal()
 {
 	double NewStuckDuration = 1.0;
 
@@ -429,8 +439,27 @@ void FThreadHeartBeat::HeartBeat(bool bReadConfig)
 	FScopeLock HeartBeatLock(&HeartBeatCritical);
 	if (bReadConfig && ThreadId == GGameThreadId && GConfig)
 	{
-		InitSettings();
+		InitSettingsInternal();
+	}
+	FHeartBeatInfo& HeartBeatInfo = ThreadHeartBeat.FindOrAdd(ThreadId);
+	HeartBeatInfo.LastHeartBeatTime = Clock.Seconds();
+	HeartBeatInfo.HangDuration = CurrentHangDuration;
+	HeartBeatInfo.StuckDuration = CurrentStuckDuration;
+#endif
+}
+
+void FThreadHeartBeat::HeartBeat()
+{
+#if USE_HANG_DETECTION
+	// disable on platforms that don't start the thread
+	if (FPlatformMisc::AllowThreadHeartBeat() == false)
+	{
+		return;
 	}
+
+	uint32 ThreadId = FPlatformTLS::GetCurrentThreadId();
+	FScopeLock HeartBeatLock(&HeartBeatCritical);
+
 	FHeartBeatInfo& HeartBeatInfo = ThreadHeartBeat.FindOrAdd(ThreadId);
 	HeartBeatInfo.LastHeartBeatTime = Clock.Seconds();
 	HeartBeatInfo.HangDuration = CurrentHangDuration;
@@ -860,7 +889,7 @@ void FThreadHeartBeat::SetDurationMultiplier(double NewMultiplier)
 	FScopeLock HeartBeatLock(&HeartBeatCritical);
 
 	HangDurationMultiplier = NewMultiplier;
-	InitSettings();
+	InitSettingsInternal();
 
 	UE_LOG(LogCore, Display, TEXT("Setting hang detector multiplier to %.4fs. New hang duration: %.4fs. New present duration: %.4fs."), NewMultiplier, CurrentHangDuration, CurrentPresentDuration);
 #endif
@@ -878,7 +907,7 @@ FGameThreadHitchHeartBeatThreaded::FGameThreadHitchHeartBeatThreaded()
 {
 	// We don't care about programs for now so no point in spawning the extra thread
 #if USE_HITCH_DETECTION
-	InitSettings();
+	InitSettingsInternal();
 #endif
 }
 
@@ -935,6 +964,14 @@ bool FGameThreadHitchHeartBeatThreaded::Init()
 
 void FGameThreadHitchHeartBeatThreaded::InitSettings()
 {
+#if USE_HANG_DETECTION
+	FScopeLock HeartBeatLock(&HeartBeatCritical);
+	InitSettingsInternal();
+#endif
+}
+
+void FGameThreadHitchHeartBeatThreaded::InitSettingsInternal()
+{
 #if USE_HITCH_DETECTION
 	static bool bFirst = true;
 	static bool bHasCmdLine = false;
@@ -1126,7 +1163,7 @@ void FGameThreadHitchHeartBeatThreaded::FrameStart(bool bSkipThisFrame)
 	// Grab this everytime to handle hotfixes
 	if (!bSkipThisFrame)
 	{
-		InitSettings();
+		InitSettingsInternal();
 	}
 	double Now = Clock.Seconds();
 	if (FirstStartTime == 0.0)
diff --git a/Engine/Source/Runtime/Core/Public/HAL/ThreadHeartBeat.h b/Engine/Source/Runtime/Core/Public/HAL/ThreadHeartBeat.h
--- a/Engine/Source/Runtime/Core/Public/HAL/ThreadHeartBeat.h
+++ b/Engine/Source/Runtime/Core/Public/HAL/ThreadHeartBeat.h
@@ -148,13 +148,13 @@ class FThreadHeartBeat : public FRunnable
 	CORE_API FThreadHeartBeat();
 	CORE_API virtual ~FThreadHeartBeat();
 
-	CORE_API void InitSettings();
-
 	CORE_API void FORCENOINLINE OnHang(double HangDuration, uint32 ThreadThatHung);
 	CORE_API void FORCENOINLINE OnPresentHang(double HangDuration);
 
 	CORE_API bool IsEnabled();
 
+	void InitSettingsInternal();
+
 public:
 
 	enum EConstants
@@ -172,8 +172,12 @@ public:
 
 	/** Begin measuring heartbeat */
 	CORE_API void Start();
+	/** Update settings from command line or ini */
+	CORE_API void InitSettings();
 	/** Called from a thread once per frame to update the heartbeat time */
-	CORE_API void HeartBeat(bool bReadConfig = false);
+	UE_DEPRECATED(5.7, "HeartBeat(bool bReadConfig) is deprecated, please use HeartBeat() instead.")
+	CORE_API void HeartBeat(bool bReadConfig);
+	CORE_API void HeartBeat();
 	/** Called from the rendering or RHI thread when the platform RHI presents a frame (supported platforms only). */
 	CORE_API void PresentFrame();
 	/** Called by a supervising thread to check the threads' health */
@@ -355,11 +359,11 @@ class FGameThreadHitchHeartBeatThreaded : public FRunnable
 
 	FThreadHeartBeatClock Clock;
 
-	CORE_API void InitSettings();
-
 	CORE_API FGameThreadHitchHeartBeatThreaded();
 	CORE_API virtual ~FGameThreadHitchHeartBeatThreaded();
 
+	void InitSettingsInternal();
+
 public:
 
 	enum EConstants
@@ -372,6 +376,9 @@ public:
 	static CORE_API FGameThreadHitchHeartBeatThreaded& Get();
 	static CORE_API FGameThreadHitchHeartBeatThreaded* GetNoInit();
 
+	/** Update settings from command line or ini */
+	CORE_API void InitSettings();
+
 	/**
 	* Called at the start of a frame to register the time we are looking to detect a hitch
 	*/
diff --git a/Engine/Source/Runtime/Core/Public/Unix/UnixSignalHeartBeat.h b/Engine/Source/Runtime/Core/Public/Unix/UnixSignalHeartBeat.h
--- a/Engine/Source/Runtime/Core/Public/Unix/UnixSignalHeartBeat.h
+++ b/Engine/Source/Runtime/Core/Public/Unix/UnixSignalHeartBeat.h
@@ -43,12 +43,13 @@ public:
 
 	CORE_API void PostFork();
 
+	CORE_API void InitSettings();
+
 private:
 	FUnixSignalGameHitchHeartBeat();
 	~FUnixSignalGameHitchHeartBeat();
 
 	void Init();
-	void InitSettings();
 
     double HitchThresholdS = -1.0;
 	double StartTime = 0.0;
diff --git a/Engine/Source/Runtime/Launch/Private/LaunchEngineLoop.cpp b/Engine/Source/Runtime/Launch/Private/LaunchEngineLoop.cpp
--- a/Engine/Source/Runtime/Launch/Private/LaunchEngineLoop.cpp
+++ b/Engine/Source/Runtime/Launch/Private/LaunchEngineLoop.cpp
@@ -4850,6 +4850,9 @@ int32 FEngineLoop::Init()
 
 	FDelayedAutoRegisterHelper::RunAndClearDelayedAutoRegisterDelegates(EDelayedRegisterRunPhase::EndOfEngineInit);
 
+	FThreadHeartBeat::Get().InitSettings();
+	FGameThreadHitchHeartBeat::Get().InitSettings();
+
 	// Emit logging. Don't edit! Automation looks for this to detect failures during initialization.
 	UE_LOG(LogInit, Display, TEXT("Engine is initialized. Leaving FEngineLoop::Init()"));
 	return 0;
@@ -5506,7 +5509,7 @@ void FEngineLoop::Tick()
 	{
 		SCOPED_NAMED_EVENT(FEngineLoop_HeartBeat, SystemColor);
 		// Send a heartbeat for the diagnostics thread
-		FThreadHeartBeat::Get().HeartBeat(true);
+		FThreadHeartBeat::Get().HeartBeat();
 
 		FGameThreadHitchHeartBeat::Get().FrameStart();
 	}
