In FGameThreadHitchHeartBeatThreaded::InitSettings fixed using USE_HANG_DETECTION instead of USE_HITCH_DETECTION.

Error introduced in CL 45197514.

diff --git a/Engine/Source/Runtime/Core/Private/HAL/ThreadHeartBeat.cpp b/Engine/Source/Runtime/Core/Private/HAL/ThreadHeartBeat.cpp
--- a/Engine/Source/Runtime/Core/Private/HAL/ThreadHeartBeat.cpp
+++ b/Engine/Source/Runtime/Core/Private/HAL/ThreadHeartBeat.cpp
@@ -964,7 +964,7 @@ bool FGameThreadHitchHeartBeatThreaded::Init()
 
 void FGameThreadHitchHeartBeatThreaded::InitSettings()
 {
-#if USE_HANG_DETECTION
+#if USE_HITCH_DETECTION
 	FScopeLock HeartBeatLock(&HeartBeatCritical);
 	InitSettingsInternal();
 #endif
