Mac - Log output screen information when GetGPUDriverInfo is called with bVerbose == true.

#tests local run & preflight

diff --git a/Engine/Source/Runtime/Core/Private/Mac/MacPlatformMisc.cpp b/Engine/Source/Runtime/Core/Private/Mac/MacPlatformMisc.cpp
--- a/Engine/Source/Runtime/Core/Private/Mac/MacPlatformMisc.cpp
+++ b/Engine/Source/Runtime/Core/Private/Mac/MacPlatformMisc.cpp
@@ -33,6 +33,7 @@
 #include "PLCrashReporter.h"
 #include "Apple/PreAppleSystemHeaders.h"
 #include <dlfcn.h>
+#include <AppKit/AppKit.h>
 #include <IOKit/IOKitLib.h>
 #include <IOKit/kext/KextManager.h>
 #include <IOKit/network/IOEthernetInterface.h>
@@ -1455,9 +1456,35 @@ FString FMacPlatformMisc::GetPrimaryGPUBrand()
 	return PrimaryGPU;
 	}
 
-FGPUDriverInfo FMacPlatformMisc::GetGPUDriverInfo(const FString& DeviceDescription)
+FGPUDriverInfo FMacPlatformMisc::GetGPUDriverInfo(const FString& DeviceDescription, bool bVerbose)
 {
 	SCOPED_AUTORELEASE_POOL;
+	
+	if (bVerbose)
+	{
+		// Also report monitor information here, for lack of a better place.
+		UE_LOG(LogMac, Log, TEXT("Attached monitors:"));
+		int NumMonitors = 0;
+		
+		for (NSScreen* Screen in NSScreen.screens)
+		{
+			NSDictionary* ScreenDescription = Screen.deviceDescription;
+			CGSize PointSize = Screen.frame.size;
+			NSSize PixelSize = [[ScreenDescription objectForKey:NSDeviceSize] sizeValue];
+			NSRect WorkArea = Screen.visibleFrame;
+			NSNumber* Scale = ScreenDescription[NSDeviceResolution];
+			
+			UE_LOG(LogMac, Log, TEXT("Screen: %s, Resolution: %.0fx%.0f, Points: %.0fx%.0f, WorkArea: (%.0f, %.0f) -> (%.0f, %.0f), Scale: %.2f")
+				   , Screen.localizedName ? *FString((__bridge CFStringRef)Screen.localizedName) : TEXT("(Unnamed)")
+				   , PixelSize.width, PixelSize.height
+				   , PointSize.width, PointSize.height
+				   , WorkArea.origin.x, WorkArea.origin.y, WorkArea.origin.x + WorkArea.size.width, WorkArea.origin.y + WorkArea.size.height
+				   , Screen.backingScaleFactor);
+			NumMonitors++;
+		}
+		
+		UE_LOG(LogMac, Log, TEXT("Found %d attached monitors."), NumMonitors);
+	}
 
 	FGPUDriverInfo Info;
 	TArray<FString> NameComponents;
diff --git a/Engine/Source/Runtime/Core/Public/Mac/MacPlatformMisc.h b/Engine/Source/Runtime/Core/Public/Mac/MacPlatformMisc.h
--- a/Engine/Source/Runtime/Core/Public/Mac/MacPlatformMisc.h
+++ b/Engine/Source/Runtime/Core/Public/Mac/MacPlatformMisc.h
@@ -39,7 +39,7 @@ struct CORE_API FMacPlatformMisc : public FApplePlatformMisc
 	static void NormalizePath(FString& InPath);
 	static void NormalizePath(FStringBuilderBase& InPath);
 	static FString GetPrimaryGPUBrand();
-	static struct FGPUDriverInfo GetGPUDriverInfo(const FString& DeviceDescription);
+	static struct FGPUDriverInfo GetGPUDriverInfo(const FString& DeviceDescription, bool bVerbose = true);
 	static void GetOSVersions( FString& out_OSVersionLabel, FString& out_OSSubVersionLabel );
 	static FString GetOSVersion();
 	static NSOperatingSystemVersion GetNSOperatingSystemVersion();
