Annotated FMemory allocation functions calls with FORCENOINLINE to prevent bloating LTO only builds

diff --git a/Engine/Source/Runtime/Core/Public/HAL/UnrealMemory.h b/Engine/Source/Runtime/Core/Public/HAL/UnrealMemory.h
--- a/Engine/Source/Runtime/Core/Public/HAL/UnrealMemory.h
+++ b/Engine/Source/Runtime/Core/Public/HAL/UnrealMemory.h
@@ -209,12 +209,12 @@ struct FMemory
 	// C style memory allocation stubs.
 	//
 
-	UE_ALLOCATION_FUNCTION(1, 2) static CORE_API void* Malloc(SIZE_T Count, uint32 Alignment = DEFAULT_ALIGNMENT);
-	UE_ALLOCATION_FUNCTION(2, 3) static CORE_API void* Realloc(void* Original, SIZE_T Count, uint32 Alignment = DEFAULT_ALIGNMENT);
-	static CORE_API void Free(void* Original);
-	static CORE_API SIZE_T GetAllocSize(void* Original);
+	UE_ALLOCATION_FUNCTION(1, 2) static FORCENOINLINE CORE_API void* Malloc(SIZE_T Count, uint32 Alignment = DEFAULT_ALIGNMENT);
+	UE_ALLOCATION_FUNCTION(2, 3) static FORCENOINLINE CORE_API void* Realloc(void* Original, SIZE_T Count, uint32 Alignment = DEFAULT_ALIGNMENT);
+	static FORCENOINLINE CORE_API void Free(void* Original);
+	static FORCENOINLINE CORE_API SIZE_T GetAllocSize(void* Original);
 
-	UE_ALLOCATION_FUNCTION(1, 2) static CORE_API void* MallocZeroed(SIZE_T Count, uint32 Alignment = DEFAULT_ALIGNMENT);
+	UE_ALLOCATION_FUNCTION(1, 2) static FORCENOINLINE CORE_API void* MallocZeroed(SIZE_T Count, uint32 Alignment = DEFAULT_ALIGNMENT);
 
 	/**
 	* For some allocators this will return the actual size that should be requested to eliminate
@@ -222,7 +222,7 @@ struct FMemory
 	* and shrink containers to optimal sizes.
 	* This call is always fast and threadsafe with no locking.
 	*/
-	static CORE_API SIZE_T QuantizeSize(SIZE_T Count, uint32 Alignment = DEFAULT_ALIGNMENT);
+	static FORCENOINLINE CORE_API SIZE_T QuantizeSize(SIZE_T Count, uint32 Alignment = DEFAULT_ALIGNMENT);
 
 	/**
 	* Releases as much memory as possible. Must be called from the main thread.
