Added proper error handling for mmap().

diff --git a/Engine/Source/Runtime/Core/Private/Android/AndroidPlatformMemory.cpp b/Engine/Source/Runtime/Core/Private/Android/AndroidPlatformMemory.cpp
--- a/Engine/Source/Runtime/Core/Private/Android/AndroidPlatformMemory.cpp
+++ b/Engine/Source/Runtime/Core/Private/Android/AndroidPlatformMemory.cpp
@@ -667,9 +667,18 @@ int64 LLMMallocTotal = 0;
 void* LLMAlloc(size_t Size)
 {
 	void* Ptr = mmap(nullptr, Size, PROT_READ | PROT_WRITE, MAP_PRIVATE | MAP_ANON, -1, 0);
-	SetMmapName(Ptr, Size, "UE");
+	if (Ptr == MAP_FAILED)
+	{
+		const int ErrNo = errno;
+		UE_LOG(LogHAL, Error, TEXT("mmap(Size=%llu) failed with errno = %d (%s)"), Size,
+			ErrNo, StringCast< TCHAR >(strerror(ErrNo)).Get());
+	}
+	else
+	{
+		SetMmapName(Ptr, Size, "UE");
 
-	LLMMallocTotal += Size;
+		LLMMallocTotal += Size;
+	}
 
 	return Ptr;
 }
diff --git a/Engine/Source/Runtime/Core/Private/Apple/ApplePlatformMemory.cpp b/Engine/Source/Runtime/Core/Private/Apple/ApplePlatformMemory.cpp
--- a/Engine/Source/Runtime/Core/Private/Apple/ApplePlatformMemory.cpp
+++ b/Engine/Source/Runtime/Core/Private/Apple/ApplePlatformMemory.cpp
@@ -861,8 +861,16 @@ int64 LLMMallocTotal = 0;
 void* LLMAlloc(size_t Size)
 {
     void* Ptr = MmapWithTag(nullptr, Size, PROT_READ | PROT_WRITE, MAP_PRIVATE | MAP_ANON, MmapTag::LLMAlloc, 0);
-
-    LLMMallocTotal += Size;
+	if (Ptr == MAP_FAILED)
+	{
+		const int ErrNo = errno;
+		UE_LOG(LogHAL, Error, TEXT("mmap(Size=%llu) failed with errno = %d (%s)"), Size,
+			ErrNo, StringCast< TCHAR >(strerror(ErrNo)).Get());
+	}
+	else
+	{
+		LLMMallocTotal += Size;
+	}
     
     return Ptr;
 }
