[Box2D] Create a Line intersection function for use via the UnrealMathUtility library to match the existing one in the 2D Box

diff --git a/Engine/Source/Runtime/Core/Public/Math/Box2D.h b/Engine/Source/Runtime/Core/Public/Math/Box2D.h
--- a/Engine/Source/Runtime/Core/Public/Math/Box2D.h
+++ b/Engine/Source/Runtime/Core/Public/Math/Box2D.h
@@ -7,6 +7,7 @@
 #include "CoreTypes.h"
 #include "Math/MathFwd.h" // IWYU pragma: export
 #include "Math/UnrealMathUtility.h"
+#include "Math/Vector.h" 
 #include "Math/Vector2D.h"
 #include "Misc/AssertionMacros.h"
 #include "Misc/LargeWorldCoordinates.h"
@@ -338,11 +339,11 @@ public:
 	/**
 	 * Checks whether the given box intersects this box.
 	 *
-	 * @param other bounding box to test intersection
+	 * @param Other bounding box to test intersection
 	 * @return true if boxes intersect, false otherwise.
 	 */
-	[[nodiscard]] inline bool Intersect( const TBox2<T> & other ) const;
-
+	[[nodiscard]] inline bool Intersect( const TBox2<T> & Other ) const;
+	
 	/**
 	 * Checks whether the given point is inside this box.
 	 *
@@ -628,3 +629,59 @@ template<> struct TIsUECoreVariant<FBox2f> { enum { Value = true }; };
 //template<> struct TCanBulkSerialize<FBox2d> { enum { Value = false }; };	// LWC_TODO: This can be done (via versioning) once LWC is fixed to on.
 template<> struct TIsPODType<FBox2d> { enum { Value = true }; };
 template<> struct TIsUECoreVariant<FBox2d> { enum { Value = true }; };
+
+/* FMath inline functions
+ *****************************************************************************/
+
+template<typename FReal>
+inline bool FMath::LineBoxIntersection2D
+	(
+	const UE::Math::TBox2<FReal>&		Box,
+	const UE::Math::TVector2<FReal>&	Start,
+	const UE::Math::TVector2<FReal>&	End
+	)
+{
+	using VectorType = UE::Math::TVector2<FReal>;
+	
+	// The segment is intersecting if its endpoints lie on any of the edges or inside the box
+	if (Box.IsInsideOrOn(Start) || Box.IsInsideOrOn(End))
+	{
+		return true;
+	}
+	
+	const VectorType BoxSize = Box.GetSize();
+	const FVector SegmentStart(Start.X, Start.Y, 0);
+	const FVector SegmentEnd(End.X, End.Y, 0);
+	FVector OutIntersectionPoint; // We don't do anything with this value, but it is needed by the SegmentIntersection2D function
+
+	//  A ____ B
+	//  |      |
+	//  |      |
+	//  C ____ D
+	const FVector PointA = FVector(Box.Min.X, Box.Min.Y, 0);
+	const FVector PointB = PointA + FVector(BoxSize.X, 0, 0);
+	
+	if (SegmentIntersection2D(SegmentStart, SegmentEnd, PointA, PointB, OutIntersectionPoint))
+	{
+		return true;
+	}
+
+	const FVector PointC = PointA + FVector(0, BoxSize.Y, 0);
+	if (SegmentIntersection2D(SegmentStart, SegmentEnd, PointA, PointC, OutIntersectionPoint))
+	{
+		return true;
+	}
+	
+	const FVector PointD = PointA + FVector(BoxSize.X, BoxSize.Y, 0);;
+	if (SegmentIntersection2D(SegmentStart, SegmentEnd, PointB, PointD, OutIntersectionPoint))
+	{
+		return true;
+	}
+	
+	if (SegmentIntersection2D(SegmentStart, SegmentEnd, PointC, PointD, OutIntersectionPoint))
+	{
+		return true;
+	}
+	
+	return false;
+}
\ No newline at end of file
diff --git a/Engine/Source/Runtime/Core/Public/Math/UnrealMathUtility.h b/Engine/Source/Runtime/Core/Public/Math/UnrealMathUtility.h
--- a/Engine/Source/Runtime/Core/Public/Math/UnrealMathUtility.h
+++ b/Engine/Source/Runtime/Core/Public/Math/UnrealMathUtility.h
@@ -1998,6 +1998,10 @@ public:
 	template<typename FReal>
 	[[nodiscard]] static bool LineBoxIntersection( const UE::Math::TBox<FReal>& Box, const UE::Math::TVector<FReal>& Start, const UE::Math::TVector<FReal>& End, const UE::Math::TVector<FReal>& Direction, const UE::Math::TVector<FReal>& OneOverDirection );
 
+	/** Determines whether a line intersects a 2D box. */
+	template<typename FReal>
+	[[nodiscard]] static bool LineBoxIntersection2D( const UE::Math::TBox2<FReal>& Box, const UE::Math::TVector2<FReal>& Start, const UE::Math::TVector2<FReal>& End );
+
 	/* Swept-Box vs Box test */
 	static CORE_API bool LineExtentBoxIntersection(const FBox& inBox, const FVector& Start, const FVector& End, const FVector& Extent, FVector& HitLocation, FVector& HitNormal, float& HitTime);
 
