Added support for extra analytic information in FError.

#DisabledValidations nonUnityAndStaticAnalysis

diff --git a/Engine/Source/Runtime/Core/Private/Experimental/UnifiedError/UnifiedError.cpp b/Engine/Source/Runtime/Core/Private/Experimental/UnifiedError/UnifiedError.cpp
--- a/Engine/Source/Runtime/Core/Private/Experimental/UnifiedError/UnifiedError.cpp
+++ b/Engine/Source/Runtime/Core/Private/Experimental/UnifiedError/UnifiedError.cpp
@@ -4,6 +4,7 @@
 #include "Containers/AnsiString.h"
 #include "Internationalization/Text.h"
 #include "Internationalization/Internationalization.h"
+#include "Serialization/CompactBinarySerialization.h"
 #include "Logging/StructuredLog.h"
 #include "Logging/StructuredLogFormat.h"
 
@@ -107,7 +108,7 @@ FText FError::GetErrorMessage(bool bIncludeContext) const
 		const IErrorDetails* DetailsIt = ErrorDetails.GetReference();
 		while (DetailsIt != nullptr)
 		{
-			if (DetailsIt->ShouldIncludeInContextMessage())
+			if (!DetailsIt->ShouldInclude(EDetailFilter::IncludeInContextLogMessage))
 			{
 				AdditionalContext.Add(FText::FromString(FUtf8String::Printf("(%s:{%s})", DetailsIt->GetErrorDetailsTypeNameForLog().GetData(), DetailsIt->GetErrorDetailsTypeNameForLog().GetData())));
 				
@@ -131,23 +132,24 @@ FText FError::GetErrorMessage(bool bIncludeContext) const
 #endif
 	TCbWriter<1024> Writer;
 
-	SerializeDetailsForLog(Writer);
+	SerializeDetails(Writer, EDetailFilter::IncludeInSerialize);
 	FUtf8StringBuilderBase OutputMessage;
 
 	return Template.FormatToText(Writer.Save());
 }
 
-
-void FError::SerializeDetailsForLog(FCbWriter& Writer) const
+void FError::SerializeDetails(FCbWriter& Writer, const EDetailFilter DetailFilter, bool bIncludeRoot) const
 {
 	const IErrorDetails* DetailsIt = ErrorDetails.GetReference();
 	while (DetailsIt != nullptr)
 	{
-		Writer.SetName(DetailsIt->GetErrorDetailsTypeNameForLog());
-		DetailsIt->SerializeToCb(Writer, *this);
-
+		if (DetailsIt->ShouldInclude(DetailFilter))
+		{
+			Writer.SetName(DetailsIt->GetErrorDetailsTypeNameForLog());
+			DetailsIt->SerializeToCb(Writer, *this);
+		}
 		const IErrorDetails* Next = DetailsIt->GetInnerErrorDetails().GetReference();
-		if (Next == nullptr)
+		if (Next == nullptr && bIncludeRoot)
 		{
 			// the last one contains some special information we want to reference by name
 			Writer.SetName(ANSITEXTVIEW("Root"));
@@ -157,6 +159,27 @@ void FError::SerializeDetailsForLog(FCbWriter& Writer) const
 	}
 }
 
+FString FError::SerializeToJsonString(const EDetailFilter DetailFilter) const
+{
+	FCbWriter Writer;
+	// Writer.SetName(ANSITEXTVIEW("FError"));
+	Writer.BeginObject();
+	SerializeDetails(Writer, DetailFilter);
+	Writer.EndObject();
+
+	// SerializeBasicToCb(Writer, UTF8TEXTVIEW("AppendFormatString"), AppendFormatString);
+
+	TStringBuilder<256> Text;
+	CompactBinaryToJson(Writer.Save(), Text);
+	return Text.ToString();
+}
+
+
+FString FError::SerializeToJsonForAnalytics() const
+{
+	return SerializeToJsonString(EDetailFilter::IncludeInAnalytics);	
+}
+
 int32 FError::GetErrorCode() const
 {
 	return ErrorCode;
@@ -196,7 +219,7 @@ public:
 		return FText::Join(FText::FromString(" "), InnerFormatString, AppendFormatString);
 	}
 
-	virtual bool ShouldIncludeInContextMessage() const { return false; }
+	virtual bool ShouldInclude(const EDetailFilter Filter) const { return EnumHasAnyFlags(Filter, EDetailFilter::Default & ~EDetailFilter::IncludeInContextLogMessage); }
 };
 
 
@@ -216,7 +239,7 @@ void SerializeForLog(FCbWriter& Writer, const FError& Error)
 	// what I have to do 
 	TCbWriter<1024> InnerWriter;
 
-	Error.SerializeDetailsForLog(InnerWriter);
+	Error.SerializeDetails(InnerWriter, EDetailFilter::IncludeInSerialize);
 
 	FText FormatText = FText::Join(FText::FromString(":"), NSLOCTEXT("UnifiedError", "DefaultLogPrepend", "{Root/ModuleIdString}.{Root/ErrorCodeString}"), Error.GetFormatErrorText());
 #if USE_LOCALIZED_STRUCTURED_LOG_FOR_FERRORMESSAGE
diff --git a/Engine/Source/Runtime/Core/Public/Experimental/UnifiedError/UnifiedError.h b/Engine/Source/Runtime/Core/Public/Experimental/UnifiedError/UnifiedError.h
--- a/Engine/Source/Runtime/Core/Public/Experimental/UnifiedError/UnifiedError.h
+++ b/Engine/Source/Runtime/Core/Public/Experimental/UnifiedError/UnifiedError.h
@@ -197,14 +197,22 @@ constexpr uint64 CompileTimeHashString(const char* String)
 
 namespace UE::UnifiedError
 {
-	
-	namespace DetailsTypes
+
+	enum class EDetailFilter : uint8
 	{
-		inline constexpr int32 IERROR_DETAILS_TYPE = 1;
-		inline constexpr int32 STATIC_ERROR_DETAILS_TYPE = 2;
-		inline constexpr int32 DYNAMIC_ERROR_DETAILS_TYPE = 3;
-		inline constexpr int32 FIRST_CUSTOM_DETAILS_TYPE = 4;
-	}
+		IncludeInSerialize = 1 << 0, // included when we serialize the details objects, almost every details should have this flag
+		IncludeInAnalytics = 1 << 1, // what details objects to include when we are creating events for analytics
+		IncludeInContextLogMessage = 1 << 2, // what details objects are included when the log message includes context objects on the string
+		IncludeInLogMessage = 1 << 3, // which details objects are included when we log a message without context objects included
+		// append here
+
+
+		// standard values
+		Default = IncludeInSerialize | IncludeInContextLogMessage,
+		None = 0x00,
+		All = 0xff
+	};
+	ENUM_CLASS_FLAGS(EDetailFilter);
 
 	namespace ErrorRegistry
 	{
@@ -254,8 +262,7 @@ namespace UE::UnifiedError
 		virtual const FAnsiStringView GetErrorDetailsTypeNameForLog() const = 0;
 		virtual void SerializeToCb(FCbWriter& Writer, const FError& Error) const =0;
 
-
-		virtual bool ShouldIncludeInContextMessage() const { return true; }
+		virtual bool ShouldInclude(const EDetailFilter DetailFilter) const { return EnumHasAnyFlags( EDetailFilter::Default, DetailFilter); }
 	};
 
 
@@ -329,18 +336,20 @@ namespace UE::UnifiedError
 		/// <returns></returns>
 		CORE_API virtual const FText GetErrorFormatString(const FError& Error) const override;
 
+		virtual bool ShouldInclude(const EDetailFilter DetailFilter) const override { return EnumHasAnyFlags(DetailFilter, ~EDetailFilter::IncludeInAnalytics); }
 	};
 
+	CORE_API void SerializeForLog(FCbWriter& Writer, const class FError& Error);
 
 	template<typename T>
 	class TErrorDetails : public FDynamicErrorDetails
 	{
 	public:
 		TErrorDetails() {}
-		TErrorDetails(T&& InErrorDetail, bool bInAppendContextToMessage)
+		TErrorDetails(T&& InErrorDetail, const EDetailFilter InDetailFilterMask = EDetailFilter::Default)
 		{
 			ErrorDetail = MoveTemp(InErrorDetail);
-			bAppendContextToMessage = bInAppendContextToMessage;
+			DetailFilterMask = InDetailFilterMask;
 		}
 		/*TErrorDetails(T&& InErrorDetail, TRefCountPtr<const IErrorDetails> InInnerErrorDetails) : FDynamicErrorDetails(InInnerErrorDetails)
 		{
@@ -387,9 +396,11 @@ namespace UE::UnifiedError
 			return TErrorStructFeatures<T>::GetErrorDetailsLogEntryNameAsString();
 		}
 
+		virtual bool ShouldInclude(const EDetailFilter InDetailFilter) const override { return EnumHasAnyFlags(InDetailFilter, DetailFilterMask); }
+
 	private:
 		T ErrorDetail;
-		bool bAppendContextToMessage = false;
+		EDetailFilter DetailFilterMask = EDetailFilter::Default;
 		
 	public:
 		const T& GetErrorContext() const
@@ -404,7 +415,7 @@ namespace UE::UnifiedError
 
 		virtual const FText GetErrorFormatString(const FError& Error) const override
 		{
-			if (bAppendContextToMessage)
+			if (EnumHasAnyFlags(DetailFilterMask, EDetailFilter::IncludeInLogMessage))
 			{
 				return FText::Join(FText::FromString(" "), FDynamicErrorDetails::GetErrorFormatString(Error), FText::FromString(FUtf8String::Printf("{%s}", GetErrorDetailsTypeNameForLog().GetData())));
 			}
@@ -465,9 +476,9 @@ namespace UE::UnifiedError
 		CORE_API FText GetFormatErrorText() const;
 
 		template<typename T>
-		void PushErrorContext(T&& InErrorStruct, bool bShouldAppendToErrorMessage = false)
+		void PushErrorContext(T&& InErrorStruct, const EDetailFilter& InDetailFilter = EDetailFilter::Default)
 		{
-			PushErrorDetails(TRefCountPtr<IErrorDetails>(new TErrorDetails<T>(MoveTemp(InErrorStruct), bShouldAppendToErrorMessage)));
+			PushErrorDetails(TRefCountPtr<IErrorDetails>(new TErrorDetails<T>(MoveTemp(InErrorStruct), InDetailFilter)));
 		}
 
 		template<typename T>
@@ -495,7 +506,7 @@ namespace UE::UnifiedError
 
 		CORE_API void AppendFormatString(FText&& InFormatString);
 
-		CORE_API void SerializeDetailsForLog(FCbWriter& Writer) const;
+		CORE_API FString SerializeToJsonForAnalytics() const;
 
 		UE_FORCEINLINE_HINT bool IsValid() const
 		{
@@ -510,6 +521,11 @@ namespace UE::UnifiedError
 		}
 	private:
 
+		friend CORE_API void SerializeForLog(FCbWriter& Writer, const FError& Error);
+
+		void SerializeDetails(FCbWriter& Writer, const EDetailFilter DetailFilter, bool bIncludeRoot = true) const;
+		FString SerializeToJsonString(const EDetailFilter DetailFilter) const;
+
 		const FManditoryErrorDetails* GetManditoryErrorDetails() const;
 
 		CORE_API void PushErrorDetails(TRefCountPtr<IErrorDetails> InErrorDetails);
@@ -606,7 +622,8 @@ namespace UE::UnifiedError
 		virtual uint32 Release() const final { return 10; }
 		virtual uint32 GetRefCount() const final { return 10; }
 
-		virtual bool ShouldIncludeInContextMessage() const { return false; }
+		
+		virtual bool ShouldInclude(const EDetailFilter Filter) const override { return EnumHasAnyFlags(Filter, EDetailFilter::Default); }
 	};
 
 
@@ -639,7 +656,7 @@ namespace UE::UnifiedError
 		Writer.AddString(PropertyName, Value.ToString());
 	}
 
-	CORE_API void SerializeForLog(FCbWriter& Writer, const FError& Error);
+	
 
 } // namespace UnifiedError
 
