FPolyglotTextData fixes

* Fixed RegisterPolyglotTextData failing to add the display string when running an editor with ShouldForceLoadGameLocalization enabled.
* Fixed the native translation of a FPolyglotTextData using bIsMinimalPatch taking precedence over existing LocRes translations, even when not using the native language.

diff --git a/Engine/Source/Runtime/Core/Private/Internationalization/PolyglotTextSource.cpp b/Engine/Source/Runtime/Core/Private/Internationalization/PolyglotTextSource.cpp
--- a/Engine/Source/Runtime/Core/Private/Internationalization/PolyglotTextSource.cpp
+++ b/Engine/Source/Runtime/Core/Private/Internationalization/PolyglotTextSource.cpp
@@ -113,6 +113,7 @@ void FPolyglotTextSource::AddPolyglotDataToResource(const FPolyglotTextData& InP
 
 	const int32 BaseResourcePriority = GetPriority() * -1; // Flip the priority as larger text source priorities are more important, but smaller text resource priorities are more important
 	const int32 NativeResourcePriority = BaseResourcePriority + InPrioritizedCultures.Num(); // Native resources always prioritize below any localized ones
+	const int32 NormalNativeResourcePriority = ELocalizedTextSourcePriority::Normal + InPrioritizedCultures.Num();
 	const FString NativeCulture = InPolyglotTextData.ResolveNativeCulture();
 
 	// We skip loading the native text if we're transitioning to the native culture as there's no extra work that needs to be done
@@ -121,12 +122,16 @@ void FPolyglotTextSource::AddPolyglotDataToResource(const FPolyglotTextData& InP
 		FString LocalizedString;
 		if (GetLocalizedStringForPolyglotData(TArrayView<const FString>(&NativeCulture, 1), LocalizedString))
 		{
+			// When using bIsMinimalPatch, we need to drop the native string priority back down to Normal so that the LocRes translations are still allowed to overwrite it
+			// Without this, any FPolyglotTextData that includes a replacement for the native translation will have that native translation take precedence for any languages not included in the FPolyglotTextData
+			const int32 ResolvedNativeResourcePriority = InPolyglotTextData.IsMinimalPatch() ? NormalNativeResourcePriority : NativeResourcePriority;
+
 			InOutNativeResource.AddEntry(
 				InPolyglotTextData.GetNamespace(),
 				InPolyglotTextData.GetKey(),
 				InPolyglotTextData.GetNativeString(),
 				LocalizedString,
-				NativeResourcePriority
+				ResolvedNativeResourcePriority
 				);
 		}
 	}
diff --git a/Engine/Source/Runtime/Core/Private/Internationalization/TextLocalizationManager.cpp b/Engine/Source/Runtime/Core/Private/Internationalization/TextLocalizationManager.cpp
--- a/Engine/Source/Runtime/Core/Private/Internationalization/TextLocalizationManager.cpp
+++ b/Engine/Source/Runtime/Core/Private/Internationalization/TextLocalizationManager.cpp
@@ -943,16 +943,23 @@ void FTextLocalizationManager::RegisterPolyglotTextData(TArrayView<const FPolygl
 			// Work out which culture to use - this is typically the current language unless we're in the 
 			// editor where the game localization preview affects the language we use for game text
 			FString CultureName;
-			if (PolyglotTextData.GetCategory() != ELocalizedTextSourceCategory::Game || !GIsEditor)
-			{
-				CultureName = FInternationalization::Get().GetCurrentLanguage()->GetName();
-			}
 #if WITH_EDITOR
-			else if (bIsGameLocalizationPreviewEnabled)
+			if (PolyglotTextData.GetCategory() == ELocalizedTextSourceCategory::Game)
 			{
-				CultureName = GetConfiguredGameLocalizationPreviewLanguage();
+				if (bIsGameLocalizationPreviewEnabled)
+				{
+					CultureName = GetConfiguredGameLocalizationPreviewLanguage();
+				}
+				else if (!GIsEditor || ShouldForceLoadGameLocalization())
+				{
+					CultureName = FInternationalization::Get().GetCurrentLanguage()->GetName();
+				}
 			}
+			else
 #endif
+			{
+				CultureName = FInternationalization::Get().GetCurrentLanguage()->GetName();
+			}
 
 			if (!CultureName.IsEmpty())
 			{
