Added Flush,Seek and OutOfDiskspace I/O error codes.

#rnx

diff --git a/Engine/Source/Runtime/Core/Private/IO/IoDispatcher.cpp b/Engine/Source/Runtime/Core/Private/IO/IoDispatcher.cpp
--- a/Engine/Source/Runtime/Core/Private/IO/IoDispatcher.cpp
+++ b/Engine/Source/Runtime/Core/Private/IO/IoDispatcher.cpp
@@ -60,7 +60,10 @@ static const TCHAR* const GetIoErrorText_ErrorCodeTextArray[] =
 	TEXT("Not Installed"),
 	TEXT("Pending Host Group"),
 	TEXT("Timeout"),
-	TEXT("Delete Error")
+	TEXT("Delete Error"),
+	TEXT("Out Of Diskspace"),
+	TEXT("FileSeek Failed"),
+	TEXT("FileFlush Failed")
 };
 static_assert(std::size(GetIoErrorText_ErrorCodeTextArray) == static_cast<uint32>(EIoErrorCode::Last), "GetIoErrorText_ErrorCodeTextArray is missing entries to match EIoErrorCode. Please update the array.");
 
diff --git a/Engine/Source/Runtime/Core/Public/IO/IoStatus.h b/Engine/Source/Runtime/Core/Public/IO/IoStatus.h
--- a/Engine/Source/Runtime/Core/Public/IO/IoStatus.h
+++ b/Engine/Source/Runtime/Core/Public/IO/IoStatus.h
@@ -39,6 +39,9 @@ enum class EIoErrorCode
 	PendingHostGroup,
 	Timeout,
 	DeleteError,
+	OutOfDiskSpace,
+	FileSeekError,
+	FileFlushError,
 	Last
 };
 
