PlatformEvents channel uses callbacks

* Use new channel callback to deny enabling channel and provide user feedback, if requirements not fulfilled
* Remove global callbacks and use channel enable callback to activate tracing.

diff --git a/Engine/Plugins/TraceUtilities/Source/EditorTraceUtilities/Private/InsightsStatusBar/SInsightsStatusBar.cpp b/Engine/Plugins/TraceUtilities/Source/EditorTraceUtilities/Private/InsightsStatusBar/SInsightsStatusBar.cpp
--- a/Engine/Plugins/TraceUtilities/Source/EditorTraceUtilities/Private/InsightsStatusBar/SInsightsStatusBar.cpp
+++ b/Engine/Plugins/TraceUtilities/Source/EditorTraceUtilities/Private/InsightsStatusBar/SInsightsStatusBar.cpp
@@ -1006,8 +1006,6 @@ void SInsightsStatusBarWidget::ToggleChannel_Execute(int32 Index)
 		const FString& ChannelName = ChannelsInfo[Index].Name;
 		bool bChannelShouldBeEnabled = !ChannelsInfo[Index].bIsEnabled;
 		UE::Trace::ToggleChannel(*ChannelName, bChannelShouldBeEnabled);
-
-		FPlatformEventsTrace::OnTraceChannelUpdated(ChannelName, bChannelShouldBeEnabled);
 	}
 }
 
diff --git a/Engine/Source/Runtime/Core/Private/ProfilingDebugging/Microsoft/EventTracingForWindows.cpp b/Engine/Source/Runtime/Core/Private/ProfilingDebugging/Microsoft/EventTracingForWindows.cpp
--- a/Engine/Source/Runtime/Core/Private/ProfilingDebugging/Microsoft/EventTracingForWindows.cpp
+++ b/Engine/Source/Runtime/Core/Private/ProfilingDebugging/Microsoft/EventTracingForWindows.cpp
@@ -4,6 +4,7 @@
 #include "HAL/Runnable.h"
 #include "HAL/Event.h"
 #include "HAL/RunnableThread.h"
+#include "HAL/PlatformProcess.h"
 #include "Logging/LogMacros.h"
 #include "Experimental/Containers/SherwoodHashTable.h"
 #include "Trace/Trace.h"
@@ -20,12 +21,33 @@ DEFINE_LOG_CATEGORY_STATIC(LogPlatformEvents, Log, All);
 
 namespace {
 
+struct FHandleScope
+{
+	FHandleScope(HANDLE InHandle)
+		: Handle(InHandle)
+	{}
+
+	FHandleScope()
+		: Handle(INVALID_HANDLE_VALUE)
+	{}
+
+	~FHandleScope()
+	{ 
+		CloseHandle(Handle); 
+	}
+
+	operator HANDLE() const { return Handle; }
+
+	HANDLE Handle;
+};
+
 class FPlatformEvents : public FRunnable
 {
 public:
 	FPlatformEvents(uint32 SamplingIntervalUsec);
 	~FPlatformEvents();
 
+	bool CanEnable(FPlatformEventsTrace::EEventType Event, const TCHAR** OutReason);
 	void Enable(FPlatformEventsTrace::EEventType Event);
 	void Disable(FPlatformEventsTrace::EEventType Event);
 
@@ -231,6 +253,59 @@ FPlatformEvents::~FPlatformEvents()
 	Stop();
 }
 
+/////////////////////////////////////////////////////////////////////
+bool FPlatformEvents::CanEnable(FPlatformEventsTrace::EEventType Event, const TCHAR** OutReason)
+{
+	// In order to control tracing the process must either run under
+	// elevated privileges OR the user has to be member of PERF_LOG_USERS
+	static const TCHAR* ErrProcAccess = TEXT("Cannot access process.");
+	static const TCHAR* ErrProcRights = TEXT("Process does not have privileges to start event tracing.");
+	
+#if PLATFORM_WINDOWS
+	FHandleScope CurrentProcToken;
+	if (!::OpenProcessToken(GetCurrentProcess(), TOKEN_QUERY, &CurrentProcToken.Handle))
+	{
+		if (OutReason)
+		{
+			*OutReason = ErrProcAccess;
+		}
+		UE_LOG(LogCore, Log, TEXT("Unable to enable platform events: %s"), ErrProcAccess);
+		return false;
+	}
+
+	// Get the elevation token
+	TOKEN_ELEVATION ElevationToken;
+	DWORD ReturnSize = 0;
+	if (!::GetTokenInformation(CurrentProcToken, TokenElevation, &ElevationToken, sizeof(TOKEN_ELEVATION), &ReturnSize))
+	{
+		if (OutReason)
+		{
+			*OutReason = ErrProcAccess;
+		}
+		UE_LOG(LogCore, Log, TEXT("Unable to enable platform events: %s"), ErrProcAccess);
+		return false;
+	}
+	
+	// Check if the user is part of PERF_LOG_USERS
+	// todo: Implement
+
+	if (ElevationToken.TokenIsElevated /* || in perf log users */)
+	{
+		return true;
+	}
+	
+	if (OutReason)
+	{
+		*OutReason = ErrProcRights;
+	}
+	
+	UE_LOG(LogCore, Log, TEXT("Unable to enable platform events: %s"), ErrProcRights);
+	return false;
+#else
+	return true;
+#endif
+}
+
 /////////////////////////////////////////////////////////////////////
 
 void FPlatformEvents::Enable(FPlatformEventsTrace::EEventType Event)
@@ -550,6 +625,15 @@ void FPlatformEventsTrace::Init(uint32 SamplingIntervalUsec)
 	GPlatformEvents = new FPlatformEvents(SamplingIntervalUsec);
 }
 
+bool FPlatformEventsTrace::CanEnable(FPlatformEventsTrace::EEventType Event, const TCHAR** OutReason)
+{
+	if (GPlatformEvents)
+	{
+		return GPlatformEvents->CanEnable(Event, OutReason);
+	}
+	return false;
+}
+
 void FPlatformEventsTrace::Enable(FPlatformEventsTrace::EEventType Event)
 {
 	if (GPlatformEvents)
diff --git a/Engine/Source/Runtime/Core/Private/ProfilingDebugging/PlatformEvents.cpp b/Engine/Source/Runtime/Core/Private/ProfilingDebugging/PlatformEvents.cpp
--- a/Engine/Source/Runtime/Core/Private/ProfilingDebugging/PlatformEvents.cpp
+++ b/Engine/Source/Runtime/Core/Private/ProfilingDebugging/PlatformEvents.cpp
@@ -5,8 +5,60 @@
 
 /////////////////////////////////////////////////////////////////////
 
-UE_TRACE_CHANNEL_DEFINE(ContextSwitchChannel)
-UE_TRACE_CHANNEL_DEFINE(StackSamplingChannel)
+#if UE_TRACE_ENABLED
+class FCSChannel : public UE::Trace::FChannel
+{
+	virtual bool OnToggle(bool bNewState, const TCHAR** OutReason) override
+	{
+#if PLATFORM_SUPPORTS_PLATFORM_EVENTS
+		return FPlatformEventsTrace::CanEnable(FPlatformEventsTrace::EEventType::ContextSwitch, OutReason);
+#else
+		return false;
+#endif
+	}
+	virtual void OnToggled(bool bNewState) override
+	{
+#if PLATFORM_SUPPORTS_PLATFORM_EVENTS
+		if (bNewState)
+		{
+			FPlatformEventsTrace::Enable(FPlatformEventsTrace::EEventType::ContextSwitch);
+		}
+		else
+		{
+			FPlatformEventsTrace::Disable(FPlatformEventsTrace::EEventType::ContextSwitch);
+		}
+#endif
+	}
+};
+
+class FSSChannel : public UE::Trace::FChannel
+{
+	virtual bool OnToggle(bool bNewState, const TCHAR** OutReason) override
+	{
+#if PLATFORM_SUPPORTS_PLATFORM_EVENTS
+		return FPlatformEventsTrace::CanEnable(FPlatformEventsTrace::EEventType::StackSampling, OutReason);
+#else 
+		return false;
+#endif
+	}
+	virtual void OnToggled(bool bNewState) override
+	{
+#if PLATFORM_SUPPORTS_PLATFORM_EVENTS
+		if (bNewState)
+		{
+			FPlatformEventsTrace::Enable(FPlatformEventsTrace::EEventType::StackSampling);
+		}
+		else
+		{
+			FPlatformEventsTrace::Disable(FPlatformEventsTrace::EEventType::StackSampling);
+		}
+#endif
+	}
+};
+#endif
+
+UE_TRACE_CHANNEL_CUSTOM_DEFINE(ContextSwitchChannel, FCSChannel, "CPU context switches", false)
+UE_TRACE_CHANNEL_CUSTOM_DEFINE(StackSamplingChannel, FSSChannel, "Stack sampling", false)
 
 /////////////////////////////////////////////////////////////////////
 
diff --git a/Engine/Source/Runtime/Core/Private/Windows/WindowsPlatformEventTracingForWindows.h b/Engine/Source/Runtime/Core/Private/Windows/WindowsPlatformEventTracingForWindows.h
--- a/Engine/Source/Runtime/Core/Private/Windows/WindowsPlatformEventTracingForWindows.h
+++ b/Engine/Source/Runtime/Core/Private/Windows/WindowsPlatformEventTracingForWindows.h
@@ -5,4 +5,5 @@
 #include "Windows/AllowWindowsPlatformTypes.h"
 #	include <evntrace.h>
 #	include <evntcons.h>
+#	include <processthreadsapi.h>
 #include "Windows/HideWindowsPlatformTypes.h"
diff --git a/Engine/Source/Runtime/Core/Public/ProfilingDebugging/PlatformEvents.h b/Engine/Source/Runtime/Core/Public/ProfilingDebugging/PlatformEvents.h
--- a/Engine/Source/Runtime/Core/Public/ProfilingDebugging/PlatformEvents.h
+++ b/Engine/Source/Runtime/Core/Public/ProfilingDebugging/PlatformEvents.h
@@ -14,8 +14,8 @@ class FString;
 
 /////////////////////////////////////////////////////////////////////
 
-UE_TRACE_CHANNEL_EXTERN(ContextSwitchChannel)
-UE_TRACE_CHANNEL_EXTERN(StackSamplingChannel)
+UE_TRACE_CHANNEL_CUSTOM_EXTERN(ContextSwitchChannel, FCSChannel)
+UE_TRACE_CHANNEL_CUSTOM_EXTERN(StackSamplingChannel, FSSChannel)
 
 /////////////////////////////////////////////////////////////////////
 
@@ -35,9 +35,11 @@ public:
 
 	CORE_API static void Init(uint32 SamplingIntervalUsec);
 	CORE_API static void PostInit();
+	CORE_API static bool CanEnable(EEventType Event, const TCHAR** OutReason);
 	CORE_API static void Enable(EEventType Event);
 	CORE_API static void Disable(EEventType Event);
 	CORE_API static void Stop();
+	UE_DEPRECATED(5.7, "Use channel callback instead.")
 	CORE_API static void OnTraceChannelUpdated(const FString& ChannelName, bool bIsEnabled);
 
 	CORE_API static void OutputSettings(uint32 SamplingIntervalUsec);
