Added core error types ArgumentError and CancellationError
- added MakeError that accepts an error context and details filter
- added UE_DECLARE_ERROR_TWOPARAM

#rnx

diff --git a/Engine/Source/Runtime/Core/Private/Experimental/UnifiedError/CoreErrorTypes.cpp b/Engine/Source/Runtime/Core/Private/Experimental/UnifiedError/CoreErrorTypes.cpp
new file mode 100644
--- /dev/null
+++ b/Engine/Source/Runtime/Core/Private/Experimental/UnifiedError/CoreErrorTypes.cpp
@@ -0,0 +1,18 @@
+// Copyright Epic Games, Inc. All Rights Reserved.
+
+#include "Experimental/UnifiedError/CoreErrorTypes.h"
+
+UE_DEFINE_ERROR_MODULE(Core);
+UE_DEFINE_ERROR(ArgumentError, Core);
+UE_DEFINE_ERROR(CancellationError, Core);
+
+namespace UE::UnifiedError
+{
+
+bool IsCancellationError(const FError& Error)
+{
+	return Core::CancellationError::OfType(Error);
+}
+
+} // namespace UE::UnifiedError
+
diff --git a/Engine/Source/Runtime/Core/Public/Experimental/UnifiedError/CoreErrorTypes.h b/Engine/Source/Runtime/Core/Public/Experimental/UnifiedError/CoreErrorTypes.h
new file mode 100644
--- /dev/null
+++ b/Engine/Source/Runtime/Core/Public/Experimental/UnifiedError/CoreErrorTypes.h
@@ -0,0 +1,17 @@
+// Copyright Epic Games, Inc. All Rights Reserved.
+
+#pragma once
+
+#include "UnifiedError.h"
+
+UE_DECLARE_ERROR_MODULE(	CORE_API, Core);
+UE_DECLARE_ERROR_TWOPARAM(	CORE_API, ArgumentError,		1,	Core, NSLOCTEXT("Core", "ArgumentError", "Invalid argument '{ArgumentName}', reason '{ArgumentError}'"), FString, ArgumentName, TEXT(""), FString, ArgumentError, TEXT(""));
+UE_DECLARE_ERROR(			CORE_API, CancellationError,	2,	Core, NSLOCTEXT("Core", "CancellationError", "The operation was cancelled"));
+
+namespace UE::UnifiedError
+{
+
+CORE_API bool IsCancellationError(const FError& Error);
+
+} // namespace UE::UnifiedError
+
diff --git a/Engine/Source/Runtime/Core/Public/Experimental/UnifiedError/UnifiedError.h b/Engine/Source/Runtime/Core/Public/Experimental/UnifiedError/UnifiedError.h
--- a/Engine/Source/Runtime/Core/Public/Experimental/UnifiedError/UnifiedError.h
+++ b/Engine/Source/Runtime/Core/Public/Experimental/UnifiedError/UnifiedError.h
@@ -698,6 +698,7 @@ CORE_API constexpr uint32 GetTypeIdHash(const TCHAR* TypeName)
 			DeclareApi int32 GetErrorCodeId(); \
 			DeclareApi TRefCountPtr<const FStaticErrorDetails> GetStaticErrorDetails(); \
 			DeclareApi const FError& GetStaticError(); \
+			DeclareApi bool OfType(const FError& Other); \
 			static inline FText GetFormatString() { return FormatString; } \
 	}}}
 
@@ -724,12 +725,23 @@ CORE_API constexpr uint32 GetTypeIdHash(const TCHAR* TypeName)
 				static FError StaticError = MakeError();\
 				return StaticError;\
 			}\
+			bool OfType(const FError& Other) \
+			{\
+				return UE::UnifiedError::ModuleName::GetStaticModuleId() == Other.GetModuleId() && ErrorCodeId == Other.GetErrorCode(); \
+			}\
 	}}}
 
 #define UE_DECLARE_ERROR(DeclareApi, ErrorName, ErrorCode, ModuleName, FormatString) \
 	UE_DECLARE_ERROR_INTERNAL(DeclareApi, ErrorName, ErrorCode, ModuleName, FormatString) \
 	namespace UE::UnifiedError { namespace ModuleName { namespace ErrorName {\
 			inline FError MakeError() { return FError(UE::UnifiedError::ModuleName::GetStaticModuleId(), GetErrorCodeId(), TRefCountPtr<const UE::UnifiedError::FManditoryErrorDetails>(GetStaticErrorDetails().GetReference())); } \
+			template<typename ErrorContextType> \
+			inline FError MakeError(ErrorContextType&& Ctx, EDetailFilter DetailFilter = EDetailFilter::Default) \
+			{\
+				FError Error(UE::UnifiedError::ModuleName::GetStaticModuleId(), GetErrorCodeId(), TRefCountPtr<const UE::UnifiedError::FManditoryErrorDetails>(GetStaticErrorDetails().GetReference())); \
+				Error.PushErrorContext(MoveTemp(Ctx), DetailFilter); \
+				return Error; \
+			}\
 	}}}
 
 
@@ -760,3 +772,34 @@ CORE_API constexpr uint32 GetTypeIdHash(const TCHAR* TypeName)
 				return Error; \
 			} \
 	}}}
+
+#define UE_DECLARE_ERROR_TWOPARAM(DeclareApi, ErrorName, ErrorCode, ModuleName, FormatString, ParamOneType, ParamOneName, ParamOneDefault, ParamTwoType, ParamTwoName, ParamTwoDefault) \
+	UE_DECLARE_ERROR_INTERNAL(DeclareApi, ErrorName, ErrorCode, ModuleName, FText::FromString(UTF8TEXT("{"#ModuleName"_F"#ErrorName"}")) ) \
+	namespace UE::UnifiedError { namespace ModuleName { \
+			struct F##ErrorName \
+			{\
+				ParamOneType ParamOneName;\
+				ParamTwoType ParamTwoName;\
+			};\
+	}}\
+	UE_DECLARE_ERRORSTRUCT_FEATURES(ModuleName, F##ErrorName);\
+	inline void SerializeForLog(FCbWriter& Writer, const UE::UnifiedError::ModuleName::F##ErrorName& Context)\
+	{\
+		Writer.BeginObject();\
+		Writer.AddString(ANSITEXTVIEW("$type"), TErrorStructFeatures<UE::UnifiedError::ModuleName::F##ErrorName>::GetErrorContextTypeNameAsString());\
+		UE::SerializeLogFormat(Writer, FormatString); \
+		UE::UnifiedError::SerializeBasicToCb(Writer, ANSITEXTVIEW(#ParamOneName), Context.ParamOneName);\
+		UE::UnifiedError::SerializeBasicToCb(Writer, ANSITEXTVIEW(#ParamTwoName), Context.ParamTwoName);\
+		Writer.EndObject();\
+	}\
+	namespace UE::UnifiedError { namespace ModuleName { \
+	namespace ErrorName { \
+			inline FError MakeError(ParamOneType ParamOneName = ParamOneDefault, ParamTwoType ParamTwoName = ParamTwoDefault) \
+			{ \
+				FError Error = FError(UE::UnifiedError::ModuleName::GetStaticModuleId(), UE::UnifiedError::ModuleName::ErrorName::GetErrorCodeId(), UE::UnifiedError::ModuleName::ErrorName::GetStaticErrorDetails().GetReference()); \
+				F##ErrorName Context = {ParamOneName, ParamTwoName}; \
+				Error.PushErrorContext(MoveTemp(Context)); \
+				return Error; \
+			} \
+	}}}
+
