TScriptArray: Add Max implementation (missed in cl 43379137).

#rnx

diff --git a/Engine/Source/Runtime/Core/Public/Containers/ScriptArray.h b/Engine/Source/Runtime/Core/Public/Containers/ScriptArray.h
--- a/Engine/Source/Runtime/Core/Public/Containers/ScriptArray.h
+++ b/Engine/Source/Runtime/Core/Public/Containers/ScriptArray.h
@@ -47,6 +47,11 @@ public:
 	{
 		return ArrayNum;
 	}
+	/** @return The number of elements the array can hold before reallocation. */
+	[[nodiscard]] FORCEINLINE int32 Max() const
+	{
+		return ArrayMax;
+	}
 	void InsertZeroed( int32 Index, int32 Count, int32 NumBytesPerElement, uint32 AlignmentOfElement )
 	{
 		Insert( Index, Count, NumBytesPerElement, AlignmentOfElement );
