Added FileMove/Flush failed I/O error codes.

#rnx

diff --git a/Engine/Source/Runtime/Core/Private/IO/IoDispatcher.cpp b/Engine/Source/Runtime/Core/Private/IO/IoDispatcher.cpp
--- a/Engine/Source/Runtime/Core/Private/IO/IoDispatcher.cpp
+++ b/Engine/Source/Runtime/Core/Private/IO/IoDispatcher.cpp
@@ -63,7 +63,10 @@ static const TCHAR* const GetIoErrorText_ErrorCodeTextArray[] =
 	TEXT("Delete Error"),
 	TEXT("Out Of Diskspace"),
 	TEXT("FileSeek Failed"),
-	TEXT("FileFlush Failed")
+	TEXT("FileFlush Failed"),
+	TEXT("FileMove Failed"),
+	TEXT("FileClose Failed"),
+	TEXT("File Corrupt")
 };
 static_assert(std::size(GetIoErrorText_ErrorCodeTextArray) == static_cast<uint32>(EIoErrorCode::Last), "GetIoErrorText_ErrorCodeTextArray is missing entries to match EIoErrorCode. Please update the array.");
 
diff --git a/Engine/Source/Runtime/Core/Public/IO/IoStatus.h b/Engine/Source/Runtime/Core/Public/IO/IoStatus.h
--- a/Engine/Source/Runtime/Core/Public/IO/IoStatus.h
+++ b/Engine/Source/Runtime/Core/Public/IO/IoStatus.h
@@ -40,8 +40,11 @@ enum class EIoErrorCode
 	Timeout,
 	DeleteError,
 	OutOfDiskSpace,
-	FileSeekError,
-	FileFlushError,
+	FileSeekFailed,
+	FileFlushFailed,
+	FileMoveFailed,
+	FileCloseFailed,
+	FileCorrupt,
 	Last
 };
 
