Add TNoDestroy wrapper type to prevent atexit destructor registration for global variables

diff --git a/Engine/Source/Runtime/Core/Public/Templates/NoDestroy.h b/Engine/Source/Runtime/Core/Public/Templates/NoDestroy.h
new file mode 100644
--- /dev/null
+++ b/Engine/Source/Runtime/Core/Public/Templates/NoDestroy.h
@@ -0,0 +1,70 @@
+// Copyright Epic Games, Inc. All Rights Reserved.
+
+#pragma once
+
+#include "CoreTypes.h"
+#include "Templates/UnrealTemplate.h"
+
+// Special tag for TNoDestroy consteval constructor to avoid confusion with the inner type having
+// its own tagged consteval constructor
+enum ENoDestroyConstEval { NoDestroyConstEval };
+
+/** 
+ * Type to wrap a global variable and prevent its destructor being called. The destructor of this
+ * type cannot be trivial, but it is empty and compilers should not register an atexit call for the object. 
+ */
+template<typename T>
+union TNoDestroy
+{
+    using ElementType = T;
+
+    template<typename... ArgTypes>
+    [[nodiscard]] explicit consteval TNoDestroy(ENoDestroyConstEval, ArgTypes&&... InArgs)
+        : Value(Forward<ArgTypes>(InArgs)...)
+    {
+    }
+
+    template<typename... ArgTypes>
+    [[nodiscard]] explicit constexpr TNoDestroy(ArgTypes&&... InArgs)
+        : Value(Forward<ArgTypes>(InArgs)...)
+    {
+    }
+
+	UE_NONCOPYABLE(TNoDestroy);
+    
+    constexpr ~TNoDestroy()
+    {
+    }
+
+    constexpr ElementType* operator&()
+    {
+        return &Value;
+    }
+
+    constexpr const ElementType* operator&() const
+    {
+        return &Value;
+    }
+
+    constexpr ElementType& operator*()
+    {
+        return Value;
+    }
+
+    constexpr const ElementType& operator*() const
+    {
+        return Value;
+    }
+
+    ElementType* operator->()
+    {
+        return &Value;
+    }
+
+    const ElementType* operator->() const
+    {
+        return &Value;
+    }
+private:
+    ElementType Value;
+};
diff --git a/Engine/Source/Runtime/Core/Tests/Templates/NoDestroyTest.cpp b/Engine/Source/Runtime/Core/Tests/Templates/NoDestroyTest.cpp
new file mode 100644
--- /dev/null
+++ b/Engine/Source/Runtime/Core/Tests/Templates/NoDestroyTest.cpp
@@ -0,0 +1,57 @@
+// Copyright Epic Games, Inc. All Rights Reserved.
+
+#if WITH_TESTS
+
+#include "Templates/NoDestroy.h"
+#include "Tests/TestHarnessAdapter.h"
+
+#if WITH_LOW_LEVEL_TESTS
+#include "TestCommon/Expectations.h"
+#endif
+
+struct FNoDestroyTestType
+{
+    int32 Value = 0;
+
+    static inline int32 ConstructorCount = 0;
+    static inline int32 DestructorCount = 0;
+
+    FNoDestroyTestType(int32 InValue)
+        : Value(InValue)
+    {
+        ++ConstructorCount;
+    }
+
+    explicit consteval FNoDestroyTestType(EConstEval, int32 InValue)
+        : Value(InValue)
+    {
+    }
+
+    ~FNoDestroyTestType()
+    {
+        ++DestructorCount;
+    }
+};
+
+constinit TNoDestroy<FNoDestroyTestType> NoDestroyTestGlobal{NoDestroyConstEval, ConstEval, 17};
+
+static_assert(sizeof(FNoDestroyTestType) == sizeof(TNoDestroy<FNoDestroyTestType>), "TNoDestroy should be same size as inner type");
+
+TEST_CASE_NAMED(FNoDestroyTest, "System::Core::Templates::NoDestroy", "[Core][Templates][SmokeFilter]")
+{
+    CHECK_EQUALS("Global value is initialized", 17, NoDestroyTestGlobal->Value);
+    CHECK_EQUALS("Address of global is equal to address from overridden address operator", 
+        (void*)std::addressof(NoDestroyTestGlobal), (void*)&NoDestroyTestGlobal);
+
+    {
+        FNoDestroyTestType::ConstructorCount = 0;
+        FNoDestroyTestType::DestructorCount = 0;
+        TNoDestroy<FNoDestroyTestType> Temporary{21};
+
+        CHECK_EQUALS("Inner object constructed in place", 1, FNoDestroyTestType::ConstructorCount);
+        CHECK_EQUALS("Inner object value initialized", 21, Temporary->Value);
+    }
+    CHECK_EQUALS("Inner object not destroyed", 0, FNoDestroyTestType::DestructorCount);
+}
+
+#endif
\ No newline at end of file
