Honor compression speed flags with deflate, gzip & lz4hc compressors

gzip/deflate still defaults to default (6). lz4 default changes from maximum (12) to default (9), still using the slower lz4hc method.

[RN] Compression speed/size flags are now honored during zlib, gzip & LZ4 compression

diff --git a/Engine/Source/Runtime/Core/Private/Misc/Compression.cpp b/Engine/Source/Runtime/Core/Private/Misc/Compression.cpp
--- a/Engine/Source/Runtime/Core/Private/Misc/Compression.cpp
+++ b/Engine/Source/Runtime/Core/Private/Misc/Compression.cpp
@@ -59,6 +59,21 @@ static uint32 appGZIPVersion()
 	return uint32(ZLIB_VERNUM); // we use zlib library for gzip
 }
 
+// Convert ECompressionFlags to the equivalent zlib compression level.
+static int GetZlibCompressionLevel(const ECompressionFlags CompressionFlags)
+{
+	return (CompressionFlags & COMPRESS_BiasSize) ? Z_BEST_COMPRESSION
+		: (CompressionFlags & COMPRESS_BiasSpeed) ? Z_BEST_SPEED
+		: Z_DEFAULT_COMPRESSION;
+}
+
+static int GetLZ4HCCompressionLevel(const ECompressionFlags CompressionFlags)
+{
+	return (CompressionFlags & COMPRESS_BiasSize) ? LZ4HC_CLEVEL_MAX
+		: (CompressionFlags & COMPRESS_BiasSpeed) ? LZ4HC_CLEVEL_MIN
+		: LZ4HC_CLEVEL_DEFAULT;
+}
+
 /**
  * Thread-safe abstract compression routine. Compresses memory from uncompressed buffer and writes it to compressed
  * buffer. Updates CompressedSize with size of compressed data.
@@ -138,7 +153,7 @@ static bool appCompressMemoryZLIB(void* CompressedBuffer, int64& CompressedSize,
 	}
 }
 
-static bool appCompressMemoryGZIP(void* CompressedBuffer, int64& CompressedSize, const void* UncompressedBuffer, int64 UncompressedSize)
+static bool appCompressMemoryGZIP(void* CompressedBuffer, int64& CompressedSize, const void* UncompressedBuffer, int64 UncompressedSize, int CompressionLevel)
 {
 	TRACE_CPUPROFILER_EVENT_SCOPE(appCompressMemoryGZIP);
 	DECLARE_SCOPE_CYCLE_COUNTER( TEXT( "Compress Memory GZIP" ), STAT_appCompressMemoryGZIP, STATGROUP_Compression );
@@ -164,7 +179,7 @@ static bool appCompressMemoryGZIP(void* CompressedBuffer, int64& CompressedSize,
 	int GZIP_ENCODING = 16;
 	deflateInit2(
 		&gzipstream,
-		Z_DEFAULT_COMPRESSION,
+		CompressionLevel,
 		Z_DEFLATED,
 		windowsBits | GZIP_ENCODING,
 		MAX_MEM_LEVEL,
@@ -806,12 +821,12 @@ bool FCompression::CompressMemory(FName FormatName, void* CompressedBuffer, int6
 	if (FormatName == NAME_Zlib)
 	{
 		// hardcoded zlib
-		bCompressSucceeded = appCompressMemoryZLIB(CompressedBuffer, CompressedSize, UncompressedBuffer, UncompressedSize, CompressionData, Z_DEFAULT_COMPRESSION);
+		bCompressSucceeded = appCompressMemoryZLIB(CompressedBuffer, CompressedSize, UncompressedBuffer, UncompressedSize, CompressionData, GetZlibCompressionLevel(Flags));
 	}
 	else if (FormatName == NAME_Gzip)
 	{
 		// hardcoded gzip
-		bCompressSucceeded = appCompressMemoryGZIP(CompressedBuffer, CompressedSize, UncompressedBuffer, UncompressedSize);
+		bCompressSucceeded = appCompressMemoryGZIP(CompressedBuffer, CompressedSize, UncompressedBuffer, UncompressedSize, GetZlibCompressionLevel(Flags));
 	}
 	else if (FormatName == NAME_LZ4)
 	{
@@ -822,7 +837,7 @@ bool FCompression::CompressMemory(FName FormatName, void* CompressedBuffer, int6
 			return false;
 		}
 
-		CompressedSize = LZ4_compress_HC((const char*)UncompressedBuffer, (char*)CompressedBuffer, UncompressedSize, CompressedSize, LZ4HC_CLEVEL_MAX);
+		CompressedSize = LZ4_compress_HC((const char*)UncompressedBuffer, (char*)CompressedBuffer, UncompressedSize, CompressedSize, GetLZ4HCCompressionLevel(Flags));
 		bCompressSucceeded = CompressedSize > 0;
 	}
 	else
