Core: Annotate ReportNotNullPtr() with CA_NO_RETURN

#rnx

diff --git a/Engine/Source/Runtime/Core/Public/Misc/NotNull.h b/Engine/Source/Runtime/Core/Public/Misc/NotNull.h
--- a/Engine/Source/Runtime/Core/Public/Misc/NotNull.h
+++ b/Engine/Source/Runtime/Core/Public/Misc/NotNull.h
@@ -74,7 +74,7 @@
 		template <typename T>
 		inline constexpr bool TSupportsNotOperator_V = TSupportsNotOperator<T>::Value;
 
-		[[noreturn]] CORE_API void ReportNotNullPtr();
+		[[noreturn]] CA_NO_RETURN CORE_API void ReportNotNullPtr();
 		CORE_API void CheckLoadingNotNullPtr(FArchive& Ar);
 #endif // DO_CHECK
 
