[Trace/Insights]
* Sampling Profiler: Set the default stack sampling interval to 125 microseconds (8 kHz).
* TraceAux: Added more CPU scopes for TraceAuxiliary.
* TraceAux: Removed the deprecated FTraceAuxiliary::GetTraceDestination() API.
* TraceAux: Improved error handling for FTraceAuxiliary::Start() and FTraceAuxiliary::Relay() APIs.

diff --git a/Engine/Source/Runtime/Core/Private/ProfilingDebugging/TraceAuxiliary.cpp b/Engine/Source/Runtime/Core/Private/ProfilingDebugging/TraceAuxiliary.cpp
--- a/Engine/Source/Runtime/Core/Private/ProfilingDebugging/TraceAuxiliary.cpp
+++ b/Engine/Source/Runtime/Core/Private/ProfilingDebugging/TraceAuxiliary.cpp
@@ -458,15 +458,15 @@ bool FTraceAuxiliaryImpl::EnableChannel(ChannelType ChannelIdentifier, const FTr
 		else
 		{
 			UE_LOG_REF(
-				LogCategory, 
-				Warning, 
-				TEXT("Cannot enable channel '%s': %s"), 
-				ChannelName.ToString(), 
+				LogCategory,
+				Warning,
+				TEXT("Cannot enable channel '%s': %s"),
+				ChannelName.ToString(),
 				DenyReason
 			);
 		}
 	}
-	
+
 	UE_CLOG(bIsEnabled, LogTrace, Verbose, TEXT("Enabled channel '%s'."), ChannelName.ToString());
 	return bIsEnabled;
 }
@@ -548,58 +548,58 @@ uint32 FTraceAuxiliaryImpl::HashChannelName(const TCHAR* Name)
 ////////////////////////////////////////////////////////////////////////////////
 bool FTraceAuxiliaryImpl::Connect(FTraceAuxiliary::EConnectionType Type, const TCHAR* Parameter, const FTraceAuxiliary::FLogCategoryAlias& LogCategory, uint16 SendFlags)
 {
-	// Connect/write to file, but only if we're not already sending/writing.
-	bool bConnected = UE::Trace::IsTracing();
-	if (!bConnected)
+	// Connect to trace server or write to a file, but only if we're not already tracing.
+	if (UE::Trace::IsTracing())
 	{
-		if (Type == FTraceAuxiliary::EConnectionType::Network)
-		{
-			TRACE_CPUPROFILER_EVENT_SCOPE(FTraceAux_SendToHost);
-			bConnected = SendToHost(Parameter, LogCategory, SendFlags);
-			if (bConnected)
-			{
-				UE_LOG_REF(LogCategory, Display, TEXT("Trace started (connected to trace server %s)."), *GetDest());
-			}
-			else
-			{
-				UE_LOG_REF(LogCategory, Error, TEXT("Trace failed to connect (trace host: %s)!"), Parameter ? Parameter : TEXT(""));
-			}
+		UE_LOG_REF(LogCategory, Warning, TEXT("Already tracing!"));
+		return true;
+	}
 
+	bool bConnected = false;
+	if (Type == FTraceAuxiliary::EConnectionType::Network)
+	{
+		TRACE_CPUPROFILER_EVENT_SCOPE(FTraceAux_SendToHost);
+		bConnected = SendToHost(Parameter, LogCategory, SendFlags);
+		if (bConnected)
+		{
+			UE_LOG_REF(LogCategory, Display, TEXT("Trace started (connected to trace server %s)."), *GetDest());
 		}
-
-		else if (Type == FTraceAuxiliary::EConnectionType::File)
+		else
 		{
-			TRACE_CPUPROFILER_EVENT_SCOPE(FTraceAux_WriteToFile);
-			bConnected = WriteToFile(Parameter, LogCategory, SendFlags);
-			if (bConnected)
-			{
-				UE_LOG_REF(LogCategory, Display, TEXT("Trace started (writing to file \"%s\")."), *GetDest());
-			}
-			else
-			{
-				UE_LOG_REF(LogCategory, Error, TEXT("Trace failed to connect (file: \"%s\")!"), Parameter ? Parameter : TEXT(""));
-			}
-
+			UE_LOG_REF(LogCategory, Error, TEXT("Trace failed to connect (trace host: %s)!"), Parameter ? Parameter : TEXT(""));
 		}
-
+	}
+	else if (Type == FTraceAuxiliary::EConnectionType::File)
+	{
+		TRACE_CPUPROFILER_EVENT_SCOPE(FTraceAux_WriteToFile);
+		bConnected = WriteToFile(Parameter, LogCategory, SendFlags);
 		if (bConnected)
 		{
-			FString StartedDest;
-			FTraceAuxiliary::EConnectionType StartedType = FTraceAuxiliary::EConnectionType::None;
-
-			{
-				FReadScopeLock _(CurrentTargetLock);
-				StartedDest = CurrentTraceTarget.TraceDest;
-				StartedType = CurrentTraceTarget.TraceType;
-			}
-
-			TRACE_CPUPROFILER_EVENT_SCOPE(FTraceAux_OnTraceStarted);
-			FTraceAuxiliary::OnTraceStarted.Broadcast(StartedType, StartedDest);
+			UE_LOG_REF(LogCategory, Display, TEXT("Trace started (writing to file \"%s\")."), *GetDest());
+		}
+		else
+		{
+			UE_LOG_REF(LogCategory, Error, TEXT("Trace failed to connect (file: \"%s\")!"), Parameter ? Parameter : TEXT(""));
 		}
 	}
 	else
 	{
-		UE_LOG_REF(LogCategory, Error, TEXT("Already tracing from unknown source."));
+		UE_LOG_REF(LogCategory, Error, TEXT("Unknown connection type!"));
+	}
+
+	if (bConnected)
+	{
+		FTraceAuxiliary::EConnectionType StartedType = FTraceAuxiliary::EConnectionType::None;
+		FString StartedDest;
+
+		{
+			FReadScopeLock _(CurrentTargetLock);
+			StartedType = CurrentTraceTarget.TraceType;
+			StartedDest = CurrentTraceTarget.TraceDest;
+		}
+
+		TRACE_CPUPROFILER_EVENT_SCOPE(FTraceAux_OnTraceStarted);
+		FTraceAuxiliary::OnTraceStarted.Broadcast(StartedType, StartedDest);
 	}
 
 	return bConnected;
@@ -610,17 +610,17 @@ bool FTraceAuxiliaryImpl::ConnectRelay(UPTRINT Handle, UE::Trace::IoWriteFunc Wr
 {
 	if (UE::Trace::IsTracing())
 	{
-		UE_LOG_REF(LogTrace, Error, TEXT("Already tracing from unknown source."));
+		UE_LOG(LogTrace, Warning, TEXT("Already tracing!"));
 		return true;
 	}
 
 	const bool bConnected = UE::Trace::RelayTo(Handle, WriteFunc, CloseFunc);
 	if (bConnected)
 	{
-		UE_LOG(LogTrace, Display, TEXT("Trace started (connected to relay endpoint)"));
+		UE_LOG(LogTrace, Display, TEXT("Trace started (connected to relay endpoint)."));
 
-		FString StartedDest = FString::Printf(TEXT("Relay endpoint (Handle: 0x%llx)"), Handle);
 		FTraceAuxiliary::EConnectionType StartedType = FTraceAuxiliary::EConnectionType::Relay;
+		FString StartedDest = FString::Printf(TEXT("Relay endpoint (Handle: 0x%llx)"), Handle);
 
 		{
 			FWriteScopeLock _(CurrentTargetLock);
@@ -633,7 +633,7 @@ bool FTraceAuxiliaryImpl::ConnectRelay(UPTRINT Handle, UE::Trace::IoWriteFunc Wr
 	}
 	else
 	{
-		UE_LOG(LogTrace, Error, TEXT("Trace failed to connect relay endpoint"));
+		UE_LOG(LogTrace, Error, TEXT("Trace failed to connect relay endpoint!"));
 	}
 	return bConnected;
 }
@@ -659,9 +659,9 @@ bool FTraceAuxiliaryImpl::Stop()
 
 	{
 		FWriteScopeLock _(CurrentTargetLock);
-		StopedDest = CurrentTraceTarget.TraceDest;
 		StopedType = CurrentTraceTarget.TraceType;
 		CurrentTraceTarget.TraceType = FTraceAuxiliary::EConnectionType::None;
+		StopedDest = CurrentTraceTarget.TraceDest;
 		CurrentTraceTarget.TraceDest.Reset();
 	}
 
@@ -995,43 +995,55 @@ void TraceAuxiliaryEmitStats(const UE::Trace::FStatistics& Stats)
 	uint64 TotalMemoryUsed = Stats.BlockPoolAllocated + Stats.FixedBufferAllocated + Stats.SharedBufferAllocated + Stats.CacheAllocated;
 
 #if UE_TRACEAUX_EMIT_STAT_COUNTERS
-	SET_MEMORY_STAT(STAT_Trace_MemoryUsed, Stats.MemoryUsed);
-	SET_MEMORY_STAT(STAT_Trace_Memory_BlockPool, Stats.BlockPoolAllocated);
-	SET_MEMORY_STAT(STAT_Trace_Memory_FixedBuffers, Stats.FixedBufferAllocated);
-	SET_MEMORY_STAT(STAT_Trace_Memory_SharedBuffers, Stats.SharedBufferAllocated);
-	SET_MEMORY_STAT(STAT_Trace_Memory_Cache, Stats.CacheAllocated);
-	SET_MEMORY_STAT(STAT_Trace_Memory_Cache_Used, Stats.CacheUsed);
-	SET_MEMORY_STAT(STAT_Trace_Memory_Cache_Unused, (int64)Stats.CacheAllocated - (int64)Stats.CacheUsed);
-	SET_MEMORY_STAT(STAT_Trace_Memory_Cache_Wasted, Stats.CacheWaste);
-	SET_MEMORY_STAT(STAT_Trace_Bandwidth_Emitted, Stats.BytesEmitted);
-	SET_MEMORY_STAT(STAT_Trace_Bandwidth_EmittedNotTraced, (int64)Stats.BytesEmitted - (int64)Stats.BytesTraced);
-	SET_MEMORY_STAT(STAT_Trace_Bandwidth_Traced, Stats.BytesTraced);
-	SET_MEMORY_STAT(STAT_Trace_Bandwidth_Sent, Stats.BytesSent);
-	SET_MEMORY_STAT(STAT_Trace_MemoryError, (int64)Stats.MemoryUsed - (int64)TotalMemoryUsed);
+	if (FThreadStats::IsCollectingData())
+	{
+		TRACE_CPUPROFILER_EVENT_SCOPE(FTraceAux_Stats_STAT);
+		SET_MEMORY_STAT(STAT_Trace_MemoryUsed, Stats.MemoryUsed);
+		SET_MEMORY_STAT(STAT_Trace_Memory_BlockPool, Stats.BlockPoolAllocated);
+		SET_MEMORY_STAT(STAT_Trace_Memory_FixedBuffers, Stats.FixedBufferAllocated);
+		SET_MEMORY_STAT(STAT_Trace_Memory_SharedBuffers, Stats.SharedBufferAllocated);
+		SET_MEMORY_STAT(STAT_Trace_Memory_Cache, Stats.CacheAllocated);
+		SET_MEMORY_STAT(STAT_Trace_Memory_Cache_Used, Stats.CacheUsed);
+		SET_MEMORY_STAT(STAT_Trace_Memory_Cache_Unused, (int64)Stats.CacheAllocated - (int64)Stats.CacheUsed);
+		SET_MEMORY_STAT(STAT_Trace_Memory_Cache_Wasted, Stats.CacheWaste);
+		SET_MEMORY_STAT(STAT_Trace_Bandwidth_Emitted, Stats.BytesEmitted);
+		SET_MEMORY_STAT(STAT_Trace_Bandwidth_EmittedNotTraced, (int64)Stats.BytesEmitted - (int64)Stats.BytesTraced);
+		SET_MEMORY_STAT(STAT_Trace_Bandwidth_Traced, Stats.BytesTraced);
+		SET_MEMORY_STAT(STAT_Trace_Bandwidth_Sent, Stats.BytesSent);
+		SET_MEMORY_STAT(STAT_Trace_MemoryError, (int64)Stats.MemoryUsed - (int64)TotalMemoryUsed);
+	}
 #endif
 
 #if UE_TRACEAUX_EMIT_TRACE_COUNTERS
-	TRACE_COUNTER_SET(Trace_MemoryUsed, Stats.MemoryUsed);
-	TRACE_COUNTER_SET(Trace_Memory_BlockPool, Stats.BlockPoolAllocated);
-	TRACE_COUNTER_SET(Trace_Memory_FixedBuffers, Stats.FixedBufferAllocated);
-	TRACE_COUNTER_SET(Trace_Memory_SharedBuffers, Stats.SharedBufferAllocated);
-	TRACE_COUNTER_SET(Trace_Memory_Cache, Stats.CacheAllocated);
-	TRACE_COUNTER_SET(Trace_Memory_Cache_Used, Stats.CacheUsed);
-	TRACE_COUNTER_SET(Trace_Memory_Cache_Unused, (int64)Stats.CacheAllocated - (int64)Stats.CacheUsed);
-	TRACE_COUNTER_SET(Trace_Memory_Cache_Wasted, Stats.CacheWaste);
-	TRACE_COUNTER_SET(Trace_Bandwidth_Emitted, Stats.BytesEmitted);
-	TRACE_COUNTER_SET(Trace_Bandwidth_EmittedNotTraced, (int64)Stats.BytesEmitted - (int64)Stats.BytesTraced);
-	TRACE_COUNTER_SET(Trace_Bandwidth_Traced, Stats.BytesTraced);
-	TRACE_COUNTER_SET(Trace_Bandwidth_Sent, Stats.BytesSent);
-	TRACE_COUNTER_SET(Trace_MemoryError, (int64)Stats.MemoryUsed - (int64)TotalMemoryUsed);
+	if (UE_TRACE_CHANNELEXPR_IS_ENABLED(CountersChannel))
+	{
+		TRACE_CPUPROFILER_EVENT_SCOPE(FTraceAux_Stats_Counters);
+		TRACE_COUNTER_SET(Trace_MemoryUsed, Stats.MemoryUsed);
+		TRACE_COUNTER_SET(Trace_Memory_BlockPool, Stats.BlockPoolAllocated);
+		TRACE_COUNTER_SET(Trace_Memory_FixedBuffers, Stats.FixedBufferAllocated);
+		TRACE_COUNTER_SET(Trace_Memory_SharedBuffers, Stats.SharedBufferAllocated);
+		TRACE_COUNTER_SET(Trace_Memory_Cache, Stats.CacheAllocated);
+		TRACE_COUNTER_SET(Trace_Memory_Cache_Used, Stats.CacheUsed);
+		TRACE_COUNTER_SET(Trace_Memory_Cache_Unused, (int64)Stats.CacheAllocated - (int64)Stats.CacheUsed);
+		TRACE_COUNTER_SET(Trace_Memory_Cache_Wasted, Stats.CacheWaste);
+		TRACE_COUNTER_SET(Trace_Bandwidth_Emitted, Stats.BytesEmitted);
+		TRACE_COUNTER_SET(Trace_Bandwidth_EmittedNotTraced, (int64)Stats.BytesEmitted - (int64)Stats.BytesTraced);
+		TRACE_COUNTER_SET(Trace_Bandwidth_Traced, Stats.BytesTraced);
+		TRACE_COUNTER_SET(Trace_Bandwidth_Sent, Stats.BytesSent);
+		TRACE_COUNTER_SET(Trace_MemoryError, (int64)Stats.MemoryUsed - (int64)TotalMemoryUsed);
+	}
 #endif
 
 #if UE_TRACEAUX_EMIT_LLM_STATS && ENABLE_LOW_LEVEL_MEM_TRACKER
-	FLowLevelMemTracker::Get().SetTagAmountForTracker(ELLMTracker::Default, LLM_TAG_NAME(TraceLog), ELLMTagSet::None, Stats.MemoryUsed, true);
-	FLowLevelMemTracker::Get().SetTagAmountForTracker(ELLMTracker::Default, LLM_TAG_NAME(TraceLog_BlockPool), ELLMTagSet::None, Stats.BlockPoolAllocated, true);
-	FLowLevelMemTracker::Get().SetTagAmountForTracker(ELLMTracker::Default, LLM_TAG_NAME(TraceLog_FixedBuffers), ELLMTagSet::None, Stats.FixedBufferAllocated, true);
-	FLowLevelMemTracker::Get().SetTagAmountForTracker(ELLMTracker::Default, LLM_TAG_NAME(TraceLog_SharedBuffers), ELLMTagSet::None, Stats.SharedBufferAllocated, true);
-	FLowLevelMemTracker::Get().SetTagAmountForTracker(ELLMTracker::Default, LLM_TAG_NAME(TraceLog_Cache), ELLMTagSet::None, Stats.CacheAllocated, true);
+	if (FLowLevelMemTracker::Get().IsEnabled())
+	{
+		TRACE_CPUPROFILER_EVENT_SCOPE(FTraceAux_Stats_LLM);
+		FLowLevelMemTracker::Get().SetTagAmountForTracker(ELLMTracker::Default, LLM_TAG_NAME(TraceLog), ELLMTagSet::None, Stats.MemoryUsed, true);
+		FLowLevelMemTracker::Get().SetTagAmountForTracker(ELLMTracker::Default, LLM_TAG_NAME(TraceLog_BlockPool), ELLMTagSet::None, Stats.BlockPoolAllocated, true);
+		FLowLevelMemTracker::Get().SetTagAmountForTracker(ELLMTracker::Default, LLM_TAG_NAME(TraceLog_FixedBuffers), ELLMTagSet::None, Stats.FixedBufferAllocated, true);
+		FLowLevelMemTracker::Get().SetTagAmountForTracker(ELLMTracker::Default, LLM_TAG_NAME(TraceLog_SharedBuffers), ELLMTagSet::None, Stats.SharedBufferAllocated, true);
+		FLowLevelMemTracker::Get().SetTagAmountForTracker(ELLMTracker::Default, LLM_TAG_NAME(TraceLog_Cache), ELLMTagSet::None, Stats.CacheAllocated, true);
+	}
 #endif
 }
 
@@ -1062,6 +1074,8 @@ void TraceAuxiliaryOnEndFrameCallback()
 {
 	// This is called at end of each frame.
 
+	TRACE_CPUPROFILER_EVENT_SCOPE(FTraceAux_OnEndFrame_Stats);
+
 	// Reset the trace update callback. From now on, we only trace stats once per frame.
 	UE::Trace::SetUpdateCallback(nullptr);
 
@@ -1083,11 +1097,12 @@ void TraceAuxiliaryOnEndFrameCallback()
 		UE::Trace::GetStatistics(Stats);
 #endif
 
-		CSV_CUSTOM_STAT(Trace, MemoryUsedMb,	double(Stats.MemoryUsed) / 1024.0 / 1024.0, ECsvCustomStatOp::Set);
-		CSV_CUSTOM_STAT(Trace, BlockPoolMb,		double(Stats.BlockPoolAllocated) / 1024.0 / 1024.0, ECsvCustomStatOp::Set);
-		CSV_CUSTOM_STAT(Trace, FixedBuffersMb,	double(Stats.FixedBufferAllocated) / 1024.0 / 1024.0, ECsvCustomStatOp::Set);
-		CSV_CUSTOM_STAT(Trace, SharedBuffersMb,	double(Stats.SharedBufferAllocated) / 1024.0 / 1024.0, ECsvCustomStatOp::Set);
-		CSV_CUSTOM_STAT(Trace, CacheMb,			double(Stats.CacheAllocated) / 1024.0 / 1024.0, ECsvCustomStatOp::Set);
+		TRACE_CPUPROFILER_EVENT_SCOPE(FTraceAux_Stats_CSV);
+		CSV_CUSTOM_STAT(Trace, MemoryUsedMb,    double(Stats.MemoryUsed)            / 1024.0 / 1024.0, ECsvCustomStatOp::Set);
+		CSV_CUSTOM_STAT(Trace, BlockPoolMb,     double(Stats.BlockPoolAllocated)    / 1024.0 / 1024.0, ECsvCustomStatOp::Set);
+		CSV_CUSTOM_STAT(Trace, FixedBuffersMb,  double(Stats.FixedBufferAllocated)  / 1024.0 / 1024.0, ECsvCustomStatOp::Set);
+		CSV_CUSTOM_STAT(Trace, SharedBuffersMb, double(Stats.SharedBufferAllocated) / 1024.0 / 1024.0, ECsvCustomStatOp::Set);
+		CSV_CUSTOM_STAT(Trace, CacheMb,         double(Stats.CacheAllocated)        / 1024.0 / 1024.0, ECsvCustomStatOp::Set);
 	}
 #endif
 }
@@ -1745,7 +1760,7 @@ static bool TraceAuxiliaryStartShared(const TCHAR* Channels, const FTraceAuxilia
 bool FTraceAuxiliary::Start(EConnectionType Type, const TCHAR* Target, const TCHAR* Channels, FOptions* Options, const FLogCategoryAlias& LogCategory)
 {
 #if UE_TRACE_ENABLED
-	// Use StartCustom function to start a custom trace
+	// Use the Relay function to start a trace relay.
 	check(Type != EConnectionType::Relay);
 
 	TRACE_CPUPROFILER_EVENT_SCOPE(FTraceAux_Start);
@@ -1784,6 +1799,7 @@ bool FTraceAuxiliary::Relay(UPTRINT Handle, UE::Trace::IoWriteFunc WriteFunc, UE
 	}
 
 	uint16 SendFlags = (Options && Options->bExcludeTail) ? UE::Trace::FSendFlags::ExcludeTail : 0;
+
 	return GTraceAuxiliary.ConnectRelay(Handle, WriteFunc, CloseFunc, SendFlags);
 #else
 	return false;
@@ -1896,7 +1912,7 @@ void FTraceAuxiliary::Initialize(const TCHAR* CommandLine)
 	// Auto launch Unreal Trace Server for certain configurations
 	if (!(FParse::Param(CommandLine, TEXT("notraceserver")) || FParse::Param(CommandLine, TEXT("buildmachine"))))
 	{
-		TRACE_CPUPROFILER_EVENT_SCOPE(FTraceAux_LaunchUnrealTrace);
+		TRACE_CPUPROFILER_EVENT_SCOPE(FTraceAux_Launch_UTS);
 		FTraceServerControls::Start();
 	}
 #endif
@@ -1985,6 +2001,7 @@ void FTraceAuxiliary::Initialize(const TCHAR* CommandLine)
 
 	// Emit empty stats once (to ensure all stats/counters start from zero).
 	{
+		TRACE_CPUPROFILER_EVENT_SCOPE(FTraceAux_InitialStats);
 		UE::Trace::FStatistics Stats;
 		TraceAuxiliaryEmitStats(Stats);
 	}
@@ -2054,29 +2071,50 @@ void FTraceAuxiliary::Initialize(const TCHAR* CommandLine)
 		}
 	}
 
-	UE::Trace::Initialize(Desc);
+	{
+		TRACE_CPUPROFILER_EVENT_SCOPE_STR("UE::Trace::Initialize");
+
+		UE::Trace::Initialize(Desc);
+	}
 
 	// Workaround for the fact that even if StartFromCommandlineArguments will enable channels
 	// specified by the commandline, UE::Trace::Initialize will reset all channels.
 	GTraceAuxiliary.EnableCommandlineChannelsPostInitialize();
 
+	check(IsInGameThread());
+	UE::Trace::ThreadRegister(TEXT("GameThread"), FPlatformTLS::GetCurrentThreadId(), -1);
+
 	// Setup known on connection callbacks
 	OnConnection.AddStatic(FStringTrace::OnConnection);
 
 	// Register end frame callbacks
 	GTraceAuxiliary.RegisterEndFrameCallbacks();
 
-	// Initialize callstack tracing. If for any reason memory tracing has been initialized
-	// without callstack tracing we need to use the same allocator. Otherwise use regular malloc.
-	FMalloc* TraceAllocator = MemoryTrace_GetAllocator();
-	CallstackTrace_Create(TraceAllocator ? TraceAllocator : UE::Private::GMalloc);
-	CallstackTrace_Initialize();
+	// Initialize Callstack tracing.
+	{
+		TRACE_CPUPROFILER_EVENT_SCOPE(FTraceAux_CallstackTrace_Create);
+
+		// If for any reason, memory tracing has been initialized without callstack tracing,
+		// we need to use the same allocator; otherwise we use the regular malloc.
+		FMalloc* TraceAllocator = MemoryTrace_GetAllocator();
+		CallstackTrace_Create(TraceAllocator ? TraceAllocator : UE::Private::GMalloc);
+	}
+	{
+		TRACE_CPUPROFILER_EVENT_SCOPE(FTraceAux_CallstackTrace_Init);
+
+		CallstackTrace_Initialize();
+	}
 
-	// By default use 1 msec for stack sampling interval.
-	uint32 Microseconds = 1000;
-	FParse::Value(CommandLine, TEXT("-samplinginterval="), Microseconds);
-	FPlatformEventsTrace::Init(Microseconds);
-	FPlatformEventsTrace::PostInit();
+	// Initialize Platform Events tracing.
+	{
+		TRACE_CPUPROFILER_EVENT_SCOPE(FTraceAux_PlatformEventsTrace_Init);
+
+		// By default use 125 microseconds (8 kHz) for the stack sampling interval.
+		uint32 Microseconds = 125;
+		FParse::Value(CommandLine, TEXT("-samplinginterval="), Microseconds);
+		FPlatformEventsTrace::Init(Microseconds);
+		FPlatformEventsTrace::PostInit();
+	}
 
 	if (GTraceAutoStart)
 	{
@@ -2089,8 +2127,6 @@ void FTraceAuxiliary::Initialize(const TCHAR* CommandLine)
 		});
 	}
 
-	UE::Trace::ThreadRegister(TEXT("GameThread"), FPlatformTLS::GetCurrentThreadId(), -1);
-
 	GTraceAuxiliary.FreezeReadOnlyChannels();
 	UE_LOG(LogTrace, Log, TEXT("Finished trace initialization."));
 #endif //UE_TRACE_ENABLED
@@ -2160,16 +2196,6 @@ void FTraceAuxiliary::DisableChannels(const TCHAR* Channels)
 #endif
 }
 
-const TCHAR* FTraceAuxiliary::GetTraceDestination()
-{
-#if UE_TRACE_ENABLED
-	static FString TempUnsafe = GTraceAuxiliary.GetDest();
-	return *TempUnsafe;
-#else
-	return nullptr;
-#endif
-}
-
 FString FTraceAuxiliary::GetTraceDestinationString()
 {
 #if UE_TRACE_ENABLED
diff --git a/Engine/Source/Runtime/Core/Public/ProfilingDebugging/TraceAuxiliary.h b/Engine/Source/Runtime/Core/Public/ProfilingDebugging/TraceAuxiliary.h
--- a/Engine/Source/Runtime/Core/Public/ProfilingDebugging/TraceAuxiliary.h
+++ b/Engine/Source/Runtime/Core/Public/ProfilingDebugging/TraceAuxiliary.h
@@ -20,7 +20,7 @@
 class FTraceAuxiliary
 {
 public:
-	/** 
+	/**
 	* This enum is serialized and sent via the trace service.
 	* Do not change the values or modify the order. Only add new values to the end.
 	* Should be kept in sync with FTraceStatus::ETraceSystemStatus from ITraceController.h
@@ -214,14 +214,16 @@ public:
 	/**
 	 * Enable channels to emit events belonging to this category.
 	 * @note Note that presets cannot be used when specifying channel ids.
-	 * @param ChannelIds List of channels to enable. 
+	 * @param ChannelIds List of channels to enable.
+	 * @param OutErrors Contains deny reasons for channels that could not be enabled.
 	 */
 	static CORE_API void EnableChannels(TConstArrayView<uint32> ChannelIds, TMap<uint32, FString>* OutErrors = nullptr);
 
 	/**
 	 * Disable channels to mute events belonging to this category.
 	 * @note Note that presets cannot be used when specifying channel ids.
-	 * @param ChannelIds List of channels to disable. 
+	 * @param ChannelIds List of channels to disable.
+	 * @param OutErrors Contains deny reasons for channels that could not be disabled.
 	 */
 	static CORE_API void DisableChannels(TConstArrayView<uint32> ChannelIds, TMap<uint32, FString>* OutErrors = nullptr);
 
@@ -231,13 +233,6 @@ public:
 	 */
 	static CORE_API void DisableChannels(const TCHAR* Channels = nullptr);
 
-	/**
-	 *  Returns the destination string that is currently being traced to.
-	 *  Contains either a file path or network address. Points to an empty string if tracing is disabled.
-	 */
-	UE_DEPRECATED(5.3, "Use GetTraceDestinationString instead.")
-	static CORE_API const TCHAR* GetTraceDestination();
-
 	/**
 	 *  Returns the destination string that is currently being traced to.
 	 *  Contains either a file path or network address. Empty if tracing is disabled.
@@ -245,12 +240,12 @@ public:
 	static CORE_API FString GetTraceDestinationString();
 
 	/**
-	 *  Returns whether the trace system is currently connected to a trace sink (file or network)
+	 *  Returns whether the trace system is currently connected to a trace sink (network, file or custom relay).
 	 */
 	static CORE_API bool IsConnected();
 
 	/**
-	 * Returns whether the trace system is currently connected. If connected writes the session/trace identifiers.
+	 * Returns whether the trace system is currently connected. If connected, it writes the session/trace identifiers.
 	 * @param OutSessionGuid If connected, the session guid
 	 * @param OutTraceGuid If connected, the trace guid
 	 * @return True if connected
