Add function to enable Trace channels by name.

#rnx

diff --git a/Engine/Source/Runtime/Core/Private/ProfilingDebugging/TraceAuxiliary.cpp b/Engine/Source/Runtime/Core/Private/ProfilingDebugging/TraceAuxiliary.cpp
--- a/Engine/Source/Runtime/Core/Private/ProfilingDebugging/TraceAuxiliary.cpp
+++ b/Engine/Source/Runtime/Core/Private/ProfilingDebugging/TraceAuxiliary.cpp
@@ -2189,6 +2189,13 @@ void FTraceAuxiliary::DisableChannels(TConstArrayView<uint32> ChannelIds, TMap<u
 #endif
 }
 
+void FTraceAuxiliary::EnableChannels(const TCHAR* Channels)
+{
+#if UE_TRACE_ENABLED
+	GTraceAuxiliary.EnableChannels(Channels, LogTrace);
+#endif
+}
+
 void FTraceAuxiliary::DisableChannels(const TCHAR* Channels)
 {
 #if UE_TRACE_ENABLED
diff --git a/Engine/Source/Runtime/Core/Public/ProfilingDebugging/TraceAuxiliary.h b/Engine/Source/Runtime/Core/Public/ProfilingDebugging/TraceAuxiliary.h
--- a/Engine/Source/Runtime/Core/Public/ProfilingDebugging/TraceAuxiliary.h
+++ b/Engine/Source/Runtime/Core/Public/ProfilingDebugging/TraceAuxiliary.h
@@ -227,9 +227,15 @@ public:
 	 */
 	static CORE_API void DisableChannels(TConstArrayView<uint32> ChannelIds, TMap<uint32, FString>* OutErrors = nullptr);
 
+	/**
+	 * Enable channels to emit events from this category.
+	 * @param Channels Comma separated list of channels and presets to enable.
+	 */
+	static CORE_API void EnableChannels(const TCHAR* Channels);
+
 	/**
 	 * Disable channels to stop recording traces with them.
-	 * @param Channels List of channels (or a preset) to disable. If null it will disable all active channels.
+	 * @param Channels Comma separated list of channels and presets to disable. If null it will disable all active channels.
 	 */
 	static CORE_API void DisableChannels(const TCHAR* Channels = nullptr);
 
