[Insights] Fixed "Unknown connection type!" CIS error.

#rnx

diff --git a/Engine/Source/Runtime/Core/Private/ProfilingDebugging/TraceAuxiliary.cpp b/Engine/Source/Runtime/Core/Private/ProfilingDebugging/TraceAuxiliary.cpp
--- a/Engine/Source/Runtime/Core/Private/ProfilingDebugging/TraceAuxiliary.cpp
+++ b/Engine/Source/Runtime/Core/Private/ProfilingDebugging/TraceAuxiliary.cpp
@@ -582,9 +582,13 @@ bool FTraceAuxiliaryImpl::Connect(FTraceAuxiliary::EConnectionType Type, const T
 			UE_LOG_REF(LogCategory, Error, TEXT("Trace failed to connect (file: \"%s\")!"), Parameter ? Parameter : TEXT(""));
 		}
 	}
+	else if (Type == FTraceAuxiliary::EConnectionType::None)
+	{
+		UE_LOG_REF(LogCategory, Log, TEXT("No trace connection. Tracing to in-memory cache."));
+	}
 	else
 	{
-		UE_LOG_REF(LogCategory, Error, TEXT("Unknown connection type!"));
+		UE_LOG_REF(LogCategory, Error, TEXT("Unknown trace connection type (%u)!"), uint32(Type));
 	}
 
 	if (bConnected)
