Removed FTraceMalloc::ShouldTrace() API deprecated since UE 5.1.

diff --git a/Engine/Source/Runtime/Core/Private/ProfilingDebugging/MemoryTrace.cpp b/Engine/Source/Runtime/Core/Private/ProfilingDebugging/MemoryTrace.cpp
--- a/Engine/Source/Runtime/Core/Private/ProfilingDebugging/MemoryTrace.cpp
+++ b/Engine/Source/Runtime/Core/Private/ProfilingDebugging/MemoryTrace.cpp
@@ -885,13 +885,3 @@ void FTraceMalloc::Free(void* Original)
 	WrappedMalloc->Free(Original);
 #endif //UE_MEMORY_TRACE_ENABLED
 }
-
-/////////////////////////////////////////////////////////////////////////////
-bool FTraceMalloc::ShouldTrace()
-{
-#if UE_MEMORY_TRACE_ENABLED
-	return !GDoNotAllocateInTrace;
-#else
-	return true;
-#endif //UE_MEMORY_TRACE_ENABLED
-}
diff --git a/Engine/Source/Runtime/Core/Public/ProfilingDebugging/TraceMalloc.h b/Engine/Source/Runtime/Core/Public/ProfilingDebugging/TraceMalloc.h
--- a/Engine/Source/Runtime/Core/Public/ProfilingDebugging/TraceMalloc.h
+++ b/Engine/Source/Runtime/Core/Public/ProfilingDebugging/TraceMalloc.h
@@ -1,5 +1,7 @@
 ﻿// Copyright Epic Games, Inc. All Rights Reserved.
+
 #pragma once
+
 #include "MemoryTrace.h"
 #include "HAL/MemoryBase.h"
 
@@ -8,14 +10,11 @@ class FTraceMalloc : public FMalloc
 public:
 	FTraceMalloc(FMalloc* InMalloc);
 	virtual ~FTraceMalloc();
-	
+
 	virtual void* Malloc(SIZE_T Count, uint32 Alignment) override;
 	virtual void* Realloc(void* Original, SIZE_T Count, uint32 Alignment) override;
 	virtual void Free(void* Original) override;
 
-	UE_DEPRECATED(5.1, "ShouldTrace is deprecated. Details should be internal now.")
-	static bool ShouldTrace();
-
 	virtual void OnMallocInitialized() override
 	{
 		WrappedMalloc->OnMallocInitialized();
@@ -33,4 +32,3 @@ public:
 
 	FMalloc* WrappedMalloc;
 };
-
