* Stats: Removed code switched off by UE_STATS_MEMORY_PROFILER_ENABLED (deprecated since UE 5.3).
* Stats: Removed ProcessMetaDataOnly() and ReadFNamesAndMetadataMessages() deprecated since UE 5.2.

diff --git a/Engine/Source/Runtime/Core/Private/Stats/StatsCommand.cpp b/Engine/Source/Runtime/Core/Private/Stats/StatsCommand.cpp
--- a/Engine/Source/Runtime/Core/Private/Stats/StatsCommand.cpp
+++ b/Engine/Source/Runtime/Core/Private/Stats/StatsCommand.cpp
@@ -32,7 +32,6 @@
 
 #include "Stats/StatsData.h"
 #include "Stats/StatsFile.h"
-#include "Stats/StatsMallocProfilerProxy.h"
 
 DECLARE_CYCLE_STAT(TEXT("Hitch Scan"),STAT_HitchScan,STATGROUP_StatSystem);
 DECLARE_CYCLE_STAT(TEXT("HUD Group"),STAT_HUDGroup,STATGROUP_StatSystem);
@@ -2070,18 +2069,6 @@ static void StatCmd(FString InCmd, bool bStatCommand, FOutputDevice* Ar /*= null
 			FCommandStatsFile::Get().Stop();
 			FThreadStats::DisableRawStats();
 
-#if UE_STATS_MEMORY_PROFILER_ENABLED
-			if (FStatsMallocProfilerProxy::HasMemoryProfilerToken())
-			{
-				if (FStatsMallocProfilerProxy::Get()->GetState())
-				{
-					// Disable memory profiler and restore default stats groups.
-					FStatsMallocProfilerProxy::Get()->SetState(false);
-					IStatGroupEnableManager::Get().StatGroupEnableManagerCommand(TEXT("default"));
-				}
-			}
-#endif //UE_STATS_MEMORY_PROFILER_ENABLED
-
 			Stats.ResetStatsForRawStats();
 
 			// stopfile command happens when some threads shutdown, and depending on order of shutdown operations, the taskgraph may
diff --git a/Engine/Source/Runtime/Core/Private/Stats/StatsData.cpp b/Engine/Source/Runtime/Core/Private/Stats/StatsData.cpp
--- a/Engine/Source/Runtime/Core/Private/Stats/StatsData.cpp
+++ b/Engine/Source/Runtime/Core/Private/Stats/StatsData.cpp
@@ -823,17 +823,6 @@ void FStatsThreadState::ScanForAdvance(FStatPacketArray& NewData)
 	INC_DWORD_STAT_BY(STAT_StatFrameMessages, Count);
 }
 
-void FStatsThreadState::ProcessMetaDataOnly(TArray<FStatMessage>& Data)
-{
-	for (int32 Index = 0; Index < Data.Num() ; Index++)
-	{
-		FStatMessage& Item = Data[Index];
-		EStatOperation::Type Op = Item.NameAndInfo.GetField<EStatOperation>();
-		check(Op == EStatOperation::SetLongName);
-		FindOrAddMetaData(Item);
-	}
-}
-
 void FStatsThreadState::ProcessMetaDataOnly(TArray64<FStatMessage>& Data)
 {
 	for (int64 Index = 0; Index < Data.Num(); Index++)
@@ -873,12 +862,6 @@ void FStatsThreadState::ProcessNonFrameStats(FStatMessagesArray& Data, TSet<FNam
 			{
 				UE_LOG(LogStats, Fatal, TEXT( "Stat %s was not cleared every frame, but was used with a scope cycle counter." ), *Item.NameAndInfo.GetRawName().ToString() );
 			}
-#if UE_STATS_MEMORY_PROFILER_ENABLED
-			else if (Op == EStatOperation::Memory)
-			{
-				// Ignore any memory messages, they shouldn't be treated as regular stats messages.
-			}
-#endif //UE_STATS_MEMORY_PROFILER_ENABLED
 			else if (Op == EStatOperation::SpecialMessageMarker)
 			{
 				// Ignore special messages, they shouldn't be treated as regular stats messages.
@@ -1530,12 +1513,6 @@ void FStatsThreadState::GetRawStackStats(int64 TargetFrame, FRawStatStackNode& R
 
 					}
 				}
-#if UE_STATS_MEMORY_PROFILER_ENABLED
-				else if (Op == EStatOperation::Memory)
-				{
-					// Should never happen.
-				}
-#endif //UE_STATS_MEMORY_PROFILER_ENABLED
 				else if (OutNonStackStats)
 				{
 					FStatsUtils::AddNonStackStats( LongName, Item, Op, ThisFrameNonStackStats );
@@ -2057,13 +2034,6 @@ void FStatsUtils::AccumulateStat(FStatMessage& Dest, FStatMessage const& Item, E
 				case EStatOperation::MaxVal:
 					StatOpMaxVal_Int64( Dest.NameAndInfo, Dest.GetValue_int64(), Item.GetValue_int64() );
 					break;
-
-#if UE_STATS_MEMORY_PROFILER_ENABLED
-				// Nothing here at this moment.
-				case EStatOperation::Memory:
-					break;
-#endif //UE_STATS_MEMORY_PROFILER_ENABLED
-
 				default:
 					check(0);
 			}
@@ -2087,13 +2057,6 @@ void FStatsUtils::AccumulateStat(FStatMessage& Dest, FStatMessage const& Item, E
 				case EStatOperation::MaxVal:
 					Dest.GetValue_double() = FMath::Max<double>(Dest.GetValue_double(), Item.GetValue_double());
 					break;
-
-#if UE_STATS_MEMORY_PROFILER_ENABLED
-				// Nothing here at this moment.
-				case EStatOperation::Memory:
-					break;
-#endif //UE_STATS_MEMORY_PROFILER_ENABLED
-
 				default:
 					check(0);
 			}
diff --git a/Engine/Source/Runtime/Core/Private/Stats/StatsFile.cpp b/Engine/Source/Runtime/Core/Private/Stats/StatsFile.cpp
--- a/Engine/Source/Runtime/Core/Private/Stats/StatsFile.cpp
+++ b/Engine/Source/Runtime/Core/Private/Stats/StatsFile.cpp
@@ -907,73 +907,6 @@ void FStatsReadFile::ProcessStats()
 							StackState->Current = RawName;
 							ProcessCycleScopeStartOperation( Message, *StackState );
 						}
-#if UE_STATS_MEMORY_PROFILER_ENABLED
-						else if (Op == EStatOperation::Memory)
-						{
-							// First memory operation is Alloc or Free
-							const uint64 EncodedPtr = Message.GetValue_Ptr();
-							const EMemoryOperation MemOp = EMemoryOperation( EncodedPtr & (uint64)EMemoryOperation::Mask );
-							const uint64 Ptr = EncodedPtr & ~(uint64)EMemoryOperation::Mask;
-							if (MemOp == EMemoryOperation::Alloc)
-							{
-								// @see FStatsMallocProfilerProxy::TrackAlloc
-								// After AllocPtr message there is always alloc size message and the sequence tag.
-								Index++; CurrentStatMessageIndex++;
-								const FStatMessage& AllocSizeMessage = Data[Index];
-								const int64 AllocSize = AllocSizeMessage.GetValue_int64();
-
-								// Read OperationSequenceTag.
-								Index++; CurrentStatMessageIndex++;
-								const FStatMessage& SequenceTagMessage = Data[Index];
-								const uint32 SequenceTag = (uint32)SequenceTagMessage.GetValue_int64();
-
-								//ThreadStats->AddMemoryMessage( GET_STATFNAME( STAT_Memory_AllocPtr ), (uint64)(UPTRINT)Ptr | (uint64)EMemoryOperation::Alloc );
-								//ThreadStats->AddMemoryMessage( GET_STATFNAME( STAT_Memory_AllocSize ), Size );
-								//ThreadStats->AddMemoryMessage( GET_STATFNAME( STAT_Memory_OperationSequenceTag ), (int64)SequenceTag );
-								ProcessMemoryOperation( MemOp, Ptr, 0, AllocSize, SequenceTag, *StackState );
-							}
-							else if (MemOp == EMemoryOperation::Realloc)
-							{
-								const uint64 OldPtr = Ptr;
-
-								// Read NewPtr
-								Index++; CurrentStatMessageIndex++;
-								const FStatMessage& AllocPtrMessage = Data[Index];
-								const uint64 NewPtr = AllocPtrMessage.GetValue_Ptr() & ~(uint64)EMemoryOperation::Mask;
-
-								// After AllocPtr message there is always alloc size message and the sequence tag.
-								Index++; CurrentStatMessageIndex++;
-								const FStatMessage& ReallocSizeMessage = Data[Index];
-								const int64 ReallocSize = ReallocSizeMessage.GetValue_int64();
-
-								// Read OperationSequenceTag.
-								Index++; CurrentStatMessageIndex++;
-								const FStatMessage& SequenceTagMessage = Data[Index];
-								const uint32 SequenceTag = (uint32)SequenceTagMessage.GetValue_int64();
-
-								//ThreadStats->AddMemoryMessage( GET_STATFNAME( STAT_Memory_FreePtr ), (uint64)(UPTRINT)OldPtr | (uint64)EMemoryOperation::Realloc );
-								//ThreadStats->AddMemoryMessage( GET_STATFNAME( STAT_Memory_AllocPtr ), (uint64)(UPTRINT)NewPtr | (uint64)EMemoryOperation::Realloc );
-								//ThreadStats->AddMemoryMessage( GET_STATFNAME( STAT_Memory_AllocSize ), NewSize );
-								//ThreadStats->AddMemoryMessage( GET_STATFNAME( STAT_Memory_OperationSequenceTag ), (int64)SequenceTag );
-								ProcessMemoryOperation( MemOp, OldPtr, NewPtr, ReallocSize, SequenceTag, *StackState );
-							}
-							else if (MemOp == EMemoryOperation::Free)
-							{
-								// Read OperationSequenceTag.
-								Index++; CurrentStatMessageIndex++;
-								const FStatMessage& SequenceTagMessage = Data[Index];
-								const uint32 SequenceTag = (uint32)SequenceTagMessage.GetValue_int64();
-
-								//ThreadStats->AddMemoryMessage( GET_STATFNAME( STAT_Memory_FreePtr ), (uint64)(UPTRINT)Ptr | (uint64)EMemoryOperation::Free );	// 16 bytes total				
-								//ThreadStats->AddMemoryMessage( GET_STATFNAME( STAT_Memory_OperationSequenceTag ), (int64)SequenceTag );
-								ProcessMemoryOperation( MemOp, Ptr, 0, 0, SequenceTag, *StackState );
-							}
-							else
-							{
-								UE_LOG( LogStats, Warning, TEXT( "Pointer from a memory operation is invalid" ) );
-							}
-						}
-#endif //UE_STATS_MEMORY_PROFILER_ENABLED
 						// Set, Clear, Add, Subtract
 						else if (Op == EStatOperation::CycleScopeEnd)
 						{
diff --git a/Engine/Source/Runtime/Core/Private/Stats/StatsMallocProfilerProxy.cpp b/Engine/Source/Runtime/Core/Private/Stats/StatsMallocProfilerProxy.cpp
deleted file mode 100644
diff --git a/Engine/Source/Runtime/Core/Private/Stats/StatsSystemTypes.cpp b/Engine/Source/Runtime/Core/Private/Stats/StatsSystemTypes.cpp
--- a/Engine/Source/Runtime/Core/Private/Stats/StatsSystemTypes.cpp
+++ b/Engine/Source/Runtime/Core/Private/Stats/StatsSystemTypes.cpp
@@ -227,8 +227,6 @@ DECLARE_MEMORY_STAT(TEXT("Stats Descriptions"), STAT_StatDescMemory, STATGROUP_S
 
 #if STATS
 
-#include "Stats/StatsMallocProfilerProxy.h"
-
 TStatIdData TStatId::TStatId_NAME_None;
 
 /*-----------------------------------------------------------------------------
@@ -1255,17 +1253,6 @@ void FThreadStats::CheckForCollectingStartupStats()
 		UE::Stats::DirectStatsCommand(TEXT("stat dumpsum -start"), true);
 	}
 
-#if UE_STATS_MEMORY_PROFILER_ENABLED
-	// Now we can safely enable malloc profiler.
-	if (FStatsMallocProfilerProxy::HasMemoryProfilerToken())
-	{
-		// Enable all available groups and enable malloc profiler.
-		IStatGroupEnableManager::Get().StatGroupEnableManagerCommand(TEXT("all"));
-		FStatsMallocProfilerProxy::Get()->SetState(true);
-		UE::Stats::DirectStatsCommand(TEXT("stat startfileraw"), true);
-	}
-#endif //UE_STATS_MEMORY_PROFILER_ENABLED
-
 	STAT_ADD_CUSTOMMESSAGE_NAME(STAT_NamedMarker, TEXT("CheckForCollectingStartupStats"));
 }
 
diff --git a/Engine/Source/Runtime/Core/Public/Stats/StatsData.h b/Engine/Source/Runtime/Core/Public/Stats/StatsData.h
--- a/Engine/Source/Runtime/Core/Public/Stats/StatsData.h
+++ b/Engine/Source/Runtime/Core/Public/Stats/StatsData.h
@@ -511,10 +511,6 @@ class FStatsThreadState
 	CORE_API void ScanForAdvance(FStatPacketArray& NewData);
 
 public:
-	/** Internal method to update the internal metadata. **/
-	UE_DEPRECATED(5.2, "Please use the overload that takes TArray64.")
-	CORE_API void ProcessMetaDataOnly(TArray<FStatMessage>& Data);
-
 	CORE_API void ProcessMetaDataOnly(TArray64<FStatMessage>& Data);
 
 	/** Toggles tracking the most memory expensive stats. */
diff --git a/Engine/Source/Runtime/Core/Public/Stats/StatsFile.h b/Engine/Source/Runtime/Core/Public/Stats/StatsFile.h
--- a/Engine/Source/Runtime/Core/Public/Stats/StatsFile.h
+++ b/Engine/Source/Runtime/Core/Public/Stats/StatsFile.h
@@ -803,26 +803,6 @@ public:
 	 *	Reads FNames and metadata messages from the specified archive, only valid for finalized stats files.
 	 *	Allow unordered file access.
 	 */
-	UE_DEPRECATED(5.2, "Please use the overload that takes TArray64.")
-	void ReadFNamesAndMetadataMessages( FArchive& Ar, TArray<FStatMessage>& out_MetadataMessages )
-	{
-		// Read FNames.
-		Ar.Seek( Header.FNameTableOffset );
-		out_MetadataMessages.Reserve( (int32)Header.NumFNames );
-		for( int32 Index = 0; Index < Header.NumFNames; Index++ )
-		{
-			ReadFName( Ar, false );
-		}
-
-		// Read metadata messages.
-		Ar.Seek( Header.MetadataMessagesOffset );
-		out_MetadataMessages.Reserve((int32)Header.NumMetadataMessages );
-		for( int32 Index = 0; Index < Header.NumMetadataMessages; Index++ )
-		{
-			out_MetadataMessages.Add( ReadMessage( Ar, false ) );
-		}
-	}
-
 	void ReadFNamesAndMetadataMessages(FArchive& Ar, TArray64<FStatMessage>& out_MetadataMessages)
 	{
 		// Read FNames.
@@ -1054,29 +1034,6 @@ protected:
 	virtual void ProcessSpecialMessageMarkerOperation( const FStatMessage& Message, const FStackState& StackState )
 	{}
 
-	// #Stats: 2015-07-13 Not implemented yet.
-// 	/** Processes set operation. */
-// 	virtual void ProcessSetOperation( const FStatMessage& Message, const FStackState& StackState )
-// 	{}
-// 
-// 	/** Processes clear operation. */
-// 	virtual void ProcessClearOperation( const FStatMessage& Message, const FStackState& StackState )
-// 	{}
-// 
-// 	/** Processes add operation. */
-// 	virtual void ProcessAddOperation( const FStatMessage& Message, const FStackState& StackState )
-// 	{}
-// 
-// 	/** Processes subtract operation. */
-// 	virtual void ProcessSubtractOperation( const FStatMessage& Message, const FStackState& StackState )
-// 	{}
-
-#if UE_STATS_MEMORY_PROFILER_ENABLED
-	/** Processes memory operation. @see EMemoryOperation. */
-	virtual void ProcessMemoryOperation( EMemoryOperation MemOp, uint64 Ptr, uint64 NewPtr, int64 Size, uint32 SequenceTag, const FStackState& StackState )
-	{}
-#endif
-
 	/** Sets a new processing stage for this file. */
 	void SetProcessingStage( EStatsProcessingStage NewStage )
 	{
diff --git a/Engine/Source/Runtime/Core/Public/Stats/StatsMallocProfilerProxy.h b/Engine/Source/Runtime/Core/Public/Stats/StatsMallocProfilerProxy.h
deleted file mode 100644
diff --git a/Engine/Source/Runtime/Core/Public/Stats/StatsSystemTypes.h b/Engine/Source/Runtime/Core/Public/Stats/StatsSystemTypes.h
--- a/Engine/Source/Runtime/Core/Public/Stats/StatsSystemTypes.h
+++ b/Engine/Source/Runtime/Core/Public/Stats/StatsSystemTypes.h
@@ -1356,16 +1356,6 @@ public:
 		}
 	}
 
-	/** Pseudo-Memory operation. */
-	template<typename TValue>
-	UE_DEPRECATED(5.3, "Use Trace/MemoryInsights and/or LLM for memory profiling.")
-		UE_FORCEINLINE_HINT void AddMemoryMessage(FName InStatName, TValue Value)
-	{
-#if UE_STATS_MEMORY_PROFILER_ENABLED
-		AddStatMessage(FStatMessage(InStatName, EStatOperation::Memory, Value, false));
-#endif //UE_STATS_MEMORY_PROFILER_ENABLED
-	}
-
 	/**
 	 * Used to force a flush at the next available opportunity. This is not useful for threads other than the main and render thread.
 	 * if DiscardCallstack is true, we also dump call stacks, making the next available opportunity at the next stat or stat close.
diff --git a/Engine/Source/Runtime/Launch/Private/LaunchEngineLoop.cpp b/Engine/Source/Runtime/Launch/Private/LaunchEngineLoop.cpp
--- a/Engine/Source/Runtime/Launch/Private/LaunchEngineLoop.cpp
+++ b/Engine/Source/Runtime/Launch/Private/LaunchEngineLoop.cpp
@@ -46,7 +46,6 @@
 #include "Serialization/CompactBinary.h"
 #include "Serialization/CompactBinarySerialization.h"
 #include "Serialization/CompactBinaryWriter.h"
-#include "Stats/StatsMallocProfilerProxy.h"
 #include "Stats/StatsSystem.h"
 #include "Stats/ThreadIdleStats.h"
 #include "Trace/Trace.inl"
