Removed UE_STATS_MEMORY_PROFILER_ENABLED (deprecated since UE 5.3) and UE_DEPRECATED_PROFILER_ENABLED (Profiler deprecated since UE 5.0).

diff --git a/Engine/Source/Runtime/Core/Public/Misc/Build.h b/Engine/Source/Runtime/Core/Public/Misc/Build.h
--- a/Engine/Source/Runtime/Core/Public/Misc/Build.h
+++ b/Engine/Source/Runtime/Core/Public/Misc/Build.h
@@ -517,19 +517,6 @@
 	#define WITH_PROCESS_PRIORITY_CONTROL 0
 #endif
 
-// Controls whether or not MemoryProfiler is enabled in STATS system.
-// This functionality is deprecated in UE 5.3.
-// For memory profiling, use instead Trace/MemoryInsights and/or LLM.
-#ifndef UE_STATS_MEMORY_PROFILER_ENABLED
-	#define UE_STATS_MEMORY_PROFILER_ENABLED 0
-#endif
-
-// Controls whether old Profiler (UnrealFrontend/SessionFrontend/Profiler) is enabled or not.
-// Old Profiler is deprecated since UE 5.0. Use Trace/UnrealInsights instead.
-#ifndef UE_DEPRECATED_PROFILER_ENABLED
-	#define UE_DEPRECATED_PROFILER_ENABLED 0
-#endif
-
 // A compile time flag used to enable support for disabling actor ticking and calls to user callbacks.
 // This functionality is not intended for general use and should be used with care.
 #ifndef UE_SUPPORT_FOR_ACTOR_TICK_DISABLE
