StringBuilder: Added TExternalStringBuilder and support for TStringBuilder<0>

TExternalStringBuilder can be used for cases that want to write to an externally-owned buffer.
TStringBuilder<0> can be used for cases that legitimately do not want an inline buffer.

diff --git a/Engine/Source/Runtime/Core/Private/Misc/StringBuilder.cpp b/Engine/Source/Runtime/Core/Private/Misc/StringBuilder.cpp
--- a/Engine/Source/Runtime/Core/Private/Misc/StringBuilder.cpp
+++ b/Engine/Source/Runtime/Core/Private/Misc/StringBuilder.cpp
@@ -31,6 +31,10 @@ TStringBuilderBase<C>::~TStringBuilderBase()
 template <typename C>
 void TStringBuilderBase<C>::Extend(SIZE_T ExtraCapacity)
 {
+	checkf(!bIsExternal, TEXT("Extending the capacity of an external string builder is not supported. "
+		"Size = %" SIZE_T_FMT " Capacity = %" SIZE_T_FMT" ExtraCapacity = %" SIZE_T_FMT "."),
+		SIZE_T(CurPos - Base), SIZE_T(End - Base), ExtraCapacity);
+
 	const SIZE_T OldCapacity = End - Base;
 	const SIZE_T NewCapacity = FPlatformMath::Max<SIZE_T>(NextPowerOfTwo(OldCapacity + ExtraCapacity), 32);
 
@@ -70,9 +74,7 @@ TStringBuilderBase<C>& TStringBuilderBase<C>::AppendV(const C* Fmt, va_list Args
 {
 	for (;;)
 	{
-		const int32 RemainingSize = (int32)(End - CurPos);
-
-		if (0 != RemainingSize)
+		if (const int32 RemainingSize = (int32)(End - CurPos))
 		{
 			va_list ArgPack;
 			va_copy(ArgPack, Args);
diff --git a/Engine/Source/Runtime/Core/Public/Containers/StringFwd.h b/Engine/Source/Runtime/Core/Public/Containers/StringFwd.h
--- a/Engine/Source/Runtime/Core/Public/Containers/StringFwd.h
+++ b/Engine/Source/Runtime/Core/Public/Containers/StringFwd.h
@@ -11,9 +11,11 @@
 
 template <typename CharType> class TStringBuilderBase;
 template <typename CharType, int32 BufferSize> class TStringBuilderWithBuffer;
+template <typename CharType> class TExternalStringBuilder;
 
 template <typename CharType> struct TIsContiguousContainer<TStringBuilderBase<CharType>> { static constexpr bool Value = true; };
 template <typename CharType, int32 BufferSize> struct TIsContiguousContainer<TStringBuilderWithBuffer<CharType, BufferSize>> { static constexpr bool Value = true; };
+template <typename CharType> struct TIsContiguousContainer<TExternalStringBuilder<CharType>> { static constexpr bool Value = true; };
 
 /** The base string builder type for TCHAR. */
 using FStringBuilderBase = TStringBuilderBase<TCHAR>;
diff --git a/Engine/Source/Runtime/Core/Public/Misc/StringBuilder.h b/Engine/Source/Runtime/Core/Public/Misc/StringBuilder.h
--- a/Engine/Source/Runtime/Core/Public/Misc/StringBuilder.h
+++ b/Engine/Source/Runtime/Core/Public/Misc/StringBuilder.h
@@ -25,6 +25,24 @@
 
 template <typename T> struct TIsContiguousContainer;
 
+namespace UE::Core::Private
+{
+
+class FDynamicStringBuilderToken
+{
+	explicit FDynamicStringBuilderToken() = default;
+	template <typename CharType> friend class ::TStringBuilderBase;
+	template <typename CharType, int32 BufferSize> friend class ::TStringBuilderWithBuffer;
+};
+
+class FExternalStringBuilderToken
+{
+	explicit FExternalStringBuilderToken() = default;
+	template <typename CharType> friend class ::TExternalStringBuilder;
+};
+
+} // UE::Core::Private
+
 /**
  * String Builder
  *
@@ -88,8 +106,8 @@ public:
 	}
 
 	inline TStringBuilderBase(CharType* BufferPointer, int32 BufferCapacity)
+		: TStringBuilderBase(BufferPointer, BufferCapacity, UE::Core::Private::FDynamicStringBuilderToken{})
 	{
-		Initialize(BufferPointer, BufferCapacity);
 	}
 
 	inline int32 Len() const
@@ -419,6 +437,21 @@ private:
 	friend inline const CharType* end  (const TStringBuilderBase& Builder) { return Builder.GetData() + Builder.Len(); }
 
 protected:
+	inline TStringBuilderBase(CharType* InBase, int32 InCapacity, UE::Core::Private::FDynamicStringBuilderToken)
+		: Base(InBase)
+		, CurPos(InBase)
+		, End(InBase + InCapacity)
+	{
+	}
+
+	inline TStringBuilderBase(CharType* InBase, int32 InCapacity, UE::Core::Private::FExternalStringBuilderToken)
+		: Base(InBase)
+		, CurPos(InBase)
+		, End(InBase + InCapacity)
+		, bIsExternal(true)
+	{
+	}
+
 	inline void Initialize(CharType* InBase, int32 InCapacity)
 	{
 		Base	= InBase;
@@ -453,6 +486,7 @@ protected:
 	CharType*	CurPos = nullptr;
 	CharType*	End = nullptr;
 	bool		bIsDynamic = false;
+	bool		bIsExternal = false;
 };
 
 template <typename CharType>
@@ -473,7 +507,7 @@ class TStringBuilderWithBuffer : public TStringBuilderBase<CharType>
 {
 public:
 	inline TStringBuilderWithBuffer()
-		: TStringBuilderBase<CharType>(StringBuffer, BufferSize)
+		: TStringBuilderBase<CharType>(StringBuffer, BufferSize, UE::Core::Private::FDynamicStringBuilderToken{})
 	{
 	}
 
@@ -481,8 +515,8 @@ public:
 	 * Construct a string builder by appending the arguments using operator<<.
 	 */
 	template <typename... ArgTypes>
-	explicit TStringBuilderWithBuffer(EInPlace, ArgTypes&&... Args)
-		: TStringBuilderBase<CharType>(StringBuffer, BufferSize)
+	[[nodiscard]] inline explicit TStringBuilderWithBuffer(EInPlace, ArgTypes&&... Args)
+		: TStringBuilderBase<CharType>(StringBuffer, BufferSize, UE::Core::Private::FDynamicStringBuilderToken{})
 	{
 		(*this << ... << (ArgTypes&&)Args);
 	}
@@ -490,7 +524,26 @@ public:
 	using TStringBuilderBase<CharType>::operator=;
 
 private:
-	CharType StringBuffer[BufferSize];
+	CharType StringBuffer[BufferSize > 0 ? BufferSize : 1];
+};
+
+//////////////////////////////////////////////////////////////////////////
+
+/**
+ * A string builder with fixed-size external storage.
+ *
+ * Asserts if external storage is exhausted.
+ */
+template <typename CharType>
+class TExternalStringBuilder : public TStringBuilderBase<CharType>
+{
+public:
+	[[nodiscard]] inline explicit TExternalStringBuilder(CharType* Buffer, int32 BufferSize)
+		: TStringBuilderBase<CharType>(Buffer, BufferSize, UE::Core::Private::FExternalStringBuilderToken{})
+	{
+	}
+
+	using TStringBuilderBase<CharType>::operator=;
 };
 
 //////////////////////////////////////////////////////////////////////////
