PR #13317: Fix tweakobjectptr debug on mac os

diff --git a/Engine/Source/Runtime/Core/Private/Misc/CoreGlobals.cpp b/Engine/Source/Runtime/Core/Private/Misc/CoreGlobals.cpp
--- a/Engine/Source/Runtime/Core/Private/Misc/CoreGlobals.cpp
+++ b/Engine/Source/Runtime/Core/Private/Misc/CoreGlobals.cpp
@@ -266,7 +266,7 @@ namespace UE::CoreUObject::Private
 }
 UE::CoreUObject::Private::FStoredObjectPathDebug* GCoreComplexObjectPathDebug = nullptr;
 UE::CoreUObject::Private::FObjectHandlePackageDebugData* GCoreObjectHandlePackageDebug = nullptr;
-#if PLATFORM_UNIX
+#if PLATFORM_UNIX || PLATFORM_MAC
 uint8** CORE_API GNameBlocksDebug = FNameDebugVisualizer(FClangKeepDebugInfo{}).GetBlocks();
 FChunkedFixedUObjectArray*& CORE_API GObjectArrayForDebugVisualizers = GCoreObjectArrayForDebugVisualizers;
 UE::CoreUObject::Private::FStoredObjectPathDebug*& GComplexObjectPathDebug = GCoreComplexObjectPathDebug;
diff --git a/Engine/Source/Runtime/Core/Public/Modules/Boilerplate/ModuleBoilerplate.h b/Engine/Source/Runtime/Core/Public/Modules/Boilerplate/ModuleBoilerplate.h
--- a/Engine/Source/Runtime/Core/Public/Modules/Boilerplate/ModuleBoilerplate.h
+++ b/Engine/Source/Runtime/Core/Public/Modules/Boilerplate/ModuleBoilerplate.h
@@ -97,7 +97,7 @@ class FChunkedFixedUObjectArray;
 
 #if !UE_ENABLE_VISUALIZER_HELPERS
 	#define UE_VISUALIZERS_HELPERS
-#elif PLATFORM_UNIX
+#elif PLATFORM_UNIX || PLATFORM_MAC
 	// GDB/LLDB pretty printers don't use these - no need to export additional symbols. This also solves ODR violation reported by ASan on Linux
 	#define UE_VISUALIZERS_HELPERS
 #else
