PR #13685: Delete iAD framework (deprecated) from Core.Build.cs

diff --git a/Engine/Source/Runtime/Core/Core.Build.cs b/Engine/Source/Runtime/Core/Core.Build.cs
--- a/Engine/Source/Runtime/Core/Core.Build.cs
+++ b/Engine/Source/Runtime/Core/Core.Build.cs
@@ -159,12 +159,6 @@ public class Core : ModuleRules
 				PrivateIncludePathModuleNames.Add("ApplicationCore");
 			}
 
-			bool bSupportAdvertising = Target.Platform == UnrealTargetPlatform.IOS;
-			if (bSupportAdvertising)
-			{
-				PublicFrameworks.AddRange(new string[] { "iAD" });
-			}
-
 			// export Core symbols for embedded Dlls
 			ModuleSymbolVisibility = ModuleRules.SymbolVisibility.VisibileForDll;
 			
