Fixed Cahced OS Page allocator on MB3

diff --git a/Engine/Source/Runtime/Core/Public/HAL/MallocBinned3.h b/Engine/Source/Runtime/Core/Public/HAL/MallocBinned3.h
--- a/Engine/Source/Runtime/Core/Public/HAL/MallocBinned3.h
+++ b/Engine/Source/Runtime/Core/Public/HAL/MallocBinned3.h
@@ -19,7 +19,9 @@
 #include "Templates/Atomic.h"
 
 
-#define UE_MB3_USE_CACHED_PAGE_ALLOCATOR_FOR_LARGE_ALLOCS (0)
+#ifndef UE_MB3_USE_CACHED_PAGE_ALLOCATOR_FOR_LARGE_ALLOCS
+#	define UE_MB3_USE_CACHED_PAGE_ALLOCATOR_FOR_LARGE_ALLOCS (0)
+#endif
 
 #if PLATFORM_IOS
 #	define UE_MB3_MAX_MEMORY_PER_POOL_SIZE_MB		512
@@ -183,6 +185,10 @@ public:
 private:
 	uint32 SmallPoolInfosPerPlatformPage;
 
+#if	UE_MB3_USE_CACHED_PAGE_ALLOCATOR_FOR_LARGE_ALLOCS
+	UE::FPlatformRecursiveMutex Mutex;
+#endif
+
 #if !BINNED3_USE_SEPARATE_VM_PER_POOL
 	FORCEINLINE uint64 PoolIndexFromPtr(const void* Ptr) const		// returns a uint64 for it can also be used to check if it is an OS allocation
 	{
