Add boot timings to LogInit under Verbose category to track init stages

diff --git a/Engine/Source/Runtime/Core/Private/Misc/CoreGlobals.cpp b/Engine/Source/Runtime/Core/Private/Misc/CoreGlobals.cpp
--- a/Engine/Source/Runtime/Core/Private/Misc/CoreGlobals.cpp
+++ b/Engine/Source/Runtime/Core/Private/Misc/CoreGlobals.cpp
@@ -5,6 +5,7 @@
 #include "Internationalization/Internationalization.h"
 #include "Stats/Stats.h"
 #include "HAL/IConsoleManager.h"
+#include "Logging/StructuredLog.h"
 #include "Modules/ModuleManager.h"
 #include "Misc/CoreStats.h"
 #include "Misc/TrackedActivity.h"
@@ -534,10 +535,12 @@ FEngineTrackedActivityScope::~FEngineTrackedActivityScope()
 #if !USE_BOOT_PROFILING
 FScopedBootTiming::FScopedBootTiming(const ANSICHAR *InMessage)
 {
+	UE_LOGFMT(LogInit, Verbose, "{Message}", InMessage);
 }
 
 FScopedBootTiming::FScopedBootTiming(const ANSICHAR *InMessage, FName Suffix)
 {
+	UE_LOGFMT(LogInit, Verbose, "{Message}: {Suffix}", InMessage, Suffix);
 }
 
 FScopedBootTiming::~FScopedBootTiming()
