Clarified commenting of the return type of Algo::RemoveIf and Algo::StableRemoveIf.

#rnx

diff --git a/Engine/Source/Runtime/Core/Public/Algo/RemoveIf.h b/Engine/Source/Runtime/Core/Public/Algo/RemoveIf.h
--- a/Engine/Source/Runtime/Core/Public/Algo/RemoveIf.h
+++ b/Engine/Source/Runtime/Core/Public/Algo/RemoveIf.h
@@ -15,7 +15,7 @@ namespace Algo
 	 * @param  Range  The range of elements to manipulate.
 	 * @param  Pred   A callable which maps elements to truthy values, specifying elements to be removed.
 	 *
-	 * @return The index of the first element after those which were not removed.
+	 * @return The number of elements that should be kept.
 	 */
 	template <typename RangeType, typename Predicate>
 	[[nodiscard]] int32 RemoveIf(RangeType& Range, Predicate Pred)
@@ -73,7 +73,7 @@ namespace Algo
 	 * @param  Range  The range of elements to manipulate.
 	 * @param  Pred   A callable which maps elements to truthy values, specifying elements to be removed.
 	 *
-	 * @return The index of the first element after those which were not removed.
+	 * @return The number of elements that should be kept.
 	 */
 	template <typename RangeType, typename Predicate>
 	[[nodiscard]] int32 StableRemoveIf(RangeType&& Range, Predicate Pred)
