Fix invalid/poisoned SkeletalMesh DDC data.

LODModel num vertices did not match with the Skin Weights num.

diff --git a/Engine/Source/Runtime/Core/Private/UObject/FortniteMainBranchObjectVersion_Restricted.inl b/Engine/Source/Runtime/Core/Private/UObject/FortniteMainBranchObjectVersion_Restricted.inl
--- a/Engine/Source/Runtime/Core/Private/UObject/FortniteMainBranchObjectVersion_Restricted.inl
+++ b/Engine/Source/Runtime/Core/Private/UObject/FortniteMainBranchObjectVersion_Restricted.inl
@@ -6,6 +6,6 @@ static void AppendRestrictedFortniteMainStreamObjectVersionGuids(TMap<FGuid, FGu
 	// patch effects.
 	const FDevSystemGuids& DevGuids = FDevSystemGuids::Get();
 	SystemGuids.Add(DevGuids.NANITE_DERIVEDDATA_VER, FGuid("ACF98184978742A592537DA223A1E6BE"));
-	SystemGuids.Add(DevGuids.SkeletalMeshDerivedDataVersion, FGuid("C998F2F3B2884A02B7E290B8C25835F0"));
+	SystemGuids.Add(DevGuids.SkeletalMeshDerivedDataVersion, FGuid("F08EDFD3CD9F4CEFA5F9935DEAAD81CF"));
 	SystemGuids.Add(DevGuids.STATICMESH_DERIVEDDATA_VER, FGuid("3ABF32149F9D4D838750368AB95573BE"));
 }
