Checked zlib init and shutdown for errors.

diff --git a/Engine/Source/Runtime/Core/Private/Misc/Compression.cpp b/Engine/Source/Runtime/Core/Private/Misc/Compression.cpp
--- a/Engine/Source/Runtime/Core/Private/Misc/Compression.cpp
+++ b/Engine/Source/Runtime/Core/Private/Misc/Compression.cpp
@@ -5,6 +5,7 @@
 #include "HAL/UnrealMemory.h"
 #include "Logging/LogMacros.h"
 #include "Misc/Parse.h"
+#include "Misc/ScopeExit.h"
 #include "Misc/ScopeLock.h"
 #include "Misc/CommandLine.h"
 #include "Stats/Stats.h"
@@ -177,13 +178,27 @@ static bool appCompressMemoryGZIP(void* CompressedBuffer, int64& CompressedSize,
 	// Init deflate settings to use GZIP
 	int windowsBits = 15;
 	int GZIP_ENCODING = 16;
-	deflateInit2(
+	int InitError = deflateInit2(
 		&gzipstream,
 		CompressionLevel,
 		Z_DEFLATED,
 		windowsBits | GZIP_ENCODING,
 		MAX_MEM_LEVEL,
-		Z_DEFAULT_STRATEGY);
+		Z_DEFAULT_STRATEGY
+	);
+	if (InitError != Z_OK)
+	{
+		UE_LOG(LogCompression, Error, TEXT("Failed to init ZLIB - error code %d"), InitError);
+		return false;
+	}
+	ON_SCOPE_EXIT
+	{
+		int EndError = deflateEnd(&gzipstream);
+		if (EndError != Z_OK)
+		{
+			UE_LOG(LogCompression, Warning, TEXT("Failed to shut down ZLIB - error code %d"), EndError);
+		}
+	};
 
 	// Setup output buffer
 	const unsigned long GzipHeaderLength = 12;
@@ -192,14 +207,13 @@ static bool appCompressMemoryGZIP(void* CompressedBuffer, int64& CompressedSize,
 	gzipstream.next_out = (uint8*)CompressedBuffer;
 	gzipstream.avail_out = CompressedSize;
 
-	int status = 0;
-	while ((status = deflate(&gzipstream, Z_FINISH)) == Z_OK);
-	deflateEnd(&gzipstream);
+	int Status = 0;
+	while ((Status = deflate(&gzipstream, Z_FINISH)) == Z_OK);
 
 	// Propagate compressed size from intermediate variable back into out variable.
 	CompressedSize = gzipstream.total_out;
 
-	bool bOperationSucceeded = (status == Z_STREAM_END);
+	bool bOperationSucceeded = (Status == Z_STREAM_END);
 	return bOperationSucceeded;
 }
 
