Adapt TraceAuxiliary for minimal trace scenarios.

* Command line argument control is by default disabled for minimal tracing.
* In game console command control is by default disabled for minimal tracing.
* Presets are only available in full tracing.

diff --git a/Engine/Source/Runtime/Core/Private/ProfilingDebugging/TraceAuxiliary.cpp b/Engine/Source/Runtime/Core/Private/ProfilingDebugging/TraceAuxiliary.cpp
--- a/Engine/Source/Runtime/Core/Private/ProfilingDebugging/TraceAuxiliary.cpp
+++ b/Engine/Source/Runtime/Core/Private/ProfilingDebugging/TraceAuxiliary.cpp
@@ -3,31 +3,37 @@
 #include "ProfilingDebugging/TraceAuxiliary.h"
 
 #include "BuildSettings.h"
-#include "Containers/UnrealString.h"
-#include "CoreGlobals.h"
-#include "Misc/ScopeRWLock.h"
 #include "ProfilingDebugging/CountersTrace.h"
 #include "ProfilingDebugging/CsvProfilerConfig.h"
 #include "ProfilingDebugging/StringsTrace.h"
 #include "ProfilingDebugging/TagTrace.h"
 #include "Trace/Trace.h"
 
-#if PLATFORM_WINDOWS
-#include "Windows/AllowWindowsPlatformTypes.h"
-#include <Windows.h>
-#include "Windows/HideWindowsPlatformTypes.h"
-#endif
-#if PLATFORM_LINUX || PLATFORM_MAC
-#include <sys/wait.h>
-#include <semaphore.h>
+
+////////////////////////////////////////////////////////////////////////////////
+/// Configuration
+////////////////////////////////////////////////////////////////////////////////
+
+#if !defined(WITH_UNREAL_TRACE_COMMANDLINE_ARGS)
+	#define UE_TRACEAUX_CMD_ARGS UE_TRACE_ENABLED
+#else
+	#define UE_TRACEAUX_CMD_ARGS WITH_UNREAL_TRACE_COMMANDLINE_ARGS
 #endif
 
-#if defined(WITH_UNREAL_TRACE_LAUNCH)
-	#define UE_TRACE_SERVER_LAUNCH_ENABLED WITH_UNREAL_TRACE_LAUNCH
-	UE_DEPRECATED_MACRO(5.4, "The WITH_UNREAL_TRACE_LAUNCH macro has been deprecated in favor of UE_TRACE_SERVER_LAUNCH_ENABLED.")
-#elif !defined(UE_TRACE_SERVER_LAUNCH_ENABLED)
+#if !defined(WITH_UNREAL_TRACE_CONSOLE_CMDS)
+	#define UE_TRACEAUX_CON_CMDS UE_TRACE_ENABLED
+#else
+	#define  UE_TRACEAUX_CON_CMDS WITH_UNREAL_TRACE_CONSOLE_CMDS
+#endif
+
+#if UE_TRACE_MINIMAL_ENABLED && !UE_TRACE_ENABLED
+	#define UE_TRACEAUX_FULL 0
+#elif UE_TRACE_ENABLED
+	#define UE_TRACEAUX_FULL 1
+#endif
+
+#if !defined(UE_TRACE_SERVER_LAUNCH_ENABLED)
 	#define UE_TRACE_SERVER_LAUNCH_ENABLED (PLATFORM_DESKTOP && !UE_BUILD_SHIPPING && !IS_PROGRAM)
-	#define WITH_UNREAL_TRACE_LAUNCH UE_DEPRECATED_MACRO(5.4, "The WITH_UNREAL_TRACE_LAUNCH macro has been deprecated in favor of UE_TRACE_SERVER_LAUNCH_ENABLED.") UE_TRACE_SERVER_LAUNCH_ENABLED
 #endif
 
 #if !defined(UE_TRACE_AUTOSTART)
@@ -50,6 +56,8 @@
 	#define UE_TRACEAUX_EMIT_CSV_STATS (1 && CSV_PROFILER_STATS)
 #endif
 
+////////////////////////////////////////////////////////////////////////////////
+
 #if UE_TRACE_SERVER_LAUNCH_ENABLED || UE_TRACE_SERVER_CONTROLS_ENABLED
 #include "HAL/PlatformProcess.h"
 #include "Misc/CommandLine.h"
@@ -57,17 +65,14 @@
 #include "Misc/Paths.h"
 #endif
 
-const FTraceAuxiliary::FChannelPreset GDefaultChannels(TEXT("Default"), TEXT("cpu,gpu,frame,log,bookmark,screenshot,region"), false);
-const FTraceAuxiliary::FChannelPreset GMemoryChannels(TEXT("Memory"), TEXT("memtag,memalloc,callstack,module"), true);
-const FTraceAuxiliary::FChannelPreset GMemoryLightChannels(TEXT("Memory_Light"), TEXT("memtag,memalloc"), true);
+#if UE_TRACE_MINIMAL_ENABLED
 
-#if UE_TRACE_ENABLED
-
-#include "BuildSettings.h"
 #include "Containers/Array.h"
 #include "Containers/Map.h"
 #include "Containers/StringView.h"
 #include "Containers/UnrealString.h"
+#include "Containers/UnrealString.h"
+#include "CoreGlobals.h"
 #include "HAL/FileManager.h"
 #include "HAL/IConsoleManager.h"
 #include "HAL/PlatformTime.h"
@@ -80,6 +85,7 @@
 #include "Misc/Fork.h"
 #include "Misc/PathViews.h"
 #include "Misc/Paths.h"
+#include "Misc/ScopeRWLock.h"
 #include "Modules/ModuleManager.h"
 #include "ProfilingDebugging/CallstackTrace.h"
 #include "ProfilingDebugging/CpuProfilerTrace.h"
@@ -90,19 +96,19 @@
 #include "Templates/UnrealTemplate.h"
 #include "Trace/Trace.inl"
 
+#if PLATFORM_WINDOWS
+#include "Windows/AllowWindowsPlatformTypes.h"
+#include <Windows.h>
+#include "Windows/HideWindowsPlatformTypes.h"
+#endif
+#if PLATFORM_LINUX || PLATFORM_MAC
+#include <sys/wait.h>
+#include <semaphore.h>
+#endif
+
+////////////////////////////////////////////////////////////////////////////////
+/// Stats and logging 
 ////////////////////////////////////////////////////////////////////////////////
 
-const TCHAR* GTraceConfigSection = TEXT("Trace.Config");
-static UE::Trace::FInitializeDesc GInitializeDesc;
-
-enum class EWorkerThreadConfig : uint8 {
-	Unset,
-	OnDemand,		// Worker thread started when first trace is started
-	OnInit,			// Worker thread started on boot
-	Never,			// Worker threda never started (rely on end frame pump)
-};
-
-static EWorkerThreadConfig GWorkerThreadConfig = EWorkerThreadConfig::Unset;
-
 CSV_DEFINE_CATEGORY(Trace, true);
 
@@ -148,14 +154,6 @@
 LLM_DEFINE_TAG(TraceLog_Cache);
 #endif
 
-////////////////////////////////////////////////////////////////////////////////
-enum class ETraceConnectType
-{
-	Network,
-	File
-};
-
-////////////////////////////////////////////////////////////////////////////////
 #if NO_LOGGING
 FTraceAuxiliary::FLogCategoryAlias LogTrace;
 #else
@@ -163,7 +161,32 @@
 #endif
 
 ////////////////////////////////////////////////////////////////////////////////
-class FTraceAuxiliaryImpl
+/// Types and constants 
+////////////////////////////////////////////////////////////////////////////////
+
+enum class EWorkerThreadConfig : uint8 {
+	Unset,
+	OnDemand,		// Worker thread started when first trace is started
+	OnInit,			// Worker thread started on boot
+	Never,			// Worker threda never started (rely on end frame pump)
+};
+
+////////////////////////////////////////////////////////////////////////////////
+const FTraceAuxiliary::FChannelPreset GDefaultChannels(TEXT("Default"), TEXT("cpu,gpu,frame,log,bookmark,screenshot,region"), false);
+const FTraceAuxiliary::FChannelPreset GMemoryChannels(TEXT("Memory"), TEXT("memtag,memalloc,callstack,module"), true);
+const FTraceAuxiliary::FChannelPreset GMemoryLightChannels(TEXT("Memory_Light"), TEXT("memtag,memalloc"), true);
+
+
+////////////////////////////////////////////////////////////////////////////////
+/**
+ * FTraceAuxiliaryImpl is an implementation of the FTraceAuxiliary public interface.
+ * It is also remembers and manages the state of TraceLog when it comes to:
+ * * Channels
+ * * Connections
+ * * Command line arguments for forking
+ * * Misc state (like worker thread)
+ */
+class FTraceAuxiliaryImpl  
 {
 public:
 	FString GetDest() const;
@@ -171,9 +194,6 @@
 	bool IsConnected(FGuid& OutSessionGuid, FGuid& OutTraceGuid);
 	FTraceAuxiliary::EConnectionType GetConnectionType() const;
 	void GetActiveChannelsString(FStringBuilderBase& String) const;
-	void AddCommandlineChannels(const TCHAR* ChannelList);
-	void ResetCommandlineChannels();
-	bool HasCommandlineChannels() const { return !CommandlineChannels.IsEmpty(); }
 	void EnableChannels(const TCHAR* ChannelList, const FTraceAuxiliary::FLogCategoryAlias& LogCategory);
 	void EnableChannels(TConstArrayView<uint32> ChannelIds, TMap<uint32, FString>* OutReason);
 	void DisableChannels(const TCHAR* ChannelList, const FTraceAuxiliary::FLogCategoryAlias& LogCategory);
@@ -187,11 +207,14 @@
 	bool IsPaused();
 	void EnableCommandlineChannels();
 	void EnableCommandlineChannelsPostInitialize();
+	bool HasCommandlineChannels() const;
 	void SetTruncateFile(bool bTruncateFile);
 	void StartWorkerThread();
 	void RegisterEndFrameCallbacks();
 	bool WriteSnapshot(const TCHAR* InFilePath, const FTraceAuxiliary::FLogCategoryAlias& LogCategory);
 	bool SendSnapshot(const TCHAR* InHost, uint32 InPort, const FTraceAuxiliary::FLogCategoryAlias& LogCategory);
+	void AddCommandlineChannels(const TCHAR* ChannelList);
+	void ResetCommandlineChannels();
 
 	// True if this is parent process with forking requested before forking.
 	bool IsParentProcessAndPreFork();
@@ -209,12 +232,12 @@
 		bool bActive = false;
 	};
 
-	void AddCommandlineChannel(const TCHAR* Name, const FTraceAuxiliary::FLogCategoryAlias& LogCategory);
 	template<typename T, typename... Ts> void ForEachChannel(const TCHAR* ChannelList, bool bResolvePresets, T&& Callable, Ts&&... Args);
 	static uint32 HashChannelName(const TCHAR* Name);
 	bool SendToHost(const TCHAR* Host, const FTraceAuxiliary::FLogCategoryAlias& LogCategory, uint16 SendFlags);
 	bool WriteToFile(const TCHAR* Path, const FTraceAuxiliary::FLogCategoryAlias& LogCategory, uint16 SendFlags);
 	bool FinalizeFilePath(const TCHAR* InPath, FString& OutPath, const FTraceAuxiliary::FLogCategoryAlias& LogCategory);
+	void AddCommandlineChannel(const TCHAR* Name, const FTraceAuxiliary::FLogCategoryAlias& LogCategory);
 
 	template<typename ChannelType>
 	bool EnableChannel(ChannelType Channel, const FTraceAuxiliary::FLogCategoryAlias& LogCategory, FString* OutReason = nullptr);
@@ -236,11 +259,17 @@
 	mutable FRWLock CurrentTargetLock;
 };
 
+////////////////////////////////////////////////////////////////////////////////
+/// Statics
+////////////////////////////////////////////////////////////////////////////////
+
 static FTraceAuxiliaryImpl GTraceAuxiliary;
+static const TCHAR* GTraceConfigSection = TEXT("Trace.Config");
+static UE::Trace::FInitializeDesc GInitializeDesc;
+static EWorkerThreadConfig GWorkerThreadConfig = EWorkerThreadConfig::Unset;
 static FDelegateHandle GEndFrameDelegateHandle;
 static FDelegateHandle GEndFrameStatDelegateHandle;
 static FDelegateHandle GOnPostForkHandle;
-
 // Whether to start tracing automatically at start or wait to initiate via Console Command.
 // This value can also be set by passing '-traceautostart=[0|1]' on command line.
 static bool GTraceAutoStart = UE_TRACE_AUTOSTART ? true : false;
@@ -254,6 +283,7 @@
 ////////////////////////////////////////////////////////////////////////////////
 void FTraceAuxiliaryImpl::AddCommandlineChannel(const TCHAR* Name, const FTraceAuxiliary::FLogCategoryAlias& LogCategory)
 {
+#if UE_TRACEAUX_CMD_ARGS
 	uint32 Hash = HashChannelName(Name);
 
 	if (CommandlineChannels.Find(Hash) != nullptr)
@@ -268,18 +298,23 @@
 	{
 		Value.bActive = EnableChannel(*Value.Name, LogCategory);
 	}
+#endif //UE_TRACEAUX_CMD_ARGS
 }
 
 ////////////////////////////////////////////////////////////////////////////////
 void FTraceAuxiliaryImpl::AddCommandlineChannels(const TCHAR* ChannelList)
 {
+#if UE_TRACEAUX_CMD_ARGS
 	ForEachChannel(ChannelList, true, &FTraceAuxiliaryImpl::AddCommandlineChannel, LogTrace);
+#endif //UE_TRACEAUX_CMD_ARGS
 }
 
 ////////////////////////////////////////////////////////////////////////////////
 void FTraceAuxiliaryImpl::ResetCommandlineChannels()
 {
+#if UE_TRACEAUX_CMD_ARGS
 	CommandlineChannels.Reset();
+#endif //UE_TRACEAUX_CMD_ARGS
 }
 
 ////////////////////////////////////////////////////////////////////////////////
@@ -713,6 +748,7 @@
 ////////////////////////////////////////////////////////////////////////////////
 void FTraceAuxiliaryImpl::EnableCommandlineChannels()
 {
+#if UE_TRACEAUX_CMD_ARGS
 	if (IsParentProcessAndPreFork())
 	{
 		return;
@@ -725,16 +761,29 @@
 			ChannelPair.Value.bActive = EnableChannel(*ChannelPair.Value.Name, LogTrace);
 		}
 	}
+#endif
 }
 
 ////////////////////////////////////////////////////////////////////////////////
 void FTraceAuxiliaryImpl::EnableCommandlineChannelsPostInitialize()
 {
+#if UE_TRACEAUX_CMD_ARGS
 	for (auto& ChannelPair : CommandlineChannels)
 	{
 		// Intentionally enable channel without checking current state.
 		ChannelPair.Value.bActive = EnableChannel(*ChannelPair.Value.Name, LogTrace);
 	}
+#endif
 }
 
 ////////////////////////////////////////////////////////////////////////////////
+bool FTraceAuxiliaryImpl::HasCommandlineChannels() const 
+{ 
+#if UE_TRACEAUX_CMD_ARGS
+	return !CommandlineChannels.IsEmpty(); 
+#else
+	return false;
+#endif
+}
+
+////////////////////////////////////////////////////////////////////////////////
@@ -1506,5 +1555,10 @@
 }
 
 ////////////////////////////////////////////////////////////////////////////////
+/// Console commands
+////////////////////////////////////////////////////////////////////////////////
+#if UE_TRACEAUX_CON_CMDS
+
+////////////////////////////////////////////////////////////////////////////////
 static FAutoConsoleCommand TraceAuxiliarySendCmd(
 	TEXT("Trace.Send"),
@@ -1621,6 +1675,8 @@
 	FConsoleCommandWithArgsDelegate::CreateStatic(TraceRegionEnd)
 );
 
+#endif // UE_TRACEAUX_CON_CMDS
+
 #endif // UE_TRACE_ENABLED
 
 
@@ -1731,10 +1787,9 @@
 FTraceAuxiliary::FOnSnapshotSaved FTraceAuxiliary::OnSnapshotSaved;
 
 ////////////////////////////////////////////////////////////////////////////////
-#if UE_TRACE_ENABLED
-
 static bool TraceAuxiliaryStartShared(const TCHAR* Channels, const FTraceAuxiliary::FOptions* Options, const FTraceAuxiliary::FLogCategoryAlias& LogCategory)
 {
+#if UE_TRACE_MINIMAL_ENABLED
 	if (GTraceAuxiliary.IsParentProcessAndPreFork())
 	{
 		return false;
@@ -1756,14 +1811,15 @@
 	}
 
 	return true;
+#else // UE_TRACE_MINIMAL_ENABLED
+	return false;
+#endif
 }
 
-#endif // UE_TRACE_ENABLED
-
 ////////////////////////////////////////////////////////////////////////////////
 bool FTraceAuxiliary::Start(EConnectionType Type, const TCHAR* Target, const TCHAR* Channels, FOptions* Options, const FLogCategoryAlias& LogCategory)
 {
-#if UE_TRACE_ENABLED
+#if UE_TRACE_MINIMAL_ENABLED
 	// Use the Relay function to start a trace relay.
 	check(Type != EConnectionType::Relay);
 
@@ -1794,7 +1850,7 @@
 ////////////////////////////////////////////////////////////////////////////////
 bool FTraceAuxiliary::Relay(UPTRINT Handle, UE::Trace::IoWriteFunc WriteFunc, UE::Trace::IoCloseFunc CloseFunc, const TCHAR* Channels, const FOptions* Options)
 {
-#if UE_TRACE_ENABLED
+#if UE_TRACE_MINIMAL_ENABLED
 	TRACE_CPUPROFILER_EVENT_SCOPE(FTraceAux_Relay);
 
 	if (!TraceAuxiliaryStartShared(Channels, Options, LogTrace))
@@ -1813,7 +1869,7 @@
 ////////////////////////////////////////////////////////////////////////////////
 bool FTraceAuxiliary::Stop()
 {
-#if UE_TRACE_ENABLED
+#if UE_TRACE_MINIMAL_ENABLED
 	return GTraceAuxiliary.Stop();
 #else
 	return false;
@@ -1823,7 +1879,7 @@
 ////////////////////////////////////////////////////////////////////////////////
 bool FTraceAuxiliary::Pause()
 {
-#if UE_TRACE_ENABLED
+#if UE_TRACE_MINIMAL_ENABLED
 	GTraceAuxiliary.PauseChannels();
 #endif
 	return true;
@@ -1832,7 +1888,7 @@
 ////////////////////////////////////////////////////////////////////////////////
 bool FTraceAuxiliary::IsPaused()
 {
-#if UE_TRACE_ENABLED
+#if UE_TRACE_MINIMAL_ENABLED
 	return GTraceAuxiliary.IsPaused();
 #else
 	return false;
@@ -1842,7 +1898,7 @@
 ////////////////////////////////////////////////////////////////////////////////
 bool FTraceAuxiliary::Resume()
 {
-#if UE_TRACE_ENABLED
+#if UE_TRACE_MINIMAL_ENABLED
 	GTraceAuxiliary.ResumeChannels();
 #endif
 	return true;
@@ -1851,7 +1907,7 @@
 ////////////////////////////////////////////////////////////////////////////////
 bool FTraceAuxiliary::WriteSnapshot(const TCHAR* InFilePath)
 {
-#if UE_TRACE_ENABLED
+#if UE_TRACE_MINIMAL_ENABLED
 	TRACE_CPUPROFILER_EVENT_SCOPE(FTraceAux_WriteSnapshot);
 	return GTraceAuxiliary.WriteSnapshot(InFilePath, LogTrace);
 #else
@@ -1862,7 +1918,7 @@
 ////////////////////////////////////////////////////////////////////////////////
 bool FTraceAuxiliary::SendSnapshot(const TCHAR* Host, uint32 Port)
 {
-#if UE_TRACE_ENABLED
+#if UE_TRACE_MINIMAL_ENABLED
 	TRACE_CPUPROFILER_EVENT_SCOPE(FTraceAux_SendSnapshot);
 	return GTraceAuxiliary.SendSnapshot(Host, Port, LogTrace);
 #else
@@ -1871,7 +1927,7 @@
 }
 
 ////////////////////////////////////////////////////////////////////////////////
-#if UE_TRACE_ENABLED
+#if UE_TRACE_MINIMAL_ENABLED
 static void TraceAuxiliaryAddPostForkCallback(const TCHAR* CommandLine)
 {
 	if (GTraceAutoStart)
@@ -1921,7 +1977,7 @@
 	}
 #endif
 
-#if UE_TRACE_ENABLED
+#if UE_TRACE_MINIMAL_ENABLED
 
 	// Setup message callback so we get feedback from TraceLog
 	UE::Trace::SetMessageCallback(&TraceAuxiliaryOnMessageCallback);
@@ -1989,19 +2045,22 @@
 		<< Session2.Platform(PREPROCESSOR_TO_STRING(UBT_COMPILED_PLATFORM), PlatformLen)
 		<< Session2.AppName(AppName, AppNameLen)
 		<< Session2.ProjectName(ProjectName, ProjectNameLen)
-		<< Session2.CommandLine(CommandLine, CommandLineLen)
 		<< Session2.Branch(BranchName, BranchNameLen)
 		<< Session2.BuildVersion(BuildVersion, BuildVersionLen)
 		<< Session2.Changelist(BuildSettings::GetCurrentChangelist())
 		<< Session2.ConfigurationType(uint8(FApp::GetBuildConfiguration()))
 		<< Session2.TargetType(uint8(FApp::GetBuildTargetType()))
 		<< Session2.InstanceId(InstanceId, InstanceIdSize)
-		<< Session2.VFSPaths(VFSPaths, VFSPathsLen);
+#if UE_TRACEAUX_FULL
+		<< Session2.CommandLine(CommandLine, CommandLineLen)
+		<< Session2.VFSPaths(VFSPaths, VFSPathsLen)
+#endif
+		;
 
 	// Attempt to send trace data somewhere from the command line. It perhaps
 	// seems odd to do this before initializing Trace, but it is done this way
 	// to support disabling the "important" cache without losing any events.
-	const bool bWasConnected = StartFromCommandlineArguments(CommandLine);
+	const bool bWasConnectedFromCmdLine = StartFromCommandlineArguments(CommandLine);
 
 	// Emit empty stats once (to ensure all stats/counters start from zero).
 	{
@@ -2037,7 +2096,7 @@
 		// If allowed start worker thread if trace is started using command line argument,
 		// otherwise start on demand. For the editor we always want to start a worker thread
 		// since there is a lot of time until frames are reliably pumped.
-		Desc.bUseWorkerThread = bWasConnected || UE_EDITOR;
+		Desc.bUseWorkerThread = bWasConnectedFromCmdLine || UE_EDITOR;
 		GWorkerThreadConfig = Desc.bUseWorkerThread ? EWorkerThreadConfig::OnInit : EWorkerThreadConfig::OnDemand;
 	}
 
@@ -2133,7 +2192,7 @@
 
 	GTraceAuxiliary.FreezeReadOnlyChannels();
 	UE_LOG(LogTrace, Log, TEXT("Finished trace initialization."));
-#endif //UE_TRACE_ENABLED
+#endif //UE_TRACE_MINIMAL_ENABLED
 }
 
 ////////////////////////////////////////////////////////////////////////////////
@@ -2174,42 +2233,42 @@
 ////////////////////////////////////////////////////////////////////////////////
 void FTraceAuxiliary::EnableCommandlineChannels()
 {
-#if UE_TRACE_ENABLED
+#if UE_TRACEAUX_CMD_ARGS
 	GTraceAuxiliary.EnableCommandlineChannels();
 #endif
 }
 
 void FTraceAuxiliary::EnableChannels(TConstArrayView<uint32> ChannelIds, TMap<uint32, FString>* OutErrors)
 {
-#if UE_TRACE_ENABLED
+#if UE_TRACE_MINIMAL_ENABLED
 	GTraceAuxiliary.EnableChannels(ChannelIds, OutErrors);
 #endif
 }
 
 void FTraceAuxiliary::DisableChannels(TConstArrayView<uint32> ChannelIds, TMap<uint32, FString>* OutErrors)
 {
-#if UE_TRACE_ENABLED
+#if UE_TRACE_MINIMAL_ENABLED
 	GTraceAuxiliary.DisableChannels(ChannelIds, OutErrors);
 #endif
 }
 
 void FTraceAuxiliary::EnableChannels(const TCHAR* Channels)
 {
-#if UE_TRACE_ENABLED
+#if UE_TRACE_MINIMAL_ENABLED
 	GTraceAuxiliary.EnableChannels(Channels, LogTrace);
 #endif
 }
 
 void FTraceAuxiliary::DisableChannels(const TCHAR* Channels)
 {
-#if UE_TRACE_ENABLED
+#if UE_TRACE_MINIMAL_ENABLED
 	GTraceAuxiliary.DisableChannels(Channels, LogTrace);
 #endif
 }
 
 FString FTraceAuxiliary::GetTraceDestinationString()
 {
-#if UE_TRACE_ENABLED
+#if UE_TRACE_MINIMAL_ENABLED
 	return GTraceAuxiliary.GetDest();
 #else
 	return FString();
@@ -2218,7 +2277,7 @@
 
 bool FTraceAuxiliary::IsConnected()
 {
-#if UE_TRACE_ENABLED
+#if UE_TRACE_MINIMAL_ENABLED
 	return GTraceAuxiliary.IsConnected();
 #else
 	return false;
@@ -2227,7 +2286,7 @@
 
 bool FTraceAuxiliary::IsConnected(FGuid& OutSessionGuid, FGuid& OutTraceGuid)
 {
-#if UE_TRACE_ENABLED
+#if UE_TRACE_MINIMAL_ENABLED
 	return GTraceAuxiliary.IsConnected(OutSessionGuid, OutTraceGuid);
 #else
 	return false;
@@ -2236,7 +2295,7 @@
 
 FTraceAuxiliary::EConnectionType FTraceAuxiliary::GetConnectionType()
 {
-#if UE_TRACE_ENABLED
+#if UE_TRACE_MINIMAL_ENABLED
 	return GTraceAuxiliary.GetConnectionType();
 #else
 	return FTraceAuxiliary::EConnectionType::None;
@@ -2266,6 +2325,7 @@
 
 void FTraceAuxiliary::EnumerateFixedChannelPresets(PresetCallback Callback)
 {
+#if UE_TRACE_ENABLED
 	if (Callback(GDefaultChannels) == EEnumerateResult::Stop)
 	{
 		return;
@@ -2280,6 +2340,7 @@
 	{
 		return;
 	}
+#endif
 }
 
 void FTraceAuxiliary::EnumerateChannelPresetsFromSettings(PresetCallback Callback)
@@ -2304,7 +2365,7 @@
 
 FTraceAuxiliary::ETraceSystemStatus FTraceAuxiliary::GetTraceSystemStatus()
 {
-#if UE_TRACE_ENABLED
+#if UE_TRACE_MINIMAL_ENABLED
 	EConnectionType ConnectionType = GetConnectionType();
 	if (ConnectionType == EConnectionType::Network)
 	{
@@ -2361,6 +2422,8 @@
 }
 
 ////////////////////////////////////////////////////////////////////////////////
+/// UnrealTraceServer launching
+/////////////////////////////////////////////////////////////////////////////////
 #if UE_TRACE_SERVER_CONTROLS_ENABLED
 
 enum class ELaunchTraceServerCommand
