TSpscQueue: added internal NumElems maintenance and Num() implementation, mainly used for stats and debugging.

diff --git a/Engine/Source/Runtime/Core/Public/Containers/SpscQueue.h b/Engine/Source/Runtime/Core/Public/Containers/SpscQueue.h
--- a/Engine/Source/Runtime/Core/Public/Containers/SpscQueue.h
+++ b/Engine/Source/Runtime/Core/Public/Containers/SpscQueue.h
@@ -26,6 +26,7 @@ public:
 		FNode* Node = ::new(AllocatorType::Malloc(sizeof(FNode), alignof(FNode))) FNode;
 		Tail.store(Node, std::memory_order_relaxed);
 		Head = First = TailCopy = Node;
+		NumElems = 0;
 	}
 
 	~TSpscQueue()
@@ -61,6 +62,8 @@ public:
 
 		Head->Next.store(Node, std::memory_order_release);
 		Head = Node;
+
+		NumElems++;
 	}
 
 	// returns empty TOptional if queue is empty
@@ -78,6 +81,11 @@ public:
 		DestructItem(TailNextValue);
 
 		Tail.store(LocalTailNext, std::memory_order_release);
+
+		if (Value.IsSet())
+		{
+			NumElems--;
+		}
 		return Value;
 	}
 
@@ -95,9 +103,12 @@ public:
 
 	[[nodiscard]] bool IsEmpty() const
 	{
-		FNode* LocalTail = Tail.load(std::memory_order_relaxed);
-		FNode* LocalTailNext = LocalTail->Next.load(std::memory_order_acquire);
-		return LocalTailNext == nullptr;
+		return !NumElems;
+	}
+
+	[[nodiscard]] int32 Num() const
+	{
+		return NumElems;
 	}
 
 	// as there can be only one consumer, a consumer can safely "peek" the tail of the queue.
@@ -199,18 +210,14 @@ private:
 private:
 	// consumer part 
 	// accessed mainly by consumer, infrequently by producer 
-#ifndef PLATFORM_MAC // some projects are still built on macOS before v.10.14 that doesn't have aligned new/delete operators
-	/*alignas(PLATFORM_CACHE_LINE_SIZE) */std::atomic<FNode*> Tail; // tail of the queue 
-#else
 	std::atomic<FNode*> Tail; // tail of the queue
-#endif
+
 	// producer part 
 	// accessed only by producer 
-#ifndef PLATFORM_MAC
-	/*alignas(PLATFORM_CACHE_LINE_SIZE) */FNode* Head; // head of the queue 
-#else
 	FNode* Head; // head of the queue
-#endif
+
 	FNode* First; // last unused node (tail of node cache) 
 	FNode* TailCopy; // helper (points somewhere between First and Tail)
+
+	std::atomic<int32> NumElems;
 };
\ No newline at end of file
