StringBuilder: Deprecated public and protected ways of initializing from a buffer

diff --git a/Engine/Source/Developer/Horde/Private/HordePlatform.cpp b/Engine/Source/Developer/Horde/Private/HordePlatform.cpp
--- a/Engine/Source/Developer/Horde/Private/HordePlatform.cpp
+++ b/Engine/Source/Developer/Horde/Private/HordePlatform.cpp
@@ -88,7 +88,7 @@ void FHordePlatform::CreateUniqueIdentifier(char* NameBuffer, size_t NameBufferL
 	uint64 Time = (uint64)(ts.tv_sec%100000);
 #endif
 
-	TStringBuilderBase<char> Builder(NameBuffer, NameBufferLen);
+	TExternalStringBuilder<char> Builder(NameBuffer, NameBufferLen);
 
 	// Encode to be as short as possible.. 
 	auto AppendValue = [&](auto Value)
diff --git a/Engine/Source/Runtime/Core/Public/Misc/StringBuilder.h b/Engine/Source/Runtime/Core/Public/Misc/StringBuilder.h
--- a/Engine/Source/Runtime/Core/Public/Misc/StringBuilder.h
+++ b/Engine/Source/Runtime/Core/Public/Misc/StringBuilder.h
@@ -105,6 +105,7 @@ public:
 		return *this = ViewType(Str);
 	}
 
+	UE_DEPRECATED(5.7, "Use TStringBuilder<N> or T[CharType]StringBuilder<N> or TExternalStringBuilder<CharType>.")
 	inline TStringBuilderBase(CharType* BufferPointer, int32 BufferCapacity)
 		: TStringBuilderBase(BufferPointer, BufferCapacity, UE::Core::Private::FDynamicStringBuilderToken{})
 	{
@@ -452,6 +453,7 @@ protected:
 	{
 	}
 
+	UE_DEPRECATED(5.7, "Use TStringBuilder<N> or T[CharType]StringBuilder<N> or TExternalStringBuilder<CharType>.")
 	inline void Initialize(CharType* InBase, int32 InCapacity)
 	{
 		Base	= InBase;
