Changed FMonotonicTimeSpan::IsInfinity() to stop matching negative infinity

All existing callers intended to check only for positive infinity, which suggests that the API is not intuitive as implemented. Matching negative infinity could be a problem in a hypothetical scenario where a calculation produces negative infinity for an expired timeout but we remap it to an infinite wait because of this logic.

diff --git a/Engine/Source/Runtime/Core/Public/Misc/MonotonicTime.h b/Engine/Source/Runtime/Core/Public/Misc/MonotonicTime.h
--- a/Engine/Source/Runtime/Core/Public/Misc/MonotonicTime.h
+++ b/Engine/Source/Runtime/Core/Public/Misc/MonotonicTime.h
@@ -38,7 +38,7 @@ public:
 	constexpr double ToMilliseconds() const { return Time * 1000.0; }
 
 	constexpr bool IsZero() const { return *this == Zero(); }
-	constexpr bool IsInfinity() const { return *this == Infinity() || *this == -Infinity(); }
+	constexpr bool IsInfinity() const { return *this == Infinity(); }
 	bool IsNaN() const { return FPlatformMath::IsNaN(Time); }
 
 	constexpr bool operator==(const FMonotonicTimeSpan Other) const { return Time == Other.Time; }
