Optimized FLazyName to support constant initialization

Added support for UTF-8 string literals in FLazyName.
Fixed a race in FLazyName::UEOpEquals when names are resolved concurrently with comparisons.

diff --git a/Engine/Source/Runtime/Core/Private/UObject/UnrealNames.cpp b/Engine/Source/Runtime/Core/Private/UObject/UnrealNames.cpp
--- a/Engine/Source/Runtime/Core/Private/UObject/UnrealNames.cpp
+++ b/Engine/Source/Runtime/Core/Private/UObject/UnrealNames.cpp
@@ -2972,19 +2972,6 @@ static FWideStringViewWithWidth MakeUnconvertedView(const WIDECHAR* Str)
 	return View;
 }
 
-// @pre Str contains only digits and the number is smaller than int64 max
-template<typename CharType>
-static constexpr int64 Atoi64(const CharType* Str, int32 Len)
-{
-    int64 N = 0;
-    for (int32 Idx = 0; Idx < Len; ++Idx)
-    {
-        N = 10 * N + Str[Idx] - '0';
-    }
-
-    return N;
-}
-
 /** Templated implementations of non-templated member functions, helps keep header clean */
 struct FNameHelper
 {
@@ -2996,7 +2983,7 @@ struct FNameHelper
 			return FName();
 		}
 		
-		uint32 InternalNumber = ParseNumber(View.Str, /* may be shortened */ View.Len);
+		uint32 InternalNumber = UE::Core::Private::ParseNumberFromName(View.Str, /* may be shortened */ View.Len);
 		return MakeWithNumber(View, FindType, InternalNumber);
 	}
 
@@ -3126,7 +3113,7 @@ struct FNameHelper
 		// "_[num]" is considered "None" to keep semantics of old FName implementation
 		if (Name.IsNone() & (Str[0] == '_'))
 		{
-			ParseNumber(Str, /* in-out */ Len);
+			UE::Core::Private::ParseNumberFromName(Str, /* in-out */ Len);
 			return Len == 0;
 		}
 
@@ -3142,7 +3129,7 @@ struct FNameHelper
 
 		// Check _[number] suffix first and shorten Len if it matches
 		if (Number != NAME_NO_NUMBER_INTERNAL && 
-			Number != ParseNumber(Str, /* in-out */ Len)) 
+			Number != UE::Core::Private::ParseNumberFromName(Str, /* in-out */ Len)) 
 		{
 			return false;
 		}
@@ -3158,37 +3145,6 @@ struct FNameHelper
 			: !FPlatformString::Strnicmp(Str, Entry->GetUnterminatedName(Temp.AnsiName), Len);
 	}
 
-	template<typename CharType>
-	static uint32 ParseNumber(const CharType* Name, int32& InOutLen)
-	{
-		const int32 Len = InOutLen;
-		int32 Digits = 0;
-		for (const CharType* It = Name + Len - 1; It >= Name && *It >= '0' && *It <= '9'; --It)
-		{
-			++Digits;
-		}
-
-		const CharType* FirstDigit = Name + Len - Digits;
-		static constexpr int32 MaxDigitsInt32 = 10;
-		if (Digits && Digits < Len && *(FirstDigit - 1) == '_' && Digits <= MaxDigitsInt32)
-		{
-			// check for the case where there are multiple digits after the _ and the first one
-			// is a 0 ("Rocket_04"). Can't split this case. (So, we check if the first char
-			// is not 0 or the length of the number is 1 (since ROcket_0 is valid)
-			if (Digits == 1 || *FirstDigit != '0')
-			{
-				int64 Number = Atoi64(Name + Len - Digits, Digits);
-				if (Number < MAX_int32)
-				{
-					InOutLen -= 1 + Digits;
-					return static_cast<uint32>(NAME_EXTERNAL_TO_INTERNAL(Number));
-				}
-			}
-		}
-
-		return NAME_NO_NUMBER_INTERNAL;
-	}
-
 	// Internal helpers for Make* functions
 private:
 
@@ -4305,7 +4261,8 @@ void FName::AutoTest()
 	CheckLazyName("Hej");
 	CheckLazyName("hej");
 	CheckLazyName("HEJ");
-	CheckLazyName(TEXT("Hej"));
+	CheckLazyName(UTF8TEXT("Hej"));
+	CheckLazyName(WIDETEXT("Hej"));
 	CheckLazyName("Hej_0");
 	CheckLazyName("Hej_00");
 	CheckLazyName("Hej_1");
@@ -4320,6 +4277,13 @@ void FName::AutoTest()
 	CheckLazyName("None_0");
 	CheckLazyName("None_1");
 
+	check(FLazyName("Hej") == FLazyName(UTF8TEXT("Hej")));
+	check(FLazyName("Hej") == FLazyName(WIDETEXT("Hej")));
+	check(FLazyName(UTF8TEXT("Hej")) == FLazyName("Hej"));
+	check(FLazyName(UTF8TEXT("Hej")) == FLazyName(WIDETEXT("Hej")));
+	check(FLazyName(WIDETEXT("Hej")) == FLazyName("Hej"));
+	check(FLazyName(WIDETEXT("Hej")) == FLazyName(UTF8TEXT("Hej")));
+
 	TestNameBatch();
 
 #if 0
@@ -4514,7 +4478,7 @@ FName FLazyName::Resolve() const
 
 	if (Copy.IsName())
 	{
-		FNameEntryId ComparisonId = Copy.GetComparisionId();
+		FNameEntryId ComparisonId = Copy.GetComparisonId();
 		FNameEntryId DisplayId = Copy.GetDisplayId();
 #if UE_FNAME_OUTLINE_NUMBER
 		return FNameHelper::MakeWithNumber(FNameEntryIds{ ComparisonId, DisplayId }, FNAME_Add, Number);
@@ -4524,12 +4488,26 @@ FName FLazyName::Resolve() const
 	}
 
 	// Resolve to FName but throw away the number part
-	FName FullName = bLiteralIsWide ? FName(Copy.AsWideLiteral()) : FName(Copy.AsAnsiLiteral());
+	FName FullName;
+	switch (LiteralType)
+	{
+	case ELiteralType::AnsiLiteral:
+		FullName = FName(Copy.AnsiLiteral);
+		break;
+	case ELiteralType::Utf8Literal:
+		FullName = FName(Copy.Utf8Literal);
+		break;
+	case ELiteralType::WideLiteral:
+		FullName = FName(Copy.WideLiteral);
+		break;
+	default: checkNoEntry();
+		break;
+	}
 	FNameEntryId ComparisonId = FullName.GetComparisonIndex();
 	FNameEntryId DisplayId = FullName.GetDisplayIndex();
 
 	// Deliberately unsynchronized write of word-sized int, ok if multiple threads resolve same lazy name
-	Either = FLiteralOrName(ComparisonId, DisplayId);
+	Either.PackedName = FLiteralOrName::PackName(ComparisonId, DisplayId);
 
 #if UE_FNAME_OUTLINE_NUMBER
 	return FNameHelper::MakeWithNumber(FNameEntryIds{ ComparisonId, DisplayId}, FNAME_Add, Number);
@@ -4548,40 +4526,59 @@ CORE_API FUtf8String FLazyName::ToUtf8String() const
 	return Resolve().ToUtf8String();
 }
 
-uint32 FLazyName::ParseNumber(const ANSICHAR* Str, int32 Len)
+template <typename CharType>
+uint32 FLazyName::CallParseNumber(const CharType* Literal, int32 Len)
 {
-	return FNameHelper::ParseNumber(Str, Len);
+	return UE::Core::Private::ParseNumberFromName(Literal, Len);
 }
 
-uint32 FLazyName::ParseNumber(const WIDECHAR* Str, int32 Len)
-{
-	return FNameHelper::ParseNumber(Str, Len);
-}
+template CORE_API uint32 FLazyName::CallParseNumber(const ANSICHAR* Literal, int32 Len);
+template CORE_API uint32 FLazyName::CallParseNumber(const UTF8CHAR* Literal, int32 Len);
+template CORE_API uint32 FLazyName::CallParseNumber(const WIDECHAR* Literal, int32 Len);
 
 bool FLazyName::UEOpEquals(const FLazyName& Rhs) const
 {
+	// Copy these because a concurrent Resolve() can modify them during this function.
+	FLiteralOrName LhsEither = Either;
+	FLiteralOrName RhsEither = Rhs.Either;
+
 	// If we have started creating FNames we might as well resolve and cache both lazy names
-	if (Either.IsName() || Rhs.Either.IsName())
+	if (LhsEither.IsName() || RhsEither.IsName())
 	{
 		return Resolve() == Rhs.Resolve();
 	}
 
 	// Literal pointer comparison, can ignore width
-	if (Either.AsAnsiLiteral() == Rhs.Either.AsAnsiLiteral())
+	if (LhsEither.AnsiLiteral == RhsEither.AnsiLiteral)
 	{
 		return true;
 	}
 
-	if (bLiteralIsWide)
+	ELiteralType LhsType = LiteralType;
+	ELiteralType RhsType = Rhs.LiteralType;
+
+	// Always compare with LhsType <= RhsType to simplify the branches below.
+	if (RhsType < LhsType)
 	{
-		return Rhs.bLiteralIsWide ? FPlatformString::Stricmp(Either.AsWideLiteral(), Rhs.Either.AsWideLiteral()) == 0
-								: FPlatformString::Stricmp(Either.AsWideLiteral(), Rhs.Either.AsAnsiLiteral()) == 0;
+		Swap(LhsType, RhsType);
+		Swap(LhsEither, RhsEither);
 	}
-	else
+
+	if (RhsType == ELiteralType::WideLiteral)
 	{
-		return Rhs.bLiteralIsWide ? FPlatformString::Stricmp(Either.AsAnsiLiteral(), Rhs.Either.AsWideLiteral()) == 0
-								: FPlatformString::Stricmp(Either.AsAnsiLiteral(), Rhs.Either.AsAnsiLiteral()) == 0;	
+		switch (LhsType)
+		{
+		case ELiteralType::AnsiLiteral:
+			return FPlatformString::Stricmp(LhsEither.AnsiLiteral, RhsEither.WideLiteral) == 0;
+		case ELiteralType::WideLiteral:
+			return FPlatformString::Stricmp(LhsEither.WideLiteral, RhsEither.WideLiteral) == 0;
+		case ELiteralType::Utf8Literal:
+			return FPlatformString::Stricmp(*WriteToString<FName::StringBufferSize>(LhsEither.Utf8Literal), RhsEither.WideLiteral) == 0;
+		}
 	}
+	return (RhsType == ELiteralType::Utf8Literal) ?
+		FPlatformString::Stricmp(LhsEither.Utf8Literal, RhsEither.Utf8Literal) == 0 :
+		FPlatformString::Stricmp(LhsEither.AnsiLiteral, RhsEither.AnsiLiteral) == 0;
 }
 
 /*-----------------------------------------------------------------------------
diff --git a/Engine/Source/Runtime/Core/Public/UObject/NameTypes.h b/Engine/Source/Runtime/Core/Public/UObject/NameTypes.h
--- a/Engine/Source/Runtime/Core/Public/UObject/NameTypes.h
+++ b/Engine/Source/Runtime/Core/Public/UObject/NameTypes.h
@@ -61,9 +61,9 @@ struct FNameEntryId
 	// Default initialize to be equal to NAME_None
 	constexpr FNameEntryId() : Value(0) {}
 	constexpr FNameEntryId(ENoInit) {}
-	explicit FNameEntryId(FIntrusiveUnsetOptionalState) : Value(~0u) {}
+	constexpr explicit FNameEntryId(FIntrusiveUnsetOptionalState) : Value(~0u) {}
 
-	bool IsNone() const
+	constexpr bool IsNone() const
 	{
 		return Value == 0;
 	}
@@ -93,10 +93,10 @@ struct FNameEntryId
 	}
 
 	// Returns true if this FNameEntryId is not equivalent to NAME_None
-	explicit operator bool() const { return Value != 0; }
+	constexpr explicit operator bool() const { return Value != 0; }
 
 	/** Get process specific integer */
-	uint32 ToUnstableInt() const { return Value; }
+	constexpr uint32 ToUnstableInt() const { return Value; }
 
 	/** Create from unstable int produced by this process */
 	static FNameEntryId FromUnstableInt(uint32 UnstableInt)
@@ -200,6 +200,49 @@ enum EFindName
 	FNAME_Add
 };
 
+namespace UE::Core::Private
+{
+
+/**
+ * Parse the number from the end of a string like Prefix_1.
+ *
+ * Number must be non-negative, less than MAX_int32, and not have a leading 0 unless the number is 0.
+ * Returns the internal representation of the external number in the string. Example: Prefix_0 returns 1.
+ *
+ * @param Name   Name to parse the trailing number from.
+ * @param InOutLen   Length of the input name, updated to exclude the _<N> suffix if a number is parsed.
+ * @return Internal value of the number that was parsed, or NAME_NO_NUMBER_INTERNAL if no number was parsed.
+ */
+template <typename CharType>
+static constexpr uint32 ParseNumberFromName(const CharType* Name, int32& InOutLen)
+{
+	int32 Digits = 0;
+	for (const CharType* It = Name + InOutLen - 1; It >= Name && *It >= '0' && *It <= '9'; --It)
+	{
+		++Digits;
+	}
+
+	const CharType* FirstDigit = Name + InOutLen - Digits;
+	constexpr int32 MaxDigitsInt32 = 10;
+	if (Digits && Digits < InOutLen && FirstDigit[-1] == '_' && Digits <= MaxDigitsInt32 && (Digits == 1 || *FirstDigit != '0'))
+	{
+		int64 Number = 0;
+		for (int32 Index = 0; Index < Digits; ++Index)
+		{
+			Number = 10 * Number + (FirstDigit[Index] - '0');
+		}
+		if (Number < MAX_int32)
+		{
+			InOutLen -= 1 + Digits;
+			return static_cast<uint32>(NAME_EXTERNAL_TO_INTERNAL(Number));
+		}
+	}
+
+	return NAME_NO_NUMBER_INTERNAL;
+}
+
+} // UE::Core::Private
+
 /*----------------------------------------------------------------------------
 	FNameEntry.
 ----------------------------------------------------------------------------*/
@@ -1623,33 +1666,53 @@ private:
 /** Lazily constructed FName that helps avoid allocating FNames during static initialization */
 class FLazyName
 {
+	// NOTE: Numeric values are used for comparison in UEOpEquals.
+	enum class ELiteralType : uint8
+	{
+		None = 0,
+		AnsiLiteral = 1,
+		Utf8Literal = 2,
+		WideLiteral = 3,
+	};
+
 public:
-	FLazyName()
-		: Either(FNameEntryId(), FNameEntryId())
-	{}
+	constexpr FLazyName() = default;
 
 	/** @param Literal must be a string literal */
-	template<int N>
-	FLazyName(const WIDECHAR(&Literal)[N])
-		: Either(Literal)
+	template <int N>
+	constexpr FLazyName(const ANSICHAR (&Literal)[N])
+		: Either{.AnsiLiteral = Literal}
 		, Number(ParseNumber(Literal, N - 1))
-		, bLiteralIsWide(true)
-	{}
+		, LiteralType(ELiteralType::AnsiLiteral)
+	{
+	}
 
 	/** @param Literal must be a string literal */
-	template<int N>
-	FLazyName(const ANSICHAR(&Literal)[N])
-		: Either(Literal)
+	template <int N>
+	constexpr FLazyName(const UTF8CHAR (&Literal)[N])
+		: Either{.Utf8Literal = Literal}
 		, Number(ParseNumber(Literal, N - 1))
-		, bLiteralIsWide(false)
-	{}
+		, LiteralType(ELiteralType::Utf8Literal)
+	{
+	}
+
+	/** @param Literal must be a string literal */
+	template <int N>
+	constexpr FLazyName(const WIDECHAR (&Literal)[N])
+		: Either{.WideLiteral = Literal}
+		, Number(ParseNumber(Literal, N - 1))
+		, LiteralType(ELiteralType::WideLiteral)
+	{
+	}
 
 	explicit FLazyName(FName Name)
-		: Either(Name.GetComparisonIndex(), Name.GetDisplayIndex())
+		: Either{.PackedName = FLiteralOrName::PackName(Name.GetComparisonIndex(), Name.GetDisplayIndex())}
 		, Number(Name.GetNumber())
-	{}
+		, LiteralType(ELiteralType::None)
+	{
+	}
 	
-	operator FName() const
+	UE_REWRITE operator FName() const
 	{
 		return Resolve();
 	}
@@ -1660,76 +1723,70 @@ public:
 	CORE_API FUtf8String ToUtf8String() const;
 
 private:
-	struct FLiteralOrName
-	{
-		// NOTE: uses high bit of pointer for flag; this may be an issue in future when high byte of address may be used for features like hardware ASAN
-		static constexpr uint64 LiteralFlag = uint64(1) << (sizeof(uint64) * 8 - 1);
-		static constexpr uint32 DisplayNameShift = 32;
-
-		explicit FLiteralOrName(const ANSICHAR* Literal)
-			: Int(reinterpret_cast<uint64>(Literal) | LiteralFlag)
-		{}
-		
-		explicit FLiteralOrName(const WIDECHAR* Literal)
-			: Int(reinterpret_cast<uint64>(Literal) | LiteralFlag)
-		{}
-
-		explicit FLiteralOrName(FNameEntryId ComparisionId, FNameEntryId DisplayId)
-			: Int(ComparisionId.ToUnstableInt() |
-				(WITH_CASE_PRESERVING_NAME ? ((static_cast<uint64>(DisplayId.ToUnstableInt()) << DisplayNameShift)) : 0))
-		{
-			// FName indices fit within 31 bits -> (DisplayId & LiteralFlag) == 0 -> IsName() == true.
-			checkName(IsName());
-		}
+	union FLiteralOrName
+	{
+		// NOTE: This uses the high bit for the flag. This may be an issue in the future if the high byte
+		// of an address starts being used for features like hardware ASan.
+		static constexpr uint64 IsNameFlag = uint64(1) << (sizeof(uint64) * 8 - 1);
+		static constexpr uint32 DisplayIdShift = 32;
+
+		const ANSICHAR* AnsiLiteral;
+		const UTF8CHAR* Utf8Literal;
+		const WIDECHAR* WideLiteral;
+		mutable uint64 PackedName = 0;
 
-		bool IsName() const
+		static constexpr uint64 PackName(FNameEntryId ComparisonId, FNameEntryId DisplayId)
 		{
-			return (LiteralFlag & Int) == 0;
+			// FNameEntryId fits in 31 bits -> (DisplayId & IsNameFlag) != 0 -> IsName() == true.
+			return IsNameFlag | ComparisonId.ToUnstableInt() |
+				(WITH_CASE_PRESERVING_NAME ? (uint64(DisplayId.ToUnstableInt()) << DisplayIdShift) : 0);
 		}
 
-		bool IsLiteral() const
+		constexpr bool IsName() const
 		{
-			return (LiteralFlag & Int) != 0;
+			return !!(PackedName & IsNameFlag);
 		}
 
-		FNameEntryId GetComparisionId() const
+		inline FNameEntryId GetComparisonId() const
 		{
-			return FNameEntryId::FromUnstableInt(static_cast<uint32>(Int));
+			return FNameEntryId::FromUnstableInt(static_cast<uint32>(PackedName));
 		}
 
-		FNameEntryId GetDisplayId() const
+		inline FNameEntryId GetDisplayId() const
 		{
 #if WITH_CASE_PRESERVING_NAME
-			return FNameEntryId::FromUnstableInt(static_cast<uint32>(Int >> DisplayNameShift)); // Can assume LiteralFlag == 0 
+			return FNameEntryId::FromUnstableInt(static_cast<uint32>((PackedName & ~IsNameFlag) >> DisplayIdShift));
 #else
-			return GetComparisionId();
+			return GetComparisonId();
 #endif
 		}
-		
-		const ANSICHAR* AsAnsiLiteral() const
+	};
+
+	FLiteralOrName Either;
+	uint32 Number = NAME_NO_NUMBER_INTERNAL;
+	ELiteralType LiteralType = ELiteralType::None;
+
+	// Parse the number at compile time when possible to avoid runtime initialization overhead.
+	// Number must be stored upon construction to allow Resolve() to convert a literal
+	// to a name without synchronization. If Number had to be written at the same time
+	// it would require synchronization.
+	template <typename CharType>
+	UE_REWRITE static constexpr uint32 ParseNumber(const CharType* Literal, int32 Len)
+	{
+		UE_IF_CONSTEVAL
 		{
-			return reinterpret_cast<const ANSICHAR*>(Int & ~LiteralFlag);
+			return UE::Core::Private::ParseNumberFromName(Literal, Len);
 		}
-
-		const WIDECHAR* AsWideLiteral() const
+		else
 		{
-			return reinterpret_cast<const WIDECHAR*>(Int & ~LiteralFlag);
+			return CallParseNumber(Literal, Len);
 		}
+	}
 
-		uint64 Int;
-	};
-
-	mutable FLiteralOrName Either;
-	mutable uint32 Number = 0;
-
-	// Distinguishes WIDECHAR* and ANSICHAR* literals, doesn't indicate if literal contains any wide characters 
-	bool bLiteralIsWide = false;
-	
-	CORE_API static uint32 ParseNumber(const WIDECHAR* Literal, int32 Len);
-	CORE_API static uint32 ParseNumber(const ANSICHAR* Literal, int32 Len);
+	template <typename CharType>
+	CORE_API static uint32 CallParseNumber(const CharType* Literal, int32 Len);
 
 public:
-
 	bool UEOpEquals(FName Rhs) const
 	{
 		// If !Rhs.IsNone(), we have started creating FNames
@@ -1738,13 +1795,17 @@ public:
 		{
 			return Rhs == Resolve();
 		}
-		else if (!bLiteralIsWide)
+		else if (LiteralType == ELiteralType::AnsiLiteral)
+		{
+			return Rhs == Either.AnsiLiteral;
+		}
+		else if (LiteralType == ELiteralType::Utf8Literal)
 		{
-			return Rhs == Either.AsAnsiLiteral();
+			return Rhs == Either.Utf8Literal;
 		}
 		else
 		{
-			return Rhs == Either.AsWideLiteral();
+			return Rhs == Either.WideLiteral;
 		}
 	}
 
