Added UE::CFloatingPoint and used it in FMath.
Removed deprecated FMath::Wrap for integral types.

diff --git a/Engine/Source/Runtime/Core/Public/Concepts/FloatingPoint.h b/Engine/Source/Runtime/Core/Public/Concepts/FloatingPoint.h
new file mode 100644
--- /dev/null
+++ b/Engine/Source/Runtime/Core/Public/Concepts/FloatingPoint.h
@@ -0,0 +1,14 @@
+// Copyright Epic Games, Inc. All Rights Reserved.
+
+#pragma once
+
+#include <type_traits>
+
+namespace UE
+{
+	/**
+	 * Concept which describes a floating point type.  We use this instead of std::floating_point because <concepts> isn't a well supported header yet.
+	 */
+	template <typename T>
+	concept CFloatingPoint = std::is_floating_point_v<T>;
+}
diff --git a/Engine/Source/Runtime/Core/Public/Math/UnrealMathUtility.h b/Engine/Source/Runtime/Core/Public/Math/UnrealMathUtility.h
--- a/Engine/Source/Runtime/Core/Public/Math/UnrealMathUtility.h
+++ b/Engine/Source/Runtime/Core/Public/Math/UnrealMathUtility.h
@@ -3,6 +3,7 @@
 #pragma once
 
 #include "CoreTypes.h"
+#include "Concepts/FloatingPoint.h"
 #include "Concepts/Integral.h"
 #include "Misc/AssertionMacros.h"
 #include "HAL/PlatformMath.h"
@@ -603,9 +604,11 @@ public:
 	/** Clamps X to be between Min and Max, inclusive. Overload to support mixed int64/int32 types. */
 	[[nodiscard]] static constexpr UE_FORCEINLINE_HINT int64 Clamp(const int64 X, const int32 Min, const int32 Max) { return Clamp<int64>(X, Min, Max); }
 
-private:
-	template< class T >
-	[[nodiscard]] static constexpr inline T WrapImpl(const T X, const T Min, const T Max)
+public:
+	/** Wraps X to be between Min and Max, inclusive. */
+	/** When X can wrap to both Min and Max, it will wrap to Min if it lies below the range and wrap to Max if it is above the range. */
+	template <UE::CFloatingPoint T>
+	[[nodiscard]] static constexpr UE_FORCEINLINE_HINT T Wrap(const T X, const T Min, const T Max)
 	{
 		// Use unsigned type for integers to allow for large ranges which don't overflow on subtraction
 		// We don't do that for floating point types because there are no unsigned versions of those.  We
@@ -637,27 +640,6 @@ private:
 		return EndVal;
 	}
 
-public:
-	/** Wraps X to be between Min and Max, inclusive. */
-	/** When X can wrap to both Min and Max, it will wrap to Min if it lies below the range and wrap to Max if it is above the range. */
-	template <
-		typename T
-		UE_REQUIRES(std::is_floating_point_v<T>)
-	>
-	[[nodiscard]] static constexpr UE_FORCEINLINE_HINT T Wrap(const T X, const T Min, const T Max)
-	{
-		return WrapImpl(X, Min, Max);
-	}
-	template <
-		typename T
-		UE_REQUIRES(!std::is_floating_point_v<T>)
-	>
-	UE_DEPRECATED(5.5, "Use of FMath::Wrap with non-floating point arguments is deprecated - consider FMath::WrapExclusive instead.")
-	[[nodiscard]] static constexpr UE_FORCEINLINE_HINT T Wrap(const T X, const T Min, const T Max)
-	{
-		return WrapImpl(X, Min, Max);
-	}
-
 	/** Wraps X to be between Min and Max, exclusive. Will never return Max. */
 	template <UE::CIntegral T>
 	[[nodiscard]] static constexpr inline T WrapExclusive(const T X, const T Min, const T Max)
@@ -774,10 +756,7 @@ public:
 	* @param ScalarCos	Pointer to where the Cos result should be stored
 	* @param Value  input angles 
 	*/
-	template <
-		typename T
-		UE_REQUIRES(std::is_floating_point_v<T>)
-	>
+	template <UE::CFloatingPoint T>
 	static constexpr inline void SinCos(std::decay_t<T>* ScalarSin, std::decay_t<T>* ScalarCos, T  Value )
 	{
 		// Map Value to y in [-pi,pi], x = 2*pi*quotient + remainder.
@@ -973,10 +952,7 @@ public:
 	}
 
 	/** Given a heading which may be outside the +/- PI range, 'unwind' it back into that range. */
-	template <
-		typename T
-		UE_REQUIRES(std::is_floating_point_v<T>)
-	>
+	template <UE::CFloatingPoint T>
 	[[nodiscard]] static constexpr T UnwindRadians(T A)
 	{
 		while(A > UE_PI)
@@ -993,10 +969,7 @@ public:
 	}
 
 	/** Utility to ensure angle is between +/- 180 degrees by unwinding. */
-	template <
-		typename T
-		UE_REQUIRES(std::is_floating_point_v<T>)
-	>
+	template <UE::CFloatingPoint T>
 	[[nodiscard]] static constexpr T UnwindDegrees(T A)
 	{
 		while(A > 180.f)
@@ -1099,11 +1072,7 @@ public:
 	// Interpolation Functions
 
 	/** Calculates the percentage along a line from MinValue to MaxValue that Value is. */
-	template <
-		typename T,
-		typename T2
-		UE_REQUIRES(std::is_floating_point_v<T>)
-	>
+	template <UE::CFloatingPoint T, typename T2>
 	[[nodiscard]] static constexpr inline auto GetRangePct(T MinValue, T MaxValue, T2 Value)
 	{
 		// Avoid Divide by Zero.
@@ -1119,22 +1088,14 @@ public:
 	}
 
 	/** Same as above, but taking a 2d vector as the range. */
-	template <
-		typename T,
-		typename T2
-		UE_REQUIRES(std::is_floating_point_v<T>)
-	>
+	template <UE::CFloatingPoint T, typename T2>
 	[[nodiscard]] static auto GetRangePct(UE::Math::TVector2<T> const& Range, T2 Value)
 	{
 		return GetRangePct(Range.X, Range.Y, Value);
 	}
 	
 	/** Basically a Vector2d version of Lerp. */
-	template <
-		typename T,
-		typename T2
-		UE_REQUIRES(std::is_floating_point_v<T>)
-	>
+	template <UE::CFloatingPoint T, typename T2>
 	[[nodiscard]] static auto GetRangeValue(UE::Math::TVector2<T> const& Range, T2 Pct)
 	{
 		return Lerp(Range.X, Range.Y, Pct);
@@ -1305,11 +1266,7 @@ public:
 	 *
 	 * @return  Interpolated value
 	 */
-	template <
-		typename T,
-		typename U
-		UE_REQUIRES(std::is_floating_point_v<U>)
-	>
+	template <typename T, UE::CFloatingPoint U>
 	[[nodiscard]] static constexpr inline T CubicInterpDerivative( const T& P0, const T& T0, const T& P1, const T& T1, const U& A )
 	{
 		T a = 6.f*P0 + 3.f*T0 + 3.f*T1 - 6.f*P1;
@@ -1330,11 +1287,7 @@ public:
 	 *
 	 * @return  Interpolated value
 	 */
-	template <
-		typename T,
-		typename U
-		UE_REQUIRES(std::is_floating_point_v<U>)
-	>
+	template <typename T, UE::CFloatingPoint U>
 	[[nodiscard]] static constexpr inline T CubicInterpSecondDerivative( const T& P0, const T& T0, const T& P1, const T& T1, const U& A )
 	{
 		T a = 12.f*P0 + 6.f*T0 + 6.f*T1 - 12.f*P1;
@@ -2665,11 +2618,9 @@ namespace LWC
 	 */
 
 	// Generic float type to int type, to enable specializations below.
-	template<typename OutIntType, typename InFloatType>
+	template <UE::CIntegral OutIntType, UE::CFloatingPoint InFloatType>
 	inline OutIntType FloatToIntCastChecked(InFloatType FloatValue)
 	{
-		static_assert(std::is_floating_point_v<InFloatType>, "Only floating point input type supported!");
-		static_assert(std::is_integral_v<OutIntType>, "Only integral output type supported!");
 		return (OutIntType)(FloatValue);
 	}
 
diff --git a/Engine/Source/Runtime/Engine/Classes/Kismet/KismetMathLibrary.inl b/Engine/Source/Runtime/Engine/Classes/Kismet/KismetMathLibrary.inl
--- a/Engine/Source/Runtime/Engine/Classes/Kismet/KismetMathLibrary.inl
+++ b/Engine/Source/Runtime/Engine/Classes/Kismet/KismetMathLibrary.inl
@@ -271,9 +271,8 @@ int32 UKismetMathLibrary::Clamp(int32 V, int32 A, int32 B)
 KISMET_MATH_INLINE
 int32 UKismetMathLibrary::Wrap(int32 Value, int32 Min, int32 Max)
 {
-PRAGMA_DISABLE_DEPRECATION_WARNINGS
-	return FMath::Wrap(Value, Min, Max);
-PRAGMA_ENABLE_DEPRECATION_WARNINGS
+	// FMath::Wrap only takes floating point types, so operate in doubles, which can hold int32 losslessly.
+	return (int32)FMath::Wrap((double)Value, (double)Min, (double)Max);
 }
 
 KISMET_MATH_INLINE
