Removed a redundant comment and handling of integral types from FMath::Wrap.

diff --git a/Engine/Source/Runtime/Core/Public/Math/UnrealMathUtility.h b/Engine/Source/Runtime/Core/Public/Math/UnrealMathUtility.h
--- a/Engine/Source/Runtime/Core/Public/Math/UnrealMathUtility.h
+++ b/Engine/Source/Runtime/Core/Public/Math/UnrealMathUtility.h
@@ -8,7 +8,6 @@
 #include "Misc/AssertionMacros.h"
 #include "HAL/PlatformMath.h"
 #include "Math/MathFwd.h"
-#include "Templates/Identity.h"
 #include "Templates/Requires.h"
 
 // Assert on non finite numbers. Used to track NaNs.
@@ -610,16 +609,10 @@ public:
 	template <UE::CFloatingPoint T>
 	[[nodiscard]] static constexpr UE_FORCEINLINE_HINT T Wrap(const T X, const T Min, const T Max)
 	{
-		// Use unsigned type for integers to allow for large ranges which don't overflow on subtraction
-		// We don't do that for floating point types because there are no unsigned versions of those.  We
-		// could bump up to the next size up (though there's no `long double` on some platforms), but such
-		// values near the extremes are unlikely.
-		using SizeType = typename std::conditional_t<std::is_integral_v<T>, std::make_unsigned<T>, TIdentity<T>>::type;
-
 		// Our asserts are not constexpr-friendly yet
 		// checkSlow(Min <= Max);
 
-		SizeType Size = (SizeType)Max - (SizeType)Min;
+		T Size = Max - Min;
 		if (Size == 0)
 		{
 			// Guard against zero-sized ranges causing division by zero.
@@ -629,13 +622,13 @@ public:
 		T EndVal = X;
 		if (EndVal < Min)
 		{
-			SizeType Mod = FMath::Modulo((SizeType)((SizeType)Min - (SizeType)EndVal), Size);
-			EndVal = (Mod != (T)0) ? (T)((SizeType)Max - Mod) : Min;
+			T Mod = FMath::Modulo(Min - EndVal, Size);
+			EndVal = (Mod != (T)0) ? Max - Mod : Min;
 		}
 		else if (EndVal > Max)
 		{
-			SizeType Mod = FMath::Modulo((SizeType)((SizeType)EndVal - (SizeType)Max), Size);
-			EndVal = (Mod != (T)0) ? (T)((SizeType)Min + Mod) : Max;
+			T Mod = FMath::Modulo(EndVal - Max, Size);
+			EndVal = (Mod != (T)0) ? Min + Mod : Max;
 		}
 		return EndVal;
 	}
@@ -2662,8 +2655,6 @@ namespace LWC
 } // namespace LWC
 } // namespace UE
 
-#if UE_ENABLE_INCLUDE_ORDER_DEPRECATED_IN_5_4
-#include "Templates/Decay.h"
-#include "Templates/IsFloatingPoint.h"
-#include "Templates/IsIntegral.h"
+#if UE_ENABLE_INCLUDE_ORDER_DEPRECATED_IN_5_7
+#include "Templates/Identity.h"
 #endif
