Fixed some unused variable warnings.

diff --git a/Engine/Source/Runtime/Core/Public/Misc/CString.h b/Engine/Source/Runtime/Core/Public/Misc/CString.h
--- a/Engine/Source/Runtime/Core/Public/Misc/CString.h
+++ b/Engine/Source/Runtime/Core/Public/Misc/CString.h
@@ -909,7 +909,6 @@ const typename TCString<T>::CharType* TCString<T>::Strnstr(const CharType* Str,
 		constexpr uint64 SignBit64  = sizeof(CharType) == 1 ? 0x8080808080808080ull : 0x8000800080008000ull;
 		constexpr uint64 CharPer64  = sizeof(uint64) / sizeof(CharType);
 		constexpr uint64 BitPerChar = sizeof(CharType) == 1 ? 8 : 16;
-		constexpr uint64 BitShift   = sizeof(CharType) == 1 ? 3 : 4;
 
 		InStrLen -= FindSuffixLength;
 		if (InStrLen >= CharPer64)
diff --git a/Engine/Source/Runtime/Core/Public/Templates/Function.h b/Engine/Source/Runtime/Core/Public/Templates/Function.h
--- a/Engine/Source/Runtime/Core/Public/Templates/Function.h
+++ b/Engine/Source/Runtime/Core/Public/Templates/Function.h
@@ -145,10 +145,9 @@ namespace UE::Core::Private::Function
 		FFunctionStorage& operator=(FFunctionStorage&& Other) = delete;
 		FFunctionStorage& operator=(const FFunctionStorage& Other) = delete;
 
-		void* BindCopy(const FFunctionStorage& Other)
+		void BindCopy(const FFunctionStorage& Other)
 		{
-			void* NewObj = Other.GetBoundObject()->CloneToEmptyStorage(this);
-			return NewObj;
+			Other.GetBoundObject()->CloneToEmptyStorage(this);
 		}
 
 		IFunction_OwnedObject* GetBoundObject() const
@@ -359,7 +358,7 @@ namespace UE::Core::Private::Function
 				}
 			}
 
-			void* NewPtr = Storage.BindCopy(Other.Storage);
+			Storage.BindCopy(Other.Storage);
 		}
 
 		TFunctionRefBase(const TFunctionRefBase& Other)
@@ -373,7 +372,7 @@ namespace UE::Core::Private::Function
 				}
 			}
 
-			void* NewPtr = Storage.BindCopy(Other.Storage);
+			Storage.BindCopy(Other.Storage);
 		}
 
 		template <UE::CNotCVRefTo<TFunctionRefBase> FunctorType>
@@ -467,11 +466,9 @@ namespace UE::Core::Private::Function
 			return &InFunc;
 		}
 
-		void* BindCopy(const FFunctionRefStoragePolicy& Other)
+		void BindCopy(const FFunctionRefStoragePolicy& Other)
 		{
-			void* OtherPtr = Other.Ptr;
-			Ptr = OtherPtr;
-			return OtherPtr;
+			Ptr = Other.Ptr;
 		}
 
 		/**
