Fix to missing include and warning in JNI refactor

diff --git a/Engine/Source/Runtime/Core/Public/Android/AndroidJavaEnv.h b/Engine/Source/Runtime/Core/Public/Android/AndroidJavaEnv.h
--- a/Engine/Source/Runtime/Core/Public/Android/AndroidJavaEnv.h
+++ b/Engine/Source/Runtime/Core/Public/Android/AndroidJavaEnv.h
@@ -4,6 +4,7 @@
 
 #include "CoreTypes.h"
 #include "Containers/Array.h"
+#include "Containers/StaticArray.h"
 #include "Containers/UnrealString.h"
 #include <algorithm>
 #include <type_traits>
diff --git a/Engine/Source/Runtime/Core/Public/Android/AndroidJniGameActivity.h b/Engine/Source/Runtime/Core/Public/Android/AndroidJniGameActivity.h
--- a/Engine/Source/Runtime/Core/Public/Android/AndroidJniGameActivity.h
+++ b/Engine/Source/Runtime/Core/Public/Android/AndroidJniGameActivity.h
@@ -130,7 +130,7 @@ namespace UE::Jni::Asis
 	{
 		static constexpr FAnsiStringView ClassName = "com/epicgames/asis/AsisGameActivity";
 	
-		static jstring JNICALL nativeGetObbComment(JNIEnv* env, jobject clazz);
+		static jstring JNICALL nativeGetObbComment(JNIEnv* env, jclass clazz);
 		static void JNICALL nativeMain(JNIEnv* env, jclass clazz, jstring projectModule);
 		static void JNICALL nativeSetCommandline(JNIEnv* env, jclass clazz, jstring commandline);
 		static void JNICALL nativeAppCommand(JNIEnv* env, jclass clazz, jint cmd);
