Fix to USE_ANDROID_JNI missing sround some FGameActivity usage

diff --git a/Engine/Source/Runtime/Core/Private/Android/AndroidJavaEnv.cpp b/Engine/Source/Runtime/Core/Private/Android/AndroidJavaEnv.cpp
--- a/Engine/Source/Runtime/Core/Private/Android/AndroidJavaEnv.cpp
+++ b/Engine/Source/Runtime/Core/Private/Android/AndroidJavaEnv.cpp
@@ -325,6 +325,7 @@ namespace UE::Jni
 	
 	void FEnv::Initialize() const
 	{
+#if USE_ANDROID_JNI
 		using namespace Java::Lang;
 
 		Class<FString> = NewGlobalRef(*FScopedJavaObject<TClass<FString>*>{Value->FindClass(ClassName<FString>.GetData())});
@@ -335,5 +336,6 @@ namespace UE::Jni
 		Initialize<FGameActivity, false>(Class<FGameActivity>);
 
 		ClassLoader = NewGlobalRef(*Get<FClassLoader* ()>(*FScopedJavaObject<TClass<FClass>*>{Value->FindClass(ClassName<FClass>.GetData())}, "getClassLoader")(Class<FGameActivity>));
-	}
+#endif
+		}
 }
diff --git a/Engine/Source/Runtime/Core/Private/Android/AndroidPlatformCrashContext.cpp b/Engine/Source/Runtime/Core/Private/Android/AndroidPlatformCrashContext.cpp
--- a/Engine/Source/Runtime/Core/Private/Android/AndroidPlatformCrashContext.cpp
+++ b/Engine/Source/Runtime/Core/Private/Android/AndroidPlatformCrashContext.cpp
@@ -52,6 +52,7 @@ static void SetCrashContextOnGameThread(const FString KeyIn, const FString Value
 	}
 }
 
+#if USE_ANDROID_JNI
 void JNICALL UE::Jni::FGameActivity::nativeCrashContextSetStringKey(JNIEnv* jenv, jobject thiz, jstring JavaKey, jstring JavaValue)
 {
 	FString Key, Value;
@@ -89,6 +90,7 @@ void JNICALL UE::Jni::FGameActivity::nativeCrashContextSetDoubleKey(JNIEnv* jenv
 	SetCrashContextOnGameThread(Key, *TTypeToString<double>::ToString(JavaValue));
 }
 
+#endif
 /** Implement platform specific static cleanup function */
 void FGenericCrashContext::CleanupPlatformSpecificFiles()
 {
diff --git a/Engine/Source/Runtime/Core/Private/Android/AndroidPlatformFile.cpp b/Engine/Source/Runtime/Core/Private/Android/AndroidPlatformFile.cpp
--- a/Engine/Source/Runtime/Core/Private/Android/AndroidPlatformFile.cpp
+++ b/Engine/Source/Runtime/Core/Private/Android/AndroidPlatformFile.cpp
@@ -127,6 +127,7 @@ static FString AndroidLogDir;
 extern jobject AndroidJNI_GetJavaAssetManager();
 extern AAssetManager * AndroidThunkCpp_GetAssetManager();
 
+#if USE_ANDROID_JNI
 //This function is declared in the Java-defined class, GameActivity.java: "public native void nativeSetObbInfo(String PackageName, int Version, int PatchVersion);"
 void JNICALL UE::Jni::FGameActivity::nativeSetObbInfo(JNIEnv* jenv, jobject thiz, jstring ProjectName, jstring PackageName, jint Version, jint PatchVersion, jstring AppType)
 {
@@ -144,8 +145,9 @@ jstring JNICALL UE::Jni::Asis::FAsisGameActivity::nativeGetObbComment(JNIEnv* je
 {
 	return jenv->NewStringUTF(TCHAR_TO_UTF8(*GLastOBBComment));
 }
-#endif
 
+#endif
+#endif
 // Constructs the base path for any files which are not in OBB/pak data
 const FString &GetFileBasePath()
 {
diff --git a/Engine/Source/Runtime/Core/Private/Android/AndroidPlatformMisc.cpp b/Engine/Source/Runtime/Core/Private/Android/AndroidPlatformMisc.cpp
--- a/Engine/Source/Runtime/Core/Private/Android/AndroidPlatformMisc.cpp
+++ b/Engine/Source/Runtime/Core/Private/Android/AndroidPlatformMisc.cpp
@@ -2482,6 +2482,7 @@ bool FAndroidMisc::AllowThreadHeartBeat()
 	return AllowThreadHeartBeatOnce == 1;
 }
 
+#if USE_ANDROID_JNI
 void JNICALL UE::Jni::FGameActivity::nativeSetConfigRulesVariables(JNIEnv* jenv, jobject thiz, Java::Lang::TArray<jstring>* KeyValuePairs)
 {
 	int32 Count = jenv->GetArrayLength(KeyValuePairs);
@@ -2498,10 +2499,12 @@ void JNICALL UE::Jni::FGameActivity::nativeSetConfigRulesVariables(JNIEnv* jenv,
 	AvailableConfigRuleVarsEvent.Notify();
 }
 
+#endif
 extern bool AndroidThunkCpp_HasMetaDataKey(const FString& Key);
 
 static bool bDetectedDebugger = false;
 
+#if USE_ANDROID_JNI
 void JNICALL UE::Jni::FGameActivity::nativeSetAndroidStartupState(JNIEnv* jenv, jobject thiz, jboolean bDebuggerAttached)
 {
 	// if Java debugger attached, mark detected (but don't lose previous trigger state)
@@ -2511,6 +2514,7 @@ void JNICALL UE::Jni::FGameActivity::nativeSetAndroidStartupState(JNIEnv* jenv,
 	}
 }
 
+#endif
 bool FAndroidMisc::UseNewWindowBehavior()
 {
 #if USE_ANDROID_STANDALONE
@@ -2814,7 +2818,6 @@ bool FAndroidMisc::HasActiveWiFiConnection()
 	return (ConnectionType == ENetworkConnectionType::WiFi ||
 			ConnectionType == ENetworkConnectionType::WiMAX);
 }
-#endif
 
 void JNICALL UE::Jni::FGameActivity::nativeNetworkChanged(JNIEnv* jenv, jobject thiz)
 {
@@ -2830,6 +2833,7 @@ void JNICALL UE::Jni::FGameActivity::nativeNetworkChanged(JNIEnv* jenv, jobject
 	}
 }
 
+#endif
 static FAndroidMisc::ReInitWindowCallbackType OnReInitWindowCallback;
 
 FAndroidMisc::ReInitWindowCallbackType FAndroidMisc::GetOnReInitWindowCallback()
diff --git a/Engine/Source/Runtime/Core/Private/Android/AndroidPlatformProcess.cpp b/Engine/Source/Runtime/Core/Private/Android/AndroidPlatformProcess.cpp
--- a/Engine/Source/Runtime/Core/Private/Android/AndroidPlatformProcess.cpp
+++ b/Engine/Source/Runtime/Core/Private/Android/AndroidPlatformProcess.cpp
@@ -261,6 +261,7 @@ static bool EnableLittleCoreAffinity = false;
 static int32 BigCoreMask = 0;
 static int32 LittleCoreMask = 0;
 
+#if USE_ANDROID_JNI
 //This function is declared in the Java-defined class, GameActivity.java: "public native void nativeSetAffinityInfo(boolean bEnableAffinity, int bigCoreMask, int littleCoreMask);
 void JNICALL UE::Jni::FGameActivity::nativeSetAffinityInfo(JNIEnv* jenv, jobject thiz, jboolean bEnableAffinity, jint bigCoreMask, jint littleCoreMask)
 {
@@ -269,6 +270,7 @@ void JNICALL UE::Jni::FGameActivity::nativeSetAffinityInfo(JNIEnv* jenv, jobject
 	LittleCoreMask = littleCoreMask;
 }
 
+#endif
 uint64 FAndroidAffinity::GetLittleCoreMask()
 {
 	static int Mask = 0;
