Mobile: Use deferred shading by default
Deferred shading will use a multi-pass rendering by default
Forward shading will still use FBF approach when possible
#DisabledValidations NonUnityAndStaticAnalysis

diff --git a/Engine/Source/Runtime/Core/Private/HAL/ConsoleManager.cpp b/Engine/Source/Runtime/Core/Private/HAL/ConsoleManager.cpp
--- a/Engine/Source/Runtime/Core/Private/HAL/ConsoleManager.cpp
+++ b/Engine/Source/Runtime/Core/Private/HAL/ConsoleManager.cpp
@@ -3765,30 +3765,30 @@ static TAutoConsoleVariable<int32> CVarMobileHDR(
 
 static TAutoConsoleVariable<int32> CVarMobileShadingPath(
 	TEXT("r.Mobile.ShadingPath"),
-	0,
-	TEXT("0: Forward shading (default)\n"
-		 "1: Deferred shading (Mobile HDR is required for Deferred)"),
+	1,
+	TEXT("0: Forward shading\n"
+		 "1: Deferred shading (Mobile HDR is required for Deferred) (default)"),
 	ECVF_RenderThreadSafe | ECVF_ReadOnly);
 
 static TAutoConsoleVariable<int32> CVarMobileAllowDeferredShadingOpenGL(
 	TEXT("r.Mobile.AllowDeferredShadingOpenGL"),
-	0,
-	TEXT("0: Do not Allow Deferred Shading on OpenGL (default)\n"
-		 "1: Allow Deferred Shading on OpenGL"),
+	1,
+	TEXT("0: Do not Allow Deferred Shading on OpenGL\n"
+		 "1: Allow Deferred Shading on OpenGL (default)"),
 	ECVF_RenderThreadSafe | ECVF_ReadOnly);
 
 static TAutoConsoleVariable<int32> CVarMobileAllowFramebufferFetchOpenGL(
 	TEXT("r.Mobile.AllowFramebufferFetchOpenGL"),
-	1,
-	TEXT("0: Use multi-pass rendering without FBF and PLS extensions\n"
-		"1: Allow use of framebuffer fetch and PLS extensions (default)"),
+	0,
+	TEXT("0: Use multi-pass rendering without FBF and PLS extensions (only for a deferred shading) (default)\n"
+		"1: Allow use of framebuffer fetch and PLS extensions"),
 	ECVF_RenderThreadSafe | ECVF_ReadOnly);
 
 static TAutoConsoleVariable<int32> CVarMobileAllowFramebufferFetch(
 	TEXT("r.Mobile.AllowFramebufferFetch"),
-	1,
-	TEXT("0: Multi-pass rendering without using sub-passes in Vulkan or programmable blending on iOS. For OpenGL use r.Mobile.AllowFramebufferFetchOpenGL\n"
-		 "1: Allow use of sub-passes and programmable blending (default)"),
+	0,
+	TEXT("0: Multi-pass rendering without using sub-passes in Vulkan or programmable blending on iOS. For OpenGL use r.Mobile.AllowFramebufferFetchOpenGL (only for a deferred shading) (default)\n"
+		 "1: Allow use of sub-passes and programmable blending"),
 	ECVF_RenderThreadSafe | ECVF_ReadOnly);
 
 static TAutoConsoleVariable<int32> CVarMobileEnableStaticAndCSMShadowReceivers(
diff --git a/Engine/Source/Runtime/RenderCore/Private/RenderUtils.cpp b/Engine/Source/Runtime/RenderCore/Private/RenderUtils.cpp
--- a/Engine/Source/Runtime/RenderCore/Private/RenderUtils.cpp
+++ b/Engine/Source/Runtime/RenderCore/Private/RenderUtils.cpp
@@ -468,14 +468,16 @@ RENDERCORE_API bool MobileAllowFramebufferFetch(const FStaticShaderPlatform Plat
 	if (Platform == SP_OPENGL_ES3_1_ANDROID)
 	{
 		const static IConsoleVariable* CVarAllowFramebufferFetchOpenGL = IConsoleManager::Get().FindConsoleVariable(TEXT("r.Mobile.AllowFramebufferFetchOpenGL"));
-		return CVarAllowFramebufferFetchOpenGL->GetBool();
+		return !IsMobileDeferredShadingEnabled(Platform) || CVarAllowFramebufferFetchOpenGL->GetBool();
 	}
 	else if (IsVulkanMobilePlatform(Platform) || IsMetalMobilePlatform(Platform))
 	{
 		static FShaderPlatformCachedIniValue<bool> MobileAllowFramebufferFetchIniValue(TEXT("r.Mobile.AllowFramebufferFetch"));
-		return MobileAllowFramebufferFetchIniValue.Get(Platform);
+		return !IsMobileDeferredShadingEnabled(Platform) || MobileAllowFramebufferFetchIniValue.Get(Platform);
 	}
-	return false;
+	
+	// Only GLES, Vulkan and Metal have FBF
+	return false; 
 }
 
 RENDERCORE_API bool MobileRequiresSceneDepthAux(const FStaticShaderPlatform Platform)
