[GMalloc]
- Suppress warnings about GMalloc being read relaxed by using TSAN_ATOMIC

diff --git a/Engine/Source/Runtime/Core/Private/Misc/CoreGlobals.cpp b/Engine/Source/Runtime/Core/Private/Misc/CoreGlobals.cpp
--- a/Engine/Source/Runtime/Core/Private/Misc/CoreGlobals.cpp
+++ b/Engine/Source/Runtime/Core/Private/Misc/CoreGlobals.cpp
@@ -55,7 +55,11 @@ FConfigCacheIni*			GConfig					= nullptr;				/* Configuration database cache */
 ITransaction*				GUndo					= nullptr;				/* Transaction tracker, non-NULL when a transaction is in progress */
 FOutputDeviceConsole*		GLogConsole				= nullptr;				/* Console log hook */
 CORE_API FMalloc**			GFixedMallocLocationPtr = nullptr;				/* Memory allocator pointer location when PLATFORM_USES_FIXED_GMalloc_CLASS is true */
-CORE_API FMalloc*			UE::Private::GMalloc	= nullptr;				/* Memory allocator */
+#if USING_THREAD_SANITISER || USING_INSTRUMENTATION
+CORE_API TSAN_ATOMIC(FMalloc*) UE::Private::GMalloc;						/* Memory allocator */
+#else
+CORE_API FMalloc*			UE::Private::GMalloc = nullptr;				/* Memory allocator */
+#endif
 CORE_API FMalloc* const&	GMalloc	                = UE::Private::GMalloc;	/* Provides backwards compatibility for external users of un-namespaced GMalloc. */
 
 class FPropertyWindowManager*	GPropertyWindowManager	= nullptr;		/* Manages and tracks property editing windows */
diff --git a/Engine/Source/Runtime/Core/Private/ProfilingDebugging/TraceAuxiliary.cpp b/Engine/Source/Runtime/Core/Private/ProfilingDebugging/TraceAuxiliary.cpp
--- a/Engine/Source/Runtime/Core/Private/ProfilingDebugging/TraceAuxiliary.cpp
+++ b/Engine/Source/Runtime/Core/Private/ProfilingDebugging/TraceAuxiliary.cpp
@@ -2161,7 +2161,7 @@ void FTraceAuxiliary::Initialize(const TCHAR* CommandLine)
 		// If for any reason, memory tracing has been initialized without callstack tracing,
 		// we need to use the same allocator; otherwise we use the regular malloc.
 		FMalloc* TraceAllocator = MemoryTrace_GetAllocator();
-		CallstackTrace_Create(TraceAllocator ? TraceAllocator : UE::Private::GMalloc);
+		CallstackTrace_Create(TraceAllocator ? TraceAllocator : (FMalloc*)UE::Private::GMalloc);
 	}
 	{
 		TRACE_CPUPROFILER_EVENT_SCOPE(FTraceAux_CallstackTrace_Init);
diff --git a/Engine/Source/Runtime/Core/Public/Clang/ClangPlatformCodeAnalysis.h b/Engine/Source/Runtime/Core/Public/Clang/ClangPlatformCodeAnalysis.h
--- a/Engine/Source/Runtime/Core/Public/Clang/ClangPlatformCodeAnalysis.h
+++ b/Engine/Source/Runtime/Core/Public/Clang/ClangPlatformCodeAnalysis.h
@@ -175,6 +175,8 @@
 			__attribute__((always_inline)) T operator--() { return Value.fetch_sub(1, std::memory_order_relaxed) - 1; }
 			__attribute__((always_inline)) T operator--(int) { return Value.fetch_sub(1, std::memory_order_relaxed); }
 
+			__attribute__((always_inline)) T operator->() { return Value.load(std::memory_order_relaxed); }
+
 			__attribute__((always_inline)) TTSANSafeValue(const TTSANSafeValue& Other)
 			{
 				Value.store((T)Other, std::memory_order_relaxed);
diff --git a/Engine/Source/Runtime/Core/Public/HAL/MemoryBase.h b/Engine/Source/Runtime/Core/Public/HAL/MemoryBase.h
--- a/Engine/Source/Runtime/Core/Public/HAL/MemoryBase.h
+++ b/Engine/Source/Runtime/Core/Public/HAL/MemoryBase.h
@@ -35,7 +35,7 @@ enum
  */
 namespace UE::Private
 {
-	CORE_API extern class FMalloc* GMalloc;
+	CORE_API extern TSAN_ATOMIC(class FMalloc*) GMalloc;
 }
 
 /** 
