Fixed FMath::FindDeltaAngleDegrees and FMath::FindDeltaAngleRadians for input deltas outside [-540, 540] degrees.

diff --git a/Engine/Source/Runtime/Core/Public/Math/UnrealMathUtility.h b/Engine/Source/Runtime/Core/Public/Math/UnrealMathUtility.h
--- a/Engine/Source/Runtime/Core/Public/Math/UnrealMathUtility.h
+++ b/Engine/Source/Runtime/Core/Public/Math/UnrealMathUtility.h
@@ -894,26 +894,15 @@ public:
 		typename T2
 		UE_REQUIRES(std::is_floating_point_v<T> || std::is_floating_point_v<T2>)
 	>
-	[[nodiscard]] static constexpr auto FindDeltaAngleDegrees(T A1, T2 A2) -> decltype(A1 * A2)
+	[[nodiscard]] static constexpr auto FindDeltaAngleDegrees(T A1, T2 A2) -> decltype(A1 - A2)
 	{
 		// Find the difference
 		auto Delta = A2 - A1;
 
-		// If change is larger than 180
-		if (Delta > 180.0f)
-		{
-			// Flip to negative equivalent
-			Delta = Delta - 360.0f;
-		}
-		else if (Delta < -180.0f)
-		{
-			// Otherwise, if change is smaller than -180
-			// Flip to positive equivalent
-			Delta = Delta + 360.0f;
-		}
+		using DeltaType = decltype(Delta);
 
 		// Return delta in [-180,180] range
-		return Delta;
+		return Wrap(Delta, (DeltaType)-180, (DeltaType)180);
 	}
 
 	/** Find the smallest angle between two headings (in radians) */
@@ -922,26 +911,15 @@ public:
 		typename T2
 		UE_REQUIRES(std::is_floating_point_v<T> || std::is_floating_point_v<T2>)
 	>
-	[[nodiscard]] static constexpr auto FindDeltaAngleRadians(T A1, T2 A2) -> decltype(A1 * A2)
+	[[nodiscard]] static constexpr auto FindDeltaAngleRadians(T A1, T2 A2) -> decltype(A1 - A2)
 	{
 		// Find the difference
 		auto Delta = A2 - A1;
 
-		// If change is larger than PI
-		if (Delta > UE_PI)
-		{
-			// Flip to negative equivalent
-			Delta = Delta - UE_TWO_PI;
-		}
-		else if (Delta < -UE_PI)
-		{
-			// Otherwise, if change is smaller than -PI
-			// Flip to positive equivalent
-			Delta = Delta + UE_TWO_PI;
-		}
+		using DeltaType = decltype(Delta);
 
 		// Return delta in [-PI,PI] range
-		return Delta;
+		return Wrap(Delta, (DeltaType)-UE_PI, (DeltaType)UE_PI);
 	}
 
 	/** Given a heading which may be outside the +/- PI range, 'unwind' it back into that range. */
