Enable stat namedevents by default in Test builds.

- This matches established practice in internal projects and provides the ability to see much more detailed insight traces in Test (crucial for using APT Insights pass).

diff --git a/Engine/Plugins/Performance/AutomatedPerfTesting/Build/Scripts/AutomatedPerfTestNode.cs b/Engine/Plugins/Performance/AutomatedPerfTesting/Build/Scripts/AutomatedPerfTestNode.cs
--- a/Engine/Plugins/Performance/AutomatedPerfTesting/Build/Scripts/AutomatedPerfTestNode.cs
+++ b/Engine/Plugins/Performance/AutomatedPerfTesting/Build/Scripts/AutomatedPerfTestNode.cs
@@ -846,6 +846,7 @@ namespace AutomatedPerfTest
 
 			if (Config.DoInsightsTrace)
 			{
+				ClientRole.CommandLineParams.Add("statnamedevents");
 				ClientRole.CommandLineParams.Add("AutomatedPerfTest.DoInsightsTrace");
 				if (Config.TraceChannels != String.Empty)
 				{
diff --git a/Engine/Source/Runtime/Core/Public/Misc/Build.h b/Engine/Source/Runtime/Core/Public/Misc/Build.h
--- a/Engine/Source/Runtime/Core/Public/Misc/Build.h
+++ b/Engine/Source/Runtime/Core/Public/Misc/Build.h
@@ -225,18 +225,6 @@
 	#define FORCE_ANSI_ALLOCATOR 0
 #endif
 
-/**
- *	Optionally enable support for named events from the stat macros without the stat system overhead
- *	This will attempt to disable regular stats system and use named events instead
- */
-#ifndef ENABLE_STATNAMEDEVENTS
-	#define ENABLE_STATNAMEDEVENTS	0
-#endif
-
-#ifndef ENABLE_STATNAMEDEVENTS_UOBJECT
-	#define ENABLE_STATNAMEDEVENTS_UOBJECT 0
-#endif
-
 /*--------------------------------------------------------------------------------
 	Basic options that by default depend on the build configuration and platform
 
@@ -246,6 +234,7 @@
 	STATS											If true, then the stats system is compiled into the executable.
 	ALLOW_DEBUG_FILES								If true, then debug files like screen shots and profiles can be saved from the executable.
 	NO_LOGGING										If true, then no logs or text output will be produced
+	ENABLE_STATNAMEDEVENTS							If true, then enable support for named events from the stat macros without the stat system overhead. This will attempt to disable regular stats system and use named events instead.
 
 --------------------------------------------------------------------------------*/
 
@@ -259,6 +248,12 @@
 	#ifndef DO_ENSURE
 		#define DO_ENSURE										1
 	#endif
+	#ifndef ENABLE_STATNAMEDEVENTS
+		#define ENABLE_STATNAMEDEVENTS							0
+	#endif
+	#ifndef ENABLE_STATNAMEDEVENTS_UOBJECT
+		#define ENABLE_STATNAMEDEVENTS_UOBJECT 					0
+	#endif
 	#ifndef STATS
 		#define STATS											((WITH_UNREAL_DEVELOPER_TOOLS || !WITH_EDITORONLY_DATA || USE_STATS_WITHOUT_ENGINE || FORCE_USE_STATS) && !ENABLE_STATNAMEDEVENTS)
 	#endif
@@ -281,6 +276,12 @@
 	#ifndef DO_ENSURE
 		#define DO_ENSURE										1
 	#endif
+	#ifndef ENABLE_STATNAMEDEVENTS
+		#define ENABLE_STATNAMEDEVENTS							0
+	#endif
+	#ifndef ENABLE_STATNAMEDEVENTS_UOBJECT
+		#define ENABLE_STATNAMEDEVENTS_UOBJECT 					0
+	#endif
 	#ifndef STATS
 		#define STATS											((WITH_UNREAL_DEVELOPER_TOOLS || !WITH_EDITORONLY_DATA || USE_STATS_WITHOUT_ENGINE || FORCE_USE_STATS) && !ENABLE_STATNAMEDEVENTS)
 	#endif
@@ -303,6 +304,12 @@
 	#ifndef DO_ENSURE
 		#define DO_ENSURE										USE_ENSURES_IN_SHIPPING
 	#endif
+	#ifndef ENABLE_STATNAMEDEVENTS
+		#define ENABLE_STATNAMEDEVENTS							(!FORCE_USE_STATS && !USE_STATS_WITHOUT_ENGINE)
+	#endif
+	#ifndef ENABLE_STATNAMEDEVENTS_UOBJECT
+		#define ENABLE_STATNAMEDEVENTS_UOBJECT 					(ENABLE_STATNAMEDEVENTS)
+	#endif
 	#ifndef STATS
 		#define STATS											(FORCE_USE_STATS && !ENABLE_STATNAMEDEVENTS)
 	#endif
@@ -325,6 +332,12 @@
 	#ifndef DO_ENSURE
 		#define DO_ENSURE									USE_ENSURES_IN_SHIPPING
 	#endif
+	#ifndef ENABLE_STATNAMEDEVENTS
+		#define ENABLE_STATNAMEDEVENTS						0
+	#endif
+	#ifndef ENABLE_STATNAMEDEVENTS_UOBJECT
+		#define ENABLE_STATNAMEDEVENTS_UOBJECT 				0
+	#endif
 	#ifndef STATS
 		#define STATS										(FORCE_USE_STATS && !ENABLE_STATNAMEDEVENTS)
 	#endif
