Mark TFunction class hierarchy as AUTORTFM_INFER.

This is a necessary prerequisite to labeling VM interpreter code as AUTORTFM_DISABLE, as TFunction is used inside VCell::Equal.

diff --git a/Engine/Source/Runtime/Core/Public/Templates/Function.h b/Engine/Source/Runtime/Core/Public/Templates/Function.h
--- a/Engine/Source/Runtime/Core/Public/Templates/Function.h
+++ b/Engine/Source/Runtime/Core/Public/Templates/Function.h
@@ -3,6 +3,7 @@
 #pragma once
 
 #include "CoreTypes.h"
+#include "AutoRTFM.h"
 #include "Concepts/NotCVRefTo.h"
 #include "Misc/AssertionMacros.h"
 #include "Misc/IntrusiveUnsetOptionalState.h"
@@ -69,7 +70,7 @@ namespace UE::Core::Private::Function
 	};
 
 	template <typename T, bool bUnique, bool bOnHeap>
-	struct TFunction_OwnedObject : public IFunction_OwnedObject
+	struct AUTORTFM_INFER TFunction_OwnedObject : public IFunction_OwnedObject
 	{
 		template <typename ArgType>
 		explicit TFunction_OwnedObject(ArgType&& Arg)
@@ -123,7 +124,7 @@ namespace UE::Core::Private::Function
 		}
 	}
 
-	struct FFunctionStorage
+	struct AUTORTFM_INFER FFunctionStorage
 	{
 		constexpr static bool bCanBeNull = true;
 
@@ -199,7 +200,7 @@ namespace UE::Core::Private::Function
 	};
 
 	template <bool bUnique>
-	struct TFunctionStorage : FFunctionStorage
+	struct AUTORTFM_INFER TFunctionStorage : FFunctionStorage
 	{
 		TFunctionStorage() = default;
 
@@ -282,7 +283,7 @@ namespace UE::Core::Private::Function
 	 * A class which is used to instantiate the code needed to call a bound function.
 	 */
 	template <typename Functor, typename Ret, typename... ParamTypes>
-	struct TFunctionRefCaller
+	struct AUTORTFM_INFER TFunctionRefCaller
 	{
 		static Ret Call(void* Obj, ParamTypes&... Params)
 		{
@@ -311,7 +312,7 @@ namespace UE::Core::Private::Function
 	struct TFunctionRefBase;
 
 	template <typename StorageType, typename Ret, typename... ParamTypes>
-	struct TFunctionRefBase<StorageType, Ret (ParamTypes...)>
+	struct AUTORTFM_INFER TFunctionRefBase<StorageType, Ret (ParamTypes...)>
 	{
 		template <typename OtherStorageType, typename OtherFuncType>
 		friend struct TFunctionRefBase;
@@ -552,7 +553,7 @@ namespace UE::Core::Private::Function
  * }
  */
 template <typename Ret, typename... ParamTypes>
-class TFunctionRef<Ret(ParamTypes...)> final : public UE::Core::Private::Function::TFunctionRefBase<UE::Core::Private::Function::FFunctionRefStoragePolicy, Ret(ParamTypes...)>
+class AUTORTFM_INFER TFunctionRef<Ret(ParamTypes...)> final : public UE::Core::Private::Function::TFunctionRefBase<UE::Core::Private::Function::FFunctionRefStoragePolicy, Ret(ParamTypes...)>
 {
 	using Super = UE::Core::Private::Function::TFunctionRefBase<UE::Core::Private::Function::FFunctionRefStoragePolicy, Ret(ParamTypes...)>;
 
@@ -635,7 +636,7 @@ public:
  * }
  */
 template <typename Ret, typename... ParamTypes>
-class TFunction<Ret(ParamTypes...)> final : public UE::Core::Private::Function::TFunctionRefBase<UE::Core::Private::Function::TFunctionStorage<false>, Ret(ParamTypes...)>
+class AUTORTFM_INFER TFunction<Ret(ParamTypes...)> final : public UE::Core::Private::Function::TFunctionRefBase<UE::Core::Private::Function::TFunctionStorage<false>, Ret(ParamTypes...)>
 {
 	using Super = UE::Core::Private::Function::TFunctionRefBase<UE::Core::Private::Function::TFunctionStorage<false>, Ret(ParamTypes...)>;
 
@@ -750,7 +751,7 @@ public:
  * Foo(MoveTemp(MovableFunc)); // ok
  */
 template <typename Ret, typename... ParamTypes>
-class TUniqueFunction<Ret(ParamTypes...)> final : public UE::Core::Private::Function::TFunctionRefBase<UE::Core::Private::Function::TFunctionStorage<true>, Ret(ParamTypes...)>
+class AUTORTFM_INFER TUniqueFunction<Ret(ParamTypes...)> final : public UE::Core::Private::Function::TFunctionRefBase<UE::Core::Private::Function::TFunctionStorage<true>, Ret(ParamTypes...)>
 {
 	using Super = UE::Core::Private::Function::TFunctionRefBase<UE::Core::Private::Function::TFunctionStorage<true>, Ret(ParamTypes...)>;
 
diff --git a/Engine/Source/Runtime/Core/Public/Templates/Invoke.h b/Engine/Source/Runtime/Core/Public/Templates/Invoke.h
--- a/Engine/Source/Runtime/Core/Public/Templates/Invoke.h
+++ b/Engine/Source/Runtime/Core/Public/Templates/Invoke.h
@@ -3,6 +3,7 @@
 #pragma once
 
 #include "CoreTypes.h"
+#include "AutoRTFM.h"
 #include "Traits/MemberFunctionPtrOuter.h"
 #include <type_traits>
 
@@ -40,14 +41,14 @@ namespace UE::Core::Private
  * See: http://en.cppreference.com/w/cpp/utility/functional/invoke
  */
 template <typename FuncType, typename... ArgTypes>
-UE_FORCEINLINE_HINT constexpr auto Invoke(FuncType&& Func, ArgTypes&&... Args)
+AUTORTFM_INFER UE_FORCEINLINE_HINT constexpr auto Invoke(FuncType&& Func, ArgTypes&&... Args)
 	-> decltype(((FuncType&&)Func)((ArgTypes&&)Args...))
 {
 	return ((FuncType&&)Func)((ArgTypes&&)Args...);
 }
 
 template <typename ReturnType, typename ObjType, typename TargetType>
-constexpr auto Invoke(ReturnType ObjType::*pdm, TargetType&& Target)
+AUTORTFM_INFER constexpr auto Invoke(ReturnType ObjType::*pdm, TargetType&& Target)
 	-> decltype(UE::Core::Private::DereferenceIfNecessary<ObjType>((TargetType&&)Target, &Target).*pdm)
 {
 	return UE::Core::Private::DereferenceIfNecessary<ObjType>((TargetType&&)Target, &Target).*pdm;
@@ -59,7 +60,7 @@ template <
 	typename... ArgTypes,
 	typename    ObjType = TMemberFunctionPtrOuter_T<PtrMemFunType>
 >
-constexpr auto Invoke(PtrMemFunType PtrMemFun, TargetType&& Target, ArgTypes&&... Args)
+AUTORTFM_INFER constexpr auto Invoke(PtrMemFunType PtrMemFun, TargetType&& Target, ArgTypes&&... Args)
 	-> decltype((UE::Core::Private::DereferenceIfNecessary<ObjType>((TargetType&&)Target, &Target).*PtrMemFun)((ArgTypes&&)Args...))
 {
 	return (UE::Core::Private::DereferenceIfNecessary<ObjType>((TargetType&&)Target, &Target).*PtrMemFun)((ArgTypes&&)Args...);
@@ -74,7 +75,7 @@ constexpr auto Invoke(PtrMemFunType PtrMemFun, TargetType&& Target, ArgTypes&&..
  *
  * TArray<FMyType> Array = ...;
  *
- * // Doesn't compile, because you can't take the address of an overloaded function when its type needs to be deduced. 
+ * // Doesn't compile, because you can't take the address of an overloaded function when its type needs to be deduced.
  * Algo::SortBy(Array, &LexToString);
  *
  * // Works as expected
diff --git a/Engine/Source/Runtime/Core/Public/Templates/Projection.h b/Engine/Source/Runtime/Core/Public/Templates/Projection.h
--- a/Engine/Source/Runtime/Core/Public/Templates/Projection.h
+++ b/Engine/Source/Runtime/Core/Public/Templates/Projection.h
@@ -2,15 +2,16 @@
 
 #pragma once
 
+#include "AutoRTFM.h"
 #include <type_traits>
 
 namespace UE::Core::Private
 {
 	template <typename InvocableType>
-	struct TProjectionMemberFunction;
+	struct AUTORTFM_INFER TProjectionMemberFunction;
 
 	template <typename ClassType, typename FunctionType>
-	struct TProjectionMemberFunction<FunctionType ClassType::*>
+	struct AUTORTFM_INFER TProjectionMemberFunction<FunctionType ClassType::*>
 	{
 		FunctionType ClassType::* MemberFunctionPtr;
 
@@ -100,11 +101,11 @@ namespace UE::Core::Private
  *     FInner Inner;
  * };
  *
- * // Sort array of outers by the names of inner, 
+ * // Sort array of outers by the names of inner,
  * Algo::SortBy(ArrayOfOuters, Projection(&FOuter::Inner, &FInner::Name));
  */
 template <typename Invocable0Type, typename... InvocableTypes>
-[[nodiscard]] constexpr auto Projection(Invocable0Type&& Invocable0, InvocableTypes&&... Invocables)
+AUTORTFM_INFER [[nodiscard]] constexpr auto Projection(Invocable0Type&& Invocable0, InvocableTypes&&... Invocables)
 {
 	if constexpr (sizeof...(InvocableTypes) == 0)
 	{
