Fix to missing include and small formmatting change in AndroidJavaEnv

diff --git a/Engine/Source/Runtime/Core/Public/Android/AndroidJavaEnv.h b/Engine/Source/Runtime/Core/Public/Android/AndroidJavaEnv.h
--- a/Engine/Source/Runtime/Core/Public/Android/AndroidJavaEnv.h
+++ b/Engine/Source/Runtime/Core/Public/Android/AndroidJavaEnv.h
@@ -5,6 +5,7 @@
 #include "CoreTypes.h"
 #include "Containers/Array.h"
 #include "Containers/StaticArray.h"
+#include "Containers/StringView.h"
 #include "Containers/UnrealString.h"
 #include <algorithm>
 #include <type_traits>
