Add FLogCategory ctor that takes const TCHAR* not FName

The log category declarations via macro (i.e. most of them) pass a literal,
so put the FName creation in the callee.

diff --git a/Engine/Source/Runtime/Core/Private/Logging/LogCategory.cpp b/Engine/Source/Runtime/Core/Private/Logging/LogCategory.cpp
--- a/Engine/Source/Runtime/Core/Private/Logging/LogCategory.cpp
+++ b/Engine/Source/Runtime/Core/Private/Logging/LogCategory.cpp
@@ -26,6 +26,11 @@ FLogCategoryBase::FLogCategoryBase(const FLogCategoryName& InCategoryName, ELogV
 	checkSlow(!(Verbosity & ELogVerbosity::BreakOnLog)); // this bit is factored out of this variable, always
 }
 
+FLogCategoryBase::FLogCategoryBase(const TCHAR* InCategoryName, ELogVerbosity::Type InDefaultVerbosity, ELogVerbosity::Type InCompileTimeVerbosity)
+	: FLogCategoryBase(FLogCategoryName{InCategoryName}, InDefaultVerbosity, InCompileTimeVerbosity)
+{
+}
+
 FLogCategoryBase::~FLogCategoryBase()
 {
 	checkSlow(!(Verbosity & ELogVerbosity::BreakOnLog)); // this bit is factored out of this variable, always
diff --git a/Engine/Source/Runtime/Core/Public/Logging/LogCategory.h b/Engine/Source/Runtime/Core/Public/Logging/LogCategory.h
--- a/Engine/Source/Runtime/Core/Public/Logging/LogCategory.h
+++ b/Engine/Source/Runtime/Core/Public/Logging/LogCategory.h
@@ -28,6 +28,15 @@ struct FLogCategoryBase
 	 */
 	UE_API FLogCategoryBase(const FLogCategoryName& CategoryName, ELogVerbosity::Type DefaultVerbosity, ELogVerbosity::Type CompileTimeVerbosity);
 
+	/**
+	 * Constructor, registers with the log suppression system and sets up the default values.
+	 *
+	 * @param CategoryName           Name of the category.
+	 * @param DefaultVerbosity       Default verbosity used to filter this category at runtime.
+	 * @param CompileTimeVerbosity   Verbosity used to filter this category at compile time.
+	 */
+	UE_API FLogCategoryBase(const TCHAR* CategoryName, ELogVerbosity::Type DefaultVerbosity, ELogVerbosity::Type CompileTimeVerbosity);
+
 	/** Destructor, unregisters from the log suppression system. */
 	UE_API ~FLogCategoryBase();
 
@@ -86,6 +95,11 @@ struct FLogCategory : public FLogCategoryBase
 		: FLogCategoryBase(InCategoryName, InDefaultVerbosity, CompileTimeVerbosity)
 	{
 	}
+
+	UE_FORCEINLINE_HINT FLogCategory(const TCHAR* InCategoryName)
+		: FLogCategoryBase(InCategoryName, InDefaultVerbosity, CompileTimeVerbosity)
+	{
+	}
 };
 
 #undef UE_API
