Fix to include in AndroidJavaEnv.h

diff --git a/Engine/Source/Runtime/Core/Public/Android/AndroidJavaEnv.h b/Engine/Source/Runtime/Core/Public/Android/AndroidJavaEnv.h
--- a/Engine/Source/Runtime/Core/Public/Android/AndroidJavaEnv.h
+++ b/Engine/Source/Runtime/Core/Public/Android/AndroidJavaEnv.h
@@ -6,7 +6,9 @@
 
 #include "CoreTypes.h"
 #include "Containers/Array.h"
+#include "Containers/ArrayView.h"
 #include "Containers/StaticArray.h"
+#include "Containers/StringConv.h"
 #include "Containers/StringView.h"
 #include "Containers/UnrealString.h"
 #include <algorithm>
