Fix -Wimplicit-int-conversion-on-negation

#rnx

diff --git a/Engine/Source/Runtime/Core/Public/GenericPlatform/GenericPlatformMath.h b/Engine/Source/Runtime/Core/Public/GenericPlatform/GenericPlatformMath.h
--- a/Engine/Source/Runtime/Core/Public/GenericPlatform/GenericPlatformMath.h
+++ b/Engine/Source/Runtime/Core/Public/GenericPlatform/GenericPlatformMath.h
@@ -946,7 +946,7 @@ struct FGenericPlatformMath
 	template< class T > 
 	[[nodiscard]] static constexpr UE_FORCEINLINE_HINT T Abs( const T A )
 	{
-		return (A < (T)0) ? -A : A;
+		return (A < (T)0) ? (T)-A : A;
 	}
 
 	/** Returns 1, 0, or -1 depending on relation of T to 0 */
