Made FFloat16 trivially copyable

#rnx

diff --git a/Engine/Source/Runtime/Core/Public/Math/Float16.h b/Engine/Source/Runtime/Core/Public/Math/Float16.h
--- a/Engine/Source/Runtime/Core/Public/Math/Float16.h
+++ b/Engine/Source/Runtime/Core/Public/Math/Float16.h
@@ -37,22 +37,22 @@ public:
 	/* Float16 can store values in [-MaxF16Float,MaxF16Float] */
 	constexpr static float MaxF16Float = 65504.f;
 
-	uint16 Encoded;
+	uint16 Encoded = 0;
 
 	/** Default constructor */
-	[[nodiscard]] FFloat16();
+	[[nodiscard]] FFloat16() = default;
 
 	/** Copy constructor. */
-	[[nodiscard]] FFloat16(const FFloat16& FP16Value);
+	[[nodiscard]] FFloat16(const FFloat16& FP16Value) = default;
 
 	/** Conversion constructor. Convert from Fp32 to Fp16. */
-	[[nodiscard]] FFloat16(float FP32Value);	
+	[[nodiscard]] FFloat16(float FP32Value);
 
 	/** Assignment operator. Convert from Fp32 to Fp16. */
 	FFloat16& operator=(float FP32Value);
 
 	/** Assignment operator. Copy Fp16 value. */
-	FFloat16& operator=(const FFloat16& FP16Value);
+	FFloat16& operator=(const FFloat16& FP16Value) = default;
 
 	/** Convert from Fp16 to Fp32. */
 	operator float() const;
@@ -120,17 +120,6 @@ template<> struct TCanBulkSerialize<FFloat16> { enum { Value = true }; };
 
 DECLARE_INTRINSIC_TYPE_LAYOUT(FFloat16);
 
-UE_FORCEINLINE_HINT FFloat16::FFloat16()
-	:	Encoded(0)
-{ }
-
-
-UE_FORCEINLINE_HINT FFloat16::FFloat16(const FFloat16& FP16Value)
-{
-	Encoded = FP16Value.Encoded;
-}
-
-
 UE_FORCEINLINE_HINT FFloat16::FFloat16(float FP32Value)
 {
 	Set(FP32Value);
@@ -144,13 +133,6 @@ inline FFloat16& FFloat16::operator=(float FP32Value)
 }
 
 
-inline FFloat16& FFloat16::operator=(const FFloat16& FP16Value)
-{
-	Encoded = FP16Value.Encoded;
-	return *this;
-}
-
-
 UE_FORCEINLINE_HINT FFloat16::operator float() const
 {
 	return GetFloat();
