[LLM] Scope ConfigSystem allocations to proper LLM tag

#rnx

diff --git a/Engine/Source/Runtime/Core/Private/Misc/ConfigCacheIni.cpp b/Engine/Source/Runtime/Core/Private/Misc/ConfigCacheIni.cpp
--- a/Engine/Source/Runtime/Core/Private/Misc/ConfigCacheIni.cpp
+++ b/Engine/Source/Runtime/Core/Private/Misc/ConfigCacheIni.cpp
@@ -1822,6 +1822,8 @@ static void ConditionalInitializeLoadType(FConfigCommandStream* File, UE::Config
 template<typename FileType>
 void FillFileFromBuffer(FileType* File, FStringView Buffer, bool bHandleSymbolCommands, const FString& FileHint)
 {
+	LLM_SCOPE(ELLMTag::ConfigSystem);
+
 	static const FName ConfigFileClassName = TEXT("ConfigFile");
 	const FName FileName = FName(*FileHint);
 	LLM_SCOPE_DYNAMIC_STAT_OBJECTPATH_FNAME(FileName, ELLMTagSet::Assets);
@@ -3429,6 +3431,7 @@ bool FConfigBranch::AddDynamicLayersToHierarchy(const TArray<FString>& Filenames
 
 bool FConfigBranch::AddDynamicLayersToHierarchy(const TArray<FDynamicLayerInfo>& Layers, FConfigModificationTracker* ModificationTracker, TSet<FString>* GlobalConfigFileCache, TSet<FString>* PluginConfigFileCache, bool bForceFullDynamicLayerUpdate)
 {
+	LLM_SCOPE(ELLMTag::ConfigSystem);
 	if (bIsSafeUnloaded)
 	{
 		SafeReload();
@@ -3583,6 +3586,7 @@ bool FConfigBranch::AddDynamicLayersToHierarchy(const TArray<FDynamicLayerInfo>&
 
 bool FConfigBranch::AddDynamicLayerStringToHierarchy(const FString& Filename, const FString& Contents, FName Tag, DynamicLayerPriority Priority, FConfigModificationTracker* ModificationTracker)
 {
+	LLM_SCOPE(ELLMTag::ConfigSystem);
 	if (bIsSafeUnloaded)
 	{
 		SafeReload();
@@ -3793,6 +3797,8 @@ void FConfigBranch::SafeUnload()
 
 void FConfigBranch::SafeReload()
 {
+	LLM_SCOPE(ELLMTag::ConfigSystem);
+
 	double StartTime = FPlatformTime::Seconds();
 
 	// read static layers back in from disk
@@ -6885,6 +6891,7 @@ void FConfigCacheIni::RemoveMultipleTagsFromAllBranches(const TArray<FName>& Tag
 
 void FConfigCacheIni::AddPluginsToBranches(const TArray<FName>& PluginNames, FConfigModificationTracker* ModificationTracker)
 {
+	LLM_SCOPE(ELLMTag::ConfigSystem);
 	// @todo make sure we are still pending
 	
 	TMap<FConfigBranch*, TArray<FDynamicLayerInfo>> AllDynamicLayers;
