[UBT][CppCompileWarnings] Engine compliance - wundef - monolothic

diff --git a/Engine/Source/Runtime/Core/Public/Windows/WindowsPlatformMisc.h b/Engine/Source/Runtime/Core/Public/Windows/WindowsPlatformMisc.h
--- a/Engine/Source/Runtime/Core/Public/Windows/WindowsPlatformMisc.h
+++ b/Engine/Source/Runtime/Core/Public/Windows/WindowsPlatformMisc.h
@@ -392,6 +392,6 @@ struct FWindowsPlatformMisc
 };
 
 
-#if WINDOWS_USE_FEATURE_PLATFORMMISC_CLASS
+#if defined(WINDOWS_USE_FEATURE_PLATFORMMISC_CLASS) && WINDOWS_USE_FEATURE_PLATFORMMISC_CLASS
 typedef FWindowsPlatformMisc FPlatformMisc;
 #endif
