constexpresify FMath::CountLeading\TrailingZeros

diff --git a/Engine/Source/Runtime/Core/Public/Clang/ClangPlatformMath.h b/Engine/Source/Runtime/Core/Public/Clang/ClangPlatformMath.h
--- a/Engine/Source/Runtime/Core/Public/Clang/ClangPlatformMath.h
+++ b/Engine/Source/Runtime/Core/Public/Clang/ClangPlatformMath.h
@@ -75,11 +75,6 @@ struct FClangPlatformMath : public FGenericPlatformMath
 	
 		return (uint32)__builtin_ctz(Value);
 	}
-
-	static UE_FORCEINLINE_HINT constexpr uint32 CountTrailingZerosConstExpr(uint32 Value)
-	{
-		return CountTrailingZeros(Value);
-	}
 	
 	/**
 	 * Counts the number of trailing zeros in the bit representation of the value
@@ -98,12 +93,6 @@ struct FClangPlatformMath : public FGenericPlatformMath
 		return (uint64)__builtin_ctzll(Value);
 	}
 
-	static UE_FORCEINLINE_HINT constexpr uint64 CountTrailingZeros64ConstExpr(uint64 Value)
-	{
-		return CountTrailingZeros64(Value);
-	}
-	
-
 	static UE_FORCEINLINE_HINT constexpr uint32 FloorLog2(uint32 Value)
 	{
 		return 31 - __builtin_clz(Value | 1);
diff --git a/Engine/Source/Runtime/Core/Public/GenericPlatform/GenericPlatformMath.h b/Engine/Source/Runtime/Core/Public/GenericPlatform/GenericPlatformMath.h
--- a/Engine/Source/Runtime/Core/Public/GenericPlatform/GenericPlatformMath.h
+++ b/Engine/Source/Runtime/Core/Public/GenericPlatform/GenericPlatformMath.h
@@ -744,7 +744,7 @@ struct FGenericPlatformMath
 	 *
 	 * @return the number of zeros after the last "on" bit
 	 */
-	[[nodiscard]] static constexpr inline uint32 CountTrailingZerosConstExpr(uint32 Value)
+	[[nodiscard]] static constexpr inline uint32 CountTrailingZeros(uint32 Value)
 	{
 		if (Value == 0)
 		{
@@ -759,9 +759,10 @@ struct FGenericPlatformMath
 		return Result;
 	}
 
-	[[nodiscard]] static constexpr UE_FORCEINLINE_HINT uint32 CountTrailingZeros(uint32 Value)
+	UE_DEPRECATED(5.7, "Please use CountTrailingZeros instead")
+	[[nodiscard]] static constexpr UE_FORCEINLINE_HINT uint32 CountTrailingZerosConstExpr(uint32 Value)
 	{
-		return CountTrailingZerosConstExpr(Value);
+		return CountTrailingZeros(Value);
 	}
 
 	/**
@@ -771,7 +772,7 @@ struct FGenericPlatformMath
 	 *
 	 * @return the number of zeros after the last "on" bit
 	 */
-	[[nodiscard]] static constexpr inline uint64 CountTrailingZeros64ConstExpr(uint64 Value)
+	[[nodiscard]] static constexpr inline uint64 CountTrailingZeros64(uint64 Value)
 	{
 		if (Value == 0)
 		{
@@ -786,9 +787,10 @@ struct FGenericPlatformMath
 		return Result;
 	}
 
-	[[nodiscard]] static constexpr UE_FORCEINLINE_HINT uint64 CountTrailingZeros64(uint64 Value)
+	UE_DEPRECATED(5.7, "Please use CountTrailingZeros64 instead")
+	[[nodiscard]] static constexpr UE_FORCEINLINE_HINT uint64 CountTrailingZeros64ConstExpr(uint64 Value)
 	{
-		return CountTrailingZeros64ConstExpr(Value);
+		return CountTrailingZeros64(Value);
 	}
 
 	/**
diff --git a/Engine/Source/Runtime/Core/Public/Microsoft/MicrosoftPlatformMath.h b/Engine/Source/Runtime/Core/Public/Microsoft/MicrosoftPlatformMath.h
--- a/Engine/Source/Runtime/Core/Public/Microsoft/MicrosoftPlatformMath.h
+++ b/Engine/Source/Runtime/Core/Public/Microsoft/MicrosoftPlatformMath.h
@@ -34,62 +34,34 @@ struct FMicrosoftPlatformMathBase : public TUnrealPlatformMathSSE4Base<FGenericP
 		return BitIndex;
 	}
 
-	static inline constexpr uint8 CountLeadingZeros8(uint8 Value)
+	static constexpr inline uint8 CountLeadingZeros8(uint8 Value)
 	{
 		// return 8 if value was 0
 
 		UE_IF_CONSTEVAL
 		{
-			if (Value == 0)
-			{
-				return 8;
-			}
-
-			uint8 Result = 0;
-			for (;;)
-			{
-				if (Value & 0x80)
-				{
-					return Result;
-				}
-				Value <<= 1;
-				++Result;
-			}
+			return FGenericPlatformMath::CountLeadingZeros8(Value);
 		}
 		else
 		{
 			unsigned long BitIndex;
-			_BitScanReverse(&BitIndex, uint32(Value)*2 + 1);
+			_BitScanReverse(&BitIndex, uint32(Value) * 2 + 1);
 			return uint8(8 - BitIndex);
 		}
 	}
 
-	static inline constexpr uint32 CountTrailingZeros(uint32 Value)
+	static constexpr inline uint32 CountTrailingZeros(uint32 Value)
 	{
-		// return 32 if value was 0
 
 		UE_IF_CONSTEVAL
 		{
-			if (Value == 0)
-			{
-				return 32;
-			}
-
-			uint32 Result = 0;
-			for (;;)
-			{
-				if (Value & 1)
-				{
-					return Result;
-				}
-				Value >>= 1;
-				++Result;
-			}
+			return FGenericPlatformMath::CountTrailingZeros(Value);
 		}
 		else
 		{
+			// return 32 if value was 0
 			unsigned long BitIndex;	// 0-based, where the LSB is 0 and MSB is 31
-			return _BitScanForward( &BitIndex, Value ) ? BitIndex : 32;
+			return _BitScanForward(&BitIndex, Value) ? BitIndex : 32;
 		}
 	}
 
@@ -110,8 +82,6 @@ struct FMicrosoftPlatformMathBase : public TUnrealPlatformMathSSE4Base<FGenericP
 		return uint64(1) << CeilLogTwo64(Arg);
 	}
 
-#if PLATFORM_64BITS
-
 	#pragma intrinsic(_BitScanReverse64)
 
 	static inline uint64 FloorLog2_64(uint64 Value)
@@ -134,108 +104,54 @@ struct FMicrosoftPlatformMathBase : public TUnrealPlatformMathSSE4Base<FGenericP
 		return 64 - CountLeadingZeros64(Arg - 1);
 	}
 
-	static inline constexpr uint64 CountLeadingZeros64(uint64 Value)
+	static constexpr inline uint64 CountLeadingZeros64(uint64 Value)
 	{
-		// return 64 if value if was 0
 
 		UE_IF_CONSTEVAL
 		{
-			if (Value == 0)
-			{
-				return 64;
-			}
-
-			uint64 Result = 0;
-			for (;;)
-			{
-				if (Value & 0x8000000000000000)
-				{
-					return Result;
-				}
-				Value <<= 1;
-				++Result;
-			}
+			return FGenericPlatformMath::CountLeadingZeros64(Value);
 		}
 		else
 		{
-			//https://godbolt.org/z/Ejh5G4vPK
+			//https://godbolt.org/z/Ejh5G4vPK	
+			// return 64 if value if was 0
 			unsigned long BitIndex;
-			if ( ! _BitScanReverse64(&BitIndex, Value) ) BitIndex = -1;
+			if (!_BitScanReverse64(&BitIndex, Value)) BitIndex = -1;
 			return 63 - BitIndex;
 		}
 	}
 
-	static inline constexpr uint64 CountTrailingZeros64(uint64 Value)
+	static constexpr inline uint64 CountTrailingZeros64(uint64 Value)
 	{
-		// return 64 if Value is 0
 
 		UE_IF_CONSTEVAL
 		{
-			if (Value == 0)
-			{
-				return 64;
-			}
-
-			uint64 Result = 0;
-			for (;;)
-			{
-				if (Value & 1)
-				{
-					return Result;
-				}
-				Value >>= 1;
-				++Result;
-			}
+			return FGenericPlatformMath::CountTrailingZeros64(Value);
 		}
 		else
 		{
+			// return 64 if Value is 0
 			unsigned long BitIndex;	// 0-based, where the LSB is 0 and MSB is 63
-			return _BitScanForward64( &BitIndex, Value ) ? BitIndex : 64;
+			return _BitScanForward64(&BitIndex, Value) ? BitIndex : 64;
 		}
 	}
 
-	static inline constexpr uint32 CountLeadingZeros(uint32 Value)
+	static constexpr inline uint32 CountLeadingZeros(uint32 Value)
 	{
-		// return 32 if value is zero
 
 		UE_IF_CONSTEVAL
 		{
-			if (Value == 0)
-			{
-				return 32;
-			}
-
-			uint32 Result = 0;
-			for (;;)
-			{
-				if (Value & 0x80000000)
-				{
-					return Result;
-				}
-				Value <<= 1;
-				++Result;
-			}
+			return FGenericPlatformMath::CountLeadingZeros(Value);
 		}
 		else
 		{
+			// return 32 if value is zero
 			unsigned long BitIndex;
-			_BitScanReverse64(&BitIndex, uint64(Value)*2 + 1);
+			_BitScanReverse64(&BitIndex, uint64(Value) * 2 + 1);
 			return 32 - BitIndex;
 		}
 	}
 
-#else // 32-bit
-
-	static inline uint32 CountLeadingZeros(uint32 Value)
-	{
-		// return 32 if value is zero
-		unsigned long BitIndex;
-		if ( ! _BitScanReverse(&BitIndex, Value) ) BitIndex = -1;
-		return 31 - BitIndex;
-	}
-
-#endif
-
 #if PLATFORM_ENABLE_POPCNT_INTRINSIC
 	static UE_FORCEINLINE_HINT int32 CountBits(uint64 Bits)
 	{
diff --git a/Engine/Source/Runtime/CoreUObject/Public/UObject/UObjectArray.h b/Engine/Source/Runtime/CoreUObject/Public/UObject/UObjectArray.h
--- a/Engine/Source/Runtime/CoreUObject/Public/UObject/UObjectArray.h
+++ b/Engine/Source/Runtime/CoreUObject/Public/UObject/UObjectArray.h
@@ -90,7 +90,7 @@ private:
 	// EInternalObjectFlags_MinFlagBitIndex at the time of writing this is 14 and we have only 1 bit left in the EInternalObjectFlags for future use
 	// We can increase UObject alignment to 16 bytes to get one more bit and reduce the overall addressable virtual memory range to get more bits if necessary
 	constexpr static int32 UObjectAlignment = 8;
-	constexpr static int32 UObjectPtrTrailingZeroes = FMath::CountTrailingZerosConstExpr(UObjectAlignment);
+	constexpr static int32 UObjectPtrTrailingZeroes = FMath::CountTrailingZeros(UObjectAlignment);
 	static_assert(int(EInternalObjectFlags_MinFlagBitIndex) >= 48 - 32 - UObjectPtrTrailingZeroes, "We need at least 13 bits to pack higher bits of a UObject pointer into Flags");
 	constexpr static int32 FlagsMask = 0xFFFFFFFF << int(EInternalObjectFlags_MinFlagBitIndex);
 	constexpr static int32 PtrMask = ~FlagsMask;
diff --git a/Engine/Source/Runtime/Experimental/IoStore/OnDemand/Private/OnDemandIoStore.cpp b/Engine/Source/Runtime/Experimental/IoStore/OnDemand/Private/OnDemandIoStore.cpp
--- a/Engine/Source/Runtime/Experimental/IoStore/OnDemand/Private/OnDemandIoStore.cpp
+++ b/Engine/Source/Runtime/Experimental/IoStore/OnDemand/Private/OnDemandIoStore.cpp
@@ -2230,7 +2230,7 @@ FStringBuilderBase& operator<<(FStringBuilderBase& Sb, UE::IoStore::EOnDemandCon
 		return Sb;
 	}
 
-	constexpr uint32 BitCount = 1 + FMath::CountTrailingZerosConstExpr(
+	constexpr uint32 BitCount = 1 + FMath::CountTrailingZeros(
 		static_cast<std::underlying_type_t<EOnDemandContainerFlags>>(EOnDemandContainerFlags::Last));
 	static_assert(UE_ARRAY_COUNT(Names) == BitCount + 1, "Please update names list");
 
diff --git a/Engine/Source/Runtime/Experimental/IoStore/OnDemandCore/Private/OnDemandToc.cpp b/Engine/Source/Runtime/Experimental/IoStore/OnDemandCore/Private/OnDemandToc.cpp
--- a/Engine/Source/Runtime/Experimental/IoStore/OnDemandCore/Private/OnDemandToc.cpp
+++ b/Engine/Source/Runtime/Experimental/IoStore/OnDemandCore/Private/OnDemandToc.cpp
@@ -708,7 +708,7 @@ FStringBuilderBase& operator<<(FStringBuilderBase& Sb, UE::IoStore::EOnDemandToc
 		return Sb;
 	}
 
-	constexpr uint32 BitCount = 1 + FMath::CountTrailingZerosConstExpr(
+	constexpr uint32 BitCount = 1 + FMath::CountTrailingZeros(
 		static_cast<std::underlying_type_t<EOnDemandTocFlags>>(EOnDemandTocFlags::Last));
 	static_assert(UE_ARRAY_COUNT(FlagText) == BitCount + 1, "Please update flag text list");
 
