MallocBinned2\3 code cleanup

diff --git a/Engine/Source/Runtime/Core/Private/HAL/MallocBinned2.cpp b/Engine/Source/Runtime/Core/Private/HAL/MallocBinned2.cpp
--- a/Engine/Source/Runtime/Core/Private/HAL/MallocBinned2.cpp
+++ b/Engine/Source/Runtime/Core/Private/HAL/MallocBinned2.cpp
@@ -98,13 +98,13 @@ static uint32 GMB2PageSize = 0;
 uint16 FMallocBinned2::SmallBinSizes[UE_MB2_SMALL_POOL_COUNT] = { 0 };
 FMallocBinned2* FMallocBinned2::MallocBinned2 = nullptr;
 // Mapping of sizes to small table indices
-uint8 FMallocBinned2::MemSizeToPoolIndex[1 + (UE_MB2_MAX_SMALL_POOL_SIZE >> UE_MBC_MIN_SMALL_POOL_ALIGNMENT_SHIFT)] = { 0 };
+uint8 FMallocBinned2::MemSizeToPoolIndex[SIZE_TO_POOL_INDEX_NUM] = { 0 };
 
 
 static FORCEINLINE bool IsSupportedSize(SIZE_T Size)
 {
 	bool bResult = IsAligned(Size, UE_MBC_MIN_SMALL_POOL_ALIGNMENT);
-	bResult = bResult && (Size >> UE_MBC_MIN_SMALL_POOL_ALIGNMENT_SHIFT) <= SIZE_T(MAX_uint32);
+	bResult = bResult && (Size >> UE_MBC_BIN_SIZE_SHIFT) <= SIZE_T(MAX_uint32);
 	return bResult;
 }
 
@@ -180,7 +180,7 @@ void* FMallocBinned2::FPoolInfo::AllocateBin()
 
 SIZE_T FMallocBinned2::FPoolInfo::GetOSRequestedBytes() const
 {
-	return SIZE_T(AllocSize) << UE_MBC_MIN_SMALL_POOL_ALIGNMENT_SHIFT;
+	return SIZE_T(AllocSize) << UE_MBC_BIN_SIZE_SHIFT;
 }
 
 SIZE_T FMallocBinned2::FPoolInfo::GetOsAllocatedBytes() const
@@ -196,7 +196,7 @@ void FMallocBinned2::FPoolInfo::SetOSAllocationSizes(SIZE_T InRequestedBytes, UP
 	checkSlow(InAllocatedBytes >= InRequestedBytes); // We must be allocating at least as much as we requested
 	checkSlow(IsSupportedSize(InRequestedBytes));    // We must be allocating a size we can store
 
-	AllocSize      = uint32(InRequestedBytes >> UE_MBC_MIN_SMALL_POOL_ALIGNMENT_SHIFT);
+	AllocSize      = uint32(InRequestedBytes >> UE_MBC_BIN_SIZE_SHIFT);
 	FirstFreeBlock = (FFreeBlock*)InAllocatedBytes;
 }
 
@@ -398,7 +398,7 @@ FMallocBinned2::FMallocBinned2()
 	FMemory::Memcpy(SmallBinSizes, SmallBinSizesInternal);
 	
 	static_assert(sizeof(FFreeBlock) <= SmallBinSizesInternal[0], "sizeof(FFreeBlock)  must be fit in smallest allocation size handled by FMallocBinned2.");
-	static_assert(sizeof(FMallocBinned2::FBundleNode) <= SmallBinSizesInternal[0], "sizeof(FBundleNode) must fit in smallest allocation size handled by FMallocBinned2.");
+	static_assert(sizeof(FBundleNode) <= SmallBinSizesInternal[0], "sizeof(FBundleNode) must fit in smallest allocation size handled by FMallocBinned2.");
 
 	const FGenericPlatformMemoryConstants Constants = FPlatformMemory::GetConstants();
 	GMB2PageSize = Constants.BinnedPageSize;
@@ -411,10 +411,8 @@ FMallocBinned2::FMallocBinned2()
 	checkf(GMB2PageSize % UE_MB2_LARGE_ALLOC == 0, TEXT("OS page size must be a multiple of UE_MB2_LARGE_ALLOC"));
 
 	static_assert(SmallBinSizesInternal[UE_MB2_SMALL_POOL_COUNT - 1] == UE_MB2_MAX_SMALL_POOL_SIZE, "UE_MB2_MAX_SMALL_POOL_SIZE must equal the largest bin size");
-	static_assert(sizeof(FMallocBinned2::FFreeBlock) <= SmallBinSizesInternal[0], "Pool header must be able to fit into the smallest bin");
 	static_assert(UE_ARRAY_COUNT(SmallBinSizesInternal) == UE_MB2_SMALL_POOL_COUNT, "Number of bins in SmallBinSizesInternal must match UE_MB2_SMALL_POOL_COUNT");
 	static_assert(UE_ARRAY_COUNT(SmallBinSizesInternal) <= 256, "Number of bins in SmallBinSizesInternal fit in one byte");
-	static_assert(sizeof(FFreeBlock) <= UE_MBC_MIN_SMALL_POOL_ALIGNMENT, "Free block struct must be small enough to fit into the smallest bin.");
 
 	// Init pool tables.
 	for (uint32 Index = 0; Index != UE_MB2_SMALL_POOL_COUNT; ++Index)
@@ -429,9 +427,9 @@ FMallocBinned2::FMallocBinned2()
 	// Set up pool mappings
 	uint8* IndexEntry = MemSizeToPoolIndex;
 	uint32 PoolIndex  = 0;
-	for (uint32 Index = 0; Index != 1 + (UE_MB2_MAX_SMALL_POOL_SIZE >> UE_MBC_MIN_SMALL_POOL_ALIGNMENT_SHIFT); ++Index)
+	for (uint32 Index = 0; Index != SIZE_TO_POOL_INDEX_NUM; ++Index)
 	{
-		const uint32 BinSize = Index << UE_MBC_MIN_SMALL_POOL_ALIGNMENT_SHIFT; // inverse of int32 Index = int32((Size >> UE_MBC_MIN_SMALL_POOL_ALIGNMENT_SHIFT));
+		const uint32 BinSize = Index << UE_MBC_BIN_SIZE_SHIFT; // inverse of int32 Index = int32((Size >> UE_MBC_BIN_SIZE_SHIFT));
 		while (SmallBinSizes[PoolIndex] < BinSize)
 		{
 			++PoolIndex;
@@ -623,8 +621,6 @@ void* FMallocBinned2::MallocExternalSmall(SIZE_T Size, uint32 Alignment)
 
 void* FMallocBinned2::MallocExternalLarge(SIZE_T Size, uint32 Alignment)
 {
-	static_assert(DEFAULT_ALIGNMENT <= UE_MBC_MIN_SMALL_POOL_ALIGNMENT, "DEFAULT_ALIGNMENT is assumed to be zero"); // used below
-
 	Alignment = FMath::Max<uint32>(Alignment, UE_MBC_MIN_SMALL_POOL_ALIGNMENT);
 	Size = Align(FMath::Max((SIZE_T)1, Size), Alignment);
 
@@ -678,8 +674,6 @@ void* FMallocBinned2::Realloc(void* Ptr, SIZE_T NewSize, uint32 Alignment)
 		return nullptr;
 	}
 
-	static_assert(DEFAULT_ALIGNMENT <= UE_MBC_MIN_SMALL_POOL_ALIGNMENT, "DEFAULT_ALIGNMENT is assumed to be zero"); // used below
-
 	if (!IsOSAllocation(Ptr))
 	{
 		check(Ptr); // null is 64k aligned so we should not be here
diff --git a/Engine/Source/Runtime/Core/Private/HAL/MallocBinned3.cpp b/Engine/Source/Runtime/Core/Private/HAL/MallocBinned3.cpp
--- a/Engine/Source/Runtime/Core/Private/HAL/MallocBinned3.cpp
+++ b/Engine/Source/Runtime/Core/Private/HAL/MallocBinned3.cpp
@@ -59,7 +59,7 @@ uint16 FMallocBinned3::SmallBinSizesShifted[UE_MB3_SMALL_POOL_COUNT + 1] = { 0 }
 
 FMallocBinned3* FMallocBinned3::MallocBinned3 = nullptr;
 // Mapping of sizes to small table indices
-uint8 FMallocBinned3::MemSizeToPoolIndex[1 + (UE_MB3_MAX_SMALL_POOL_SIZE >> UE_MBC_MIN_SMALL_POOL_ALIGNMENT_SHIFT)] = { 0 };
+uint8 FMallocBinned3::MemSizeToPoolIndex[SIZE_TO_POOL_INDEX_NUM] = { 0 };
 
 template <class T>
 static void SetCanaryInternal(T& PoolInfo, typename T::ECanary ShouldBe, bool bPreexisting, bool bGuarnteedToBeNew)
@@ -108,12 +108,12 @@ struct FFreeBlock
 	};
 
 	FORCEINLINE FFreeBlock(uint32 InBlockSize, uint32 InBinSize, uint8 InPoolIndex)
-		: BinSizeShifted(InBinSize >> UE_MBC_MIN_SMALL_POOL_ALIGNMENT_SHIFT)
+		: BinSizeShifted(InBinSize >> UE_MBC_BIN_SIZE_SHIFT)
 		, PoolIndex(InPoolIndex)
 		, Canary(CANARY_VALUE)
 		, NextFreeBlockIndex(MAX_uint32)
 	{
-		check(InPoolIndex < MAX_uint8 && (InBinSize >> UE_MBC_MIN_SMALL_POOL_ALIGNMENT_SHIFT) <= MAX_uint16);
+		check(InPoolIndex < MAX_uint8 && (InBinSize >> UE_MBC_BIN_SIZE_SHIFT) <= MAX_uint16);
 		NumFreeBins = InBlockSize / InBinSize;
 	}
 
@@ -133,7 +133,7 @@ struct FFreeBlock
 		{
 			CanaryFail();
 		}
-		//checkSlow(PoolIndex == BoundSizeToPoolIndex(uint32(BinSizeShifted) << UE_MBC_MIN_SMALL_POOL_ALIGNMENT_SHIFT));
+		//checkSlow(PoolIndex == BoundSizeToPoolIndex(uint32(BinSizeShifted) << UE_MBC_BIN_SIZE_SHIFT));
 	}
 
 	FORCEINLINE void CanaryFail() const
@@ -144,10 +144,10 @@ struct FFreeBlock
 	FORCEINLINE void* AllocateBin()
 	{
 		--NumFreeBins;
-		return (uint8*)this + NumFreeBins * (uint32(BinSizeShifted) << UE_MBC_MIN_SMALL_POOL_ALIGNMENT_SHIFT);
+		return (uint8*)this + NumFreeBins * (uint32(BinSizeShifted) << UE_MBC_BIN_SIZE_SHIFT);
 	}
 
-	uint16 BinSizeShifted;		// Size of the bins that this list points to >> UE_MBC_MIN_SMALL_POOL_ALIGNMENT_SHIFT
+	uint16 BinSizeShifted;		// Size of the bins that this list points to >> UE_MBC_BIN_SIZE_SHIFT
 	uint8  PoolIndex;			// Index of this pool
 	uint8  Canary;				// Constant value of 0xe7
 	uint32 NumFreeBins;			// Number of consecutive free bins here, at least 1.
@@ -343,7 +343,7 @@ struct FMallocBinned3::Private
 				FFreeBlock* Free = (FFreeBlock*)Node;
 				Free->NumFreeBins = 1;
 				Free->NextFreeBlockIndex = NodePool->NoFirstFreeIndex ? MAX_uint32 : NodePool->FirstFreeIndex;
-				Free->BinSizeShifted = (BinSize >> UE_MBC_MIN_SMALL_POOL_ALIGNMENT_SHIFT);
+				Free->BinSizeShifted = (BinSize >> UE_MBC_BIN_SIZE_SHIFT);
 				Free->Canary = FFreeBlock::CANARY_VALUE;
 				Free->PoolIndex = InPoolIndex;
 				NodePool->FirstFreeIndex = BinIndexWithinBlock;
@@ -553,7 +553,7 @@ FMallocBinned3::FMallocBinned3(const FPlatformMemory::FPlatformVirtualMemoryBloc
 		SmallPoolTables[Index].NumMemoryPagesPerBlock = SizeTable[Index].NumMemoryPagesPerBlock;
 		SmallPoolTables[Index].BlockSize = SizeTable[Index].NumMemoryPagesPerBlock * OsAllocationGranularity;
 
-		SmallBinSizesShifted[Index + 1] = (SizeTable[Index].BinSize >> UE_MBC_MIN_SMALL_POOL_ALIGNMENT_SHIFT);
+		SmallBinSizesShifted[Index + 1] = (SizeTable[Index].BinSize >> UE_MBC_BIN_SIZE_SHIFT);
 
 		const int64 TotalNumberOfBlocks = UE_MB3_MAX_MEMORY_PER_POOL_SIZE / SmallPoolTables[Index].BlockSize;
 		const uint32 Size = Align(sizeof(FPoolInfoSmall**) * (TotalNumberOfBlocks + SmallPoolInfosPerPlatformPage - 1) / SmallPoolInfosPerPlatformPage, PLATFORM_CACHE_LINE_SIZE);
@@ -593,9 +593,9 @@ FMallocBinned3::FMallocBinned3(const FPlatformMemory::FPlatformVirtualMemoryBloc
 	// Set up pool mappings
 	uint8* IndexEntry = MemSizeToPoolIndex;
 	uint32 PoolIndex  = 0;
-	for (uint32 Index = 0; Index != 1 + (UE_MB3_MAX_SMALL_POOL_SIZE >> UE_MBC_MIN_SMALL_POOL_ALIGNMENT_SHIFT); ++Index)
+	for (uint32 Index = 0; Index != SIZE_TO_POOL_INDEX_NUM; ++Index)
 	{
-		const uint32 BinSize = Index << UE_MBC_MIN_SMALL_POOL_ALIGNMENT_SHIFT; // inverse of int32 Index = int32((Size >> UE_MBC_MIN_SMALL_POOL_ALIGNMENT_SHIFT));
+		const uint32 BinSize = Index << UE_MBC_BIN_SIZE_SHIFT; // inverse of int32 Index = int32((Size >> UE_MBC_BIN_SIZE_SHIFT));
 		while (SizeTable[PoolIndex].BinSize < BinSize)
 		{
 			++PoolIndex;
@@ -670,8 +670,6 @@ bool FMallocBinned3::IsInternallyThreadSafe() const
 
 void* FMallocBinned3::Malloc(SIZE_T Size, uint32 Alignment)
 {
-	static_assert(DEFAULT_ALIGNMENT <= UE_MBC_MIN_SMALL_POOL_ALIGNMENT, "DEFAULT_ALIGNMENT is assumed to be zero"); // used below
-
 	// Fast path: Allocate from the small pools if the size is small enough and the alignment <= binned3 min alignment.
 	//            Larger alignments can waste a lot of memory allocating an entire page, so some smaller alignments are 
 	//			  handled in the fallback path if less than a predefined max small pool alignment. 
@@ -770,9 +768,6 @@ void* FMallocBinned3::Malloc(SIZE_T Size, uint32 Alignment)
 		return Result;
 	}
 	Alignment = FMath::Max<uint32>(Alignment, UE_MBC_MIN_SMALL_POOL_ALIGNMENT);
-	Size = Align(FMath::Max((SIZE_T)1, Size), Alignment);
-
-	check(FMath::IsPowerOfTwo(Alignment));
 
 	// Use OS for non-pooled allocations.
 	const uint64 AlignedSize = Align(Size, FPlatformMemory::FPlatformVirtualMemoryBlock::GetCommitAlignment());
@@ -845,7 +840,7 @@ void* FMallocBinned3::Realloc(void* Ptr, SIZE_T NewSize, uint32 Alignment)
 		FMallocBinned3::FreeExternal(Ptr, PoolIndex);
 		return nullptr;
 	}
-	static_assert(DEFAULT_ALIGNMENT <= UE_MBC_MIN_SMALL_POOL_ALIGNMENT, "DEFAULT_ALIGNMENT is assumed to be zero"); // used below
+
 	check(FMath::IsPowerOfTwo(Alignment));
 	check(Alignment <= OsAllocationGranularity);
 
diff --git a/Engine/Source/Runtime/Core/Private/HAL/MallocBinnedCommon.cpp b/Engine/Source/Runtime/Core/Private/HAL/MallocBinnedCommon.cpp
--- a/Engine/Source/Runtime/Core/Private/HAL/MallocBinnedCommon.cpp
+++ b/Engine/Source/Runtime/Core/Private/HAL/MallocBinnedCommon.cpp
@@ -21,6 +21,8 @@ PRAGMA_DISABLE_UNSAFE_TYPECAST_WARNINGS
 
 CSV_DEFINE_CATEGORY(MallocBinned, false);
 
+static_assert(DEFAULT_ALIGNMENT <= UE_MBC_MIN_SMALL_POOL_ALIGNMENT, "DEFAULT_ALIGNMENT is assumed to be zero");
+
 // Bin sizes are based around getting the maximum amount of allocations per block, with as little alignment waste as possible.
 // Bin sizes should be close to even divisors of the system page size, and well distributed.
 // They must be 16-byte aligned as well.
@@ -137,6 +139,8 @@ static inline void FillTable(FSizeTableEntry* SizeTable, int32& Index, const uin
 
 uint8 FSizeTableEntry::FillSizeTable(uint64 PlatformPageSize, FSizeTableEntry* SizeTable, uint32 BasePageSize, uint32 MaxSize, uint32 SizeIncrement)
 {
+	check(BasePageSize <= UE_MBC_MAX_SUPPORTED_PLATFORM_PAGE_SIZE);
+
 	int32 Index = 0;
 	FillTable(SizeTable, Index, BinnedCommonSmallBinSizes4k,  UE_ARRAY_COUNT(BinnedCommonSmallBinSizes4k),  1, PlatformPageSize, BasePageSize);
 	FillTable(SizeTable, Index, BinnedCommonSmallBinSizes8k,  UE_ARRAY_COUNT(BinnedCommonSmallBinSizes8k),  2, PlatformPageSize, BasePageSize);
diff --git a/Engine/Source/Runtime/Core/Public/HAL/MallocBinned2.h b/Engine/Source/Runtime/Core/Public/HAL/MallocBinned2.h
--- a/Engine/Source/Runtime/Core/Public/HAL/MallocBinned2.h
+++ b/Engine/Source/Runtime/Core/Public/HAL/MallocBinned2.h
@@ -318,7 +318,7 @@ public:
 	static uint16 SmallBinSizes[UE_MB2_SMALL_POOL_COUNT];
 	static FMallocBinned2* MallocBinned2;
 	// Mapping of sizes to small table indices
-	static uint8 MemSizeToPoolIndex[1 + (UE_MB2_MAX_SMALL_POOL_SIZE >> UE_MBC_MIN_SMALL_POOL_ALIGNMENT_SHIFT)];
+	static uint8 MemSizeToPoolIndex[1 + (UE_MB2_MAX_SMALL_POOL_SIZE >> UE_MBC_BIN_SIZE_SHIFT)];
 
 	static void* AllocateMetaDataMemory(SIZE_T Size);
 	static void FreeMetaDataMemory(void* Ptr, SIZE_T Size);
diff --git a/Engine/Source/Runtime/Core/Public/HAL/MallocBinned3.h b/Engine/Source/Runtime/Core/Public/HAL/MallocBinned3.h
--- a/Engine/Source/Runtime/Core/Public/HAL/MallocBinned3.h
+++ b/Engine/Source/Runtime/Core/Public/HAL/MallocBinned3.h
@@ -371,11 +371,11 @@ public:
 	FPlatformMemory::FPlatformVirtualMemoryBlock PoolBaseVMBlock[UE_MB3_SMALL_POOL_COUNT];
 #endif
 	// Mapping of sizes to small table indices
-	static uint8 MemSizeToPoolIndex[1 + (UE_MB3_MAX_SMALL_POOL_SIZE >> UE_MBC_MIN_SMALL_POOL_ALIGNMENT_SHIFT)];
+	static uint8 MemSizeToPoolIndex[SIZE_TO_POOL_INDEX_NUM];
 
 	FORCEINLINE uint32 PoolIndexToBinSize(uint32 PoolIndex) const
 	{
-		return uint32(SmallBinSizesShifted[PoolIndex + 1]) << UE_MBC_MIN_SMALL_POOL_ALIGNMENT_SHIFT;
+		return uint32(SmallBinSizesShifted[PoolIndex + 1]) << UE_MBC_BIN_SIZE_SHIFT;
 	}
 
 	void FreeBundles(FBundleNode* Bundles, uint32 PoolIndex);
diff --git a/Engine/Source/Runtime/Core/Public/HAL/MallocBinnedCommon.h b/Engine/Source/Runtime/Core/Public/HAL/MallocBinnedCommon.h
--- a/Engine/Source/Runtime/Core/Public/HAL/MallocBinnedCommon.h
+++ b/Engine/Source/Runtime/Core/Public/HAL/MallocBinnedCommon.h
@@ -43,7 +43,7 @@
 #endif
 
 #define UE_MBC_MIN_SMALL_POOL_ALIGNMENT					16	// Minimum alignment of bins
-#define UE_MBC_MIN_SMALL_POOL_ALIGNMENT_SHIFT			4	// Alignment of bins, expressed as a shift
+#define UE_MBC_BIN_SIZE_SHIFT							4	// Alignment of bins, expressed as a shift
 #define UE_MBC_MAX_SMALL_POOL_ALIGNMENT					256
 
 #if !defined(AGGRESSIVE_MEMORY_SAVING)
@@ -104,6 +104,8 @@
 #	error "MB lightweight bin callstack tracker needs UE_MBC_ALLOCATOR_STATS to be enabled."
 #endif
 
+#define UE_MBC_MAX_SUPPORTED_PLATFORM_PAGE_SIZE			16 * 1024	// MB3 and it's bins were designed around 4\16KB native page size. Having larger page size will increase memory waste and will be inefficient for current allocator design
+
 CSV_DECLARE_CATEGORY_EXTERN(MallocBinned);
 
 extern CORE_API float GMallocBinnedFlushThreadCacheMaxWaitTime;
@@ -489,6 +491,8 @@ public:
 	}
 
 protected:
+	static constexpr int SIZE_TO_POOL_INDEX_NUM = 1 + (MAX_SMALL_POOL_SIZE >> UE_MBC_BIN_SIZE_SHIFT);
+
 	struct FPerThreadFreeBlockLists
 	{
 		inline static FPerThreadFreeBlockLists* Get() TSAN_SAFE
@@ -825,7 +829,6 @@ protected:
 
 	inline SIZE_T QuantizeSizeCommon(SIZE_T Count, uint32 Alignment, const AllocType& Alloc) const
 	{
-		static_assert(DEFAULT_ALIGNMENT <= UE_MBC_MIN_SMALL_POOL_ALIGNMENT, "DEFAULT_ALIGNMENT is assumed to be zero"); // used below
 		checkSlow(FMath::IsPowerOfTwo(Alignment));
 		SIZE_T SizeOut;
 		if ((Count <= MAX_SMALL_POOL_SIZE) & (Alignment <= UE_MBC_MIN_SMALL_POOL_ALIGNMENT)) // one branch, not two
@@ -859,10 +862,10 @@ protected:
 		return SizeOut;
 	}
 
-	inline uint32 BoundSizeToPoolIndex(SIZE_T Size, const uint8(&MemSizeToPoolIndex)[1 + (MAX_SMALL_POOL_SIZE >> UE_MBC_MIN_SMALL_POOL_ALIGNMENT_SHIFT)]) const
+	inline uint32 BoundSizeToPoolIndex(SIZE_T Size, const uint8(&MemSizeToPoolIndex)[SIZE_TO_POOL_INDEX_NUM]) const
 	{
-		const auto Index = ((Size + UE_MBC_MIN_SMALL_POOL_ALIGNMENT - 1) >> UE_MBC_MIN_SMALL_POOL_ALIGNMENT_SHIFT);
-		checkSlow(Index >= 0 && Index <= (MAX_SMALL_POOL_SIZE >> UE_MBC_MIN_SMALL_POOL_ALIGNMENT_SHIFT)); // and it should be in the table
+		const auto Index = ((Size + UE_MBC_MIN_SMALL_POOL_ALIGNMENT - 1) >> UE_MBC_BIN_SIZE_SHIFT);
+		checkSlow(Index >= 0 && Index < SIZE_TO_POOL_INDEX_NUM); // and it should be in the table
 		const uint32 PoolIndex = uint32(MemSizeToPoolIndex[Index]);
 		checkSlow(PoolIndex >= 0 && PoolIndex < NUM_SMALL_POOLS);
 		return PoolIndex;
