More MB2\3 code cleanup

diff --git a/Engine/Source/Runtime/Core/Private/HAL/MallocBinned2.cpp b/Engine/Source/Runtime/Core/Private/HAL/MallocBinned2.cpp
--- a/Engine/Source/Runtime/Core/Private/HAL/MallocBinned2.cpp
+++ b/Engine/Source/Runtime/Core/Private/HAL/MallocBinned2.cpp
@@ -622,13 +622,12 @@ void* FMallocBinned2::MallocExternalSmall(SIZE_T Size, uint32 Alignment)
 void* FMallocBinned2::MallocExternalLarge(SIZE_T Size, uint32 Alignment)
 {
 	Alignment = FMath::Max<uint32>(Alignment, UE_MBC_MIN_SMALL_POOL_ALIGNMENT);
-	Size = Align(FMath::Max((SIZE_T)1, Size), Alignment);
 
 	checkf(FMath::IsPowerOfTwo(Alignment), TEXT("Invalid Malloc alignment: '%u' is not a power of two"), Alignment);
 	checkf(Alignment <= GMB2PageSize, TEXT("Invalid Malloc alignment: '%u' is greater than the page size '%u'"), Alignment, GMB2PageSize);
-	checkf(IsSupportedSize(Size), TEXT("Invalid Malloc size: '%" SIZE_T_FMT "'"), Size);
-
+	
 	const SIZE_T AlignedSize = Align(Size, OsAllocationGranularity);
+	checkf(IsSupportedSize(AlignedSize), TEXT("Invalid Malloc size: '%" SIZE_T_FMT "'"), AlignedSize);
 
 	LLM_PLATFORM_SCOPE(ELLMTag::FMalloc);
 
@@ -660,7 +659,7 @@ void* FMallocBinned2::MallocExternalLarge(SIZE_T Size, uint32 Alignment)
 #endif
 
 	// Create pool.
-	check(Size > 0 && Size <= AlignedSize && AlignedSize >= OsAllocationGranularity);
+	check(Size > 0 && AlignedSize >= OsAllocationGranularity);
 	Pool->SetOSAllocationSizes(Size, AlignedSize);
 
 	return Result;
@@ -674,6 +673,8 @@ void* FMallocBinned2::Realloc(void* Ptr, SIZE_T NewSize, uint32 Alignment)
 		return nullptr;
 	}
 
+	checkf(FMath::IsPowerOfTwo(Alignment), TEXT("Invalid Realloc alignment: '%u' is not a power of two"), Alignment);
+
 	if (!IsOSAllocation(Ptr))
 	{
 		check(Ptr); // null is 64k aligned so we should not be here
@@ -749,9 +750,7 @@ void* FMallocBinned2::Realloc(void* Ptr, SIZE_T NewSize, uint32 Alignment)
 	ExternalAllocMutex.Unlock();
 
 	Alignment = FMath::Max<uint32>(Alignment, UE_MBC_MIN_SMALL_POOL_ALIGNMENT);
-	NewSize = Align(FMath::Max((SIZE_T)1, NewSize), Alignment);
 
-	checkf(FMath::IsPowerOfTwo(Alignment), TEXT("Invalid Realloc alignment: '%u' is not a power of two"), Alignment);
 	checkf(Alignment <= GMB2PageSize, TEXT("Invalid Realloc alignment: '%u' is greater than the page size '%u'"), Alignment, GMB2PageSize);
 	checkf(IsSupportedSize(NewSize), TEXT("Invalid Realloc size: '%" SIZE_T_FMT "'"), NewSize);
 
diff --git a/Engine/Source/Runtime/Core/Private/HAL/MallocBinned3.cpp b/Engine/Source/Runtime/Core/Private/HAL/MallocBinned3.cpp
--- a/Engine/Source/Runtime/Core/Private/HAL/MallocBinned3.cpp
+++ b/Engine/Source/Runtime/Core/Private/HAL/MallocBinned3.cpp
@@ -675,6 +675,8 @@ void* FMallocBinned3::Malloc(SIZE_T Size, uint32 Alignment)
 	//			  handled in the fallback path if less than a predefined max small pool alignment. 
 
 	bool UsePools = (Size <= UE_MB3_MAX_SMALL_POOL_SIZE);
+
+	check(FMath::IsPowerOfTwo(Alignment));
 	
 	if (UNLIKELY(UsePools && (Alignment > UE_MBC_MIN_SMALL_POOL_ALIGNMENT)))
 	{
