Added 8 bytes bin to MB3

diff --git a/Engine/Source/Runtime/Core/Private/HAL/MallocBinned2.cpp b/Engine/Source/Runtime/Core/Private/HAL/MallocBinned2.cpp
--- a/Engine/Source/Runtime/Core/Private/HAL/MallocBinned2.cpp
+++ b/Engine/Source/Runtime/Core/Private/HAL/MallocBinned2.cpp
@@ -240,76 +240,68 @@ struct FMallocBinned2::Private
 
 		UE::TUniqueLock Lock(Table.Mutex);
 
-		FBundleNode* Bundle = BundlesToRecycle;
-		while (Bundle)
+		FBundleNode* Node = BundlesToRecycle;
+		do
 		{
-			FBundleNode* NextBundle = Bundle->NextBundle;
+			FBundleNode* NextNode = Node->GetNextNodeInCurrentBundle();
+			FPoolInfo*   NodePool = Internal::FindPoolInfo(Allocator, Node);
+			if (!NodePool)
+			{
+				UE_LOG(LogMemory, Fatal, TEXT("FMallocBinned2 Attempt to free an unrecognized small block %p"), Node);
+			}
+			NodePool->CheckCanary(FPoolInfo::ECanary::FirstFreeBlockIsPtr);
 
-			FBundleNode* Node = Bundle;
-			do
+			// If this pool was exhausted, move to available list.
+			if (!NodePool->FirstFreeBlock)
 			{
-				FBundleNode* NextNode = Node->NextNodeInCurrentBundle;
-				FPoolInfo*   NodePool = Internal::FindPoolInfo(Allocator, Node);
-				if (!NodePool)
-				{
-					UE_LOG(LogMemory, Fatal, TEXT("FMallocBinned2 Attempt to free an unrecognized small block %p"), Node);
-				}
-				NodePool->CheckCanary(FPoolInfo::ECanary::FirstFreeBlockIsPtr);
+				Table.ActivePools.LinkToFront(NodePool);
+			}
+			else
+			{
+				// If we are freeing memory in this pool it must have the current canary and not the pre-fork one. All caches should have been cleared when forking.
+				check(NodePool->FirstFreeBlock->CanaryAndForkState == EBlockCanary::Zero || NodePool->FirstFreeBlock->CanaryAndForkState == Allocator.CurrentCanary);
+			}
 
-				// If this pool was exhausted, move to available list.
-				if (!NodePool->FirstFreeBlock)
-				{
-					Table.ActivePools.LinkToFront(NodePool);
-				}
-				else
-				{
-					// If we are freeing memory in this pool it must have the current canary and not the pre-fork one. All caches should have been cleared when forking.
-					check(NodePool->FirstFreeBlock->CanaryAndForkState == EBlockCanary::Zero || NodePool->FirstFreeBlock->CanaryAndForkState == Allocator.CurrentCanary);
-				}
+			// Free a pooled allocation.
+			FFreeBlock* Free = (FFreeBlock*)Node;
+			Free->NumFreeBins = 1;
+			Free->NextFreeBlock = NodePool->FirstFreeBlock;
+			Free->BinSize = BinSize;
+			Free->CanaryAndForkState = Allocator.CurrentCanary;
+			Free->PoolIndex = InPoolIndex;
+			NodePool->FirstFreeBlock = Free;
 
-				// Free a pooled allocation.
-				FFreeBlock* Free = (FFreeBlock*)Node;
-				Free->NumFreeBins = 1;
-				Free->NextFreeBlock = NodePool->FirstFreeBlock;
-				Free->BinSize = BinSize;
-				Free->CanaryAndForkState = Allocator.CurrentCanary;
-				Free->PoolIndex = InPoolIndex;
-				NodePool->FirstFreeBlock = Free;
+			UE_MBC_UPDATE_STATS(--Table.TotalUsedBins);
 
-				UE_MBC_UPDATE_STATS(--Table.TotalUsedBins);
+			// Free this pool.
+			check(NodePool->Taken >= 1);
+			if (--NodePool->Taken == 0)
+			{
+				NodePool->SetCanary(FPoolInfo::ECanary::Unassigned, true, false);
+				FFreeBlock* BasePtrOfNode = GetPoolHeaderFromPointer(Node);
 
-				// Free this pool.
-				check(NodePool->Taken >= 1);
-				if (--NodePool->Taken == 0)
+				// Free the OS memory.
+				NodePool->Unlink();
 				{
-					NodePool->SetCanary(FPoolInfo::ECanary::Unassigned, true, false);
-					FFreeBlock* BasePtrOfNode = GetPoolHeaderFromPointer(Node);
-
-					// Free the OS memory.
-					NodePool->Unlink();
-					{
-						UE::TUniqueLock InnerLock(Allocator.ExternalAllocMutex);
-						Allocator.CachedOSPageAllocator.Free(BasePtrOfNode, GMB2PageSize);
-					}
+					UE::TUniqueLock InnerLock(Allocator.ExternalAllocMutex);
+					Allocator.CachedOSPageAllocator.Free(BasePtrOfNode, GMB2PageSize);
+				}
 
 #if UE_MB2_ALLOCATOR_STATS
-					int NumBinsInBlock = GMB2PageSize / BinSize;
-					if (NumBinsInBlock * BinSize + sizeof(FFreeBlock) > GMB2PageSize)
-					{
-						NumBinsInBlock--;
-					}
-
-					Table.TotalAllocatedBins -= NumBinsInBlock;
-					Table.TotalAllocatedMem -= GMB2PageSize;
-					AllocatedOSSmallPoolMemory.fetch_sub((int64)GMB2PageSize, std::memory_order_relaxed);
-#endif
+				int NumBinsInBlock = GMB2PageSize / BinSize;
+				if (NumBinsInBlock * BinSize + sizeof(FFreeBlock) > GMB2PageSize)
+				{
+					NumBinsInBlock--;
 				}
 
-				Node = NextNode;
-			} while (Node);
+				Table.TotalAllocatedBins -= NumBinsInBlock;
+				Table.TotalAllocatedMem -= GMB2PageSize;
+				AllocatedOSSmallPoolMemory.fetch_sub((int64)GMB2PageSize, std::memory_order_relaxed);
+#endif
+			}
 
-			Bundle = NextBundle;
-		}
+			Node = NextNode;
+		} while (Node);
 	}
 
 #if BINNED2_FORK_SUPPORT
@@ -397,7 +389,7 @@ FMallocBinned2::FMallocBinned2()
 
 	FMemory::Memcpy(SmallBinSizes, SmallBinSizesInternal);
 	
-	static_assert(sizeof(FFreeBlock) <= SmallBinSizesInternal[0], "sizeof(FFreeBlock)  must be fit in smallest allocation size handled by FMallocBinned2.");
+	static_assert(sizeof(FFreeBlock)  <= SmallBinSizesInternal[0], "sizeof(FFreeBlock)  must be fit in smallest allocation size handled by FMallocBinned2.");
 	static_assert(sizeof(FBundleNode) <= SmallBinSizesInternal[0], "sizeof(FBundleNode) must fit in smallest allocation size handled by FMallocBinned2.");
 
 	const FGenericPlatformMemoryConstants Constants = FPlatformMemory::GetConstants();
@@ -484,7 +476,6 @@ void FMallocBinned2::OnPreFork()
 	{
 		while (FBundleNode* Node = FMallocBinned2::Private::GGlobalRecycler.PopBundle(PoolIndex))
 		{
-			Node->NextBundle = nullptr; // We need to override this because it's in a union with Count which is not needed for the freeing work
 			FMallocBinned2::Private::FreeBundles(*this, Node, PoolIndex);
 		}
 	}
@@ -529,7 +520,7 @@ bool FMallocBinned2::IsInternallyThreadSafe() const
 
 FORCEINLINE static bool UseSmallAlloc(SIZE_T Size, uint32 Alignment)
 {
-	return ((Size <= UE_MB2_MAX_SMALL_POOL_SIZE) & (Alignment <= UE_MBC_MIN_SMALL_POOL_ALIGNMENT)); // one branch, not two
+	return ((Size <= UE_MB2_MAX_SMALL_POOL_SIZE) & (Alignment <= UE_MBC_STANDARD_ALIGNMENT)); // one branch, not two
 }
 
 void* FMallocBinned2::Malloc(SIZE_T Size, uint32 Alignment)
@@ -621,7 +612,7 @@ void* FMallocBinned2::MallocExternalSmall(SIZE_T Size, uint32 Alignment)
 
 void* FMallocBinned2::MallocExternalLarge(SIZE_T Size, uint32 Alignment)
 {
-	Alignment = FMath::Max<uint32>(Alignment, UE_MBC_MIN_SMALL_POOL_ALIGNMENT);
+	Alignment = FMath::Max<uint32>(Alignment, UE_MBC_STANDARD_ALIGNMENT);
 
 	checkf(FMath::IsPowerOfTwo(Alignment), TEXT("Invalid Malloc alignment: '%u' is not a power of two"), Alignment);
 	checkf(Alignment <= GMB2PageSize, TEXT("Invalid Malloc alignment: '%u' is greater than the page size '%u'"), Alignment, GMB2PageSize);
@@ -749,7 +740,7 @@ void* FMallocBinned2::Realloc(void* Ptr, SIZE_T NewSize, uint32 Alignment)
 
 	ExternalAllocMutex.Unlock();
 
-	Alignment = FMath::Max<uint32>(Alignment, UE_MBC_MIN_SMALL_POOL_ALIGNMENT);
+	Alignment = FMath::Max<uint32>(Alignment, UE_MBC_STANDARD_ALIGNMENT);
 
 	checkf(Alignment <= GMB2PageSize, TEXT("Invalid Realloc alignment: '%u' is greater than the page size '%u'"), Alignment, GMB2PageSize);
 	checkf(IsSupportedSize(NewSize), TEXT("Invalid Realloc size: '%" SIZE_T_FMT "'"), NewSize);
@@ -793,12 +784,11 @@ void FMallocBinned2::Free(void* Ptr)
 		else
 		{
 			BundlesToRecycle = (FBundleNode*)Ptr;
-			BundlesToRecycle->NextNodeInCurrentBundle = nullptr;
+			BundlesToRecycle->SetNextNodeInCurrentBundle(nullptr);
 		}
 
 		if (BundlesToRecycle)
 		{
-			BundlesToRecycle->NextBundle = nullptr;
 			Private::FreeBundles(*this, BundlesToRecycle, PoolIndex);
 #if UE_MB2_ALLOCATOR_STATS
 			if (!Lists)
diff --git a/Engine/Source/Runtime/Core/Private/HAL/MallocBinned3.cpp b/Engine/Source/Runtime/Core/Private/HAL/MallocBinned3.cpp
--- a/Engine/Source/Runtime/Core/Private/HAL/MallocBinned3.cpp
+++ b/Engine/Source/Runtime/Core/Private/HAL/MallocBinned3.cpp
@@ -111,7 +111,7 @@ struct FFreeBlock
 		: BinSizeShifted(InBinSize >> UE_MBC_BIN_SIZE_SHIFT)
 		, PoolIndex(InPoolIndex)
 		, Canary(CANARY_VALUE)
-		, NextFreeBlockIndex(MAX_uint32)
+		, NextFreeBlockIndex(InvalidNextFreeBlock)
 	{
 		check(InPoolIndex < MAX_uint8 && (InBinSize >> UE_MBC_BIN_SIZE_SHIFT) <= MAX_uint16);
 		NumFreeBins = InBlockSize / InBinSize;
@@ -147,13 +147,20 @@ struct FFreeBlock
 		return (uint8*)this + NumFreeBins * (uint32(BinSizeShifted) << UE_MBC_BIN_SIZE_SHIFT);
 	}
 
-	uint16 BinSizeShifted;		// Size of the bins that this list points to >> UE_MBC_BIN_SIZE_SHIFT
-	uint8  PoolIndex;			// Index of this pool
-	uint8  Canary;				// Constant value of 0xe7
-	uint32 NumFreeBins;			// Number of consecutive free bins here, at least 1.
-	uint32 NextFreeBlockIndex;	// Next free block or MAX_uint32
+	constexpr static uint32 MaxBitsForBinSizeShifted = 32 - FMath::CountLeadingZeros(UE_MB3_MAX_SMALL_POOL_SIZE >> UE_MBC_BIN_SIZE_SHIFT);
+	constexpr static uint32 MaxBitsForPoolIndex = 32 - FMath::CountLeadingZeros(UE_MB3_SMALL_POOL_COUNT);
+	constexpr static uint32 MaxBitsForFreeBins = 32 - FMath::CountLeadingZeros(UE_MBC_MAX_SUPPORTED_PLATFORM_PAGE_SIZE / UE_MBC_MIN_BIN_SIZE);	// max supported page size is 16 kb (on iOS) and the highest amount of free bins in one block will be for the 8 bytes bin in the 16kb block
+	constexpr static uint32 InvalidNextFreeBlock = 1 << MaxBitsForFreeBins;
+
+	uint32 BinSizeShifted : MaxBitsForBinSizeShifted;	// Size of the bins that this list points to >> UE_MBC_BIN_SIZE_SHIFT
+	uint32 PoolIndex : MaxBitsForPoolIndex;				// Index of this pool
+	uint32 Canary : 8;									// Constant value of 0xe7
+	uint32 NumFreeBins : MaxBitsForFreeBins;			// Number of consecutive free bins here, at least 1 for a newly initialized FFreeBlock
+	uint32 NextFreeBlockIndex : MaxBitsForFreeBins + 1;	// Next free FFreeBlock or InvalidNextFreeBlock
 };
 
+static_assert(sizeof(FFreeBlock) <= UE_MBC_MIN_BIN_SIZE, "FFreeBlock should fit into 8 bytes to support 8 bytes bins");
+
 
 struct FMallocBinned3::FPoolInfoSmall		//This is more like BlockInfoSmall as it stores info per block
 {
@@ -205,7 +212,7 @@ struct FMallocBinned3::FPoolInfoSmall		//This is more like BlockInfoSmall as it
 		void* Result = Free->AllocateBin();
 		if (Free->GetNumFreeBins() == 0)
 		{
-			if (Free->NextFreeBlockIndex == MAX_uint32)
+			if (Free->NextFreeBlockIndex == FFreeBlock::InvalidNextFreeBlock)
 			{
 				FirstFreeIndex = 0;
 				NoFirstFreeIndex = 1;
@@ -313,76 +320,68 @@ struct FMallocBinned3::Private
 		FPoolTable& Table = Allocator.SmallPoolTables[InPoolIndex];
 		UE::TUniqueLock Lock(Table.Mutex);
 
-		FBundleNode* Bundle = BundlesToRecycle;
-		while (Bundle)
+		FBundleNode* Node = BundlesToRecycle;
+		do
 		{
-			FBundleNode* NextBundle = Bundle->NextBundle;
+			FBundleNode* NextNode = Node->GetNextNodeInCurrentBundle();
 
-			FBundleNode* Node = Bundle;
-			do
+			uint32 OutBlockIndex;
+			const uint32 BinSize = Allocator.SmallPoolTables[InPoolIndex].BinSize;
+			void* BaseBlockPtr = Allocator.BlockPointerFromContainedPtr(Node, Allocator.SmallPoolTables[InPoolIndex].NumMemoryPagesPerBlock, OutBlockIndex);
+			const uint32 BinIndexWithinBlock = (((uint8*)Node) - ((uint8*)BaseBlockPtr)) / BinSize;
+
+			FPoolInfoSmall* NodePoolBlock = Allocator.SmallPoolTables[InPoolIndex].PoolInfos[OutBlockIndex / Allocator.SmallPoolInfosPerPlatformPage];
+			if (!NodePoolBlock)
 			{
-				FBundleNode* NextNode = Node->NextNodeInCurrentBundle;
+				UE_LOG(LogMemory, Fatal, TEXT("FMallocBinned3 Attempt to free an unrecognized small block %p"), Node);
+			}
+			FPoolInfoSmall* NodePool = &NodePoolBlock[OutBlockIndex % Allocator.SmallPoolInfosPerPlatformPage];
 
-				uint32 OutBlockIndex;
-				const uint32 BinSize = Allocator.SmallPoolTables[InPoolIndex].BinSize;
-				void* BaseBlockPtr = Allocator.BlockPointerFromContainedPtr(Node, Allocator.SmallPoolTables[InPoolIndex].NumMemoryPagesPerBlock, OutBlockIndex);
-				const uint32 BinIndexWithinBlock = (((uint8*)Node) - ((uint8*)BaseBlockPtr)) / BinSize;
+			NodePool->CheckCanary(FPoolInfoSmall::ECanary::Assigned);
 
-				FPoolInfoSmall* NodePoolBlock = Allocator.SmallPoolTables[InPoolIndex].PoolInfos[OutBlockIndex / Allocator.SmallPoolInfosPerPlatformPage];
-				if (!NodePoolBlock)
-				{
-					UE_LOG(LogMemory, Fatal, TEXT("FMallocBinned3 Attempt to free an unrecognized small block %p"), Node);
-				}
-				FPoolInfoSmall* NodePool = &NodePoolBlock[OutBlockIndex % Allocator.SmallPoolInfosPerPlatformPage];
-
-				NodePool->CheckCanary(FPoolInfoSmall::ECanary::Assigned);
+			const bool bWasExhaused = NodePool->NoFirstFreeIndex;
 
-				const bool bWasExhaused = NodePool->NoFirstFreeIndex;
+			// Free a pooled allocation.
+			FFreeBlock* Free = (FFreeBlock*)Node;
+			Free->NumFreeBins = 1;
+			Free->NextFreeBlockIndex = NodePool->NoFirstFreeIndex ? FFreeBlock::InvalidNextFreeBlock : NodePool->FirstFreeIndex;
+			Free->BinSizeShifted = (BinSize >> UE_MBC_BIN_SIZE_SHIFT);
+			Free->Canary = FFreeBlock::CANARY_VALUE;
+			Free->PoolIndex = InPoolIndex;
+			NodePool->FirstFreeIndex = BinIndexWithinBlock;
+			NodePool->NoFirstFreeIndex = 0;
+			check(uint32(NodePool->FirstFreeIndex) == BinIndexWithinBlock);
 
-				// Free a pooled allocation.
-				FFreeBlock* Free = (FFreeBlock*)Node;
-				Free->NumFreeBins = 1;
-				Free->NextFreeBlockIndex = NodePool->NoFirstFreeIndex ? MAX_uint32 : NodePool->FirstFreeIndex;
-				Free->BinSizeShifted = (BinSize >> UE_MBC_BIN_SIZE_SHIFT);
-				Free->Canary = FFreeBlock::CANARY_VALUE;
-				Free->PoolIndex = InPoolIndex;
-				NodePool->FirstFreeIndex = BinIndexWithinBlock;
-				NodePool->NoFirstFreeIndex = 0;
-				check(uint32(NodePool->FirstFreeIndex) == BinIndexWithinBlock);
+			UE_MBC_UPDATE_STATS(--Table.TotalUsedBins);
 
-				UE_MBC_UPDATE_STATS(--Table.TotalUsedBins);
+			// Free this pool.
+			check(NodePool->Taken >= 1);
+			if (--NodePool->Taken == 0)
+			{
+				NodePool->SetCanary(FPoolInfoSmall::ECanary::Unassigned, true, false);
+				Table.BlocksAllocatedBits.FreeBit(OutBlockIndex);
 
-				// Free this pool.
-				check(NodePool->Taken >= 1);
-				if (--NodePool->Taken == 0)
+				if (!bWasExhaused)
 				{
-					NodePool->SetCanary(FPoolInfoSmall::ECanary::Unassigned, true, false);
-					Table.BlocksAllocatedBits.FreeBit(OutBlockIndex);
-
-					if (!bWasExhaused)
-					{
-						Table.BlocksExhaustedBits.AllocBit(OutBlockIndex);
-					}
+					Table.BlocksExhaustedBits.AllocBit(OutBlockIndex);
+				}
 
-					const uint64 BlockSize = Allocator.SmallPoolTables[InPoolIndex].BlockSize;
-					Allocator.Decommit(InPoolIndex, BaseBlockPtr, BlockSize);
+				const uint64 BlockSize = Allocator.SmallPoolTables[InPoolIndex].BlockSize;
+				Allocator.Decommit(InPoolIndex, BaseBlockPtr, BlockSize);
 
 #if UE_MB3_ALLOCATOR_STATS
-					Table.TotalAllocatedBins -= BlockSize / BinSize;
-					Table.TotalAllocatedMem -= BlockSize;
-					AllocatedOSSmallPoolMemory.fetch_sub(BlockSize, std::memory_order_relaxed);
+				Table.TotalAllocatedBins -= BlockSize / BinSize;
+				Table.TotalAllocatedMem -= BlockSize;
+				AllocatedOSSmallPoolMemory.fetch_sub(BlockSize, std::memory_order_relaxed);
 #endif
-				}
-				else if (bWasExhaused)
-				{
-					Table.BlocksExhaustedBits.FreeBit(OutBlockIndex);
-				}
-
-				Node = NextNode;
-			} while (Node);
+			}
+			else if (bWasExhaused)
+			{
+				Table.BlocksExhaustedBits.FreeBit(OutBlockIndex);
+			}
 
-			Bundle = NextBundle;
-		}
+			Node = NextNode;
+		} while (Node);
 	}
 };
 
@@ -533,14 +532,14 @@ FMallocBinned3::FMallocBinned3(const FPlatformMemory::FPlatformVirtualMemoryBloc
 
 	checkf(Constants.AddressLimit > OsAllocationGranularity, TEXT("OS address limit must be greater than the page size")); // Check to catch 32 bit overflow in AddressLimit
 	static_assert(UE_MB3_SMALL_POOL_COUNT <= 256, "Small bins size array size must fit in a byte");
-	static_assert(sizeof(FFreeBlock) <= UE_MBC_MIN_SMALL_POOL_ALIGNMENT, "Free block struct must be small enough to fit into the smallest bin");
+	static_assert(sizeof(FFreeBlock) <= UE_MBC_MIN_BIN_SIZE, "Free block struct must be small enough to fit into the smallest bin");
 
 	// Init pool tables.
 	FSizeTableEntry SizeTable[UE_MB3_SMALL_POOL_COUNT];
 
 	verify(FSizeTableEntry::FillSizeTable(OsAllocationGranularity, SizeTable, UE_MB3_BASE_PAGE_SIZE, UE_MB3_MAX_SMALL_POOL_SIZE, UE_MB3_BASE_PAGE_SIZE) == UE_MB3_SMALL_POOL_COUNT);
 	checkf(SizeTable[UE_MB3_SMALL_POOL_COUNT - 1].BinSize == UE_MB3_MAX_SMALL_POOL_SIZE, TEXT("UE_MB3_MAX_SMALL_POOL_SIZE must be equal to the largest bin size"));
-	checkf(sizeof(FFreeBlock) <= SizeTable[0].BinSize, TEXT("Pool header must be able to fit into the smallest bin"));
+	checkf(sizeof(FFreeBlock)  <= SizeTable[0].BinSize, TEXT("Pool header must be able to fit into the smallest bin"));
 
 	SmallPoolInfosPerPlatformPage = OsAllocationGranularity / sizeof(FPoolInfoSmall);
 
@@ -678,11 +677,11 @@ void* FMallocBinned3::Malloc(SIZE_T Size, uint32 Alignment)
 
 	check(FMath::IsPowerOfTwo(Alignment));
 	
-	if (UNLIKELY(UsePools && (Alignment > UE_MBC_MIN_SMALL_POOL_ALIGNMENT)))
+	if (UNLIKELY(UsePools && (Alignment > UE_MBC_STANDARD_ALIGNMENT)))
 	{
-		// check if allocations that require alignment larger than UE_MBC_MIN_SMALL_POOL_ALIGNMENT can be promoted to a bin with a natural alignment that matches
+		// check if allocations that require alignment larger than UE_MBC_STANDARD_ALIGNMENT can be promoted to a bin with a natural alignment that matches
 		// i.e. 16 bytes allocation with 128 bytes alignment can be promoted to 128 bytes bin
-		// this will save us a lot of memory as otherwise allocations will be promoted to OS allocs that are at least 64 KB large, depending on a UE_MB3_MAX_SMALL_POOL_SIZE
+		// this will save us a lot of memory as otherwise allocations will be promoted to OS allocs that are at least 4 KB large, depending on a platform
 		UsePools = PromoteToLargerBin(Size, Alignment, *this);
 	}
 
@@ -769,7 +768,7 @@ void* FMallocBinned3::Malloc(SIZE_T Size, uint32 Alignment)
 
 		return Result;
 	}
-	Alignment = FMath::Max<uint32>(Alignment, UE_MBC_MIN_SMALL_POOL_ALIGNMENT);
+	Alignment = FMath::Max<uint32>(Alignment, UE_MBC_STANDARD_ALIGNMENT);
 
 	// Use OS for non-pooled allocations.
 	const uint64 AlignedSize = Align(Size, FPlatformMemory::FPlatformVirtualMemoryBlock::GetCommitAlignment());
@@ -887,7 +886,7 @@ void* FMallocBinned3::Realloc(void* Ptr, SIZE_T NewSize, uint32 Alignment)
 	const uint32 PoolOSRequestedBytes = Pool->GetOSRequestedBytes();
 	checkf(PoolOSRequestedBytes <= PoolOsBytes, TEXT("FMallocBinned3::ReallocExternal %u %u"), PoolOSRequestedBytes, PoolOsBytes);
 	if (NewSize > PoolOsBytes || // can't fit in the old block
-		(NewSize <= UE_MB3_MAX_SMALL_POOL_SIZE && Alignment <= UE_MBC_MIN_SMALL_POOL_ALIGNMENT) || // can switch to the small bin allocator
+		(NewSize <= UE_MB3_MAX_SMALL_POOL_SIZE && Alignment <= UE_MBC_STANDARD_ALIGNMENT) || // can switch to the small bin allocator
 		Align(NewSize, OsAllocationGranularity) < PoolOsBytes) // we can get some pages back
 	{
 		ExternalAllocMutex.Unlock();
@@ -931,11 +930,10 @@ void FMallocBinned3::FreeExternal(void* Ptr, uint64 PoolIndex)
 		else
 		{
 			BundlesToRecycle = (FBundleNode*)Ptr;
-			BundlesToRecycle->NextNodeInCurrentBundle = nullptr;
+			BundlesToRecycle->SetNextNodeInCurrentBundle(nullptr);
 		}
 		if (BundlesToRecycle)
 		{
-			BundlesToRecycle->NextBundle = nullptr;
 			Private::FreeBundles(*this, BundlesToRecycle, PoolIndex);
 #if UE_MB3_ALLOCATOR_STATS
 			if (!Lists)
diff --git a/Engine/Source/Runtime/Core/Private/HAL/MallocBinnedCommon.cpp b/Engine/Source/Runtime/Core/Private/HAL/MallocBinnedCommon.cpp
--- a/Engine/Source/Runtime/Core/Private/HAL/MallocBinnedCommon.cpp
+++ b/Engine/Source/Runtime/Core/Private/HAL/MallocBinnedCommon.cpp
@@ -22,13 +22,15 @@ PRAGMA_DISABLE_UNSAFE_TYPECAST_WARNINGS
 CSV_DEFINE_CATEGORY(MallocBinned, false);
 
 static_assert(DEFAULT_ALIGNMENT <= UE_MBC_MIN_SMALL_POOL_ALIGNMENT, "DEFAULT_ALIGNMENT is assumed to be zero");
+static_assert(sizeof(FMallocBinnedCommonBase::FBundleNode) <= UE_MBC_MIN_BIN_SIZE, "FBundleNode must fit into the smallest bin size");
 
 // Bin sizes are based around getting the maximum amount of allocations per block, with as little alignment waste as possible.
 // Bin sizes should be close to even divisors of the system page size, and well distributed.
 // They must be 16-byte aligned as well.
 static constexpr uint32 BinnedCommonSmallBinSizes4k[] =
 {
-	16, 32, 48, 64, 80, 96, 112, 128, 144, 160, 176, 192, // +16
+	8, 16, // +8
+	32, 48, 64, 80, 96, 112, 128, 144, 160, 176, 192, // +16
 	224, 256, 288, 320, // +32
 	368,  // /11 ish
 	400,  // /10 ish
@@ -103,7 +105,7 @@ FSizeTableEntry::FSizeTableEntry(uint32 InBinSize, uint64 PlatformPageSize, uint
 	: BinSize(InBinSize)
 {
 	check((PlatformPageSize & (BasePageSize - 1)) == 0 && PlatformPageSize >= BasePageSize);
-	checkf(InBinSize % UE_MBC_MIN_SMALL_POOL_ALIGNMENT == 0, TEXT("Small bin size must be a multiple of UE_MBC_MINIMUM_ALIGNMENT"));
+	checkf(InBinSize % UE_MBC_MIN_BIN_SIZE == 0, TEXT("Small bin size must be a multiple of UE_MBC_MIN_BIN_SIZE"));
 
 	const uint64 NumBasePagesPerPlatformPage = PlatformPageSize / BasePageSize;
 
diff --git a/Engine/Source/Runtime/Core/Public/HAL/MallocBinnedCommon.h b/Engine/Source/Runtime/Core/Public/HAL/MallocBinnedCommon.h
--- a/Engine/Source/Runtime/Core/Public/HAL/MallocBinnedCommon.h
+++ b/Engine/Source/Runtime/Core/Public/HAL/MallocBinnedCommon.h
@@ -31,7 +31,7 @@
 #endif
 
 #if !defined(UE_MBC_NUM_LISTED_SMALL_POOLS)
-#	define UE_MBC_NUM_LISTED_SMALL_POOLS				51
+#	define UE_MBC_NUM_LISTED_SMALL_POOLS				52
 #endif
 
 #if !defined(BINNEDCOMMON_USE_SEPARATE_VM_PER_POOL)
@@ -42,9 +42,11 @@
 #	endif
 #endif
 
-#define UE_MBC_MIN_SMALL_POOL_ALIGNMENT					16	// Minimum alignment of bins
-#define UE_MBC_BIN_SIZE_SHIFT							4	// Alignment of bins, expressed as a shift
+#define UE_MBC_MIN_SMALL_POOL_ALIGNMENT					8	// Minimum alignment of bins. Added to support 8 bytes bin. If you need to mask lower bits - use this!
 #define UE_MBC_MAX_SMALL_POOL_ALIGNMENT					256
+#define UE_MBC_STANDARD_ALIGNMENT						16	// Standard alignment for all allocations, except 8 bytes that is aligned by 8 bytes
+#define UE_MBC_MIN_BIN_SIZE								8	// Minimum supported bin size
+#define UE_MBC_BIN_SIZE_SHIFT							3	// Shift for a bin size to save memory
 
 #if !defined(AGGRESSIVE_MEMORY_SAVING)
 #	error "AGGRESSIVE_MEMORY_SAVING must be defined"
@@ -184,18 +186,20 @@ public:
 	// This needs to be small enough to fit inside the smallest allocation handled by MallocBinned2\3, hence the union.
 	struct FBundleNode
 	{
-		FBundleNode* NextNodeInCurrentBundle;
+		uint64 NextNodeInCurrentBundle : 48;
+		uint64 Count : 8;		// 8 bits is enough to store count as UE_DEFAULT_GMallocBinnedBundleCount is 64
+		uint64 Reserved : 8;	// Reserved for ARM HW for TBI, MTE, etc
 
-		// NextBundle ptr is valid when node is stored in FFreeBlockList in a thread-local list of reusable allocations.
-		// Count is valid when node is stored in global recycler and caches the number of nodes in the list formed by NextNodeInCurrentBundle.
-		union
+		void SetNextNodeInCurrentBundle(FBundleNode* Next)
 		{
-			FBundleNode* NextBundle;
-			int32 Count;
-		};
-	};
+			NextNodeInCurrentBundle = reinterpret_cast<uint64>(Next);
+		}
 
-	static_assert(sizeof(FBundleNode) <= UE_MBC_MIN_SMALL_POOL_ALIGNMENT, "Bundle nodes must fit into the smallest block size");
+		FBundleNode* GetNextNodeInCurrentBundle()
+		{
+			return (FBundleNode*)NextNodeInCurrentBundle;
+		}
+	};
 
 	CORE_API virtual void OnMallocInitialized() override; 
 
@@ -274,8 +278,7 @@ private:
 
 		inline void PushHead(FBundleNode* Node)
 		{
-			Node->NextNodeInCurrentBundle = Head;
-			Node->NextBundle = nullptr;
+			Node->SetNextNodeInCurrentBundle(Head);
 			Head = Node;
 			Count++;
 		}
@@ -285,7 +288,7 @@ private:
 			FBundleNode* Result = Head;
 
 			Count--;
-			Head = Head->NextNodeInCurrentBundle;
+			Head = Head->GetNextNodeInCurrentBundle();
 			return Result;
 		}
 
@@ -340,7 +343,6 @@ protected:
 				if (!InGlobalRecycler.PushBundle(InPoolIndex, FullBundle.Head))
 				{
 					Result = FullBundle.Head;
-					Result->NextBundle = nullptr;
 				}
 				FullBundle.Reset();
 			}
@@ -357,7 +359,6 @@ protected:
 				if (PartialBundle.Head)
 				{
 					PartialBundle.Count = PartialBundle.Head->Count;
-					PartialBundle.Head->NextBundle = nullptr;
 					return true;
 				}
 				return false;
@@ -371,20 +372,25 @@ protected:
 			if (Partial)
 			{
 				PartialBundle.Reset();
-				Partial->NextBundle = nullptr;
 			}
 
 			FBundleNode* Full = FullBundle.Head;
 			if (Full)
 			{
 				FullBundle.Reset();
-				Full->NextBundle = nullptr;
 			}
 
 			FBundleNode* Result = Partial;
 			if (Result)
 			{
-				Result->NextBundle = Full;
+				FBundleNode* Prev = Result;
+				FBundleNode* Next = Result->GetNextNodeInCurrentBundle();
+				while (Next)
+				{
+					Prev = Next;
+					Next = Next->GetNextNodeInCurrentBundle();
+				}
+				Prev->SetNextNodeInCurrentBundle(Full);
 			}
 			else
 			{
@@ -831,13 +837,13 @@ protected:
 	{
 		checkSlow(FMath::IsPowerOfTwo(Alignment));
 		SIZE_T SizeOut;
-		if ((Count <= MAX_SMALL_POOL_SIZE) & (Alignment <= UE_MBC_MIN_SMALL_POOL_ALIGNMENT)) // one branch, not two
+		if ((Count <= MAX_SMALL_POOL_SIZE) & (Alignment <= UE_MBC_STANDARD_ALIGNMENT)) // one branch, not two
 		{
 			SizeOut = Alloc.PoolIndexToBinSize(BoundSizeToPoolIndex(Count, Alloc.MemSizeToPoolIndex));
 			check(SizeOut >= Count);
 			return SizeOut;
 		}
-		Alignment = FMath::Max<uint32>(Alignment, UE_MBC_MIN_SMALL_POOL_ALIGNMENT);
+		Alignment = FMath::Max<uint32>(Alignment, UE_MBC_STANDARD_ALIGNMENT);
 		Count = Align(Count, Alignment);
 		if ((Count <= MAX_SMALL_POOL_SIZE) & (Alignment <= UE_MBC_MAX_SMALL_POOL_ALIGNMENT))
 		{
@@ -875,9 +881,9 @@ protected:
 	FORCENOINLINE bool PromoteToLargerBin(SIZE_T& Size, uint32& Alignment, const AllocType& Alloc) const
 	{
 		// try to promote our allocation request to a larger bin with a matching natural alignment
-		// if requested alignment is larger than UE_MBC_MIN_SMALL_POOL_ALIGNMENT but smaller than UE_MBC_MAX_SMALL_POOL_ALIGNMENT
+		// if requested alignment is larger than UE_MBC_STANDARD_ALIGNMENT but smaller than UE_MBC_MAX_SMALL_POOL_ALIGNMENT
 		// so we don't do a page allocation with a lot of memory waste
-		Alignment = FMath::Max<uint32>(Alignment, UE_MBC_MIN_SMALL_POOL_ALIGNMENT);
+		Alignment = FMath::Max<uint32>(Alignment, UE_MBC_STANDARD_ALIGNMENT);
 		const SIZE_T AlignedSize = Align(Size, Alignment);
 		if (UNLIKELY((AlignedSize <= MAX_SMALL_POOL_SIZE) && (Alignment <= UE_MBC_MAX_SMALL_POOL_ALIGNMENT)))
 		{
@@ -889,7 +895,7 @@ protected:
 				{
 					// we found a matching pool for our alignment and size requirements, so modify the size request to match
 					Size = SIZE_T(BlockSize);
-					Alignment = UE_MBC_MIN_SMALL_POOL_ALIGNMENT;
+					Alignment = UE_MBC_STANDARD_ALIGNMENT;
 					return true;
 				}
 
diff --git a/Engine/Source/Runtime/Core/Public/Modules/Boilerplate/ModuleBoilerplate.h b/Engine/Source/Runtime/Core/Public/Modules/Boilerplate/ModuleBoilerplate.h
--- a/Engine/Source/Runtime/Core/Public/Modules/Boilerplate/ModuleBoilerplate.h
+++ b/Engine/Source/Runtime/Core/Public/Modules/Boilerplate/ModuleBoilerplate.h
@@ -35,8 +35,8 @@
 #endif
 
 #define UE_DEFINE_FMEMORY_WRAPPERS \
-	void* FMemory_Malloc ( size_t Size, size_t Alignment                 ) { return FMemory::Malloc( Size ? Size : 1, Alignment ); } \
-	void* FMemory_Realloc( void* Original, size_t Size, size_t Alignment ) { return FMemory::Realloc(Original, Size ? Size : 1, Alignment ); } \
+	void* FMemory_Malloc ( size_t Size, size_t Alignment                 ) { return FMemory::Malloc( Size, Alignment ); } \
+	void* FMemory_Realloc( void* Original, size_t Size, size_t Alignment ) { return FMemory::Realloc( Original, Size, Alignment ); } \
 	void  FMemory_Free   ( void *Ptr                                     ) { FMemory::Free( Ptr ); }
 
 // Disable the replacement new/delete when running the Clang static analyzer, due to false positives in 15.0.x:
@@ -44,14 +44,14 @@
 #if !FORCE_ANSI_ALLOCATOR && !defined(__clang_analyzer__)
 	static_assert(__STDCPP_DEFAULT_NEW_ALIGNMENT__ <= 16, "Expecting 16-byte default operator new alignment - alignments > 16 may have bloat");
 	#define REPLACEMENT_OPERATOR_NEW_AND_DELETE \
-		FORCENOINLINE OPERATOR_NEW_MSVC_PRAGMA void* operator new  ( size_t Size                                                    ) OPERATOR_NEW_THROW_SPEC      { return FMemory_MallocInline( Size ? Size : 1, __STDCPP_DEFAULT_NEW_ALIGNMENT__ ); } \
-		FORCENOINLINE OPERATOR_NEW_MSVC_PRAGMA void* operator new[]( size_t Size                                                    ) OPERATOR_NEW_THROW_SPEC      { return FMemory_MallocInline( Size ? Size : 1, __STDCPP_DEFAULT_NEW_ALIGNMENT__ ); } \
-		FORCENOINLINE OPERATOR_NEW_MSVC_PRAGMA void* operator new  ( size_t Size,                             const std::nothrow_t& ) OPERATOR_NEW_NOTHROW_SPEC    { return FMemory_MallocInline( Size ? Size : 1, __STDCPP_DEFAULT_NEW_ALIGNMENT__ ); } \
-		FORCENOINLINE OPERATOR_NEW_MSVC_PRAGMA void* operator new[]( size_t Size,                             const std::nothrow_t& ) OPERATOR_NEW_NOTHROW_SPEC    { return FMemory_MallocInline( Size ? Size : 1, __STDCPP_DEFAULT_NEW_ALIGNMENT__ ); } \
-		FORCENOINLINE OPERATOR_NEW_MSVC_PRAGMA void* operator new  ( size_t Size, std::align_val_t Alignment                        ) OPERATOR_NEW_THROW_SPEC      { return FMemory_MallocInline( Size ? Size : 1, (std::size_t)Alignment ); } \
-		FORCENOINLINE OPERATOR_NEW_MSVC_PRAGMA void* operator new[]( size_t Size, std::align_val_t Alignment                        ) OPERATOR_NEW_THROW_SPEC      { return FMemory_MallocInline( Size ? Size : 1, (std::size_t)Alignment ); } \
-		FORCENOINLINE OPERATOR_NEW_MSVC_PRAGMA void* operator new  ( size_t Size, std::align_val_t Alignment, const std::nothrow_t& ) OPERATOR_NEW_NOTHROW_SPEC    { return FMemory_MallocInline( Size ? Size : 1, (std::size_t)Alignment ); } \
-		FORCENOINLINE OPERATOR_NEW_MSVC_PRAGMA void* operator new[]( size_t Size, std::align_val_t Alignment, const std::nothrow_t& ) OPERATOR_NEW_NOTHROW_SPEC    { return FMemory_MallocInline( Size ? Size : 1, (std::size_t)Alignment ); } \
+		FORCENOINLINE OPERATOR_NEW_MSVC_PRAGMA void* operator new  ( size_t Size                                                    ) OPERATOR_NEW_THROW_SPEC      { return FMemory_MallocInline( Size, Size <= 8 ? 8 : __STDCPP_DEFAULT_NEW_ALIGNMENT__ ); } \
+		FORCENOINLINE OPERATOR_NEW_MSVC_PRAGMA void* operator new[]( size_t Size                                                    ) OPERATOR_NEW_THROW_SPEC      { return FMemory_MallocInline( Size, Size <= 8 ? 8 : __STDCPP_DEFAULT_NEW_ALIGNMENT__ ); } \
+		FORCENOINLINE OPERATOR_NEW_MSVC_PRAGMA void* operator new  ( size_t Size,                             const std::nothrow_t& ) OPERATOR_NEW_NOTHROW_SPEC    { return FMemory_MallocInline( Size, Size <= 8 ? 8 : __STDCPP_DEFAULT_NEW_ALIGNMENT__ ); } \
+		FORCENOINLINE OPERATOR_NEW_MSVC_PRAGMA void* operator new[]( size_t Size,                             const std::nothrow_t& ) OPERATOR_NEW_NOTHROW_SPEC    { return FMemory_MallocInline( Size, Size <= 8 ? 8 : __STDCPP_DEFAULT_NEW_ALIGNMENT__ ); } \
+		FORCENOINLINE OPERATOR_NEW_MSVC_PRAGMA void* operator new  ( size_t Size, std::align_val_t Alignment                        ) OPERATOR_NEW_THROW_SPEC      { return FMemory_MallocInline( Size, (std::size_t)Alignment ); } \
+		FORCENOINLINE OPERATOR_NEW_MSVC_PRAGMA void* operator new[]( size_t Size, std::align_val_t Alignment                        ) OPERATOR_NEW_THROW_SPEC      { return FMemory_MallocInline( Size, (std::size_t)Alignment ); } \
+		FORCENOINLINE OPERATOR_NEW_MSVC_PRAGMA void* operator new  ( size_t Size, std::align_val_t Alignment, const std::nothrow_t& ) OPERATOR_NEW_NOTHROW_SPEC    { return FMemory_MallocInline( Size, (std::size_t)Alignment ); } \
+		FORCENOINLINE OPERATOR_NEW_MSVC_PRAGMA void* operator new[]( size_t Size, std::align_val_t Alignment, const std::nothrow_t& ) OPERATOR_NEW_NOTHROW_SPEC    { return FMemory_MallocInline( Size, (std::size_t)Alignment ); } \
 		FORCENOINLINE void operator delete  ( void* Ptr                                                                             ) OPERATOR_DELETE_THROW_SPEC   { FMemory_FreeInline( Ptr ); } \
 		FORCENOINLINE void operator delete[]( void* Ptr                                                                             ) OPERATOR_DELETE_THROW_SPEC   { FMemory_FreeInline( Ptr ); } \
 		FORCENOINLINE void operator delete  ( void* Ptr,                                                      const std::nothrow_t& ) OPERATOR_DELETE_NOTHROW_SPEC { FMemory_FreeInline( Ptr ); } \
