Reverted change that removed guarantee that new allocates at least 1 byte, since some allocators can technically return null

diff --git a/Engine/Source/Runtime/Core/Public/Modules/Boilerplate/ModuleBoilerplate.h b/Engine/Source/Runtime/Core/Public/Modules/Boilerplate/ModuleBoilerplate.h
--- a/Engine/Source/Runtime/Core/Public/Modules/Boilerplate/ModuleBoilerplate.h
+++ b/Engine/Source/Runtime/Core/Public/Modules/Boilerplate/ModuleBoilerplate.h
@@ -44,10 +44,10 @@
 #if !FORCE_ANSI_ALLOCATOR && !defined(__clang_analyzer__)
 	static_assert(__STDCPP_DEFAULT_NEW_ALIGNMENT__ <= 16, "Expecting 16-byte default operator new alignment - alignments > 16 may have bloat");
 	#define REPLACEMENT_OPERATOR_NEW_AND_DELETE \
-		FORCENOINLINE OPERATOR_NEW_MSVC_PRAGMA void* operator new  ( size_t Size                                                    ) OPERATOR_NEW_THROW_SPEC      { return FMemory_MallocInline( Size, Size <= 8 ? 8 : __STDCPP_DEFAULT_NEW_ALIGNMENT__ ); } \
-		FORCENOINLINE OPERATOR_NEW_MSVC_PRAGMA void* operator new[]( size_t Size                                                    ) OPERATOR_NEW_THROW_SPEC      { return FMemory_MallocInline( Size, Size <= 8 ? 8 : __STDCPP_DEFAULT_NEW_ALIGNMENT__ ); } \
-		FORCENOINLINE OPERATOR_NEW_MSVC_PRAGMA void* operator new  ( size_t Size,                             const std::nothrow_t& ) OPERATOR_NEW_NOTHROW_SPEC    { return FMemory_MallocInline( Size, Size <= 8 ? 8 : __STDCPP_DEFAULT_NEW_ALIGNMENT__ ); } \
-		FORCENOINLINE OPERATOR_NEW_MSVC_PRAGMA void* operator new[]( size_t Size,                             const std::nothrow_t& ) OPERATOR_NEW_NOTHROW_SPEC    { return FMemory_MallocInline( Size, Size <= 8 ? 8 : __STDCPP_DEFAULT_NEW_ALIGNMENT__ ); } \
+		FORCENOINLINE OPERATOR_NEW_MSVC_PRAGMA void* operator new  ( size_t Size                                                    ) OPERATOR_NEW_THROW_SPEC      { return FMemory_MallocInline( Size ? Size : 1, Size <= 8 ? 8 : __STDCPP_DEFAULT_NEW_ALIGNMENT__ ); } \
+		FORCENOINLINE OPERATOR_NEW_MSVC_PRAGMA void* operator new[]( size_t Size                                                    ) OPERATOR_NEW_THROW_SPEC      { return FMemory_MallocInline( Size ? Size : 1, Size <= 8 ? 8 : __STDCPP_DEFAULT_NEW_ALIGNMENT__ ); } \
+		FORCENOINLINE OPERATOR_NEW_MSVC_PRAGMA void* operator new  ( size_t Size,                             const std::nothrow_t& ) OPERATOR_NEW_NOTHROW_SPEC    { return FMemory_MallocInline( Size ? Size : 1, Size <= 8 ? 8 : __STDCPP_DEFAULT_NEW_ALIGNMENT__ ); } \
+		FORCENOINLINE OPERATOR_NEW_MSVC_PRAGMA void* operator new[]( size_t Size,                             const std::nothrow_t& ) OPERATOR_NEW_NOTHROW_SPEC    { return FMemory_MallocInline( Size ? Size : 1, Size <= 8 ? 8 : __STDCPP_DEFAULT_NEW_ALIGNMENT__ ); } \
 		FORCENOINLINE OPERATOR_NEW_MSVC_PRAGMA void* operator new  ( size_t Size, std::align_val_t Alignment                        ) OPERATOR_NEW_THROW_SPEC      { return FMemory_MallocInline( Size, (std::size_t)Alignment ); } \
 		FORCENOINLINE OPERATOR_NEW_MSVC_PRAGMA void* operator new[]( size_t Size, std::align_val_t Alignment                        ) OPERATOR_NEW_THROW_SPEC      { return FMemory_MallocInline( Size, (std::size_t)Alignment ); } \
 		FORCENOINLINE OPERATOR_NEW_MSVC_PRAGMA void* operator new  ( size_t Size, std::align_val_t Alignment, const std::nothrow_t& ) OPERATOR_NEW_NOTHROW_SPEC    { return FMemory_MallocInline( Size, (std::size_t)Alignment ); } \
