Fixed missing 4 operator new overloads and added missing guarantee that new allocates at least 1 byte, since some allocators can technically return null

diff --git a/Engine/Source/Runtime/Core/Public/Modules/Boilerplate/ModuleBoilerplate.h b/Engine/Source/Runtime/Core/Public/Modules/Boilerplate/ModuleBoilerplate.h
--- a/Engine/Source/Runtime/Core/Public/Modules/Boilerplate/ModuleBoilerplate.h
+++ b/Engine/Source/Runtime/Core/Public/Modules/Boilerplate/ModuleBoilerplate.h
@@ -48,10 +48,10 @@
 		FORCENOINLINE OPERATOR_NEW_MSVC_PRAGMA void* operator new[]( size_t Size                                                    ) OPERATOR_NEW_THROW_SPEC      { return FMemory_MallocInline( Size ? Size : 1, Size <= 8 ? 8 : __STDCPP_DEFAULT_NEW_ALIGNMENT__ ); } \
 		FORCENOINLINE OPERATOR_NEW_MSVC_PRAGMA void* operator new  ( size_t Size,                             const std::nothrow_t& ) OPERATOR_NEW_NOTHROW_SPEC    { return FMemory_MallocInline( Size ? Size : 1, Size <= 8 ? 8 : __STDCPP_DEFAULT_NEW_ALIGNMENT__ ); } \
 		FORCENOINLINE OPERATOR_NEW_MSVC_PRAGMA void* operator new[]( size_t Size,                             const std::nothrow_t& ) OPERATOR_NEW_NOTHROW_SPEC    { return FMemory_MallocInline( Size ? Size : 1, Size <= 8 ? 8 : __STDCPP_DEFAULT_NEW_ALIGNMENT__ ); } \
-		FORCENOINLINE OPERATOR_NEW_MSVC_PRAGMA void* operator new  ( size_t Size, std::align_val_t Alignment                        ) OPERATOR_NEW_THROW_SPEC      { return FMemory_MallocInline( Size, (std::size_t)Alignment ); } \
-		FORCENOINLINE OPERATOR_NEW_MSVC_PRAGMA void* operator new[]( size_t Size, std::align_val_t Alignment                        ) OPERATOR_NEW_THROW_SPEC      { return FMemory_MallocInline( Size, (std::size_t)Alignment ); } \
-		FORCENOINLINE OPERATOR_NEW_MSVC_PRAGMA void* operator new  ( size_t Size, std::align_val_t Alignment, const std::nothrow_t& ) OPERATOR_NEW_NOTHROW_SPEC    { return FMemory_MallocInline( Size, (std::size_t)Alignment ); } \
-		FORCENOINLINE OPERATOR_NEW_MSVC_PRAGMA void* operator new[]( size_t Size, std::align_val_t Alignment, const std::nothrow_t& ) OPERATOR_NEW_NOTHROW_SPEC    { return FMemory_MallocInline( Size, (std::size_t)Alignment ); } \
+		FORCENOINLINE OPERATOR_NEW_MSVC_PRAGMA void* operator new  ( size_t Size, std::align_val_t Alignment                        ) OPERATOR_NEW_THROW_SPEC      { return FMemory_MallocInline( Size ? Size : 1, (std::size_t)Alignment ); } \
+		FORCENOINLINE OPERATOR_NEW_MSVC_PRAGMA void* operator new[]( size_t Size, std::align_val_t Alignment                        ) OPERATOR_NEW_THROW_SPEC      { return FMemory_MallocInline( Size ? Size : 1, (std::size_t)Alignment ); } \
+		FORCENOINLINE OPERATOR_NEW_MSVC_PRAGMA void* operator new  ( size_t Size, std::align_val_t Alignment, const std::nothrow_t& ) OPERATOR_NEW_NOTHROW_SPEC    { return FMemory_MallocInline( Size ? Size : 1, (std::size_t)Alignment ); } \
+		FORCENOINLINE OPERATOR_NEW_MSVC_PRAGMA void* operator new[]( size_t Size, std::align_val_t Alignment, const std::nothrow_t& ) OPERATOR_NEW_NOTHROW_SPEC    { return FMemory_MallocInline( Size ? Size : 1, (std::size_t)Alignment ); } \
 		FORCENOINLINE void operator delete  ( void* Ptr                                                                             ) OPERATOR_DELETE_THROW_SPEC   { FMemory_FreeInline( Ptr ); } \
 		FORCENOINLINE void operator delete[]( void* Ptr                                                                             ) OPERATOR_DELETE_THROW_SPEC   { FMemory_FreeInline( Ptr ); } \
 		FORCENOINLINE void operator delete  ( void* Ptr,                                                      const std::nothrow_t& ) OPERATOR_DELETE_NOTHROW_SPEC { FMemory_FreeInline( Ptr ); } \
