Fixed asserts on size check for large allocations in MB3 caused by a removal of size rounding up to alignmnet in a previous change

diff --git a/Engine/Source/Runtime/Core/Private/HAL/MallocBinned2.cpp b/Engine/Source/Runtime/Core/Private/HAL/MallocBinned2.cpp
--- a/Engine/Source/Runtime/Core/Private/HAL/MallocBinned2.cpp
+++ b/Engine/Source/Runtime/Core/Private/HAL/MallocBinned2.cpp
@@ -613,6 +613,7 @@ void* FMallocBinned2::MallocExternalSmall(SIZE_T Size, uint32 Alignment)
 void* FMallocBinned2::MallocExternalLarge(SIZE_T Size, uint32 Alignment)
 {
 	Alignment = FMath::Max<uint32>(Alignment, UE_MBC_STANDARD_ALIGNMENT);
+	Size = Align(FMath::Max((SIZE_T)1, Size), Alignment);
 
 	checkf(FMath::IsPowerOfTwo(Alignment), TEXT("Invalid Malloc alignment: '%u' is not a power of two"), Alignment);
 	checkf(Alignment <= GMB2PageSize, TEXT("Invalid Malloc alignment: '%u' is greater than the page size '%u'"), Alignment, GMB2PageSize);
@@ -741,6 +742,7 @@ void* FMallocBinned2::Realloc(void* Ptr, SIZE_T NewSize, uint32 Alignment)
 	ExternalAllocMutex.Unlock();
 
 	Alignment = FMath::Max<uint32>(Alignment, UE_MBC_STANDARD_ALIGNMENT);
+	NewSize = Align(FMath::Max((SIZE_T)1, NewSize), Alignment);
 
 	checkf(Alignment <= GMB2PageSize, TEXT("Invalid Realloc alignment: '%u' is greater than the page size '%u'"), Alignment, GMB2PageSize);
 	checkf(IsSupportedSize(NewSize), TEXT("Invalid Realloc size: '%" SIZE_T_FMT "'"), NewSize);
diff --git a/Engine/Source/Runtime/Core/Public/HAL/MallocBinnedCommon.h b/Engine/Source/Runtime/Core/Public/HAL/MallocBinnedCommon.h
--- a/Engine/Source/Runtime/Core/Public/HAL/MallocBinnedCommon.h
+++ b/Engine/Source/Runtime/Core/Public/HAL/MallocBinnedCommon.h
@@ -183,7 +183,7 @@ struct FSizeTableEntry
 class FMallocBinnedCommonBase : public FMalloc
 {
 public:
-	// This needs to be small enough to fit inside the smallest allocation handled by MallocBinned2\3, hence the union.
+	// This needs to be small enough to fit inside the smallest allocation handled by MallocBinned2\3
 	struct FBundleNode
 	{
 		uint64 NextNodeInCurrentBundle : 48;
