[VirtualStackAllocator] Disable ASAN poisoning when building with AutoRTFM instrumentation.

Attempting to unpoison the memory could now be done with the new AutoRTFM::OnPreAbort() handler, but to do this efficiently (e.g. not enqueuing a separate handler per small stack allocation) would involve a bunch of engineering effort and complexity.

The false-positive ASAN hits are impacting our ability to investigate important crashes right now, so just disable the poisoning when instrumenting with both ASAN + AutoRTFM.

diff --git a/Engine/Source/Runtime/Core/Private/Memory/VirtualStackAllocator.cpp b/Engine/Source/Runtime/Core/Private/Memory/VirtualStackAllocator.cpp
--- a/Engine/Source/Runtime/Core/Private/Memory/VirtualStackAllocator.cpp
+++ b/Engine/Source/Runtime/Core/Private/Memory/VirtualStackAllocator.cpp
@@ -43,7 +43,7 @@ FVirtualStackAllocator::FVirtualStackAllocator(size_t RequestedStackSize, EVirtu
 		NextAllocationStart = NextUncommittedPage;
 		RecentHighWaterMark = NextUncommittedPage;
 
-#if PLATFORM_HAS_ASAN_INCLUDE
+#if PLATFORM_HAS_ASAN_INCLUDE && !UE_AUTORTFM // Disable poisoning when running with AutoRTFM, as poisoning cannot be safely rolled back.
 		ASAN_POISON_MEMORY_REGION(NextUncommittedPage, TotalReservationSize);
 #endif
 	}
@@ -98,11 +98,8 @@ void* FVirtualStackAllocator::Allocate(size_t Size, size_t Alignment)
 
 		NextAllocationStart = AllocationEnd;
 
-#if PLATFORM_HAS_ASAN_INCLUDE
-		UE_AUTORTFM_OPEN
-		{
-			ASAN_UNPOISON_MEMORY_REGION(AllocationStart, Size);
-		};
+#if PLATFORM_HAS_ASAN_INCLUDE && !UE_AUTORTFM // Disable poisoning when running with AutoRTFM, as poisoning cannot be safely rolled back.
+		ASAN_UNPOISON_MEMORY_REGION(AllocationStart, Size);
 #endif
 	}
  
diff --git a/Engine/Source/Runtime/Core/Public/Memory/VirtualStackAllocator.h b/Engine/Source/Runtime/Core/Public/Memory/VirtualStackAllocator.h
--- a/Engine/Source/Runtime/Core/Public/Memory/VirtualStackAllocator.h
+++ b/Engine/Source/Runtime/Core/Public/Memory/VirtualStackAllocator.h
@@ -110,11 +110,8 @@ private:
 	{
 		check(RestorePointer <= NextAllocationStart);
 
-#if PLATFORM_HAS_ASAN_INCLUDE
-		UE_AUTORTFM_OPEN
-		{
-			ASAN_POISON_MEMORY_REGION(RestorePointer, static_cast<uint8*>(NextAllocationStart) - static_cast<uint8*>(RestorePointer));
-		};
+#if PLATFORM_HAS_ASAN_INCLUDE && !UE_AUTORTFM // Disable poisoning when running with AutoRTFM, as poisoning cannot be safely rolled back.
+		ASAN_POISON_MEMORY_REGION(RestorePointer, static_cast<uint8*>(NextAllocationStart) - static_cast<uint8*>(RestorePointer));
 #endif
 
 		NextAllocationStart = RestorePointer;
