[UBT][CppCompileWarnings] Engine compliance wundef

diff --git a/Engine/Source/Runtime/Core/Public/Android/AndroidJavaEnv.h b/Engine/Source/Runtime/Core/Public/Android/AndroidJavaEnv.h
--- a/Engine/Source/Runtime/Core/Public/Android/AndroidJavaEnv.h
+++ b/Engine/Source/Runtime/Core/Public/Android/AndroidJavaEnv.h
@@ -2,7 +2,7 @@
 
 #pragma once
 
-#if USE_ANDROID_JNI
+#if defined(USE_ANDROID_JNI) && USE_ANDROID_JNI
 
 #include "CoreTypes.h"
 #include "Containers/Array.h"
