[UBT][CppCompileWarnings] Engine compliance wundef

diff --git a/Engine/Source/Runtime/Core/Private/Android/AndroidJavaEnv.cpp b/Engine/Source/Runtime/Core/Private/Android/AndroidJavaEnv.cpp
--- a/Engine/Source/Runtime/Core/Private/Android/AndroidJavaEnv.cpp
+++ b/Engine/Source/Runtime/Core/Private/Android/AndroidJavaEnv.cpp
@@ -1,6 +1,6 @@
 // Copyright Epic Games, Inc. All Rights Reserved.
 
-#if USE_ANDROID_JNI
+#if defined(USE_ANDROID_JNI) && USE_ANDROID_JNI
 
 #include "Android/AndroidJavaEnv.h"
 #include "Android/AndroidJniGameActivity.h"
