Logging: Fixed early logs being duplicated in the first buffered output device

In a typical configuration, early logs could be written to the log file twice because both the backlog and BufferedItems contained a copy of those logs and both were being serialized to the log file when it was created.

#rnx

diff --git a/Engine/Source/Runtime/Core/Private/Misc/OutputDeviceRedirector.cpp b/Engine/Source/Runtime/Core/Private/Misc/OutputDeviceRedirector.cpp
--- a/Engine/Source/Runtime/Core/Private/Misc/OutputDeviceRedirector.cpp
+++ b/Engine/Source/Runtime/Core/Private/Misc/OutputDeviceRedirector.cpp
@@ -229,6 +229,9 @@ struct FOutputDeviceRedirectorState
 	void ConditionalFlushBufferedItems(EOutputDeviceRedirectorFlushOptions Options);
 	void FlushBufferedItems();
 
+	template <typename VisitorType>
+	void FlushBufferedItemsTo(VisitorType&& Visitor);
+
 	uint64 CreateFence();
 	void WaitFence(uint64 Index);
 
@@ -251,6 +254,7 @@ struct FOutputDeviceRedirectorState
 		}
 	}
 
+	FORCENOINLINE void EnableBacklog(bool bEnable);
 	FORCENOINLINE void AddToBacklog(const FLogRecord& Record);
 	FORCENOINLINE void AddToBacklog(const TCHAR* Data, ELogVerbosity::Type Verbosity, const FName& Category, double Time);
 	FORCENOINLINE void SerializeBacklog(TNotNull<FOutputDevice*> OutputDevice);
@@ -537,13 +541,42 @@ void FOutputDeviceRedirectorState::FlushBufferedItems()
 		return;
 	}
 
+	struct FVisitor
+	{
+		FOutputDeviceRedirectorState& State;
+		const uint32 ThreadId = FPlatformTLS::GetCurrentThreadId();
+		const bool bIsPanicThread = State.IsPanicThread(ThreadId);
+
+		inline void operator()(const FOutputDeviceLine& Line)
+		{
+			State.BroadcastTo(ThreadId, State.BufferedOutputDevices, State.BufferedOutputDevicesCanBeUsedOnPanicThread,
+				UE_PROJECTION_MEMBER(FOutputDevice, Serialize),
+				Line.Data, Line.Verbosity, Line.Category, Line.Time);
+		}
+
+		inline void operator()(const FLogRecord& Record)
+		{
+			State.BroadcastTo(ThreadId, State.BufferedOutputDevices, State.BufferedOutputDevicesCanBeUsedOnPanicThread,
+				UE_PROJECTION_MEMBER(FOutputDevice, SerializeRecord), Record);
+		}
+	};
+	FlushBufferedItemsTo(FVisitor{*this});
+}
+
+template <typename VisitorType>
+UE_AUTORTFM_NOAUTORTFM
+void FOutputDeviceRedirectorState::FlushBufferedItemsTo(VisitorType&& Visitor)
+{
+	if (BufferedItems.IsEmpty())
+	{
+		return;
+	}
+
 	TRACE_CPUPROFILER_EVENT_SCOPE(FOutputDeviceRedirector::FlushBufferedItems);
 
-	const uint32 ThreadId = FPlatformTLS::GetCurrentThreadId();
-	bool bIsPanicThread = this->IsPanicThread(ThreadId);
-	BufferedItems.ConsumeAllFifo([this, ThreadId, bIsPanicThread](FOutputDeviceItem&& Item)
+	BufferedItems.ConsumeAllFifo([this, &Visitor](FOutputDeviceItem&& Item)
 	{
-		if (!bIsPanicThread && HasPanicThread())
+		if (!Visitor.bIsPanicThread && HasPanicThread())
 		{
 			// Re-queue items to be consumed by the panic thread.
 			// ProduceItem within ConsumeAllFifo is okay because the items passed to the consumer
@@ -553,21 +586,18 @@ void FOutputDeviceRedirectorState::FlushBufferedItems()
 		}
 		else
 		{
-			Visit([this, ThreadId](auto&& Value)
+			Visit([this, &Visitor](auto&& Value)
 			{
 				using ValueType = std::decay_t<decltype(Value)>;
 				if constexpr (std::is_same_v<ValueType, FOutputDeviceLine>)
 				{
 					const FOutputDeviceLine& Line = Value;
-					BroadcastTo(ThreadId, BufferedOutputDevices, BufferedOutputDevicesCanBeUsedOnPanicThread,
-						UE_PROJECTION_MEMBER(FOutputDevice, Serialize),
-						Line.Data, Line.Verbosity, Line.Category, Line.Time);
+					Visitor(Line);
 				}
 				else if constexpr (std::is_same_v<ValueType, FLogRecord>)
 				{
 					const FLogRecord& Record = Value;
-					BroadcastTo(ThreadId, BufferedOutputDevices, BufferedOutputDevicesCanBeUsedOnPanicThread,
-						UE_PROJECTION_MEMBER(FOutputDevice, SerializeRecord), Record);
+					Visitor(Record);
 				}
 				else if constexpr (std::is_same_v<ValueType, uint64>)
 				{
@@ -618,6 +648,48 @@ void FOutputDeviceRedirectorState::WaitFence(uint64 Index)
 	}
 }
 
+UE_AUTORTFM_NOAUTORTFM
+void FOutputDeviceRedirectorState::EnableBacklog(bool bEnable)
+{
+	FWriteScopeLock ScopeLock(BacklogLock);
+	const bool bWasEnabled = bEnableBacklog;
+	bEnableBacklog = bEnable;
+	if (!bEnable)
+	{
+		BacklogLines.Empty();
+	}
+	else if (!bWasEnabled)
+	{
+		// Move any buffered items into the backlog if there are no buffered output devices
+		// registered yet. This can only occur in program targets. The backlog is a more robust
+		// means of capturing early logging.
+		const uint32 ThreadId = FPlatformTLS::GetCurrentThreadId();
+		if (FOutputDevicesPrimaryScopeLock PrimaryLock(*this);
+			PrimaryLock.IsLocked() && IsPrimaryThread(ThreadId) && BufferedOutputDevices.IsEmpty())
+		{
+			struct FVisitor
+			{
+				FOutputDeviceRedirectorState& State;
+				const uint32 ThreadId;
+				const bool bIsPanicThread = State.IsPanicThread(ThreadId);
+
+				inline void operator()(const FOutputDeviceLine& Line)
+				{
+					State.BacklogLines.Emplace(Line.Data, Line.Category, Line.Verbosity, Line.Time);
+				}
+
+				inline void operator()(const FLogRecord& Record)
+				{
+					TStringBuilder<512> Text;
+					Record.FormatMessageTo(Text);
+					State.BacklogLines.Emplace(*Text, Record.GetCategory(), Record.GetVerbosity(), FPlatformTime::Seconds() - GStartTime);
+				}
+			};
+			FlushBufferedItemsTo(FVisitor{*this, ThreadId});
+		}
+	}
+}
+
 UE_AUTORTFM_NOAUTORTFM
 void FOutputDeviceRedirectorState::AddToBacklog(const FLogRecord& Record)
 {
@@ -684,7 +756,15 @@ void FOutputDeviceRedirector::AddOutputDevice(FOutputDevice* OutputDevice)
 {
 	if (OutputDevice)
 	{
+		if (State->bEnableBacklog)
+		{
+			// Flush buffered items before adding the output device to avoid an item being logged
+			// twice due to being in both the backlog and the buffer.
+			State->ConditionalFlushBufferedItems(EOutputDeviceRedirectorFlushOptions::None);
+		}
+
 		State->AddOutputDevice(OutputDevice);
+
 		if (State->bEnableBacklog)
 		{
 			State->SerializeBacklog(OutputDevice);
@@ -723,12 +803,7 @@ void FOutputDeviceRedirector::SerializeBacklog(FOutputDevice* OutputDevice)
 UE_AUTORTFM_NOAUTORTFM // Not currently supported in a closed transaction
 void FOutputDeviceRedirector::EnableBacklog(bool bEnable)
 {
-	FWriteScopeLock ScopeLock(State->BacklogLock);
-	State->bEnableBacklog = bEnable;
-	if (!bEnable)
-	{
-		State->BacklogLines.Empty();
-	}
+	State->EnableBacklog(bEnable);
 }
 
 UE_AUTORTFM_NOAUTORTFM // Not currently supported in a closed transaction
