Added a test for FPlatformManualResetEvent

#rnx

diff --git a/Engine/Source/Runtime/Core/Private/Microsoft/MicrosoftPlatformManualResetEvent.h b/Engine/Source/Runtime/Core/Private/Microsoft/MicrosoftPlatformManualResetEvent.h
--- a/Engine/Source/Runtime/Core/Private/Microsoft/MicrosoftPlatformManualResetEvent.h
+++ b/Engine/Source/Runtime/Core/Private/Microsoft/MicrosoftPlatformManualResetEvent.h
@@ -36,11 +36,33 @@ public:
 		}
 	}
 
+	/**
+	 * Waits for the wait time for Notify() to be called.
+	 *
+	 * Notify() may be called prior to WaitFor(), and this will return immediately in that case.
+	 *
+	 * NOTE: Windows will round the wait time to the nearest system tick and may coalesce timers
+	 *       for power efficiency. These details can cause a timeout to elapse slightly early.
+	 *
+	 * @param WaitTime   Relative time after which waiting is canceled and the thread wakes.
+	 * @return True if Notify() was called before the wait time elapsed, otherwise false.
+	 */
 	inline bool WaitFor(FMonotonicTimeSpan WaitTime)
 	{
 		return !bWait.load(std::memory_order_acquire) || WaitForSlow(WaitTime);
 	}
 
+	/**
+	 * Waits until the wait time for Notify() to be called.
+	 *
+	 * Notify() may be called prior to WaitUntil(), and this will return immediately in that case.
+	 *
+	 * NOTE: Windows will round the wait time to the nearest system tick and may coalesce timers
+	 *       for power efficiency. These details can cause a timeout to elapse slightly early.
+	 *
+	 * @param WaitTime   Absolute time after which waiting is canceled and the thread wakes.
+	 * @return True if Notify() was called before the wait time elapsed, otherwise false.
+	 */
 	inline bool WaitUntil(FMonotonicTimePoint WaitTime)
 	{
 		return !bWait.load(std::memory_order_acquire) || WaitUntilSlow(WaitTime);
