Fixed infinite waits in FPlatformManualResetEvent

#rnx

diff --git a/Engine/Source/Runtime/Core/Tests/HAL/PlatformManualResetEventTest.cpp b/Engine/Source/Runtime/Core/Tests/HAL/PlatformManualResetEventTest.cpp
--- a/Engine/Source/Runtime/Core/Tests/HAL/PlatformManualResetEventTest.cpp
+++ b/Engine/Source/Runtime/Core/Tests/HAL/PlatformManualResetEventTest.cpp
@@ -58,6 +58,7 @@ TEST_CASE("Core::Async::PlatformManualResetEvent", "[Core][Async]")
 	CHECK_FALSE(Event.WaitUntil(Timer.Now + FMonotonicTimeSpan::FromMilliseconds(32.0)));
 	CHECK(Timer.Step().ToMilliseconds() >= 16.0);
 
+	// Wait
 	Tasks::Launch(UE_SOURCE_LOCATION, [&Event]
 	{
 		FPlatformProcess::SleepNoStats(0.02f);
@@ -66,6 +67,26 @@ TEST_CASE("Core::Async::PlatformManualResetEvent", "[Core][Async]")
 	Event.Wait();
 	CHECK(Timer.Step().ToMilliseconds() >= 16.0);
 
+	// WaitFor(Infinity)
+	Event.Reset();
+	Tasks::Launch(UE_SOURCE_LOCATION, [&Event]
+	{
+		FPlatformProcess::SleepNoStats(0.02f);
+		Event.Notify();
+	});
+	CHECK(Event.WaitFor(FMonotonicTimeSpan::Infinity()));
+	CHECK(Timer.Step().ToMilliseconds() >= 16.0);
+
+	// WaitUntil(Infinity)
+	Event.Reset();
+	Tasks::Launch(UE_SOURCE_LOCATION, [&Event]
+	{
+		FPlatformProcess::SleepNoStats(0.02f);
+		Event.Notify();
+	});
+	CHECK(Event.WaitUntil(FMonotonicTimePoint::Infinity()));
+	CHECK(Timer.Step().ToMilliseconds() >= 16.0);
+
 	Event.Notify();
 	Event.Reset();
 }
