Fixed negative infinite waits in FPlatformManualResetEvent

WaitFor(-INFINITY) would have waited forever in the affected implementations, and will now correctly return immediately.

#rnx

diff --git a/Engine/Source/Runtime/Core/Private/Android/AndroidPlatformManualResetEvent.cpp b/Engine/Source/Runtime/Core/Private/Android/AndroidPlatformManualResetEvent.cpp
--- a/Engine/Source/Runtime/Core/Private/Android/AndroidPlatformManualResetEvent.cpp
+++ b/Engine/Source/Runtime/Core/Private/Android/AndroidPlatformManualResetEvent.cpp
@@ -3,8 +3,8 @@
 #include "Android/AndroidPlatformManualResetEvent.h"
 
 #include "Async/Fundamental/Oversubscription.h"
-#include "HAL/PlatformMath.h"
 
+#include <limits>
 #include <linux/futex.h>
 #include <sys/syscall.h>
 #include <time.h>
@@ -15,7 +15,7 @@ namespace UE::HAL::Private
 
 inline static const struct timespec* SecondsToTimeSpec(double Seconds, struct timespec& OutTime)
 {
-	if (FPlatformMath::IsFinite(Seconds))
+	if (Seconds != std::numeric_limits<float>::infinity())
 	{
 		const double TotalSec = Seconds >= 0.0 ? Seconds : 0.0;
 		int32 Sec = int32(TotalSec);
diff --git a/Engine/Source/Runtime/Core/Private/Unix/UnixPlatformManualResetEvent.cpp b/Engine/Source/Runtime/Core/Private/Unix/UnixPlatformManualResetEvent.cpp
--- a/Engine/Source/Runtime/Core/Private/Unix/UnixPlatformManualResetEvent.cpp
+++ b/Engine/Source/Runtime/Core/Private/Unix/UnixPlatformManualResetEvent.cpp
@@ -3,9 +3,9 @@
 #include "Unix/UnixPlatformManualResetEvent.h"
 
 #include "Async/Fundamental/Oversubscription.h"
-#include "HAL/PlatformMath.h"
 #include "HAL/PlatformTime.h"
 
+#include <limits>
 #include <linux/futex.h>
 #include <sys/syscall.h>
 #include <time.h>
@@ -16,7 +16,7 @@ namespace UE::HAL::Private
 
 inline static const struct timespec* SecondsToTimeSpec(double Seconds, struct timespec& OutTime)
 {
-	if (FPlatformMath::IsFinite(Seconds))
+	if (Seconds != std::numeric_limits<float>::infinity())
 	{
 		const double TotalSec = Seconds >= 0.0 ? Seconds : 0.0;
 		int32 Sec = int32(TotalSec);
diff --git a/Engine/Source/Runtime/Core/Tests/HAL/PlatformManualResetEventTest.cpp b/Engine/Source/Runtime/Core/Tests/HAL/PlatformManualResetEventTest.cpp
--- a/Engine/Source/Runtime/Core/Tests/HAL/PlatformManualResetEventTest.cpp
+++ b/Engine/Source/Runtime/Core/Tests/HAL/PlatformManualResetEventTest.cpp
@@ -58,6 +58,8 @@ TEST_CASE("Core::Async::PlatformManualResetEvent", "[Core][Async]")
 	CHECK_FALSE(Event.WaitUntil(Timer.Now + FMonotonicTimeSpan::FromMilliseconds(32.0)));
 	CHECK(Timer.Step().ToMilliseconds() >= 16.0);
 
+	CHECK_FALSE(Event.WaitFor(-FMonotonicTimeSpan::Infinity()));
+
 	// Wait
 	Tasks::Launch(UE_SOURCE_LOCATION, [&Event]
 	{
