[LLM] Scope ConfigSystem allocations to proper LLM tag. Add scope to missing function

diff --git a/Engine/Source/Runtime/Core/Private/Misc/ConfigCacheIni.cpp b/Engine/Source/Runtime/Core/Private/Misc/ConfigCacheIni.cpp
--- a/Engine/Source/Runtime/Core/Private/Misc/ConfigCacheIni.cpp
+++ b/Engine/Source/Runtime/Core/Private/Misc/ConfigCacheIni.cpp
@@ -3696,6 +3696,8 @@ bool FConfigBranch::RemoveDynamicLayerFromHierarchy(const FString& Filename, FCo
 
 bool FConfigBranch::RemoveDynamicLayersFromHierarchy(const TArray<FString>& Filenames, FConfigModificationTracker* ModificationTracker)
 {
+	LLM_SCOPE(ELLMTag::ConfigSystem);
+
 	if (bIsSafeUnloaded)
 	{
 		SafeReload();
