Make sure boot hotfixable saved values can be applied for CVars defined in dlls.

diff --git a/Engine/Source/Runtime/Core/Private/HAL/ConsoleManager.cpp b/Engine/Source/Runtime/Core/Private/HAL/ConsoleManager.cpp
--- a/Engine/Source/Runtime/Core/Private/HAL/ConsoleManager.cpp
+++ b/Engine/Source/Runtime/Core/Private/HAL/ConsoleManager.cpp
@@ -3123,6 +3123,13 @@ IConsoleObject* FConsoleManager::AddConsoleObject(const TCHAR* Name, IConsoleObj
 		}
 	}
 
+#if UE_EDITOR
+	if (Obj->GetFlags() & ECVF_SaveForNextBoot)
+	{
+		UE::ConfigUtilities::ApplyCVarOfDllFromBootHotfix(Name, Obj);
+	}
+#endif
+
 	if(ExistingObj)
 	{
 		// An existing console object was found that has the same name as the object being registered.
diff --git a/Engine/Source/Runtime/Core/Private/Misc/ConfigUtilities.cpp b/Engine/Source/Runtime/Core/Private/Misc/ConfigUtilities.cpp
--- a/Engine/Source/Runtime/Core/Private/Misc/ConfigUtilities.cpp
+++ b/Engine/Source/Runtime/Core/Private/Misc/ConfigUtilities.cpp
@@ -263,23 +263,50 @@ void SaveCVarForNextBoot(const TCHAR* Key, const TCHAR* Value)
 #endif // !UE_SERVER
 }
 
+#if UE_EDITOR
+TMap<FString, FString> CVarsToApplyForDlls;
+#endif
+
 void ApplyCVarsFromBootHotfix()
 {
 #if !UE_SERVER
 	TMap<FString, FString> CVarsToApply;
 	LoadCVarsFromFileForNextBoot(CVarsToApply);
 
-	for (const TPair<FString, FString>& CVarPair : CVarsToApply)
+	for (auto It = CVarsToApply.CreateIterator(); It; ++It)
 	{
-		IConsoleVariable* CVar = IConsoleManager::Get().FindConsoleVariable(*CVarPair.Key); 
+		IConsoleVariable* CVar = IConsoleManager::Get().FindConsoleVariable(*It->Key); 
 		if (CVar && CVar->TestFlags(ECVF_SaveForNextBoot))
 		{
-			CVar->Set(*CVarPair.Value, ECVF_SetByHotfix);
+			CVar->Set(*It->Value, ECVF_SetByHotfix);
+			It.RemoveCurrent();
 		}
 	}
+
+#if UE_EDITOR
+	if (!CVarsToApply.IsEmpty())
+	{
+		CVarsToApplyForDlls = MoveTemp(CVarsToApply);
+	}
+#endif // UE_EDITOR
 #endif // !UE_SERVER
 }
 
+#if UE_EDITOR
+void ApplyCVarOfDllFromBootHotfix(const TCHAR* Key, IConsoleObject* Obj)
+{
+	FString* Value = CVarsToApplyForDlls.Find(Key);
+	if (Value)
+	{
+		if (IConsoleVariable* CVar = Obj->AsVariable())
+		{
+			CVar->Set(*(*Value), ECVF_SetByHotfix);
+		}
+		CVarsToApplyForDlls.Remove(Key);
+	}
+}
+#endif // UE_EDITOR
+
 void OnSetCVarFromIniEntry(const TCHAR *IniFile, const TCHAR *Key, const TCHAR* Value, uint32 SetBy, bool bAllowCheating, bool bNoLogging, FName Tag)
 {
 	check(IniFile && Key && Value);
diff --git a/Engine/Source/Runtime/Core/Public/Misc/ConfigUtilities.h b/Engine/Source/Runtime/Core/Public/Misc/ConfigUtilities.h
--- a/Engine/Source/Runtime/Core/Public/Misc/ConfigUtilities.h
+++ b/Engine/Source/Runtime/Core/Public/Misc/ConfigUtilities.h
@@ -57,6 +57,10 @@ namespace UE::ConfigUtilities
 
 	CORE_API void SaveCVarForNextBoot(const TCHAR* Key, const TCHAR* Value);
 
+#if UE_EDITOR
+	CORE_API void ApplyCVarOfDllFromBootHotfix(const TCHAR* Key, IConsoleObject* Obj);
+#endif
+
 	/**
 	 * Helper function to operate a user defined function for each CVar key/value pair in the specified section in an ini file
 	 * @param InSectionName - The name of the section to apply cvars from
