Fixed infinity constant used in FPlatformManualResetEvent

#rnx

diff --git a/Engine/Source/Runtime/Core/Private/Android/AndroidPlatformManualResetEvent.cpp b/Engine/Source/Runtime/Core/Private/Android/AndroidPlatformManualResetEvent.cpp
--- a/Engine/Source/Runtime/Core/Private/Android/AndroidPlatformManualResetEvent.cpp
+++ b/Engine/Source/Runtime/Core/Private/Android/AndroidPlatformManualResetEvent.cpp
@@ -15,7 +15,7 @@ namespace UE::HAL::Private
 
 inline static const struct timespec* SecondsToTimeSpec(double Seconds, struct timespec& OutTime)
 {
-	if (Seconds != std::numeric_limits<float>::infinity())
+	if (Seconds != std::numeric_limits<double>::infinity())
 	{
 		const double TotalSec = Seconds >= 0.0 ? Seconds : 0.0;
 		int32 Sec = int32(TotalSec);
diff --git a/Engine/Source/Runtime/Core/Private/Unix/UnixPlatformManualResetEvent.cpp b/Engine/Source/Runtime/Core/Private/Unix/UnixPlatformManualResetEvent.cpp
--- a/Engine/Source/Runtime/Core/Private/Unix/UnixPlatformManualResetEvent.cpp
+++ b/Engine/Source/Runtime/Core/Private/Unix/UnixPlatformManualResetEvent.cpp
@@ -16,7 +16,7 @@ namespace UE::HAL::Private
 
 inline static const struct timespec* SecondsToTimeSpec(double Seconds, struct timespec& OutTime)
 {
-	if (Seconds != std::numeric_limits<float>::infinity())
+	if (Seconds != std::numeric_limits<double>::infinity())
 	{
 		const double TotalSec = Seconds >= 0.0 ? Seconds : 0.0;
 		int32 Sec = int32(TotalSec);
