Fixed a narrowing cast in the timespec conversion in FPlatformManualResetEvent

#rnx

diff --git a/Engine/Source/Runtime/Core/Private/Android/AndroidPlatformManualResetEvent.cpp b/Engine/Source/Runtime/Core/Private/Android/AndroidPlatformManualResetEvent.cpp
--- a/Engine/Source/Runtime/Core/Private/Android/AndroidPlatformManualResetEvent.cpp
+++ b/Engine/Source/Runtime/Core/Private/Android/AndroidPlatformManualResetEvent.cpp
@@ -18,10 +18,8 @@ inline static const struct timespec* SecondsToTimeSpec(double Seconds, struct ti
 	if (Seconds != std::numeric_limits<double>::infinity())
 	{
 		const double TotalSec = Seconds >= 0.0 ? Seconds : 0.0;
-		int32 Sec = int32(TotalSec);
-		int32 NSec = int32((TotalSec - Sec) * 1e9);
-		OutTime.tv_sec = Sec;
-		OutTime.tv_nsec = NSec;
+		OutTime.tv_sec = time_t(TotalSec);
+		OutTime.tv_nsec = long((TotalSec - double(OutTime.tv_sec)) * 1e9);
 		return &OutTime;
 	}
 	return nullptr;
diff --git a/Engine/Source/Runtime/Core/Private/Unix/UnixPlatformManualResetEvent.cpp b/Engine/Source/Runtime/Core/Private/Unix/UnixPlatformManualResetEvent.cpp
--- a/Engine/Source/Runtime/Core/Private/Unix/UnixPlatformManualResetEvent.cpp
+++ b/Engine/Source/Runtime/Core/Private/Unix/UnixPlatformManualResetEvent.cpp
@@ -19,10 +19,8 @@ inline static const struct timespec* SecondsToTimeSpec(double Seconds, struct ti
 	if (Seconds != std::numeric_limits<double>::infinity())
 	{
 		const double TotalSec = Seconds >= 0.0 ? Seconds : 0.0;
-		int32 Sec = int32(TotalSec);
-		int32 NSec = int32((TotalSec - Sec) * 1e9);
-		OutTime.tv_sec = Sec;
-		OutTime.tv_nsec = NSec;
+		OutTime.tv_sec = time_t(TotalSec);
+		OutTime.tv_nsec = long((TotalSec - double(OutTime.tv_sec)) * 1e9);
 		return &OutTime;
 	}
 	return nullptr;
