Allow FIoStatus to work directly with structured logging.

#rnx

- When used with standard logging this will output the results of FIoStatus::ToString
- When used with structured logging we will output each component of the status as their own entry:
-- ErrorCode
-- SystemErrorCode
-- ErrorMessage
- FIoStatus will now work directly with UE_LOGFMT

diff --git a/Engine/Source/Runtime/Core/Private/IO/IoStatus.cpp b/Engine/Source/Runtime/Core/Private/IO/IoStatus.cpp
--- a/Engine/Source/Runtime/Core/Private/IO/IoStatus.cpp
+++ b/Engine/Source/Runtime/Core/Private/IO/IoStatus.cpp
@@ -104,6 +104,21 @@ FString FIoStatus::ToString() const
 	return FString::Format(TEXT("{0} ({1}, SystemErrorCode={2})"), { *ErrorMessage, GetIoErrorText(ErrorCode), SystemErrorCode });
 }
 
+void SerializeForLog(FCbWriter& Writer, const FIoStatus& Status)
+{
+	Writer.BeginObject();
+
+	Writer.AddString(ANSITEXTVIEW("$type"), ANSITEXTVIEW("IoStatus"));
+	// This probably could be optimized to avoid the FString creation
+	Writer.AddString(ANSITEXTVIEW("$text"), Status.ToString());
+	
+	Writer.AddInteger(ANSITEXTVIEW("ErrorCode"), static_cast<uint32>(Status.ErrorCode));
+	Writer.AddInteger(ANSITEXTVIEW("SystemErrorCode"), Status.SystemErrorCode);
+	Writer.AddString(ANSITEXTVIEW("ErrorMessage"), *Status.ErrorMessage);
+
+	Writer.EndObject();
+}
+
 void StatusOrCrash(const FIoStatus& Status)
 {
 	UE_LOG(LogIoStatus, Fatal, TEXT("I/O Error '%s'"), *Status.ToString());
diff --git a/Engine/Source/Runtime/Core/Public/IO/IoStatus.h b/Engine/Source/Runtime/Core/Public/IO/IoStatus.h
--- a/Engine/Source/Runtime/Core/Public/IO/IoStatus.h
+++ b/Engine/Source/Runtime/Core/Public/IO/IoStatus.h
@@ -79,6 +79,8 @@ public:
 	const TCHAR*		GetErrorMessage() const { return ErrorMessage.IsEmpty() ? nullptr : *ErrorMessage; }
 	CORE_API FString	ToString() const;
 
+	CORE_API friend void SerializeForLog(FCbWriter& Writer, const FIoStatus& Value);
+
 	CORE_API static const FIoStatus Ok;
 	CORE_API static const FIoStatus Unknown;
 	CORE_API static const FIoStatus Invalid;
