Fix possible compiler error depending on include order of Optional.h and TypeHash.h causing failure to instantiate GetTypeHash(TOptional<int>)

diff --git a/Engine/Source/Runtime/Core/Public/Misc/Optional.h b/Engine/Source/Runtime/Core/Public/Misc/Optional.h
--- a/Engine/Source/Runtime/Core/Public/Misc/Optional.h
+++ b/Engine/Source/Runtime/Core/Public/Misc/Optional.h
@@ -7,6 +7,7 @@
 #include "Misc/IntrusiveUnsetOptionalState.h"
 #include "Misc/OptionalFwd.h"
 #include "Templates/MemoryOps.h"
+#include "Templates/TypeHash.h" // Definitions of GetTypeHash for primitive types must be visible when GetTypeHash for TOptional is defined
 #include "Templates/UnrealTemplate.h"
 #include "Serialization/Archive.h"
 #include <type_traits>
