Changed Fatal log to a Warning when posix_spawn fails in Mac CreateProcInternal. This is to be consistent with Windows, and prevent unnecessary crashing

diff --git a/Engine/Source/Runtime/Core/Private/Mac/MacPlatformProcess.cpp b/Engine/Source/Runtime/Core/Private/Mac/MacPlatformProcess.cpp
--- a/Engine/Source/Runtime/Core/Private/Mac/MacPlatformProcess.cpp
+++ b/Engine/Source/Runtime/Core/Private/Mac/MacPlatformProcess.cpp
@@ -587,8 +587,8 @@ FProcHandle FMacPlatformProcess::CreateProcInternal(const TCHAR* URL, const TCHA
 
 	if (PosixSpawnErrNo != 0)
 	{
-		UE_LOG(LogHAL, Fatal, TEXT("FMacPlatformProcess::CreateProc: posix_spawn(\"%s\") failed (%d, %s)"), *ProcessPath, PosixSpawnErrNo, UTF8_TO_TCHAR(strerror(PosixSpawnErrNo)));
-		return FProcHandle();	// produce knowingly invalid handle if for some reason Fatal log (above) returns
+		UE_LOG(LogHAL, Warning, TEXT("FMacPlatformProcess::CreateProc: posix_spawn(\"%s\") failed (%d, %s)"), *ProcessPath, PosixSpawnErrNo, UTF8_TO_TCHAR(strerror(PosixSpawnErrNo)));
+		return FProcHandle();
 	}
 
 	if (PriorityModifier != 0)
