[Backout] - CL45785126
#rnx
Original CL Desc
-----------------------------------------------------------------
[Core] Add support to compile switch between using sparse or compact sets as the default set container

#tests Using the editor, running PIE and game client, making builds

diff --git a/Engine/Plugins/Experimental/MeshModelingToolsetExp/Source/MeshModelingToolsExp/Private/ExtractSplineTool.cpp b/Engine/Plugins/Experimental/MeshModelingToolsetExp/Source/MeshModelingToolsExp/Private/ExtractSplineTool.cpp
--- a/Engine/Plugins/Experimental/MeshModelingToolsetExp/Source/MeshModelingToolsExp/Private/ExtractSplineTool.cpp
+++ b/Engine/Plugins/Experimental/MeshModelingToolsetExp/Source/MeshModelingToolsExp/Private/ExtractSplineTool.cpp
@@ -26,7 +26,6 @@
 
 #include "CuttingOps/PlaneCutOp.h"
 
-#include "Containers/SparseSet.h"
 #include "Misc/MessageDialog.h"
 
 #include "ModelingToolTargetUtil.h"
@@ -468,7 +467,7 @@ void UExtractSplineTool::GatherSplineDataFromPolygroupSelection()
 	// Basic algorithim:
 	// Starting with an edge from the set, sequentially pull edges out of the set which connect to the last selected edge. 
 	// Once an ordered list of boundary edges is generated, generate a list of spline points from these polygroup edges which connect into a continous polyline.
-	auto ProcessSelectedEdges = [this](TSparseSet<int32>& ToProcessEdges)
+	auto ProcessSelectedEdges = [this](TSet<int32>& ToProcessEdges)
 	{
 		if (ToProcessEdges.IsEmpty())
 		{
@@ -581,9 +580,8 @@ void UExtractSplineTool::GatherSplineDataFromPolygroupSelection()
 		return;
 	}
 
-	// Use sparse set so we can remove multiple and compact once
-	TSparseSet<int32> MultipleCountedEdges;
-	TSparseSet<int32> ToProcessEdges;
+	TSet<int32> MultipleCountedEdges;
+	TSet<int32> ToProcessEdges;
 	if (!Selection.SelectedGroupIDs.IsEmpty())
 	{
 		// Use this to develop equivalence relations between groups.
@@ -640,12 +638,7 @@ void UExtractSplineTool::GatherSplineDataFromPolygroupSelection()
 	}
 	else
 	{
-		// Manual copy since SelectedEdgeIDs may be a different set type based on compile switches
-		ToProcessEdges.Empty(Selection.SelectedEdgeIDs.Num());
-		for (const int32 Id : Selection.SelectedEdgeIDs)
-		{
-			ToProcessEdges.Add(Id);
-		}
+		ToProcessEdges = Selection.SelectedEdgeIDs;		
 
 		ProcessSelectedEdges(ToProcessEdges);
 	}
diff --git a/Engine/Plugins/Experimental/PlainProps/Source/Public/PlainPropsUeCoreBindings.h b/Engine/Plugins/Experimental/PlainProps/Source/Public/PlainPropsUeCoreBindings.h
--- a/Engine/Plugins/Experimental/PlainProps/Source/Public/PlainPropsUeCoreBindings.h
+++ b/Engine/Plugins/Experimental/PlainProps/Source/Public/PlainPropsUeCoreBindings.h
@@ -523,26 +523,19 @@ struct TSetBinding : IItemRangeBinding
 		}
 	}
 
-#if UE_USE_COMPACT_SET_AS_DEFAULT
 	virtual void ReadItems(FSaveRangeContext& Ctx) const override
 	{
-		const SetType& Set = Ctx.Request.GetRange<SetType>();
-		Ctx.Items.SetAll(Set.GetData(), Set.Num());
-	}
-#else
-	virtual void ReadItems(FSaveRangeContext& Ctx) const override
-	{
-		static_assert(offsetof(TSparseSetElement<T>, Value) == 0);
-		const TSparseArray<TSparseSetElement<T>>& Elems = Ctx.Request.GetRange<TSparseArray<TSparseSetElement<T>>>();
+		static_assert(offsetof(TSetElement<T>, Value) == 0);
+		const TSparseArray<TSetElement<T>>& Elems = Ctx.Request.GetRange<TSparseArray<TSetElement<T>>>();
 
 		if (Elems.IsEmpty())
 		{
-			Ctx.Items.SetAll(nullptr, 0, sizeof(TSparseSetElement<T>));
+			Ctx.Items.SetAll(nullptr, 0, sizeof(TSetElement<T>));
 		}
 		else if (FExistingItemSlice LastRead = Ctx.Items.Slice)
 		{
 			// Continue partial response
-			const TSparseSetElement<T>* NextElem = static_cast<const TSparseSetElement<T>*>(LastRead.Data) + LastRead.Num + /* skip known invalid */ 1;
+			const TSetElement<T>* NextElem = static_cast<const TSetElement<T>*>(LastRead.Data) + LastRead.Num + /* skip known invalid */ 1;
 			Ctx.Items.Slice = GetContiguousSlice(Elems.PointerToIndex(NextElem), Elems);
 		}
 		else if (Elems.IsCompact())
@@ -553,19 +546,18 @@ struct TSetBinding : IItemRangeBinding
 		{
 			// Start partial response
 			Ctx.Items.NumTotal = Elems.Num();
-			Ctx.Items.Stride = sizeof(TSparseSetElement<T>);
+			Ctx.Items.Stride = sizeof(TSetElement<T>);
 			Ctx.Items.Slice = GetContiguousSlice(0, Elems);
 		}
 	}
 
-	static FExistingItemSlice GetContiguousSlice(int32 Idx, const TSparseArray<TSparseSetElement<T>>& Elems)
+	static FExistingItemSlice GetContiguousSlice(int32 Idx, const TSparseArray<TSetElement<T>>& Elems)
 	{
 		int32 Num = 1;
 		for (;!Elems.IsValidIndex(Idx); ++Idx) {}
 		for (; Elems.IsValidIndex(Idx + Num); ++Num) {}
 		return { &Elems[Idx], static_cast<uint64>(Num) };
 	}
-#endif
 };
 
 //////////////////////////////////////////////////////////////////////////
diff --git a/Engine/Plugins/Experimental/PlainPropsUObject/Source/Private/PlainPropsUObjectBindings.cpp b/Engine/Plugins/Experimental/PlainPropsUObject/Source/Private/PlainPropsUObjectBindings.cpp
--- a/Engine/Plugins/Experimental/PlainPropsUObject/Source/Private/PlainPropsUObjectBindings.cpp
+++ b/Engine/Plugins/Experimental/PlainPropsUObject/Source/Private/PlainPropsUObjectBindings.cpp
@@ -1663,10 +1663,31 @@ using TPropertyRangeSerializer = typename TSelectRangeSerializer<Kind>::Type;
 
 ////////////////////////////////////////////////////////////////////////////////////////////////
 
-#if !UE_USE_COMPACT_SET_AS_DEFAULT
 inline FScriptSparseArray& AsSparseArray(FScriptSet& In) { return reinterpret_cast<FScriptSparseArray&>(In); }
 inline FScriptSparseArray& AsSparseArray(FScriptMap& In) { return reinterpret_cast<FScriptSparseArray&>(In); }
 
+template<typename ScriptSet>
+inline bool IsCompact(const ScriptSet& Set)
+{
+	return Set.NumUnchecked() == Set.GetMaxIndex();
+}
+
+// There's no TScriptSparseArray::SetNumUninitialized() (yet), 
+// reserve using Empty() and add items one by one instead
+template<class ScriptType, class LayoutType>
+uint8* SetNumUninitialized(ScriptType& Dst, const LayoutType& Layout, uint64 Num)
+{
+	check(Dst.IsEmpty());
+	Dst.Empty(static_cast<int32>(Num), Layout);
+	for (uint64 Idx = 0; Idx < Num; ++Idx)
+	{
+		Dst.AddUninitialized(Layout);
+	}
+	check(IsCompact(Dst));
+
+	return static_cast<uint8*>(Dst.GetData(0, Layout));
+}
+
 // @pre Elems.Num() > 0
 inline FExistingItemSlice GetContiguousSlice(int32 Idx, const FScriptSparseArray& Elems, const uint8* Data, SIZE_T Stride)
 {
@@ -1708,29 +1729,6 @@ inline void ReadSparseItems(FExistingItems& Dst, const FScriptSparseArray& Src,
 		Dst.Slice = GetContiguousSlice(0, Src, Data, Layout.Size);
 	}
 }
-#endif
-
-template<typename ScriptSet>
-inline bool IsCompact(const ScriptSet& Set)
-{
-	return Set.NumUnchecked() == Set.GetMaxIndex();
-}
-
-// There's no TScriptSparseArray::SetNumUninitialized() (yet), 
-// reserve using Empty() and add items one by one instead
-template<class ScriptType, class LayoutType>
-uint8* SetNumUninitialized(ScriptType& Dst, const LayoutType& Layout, uint64 Num)
-{
-	check(Dst.IsEmpty());
-	Dst.Empty(static_cast<int32>(Num), Layout);
-	for (uint64 Idx = 0; Idx < Num; ++Idx)
-	{
-		Dst.AddUninitialized(Layout);
-	}
-	check(IsCompact(Dst));
-
-	return static_cast<uint8*>(Dst.GetData(0, Layout));
-}
 
 ////////////////////////////////////////////////////////////////////////////////////////////////
 
@@ -1975,11 +1973,7 @@ struct TSetPropertyBinding : IItemRangeBinding, ICustomBinding
 		return Declare({Id, NoId, 0, EMemberPresence::AllowSparse, Members, {RangeSpec, RangeSpec, RangeSpec}});
 	}
 
-#if UE_USE_COMPACT_SET_AS_DEFAULT
-	inline SIZE_T GetStride() const requires (bLeaves)	{ return sizeof(LeafType); }
-#else
-	inline SIZE_T GetStride() const requires (bLeaves)	{ return sizeof(TSparseSetElement<LeafType>); }
-#endif
+	inline SIZE_T GetStride() const requires (bLeaves)	{ return sizeof(TSetElement<LeafType>); }
 	inline SIZE_T GetStride() const						{ return Layout.Size; }
 
 	inline int32 FindIndex(const FScriptSet& Set, const void* Elem) const
@@ -1992,11 +1986,7 @@ struct TSetPropertyBinding : IItemRangeBinding, ICustomBinding
 		if (int32 Idx = FindIndex(Set, Elem); Idx != INDEX_NONE)
 		{
 			DestroyElem(Set, Idx);
-#if UE_USE_COMPACT_SET_AS_DEFAULT
-			Set.RemoveAt(Idx, Layout, MakeHashFn(Inner));
-#else
 			Set.RemoveAt(Idx, Layout);
-#endif
 		}
 	}
 	
@@ -2010,6 +2000,11 @@ struct TSetPropertyBinding : IItemRangeBinding, ICustomBinding
 		return HasItem(Set, Elem);
 	}
 
+	inline void Rehash(FScriptSet& Set) const
+	{
+		Set.Rehash(Layout, MakeHashFn(Inner));
+	}
+
 	inline void DestroyElem(FScriptSet& Set) const requires (bLeaves) {}
 	inline void DestroyElem(FScriptSet& Set, int32 Idx) const
 	{
@@ -2040,12 +2035,7 @@ struct TSetPropertyBinding : IItemRangeBinding, ICustomBinding
 	
 	virtual void ReadItems(FSaveRangeContext& Ctx) const override
 	{
-#if UE_USE_COMPACT_SET_AS_DEFAULT
-		const FScriptSet& Set = Ctx.Request.GetRange<FScriptSet>();
-		Ctx.Items.SetAll(Set.GetData(0, Layout), Set.Num(), Layout.Size);
-#else
 		ReadSparseItems(/* out */ Ctx.Items, Ctx.Request.GetRange<FScriptSparseArray>(), Layout.SparseArrayLayout);
-#endif
 	}
 
 	virtual void MakeItems(FLoadRangeContext& Ctx) const override
@@ -2068,7 +2058,7 @@ struct TSetPropertyBinding : IItemRangeBinding, ICustomBinding
 		else
 		{
 			check(Ctx.Request.IsFinalCall());
-			Set.Rehash(Layout, MakeHashFn(Inner));
+			Rehash(Set);
 		}
 	}
 
@@ -2121,42 +2111,38 @@ struct TSetPropertyBinding : IItemRangeBinding, ICustomBinding
 			}
 		}
 
-		Dst.Rehash(Layout, MakeHashFn(Inner));
+		Rehash(Dst);
 	}
 
 	// Load leaves into non-empty set
 	inline void InsertNonEmpty(FScriptSet& Dst, FRangeLoadView Src) const requires(bLeaves)
 	{	
-		auto HashFn = MakeHashFn(Inner);
-
 		for (LeafType Item : CastAs<LeafType>(Src.AsLeaves()))
 		{
 			if (!HasItem(Dst, &Item))
 			{
 				 void* Elem = Dst.GetData(Dst.AddUninitialized(Layout), Layout);
-				 Dst.CommitLastUninitialized(Layout, HashFn);
 				 *static_cast<LeafType*>(Elem) = Item;
 			}
 		}
 
-		Dst.CommitAllUninitialized(Layout, HashFn);
+		Rehash(Dst);
 	}
 
-	inline void* AddItem(FScriptSet& Dst, int32& OutIdx) const
+	inline void* AddItem(FScriptSparseArray& Dst, int32& OutIdx) const
 	{
-		OutIdx = Dst.AddUninitialized(Layout);
-		void* Out = Dst.GetData(OutIdx, Layout);
+		OutIdx = Dst.AddUninitialized(Layout.SparseArrayLayout);
+		void* Out = Dst.GetData(OutIdx, Layout.SparseArrayLayout);
 		Inner.InitItem(Out);
 		return Out;
 	}
 
 	// Load structs or ranges into non-empty set
-	inline void InsertNonEmpty(FScriptSet& Dst, FRangeLoadView Src) const
+	inline void InsertNonEmpty(FScriptSet& DstSet, FRangeLoadView Src) const
 	{
-		auto HashFn = MakeHashFn(Inner);
-
 		// Written to avoid FProperty::CopyCompleteValue_InContainer dependency
 		// Items are loaded directly into sparse array and then removed if a duplicate existed
+		FScriptSparseArray& Dst = AsSparseArray(DstSet);
 		const int32 OldNum = Dst.NumUnchecked();
 		int32 TmpIdx;
 		void* Tmp = AddItem(Dst, /* out */ TmpIdx);
@@ -2166,11 +2152,7 @@ struct TSetPropertyBinding : IItemRangeBinding, ICustomBinding
 			for (FRangeLoadView Item : Src.AsRanges())
 			{
 				LoadRange(Tmp, Item, InnerBindings);
-				if (!HasItem(Dst, Tmp))
-				{
-					Dst.CommitLastUninitialized(Layout, HashFn);
-					Tmp = AddItem(Dst, /* out */ TmpIdx);
-				}
+				Tmp = HasItem(DstSet, Tmp) ? Tmp : AddItem(Dst, /* out */ TmpIdx);
 			}
 		}
 		else
@@ -2178,20 +2160,16 @@ struct TSetPropertyBinding : IItemRangeBinding, ICustomBinding
 			for (FStructLoadView Item : Src.AsStructs())
 			{
 				LoadStruct(Tmp, Item);
-				if (!HasItem(Dst, Tmp))
-				{
-					Dst.CommitLastUninitialized(Layout, HashFn);
-					Tmp = AddItem(Dst, /* out */ TmpIdx);
-				}
+				Tmp = HasItem(DstSet, Tmp) ? Tmp : AddItem(Dst, /* out */ TmpIdx);
 			}
 		}
 
 		Inner.DestroyItem(Tmp);
-		Dst.RemoveAtUninitialized(Layout, TmpIdx);
+		Dst.RemoveAtUninitialized(Layout.SparseArrayLayout, TmpIdx, 1);
 
 		if (Dst.NumUnchecked() != OldNum)
 		{
-			Dst.CommitAllUninitialized(Layout, HashFn);
+			Rehash(DstSet);
 		}
 	}
 
@@ -2235,7 +2213,7 @@ struct TSetPropertyBinding : IItemRangeBinding, ICustomBinding
 	{
 		checkSlow(Idx < Dst.GetMaxIndex());
 		checkSlow(!Dst.IsValidIndex(Idx));
-		void* Elem = Dst.GetData(Idx, Layout);
+		void* Elem = AsSparseArray(Dst).GetData(Idx, Layout.SparseArrayLayout);
 		Inner.DestroyItem(Elem);
 	}
 };
@@ -2357,12 +2335,7 @@ struct TMapPropertyItemBinding : IItemRangeBinding
 
 	virtual void ReadItems(FSaveRangeContext& Ctx) const override
 	{
-#if UE_USE_COMPACT_SET_AS_DEFAULT
-		const FScriptMap& Set = Ctx.Request.GetRange<FScriptMap>();
-		Ctx.Items.SetAll(Set.GetData(0, Layout), Set.Num(), Layout.SetLayout.Size);
-#else
 		ReadSparseItems(/* out */ Ctx.Items, Ctx.Request.GetRange<FScriptSparseArray>(), Layout.SetLayout.SparseArrayLayout);
-#endif
 	}
 
 	virtual void MakeItems(FLoadRangeContext& Ctx) const override
@@ -2458,12 +2431,12 @@ struct TMapPropertyCustomBinding : TMapPropertyItemBinding<FScriptMap, KeyKind,
 		return false;
 	}
 
-	inline uint8* AddPair(FScriptMap& Dst, int32& OutIdx) const
+	inline uint8* AddPair(FScriptSparseArray& Dst, int32& OutIdx, int32 ValueOffset) const
 	{
-		OutIdx = Dst.AddUninitialized(Layout);
-		uint8* Out = static_cast<uint8*>(Dst.GetData(OutIdx, Layout));
+		OutIdx = Dst.AddUninitialized(Layout.SetLayout.SparseArrayLayout);
+		uint8* Out = static_cast<uint8*>(Dst.GetData(OutIdx, Layout.SetLayout.SparseArrayLayout));
 		InnerKey.InitItem(Out);
-		InnerValue.InitItem(Out + Layout.ValueOffset);
+		InnerValue.InitItem(Out + ValueOffset);
 		return Out;
 	}
 
@@ -2482,11 +2455,7 @@ struct TMapPropertyCustomBinding : TMapPropertyItemBinding<FScriptMap, KeyKind,
 		if (int32 Idx = FindKey(Map, Key); Idx != INDEX_NONE)
 		{
 			DestroyPair(Map, Idx);
-#if UE_USE_COMPACT_SET_AS_DEFAULT
-			Map.RemoveAt(Idx, Layout, MakeHashFn(InnerKey));
-#else
 			Map.RemoveAt(Idx, Layout);
-#endif
 		}
 	}
 
@@ -2530,15 +2499,14 @@ struct TMapPropertyCustomBinding : TMapPropertyItemBinding<FScriptMap, KeyKind,
 	// Load structs or ranges into non-empty map
 	inline void InsertNonEmpty(FScriptMap& Dst, FRangeLoadView Src) const
 	{
-		auto HashFn = MakeHashFn(InnerKey);
-
 		// Written to avoid FProperty::CopyCompleteValue_InContainer dependency
 		// Items are loaded directly into sparse array and then removed if a duplicate existed
 
-		const int32 OldNum = Dst.NumUnchecked();
+		FScriptSparseArray& DstArray = AsSparseArray(Dst);
+		const int32 OldNum = DstArray.NumUnchecked();
 		const int32 ValueOffset = Layout.ValueOffset;
 		int32 TmpIdx;
-		uint8* Tmp = AddPair(Dst, /* out*/ TmpIdx);
+		uint8* Tmp = AddPair(DstArray, /* out*/ TmpIdx, ValueOffset);
 
 		FOptionalSchemaId InnerLoadIds[2];
 		FStructRangeLoadView Structs = Src.AsStructs();
@@ -2552,25 +2520,24 @@ struct TMapPropertyCustomBinding : TMapPropertyItemBinding<FScriptMap, KeyKind,
 			if (int32 Idx = FindKey(Dst, Tmp); Idx != INDEX_NONE)
 			{
 				// Load value into existing pair
-				uint8* Pair = static_cast<uint8*>(Dst.GetData(Idx, Layout));
+				uint8* Pair = static_cast<uint8*>(DstArray.GetData(Idx, Layout.SetLayout.SparseArrayLayout));
 				ValueRange.LoadItem(Pair + ValueOffset,	/* in-out */ Struct.Values, /* in-out */ Bits, InnerLoadIds[1], Struct.Schema.Batch);
 			}
 			else
 			{
 				// Load value into tmp pair and add new temporary
 				ValueRange.LoadItem(Tmp + ValueOffset,	/* in-out */ Struct.Values, /* in-out */ Bits, InnerLoadIds[1], Struct.Schema.Batch);
-				Dst.CommitLastUninitialized(Layout, HashFn);
-				Tmp = AddPair(Dst, /* out*/ TmpIdx);
+				Tmp = AddPair(DstArray, /* out*/ TmpIdx, ValueOffset);
 			}
 		}
 
 		InnerKey.DestroyItem(Tmp);
 		InnerValue.DestroyItem(Tmp + ValueOffset);
-		Dst.RemoveAtUninitialized(Layout, TmpIdx);
+		DstArray.RemoveAtUninitialized(Layout.SetLayout.SparseArrayLayout, TmpIdx, 1);
 
-		if (Dst.NumUnchecked() != OldNum)
+		if (DstArray.NumUnchecked() != OldNum)
 		{
-			Dst.CommitAllUninitialized(Layout, HashFn);
+			Rehash(Dst);
 		}
 	}
 
@@ -3536,7 +3503,7 @@ static void BindStruct(FBindId Id, const UStruct* Struct)
 	BindMembers(/* out */ Binder, Struct);
 
 	GUE.Schemas.BindStruct(Id, Binder.GetMembers(), Binder.Declare());
-	
+		
 	// Don't bind CDO defaults, object defaults are passed in from top and objects aren't owned by containers 
 	if (Struct->HasAnyCastFlags(CASTCLASS_UScriptStruct))
 	{
@@ -4890,18 +4857,18 @@ int32 RoundtripViaBatch(TConstArrayView<UObject*> Objects, ERoundtrip Options)
 	}
 	else
 	{
-	UE_LOGFMT(LogPlainPropsUObject, Display, "Gathering all non-empty UObjects...");
-	static constexpr EObjectFlags SkipFlags = RF_ClassDefaultObject | RF_MirroredGarbage | RF_InheritableComponentTemplate;
-	for (TObjectIterator<UObject> It(SkipFlags); It; ++It)
-	{
-		UObject* Object = *It;
+		UE_LOGFMT(LogPlainPropsUObject, Display, "Gathering all non-empty UObjects...");
+		static constexpr EObjectFlags SkipFlags = RF_ClassDefaultObject | RF_MirroredGarbage | RF_InheritableComponentTemplate;
+		for (TObjectIterator<UObject> It(SkipFlags); It; ++It)
+		{
+			UObject* Object = *It;
 			UClass* Class = Object->GetClass();
 			if (ShouldBind(Class) && IncludeClass(Class))
-		{
-			Instances.Emplace(Object->GetPathName(), Object);
+			{
+				Instances.Emplace(Object->GetPathName(), Object);
+			}
 		}
 	}
-	}
 
 	UE_LOGFMT(LogPlainPropsUObject, Display, "Sorting {Num} UObjects ...", Instances.Num());
 	Algo::Sort(Instances, [](const FInstance& A, const FInstance& B)
@@ -5140,20 +5107,20 @@ int32 RoundtripViaPackages(TConstArrayView<UObject*> Objects, ERoundtrip Options
 		{
 			UClass* Class = Object->GetClass();
 			if (Class->HasAnyClassFlags(CLASS_CompiledFromBlueprint))
-		{
+			{
 				bExclude = true;
 				return false;
-		}
+			}
 			if (!IncludeClass(Class))
 			{
 				UE_LOGFMT(LogPlainPropsUObject, Log, " Excluding PlainProps for {Package}", Object->GetPackage()->GetName());
 				bExclude = true;
 				return false;
-	}
+			}
 			return true;
 		});
 		if (bExclude)
-	{
+		{
 			continue;
 		}
 		Packages.Add(Package);
@@ -5170,7 +5137,7 @@ int32 RoundtripViaPackages(TConstArrayView<UObject*> Objects, ERoundtrip Options
 
 	UE_LOGFMT(LogPlainPropsUObject, Display, "Saving {SaveNum} packages...", Packages.Num());
 	for (UPackage* Package : Packages)
-		{
+	{
 		UObject* Asset = Package->FindAssetInPackage();
 		FSavePackageArgs SaveArgs;
 		SaveArgs.TopLevelFlags = RF_Standalone;
@@ -5255,4 +5222,4 @@ void FPlainPropsUObjectModule::ShutdownModule()
 {
 	BindingScopePtr.Reset();
 	DebugIdScopePtr.Reset();
-}
\ No newline at end of file
+}
diff --git a/Engine/Plugins/Runtime/AudioModulation/Source/AudioModulation/Private/AudioModulationDebugger.cpp b/Engine/Plugins/Runtime/AudioModulation/Source/AudioModulation/Private/AudioModulationDebugger.cpp
--- a/Engine/Plugins/Runtime/AudioModulation/Source/AudioModulation/Private/AudioModulationDebugger.cpp
+++ b/Engine/Plugins/Runtime/AudioModulation/Source/AudioModulation/Private/AudioModulationDebugger.cpp
@@ -565,7 +565,7 @@ namespace AudioModulation
 				// Bus mix is global if it only has one stage and its ID is the Control Bus's ID.
 				if (Proxy->Stages.Num() == 1)
 				{
-					uint32 BusId = Proxy->Stages.FindArbitraryElement()->Key;
+					uint32 BusId = Proxy->Stages.begin().Key();
 					DebugInfo.bIsGlobal = ModulationSystem->IsBusMixGlobal(BusId);
 				}
 			
diff --git a/Engine/Source/Editor/ContentBrowser/Private/SAssetView.cpp b/Engine/Source/Editor/ContentBrowser/Private/SAssetView.cpp
--- a/Engine/Source/Editor/ContentBrowser/Private/SAssetView.cpp
+++ b/Engine/Source/Editor/ContentBrowser/Private/SAssetView.cpp
@@ -625,7 +625,7 @@ private:
 
 bool FAssetViewItemCollection::RefreshLookup()
 {
-	uint32 HashSize = FDefaultSparseSetAllocator::GetNumberOfHashBuckets(Items.Num());
+	uint32 HashSize = TSetAllocator<>::GetNumberOfHashBuckets(Items.Num());
 	if (Lookup.GetHashSize() < HashSize)
 	{
 		Lookup.Clear(HashSize, Items.Max());
@@ -3484,7 +3484,7 @@ void FAssetViewItemCollection::RefreshItemsFromBackend(const FAssetViewContentSo
 	bItemsPendingRemove = false;
 
 	// Create or recyle FAssetViewItem for each FContentBrowserItemData. Build new hashtable concurrently at the same time
-	uint32 HashSize = FDefaultSparseSetAllocator::GetNumberOfHashBuckets(Items.Num());
+	uint32 HashSize = TSetAllocator<>::GetNumberOfHashBuckets(Items.Num());
 	std::atomic<bool> bAnyFolders(false);
 	std::atomic<bool> bAnyRecycled(false);
 	Lookup.Clear(HashSize, Items.Num());
diff --git a/Engine/Source/Editor/UnrealEd/Private/DataTableEditorUtils.cpp b/Engine/Source/Editor/UnrealEd/Private/DataTableEditorUtils.cpp
--- a/Engine/Source/Editor/UnrealEd/Private/DataTableEditorUtils.cpp
+++ b/Engine/Source/Editor/UnrealEd/Private/DataTableEditorUtils.cpp
@@ -465,12 +465,15 @@ bool FDataTableEditorUtils::RemoveRow(UDataTable* DataTable, FName Name)
 		BroadcastPreChange(DataTable, EDataTableChangeInfo::RowList);
 		DataTable->Modify();
 		uint8* RowData = nullptr;
-		const bool bRemoved = DataTable->GetNonConstRowMap().RemoveAndCopyValueStable(Name, RowData);
+		const bool bRemoved = DataTable->GetNonConstRowMap().RemoveAndCopyValue(Name, RowData);
 		if (bRemoved && RowData)
 		{
 			DataTable->RowStruct->DestroyStruct(RowData);
 			FMemory::Free(RowData);
 			bResult = true;
+
+			// Compact the map so that a subsequent add goes at the end of the table
+			DataTable->GetNonConstRowMap().CompactStable();
 		}
 		BroadcastPostChange(DataTable, EDataTableChangeInfo::RowList);
 	}
diff --git a/Engine/Source/Runtime/Core/Private/Misc/ConfigCacheIni.cpp b/Engine/Source/Runtime/Core/Private/Misc/ConfigCacheIni.cpp
--- a/Engine/Source/Runtime/Core/Private/Misc/ConfigCacheIni.cpp
+++ b/Engine/Source/Runtime/Core/Private/Misc/ConfigCacheIni.cpp
@@ -1766,11 +1766,13 @@ void FConfigFile::ProcessCommand(FConfigSection* Section, FStringView SectionNam
 			break;
 		case FConfigValue::EValueType::Remove:
 			// Remove if present.
-			Section->RemoveSingleStable(Key, Value);
+			Section->RemoveSingle(Key, Value);
+			Section->CompactStable();
 			break;
 		case FConfigValue::EValueType::Clear:
 			// Remove if present.
-			Section->RemoveStable(Key);
+			Section->Remove(Key);
+			Section->CompactStable();
 			
 			// clear any empty initialization, which will reset it to using code defaults if nothing else adds anything else
 			Section->EmptyInitializedKeys.Remove(Key);
@@ -1780,7 +1782,8 @@ void FConfigFile::ProcessCommand(FConfigSection* Section, FStringView SectionNam
 			Section->EmptyInitializedKeys.Add(Key);
 			
 			// also clear any existing entries
-			Section->RemoveStable(Key);
+			Section->Remove(Key);
+			Section->CompactStable();
 			break;
 		case FConfigValue::EValueType::ArrayOfStructKey:
 			// track a key to show uniqueness for arrays of structs
diff --git a/Engine/Source/Runtime/Core/Public/Containers/Array.h b/Engine/Source/Runtime/Core/Public/Containers/Array.h
--- a/Engine/Source/Runtime/Core/Public/Containers/Array.h
+++ b/Engine/Source/Runtime/Core/Public/Containers/Array.h
@@ -216,18 +216,6 @@ UE_NODEBUG [[nodiscard]] UE_FORCEINLINE_HINT TIndexedContainerIterator<Container
 		{
 		}
 
-		UE_NODEBUG [[nodiscard]] UE_FORCEINLINE_HINT ElementType* operator->() const
-		{
-			if constexpr (bReverse)
-			{
-				return Ptr - 1;
-			}
-			else
-			{
-				return Ptr;
-			}
-		}
-
 		UE_NODEBUG [[nodiscard]] UE_FORCEINLINE_HINT ElementType& operator*() const
 		{
 			if constexpr (bReverse)
@@ -277,12 +265,6 @@ UE_NODEBUG [[nodiscard]] UE_FORCEINLINE_HINT TIndexedContainerIterator<Container
 			return Ptr != Rhs.Ptr;
 		}
 
-		UE_NODEBUG [[nodiscard]] FORCEINLINE bool operator==(const TCheckedPointerIterator& Rhs) const
-		{
-			return !(*this != Rhs);
-		}
-
-
 	private:
 		ElementType*    Ptr;
 		const SizeType& CurrentNum;
diff --git a/Engine/Source/Runtime/Core/Public/Containers/CompactSet.h b/Engine/Source/Runtime/Core/Public/Containers/CompactSet.h
--- a/Engine/Source/Runtime/Core/Public/Containers/CompactSet.h
+++ b/Engine/Source/Runtime/Core/Public/Containers/CompactSet.h
@@ -111,10 +111,6 @@ private:
 
 		UE_FORCEINLINE_HINT TBaseIterator& operator++()
 		{
-#if DO_CHECK
-			const int32 SetNum = Set.Num();
-			checkf(SetNum >= InitialNum, TEXT("Sets/Maps should never have elements removed during iteration outside of Iterator.RemoveCurrent(). InitialNum %d CurrentNum %d"), InitialNum, SetNum);
-#endif
 			++Index;
 			return *this;
 		}
@@ -127,6 +123,9 @@ private:
 
 		[[nodiscard]] inline bool operator==(const TBaseIterator& Rhs) const
 		{
+#if DO_CHECK
+			checkf(Set.Num() == InitialNum, TEXT("Sets/Maps should never be mutated during range-for iteration. Use CreateIterator form to allow more flexibility. InitialNum %d CurrentNum %d"), InitialNum, Set.Num());
+#endif
 			checkSlow(&Set == &Rhs.Set);
 			return Index == Rhs.Index;
 		}
@@ -186,7 +185,7 @@ private:
 
 			// Note: Adding new elements is safe, however they will be guaranteed to be missed by the current iteration
 #if DO_CHECK
-			checkf(SetNum >= InitialNum, TEXT("Sets/Maps should never have elements removed during iteration outside of Iterator.RemoveCurrent(). InitialNum %d CurrentNum %d"), InitialNum, SetNum);
+			checkf(SetNum >= InitialNum, TEXT("Sets/Maps should never have elements removed during iteration outside of Iterator.RemoveCurrent(). InitialNum %d CurrentNum %d"), InitialNum, Set.Num());
 #endif
 
 			Index = NextIndex;
@@ -301,11 +300,6 @@ public:
 
 	[[nodiscard]] UE_FORCEINLINE_HINT TCompactSet() = default;
 
-	[[nodiscard]] explicit consteval TCompactSet(EConstEval)
-		: Super(ConstEval)
-	{
-	}
-
 	[[nodiscard]] UE_FORCEINLINE_HINT TCompactSet(const TCompactSet& Copy)
 	{
 		*this = Copy;
@@ -559,7 +553,7 @@ public:
 	 */
 	[[nodiscard]] UE_FORCEINLINE_HINT SIZE_T GetAllocatedSize() const
 	{
-		return Super::GetAllocatedSize(GetSetLayout());
+		return this->Elements.GetAllocationSize();
 	}
 
 	/** Tracks the container's memory use through an archive. */
@@ -612,10 +606,10 @@ public:
 	}
 
 	/**
-	 * Checks if index is in array range.
-	 *
-	 * @param Index Index to check.
-	 */
+		* Checks if index is in array range.
+		*
+		* @param Index Index to check.
+		*/
 	inline void RangeCheck(FSetElementId Id) const
 	{
 		CheckInvariants();
@@ -675,7 +669,6 @@ public:
 	 * Adds an element to the set.
 	 *
 	 * @see		Class documentation section on ByHash() functions
-	 * @param	KeyHash						A precomputed hash value, calculated in the same way as ElementType is hashed.
 	 * @param	InElement					Element to add to set
 	 * @param	bIsAlreadyInSetPtr	[out]	Optional pointer to bool that will be set depending on whether element is already in set
      * @return  A handle to the element stored in the set
@@ -692,35 +685,15 @@ public:
 	/**
 	 * Adds an element to the set.
 	 *
-	 * @param	Arg							The argument(s) to be forwarded to the set element's constructor.
+	 * @param	Args						The argument(s) to be forwarded to the set element's constructor.
 	 * @param	bIsAlreadyInSetPtr	[out]	Optional pointer to bool that will be set depending on whether element is already in set
 	 * @return	A handle to the element stored in the set.
 	 */
-	template<typename ArgType = ElementType>
-	UE_FORCEINLINE_HINT FSetElementId Emplace(ArgType && Arg, bool* bIsAlreadyInSetPtr = nullptr)
+	template<typename ArgsType = ElementType>
+	FSetElementId Emplace(ArgsType && Args, bool* bIsAlreadyInSetPtr = nullptr)
 	{
-		TPair<FSetElementId, bool> Result = Emplace(InPlace, Forward<ArgType>(Arg));
-
-		if (bIsAlreadyInSetPtr)
-		{
-			*bIsAlreadyInSetPtr = Result.Value;
-		}
-
-		return Result.Key;
-	}
-
-    /**
-     * Adds an element to the set by constructing the ElementType in-place with multiple args.
-     *
-     * @param   EInPlace   Tag to disambiguate in-place construction.
-     * @param   InArgs     Arguments forwarded to ElementType's constructor.
-     * @return  Pair of (element id, whether an equivalent element already existed).
-     */
-    template <typename... ArgTypes>
-    TPair<FSetElementId, bool> Emplace(EInPlace, ArgTypes&&... InArgs)
-    {
 		alignas(alignof(ElementType))char StackElement[sizeof(ElementType)];
-		ElementType& TempElement = *new (StackElement) ElementType(Forward<ArgTypes>(InArgs)...);
+		ElementType& TempElement = *new (StackElement) ElementType(Forward<ArgsType>(Args));
 		SizeType ExistingIndex(INDEX_NONE);
 
 		const KeyInitType Key = KeyFuncs::GetSetKey(TempElement);
@@ -731,6 +704,11 @@ public:
 			ExistingIndex = FindIndexByHash(KeyHash, Key);
 		}
 
+		if (bIsAlreadyInSetPtr)
+		{
+			*bIsAlreadyInSetPtr = ExistingIndex != INDEX_NONE;
+		}
+
 		if (ExistingIndex == INDEX_NONE)
 		{
 			ExistingIndex = this->NumElements;
@@ -743,45 +721,22 @@ public:
 			MoveByRelocate(ExistingElement, TempElement);
 		}
 
-		return { FSetElementId::FromInteger(ExistingIndex), ExistingIndex != INDEX_NONE };
+		return FSetElementId::FromInteger(ExistingIndex);
 	}
 	
 	/**
 	 * Adds an element to the set.
 	 *
 	 * @see		Class documentation section on ByHash() functions
-	 * @param	KeyHash						A precomputed hash value, calculated in the same way as ElementType is hashed.
-	 * @param	Arg							The argument(s) to be forwarded to the set element's constructor.
+	 * @param	Args						The argument(s) to be forwarded to the set element's constructor.
 	 * @param	bIsAlreadyInSetPtr	[out]	Optional pointer to bool that will be set depending on whether element is already in set
 	 * @return	A handle to the element stored in the set.
 	 */
-	template<typename ArgType = ElementType>
-	UE_FORCEINLINE_HINT FSetElementId EmplaceByHash(uint32 KeyHash, ArgType && Arg, bool* bIsAlreadyInSetPtr = nullptr)
+	template<typename ArgsType = ElementType>
+	FSetElementId EmplaceByHash(uint32 KeyHash, ArgsType && Args, bool* bIsAlreadyInSetPtr = nullptr)
 	{
-		TPair<FSetElementId, bool> Result = EmplaceByHash(InPlace, KeyHash, Forward<ArgType>(Arg));
-
-		if (bIsAlreadyInSetPtr)
-		{
-			*bIsAlreadyInSetPtr = Result.Value;
-		}
-
-		return Result.Key;
-	}
-
-    /**
-     * Adds an element to the set by constructing in-place with multiple args, using a precomputed hash.
-	 *
-     * @see     Class documentation section on ByHash() functions.
-	 * @param	KeyHash    A precomputed hash value, calculated in the same way as ElementType is hashed.
-     * @param   EInPlace   Tag to disambiguate in-place construction.
-     * @param   InArgs     Arguments forwarded to ElementType's constructor.
-     * @return  Pair of (element id, whether an equivalent element already existed).
-     */
-    template <typename... ArgTypes>
-    TPair<FSetElementId, bool> EmplaceByHash(EInPlace, uint32 KeyHash, ArgTypes&&... InArgs)
-    {
 		alignas(alignof(ElementType))char StackElement[sizeof(ElementType)];
-		ElementType& TempElement = *new (StackElement) ElementType(Forward<ArgTypes>(InArgs)...);
+		ElementType& TempElement = *new (StackElement) ElementType(Forward<ArgsType>(Args));
 		SizeType ExistingIndex(INDEX_NONE);
 
 		if constexpr (!KeyFuncs::bAllowDuplicateKeys)
@@ -789,6 +744,11 @@ public:
 			ExistingIndex = FindIndexByHash(KeyHash, KeyFuncs::GetSetKey(TempElement));
 		}
 
+		if (bIsAlreadyInSetPtr)
+		{
+			*bIsAlreadyInSetPtr = ExistingIndex != INDEX_NONE;
+		}
+
 		if (ExistingIndex == INDEX_NONE)
 		{
 			ExistingIndex = this->NumElements;
@@ -801,16 +761,16 @@ public:
 			MoveByRelocate(ExistingElement, TempElement);
 		}
 
-		return { FSetElementId::FromInteger(ExistingIndex), ExistingIndex != INDEX_NONE };
+		return FSetElementId::FromInteger(ExistingIndex);
 	}
 
 	/**
-	 * Adds an element to the set if not already present and returns a reference to the added or existing element.
-	 *
-	 * @param	InElement					Element to add to set
-	 * @param	bIsAlreadyInSetPtr	[out]	Optional pointer to bool that will be set depending on whether element is already in set
-	 * @return	A reference to the element stored in the set.
-	 */
+		* Adds an element to the set if not already present and returns a reference to the added or existing element.
+		*
+		* @param	InElement					Element to add to set
+		* @param	bIsAlreadyInSetPtr	[out]	Optional pointer to bool that will be set depending on whether element is already in set
+		* @return	A reference to the element stored in the set.
+		*/
 	UE_FORCEINLINE_HINT ElementType& FindOrAdd(const InElementType& InElement, bool* bIsAlreadyInSetPtr = nullptr)
 	{
 		return FindOrAddByHash(KeyFuncs::GetKeyHash(KeyFuncs::GetSetKey(InElement)), InElement, bIsAlreadyInSetPtr);
@@ -824,7 +784,6 @@ public:
 	 * Adds an element to the set if not already present and returns a reference to the added or existing element.
 	 *
 	 * @see		Class documentation section on ByHash() functions
-	 * @param	KeyHash						A precomputed hash value, calculated in the same way as ElementType is hashed.
 	 * @param	InElement					Element to add to set
 	 * @param	bIsAlreadyInSetPtr	[out]	Optional pointer to bool that will be set depending on whether element is already in set
 	 * @return  A reference to the element stored in the set
@@ -944,10 +903,10 @@ public:
 	}
 
 	/**
-	 * Finds any element in the set and returns a pointer to it.
-	 * Callers should not depend on particular patterns in the behaviour of this function.
-	 * @return A pointer to an arbitrary element, or nullptr if the container is empty.
-	 */
+		* Finds any element in the set and returns a pointer to it.
+		* Callers should not depend on particular patterns in the behaviour of this function.
+		* @return A pointer to an arbitrary element, or nullptr if the container is empty.
+		*/
 	[[nodiscard]] ElementType* FindArbitraryElement()
 	{
 		return (this->NumElements > 0) ? GetData() : nullptr;
@@ -1443,10 +1402,10 @@ private:
 	}
 
 	/**
-	 * Finds an element with a pre-calculated hash and a key that can be compared to KeyType
-	 * @see	Class documentation section on ByHash() functions
-	 * @return The element id that matches the key and hash or an invalid element id
-	 */
+		* Finds an element with a pre-calculated hash and a key that can be compared to KeyType
+		* @see	Class documentation section on ByHash() functions
+		* @return The element id that matches the key and hash or an invalid element id
+		*/
 	template<typename ComparableKey>
 	[[nodiscard]] SizeType FindIndexByHash(uint32 KeyHash, const ComparableKey& Key) const
 	{
diff --git a/Engine/Source/Runtime/Core/Public/Containers/CompactSetBase.h b/Engine/Source/Runtime/Core/Public/Containers/CompactSetBase.h
--- a/Engine/Source/Runtime/Core/Public/Containers/CompactSetBase.h
+++ b/Engine/Source/Runtime/Core/Public/Containers/CompactSetBase.h
@@ -6,6 +6,89 @@
 #include "Containers/CompactHashTable.h"
 #include "Templates/AlignmentTemplates.h"
 
+//
+// Compact set allocation definitions - It operates entirely with uint8 because it includes its lookup table inline with the overall allocation
+// Todo: move to ContainerAllocationPolicies.h
+//
+
+namespace UE::Core::Private::CompactSetAllocatorHelpers
+{
+	template <uint32 NumInlineElements, int32 ElementSize>
+	constexpr int32 CalculateRequiredBytes()
+	{
+		constexpr uint32 TypeSize = 1 + (NumInlineElements > 0xff) + (NumInlineElements > 0xffff) * 2;
+		constexpr uint32 HashSize = NumInlineElements < 8 ? 4 : FGenericPlatformMath::RoundUpToPowerOfTwo((NumInlineElements / 2) + 1);
+		return Align(NumInlineElements*ElementSize, 4) + 4 + (NumInlineElements + HashSize) * TypeSize;
+	}
+
+	/** Regular iterator */
+	template<typename InElementAllocator = FDefaultAllocator>
+	struct TCompactSetAllocator
+	{
+		/** CompactSets use raw uint8, so we need extra alignment information for inline allocators */
+		template <typename ElementType>
+		struct AllocatorAlignment
+		{
+			static constexpr size_t Value = alignof(typename InElementAllocator::template ForElementType<uint8>);
+		};
+
+		template<int32 ElementSize>
+		using ElementAllocator = InElementAllocator;
+	};
+
+	/** An inline allocator that allows sizing of the inline allocations for a set number of elements. */
+	template<uint32 NumInlineElements, typename SecondaryAllocator = TCompactSetAllocator<>>
+	struct TInlineCompactSetAllocator
+	{
+		template<int32 ElementSize>
+		using ElementAllocator = TInlineAllocator<UE::Core::Private::CompactSetAllocatorHelpers::CalculateRequiredBytes<NumInlineElements, ElementSize>(),typename SecondaryAllocator::template ElementAllocator<ElementSize>>;
+
+		/** CompactSets use raw uint8, so we need extra alignment information for inline allocators */
+		template <typename ElementType>
+		struct AllocatorAlignment
+		{
+			static constexpr size_t ElementAlignof = alignof(ElementType);
+			static constexpr size_t AllocatorAlignof = alignof(typename ElementAllocator<sizeof(ElementType)>::template ForElementType<uint8>);
+			static constexpr size_t Value = FGenericPlatformMath::Max(ElementAlignof, AllocatorAlignof);
+		};
+	};
+
+	/** An inline allocator that doesn't have any secondary storage. */
+	template <uint32 NumInlineElements>
+	struct TFixedCompactSetAllocator
+	{
+		template<int32 ElementSize>
+		using ElementAllocator = TFixedAllocator<UE::Core::Private::CompactSetAllocatorHelpers::CalculateRequiredBytes<NumInlineElements, ElementSize>()>;
+	
+		/** CompactSets use raw uint8, so we need extra alignment information for inline allocators */
+		template <typename ElementType>
+		struct AllocatorAlignment
+		{
+			static constexpr size_t ElementAlignof = alignof(ElementType);
+			static constexpr size_t AllocatorAlignof = alignof(typename ElementAllocator<sizeof(ElementType)>::template ForElementType<uint8>);
+			static constexpr size_t Value = FGenericPlatformMath::Max(ElementAlignof, AllocatorAlignof);
+		};
+	};
+}
+
+template<typename InElementAllocator>
+struct TAllocatorTraits<UE::Core::Private::CompactSetAllocatorHelpers::TCompactSetAllocator<InElementAllocator>> : TAllocatorTraitsBase<UE::Core::Private::CompactSetAllocatorHelpers::TCompactSetAllocator<InElementAllocator>>
+{
+	enum
+	{
+		SupportsFreezeMemoryImage = TAllocatorTraits<InElementAllocator>::SupportsFreezeMemoryImage,
+	};
+};
+
+class FDefaultCompactSetAllocator  : public UE::Core::Private::CompactSetAllocatorHelpers::TCompactSetAllocator<>  { public: typedef UE::Core::Private::CompactSetAllocatorHelpers::TCompactSetAllocator<>  Typedef; };
+
+// End Compact set allocation definitions
+
+// Todo: Move to ContainersFwd.h
+class FDefaultCompactSetAllocator;
+template<typename InElementType, typename KeyFuncs = DefaultKeyFuncs<InElementType>, typename Allocator = FDefaultCompactSetAllocator> class TCompactSet;
+// End Move to ContainersFwd.h
+
 /** Used to describe the data layout of the contents of the set */
 struct FCompactSetLayout
 {
@@ -26,50 +109,22 @@ public:
 		return MaxElements == INDEX_NONE;
 	}
 
-	/**
-	 * Returns true if the sets is empty and contains no elements. 
-	 *
-	 * @returns True if the set is empty.
-	 * @see Num
-	 */
 	[[nodiscard]] UE_FORCEINLINE_HINT bool IsEmpty() const
 	{
 		return NumElements == 0;
 	}
 
-	/** @return the number of elements. */
 	[[nodiscard]] UE_FORCEINLINE_HINT int32 Num() const
 	{
 		return NumElements;
 	}
 
-	/** @return The number of elements the set can hold before reallocation. */
-	[[nodiscard]] UE_FORCEINLINE_HINT int32 Max() const
-	{
-		return MaxElements;
-	}
-
 	/** Deprecated - unnecessary, keeping this here so TCompactSet can be swapped with TSet without changing code */
 	[[nodiscard]] UE_FORCEINLINE_HINT int32 GetMaxIndex() const
 	{
 		return NumElements;
 	}
 
-	/** 
-	 * Helper function to return the amount of memory allocated by this container 
-	 * Only returns the size of allocations made directly by the container, not the elements themselves.
-	 * @return number of bytes allocated by this container
-	 */
-	[[nodiscard]] UE_FORCEINLINE_HINT SIZE_T GetAllocatedSize(const FCompactSetLayout Layout) const
-	{
-		if (MaxElements == 0)
-		{
-			return 0;
-		}
-
-		return GetTotalMemoryRequiredInBytes(MaxElements, *GetHashTableMemory(Layout), Layout);
-	}
-
 protected:
 	using HashCountType = uint32;
 	static constexpr size_t HashCountSize = sizeof(HashCountType);
@@ -78,14 +133,7 @@ protected:
 	static_assert(std::is_same_v<SizeType, int32>, "TCompactSet currently only supports 32-bit allocators");
 	static_assert(sizeof(HashCountType) == UE::Core::CompactHashTable::GetMemoryAlignment(), "Hashtable alignment changed, need to make sure HashCountType is stored correctly");
 
-	[[nodiscard]] UE_FORCEINLINE_HINT TCompactSetBase() = default;
-
-	[[nodiscard]] explicit consteval TCompactSetBase(EConstEval)
-	: Elements(ConstEval)
-	, NumElements(0)
-	, MaxElements(0)
-	{
-	}
+	[[nodiscard]] UE_FORCEINLINE_HINT TCompactSetBase() = default;	
 
 	[[nodiscard]] explicit TCompactSetBase(FIntrusiveUnsetOptionalState Tag)
 	: NumElements(0)
diff --git a/Engine/Source/Runtime/Core/Public/Containers/ContainerAllocationPolicies.h b/Engine/Source/Runtime/Core/Public/Containers/ContainerAllocationPolicies.h
--- a/Engine/Source/Runtime/Core/Public/Containers/ContainerAllocationPolicies.h
+++ b/Engine/Source/Runtime/Core/Public/Containers/ContainerAllocationPolicies.h
@@ -8,7 +8,6 @@
 #include "HAL/UnrealMemory.h"
 #include "Math/NumericLimits.h"
 #include "Misc/AssertionMacros.h"
-#include "Templates/AlignmentTemplates.h"
 #include "Templates/IsPolymorphic.h"
 #include "Templates/MemoryOps.h"
 #include "Templates/TypeCompatibleBytes.h"
@@ -135,6 +134,12 @@ struct FArraySlackTrackingHeader
 static_assert(UE_CONTAINER_SLACK_GROWTH_FACTOR_DENOMINATOR > 0, "UE_CONTAINER_SLACK_GROWTH_FACTOR_DENOMINATOR must be greater than 0");
 static_assert(UE_CONTAINER_SLACK_GROWTH_FACTOR_DENOMINATOR > UE_CONTAINER_SLACK_GROWTH_FACTOR_NUMERATOR, "UE_CONTAINER_SLACK_GROWTH_FACTOR_DENOMINATOR must be greater than UE_CONTAINER_SLACK_GROWTH_FACTOR_NUMERATOR");
 
+
+class FDefaultBitArrayAllocator;
+
+template<int IndexSize> class TSizedDefaultAllocator;
+using FDefaultAllocator = TSizedDefaultAllocator<32>;
+
 template <typename SizeType>
 UE_FORCEINLINE_HINT SizeType DefaultCalculateSlackShrink(SizeType NewMax, SizeType CurrentMax, SIZE_T BytesPerElement, bool bAllowQuantize, uint32 Alignment = DEFAULT_ALIGNMENT)
 {
@@ -618,7 +623,6 @@ class TSizedHeapAllocator
 {
 public:
 	using SizeType = typename TBitsToSizeType<IndexSize>::Type;
-	using BaseMalloc = BaseMallocType;
 
 private:
 	using USizeType = std::make_unsigned_t<SizeType>;
@@ -827,9 +831,6 @@ public:
 	};
 };
 
-template <int IndexSize> class TSizedDefaultAllocator : public TSizedHeapAllocator<IndexSize> { public: typedef TSizedHeapAllocator<IndexSize> Typedef; };
-using FDefaultAllocator = TSizedDefaultAllocator<32>;
-
 // Define the ResizeAllocation functions with the regular allocator as exported to avoid bloat
 #if !IS_MERGEDMODULES && !defined(UE_HEADER_UNITS)
 extern template CORE_API FORCENOINLINE void TSizedHeapAllocator<32, FMemory>::ForAnyElementType::ResizeAllocation(SizeType CurrentNum, SizeType NewMax, SIZE_T NumBytesPerElement);
@@ -919,9 +920,6 @@ public:
 
 		void ResizeAllocation(SizeType CurrentNum, SizeType NewMax,SIZE_T NumBytesPerElement)
 		{
-			// Make sure the number of live elements is still withing the allocation since we only memmove, not destruct
-			checkSlow(CurrentNum <= NewMax);
-
 			// Check if the new allocation will fit in the inline data area.
 			if(NewMax <= NumInlineElements)
 			{
@@ -1115,9 +1113,6 @@ public:
 
 		void ResizeAllocation(SizeType CurrentNum, SizeType NewMax,SIZE_T NumBytesPerElement)
 		{
-			// Make sure the number of live elements is still withing the allocation since we only memmove, not destruct
-			checkSlow(CurrentNum <= NewMax);
-
 			// Check if the new allocation will fit in the inline data area.
 			if(NewMax <= NumInlineElements)
 			{
@@ -1333,8 +1328,6 @@ public:
 #define NumBitsPerDWORD ((int32)32)
 #define NumBitsPerDWORDLogTwo ((int32)5)
 
-class FDefaultBitArrayAllocator;
-
 //
 // Sparse array allocation definitions
 //
@@ -1390,79 +1383,6 @@ public:
 	typedef TFixedAllocator<InlineBitArrayDWORDs> BitArrayAllocator;
 };
 
-//
-// Compact set allocation definitions - It operates entirely with uint8 because it includes its lookup table inline with the overall allocation
-//
-
-namespace UE::Core::Private::CompactSetAllocatorHelpers
-{
-	template <uint32 NumInlineElements, int32 ElementSize>
-	constexpr int32 CalculateRequiredBytes()
-	{
-		constexpr uint32 TypeSize = 1 + (NumInlineElements > 0xff) + (NumInlineElements > 0xffff) * 2;
-		constexpr uint32 HashSize = NumInlineElements < 8 ? 4 : FGenericPlatformMath::RoundUpToPowerOfTwo((NumInlineElements / 2) + 1);
-		return Align(NumInlineElements*ElementSize, 4) + 4 + (NumInlineElements + HashSize) * TypeSize;
-	}
-}
-
-/** Regular iterator */
-template<typename InElementAllocator = FDefaultAllocator>
-struct TCompactSetAllocator
-{
-	/** CompactSets use raw uint8, so we need extra alignment information for inline allocators */
-	template <typename ElementType>
-	struct AllocatorAlignment
-	{
-		static constexpr size_t Value = alignof(typename InElementAllocator::template ForElementType<uint8>);
-	};
-
-	template<int32 ElementSize>
-	using ElementAllocator = InElementAllocator;
-};
-
-template<typename InElementAllocator>
-struct TAllocatorTraits<TCompactSetAllocator<InElementAllocator>> : TAllocatorTraitsBase<TCompactSetAllocator<InElementAllocator>>
-{
-	enum
-	{
-		SupportsFreezeMemoryImage = TAllocatorTraits<InElementAllocator>::SupportsFreezeMemoryImage,
-	};
-};
-
-/** An inline allocator that allows sizing of the inline allocations for a set number of elements. */
-template<uint32 NumInlineElements, typename SecondaryAllocator = TCompactSetAllocator<>>
-struct TInlineCompactSetAllocator
-{
-	template<int32 ElementSize>
-	using ElementAllocator = TInlineAllocator<UE::Core::Private::CompactSetAllocatorHelpers::CalculateRequiredBytes<NumInlineElements, ElementSize>(),typename SecondaryAllocator::template ElementAllocator<ElementSize>>;
-
-	/** CompactSets use raw uint8, so we need extra alignment information for inline allocators */
-	template <typename ElementType>
-	struct AllocatorAlignment
-	{
-		static constexpr size_t ElementAlignof = alignof(ElementType);
-		static constexpr size_t AllocatorAlignof = alignof(typename ElementAllocator<sizeof(ElementType)>::template ForElementType<uint8>);
-		static constexpr size_t Value = FGenericPlatformMath::Max(ElementAlignof, AllocatorAlignof);
-	};
-};
-
-/** An inline allocator that doesn't have any secondary storage. */
-template <uint32 NumInlineElements>
-struct TFixedCompactSetAllocator
-{
-	template<int32 ElementSize>
-	using ElementAllocator = TFixedAllocator<UE::Core::Private::CompactSetAllocatorHelpers::CalculateRequiredBytes<NumInlineElements, ElementSize>()>;
-	
-	/** CompactSets use raw uint8, so we need extra alignment information for inline allocators */
-	template <typename ElementType>
-	struct AllocatorAlignment
-	{
-		static constexpr size_t ElementAlignof = alignof(ElementType);
-		static constexpr size_t AllocatorAlignof = alignof(typename ElementAllocator<sizeof(ElementType)>::template ForElementType<uint8>);
-		static constexpr size_t Value = FGenericPlatformMath::Max(ElementAlignof, AllocatorAlignof);
-	};
-};
-
 
 
 //
@@ -1483,7 +1403,7 @@ template<
 	uint32   BaseNumberOfHashBuckets              = DEFAULT_BASE_NUMBER_OF_HASH_BUCKETS,
 	uint32   MinNumberOfHashedElements            = DEFAULT_MIN_NUMBER_OF_HASHED_ELEMENTS
 	>
-class TSparseSetAllocator
+class TSetAllocator
 {
 public:
 
@@ -1498,8 +1418,8 @@ public:
 		return 1;
 	}
 
-	using SparseArrayAllocator = InSparseArrayAllocator;
-	using HashAllocator = InHashAllocator;
+	typedef InSparseArrayAllocator SparseArrayAllocator;
+	typedef InHashAllocator        HashAllocator;
 };
 
 template<
@@ -1509,8 +1429,8 @@ template<
 	uint32   BaseNumberOfHashBuckets,
 	uint32   MinNumberOfHashedElements
 >
-struct TAllocatorTraits<TSparseSetAllocator<InSparseArrayAllocator, InHashAllocator, AverageNumberOfElementsPerHashBucket, BaseNumberOfHashBuckets, MinNumberOfHashedElements>> :
-	TAllocatorTraitsBase<TSparseSetAllocator<InSparseArrayAllocator, InHashAllocator, AverageNumberOfElementsPerHashBucket, BaseNumberOfHashBuckets, MinNumberOfHashedElements>>
+struct TAllocatorTraits<TSetAllocator<InSparseArrayAllocator, InHashAllocator, AverageNumberOfElementsPerHashBucket, BaseNumberOfHashBuckets, MinNumberOfHashedElements>> :
+	TAllocatorTraitsBase<TSetAllocator<InSparseArrayAllocator, InHashAllocator, AverageNumberOfElementsPerHashBucket, BaseNumberOfHashBuckets, MinNumberOfHashedElements>>
 {
 	enum
 	{
@@ -1521,11 +1441,11 @@ struct TAllocatorTraits<TSparseSetAllocator<InSparseArrayAllocator, InHashAlloca
 /** An inline set allocator that allows sizing of the inline allocations for a set number of elements. */
 template<
 	uint32   NumInlineElements,
-	typename SecondaryAllocator                   = TSparseSetAllocator<TSparseArrayAllocator<FDefaultAllocator,FDefaultAllocator>,FDefaultAllocator>,
+	typename SecondaryAllocator                   = TSetAllocator<TSparseArrayAllocator<FDefaultAllocator,FDefaultAllocator>,FDefaultAllocator>,
 	uint32   AverageNumberOfElementsPerHashBucket = DEFAULT_NUMBER_OF_ELEMENTS_PER_HASH_BUCKET,
 	uint32   MinNumberOfHashedElements            = DEFAULT_MIN_NUMBER_OF_HASHED_ELEMENTS
 	>
-class TInlineSparseSetAllocator
+class TInlineSetAllocator
 {
 private:
 
@@ -1552,8 +1472,8 @@ public:
 		return NumDesiredHashBuckets;
 	}
 
-	using SparseArrayAllocator = TInlineSparseArrayAllocator<NumInlineElements,typename SecondaryAllocator::SparseArrayAllocator>;
-	using HashAllocator = TInlineAllocator<NumInlineHashBuckets,typename SecondaryAllocator::HashAllocator>;
+	typedef TInlineSparseArrayAllocator<NumInlineElements,typename SecondaryAllocator::SparseArrayAllocator> SparseArrayAllocator;
+	typedef TInlineAllocator<NumInlineHashBuckets,typename SecondaryAllocator::HashAllocator>                HashAllocator;
 };
 
 /** An inline set allocator that doesn't have any secondary storage. */
@@ -1562,7 +1482,7 @@ template<
 	uint32 AverageNumberOfElementsPerHashBucket = DEFAULT_NUMBER_OF_ELEMENTS_PER_HASH_BUCKET,
 	uint32 MinNumberOfHashedElements            = DEFAULT_MIN_NUMBER_OF_HASHED_ELEMENTS
 	>
-class TFixedSparseSetAllocator
+class TFixedSetAllocator
 {
 private:
 
@@ -1589,8 +1509,8 @@ public:
 		return NumDesiredHashBuckets;
 	}
 
-	using SparseArrayAllocator = TFixedSparseArrayAllocator<NumInlineElements>;
-	using HashAllocator = TFixedAllocator<NumInlineHashBuckets>;
+	typedef TFixedSparseArrayAllocator<NumInlineElements> SparseArrayAllocator;
+	typedef TFixedAllocator<NumInlineHashBuckets>         HashAllocator;
 };
 
 
@@ -1601,38 +1521,16 @@ public:
  * 'forward' these TAllocatorTraits specializations below.
  */
 
-class FDefaultSparseSetAllocator   : public TSparseSetAllocator<>   { public: typedef TSparseSetAllocator<>   Typedef; };
-class FDefaultCompactSetAllocator  : public TCompactSetAllocator<>  { public: typedef TCompactSetAllocator<>  Typedef; };
+template <int IndexSize> class TSizedDefaultAllocator : public TSizedHeapAllocator<IndexSize> { public: typedef TSizedHeapAllocator<IndexSize> Typedef; };
+
+class FDefaultSetAllocator         : public TSetAllocator<>         { public: typedef TSetAllocator<>         Typedef; };
 class FDefaultBitArrayAllocator    : public TInlineAllocator<4>     { public: typedef TInlineAllocator<4>     Typedef; };
 class FDefaultSparseArrayAllocator : public TSparseArrayAllocator<> { public: typedef TSparseArrayAllocator<> Typedef; };
 
-#ifndef UE_USE_COMPACT_SET_AS_DEFAULT
-#define UE_USE_COMPACT_SET_AS_DEFAULT 0
-#endif
-
-#if UE_USE_COMPACT_SET_AS_DEFAULT
-class FDefaultSetAllocator         : public TCompactSetAllocator<>  { public: typedef TCompactSetAllocator<>  Typedef; };
-template<typename InSparseArrayAllocator = TSparseArrayAllocator<>, typename InHashAllocator = TInlineAllocator<1,FDefaultAllocator>, uint32... N>
-class TSetAllocator        		   : public TCompactSetAllocator<InHashAllocator>  { public: typedef TCompactSetAllocator<InHashAllocator>  Typedef; };
-template <uint32 N, typename S = TCompactSetAllocator<>, uint32... NN>
-class TInlineSetAllocator          : public TInlineCompactSetAllocator<N, S>  { public: typedef TInlineCompactSetAllocator<N, S>  Typedef; };
-template <uint32 N, uint32... Y>
-class TFixedSetAllocator           : public TFixedCompactSetAllocator<N>  { public: typedef TFixedCompactSetAllocator<N>  Typedef; };
-#else
-class FDefaultSetAllocator         : public TSparseSetAllocator<>         { public: typedef TSparseSetAllocator<>         Typedef; };
-template<typename InSparseArrayAllocator = TSparseArrayAllocator<>, typename InHashAllocator = TInlineAllocator<1,FDefaultAllocator>, uint32... N>
-class TSetAllocator         	   : public TSparseSetAllocator<InSparseArrayAllocator, InHashAllocator, N...>  { public: typedef TSparseSetAllocator<InSparseArrayAllocator, InHashAllocator, N...>  Typedef; };
-template<uint32 N, typename S = TSparseSetAllocator<TSparseArrayAllocator<FDefaultAllocator,FDefaultAllocator>,FDefaultAllocator>, uint32... NN>
-class TInlineSetAllocator          : public TInlineSparseSetAllocator<N, S, NN...>  { public: typedef TInlineSparseSetAllocator<N, S, NN...>  Typedef; };
-template <uint32... N>
-class TFixedSetAllocator           : public TFixedSparseSetAllocator<N...>  { public: typedef TFixedSparseSetAllocator<N...>  Typedef; };
-#endif
-
 template <int IndexSize> struct TAllocatorTraits<TSizedDefaultAllocator<IndexSize>> : TAllocatorTraits<typename TSizedDefaultAllocator<IndexSize>::Typedef> {};
 
 template <> struct TAllocatorTraits<FDefaultAllocator>            : TAllocatorTraits<typename FDefaultAllocator           ::Typedef> {};
 template <> struct TAllocatorTraits<FDefaultSetAllocator>         : TAllocatorTraits<typename FDefaultSetAllocator        ::Typedef> {};
-template <> struct TAllocatorTraits<FDefaultSparseSetAllocator>   : TAllocatorTraits<typename FDefaultSparseSetAllocator  ::Typedef> {};
 template <> struct TAllocatorTraits<FDefaultBitArrayAllocator>    : TAllocatorTraits<typename FDefaultBitArrayAllocator   ::Typedef> {};
 template <> struct TAllocatorTraits<FDefaultSparseArrayAllocator> : TAllocatorTraits<typename FDefaultSparseArrayAllocator::Typedef> {};
 
diff --git a/Engine/Source/Runtime/Core/Public/Containers/ContainersFwd.h b/Engine/Source/Runtime/Core/Public/Containers/ContainersFwd.h
--- a/Engine/Source/Runtime/Core/Public/Containers/ContainersFwd.h
+++ b/Engine/Source/Runtime/Core/Public/Containers/ContainersFwd.h
@@ -8,13 +8,9 @@
 
 /// @cond DOXYGEN_WARNINGS
 template<int IndexSize> class TSizedDefaultAllocator;
-template<typename ElementType,bool bInAllowDuplicateKeys = false> struct DefaultKeyFuncs;
-
 using FDefaultAllocator = TSizedDefaultAllocator<32>;
 using FDefaultAllocator64 = TSizedDefaultAllocator<64>;
 class FDefaultSetAllocator;
-class FDefaultCompactSetAllocator;
-class FDefaultSparseSetAllocator;
 
 class FString;
 class FAnsiString;
@@ -69,18 +65,13 @@ template<typename T, typename SizeType = int32> using TConstArrayView = TArrayVi
 template<typename T> using TConstArrayView64 = TConstArrayView<T, int64>;
 template<typename T> class TTransArray;
 template<typename InKeyType, typename InValueType, bool bInAllowDuplicateKeys> struct TDefaultMapHashableKeyFuncs;
-template<typename InKeyType, typename InValueType, typename SetAllocator = FDefaultSetAllocator, typename KeyFuncs = TDefaultMapHashableKeyFuncs<InKeyType, InValueType, false> > class TSparseMap;
-template<typename InKeyType, typename InValueType, typename SetAllocator = FDefaultSetAllocator, typename KeyFuncs = TDefaultMapHashableKeyFuncs<InKeyType, InValueType, false> > class TCompactMap;
 template<typename InKeyType, typename InValueType, typename SetAllocator = FDefaultSetAllocator, typename KeyFuncs = TDefaultMapHashableKeyFuncs<InKeyType, InValueType, false> > class TMap;
-template<typename KeyType, typename ValueType, typename SetAllocator = FDefaultSetAllocator, typename KeyFuncs = TDefaultMapHashableKeyFuncs<KeyType, ValueType, true > > class TSparseMultiMap;
-template<typename KeyType, typename ValueType, typename SetAllocator = FDefaultSetAllocator, typename KeyFuncs = TDefaultMapHashableKeyFuncs<KeyType, ValueType, true > > class TCompactMultiMap;
 template<typename KeyType, typename ValueType, typename SetAllocator = FDefaultSetAllocator, typename KeyFuncs = TDefaultMapHashableKeyFuncs<KeyType, ValueType, true > > class TMultiMap;
 template <typename T = void > struct TLess;
 template <typename> struct TTypeTraits;
 template<typename InKeyType, typename InValueType, typename ArrayAllocator = FDefaultAllocator, typename SortPredicate = TLess<typename TTypeTraits<InKeyType>::ConstPointerType> > class TSortedMap;
 template <typename InElementType, typename ArrayAllocator = FDefaultAllocator, typename SortPredicate = TLess<InElementType>> class TSortedSet;
-template<typename InElementType, typename KeyFuncs = DefaultKeyFuncs<InElementType>, typename Allocator = FDefaultSparseSetAllocator> class TSparseSet;
-template<typename InElementType, typename KeyFuncs = DefaultKeyFuncs<InElementType>, typename Allocator = FDefaultCompactSetAllocator> class TCompactSet;
+template<typename ElementType,bool bInAllowDuplicateKeys = false> struct DefaultKeyFuncs;
 template<typename InElementType, typename KeyFuncs = DefaultKeyFuncs<InElementType>, typename Allocator = FDefaultSetAllocator> class TSet;
 template<typename InElementType, typename InSizeType = int32> class TStridedView;
 template<typename T, typename SizeType = int32> using TConstStridedView = TStridedView<const T, SizeType>;
diff --git a/Engine/Source/Runtime/Core/Public/Containers/Map.h b/Engine/Source/Runtime/Core/Public/Containers/Map.h
--- a/Engine/Source/Runtime/Core/Public/Containers/Map.h
+++ b/Engine/Source/Runtime/Core/Public/Containers/Map.h
@@ -116,25 +116,1767 @@ struct TDefaultMapHashableKeyFuncs : TDefaultMapKeyFuncs<KeyType, ValueType, bIn
 	using HashabilityCheck = decltype(GetTypeHash(DeclVal<const KeyType>()));
 };
 
+/** 
+ * The base class of maps from keys to values.  Implemented using a TSet of key-value pairs with a custom KeyFuncs, 
+ * with the same O(1) addition, removal, and finding. 
+ *
+ * The ByHash() functions are somewhat dangerous but particularly useful in two scenarios:
+ * -- Heterogeneous lookup to avoid creating expensive keys like FString when looking up by const TCHAR*.
+ *	  You must ensure the hash is calculated in the same way as ElementType is hashed.
+ *    If possible put both ComparableKey and ElementType hash functions next to each other in the same header
+ *    to avoid bugs when the ElementType hash function is changed.
+ * -- Reducing contention around hash tables protected by a lock. It is often important to incur
+ *    the cache misses of reading key data and doing the hashing *before* acquiring the lock.
+ **/
+template <typename KeyType, typename ValueType, typename SetAllocator, typename KeyFuncs>
+class TMapBase
+{
+	template <typename OtherKeyType, typename OtherValueType, typename OtherSetAllocator, typename OtherKeyFuncs>
+	friend class TMapBase;
+
+public:
+	typedef typename TTypeTraits<KeyType  >::ConstPointerType KeyConstPointerType;
+	typedef typename TTypeTraits<KeyType  >::ConstInitType    KeyInitType;
+	typedef typename TTypeTraits<ValueType>::ConstInitType    ValueInitType;
+	typedef TPair<KeyType, ValueType> ElementType;
+
+protected:
+	[[nodiscard]] TMapBase() = default;
+	[[nodiscard]] explicit consteval TMapBase(EConstEval)
+	: Pairs(ConstEval)
+	{
+	} 
+	[[nodiscard]] TMapBase(TMapBase&&) = default;
+	[[nodiscard]] TMapBase(const TMapBase&) = default;
+	TMapBase& operator=(TMapBase&&) = default;
+	TMapBase& operator=(const TMapBase&) = default;
+
+	/** Constructor for moving elements from a TMap with a different SetAllocator */
+	template<typename OtherSetAllocator>
+	[[nodiscard]] TMapBase(TMapBase<KeyType, ValueType, OtherSetAllocator, KeyFuncs>&& Other)
+		: Pairs(MoveTemp(Other.Pairs))
+	{
+	}
+
+	/** Constructor for copying elements from a TMap with a different SetAllocator */
+	template<typename OtherSetAllocator>
+	[[nodiscard]] TMapBase(const TMapBase<KeyType, ValueType, OtherSetAllocator, KeyFuncs>& Other)
+		: Pairs(Other.Pairs)
+	{
+	}
+
+	~TMapBase()
+	{
+		UE_STATIC_ASSERT_WARN(TIsTriviallyRelocatable_V<KeyType> && TIsTriviallyRelocatable_V<ValueType>, "TMapBase can only be used with trivially relocatable types");
+	}
+
+	/////////////////////////////////////////////////
+	// Start - intrusive TOptional<TMapBase> state //
+	/////////////////////////////////////////////////
+	constexpr static bool bHasIntrusiveUnsetOptionalState = true;
+	using IntrusiveUnsetOptionalStateType = TMapBase;
+
+	[[nodiscard]] explicit TMapBase(FIntrusiveUnsetOptionalState Tag)
+		: Pairs(Tag)
+	{
+	}
+	[[nodiscard]] bool operator==(FIntrusiveUnsetOptionalState Tag) const
+	{
+		return Pairs == Tag;
+	}
+	///////////////////////////////////////////////
+	// End - intrusive TOptional<TMapBase> state //
+	///////////////////////////////////////////////
+
+	/** Assignment operator for moving elements from a TMap with a different SetAllocator */
+	template<typename OtherSetAllocator>
+	TMapBase& operator=(TMapBase<KeyType, ValueType, OtherSetAllocator, KeyFuncs>&& Other)
+	{
+		Pairs = MoveTemp(Other.Pairs);
+		return *this;
+	}
+
+	/** Assignment operator for copying elements from a TMap with a different SetAllocator */
+	template<typename OtherSetAllocator>
+	TMapBase& operator=(const TMapBase<KeyType, ValueType, OtherSetAllocator, KeyFuncs>& Other)
+	{
+		Pairs = Other.Pairs;
+		return *this;
+	}
+
+public:
+
+	/**
+	 * Compare this map with another for equality. Does not make any assumptions about Key order.
+	 * NOTE: this might be a candidate for operator== but it was decided to make it an explicit function
+	 *  since it can potentially be quite slow.
+	 *
+	 * @param Other The other map to compare against
+	 * @returns True if both this and Other contain the same keys with values that compare ==
+	 */
+	[[nodiscard]] bool OrderIndependentCompareEqual(const TMapBase& Other) const
+	{
+		// first check counts (they should be the same obviously)
+		if (Num() != Other.Num())
+		{
+			return false;
+		}
+
+		// since we know the counts are the same, we can just iterate one map and check for existence in the other
+		for (typename ElementSetType::TConstIterator It(Pairs); It; ++It)
+		{
+			const ValueType* BVal = Other.Find(It->Key);
+			if (BVal == nullptr)
+			{
+				return false;
+			}
+			if (!(*BVal == It->Value))
+			{
+				return false;
+			}
+		}
+
+		// all fields in A match B and A and B's counts are the same (so there can be no fields in B not in A)
+		return true;
+	}
+
+	/**
+	 * Removes all elements from the map.
+	 *
+	 * This method potentially leaves space allocated for an expected
+	 * number of elements about to be added.
+	 *
+	 * @param ExpectedNumElements The number of elements about to be added to the set.
+	 */
+	UE_FORCEINLINE_HINT void Empty(int32 ExpectedNumElements = 0)
+	{
+		Pairs.Empty(ExpectedNumElements);
+	}
+
+	/** Efficiently empties out the map but preserves all allocations and capacities */
+	UE_FORCEINLINE_HINT void Reset()
+	{
+		Pairs.Reset();
+	}
+
+	/** Shrinks the pair set to avoid slack. */
+	UE_FORCEINLINE_HINT void Shrink()
+	{
+		Pairs.Shrink();
+	}
+
+	/** Compacts the pair set to remove holes */
+	UE_FORCEINLINE_HINT void Compact()
+	{
+		Pairs.Compact();
+	}
+
+	/** Compacts the pair set to remove holes. Does not change the iteration order of the elements. */
+	UE_FORCEINLINE_HINT void CompactStable()
+	{
+		Pairs.CompactStable();
+	}
+
+	/** Preallocates enough memory to contain Number elements */
+	UE_FORCEINLINE_HINT void Reserve(int32 Number)
+	{
+		Pairs.Reserve(Number);
+	}
+
+	/**
+	 * Returns true if the map is empty and contains no elements. 
+	 *
+	 * @returns True if the map is empty.
+	 * @see Num
+	 */
+	[[nodiscard]] bool IsEmpty() const
+	{
+		return Pairs.IsEmpty();
+	}
+
+	/** @return The number of elements in the map. */
+	[[nodiscard]] UE_FORCEINLINE_HINT int32 Num() const
+	{
+		return Pairs.Num();
+	}
+
+	/** @return The number of elements the map can hold before reallocation. */
+	[[nodiscard]] UE_FORCEINLINE_HINT int32 Max() const
+	{
+		return Pairs.Max();
+	}
+
+	/** @return The non-inclusive maximum index of elements in the map. */
+	[[nodiscard]] UE_FORCEINLINE_HINT int32 GetMaxIndex() const
+	{
+		return Pairs.GetMaxIndex();
+	}
+
+	/**
+	 * Checks whether an element id is valid.
+	 * @param Id - The element id to check.
+	 * @return true if the element identifier refers to a valid element in this map.
+	 */
+	[[nodiscard]] UE_FORCEINLINE_HINT bool IsValidId(FSetElementId Id) const
+	{
+		return Pairs.IsValidId(Id);
+	}
+
+	/** Return a mapped pair by internal identifier. Element must be valid (see @IsValidId). */
+	[[nodiscard]] UE_FORCEINLINE_HINT ElementType& Get(FSetElementId Id)
+	{
+		return Pairs[Id];
+	}
+
+	/** Return a mapped pair by internal identifier.  Element must be valid (see @IsValidId).*/
+	[[nodiscard]] UE_FORCEINLINE_HINT const ElementType& Get(FSetElementId Id) const
+	{
+		return Pairs[Id];
+	}
+
+	/**
+	 * Get the unique keys contained within this map.
+	 *
+	 * @param OutKeys Upon return, contains the set of unique keys in this map.
+	 * @return The number of unique keys in the map.
+	 */
+	template<typename Allocator> int32 GetKeys(TArray<KeyType, Allocator>& OutKeys) const
+	{
+		OutKeys.Reset();
+
+		TSet<KeyType> VisitedKeys;
+		VisitedKeys.Reserve(Num());
+
+		// Presize the array if we know there are supposed to be no duplicate keys
+		if constexpr (!KeyFuncs::bAllowDuplicateKeys)
+		{
+			OutKeys.Reserve(Num());
+		}
+
+		for (typename ElementSetType::TConstIterator It(Pairs); It; ++It)
+		{
+			// Even if bAllowDuplicateKeys is false, we still want to filter for duplicate
+			// keys due to maps with keys that can be invalidated (UObjects, TWeakObj, etc.)
+			if (!VisitedKeys.Contains(It->Key))
+			{
+				OutKeys.Add(It->Key);
+				VisitedKeys.Add(It->Key);
+			}
+		}
+
+		return OutKeys.Num();
+	}
+
+	/**
+	 * Get the unique keys contained within this map.
+	 *
+	 * @param OutKeys Upon return, contains the set of unique keys in this map.
+	 * @return The number of unique keys in the map.
+	 */
+	template<typename InSetKeyFuncs, typename InSetAllocator> int32 GetKeys(TSet<KeyType, InSetKeyFuncs, InSetAllocator>& OutKeys) const
+	{
+		OutKeys.Reset();
+		
+		// Presize the set if we know there are supposed to be no duplicate keys
+		if (!KeyFuncs::bAllowDuplicateKeys)
+		{
+			OutKeys.Reserve(Num());
+		}
+
+		for (typename ElementSetType::TConstIterator It(Pairs); It; ++It)
+		{
+			OutKeys.Add(It->Key);
+		}
+
+		return OutKeys.Num();
+	}
+
+	/** 
+	 * Helper function to return the amount of memory allocated by this container .
+	 * Only returns the size of allocations made directly by the container, not the elements themselves.
+	 *
+	 * @return Number of bytes allocated by this container.
+	 * @see CountBytes
+	 */
+	[[nodiscard]] UE_FORCEINLINE_HINT SIZE_T GetAllocatedSize() const
+	{
+		return Pairs.GetAllocatedSize();
+	}
+
+	/**
+	 * Track the container's memory use through an archive.
+	 *
+	 * @param Ar The archive to use.
+	 * @see GetAllocatedSize
+	 */
+	UE_FORCEINLINE_HINT void CountBytes(FArchive& Ar) const
+	{
+		Pairs.CountBytes(Ar);
+	}
+
+	/**
+	 * Set the value associated with a key.
+	 *
+	 * @param InKey The key to associate the value with.
+	 * @param InValue The value to associate with the key.
+	 * @return A reference to the value as stored in the map. The reference is only valid until the next change to any key in the map.
+	 */
+	UE_FORCEINLINE_HINT ValueType& Add(const KeyType&  InKey, const ValueType&  InValue) { return Emplace(InKey, InValue); }
+	UE_FORCEINLINE_HINT ValueType& Add(const KeyType&  InKey,		ValueType&& InValue) { return Emplace(InKey, MoveTempIfPossible(InValue)); }
+	UE_FORCEINLINE_HINT ValueType& Add(		 KeyType&& InKey, const ValueType&  InValue) { return Emplace(MoveTempIfPossible(InKey), InValue); }
+	UE_FORCEINLINE_HINT ValueType& Add(		 KeyType&& InKey,		ValueType&& InValue) { return Emplace(MoveTempIfPossible(InKey), MoveTempIfPossible(InValue)); }
+
+	/** See Add() and class documentation section on ByHash() functions */
+	UE_FORCEINLINE_HINT ValueType& AddByHash(uint32 KeyHash, const KeyType&  InKey, const ValueType&  InValue) { return EmplaceByHash(KeyHash, InKey, InValue); }
+	UE_FORCEINLINE_HINT ValueType& AddByHash(uint32 KeyHash, const KeyType&  InKey,		  ValueType&& InValue) { return EmplaceByHash(KeyHash, InKey, MoveTempIfPossible(InValue)); }
+	UE_FORCEINLINE_HINT ValueType& AddByHash(uint32 KeyHash,	   KeyType&& InKey, const ValueType&  InValue) { return EmplaceByHash(KeyHash, MoveTempIfPossible(InKey), InValue); }
+	UE_FORCEINLINE_HINT ValueType& AddByHash(uint32 KeyHash,	   KeyType&& InKey,		  ValueType&& InValue) { return EmplaceByHash(KeyHash, MoveTempIfPossible(InKey), MoveTempIfPossible(InValue)); }
+
+	/**
+	 * Set a default value associated with a key.
+	 *
+	 * @param InKey The key to associate the value with.
+	 * @return A reference to the value as stored in the map. The reference is only valid until the next change to any key in the map.
+	 */
+	UE_FORCEINLINE_HINT ValueType& Add(const KeyType&  InKey) { return Emplace(InKey); }
+	UE_FORCEINLINE_HINT ValueType& Add(		 KeyType&& InKey) { return Emplace(MoveTempIfPossible(InKey)); }
+
+	/** See Add() and class documentation section on ByHash() functions */
+	UE_FORCEINLINE_HINT ValueType& AddByHash(uint32 KeyHash, const KeyType&  InKey) { return EmplaceByHash(KeyHash, InKey); }
+	UE_FORCEINLINE_HINT ValueType& AddByHash(uint32 KeyHash,	   KeyType&& InKey)	{ return EmplaceByHash(KeyHash, MoveTempIfPossible(InKey)); }
+
+	/**
+	 * Set the value associated with a key.
+	 *
+	 * @param InKeyValue A Tuple containing the Key and Value to associate together
+	 * @return A reference to the value as stored in the map. The reference is only valid until the next change to any key in the map.
+	 */
+	UE_FORCEINLINE_HINT ValueType& Add(const TTuple<KeyType, ValueType>&  InKeyValue) { return Emplace(InKeyValue.Key, InKeyValue.Value); }
+	UE_FORCEINLINE_HINT ValueType& Add(		 TTuple<KeyType, ValueType>&& InKeyValue) { return Emplace(MoveTempIfPossible(InKeyValue.Key), MoveTempIfPossible(InKeyValue.Value)); }
+
+	/**
+	 * Sets the value associated with a key.
+	 *
+	 * @param InKey The key to associate the value with.
+	 * @param InValue The value to associate with the key.
+	 * @return A reference to the value as stored in the map. The reference is only valid until the next change to any key in the map.	 */
+	template <typename InitKeyType = KeyType, typename InitValueType = ValueType>
+	ValueType& Emplace(InitKeyType&& InKey, InitValueType&& InValue)
+	{
+		const FSetElementId PairId = Pairs.Emplace(TPairInitializer<InitKeyType&&, InitValueType&&>(Forward<InitKeyType>(InKey), Forward<InitValueType>(InValue)));
+
+		return Pairs[PairId].Value;
+	}
+
+	/** See Emplace() and class documentation section on ByHash() functions */
+	template <typename InitKeyType = KeyType, typename InitValueType = ValueType>
+	ValueType& EmplaceByHash(uint32 KeyHash, InitKeyType&& InKey, InitValueType&& InValue)
+	{
+		const FSetElementId PairId = Pairs.EmplaceByHash(KeyHash, TPairInitializer<InitKeyType&&, InitValueType&&>(Forward<InitKeyType>(InKey), Forward<InitValueType>(InValue)));
+
+		return Pairs[PairId].Value;
+	}
+
+	/**
+	 * Set a default value associated with a key.
+	 *
+	 * @param InKey The key to associate the value with.
+	 * @return A reference to the value as stored in the map. The reference is only valid until the next change to any key in the map.
+	 */
+	template <typename InitKeyType = KeyType>
+	ValueType& Emplace(InitKeyType&& InKey)
+	{
+		const FSetElementId PairId = Pairs.Emplace(TKeyInitializer<InitKeyType&&>(Forward<InitKeyType>(InKey)));
+
+		return Pairs[PairId].Value;
+	}
+
+	/** See Emplace() and class documentation section on ByHash() functions */
+	template <typename InitKeyType = KeyType>
+	ValueType& EmplaceByHash(uint32 KeyHash, InitKeyType&& InKey)
+	{
+		const FSetElementId PairId = Pairs.EmplaceByHash(KeyHash, TKeyInitializer<InitKeyType&&>(Forward<InitKeyType>(InKey)));
+
+		return Pairs[PairId].Value;
+	}
+
+	/**
+	 * Remove all value associations for a key.
+	 *
+	 * @param InKey The key to remove associated values for.
+	 * @return The number of values that were associated with the key.
+	 */
+	inline int32 Remove(KeyConstPointerType InKey)
+	{
+		const int32 NumRemovedPairs = Pairs.Remove(InKey);
+		return NumRemovedPairs;
+	}
+
+	/** See Remove() and class documentation section on ByHash() functions */
+	template<typename ComparableKey>
+	inline int32 RemoveByHash(uint32 KeyHash, const ComparableKey& Key)
+	{
+		const int32 NumRemovedPairs = Pairs.RemoveByHash(KeyHash, Key);
+		return NumRemovedPairs;
+	}
+
+	/**
+	 * Removes the element at the specified index. The caller has to ensure that the index is valid.
+	 * @param Id The index of the element to remove.
+	 */
+	UE_FORCEINLINE_HINT void Remove(FSetElementId Id)
+	{
+		Pairs.Remove(Id);
+	}
+
+	/**
+	 * Find the key associated with the specified value.
+	 *
+	 * The time taken is O(N) in the number of pairs.
+	 *
+	 * @param Value The value to search for
+	 * @return A pointer to the key associated with the specified value,
+	 *     or nullptr if the value isn't contained in this map. The pointer
+	 *     is only valid until the next change to any key in the map.
+	 */
+	[[nodiscard]] const KeyType* FindKey(ValueInitType Value) const
+	{
+		for (typename ElementSetType::TConstIterator PairIt(Pairs); PairIt; ++PairIt)
+		{
+			if (PairIt->Value == Value)
+			{
+				return &PairIt->Key;
+			}
+		}
+		return nullptr;
+	}
+
+	/**
+	 * Filters the elements in the map based on a predicate functor.
+	 *
+	 * @param Pred The functor to apply to each element.
+	 * @returns TMap with the same type as this object which contains
+	 *          the subset of elements for which the functor returns true.
+	 */
+	template <typename Predicate>
+	[[nodiscard]] TMap<KeyType, ValueType, SetAllocator, KeyFuncs> FilterByPredicate(Predicate Pred) const
+	{
+		TMap<KeyType, ValueType, SetAllocator, KeyFuncs> FilterResults;
+		FilterResults.Reserve(Pairs.Num());
+		for (const ElementType& Pair : Pairs)
+		{
+			if (Pred(Pair))
+			{
+				FilterResults.Add(Pair);
+			}
+		}
+		return FilterResults;
+	}
+
+	/**
+	 * Find the value associated with a specified key.
+	 *
+	 * @param Key The key to search for.
+	 * @return A pointer to the value associated with the specified key, or nullptr if the key isn't contained in this map.  The pointer
+	 *			is only valid until the next change to any key in the map.
+	 */
+	[[nodiscard]] inline ValueType* Find(KeyConstPointerType Key)
+	{
+		if (auto* Pair = Pairs.Find(Key))
+		{
+			return &Pair->Value;
+		}
+
+		return nullptr;
+	}
+	[[nodiscard]] UE_FORCEINLINE_HINT const ValueType* Find(KeyConstPointerType Key) const
+	{
+		return const_cast<TMapBase*>(this)->Find(Key);
+	}
+
+	/** See Find() and class documentation section on ByHash() functions */
+	template<typename ComparableKey>
+	[[nodiscard]] inline ValueType* FindByHash(uint32 KeyHash, const ComparableKey& Key)
+	{
+		if (auto* Pair = Pairs.FindByHash(KeyHash, Key))
+		{
+			return &Pair->Value;
+		}
+
+		return nullptr;
+	}
+	template<typename ComparableKey>
+	[[nodiscard]] UE_FORCEINLINE_HINT const ValueType* FindByHash(uint32 KeyHash, const ComparableKey& Key) const
+	{
+		return const_cast<TMapBase*>(this)->FindByHash(KeyHash, Key);
+	}
+
+	template<typename ComparableKey>
+	[[nodiscard]] inline ValueType& FindByHashChecked(uint32 KeyHash, const ComparableKey& Key)
+	{
+		auto* Pair = Pairs.FindByHash(KeyHash, Key);
+		check(Pair != nullptr);
+		return Pair->Value;
+	}
+	template<typename ComparableKey>
+	[[nodiscard]] UE_FORCEINLINE_HINT const ValueType& FindByHashChecked(uint32 KeyHash, const ComparableKey& Key) const
+	{
+		return const_cast<TMapBase*>(this)->FindByHashChecked(KeyHash, Key);
+	}
+
+	/**
+	 * Finds the index of the first element that is assigned to the specified key.
+	 * The returned index is only valid until the map is changed.
+	 *
+	 * @param Key The key to search for.
+	 */
+	UE_FORCEINLINE_HINT FSetElementId FindId(KeyInitType Key) const
+	{
+		return Pairs.FindId(Key);
+	}
+
+	/** See FindId() and class documentation section on ByHash() functions */
+	template<typename ComparableKey>
+	UE_FORCEINLINE_HINT FSetElementId FindIdByHash(uint32 KeyHash, const ComparableKey& Key) const
+	{
+		return Pairs.FindIdByHash(KeyHash, Key);
+	}
+
+private:
+	[[nodiscard]] UE_FORCEINLINE_HINT static uint32 HashKey(const KeyType& Key)
+	{
+		return KeyFuncs::GetKeyHash(Key);
+	}
+
+	/**
+	 * Find the value associated with a specified key, or if none exists, 
+	 * adds a value using the default constructor.
+	 *
+	 * @param Key The key to search for.
+	 * @return A reference to the value associated with the specified key.
+	 */
+	template <typename InitKeyType>
+	[[nodiscard]] ValueType& FindOrAddImpl(uint32 KeyHash, InitKeyType&& Key)
+	{
+		if (auto* Pair = Pairs.FindByHash(KeyHash, Key))
+		{
+			return Pair->Value;
+		}
+
+		return AddByHash(KeyHash, Forward<InitKeyType>(Key));
+	}
+
+	/**
+	 * Find the value associated with a specified key, or if none exists,
+	 * adds the value
+	 *
+	 * @param Key The key to search for.
+	 * @param Value The value to associate with the key.
+	 * @return A reference to the value associated with the specified key.
+	 */
+	template <typename InitKeyType, typename InitValueType>
+	[[nodiscard]] ValueType& FindOrAddImpl(uint32 KeyHash, InitKeyType&& Key, InitValueType&& Value)
+	{
+		if (auto* Pair = Pairs.FindByHash(KeyHash, Key))
+		{
+			return Pair->Value;
+		}
+
+		return AddByHash(KeyHash, Forward<InitKeyType>(Key), Forward<InitValueType>(Value));
+	}
+
+public:
+
+	/**
+	 * Find the value associated with a specified key, or if none exists, 
+	 * adds a value using the default constructor.
+	 *
+	 * @param Key The key to search for.
+	 * @return A reference to the value associated with the specified key.
+	 */
+	UE_FORCEINLINE_HINT ValueType& FindOrAdd(const KeyType&  Key) { return FindOrAddImpl(HashKey(Key),					  Key); }
+	UE_FORCEINLINE_HINT ValueType& FindOrAdd(      KeyType&& Key) { return FindOrAddImpl(HashKey(Key), MoveTempIfPossible(Key)); }
+
+	/** See FindOrAdd() and class documentation section on ByHash() functions */
+	UE_FORCEINLINE_HINT ValueType& FindOrAddByHash(uint32 KeyHash, const KeyType&  Key) { return FindOrAddImpl(KeyHash,                    Key); }
+	UE_FORCEINLINE_HINT ValueType& FindOrAddByHash(uint32 KeyHash,       KeyType&& Key) { return FindOrAddImpl(KeyHash, MoveTempIfPossible(Key)); }
+
+	/**
+	 * Find the value associated with a specified key, or if none exists, 
+	 * adds a value using the default constructor.
+	 *
+	 * @param Key The key to search for.
+	 * @param Value The value to associate with the key.
+	 * @return A reference to the value associated with the specified key.
+	 */
+	UE_FORCEINLINE_HINT ValueType& FindOrAdd(const KeyType&  Key, const ValueType&  Value) { return FindOrAddImpl(HashKey(Key),						Key,                    Value  ); }
+	UE_FORCEINLINE_HINT ValueType& FindOrAdd(const KeyType&  Key, ValueType&&       Value) { return FindOrAddImpl(HashKey(Key),						Key, MoveTempIfPossible(Value) ); }
+	UE_FORCEINLINE_HINT ValueType& FindOrAdd(      KeyType&& Key, const ValueType&  Value) { return FindOrAddImpl(HashKey(Key), MoveTempIfPossible(Key),                    Value  ); }
+	UE_FORCEINLINE_HINT ValueType& FindOrAdd(      KeyType&& Key, ValueType&&       Value) { return FindOrAddImpl(HashKey(Key), MoveTempIfPossible(Key), MoveTempIfPossible(Value) ); }
+
+	/** See FindOrAdd() and class documentation section on ByHash() functions */
+	UE_FORCEINLINE_HINT ValueType& FindOrAddByHash(uint32 KeyHash, const KeyType&  Key, const ValueType&  Value) { return FindOrAddImpl(KeyHash,                     Key,                    Value); }
+	UE_FORCEINLINE_HINT ValueType& FindOrAddByHash(uint32 KeyHash, const KeyType&  Key,       ValueType&& Value) { return FindOrAddImpl(KeyHash,                     Key, MoveTempIfPossible(Value)); }
+	UE_FORCEINLINE_HINT ValueType& FindOrAddByHash(uint32 KeyHash,       KeyType&& Key, const ValueType&  Value) { return FindOrAddImpl(KeyHash, MoveTempIfPossible(Key),                    Value); }
+	UE_FORCEINLINE_HINT ValueType& FindOrAddByHash(uint32 KeyHash,       KeyType&& Key,       ValueType&& Value) { return FindOrAddImpl(KeyHash, MoveTempIfPossible(Key), MoveTempIfPossible(Value)); }
+
+	/**
+	 * Find a reference to the value associated with a specified key.
+	 *
+	 * @param Key The key to search for.
+	 * @return The value associated with the specified key, or triggers an assertion if the key does not exist.
+	 */
+	[[nodiscard]] inline const ValueType& FindChecked(KeyConstPointerType Key) const
+	{
+		const auto* Pair = Pairs.Find(Key);
+		check(Pair != nullptr);
+		return Pair->Value;
+	}
+
+	/**
+	 * Find a reference to the value associated with a specified key.
+	 *
+	 * @param Key The key to search for.
+	 * @return The value associated with the specified key, or triggers an assertion if the key does not exist.
+	 */
+	[[nodiscard]] inline ValueType& FindChecked(KeyConstPointerType Key)
+	{
+		auto* Pair = Pairs.Find(Key);
+		check(Pair != nullptr);
+		return Pair->Value;
+	}
+
+	/**
+	 * Find the value associated with a specified key.
+	 *
+	 * @param Key The key to search for.
+	 * @return The value associated with the specified key, or the default value for the ValueType if the key isn't contained in this map.
+	 */
+	[[nodiscard]] inline ValueType FindRef(KeyConstPointerType Key) const
+	{
+		if (const auto* Pair = Pairs.Find(Key))
+		{
+			return Pair->Value;
+		}
+
+		return ValueType();
+	}
+
+	/**
+	 * Find the value associated with a specified key.
+	 *
+	 * @param Key The key to search for.
+	 * @param DefaultValue The fallback value if the key is not found.
+	 * @return The value associated with the specified key, or DefaultValue if the key isn't contained in this map.
+	 */
+	[[nodiscard]] inline ValueType FindRef(KeyConstPointerType Key, ValueType DefaultValue) const
+	{
+		if (const auto* Pair = Pairs.Find(Key))
+		{
+			return Pair->Value;
+		}
+
+		return DefaultValue;
+	}
+
+	/**
+	 * Finds any pair in the map and returns a pointer to it.
+	 * Callers should not depend on particular patterns in the behaviour of this function.
+	 * @return A pointer to an arbitrary pair, or nullptr if the container is empty.
+	 */
+	[[nodiscard]] ElementType* FindArbitraryElement()
+	{
+		// The goal of this function is to be fast, and so the implementation may be improved at any time even if it gives different results.
+
+		return Pairs.FindArbitraryElement();
+	}
+	[[nodiscard]] const ElementType* FindArbitraryElement() const
+	{
+		return const_cast<TMapBase*>(this)->FindArbitraryElement();
+	}
+
+	/**
+	 * Check if map contains the specified key.
+	 *
+	 * @param Key The key to check for.
+	 * @return true if the map contains the key.
+	 */
+	[[nodiscard]] UE_FORCEINLINE_HINT bool Contains(KeyConstPointerType Key) const
+	{
+		return Pairs.Contains(Key);
+	}
+
+	/** See Contains() and class documentation section on ByHash() functions */
+	template<typename ComparableKey>
+	[[nodiscard]] UE_FORCEINLINE_HINT bool ContainsByHash(uint32 KeyHash, const ComparableKey& Key) const
+	{
+		return Pairs.ContainsByHash(KeyHash, Key);
+	}
+
+	/** Copy the key/value pairs in this map into an array. */
+	[[nodiscard]] TArray<ElementType> Array() const
+	{
+		return Pairs.Array();
+	}
+
+	/**
+	 * Generate an array from the keys in this map.
+	 *
+	 * @param OutArray Will contain the collection of keys.
+	 */
+	template<typename Allocator> void GenerateKeyArray(TArray<KeyType, Allocator>& OutArray) const
+	{
+		OutArray.Empty(Pairs.Num());
+		for (typename ElementSetType::TConstIterator PairIt(Pairs); PairIt; ++PairIt)
+		{
+			OutArray.Add(PairIt->Key);
+		}
+	}
+
+	/**
+	 * Generate an array from the values in this map.
+	 *
+	 * @param OutArray Will contain the collection of values.
+	 */
+	template<typename Allocator> void GenerateValueArray(TArray<ValueType, Allocator>& OutArray) const
+	{
+		OutArray.Empty(Pairs.Num());
+		for (typename ElementSetType::TConstIterator PairIt(Pairs); PairIt; ++PairIt)
+		{
+			OutArray.Add(PairIt->Value);
+		}
+	}
+
+	/**
+	 * Describes the map's contents through an output device.
+	 *
+	 * @param Ar The output device to describe the map's contents through.
+	 */
+	void Dump(FOutputDevice& Ar)
+	{
+		Pairs.Dump(Ar);
+	}
+
+protected:
+	typedef TSet<ElementType, KeyFuncs, SetAllocator> ElementSetType;
+
+	/** The base of TMapBase iterators. */
+	template<bool bConst, bool bRangedFor = false>
+	class TBaseIterator
+	{
+	public:
+		typedef std::conditional_t<
+			bConst,
+			std::conditional_t<bRangedFor, typename ElementSetType::TRangedForConstIterator, typename ElementSetType::TConstIterator>,
+			std::conditional_t<bRangedFor, typename ElementSetType::TRangedForIterator, typename ElementSetType::TIterator     >
+		> PairItType;
+	private:
+		typedef std::conditional_t<bConst, const TMapBase, TMapBase> MapType;
+		typedef std::conditional_t<bConst, const KeyType, KeyType> ItKeyType;
+		typedef std::conditional_t<bConst, const ValueType, ValueType> ItValueType;
+		typedef std::conditional_t<bConst, const typename ElementSetType::ElementType, typename ElementSetType::ElementType> PairType;
+
+	public:
+		[[nodiscard]] UE_FORCEINLINE_HINT TBaseIterator(const PairItType& InElementIt)
+			: PairIt(InElementIt)
+		{
+		}
+
+		UE_FORCEINLINE_HINT TBaseIterator& operator++()
+		{
+			++PairIt;
+			return *this;
+		}
+
+		/** conversion to "bool" returning true if the iterator is valid. */
+		[[nodiscard]] UE_FORCEINLINE_HINT explicit operator bool() const
+		{ 
+			return !!PairIt; 
+		}
+		/** inverse of the "bool" operator */
+		[[nodiscard]] UE_FORCEINLINE_HINT bool operator !() const
+		{
+			return !(bool)*this;
+		}
+
+		[[nodiscard]] UE_FORCEINLINE_HINT bool operator==(const TBaseIterator& Rhs) const
+		{
+			return PairIt == Rhs.PairIt;
+		}
+		[[nodiscard]] UE_FORCEINLINE_HINT bool operator!=(const TBaseIterator& Rhs) const
+		{
+			return PairIt != Rhs.PairIt;
+		}
+
+		[[nodiscard]] UE_FORCEINLINE_HINT ItKeyType& Key() const
+		{
+			return PairIt->Key;
+		}
+		[[nodiscard]] UE_FORCEINLINE_HINT ItValueType& Value() const
+		{
+			return PairIt->Value;
+		}
+
+		[[nodiscard]] UE_FORCEINLINE_HINT FSetElementId GetId() const
+		{
+			return PairIt.GetId();
+		}
+
+		[[nodiscard]] UE_FORCEINLINE_HINT PairType& operator*() const
+		{
+			return *PairIt;
+		}
+		[[nodiscard]] UE_FORCEINLINE_HINT PairType* operator->() const
+		{
+			return &*PairIt;
+		}
+
+	protected:
+		PairItType PairIt;
+	};
+
+	/** The base type of iterators that iterate over the values associated with a specified key. */
+	template<bool bConst>
+	class TBaseKeyIterator
+	{
+	private:
+		typedef std::conditional_t<bConst, typename ElementSetType::TConstKeyIterator, typename ElementSetType::TKeyIterator> SetItType;
+		typedef std::conditional_t<bConst, const KeyType, KeyType> ItKeyType;
+		typedef std::conditional_t<bConst, const ValueType, ValueType> ItValueType;
+
+	public:
+		/** Initialization constructor. */
+		[[nodiscard]] UE_FORCEINLINE_HINT TBaseKeyIterator(const SetItType& InSetIt)
+			: SetIt(InSetIt)
+		{
+		}
+
+		UE_FORCEINLINE_HINT TBaseKeyIterator& operator++()
+		{
+			++SetIt;
+			return *this;
+		}
+
+		/** conversion to "bool" returning true if the iterator is valid. */
+		[[nodiscard]] UE_FORCEINLINE_HINT explicit operator bool() const
+		{
+			return !!SetIt; 
+		}
+		/** inverse of the "bool" operator */
+		[[nodiscard]] UE_FORCEINLINE_HINT bool operator !() const 
+		{
+			return !(bool)*this;
+		}
+
+		[[nodiscard]] UE_FORCEINLINE_HINT FSetElementId GetId() const
+		{
+			return SetIt.GetId();
+		}
+		[[nodiscard]] UE_FORCEINLINE_HINT ItKeyType& Key() const
+		{
+			return SetIt->Key;
+		}
+		[[nodiscard]] UE_FORCEINLINE_HINT ItValueType& Value() const
+		{
+			return SetIt->Value;
+		}
+
+		[[nodiscard]] UE_FORCEINLINE_HINT decltype(auto) operator*() const
+		{
+			return SetIt.operator*();
+		}
+		[[nodiscard]] UE_FORCEINLINE_HINT decltype(auto) operator->() const
+		{
+			return SetIt.operator->();
+		}
+
+	protected:
+		SetItType SetIt;
+	};
+
+	/** A set of the key-value pairs in the map. */
+	ElementSetType Pairs;
+
+public:
+	void WriteMemoryImage(FMemoryImageWriter& Writer) const
+	{
+		Pairs.WriteMemoryImage(Writer);
+	}
+
+	void CopyUnfrozen(const FMemoryUnfreezeContent& Context, void* Dst) const
+	{
+		Pairs.CopyUnfrozen(Context, Dst);
+	}
+
+	static void AppendHash(const FPlatformTypeLayoutParameters& LayoutParams, FSHA1& Hasher)
+	{
+		ElementSetType::AppendHash(LayoutParams, Hasher);
+	}
+
+	/** Map iterator. */
+	class TIterator : public TBaseIterator<false>
+	{
+	public:
+
+		/** Initialization constructor. */
+		[[nodiscard]] inline TIterator(TMapBase& InMap, bool bInRequiresRehashOnRemoval = false)
+			: TBaseIterator<false>(InMap.Pairs.CreateIterator())
+			, Map(InMap)
+			, bElementsHaveBeenRemoved(false)
+			, bRequiresRehashOnRemoval(bInRequiresRehashOnRemoval)
+		{
+		}
+
+		/** Destructor. */
+		inline ~TIterator()
+		{
+			if (bElementsHaveBeenRemoved && bRequiresRehashOnRemoval)
+			{
+				Map.Pairs.Relax();
+			}
+		}
+
+		/** Removes the current pair from the map without losing the iteration
+		 * position. Increment before using the iterator again, but after that
+		 * it will point at the element that was after the removed element.
+		 */
+		inline void RemoveCurrent()
+		{
+			TBaseIterator<false>::PairIt.RemoveCurrent();
+			bElementsHaveBeenRemoved = true;
+		}
+
+	private:
+		TMapBase& Map;
+		bool      bElementsHaveBeenRemoved;
+		bool      bRequiresRehashOnRemoval;
+	};
+
+	/** Const map iterator. */
+	class TConstIterator : public TBaseIterator<true>
+	{
+	public:
+		[[nodiscard]] UE_FORCEINLINE_HINT TConstIterator(const TMapBase& InMap)
+			: TBaseIterator<true>(InMap.Pairs.CreateConstIterator())
+		{
+		}
+	};
+
+	using TRangedForIterator = TBaseIterator<false, true>;
+	using TRangedForConstIterator = TBaseIterator<true, true>;
+
+	/** Iterates over values associated with a specified key in a const map. */
+	class TConstKeyIterator : public TBaseKeyIterator<true>
+	{
+	private:
+		using Super        = TBaseKeyIterator<true>;
+		using IteratorType = typename ElementSetType::TConstKeyIterator;
+
+	public:
+		using KeyArgumentType = typename IteratorType::KeyArgumentType;
+
+		[[nodiscard]] UE_FORCEINLINE_HINT TConstKeyIterator(const TMapBase& InMap, KeyArgumentType InKey)
+			: Super(IteratorType(InMap.Pairs, InKey))
+		{
+		}
+	};
+
+	/** Iterates over values associated with a specified key in a map. */
+	class TKeyIterator : public TBaseKeyIterator<false>
+	{
+	private:
+		using Super        = TBaseKeyIterator<false>;
+		using IteratorType = typename ElementSetType::TKeyIterator;
+
+	public:
+		using KeyArgumentType = typename IteratorType::KeyArgumentType;
+
+		[[nodiscard]] UE_FORCEINLINE_HINT TKeyIterator(TMapBase& InMap, KeyArgumentType InKey)
+			: Super(IteratorType(InMap.Pairs, InKey))
+		{
+		}
+
+		/** Removes the current key-value pair from the map. */
+		UE_FORCEINLINE_HINT void RemoveCurrent()
+		{
+			TBaseKeyIterator<false>::SetIt.RemoveCurrent();
+		}
+	};
+
+	/** Creates an iterator over all the pairs in this map */
+	[[nodiscard]] UE_FORCEINLINE_HINT TIterator CreateIterator()
+	{
+		return TIterator(*this);
+	}
+
+	/** Creates a const iterator over all the pairs in this map */
+	[[nodiscard]] UE_FORCEINLINE_HINT TConstIterator CreateConstIterator() const
+	{
+		return TConstIterator(*this);
+	}
+
+	/** Creates an iterator over the values associated with a specified key in a map */
+	[[nodiscard]] UE_FORCEINLINE_HINT TKeyIterator CreateKeyIterator(typename TKeyIterator::KeyArgumentType InKey)
+	{
+		return TKeyIterator(*this, InKey);
+	}
+
+	/** Creates a const iterator over the values associated with a specified key in a map */
+	[[nodiscard]] UE_FORCEINLINE_HINT TConstKeyIterator CreateConstKeyIterator(typename TConstKeyIterator::KeyArgumentType InKey) const
+	{
+		return TConstKeyIterator(*this, InKey);
+	}
+
+	friend struct TMapPrivateFriend;
+
+public:
+	/**
+	 * DO NOT USE DIRECTLY
+	 * STL-like iterators to enable range-based for loop support.
+	 */
+	[[nodiscard]] UE_FORCEINLINE_HINT TRangedForIterator      begin() { return TRangedForIterator(Pairs.begin()); }
+	[[nodiscard]] UE_FORCEINLINE_HINT TRangedForConstIterator begin() const { return TRangedForConstIterator(Pairs.begin()); }
+	[[nodiscard]] UE_FORCEINLINE_HINT TRangedForIterator      end() { return TRangedForIterator(Pairs.end()); }
+	[[nodiscard]] UE_FORCEINLINE_HINT TRangedForConstIterator end() const { return TRangedForConstIterator(Pairs.end()); }
+
+	// Maps are deliberately prevented from being hashed or compared, because this would hide potentially major performance problems behind default operations.
+	friend uint32 GetTypeHash(const TMapBase& Map) = delete;
+	friend bool operator==(const TMapBase&, const TMapBase&) = delete;
+	friend bool operator!=(const TMapBase&, const TMapBase&) = delete;
+};
+
+/** The base type of sortable maps. */
+template <typename KeyType, typename ValueType, typename SetAllocator, typename KeyFuncs>
+class TSortableMapBase : public TMapBase<KeyType, ValueType, SetAllocator, KeyFuncs>
+{
+protected:
+	typedef TMapBase<KeyType, ValueType, SetAllocator, KeyFuncs> Super;
+
+	[[nodiscard]] TSortableMapBase() = default;
+	[[nodiscard]] explicit consteval TSortableMapBase(EConstEval)
+	: Super(ConstEval)
+	{
+	}
+	[[nodiscard]] TSortableMapBase(TSortableMapBase&&) = default;
+	[[nodiscard]] TSortableMapBase(const TSortableMapBase&) = default;
+	TSortableMapBase& operator=(TSortableMapBase&&) = default;
+	TSortableMapBase& operator=(const TSortableMapBase&) = default;
+
+	/** Constructor for moving elements from a TMap with a different SetAllocator */
+	template<typename OtherSetAllocator>
+	[[nodiscard]] TSortableMapBase(TSortableMapBase<KeyType, ValueType, OtherSetAllocator, KeyFuncs>&& Other)
+		: Super(MoveTemp(Other))
+	{
+	}
+
+	/** Constructor for copying elements from a TMap with a different SetAllocator */
+	template<typename OtherSetAllocator>
+	[[nodiscard]] TSortableMapBase(const TSortableMapBase<KeyType, ValueType, OtherSetAllocator, KeyFuncs>& Other)
+		: Super(Other)
+	{
+	}
+
+	/** Assignment operator for moving elements from a TMap with a different SetAllocator */
+	template<typename OtherSetAllocator>
+	TSortableMapBase& operator=(TSortableMapBase<KeyType, ValueType, OtherSetAllocator, KeyFuncs>&& Other)
+	{
+		(Super&)*this = MoveTemp(Other);
+		return *this;
+	}
+
+	/** Assignment operator for copying elements from a TMap with a different SetAllocator */
+	template<typename OtherSetAllocator>
+	TSortableMapBase& operator=(const TSortableMapBase<KeyType, ValueType, OtherSetAllocator, KeyFuncs>& Other)
+	{
+		(Super&)*this = Other;
+		return *this;
+	}
+
+	/////////////////////////////////////////////////////////
+	// Start - intrusive TOptional<TSortableMapBase> state //
+	/////////////////////////////////////////////////////////
+	constexpr static bool bHasIntrusiveUnsetOptionalState = true;
+	using IntrusiveUnsetOptionalStateType = TSortableMapBase;
+
+	[[nodiscard]] explicit TSortableMapBase(FIntrusiveUnsetOptionalState Tag)
+		: Super(Tag)
+	{
+	}
+	[[nodiscard]] bool operator==(FIntrusiveUnsetOptionalState Tag) const
+	{
+		return Super::operator==(Tag);
+	}
+	///////////////////////////////////////////////////////
+	// End - intrusive TOptional<TSortableMapBase> state //
+	///////////////////////////////////////////////////////
+
+public:
+	/**
+	 * Sorts the pairs array using each pair's Key as the sort criteria, then rebuilds the map's hash.
+	 * Invoked using "MyMapVar.KeySort( PREDICATE_CLASS() );"
+	 */
+	template<typename PREDICATE_CLASS>
+	UE_FORCEINLINE_HINT void KeySort(const PREDICATE_CLASS& Predicate)
+	{
+		Super::Pairs.Sort(FKeyComparisonClass<PREDICATE_CLASS>(Predicate));
+	}
+
+	/**
+	 * Stable sorts the pairs array using each pair's Key as the sort criteria, then rebuilds the map's hash.
+	 * Invoked using "MyMapVar.KeySort( PREDICATE_CLASS() );"
+	 */
+	template<typename PREDICATE_CLASS>
+	UE_FORCEINLINE_HINT void KeyStableSort(const PREDICATE_CLASS& Predicate)
+	{
+		Super::Pairs.StableSort(FKeyComparisonClass<PREDICATE_CLASS>(Predicate));
+	}
+
+	/**
+	 * Sorts the pairs array using each pair's Value as the sort criteria, then rebuilds the map's hash.
+	 * Invoked using "MyMapVar.ValueSort( PREDICATE_CLASS() );"
+	 */
+	template<typename PREDICATE_CLASS>
+	UE_FORCEINLINE_HINT void ValueSort(const PREDICATE_CLASS& Predicate)
+	{
+		Super::Pairs.Sort(FValueComparisonClass<PREDICATE_CLASS>(Predicate));
+	}
+
+	/**
+	 * Stable sorts the pairs array using each pair's Value as the sort criteria, then rebuilds the map's hash.
+	 * Invoked using "MyMapVar.ValueSort( PREDICATE_CLASS() );"
+	 */
+	template<typename PREDICATE_CLASS>
+	UE_FORCEINLINE_HINT void ValueStableSort(const PREDICATE_CLASS& Predicate)
+	{
+		Super::Pairs.StableSort(FValueComparisonClass<PREDICATE_CLASS>(Predicate));
+	}
+
+	/**
+	 * Sort the free element list so that subsequent additions will occur in the lowest available
+	 * TSet index resulting in tighter packing without moving any existing items. Also useful for
+	 * some types of determinism. @see TSparseArray::SortFreeList() for more info.
+	 */
+	void SortFreeList()
+	{
+		Super::Pairs.SortFreeList();
+	}
+
+private:
+
+	/** Extracts the pair's key from the map's pair structure and passes it to the user provided comparison class. */
+	template<typename PREDICATE_CLASS>
+	class FKeyComparisonClass
+	{
+		TDereferenceWrapper< KeyType, PREDICATE_CLASS> Predicate;
+
+	public:
+
+		[[nodiscard]] UE_FORCEINLINE_HINT FKeyComparisonClass(const PREDICATE_CLASS& InPredicate)
+			: Predicate(InPredicate)
+		{
+		}
+
+		[[nodiscard]] inline bool operator()(const typename Super::ElementType& A, const typename Super::ElementType& B) const
+		{
+			return Predicate(A.Key, B.Key);
+		}
+	};
+
+	/** Extracts the pair's value from the map's pair structure and passes it to the user provided comparison class. */
+	template<typename PREDICATE_CLASS>
+	class FValueComparisonClass
+	{
+		TDereferenceWrapper< ValueType, PREDICATE_CLASS> Predicate;
+
+	public:
+
+		[[nodiscard]] UE_FORCEINLINE_HINT FValueComparisonClass(const PREDICATE_CLASS& InPredicate)
+			: Predicate(InPredicate)
+		{
+		}
+
+		[[nodiscard]] inline bool operator()(const typename Super::ElementType& A, const typename Super::ElementType& B) const
+		{
+			return Predicate(A.Value, B.Value);
+		}
+	};
+};
+
 template <typename AllocatorType, typename InDerivedType = void>
 class TScriptMap;
 
+/** A TMapBase specialization that only allows a single value associated with each key.*/
+template<typename InKeyType, typename InValueType, typename SetAllocator /*= FDefaultSetAllocator*/, typename KeyFuncs /*= TDefaultMapHashableKeyFuncs<KeyType,ValueType,false>*/>
+class TMap : public TSortableMapBase<InKeyType, InValueType, SetAllocator, KeyFuncs>
+{
+	template <typename, typename>
+	friend class TScriptMap;
+
+	static_assert(!KeyFuncs::bAllowDuplicateKeys, "TMap cannot be instantiated with a KeyFuncs which allows duplicate keys");
+
+public:
+	typedef InKeyType      KeyType;
+	typedef InValueType    ValueType;
+	typedef SetAllocator   SetAllocatorType;
+	typedef KeyFuncs       KeyFuncsType;
+
+	typedef TSortableMapBase<KeyType, ValueType, SetAllocator, KeyFuncs> Super;
+	typedef typename Super::KeyInitType KeyInitType;
+	typedef typename Super::KeyConstPointerType KeyConstPointerType;
+
+	[[nodiscard]] TMap() = default;
+	[[nodiscard]] TMap(TMap&&) = default;
+	[[nodiscard]] TMap(const TMap&) = default;
+	TMap& operator=(TMap&&) = default;
+	TMap& operator=(const TMap&) = default;
+
+	[[nodiscard]] explicit consteval TMap(EConstEval)
+	: Super(ConstEval)
+	{
+	}
+
+	/** Constructor for moving elements from a TMap with a different SetAllocator */
+	template<typename OtherSetAllocator>
+	[[nodiscard]] TMap(TMap<KeyType, ValueType, OtherSetAllocator, KeyFuncs>&& Other)
+		: Super(MoveTemp(Other))
+	{
+	}
+
+	/** Constructor for copying elements from a TMap with a different SetAllocator */
+	template<typename OtherSetAllocator>
+	[[nodiscard]] TMap(const TMap<KeyType, ValueType, OtherSetAllocator, KeyFuncs>& Other)
+		: Super(Other)
+	{
+	}
+
+	/** Constructor which gets its elements from a native initializer list */
+	[[nodiscard]] TMap(std::initializer_list<TPairInitializer<const KeyType&, const ValueType&>> InitList)
+	{
+		this->Reserve((int32)InitList.size());
+		for (const TPairInitializer<const KeyType&, const ValueType&>& Element : InitList)
+		{
+			this->Add(Element.Key, Element.Value);
+		}
+	}
+
+	/////////////////////////////////////////////
+	// Start - intrusive TOptional<TMap> state //
+	/////////////////////////////////////////////
+	constexpr static bool bHasIntrusiveUnsetOptionalState = true;
+	using IntrusiveUnsetOptionalStateType = TMap;
+
+	[[nodiscard]] explicit TMap(FIntrusiveUnsetOptionalState Tag)
+		: Super(Tag)
+	{
+	}
+	[[nodiscard]] bool operator==(FIntrusiveUnsetOptionalState Tag) const
+	{
+		return Super::operator==(Tag);
+	}
+	///////////////////////////////////////////
+	// End - intrusive TOptional<TMap> state //
+	///////////////////////////////////////////
+
+	/** Assignment operator for moving elements from a TMap with a different SetAllocator */
+	template<typename OtherSetAllocator>
+	TMap& operator=(TMap<KeyType, ValueType, OtherSetAllocator, KeyFuncs>&& Other)
+	{
+		(Super&)*this = MoveTemp(Other);
+		return *this;
+	}
+
+	/** Assignment operator for copying elements from a TMap with a different SetAllocator */
+	template<typename OtherSetAllocator>
+	TMap& operator=(const TMap<KeyType, ValueType, OtherSetAllocator, KeyFuncs>& Other)
+	{
+		(Super&)*this = Other;
+		return *this;
+	}
+
+	/** Assignment operator which gets its elements from a native initializer list */
+	TMap& operator=(std::initializer_list<TPairInitializer<const KeyType&, const ValueType&>> InitList)
+	{
+		this->Empty((int32)InitList.size());
+		for (const TPairInitializer<const KeyType&, const ValueType&>& Element : InitList)
+		{
+			this->Add(Element.Key, Element.Value);
+		}
+		return *this;
+	}
+
 	/**
-* Traits class which determines whether or not a type is a TMAP.
+	 * Remove the pair with the specified key and copies the value
+	 * that was removed to the ref parameter
+	 *
+	 * @param Key The key to search for
+	 * @param OutRemovedValue If found, the value that was removed (not modified if the key was not found)
+	 * @return whether or not the key was found
 	 */
-template <typename T> struct TIsTMap { enum { Value = false }; };
+	inline bool RemoveAndCopyValue(KeyInitType Key, ValueType& OutRemovedValue)
+	{
+		const FSetElementId PairId = Super::Pairs.FindId(Key);
+		if (!PairId.IsValidId())
+		{
+			return false;
+		}
+
+		OutRemovedValue = MoveTempIfPossible(Super::Pairs[PairId].Value);
+		Super::Pairs.Remove(PairId);
+		return true;
+	}
+	
+	/** See RemoveAndCopyValue() and class documentation section on ByHash() functions */
+	template<typename ComparableKey>
+	inline bool RemoveAndCopyValueByHash(uint32 KeyHash, const ComparableKey& Key, ValueType& OutRemovedValue)
+	{
+		const FSetElementId PairId = Super::Pairs.FindIdByHash(KeyHash, Key);
+		if (!PairId.IsValidId())
+		{
+			return false;
+		}
+
+		OutRemovedValue = MoveTempIfPossible(Super::Pairs[PairId].Value);
+		Super::Pairs.Remove(PairId);
+		return true;
+	}
+
+	/**
+	 * Find a pair with the specified key, removes it from the map, and returns the value part of the pair.
+	 *
+	 * If no pair was found, an exception is thrown.
+	 *
+	 * @param Key the key to search for
+	 * @return whether or not the key was found
+	 */
+	inline ValueType FindAndRemoveChecked(KeyConstPointerType Key)
+	{
+		const FSetElementId PairId = Super::Pairs.FindId(Key);
+		check(PairId.IsValidId());
+		ValueType Result = MoveTempIfPossible(Super::Pairs[PairId].Value);
+		Super::Pairs.Remove(PairId);
+		return Result;
+	}
+
+	/**
+	 * Move all items from another map into our map (if any keys are in both,
+	 * the value from the other map wins) and empty the other map.
+	 *
+	 * @param OtherMap The other map of items to move the elements from.
+	 */
+	template<typename OtherSetAllocator>
+	void Append(TMap<KeyType, ValueType, OtherSetAllocator, KeyFuncs>&& OtherMap)
+	{
+		this->Reserve(this->Num() + OtherMap.Num());
+		for (auto& Pair : OtherMap)
+		{
+			this->Add(MoveTempIfPossible(Pair.Key), MoveTempIfPossible(Pair.Value));
+		}
+
+		OtherMap.Reset();
+	}
+
+	/**
+	 * Add all items from another map to our map (if any keys are in both,
+	 * the value from the other map wins).
+	 *
+	 * @param OtherMap The other map of items to add.
+	 */
+	template<typename OtherSetAllocator>
+	void Append(const TMap<KeyType, ValueType, OtherSetAllocator, KeyFuncs>& OtherMap)
+	{
+		this->Reserve(this->Num() + OtherMap.Num());
+		for (auto& Pair : OtherMap)
+		{
+			this->Add(Pair.Key, Pair.Value);
+		}
+	}
+
+	[[nodiscard]] UE_FORCEINLINE_HINT ValueType& operator[](KeyConstPointerType Key)
+	{
+		return this->FindChecked(Key);
+	}
+	[[nodiscard]] UE_FORCEINLINE_HINT const ValueType& operator[](KeyConstPointerType Key) const
+	{
+		return this->FindChecked(Key);
+	}
+};
+
+namespace Freeze
+{
+	template <typename KeyType, typename ValueType, typename SetAllocator, typename KeyFuncs>
+	void IntrinsicWriteMemoryImage(FMemoryImageWriter& Writer, const TMap<KeyType, ValueType, SetAllocator, KeyFuncs>& Object, const FTypeLayoutDesc&)
+	{
+		Object.WriteMemoryImage(Writer);
+	}
+
+	template <typename KeyType, typename ValueType, typename SetAllocator, typename KeyFuncs>
+	uint32 IntrinsicUnfrozenCopy(const FMemoryUnfreezeContent& Context, const TMap<KeyType, ValueType, SetAllocator, KeyFuncs>& Object, void* OutDst)
+	{
+		Object.CopyUnfrozen(Context, OutDst);
+		return sizeof(Object);
+	}
+
+	template <typename KeyType, typename ValueType, typename SetAllocator, typename KeyFuncs>
+	uint32 IntrinsicAppendHash(const TMap<KeyType, ValueType, SetAllocator, KeyFuncs>* DummyObject, const FTypeLayoutDesc& TypeDesc, const FPlatformTypeLayoutParameters& LayoutParams, FSHA1& Hasher)
+	{
+		TMap<KeyType, ValueType, SetAllocator, KeyFuncs>::AppendHash(LayoutParams, Hasher);
+		return DefaultAppendHash(TypeDesc, LayoutParams, Hasher);
+	}
+}
+
+DECLARE_TEMPLATE_INTRINSIC_TYPE_LAYOUT((template <typename KeyType, typename ValueType, typename SetAllocator, typename KeyFuncs>), (TMap<KeyType, ValueType, SetAllocator, KeyFuncs>));
+
+/** A TMapBase specialization that allows multiple values to be associated with each key. */
+template<typename KeyType, typename ValueType, typename SetAllocator /* = FDefaultSetAllocator */, typename KeyFuncs /*= TDefaultMapHashableKeyFuncs<KeyType,ValueType,true>*/>
+class TMultiMap : public TSortableMapBase<KeyType, ValueType, SetAllocator, KeyFuncs>
+{
+	static_assert(KeyFuncs::bAllowDuplicateKeys, "TMultiMap cannot be instantiated with a KeyFuncs which disallows duplicate keys");
+
+public:
+	typedef TSortableMapBase<KeyType, ValueType, SetAllocator, KeyFuncs> Super;
+	typedef typename Super::KeyConstPointerType KeyConstPointerType;
+	typedef typename Super::KeyInitType KeyInitType;
+	typedef typename Super::ValueInitType ValueInitType;
+
+	[[nodiscard]] TMultiMap() = default;
+	[[nodiscard]] TMultiMap(TMultiMap&&) = default;
+	[[nodiscard]] TMultiMap(const TMultiMap&) = default;
+	TMultiMap& operator=(TMultiMap&&) = default;
+	TMultiMap& operator=(const TMultiMap&) = default;
+
+	/** Constructor for moving elements from a TMap with a different SetAllocator */
+	template<typename OtherSetAllocator>
+	[[nodiscard]] TMultiMap(TMultiMap<KeyType, ValueType, OtherSetAllocator, KeyFuncs>&& Other)
+		: Super(MoveTemp(Other))
+	{
+	}
+
+	/** Constructor for copying elements from a TMap with a different SetAllocator */
+	template<typename OtherSetAllocator>
+	[[nodiscard]] TMultiMap(const TMultiMap<KeyType, ValueType, OtherSetAllocator, KeyFuncs>& Other)
+		: Super(Other)
+	{
+	}
+
+	/** Constructor which gets its elements from a native initializer list */
+	[[nodiscard]] TMultiMap(std::initializer_list<TPairInitializer<const KeyType&, const ValueType&>> InitList)
+	{
+		this->Reserve((int32)InitList.size());
+		for (const TPairInitializer<const KeyType&, const ValueType&>& Element : InitList)
+		{
+			this->Add(Element.Key, Element.Value);
+		}
+	}
+
+	//////////////////////////////////////////////////
+	// Start - intrusive TOptional<TMultiMap> state //
+	//////////////////////////////////////////////////
+	constexpr static bool bHasIntrusiveUnsetOptionalState = true;
+	using IntrusiveUnsetOptionalStateType = TMultiMap;
+
+	[[nodiscard]] explicit TMultiMap(FIntrusiveUnsetOptionalState Tag)
+		: Super(Tag)
+	{
+	}
+	[[nodiscard]] bool operator==(FIntrusiveUnsetOptionalState Tag) const
+	{
+		return Super::operator==(Tag);
+	}
+	////////////////////////////////////////////////
+	// End - intrusive TOptional<TMultiMap> state //
+	////////////////////////////////////////////////
+
+	/** Assignment operator for moving elements from a TMap with a different SetAllocator */
+	template<typename OtherSetAllocator>
+	TMultiMap& operator=(TMultiMap<KeyType, ValueType, OtherSetAllocator, KeyFuncs>&& Other)
+	{
+		(Super&)*this = MoveTemp(Other);
+		return *this;
+	}
+
+	/** Assignment operator for copying elements from a TMap with a different SetAllocator */
+	template<typename OtherSetAllocator>
+	TMultiMap& operator=(const TMultiMap<KeyType, ValueType, OtherSetAllocator, KeyFuncs>& Other)
+	{
+		(Super&)*this = Other;
+		return *this;
+	}
+
+	/** Assignment operator which gets its elements from a native initializer list */
+	TMultiMap& operator=(std::initializer_list<TPairInitializer<const KeyType&, const ValueType&>> InitList)
+	{
+		this->Empty((int32)InitList.size());
+		for (const TPairInitializer<const KeyType&, const ValueType&>& Element : InitList)
+		{
+			this->Add(Element.Key, Element.Value);
+		}
+		return *this;
+	}
+
+	/**
+	 * Finds all values associated with the specified key.
+	 *
+	 * @param Key The key to find associated values for.
+	 * @param OutValues Upon return, contains the values associated with the key.
+	 * @param bMaintainOrder true if the Values array should be in the same order as the map's pairs.
+	 */
+	template<typename Allocator> void MultiFind(KeyInitType Key, TArray<ValueType, Allocator>& OutValues, bool bMaintainOrder = false) const
+	{
+		for (typename Super::ElementSetType::TConstKeyIterator It(Super::Pairs, Key); It; ++It)
+		{
+			OutValues.Add(It->Value);
+		}
+
+		if (bMaintainOrder)
+		{
+			Algo::Reverse(OutValues);
+		}
+	}
+
+	/**
+	 * Finds all values associated with the specified key.
+	 *
+	 * @param Key The key to find associated values for.
+	 * @param OutValues Upon return, contains pointers to the values associated with the key.
+	 *					Pointers are only valid until the next change to any key in the map.
+	 * @param bMaintainOrder true if the Values array should be in the same order as the map's pairs.
+	 */
+	template<typename Allocator> void MultiFindPointer(KeyInitType Key, TArray<const ValueType*, Allocator>& OutValues, bool bMaintainOrder = false) const
+	{
+		for (typename Super::ElementSetType::TConstKeyIterator It(Super::Pairs, Key); It; ++It)
+		{
+			OutValues.Add(&It->Value);
+		}
+
+		if (bMaintainOrder)
+		{
+			Algo::Reverse(OutValues);
+		}
+	}
+	template<typename Allocator> void MultiFindPointer(KeyInitType Key, TArray<ValueType*, Allocator>& OutValues, bool bMaintainOrder = false)
+	{
+		for (typename Super::ElementSetType::TKeyIterator It(Super::Pairs, Key); It; ++It)
+		{
+			OutValues.Add(&It->Value);
+		}
+
+		if (bMaintainOrder)
+		{
+			Algo::Reverse(OutValues);
+		}
+	}
+
+	/**
+	 * Add a key-value association to the map.  The association doesn't replace any of the key's existing associations.
+	 * However, if both the key and value match an existing association in the map, no new association is made and the existing association's
+	 * value is returned.
+	 *
+	 * @param InKey The key to associate.
+	 * @param InValue The value to associate.
+	 * @return A reference to the value as stored in the map; the reference is only valid until the next change to any key in the map.
+	 */
+	UE_FORCEINLINE_HINT ValueType& AddUnique(const KeyType&  InKey, const ValueType&  InValue) { return EmplaceUnique(InKey, InValue); }
+	UE_FORCEINLINE_HINT ValueType& AddUnique(const KeyType&  InKey, ValueType&& InValue) { return EmplaceUnique(InKey, MoveTempIfPossible(InValue)); }
+	UE_FORCEINLINE_HINT ValueType& AddUnique(KeyType&& InKey, const ValueType&  InValue) { return EmplaceUnique(MoveTempIfPossible(InKey), InValue); }
+	UE_FORCEINLINE_HINT ValueType& AddUnique(KeyType&& InKey, ValueType&& InValue) { return EmplaceUnique(MoveTempIfPossible(InKey), MoveTempIfPossible(InValue)); }
+
+	/**
+	 * Add a key-value association to the map.
+	 *
+	 * The association doesn't replace any of the key's existing associations.
+	 * However, if both key and value match an existing association in the map,
+	 * no new association is made and the existing association's value is returned.
+	 *
+	 * @param InKey The key to associate.
+	 * @param InValue The value to associate.
+	 * @return A reference to the value as stored in the map; the reference is only valid until the next change to any key in the map.
+	 */
+	template <typename InitKeyType, typename InitValueType>
+	ValueType& EmplaceUnique(InitKeyType&& InKey, InitValueType&& InValue)
+	{
+		if (ValueType* Found = FindPair(InKey, InValue))
+		{
+			return *Found;
+		}
+
+		// If there's no existing association with the same key and value, create one.
+		return Super::Add(Forward<InitKeyType>(InKey), Forward<InitValueType>(InValue));
+	}
+
+	/**
+	 * Remove all value associations for a key.
+	 *
+	 * @param InKey The key to remove associated values for.
+	 * @return The number of values that were associated with the key.
+	 */
+	UE_FORCEINLINE_HINT int32 Remove(KeyConstPointerType InKey)
+	{
+		return Super::Remove(InKey);
+	}
+
+	/**
+	 * Remove associations between the specified key and value from the map.
+	 *
+	 * @param InKey The key part of the pair to remove.
+	 * @param InValue The value part of the pair to remove.
+	 * @return The number of associations removed.
+	 */
+	int32 Remove(KeyInitType InKey, ValueInitType InValue)
+	{
+		// Iterate over pairs with a matching key.
+		int32 NumRemovedPairs = 0;
+		for (typename Super::ElementSetType::TKeyIterator It(Super::Pairs, InKey); It; ++It)
+		{
+			// If this pair has a matching value as well, remove it.
+			if (It->Value == InValue)
+			{
+				It.RemoveCurrent();
+				++NumRemovedPairs;
+			}
+		}
+		return NumRemovedPairs;
+	}
+
+	/**
+	 * Remove the first association between the specified key and value from the map.
+	 *
+	 * @param InKey The key part of the pair to remove.
+	 * @param InValue The value part of the pair to remove.
+	 * @return The number of associations removed.
+	 */
+	int32 RemoveSingle(KeyInitType InKey, ValueInitType InValue)
+	{
+		// Iterate over pairs with a matching key.
+		int32 NumRemovedPairs = 0;
+		for (typename Super::ElementSetType::TKeyIterator It(Super::Pairs, InKey); It; ++It)
+		{
+			// If this pair has a matching value as well, remove it.
+			if (It->Value == InValue)
+			{
+				It.RemoveCurrent();
+				++NumRemovedPairs;
+
+				// We were asked to remove only the first association, so bail out.
+				break;
+			}
+		}
+		return NumRemovedPairs;
+	}
+
+	/**
+	 * Find an association between a specified key and value. (const)
+	 *
+	 * @param Key The key to find.
+	 * @param Value The value to find.
+	 * @return If the map contains a matching association, a pointer to the value in the map is returned.  Otherwise nullptr is returned.
+	 *			The pointer is only valid as long as the map isn't changed.
+	 */
+	[[nodiscard]] UE_FORCEINLINE_HINT const ValueType* FindPair(KeyInitType Key, ValueInitType Value) const
+	{
+		return const_cast<TMultiMap*>(this)->FindPair(Key, Value);
+	}
+
+	/**
+	 * Find an association between a specified key and value.
+	 *
+	 * @param Key The key to find.
+	 * @param Value The value to find.
+	 * @return If the map contains a matching association, a pointer to the value in the map is returned.  Otherwise nullptr is returned.
+	 *			The pointer is only valid as long as the map isn't changed.
+	 */
+	[[nodiscard]] ValueType* FindPair(KeyInitType Key, ValueInitType Value)
+	{
+		// Iterate over pairs with a matching key.
+		for (typename Super::ElementSetType::TKeyIterator It(Super::Pairs, Key); It; ++It)
+		{
+			// If the pair's value matches, return a pointer to it.
+			if (It->Value == Value)
+			{
+				return &It->Value;
+			}
+		}
+
+		return nullptr;
+	}
+
+	/** Returns the number of values within this map associated with the specified key */
+	[[nodiscard]] int32 Num(KeyInitType Key) const
+	{
+		// Iterate over pairs with a matching key.
+		int32 NumMatchingPairs = 0;
+		for (typename Super::ElementSetType::TConstKeyIterator It(Super::Pairs, Key); It; ++It)
+		{
+			++NumMatchingPairs;
+		}
+		return NumMatchingPairs;
+	}
+
+	// Since we implement an overloaded Num() function in TMultiMap, we need to reimplement TMapBase::Num to make it visible.
+	[[nodiscard]] UE_FORCEINLINE_HINT int32 Num() const
+	{
+		return Super::Num();
+	}
+
+	/**
+	 * Move all items from another map into our map (if any keys are in both,
+	 * the value from the other map wins) and empty the other map.
+	 *
+	 * @param OtherMultiMap The other map of items to move the elements from.
+	 */
+	template<typename OtherSetAllocator>
+	void Append(TMultiMap<KeyType, ValueType, OtherSetAllocator, KeyFuncs>&& OtherMultiMap)
+	{
+		this->Reserve(this->Num() + OtherMultiMap.Num());
+		for (auto& Pair : OtherMultiMap)
+		{
+			this->Add(MoveTempIfPossible(Pair.Key), MoveTempIfPossible(Pair.Value));
+		}
 
-#define UE_TMAP_PREFIX T
-#include "Map.h.inl"
-#undef UE_TMAP_PREFIX
+		OtherMultiMap.Reset();
+	}
+
+	/**
+	 * Add all items from another map to our map (if any keys are in both,
+	 * the value from the other map wins).
+	 *
+	 * @param OtherMultiMap The other map of items to add.
+	 */
+	template<typename OtherSetAllocator>
+	void Append(const TMultiMap<KeyType, ValueType, OtherSetAllocator, KeyFuncs>& OtherMultiMap)
+	{
+		this->Reserve(this->Num() + OtherMultiMap.Num());
+		for (auto& Pair : OtherMultiMap)
+		{
+			this->Add(Pair.Key, Pair.Value);
+		}
+	}
+};
+
+namespace Freeze
+{
+	template <typename KeyType, typename ValueType, typename SetAllocator, typename KeyFuncs>
+	void IntrinsicWriteMemoryImage(FMemoryImageWriter& Writer, const TMultiMap<KeyType, ValueType, SetAllocator, KeyFuncs>& Object, const FTypeLayoutDesc&)
+	{
+		Object.WriteMemoryImage(Writer);
+	}
+
+	template <typename KeyType, typename ValueType, typename SetAllocator, typename KeyFuncs>
+	uint32 IntrinsicUnfrozenCopy(const FMemoryUnfreezeContent& Context, const TMultiMap<KeyType, ValueType, SetAllocator, KeyFuncs>& Object, void* OutDst)
+	{
+		Object.CopyUnfrozen(Context, OutDst);
+		return sizeof(Object);
+	}
+
+	template <typename KeyType, typename ValueType, typename SetAllocator, typename KeyFuncs>
+	uint32 IntrinsicAppendHash(const TMultiMap<KeyType, ValueType, SetAllocator, KeyFuncs>* DummyObject, const FTypeLayoutDesc& TypeDesc, const FPlatformTypeLayoutParameters& LayoutParams, FSHA1& Hasher)
+	{
+		TMultiMap<KeyType, ValueType, SetAllocator, KeyFuncs>::AppendHash(LayoutParams, Hasher);
+		return DefaultAppendHash(TypeDesc, LayoutParams, Hasher);
+	}
+}
 
-#define UE_TMAP_PREFIX TSparse
-#include "Map.h.inl"
-#undef UE_TMAP_PREFIX
+DECLARE_TEMPLATE_INTRINSIC_TYPE_LAYOUT((template <typename KeyType, typename ValueType, typename SetAllocator, typename KeyFuncs>), (TMultiMap<KeyType, ValueType, SetAllocator, KeyFuncs>));
 
-#define UE_TMAP_PREFIX TCompact
-#include "Map.h.inl"
-#undef UE_TMAP_PREFIX
 
 struct FScriptMapLayout
 {
@@ -246,17 +1988,10 @@ public:
 		Pairs.Empty(Slack, Layout.SetLayout);
 	}
 
-#if UE_USE_COMPACT_SET_AS_DEFAULT
-	void RemoveAt(int32 Index, const FScriptMapLayout& Layout, TFunctionRef<uint32 (const void*)> GetKeyHash, TFunctionRef<void (void*)> DestructItem)
-	{
-		Pairs.RemoveAt(Index, Layout.SetLayout, GetKeyHash, DestructItem);
-	}
-#else
 	void RemoveAt(int32 Index, const FScriptMapLayout& Layout)
 	{
 		Pairs.RemoveAt(Index, Layout.SetLayout);
 	}
-#endif
 
 	/**
 	 * Adds an uninitialized object to the map.
@@ -269,21 +2004,6 @@ public:
 		return Pairs.AddUninitialized(Layout.SetLayout);
 	}
 
-	void RemoveAtUninitialized(const FScriptMapLayout& Layout, int32 Index)
-	{
-		Pairs.RemoveAtUninitialized(Layout.SetLayout, Index);
-	}
-
-	void CommitLastUninitialized(const FScriptMapLayout& Layout, TFunctionRef<uint32 (const void*)> GetKeyHash)
-	{
-		Pairs.CommitLastUninitialized(Layout.SetLayout, GetKeyHash);
-	}
-
-	void CommitAllUninitialized(const FScriptMapLayout& Layout, TFunctionRef<uint32 (const void*)> GetKeyHash)
-	{
-		Pairs.CommitAllUninitialized(Layout.SetLayout, GetKeyHash);
-	}
-
 	void Rehash(const FScriptMapLayout& Layout, TFunctionRef<uint32(const void*)> GetKeyHash)
 	{
 		Pairs.Rehash(Layout.SetLayout, GetKeyHash);
@@ -417,6 +2137,16 @@ struct TIsZeroConstructType<TScriptMap<AllocatorType>>
 	enum { Value = true };
 };
 
+/**
+ * Traits class which determines whether or not a type is a TMap.
+ */
+template <typename T> struct TIsTMap { enum { Value = false }; };
+
+template <typename KeyType, typename ValueType, typename SetAllocator, typename KeyFuncs> struct TIsTMap<               TMap<KeyType, ValueType, SetAllocator, KeyFuncs>> { enum { Value = true }; };
+template <typename KeyType, typename ValueType, typename SetAllocator, typename KeyFuncs> struct TIsTMap<const          TMap<KeyType, ValueType, SetAllocator, KeyFuncs>> { enum { Value = true }; };
+template <typename KeyType, typename ValueType, typename SetAllocator, typename KeyFuncs> struct TIsTMap<      volatile TMap<KeyType, ValueType, SetAllocator, KeyFuncs>> { enum { Value = true }; };
+template <typename KeyType, typename ValueType, typename SetAllocator, typename KeyFuncs> struct TIsTMap<const volatile TMap<KeyType, ValueType, SetAllocator, KeyFuncs>> { enum { Value = true }; };
+
 class FScriptMap : public TScriptMap<FDefaultSetAllocator, FScriptMap>
 {
 	using Super = TScriptMap<FDefaultSetAllocator, FScriptMap>;
@@ -432,3 +2162,84 @@ public:
 	// End - intrusive TOptional<FScriptMap> state //
 	/////////////////////////////////////////////////
 };
+
+struct TMapPrivateFriend
+{
+	template <typename KeyType, typename ValueType, typename SetAllocator, typename KeyFuncs>
+	inline static FArchive& Serialize(FArchive& Ar, TMapBase<KeyType, ValueType, SetAllocator, KeyFuncs>& Map)
+	{
+		Ar << Map.Pairs;
+		return Ar;
+	}
+
+	template <typename KeyType, typename ValueType, typename SetAllocator, typename KeyFuncs>
+	UE_FORCEINLINE_HINT static void SerializeStructured(FStructuredArchive::FSlot Slot, TMapBase<KeyType, ValueType, SetAllocator, KeyFuncs>& InMap)
+	{
+		/*
+		if (Slot.GetUnderlyingArchive().IsTextFormat())
+		{
+			int32 Num = InMap.Num();
+			FStructuredArchive::FMap Map = Slot.EnterMap(Num);
+
+			if (Slot.GetUnderlyingArchive().IsLoading())
+			{
+				FString KeyString;
+				KeyType Key;
+
+				for (int32 Index = 0; Index < Num; ++Index)
+				{
+					FStructuredArchive::FSlot ValueSlot = Map.EnterElement(KeyString);
+					LexFromString(Key, *KeyString);
+					ValueSlot << InMap.Add(Key);
+				}
+			}
+			else
+			{
+				FString StringK;
+				for (TMapBase::TIterator It(InMap); It; ++It)
+				{
+					StringK = LexToString(It->Key);
+					FStructuredArchive::FSlot ValueSlot = Map.EnterElement(StringK);
+					ValueSlot << It->Value;
+				}
+			}
+		}
+		else
+		*/
+		{
+			Slot << InMap.Pairs;
+		}
+	}
+
+	template <typename KeyType, typename ValueType, typename SetAllocator, typename KeyFuncs>
+	[[nodiscard]] static bool LegacyCompareEqual(const TMapBase<KeyType, ValueType, SetAllocator, KeyFuncs>& A, const TMapBase<KeyType, ValueType, SetAllocator, KeyFuncs>& B)
+	{
+		return TSetPrivateFriend::LegacyCompareEqual(A.Pairs, B.Pairs);
+	}
+};
+
+/** Serializer. */
+template <typename KeyType, typename ValueType, typename SetAllocator, typename KeyFuncs>
+UE_FORCEINLINE_HINT FArchive& operator<<(FArchive& Ar, TMapBase<KeyType, ValueType, SetAllocator, KeyFuncs>& Map)
+{
+	return TMapPrivateFriend::Serialize(Ar, Map);
+}
+
+/** Structured archive serializer. */
+template <typename KeyType, typename ValueType, typename SetAllocator, typename KeyFuncs>
+UE_FORCEINLINE_HINT void operator<<(FStructuredArchive::FSlot Slot, TMapBase<KeyType, ValueType, SetAllocator, KeyFuncs>& InMap)
+{
+	TMapPrivateFriend::SerializeStructured(Slot, InMap);
+}
+
+// Legacy comparison operators.  Note that these also test whether the map's key-value pairs were added in the same order!
+template <typename KeyType, typename ValueType, typename SetAllocator, typename KeyFuncs>
+[[nodiscard]] bool LegacyCompareEqual(const TMapBase<KeyType, ValueType, SetAllocator, KeyFuncs>& A, const TMapBase<KeyType, ValueType, SetAllocator, KeyFuncs>& B)
+{
+	return TMapPrivateFriend::LegacyCompareEqual(A, B);
+}
+template <typename KeyType, typename ValueType, typename SetAllocator, typename KeyFuncs>
+[[nodiscard]] bool LegacyCompareNotEqual(const TMapBase<KeyType, ValueType, SetAllocator, KeyFuncs>& A, const TMapBase<KeyType, ValueType, SetAllocator, KeyFuncs>& B)
+{
+	return !TMapPrivateFriend::LegacyCompareEqual(A, B);
+}
diff --git a/Engine/Source/Runtime/Core/Public/Containers/Map.h.inl b/Engine/Source/Runtime/Core/Public/Containers/Map.h.inl
deleted file mode 100644
diff --git a/Engine/Source/Runtime/Core/Public/Containers/ScriptCompactSet.h b/Engine/Source/Runtime/Core/Public/Containers/ScriptCompactSet.h
--- a/Engine/Source/Runtime/Core/Public/Containers/ScriptCompactSet.h
+++ b/Engine/Source/Runtime/Core/Public/Containers/ScriptCompactSet.h
@@ -79,25 +79,19 @@ public:
 		this->NumElements = 0;
 	}
 
-	void RemoveAt(int32 Index, const FScriptCompactSetLayout& Layout, TFunctionRef<uint32 (const void*)> GetKeyHash, TFunctionRef<void (void*)> DestructItem)
+	void RemoveAt(int32 Index, const FScriptCompactSetLayout& Layout, TFunctionRef<uint32 (const void*)> GetKeyHash)
 	{
 		check(IsValidIndex(Index));
 
-		void* Dst = GetData(Index, Layout);
-		const uint32 KeyHash = GetKeyHash(Dst);
+		const uint32 KeyHash = GetKeyHash(GetData(Index, Layout));
 
 		if (Index == this->NumElements - 1)
 		{
 			this->GetHashTableView(Layout).Remove(Index, KeyHash, Index, 0);
-			DestructItem(Dst);
 		}
 		else
 		{
-			void* Src = GetData(this->NumElements - 1, Layout);
-
-			this->GetHashTableView(Layout).Remove(Index, KeyHash, this->NumElements - 1, GetKeyHash(Src));
-			DestructItem(Dst);
-			FMemory::Memmove(Dst, Src, Layout.Size);
+			this->GetHashTableView(Layout).Remove(Index, KeyHash, this->NumElements - 1, GetKeyHash(GetData(this->NumElements - 1, Layout)));
 		}
 
 		--this->NumElements;
diff --git a/Engine/Source/Runtime/Core/Public/Containers/ScriptSparseSet.h b/Engine/Source/Runtime/Core/Public/Containers/ScriptSparseSet.h
deleted file mode 100644
diff --git a/Engine/Source/Runtime/Core/Public/Containers/Set.h b/Engine/Source/Runtime/Core/Public/Containers/Set.h
--- a/Engine/Source/Runtime/Core/Public/Containers/Set.h
+++ b/Engine/Source/Runtime/Core/Public/Containers/Set.h
@@ -3,193 +3,2312 @@
 #pragma once
 
 #include "Containers/ContainerAllocationPolicies.h"
+#include "Containers/ContainerElementTypeCompatibility.h"
+#include "Containers/SetUtilities.h"
+#include "Containers/SparseArray.h"
+#include "ContainersFwd.h"
+#include "Math/UnrealMathUtility.h"
+#include "Misc/AssertionMacros.h"
+#include "Misc/StructBuilder.h"
+#include "Serialization/MemoryImageWriter.h"
+#include "Serialization/StructuredArchive.h"
+#include "Templates/Function.h"
+#include "Templates/RetainedRef.h"
+#include "Templates/Sorting.h"
+#include "Templates/Tuple.h"
+#include "Templates/TypeHash.h"
+#include "Templates/UnrealTemplate.h"
+#include "Traits/IsTriviallyRelocatable.h"
+#include <initializer_list>
+#include <type_traits>
 
-#include "Containers/SparseSet.h"
+template <typename AllocatorType, typename InDerivedType = void>
+class TScriptSet;
 
-#if UE_USE_COMPACT_SET_AS_DEFAULT
-#include "Containers/CompactSet.h"
-#include "Containers/ScriptCompactSet.h"
+template<typename InElementType, bool bTypeLayout>
+class TSetElementBase
+{
+public:
+	typedef InElementType ElementType;
+
+	[[nodiscard]] TSetElementBase() = default;
+
+	/** Initialization constructor. */
+	template <
+		typename... InitType
+		UE_REQUIRES((sizeof...(InitType) != 1) || (!std::is_same_v<TSetElementBase, std::decay_t<InitType>> && ...))
+	>
+	[[nodiscard]] explicit UE_FORCEINLINE_HINT TSetElementBase(InitType&&... InValue)
+		: Value(Forward<InitType>(InValue)...)
+	{
+	}
 
-using FScriptSetLayout = FScriptCompactSetLayout;
+	[[nodiscard]] TSetElementBase(TSetElementBase&&) = default;
+	[[nodiscard]] TSetElementBase(const TSetElementBase&) = default;
+	TSetElementBase& operator=(TSetElementBase&&) = default;
+	TSetElementBase& operator=(const TSetElementBase&) = default;
 
-template<typename Allocator>
-using TScriptSet = TScriptCompactSet<Allocator>;
+	/** The element's value. */
+	ElementType Value;
 
-using TSetPrivateFriend = TCompactSetPrivateFriend;
+	/** The id of the next element in the same hash bucket. */
+	mutable FSetElementId HashNextId;
 
-template <typename InElementType, typename KeyFuncs, typename Allocator>
-using TSetBase = TCompactSet<InElementType, KeyFuncs, Allocator>;
+	/** The hash bucket that the element is currently linked to. */
+	mutable int32 HashIndex;
+};
 
-#else
+template<typename InElementType>
+class TSetElementBase<InElementType, true>
+{
+	DECLARE_INLINE_TYPE_LAYOUT(TSetElementBase, NonVirtual);
+public:
+	typedef InElementType ElementType;
 
-#include "Containers/ScriptSparseSet.h"
+	[[nodiscard]] TSetElementBase() = default;
 
-using FScriptSetLayout = FScriptSparseSetLayout;
+	/** Initialization constructor. */
+	template <
+		typename... InitType
+		UE_REQUIRES((sizeof...(InitType) != 1) || (!std::is_same_v<TSetElementBase, std::decay_t<InitType>> && ...))
+	>
+	[[nodiscard]] explicit UE_FORCEINLINE_HINT TSetElementBase(InitType&&... InValue)
+		: Value(Forward<InitType>(InValue)...)
+	{
+	}
 
-template<typename Allocator>
-using TScriptSet = TScriptSparseSet<Allocator>;
+	[[nodiscard]] TSetElementBase(TSetElementBase&&) = default;
+	[[nodiscard]] TSetElementBase(const TSetElementBase&) = default;
+	TSetElementBase& operator=(TSetElementBase&&) = default;
+	TSetElementBase& operator=(const TSetElementBase&) = default;
 
-using TSetPrivateFriend = TSparseSetPrivateFriend;
+	/** The element's value. */
+	LAYOUT_FIELD(ElementType, Value);
 
-template <typename InElementType, typename KeyFuncs, typename Allocator>
-using TSetBase = TSparseSet<InElementType, KeyFuncs, Allocator>;
+	/** The id of the next element in the same hash bucket. */
+	LAYOUT_MUTABLE_FIELD(FSetElementId, HashNextId);
 
-#endif
+	/** The hash bucket that the element is currently linked to. */
+	LAYOUT_MUTABLE_FIELD(int32, HashIndex);
+};
+
+/** An element in the set. */
+template <typename InElementType>
+class TSetElement : public TSetElementBase<InElementType, THasTypeLayout<InElementType>::Value>
+{
+	using Super = TSetElementBase<InElementType, THasTypeLayout<InElementType>::Value>;
+public:
+	/** Default constructor. */
+	[[nodiscard]] TSetElement() = default;
+
+	/** Initialization constructor. */
+	template <
+		typename... InitType
+		UE_REQUIRES((sizeof...(InitType) != 1) || (!std::is_same_v<TSetElement, std::decay_t<InitType>> && ...))
+	>
+	[[nodiscard]] explicit UE_FORCEINLINE_HINT TSetElement(InitType&&... InValue)
+		: Super(Forward<InitType>(InValue)...)
+	{
+	}
+
+	[[nodiscard]] TSetElement(TSetElement&&) = default;
+	[[nodiscard]] TSetElement(const TSetElement&) = default;
+	TSetElement& operator=(TSetElement&&) = default;
+	TSetElement& operator=(const TSetElement&) = default;
+
+	// Comparison operators
+	[[nodiscard]] UE_FORCEINLINE_HINT bool operator==(const TSetElement& Other) const
+	{
+		return this->Value == Other.Value;
+	}
+	[[nodiscard]] UE_FORCEINLINE_HINT bool operator!=(const TSetElement& Other) const
+	{
+		return this->Value != Other.Value;
+	}
+};
 
-using FScriptSet = TScriptSet<FDefaultSetAllocator>;
+namespace UE::Core::Private
+{
+	[[noreturn]] CORE_API void OnInvalidSetNum(unsigned long long NewNum);
+
+	template<typename HashType>
+	void CopyHash(HashType& Hash, int32& HashSize, const HashType& Copy, int32 HashSizeCopy)
+	{
+		DestructItems((FSetElementId*)Hash.GetAllocation(), HashSize);
+		Hash.ResizeAllocation(0, HashSizeCopy, sizeof(FSetElementId));
+		ConstructItems<FSetElementId>(Hash.GetAllocation(), (FSetElementId*)Copy.GetAllocation(), HashSizeCopy);
+		HashSize = HashSizeCopy;
+	}
+
+	template<typename HashType>
+	[[nodiscard]] FORCEINLINE FSetElementId& GetTypedHash(HashType& Hash, int32 HashIndex, int32 HashSize)
+	{
+		return ((FSetElementId*)Hash.GetAllocation())[HashIndex & (HashSize - 1)];
+	}
+
+	template<typename HashType>
+	void Rehash(HashType& Hash, int32 HashSize)
+	{
+		// Free the old hash.
+		Hash.ResizeAllocation(0, 0, sizeof(FSetElementId));
+
+		if (HashSize)
+		{
+			// Allocate the new hash.
+			checkSlow(FMath::IsPowerOfTwo(HashSize));
+			Hash.ResizeAllocation(0, HashSize, sizeof(FSetElementId));
+			for (int32 HashIndex = 0; HashIndex < HashSize; ++HashIndex)
+			{
+				GetTypedHash(Hash, HashIndex, HashSize) = FSetElementId();
+			}
+		}
+	}
+}
 
-template<typename InElementType, typename KeyFuncs, typename Allocator>
-class TSet : public TSetBase<InElementType, KeyFuncs, Allocator>
+/**
+ * A set with an optional KeyFuncs parameters for customizing how the elements are compared and searched.
+ * E.g. You can specify a mapping from elements to keys if you want to find elements by specifying a subset of
+ * the element type.  It uses a TSparseArray of the elements, and also links the elements into a hash with a
+ * number of buckets proportional to the number of elements.  Addition, removal, and finding are O(1).
+ *
+ * The ByHash() functions are somewhat dangerous but particularly useful in two scenarios:
+ * -- Heterogeneous lookup to avoid creating expensive keys like FString when looking up by const TCHAR*.
+ *	  You must ensure the hash is calculated in the same way as ElementType is hashed.
+ *    If possible put both ComparableKey and ElementType hash functions next to each other in the same header
+ *    to avoid bugs when the ElementType hash function is changed.
+ * -- Reducing contention around hash tables protected by a lock. It is often important to incur
+ *    the cache misses of reading key data and doing the hashing *before* acquiring the lock.
+ *
+ **/
+template<
+	typename InElementType,
+	typename KeyFuncs /*= DefaultKeyFuncs<ElementType>*/,
+	typename Allocator /*= FDefaultSetAllocator*/
+	>
+class TSet
 {
 public:
-	using SetBase = TSetBase<InElementType, KeyFuncs, Allocator>;
+	typedef InElementType ElementType;
+	typedef KeyFuncs    KeyFuncsType;
+	typedef Allocator   AllocatorType;
+
+	using SizeType = typename Allocator::SparseArrayAllocator::ElementAllocator::SizeType;
+
+	static_assert(std::is_same_v<SizeType, int32>, "TSet currently only supports 32-bit allocators");
+
+private:
+	using USizeType = std::make_unsigned_t<SizeType>;
 
+	template <typename, typename>
+	friend class TScriptSet;
+
+	typedef typename KeyFuncs::KeyInitType     KeyInitType;
+	typedef typename KeyFuncs::ElementInitType ElementInitType;
+
+	typedef TSetElement<InElementType> SetElementType;
+
+public:
+	/** Initialization constructor. */
 	[[nodiscard]] UE_FORCEINLINE_HINT TSet() = default;
 
-	[[nodiscard]] UE_FORCEINLINE_HINT TSet(const SetBase& Copy)
+	[[nodiscard]] explicit consteval TSet(EConstEval)
+	: Elements(ConstEval)
+	, Hash(ConstEval)
+	{
+	}
+
+	/** Copy constructor. */
+	[[nodiscard]] UE_FORCEINLINE_HINT TSet(const TSet& Copy)
 	{
 		*this = Copy;
 	}
-	[[nodiscard]] UE_FORCEINLINE_HINT TSet(SetBase&& Other)
+
+	[[nodiscard]] UE_FORCEINLINE_HINT explicit TSet(TArrayView<const ElementType> InArrayView)
+	{
+		Append(InArrayView);
+	}
+
+	[[nodiscard]] UE_FORCEINLINE_HINT explicit TSet(TArray<ElementType>&& InArray)
 	{
-		*this = MoveTemp(Other);
+		Append(MoveTemp(InArray));
 	}
 
-	[[nodiscard]] UE_FORCEINLINE_HINT explicit TSet(TArrayView<const InElementType> InArrayView)
+	/** Destructor. */
+	UE_FORCEINLINE_HINT ~TSet()
 	{
-		this->Append(InArrayView);
+		UE_STATIC_ASSERT_WARN(TIsTriviallyRelocatable_V<InElementType>, "TSet can only be used with trivially relocatable types");
+
+		HashSize = 0;
 	}
 
-	[[nodiscard]] UE_FORCEINLINE_HINT explicit TSet(TArray<InElementType>&& InArray)
+	/////////////////////////////////////////////
+	// Start - intrusive TOptional<TSet> state //
+	/////////////////////////////////////////////
+	constexpr static bool bHasIntrusiveUnsetOptionalState = true;
+	using IntrusiveUnsetOptionalStateType = TSet;
+
+	[[nodiscard]] explicit TSet(FIntrusiveUnsetOptionalState Tag)
+		: Elements(Tag)
+	{
+	}
+	[[nodiscard]] bool operator==(FIntrusiveUnsetOptionalState Tag) const
 	{
-		this->Append(MoveTemp(InArray));
+		return Elements == Tag;
 	}
+	///////////////////////////////////////////
+	// End - intrusive TOptional<TSet> state //
+	///////////////////////////////////////////
 
-	[[nodiscard]] UE_FORCEINLINE_HINT TSet(std::initializer_list<InElementType> InitList)
+	/** Assignment operator. */
+	TSet& operator=(const TSet& Copy)
 	{
-		this->Append(InitList);
+		if (this != &Copy)
+		{
+			UE::Core::Private::CopyHash(Hash, HashSize, Copy.Hash, Copy.HashSize);
+
+			Elements = Copy.Elements;
+		}
+		return *this;
 	}
 
-	template<typename OtherAllocator>
-	[[nodiscard]] UE_FORCEINLINE_HINT TSet(TSetBase<InElementType, KeyFuncs, OtherAllocator>&& Other)
+private:
+	template <typename SetType>
+	static inline void Move(SetType& ToSet, SetType& FromSet)
 	{
-		this->Append(MoveTemp(Other));
+		ToSet.Elements = (ElementArrayType&&)FromSet.Elements;
+
+		ToSet.Hash.MoveToEmpty(FromSet.Hash);
+
+		ToSet  .HashSize = FromSet.HashSize;
+		FromSet.HashSize = 0;
 	}
 
-	template<typename OtherAllocator>
-	[[nodiscard]] UE_FORCEINLINE_HINT TSet(const TSetBase<InElementType, KeyFuncs, OtherAllocator>& Other)
+public:
+	/** Initializer list constructor. */
+	[[nodiscard]] TSet(std::initializer_list<ElementType> InitList)
+		: HashSize(0)
 	{
-		this->Append(Other);
+		Append(InitList);
 	}
 
-	TSet& operator=(const SetBase& Copy)
+	/** Move constructor. */
+	[[nodiscard]] TSet(TSet&& Other)
+		: HashSize(0)
 	{
-		SetBase::operator=(Copy);
-		return *this;
+		this->Move(*this, Other);
 	}
 
 	/** Move assignment operator. */
-	TSet& operator=(SetBase&& Other)
+	TSet& operator=(TSet&& Other)
 	{
-		SetBase::operator=(MoveTemp(Other));
+		if (this != &Other)
+		{
+			this->Move(*this, Other);
+		}
+
 		return *this;
 	}
 
-	/** Assignment operator for moving elements from a TSparseSet with a different SetAllocator */
+	/** Constructor for moving elements from a TSet with a different SetAllocator */
+	template<typename OtherAllocator>
+	[[nodiscard]] TSet(TSet<ElementType, KeyFuncs, OtherAllocator>&& Other)
+		: HashSize(0)
+	{
+		Append(MoveTemp(Other));
+	}
+
+	/** Constructor for copying elements from a TSet with a different SetAllocator */
+	template<typename OtherAllocator>
+	[[nodiscard]] TSet(const TSet<ElementType, KeyFuncs, OtherAllocator>& Other)
+		: HashSize(0)
+	{
+		Append(Other);
+	}
+
+	/** Assignment operator for moving elements from a TSet with a different SetAllocator */
 	template<typename OtherAllocator>
-	TSet& operator=(TSetBase<InElementType, KeyFuncs, OtherAllocator>&& Other)
+	TSet& operator=(TSet<ElementType, KeyFuncs, OtherAllocator>&& Other)
 	{
-		this->Reset();
-		this->Append(MoveTemp(Other));
+		Reset();
+		Append(MoveTemp(Other));
 		return *this;
 	}
 
-	/** Assignment operator for copying elements from a TSparseSet with a different SetAllocator */
+	/** Assignment operator for copying elements from a TSet with a different SetAllocator */
 	template<typename OtherAllocator>
-	TSet& operator=(const TSetBase<InElementType, KeyFuncs, OtherAllocator>& Other)
+	TSet& operator=(const TSet<ElementType, KeyFuncs, OtherAllocator>& Other)
 	{
-		this->Reset();
-		this->Append(Other);
+		Reset();
+		Append(Other);
 		return *this;
 	}
 
 	/** Initializer list assignment operator */
-	TSet& operator=(std::initializer_list<InElementType> InitList)
+	TSet& operator=(std::initializer_list<ElementType> InitList)
 	{
-		this->Reset();
-		this->Append(InitList);
+		Reset();
+		Append(InitList);
 		return *this;
 	}
 
 	/**
-	 * Move assignment operator.
-	 * Compatible element type version.
+	 * Removes all elements from the set, potentially leaving space allocated for an expected number of elements about to be added.
+	 * @param ExpectedNumElements - The number of elements about to be added to the set.
+	 */
+	void Empty(int32 ExpectedNumElements = 0)
+	{
+		// Empty the elements array, and reallocate it for the expected number of elements.
+		const int32 DesiredHashSize = Allocator::GetNumberOfHashBuckets(ExpectedNumElements);
+		const bool ShouldDoRehash = ShouldRehash(ExpectedNumElements, DesiredHashSize, EAllowShrinking::Yes);
+
+		if (!ShouldDoRehash)
+		{
+			// If the hash was already the desired size, clear the references to the elements that have now been removed.
+			UnhashElements();
+		}
+
+		Elements.Empty(ExpectedNumElements);
+
+		// Resize the hash to the desired size for the expected number of elements.
+		if (ShouldDoRehash)
+		{
+			HashSize = DesiredHashSize;
+			Rehash();
+		}
+	}
+
+	/** Efficiently empties out the set but preserves all allocations and capacities */
+	void Reset()
+	{
+		if (Num() == 0)
+		{
+			return;
+		}
+
+		// Reset the elements array.
+		UnhashElements();
+		Elements.Reset();
+	}
+
+	/** Shrinks the set's element storage to avoid slack. */
+	inline void Shrink()
+	{
+		Elements.Shrink();
+		Relax();
+	}
+
+	/** Compacts the allocated elements into a contiguous range. */
+	inline void Compact()
+	{
+		if (Elements.Compact())
+		{
+			HashSize = Allocator::GetNumberOfHashBuckets(Elements.Num());
+			Rehash();
+		}
+	}
+
+	/** Compacts the allocated elements into a contiguous range. Does not change the iteration order of the elements. */
+	inline void CompactStable()
+	{
+		if (Elements.CompactStable())
+		{
+			HashSize = Allocator::GetNumberOfHashBuckets(Elements.Num());
+			Rehash();
+		}
+	}
+
+	/** Preallocates enough memory to contain Number elements */
+	inline void Reserve(int32 Number)
+	{
+		// makes sense only when Number > Elements.Num() since TSparseArray::Reserve
+		// does any work only if that's the case
+		if ((USizeType)Number > (USizeType)Elements.Num())
+		{
+			// Trap negative reserves
+			if (Number < 0)
+			{
+				UE::Core::Private::OnInvalidSetNum((unsigned long long)Number);
+			}
+
+			// Preallocates memory for array of elements
+			Elements.Reserve(Number);
+
+			// Calculate the corresponding hash size for the specified number of elements.
+			const int32 NewHashSize = Allocator::GetNumberOfHashBuckets(Number);
+
+			// If the hash hasn't been created yet, or is smaller than the corresponding hash size, rehash
+			// to force a preallocation of the hash table
+			if(!HashSize || HashSize < NewHashSize)
+			{
+				HashSize = NewHashSize;
+				Rehash();
+			}
+		}
+	}
+
+	/** Relaxes the set's hash to a size strictly bounded by the number of elements in the set. */
+	UE_FORCEINLINE_HINT void Relax()
+	{
+		ConditionalRehash(Elements.Num(), EAllowShrinking::Yes);
+	}
+
+	/**
+	 * Helper function to return the amount of memory allocated by this container
+	 * Only returns the size of allocations made directly by the container, not the elements themselves.
+	 * @return number of bytes allocated by this container
+	 */
+	[[nodiscard]] UE_FORCEINLINE_HINT SIZE_T GetAllocatedSize( void ) const
+	{
+		return Elements.GetAllocatedSize() + Hash.GetAllocatedSize(HashSize, sizeof(FSetElementId));
+	}
+
+	/** Tracks the container's memory use through an archive. */
+	inline void CountBytes(FArchive& Ar) const
+	{
+		Elements.CountBytes(Ar);
+		Ar.CountBytes(HashSize * sizeof(int32),HashSize * sizeof(FSetElementId));
+	}
+
+	/**
+	 * Returns true if the sets is empty and contains no elements.
 	 *
-	 * @param Other Set to assign and move from.
+	 * @returns True if the set is empty.
+	 * @see Num
 	 */
-	template <
-	typename OtherKeyFuncs,
-	typename AliasElementType = InElementType
-	UE_REQUIRES(TIsContainerElementTypeCopyable_V<AliasElementType>)
-	>
-	TSet& operator=(TSetBase<typename TContainerElementTypeCompatibility<InElementType>::CopyFromOtherType, OtherKeyFuncs, Allocator>&& Other)
+	[[nodiscard]] bool IsEmpty() const
 	{
-		TContainerElementTypeCompatibility<InElementType>::CopyingFromOtherType();
-		this->Reset();
-		this->Append(MoveTemp(Other));
-		return *this;
+		return Elements.IsEmpty();
+	}
+
+	/** @return the number of elements. */
+	[[nodiscard]] UE_FORCEINLINE_HINT int32 Num() const
+	{
+		return Elements.Num();
+	}
+
+	/** @return The number of elements the set can hold before reallocation. */
+	[[nodiscard]] UE_FORCEINLINE_HINT int32 Max() const
+	{
+		return Elements.Max();
+	}
+
+	/** @return The non-inclusive maximum index of elements in the set. */
+	[[nodiscard]] UE_FORCEINLINE_HINT int32 GetMaxIndex() const
+	{
+		return Elements.GetMaxIndex();
 	}
 
 	/**
-	 * Assignment operator. First deletes all currently contained elements
-	 * and then copies from other set.
-	 * Compatible element type version.
+	 * Checks whether an element id is valid.
+	 * @param Id - The element id to check.
+	 * @return true if the element identifier refers to a valid element in this set.
+	 */
+	[[nodiscard]] inline bool IsValidId(FSetElementId Id) const
+	{
+		SizeType Index = Id.AsInteger();
+		return Index != INDEX_NONE &&
+			Index >= 0 &&
+			Index < Elements.GetMaxIndex() &&
+			Elements.IsAllocated(Index);
+	}
+
+	/** Accesses the identified element's value. Element must be valid (see @IsValidId). */
+	[[nodiscard]] UE_FORCEINLINE_HINT ElementType& operator[](FSetElementId Id)
+	{
+		return Elements[Id.AsInteger()].Value;
+	}
+
+	/** Accesses the identified element's value. Element must be valid (see @IsValidId). */
+	[[nodiscard]] UE_FORCEINLINE_HINT const ElementType& operator[](FSetElementId Id) const
+	{
+		return Elements[Id.AsInteger()].Value;
+	}
+
+	/** Accesses the identified element's value. Element must be valid (see @IsValidId). */
+	[[nodiscard]] UE_FORCEINLINE_HINT ElementType& Get(FSetElementId Id)
+	{
+		return Elements[Id.AsInteger()].Value;
+	}
+
+	/** Accesses the identified element's value. Element must be valid (see @IsValidId). */
+	[[nodiscard]] UE_FORCEINLINE_HINT const ElementType& Get(FSetElementId Id) const
+	{
+		return Elements[Id.AsInteger()].Value;
+	}
+
+	/**
+	 * Adds an element to the set.
 	 *
-	 * @param Other The source set to assign from.
+	 * @param	InElement					Element to add to set
+	 * @param	bIsAlreadyInSetPtr	[out]	Optional pointer to bool that will be set depending on whether element is already in set
+	 * @return	A pointer to the element stored in the set.
 	 */
-	template <
-	typename OtherKeyFuncs,
-	typename OtherAllocator,
-	typename AliasElementType = InElementType
-	UE_REQUIRES(TIsContainerElementTypeCopyable_V<AliasElementType>)
-	>
-	TSet& operator=(const TSetBase<typename TContainerElementTypeCompatibility<InElementType>::CopyFromOtherType, OtherKeyFuncs, OtherAllocator>& Other)
+	UE_FORCEINLINE_HINT FSetElementId Add(const InElementType&  InElement, bool* bIsAlreadyInSetPtr = nullptr)
 	{
-		TContainerElementTypeCompatibility<InElementType>::CopyingFromOtherType();
-		this->Reset();
-		this->Append(Other);
-		return *this;
+		return Emplace(InElement, bIsAlreadyInSetPtr);
+	}
+	UE_FORCEINLINE_HINT FSetElementId Add(InElementType&& InElement, bool* bIsAlreadyInSetPtr = nullptr)
+	{
+		return Emplace(MoveTempIfPossible(InElement), bIsAlreadyInSetPtr);
 	}
 
-	/////////////////////////////////////////////
-	// Start - intrusive TOptional<TSet> state //
-	/////////////////////////////////////////////
-	constexpr static bool bHasIntrusiveUnsetOptionalState = true;
-	using IntrusiveUnsetOptionalStateType = TSet;
+	/**
+	 * Adds an element to the set if not already present and returns a reference to the added or existing element.
+	 *
+	 * @param	InElement					Element to add to set
+	 * @param	bIsAlreadyInSetPtr	[out]	Optional pointer to bool that will be set depending on whether element is already in set
+	 * @return	A reference to the element stored in the set.
+	 */
+	UE_FORCEINLINE_HINT ElementType& FindOrAdd(const InElementType& InElement, bool* bIsAlreadyInSetPtr = nullptr)
+	{
+		return FindOrAddByHash(KeyFuncs::GetKeyHash(KeyFuncs::GetSetKey(InElement)), InElement, bIsAlreadyInSetPtr);
+	}
+	UE_FORCEINLINE_HINT ElementType& FindOrAdd(InElementType&& InElement, bool* bIsAlreadyInSetPtr = nullptr)
+	{
+		return FindOrAddByHash(KeyFuncs::GetKeyHash(KeyFuncs::GetSetKey(InElement)), MoveTempIfPossible(InElement), bIsAlreadyInSetPtr);
+	}
 
-	[[nodiscard]] explicit TSet(FIntrusiveUnsetOptionalState Tag)
-	: SetBase(Tag)
+	/**
+	 * Adds an element to the set.
+	 *
+	 * @see		Class documentation section on ByHash() functions
+	 * @param	KeyHash						A precomputed hash value, calculated in the same way as ElementType is hashed.
+	 * @param	InElement					Element to add to set
+	 * @param	bIsAlreadyInSetPtr	[out]	Optional pointer to bool that will be set depending on whether element is already in set
+     * @return  A handle to the element stored in the set
+     */
+	UE_FORCEINLINE_HINT FSetElementId AddByHash(uint32 KeyHash, const InElementType& InElement, bool* bIsAlreadyInSetPtr = nullptr)
+	{
+		return EmplaceByHash(KeyHash, InElement, bIsAlreadyInSetPtr);
+	}
+	UE_FORCEINLINE_HINT FSetElementId AddByHash(uint32 KeyHash,		 InElementType&& InElement, bool* bIsAlreadyInSetPtr = nullptr)
 	{
+		return EmplaceByHash(KeyHash, MoveTempIfPossible(InElement), bIsAlreadyInSetPtr);
 	}
 
-	///////////////////////////////////////////
-	// End - intrusive TOptional<TSet> state //
-	///////////////////////////////////////////
-};
+	/**
+	 * Adds an element to the set if not already present and returns a reference to the added or existing element.
+	 *
+	 * @see		Class documentation section on ByHash() functions
+	 * @param	KeyHash						A precomputed hash value, calculated in the same way as ElementType is hashed.
+	 * @param	InElement					Element to add to set
+	 * @param	bIsAlreadyInSetPtr	[out]	Optional pointer to bool that will be set depending on whether element is already in set
+	 * @return  A reference to the element stored in the set
+	 */
+	template <typename ElementReferenceType>
+	ElementType& FindOrAddByHash(uint32 KeyHash, ElementReferenceType&& InElement, bool* bIsAlreadyInSetPtr = nullptr)
+	{
+		SizeType ExistingIndex = FindIndexByHash(KeyHash, KeyFuncs::GetSetKey(InElement));
+		bool bIsAlreadyInSet = ExistingIndex != INDEX_NONE;
+		if (bIsAlreadyInSetPtr)
+		{
+			*bIsAlreadyInSetPtr = bIsAlreadyInSet;
+		}
+		if (bIsAlreadyInSet)
+		{
+			return Elements[ExistingIndex].Value;
+		}
 
-template <typename RangeType>
-TSet(RangeType&&) -> TSet<TElementType_T<RangeType>>;
+		// Create a new element.
+		FSparseArrayAllocationInfo ElementAllocation = Elements.AddUninitialized();
+		SetElementType& Element = *new (ElementAllocation) SetElementType(Forward<ElementReferenceType>(InElement));
+		RehashOrLink(KeyHash, Element, ElementAllocation.Index);
+		return Element.Value;
+	}
 
-/**
- * Traits class which determines whether or not a type is a TSet.
- */
-template <typename T> struct TIsTSet { enum { Value = false }; };
+private:
+	bool TryReplaceExisting(uint32 KeyHash, SetElementType& Element, SizeType& InOutElementIndex, bool* bIsAlreadyInSetPtr)
+	{
+		bool bIsAlreadyInSet = false;
+		if constexpr (!KeyFuncs::bAllowDuplicateKeys)
+		{
+			// If the set doesn't allow duplicate keys, check for an existing element with the same key as the element being added.
 
-template <typename ElementType, typename KeyFuncs, typename Allocator> struct TIsTSet<               TSet<ElementType, KeyFuncs, Allocator>> { enum { Value = true }; };
-template <typename ElementType, typename KeyFuncs, typename Allocator> struct TIsTSet<const          TSet<ElementType, KeyFuncs, Allocator>> { enum { Value = true }; };
-template <typename ElementType, typename KeyFuncs, typename Allocator> struct TIsTSet<      volatile TSet<ElementType, KeyFuncs, Allocator>> { enum { Value = true }; };
-template <typename ElementType, typename KeyFuncs, typename Allocator> struct TIsTSet<const volatile TSet<ElementType, KeyFuncs, Allocator>> { enum { Value = true }; };
+			// Don't bother searching for a duplicate if this is the first element we're adding
+			if (Elements.Num() != 1)
+			{
+				SizeType ExistingIndex = FindIndexByHash(KeyHash, KeyFuncs::GetSetKey(Element.Value));
+				bIsAlreadyInSet = ExistingIndex != INDEX_NONE;
+				if (bIsAlreadyInSet)
+				{
+					// If there's an existing element with the same key as the new element, replace the existing element with the new element.
+					MoveByRelocate(Elements[ExistingIndex].Value, Element.Value);
+
+					// Then remove the new element.
+					Elements.RemoveAtUninitialized(InOutElementIndex);
+
+					// Then point the return value at the replaced element.
+					InOutElementIndex = ExistingIndex;
+				}
+			}
+		}
+		if (bIsAlreadyInSetPtr)
+		{
+			*bIsAlreadyInSetPtr = bIsAlreadyInSet;
+		}
+		return bIsAlreadyInSet;
+	}
+
+	inline void RehashOrLink(uint32 KeyHash, SetElementType& Element, SizeType ElementIndex)
+	{
+		// Check if the hash needs to be resized.
+		if (!ConditionalRehash(Elements.Num(), EAllowShrinking::No))
+		{
+			// If the rehash didn't add the new element to the hash, add it.
+			LinkElement(ElementIndex, Element, KeyHash);
+		}
+	}
+
+public:
+	/**
+	 * Adds an element to the set.
+	 *
+	 * @param	Arg							The argument to be forwarded to the set element's constructor.
+	 * @param	bIsAlreadyInSetPtr	[out]	Optional pointer to bool that will be set depending on whether element is already in set
+	 * @return	A handle to the element stored in the set.
+	 */
+	template <typename ArgType = ElementType>
+	FSetElementId Emplace(ArgType&& Arg, bool* bIsAlreadyInSetPtr = nullptr)
+	{
+		// Create a new element.
+		FSparseArrayAllocationInfo ElementAllocation = Elements.AddUninitialized();
+		SetElementType& Element = *new (ElementAllocation) SetElementType(Forward<ArgType>(Arg));
+
+		SizeType NewHashIndex = ElementAllocation.Index;
+
+		uint32 KeyHash = KeyFuncs::GetKeyHash(KeyFuncs::GetSetKey(Element.Value));
+		if (!TryReplaceExisting(KeyHash, Element, NewHashIndex, bIsAlreadyInSetPtr))
+		{
+			RehashOrLink(KeyHash, Element, NewHashIndex);
+		}
+		return FSetElementId::FromInteger(NewHashIndex);
+	}
+
+    /**
+     * Adds an element to the set by constructing the ElementType in-place with multiple args.
+     *
+     * @param   EInPlace   Tag to disambiguate in-place construction.
+     * @param   InArgs     Arguments forwarded to ElementType's constructor.
+     * @return  Pair of (element id, whether an equivalent element already existed).
+     */
+    template <typename... ArgTypes>
+    TPair<FSetElementId, bool> Emplace(EInPlace, ArgTypes&&... InArgs)
+    {
+		// Create a new element.
+        FSparseArrayAllocationInfo ElementAllocation = Elements.AddUninitialized();
+		SetElementType& Element = *new (ElementAllocation) SetElementType(Forward<ArgTypes>(InArgs)...);
+
+        SizeType NewHashIndex = ElementAllocation.Index;
+        const uint32 KeyHash = KeyFuncs::GetKeyHash(KeyFuncs::GetSetKey(Element.Value));
+
+        bool bAlreadyInSet = false;
+        if (!TryReplaceExisting(KeyHash, Element, NewHashIndex, &bAlreadyInSet))
+        {
+            RehashOrLink(KeyHash, Element, NewHashIndex);
+        }
+
+        return TPair<FSetElementId, bool>(FSetElementId::FromInteger(NewHashIndex), bAlreadyInSet);
+    }
+
+	/**
+	 * Adds an element to the set.
+	 *
+	 * @see		Class documentation section on ByHash() functions
+	 * @param	KeyHash                     A precomputed hash value, calculated in the same way as ElementType is hashed.
+	 * @param	Arg							The argument to be forwarded to the set element's constructor.
+	 * @param	bIsAlreadyInSetPtr	[out]	Optional pointer to bool that will be set depending on whether element is already in set
+	 * @return	A handle to the element stored in the set.
+	 */
+	template <typename ArgType = ElementType>
+	FSetElementId EmplaceByHash(uint32 KeyHash, ArgType&& Arg, bool* bIsAlreadyInSetPtr = nullptr)
+	{
+		// Create a new element.
+		FSparseArrayAllocationInfo ElementAllocation = Elements.AddUninitialized();
+		SetElementType& Element = *new (ElementAllocation) SetElementType(Forward<ArgType>(Arg));
+
+		SizeType NewHashIndex = ElementAllocation.Index;
+
+		if (!TryReplaceExisting(KeyHash, Element, NewHashIndex, bIsAlreadyInSetPtr))
+		{
+			RehashOrLink(KeyHash, Element, NewHashIndex);
+		}
+		return FSetElementId::FromInteger(NewHashIndex);
+	}
+
+    /**
+     * Adds an element to the set by constructing in-place with multiple args, using a precomputed hash.
+	 *
+     * @see     Class documentation section on ByHash() functions.
+	 * @param	KeyHash    A precomputed hash value, calculated in the same way as ElementType is hashed.
+     * @param   EInPlace   Tag to disambiguate in-place construction.
+     * @param   InArgs     Arguments forwarded to ElementType's constructor.
+     * @return  Pair of (element id, whether an equivalent element already existed).
+     */
+    template <typename... ArgTypes>
+    TPair<FSetElementId, bool> EmplaceByHash(EInPlace, uint32 KeyHash, ArgTypes&&... InArgs)
+    {
+		// Create a new element.
+        FSparseArrayAllocationInfo ElementAllocation = Elements.AddUninitialized();
+		SetElementType& Element = *new (ElementAllocation) SetElementType(Forward<ArgTypes>(InArgs)...);
+
+        SizeType NewHashIndex = ElementAllocation.Index;
+        bool bAlreadyInSet = false;
+        if (!TryReplaceExisting(KeyHash, Element, NewHashIndex, &bAlreadyInSet))
+        {
+            RehashOrLink(KeyHash, Element, NewHashIndex);
+        }
+        return TPair<FSetElementId, bool>(FSetElementId::FromInteger(NewHashIndex), bAlreadyInSet);
+    }
+
+ 	void Append(TArrayView<const ElementType> InElements)
+	{
+		Reserve(Elements.Num() + InElements.Num());
+		for (const ElementType& Element : InElements)
+		{
+			Add(Element);
+		}
+	}
+
+	template<typename ArrayAllocator>
+	void Append(TArray<ElementType, ArrayAllocator>&& InElements)
+	{
+		Reserve(Elements.Num() + InElements.Num());
+		for (ElementType& Element : InElements)
+		{
+			Add(MoveTempIfPossible(Element));
+		}
+		InElements.Reset();
+	}
+
+	/**
+	 * Add all items from another set to our set (union without creating a new set)
+	 * @param OtherSet - The other set of items to add.
+	 */
+	template<typename OtherAllocator>
+	void Append(const TSet<ElementType, KeyFuncs, OtherAllocator>& OtherSet)
+	{
+		Reserve(Elements.Num() + OtherSet.Num());
+		for (const ElementType& Element : OtherSet)
+		{
+			Add(Element);
+		}
+	}
+
+	template<typename OtherAllocator>
+	void Append(TSet<ElementType, KeyFuncs, OtherAllocator>&& OtherSet)
+	{
+		Reserve(Elements.Num() + OtherSet.Num());
+		for (ElementType& Element : OtherSet)
+		{
+			Add(MoveTempIfPossible(Element));
+		}
+		OtherSet.Reset();
+	}
+
+	void Append(std::initializer_list<ElementType> InitList)
+	{
+		Reserve(Elements.Num() + (int32)InitList.size());
+		for (const ElementType& Element : InitList)
+		{
+			Add(Element);
+		}
+	}
+
+private:
+	void RemoveByIndex(SizeType ElementIndex)
+	{
+		checkf(Elements.IsValidIndex(ElementIndex), TEXT("Invalid ElementIndex passed to TSet::RemoveByIndex"));
 
+		const SetElementType& ElementBeingRemoved = Elements[ElementIndex];
+
+		// Remove the element from the hash.
+		FSetElementId* HashPtr              = Hash.GetAllocation();
+		FSetElementId* NextElementIndexIter = &HashPtr[ElementBeingRemoved.HashIndex];
+		for (;;)
+		{
+			SizeType NextElementIndex = NextElementIndexIter->AsInteger();
+			checkf(NextElementIndex != INDEX_NONE, TEXT("Corrupt hash"));
+
+			if (NextElementIndex == ElementIndex)
+			{
+				*NextElementIndexIter = ElementBeingRemoved.HashNextId;
+				break;
+			}
+
+			NextElementIndexIter = &Elements[NextElementIndex].HashNextId;
+		}
+
+		// Remove the element from the elements array.
+		Elements.RemoveAt(ElementIndex);
+	}
+
+public:
+	/**
+	 * Removes an element from the set.
+	 * @param Element - A pointer to the element in the set, as returned by Add or Find.
+	 */
+	void Remove(FSetElementId ElementId)
+	{
+		RemoveByIndex(ElementId.AsInteger());
+	}
+
+private:
+	/**
+	 * Finds an element with a pre-calculated hash and a key that can be compared to KeyType
+	 * @see	Class documentation section on ByHash() functions
+	 * @return The element id that matches the key and hash or an invalid element id
+	 */
+	template <typename ComparableKey>
+	[[nodiscard]] SizeType FindIndexByHash(uint32 KeyHash, const ComparableKey& Key) const
+	{
+		if (Elements.Num() == 0)
+		{
+			return INDEX_NONE;
+		}
+
+		FSetElementId* HashPtr      = Hash.GetAllocation();
+		SizeType       ElementIndex = HashPtr[KeyHash & (HashSize - 1)].AsInteger();
+		for (;;)
+		{
+			if (ElementIndex == INDEX_NONE)
+			{
+				return INDEX_NONE;
+			}
+
+			if (KeyFuncs::Matches(KeyFuncs::GetSetKey(Elements[ElementIndex].Value), Key))
+			{
+				// Return the first match, regardless of whether the set has multiple matches for the key or not.
+				return ElementIndex;
+			}
+
+			ElementIndex = Elements[ElementIndex].HashNextId.AsInteger();
+		}
+	}
+
+public:
+	/**
+	 * Finds any element in the set and returns a pointer to it.
+	 * Callers should not depend on particular patterns in the behaviour of this function.
+	 * @return A pointer to an arbitrary element, or nullptr if the container is empty.
+	 */
+	[[nodiscard]] ElementType* FindArbitraryElement()
+	{
+		// The goal of this function is to be fast, and so the implementation may be improved at any time even if it gives different results.
+
+		int32 Result = Elements.FindArbitraryElementIndex();
+		return (Result != INDEX_NONE) ? &Elements[Result].Value : nullptr;
+	}
+	[[nodiscard]] const ElementType* FindArbitraryElement() const
+	{
+		return const_cast<TSet*>(this)->FindArbitraryElement();
+	}
+
+	/**
+	 * Finds an element with the given key in the set.
+	 * @param Key - The key to search for.
+	 * @return The id of the set element matching the given key, or the NULL id if none matches.
+	 */
+	[[nodiscard]] FSetElementId FindId(KeyInitType Key) const
+	{
+		return FSetElementId::FromInteger(FindIndexByHash(KeyFuncs::GetKeyHash(Key), Key));
+	}
+
+	/**
+	 * Finds an element with a pre-calculated hash and a key that can be compared to KeyType
+	 * @see	Class documentation section on ByHash() functions
+	 * @return The element id that matches the key and hash or an invalid element id
+	 */
+	template<typename ComparableKey>
+	[[nodiscard]] FSetElementId FindIdByHash(uint32 KeyHash, const ComparableKey& Key) const
+	{
+		PRAGMA_DISABLE_DEPRECATION_WARNINGS
+		// Disable deprecations warnings to stop warnings being thrown by our check macro.
+		checkSlow(KeyHash == KeyFuncs::GetKeyHash(Key));
+		PRAGMA_ENABLE_DEPRECATION_WARNINGS
+
+		return FSetElementId::FromInteger(FindIndexByHash(KeyHash, Key));
+	}
+
+	/**
+	 * Finds an element with the given key in the set.
+	 * @param Key - The key to search for.
+	 * @return A pointer to an element with the given key.  If no element in the set has the given key, this will return NULL.
+	 */
+	[[nodiscard]] inline ElementType* Find(KeyInitType Key)
+	{
+		SizeType ElementIndex = FindIndexByHash(KeyFuncs::GetKeyHash(Key), Key);
+		if (ElementIndex != INDEX_NONE)
+		{
+			return &Elements[ElementIndex].Value;
+		}
+		else
+		{
+			return nullptr;
+		}
+	}
+
+	/**
+	 * Finds an element with the given key in the set.
+	 * @param Key - The key to search for.
+	 * @return A const pointer to an element with the given key.  If no element in the set has the given key, this will return NULL.
+	 */
+	[[nodiscard]] UE_FORCEINLINE_HINT const ElementType* Find(KeyInitType Key) const
+	{
+		return const_cast<TSet*>(this)->Find(Key);
+	}
+
+	/**
+	 * Finds an element with a pre-calculated hash and a key that can be compared to KeyType.
+	 * @see	Class documentation section on ByHash() functions
+	 * @return A pointer to the contained element or nullptr.
+	 */
+	template<typename ComparableKey>
+	[[nodiscard]] ElementType* FindByHash(uint32 KeyHash, const ComparableKey& Key)
+	{
+		SizeType ElementIndex = FindIndexByHash(KeyHash, Key);
+		if (ElementIndex != INDEX_NONE)
+		{
+			return &Elements[ElementIndex].Value;
+		}
+		else
+		{
+			return nullptr;
+		}
+	}
+
+	template<typename ComparableKey>
+	[[nodiscard]] const ElementType* FindByHash(uint32 KeyHash, const ComparableKey& Key) const
+	{
+		return const_cast<TSet*>(this)->FindByHash(KeyHash, Key);
+	}
+
+private:
+	template<typename ComparableKey>
+	inline int32 RemoveImpl(uint32 KeyHash, const ComparableKey& Key)
+	{
+		int32 NumRemovedElements = 0;
+
+		FSetElementId* NextElementId = &GetTypedHash(KeyHash);
+		while (NextElementId->IsValidId())
+		{
+			const int32 ElementIndex = NextElementId->AsInteger();
+			SetElementType& Element = Elements[ElementIndex];
+
+			if (KeyFuncs::Matches(KeyFuncs::GetSetKey(Element.Value), Key))
+			{
+				// This element matches the key, remove it from the set.  Note that Remove sets *NextElementId to point to the next
+				// element after the removed element in the hash bucket.
+				RemoveByIndex(ElementIndex);
+				NumRemovedElements++;
+
+				if constexpr (!KeyFuncs::bAllowDuplicateKeys)
+				{
+					// If the hash disallows duplicate keys, we're done removing after the first matched key.
+					break;
+				}
+			}
+			else
+			{
+				NextElementId = &Element.HashNextId;
+			}
+		}
+
+		return NumRemovedElements;
+	}
+
+public:
+	/**
+	 * Removes all elements from the set matching the specified key.
+	 * @param Key - The key to match elements against.
+	 * @return The number of elements removed.
+	 */
+	int32 Remove(KeyInitType Key)
+	{
+		if (Elements.Num())
+		{
+			return RemoveImpl(KeyFuncs::GetKeyHash(Key), Key);
+		}
+
+		return 0;
+	}
+
+	/**
+	 * Removes all elements from the set matching the specified key.
+	 *
+	 * @see		Class documentation section on ByHash() functions
+	 * @param	Key - The key to match elements against.
+	 * @return	The number of elements removed.
+	 */
+	template<typename ComparableKey>
+	int32 RemoveByHash(uint32 KeyHash, const ComparableKey& Key)
+	{
+		PRAGMA_DISABLE_DEPRECATION_WARNINGS
+		// Disable deprecations warnings to stop warnings being thrown by our check macro.
+		checkSlow(KeyHash == KeyFuncs::GetKeyHash(Key));
+		PRAGMA_ENABLE_DEPRECATION_WARNINGS
+
+		if (Elements.Num())
+		{
+			return RemoveImpl(KeyHash, Key);
+		}
+
+		return 0;
+	}
+
+	/**
+	 * Checks if the element contains an element with the given key.
+	 * @param Key - The key to check for.
+	 * @return true if the set contains an element with the given key.
+	 */
+	[[nodiscard]] UE_FORCEINLINE_HINT bool Contains(KeyInitType Key) const
+	{
+		return FindIndexByHash(KeyFuncs::GetKeyHash(Key), Key) != INDEX_NONE;
+	}
+
+	/**
+	 * Checks if the element contains an element with the given key.
+	 *
+	 * @see	Class documentation section on ByHash() functions
+	 */
+	template<typename ComparableKey>
+	[[nodiscard]] inline bool ContainsByHash(uint32 KeyHash, const ComparableKey& Key) const
+	{
+		PRAGMA_DISABLE_DEPRECATION_WARNINGS
+		// Disable deprecations warnings to stop warnings being thrown by our check macro.
+		checkSlow(KeyHash == KeyFuncs::GetKeyHash(Key));
+		PRAGMA_ENABLE_DEPRECATION_WARNINGS
+
+		return FindIndexByHash(KeyHash, Key) != INDEX_NONE;
+	}
+
+	/**
+	 * Sorts the set's elements using the provided comparison class.
+	 */
+	template <typename PREDICATE_CLASS>
+	void Sort( const PREDICATE_CLASS& Predicate )
+	{
+		// Sort the elements according to the provided comparison class.
+		Elements.Sort( FElementCompareClass< PREDICATE_CLASS >( Predicate ) );
+
+		// Rehash.
+		Rehash();
+	}
+
+	/**
+	 * Stable sorts the set's elements using the provided comparison class.
+	 */
+	template <typename PREDICATE_CLASS>
+	void StableSort(const PREDICATE_CLASS& Predicate)
+	{
+		// Sort the elements according to the provided comparison class.
+		Elements.StableSort(FElementCompareClass< PREDICATE_CLASS >(Predicate));
+
+		// Rehash.
+		Rehash();
+	}
+
+	/**
+	* Sort the free element list so that subsequent additions will occur in the lowest available
+	* TSparseArray index resulting in tighter packing without moving any existing items. Also useful for
+	* some types of determinism. @see TSparseArray::SortFreeList() for more info.
+	*/
+	void SortFreeList()
+	{
+		Elements.SortFreeList();
+	}
+
+	/**
+	 * Describes the set's contents through an output device.
+	 * @param Ar - The output device to describe the set's contents through.
+	 */
+	void Dump(FOutputDevice& Ar)
+	{
+		Ar.Logf( TEXT("TSet: %i elements, %i hash slots"), Elements.Num(), HashSize );
+		for (int32 HashIndex = 0, LocalHashSize = HashSize; HashIndex < LocalHashSize; ++HashIndex)
+		{
+			// Count the number of elements in this hash bucket.
+			int32 NumElementsInBucket = 0;
+			for(FSetElementId ElementId = GetTypedHash(HashIndex);
+				ElementId.IsValidId();
+				ElementId = Elements[ElementId.AsInteger()].HashNextId)
+			{
+				NumElementsInBucket++;
+			}
+
+			Ar.Logf(TEXT("   Hash[%i] = %i"),HashIndex,NumElementsInBucket);
+		}
+	}
+
+	[[nodiscard]] bool VerifyHashElementsKey(KeyInitType Key) const
+	{
+		bool bResult=true;
+		if (Elements.Num())
+		{
+			// iterate over all elements for the hash entry of the given key
+			// and verify that the ids are valid
+			FSetElementId ElementId = GetTypedHash(KeyFuncs::GetKeyHash(Key));
+			while( ElementId.IsValidId() )
+			{
+				if( !IsValidId(ElementId) )
+				{
+					bResult=false;
+					break;
+				}
+				ElementId = Elements[ElementId.AsInteger()].HashNextId;
+			}
+		}
+		return bResult;
+	}
+
+	void DumpHashElements(FOutputDevice& Ar)
+	{
+		for (int32 HashIndex = 0, LocalHashSize = HashSize; HashIndex < LocalHashSize; ++HashIndex)
+		{
+			Ar.Logf(TEXT("   Hash[%i]"),HashIndex);
+
+			// iterate over all elements for the all hash entries
+			// and dump info for all elements
+			FSetElementId ElementId = GetTypedHash(HashIndex);
+			while( ElementId.IsValidId() )
+			{
+				if( !IsValidId(ElementId) )
+				{
+					Ar.Logf(TEXT("		!!INVALID!! ElementId = %d"),ElementId.AsInteger());
+				}
+				else
+				{
+					Ar.Logf(TEXT("		VALID ElementId = %d"),ElementId.AsInteger());
+				}
+				ElementId = Elements[ElementId].HashNextId;
+			}
+		}
+	}
+
+	/** @return the intersection of two sets. (A AND B)*/
+	[[nodiscard]] TSet Intersect(const TSet& OtherSet) const
+	{
+		const bool bOtherSmaller = (Num() > OtherSet.Num());
+		const TSet& A = (bOtherSmaller ? OtherSet : *this);
+		const TSet& B = (bOtherSmaller ? *this : OtherSet);
+
+		TSet Result;
+		Result.Reserve(A.Num()); // Worst case is everything in smaller is in larger
+
+		for(TConstIterator SetIt(A);SetIt;++SetIt)
+		{
+			if(B.Contains(KeyFuncs::GetSetKey(*SetIt)))
+			{
+				Result.Add(*SetIt);
+			}
+		}
+		return Result;
+	}
+
+	/** @return the union of two sets. (A OR B)*/
+	[[nodiscard]] TSet Union(const TSet& OtherSet) const
+	{
+		TSet Result;
+		Result.Reserve(Num() + OtherSet.Num()); // Worst case is 2 totally unique Sets
+
+		for(TConstIterator SetIt(*this);SetIt;++SetIt)
+		{
+			Result.Add(*SetIt);
+		}
+		for(TConstIterator SetIt(OtherSet);SetIt;++SetIt)
+		{
+			Result.Add(*SetIt);
+		}
+		return Result;
+	}
+
+	/** @return the complement of two sets. (A not in B where A is this and B is Other)*/
+	[[nodiscard]] TSet Difference(const TSet& OtherSet) const
+	{
+		TSet Result;
+		Result.Reserve(Num()); // Worst case is no elements of this are in Other
+
+		for(TConstIterator SetIt(*this);SetIt;++SetIt)
+		{
+			if(!OtherSet.Contains(KeyFuncs::GetSetKey(*SetIt)))
+			{
+				Result.Add(*SetIt);
+			}
+		}
+		return Result;
+	}
+
+	/**
+	 * Determine whether the specified set is entirely included within this set
+	 *
+	 * @param OtherSet	Set to check
+	 *
+	 * @return True if the other set is entirely included in this set, false if it is not
+	 */
+	[[nodiscard]] bool Includes(const TSet<ElementType,KeyFuncs,Allocator>& OtherSet) const
+	{
+		bool bIncludesSet = true;
+		if (OtherSet.Num() <= Num())
+		{
+			for(TConstIterator OtherSetIt(OtherSet); OtherSetIt; ++OtherSetIt)
+			{
+				if (!Contains(KeyFuncs::GetSetKey(*OtherSetIt)))
+				{
+					bIncludesSet = false;
+					break;
+				}
+			}
+		}
+		else
+		{
+			// Not possible to include if it is bigger than us
+			bIncludesSet = false;
+		}
+		return bIncludesSet;
+	}
+
+	/** @return a TArray of the elements */
+	[[nodiscard]] TArray<ElementType> Array() const
+	{
+		TArray<ElementType> Result;
+		Result.Reserve(Num());
+		for(TConstIterator SetIt(*this);SetIt;++SetIt)
+		{
+			Result.Add(*SetIt);
+		}
+		return Result;
+	}
+
+	/**
+	 * Checks that the specified address is not part of an element within the container.  Used for implementations
+	 * to check that reference arguments aren't going to be invalidated by possible reallocation.
+	 *
+	 * @param Addr The address to check.
+	 */
+	UE_FORCEINLINE_HINT void CheckAddress(const ElementType* Addr) const
+	{
+		Elements.CheckAddress(Addr);
+	}
+
+	/**
+	 * Move assignment operator.
+	 * Compatible element type version.
+	 *
+	 * @param Other Set to assign and move from.
+	 */
+	template <
+		typename OtherKeyFuncs,
+		typename AliasElementType = ElementType
+		UE_REQUIRES(TIsContainerElementTypeCopyable_V<AliasElementType>)
+	>
+	TSet& operator=(TSet<typename TContainerElementTypeCompatibility<ElementType>::CopyFromOtherType, OtherKeyFuncs, Allocator>&& Other)
+	{
+		TContainerElementTypeCompatibility<ElementType>::CopyingFromOtherType();
+		Reset();
+		Append(MoveTemp(Other));
+		return *this;
+	}
+
+	/**
+	 * Assignment operator. First deletes all currently contained elements
+	 * and then copies from other set.
+	 * Compatible element type version.
+	 *
+	 * @param Other The source set to assign from.
+	 */
+	template <
+		typename OtherKeyFuncs,
+		typename OtherAllocator,
+		typename AliasElementType = ElementType
+		UE_REQUIRES(TIsContainerElementTypeCopyable_V<AliasElementType>)
+	>
+	TSet& operator=(const TSet<typename TContainerElementTypeCompatibility<ElementType>::CopyFromOtherType, OtherKeyFuncs, OtherAllocator>& Other)
+	{
+		TContainerElementTypeCompatibility<ElementType>::CopyingFromOtherType();
+		Reset();
+		Append(Other);
+		return *this;
+	}
+
+	/**
+	 * Add all items from another set to our set (union without creating a new set)
+	 * Compatible element type version.
+	 * @param OtherSet - The other set of items to add.
+	 */
+	template <
+		typename OtherKeyFuncs,
+		typename OtherAllocator,
+		typename AliasElementType = ElementType
+		UE_REQUIRES(TIsContainerElementTypeCopyable_V<AliasElementType>)
+	>
+	void Append(const TSet<typename TContainerElementTypeCompatibility<ElementType>::CopyFromOtherType, OtherKeyFuncs, OtherAllocator>& OtherSet)
+	{
+		TContainerElementTypeCompatibility<ElementType>::CopyingFromOtherType();
+		Reserve(Elements.Num() + OtherSet.Num());
+		for (const ElementType& Element : OtherSet)
+		{
+			Add(Element);
+		}
+	}
+
+	/**
+	 * Add all items from another set to our set (union without creating a new set)
+	 * Compatible element type version.
+	 * @param OtherSet - The other set of items to add.
+	 */
+	template <
+		typename OtherKeyFuncs,
+		typename AliasElementType = ElementType
+		UE_REQUIRES(TIsContainerElementTypeCopyable_V<AliasElementType>)
+	>
+	void Append(TSet<typename TContainerElementTypeCompatibility<ElementType>::CopyFromOtherType, OtherKeyFuncs, Allocator>&& OtherSet)
+	{
+		TContainerElementTypeCompatibility<ElementType>::CopyingFromOtherType();
+		Reserve(Elements.Num() + OtherSet.Num());
+		for (ElementType& Element : OtherSet)
+		{
+			Add(MoveTempIfPossible(Element));
+		}
+		OtherSet.Reset();
+	}
+
+private:
+	/** Extracts the element value from the set's element structure and passes it to the user provided comparison class. */
+	template <typename PREDICATE_CLASS>
+	class FElementCompareClass
+	{
+		TDereferenceWrapper< ElementType, PREDICATE_CLASS > Predicate;
+
+	public:
+		[[nodiscard]] UE_FORCEINLINE_HINT FElementCompareClass( const PREDICATE_CLASS& InPredicate )
+			: Predicate( InPredicate )
+		{
+		}
+
+		[[nodiscard]] UE_FORCEINLINE_HINT bool operator()( const SetElementType& A,const SetElementType& B ) const
+		{
+			return Predicate( A.Value, B.Value );
+		}
+	};
+
+	using ElementArrayType = TSparseArray<SetElementType, typename Allocator::SparseArrayAllocator>;
+	using HashType         = typename Allocator::HashAllocator::template ForElementType<FSetElementId>;
+
+	ElementArrayType Elements;
+
+	HashType Hash;
+	int32	 HashSize = 0;
+
+public:
+	void WriteMemoryImage(FMemoryImageWriter& Writer) const
+	{
+		checkf(!Writer.Is32BitTarget(), TEXT("TSet does not currently support freezing for 32bits"));
+		if constexpr (TAllocatorTraits<Allocator>::SupportsFreezeMemoryImage && THasTypeLayout<InElementType>::Value)
+		{
+			this->Elements.WriteMemoryImage(Writer);
+			this->Hash.WriteMemoryImage(Writer, StaticGetTypeLayoutDesc<FSetElementId>(), this->HashSize);
+			Writer.WriteBytes(this->HashSize);
+		}
+		else
+		{
+			Writer.WriteBytes(TSet());
+		}
+	}
+
+	void CopyUnfrozen(const FMemoryUnfreezeContent& Context, void* Dst) const
+	{
+		if constexpr (TAllocatorTraits<Allocator>::SupportsFreezeMemoryImage && THasTypeLayout<InElementType>::Value)
+		{
+			TSet* DstObject = static_cast<TSet*>(Dst);
+			this->Elements.CopyUnfrozen(Context, &DstObject->Elements);
+
+			::new((void*)&DstObject->Hash) HashType();
+			DstObject->Hash.ResizeAllocation(0, this->HashSize, sizeof(FSetElementId));
+			FMemory::Memcpy(DstObject->Hash.GetAllocation(), this->Hash.GetAllocation(), sizeof(FSetElementId) * this->HashSize);
+			DstObject->HashSize = this->HashSize;
+		}
+		else
+		{
+			::new(Dst) TSet();
+		}
+	}
+
+	static void AppendHash(const FPlatformTypeLayoutParameters& LayoutParams, FSHA1& Hasher)
+	{
+		ElementArrayType::AppendHash(LayoutParams, Hasher);
+	}
+
+private:
+
+	[[nodiscard]] UE_FORCEINLINE_HINT FSetElementId& GetTypedHash(int32 HashIndex) const
+	{
+		return ((FSetElementId*)Hash.GetAllocation())[HashIndex & (HashSize - 1)];
+	}
+
+	/** Links an added element to the hash chain. */
+	inline void LinkElement(SizeType ElementIndex, const SetElementType& Element, uint32 KeyHash) const
+	{
+		// Compute the hash bucket the element goes in.
+		Element.HashIndex = KeyHash & (HashSize - 1);
+
+		// Link the element into the hash bucket.
+		FSetElementId& TypedHash = GetTypedHash(Element.HashIndex);
+		Element.HashNextId = TypedHash;
+		TypedHash = FSetElementId::FromInteger(ElementIndex);
+	}
+
+	/** Hashes and links an added element to the hash chain. */
+	UE_FORCEINLINE_HINT void HashElement(SizeType ElementIndex, const SetElementType& Element) const
+	{
+		LinkElement(ElementIndex, Element, KeyFuncs::GetKeyHash(KeyFuncs::GetSetKey(Element.Value)));
+	}
+
+	/** Reset hash buckets to invalid */
+	void UnhashElements()
+	{
+		FSetElementId* HashPtr = Hash.GetAllocation();
+
+		// Check if it should be faster to clear the hash by going through elements instead of resetting the whole hash
+		if (Num() < (HashSize / 4))
+		{
+			// Faster path: only reset hash buckets to invalid for elements in the hash
+			for (const SetElementType& Element: Elements)
+			{
+				HashPtr[Element.HashIndex] = FSetElementId();
+			}
+		}
+		else
+		{
+			/** Reset the whole hash to invalid */
+			for (int32 I = 0; I < HashSize; ++I)
+			{
+				HashPtr[I] = FSetElementId();
+			}
+		}
+	}
+
+	/**
+	 * Checks if the hash has an appropriate number of buckets, and if it should be resized.
+	 * @param NumHashedElements - The number of elements to size the hash for.
+	 * @param DesiredHashSize - Desired size if we should rehash.
+	 * @param AllowShrinking - If the hash is allowed to shrink.
+	 * @return true if the set should berehashed.
+	 */
+	[[nodiscard]] UE_FORCEINLINE_HINT bool ShouldRehash(int32 NumHashedElements, int32 DesiredHashSize, EAllowShrinking AllowShrinking) const
+	{
+		// If the hash hasn't been created yet, or is smaller than the desired hash size, rehash.
+		// If shrinking is allowed and the hash is bigger than the desired hash size, rehash.
+		return ((NumHashedElements > 0 && HashSize < DesiredHashSize) || (AllowShrinking == EAllowShrinking::Yes && HashSize > DesiredHashSize));
+	}
+
+	/**
+	 * Checks if the hash has an appropriate number of buckets, and if not resizes it.
+	 * @param NumHashedElements - The number of elements to size the hash for.
+	 * @param AllowShrinking - If the hash is allowed to shrink.
+	 * @return true if the set was rehashed.
+	 */
+	bool ConditionalRehash(int32 NumHashedElements, EAllowShrinking AllowShrinking)
+	{
+		// Calculate the desired hash size for the specified number of elements.
+		const int32 DesiredHashSize = Allocator::GetNumberOfHashBuckets(NumHashedElements);
+
+		if (ShouldRehash(NumHashedElements, DesiredHashSize, AllowShrinking))
+		{
+			HashSize = DesiredHashSize;
+			Rehash();
+			return true;
+		}
+
+		return false;
+	}
+
+	/** Resizes the hash. */
+	void Rehash()
+	{
+		UE::Core::Private::Rehash(Hash, HashSize);
+
+		if (HashSize)
+		{
+			// Add the existing elements to the new hash.
+			for(typename ElementArrayType::TConstIterator ElementIt(Elements);ElementIt;++ElementIt)
+			{
+				HashElement(ElementIt.GetIndex(), *ElementIt);
+			}
+		}
+	}
+
+	/** The base type of whole set iterators. */
+	template<bool bConst, bool bRangedFor = false>
+	class TBaseIterator
+	{
+	private:
+		friend class TSet;
+
+		typedef std::conditional_t<bConst,const ElementType,ElementType> ItElementType;
+
+	public:
+		typedef std::conditional_t<
+			bConst,
+			std::conditional_t<bRangedFor, typename ElementArrayType::TRangedForConstIterator, typename ElementArrayType::TConstIterator>,
+			std::conditional_t<bRangedFor, typename ElementArrayType::TRangedForIterator,      typename ElementArrayType::TIterator     >
+		> ElementItType;
+
+		[[nodiscard]] UE_FORCEINLINE_HINT TBaseIterator(const ElementItType& InElementIt)
+			: ElementIt(InElementIt)
+		{
+		}
+
+		/** Advances the iterator to the next element. */
+		UE_FORCEINLINE_HINT TBaseIterator& operator++()
+		{
+			++ElementIt;
+			return *this;
+		}
+
+		/** conversion to "bool" returning true if the iterator is valid. */
+		[[nodiscard]] UE_FORCEINLINE_HINT explicit operator bool() const
+		{
+			return !!ElementIt;
+		}
+		/** inverse of the "bool" operator */
+		[[nodiscard]] UE_FORCEINLINE_HINT bool operator!() const
+		{
+			return !(bool)*this;
+		}
+
+		// Accessors.
+		[[nodiscard]] UE_FORCEINLINE_HINT FSetElementId GetId() const
+		{
+			return FSetElementId::FromInteger(ElementIt.GetIndex());
+		}
+		[[nodiscard]] UE_FORCEINLINE_HINT ItElementType* operator->() const
+		{
+			return &ElementIt->Value;
+		}
+		[[nodiscard]] UE_FORCEINLINE_HINT ItElementType& operator*() const
+		{
+			return ElementIt->Value;
+		}
+
+		[[nodiscard]] UE_FORCEINLINE_HINT bool operator!=(const TBaseIterator& Rhs) const { return ElementIt != Rhs.ElementIt; }
+		[[nodiscard]] UE_FORCEINLINE_HINT bool operator==(const TBaseIterator& Rhs) const { return ElementIt == Rhs.ElementIt; }
+
+		ElementItType ElementIt;
+	};
+
+	/** The base type of whole set iterators. */
+	template<bool bConst>
+	class TBaseKeyIterator
+	{
+	private:
+		typedef std::conditional_t<bConst, const TSet, TSet> SetType;
+		typedef std::conditional_t<bConst,const ElementType,ElementType> ItElementType;
+		typedef typename TTypeTraits<typename KeyFuncs::KeyType>::ConstPointerType ReferenceOrValueType;
+
+	public:
+		using KeyArgumentType =
+			std::conditional_t<
+				std::is_reference_v<ReferenceOrValueType>,
+				TRetainedRef<std::remove_reference_t<ReferenceOrValueType>>,
+				KeyInitType
+			>;
+
+		/** Initialization constructor. */
+		[[nodiscard]] inline TBaseKeyIterator(SetType& InSet, KeyArgumentType InKey)
+			: Set  (InSet)
+			, Key  (InKey) //-V1041
+			, Index(INDEX_NONE)
+		{
+			if (Set.HashSize)
+			{
+				NextIndex = Set.GetTypedHash(KeyFuncs::GetKeyHash(Key)).AsInteger();
+				++(*this);
+			}
+			else
+			{
+				NextIndex = INDEX_NONE;
+			}
+		}
+
+		/** Advances the iterator to the next element. */
+		inline TBaseKeyIterator& operator++()
+		{
+			Index = NextIndex;
+
+			while (Index != INDEX_NONE)
+			{
+				NextIndex = Set.Elements[Index].HashNextId.AsInteger();
+				checkSlow(Index != NextIndex);
+
+				if (KeyFuncs::Matches(KeyFuncs::GetSetKey(Set.Elements[Index].Value),Key))
+				{
+					break;
+				}
+
+				Index = NextIndex;
+			}
+			return *this;
+		}
+
+		/** conversion to "bool" returning true if the iterator is valid. */
+		[[nodiscard]] UE_FORCEINLINE_HINT explicit operator bool() const
+		{
+			return Index != INDEX_NONE;
+		}
+		/** inverse of the "bool" operator */
+		[[nodiscard]] UE_FORCEINLINE_HINT bool operator!() const
+		{
+			return !(bool)*this;
+		}
+
+		// Accessors.
+		[[nodiscard]] UE_FORCEINLINE_HINT FSetElementId GetId() const
+		{
+			return FSetElementId::FromInteger(Index);
+		}
+		[[nodiscard]] UE_FORCEINLINE_HINT ItElementType* operator->() const
+		{
+			return &Set.Elements[Index].Value;
+		}
+		[[nodiscard]] UE_FORCEINLINE_HINT ItElementType& operator*() const
+		{
+			return Set.Elements[Index].Value;
+		}
+
+	protected:
+		SetType& Set;
+		ReferenceOrValueType Key;
+		SizeType Index;
+		SizeType NextIndex;
+	};
+
+public:
+
+	/** Used to iterate over the elements of a const TSet. */
+	class TConstIterator : public TBaseIterator<true>
+	{
+		friend class TSet;
+
+	public:
+		[[nodiscard]] UE_FORCEINLINE_HINT TConstIterator(const TSet& InSet)
+			: TBaseIterator<true>(InSet.Elements.begin())
+		{
+		}
+	};
+
+	/** Used to iterate over the elements of a TSet. */
+	class TIterator : public TBaseIterator<false>
+	{
+		friend class TSet;
+
+	public:
+		[[nodiscard]] inline TIterator(TSet& InSet)
+			: TBaseIterator<false>(InSet.Elements.begin())
+			, Set                 (InSet)
+		{
+		}
+
+		/** Removes the current element from the set. */
+		UE_FORCEINLINE_HINT void RemoveCurrent()
+		{
+			Set.RemoveByIndex(TBaseIterator<false>::ElementIt.GetIndex());
+		}
+
+	private:
+		TSet& Set;
+	};
+
+	using TRangedForConstIterator = TBaseIterator<true, true>;
+	using TRangedForIterator      = TBaseIterator<false, true>;
+
+	/** Used to iterate over the elements of a const TSet. */
+	class TConstKeyIterator : public TBaseKeyIterator<true>
+	{
+	private:
+		using Super = TBaseKeyIterator<true>;
+
+	public:
+		using KeyArgumentType = typename Super::KeyArgumentType;
+
+		[[nodiscard]] UE_FORCEINLINE_HINT TConstKeyIterator(const TSet& InSet, KeyArgumentType InKey)
+			: Super(InSet, InKey)
+		{
+		}
+	};
+
+	/** Used to iterate over the elements of a TSet. */
+	class TKeyIterator : public TBaseKeyIterator<false>
+	{
+	private:
+		using Super = TBaseKeyIterator<false>;
+
+	public:
+		using KeyArgumentType = typename Super::KeyArgumentType;
+
+		[[nodiscard]] UE_FORCEINLINE_HINT TKeyIterator(TSet& InSet, KeyArgumentType InKey)
+			: Super(InSet, InKey)
+		{
+		}
+
+		/** Removes the current element from the set. */
+		inline void RemoveCurrent()
+		{
+			this->Set.RemoveByIndex(TBaseKeyIterator<false>::Index);
+			TBaseKeyIterator<false>::Index = INDEX_NONE;
+		}
+	};
+
+	/** Creates an iterator for the contents of this set */
+	[[nodiscard]] UE_FORCEINLINE_HINT TIterator CreateIterator()
+	{
+		return TIterator(*this);
+	}
+
+	/** Creates a const iterator for the contents of this set */
+	[[nodiscard]] UE_FORCEINLINE_HINT TConstIterator CreateConstIterator() const
+	{
+		return TConstIterator(*this);
+	}
+
+	friend struct TSetPrivateFriend;
+
+public:
+	/**
+	 * DO NOT USE DIRECTLY
+	 * STL-like iterators to enable range-based for loop support.
+	 */
+	[[nodiscard]] UE_FORCEINLINE_HINT TRangedForIterator begin()
+	{
+		return TRangedForIterator(Elements.begin());
+	}
+	[[nodiscard]] UE_FORCEINLINE_HINT TRangedForConstIterator begin() const
+	{
+		return TRangedForConstIterator(Elements.begin());
+	}
+	[[nodiscard]] UE_FORCEINLINE_HINT TRangedForIterator end()
+	{
+		return TRangedForIterator(Elements.end());
+	}
+	[[nodiscard]] UE_FORCEINLINE_HINT TRangedForConstIterator end() const
+	{
+		return TRangedForConstIterator(Elements.end());
+	}
+
+	// Maps are deliberately prevented from being hashed or compared, because this would hide potentially major performance problems behind default operations.
+	friend uint32 GetTypeHash(const TSet& Set) = delete;
+	friend bool operator==(const TSet&, const TSet&) = delete;
+	friend bool operator!=(const TSet&, const TSet&) = delete;
+};
+
+template <typename RangeType>
+TSet(RangeType&&) -> TSet<TElementType_T<RangeType>>;
+
+namespace Freeze
+{
+	template<typename ElementType, typename KeyFuncs, typename Allocator>
+	void IntrinsicWriteMemoryImage(FMemoryImageWriter& Writer, const TSet<ElementType, KeyFuncs, Allocator>& Object, const FTypeLayoutDesc&)
+	{
+		Object.WriteMemoryImage(Writer);
+	}
+
+	template<typename ElementType, typename KeyFuncs, typename Allocator>
+	uint32 IntrinsicUnfrozenCopy(const FMemoryUnfreezeContent& Context, const TSet<ElementType, KeyFuncs, Allocator>& Object, void* OutDst)
+	{
+		Object.CopyUnfrozen(Context, OutDst);
+		return sizeof(Object);
+	}
+
+	template<typename ElementType, typename KeyFuncs, typename Allocator>
+	uint32 IntrinsicAppendHash(const TSet<ElementType, KeyFuncs, Allocator>* DummyObject, const FTypeLayoutDesc& TypeDesc, const FPlatformTypeLayoutParameters& LayoutParams, FSHA1& Hasher)
+	{
+		TSet<ElementType, KeyFuncs, Allocator>::AppendHash(LayoutParams, Hasher);
+		return DefaultAppendHash(TypeDesc, LayoutParams, Hasher);
+	}
+}
+
+DECLARE_TEMPLATE_INTRINSIC_TYPE_LAYOUT((template <typename ElementType, typename KeyFuncs, typename Allocator>), (TSet<ElementType, KeyFuncs, Allocator>));
+
+struct FScriptSetLayout
+{
+	// int32 ElementOffset = 0; // always at zero offset from the TSetElement - not stored here
+	int32 HashNextIdOffset;
+	int32 HashIndexOffset;
+	int32 Size;
+
+	FScriptSparseArrayLayout SparseArrayLayout;
+};
+
+// Untyped set type for accessing TSet data, like FScriptArray for TArray.
+// Must have the same memory representation as a TSet.
+template <typename Allocator, typename InDerivedType>
+class TScriptSet
+{
+	using DerivedType = std::conditional_t<std::is_void_v<InDerivedType>, TScriptSet, InDerivedType>;
+
+public:
+	[[nodiscard]] static constexpr FScriptSetLayout GetScriptLayout(int32 ElementSize, int32 ElementAlignment)
+	{
+		FScriptSetLayout Result;
+
+		// TSetElement<TPair<Key, Value>>
+		FStructBuilder SetElementStruct;
+		int32 ElementOffset      = SetElementStruct.AddMember(ElementSize,           ElementAlignment);
+		Result.HashNextIdOffset  = SetElementStruct.AddMember(sizeof(FSetElementId), alignof(FSetElementId));
+		Result.HashIndexOffset   = SetElementStruct.AddMember(sizeof(int32),         alignof(int32));
+		Result.Size              = SetElementStruct.GetSize();
+		Result.SparseArrayLayout = FScriptSparseArray::GetScriptLayout(SetElementStruct.GetSize(), SetElementStruct.GetAlignment());
+
+		checkf(ElementOffset == 0, TEXT("The element inside the TSetElement is expected to be at the start of the struct"));
+
+		return Result;
+	}
+
+	[[nodiscard]] TScriptSet()
+		: HashSize(0)
+	{
+	}
+
+	///////////////////////////////////////////////////
+	// Start - intrusive TOptional<TScriptSet> state //
+	///////////////////////////////////////////////////
+	constexpr static bool bHasIntrusiveUnsetOptionalState = true;
+	using IntrusiveUnsetOptionalStateType = TScriptSet;
+
+	[[nodiscard]] explicit TScriptSet(FIntrusiveUnsetOptionalState Tag)
+		: Elements(Tag)
+	{
+	}
+	[[nodiscard]] bool operator==(FIntrusiveUnsetOptionalState Tag) const
+	{
+		return Elements == Tag;
+	}
+	/////////////////////////////////////////////////
+	// End - intrusive TOptional<TScriptSet> state //
+	/////////////////////////////////////////////////
+
+	[[nodiscard]] bool IsValidIndex(int32 Index) const
+	{
+		return Elements.IsValidIndex(Index);
+	}
+
+	[[nodiscard]] bool IsEmpty() const
+	{
+		return Elements.IsEmpty();
+	}
+
+	[[nodiscard]] int32 Num() const
+	{
+		return Elements.Num();
+	}
+
+	/** @return The number of elements the set can hold before reallocation. */
+	[[nodiscard]] int32 Max() const
+	{
+		return Elements.Max();
+	}
+
+	[[nodiscard]] int32 NumUnchecked() const
+	{
+		return Elements.NumUnchecked();
+	}
+
+	[[nodiscard]] int32 GetMaxIndex() const
+	{
+		return Elements.GetMaxIndex();
+	}
+
+	[[nodiscard]] void* GetData(int32 Index, const FScriptSetLayout& Layout)
+	{
+		return Elements.GetData(Index, Layout.SparseArrayLayout);
+	}
+
+	[[nodiscard]] const void* GetData(int32 Index, const FScriptSetLayout& Layout) const
+	{
+		return Elements.GetData(Index, Layout.SparseArrayLayout);
+	}
+
+	void MoveAssign(DerivedType& Other, const FScriptSetLayout& Layout)
+	{
+		checkSlow(this != &Other);
+		Empty(0, Layout);
+		Elements.MoveAssign(Other.Elements, Layout.SparseArrayLayout);
+		Hash.MoveToEmpty(Other.Hash);
+		HashSize = Other.HashSize; Other.HashSize = 0;
+	}
+
+	void Empty(int32 Slack, const FScriptSetLayout& Layout)
+	{
+		// Empty the elements array, and reallocate it for the expected number of elements.
+		Elements.Empty(Slack, Layout.SparseArrayLayout);
+
+		// Calculate the desired hash size for the specified number of elements.
+		const int32 DesiredHashSize = Allocator::GetNumberOfHashBuckets(Slack);
+
+		// If the hash hasn't been created yet, or is smaller than the desired hash size, rehash.
+		if (Slack != 0 && (HashSize == 0 || HashSize != DesiredHashSize))
+		{
+			HashSize = DesiredHashSize;
+
+			// Free the old hash.
+			Hash.ResizeAllocation(0, HashSize, sizeof(FSetElementId));
+		}
+
+		FSetElementId* HashPtr = Hash.GetAllocation();
+		for (int32 I = 0; I < HashSize; ++I)
+		{
+			HashPtr[I] = FSetElementId();
+		}
+	}
+
+	void RemoveAt(int32 Index, const FScriptSetLayout& Layout)
+	{
+		check(IsValidIndex(Index));
+
+		void* ElementBeingRemoved = Elements.GetData(Index, Layout.SparseArrayLayout);
+
+		// Remove the element from the hash.
+		for (FSetElementId* NextElementId = &GetTypedHash(GetHashIndexRef(ElementBeingRemoved, Layout)); NextElementId->IsValidId(); NextElementId = &GetHashNextIdRef(Elements.GetData(NextElementId->AsInteger(), Layout.SparseArrayLayout), Layout))
+		{
+			if (NextElementId->AsInteger() == Index)
+			{
+				*NextElementId = GetHashNextIdRef(ElementBeingRemoved, Layout);
+				break;
+			}
+		}
+
+		// Remove the element from the elements array.
+		Elements.RemoveAtUninitialized(Layout.SparseArrayLayout, Index);
+	}
+
+	/**
+	 * Adds an uninitialized object to the set.
+	 * The set will need rehashing at some point after this call to make it valid.
+	 *
+	 * @return  The index of the added element.
+	 */
+	int32 AddUninitialized(const FScriptSetLayout& Layout)
+	{
+		int32 Result = Elements.AddUninitialized(Layout.SparseArrayLayout);
+		return Result;
+	}
+
+	void Rehash(const FScriptSetLayout& Layout, TFunctionRef<uint32 (const void*)> GetKeyHash)
+	{
+		// Free the old hash.
+		Hash.ResizeAllocation(0,0,sizeof(FSetElementId));
+
+		HashSize = Allocator::GetNumberOfHashBuckets(Elements.Num());
+		if (HashSize)
+		{
+			// Allocate the new hash.
+			checkSlow(FMath::IsPowerOfTwo(HashSize));
+			Hash.ResizeAllocation(0, HashSize, sizeof(FSetElementId));
+			for (int32 HashIndex = 0; HashIndex < HashSize; ++HashIndex)
+			{
+				GetTypedHash(HashIndex) = FSetElementId();
+			}
+
+			// Add the existing elements to the new hash.
+			int32 Index = 0;
+			int32 Count = Elements.Num();
+			while (Count)
+			{
+				if (Elements.IsValidIndex(Index))
+				{
+					FSetElementId ElementId = FSetElementId::FromInteger(Index);
+
+					void* Element = (uint8*)Elements.GetData(Index, Layout.SparseArrayLayout);
+
+					// Compute the hash bucket the element goes in.
+					uint32 KeyHash = GetKeyHash(Element);
+					int32  HashIndex   = KeyHash & (HashSize - 1);
+					GetHashIndexRef(Element, Layout) = KeyHash & (HashSize - 1);
+
+					// Link the element into the hash bucket.
+					GetHashNextIdRef(Element, Layout) = GetTypedHash(HashIndex);
+					GetTypedHash(HashIndex) = ElementId;
+
+					--Count;
+				}
+
+				++Index;
+			}
+		}
+	}
+
+private:
+	[[nodiscard]] int32 FindIndexImpl(const void* Element, const FScriptSetLayout& Layout, uint32 KeyHash, TFunctionRef<bool (const void*, const void*)> EqualityFn) const
+	{
+		const int32  HashIndex = KeyHash & (HashSize - 1);
+
+		uint8* CurrentElement = nullptr;
+		for (FSetElementId ElementId = GetTypedHash(HashIndex);
+			ElementId.IsValidId();
+			ElementId = GetHashNextIdRef(CurrentElement, Layout))
+		{
+			const int32 ElementIndex = ElementId.AsInteger();
+			CurrentElement = (uint8*)Elements.GetData(ElementIndex, Layout.SparseArrayLayout);
+			if (EqualityFn(Element, CurrentElement))
+			{
+				return ElementIndex;
+			}
+		}
+
+		return INDEX_NONE;
+	}
+
+public:
+	[[nodiscard]] int32 FindIndex(const void* Element, const FScriptSetLayout& Layout, TFunctionRef<uint32 (const void*)> GetKeyHash, TFunctionRef<bool (const void*, const void*)> EqualityFn) const
+	{
+		if (Elements.Num())
+		{
+			return FindIndexImpl(Element, Layout, GetKeyHash(Element), EqualityFn);
+		}
+
+		return INDEX_NONE;
+	}
+
+
+	[[nodiscard]] int32 FindIndexByHash(const void* Element, const FScriptSetLayout& Layout, uint32 KeyHash, TFunctionRef<bool (const void*, const void*)> EqualityFn) const
+	{
+		if (Elements.Num())
+		{
+			return FindIndexImpl(Element, Layout, KeyHash, EqualityFn);
+		}
+
+		return INDEX_NONE;
+	}
+
+	int32 FindOrAdd(const void* Element, const FScriptSetLayout& Layout, TFunctionRef<uint32(const void*)> GetKeyHash, TFunctionRef<bool(const void*, const void*)> EqualityFn, TFunctionRef<void(void*)> ConstructFn)
+	{
+		uint32 KeyHash = GetKeyHash(Element);
+		int32 OldElementIndex = FindIndexByHash(Element, Layout, KeyHash, EqualityFn);
+		if (OldElementIndex != INDEX_NONE)
+		{
+			return OldElementIndex;
+		}
+
+		return AddNewElement(Layout, GetKeyHash, KeyHash, ConstructFn);
+	}
+
+	void Add(const void* Element, const FScriptSetLayout& Layout, TFunctionRef<uint32(const void*)> GetKeyHash, TFunctionRef<bool(const void*, const void*)> EqualityFn, TFunctionRef<void(void*)> ConstructFn, TFunctionRef<void(void*)> DestructFn)
+	{
+		uint32 KeyHash = GetKeyHash(Element);
+		int32 OldElementIndex = FindIndexByHash(Element, Layout, KeyHash, EqualityFn);
+		if (OldElementIndex != INDEX_NONE)
+		{
+			void* ElementPtr = Elements.GetData(OldElementIndex, Layout.SparseArrayLayout);
+
+			DestructFn(ElementPtr);
+			ConstructFn(ElementPtr);
+
+			// We don't update the hash because we don't need to - the new element
+			// should have the same hash, but let's just check.
+			PRAGMA_DISABLE_DEPRECATION_WARNINGS
+			// Disable deprecations warnings to stop warnings being thrown by our check macro.
+			checkSlow(KeyHash == GetKeyHash(ElementPtr));
+			PRAGMA_ENABLE_DEPRECATION_WARNINGS
+		}
+		else
+		{
+			AddNewElement(Layout, GetKeyHash, KeyHash, ConstructFn);
+		}
+	}
+
+private:
+	int32 AddNewElement(const FScriptSetLayout& Layout, TFunctionRef<uint32(const void*)> GetKeyHash, uint32 KeyHash, TFunctionRef<void(void*)> ConstructFn)
+	{
+		int32 NewElementIndex = Elements.AddUninitialized(Layout.SparseArrayLayout);
+
+		void* ElementPtr = Elements.GetData(NewElementIndex, Layout.SparseArrayLayout);
+		ConstructFn(ElementPtr);
+
+		const int32 DesiredHashSize = FDefaultSetAllocator::GetNumberOfHashBuckets(Num());
+		if (!HashSize || HashSize < DesiredHashSize)
+		{
+			// rehash, this will link in our new element if needed:
+			Rehash(Layout, GetKeyHash);
+		}
+		else
+		{
+			// link the new element into the set:
+			int32 HashIndex = KeyHash & (HashSize - 1);
+			FSetElementId& TypedHash = GetTypedHash(HashIndex);
+			GetHashIndexRef(ElementPtr, Layout) = HashIndex;
+			GetHashNextIdRef(ElementPtr, Layout) = TypedHash;
+			TypedHash = FSetElementId::FromInteger(NewElementIndex);
+		}
+
+		return NewElementIndex;
+	}
+
+	typedef TScriptSparseArray<typename Allocator::SparseArrayAllocator> ElementArrayType;
+	typedef typename Allocator::HashAllocator::template ForElementType<FSetElementId> HashType;
+
+	ElementArrayType Elements;
+	HashType         Hash;
+	int32            HashSize;
+
+	[[nodiscard]] UE_FORCEINLINE_HINT FSetElementId& GetTypedHash(int32 HashIndex) const
+	{
+		return ((FSetElementId*)Hash.GetAllocation())[HashIndex & (HashSize - 1)];
+	}
+
+	[[nodiscard]] static FSetElementId& GetHashNextIdRef(const void* Element, const FScriptSetLayout& Layout)
+	{
+		return *(FSetElementId*)((uint8*)Element + Layout.HashNextIdOffset);
+	}
+
+	[[nodiscard]] static int32& GetHashIndexRef(const void* Element, const FScriptSetLayout& Layout)
+	{
+		return *(int32*)((uint8*)Element + Layout.HashIndexOffset);
+	}
+
+	// This function isn't intended to be called, just to be compiled to validate the correctness of the type.
+	static void CheckConstraints()
+	{
+		typedef TScriptSet  ScriptType;
+		typedef TSet<int32> RealType;
+
+		// Check that the class footprint is the same
+		static_assert(sizeof (ScriptType) == sizeof (RealType), "TScriptSet's size doesn't match TSet");
+		static_assert(alignof(ScriptType) == alignof(RealType), "TScriptSet's alignment doesn't match TSet");
+
+		// Check member sizes
+		static_assert(sizeof(DeclVal<ScriptType>().Elements) == sizeof(DeclVal<RealType>().Elements), "TScriptSet's Elements member size does not match TSet's");
+		static_assert(sizeof(DeclVal<ScriptType>().Hash)     == sizeof(DeclVal<RealType>().Hash),     "TScriptSet's Hash member size does not match TSet's");
+		static_assert(sizeof(DeclVal<ScriptType>().HashSize) == sizeof(DeclVal<RealType>().HashSize), "TScriptSet's HashSize member size does not match TSet's");
+
+		// Check member offsets
+		static_assert(STRUCT_OFFSET(ScriptType, Elements) == STRUCT_OFFSET(RealType, Elements), "TScriptSet's Elements member offset does not match TSet's");
+		static_assert(STRUCT_OFFSET(ScriptType, Hash)     == STRUCT_OFFSET(RealType, Hash),     "TScriptSet's Hash member offset does not match TSet's");
+		static_assert(STRUCT_OFFSET(ScriptType, HashSize) == STRUCT_OFFSET(RealType, HashSize), "TScriptSet's FirstFreeIndex member offset does not match TSet's");
+	}
+
+public:
+	// These should really be private, because they shouldn't be called, but there's a bunch of code
+	// that needs to be fixed first.
+	TScriptSet(const TScriptSet&) { check(false); }
+	void operator=(const TScriptSet&) { check(false); }
+};
+
+template <typename AllocatorType, typename InDerivedType>
+struct TIsZeroConstructType<TScriptSet<AllocatorType, InDerivedType>>
+{
+	enum { Value = true };
+};
+
+/**
+ * Traits class which determines whether or not a type is a TSet.
+ */
+template <typename T> struct TIsTSet { enum { Value = false }; };
+
+template <typename ElementType, typename KeyFuncs, typename Allocator> struct TIsTSet<               TSet<ElementType, KeyFuncs, Allocator>> { enum { Value = true }; };
+template <typename ElementType, typename KeyFuncs, typename Allocator> struct TIsTSet<const          TSet<ElementType, KeyFuncs, Allocator>> { enum { Value = true }; };
+template <typename ElementType, typename KeyFuncs, typename Allocator> struct TIsTSet<      volatile TSet<ElementType, KeyFuncs, Allocator>> { enum { Value = true }; };
+template <typename ElementType, typename KeyFuncs, typename Allocator> struct TIsTSet<const volatile TSet<ElementType, KeyFuncs, Allocator>> { enum { Value = true }; };
+
+class FScriptSet : public TScriptSet<FDefaultSetAllocator, FScriptSet>
+{
+	using Super = TScriptSet<FDefaultSetAllocator, FScriptSet>;
+
+public:
+	using Super::Super;
+
+	///////////////////////////////////////////////////
+	// Start - intrusive TOptional<FScriptSet> state //
+	///////////////////////////////////////////////////
+	using IntrusiveUnsetOptionalStateType = FScriptSet;
+	/////////////////////////////////////////////////
+	// End - intrusive TOptional<FScriptSet> state //
+	/////////////////////////////////////////////////
+};
+
+struct TSetPrivateFriend
+{
+	/** Serializer. */
+	template<typename ElementType, typename KeyFuncs,typename Allocator>
+	static FArchive& Serialize(FArchive& Ar,TSet<ElementType, KeyFuncs, Allocator>& Set)
+	{
+		// Load the set's new elements.
+		Ar << Set.Elements;
+
+		if(Ar.IsLoading() || (Ar.IsModifyingWeakAndStrongReferences() && !Ar.IsSaving()))
+		{
+			// Free the old hash.
+			Set.Hash.ResizeAllocation(0,0,sizeof(FSetElementId));
+			Set.HashSize = 0;
+
+			// Hash the newly loaded elements.
+			Set.ConditionalRehash(Set.Elements.Num(), EAllowShrinking::No);
+		}
+
+		return Ar;
+	}
+
+	/** Structured archive serializer. */
+	template<typename ElementType, typename KeyFuncs,typename Allocator>
+ 	static void SerializeStructured(FStructuredArchive::FSlot Slot, TSet<ElementType, KeyFuncs, Allocator>& Set)
+ 	{
+		Slot << Set.Elements;
+
+		if (Slot.GetUnderlyingArchive().IsLoading() || (Slot.GetUnderlyingArchive().IsModifyingWeakAndStrongReferences() && !Slot.GetUnderlyingArchive().IsSaving()))
+		{
+			// Free the old hash.
+			Set.Hash.ResizeAllocation(0, 0, sizeof(FSetElementId));
+			Set.HashSize = 0;
+
+			// Hash the newly loaded elements.
+			Set.ConditionalRehash(Set.Elements.Num(), EAllowShrinking::No);
+		}
+ 	}
+
+	// Legacy comparison operators.  Note that these also test whether the set's elements were added in the same order!
+	template<typename ElementType, typename KeyFuncs,typename Allocator>
+	[[nodiscard]] static bool LegacyCompareEqual(const TSet<ElementType, KeyFuncs, Allocator>& A, const TSet<ElementType, KeyFuncs, Allocator>& B)
+	{
+		return A.Elements == B.Elements;
+	}
+};
+
+/** Serializer. */
+template <typename ElementType>
+UE_FORCEINLINE_HINT FArchive& operator<<(FArchive& Ar, TSetElement<ElementType>& Element)
+{
+	return Ar << Element.Value;
+}
+
+/** Structured archive serializer. */
+template <typename ElementType>
+UE_FORCEINLINE_HINT void operator<<(FStructuredArchive::FSlot& Ar, TSetElement<ElementType>& Element)
+{
+	Ar << Element.Value;
+}
+
+/** Serializer. */
+template<typename ElementType, typename KeyFuncs,typename Allocator>
+FArchive& operator<<(FArchive& Ar, TSet<ElementType, KeyFuncs, Allocator>& Set)
+{
+	return TSetPrivateFriend::Serialize(Ar, Set);
+}
+
+/** Structured archive serializer. */
+template<typename ElementType, typename KeyFuncs,typename Allocator>
+void operator<<(FStructuredArchive::FSlot& Ar, TSet<ElementType, KeyFuncs, Allocator>& Set)
+{
+	TSetPrivateFriend::SerializeStructured(Ar, Set);
+}
+
+// Legacy comparison operators.  Note that these also test whether the set's elements were added in the same order!
+template<typename ElementType, typename KeyFuncs,typename Allocator>
+[[nodiscard]] bool LegacyCompareEqual(const TSet<ElementType, KeyFuncs, Allocator>& A,const TSet<ElementType, KeyFuncs, Allocator>& B)
+{
+	return TSetPrivateFriend::LegacyCompareEqual(A, B);
+}
+template<typename ElementType, typename KeyFuncs,typename Allocator>
+[[nodiscard]] bool LegacyCompareNotEqual(const TSet<ElementType, KeyFuncs, Allocator>& A,const TSet<ElementType, KeyFuncs, Allocator>& B)
+{
+	return !TSetPrivateFriend::LegacyCompareEqual(A, B);
+}
 
 #if UE_ENABLE_INCLUDE_ORDER_DEPRECATED_IN_5_4
 #include "Templates/Decay.h"
diff --git a/Engine/Source/Runtime/Core/Public/Containers/SetUtilities.h b/Engine/Source/Runtime/Core/Public/Containers/SetUtilities.h
--- a/Engine/Source/Runtime/Core/Public/Containers/SetUtilities.h
+++ b/Engine/Source/Runtime/Core/Public/Containers/SetUtilities.h
@@ -5,7 +5,6 @@
 #include "CoreTypes.h"
 #include "Serialization/MemoryLayout.h"
 #include "Templates/MemoryOps.h"
-#include "Templates/Tuple.h"
 #include "Templates/UnrealTypeTraits.h"
 
 /**
diff --git a/Engine/Source/Runtime/Core/Public/Containers/SparseSet.h b/Engine/Source/Runtime/Core/Public/Containers/SparseSet.h
deleted file mode 100644
diff --git a/Engine/Source/Runtime/Core/Tests/Containers/SetTests.cpp b/Engine/Source/Runtime/Core/Tests/Containers/SetTests.cpp
--- a/Engine/Source/Runtime/Core/Tests/Containers/SetTests.cpp
+++ b/Engine/Source/Runtime/Core/Tests/Containers/SetTests.cpp
@@ -5,9 +5,8 @@
 #include "Algo/IndexOf.h"
 #include "Algo/RandomShuffle.h"
 #include "Containers/CompactSet.h"
+#include "Containers/Set.h"
 #include "Containers/ScriptCompactSet.h"
-#include "Containers/SparseSet.h"
-#include "Containers/ScriptSparseSet.h"
 #include "Containers/Map.h"
 #include "Tests/TestHarnessAdapter.h"
 
@@ -17,22 +16,22 @@
 #include <catch2/benchmark/catch_benchmark.hpp>
 #endif
 
-struct FSparseSetTestResolver
+struct FSparseSetResolver
 {
-	template<typename InElementType, typename KeyFuncs = DefaultKeyFuncs<InElementType>, typename Allocator = FDefaultSparseSetAllocator>
-	using SetType = TSparseSet<InElementType, KeyFuncs, Allocator>;
+	template<typename InElementType, typename KeyFuncs = DefaultKeyFuncs<InElementType>, typename Allocator = FDefaultSetAllocator>
+	using SetType = TSet<InElementType, KeyFuncs, Allocator>;
 
 	template <typename Type, int32 Size>
-	using FixedSetType = TSparseSet<Type, DefaultKeyFuncs<Type>, TFixedSparseSetAllocator<Size>>;
+	using FixedSetType = TSet<Type, DefaultKeyFuncs<Type>, TFixedSetAllocator<Size>>;
 };
 
-struct FCompactSetTestResolver
+struct FCompactSetResolver
 {
 	template<typename InElementType, typename KeyFuncs = DefaultKeyFuncs<InElementType>, typename Allocator = FDefaultCompactSetAllocator>
 	using SetType = TCompactSet<InElementType, KeyFuncs, Allocator>;
 
 	template <typename Type, int32 Size>
-	using FixedSetType = TCompactSet<Type, DefaultKeyFuncs<Type>, TFixedCompactSetAllocator<Size>>;
+	using FixedSetType = TCompactSet<Type, DefaultKeyFuncs<Type>, UE::Core::Private::CompactSetAllocatorHelpers::TFixedCompactSetAllocator<Size>>;
 };
 
 template<typename Type, typename SetType>
@@ -76,7 +75,7 @@ uint32 GetTypeHash(const FSetTestNonCopyableTestData& A)
 }
 
 
-template<typename SetResolver = FCompactSetTestResolver>
+template<typename SetResolver = FCompactSetResolver>
 struct TSetTestHelper
 {
 	template <typename Type>
@@ -251,23 +250,6 @@ struct TSetTestHelper
 			}
 		}
 
-		// ensure that duplicate entries overrides existing elements and doesn't cause reordering
-		SECTION("Duplicate Entry Behavior")
-		{
-			SetType<int32> Set;
-
-			Set.Add(0);
-			Set.Add(1);
-			Set.Add(2);
-			Set.Add(3);
-			Set.Add(2);
-
-			for (int32 Index = 0; Index < Set.Num(); ++Index)
-			{
-				CHECK(Set[FSetElementId::FromInteger(Index)] == Index);
-			}
-		}
-
 		SECTION("Duplicate Support")
 		{
 			FString StringTestItems[] = { TEXT("Zero"), TEXT("One"), TEXT("Two"), TEXT("Three") };
@@ -470,27 +452,15 @@ struct TSetTestHelper
 
 TEST_CASE_NAMED(FCompactSetTest, "System::Core::Containers::CompactSet", "[Core][Containers][CompactSet]")
 {
-	TSetTestHelper<FCompactSetTestResolver>();
+	static_assert(TCompactSet<FString>::GetTotalMemoryRequiredInBytes(97) == UE::Core::Private::CompactSetAllocatorHelpers::CalculateRequiredBytes<97, sizeof(FString)>());
 
-	SECTION("Allocation Size")
-	{
-		static const int32 SizeToTest = 139;
-		static_assert(TCompactSet<int32>::GetTotalMemoryRequiredInBytes(SizeToTest) == UE::Core::Private::CompactSetAllocatorHelpers::CalculateRequiredBytes<SizeToTest, sizeof(int32)>());
-		TCompactSet<int32> Container;
-		for (int32 Index = 0; Index < SizeToTest; ++Index)
-		{
-			Container.Add(Index);
-		}
-
-		// Memory allocator may have additional room in block allocated, so the set will fill up that space which can lead to increased memory usage
-		CHECK(Container.GetAllocatedSize() >= TCompactSet<int32>::GetTotalMemoryRequiredInBytes(SizeToTest));
-	}
+	TSetTestHelper<FCompactSetResolver>();
 
 	// Todo: Move into TSetTestHelper when Stable moved to TSparseSet
 	SECTION("Stability Support")
 	{
 		TArray<FString> StringTestItems = { TEXT("Zero"), TEXT("One"), TEXT("Two"), TEXT("Three") };
-		TSetTestHelper<FCompactSetTestResolver>::DuplicateSetType<FString> Set(StringTestItems);
+		TSetTestHelper<FCompactSetResolver>::DuplicateSetType<FString> Set(StringTestItems);
 
 		Set.RemoveStable(TEXT("One"));
 		StringTestItems.RemoveAt(1);
@@ -520,7 +490,7 @@ TEST_CASE_NAMED(FCompactSetTest, "System::Core::Containers::CompactSet", "[Core]
 
 TEST_CASE_NAMED(FSparseSetTest, "System::Core::Containers::SparseSet", "[Core][Containers][CompactSet]")
 {
-	TSetTestHelper<FSparseSetTestResolver>();
+	TSetTestHelper<FSparseSetResolver>();
 }
 
 TEST_CASE_NAMED(FScriptCompactSetTest, "System::Core::Containers::ScriptCompactSet", "[Core][Containers][CompactSet]")
@@ -578,7 +548,6 @@ TEST_CASE_NAMED(FScriptCompactSetTest, "System::Core::Containers::ScriptCompactS
 
 		auto GetKeyHash = [](const void* Data) { return * (const int32 *)Data; };
 		auto EqualityFn = [](const void* Lhs, const void* Rhs) { return * (const int32 *)Lhs == * (const int32 *)Rhs; };
-		auto DestructFn = [](void* Data) { };
 
 		for (int32 TestItem : IntTestItems)
 		{
@@ -595,7 +564,7 @@ TEST_CASE_NAMED(FScriptCompactSetTest, "System::Core::Containers::ScriptCompactS
 		{
 			if (Index % 2)
 			{
-				Set.RemoveAt(Index, Layout, GetKeyHash, DestructFn);
+				Set.RemoveAt(Index, Layout, GetKeyHash);
 				*(int32 *)Set.GetData(Index, Layout) = *(int32 *)Set.GetData(Set.Num() - 1, Layout);
 				CHECK(Set.FindIndex(&IntTestItems[Index], Layout, GetKeyHash, EqualityFn) == INDEX_NONE);
 			}
diff --git a/Engine/Source/Runtime/CoreUObject/Private/UObject/GarbageCollection.cpp b/Engine/Source/Runtime/CoreUObject/Private/UObject/GarbageCollection.cpp
--- a/Engine/Source/Runtime/CoreUObject/Private/UObject/GarbageCollection.cpp
+++ b/Engine/Source/Runtime/CoreUObject/Private/UObject/GarbageCollection.cpp
@@ -1924,7 +1924,7 @@ private:
 struct FSparseStructArray
 {
 	FSchemaView Schema{NoInit};
-	FScriptSet* Array;
+	FScriptSparseArray* Array;
 	int32 Num;
 	uint32 Stride;
 };
@@ -1959,7 +1959,7 @@ class FStructBatcher : public FBatcherBase
 public:
 	explicit FStructBatcher(FWorkerContext& Context) : ValidatedStructArrays(Context) {}
 
-	FORCEINLINE_DEBUGGABLE void PushSparseStructArray(FSchemaView Schema, FScriptSet& Array)
+	FORCEINLINE_DEBUGGABLE void PushSparseStructArray(FSchemaView Schema, FScriptSparseArray& Array)
 	{
 		CheckValid(Schema);
 
@@ -2010,11 +2010,8 @@ private:
 		{
 			if (Entry.Num)
 			{
-				FScriptSet& Array = *Entry.Array;
-				FStructArray Slice = { Entry.Schema, Private::GetSetData(Array), Entry.Num, Entry.Stride };
-#if UE_USE_COMPACT_SET_AS_DEFAULT
-				ValidatedStructArrays.Push(Slice);
-#else
+				FScriptSparseArray& Array = *Entry.Array;
+				FStructArray Slice = { Entry.Schema, Private::GetSparseData(Array), Entry.Num, Entry.Stride };
 				if (Array.IsCompact())
 				{
 					ValidatedStructArrays.Push(Slice);
@@ -2024,7 +2021,7 @@ private:
 					Slice.Num = 0;
 					for (uint32 Idx = 0, NumIndices = Array.GetMaxIndex(); Idx < NumIndices; ++Idx)
 					{
-						if (Array.IsValidIndex(Idx))
+						if (Array.IsAllocated(Idx))
 						{
 							++Slice.Num;
 						}
@@ -2046,7 +2043,6 @@ private:
 						ValidatedStructArrays.Push(Slice);
 					}
 				}
-#endif
 			}
 		}
 
@@ -3200,7 +3196,7 @@ struct TBatchDispatcher
 		StructBatcher.PushStructArray(Schema, Data, Num);
 	}
 
-	FORCEINLINE_DEBUGGABLE void QueueSparseStructArray(FSchemaView Schema, FScriptSet& Array)
+	FORCEINLINE_DEBUGGABLE void QueueSparseStructArray(FSchemaView Schema, FScriptSparseArray& Array)
 	{
 		StructBatcher.PushSparseStructArray(Schema, Array);
 	}
@@ -6727,7 +6723,7 @@ void FMapProperty::EmitReferenceInfo(UE::GC::FSchemaBuilder& Schema, int32 BaseO
 			}
 		}
 
-		Schema.Add(UE::GC::DeclareMember(DebugPath, BaseOffset + GetOffset_ForGC(), EMemberType::StructSet, PairSchema.Build()));
+		Schema.Add(UE::GC::DeclareMember(DebugPath, BaseOffset + GetOffset_ForGC(), EMemberType::SparseStructArray, PairSchema.Build()));
 	}
 }
 
@@ -6739,7 +6735,7 @@ void FSetProperty::EmitReferenceInfo(UE::GC::FSchemaBuilder& Schema, int32 BaseO
 		FSchemaBuilder ElementSchema(GetStride());
 		FPropertyStackScope PropertyScope(DebugPath, ElementProp);
 		ElementProp->EmitReferenceInfo(ElementSchema, 0, EncounteredStructProps, DebugPath);
-		Schema.Add(UE::GC::DeclareMember(DebugPath, BaseOffset + GetOffset_ForGC(), EMemberType::StructSet, ElementSchema.Build()));
+		Schema.Add(UE::GC::DeclareMember(DebugPath, BaseOffset + GetOffset_ForGC(), EMemberType::SparseStructArray, ElementSchema.Build()));
 	}
 }
 
diff --git a/Engine/Source/Runtime/CoreUObject/Private/UObject/GarbageCollectionSchema.cpp b/Engine/Source/Runtime/CoreUObject/Private/UObject/GarbageCollectionSchema.cpp
--- a/Engine/Source/Runtime/CoreUObject/Private/UObject/GarbageCollectionSchema.cpp
+++ b/Engine/Source/Runtime/CoreUObject/Private/UObject/GarbageCollectionSchema.cpp
@@ -14,7 +14,7 @@ namespace UE::GC {
 //////////////////////////////////////////////////////////////////////////
 
 static constexpr uint32 ToBit(EMemberType Type)		{ return 1u << static_cast<uint32>(Type); }
-static constexpr uint32 InnerSchemaTypes =			ToBit(EMemberType::StructArray) | ToBit(EMemberType::StructSet) | ToBit(EMemberType::FreezableStructArray) | ToBit(EMemberType::Optional);
+static constexpr uint32 InnerSchemaTypes =			ToBit(EMemberType::StructArray) | ToBit(EMemberType::SparseStructArray) | ToBit(EMemberType::FreezableStructArray) | ToBit(EMemberType::Optional);
 static constexpr uint32 AROTypes =					ToBit(EMemberType::ARO) | ToBit(EMemberType::SlowARO) | ToBit(EMemberType::MemberARO);
 static constexpr uint32 AROPointerTypes =			ToBit(EMemberType::ARO) | ToBit(EMemberType::MemberARO);
 static constexpr uint32 ExtraWordTypes =			InnerSchemaTypes | AROPointerTypes | ToBit(EMemberType::StridedArray);
@@ -46,7 +46,7 @@ FName ToName(EMemberType Type)
 			Names[(uint8)EMemberType::ReferenceArray] =					"ReferenceArray";
 			Names[(uint8)EMemberType::StructArray] =					"StructArray";
 			Names[(uint8)EMemberType::StridedArray] =					"StridedArray";
-			Names[(uint8)EMemberType::StructSet] =						"StructSet";
+			Names[(uint8)EMemberType::SparseStructArray] =				"SparseStructArray";
 			Names[(uint8)EMemberType::FieldPath] =						"FieldPath";
 			Names[(uint8)EMemberType::FieldPathArray] =					"FieldPathArray";
 			Names[(uint8)EMemberType::FreezableReferenceArray] =		"FreezableReferenceArray";
@@ -649,7 +649,7 @@ struct FStructMemberStats
 		{ 
 			case EMemberType::StructArray:			VisitArray(*(const FScriptArray*)Container, Visitor);
 			break;
-			case EMemberType::StructSet:			VisitSet(*(const FScriptSet*)Container, Visitor);
+			case EMemberType::SparseStructArray:	VisitSparseArray(*(const FScriptSparseArray*)Container, Visitor);
 			break;
 			case EMemberType::FreezableStructArray:	VisitArray(*(const FFreezableScriptArray*)Container, Visitor);
 			break;
@@ -678,13 +678,13 @@ private:
 	}
 
 	template<typename Fn>
-	void VisitSet(const FScriptSet& Set, Fn Visitor) const
+	void VisitSparseArray(const FScriptSparseArray& Array, Fn Visitor) const
 	{
-		const uint8* Data = Private::GetSetData(Set);
+		const uint8* Data = Private::GetSparseData(Array);
 		const uint32 Stride = Schema.GetStructStride();
-		for (int32 Idx = 0, Num = Set.Num(); Idx < Num; ++Idx)
+		for (int32 Idx = 0, Num = Array.Num(); Idx < Num; ++Idx)
 		{
-			if (Set.IsValidIndex(Idx))
+			if (Array.IsValidIndex(Idx))
 			{
 				Visitor(Data + Idx * Stride);
 			}
diff --git a/Engine/Source/Runtime/CoreUObject/Public/UObject/FastReferenceCollector.h b/Engine/Source/Runtime/CoreUObject/Public/UObject/FastReferenceCollector.h
--- a/Engine/Source/Runtime/CoreUObject/Public/UObject/FastReferenceCollector.h
+++ b/Engine/Source/Runtime/CoreUObject/Public/UObject/FastReferenceCollector.h
@@ -495,14 +495,14 @@ FORCEINLINE FStridedReferenceView ToView(FStridedReferenceArray In)
 	return { reinterpret_cast<UObject**>(In.Array->GetData()) + In.Layout.WordOffset,  In.Array->Num(), In.Layout.WordStride };
 }
 
-FORCEINLINE uint8* GetSetData(FScriptSet& Set)
+FORCEINLINE uint8* GetSparseData(FScriptSparseArray& Array)
 {
-	return reinterpret_cast<uint8*>(Set.GetData(0, FScriptSetLayout{})); 
+	return reinterpret_cast<uint8*>(Array.GetData(0, FScriptSparseArrayLayout{})); 
 }
 
-FORCEINLINE const uint8* GetSetData(const FScriptSet& Set)
+FORCEINLINE const uint8* GetSparseData(const FScriptSparseArray& Array)
 {
-	return reinterpret_cast<const uint8*>(Set.GetData(0, FScriptSetLayout{})); 
+	return reinterpret_cast<const uint8*>(Array.GetData(0, FScriptSparseArrayLayout{})); 
 }
 
 template<class DispatcherType>
@@ -534,20 +534,20 @@ FORCEINLINE_DEBUGGABLE void VisitStructArray(DispatcherType& Dispatcher, FSchema
 }
 
 template<class DispatcherType>
-FORCEINLINE_DEBUGGABLE void VisitStructSet(DispatcherType& Dispatcher, FSchemaView StructSchema, FScriptSet& Set)
+FORCEINLINE_DEBUGGABLE void VisitSparseStructArray(DispatcherType& Dispatcher, FSchemaView StructSchema, FScriptSparseArray& Array)
 {
 	check(!StructSchema.IsEmpty());
 	if constexpr (DispatcherType::bBatching)
 	{
-		Dispatcher.QueueSparseStructArray(StructSchema, Set);
+		Dispatcher.QueueSparseStructArray(StructSchema, Array);
 	}
-	else if (int32 Num = Set.Num())
+	else if (int32 Num = Array.Num())
 	{
-		uint8* It = GetSetData(Set);
+		uint8* It = GetSparseData(Array);
 		const uint32 Stride = StructSchema.GetStructStride();
-		for (int32 Idx = 0, MaxIdx = Set.GetMaxIndex(); Idx < MaxIdx; ++Idx, It += Stride)
+		for (int32 Idx = 0, MaxIdx = Array.GetMaxIndex(); Idx < MaxIdx; ++Idx, It += Stride)
 		{
-			if (Set.IsValidIndex(Idx))
+			if (Array.IsAllocated(Idx))
 			{
 				typename DispatcherType::SchemaStackScopeType SchemaStack(Dispatcher.Context, StructSchema);
 				VisitNestedStructMembers(Dispatcher, StructSchema, It);
@@ -709,7 +709,7 @@ FORCEINLINE_DEBUGGABLE void VisitMembers(DispatcherType& Dispatcher, FSchemaView
 			break;
 			case EMemberType::StructArray:				VisitStructArray(			Dispatcher, FSchemaView((++WordIt)->InnerSchema, Origin), *(FScriptArray*)MemberPtr);
 			break;
-			case EMemberType::StructSet:				VisitStructSet(				Dispatcher, FSchemaView((++WordIt)->InnerSchema, Origin), *(FScriptSet*)MemberPtr);
+			case EMemberType::SparseStructArray:		VisitSparseStructArray(		Dispatcher, FSchemaView((++WordIt)->InnerSchema, Origin), *(FScriptSparseArray*)MemberPtr);
 			break;
 			case EMemberType::FreezableStructArray:		VisitStructArray(			Dispatcher, FSchemaView((++WordIt)->InnerSchema, Origin), *(FFreezableScriptArray*)MemberPtr);
 			break;
diff --git a/Engine/Source/Runtime/CoreUObject/Public/UObject/GarbageCollectionSchema.h b/Engine/Source/Runtime/CoreUObject/Public/UObject/GarbageCollectionSchema.h
--- a/Engine/Source/Runtime/CoreUObject/Public/UObject/GarbageCollectionSchema.h
+++ b/Engine/Source/Runtime/CoreUObject/Public/UObject/GarbageCollectionSchema.h
@@ -29,7 +29,7 @@ enum class EMemberType : uint8
 	ReferenceArray,						// Member - Array of references, e.g. TArray<MyObject*> and TArray<TObjectPtr<MyObject>>
 	StructArray,						// Array of structs
 	StridedArray,						// Array of structs with single reference per struct (~half of struct instances)
-	StructSet,							// TMap/TSet of structs
+	SparseStructArray,					// TMap/TSet of structs
 	FieldPath,							// Field path strong reference to owner
 	FieldPathArray,					 	// Array of field paths
 	FreezableReferenceArray,			// Freezable array of references
diff --git a/Engine/Source/Runtime/CoreUObject/Public/UObject/UnrealType.h b/Engine/Source/Runtime/CoreUObject/Public/UObject/UnrealType.h
--- a/Engine/Source/Runtime/CoreUObject/Public/UObject/UnrealType.h
+++ b/Engine/Source/Runtime/CoreUObject/Public/UObject/UnrealType.h
@@ -4773,33 +4773,7 @@ public:
 		return WithScriptMap([this, InternalIndex, Count](auto* Map)
 		{
 			check(Map->IsValidIndex(InternalIndex));
-			check(InternalIndex + Count <= Map->Max());
 
-#if UE_USE_COMPACT_SET_AS_DEFAULT
-			auto GetHashKey = [this](const void* ElementKey)
-			{
-				return KeyProp->GetValueTypeHash(ElementKey);
-			};
-
-			auto DestroyElement = [this](void* Element)
-			{
-				if (!(KeyProp->PropertyFlags & (CPF_IsPlainOldData | CPF_NoDestructor)))
-				{
-					KeyProp->DestroyValue(Element);
-				}
-				void* ValuePtr = (uint8*)Element + MapLayout.ValueOffset;
-				if (!(ValueProp->PropertyFlags & (CPF_IsPlainOldData | CPF_NoDestructor)))
-				{
-					ValueProp->DestroyValue(ValuePtr);
-				}
-			};
-
-			// Last element will get swapped into current spot, so go in reverse so everything moves safely
-			for (int32 Index = InternalIndex + Count - 1; Index >= InternalIndex; --Index)
-			{
-				Map->RemoveAt(Index, MapLayout, GetHashKey, DestroyElement);
-			}
-#else
 			DestructItems(InternalIndex, Count);
 			for (int32 LocalCount = Count, LocalIndex = InternalIndex; LocalCount; ++LocalIndex)
 			{
@@ -4809,7 +4783,6 @@ public:
 					--LocalCount;
 				}
 			}
-#endif
 		});
 	}
 
@@ -5611,28 +5584,7 @@ public:
 	void RemoveAt(int32 InternalIndex, int32 Count = 1)
 	{
 		check(IsValidIndex(InternalIndex));
-		check(InternalIndex + Count <= Set->Max());
-
-#if UE_USE_COMPACT_SET_AS_DEFAULT
-		auto GetHashKey = [this](const void* ElementKey)
-		{
-			return ElementProp->GetValueTypeHash(ElementKey);
-		};
 
-		auto DestroyElement = [this](void* Element)
-		{
-			if (!(ElementProp->PropertyFlags & (CPF_IsPlainOldData | CPF_NoDestructor)))
-			{
-				ElementProp->DestroyValue(Element);
-			}
-		};
-
-		// Last element will get swapped into current spot, so go in reverse so everything moves safely
-		for (int32 Index = InternalIndex + Count - 1; Index >= InternalIndex; --Index)
-		{
-			Set->RemoveAt(Index, SetLayout, GetHashKey, DestroyElement);
-		}
-#else
 		DestructItems(InternalIndex, Count);
 		for (; Count; ++InternalIndex)
 		{
@@ -5642,7 +5594,6 @@ public:
 				--Count;
 			}
 		}
-#endif
 	}
 
 	/**
diff --git a/Engine/Source/Runtime/Engine/Private/Materials/MaterialInstance.cpp b/Engine/Source/Runtime/Engine/Private/Materials/MaterialInstance.cpp
--- a/Engine/Source/Runtime/Engine/Private/Materials/MaterialInstance.cpp
+++ b/Engine/Source/Runtime/Engine/Private/Materials/MaterialInstance.cpp
@@ -4917,14 +4917,14 @@ void UMaterialInstance::GetResourceSizeEx(FResourceSizeEx& CumulativeResourceSiz
 		CumulativeResourceSize.AddDedicatedSystemMemoryBytes(ParameterCollectionParameterValues.Num() * sizeof(THashedMaterialParameterMap<const UMaterialParameterCollection*>::TNamedParameter));
 
 		// Record space for hash tables as well..
-		if (ScalarParameterValues.Num()) CumulativeResourceSize.AddDedicatedSystemMemoryBytes(FDefaultSparseSetAllocator::GetNumberOfHashBuckets(ScalarParameterValues.Num()) * sizeof(uint16));
-		if (VectorParameterValues.Num()) CumulativeResourceSize.AddDedicatedSystemMemoryBytes(FDefaultSparseSetAllocator::GetNumberOfHashBuckets(VectorParameterValues.Num()) * sizeof(uint16));
-		if (DoubleVectorParameterValues.Num()) CumulativeResourceSize.AddDedicatedSystemMemoryBytes(FDefaultSparseSetAllocator::GetNumberOfHashBuckets(DoubleVectorParameterValues.Num()) * sizeof(uint16));
-		if (TextureParameterValues.Num()) CumulativeResourceSize.AddDedicatedSystemMemoryBytes(FDefaultSparseSetAllocator::GetNumberOfHashBuckets(TextureParameterValues.Num()) * sizeof(uint16));
-		if (RuntimeVirtualTextureParameterValues.Num()) CumulativeResourceSize.AddDedicatedSystemMemoryBytes(FDefaultSparseSetAllocator::GetNumberOfHashBuckets(RuntimeVirtualTextureParameterValues.Num()) * sizeof(uint16));
-		if (SparseVolumeTextureParameterValues.Num()) CumulativeResourceSize.AddDedicatedSystemMemoryBytes(FDefaultSparseSetAllocator::GetNumberOfHashBuckets(SparseVolumeTextureParameterValues.Num()) * sizeof(uint16));
-		if (FontParameterValues.Num()) CumulativeResourceSize.AddDedicatedSystemMemoryBytes(FDefaultSparseSetAllocator::GetNumberOfHashBuckets(FontParameterValues.Num()) * sizeof(uint16));
-		if (ParameterCollectionParameterValues.Num()) CumulativeResourceSize.AddDedicatedSystemMemoryBytes(FDefaultSparseSetAllocator::GetNumberOfHashBuckets(ParameterCollectionParameterValues.Num()) * sizeof(uint16));
+		if (ScalarParameterValues.Num()) CumulativeResourceSize.AddDedicatedSystemMemoryBytes(FDefaultSetAllocator::GetNumberOfHashBuckets(ScalarParameterValues.Num()) * sizeof(uint16));
+		if (VectorParameterValues.Num()) CumulativeResourceSize.AddDedicatedSystemMemoryBytes(FDefaultSetAllocator::GetNumberOfHashBuckets(VectorParameterValues.Num()) * sizeof(uint16));
+		if (DoubleVectorParameterValues.Num()) CumulativeResourceSize.AddDedicatedSystemMemoryBytes(FDefaultSetAllocator::GetNumberOfHashBuckets(DoubleVectorParameterValues.Num()) * sizeof(uint16));
+		if (TextureParameterValues.Num()) CumulativeResourceSize.AddDedicatedSystemMemoryBytes(FDefaultSetAllocator::GetNumberOfHashBuckets(TextureParameterValues.Num()) * sizeof(uint16));
+		if (RuntimeVirtualTextureParameterValues.Num()) CumulativeResourceSize.AddDedicatedSystemMemoryBytes(FDefaultSetAllocator::GetNumberOfHashBuckets(RuntimeVirtualTextureParameterValues.Num()) * sizeof(uint16));
+		if (SparseVolumeTextureParameterValues.Num()) CumulativeResourceSize.AddDedicatedSystemMemoryBytes(FDefaultSetAllocator::GetNumberOfHashBuckets(SparseVolumeTextureParameterValues.Num()) * sizeof(uint16));
+		if (FontParameterValues.Num()) CumulativeResourceSize.AddDedicatedSystemMemoryBytes(FDefaultSetAllocator::GetNumberOfHashBuckets(FontParameterValues.Num()) * sizeof(uint16));
+		if (ParameterCollectionParameterValues.Num()) CumulativeResourceSize.AddDedicatedSystemMemoryBytes(FDefaultSetAllocator::GetNumberOfHashBuckets(ParameterCollectionParameterValues.Num()) * sizeof(uint16));
 	}
 }
 
diff --git a/Engine/Source/Runtime/Engine/Private/Materials/MaterialInstanceSupport.h b/Engine/Source/Runtime/Engine/Private/Materials/MaterialInstanceSupport.h
--- a/Engine/Source/Runtime/Engine/Private/Materials/MaterialInstanceSupport.h
+++ b/Engine/Source/Runtime/Engine/Private/Materials/MaterialInstanceSupport.h
@@ -177,7 +177,7 @@ public:
 		// If zero is passed in (or anything invalid), choose a default size
 		if (NumHashBuckets <= 0)
 		{
-			NumHashBuckets = FDefaultSparseSetAllocator::GetNumberOfHashBuckets(Array.Num());
+			NumHashBuckets = FDefaultSetAllocator::GetNumberOfHashBuckets(Array.Num());
 		}
 		if (Array.IsEmpty())
 		{
@@ -278,7 +278,7 @@ public:
 			// FHashedMaterialParameterInfoPacked::HashNext is limited to 14 bits
 			check(ValueArray.Array.Num() <= (1<<14));
 
-			const int32 NumHashBuckets = FDefaultSparseSetAllocator::GetNumberOfHashBuckets(ValueArray.Array.Num());
+			const int32 NumHashBuckets = FDefaultSetAllocator::GetNumberOfHashBuckets(ValueArray.Array.Num());
 			if (NumHashBuckets == ValueArray.HashTable.Num())
 			{
 				ValueArray.HashAddOneItem(ValueArray.Array.Num() - 1);
diff --git a/Engine/Source/Runtime/JsonUtilities/Private/JsonUClassExporter.cpp b/Engine/Source/Runtime/JsonUtilities/Private/JsonUClassExporter.cpp
--- a/Engine/Source/Runtime/JsonUtilities/Private/JsonUClassExporter.cpp
+++ b/Engine/Source/Runtime/JsonUtilities/Private/JsonUClassExporter.cpp
@@ -169,18 +169,13 @@ void WriteFScriptSparseArrayLayout(const FScriptSparseArrayLayout& SparseArrayLa
 
 void WriteFScriptSetLayout(const FScriptSetLayout& ScriptSetLayout, FJsonUETypeWriter& Writer)
 {
-	Writer.WriteValue(TEXT("Size"), ScriptSetLayout.Size);
-
-#if UE_USE_COMPACT_SET_AS_DEFAULT
-	Writer.WriteValue(TEXT("Alignment"), ScriptSetLayout.Alignment);
-#else
 	Writer.WriteValue(TEXT("HashNextIdOffset"), ScriptSetLayout.HashNextIdOffset);
 	Writer.WriteValue(TEXT("HashIndexOffset"), ScriptSetLayout.HashIndexOffset);
+	Writer.WriteValue(TEXT("Size"), ScriptSetLayout.Size);
 
 	Writer.WriteObjectStart(TEXT("SparseArrayLayout"));
 	WriteFScriptSparseArrayLayout(ScriptSetLayout.SparseArrayLayout, Writer);
 	Writer.WriteObjectEnd();
-#endif
 }
 
 void WriteFScriptMapLayout(const FScriptMapLayout& MapLayout, FJsonUETypeWriter& Writer)
diff --git a/Engine/Source/Runtime/Renderer/Private/DistanceFieldStreaming.cpp b/Engine/Source/Runtime/Renderer/Private/DistanceFieldStreaming.cpp
--- a/Engine/Source/Runtime/Renderer/Private/DistanceFieldStreaming.cpp
+++ b/Engine/Source/Runtime/Renderer/Private/DistanceFieldStreaming.cpp
@@ -1091,7 +1091,7 @@ bool FDistanceFieldSceneData::ResizeIndirectionAtlasIfNeeded(FRDGBuilder& GraphB
 	}
 	else if (CurrentSize != DesiredSize)
 	{
-		for (TSparseSet<FDistanceFieldAssetState, TFDistanceFieldAssetStateFuncs>::TConstIterator It(AssetStateArray); It; ++It)
+		for (TSet<FDistanceFieldAssetState, TFDistanceFieldAssetStateFuncs>::TConstIterator It(AssetStateArray); It; ++It)
 		{
 			const FDistanceFieldAssetState& AssetState = *It;
 			for (int32 ReversedMipIndex = 0; ReversedMipIndex < AssetState.ReversedMips.Num(); ReversedMipIndex++)
@@ -1164,7 +1164,7 @@ void FDistanceFieldSceneData::DefragmentIndirectionAtlas(FIntVector MinSize, TAr
 	};
 	TArray<FEntry> Entries;
 
-	for (TSparseSet<FDistanceFieldAssetState, TFDistanceFieldAssetStateFuncs>::TConstIterator It(AssetStateArray); It; ++It)
+	for (TSet<FDistanceFieldAssetState, TFDistanceFieldAssetStateFuncs>::TConstIterator It(AssetStateArray); It; ++It)
 	{
 		const FDistanceFieldAssetState& AssetState = *It;
 		for (int32 ReversedMipIndex = 0; ReversedMipIndex < AssetState.ReversedMips.Num(); ReversedMipIndex++)
@@ -1389,7 +1389,7 @@ void FDistanceFieldSceneData::UploadAllAssetData(FRDGBuilder& GraphBuilder, FRDG
 
 	AssetDataUploadBuffer.Init(GraphBuilder, NumUploads, AssetDataMipStrideFloat4s * sizeof(FVector4f), true, TEXT("DistanceFields.DFAssetDataUploadBuffer"));
 
-	for (TSparseSet<FDistanceFieldAssetState, TFDistanceFieldAssetStateFuncs>::TConstIterator It(AssetStateArray); It; ++It)
+	for (TSet<FDistanceFieldAssetState, TFDistanceFieldAssetStateFuncs>::TConstIterator It(AssetStateArray); It; ++It)
 	{
 		const FDistanceFieldAssetState& AssetState = *It;
 		const FSetElementId AssetId = It.GetId();
@@ -1712,7 +1712,7 @@ void FDistanceFieldSceneData::ListMeshDistanceFields(bool bDumpAssetStats) const
 
 	const uint32 BrickSizeBytes = GPixelFormats[DistanceField::DistanceFieldFormat].BlockBytes * DistanceField::BrickSize * DistanceField::BrickSize * DistanceField::BrickSize;
 
-	for (TSparseSet<FDistanceFieldAssetState, TFDistanceFieldAssetStateFuncs>::TConstIterator It(AssetStateArray); It; ++It)
+	for (TSet<FDistanceFieldAssetState, TFDistanceFieldAssetStateFuncs>::TConstIterator It(AssetStateArray); It; ++It)
 	{
 		const FDistanceFieldAssetState& AssetState = *It;
 
diff --git a/Engine/Source/Runtime/Renderer/Private/ScenePrivate.h b/Engine/Source/Runtime/Renderer/Private/ScenePrivate.h
--- a/Engine/Source/Runtime/Renderer/Private/ScenePrivate.h
+++ b/Engine/Source/Runtime/Renderer/Private/ScenePrivate.h
@@ -8,7 +8,6 @@
 
 // Dependencies.
 
-#include "Containers/SparseSet.h"
 #include "Misc/Guid.h"
 #include "Math/RandomStream.h"
 #include "Templates/PimplPtr.h"
@@ -2158,7 +2157,7 @@ public:
 	TArray<int32> IndicesToUpdateInObjectBuffers;
 	TSet<int32> IndicesToUpdateInObjectBuffersSet;
 
-	TSparseSet<FDistanceFieldAssetState, TFDistanceFieldAssetStateFuncs> AssetStateArray;
+	TSet<FDistanceFieldAssetState, TFDistanceFieldAssetStateFuncs> AssetStateArray;
 	TRefCountPtr<FRDGPooledBuffer> AssetDataBuffer;
 	FRDGScatterUploadBuffer AssetDataUploadBuffer;
 
