Fixed UE_TEXTKEY_SPLIT_GUID causing GUID case variance issues

Lowercase GUIDs are now ignored by the UE_TEXTKEY_SPLIT_GUID optimization, as they would be reconstructed as an uppercase GUID.

FGuid::ParseExact doesn't separate the parsing logic for Digits and DigitsLowercase, but we only optimize for uppercase as we need the String -> TextKey -> String conversion to be lossless for case.

We could support other GUID formats if needed, but it would mean tracking what kind of GUID was parsed (which would use more memory) so that it could be reconstructed corectly when used with FGuid::ToString.

diff --git a/Engine/Source/Runtime/Core/Private/Internationalization/TextKey.cpp b/Engine/Source/Runtime/Core/Private/Internationalization/TextKey.cpp
--- a/Engine/Source/Runtime/Core/Private/Internationalization/TextKey.cpp
+++ b/Engine/Source/Runtime/Core/Private/Internationalization/TextKey.cpp
@@ -386,8 +386,31 @@ private:
 		int32 Index = INDEX_NONE;
 		{
 #if UE_TEXTKEY_SPLIT_GUID
+			auto IsCompatibleGuidString =
+				[InStr](FGuid& OutGuid)
+				{
+					// Only checking for EGuidFormats::Digits as that's the default of FGuid::ToString() as used by FText keys
+					if (FGuid::ParseExact(InStr, EGuidFormats::Digits, OutGuid))
+					{
+						// Secondary check that every digit character parsed was uppercase, as ParseExact doesn't separate 
+						// the parsing logic for Digits and DigitsLowercase, but we only optimize for uppercase as we need 
+						// the String -> TextKey -> String conversion to be lossless for case
+						// We could support other GUID formats if needed, but it would mean tracking what kind of GUID
+						// was parsed so that it could be reconstructed corectly when used with FGuid::ToString
+						for (const TCHAR GuidChar : InStr)
+						{
+							if (!FChar::IsDigit(GuidChar) && !FChar::IsUpper(GuidChar))
+							{
+								return false;
+							}
+						}
+						return true;
+					}
+					return false;
+				};
+
 			FGuid KeyGuid;
-			if (FGuid::ParseExact(InStr, EGuidFormats::Digits, KeyGuid)) // Only checking for EGuidFormats::Digits as that's the default of FGuid::ToString() as used by FText keys
+			if (IsCompatibleGuidString(KeyGuid))
 			{
 				const FKeyDataView KeyData = FKeyDataView(KeyGuid, InStrHash);
 				Index = FindOrAddData(KeyData);
