Add the engine release number to the CSVProfiling metadata.

The CSVProfiling metadata includes EngineVersion which is the full version string (Major.Minor.Patch, branch, and changelist). Isolating the version number portion to group captures by engine builds.

diff --git a/Engine/Source/Runtime/Core/Private/ProfilingDebugging/CsvProfiler.cpp b/Engine/Source/Runtime/Core/Private/ProfilingDebugging/CsvProfiler.cpp
--- a/Engine/Source/Runtime/Core/Private/ProfilingDebugging/CsvProfiler.cpp
+++ b/Engine/Source/Runtime/Core/Private/ProfilingDebugging/CsvProfiler.cpp
@@ -3475,6 +3475,7 @@ FCsvProfiler::FCsvProfiler()
 	CommandlineStr.ReplaceInline(TEXT("\r"), TEXT(""));
 	FString BuildVersionString = FApp::GetBuildVersion();
 	FString EngineVersionString = FEngineVersion::Current().ToString();
+	FString EngineReleaseVersionString = FEngineVersion::Current().ToString(EVersionComponent::Patch);
 
 	FString OSMajor, OSMinor;
 	FPlatformMisc::GetOSVersions(OSMajor, OSMinor);
@@ -3488,6 +3489,7 @@ FCsvProfiler::FCsvProfiler()
 	SetMetadataInternal(TEXT("Config"), *BuildConfigurationStr);
 	SetMetadataInternal(TEXT("BuildVersion"), *BuildVersionString);
 	SetMetadataInternal(TEXT("EngineVersion"), *EngineVersionString);
+	SetMetadataInternal(TEXT("EngineReleaseVersion"), *EngineReleaseVersionString);
 	SetMetadataInternal(TEXT("OS"), *OSString);
 	SetMetadataInternal(TEXT("CPU"), *FPlatformMisc::GetDeviceMakeAndModel());
 	SetMetadataInternal(TEXT("PGOEnabled"), FPlatformMisc::IsPGOEnabled() ? TEXT("1") : TEXT("0"));
