Add missing EConstEval constructor to TSet

diff --git a/Engine/Source/Runtime/Core/Public/Containers/Set.h b/Engine/Source/Runtime/Core/Public/Containers/Set.h
--- a/Engine/Source/Runtime/Core/Public/Containers/Set.h
+++ b/Engine/Source/Runtime/Core/Public/Containers/Set.h
@@ -46,6 +46,11 @@ public:
 
 	[[nodiscard]] UE_FORCEINLINE_HINT TSet() = default;
 
+	[[nodiscard]] UE_FORCEINLINE_HINT consteval TSet(EConstEval)
+		: SetBase(ConstEval)
+	{
+	}
+
 	[[nodiscard]] UE_FORCEINLINE_HINT TSet(const SetBase& Copy)
 	{
 		*this = Copy;
