merged UEOps fixes for types with templated conversion operators

diff --git a/Engine/Source/Runtime/Core/Public/Misc/UEOps.h b/Engine/Source/Runtime/Core/Public/Misc/UEOps.h
--- a/Engine/Source/Runtime/Core/Public/Misc/UEOps.h
+++ b/Engine/Source/Runtime/Core/Public/Misc/UEOps.h
@@ -8,7 +8,8 @@
 namespace UE::Core::Private
 {
 	// This only exists to bind to NULL, which is not perfectly-forwardable but which is commonly used in comparison operators.
-	// The type is unique so that it can't be confused with a type that will actually be used by real user code.
+	// The type is unique so that it can't be confused with a type that will actually be used by real user code, and it used to
+	// create a pointer-to-member type which is highly unlikely to be deduced.
 	struct FIncomplete;
 
 	template <typename LhsType, typename RhsType>
@@ -65,7 +66,7 @@ template <typename LhsType, typename RhsType>
 // This overload exists to support A == NULL, which is common in comparison operators, but NULL is not perfectly-forwardable.
 template <typename LhsType>
 	requires UE::Core::Private::CWithUEOpEquals<LhsType, decltype(nullptr)>
-[[nodiscard]] UE_REWRITE constexpr bool operator==(const LhsType& Lhs, UE::Core::Private::FIncomplete* Rhs)
+[[nodiscard]] UE_REWRITE constexpr bool operator==(const LhsType& Lhs, int UE::Core::Private::FIncomplete::* Rhs)
 {
 	return Lhs.UEOpEquals(nullptr);
 }
