Optimized TSet::UnhashElements

diff --git a/Engine/Source/Runtime/Core/Public/Containers/SetUtilities.h b/Engine/Source/Runtime/Core/Public/Containers/SetUtilities.h
--- a/Engine/Source/Runtime/Core/Public/Containers/SetUtilities.h
+++ b/Engine/Source/Runtime/Core/Public/Containers/SetUtilities.h
@@ -90,7 +90,7 @@ class FSetElementId
 {
 public:
 	/** Default constructor. */
-	UE_FORCEINLINE_HINT FSetElementId() = default;
+	constexpr UE_FORCEINLINE_HINT FSetElementId() = default;
 
 	/** @return a boolean value representing whether the id is NULL. */
 	[[nodiscard]] UE_FORCEINLINE_HINT bool IsValidId() const
@@ -108,7 +108,7 @@ public:
 		return A.Index != B.Index;
 	}
 
-	[[nodiscard]] UE_FORCEINLINE_HINT int32 AsInteger() const
+	[[nodiscard]] constexpr UE_FORCEINLINE_HINT int32 AsInteger() const
 	{
 		return Index;
 	}
diff --git a/Engine/Source/Runtime/Core/Public/Containers/SparseSet.h b/Engine/Source/Runtime/Core/Public/Containers/SparseSet.h
--- a/Engine/Source/Runtime/Core/Public/Containers/SparseSet.h
+++ b/Engine/Source/Runtime/Core/Public/Containers/SparseSet.h
@@ -1503,11 +1503,8 @@ private:
 		}
 		else
 		{
-			/** Reset the whole hash to invalid */
-			for (int32 I = 0; I < HashSize; ++I)
-			{
-				HashPtr[I] = FSetElementId();
-			}
+			static_assert(FSetElementId().AsInteger() == -1);
+			FMemory::Memset(HashPtr, 0xFF, HashSize * sizeof(FSetElementId));
 		}
 	}
 
