Set memory optimal MB3 settings by default for iOS\Android

diff --git a/Engine/Source/Runtime/Core/Core.Build.cs b/Engine/Source/Runtime/Core/Core.Build.cs
--- a/Engine/Source/Runtime/Core/Core.Build.cs
+++ b/Engine/Source/Runtime/Core/Core.Build.cs
@@ -215,6 +215,15 @@ public class Core : ModuleRules
 			}
 		}
 
+		if (Target.IsInPlatformGroup(UnrealPlatformGroup.IOS) || Target.IsInPlatformGroup(UnrealPlatformGroup.Android))
+		{
+			// For MB3 we can reduce memory retained in the allocator by limiting the size of allocations coming from bins
+			// In this case if an allocation is above 14KB - we'll use page cache which is shared with other allocators and is pretty fast to allocate from
+			PrivateDefinitions.Add("UE_MB3_MAX_SMALL_POOL_SIZE=14336");
+			PrivateDefinitions.Add("UE_MBC_MAX_LISTED_SMALL_POOL_SIZE=14336");
+			PrivateDefinitions.Add("UE_MBC_NUM_LISTED_SMALL_POOLS=48");
+		}
+
 		if (Target.bCompileICU == true)
 		{
 			AddEngineThirdPartyPrivateStaticDependencies(Target, "ICU");
