Fix error context for missing chunks not getting printed

diff --git a/Engine/Source/Runtime/Core/Private/Experimental/UnifiedError/UnifiedError.cpp b/Engine/Source/Runtime/Core/Private/Experimental/UnifiedError/UnifiedError.cpp
--- a/Engine/Source/Runtime/Core/Private/Experimental/UnifiedError/UnifiedError.cpp
+++ b/Engine/Source/Runtime/Core/Private/Experimental/UnifiedError/UnifiedError.cpp
@@ -96,8 +96,6 @@ FAnsiString FError::GetModuleIdString() const
 
 FText FError::GetErrorMessage(bool bIncludeContext) const
 {
-
-
 	FText FormatText = FText::Join(FText::FromString(": "), NSLOCTEXT("UnifiedError", "DefaultLogPrepend", "{Root/ModuleIdString}.{Root/ErrorCodeString}"), GetFormatErrorText());
 
 	if (bIncludeContext)
@@ -108,10 +106,9 @@ FText FError::GetErrorMessage(bool bIncludeContext) const
 		const IErrorDetails* DetailsIt = ErrorDetails.GetReference();
 		while (DetailsIt != nullptr)
 		{
-			if (!DetailsIt->ShouldInclude(EDetailFilter::IncludeInContextLogMessage))
+			if (DetailsIt->ShouldInclude(EDetailFilter::IncludeInContextLogMessage))
 			{
 				AdditionalContext.Add(FText::FromString(FUtf8String::Printf("(%s:{%s})", DetailsIt->GetErrorDetailsTypeNameForLog().GetData(), DetailsIt->GetErrorDetailsTypeNameForLog().GetData())));
-				
 			}
 			const IErrorDetails* Next = DetailsIt->GetInnerErrorDetails().GetReference(); 
 			DetailsIt = Next;
diff --git a/Engine/Source/Runtime/Experimental/IoStore/OnDemandCore/Public/IO/OnDemandError.h b/Engine/Source/Runtime/Experimental/IoStore/OnDemandCore/Public/IO/OnDemandError.h
--- a/Engine/Source/Runtime/Experimental/IoStore/OnDemandCore/Public/IO/OnDemandError.h
+++ b/Engine/Source/Runtime/Experimental/IoStore/OnDemandCore/Public/IO/OnDemandError.h
@@ -199,10 +199,14 @@ inline void SerializeForLog(FCbWriter& Writer, const UE::UnifiedError::IoStoreOn
 {
 	Writer.BeginObject();
 	Writer.AddString(ANSITEXTVIEW("$type"), TErrorStructFeatures<UE::UnifiedError::IoStoreOnDemand::FChunkMissingErrorContext>::GetErrorContextTypeNameAsString());
-	Writer.AddString(ANSITEXTVIEW("$format"), TEXT("(MissingCount: {MissingCount})"));
+	Writer.AddString(ANSITEXTVIEW("$format"), TEXT("(MissingCount: {MissingCount}, ChunkIds: [{ChunkIds}])"));
+
 	Writer.AddInteger(ANSITEXTVIEW("MissingCount"), Ctx.ChunkIds.Num());
-	//Writer.BeginArray(TEXT("ChunkIds"));
-	//Writer.EndArray();
+
+	TUtf8StringBuilder<1024> Builder;
+	Builder.Join(Ctx.ChunkIds, UTF8TEXTVIEW(", "));
+	Writer.AddString(ANSITEXTVIEW("ChunkIds"), Builder);
+	
 	Writer.EndObject();
 }
 ////////////////////////////////////////////////////////////////////////////////
