Disable new parser for PVS-Studio analysis

#rnx

diff --git a/Engine/Source/Runtime/Core/Public/Analysis/PvsStudioCodeAnalysis.h b/Engine/Source/Runtime/Core/Public/Analysis/PvsStudioCodeAnalysis.h
--- a/Engine/Source/Runtime/Core/Public/Analysis/PvsStudioCodeAnalysis.h
+++ b/Engine/Source/Runtime/Core/Public/Analysis/PvsStudioCodeAnalysis.h
@@ -5,6 +5,9 @@
 #if defined( PVS_STUDIO )
 
 	//PVS-Studio settings:
+	// v7.38 currently crashes using the new parser, use the old parser until resolved
+	//V_USE_OLD_PARSER
+
 	//-V::505,542,677,704,719,720,730,735,780,1002,1055,1062,1100
 	//-V:TRYCOMPRESSION:519,547
 	//-V:check(:501,547,560,605
