[Core] Fix crash caused when copying empty compact set with inline allocators
* CompactSet: add header dependency causing build breakage when enabled
* Move GetData/GetAlignment methods to private as they are not in the sparseset interface and GetData is dangerous to expose directly
* Addressing bughawk feedback and other cleanup
** Add comments around areas bughawk called out but are ok
** Fix FilterByPredicate to return the correct map type
** Rewrite compact set Emplace and EmplaceByHash to be cleaner and more safely return results
** Fix bug in TSetPropertyBinding::InsertNonEmpty which would cause an element to not be correctly hashed for compact sets

#tests Compiled with Compact set enabled

diff --git a/Engine/Source/Runtime/Core/Public/Containers/CompactSet.h.inl b/Engine/Source/Runtime/Core/Public/Containers/CompactSet.h.inl
--- a/Engine/Source/Runtime/Core/Public/Containers/CompactSet.h.inl
+++ b/Engine/Source/Runtime/Core/Public/Containers/CompactSet.h.inl
@@ -579,23 +579,6 @@ public:
 		return Super::GetTotalMemoryRequiredInBytes(NumElements, GetSetLayout());
 	}
 
-
-	/* Get the alignment required for the allocation */
-	[[nodiscard]] UE_FORCEINLINE_HINT static constexpr size_t GetAlignment()
-	{
-		return FGenericPlatformMath::Max<size_t>(alignof(ElementType), UE::Core::CompactHashTable::GetMemoryAlignment());
-	}
-
-	[[nodiscard]] UE_FORCEINLINE_HINT ElementType* GetData()
-	{
-		return (ElementType*)this->Elements.GetAllocation();
-	}
-
-	[[nodiscard]] UE_FORCEINLINE_HINT const ElementType* GetData() const
-	{
-		return (const ElementType*)this->Elements.GetAllocation();
-	}
-
 	/**
 	 * Checks whether an element id is valid.
 	 * @param Id - The element id to check.
@@ -735,7 +718,9 @@ public:
 			ExistingIndex = FindIndexByHash(KeyHash, Key);
 		}
 
-		if (ExistingIndex == INDEX_NONE)
+		const bool bAlreadyInSet = ExistingIndex != INDEX_NONE;
+
+		if (!bAlreadyInSet)
 		{
 			ExistingIndex = this->NumElements;
 			ElementType& NewElement = AddUninitialized(KeyHash);
@@ -747,7 +732,7 @@ public:
 			MoveByRelocate(ExistingElement, TempElement);
 		}
 
-		return { FSetElementId::FromInteger(ExistingIndex), ExistingIndex != INDEX_NONE };
+		return { FSetElementId::FromInteger(ExistingIndex), bAlreadyInSet };
 	}
 	
 	/**
@@ -793,21 +778,21 @@ public:
 			ExistingIndex = FindIndexByHash(KeyHash, KeyFuncs::GetSetKey(TempElement));
 		}
 
-		if (ExistingIndex == INDEX_NONE)
+		const bool bAlreadyInSet = ExistingIndex != INDEX_NONE;
+
+		if (!bAlreadyInSet)
 		{
 			ExistingIndex = this->NumElements;
 			ElementType& NewElement = AddUninitialized(KeyHash);
 			RelocateConstructItem<ElementType, ElementType, SizeType>(&NewElement, &TempElement);
-
-			return { FSetElementId::FromInteger(ExistingIndex), false };
 		}
 		else
 		{
 			ElementType& ExistingElement = GetData()[ExistingIndex];
 			MoveByRelocate(ExistingElement, TempElement);
-
-			return { FSetElementId::FromInteger(ExistingIndex), true };
 		}
+
+		return { FSetElementId::FromInteger(ExistingIndex), bAlreadyInSet };
 	}
 
 	/**
@@ -1394,6 +1379,22 @@ public:
 	}
 
 private:
+	/* Get the alignment required for the allocation */
+	[[nodiscard]] UE_FORCEINLINE_HINT static constexpr size_t GetAlignment()
+	{
+		return FGenericPlatformMath::Max<size_t>(alignof(ElementType), UE::Core::CompactHashTable::GetMemoryAlignment());
+	}
+
+	[[nodiscard]] UE_FORCEINLINE_HINT ElementType* GetData()
+	{
+		return (ElementType*)this->Elements.GetAllocation();
+	}
+
+	[[nodiscard]] UE_FORCEINLINE_HINT const ElementType* GetData() const
+	{
+		return (const ElementType*)this->Elements.GetAllocation();
+	}
+
 	[[nodiscard]] UE_FORCEINLINE_HINT static constexpr FCompactSetLayout GetSetLayout()
 	{
 		return { sizeof(ElementType), GetAlignment() };
@@ -1424,8 +1425,9 @@ private:
 	/** Recalculate the lookup table, used after bulk operations are performed */
 	void Rehash()
 	{
-		if (const ElementType* ElementData = GetData())
+		if (this->MaxElements > 0)
 		{
+			const ElementType* ElementData = GetData();
 			FCompactHashTableView HashTable = GetHashTableView();
 
 			HashTable.Reset();
diff --git a/Engine/Source/Runtime/Core/Public/Containers/Map.h.inl b/Engine/Source/Runtime/Core/Public/Containers/Map.h.inl
--- a/Engine/Source/Runtime/Core/Public/Containers/Map.h.inl
+++ b/Engine/Source/Runtime/Core/Public/Containers/Map.h.inl
@@ -462,9 +462,9 @@ public:
 	 *          the subset of elements for which the functor returns true.
 	 */
 	template <typename Predicate>
-	[[nodiscard]] TMap<KeyType, ValueType, SetAllocator, KeyFuncs> FilterByPredicate(Predicate Pred) const
+	[[nodiscard]] TMAP<KeyType, ValueType, SetAllocator, KeyFuncs> FilterByPredicate(Predicate Pred) const
 	{
-		TMap<KeyType, ValueType, SetAllocator, KeyFuncs> FilterResults;
+		TMAP<KeyType, ValueType, SetAllocator, KeyFuncs> FilterResults;
 		FilterResults.Reserve(Pairs.Num());
 		for (const ElementType& Pair : Pairs)
 		{
diff --git a/Engine/Source/Runtime/Core/Public/Containers/ScriptCompactSet.h b/Engine/Source/Runtime/Core/Public/Containers/ScriptCompactSet.h
--- a/Engine/Source/Runtime/Core/Public/Containers/ScriptCompactSet.h
+++ b/Engine/Source/Runtime/Core/Public/Containers/ScriptCompactSet.h
@@ -97,6 +97,8 @@ public:
 
 			this->GetHashTableView(Layout).Remove(Index, KeyHash, this->NumElements - 1, GetKeyHash(Src));
 			DestructItem(Dst);
+
+			// Memmove is fine here as our containers already only work if data is trivially relocatable
 			FMemory::Memmove(Dst, Src, Layout.Size);
 		}
 
diff --git a/Engine/Source/Runtime/Core/Public/Containers/Set.h b/Engine/Source/Runtime/Core/Public/Containers/Set.h
--- a/Engine/Source/Runtime/Core/Public/Containers/Set.h
+++ b/Engine/Source/Runtime/Core/Public/Containers/Set.h
@@ -21,6 +21,9 @@ using TScriptSet = TScriptCompactSet<Allocator>;
 #include "Containers/CompactSet.h.inl"
 #undef UE_TCOMPACT_SET
 
+// Todo: This should be removed when existing code is fixed to explicitly include what it uses
+#include "Containers/BitArray.h"
+
 #else
 
 #include "Containers/ScriptSparseSet.h"
diff --git a/Engine/Source/Runtime/CoreUObject/Public/UObject/UnrealType.h b/Engine/Source/Runtime/CoreUObject/Public/UObject/UnrealType.h
--- a/Engine/Source/Runtime/CoreUObject/Public/UObject/UnrealType.h
+++ b/Engine/Source/Runtime/CoreUObject/Public/UObject/UnrealType.h
@@ -5092,6 +5092,8 @@ public:
 		return WithScriptMap([this, InternalIndex, Count](auto* Map)
 		{
 			check(Map->IsValidIndex(InternalIndex));
+
+			// Checking against max as sparse sets may want to remove a range including invalid items, so num is not a correct validation here
 			check(InternalIndex + Count <= Map->Max());
 
 #if UE_USE_COMPACT_SET_AS_DEFAULT
@@ -5930,6 +5932,8 @@ public:
 	void RemoveAt(int32 InternalIndex, int32 Count = 1)
 	{
 		check(IsValidIndex(InternalIndex));
+
+		// Checking against max as sparse sets may want to remove a range including invalid items, so num is not a correct validation here
 		check(InternalIndex + Count <= Set->Max());
 
 #if UE_USE_COMPACT_SET_AS_DEFAULT
