Reimplement TVectorRegisterType with partial specializations rather than std::conditional_t

TVectorRegisterType was not inheriting the alignment of the result in struct packing situations when building with Clang for Windows. This should be working now.

Fix courtesy of [REDACTED_PERSON].

diff --git a/Engine/Source/Runtime/Core/Public/Math/VectorRegister.h b/Engine/Source/Runtime/Core/Public/Math/VectorRegister.h
--- a/Engine/Source/Runtime/Core/Public/Math/VectorRegister.h
+++ b/Engine/Source/Runtime/Core/Public/Math/VectorRegister.h
@@ -22,9 +22,31 @@
 #define SIMD_ALIGNMENT (alignof(VectorRegister))
 
 
+namespace UE::Math::VectorRegisterPrivate
+{
+	// Have to do this with partial specializations to work around Clang bug on Windows
+	template<typename T>
+	struct TVectorRegisterTypeHelper
+	{
+		using Type = void;
+	};
+
+	template<>
+	struct TVectorRegisterTypeHelper<float>
+	{
+		using Type = VectorRegister4Float;
+	};
+
+	template<>
+	struct TVectorRegisterTypeHelper<double>
+	{
+		using Type = VectorRegister4Double;
+	};
+}
+
 // Alias which helps resolve a template type (float/double) to the appropriate VectorRegister type (VectorRegister4Float/VectorRegister4Double).
 template<typename T>
-using TVectorRegisterType = std::conditional_t<std::is_same_v<T, float>, VectorRegister4Float, std::conditional_t<std::is_same_v<T, double>, VectorRegister4Double, void> >;
+using TVectorRegisterType = typename UE::Math::VectorRegisterPrivate::TVectorRegisterTypeHelper<T>::Type;
 
 // 'Cross-platform' vector intrinsics (built on the platform-specific ones defined above)
 #include "Math/UnrealMathVectorCommon.h.inl"
