CrashHandler - Include additional crash context providers in crash info folder on Apple platforms & Linux

#tests Local test on macOS & build iOS

diff --git a/Engine/Source/Runtime/Core/Private/GenericPlatform/GenericPlatformCrashContext.cpp b/Engine/Source/Runtime/Core/Private/GenericPlatform/GenericPlatformCrashContext.cpp
--- a/Engine/Source/Runtime/Core/Private/GenericPlatform/GenericPlatformCrashContext.cpp
+++ b/Engine/Source/Runtime/Core/Private/GenericPlatform/GenericPlatformCrashContext.cpp
@@ -1888,7 +1888,10 @@ void FGenericCrashContext::DumpAdditionalContext(const TCHAR* CrashFolderAbsolut
 {
 #if WITH_ADDITIONAL_CRASH_CONTEXTS 
 	FCrashContextExtendedWriterImpl Writer(CrashFolderAbsolute);
-	AdditionalCrashContextDelegate.Broadcast(Writer);
+	if (!UE::IsGCLocked())
+	{
+		AdditionalCrashContextDelegate.Broadcast(Writer);
+	}
 	FAdditionalCrashContextStack::ExecuteProviders(Writer);
 #endif
 }
diff --git a/Engine/Source/Runtime/Core/Private/IOS/IOSPlatformCrashContext.cpp b/Engine/Source/Runtime/Core/Private/IOS/IOSPlatformCrashContext.cpp
--- a/Engine/Source/Runtime/Core/Private/IOS/IOSPlatformCrashContext.cpp
+++ b/Engine/Source/Runtime/Core/Private/IOS/IOSPlatformCrashContext.cpp
@@ -95,6 +95,9 @@ void FIOSCrashContext::GenerateInfoInFolder(char const* const InfoFolder, bool b
             
             close(ReportFile);
         }
+		
+		// generate additional files from providers.
+		DumpAdditionalContext(*FString(CrashInfoFolder));
                 
         // generate "minidump" (Apple crash log format)
         FCStringAnsi::Strncpy(FilePath, CrashInfoFolder, PATH_MAX);
diff --git a/Engine/Source/Runtime/Core/Private/Mac/MacPlatformMisc.cpp b/Engine/Source/Runtime/Core/Private/Mac/MacPlatformMisc.cpp
--- a/Engine/Source/Runtime/Core/Private/Mac/MacPlatformMisc.cpp
+++ b/Engine/Source/Runtime/Core/Private/Mac/MacPlatformMisc.cpp
@@ -2263,6 +2263,9 @@ void FMacCrashContext::GenerateInfoInFolder(char const* const InfoFolder) const
 	{
 		char FilePath[PATH_MAX] = {};
 		
+		// generate additional files from providers.
+		DumpAdditionalContext(*FString(CrashInfoFolder));
+		
 		// generate "minidump" (Apple crash log format)
 		FCStringAnsi::Strncpy(FilePath, CrashInfoFolder, PATH_MAX);
 		FCStringAnsi::StrncatTruncateDest(FilePath, PATH_MAX, "/minidump.dmp");
diff --git a/Engine/Source/Runtime/Core/Private/Unix/UnixPlatformCrashContext.cpp b/Engine/Source/Runtime/Core/Private/Unix/UnixPlatformCrashContext.cpp
--- a/Engine/Source/Runtime/Core/Private/Unix/UnixPlatformCrashContext.cpp
+++ b/Engine/Source/Runtime/Core/Private/Unix/UnixPlatformCrashContext.cpp
@@ -619,6 +619,9 @@ void FUnixCrashContext::GenerateCrashInfoAndLaunchReporter() const
 		FString CrashInfoAbsolute = FPaths::ConvertRelativePathToFull(CrashInfoFolder);
 		if (IFileManager::Get().MakeDirectory(*CrashInfoAbsolute, true))
 		{
+			// generate additional files from providers.
+			DumpAdditionalContext(*CrashInfoFolder);
+
 			// generate "minidump"
 			GenerateReport(FPaths::Combine(*CrashInfoAbsolute, TEXT("Diagnostics.txt")));
 
