Add delegate lifetime checking

- Add new CheckDelegateLifetime UBT flag to have mutlicast delegate validate instances before any call
- Add callsite/module name storage to a new FTrackedDelegateInstanceExtras layer
- Add module manager helper to confirm whether a given module is currently safe to call code for
- Add new GetInvocationNames mutlicast helper to dump all current users, for use in debugging

#rnx

diff --git a/Engine/Source/Programs/UnrealBuildTool/Configuration/ReadOnlyTargetRules.cs b/Engine/Source/Programs/UnrealBuildTool/Configuration/ReadOnlyTargetRules.cs
--- a/Engine/Source/Programs/UnrealBuildTool/Configuration/ReadOnlyTargetRules.cs
+++ b/Engine/Source/Programs/UnrealBuildTool/Configuration/ReadOnlyTargetRules.cs
@@ -517,6 +517,8 @@ namespace UnrealBuildTool
 
 		public bool bMergeModules => Inner.bMergeModules;
 
+		public bool bCheckDelegateLifetime => Inner.bCheckDelegateLifetime;
+
 		public IReadOnlySet<string> MergeAdditionalCommonModules => Inner.MergeAdditionalCommonModules;
 
 		public IReadOnlyList<string> MergePlugins => Inner.MergePlugins;
diff --git a/Engine/Source/Programs/UnrealBuildTool/Configuration/TargetRules.cs b/Engine/Source/Programs/UnrealBuildTool/Configuration/TargetRules.cs
--- a/Engine/Source/Programs/UnrealBuildTool/Configuration/TargetRules.cs
+++ b/Engine/Source/Programs/UnrealBuildTool/Configuration/TargetRules.cs
@@ -2559,6 +2559,13 @@ namespace UnrealBuildTool
 		[CommandLine("-MergeModules")]
 		public bool bMergeModules { get; set; }
 
+		/// <summary>
+		/// Experimental: Check lifetimes for delegate instances. Recommended for merged modules testing
+		/// </summary>
+		[RequiresUniqueBuildEnvironment]
+		[CommandLine("-CheckDelegateLifetime")]
+		public bool bCheckDelegateLifetime { get; set; }
+
 		/// <summary>
 		/// Experimental: List of modules to exclude from Engine to be placed in Common
 		/// </summary>
diff --git a/Engine/Source/Runtime/Core/Core.Build.cs b/Engine/Source/Runtime/Core/Core.Build.cs
--- a/Engine/Source/Runtime/Core/Core.Build.cs
+++ b/Engine/Source/Runtime/Core/Core.Build.cs
@@ -308,6 +308,7 @@ public class Core : ModuleRules
 		}
 
 		PublicDefinitions.Add("UE_MERGED_MODULES=" + (Target.bMergeModules ? "1" : "0"));
+		PublicDefinitions.Add("UE_DELEGATE_CHECK_LIFETIME=" + (Target.bCheckDelegateLifetime ? "1" : "0"));
 	}
 
 	protected virtual bool SupportsBinaryConfig(ReadOnlyTargetRules Target)
diff --git a/Engine/Source/Runtime/Core/Private/Delegates/DelegateBase.cpp b/Engine/Source/Runtime/Core/Private/Delegates/DelegateBase.cpp
--- a/Engine/Source/Runtime/Core/Private/Delegates/DelegateBase.cpp
+++ b/Engine/Source/Runtime/Core/Private/Delegates/DelegateBase.cpp
@@ -2,6 +2,33 @@
 
 #include "Delegates/DelegateBase.h"
 
+#if UE_DELEGATE_CHECK_LIFETIME
+#include "Modules/ModuleManager.h"
+
+DEFINE_LOG_CATEGORY_STATIC(LogDelegates, Log, All);
+
+bool FTrackedDelegateInstanceExtras::IsValid() const
+{
+	return ModuleName == NAME_None || FModuleManager::Get().IsModuleSafeToUse(ModuleName);
+}
+
+bool FTrackedDelegateInstanceExtras::CheckValid() const
+{
+	if (!IsValid())
+	{
+		UE_LOG(LogDelegates, Fatal, TEXT("FTrackedDelegateInstanceExtras::CheckValid: bad delegate from %s (%s)"),
+			*GetBoundFunctionName().ToString(), *ModuleName.ToString());
+
+		return false;
+	}
+	else
+	{
+		return true;
+	}
+}
+
+#endif //UE_DELEGATE_CHECK_LIFETIME
+
 void* UE::Core::Private::DelegateAllocate(size_t Size, FDelegateAllocation& Allocation)
 {
 	int32 NewDelegateSize = FMath::DivideAndRoundUp((int32)Size, (int32)sizeof(FAlignedInlineDelegateType));
diff --git a/Engine/Source/Runtime/Core/Private/Modules/ModuleManager.cpp b/Engine/Source/Runtime/Core/Private/Modules/ModuleManager.cpp
--- a/Engine/Source/Runtime/Core/Private/Modules/ModuleManager.cpp
+++ b/Engine/Source/Runtime/Core/Private/Modules/ModuleManager.cpp
@@ -186,6 +186,7 @@ FModuleManager& FModuleManager::Get()
 FModuleManager::FModuleManager(FPrivateToken)
 	: bCanProcessNewlyLoadedObjects(false)
 	, bExtraBinarySearchPathsAdded(false)
+	, bIsLoadingDynamicLibrary(false)
 {
 	check(IsInGameThread());
 
@@ -552,7 +553,10 @@ void* FModuleManager::InternalLoadLibrary(FName ModuleName, const FString& Libra
 
 	// Do the actual loading
 	const uint64_t InitialMemUsed = FPlatformMemory::GetMemoryUsedFast();
-	Handle = FPlatformProcess::GetDllHandle(*LibraryToLoad);
+	{
+		TGuardValue<bool> IsLoadingDynamicLibraryGuard(bIsLoadingDynamicLibrary, true);
+		Handle = FPlatformProcess::GetDllHandle(*LibraryToLoad);
+	}
 	const uint64_t FinalMemUsed = FPlatformMemory::GetMemoryUsedFast();
 
 #if UE_MERGED_MODULES
@@ -1444,6 +1448,45 @@ IModuleInterface* FModuleManager::GetModule( const FName InModuleName )
 	return nullptr;
 }
 
+bool FModuleManager::IsModuleSafeToUse(const FName InModuleName) const
+{
+	FScopeLock Lock(&ModulesCriticalSection);
+
+	if (bIsLoadingDynamicLibrary)
+	{
+		// Dynamic library loading triggers a lot of construction at a time where we don't have good bookkeeping on modules
+		return true;
+	}
+
+	if (!ModulePathsCache.Contains(InModuleName))
+	{
+		// Module is compiled-in
+		return true;
+	}
+
+	else
+	{
+#if UE_MERGED_MODULES
+
+		const FString& ModuleLibrary = ModulePathsCache[InModuleName];
+		if (LibraryHandles.Contains(ModuleLibrary))
+		{
+			// The merged library for this module is currently live
+			return true;
+		}
+
+#else
+
+		// Just check if the module is loaded
+		return IsModuleLoaded(InModuleName);
+
+#endif // UE_MERGED_MODULES
+	}
+
+	// Module is neither compiled-in, nor currently loaded
+	return false;
+}
+
 bool FModuleManager::Exec_Dev( UWorld* Inworld, const TCHAR* Cmd, FOutputDevice& Ar )
 {
 #if !UE_BUILD_SHIPPING
diff --git a/Engine/Source/Runtime/Core/Public/Delegates/DelegateBase.h b/Engine/Source/Runtime/Core/Public/Delegates/DelegateBase.h
--- a/Engine/Source/Runtime/Core/Public/Delegates/DelegateBase.h
+++ b/Engine/Source/Runtime/Core/Public/Delegates/DelegateBase.h
@@ -35,6 +35,87 @@ class TDelegateBase;
 
 ALIAS_TEMPLATE_TYPE_LAYOUT(template<typename ElementType>, FDelegateAllocatorType::ForElementType<ElementType>, void*);
 
+#if UE_DELEGATE_CHECK_LIFETIME
+
+// This tracking layer allows us to safely check if the code for virtual calls in IDelegateInstance are safe
+// This code will be called on instances that potentially have their code unloaded!
+class FTrackedDelegateInstanceExtras : public IDelegateInstance
+{
+public:
+	FTrackedDelegateInstanceExtras()
+#ifdef UE_MODULE_NAME
+		: ModuleName(UE_MODULE_NAME)
+#else
+		: ModuleName(NAME_None)
+#endif //UE_MODULE_NAME
+		, BoundFunctionName(NAME_None)
+	{}
+
+	/**
+	 * Set the name of the module this was compiled in.
+	 */
+	void SetModuleName(FName InModuleName)
+	{
+		ModuleName = InModuleName;
+	}
+
+	/**
+	 * Retrieves the name of the module this was compiled in.
+	 */
+	FName GetModuleName() const
+	{
+		return ModuleName;
+	}
+
+	/**
+	 * Stores the debug name for this delegate's call site - could be a context object, filename, etc.
+	 */
+	void SetBoundFunctionName(FName InBoundFunctionName)
+	{
+		BoundFunctionName = InBoundFunctionName;
+	}
+
+	/**
+	 * Retrieves the debug name of the bound function or call site.
+	 * This is similar to TryGetBoundFunctionName in IDelegateInstance but we need a non-virtual version with its own storage.
+	 * In turn, TryGetBoundFunctionName will call this if compiled with UE_DELEGATE_CHECK_LIFETIME.
+	 */
+	FName GetBoundFunctionName() const
+	{
+		return BoundFunctionName;
+	}
+
+	/**
+	 * Checks whether this delegate instance is safe to call virtual functions on (i.e. the module for it is loaded).
+	 * 
+	 * Returns whether the instance is valid.
+	 */
+	CORE_API bool IsValid() const;
+
+	/**
+	 * Checks whether this delegate instance is safe, and produce debug logging if not.
+	 * 
+	 * Returns whether the instance is valid.
+	 */
+	CORE_API bool CheckValid() const;
+
+protected:
+
+	FName ModuleName;
+	FName BoundFunctionName;
+};
+
+// Convenience macro to check for lifetime validity as widely as possible
+#define CHECK_DELEGATE_LIFETIME(DelegateInstance)                                           \
+    if (DelegateInstance != nullptr)                                                        \
+    {                                                                                       \
+    	static_cast<const FTrackedDelegateInstanceExtras*>(DelegateInstance)->CheckValid(); \
+    }
+
+#else
+#define CHECK_DELEGATE_LIFETIME(DelegateInstance)
+#endif // UE_DELEGATE_CHECK_LIFETIME
+
 // not thread-safe version, with automatic race detection in dev builds
 struct FDefaultDelegateUserPolicy
 {
@@ -61,7 +142,12 @@ struct FDefaultDelegateUserPolicy
 	//     holds an array of FDelegateExtras-derived objects which is the invocation list).
 	//   - Public data members and member functions are accessible directly through the TMulticastDelegate object.
 
+#if UE_DELEGATE_CHECK_LIFETIME
+	using FDelegateInstanceExtras = FTrackedDelegateInstanceExtras;
+#else
 	using FDelegateInstanceExtras = IDelegateInstance;
+#endif // UE_DELEGATE_CHECK_LIFETIME
+
 	using FThreadSafetyMode =
 #if UE_DETECT_DELEGATES_RACE_CONDITIONS
 		FNotThreadSafeDelegateMode;
@@ -76,7 +162,13 @@ struct FDefaultDelegateUserPolicy
 struct FDefaultTSDelegateUserPolicy
 {
 	// see `FDefaultDelegateUserPolicy` for documentation
+
+#if UE_DELEGATE_CHECK_LIFETIME
+	using FDelegateInstanceExtras = FTrackedDelegateInstanceExtras;
+#else
 	using FDelegateInstanceExtras = IDelegateInstance;
+#endif // UE_DELEGATE_CHECK_LIFETIME
+
 	using FThreadSafetyMode = FThreadSafeDelegateMode;
 	using FDelegateExtras = TDelegateBase<FThreadSafetyMode>;
 	using FMulticastDelegateExtras = TMulticastDelegateBase<FDefaultTSDelegateUserPolicy>;
@@ -85,7 +177,12 @@ struct FDefaultTSDelegateUserPolicy
 // not thread-safe version, no race detection. used primarily for deprecated unsafe delegates that must be kept running for backward compatibility
 struct FNotThreadSafeNotCheckedDelegateUserPolicy
 {
+#if UE_DELEGATE_CHECK_LIFETIME
+	using FDelegateInstanceExtras = FTrackedDelegateInstanceExtras;
+#else
 	using FDelegateInstanceExtras = IDelegateInstance;
+#endif // UE_DELEGATE_CHECK_LIFETIME
+
 	using FThreadSafetyMode = FNotThreadSafeNotCheckedDelegateMode;
 	using FDelegateExtras = TDelegateBase<FThreadSafetyMode>;
 	using FMulticastDelegateExtras = TMulticastDelegateBase<FNotThreadSafeNotCheckedDelegateUserPolicy>;
diff --git a/Engine/Source/Runtime/Core/Public/Delegates/DelegateInstancesImpl.h b/Engine/Source/Runtime/Core/Public/Delegates/DelegateInstancesImpl.h
--- a/Engine/Source/Runtime/Core/Public/Delegates/DelegateInstancesImpl.h
+++ b/Engine/Source/Runtime/Core/Public/Delegates/DelegateInstancesImpl.h
@@ -22,6 +22,17 @@
 #include "UObject/NameTypes.h"
 #include "UObject/WeakObjectPtrTemplates.h"
 
+#if UE_DELEGATE_CHECK_LIFETIME
+#  ifdef UE_MODULE_NAME
+#    define INSTRUMENT_DELEGATE_MODULE() FTrackedDelegateInstanceExtras::SetModuleName(UE_MODULE_NAME)
+#  else
+#    define INSTRUMENT_DELEGATE_MODULE()
+#  endif // UE_MODULE_NAME
+#  define INSTRUMENT_DELEGATE_FUNCTION(Function) FTrackedDelegateInstanceExtras::SetBoundFunctionName(Function)
+#else
+#  define INSTRUMENT_DELEGATE_MODULE()
+#  define INSTRUMENT_DELEGATE_FUNCTION(Function)
+#endif // UE_DELEGATE_CHECK_LIFETIME
 
 namespace UE::Delegates::Private
 {
@@ -30,7 +41,7 @@ namespace UE::Delegates::Private
 }
 
 template <typename FuncType, typename UserPolicy, typename... VarTypes>
-class TCommonDelegateInstanceState : IBaseDelegateInstance<FuncType, UserPolicy>
+class TCommonDelegateInstanceState : protected IBaseDelegateInstance<FuncType, UserPolicy>
 {
 public:
 	template <typename... InVarTypes>
@@ -74,6 +85,9 @@ public:
 	{
 		check(InFunctionName != NAME_None);
 
+		INSTRUMENT_DELEGATE_MODULE();
+		INSTRUMENT_DELEGATE_FUNCTION(*UserClass::StaticClass()->GetName());
+
 		if (InUserObject != nullptr)
 		{
 			CachedFunction = UserObjectPtr->FindFunctionChecked(InFunctionName);
@@ -211,6 +225,9 @@ public:
 		// NOTE: Shared pointer delegates are allowed to have a null incoming object pointer.  Weak pointers can expire,
 		//       an it is possible for a copy of a delegate instance to end up with a null pointer.
 		checkSlow(MethodPtr != nullptr);
+
+		INSTRUMENT_DELEGATE_MODULE();
+		INSTRUMENT_DELEGATE_FUNCTION(__FUNCTION__);
 	}
 
 	// IDelegateInstance interface
@@ -341,6 +358,8 @@ public:
 		, ContextObject(InContextObject)
 		, Functor(Forward<InFunctorType>(InFunctor))
 	{
+		INSTRUMENT_DELEGATE_MODULE();
+		INSTRUMENT_DELEGATE_FUNCTION(__FUNCTION__);
 	}
 
 	// IDelegateInstance interface
@@ -451,6 +470,9 @@ public:
 	{
 		// Non-expirable delegates must always have a non-null object pointer on creation (otherwise they could never execute.)
 		check(InUserObject != nullptr && MethodPtr != nullptr);
+
+		INSTRUMENT_DELEGATE_MODULE();
+		INSTRUMENT_DELEGATE_FUNCTION(__FUNCTION__);
 	}
 
 	// IDelegateInstance interface
@@ -577,6 +599,9 @@ public:
 		// NOTE: UObject delegates are allowed to have a null incoming object pointer.  UObject weak pointers can expire,
 		//       an it is possible for a copy of a delegate instance to end up with a null pointer.
 		checkSlow(MethodPtr != nullptr);
+
+		INSTRUMENT_DELEGATE_MODULE();
+		INSTRUMENT_DELEGATE_FUNCTION(*UserClass::StaticClass()->GetName());
 	}
 
 	// IDelegateInstance interface
@@ -718,6 +743,8 @@ public:
 		, StaticFuncPtr(InStaticFuncPtr)
 	{
 		check(StaticFuncPtr != nullptr);
+
+		INSTRUMENT_DELEGATE_MODULE();
 	}
 
 	// IDelegateInstance interface
@@ -822,6 +849,7 @@ public:
 		: Super  (Forward<InVarTypes>(Vars)...)
 		, Functor(Forward<InFunctorType>(InFunctor))
 	{
+		INSTRUMENT_DELEGATE_MODULE();
 	}
 
 	// IDelegateInstance interface
@@ -922,6 +950,8 @@ public:
 		, ContextObject(InContextObject)
 		, Functor      (Forward< InFunctorType>(InFunctor))
 	{
+		INSTRUMENT_DELEGATE_MODULE();
+		INSTRUMENT_DELEGATE_FUNCTION(*InContextObject->GetName());
 	}
 
 	// IDelegateInstance interface
diff --git a/Engine/Source/Runtime/Core/Public/Delegates/DelegateSignatureImpl.inl b/Engine/Source/Runtime/Core/Public/Delegates/DelegateSignatureImpl.inl
--- a/Engine/Source/Runtime/Core/Public/Delegates/DelegateSignatureImpl.inl
+++ b/Engine/Source/Runtime/Core/Public/Delegates/DelegateSignatureImpl.inl
@@ -605,6 +605,7 @@ public:
 		FReadAccessScope ReadScope = GetReadAccessScope();
 
 		const DelegateInstanceInterfaceType* LocalDelegateInstance = GetDelegateInstanceProtected();
+		CHECK_DELEGATE_LIFETIME(LocalDelegateInstance);
 
 		// If this assert goes off, Execute() was called before a function was bound to the delegate.
 		// Consider using ExecuteIfBound() instead.
@@ -626,6 +627,7 @@ public:
 
 		if (const DelegateInstanceInterfaceType* Ptr = GetDelegateInstanceProtected())
 		{
+			CHECK_DELEGATE_LIFETIME(Ptr);
 			return Ptr->ExecuteIfSafe(Forward<ParamTypes>(Params)...);
 		}
 
diff --git a/Engine/Source/Runtime/Core/Public/Delegates/MulticastDelegateBase.h b/Engine/Source/Runtime/Core/Public/Delegates/MulticastDelegateBase.h
--- a/Engine/Source/Runtime/Core/Public/Delegates/MulticastDelegateBase.h
+++ b/Engine/Source/Runtime/Core/Public/Delegates/MulticastDelegateBase.h
@@ -114,6 +114,8 @@ public:
 		for (const UnicastDelegateType& DelegateBaseRef : InvocationList)
 		{
 			const IDelegateInstance* DelegateInstance = DelegateBaseRef.GetDelegateInstanceProtected();
+			CHECK_DELEGATE_LIFETIME(DelegateInstance);
+
 			if ((DelegateInstance != nullptr) && DelegateInstance->HasSameObject(InUserObject))
 			{
 				return true;
@@ -140,6 +142,8 @@ public:
 			for (UnicastDelegateType& DelegateBaseRef : InvocationList)
 			{
 				IDelegateInstance* DelegateInstance = DelegateBaseRef.GetDelegateInstanceProtected();
+				CHECK_DELEGATE_LIFETIME(DelegateInstance);
+
 				if ((DelegateInstance != nullptr) && DelegateInstance->HasSameObject(InUserObject))
 				{
 					// Manually unbind the delegate here so the compaction will find and remove it.
@@ -162,6 +166,8 @@ public:
 				UnicastDelegateType& DelegateBaseRef = InvocationList[InvocationListIndex];
 
 				IDelegateInstance* DelegateInstance = DelegateBaseRef.GetDelegateInstanceProtected();
+				CHECK_DELEGATE_LIFETIME(DelegateInstance);
+
 				if (DelegateInstance == nullptr
 					|| DelegateInstance->HasSameObject(InUserObject)
 					|| DelegateInstance->IsCompactable())
@@ -199,6 +205,39 @@ public:
 		return Size;
 	}
 
+#if UE_DELEGATE_CHECK_LIFETIME
+
+	/**
+	 * Gets a list of all delegate call sites & module names registered to this multicast
+	 *
+	 * @return A list of pair of call site + module name, including for delegates that were unbound but not yet compacted.
+	 */
+	TArray<TPair<FName, FName>> GetInvocationNames() const
+	{
+		TArray<TPair<FName, FName>> FoundFunctionsAndModules;
+
+		FReadAccessScope ReadScope = GetReadAccessScope();
+
+		for (const UnicastDelegateType& DelegateBaseRef : InvocationList)
+		{
+			const IDelegateInstance* DelegateInstance = DelegateBaseRef.GetDelegateInstanceProtected();
+			CHECK_DELEGATE_LIFETIME(DelegateInstance);
+
+			if (DelegateInstance != nullptr)
+			{
+				const FTrackedDelegateInstanceExtras* TrackedDelegateInstance = static_cast<const FTrackedDelegateInstanceExtras*>(DelegateInstance);
+
+				FoundFunctionsAndModules.Add(TPair<FName, FName>(
+					TrackedDelegateInstance->GetBoundFunctionName(),
+					TrackedDelegateInstance->GetModuleName()));
+			}
+		}
+
+		return FoundFunctionsAndModules;
+	}
+
+#endif // UE_DELEGATE_CHECK_LIFETIME
+
 protected:
 
 	/** Hidden default constructor. */
@@ -221,6 +260,8 @@ protected:
 		    {
 				if (const IDelegateInstance* OtherInstance = OtherDelegateRef.GetDelegateInstanceProtected())
 			    {
+					CHECK_DELEGATE_LIFETIME(OtherInstance);
+
 				    UnicastDelegateType TempDelegate;
 				    static_cast<const DelegateInstanceInterfaceType*>(OtherInstance)->CreateCopy(TempDelegate);
 				    TempInvocationList.Add(MoveTemp(TempDelegate));
@@ -255,6 +296,8 @@ protected:
 				const UnicastDelegateType& DelegateBase = LocalInvocationList[InvocationListIndex];
 
 				const IDelegateInstance* DelegateInstanceInterface = DelegateBase.GetDelegateInstanceProtected();
+				CHECK_DELEGATE_LIFETIME(DelegateInstanceInterface);
+
 				if (DelegateInstanceInterface == nullptr || !static_cast<const DelegateInstanceInterfaceType*>(DelegateInstanceInterface)->ExecuteIfSafe(Params...))
 				{
 					NeedsCompaction = true;
@@ -305,6 +348,8 @@ protected:
 			UnicastDelegateType& DelegateBase = InvocationList[InvocationListIndex];
 
 			IDelegateInstance* DelegateInstance = DelegateBase.GetDelegateInstanceProtected();
+			CHECK_DELEGATE_LIFETIME(DelegateInstance);
+
 			if (DelegateInstance && DelegateInstance->GetHandle() == Handle)
 			{
 				DelegateBase.Unbind();
@@ -360,6 +405,8 @@ private:
 			auto& DelegateBaseRef = InvocationList[InvocationListIndex];
 
 			IDelegateInstance* DelegateInstance = DelegateBaseRef.GetDelegateInstanceProtected();
+			CHECK_DELEGATE_LIFETIME(DelegateInstance);
+
 			if (DelegateInstance == nullptr	|| DelegateInstance->IsCompactable())
 			{
 				InvocationList.RemoveAtSwap(InvocationListIndex, EAllowShrinking::No);
@@ -427,6 +474,11 @@ private:
 	{
 		for (UnicastDelegateType& DelegateBaseRef : InvocationList)
 		{
+#if UE_DELEGATE_CHECK_LIFETIME
+			IDelegateInstance* DelegateInstance = DelegateBaseRef.GetDelegateInstanceProtected();
+			CHECK_DELEGATE_LIFETIME(DelegateInstance);
+#endif // UE_DELEGATE_CHECK_LIFETIME
+
 			DelegateBaseRef.Unbind();
 		}
 
diff --git a/Engine/Source/Runtime/Core/Public/Modules/ModuleManager.h b/Engine/Source/Runtime/Core/Public/Modules/ModuleManager.h
--- a/Engine/Source/Runtime/Core/Public/Modules/ModuleManager.h
+++ b/Engine/Source/Runtime/Core/Public/Modules/ModuleManager.h
@@ -217,6 +217,14 @@ public:
 	 */
 	CORE_API IModuleInterface* GetModule( const FName InModuleName );
 
+	/**
+	 * Checks whether a module is compiled-in, or currently loaded, with all code inside safe to call.
+	 * 
+	 * @param InModuleName Name of the module to check.
+	 * @return true for a module that exists and is not compiled into the executable.
+	 */
+	CORE_API bool IsModuleSafeToUse( const FName InModuleName ) const;
+
 	/**
 	 * Checks whether the specified module is currently loaded.
 	 *
@@ -742,6 +750,9 @@ private:
 	/** True once AddExtraBinarySearchPaths has been called */
 	bool bExtraBinarySearchPathsAdded;
 
+	/** Check if we're in the middle of loading a dynamic library */
+	bool bIsLoadingDynamicLibrary;
+
 	/** Cache of known module paths. Used for performance. Can increase editor startup times by up to 30% */
 	mutable TMap<FName, FString> ModulePathsCache;
 
