Fixed inability to override MB2\MB3 with a crash allocator in cases when PLATFORM_USES_FIXED_GMalloc_CLASS is set to 1

diff --git a/Engine/Source/Runtime/Core/Private/GenericPlatform/GenericPlatformMallocCrash.cpp b/Engine/Source/Runtime/Core/Private/GenericPlatform/GenericPlatformMallocCrash.cpp
--- a/Engine/Source/Runtime/Core/Private/GenericPlatform/GenericPlatformMallocCrash.cpp
+++ b/Engine/Source/Runtime/Core/Private/GenericPlatform/GenericPlatformMallocCrash.cpp
@@ -314,7 +314,7 @@ void FGenericPlatformMallocCrash::SetAsGMalloc()
 	FPlatformAtomics::InterlockedExchangePtr((void**)&UE::Private::GMalloc, this);
 	if (PLATFORM_USES_FIXED_GMalloc_CLASS && GFixedMallocLocationPtr)
 	{
-		*GFixedMallocLocationPtr = nullptr; // this disables any fast-path inline allocators
+		*GFixedMallocLocationPtr = this; // replace fixed allocator with crash allocator
 	}
 	CrashedThreadId = FPlatformTLS::GetCurrentThreadId();
 	GIsMallocCrashActive = true;
@@ -697,7 +697,7 @@ void FGenericStackBasedMallocCrash::SetAsGMalloc()
 {
 	if (PLATFORM_USES_FIXED_GMalloc_CLASS && GFixedMallocLocationPtr)
 	{
-		*GFixedMallocLocationPtr = nullptr; // this disables any fast-path inline allocators
+		*GFixedMallocLocationPtr = this; // replace fixed allocator with crash allocator
 	}
 	UE::Private::GMalloc = this;
 	GIsMallocCrashActive = true;
diff --git a/Engine/Source/Runtime/Core/Private/IOS/IOSMallocZone.cpp b/Engine/Source/Runtime/Core/Private/IOS/IOSMallocZone.cpp
--- a/Engine/Source/Runtime/Core/Private/IOS/IOSMallocZone.cpp
+++ b/Engine/Source/Runtime/Core/Private/IOS/IOSMallocZone.cpp
@@ -88,7 +88,7 @@ void FIOSMallocCrashHandler::Enable( FIOSCrashContext* Context, uint32 CrashedTh
 	OriginalHeap = UE::Private::GMalloc;
 	if (PLATFORM_USES_FIXED_GMalloc_CLASS && GFixedMallocLocationPtr)
 	{
-		*GFixedMallocLocationPtr = nullptr; // this disables any fast-path inline allocators
+		*GFixedMallocLocationPtr = this; // replace fixed allocator with crash allocator
 	}
 	UE::Private::GMalloc = this;
 }
diff --git a/Engine/Source/Runtime/Core/Private/Mac/MacMallocZone.cpp b/Engine/Source/Runtime/Core/Private/Mac/MacMallocZone.cpp
--- a/Engine/Source/Runtime/Core/Private/Mac/MacMallocZone.cpp
+++ b/Engine/Source/Runtime/Core/Private/Mac/MacMallocZone.cpp
@@ -88,7 +88,7 @@ void FMacMallocCrashHandler::Enable( FMacCrashContext* Context, uint32 CrashedTh
 	OriginalHeap = UE::Private::GMalloc;
 	if (PLATFORM_USES_FIXED_GMalloc_CLASS && GFixedMallocLocationPtr)
 	{
-		*GFixedMallocLocationPtr = nullptr; // this disables any fast-path inline allocators
+		*GFixedMallocLocationPtr = this; // replace fixed allocator with crash allocator
 	}
 	UE::Private::GMalloc = this;
 }
