[Core] Fix FName dump. If Name has symbols that look like format qualifiers it will read from past the buffer.

#rnx

diff --git a/Engine/Source/Runtime/Core/Private/UObject/UnrealNames.cpp b/Engine/Source/Runtime/Core/Private/UObject/UnrealNames.cpp
--- a/Engine/Source/Runtime/Core/Private/UObject/UnrealNames.cpp
+++ b/Engine/Source/Runtime/Core/Private/UObject/UnrealNames.cpp
@@ -4437,7 +4437,7 @@ void FNameEntry::DebugDump(FOutputDevice& Out) const
 	{
 		TCHAR Buffer[NAME_SIZE];
 		GetName(Buffer);
-		Out.Logf(Buffer);
+		Out.Log(Buffer);
 	}
 }
 
