- Fixed an uninitialized property in FArchiveState

diff --git a/Engine/Source/Runtime/Core/Private/Serialization/Archive.cpp b/Engine/Source/Runtime/Core/Private/Serialization/Archive.cpp
--- a/Engine/Source/Runtime/Core/Private/Serialization/Archive.cpp
+++ b/Engine/Source/Runtime/Core/Private/Serialization/Archive.cpp
@@ -189,6 +189,7 @@ PRAGMA_ENABLE_DEPRECATION_WARNINGS
 	ArShouldSkipCompilingAssets			= false;
 	ArMaxSerializeSize					= 0;
 	ArIsFilterEditorOnly				= false;
+	ArIsSavingOptionalObject			= false;
 	ArIsSaveGame						= false;
 	ArIsNetArchive						= false;
 	ArCustomPropertyList				= nullptr;
