Fix an issue where the trace auxiliary API in relay mode does not respect the trace options provided.

This happened because the Relay method was not passing the SendFlags to the internal write function. The consecuence is that tools that rely on options like disabling the tail buffer can't work properly in relay mode (like the Chaos Visual Debugger)

diff --git a/Engine/Source/Runtime/Core/Private/ProfilingDebugging/TraceAuxiliary.cpp b/Engine/Source/Runtime/Core/Private/ProfilingDebugging/TraceAuxiliary.cpp
--- a/Engine/Source/Runtime/Core/Private/ProfilingDebugging/TraceAuxiliary.cpp
+++ b/Engine/Source/Runtime/Core/Private/ProfilingDebugging/TraceAuxiliary.cpp
@@ -654,7 +654,7 @@ bool FTraceAuxiliaryImpl::ConnectRelay(UPTRINT Handle, UE::Trace::IoWriteFunc Wr
 		return true;
 	}
 
-	const bool bConnected = UE::Trace::RelayTo(Handle, WriteFunc, CloseFunc);
+	const bool bConnected = UE::Trace::RelayTo(Handle, WriteFunc, CloseFunc, SendFlags);
 	if (bConnected)
 	{
 		UE_LOG(LogTrace, Display, TEXT("Trace started (connected to relay endpoint)."));
