Honor `ShrinkByDefault` in various existing array classes.

This CL replaces `EAllowShrinking::Default` with `AllowShrinkingByDefault<Allocator>` in the containers which explicitly specify an allocator.

#FailedValidations Preflight

diff --git a/Engine/Plugins/Experimental/UAF/UAF/Source/UAF/Public/TransformArray.h b/Engine/Plugins/Experimental/UAF/UAF/Source/UAF/Public/TransformArray.h
--- a/Engine/Plugins/Experimental/UAF/UAF/Source/UAF/Public/TransformArray.h
+++ b/Engine/Plugins/Experimental/UAF/UAF/Source/UAF/Public/TransformArray.h
@@ -40,7 +40,7 @@ struct TTransformArrayAoS
 		Transforms.Reset(NumTransforms);
 	}
 
-	void SetNum(int32 NumTransforms, EAllowShrinking AllowShrinking = EAllowShrinking::Default)
+	void SetNum(int32 NumTransforms, EAllowShrinking AllowShrinking = UE::Core::Private::AllowShrinkingByDefault<AllocatorType>())
 	{
 		Transforms.SetNum(NumTransforms, AllowShrinking);
 	}
@@ -50,7 +50,7 @@ struct TTransformArrayAoS
 		SetNum(NumTransforms, bAllowShrinking ? EAllowShrinking::Yes : EAllowShrinking::No);
 	}
 
-	void SetNumUninitialized(int32 NumTransforms, EAllowShrinking AllowShrinking = EAllowShrinking::Default)
+	void SetNumUninitialized(int32 NumTransforms, EAllowShrinking AllowShrinking = UE::Core::Private::AllowShrinkingByDefault<AllocatorType>())
 	{
 		Transforms.SetNumUninitialized(NumTransforms, AllowShrinking);
 	}
@@ -69,7 +69,7 @@ struct TTransformArrayAoS
 	{
 		check(Index < Num() && Index < Other.Num());
 		check(Index + (NumTransforms - 1) < Num() && Index + (NumTransforms - 1) < Other.Num());
-		
+
 		std::copy(Other.Transforms.GetData() + Index, Other.Transforms.GetData() + (Index + NumTransforms), Transforms.GetData() + Index);
 	}
 
@@ -353,7 +353,7 @@ struct TTransformArraySoA
 		AllocatedMemory.Reset(NumTransforms * TransformSize + Padding);
 	}
 
-	void SetNum(int32 NumTransforms, EAllowShrinking AllowShrinking = EAllowShrinking::Default)
+	void SetNum(int32 NumTransforms, EAllowShrinking AllowShrinking = UE::Core::Private::AllowShrinkingByDefault<AllocatorType>())
 	{
 		const int32 Padding = NumTransforms != 0 ? PaddingSize : 0;
 		AllocatedMemory.SetNum(NumTransforms * TransformSize + Padding, AllowShrinking);
@@ -361,7 +361,7 @@ struct TTransformArraySoA
 		UpdateViews(AllocatedMemory.GetData(), NumTransforms);
 	}
 
-	void SetNumUninitialized(int32 NumTransforms, EAllowShrinking AllowShrinking = EAllowShrinking::Default)
+	void SetNumUninitialized(int32 NumTransforms, EAllowShrinking AllowShrinking = UE::Core::Private::AllowShrinkingByDefault<AllocatorType>())
 	{
 		const int32 Padding = NumTransforms != 0 ? PaddingSize : 0;
 		AllocatedMemory.SetNumUninitialized(NumTransforms * TransformSize + Padding, AllowShrinking);
diff --git a/Engine/Source/Runtime/Core/Public/Containers/IndirectArray.h b/Engine/Source/Runtime/Core/Public/Containers/IndirectArray.h
--- a/Engine/Source/Runtime/Core/Public/Containers/IndirectArray.h
+++ b/Engine/Source/Runtime/Core/Public/Containers/IndirectArray.h
@@ -76,7 +76,7 @@ public:
 	}
 
 	/**
-	 * Returns true if the array is empty and contains no elements. 
+	 * Returns true if the array is empty and contains no elements.
 	 *
 	 * @returns True if the array is empty.
 	 * @see Num
@@ -116,7 +116,7 @@ public:
 		return (const T**)Array.GetData();
 	}
 
-	/** 
+	/**
 	 * Helper function returning the size of the inner type.
 	 *
 	 * @returns Size in bytes of array type.
@@ -251,7 +251,7 @@ public:
 	 * @param AllowShrinking (Optional) Tells if this call can shrink array if
 	 *                        suitable after remove. Default is yes.
 	 */
-	void RemoveAt(int32 Index, EAllowShrinking AllowShrinking = EAllowShrinking::Default)
+	void RemoveAt(int32 Index, EAllowShrinking AllowShrinking = UE::Core::Private::AllowShrinkingByDefault<Allocator>())
 	{
 		check(Index >= 0);
 		check(Index < Array.Num());
@@ -269,7 +269,7 @@ public:
 	 * @param AllowShrinking (Optional) Tells if this call can shrink array if
 	 *                        suitable after remove. Default is yes.
 	 */
-	void RemoveAt(int32 Index, int32 Count, EAllowShrinking AllowShrinking = EAllowShrinking::Default)
+	void RemoveAt(int32 Index, int32 Count, EAllowShrinking AllowShrinking = UE::Core::Private::AllowShrinkingByDefault<Allocator>())
 	{
 		check(Index >= 0);
 		check(Index <= Array.Num());
@@ -299,7 +299,7 @@ public:
 	 * @param AllowShrinking (Optional) Tells if this call can shrink array if
 	 *                        suitable after remove. Default is yes.
 	 */
-	void RemoveAtSwap(int32 Index, EAllowShrinking AllowShrinking = EAllowShrinking::Default)
+	void RemoveAtSwap(int32 Index, EAllowShrinking AllowShrinking = UE::Core::Private::AllowShrinkingByDefault<Allocator>())
 	{
 		check(Index >= 0);
 		check(Index < Array.Num());
@@ -320,7 +320,7 @@ public:
 	 * @param AllowShrinking (Optional) Tells if this call can shrink array if
 	 *                        suitable after remove. Default is yes.
 	 */
-	void RemoveAtSwap(int32 Index, int32 Count, EAllowShrinking AllowShrinking = EAllowShrinking::Default)
+	void RemoveAtSwap(int32 Index, int32 Count, EAllowShrinking AllowShrinking = UE::Core::Private::AllowShrinkingByDefault<Allocator>())
 	{
 		check(Index >= 0);
 		check(Index <= Array.Num());
@@ -410,8 +410,8 @@ public:
 		return Array.IsValidIndex(Index);
 	}
 
-	/** 
-	 * Helper function to return the amount of memory allocated by this container 
+	/**
+	 * Helper function to return the amount of memory allocated by this container
 	 *
 	 * @returns Number of bytes allocated by this container.
 	 */
diff --git a/Engine/Source/Runtime/Core/Public/Containers/PagedArray.h b/Engine/Source/Runtime/Core/Public/Containers/PagedArray.h
--- a/Engine/Source/Runtime/Core/Public/Containers/PagedArray.h
+++ b/Engine/Source/Runtime/Core/Public/Containers/PagedArray.h
@@ -325,7 +325,7 @@ public:
 	 * Resizes array to the parameter number of elements.
 	 * The allow shrinking parameter indicates whether the container page allocation can be reduced if possible.
 	 */
-	void SetNum(SizeType NewNum, EAllowShrinking AllowShrinking = EAllowShrinking::Default)
+	void SetNum(SizeType NewNum, EAllowShrinking AllowShrinking = UE::Core::Private::AllowShrinkingByDefault<AllocatorType>())
 	{
 		const SizeType RequiredPageCount = NumRequiredPages(NewNum);
 		if (NewNum > Num())
@@ -527,7 +527,7 @@ public:
 	/*
 	 * Removes the last element in the container.
 	 */
-	void Pop(EAllowShrinking AllowShrinking = EAllowShrinking::Default)
+	void Pop(EAllowShrinking AllowShrinking = UE::Core::Private::AllowShrinkingByDefault<AllocatorType>())
 	{
 		CheckValidIndex(0);
 		const SizeType LastIndex = Num() - 1;
@@ -551,7 +551,7 @@ public:
 	 * to ensure the range is contiguous. This method provides efficient removal O(1) but it doesn't preserve the
 	 * insertion order.
 	 */
-	void RemoveAtSwap(SizeType Index, EAllowShrinking AllowShrinking = EAllowShrinking::Default)
+	void RemoveAtSwap(SizeType Index, EAllowShrinking AllowShrinking = UE::Core::Private::AllowShrinkingByDefault<AllocatorType>())
 	{
 		CheckValidIndex(Index);
 		const SizeType TargetPageIndex = GetPageIndex(Index);
diff --git a/Engine/Source/Runtime/Core/Public/Containers/ScriptArray.h b/Engine/Source/Runtime/Core/Public/Containers/ScriptArray.h
--- a/Engine/Source/Runtime/Core/Public/Containers/ScriptArray.h
+++ b/Engine/Source/Runtime/Core/Public/Containers/ScriptArray.h
@@ -110,7 +110,7 @@ public:
 			ResizeTo(ArrayNum, NumBytesPerElement, AlignmentOfElement);
 		}
 	}
-	void SetNumUninitialized(int32 NewNum, int32 NumBytesPerElement, uint32 AlignmentOfElement, EAllowShrinking AllowShrinking = EAllowShrinking::Default)
+	void SetNumUninitialized(int32 NewNum, int32 NumBytesPerElement, uint32 AlignmentOfElement, EAllowShrinking AllowShrinking = UE::Core::Private::AllowShrinkingByDefault<AllocatorType>())
 	{
 		checkSlow(NewNum >= 0);
 		int32 OldNum = Num();
@@ -193,7 +193,7 @@ public:
 		return ArrayMax - ArrayNum;
 	}
 
-	void Remove( int32 Index, int32 Count, int32 NumBytesPerElement, uint32 AlignmentOfElement, EAllowShrinking AllowShrinking = EAllowShrinking::Default )
+	void Remove(int32 Index, int32 Count, int32 NumBytesPerElement, uint32 AlignmentOfElement, EAllowShrinking AllowShrinking = UE::Core::Private::AllowShrinkingByDefault<AllocatorType>())
 	{
 		if (Count)
 		{
