Replace IS_MERGEDMODULES with UE_MERGED_MODULES

#rnx
#FailedValidations Preflight

diff --git a/Engine/Source/Programs/UnrealBuildTool/Configuration/UEBuildTarget.cs b/Engine/Source/Programs/UnrealBuildTool/Configuration/UEBuildTarget.cs
--- a/Engine/Source/Programs/UnrealBuildTool/Configuration/UEBuildTarget.cs
+++ b/Engine/Source/Programs/UnrealBuildTool/Configuration/UEBuildTarget.cs
@@ -5848,7 +5848,7 @@ namespace UnrealBuildTool
 			GlobalCompileEnvironment.Definitions.Add("__UNREAL__");
 
 			GlobalCompileEnvironment.Definitions.Add(String.Format("IS_MONOLITHIC={0}", ShouldCompileMonolithic() ? "1" : "0"));
-			GlobalCompileEnvironment.Definitions.Add(String.Format("IS_MERGEDMODULES={0}", Rules.bMergeModules ? "1" : "0"));
+			GlobalCompileEnvironment.Definitions.Add(String.Format("UE_MERGED_MODULES={0}", Rules.bMergeModules ? "1" : "0"));
 
 			GlobalCompileEnvironment.Definitions.Add(String.Format("WITH_ENGINE={0}", Rules.bCompileAgainstEngine ? "1" : "0"));
 			GlobalCompileEnvironment.Definitions.Add(String.Format("WITH_UNREAL_DEVELOPER_TOOLS={0}", Rules.bBuildDeveloperTools ? "1" : "0"));
diff --git a/Engine/Source/Runtime/Core/Core.Build.cs b/Engine/Source/Runtime/Core/Core.Build.cs
--- a/Engine/Source/Runtime/Core/Core.Build.cs
+++ b/Engine/Source/Runtime/Core/Core.Build.cs
@@ -307,7 +307,6 @@ public class Core : ModuleRules
 			PrivateDependencyModuleNames.Add("Detours");
 		}
 
-		PublicDefinitions.Add("UE_MERGED_MODULES=" + (Target.bMergeModules ? "1" : "0"));
 		PublicDefinitions.Add("UE_DELEGATE_CHECK_LIFETIME=" + (Target.bCheckDelegateLifetime ? "1" : "0"));
 	}
 
diff --git a/Engine/Source/Runtime/Core/Private/Containers/ContainerAllocationPolicies.cpp b/Engine/Source/Runtime/Core/Private/Containers/ContainerAllocationPolicies.cpp
--- a/Engine/Source/Runtime/Core/Private/Containers/ContainerAllocationPolicies.cpp
+++ b/Engine/Source/Runtime/Core/Private/Containers/ContainerAllocationPolicies.cpp
@@ -15,7 +15,7 @@ FORCENOINLINE void UE::Core::Private::OnInvalidSizedHeapAllocatorNum(int32 Index
 	for (;;);
 }
 
-#if !IS_MERGEDMODULES && !defined(UE_HEADER_UNITS)
+#if !UE_MERGED_MODULES && !defined(UE_HEADER_UNITS)
 template void TSizedHeapAllocator<32, FMemory>::ForAnyElementType::ResizeAllocation(SizeType CurrentNum, SizeType NewMax, SIZE_T NumBytesPerElement);
 template void TSizedHeapAllocator<32, FMemory>::ForAnyElementType::ResizeAllocation(SizeType CurrentNum, SizeType NewMax, SIZE_T NumBytesPerElement, uint32 AlignmentOfElement);
 template void TSizedHeapAllocator<64, FMemory>::ForAnyElementType::ResizeAllocation(SizeType CurrentNum, SizeType NewMax, SIZE_T NumBytesPerElement);
diff --git a/Engine/Source/Runtime/Core/Public/Containers/ContainerAllocationPolicies.h b/Engine/Source/Runtime/Core/Public/Containers/ContainerAllocationPolicies.h
--- a/Engine/Source/Runtime/Core/Public/Containers/ContainerAllocationPolicies.h
+++ b/Engine/Source/Runtime/Core/Public/Containers/ContainerAllocationPolicies.h
@@ -839,7 +839,7 @@ public:
 };
 
 // Define the ResizeAllocation functions with the regular allocator as exported to avoid bloat
-#if !IS_MERGEDMODULES && !defined(UE_HEADER_UNITS)
+#if !UE_MERGED_MODULES && !defined(UE_HEADER_UNITS)
 extern template CORE_API FORCENOINLINE void TSizedHeapAllocator<32, FMemory>::ForAnyElementType::ResizeAllocation(SizeType CurrentNum, SizeType NewMax, SIZE_T NumBytesPerElement);
 extern template CORE_API FORCENOINLINE void TSizedHeapAllocator<32, FMemory>::ForAnyElementType::ResizeAllocation(SizeType CurrentNum, SizeType NewMax, SIZE_T NumBytesPerElement, uint32 AlignmentOfElement);
 extern template CORE_API FORCENOINLINE void TSizedHeapAllocator<64, FMemory>::ForAnyElementType::ResizeAllocation(SizeType CurrentNum, SizeType NewMax, SIZE_T NumBytesPerElement);
diff --git a/Engine/Source/Runtime/Experimental/Chaos/Private/Chaos/BoundingVolume.cpp b/Engine/Source/Runtime/Experimental/Chaos/Private/Chaos/BoundingVolume.cpp
--- a/Engine/Source/Runtime/Experimental/Chaos/Private/Chaos/BoundingVolume.cpp
+++ b/Engine/Source/Runtime/Experimental/Chaos/Private/Chaos/BoundingVolume.cpp
@@ -26,7 +26,7 @@ namespace Chaos
 	    MaxDirtyElements,
 	    TEXT("The max number of dirty elements. This forces a flush which is very expensive"));
 
-#if !IS_MERGEDMODULES
+#if !UE_MERGED_MODULES
 	template class CHAOS_API Chaos::TBoundingVolume<int32, FReal, 3>;
 	template class CHAOS_API Chaos::TBoundingVolume<FAccelerationStructureHandle, FReal, 3>;
 #endif
diff --git a/Engine/Source/Runtime/Experimental/Chaos/Private/Chaos/SpatialAcceleration.cpp b/Engine/Source/Runtime/Experimental/Chaos/Private/Chaos/SpatialAcceleration.cpp
--- a/Engine/Source/Runtime/Experimental/Chaos/Private/Chaos/SpatialAcceleration.cpp
+++ b/Engine/Source/Runtime/Experimental/Chaos/Private/Chaos/SpatialAcceleration.cpp
@@ -29,7 +29,7 @@ namespace Chaos
 		}
 	}
 
-	#if !IS_MERGEDMODULES
+	#if !UE_MERGED_MODULES
 	template class CHAOS_API Chaos::ISpatialAcceleration<int32, FReal, 3>;
 	template class CHAOS_API Chaos::ISpatialAcceleration<FAccelerationStructureHandle, FReal, 3>;
 	template class CHAOS_API Chaos::ISpatialVisitor<int32, FReal>;
diff --git a/Engine/Source/Runtime/Experimental/Chaos/Public/Chaos/BoundingVolume.h b/Engine/Source/Runtime/Experimental/Chaos/Public/Chaos/BoundingVolume.h
--- a/Engine/Source/Runtime/Experimental/Chaos/Public/Chaos/BoundingVolume.h
+++ b/Engine/Source/Runtime/Experimental/Chaos/Public/Chaos/BoundingVolume.h
@@ -1381,7 +1381,7 @@ FArchive& operator<<(FChaosArchive& Ar, TBoundingVolume<TPayloadType, T, d>& Bou
 	return Ar;
 }
 
-#if !IS_MERGEDMODULES
+#if !UE_MERGED_MODULES
 #if PLATFORM_COMPILER_CLANG
 extern template class CHAOS_API Chaos::TBoundingVolume<int32, Chaos::FReal, 3>;
 extern template class CHAOS_API Chaos::TBoundingVolume<Chaos::FAccelerationStructureHandle, Chaos::FReal, 3>;
diff --git a/Engine/Source/Runtime/Experimental/Chaos/Public/Chaos/ISpatialAcceleration.h b/Engine/Source/Runtime/Experimental/Chaos/Public/Chaos/ISpatialAcceleration.h
--- a/Engine/Source/Runtime/Experimental/Chaos/Public/Chaos/ISpatialAcceleration.h
+++ b/Engine/Source/Runtime/Experimental/Chaos/Public/Chaos/ISpatialAcceleration.h
@@ -1030,7 +1030,7 @@ FORCEINLINE bool PrePreFilterHelper(const int32 Payload, const TVisitor& Visitor
 	return false;
 }
 
-#if !IS_MERGEDMODULES
+#if !UE_MERGED_MODULES
 #if PLATFORM_COMPILER_CLANG
 extern template class CHAOS_API ISpatialAcceleration<int32, FReal, 3>;
 extern template class CHAOS_API ISpatialVisitor<int32, FReal>;
diff --git a/Engine/Source/Runtime/Experimental/Chaos/Public/Chaos/ParticleHandle.h b/Engine/Source/Runtime/Experimental/Chaos/Public/Chaos/ParticleHandle.h
--- a/Engine/Source/Runtime/Experimental/Chaos/Public/Chaos/ParticleHandle.h
+++ b/Engine/Source/Runtime/Experimental/Chaos/Public/Chaos/ParticleHandle.h
@@ -4072,7 +4072,7 @@ inline void SetObjectStateHelper(IPhysicsProxyBase& Proxy, FPBDRigidParticle& Ri
 
 CHAOS_API void SetObjectStateHelper(IPhysicsProxyBase& Proxy, FPBDRigidParticleHandle& Rigid, EObjectStateType InState, bool bAllowEvents = false, bool bInvalidate = true);
 
-#if !IS_MERGEDMODULES
+#if !UE_MERGED_MODULES
 #if PLATFORM_COMPILER_CLANG
 extern template class CHAOS_API ISpatialAcceleration<FAccelerationStructureHandle, FReal, 3>;
 extern template class CHAOS_API ISpatialVisitor<FAccelerationStructureHandle, FReal>;
diff --git a/Engine/Source/Runtime/GeometryCore/Public/DynamicMesh/DynamicMeshOverlay.h b/Engine/Source/Runtime/GeometryCore/Public/DynamicMesh/DynamicMeshOverlay.h
--- a/Engine/Source/Runtime/GeometryCore/Public/DynamicMesh/DynamicMeshOverlay.h
+++ b/Engine/Source/Runtime/GeometryCore/Public/DynamicMesh/DynamicMeshOverlay.h
@@ -830,7 +830,7 @@ public:
 #define UE_EXTERN_TEMPLATE_API
 #endif
 
-#if !IS_MERGEDMODULES
+#if !UE_MERGED_MODULES
 extern template class UE_EXTERN_TEMPLATE_API TDynamicMeshOverlay<float, 1>;
 extern template class UE_EXTERN_TEMPLATE_API TDynamicMeshOverlay<double, 1>;
 extern template class UE_EXTERN_TEMPLATE_API TDynamicMeshOverlay<int, 1>;
diff --git a/Engine/Source/Runtime/Launch/Launch.Build.cs b/Engine/Source/Runtime/Launch/Launch.Build.cs
--- a/Engine/Source/Runtime/Launch/Launch.Build.cs
+++ b/Engine/Source/Runtime/Launch/Launch.Build.cs
@@ -278,10 +278,5 @@ public class Launch : ModuleRules
 			PrivateDefinitions.Add(string.Format("COMPILED_IN_COMPATIBLE_CL={0}", Target.Version.EffectiveCompatibleChangelist));
 			PrivateDefinitions.Add(string.Format("COMPILED_IN_BRANCH_NAME={0}", (Target.Version.BranchName == null || Target.Version.BranchName.Length == 0) ? "UE" : Target.Version.BranchName));
 		}
-
-		if (Target.bMergeModules)
-		{
-			PrivateDefinitions.Add("UE_MERGED_MODULES=1");
-		}
 	}
 }
