Only start TraceLog end frame update if there is no worker thread.

* Remove end frame pump when worker thread is started. With the introduction of block pool limits running both at the same time risks deadlocking the game thread in certain scenarios. Also with variable sleep the end frame pump can also "sneak" in when the worker thread is sleeping between updates.
* When using autoconnect, start worker thread so that end frame update does not run at the same time as worker thread startup.
* Don't auto connect on parent process when forking

diff --git a/Engine/Source/Runtime/Core/Private/ProfilingDebugging/TraceAuxiliary.cpp b/Engine/Source/Runtime/Core/Private/ProfilingDebugging/TraceAuxiliary.cpp
--- a/Engine/Source/Runtime/Core/Private/ProfilingDebugging/TraceAuxiliary.cpp
+++ b/Engine/Source/Runtime/Core/Private/ProfilingDebugging/TraceAuxiliary.cpp
@@ -211,6 +211,7 @@ public:
 	bool HasCommandlineChannels() const;
 	void SetTruncateFile(bool bTruncateFile);
 	void StartWorkerThread();
+	void StartEndFrameUpdate();
 	void RegisterEndFrameCallbacks();
 	bool WriteSnapshot(const TCHAR* InFilePath, const FTraceAuxiliary::FLogCategoryAlias& LogCategory);
 	bool SendSnapshot(const TCHAR* InHost, uint32 InPort, const FTraceAuxiliary::FLogCategoryAlias& LogCategory);
@@ -239,6 +240,7 @@ private:
 	bool WriteToFile(const TCHAR* Path, const FTraceAuxiliary::FLogCategoryAlias& LogCategory, uint16 SendFlags);
 	bool FinalizeFilePath(const TCHAR* InPath, FString& OutPath, const FTraceAuxiliary::FLogCategoryAlias& LogCategory);
 	void AddCommandlineChannel(const TCHAR* Name, const FTraceAuxiliary::FLogCategoryAlias& LogCategory);
+	static void FrameBasedUpdate();
 
 	template<typename ChannelType>
 	bool EnableChannel(ChannelType Channel, const FTraceAuxiliary::FLogCategoryAlias& LogCategory, FString* OutReason = nullptr);
@@ -251,6 +253,7 @@ private:
 	bool bTruncateFile = false;
 	bool bReadOnlyChannelsFrozen = false;
 	FString PausedPreset;
+	FDelegateHandle EndFrameUpdateHandle;
 
 	struct FCurrentTraceTarget
 	{
@@ -269,7 +272,6 @@ static const TCHAR* GTraceConfigSection = TEXT("Trace.Config");
 static UE::Trace::FInitializeDesc GInitializeDesc;
 static EWorkerThreadConfig GWorkerThreadConfig = EWorkerThreadConfig::Unset;
 static FDelegateHandle GEndFrameDelegateHandle;
-static FDelegateHandle GEndFrameStatDelegateHandle;
 static FDelegateHandle GOnPostForkHandle;
 // Whether to start tracing automatically at start or wait to initiate via Console Command.
 // This value can also be set by passing '-traceautostart=[0|1]' on command line.
@@ -894,6 +896,13 @@ bool FTraceAuxiliaryImpl::FinalizeFilePath(const TCHAR* InPath, FString& OutPath
 	return true;
 }
 
+////////////////////////////////////////////////////////////////////////////////
+void FTraceAuxiliaryImpl::FrameBasedUpdate()
+{
+	TRACE_CPUPROFILER_EVENT_SCOPE(FTraceAux_Trace_Update);
+	UE::Trace::Update();
+}
+
 ////////////////////////////////////////////////////////////////////////////////
 bool FTraceAuxiliaryImpl::WriteToFile(const TCHAR* Path, const FTraceAuxiliary::FLogCategoryAlias& LogCategory, uint16 SendFlags)
 {
@@ -1038,11 +1047,19 @@ void FTraceAuxiliaryImpl::StartWorkerThread()
 	if (!bWorkerThreadStarted && bShouldStart)
 	{
 		TRACE_CPUPROFILER_EVENT_SCOPE(FTraceAux_StartWorkerThread);
+		// Remove end frame delegate as we don't need it anymore
+		FCoreDelegates::OnEndFrame.Remove(EndFrameUpdateHandle);
 		UE::Trace::StartWorkerThread();
 		bWorkerThreadStarted = true;
 	}
 }
 
+////////////////////////////////////////////////////////////////////////////////
+void FTraceAuxiliaryImpl::StartEndFrameUpdate()
+{
+	EndFrameUpdateHandle = FCoreDelegates::OnEndFrame.AddStatic(&FTraceAuxiliaryImpl::FrameBasedUpdate);
+}
+
 ////////////////////////////////////////////////////////////////////////////////
 void TraceAuxiliaryEmitStats(const UE::Trace::FStatistics& Stats)
 {
@@ -1127,7 +1144,6 @@ void TraceAuxiliaryOnTraceUpdateCallback()
 void TraceAuxiliaryOnEndFrameCallback()
 {
 	// This is called at end of each frame.
-
 	TRACE_CPUPROFILER_EVENT_SCOPE(FTraceAux_OnEndFrame_Stats);
 
 	// Reset the trace update callback. From now on, we only trace stats once per frame.
@@ -1164,16 +1180,8 @@ void TraceAuxiliaryOnEndFrameCallback()
 ////////////////////////////////////////////////////////////////////////////////
 void FTraceAuxiliaryImpl::RegisterEndFrameCallbacks()
 {
-	// Always register end frame updates. This path is short circuited if a worker thread exists.
-	// If the worker thread is disabled, pump the update from end frame.
-	FCoreDelegates::OnEndFrame.AddLambda([]()
-		{
-			TRACE_CPUPROFILER_EVENT_SCOPE(FTraceAux_Trace_Update);
-			UE::Trace::Update();
-		});
-
 	// Update stats every frame.
-	FCoreDelegates::OnEndFrame.AddStatic(TraceAuxiliaryOnEndFrameCallback);
+	GEndFrameDelegateHandle = FCoreDelegates::OnEndFrame.AddStatic(TraceAuxiliaryOnEndFrameCallback);
 }
 
 ////////////////////////////////////////////////////////////////////////////////
@@ -2101,6 +2109,11 @@ void FTraceAuxiliary::Initialize(const TCHAR* CommandLine)
 		GWorkerThreadConfig = Desc.bUseWorkerThread ? EWorkerThreadConfig::OnInit : EWorkerThreadConfig::OnDemand;
 	}
 
+	if (!Desc.bUseWorkerThread)
+	{
+		GTraceAuxiliary.StartEndFrameUpdate();
+	}
+
 	FGuid SessionGuid;
 	if (!FParse::Value(CommandLine, TEXT("-tracesessionguid="), SessionGuid))
 	{
@@ -2395,7 +2408,7 @@ void FTraceAuxiliary::TryAutoConnect()
 #if UE_TRACE_ENABLED
 	TRACE_CPUPROFILER_EVENT_SCOPE(FTraceAux_TryAutoConnect);
 #if PLATFORM_WINDOWS
-	if (GTraceAutoStart && !IsConnected())
+	if (GTraceAutoStart && !IsConnected() && !GTraceAuxiliary.IsParentProcessAndPreFork())
 	{
 		// If we can detect a named event it means UnrealInsights (Browser Mode) is running.
 		// In this case, we try to auto-connect with the Trace Server.
@@ -2403,12 +2416,13 @@ void FTraceAuxiliary::TryAutoConnect()
 		if (KnownEvent != nullptr)
 		{
 			UE_LOG(LogTrace, Display, TEXT("Unreal Insights instance detected, auto-connecting to local trace server..."));
+			GTraceAuxiliary.StartWorkerThread();
 			Start(EConnectionType::Network, TEXT("127.0.0.1"), GTraceAuxiliary.HasCommandlineChannels() ? nullptr : TEXT("default"), nullptr);
 			::CloseHandle(KnownEvent);
 		}
 	}
 #elif PLATFORM_MAC || PLATFORM_LINUX
-	if (GTraceAutoStart && !IsConnected())
+	if (GTraceAutoStart && !IsConnected() && !GTraceAuxiliary.IsParentProcessAndPreFork())
 	{
 		sem_t* AutoConnectSemaphore = sem_open("/UnrealInsightsAutoConnect", O_RDONLY);
 		if (AutoConnectSemaphore != SEM_FAILED)
