Fix misplaced parantheses detected by unreachable code warning

#rnx

diff --git a/Engine/Source/Runtime/Core/Tests/Misc/PathsTest.cpp b/Engine/Source/Runtime/Core/Tests/Misc/PathsTest.cpp
--- a/Engine/Source/Runtime/Core/Tests/Misc/PathsTest.cpp
+++ b/Engine/Source/Runtime/Core/Tests/Misc/PathsTest.cpp
@@ -242,9 +242,9 @@ void TestCreateStandardFilename()
 	FStringView StandardPathOfBaseParent = FPathViews::GetPath(StandardPathOfBase);
 	FStringView StandardPathOfBaseParent2 = FPathViews::GetPath(StandardPathOfBaseParent);
 	FStringView StandardPathOfBaseParent3 = FPathViews::GetPath(StandardPathOfBaseParent2);
-	if (StandardPathOfBaseParent.EndsWith(TEXT("..") || StandardPathOfBaseParent.IsEmpty())
-		|| StandardPathOfBaseParent2.EndsWith(TEXT("..") || StandardPathOfBaseParent2.IsEmpty())
-		|| StandardPathOfBaseParent3.EndsWith(TEXT("..") || StandardPathOfBaseParent3.IsEmpty()))
+	if (StandardPathOfBaseParent.EndsWith(TEXT("..")) || StandardPathOfBaseParent.IsEmpty()
+		|| StandardPathOfBaseParent2.EndsWith(TEXT("..")) || StandardPathOfBaseParent2.IsEmpty()
+		|| StandardPathOfBaseParent3.EndsWith(TEXT("..")) || StandardPathOfBaseParent3.IsEmpty())
 	{
 		// We expect StandardPathOfBase == ../../../Engine/Binaries/Win64
 		// We want to be able to use this to know that ".." -> "../../../Engine/Binaries" and
