Fix doc-comments related to AllowShrinking parameter.

BugHawk noticed that I had changed the default value of this parameter in https://p4-swarm.epicgames.net/changes/46229319, but had not updated the doc-comment to match. The description also varied between APIs so I standardized on one phrasing across the various doc-comments.

diff --git a/Engine/Source/Runtime/Core/Public/Containers/Array.h b/Engine/Source/Runtime/Core/Public/Containers/Array.h
--- a/Engine/Source/Runtime/Core/Public/Containers/Array.h
+++ b/Engine/Source/Runtime/Core/Public/Containers/Array.h
@@ -1189,7 +1189,8 @@ public:
 	/**
 	 * Pops element from the array.
 	 *
-	 * @param AllowShrinking If this call allows shrinking of the array during element remove.
+	 * @param AllowShrinking (Optional) By default, arrays with large amounts of slack will automatically shrink.
+	 *                       Use FNonshrinkingAllocator or pass EAllowShrinking::No to prevent this.
 	 * @returns Popped element.
 	 */
 	ElementType Pop(EAllowShrinking AllowShrinking = UE::Core::Private::AllowShrinkingByDefault<AllocatorType>())
@@ -2077,7 +2078,8 @@ public:
 	 * the array.
 	 *
 	 * @param Index Location in array of the element to remove.
-	 * @param AllowShrinking (Optional) Tells if this call can shrink array if suitable after remove. Default is yes.
+	 * @param AllowShrinking (Optional) By default, arrays with large amounts of slack will automatically shrink.
+	 *                       Use FNonshrinkingAllocator or pass EAllowShrinking::No to prevent this.
 	 */
 	void RemoveAt(SizeType Index, EAllowShrinking AllowShrinking = UE::Core::Private::AllowShrinkingByDefault<AllocatorType>())
 	{
@@ -2101,7 +2103,8 @@ public:
 	 *
 	 * @param Index Location in array of the element to remove.
 	 * @param Count (Optional) Number of elements to remove. Default is 1.
-	 * @param AllowShrinking (Optional) Tells if this call can shrink array if suitable after remove. Default is yes.
+	 * @param AllowShrinking (Optional) By default, arrays with large amounts of slack will automatically shrink.
+	 *                       Use FNonshrinkingAllocator or pass EAllowShrinking::No to prevent this.
 	 */
 	template <UE::CIntegral CountType>
 	UE_FORCEINLINE_HINT void RemoveAt(SizeType Index, CountType Count, EAllowShrinking AllowShrinking = UE::Core::Private::AllowShrinkingByDefault<AllocatorType>())
@@ -2178,8 +2181,8 @@ public:
 	 * O(ArrayNum)), but does not preserve the order.
 	 *
 	 * @param Index Location in array of the element to remove.
-	 * @param AllowShrinking (Optional) Tells if this call can shrink array if
-	 *                        suitable after remove. Default is yes.
+	 * @param AllowShrinking (Optional) By default, arrays with large amounts of slack will automatically shrink.
+	 *                       Use FNonshrinkingAllocator or pass EAllowShrinking::No to prevent this.
 	 */
 	UE_FORCEINLINE_HINT void RemoveAtSwap(SizeType Index, EAllowShrinking AllowShrinking = UE::Core::Private::AllowShrinkingByDefault<AllocatorType>())
 	{
@@ -2206,8 +2209,8 @@ public:
 	 *
 	 * @param Index Location in array of the element to remove.
 	 * @param Count (Optional) Number of elements to remove. Default is 1.
-	 * @param AllowShrinking (Optional) Tells if this call can shrink array if
-	 *                        suitable after remove. Default is yes.
+	 * @param AllowShrinking (Optional) By default, arrays with large amounts of slack will automatically shrink.
+	 *                       Use FNonshrinkingAllocator or pass EAllowShrinking::No to prevent this.
 	 */
 	template <UE::CIntegral CountType>
 	UE_FORCEINLINE_HINT void RemoveAtSwap(SizeType Index, CountType Count, EAllowShrinking AllowShrinking = UE::Core::Private::AllowShrinkingByDefault<AllocatorType>())
@@ -2301,7 +2304,8 @@ public:
 	 * Resizes array to given number of elements.
 	 *
 	 * @param NewNum New size of the array.
-	 * @param AllowShrinking Tell if this function can shrink the memory in-use if suitable.
+	 * @param AllowShrinking (Optional) By default, arrays with large amounts of slack will automatically shrink.
+	 *                       Use FNonshrinkingAllocator or pass EAllowShrinking::No to prevent this.
 	 */
 	void SetNum(SizeType NewNum, EAllowShrinking AllowShrinking = UE::Core::Private::AllowShrinkingByDefault<AllocatorType>())
 	{
@@ -2332,7 +2336,8 @@ public:
 	 * New elements will be zeroed.
 	 *
 	 * @param NewNum New size of the array.
-	 * @param AllowShrinking Tell if this function can shrink the memory in-use if suitable.
+	 * @param AllowShrinking (Optional) By default, arrays with large amounts of slack will automatically shrink.
+	 *                       Use FNonshrinkingAllocator or pass EAllowShrinking::No to prevent this.
 	 */
 	void SetNumZeroed(SizeType NewNum, EAllowShrinking AllowShrinking = UE::Core::Private::AllowShrinkingByDefault<AllocatorType>())
 	{
@@ -2360,7 +2365,8 @@ public:
 	 * Resizes array to given number of elements. New elements will be uninitialized.
 	 *
 	 * @param NewNum New size of the array.
-	 * @param AllowShrinking Tell if this function can shrink the memory in-use if suitable.
+	 * @param AllowShrinking (Optional) By default, arrays with large amounts of slack will automatically shrink.
+	 *                       Use FNonshrinkingAllocator or pass EAllowShrinking::No to prevent this.
 	 */
 	void SetNumUninitialized(SizeType NewNum, EAllowShrinking AllowShrinking = UE::Core::Private::AllowShrinkingByDefault<AllocatorType>())
 	{
@@ -3151,7 +3157,8 @@ public:
 	 * Remove all instances that match the predicate
 	 *
 	 * @param Predicate Predicate class instance
-	 * @param AllowShrinking Tell if this function can shrink the memory in-use if suitable.
+	 * @param AllowShrinking (Optional) By default, arrays with large amounts of slack will automatically shrink.
+	 *                       Use FNonshrinkingAllocator or pass EAllowShrinking::No to prevent this.
 	 * @see Remove, RemoveSingle, RemoveSingleSwap, RemoveSwap
 	 */
 	template <class PREDICATE_CLASS>
@@ -3197,7 +3204,8 @@ public:
 	 * O(Count) instead of O(ArrayNum), but does not preserve the order
 	 *
 	 * @param Item The item to remove
-	 * @param AllowShrinking Tell if this function can shrink the memory in-use if suitable.
+	 * @param AllowShrinking (Optional) By default, arrays with large amounts of slack will automatically shrink.
+	 *                       Use FNonshrinkingAllocator or pass EAllowShrinking::No to prevent this.
 	 *
 	 * @returns The number of items removed. For RemoveSingleItem, this is always either 0 or 1.
 	 * @see Add, Insert, Remove, RemoveAll, RemoveAllSwap, RemoveSwap
@@ -3229,7 +3237,8 @@ public:
 	 * but does not preserve the order.
 	 *
 	 * @param Item The item to remove
-	 * @param AllowShrinking Tell if this function can shrink the memory in-use if suitable.
+	 * @param AllowShrinking (Optional) By default, arrays with large amounts of slack will automatically shrink.
+	 *                       Use FNonshrinkingAllocator or pass EAllowShrinking::No to prevent this.
 	 *
 	 * @returns Number of elements removed.
 	 * @see Add, Insert, Remove, RemoveAll, RemoveAllSwap
@@ -3730,7 +3739,8 @@ public:
 	 *
 	 * @param OutItem The removed item.
 	 * @param Predicate Predicate class instance.
-	 * @param AllowShrinking Tell if this function can shrink the memory in-use if suitable.
+	 * @param AllowShrinking (Optional) By default, arrays with large amounts of slack will automatically shrink.
+	 *                       Use FNonshrinkingAllocator or pass EAllowShrinking::No to prevent this.
 	 *
 	 * @note: If your array contains raw pointers, they will be automatically dereferenced during heapification.
 	 *        Therefore, your predicate will be passed references rather than pointers.
@@ -3757,7 +3767,8 @@ public:
 	 * the template type.
 	 *
 	 * @param OutItem The removed item.
-	 * @param AllowShrinking (Optional) Tells if this call can shrink the array allocation if suitable after the pop. Default is yes.
+	 * @param AllowShrinking (Optional) By default, arrays with large amounts of slack will automatically shrink.
+	 *                       Use FNonshrinkingAllocator or pass EAllowShrinking::No to prevent this.
 	 *
 	 * @note: If your array contains raw pointers, they will be automatically dereferenced during heapification.
 	 *        Therefore, your array will be heapified by the values being pointed to, rather than the pointers' values.
@@ -3788,7 +3799,8 @@ public:
 	 * Removes the top element from the heap.
 	 *
 	 * @param Predicate Predicate class instance.
-	 * @param AllowShrinking (Optional) Tells if this call can shrink the array allocation if suitable after the discard. Default is yes.
+	 * @param AllowShrinking (Optional) By default, arrays with large amounts of slack will automatically shrink.
+	 *                       Use FNonshrinkingAllocator or pass EAllowShrinking::No to prevent this.
 	 *
 	 * @note: If your array contains raw pointers, they will be automatically dereferenced during heapification.
 	 *        Therefore, your predicate will be passed references rather than pointers.
@@ -3811,8 +3823,8 @@ public:
 	/**
 	 * Removes the top element from the heap. Assumes < operator is defined for the template type.
 	 *
-	 * @param AllowShrinking (Optional) Tells if this call can shrink the array
-	 *		allocation if suitable after the discard. Default is yes.
+	 * @param AllowShrinking (Optional) By default, arrays with large amounts of slack will automatically shrink.
+	 *                       Use FNonshrinkingAllocator or pass EAllowShrinking::No to prevent this.
 	 *
 	 * @note: If your array contains raw pointers, they will be automatically dereferenced during heapification.
 	 *        Therefore, your array will be heapified by the values being pointed to, rather than the pointers' values.
@@ -3855,8 +3867,8 @@ public:
 	 *
 	 * @param Index Position at which to remove item.
 	 * @param Predicate Predicate class instance.
-	 * @param AllowShrinking (Optional) Tells if this call can shrink the array allocation
-	 *		if suitable after the remove (default = yes).
+	 * @param AllowShrinking (Optional) By default, arrays with large amounts of slack will automatically shrink.
+	 *                       Use FNonshrinkingAllocator or pass EAllowShrinking::No to prevent this.
 	 *
 	 * @note: If your array contains raw pointers, they will be automatically dereferenced during heapification.
 	 *        Therefore, your predicate will be passed references rather than pointers.
@@ -3882,8 +3894,8 @@ public:
 	 * Removes an element from the heap. Assumes < operator is defined for the template type.
 	 *
 	 * @param Index Position at which to remove item.
-	 * @param AllowShrinking (Optional) Tells if this call can shrink the array allocation
-	 *		if suitable after the remove (default = yes).
+	 * @param AllowShrinking (Optional) By default, arrays with large amounts of slack will automatically shrink.
+	 *                       Use FNonshrinkingAllocator or pass EAllowShrinking::No to prevent this.
 	 *
 	 * @note: If your array contains raw pointers, they will be automatically dereferenced during heapification.
 	 *        Therefore, your array will be heapified by the values being pointed to, rather than the pointers' values.
diff --git a/Engine/Source/Runtime/Core/Public/Containers/IndirectArray.h b/Engine/Source/Runtime/Core/Public/Containers/IndirectArray.h
--- a/Engine/Source/Runtime/Core/Public/Containers/IndirectArray.h
+++ b/Engine/Source/Runtime/Core/Public/Containers/IndirectArray.h
@@ -248,8 +248,8 @@ public:
 	 * the array.
 	 *
 	 * @param Index Location in array of the element to remove.
-	 * @param AllowShrinking (Optional) Tells if this call can shrink array if
-	 *                        suitable after remove. Default is yes.
+	 * @param AllowShrinking (Optional) By default, arrays with large amounts of slack will automatically shrink.
+	 *                       Use FNonshrinkingAllocator or pass EAllowShrinking::No to prevent this.
 	 */
 	void RemoveAt(int32 Index, EAllowShrinking AllowShrinking = UE::Core::Private::AllowShrinkingByDefault<Allocator>())
 	{
@@ -266,8 +266,8 @@ public:
 	 *
 	 * @param Index Location in array of the first element to remove.
 	 * @param Count Number of elements to remove.
-	 * @param AllowShrinking (Optional) Tells if this call can shrink array if
-	 *                        suitable after remove. Default is yes.
+	 * @param AllowShrinking (Optional) By default, arrays with large amounts of slack will automatically shrink.
+	 *                       Use FNonshrinkingAllocator or pass EAllowShrinking::No to prevent this.
 	 */
 	void RemoveAt(int32 Index, int32 Count, EAllowShrinking AllowShrinking = UE::Core::Private::AllowShrinkingByDefault<Allocator>())
 	{
@@ -296,8 +296,8 @@ public:
 	 * O(ArrayNum)), but does not preserve the order.
 	 *
 	 * @param Index Location in array of the element to remove.
-	 * @param AllowShrinking (Optional) Tells if this call can shrink array if
-	 *                        suitable after remove. Default is yes.
+	 * @param AllowShrinking (Optional) By default, arrays with large amounts of slack will automatically shrink.
+	 *                       Use FNonshrinkingAllocator or pass EAllowShrinking::No to prevent this.
 	 */
 	void RemoveAtSwap(int32 Index, EAllowShrinking AllowShrinking = UE::Core::Private::AllowShrinkingByDefault<Allocator>())
 	{
@@ -317,8 +317,8 @@ public:
 	 *
 	 * @param Index Location in array of the rirst element to remove.
 	 * @param Count Number of elements to remove.
-	 * @param AllowShrinking (Optional) Tells if this call can shrink array if
-	 *                        suitable after remove. Default is yes.
+	 * @param AllowShrinking (Optional) By default, arrays with large amounts of slack will automatically shrink.
+	 *                       Use FNonshrinkingAllocator or pass EAllowShrinking::No to prevent this.
 	 */
 	void RemoveAtSwap(int32 Index, int32 Count, EAllowShrinking AllowShrinking = UE::Core::Private::AllowShrinkingByDefault<Allocator>())
 	{
