Fix more BugHawk-discovered comment inaccuracies.

The doc-comments here had further subtle issues that BugHawk failed to mention the first time around.

diff --git a/Engine/Source/Runtime/Core/Public/Containers/Array.h b/Engine/Source/Runtime/Core/Public/Containers/Array.h
--- a/Engine/Source/Runtime/Core/Public/Containers/Array.h
+++ b/Engine/Source/Runtime/Core/Public/Containers/Array.h
@@ -2074,8 +2074,7 @@ private:
 
 public:
 	/**
-	 * Removes an element (or elements) at given location, then shrinks
-	 * the array.
+	 * Removes an element at the given location, optionally shrinking the array.
 	 *
 	 * @param Index Location in array of the element to remove.
 	 * @param AllowShrinking (Optional) By default, arrays with large amounts of slack will automatically shrink.
@@ -2174,8 +2173,7 @@ private:
 
 public:
 	/**
-	 * Removes an element (or elements) at given location, then shrinks
-	 * the array.
+	 * Removes an element at the given location, optionally shrinking the array.
 	 *
 	 * This version is much more efficient than RemoveAt (O(Count) instead of
 	 * O(ArrayNum)), but does not preserve the order.
