Add safety switches to merged library unloading

In case of emergency situations, it might be useful to be able to remotely stop unmounting merged libraries. This could be required to prevent crashes or broken gameplay.

- Add Modules.EnableMergedLibraryUnloading to completely disable unloading across the board
- Add Modules.PersistentMergedLibraries to make some specific libraries persistent

#rnx

diff --git a/Engine/Source/Runtime/Core/Private/Modules/ModuleManager.cpp b/Engine/Source/Runtime/Core/Private/Modules/ModuleManager.cpp
--- a/Engine/Source/Runtime/Core/Private/Modules/ModuleManager.cpp
+++ b/Engine/Source/Runtime/Core/Private/Modules/ModuleManager.cpp
@@ -31,7 +31,25 @@ DEFINE_LOG_CATEGORY_STATIC(LogModuleManager, Log, All);
 	CORE_API bool GIsHotReload = false;
 #endif // WITH_HOT_RELOAD
 
-#if UE_MERGED_MODULES && !UE_BUILD_SHIPPING
+#if UE_MERGED_MODULES
+
+	// Toggle switch for merged modules unloading
+	static bool GEnableMergedLibraryUnloading = true;
+	static FAutoConsoleVariableRef EnableMergedLibraryUnloadingCVar(
+		TEXT("Modules.EnableMergedLibraryUnloading"),
+		GEnableMergedLibraryUnloading,
+		TEXT("When set, enable the unloading of a merged library when none of its modules are loaded."),
+		ECVF_Default);
+
+	// Off switch for unloading specific merged libraries
+	static FString GPersistentMergedLibraries = "";
+	static FAutoConsoleVariableRef PersistentMergedLibrariesCVar(
+		TEXT("Modules.PersistentMergedLibraries"),
+		GPersistentMergedLibraries,
+		TEXT("List of comma-separated patterns of merged library names to never unload. The search rule matches any name containing one entry."),
+		ECVF_Default);
+
+#if !UE_BUILD_SHIPPING
 
 	// Log user modules of merged libraries
 	FAutoConsoleCommand LogMergedLibraryUsageCommand(
@@ -66,8 +84,8 @@ DEFINE_LOG_CATEGORY_STATIC(LogModuleManager, Log, All);
 			}
 		}));
 
-
-#endif // UE_MERGED_MODULES && !UE_BUILD_SHIPPING
+#endif // !UE_BUILD_SHIPPING
+#endif // UE_MERGED_MODULES
 
 #if WITH_ENGINE
 	TMap<UClass*, UClass*>& GetClassesToReinstanceForHotReload()
@@ -617,7 +635,32 @@ void FModuleManager::InternalFreeLibrary(FName ModuleName, void* Handle)
 		UE_LOG(LogModuleManager, Fatal, TEXT("InternalFreeLibrary: no module name or valid handle was passed"))
 	}
 
-	if (ModuleCount == 0)
+	// Safety feature: allow disabling of merged library unloading, either wholly, or for specific libraries
+	const auto CanUnloadLibrary = [&LibraryToUnload]()
+	{
+		if (!GEnableMergedLibraryUnloading)
+		{
+			return false;
+		}
+
+		if (GPersistentMergedLibraries.Len())
+		{
+			TArray<FString> PersistentLibraries;
+			GPersistentMergedLibraries.ParseIntoArray(PersistentLibraries, TEXT(","));
+
+			for (const FString& PersistentLibrary : PersistentLibraries)
+			{
+				if (LibraryToUnload.Contains(PersistentLibrary))
+				{
+					return false;
+				}
+			}
+		}
+
+		return true;
+	};
+
+	if (ModuleCount == 0 && CanUnloadLibrary())
 	{
 		UE_LOG(LogModuleManager, Log, TEXT("InternalFreeLibrary: preparing unload for library '%s'"), *LibraryToUnload);
 
