Add enum class flag functionality to IPlatformFile::EOpenWriteFlags to make them easier to use.

#rnx

diff --git a/Engine/Source/Runtime/Core/Public/GenericPlatform/GenericPlatformFile.h b/Engine/Source/Runtime/Core/Public/GenericPlatform/GenericPlatformFile.h
--- a/Engine/Source/Runtime/Core/Public/GenericPlatform/GenericPlatformFile.h
+++ b/Engine/Source/Runtime/Core/Public/GenericPlatform/GenericPlatformFile.h
@@ -1082,4 +1082,4 @@ inline FFileJournalData::FFileJournalData()
 }
 
 ENUM_CLASS_FLAGS(IPlatformFile::EOpenReadFlags);
-
+ENUM_CLASS_FLAGS(IPlatformFile::EOpenWriteFlags);
