Fixed missing CORE_API from previous CL where member function was not visible from callers.

diff --git a/Engine/Source/Runtime/Core/Public/Apple/ApplePlatformMisc.h b/Engine/Source/Runtime/Core/Public/Apple/ApplePlatformMisc.h
--- a/Engine/Source/Runtime/Core/Public/Apple/ApplePlatformMisc.h
+++ b/Engine/Source/Runtime/Core/Public/Apple/ApplePlatformMisc.h
@@ -102,5 +102,5 @@ struct CORE_API FApplePlatformMisc : public FGenericPlatformMisc
 	static void* CreateAutoreleasePool();
 	static void ReleaseAutoreleasePool(void *Pool);
 	
-	static void SetEnvironmentVar(const TCHAR* InVariableName, const TCHAR* Value);
+	static CORE_API void SetEnvironmentVar(const TCHAR* InVariableName, const TCHAR* Value);
 };
