[Linux]Support using -basedir to override FLinuxPlatformProcess::BaseDir

[RN] [Linux] support overriding BaseDir with -basedir like other platforms
#tests Tested with and without override on Linux using quoted and unquoted paths

diff --git a/Engine/Source/Runtime/Core/Private/Linux/LinuxPlatformProcess.cpp b/Engine/Source/Runtime/Core/Private/Linux/LinuxPlatformProcess.cpp
--- a/Engine/Source/Runtime/Core/Private/Linux/LinuxPlatformProcess.cpp
+++ b/Engine/Source/Runtime/Core/Private/Linux/LinuxPlatformProcess.cpp
@@ -6,6 +6,7 @@
 #include "CoreGlobals.h"
 #include "Logging/LogMacros.h"
 #include "Misc/Paths.h"
+#include "Misc/Parse.h"
 #include "HAL/FileManager.h"
 
 const TCHAR* FLinuxPlatformProcess::BaseDir()
@@ -15,18 +16,56 @@ const TCHAR* FLinuxPlatformProcess::BaseDir()
 	if (!bHaveResult)
 	{
 		char SelfPath[UNIX_MAX_PATH] = {0};
-		if (readlink( "/proc/self/exe", SelfPath, UE_ARRAY_COUNT(SelfPath) - 1) == -1)
+
+		FILE* CmdLineFile = fopen("/proc/self/cmdline", "r");
+
+		if (!CmdLineFile)
 		{
 			int ErrNo = errno;
-			UE_LOG(LogHAL, Fatal, TEXT("readlink() failed with errno = %d (%s)"), ErrNo,
+			UE_LOG(LogHAL, Fatal, TEXT("fopen() for cmdline failed with errno = %d (%s)"), ErrNo,
 				StringCast< TCHAR >(strerror(ErrNo)).Get());
 			// unreachable
 			return CachedResult;
 		}
-		SelfPath[UE_ARRAY_COUNT(SelfPath) - 1] = 0;
 
-		FCString::Strncpy(CachedResult, UTF8_TO_TCHAR(dirname(SelfPath)), UE_ARRAY_COUNT(CachedResult));
-		FCString::StrncatTruncateDest(CachedResult, UE_ARRAY_COUNT(CachedResult), TEXT("/"));
+		// can't lseek or stat /proc/self/cmdline, so we just read in and check each argument
+		char* Arg = nullptr;
+		size_t Size = 0;
+		const char BaseDirCmdArg[] = "-basedir=";
+		while(getdelim(&Arg, &Size, 0, CmdLineFile) != -1)
+		{
+			int Length = strlen(Arg);
+			if (FCStringAnsi::Strnicmp(Arg, BaseDirCmdArg, sizeof(BaseDirCmdArg)-1) == 0)
+			{
+				char* Value = Arg + sizeof(BaseDirCmdArg)-1;
+				Length -= sizeof(BaseDirCmdArg)-1;
+
+				FCString::Strncpy(CachedResult, UTF8_TO_TCHAR(Value), UE_ARRAY_COUNT(CachedResult));
+				FCString::StrncatTruncateDest(CachedResult, UE_ARRAY_COUNT(CachedResult), TEXT("/"));
+				break;
+			}
+		}
+		free(Arg);
+		Arg = nullptr;
+		Size = 0;
+		fclose(CmdLineFile);
+
+		// if not filled in by command line, base it off of the executable location
+		if (*CachedResult == '\0')
+		{
+			if (readlink( "/proc/self/exe", SelfPath, UE_ARRAY_COUNT(SelfPath) - 1) == -1)
+			{
+				int ErrNo = errno;
+				UE_LOG(LogHAL, Fatal, TEXT("readlink() failed with errno = %d (%s)"), ErrNo,
+					StringCast< TCHAR >(strerror(ErrNo)).Get());
+				// unreachable
+				return CachedResult;
+			}
+			SelfPath[UE_ARRAY_COUNT(SelfPath) - 1] = 0;
+
+			FCString::Strncpy(CachedResult, UTF8_TO_TCHAR(dirname(SelfPath)), UE_ARRAY_COUNT(CachedResult));
+			FCString::StrncatTruncateDest(CachedResult, UE_ARRAY_COUNT(CachedResult), TEXT("/"));
+		}
 
 #ifdef UE_RELATIVE_BASE_DIR
 		FString CollapseResult(CachedResult);
