Fix invalid iPad Sharesheet popover location being used by dispatch. And take into account the scale factor.

diff --git a/Engine/Source/Runtime/Core/Private/IOS/IOSPlatformMisc.cpp b/Engine/Source/Runtime/Core/Private/IOS/IOSPlatformMisc.cpp
--- a/Engine/Source/Runtime/Core/Private/IOS/IOSPlatformMisc.cpp
+++ b/Engine/Source/Runtime/Core/Private/IOS/IOSPlatformMisc.cpp
@@ -1324,7 +1324,9 @@ void FIOSPlatformMisc::ShareURL(const FString& URL, const FText& Description, in
 {
 	NSString* SharedString = [NSString stringWithFString:Description.ToString()];
 	NSURL* SharedURL = [NSURL URLWithString:[NSString stringWithFString:URL]];
-	CGRect PopoverLocation = CGRectMake(LocationHintX, LocationHintY, 1, 1);
+    
+    float ScaleFactor = (float)[IOSAppDelegate GetDelegate].IOSView.contentScaleFactor;
+    __block CGRect PopoverLocation = CGRectMake(LocationHintX/ScaleFactor, LocationHintY/ScaleFactor, 10, 10);
 
 	dispatch_async(dispatch_get_main_queue(),^ {
 		NSArray* ObjectsToShare = @[SharedString, SharedURL];
