Collect driver information for Qualcomm GPUs and allow the use of the driver denylist.

#tests tested by QA
#FailedValidations Preflight

diff --git a/Engine/Source/Runtime/Core/Private/Windows/WindowsPlatformMisc.cpp b/Engine/Source/Runtime/Core/Private/Windows/WindowsPlatformMisc.cpp
--- a/Engine/Source/Runtime/Core/Private/Windows/WindowsPlatformMisc.cpp
+++ b/Engine/Source/Runtime/Core/Private/Windows/WindowsPlatformMisc.cpp
@@ -3710,6 +3710,10 @@ static void GetVideoDriverDetailsFromSetup(const FString& DeviceName, bool bVerb
 		{
 			Out.SetIntel();
 		}
+		else if (Out.ProviderName.Contains(TEXT("Qualcomm")))
+		{
+			Out.SetQualcomm();
+		}
 	}
 
 	Out.UserDriverVersion = Out.InternalDriverVersion;
@@ -3795,6 +3799,10 @@ static void GetVideoDriverDetails(const FString& Key, FGPUDriverInfo& Out)
 		{
 			Out.SetIntel();
 		}
+		else if (Out.ProviderName.Find(TEXT("Qualcomm")) != INDEX_NONE)
+		{
+			Out.SetQualcomm();
+		}
 	}
 
 	// technical driver version, AMD and NVIDIA
diff --git a/Engine/Source/Runtime/Core/Public/GenericPlatform/GenericPlatformDriver.h b/Engine/Source/Runtime/Core/Public/GenericPlatform/GenericPlatformDriver.h
--- a/Engine/Source/Runtime/Core/Public/GenericPlatform/GenericPlatformDriver.h
+++ b/Engine/Source/Runtime/Core/Public/GenericPlatform/GenericPlatformDriver.h
@@ -186,10 +186,17 @@ struct FGPUDriverInfo
 	void SetAMD() { VendorId = 0x1002; }
 	void SetIntel() { VendorId = 0x8086; }
 	void SetNVIDIA() { VendorId = 0x10DE; }
+	void SetQualcomm() { VendorId = 0x4D4F4351; }
 
 	bool IsAMD() const { return VendorId == 0x1002; }
 	bool IsIntel() const { return VendorId == 0x8086; }
 	bool IsNVIDIA() const { return VendorId == 0x10DE; }
+	bool IsQualcomm() const
+	{
+		// On Windows, the DXGI vendor ID is reported as 0x4D4F4351.
+		// On other platforms, it's 0x5143.
+		return VendorId == 0x5143 || VendorId == 0x4D4F4351;
+	}
 
 	bool IsSameDriverVersionGeneration(const FDriverVersion& UnifiedDriverVersion) const
 	{
@@ -554,6 +561,10 @@ private:
 		{
 			Section = TEXT("GPU_Intel");
 		}
+		else if (DriverInfo.IsQualcomm())
+		{
+			Section = TEXT("GPU_Qualcomm");
+		}
 		if (!Section)
 		{
 			return TEXT("");
diff --git a/Engine/Source/Runtime/Core/Tests/GenericPlatform/GenericPlatformDriverTest.cpp b/Engine/Source/Runtime/Core/Tests/GenericPlatform/GenericPlatformDriverTest.cpp
--- a/Engine/Source/Runtime/Core/Tests/GenericPlatform/GenericPlatformDriverTest.cpp
+++ b/Engine/Source/Runtime/Core/Tests/GenericPlatform/GenericPlatformDriverTest.cpp
@@ -123,6 +123,12 @@ TEST_CASE_NAMED(DriverVersionTest, "System::Core::GenericPlatform::Driver", "[Ap
 			Info.InternalDriverVersion = TEXT("15.200.1062.1004");
 			CHECK(Info.GetUnifiedDriverVersion() == TEXT("15.200.1062.1004"));
 		}
+		{
+			FGPUDriverInfo Info;
+			Info.SetQualcomm();
+			Info.InternalDriverVersion = TEXT("31.0.117.0");
+			CHECK(Info.GetUnifiedDriverVersion() == TEXT("31.0.117.0"));
+		}
 		{
 			FGPUDriverInfo Info;
 			Info.SetIntel();
diff --git a/Engine/Source/Runtime/RHI/Private/DynamicRHI.cpp b/Engine/Source/Runtime/RHI/Private/DynamicRHI.cpp
--- a/Engine/Source/Runtime/RHI/Private/DynamicRHI.cpp
+++ b/Engine/Source/Runtime/RHI/Private/DynamicRHI.cpp
@@ -176,6 +176,11 @@ void RHIDetectAndWarnOfBadDrivers(bool bHasEditorToken)
 				VendorString = TEXT("Intel");
 				HyperlinkText = NSLOCTEXT("MessageDialog", "DriverDownloadLinkIntel", "https://downloadcenter.intel.com/product/80939/Graphics");
 			}
+			else if (DriverInfo.IsQualcomm())
+			{
+				VendorString = TEXT("Qualcomm");
+				HyperlinkText = NSLOCTEXT("MessageDialog", "DriverDownloadLinkQualcomm", "https://softwarecenter.qualcomm.com/catalog/item/Windows_Graphics_Driver");
+			}
 
 			FFormatNamedArguments Args;
 			Args.Add(TEXT("AdapterName"), FText::FromString(DriverInfo.DeviceDescription));
