Windows/Mac/Unix output devices: Use UE_LOGFMT when reporting error so the error message inherits structured log context fields from the asserting thread.

#FailedValidations Preflight

diff --git a/Engine/Source/Programs/TestPAL/Private/Main.cpp b/Engine/Source/Programs/TestPAL/Private/Main.cpp
--- a/Engine/Source/Programs/TestPAL/Private/Main.cpp
+++ b/Engine/Source/Programs/TestPAL/Private/Main.cpp
@@ -3,6 +3,7 @@
 #include "CoreMinimal.h"
 #include "TestPALLog.h"
 #include "Parent.h"
+#include "Logging/StructuredLog.h"
 #include "Misc/Guid.h"
 #include "Stats/StatsMisc.h"
 #include "HAL/RunnableThread.h"
@@ -344,6 +345,8 @@ int32 CrashTest(const TCHAR* CommandLine)
 	GEngineLoop.PreInit(CommandLine);
 	UE_LOG(LogTestPAL, Display, TEXT("Running crash test (this should not exit)."));
 
+	UE_LOG_CONTEXT("FunctionName", "CrashTest");
+
 	if (FParse::Param(CommandLine, TEXT("logfatal")))
 	{
 		UE_LOG(LogTestPAL, Fatal, TEXT("  LogFatal!"));
diff --git a/Engine/Source/Runtime/Core/Private/Mac/MacErrorOutputDevice.cpp b/Engine/Source/Runtime/Core/Private/Mac/MacErrorOutputDevice.cpp
--- a/Engine/Source/Runtime/Core/Private/Mac/MacErrorOutputDevice.cpp
+++ b/Engine/Source/Runtime/Core/Private/Mac/MacErrorOutputDevice.cpp
@@ -27,12 +27,12 @@ void FMacErrorOutputDevice::Serialize( const TCHAR* Msg, ELogVerbosity::Type Ver
 		// pop up a crash window if we are not in unattended mode
 		if( FApp::IsUnattended() == false )
 		{ 
-			UE_LOG(LogMac, Error, TEXT("appError called: %s"), Msg );
+			UE_LOGFMT(LogMac, Error, "appError called: {Msg}", Msg );
 		} 
 		// log the warnings to the log
 		else
 		{
-			UE_LOG(LogMac, Error, TEXT("appError called: %s"), Msg );
+			UE_LOGFMT(LogMac, Error, "appError called: {Msg}", Msg );
 		}
 
 		// CheckVerifyFailedImpl writes GErrorHist including a callstack and then calls GError->Logf with only the
diff --git a/Engine/Source/Runtime/Core/Private/Unix/UnixErrorOutputDevice.cpp b/Engine/Source/Runtime/Core/Private/Unix/UnixErrorOutputDevice.cpp
--- a/Engine/Source/Runtime/Core/Private/Unix/UnixErrorOutputDevice.cpp
+++ b/Engine/Source/Runtime/Core/Private/Unix/UnixErrorOutputDevice.cpp
@@ -4,6 +4,7 @@
 #include "Containers/StringConv.h"
 #include "Logging/LogMacros.h"
 #include "CoreGlobals.h"
+#include "Logging/StructuredLog.h"
 #include "Misc/Parse.h"
 #include "Misc/CommandLine.h"
 #include "Misc/Paths.h"
@@ -34,12 +35,12 @@ void FUnixErrorOutputDevice::Serialize(const TCHAR* Msg, ELogVerbosity::Type Ver
 		if (WITH_EDITOR && FApp::IsUnattended() == false)
 		{
 			// @TODO Not implemented
-			UE_LOG(LogCore, Error, TEXT("appError called: %s"), Msg );
+			UE_LOGFMT(LogCore, Error, "appError called: {Msg}", Msg );
 		}
 		// log the warnings to the log
 		else
 		{
-			UE_LOG(LogCore, Error, TEXT("appError called: %s"), Msg );
+			UE_LOGFMT(LogCore, Error, "appError called: {Msg}", Msg );
 		}
 
 		// CheckVerifyFailedImpl writes GErrorHist including a callstack and then calls GError->Logf with only the
diff --git a/Engine/Source/Runtime/Core/Private/Windows/WindowsErrorOutputDevice.cpp b/Engine/Source/Runtime/Core/Private/Windows/WindowsErrorOutputDevice.cpp
--- a/Engine/Source/Runtime/Core/Private/Windows/WindowsErrorOutputDevice.cpp
+++ b/Engine/Source/Runtime/Core/Private/Windows/WindowsErrorOutputDevice.cpp
@@ -12,6 +12,7 @@
 #include "Misc/Parse.h"
 #include "Misc/CommandLine.h"
 #include "Internationalization/Internationalization.h"
+#include "Logging/StructuredLog.h"
 #include "Misc/ConfigCacheIni.h"
 #include "Misc/OutputDeviceHelper.h"
 #include "Misc/CoreDelegates.h"
@@ -38,7 +39,7 @@ void FWindowsErrorOutputDevice::Serialize( const TCHAR* Msg, ELogVerbosity::Type
 
 		// There is a bit of ambiguity around which SEH __except is to call HandleError or LogFormattedMessageWithCallstack
 		// We print it here so that the message, at very least, is always included with the logs
-		UE_LOG(LogWindows, Error, TEXT("appError called: %s"), Msg);
+		UE_LOGFMT(LogWindows, Error, "appError called: {Msg}", Msg);
 
 		// Windows error.
 		const int32 LastError = ::GetLastError();
