[Core] Fix natvis visualizers for sparse/compact maps and sets
* Fix declaration of sparse/compact map to use the correct allocators
* Fix error in manual access of allocator data in TMulticastDelegate

#tests Added a function with all set/map types to view results

diff --git a/Engine/Source/Runtime/Core/Public/Containers/ContainersFwd.h b/Engine/Source/Runtime/Core/Public/Containers/ContainersFwd.h
--- a/Engine/Source/Runtime/Core/Public/Containers/ContainersFwd.h
+++ b/Engine/Source/Runtime/Core/Public/Containers/ContainersFwd.h
@@ -71,11 +71,11 @@ template<typename T, typename SizeType = int32> using TConstArrayView = TArrayVi
 template<typename T> using TConstArrayView64 = TConstArrayView<T, int64>;
 template<typename T> class TTransArray;
 template<typename InKeyType, typename InValueType, bool bInAllowDuplicateKeys> struct TDefaultMapHashableKeyFuncs;
-template<typename InKeyType, typename InValueType, typename SetAllocator = FDefaultSetAllocator, typename KeyFuncs = TDefaultMapHashableKeyFuncs<InKeyType, InValueType, false> > class TSparseMap;
-template<typename InKeyType, typename InValueType, typename SetAllocator = FDefaultSetAllocator, typename KeyFuncs = TDefaultMapHashableKeyFuncs<InKeyType, InValueType, false> > class TCompactMap;
+template<typename InKeyType, typename InValueType, typename SetAllocator = FDefaultSparseSetAllocator, typename KeyFuncs = TDefaultMapHashableKeyFuncs<InKeyType, InValueType, false> > class TSparseMap;
+template<typename InKeyType, typename InValueType, typename SetAllocator = FDefaultCompactSetAllocator, typename KeyFuncs = TDefaultMapHashableKeyFuncs<InKeyType, InValueType, false> > class TCompactMap;
 template<typename InKeyType, typename InValueType, typename SetAllocator = FDefaultSetAllocator, typename KeyFuncs = TDefaultMapHashableKeyFuncs<InKeyType, InValueType, false> > class TMap;
-template<typename KeyType, typename ValueType, typename SetAllocator = FDefaultSetAllocator, typename KeyFuncs = TDefaultMapHashableKeyFuncs<KeyType, ValueType, true > > class TSparseMultiMap;
-template<typename KeyType, typename ValueType, typename SetAllocator = FDefaultSetAllocator, typename KeyFuncs = TDefaultMapHashableKeyFuncs<KeyType, ValueType, true > > class TCompactMultiMap;
+template<typename KeyType, typename ValueType, typename SetAllocator = FDefaultSparseSetAllocator, typename KeyFuncs = TDefaultMapHashableKeyFuncs<KeyType, ValueType, true > > class TSparseMultiMap;
+template<typename KeyType, typename ValueType, typename SetAllocator = FDefaultCompactSetAllocator, typename KeyFuncs = TDefaultMapHashableKeyFuncs<KeyType, ValueType, true > > class TCompactMultiMap;
 template<typename KeyType, typename ValueType, typename SetAllocator = FDefaultSetAllocator, typename KeyFuncs = TDefaultMapHashableKeyFuncs<KeyType, ValueType, true > > class TMultiMap;
 template <typename T = void > struct TLess;
 template <typename> struct TTypeTraits;
