Add missing header

#DisabledValidations nonUnityAndStaticAnalysis

diff --git a/Engine/Source/Runtime/Core/Private/Mac/MacErrorOutputDevice.cpp b/Engine/Source/Runtime/Core/Private/Mac/MacErrorOutputDevice.cpp
--- a/Engine/Source/Runtime/Core/Private/Mac/MacErrorOutputDevice.cpp
+++ b/Engine/Source/Runtime/Core/Private/Mac/MacErrorOutputDevice.cpp
@@ -1,6 +1,7 @@
 // Copyright Epic Games, Inc. All Rights Reserved.
 
 #include "Mac/MacErrorOutputDevice.h"
+#include "Logging/StructuredLog.h"
 #include "Misc/App.h"
 #include "Misc/CoreDelegates.h"
 #include "Misc/ConfigCacheIni.h"
