Add fallback method for re-loading a merged module that is in the process of winding down

Support for this was never tested before because it never occurred, but an opportunity showed up.

#rnx

#FailedValidations Preflight

diff --git a/Engine/Source/Runtime/Core/Private/Modules/ModuleManager.cpp b/Engine/Source/Runtime/Core/Private/Modules/ModuleManager.cpp
--- a/Engine/Source/Runtime/Core/Private/Modules/ModuleManager.cpp
+++ b/Engine/Source/Runtime/Core/Private/Modules/ModuleManager.cpp
@@ -34,7 +34,7 @@ DEFINE_LOG_CATEGORY_STATIC(LogModuleManager, Log, All);
 #if UE_MERGED_MODULES
 
 	// Toggle switch for merged modules unloading
-	static bool GEnableMergedLibraryUnloading = true;
+	static bool GEnableMergedLibraryUnloading = false;
 	static FAutoConsoleVariableRef EnableMergedLibraryUnloadingCVar(
 		TEXT("Modules.EnableMergedLibraryUnloading"),
 		GEnableMergedLibraryUnloading,
@@ -554,6 +554,17 @@ void* FModuleManager::InternalLoadLibrary(FName ModuleName, const FString& Libra
 
 #if UE_MERGED_MODULES
 
+	// Stop any delay unload for that library, and match the handle for the code module to the library path.
+	if (void** ExistingHandle = LibraryHandles.Find(LibraryToLoad))
+	{
+		UE_LOG(LogModuleManager, Verbose, TEXT("InternalLoadLibrary: already loaded library '%s'"), *LibraryToLoad);
+
+		DelayUnloadLibraries.Remove(LibraryToLoad);
+
+		return ExistingHandle;
+	}
+	DelayUnloadLibraries.Remove(LibraryToLoad);
+
 	// The plugin manager normally handles dependencies, but merged libraries introduce indirect dependencies.
 	// If we load plugin A that's merged with B, and B depends on C in another library...
 	// ... then we need to load C first, but the plugin manager cannot know that.
