Replaced FString with FUtf8String in FConfigFileHierarchy to save memory

diff --git a/Engine/Source/Editor/ConfigEditor/Private/SConfigEditor.cpp b/Engine/Source/Editor/ConfigEditor/Private/SConfigEditor.cpp
--- a/Engine/Source/Editor/ConfigEditor/Private/SConfigEditor.cpp
+++ b/Engine/Source/Editor/ConfigEditor/Private/SConfigEditor.cpp
@@ -101,7 +101,7 @@ void SConfigEditor::CreateDisplayObjectForSelectedTargetPlatform()
 			for (const auto& IniFile : Branch->Hierarchy)
 			{
 				UPropertyConfigFileDisplayRow* ConfigFilePropertyObj = NewObject<UPropertyConfigFileDisplayRow>(GetTransientPackage(), *IniFile.Value);
-				ConfigFilePropertyObj->InitWithConfigAndProperty(IniFile.Value, PropHelper->EditProperty.Get());
+				ConfigFilePropertyObj->InitWithConfigAndProperty(FString(IniFile.Value), PropHelper->EditProperty.Get());
 
 				PropHelper->ConfigFilePropertyObjects.Add(ConfigFilePropertyObj);
 			}
diff --git a/Engine/Source/Runtime/Core/Private/Misc/ConfigCacheIni.cpp b/Engine/Source/Runtime/Core/Private/Misc/ConfigCacheIni.cpp
--- a/Engine/Source/Runtime/Core/Private/Misc/ConfigCacheIni.cpp
+++ b/Engine/Source/Runtime/Core/Private/Misc/ConfigCacheIni.cpp
@@ -2570,7 +2570,7 @@ bool FConfigFile::IsADefaultIniWrite(const FString& Filename, int32& OutIniCombi
 		FString IniName = FPaths::GetCleanFilename(Filename);
 		for (const auto& HierarchyFileIt : Branch->Hierarchy)
 		{
-			if (FPaths::GetCleanFilename(HierarchyFileIt.Value) == IniName)
+			if (FPaths::GetCleanFilenameUtf8(HierarchyFileIt.Value) == IniName)
 			{
 				OutIniCombineThreshold = HierarchyFileIt.Key;
 				break;
@@ -3211,7 +3211,7 @@ static TArray<FString> GetSourceProperties(const FConfigFileHierarchy& SourceIni
 		HighestFileIndex = NextEntry > HighestFileIndex ? NextEntry : HighestFileIndex;
 	}
 
-	const FString& LastFileInHierarchy = SourceIniHierarchy.FindChecked(HighestFileIndex);
+	const FString LastFileInHierarchy = FString(SourceIniHierarchy.FindChecked(HighestFileIndex));
 	FConfigFile& DefaultConfigFile = Hierarchy.Add(LastFileInHierarchy, FConfigFile());
 
 	for (const auto& HierarchyFileIt : SourceIniHierarchy)
@@ -3221,7 +3221,7 @@ static TArray<FString> GetSourceProperties(const FConfigFileHierarchy& SourceIni
 		// we would be diffing against the same config we've just written to.
 		if (HierarchyFileIt.Key < IniCombineThreshold)
 		{
-			DefaultConfigFile.Combine(HierarchyFileIt.Value);
+			DefaultConfigFile.Combine(FString(HierarchyFileIt.Value));
 		}
 	}
 
@@ -3229,7 +3229,7 @@ static TArray<FString> GetSourceProperties(const FConfigFileHierarchy& SourceIni
 	// Note.	This compensates for an issue where strings in the hierarchy have a slightly different format
 	//			to how the config system wishes to serialize them.
 	TArray<FString> SourceArrayProperties;
-	Hierarchy.GetArray(*SectionName, *PropertyName, SourceArrayProperties, *LastFileInHierarchy);
+	Hierarchy.GetArray(*SectionName, *PropertyName, SourceArrayProperties, LastFileInHierarchy);
 
 	return SourceArrayProperties;
 }
@@ -5014,7 +5014,7 @@ static void DumpBranch(FOutputDevice& Ar, const FConfigBranch& Branch)
 	Ar.Logf(TEXT("Branch Name: %s"), *Branch.IniName.ToString());
 	Ar.Logf(TEXT("Branch Filename: %s"), *Branch.IniPath);
 	Ar.Logf(TEXT("Branch Static Hierarchy:"));
-	for (const TPair<int32, FString>& Pair: Branch.Hierarchy)
+	for (const TPair<int32, FUtf8String>& Pair: Branch.Hierarchy)
 	{
 		Ar.Logf(TEXT("  %s"), *Pair.Value);
 	}
@@ -5469,13 +5469,15 @@ static FConfigCommandStream* FindStaticLayer(FConfigCacheIni* ConfigSystem, cons
 			}
 		}
 
+		TUtf8StringBuilder<256> LayerKeySubstringUtf8(InPlace, LayerKeySubstring);
+
 		// we didn't find an existing layer, so look to see if we have a matching entry in the hierarchy and insert an empty layer for it
-		for (TPair<int32, FString>& IniFile : OutBranch->Hierarchy)
+		for (TPair<int32, FUtf8String>& IniFile : OutBranch->Hierarchy)
 		{
-			if (IniFile.Value.Contains(LayerKeySubstring))
+			if (IniFile.Value.Contains(LayerKeySubstringUtf8))
 			{
 				int KeyToInsert = IniFile.Key;
-				FString FilenameToInsert = IniFile.Value;
+				FString FilenameToInsert = FString(IniFile.Value);
 				bool bInserted = false;
 
 				// now for the painful process of finding where to insert it into the static layers - and since we need it ordered
@@ -5485,7 +5487,7 @@ static FConfigCommandStream* FindStaticLayer(FConfigCacheIni* ConfigSystem, cons
 				for (TPair<FString, FConfigCommandStream>& Pair : OutBranch->StaticLayers)
 				{
 					// find the entry in the hierarchy that matches this layer filename
-					if (const int32* LayerKey = OutBranch->Hierarchy.FindKey(Pair.Key))
+					if (const int32* LayerKey = OutBranch->Hierarchy.FindKey(FUtf8String(Pair.Key)))
 					{
 						if (!bInserted && *LayerKey > KeyToInsert)
 						{
@@ -5499,7 +5501,7 @@ static FConfigCommandStream* FindStaticLayer(FConfigCacheIni* ConfigSystem, cons
 				if (bInserted)
 				{
 					OutBranch->StaticLayers = NewLayers;
-					OutFilename = FilenameToInsert;
+					OutFilename = MoveTemp(FilenameToInsert);
 					return NewLayer;
 				}
 			}
@@ -7700,18 +7702,20 @@ class FIniExec : public FSelfRegisteringExec
 				// now walk the hierarchy, diffing things
 				bool bTrueBaseFound = false;
 				bool bBranchBaseFound = false;
-				for (const TPair<int32, FString>& Pair : Branch->Hierarchy)
+				for (const TPair<int32, FUtf8String>& Pair : Branch->Hierarchy)
 				{
+					const FString Value = FString(Pair.Value);
+
 					if (!bTrueBaseFound)
 					{
 						// prime the layers with the first, mainly empty layer
-						Branch->MergeStaticLayersUpToAndIncluding(Pair.Value, PreviousLayer);
+						Branch->MergeStaticLayersUpToAndIncluding(Value, PreviousLayer);
 						bTrueBaseFound = true;
 						continue;
 					}
 
 					FConfigFile CurrentCombinedLayer;
-					Branch->MergeStaticLayersUpToAndIncluding(Pair.Value, CurrentCombinedLayer);
+					Branch->MergeStaticLayersUpToAndIncluding(Value, CurrentCombinedLayer);
 					if (!bBranchBaseFound)
 					{
 						BaseLayer = CurrentCombinedLayer;
@@ -7723,7 +7727,7 @@ class FIniExec : public FSelfRegisteringExec
 					}
 
 					
-					const FConfigCommandStream* CurrentLayer = Branch->GetStaticLayer(Pair.Value);
+					const FConfigCommandStream* CurrentLayer = Branch->GetStaticLayer(Value);
 					if (CurrentLayer == nullptr)
 					{
 						// most layers won't exist
@@ -7747,7 +7751,7 @@ class FIniExec : public FSelfRegisteringExec
 								{
 									if (bShowMissingDefault)
 									{
-										LogHeaders(bShownFileHeader, bShownSectionHeader, Pair.Value, SectionPair.Key);
+										LogHeaders(bShownFileHeader, bShownSectionHeader, Value, SectionPair.Key);
 										Ar.Logf(TEXT("  %s -- No Default in Base"), *ValuePair.Key.ToString());
 									}
 								}
@@ -7755,7 +7759,7 @@ class FIniExec : public FSelfRegisteringExec
 								{
 									if (bShowSameAsPrevious)
 									{
-										LogHeaders(bShownFileHeader, bShownSectionHeader, Pair.Value, SectionPair.Key);
+										LogHeaders(bShownFileHeader, bShownSectionHeader, Value, SectionPair.Key);
 										Ar.Logf(TEXT("  %s -- Identical To Previous"), *ValuePair.Key.ToString());
 									}
 								}
@@ -7778,7 +7782,7 @@ class FIniExec : public FSelfRegisteringExec
 							{
 								if (AddRemovePair.Value == 3)
 								{
-									LogHeaders(bShownFileHeader, bShownSectionHeader, Pair.Value, SectionPair.Key);
+									LogHeaders(bShownFileHeader, bShownSectionHeader, Value, SectionPair.Key);
 									Ar.Logf(TEXT("  %s -- Has Add and Remove"), *AddRemovePair.Key);
 								}
 							}
diff --git a/Engine/Source/Runtime/Core/Private/Misc/ConfigContext.cpp b/Engine/Source/Runtime/Core/Private/Misc/ConfigContext.cpp
--- a/Engine/Source/Runtime/Core/Private/Misc/ConfigContext.cpp
+++ b/Engine/Source/Runtime/Core/Private/Misc/ConfigContext.cpp
@@ -1027,10 +1027,10 @@ bool FConfigContext::LoadIniFileHierarchy()
 
 	TRACE_CPUPROFILER_EVENT_SCOPE(LoadIniFileHierarchy);
 	// Traverse ini list back to front, merging along the way.
-	for (const TPair<int32, FString>& HierarchyIt : Branch->Hierarchy)
+	for (const TPair<int32, FUtf8String>& HierarchyIt : Branch->Hierarchy)
 	{
 		const int KeyWithoutFlags = HierarchyIt.Key & ~((1 << NumFlagsBits) - 1);
-		const FString& IniFileName = HierarchyIt.Value;
+		const FString IniFileName = FString(HierarchyIt.Value);
 
 		const bool bDoCombine = (KeyWithoutFlags != 0);
 		const bool bUseStagedGlobalCache = ConfigSystem != nullptr && (HierarchyIt.Key & KeyFlag_UseGlobalCache) != 0;
@@ -1277,7 +1277,7 @@ int32 FConfigFileHierarchy::GenerateDynamicKey()
 int32 FConfigFileHierarchy::AddStaticLayer(const FString& Filename, int32 LayerIndex, int32 ExpansionIndex /*= 0*/, int32 PlatformIndex /*= 0*/, int32 Flags)
 {
 	int32 Key = GetStaticKey(LayerIndex, ExpansionIndex, PlatformIndex, Flags);
-	Emplace(Key, Filename);
+	Emplace(Key, FUtf8String(Filename));
 	return Key;
 }
 
diff --git a/Engine/Source/Runtime/Core/Private/Misc/Paths.cpp b/Engine/Source/Runtime/Core/Private/Misc/Paths.cpp
--- a/Engine/Source/Runtime/Core/Private/Misc/Paths.cpp
+++ b/Engine/Source/Runtime/Core/Private/Misc/Paths.cpp
@@ -950,6 +950,15 @@ FString FPaths::GetCleanFilename(const FString& InPath)
 	return Result;
 }
 
+FString FPaths::GetCleanFilenameUtf8(const FUtf8String& InPath)
+{
+	int32 EndPos = InPath.FindLastCharByPredicate(UE4Paths_Private::IsNotSlashOrBackslash) + 1;
+	int32 StartPos = InPath.FindLastCharByPredicate(UE4Paths_Private::IsSlashOrBackslash) + 1;
+
+	FString Result = (StartPos <= EndPos) ? FString(InPath.Mid(StartPos, EndPos - StartPos)) : "";
+	return Result;
+}
+
 FString FPaths::GetCleanFilename(FString&& InPath)
 {
 	static_assert(INDEX_NONE == -1, "INDEX_NONE assumed to be -1");
diff --git a/Engine/Source/Runtime/Core/Public/Misc/ConfigCacheIni.h b/Engine/Source/Runtime/Core/Public/Misc/ConfigCacheIni.h
--- a/Engine/Source/Runtime/Core/Public/Misc/ConfigCacheIni.h
+++ b/Engine/Source/Runtime/Core/Public/Misc/ConfigCacheIni.h
@@ -720,7 +720,7 @@ public:
 	bool operator==( const FConfigFile& Other ) const;
 	bool operator!=( const FConfigFile& Other ) const;
 
-	CORE_API bool Combine( const FString& Filename);
+	CORE_API bool Combine(const FString& Filename);
 	CORE_API void CombineFromBuffer(const FString& Buffer, const FString& FileHint);
 	CORE_API void Read( const FString& Filename );
 
@@ -1049,7 +1049,7 @@ private:
 	 *
 	 * @param InBaseIniName Ini name.
 	 * @param InPlatformName Platform name, nullptr means to use the current platform
-	 * @param OutHierarchy An array which is to receive the generated hierachy of ini filenames.
+	 * @param OutHierarchy An array which is to receive the generated hierarchy of ini filenames.
 	 */
 	void AddStaticLayersToHierarchy(const TCHAR* InBaseIniName, const TCHAR* InPlatformName, const TCHAR* EngineConfigDir, const TCHAR* SourceConfigDir);
 	static void AddStaticLayersToHierarchy(FConfigContext& Context);
diff --git a/Engine/Source/Runtime/Core/Public/Misc/ConfigTypes.h b/Engine/Source/Runtime/Core/Public/Misc/ConfigTypes.h
--- a/Engine/Source/Runtime/Core/Public/Misc/ConfigTypes.h
+++ b/Engine/Source/Runtime/Core/Public/Misc/ConfigTypes.h
@@ -4,12 +4,13 @@
 
 #include "CoreTypes.h"
 #include "Containers/Array.h"
+#include "Containers/Utf8String.h"
 #include "Containers/UnrealString.h"
 #include "Containers/Map.h"
 #include "CoreGlobals.h"
 
 
-class FConfigFileHierarchy : public TMap<int32, FString>
+class FConfigFileHierarchy : public TMap<int32, FUtf8String>
 {
 private:
 	int32 KeyGen = 0;
diff --git a/Engine/Source/Runtime/Core/Public/Misc/Paths.h b/Engine/Source/Runtime/Core/Public/Misc/Paths.h
--- a/Engine/Source/Runtime/Core/Public/Misc/Paths.h
+++ b/Engine/Source/Runtime/Core/Public/Misc/Paths.h
@@ -508,6 +508,7 @@ public:
 
 	// Returns the filename (with extension), minus any path information.
 	static CORE_API FString GetCleanFilename(const FString& InPath);
+	static CORE_API FString GetCleanFilenameUtf8(const FUtf8String& InPath);
 
 	// Returns the filename (with extension), minus any path information.
 	static CORE_API FString GetCleanFilename(FString&& InPath);
