Fixed DumpBranch for a branch hierarchy in config hierarchy ini

diff --git a/Engine/Source/Runtime/Core/Private/Misc/ConfigCacheIni.cpp b/Engine/Source/Runtime/Core/Private/Misc/ConfigCacheIni.cpp
--- a/Engine/Source/Runtime/Core/Private/Misc/ConfigCacheIni.cpp
+++ b/Engine/Source/Runtime/Core/Private/Misc/ConfigCacheIni.cpp
@@ -5016,7 +5016,7 @@ static void DumpBranch(FOutputDevice& Ar, const FConfigBranch& Branch)
 	Ar.Logf(TEXT("Branch Static Hierarchy:"));
 	for (const TPair<int32, FUtf8String>& Pair: Branch.Hierarchy)
 	{
-		Ar.Logf(TEXT("  %s"), *Pair.Value);
+		Ar.Logf(TEXT("  %s"), *FString(Pair.Value));
 	}
 	if (Branch.DynamicLayers.GetHead())
 	{
