Detect ANGLE and replace GL_RENDERER / SRC_GPUFamily string with the Vulkan VulkanPhysicalDeviceProperties.deviceName string

diff --git a/Engine/Source/Runtime/Core/Private/Android/AndroidPlatformMisc.cpp b/Engine/Source/Runtime/Core/Private/Android/AndroidPlatformMisc.cpp
--- a/Engine/Source/Runtime/Core/Private/Android/AndroidPlatformMisc.cpp
+++ b/Engine/Source/Runtime/Core/Private/Android/AndroidPlatformMisc.cpp
@@ -2020,6 +2020,7 @@ enum class EDeviceVulkanSupportStatus
 };
 
 static FString VulkanVersionString;
+static FString VulkanDeviceNameString;
 static EDeviceVulkanSupportStatus VulkanSupport = EDeviceVulkanSupportStatus::Uninitialized;
 
 ///////////////////////////////////////////////////////////////////////////////
@@ -2173,6 +2174,7 @@ static EDeviceVulkanSupportStatus AttemptVulkanInit(void* VulkanLib)
 
 	VulkanVersionString = FString::Printf(TEXT("%d.%d.%d"), VK_VERSION_MAJOR(VulkanPhysicalDeviceProperties.apiVersion), VK_VERSION_MINOR(VulkanPhysicalDeviceProperties.apiVersion), VK_VERSION_PATCH(VulkanPhysicalDeviceProperties.apiVersion));
 	vkDestroyInstance(Instance, nullptr);
+	VulkanDeviceNameString = VulkanPhysicalDeviceProperties.deviceName;
 
 	return EDeviceVulkanSupportStatus::Supported;
 }
@@ -2392,6 +2394,12 @@ FString FAndroidMisc::GetVulkanVersion()
 	return VulkanVersionString;
 }
 
+FString FAndroidMisc::GetVulkanDeviceNameString()
+{
+	check(VulkanSupport != EDeviceVulkanSupportStatus::Uninitialized);
+	return VulkanDeviceNameString;
+}
+
 FString FAndroidMisc::GetVKQualityRecommendation()
 {
 	static bool bCachedResult = false;
diff --git a/Engine/Source/Runtime/Core/Public/Android/AndroidPlatformMisc.h b/Engine/Source/Runtime/Core/Public/Android/AndroidPlatformMisc.h
--- a/Engine/Source/Runtime/Core/Public/Android/AndroidPlatformMisc.h
+++ b/Engine/Source/Runtime/Core/Public/Android/AndroidPlatformMisc.h
@@ -211,7 +211,7 @@ public:
 	 */
 	static CORE_API bool IsVulkanAvailable();
 	static CORE_API bool IsDesktopVulkanAvailable();
-
+	
 	/* ShouldUseVulkan
 	 * @return true if Vulkan is available, and not disabled by device profile cvar
 	 */
@@ -219,6 +219,7 @@ public:
 	static CORE_API bool ShouldUseDesktopVulkan();
 	static CORE_API FString GetVulkanVersion();
 	static CORE_API FString GetVKQualityRecommendation();
+	static CORE_API FString GetVulkanDeviceNameString();
 
 	static CORE_API bool IsExternalMemoryAndroidHardwareBufferExtensionLoaded();
 	typedef TFunction<void(void* NewNativeHandle)> ReInitWindowCallbackType;
diff --git a/Engine/Source/Runtime/OpenGLDrv/Private/Android/AndroidOpenGLPrivate.cpp b/Engine/Source/Runtime/OpenGLDrv/Private/Android/AndroidOpenGLPrivate.cpp
--- a/Engine/Source/Runtime/OpenGLDrv/Private/Android/AndroidOpenGLPrivate.cpp
+++ b/Engine/Source/Runtime/OpenGLDrv/Private/Android/AndroidOpenGLPrivate.cpp
@@ -71,6 +71,14 @@ void FAndroidGPUInfo::ReadGPUFamily()
 	GPUFamily = (const ANSICHAR*)glGetString(GL_RENDERER);
 	check(!GPUFamily.IsEmpty());
 
+	// If we're ANGLE, get the raw GPU string from Vulkan.
+	if (GPUFamily.StartsWith(TEXT("ANGLE")) && FAndroidMisc::HasVulkanDriverSupport())
+	{
+		UE_LOG(LogAndroid, Log, TEXT("FAndroidGPUInfo ANGLE detected with GL_RENDERER string %s"), *GPUFamily);
+		GPUFamily = FAndroidMisc::GetVulkanDeviceNameString();
+		UE_LOG(LogAndroid, Log, TEXT("FAndroidGPUInfo Using Vulkan deviceName string %s"), *GPUFamily);
+	}
+
 	// thirdparty api requires std::unique_ptr
 	std::unique_ptr<libgpuinfo::instance> ArmGPUInfoInstance = libgpuinfo::instance::create();
 
