Re-enable merged modules loading by default, as this this the behaviour we want by default

Previously set to disabled by default with the introduction of the control cvars.

#rnx

diff --git a/Engine/Source/Runtime/Core/Private/Modules/ModuleManager.cpp b/Engine/Source/Runtime/Core/Private/Modules/ModuleManager.cpp
--- a/Engine/Source/Runtime/Core/Private/Modules/ModuleManager.cpp
+++ b/Engine/Source/Runtime/Core/Private/Modules/ModuleManager.cpp
@@ -34,7 +34,7 @@ DEFINE_LOG_CATEGORY_STATIC(LogModuleManager, Log, All);
 #if UE_MERGED_MODULES
 
 	// Toggle switch for merged modules unloading
-	static bool GEnableMergedLibraryUnloading = false;
+	static bool GEnableMergedLibraryUnloading = true;
 	static FAutoConsoleVariableRef EnableMergedLibraryUnloadingCVar(
 		TEXT("Modules.EnableMergedLibraryUnloading"),
 		GEnableMergedLibraryUnloading,
