Added stat scope for FNamePoolShard::Grow.

#rnx

diff --git a/Engine/Source/Runtime/Core/Private/UObject/UnrealNames.cpp b/Engine/Source/Runtime/Core/Private/UObject/UnrealNames.cpp
--- a/Engine/Source/Runtime/Core/Private/UObject/UnrealNames.cpp
+++ b/Engine/Source/Runtime/Core/Private/UObject/UnrealNames.cpp
@@ -1588,6 +1588,7 @@ private:
 		LLM_TAGSET_SCOPE_CLEAR(ELLMTagSet::Assets);
 		LLM_TAGSET_SCOPE_CLEAR(ELLMTagSet::AssetClasses);
 		UE_TRACE_METADATA_CLEAR_SCOPE();
+		QUICK_SCOPE_CYCLE_COUNTER(STAT_FNamePoolShard_Grow);
 		TArrayView<FNameSlot> OldSlots(Slots, Capacity());
 		const uint32 OldUsedSlots = UsedSlots;
 
