Fixed a race waking threads in FSharedMutex and FSharedRecursiveMutex

It was possible for a shared waiter to not be woken if it is the only shared waiter and started waiting after Unlock() read the State but before it cleared MayHaveWaitingLockFlag.

diff --git a/Engine/Source/Runtime/Core/Private/Async/SharedMutex.cpp b/Engine/Source/Runtime/Core/Private/Async/SharedMutex.cpp
--- a/Engine/Source/Runtime/Core/Private/Async/SharedMutex.cpp
+++ b/Engine/Source/Runtime/Core/Private/Async/SharedMutex.cpp
@@ -96,6 +96,13 @@ FORCENOINLINE void FSharedMutex::WakeWaitingThreads(uint32 LastState)
 		{
 			return;
 		}
+
+		// Reload the state if there were no shared waiters because new
+		// ones may have registered themselves since LastState was read.
+		if (!(LastState & MayHaveWaitingSharedLockFlag))
+		{
+			LastState = State.load(std::memory_order_relaxed);
+		}
 	}
 
 	if (LastState & MayHaveWaitingSharedLockFlag)
diff --git a/Engine/Source/Runtime/Core/Private/Async/SharedRecursiveMutex.cpp b/Engine/Source/Runtime/Core/Private/Async/SharedRecursiveMutex.cpp
--- a/Engine/Source/Runtime/Core/Private/Async/SharedRecursiveMutex.cpp
+++ b/Engine/Source/Runtime/Core/Private/Async/SharedRecursiveMutex.cpp
@@ -356,6 +356,13 @@ FORCENOINLINE void FSharedRecursiveMutex::WakeWaitingThreads(uint32 LastState)
 		{
 			return;
 		}
+
+		// Reload the state if there were no shared waiters because new
+		// ones may have registered themselves since LastState was read.
+		if (!(LastState & MayHaveWaitingSharedLockFlag))
+		{
+			LastState = State.load(std::memory_order_relaxed);
+		}
 	}
 
 	if (LastState & MayHaveWaitingSharedLockFlag)
