[Insights] Modules: Fixed the incorrect module size estimation for mach-o modules. Previously the size of the main module has included the large 4GiB reserved virtual memory of the __PAGEZERO segment and then the size value was trimmed to 32 bits. Now __PAGEZERO segment is excluded. Also excluded the __LINKEDIT segment.
* Fixed the incorrect ProgramSize allocation (same cause; inclusion of the 4GiB vmsize from __PAGEZERO segment).

diff --git a/Engine/Source/Runtime/Core/Private/IOS/IOSModuleDiagnostics.cpp b/Engine/Source/Runtime/Core/Private/IOS/IOSModuleDiagnostics.cpp
--- a/Engine/Source/Runtime/Core/Private/IOS/IOSModuleDiagnostics.cpp
+++ b/Engine/Source/Runtime/Core/Private/IOS/IOSModuleDiagnostics.cpp
@@ -1,6 +1,13 @@
 ﻿// Copyright Epic Games, Inc. All Rights Reserved.
 
 #include "CoreTypes.h"
+
+#if UE_BUILD_SHIPPING
+
+void Modules_Initialize() {}
+
+#else // !UE_BUILD_SHIPPING
+
 #include "HAL/LowLevelMemTracker.h"
 #include "HAL/PlatformProcess.h"
 #include "Math/UnrealMathUtility.h"
@@ -9,88 +16,118 @@
 #include "ProfilingDebugging/MemoryTrace.h"
 #include "ProfilingDebugging/MetadataTrace.h"
 
-#if !UE_BUILD_SHIPPING
 #include "Apple/PreAppleSystemHeaders.h"
 #include <mach/mach.h>
 #include <mach-o/dyld_images.h>
 #include <mach-o/loader.h>
 #include "Apple/PostAppleSystemHeaders.h"
-#endif
 
-////////////////////////////////////////////////////////////////////////////////
+//#include "CoreGlobals.h"
+//#define TRACE_MODULES_DEBUG_LOG(Format, ...) \
+//	{ \
+//		FPlatformMisc::LowLevelOutputDebugStringf(TEXT("[Modules]") Format TEXT("\n"), ##__VA_ARGS__); \
+//		UE_LOG(LogTemp, Display, TEXT("[Modules]") Format, ##__VA_ARGS__); \
+//	}
+#define TRACE_MODULES_DEBUG_LOG(Format, ...)
+
 void Modules_Initialize()
 {
-#if !UE_BUILD_SHIPPING
 	using namespace UE::Trace;
 
 	constexpr uint32 SizeOfSymbolFormatString = 4;
 	UE_TRACE_LOG(Diagnostics, ModuleInit, ModuleChannel, sizeof(ANSICHAR) * SizeOfSymbolFormatString)
 		<< ModuleInit.SymbolFormat("psym", SizeOfSymbolFormatString)
 		<< ModuleInit.ModuleBaseShift(uint8(0));
-    
-    struct task_dyld_info dyld_info;
-    mach_msg_type_number_t count = TASK_DYLD_INFO_COUNT;
-    if (task_info(mach_task_self(), TASK_DYLD_INFO, (task_info_t)&dyld_info, &count) == KERN_SUCCESS)
-    {
-        HeapId ProgramHeapId = MemoryTrace_HeapSpec(EMemoryTraceRootHeap::SystemMemory, TEXT("Program"), EMemoryTraceHeapFlags::NeverFrees);
-        FString ExecutablePath = FPaths::GetPath(FString([[NSBundle mainBundle] executablePath]));
-        
-        struct dyld_all_image_infos* infos = (struct dyld_all_image_infos*)dyld_info.all_image_info_addr;
-        // FPlatformMisc::LowLevelOutputDebugStringf(TEXT("Found %d modules"), infos->infoArrayCount);
-
-        // Send Mach-O's uuid as the BuildId. Psym generation seems to have an extra 0 but we will ignore it on the other end.
-        constexpr uint32 BuildIdSize = 16;
-        uint8 BuildId[BuildIdSize] = {0};
-
-        for (int i=0; i<infos->infoArrayCount; i++)
-        {
-            const struct dyld_image_info* Image = &infos->infoArray[i];
-            // FPlatformMisc::LowLevelOutputDebugStringf(TEXT("Found module %s"), ANSI_TO_TCHAR(Image->imageFilePath));
-
-            // calc image size by adding header and size of segments
-            struct mach_header_64* Header = (struct mach_header_64*)Image->imageLoadAddress;
-            size_t ImageSize = sizeof(*Header) + Header->sizeofcmds;
-            char* CmdPtr = (char*)(Header + 1);
-            for (int cmd=0; cmd < Header->ncmds; cmd++)
-            {
-                struct load_command* LoadCommand = (struct load_command*)CmdPtr;
-                if (LoadCommand->cmd == LC_SEGMENT_64)
-                {
-                    ImageSize += ((struct segment_command_64*)LoadCommand)->vmsize;
-                }
-                else if (LoadCommand->cmd == LC_UUID)
-                {
-                    FMemory::Memcpy(BuildId, ((struct uuid_command*)LoadCommand)->uuid, BuildIdSize);
-                }
-                CmdPtr = CmdPtr + LoadCommand->cmdsize;
-            }
-                        
-            FString ImageName(Image->imageFilePath);
-            bool bInsideExecutablePath = ImageName.StartsWith(ExecutablePath);
-            
-            // trim path to leave just image name
-            ImageName = FPaths::GetCleanFilename(ImageName);
-            uint64 ImageLoadAddress = (uint64)Image->imageLoadAddress;
-            
-            UE_TRACE_LOG(Diagnostics, ModuleLoad, ModuleChannel, sizeof(TCHAR) * ImageName.Len() + BuildIdSize)
-                << ModuleLoad.Name(*ImageName, ImageName.Len())
-                << ModuleLoad.Base(ImageLoadAddress)
-                << ModuleLoad.Size(ImageSize)
-                << ModuleLoad.ImageId(BuildId, BuildIdSize);
-                       
+
+	struct task_dyld_info dyld_info;
+	mach_msg_type_number_t count = TASK_DYLD_INFO_COUNT;
+	if (task_info(mach_task_self(), TASK_DYLD_INFO, (task_info_t)&dyld_info, &count) == KERN_SUCCESS)
+	{
+#if UE_MEMORY_TRACE_ENABLED
+		HeapId ProgramHeapId = MemoryTrace_HeapSpec(EMemoryTraceRootHeap::SystemMemory, TEXT("Program"), EMemoryTraceHeapFlags::NeverFrees);
+		FString ExecutablePath = FPaths::GetPath(FString([[NSBundle mainBundle]executablePath]));
+#endif // UE_MEMORY_TRACE_ENABLED
+
+		const struct dyld_all_image_infos* infos = (const struct dyld_all_image_infos*)dyld_info.all_image_info_addr;
+		TRACE_MODULES_DEBUG_LOG(TEXT("Found %d modules"), infos->infoArrayCount);
+
+		// Send Mach-O's uuid as the BuildId. Psym generation seems to have an extra 0 but we will ignore it on the other end.
+		constexpr uint32 BuildIdSize = 16;
+		uint8 BuildId[BuildIdSize] = {0};
+
+		for (int ImageInfoIndex = 0; ImageInfoIndex < infos->infoArrayCount; ImageInfoIndex++)
+		{
+			const struct dyld_image_info* Image = &infos->infoArray[ImageInfoIndex];
+			TRACE_MODULES_DEBUG_LOG(TEXT("Found module %s"), ANSI_TO_TCHAR(Image->imageFilePath));
+
+			const uint64 ModuleBase = uint64(Image->imageLoadAddress);
+			TRACE_MODULES_DEBUG_LOG(TEXT("    Base: 0x%llX"), ModuleBase);
+
+			const struct mach_header_64* Header = (const struct mach_header_64*)Image->imageLoadAddress;
+
+			// calc image size by adding header and size of segments
+			uint64 ImageSize = sizeof(*Header) + Header->sizeofcmds;
+			constexpr uint64 PageMask = (1 << 12) - 1;
+			ImageSize = (ImageSize + PageMask) & ~PageMask; // 4K page aligned
+			TRACE_MODULES_DEBUG_LOG(TEXT("    Header: %llu + %llu (%d cmds) -> %llu"), sizeof(*Header), Header->sizeofcmds, Header->ncmds, ImageSize);
+
+			char* CmdPtr = (char*)(Header + 1);
+			for (int CommandIndex = 0; CommandIndex < Header->ncmds; ++CommandIndex)
+			{
+				const struct load_command* LoadCommand = (const struct load_command*)CmdPtr;
+				if (LoadCommand->cmd == LC_SEGMENT_64)
+				{
+					const struct segment_command_64* Segment = (const struct segment_command_64*)LoadCommand;
+					TRACE_MODULES_DEBUG_LOG(TEXT("    LC_SEGMENT_64 %s (vmaddr=0x%llX, vmsize=%llu, filesize=%llu)"),
+						ANSI_TO_TCHAR(Segment->segname), uint64(Segment->vmaddr), uint64(Segment->vmsize), uint64(Segment->filesize));
+
+					if (uint64(Segment->vmaddr) != 0 && // skips __PAGEZERO segment (4 GiB; reserved virtual memory)
+						FPlatformString::Strcmp(Segment->segname, "__LINKEDIT") != 0) // skips __LINKEDIT segment
+					{
+						ImageSize += uint64(Segment->vmsize);
+					}
+				}
+				else if (LoadCommand->cmd == LC_UUID)
+				{
+					FMemory::Memcpy(BuildId, ((const struct uuid_command*)LoadCommand)->uuid, BuildIdSize);
+				}
+				CmdPtr = CmdPtr + LoadCommand->cmdsize;
+			}
+
+			TRACE_MODULES_DEBUG_LOG(TEXT("    ImageSize: %llu (end=0x%llX)"), ImageSize, ModuleBase + ImageSize);
+
+			FString ImageName(Image->imageFilePath);
+
 #if UE_MEMORY_TRACE_ENABLED
-            // Only count the main executable and any other libraries inside our bundle as "Program Size"
-            if (bInsideExecutablePath)
-            {
-                UE_TRACE_METADATA_CLEAR_SCOPE();
-                LLM(UE_MEMSCOPE(ELLMTag::ProgramSize));
-                MemoryTrace_Alloc(ImageLoadAddress, ImageSize, 1);
-                MemoryTrace_MarkAllocAsHeap(ImageLoadAddress, ProgramHeapId);
-                MemoryTrace_Alloc(ImageLoadAddress, ImageSize, 1);
-            }
+			bool bInsideExecutablePath = ImageName.StartsWith(ExecutablePath);
 #endif // UE_MEMORY_TRACE_ENABLED
 
-        }
-    }
-#endif
+			// trim path to leave just image name
+			ImageName = FPaths::GetCleanFilename(ImageName);
+
+			const uint32 ModuleSize = uint32(ImageSize);
+
+			UE_TRACE_LOG(Diagnostics, ModuleLoad, ModuleChannel, sizeof(TCHAR) * ImageName.Len() + BuildIdSize)
+				<< ModuleLoad.Name(*ImageName, ImageName.Len())
+				<< ModuleLoad.Base(ModuleBase)
+				<< ModuleLoad.Size(ModuleSize)
+				<< ModuleLoad.ImageId(BuildId, BuildIdSize);
+
+#if UE_MEMORY_TRACE_ENABLED
+			// Only count the main executable and any other libraries inside our bundle as "Program Size"
+			if (bInsideExecutablePath)
+			{
+				UE_TRACE_METADATA_CLEAR_SCOPE();
+				LLM(UE_MEMSCOPE(ELLMTag::ProgramSize));
+				MemoryTrace_Alloc(ModuleBase, ImageSize, 1);
+				MemoryTrace_MarkAllocAsHeap(ModuleBase, ProgramHeapId);
+				MemoryTrace_Alloc(ModuleBase, ImageSize, 1);
+			}
+#endif // UE_MEMORY_TRACE_ENABLED
+		}
+	}
 }
+
+#undef TRACE_MODULES_DEBUG_LOG
+
+#endif // !UE_BUILD_SHIPPING
