Add a generic version of GetExtensionDirs

diff --git a/Engine/Plugins/Runtime/GameFeatures/Source/GameFeatures/Private/GameFeaturesSubsystemSettings.cpp b/Engine/Plugins/Runtime/GameFeatures/Source/GameFeatures/Private/GameFeaturesSubsystemSettings.cpp
--- a/Engine/Plugins/Runtime/GameFeatures/Source/GameFeatures/Private/GameFeaturesSubsystemSettings.cpp
+++ b/Engine/Plugins/Runtime/GameFeatures/Source/GameFeatures/Private/GameFeaturesSubsystemSettings.cpp
@@ -23,9 +23,13 @@ bool UGameFeaturesSubsystemSettings::IsValidGameFeaturePlugin(const FString& Plu
 	{
 		FBuiltInGameFeaturePluginsFolders()
 		{
+			const FPaths::EGetExtensionDirsFlags ExtensionFlags =
+				FPaths::EGetExtensionDirsFlags::WithBase |
+				FPaths::EGetExtensionDirsFlags::WithRestricted;
+
 			// Get all the existing game feature paths
-			TArray<FString> RelativePaths = FPaths::GetNonPlatformExtensionDirs(
-				FPaths::ProjectDir(), FPaths::Combine(TEXT("Plugins"), TEXT("GameFeatures")));
+			TArray<FString> RelativePaths = FPaths::GetExtensionDirs(
+				FPaths::ProjectDir(), FPaths::Combine(TEXT("Plugins"), TEXT("GameFeatures")), ExtensionFlags);
 
 			// The base directory may not exist yet, add it if empty
 			if (RelativePaths.IsEmpty())
diff --git a/Engine/Source/Runtime/Core/Private/Misc/Paths.cpp b/Engine/Source/Runtime/Core/Private/Misc/Paths.cpp
--- a/Engine/Source/Runtime/Core/Private/Misc/Paths.cpp
+++ b/Engine/Source/Runtime/Core/Private/Misc/Paths.cpp
@@ -397,24 +397,16 @@ static void AddIfDirectoryExists(TArray<FString>& ExtensionDirs, FString&& Dir)
 	}
 }
 
-namespace UE::Paths
+static void GetExtensionDirsInternal(TArray<FString>& ExtensionDirs, const FString& BaseDir, const FString& SubDir, FPaths::EGetExtensionDirsFlags Flags)
 {
-	enum class EGetExtensionDirsFlags : uint8
+	if (EnumHasAnyFlags(Flags, FPaths::EGetExtensionDirsFlags::WithBase))
 	{
-		None = 0,
-		WithPlatforms = (1 << 0),
-		CheckValid = (1 << 1)
-	};
-	ENUM_CLASS_FLAGS(EGetExtensionDirsFlags);
-}
-
-static void GetExtensionDirsInternal(TArray<FString>& ExtensionDirs, const FString& BaseDir, const FString& SubDir, UE::Paths::EGetExtensionDirsFlags Flags)
-{
-	AddIfDirectoryExists(ExtensionDirs, FPaths::Combine(BaseDir, SubDir));
+		AddIfDirectoryExists(ExtensionDirs, FPaths::Combine(BaseDir, SubDir));
+	}
 
-	if (EnumHasAnyFlags(Flags, UE::Paths::EGetExtensionDirsFlags::WithPlatforms))
+	if (EnumHasAnyFlags(Flags, FPaths::EGetExtensionDirsFlags::WithPlatforms))
 	{
-		const bool bCheckValid = EnumHasAnyFlags(Flags, UE::Paths::EGetExtensionDirsFlags::CheckValid);
+		const bool bCheckValid = EnumHasAnyFlags(Flags, FPaths::EGetExtensionDirsFlags::CheckPlatformsValid);
 
 		FString PlatformExtensionBaseDir = FPaths::Combine(BaseDir, TEXT("Platforms"));
 		for (const FString& PlatformName : FDataDrivenPlatformInfoRegistry::GetPlatformDirectoryNames(bCheckValid))
@@ -423,25 +415,26 @@ static void GetExtensionDirsInternal(TArray<FString>& ExtensionDirs, const FStri
 		}
 	}
 
-	FString RestrictedBaseDir = FPaths::Combine(BaseDir, TEXT("Restricted"));
-	IFileManager::Get().IterateDirectory(*RestrictedBaseDir, [&ExtensionDirs, SubDir, Flags](const TCHAR* FilenameOrDirectory, bool bIsDirectory)  -> bool
+	if (EnumHasAnyFlags(Flags, FPaths::EGetExtensionDirsFlags::WithRestricted))
 	{
-		if (bIsDirectory)
+		FString RestrictedBaseDir = FPaths::Combine(BaseDir, TEXT("Restricted"));
+		IFileManager::Get().IterateDirectory(*RestrictedBaseDir, [&ExtensionDirs, SubDir, Flags](const TCHAR* FilenameOrDirectory, bool bIsDirectory)  -> bool
 		{
-			GetExtensionDirsInternal(ExtensionDirs, FilenameOrDirectory, SubDir, Flags);
-		}
-		return true;
-	});
+			if (bIsDirectory)
+			{
+				GetExtensionDirsInternal(ExtensionDirs, FilenameOrDirectory, SubDir, Flags);
+			}
+			return true;
+		});
+	}
 }
 
 TArray<FString> FPaths::GetExtensionDirs(const FString& BaseDir, const FString& SubDir, bool bCheckValid)
 {
-	using namespace UE::Paths;
-
-	EGetExtensionDirsFlags Flags = EGetExtensionDirsFlags::WithPlatforms;
-	if (bCheckValid)
+	EGetExtensionDirsFlags Flags = EGetExtensionDirsFlags::Default;
+	if (!bCheckValid)
 	{
-		Flags |= EGetExtensionDirsFlags::CheckValid;
+		Flags &= ~EGetExtensionDirsFlags::CheckPlatformsValid;
 	}
 
 	TArray<FString> ExtensionDirs;
@@ -449,10 +442,10 @@ TArray<FString> FPaths::GetExtensionDirs(const FString& BaseDir, const FString&
 	return ExtensionDirs;
 }
 
-TArray<FString> FPaths::GetNonPlatformExtensionDirs(const FString& BaseDir, const FString& SubDir)
+TArray<FString> FPaths::GetExtensionDirs(const FString& BaseDir, const FString& SubDir, EGetExtensionDirsFlags Flags)
 {
 	TArray<FString> ExtensionDirs;
-	GetExtensionDirsInternal(ExtensionDirs, BaseDir, SubDir, UE::Paths::EGetExtensionDirsFlags::None);
+	GetExtensionDirsInternal(ExtensionDirs, BaseDir, SubDir, Flags);
 	return ExtensionDirs;
 }
 
diff --git a/Engine/Source/Runtime/Core/Public/Misc/Paths.h b/Engine/Source/Runtime/Core/Public/Misc/Paths.h
--- a/Engine/Source/Runtime/Core/Public/Misc/Paths.h
+++ b/Engine/Source/Runtime/Core/Public/Misc/Paths.h
@@ -226,14 +226,32 @@ public:
 	 *
 	 * @return BaseDir and usable extension directories under BaseDir (either Engine or Project)
 	 */
-	static CORE_API TArray<FString> GetExtensionDirs(const FString& BaseDir, const FString& SubDir=FString(), bool bCheckValid=true);
+	static CORE_API TArray<FString> GetExtensionDirs(const FString& BaseDir, const FString& SubDir, bool bCheckValid);
+
+	enum class EGetExtensionDirsFlags : uint8
+	{
+		None						= 0,
+		// Include the Base directory
+		WithBase					= (1 << 0),
+		// Include Restricted extension directories
+		WithRestricted				= (1 << 1),
+		// Include Platform extension directories
+		WithPlatforms				= (1 << 2),
+		// Use FDataDrivePlatformInfo to determine valid platforms, otherwise just use what's present
+		CheckPlatformsValid			= (1 << 3),
+		Default						= 0xFF
+	};
+	FRIEND_ENUM_CLASS_FLAGS(EGetExtensionDirsFlags);
 
 	/**
 	 * Returns restricted extensions that are present
 	 *
 	 * @return BaseDir and usable extension directories under BaseDir (either Engine or Project)
 	 */
-	static CORE_API TArray<FString> GetNonPlatformExtensionDirs(const FString& BaseDir, const FString& SubDir = FString());
+	static CORE_API TArray<FString> GetExtensionDirs(
+		const FString& BaseDir, 
+		const FString& SubDir = FString(), 
+		EGetExtensionDirsFlags Flags = EGetExtensionDirsFlags::Default);
 
 	/**
 	 * Returns the root directory of the engine directory tree
@@ -752,3 +770,5 @@ private:
 	/** Returns, if any, the value of the -shaderworkingdir command line argument. This can be used to sandbox shader working files to a desired location */
 	static CORE_API const FString& CustomShaderDirArgument();
 };
+
+ENUM_CLASS_FLAGS(FPaths::EGetExtensionDirsFlags);
