Update commetns

diff --git a/Engine/Source/Runtime/Core/Public/Misc/Paths.h b/Engine/Source/Runtime/Core/Public/Misc/Paths.h
--- a/Engine/Source/Runtime/Core/Public/Misc/Paths.h
+++ b/Engine/Source/Runtime/Core/Public/Misc/Paths.h
@@ -222,7 +222,7 @@ public:
 	
 	/**
 	 * Returns platform and restricted extensions that are present and if bCheckValid is set, valid
-     * (for platforms, it uses FDataDrivePlatformInfo to determine valid platforms, it doesn't just use what's present)
+     * (for platforms, it uses FDataDrivenPlatformInfo to determine valid platforms, it doesn't just use what's present)
 	 *
 	 * @return BaseDir and usable extension directories under BaseDir (either Engine or Project)
 	 */
@@ -244,7 +244,8 @@ public:
 	FRIEND_ENUM_CLASS_FLAGS(EGetExtensionDirsFlags);
 
 	/**
-	 * Returns restricted extensions that are present
+	 * Returns platform and restricted extensions that are present and if CheckPlatformsValid is set, valid
+	 * (for platforms, it uses FDataDrivenPlatformInfo to determine valid platforms, it doesn't just use what's present)
 	 *
 	 * @return BaseDir and usable extension directories under BaseDir (either Engine or Project)
 	 */
