Add a general-purpose FStringOutputDevice which filters based on the log category.

In a followup CL, I intend to reuse this output device in a different test, and using a different log category.

diff --git a/Engine/Source/Runtime/Core/Public/Misc/StringOutputDevice.h b/Engine/Source/Runtime/Core/Public/Misc/StringOutputDevice.h
--- a/Engine/Source/Runtime/Core/Public/Misc/StringOutputDevice.h
+++ b/Engine/Source/Runtime/Core/Public/Misc/StringOutputDevice.h
@@ -4,6 +4,7 @@
 
 #include "Containers/UnrealString.h"
 #include "Misc/OutputDevice.h"
+#include "Templates/Tuple.h"
 #include "Traits/IsContiguousContainer.h"
 
 class FName;
@@ -160,3 +161,42 @@ struct TIsContiguousContainer<FStringOutputDeviceCountLines>
 {
 	enum { Value = true };
 };
+
+/**
+ * Serializes log lines of the requested categories to a string; discards other log lines.
+ */
+template <typename... LogCategoryTypes>
+class TLogCategoryOutputDevice : public FStringOutputDevice
+{
+public:
+	explicit TLogCategoryOutputDevice(LogCategoryTypes&... InLogCategories UE_LIFETIMEBOUND)
+		: LogCategories(&InLogCategories...)
+	{
+		bAutoEmitLineTerminator = true;
+	}
+
+	virtual void Serialize(const TCHAR* InData, ELogVerbosity::Type Verbosity, const FName& Category) override
+	{
+		bool bFoundMatchingCategory = false;
+		VisitTupleElements(
+			[&bFoundMatchingCategory, &Category](auto& Elem)
+			{
+				bFoundMatchingCategory = bFoundMatchingCategory || (Category == Elem->GetCategoryName());
+			},
+			LogCategories);
+
+		if (bFoundMatchingCategory)
+		{
+			FStringOutputDevice::Serialize(InData, Verbosity, Category);
+		}
+	}
+
+private:
+	TTuple<LogCategoryTypes*...> LogCategories;
+};
+
+template <typename... LogCategoryTypes>
+struct TIsContiguousContainer<TLogCategoryOutputDevice<LogCategoryTypes...>>
+{
+	enum { Value = true };
+};
