Fix errors with AUTORTFM_DUMP_CALLSTACK and improve its output.

Previously, the macro was never used anywhere and didn't actually generate a callstack.
Also, AutoRTFM call stacks would previously place the user message below the callstack and in a different log category:

LogOutputDevice: AutoRTFM Backtrace
LogOutputDevice:
LogOutputDevice: Foo()
LogOutputDevice: Bar()
LogOutputDevice:
LogAutoRTFM: User Message Here

Now the user message comes first and is in the same category:

LogOutputDevice: User Message Here
LogOutputDevice:
LogOutputDevice: Foo()
LogOutputDevice: Bar()
LogOutputDevice:

diff --git a/Engine/Source/Runtime/AutoRTFM/Private/Utils.h b/Engine/Source/Runtime/AutoRTFM/Private/Utils.h
--- a/Engine/Source/Runtime/AutoRTFM/Private/Utils.h
+++ b/Engine/Source/Runtime/AutoRTFM/Private/Utils.h
@@ -16,7 +16,7 @@
 // If 0 then verbose logging is compiled-out.
 #define AUTORTFM_VERBOSE_ENABLED 0
 
-// AutoRTFM internal-runtime-only attribute that can be applied to classes, 
+// AutoRTFM internal-runtime-only attribute that can be applied to classes,
 // methods or functions to prevent AutoRTFM closed function(s) from being
 // generated. Applying the annotation to a class is equivalent to applying the
 // annotation to each method of the class.
@@ -31,7 +31,7 @@ namespace AutoRTFM
 {
 
 // Reports an internal AutoRTFM issue. The behavior of this function will vary
-// based on ForTheRuntime::GetInternalAbortAction() and 
+// based on ForTheRuntime::GetInternalAbortAction() and
 // ForTheRuntime::GetEnsureOnInternalAbort().
 // If ProgramCounter is null, then the return address of the call will be used.
 AUTORTFM_DISABLE __attribute__((__format__(__printf__, 4, 5)))
@@ -236,7 +236,7 @@ inline constexpr T Lerp(T Start, T End, T Fraction)
 		}(); \
 	} } while(false)
 
-#define AUTORTFM_DUMP_STACKTRACE(Message) ::AutoRTFM::Log(Message, ##__VA_ARGS__)
+#define AUTORTFM_DUMP_STACKTRACE(Message, ...) ::AutoRTFM::LogWithCallstack(autortfm_log_error, Message, ##__VA_ARGS__)
 
 // This is all a bit funky - but for good reason! We want `AUTORTFM_ASSERT`
 // to be as close to *zero* cost if the assert wouldn't trigger.
diff --git a/Engine/Source/Runtime/Core/Private/AutoRTFM/AutoRTFMUE.cpp b/Engine/Source/Runtime/Core/Private/AutoRTFM/AutoRTFMUE.cpp
--- a/Engine/Source/Runtime/Core/Private/AutoRTFM/AutoRTFMUE.cpp
+++ b/Engine/Source/Runtime/Core/Private/AutoRTFM/AutoRTFMUE.cpp
@@ -90,7 +90,7 @@ static FAutoConsoleVariable CVarAutoRTFMRuntimeEnabled(
 	TEXT("default"),
 	TEXT("Enables the AutoRTFM runtime"),
 	FConsoleVariableDelegate::CreateLambda([] (IConsoleVariable* Variable)
-	{ 
+	{
 		FString Value = Variable->GetString().ToLower();
 		if (Value == "default")
 		{
@@ -275,23 +275,25 @@ void AutoRTFMLog(const char* File, int Line, void* ProgramCounter, autortfm_log_
 		/* DynamicData */ &DynamicData
 	);
 
+	auto Message = FmtToTChar(Format, Args);
+
 	switch (Severity)
 	{
 		case autortfm_log_verbose:
-			UE::Logging::Private::BasicLog<ELogVerbosity::Verbose>(Record, &LogAutoRTFM, FmtToTChar(Format, Args).Get());
+			UE::Logging::Private::BasicLog<ELogVerbosity::Verbose>(Record, &LogAutoRTFM, Message.Get());
 			break;
 		case autortfm_log_info:
-			UE::Logging::Private::BasicLog<ELogVerbosity::Display>(Record, &LogAutoRTFM, FmtToTChar(Format, Args).Get());
+			UE::Logging::Private::BasicLog<ELogVerbosity::Display>(Record, &LogAutoRTFM, Message.Get());
 			break;
 		case autortfm_log_warn:
-			UE::Logging::Private::BasicLog<ELogVerbosity::Warning>(Record, &LogAutoRTFM, FmtToTChar(Format, Args).Get());
+			UE::Logging::Private::BasicLog<ELogVerbosity::Warning>(Record, &LogAutoRTFM, Message.Get());
 			break;
 		case autortfm_log_error:
-			UE::Logging::Private::BasicLog<ELogVerbosity::Error>(Record, &LogAutoRTFM, FmtToTChar(Format, Args).Get());
+			UE::Logging::Private::BasicLog<ELogVerbosity::Error>(Record, &LogAutoRTFM, Message.Get());
 			break;
 		case autortfm_log_fatal:
 			FDebug::DumpStackTraceToLog(TEXT("AutoRTFM backtrace"), ELogVerbosity::Error);
-			UE::Logging::Private::BasicFatalLogWithProgramCounter(Record, &LogAutoRTFM, ProgramCounter, FmtToTChar(Format, Args).Get());
+			UE::Logging::Private::BasicFatalLogWithProgramCounter(Record, &LogAutoRTFM, ProgramCounter, Message.Get());
 			break;
 	}
 #endif
@@ -306,35 +308,24 @@ void AutoRTFMLogWithCallstack(void* ProgramCounter, autortfm_log_severity Severi
 		return;
 	}
 
-	static ::UE::Logging::Private::FStaticBasicLogDynamicData DynamicData;
-	UE::Logging::Private::FStaticBasicLogRecord Record(
-		/* Format */ TEXT("%s"),
-		/* File */ "",
-		/* Line */ 0,
-		/* DynamicData */ &DynamicData
-	);
+	auto Message = FmtToTChar(Format, Args);
 
 	switch (Severity)
 	{
 		case autortfm_log_verbose:
-			FDebug::DumpStackTraceToLog(TEXT("AutoRTFM backtrace"), ELogVerbosity::Verbose);
-			UE::Logging::Private::BasicLog<ELogVerbosity::Verbose>(Record, &LogAutoRTFM, FmtToTChar(Format, Args).Get());
+			FDebug::DumpStackTraceToLog(Message.Get(), ELogVerbosity::Verbose);
 			break;
 		case autortfm_log_info:
-			FDebug::DumpStackTraceToLog(TEXT("AutoRTFM backtrace"), ELogVerbosity::Display);
-			UE::Logging::Private::BasicLog<ELogVerbosity::Display>(Record, &LogAutoRTFM, FmtToTChar(Format, Args).Get());
+			FDebug::DumpStackTraceToLog(Message.Get(), ELogVerbosity::Display);
 			break;
 		case autortfm_log_warn:
-			FDebug::DumpStackTraceToLog(TEXT("AutoRTFM backtrace"), ELogVerbosity::Warning);
-			UE::Logging::Private::BasicLog<ELogVerbosity::Warning>(Record, &LogAutoRTFM, FmtToTChar(Format, Args).Get());
+			FDebug::DumpStackTraceToLog(Message.Get(), ELogVerbosity::Warning);
 			break;
 		case autortfm_log_error:
-			FDebug::DumpStackTraceToLog(TEXT("AutoRTFM backtrace"), ELogVerbosity::Error);
-			UE::Logging::Private::BasicLog<ELogVerbosity::Error>(Record, &LogAutoRTFM, FmtToTChar(Format, Args).Get());
+			FDebug::DumpStackTraceToLog(Message.Get(), ELogVerbosity::Error);
 			break;
 		case autortfm_log_fatal:
-			FDebug::DumpStackTraceToLog(TEXT("AutoRTFM backtrace"), ELogVerbosity::Error);
-			UE::Logging::Private::BasicFatalLogWithProgramCounter(Record, &LogAutoRTFM, ProgramCounter, FmtToTChar(Format, Args).Get());
+			FDebug::DumpStackTraceToLog(Message.Get(), ELogVerbosity::Fatal);
 			break;
 	}
 #endif
@@ -400,7 +391,7 @@ CORE_API void AutoRTFM::InitializeForUE()
 	ExternAPI.OnMemoryValidationStatisticsChanged = OnAutoRTFMMemoryValidationStatisticsChanged;
 
 	AutoRTFM::ForTheRuntime::Initialize(ExternAPI);
-	
+
 	// Call the OnAutoRTFMXXXChanged() handlers now so that values are logged
 	// and crash context data is set.
 	OnAutoRTFMRuntimeEnabledChanged();
