Integrating:
Make the ObjectPtr and native pointer LinkedListBuilder to share a common base.  LinkedListBuilder had a bug resolved but we were unaware of the duplicated code in the ObjectPtr version.

#rnx

diff --git a/Engine/Source/Runtime/Core/Public/Containers/LinkedListBuilder.h b/Engine/Source/Runtime/Core/Public/Containers/LinkedListBuilder.h
--- a/Engine/Source/Runtime/Core/Public/Containers/LinkedListBuilder.h
+++ b/Engine/Source/Runtime/Core/Public/Containers/LinkedListBuilder.h
@@ -28,19 +28,21 @@ struct TLinkedListBuilderNextLinkMemberVar
 };
 
 /**
- * Single linked list builder
+ * Single linked list builder base
  */
-template<typename InElementType, typename InLinkAccessor = TLinkedListBuilderNextLink<InElementType>>
-struct TLinkedListBuilder
+template<typename InElementType, typename InPointerType, typename InLinkAccessor>
+struct TLinkedListBuilderBase
 {
 public:
 	using ElementType = InElementType;
+	using PointerType = InPointerType;
 	using LinkAccessor = InLinkAccessor;
 
-	UE_NONCOPYABLE(TLinkedListBuilder);
+	UE_NONCOPYABLE(TLinkedListBuilderBase);
 
 private:
-	inline void WriteEndPtr(InElementType* NewValue)
+
+	inline void WriteEndPtr(PointerType NewValue)
 	{
 		// Do not overwrite the same value to avoid dirtying the cache and
 		// also prevent TSAN from thinking we are messing around with existing data.
@@ -49,9 +51,10 @@ private:
 			*EndPtr = NewValue;
 		}
 	}
+
 public:
 
-	[[nodiscard]] explicit TLinkedListBuilder(ElementType** ListStartPtr) :
+	[[nodiscard]] explicit TLinkedListBuilderBase(PointerType* ListStartPtr) :
 		StartPtr(ListStartPtr),
 		EndPtr(ListStartPtr)
 	{
@@ -90,11 +93,11 @@ private:
 	// Helper method to remove the element pointed to by CurrentLinkPtr.
 	// CurrentLinkPtr is either a pointer to the list start pointer or the 
 	// next pointer in the previous element
-	inline void RemoveImpl(ElementType** CurrentLinkPtr)
+	inline void RemoveImpl(PointerType* CurrentLinkPtr)
 	{
 
 		// This gets the address of the next pointer of the element being removed
-		ElementType** NextLinkPtr = LinkAccessor::GetNextPtr(**CurrentLinkPtr);
+		PointerType* NextLinkPtr = LinkAccessor::GetNextPtr(**CurrentLinkPtr);
 
 		// If EndPtr points to the next pointer of the element being removed,
 		// then move the EndPtr back one element so it remains valid.
@@ -120,7 +123,7 @@ public:
 
 		// With the single link list, CurrentLinkPtr will always point to the address of the
 		// variable that points to the current element to be tested, not the element itself.
-		for (ElementType** CurrentLinkPtr = StartPtr; *CurrentLinkPtr;)
+		for (PointerType* CurrentLinkPtr = StartPtr; *CurrentLinkPtr;)
 		{
 			if (Predicate(*CurrentLinkPtr))
 			{
@@ -140,7 +143,7 @@ public:
 	{
 		// With the single link list, CurrentLinkPtr will always point to the address of the
 		// variable that points to the current element to be tested, not the element itself.
-		for (ElementType** CurrentLinkPtr = StartPtr; *CurrentLinkPtr; CurrentLinkPtr = LinkAccessor::GetNextPtr(**CurrentLinkPtr))
+		for (PointerType* CurrentLinkPtr = StartPtr; *CurrentLinkPtr; CurrentLinkPtr = LinkAccessor::GetNextPtr(**CurrentLinkPtr))
 		{
 			if (*CurrentLinkPtr == &Element)
 			{
@@ -158,7 +161,7 @@ public:
 
 	inline void MoveToEnd()
 	{
-		for (ElementType* It = *StartPtr; It; It = GetNext(*It))
+		for (PointerType It = *StartPtr; It; It = GetNext(*It))
 		{
 			EndPtr = LinkAccessor::GetNextPtr(*It);
 		}
@@ -175,27 +178,40 @@ public:
 		return false;
 	}
 
-	[[nodiscard]] UE_FORCEINLINE_HINT ElementType* GetNext(ElementType& Element) const
+	[[nodiscard]] UE_FORCEINLINE_HINT PointerType GetNext(ElementType& Element) const
 	{
 		return GetNextRef(Element);
 	}
 
-	[[nodiscard]] UE_FORCEINLINE_HINT ElementType* GetListStart() const
+	[[nodiscard]] UE_FORCEINLINE_HINT PointerType GetListStart() const
 	{
 		return *StartPtr;
 	}
 
-	[[nodiscard]] UE_FORCEINLINE_HINT ElementType* GetListEnd() const
+	[[nodiscard]] UE_FORCEINLINE_HINT PointerType GetListEnd() const
 	{
 		return *EndPtr;
 	}
 
 private:
-	[[nodiscard]] UE_FORCEINLINE_HINT ElementType*& GetNextRef(ElementType& Element) const
+	[[nodiscard]] UE_FORCEINLINE_HINT PointerType& GetNextRef(ElementType& Element) const
 	{
 		return *LinkAccessor::GetNextPtr(Element);
 	}
 
-	ElementType** StartPtr;
-	ElementType** EndPtr;
+	PointerType* StartPtr;
+	PointerType* EndPtr;
+};
+
+/**
+ * Single linked list builder for raw pointers
+ */
+
+template<typename InElementType, typename InLinkAccessor = TLinkedListBuilderNextLink<InElementType>>
+struct TLinkedListBuilder : public TLinkedListBuilderBase<InElementType, InElementType*, InLinkAccessor>
+{
+	using Super = TLinkedListBuilderBase<InElementType, InElementType*, InLinkAccessor>;
+	using Super::Super;
+
+	UE_NONCOPYABLE(TLinkedListBuilder);
 };
diff --git a/Engine/Source/Runtime/CoreUObject/Public/UObject/ObjectPtrLinkedListBuilder.h b/Engine/Source/Runtime/CoreUObject/Public/UObject/ObjectPtrLinkedListBuilder.h
--- a/Engine/Source/Runtime/CoreUObject/Public/UObject/ObjectPtrLinkedListBuilder.h
+++ b/Engine/Source/Runtime/CoreUObject/Public/UObject/ObjectPtrLinkedListBuilder.h
@@ -22,126 +22,19 @@ struct TObjectPtrLinkedListBuilderNextLink
  * We wish to migrate the links to TObjectPtr.
  */
 template<typename InElementType, typename InLinkAccessor = TObjectPtrLinkedListBuilderNextLink<InElementType>>
-struct TObjectPtrLinkedListBuilder
+struct TObjectPtrLinkedListBuilder : public TLinkedListBuilderBase<InElementType, TObjectPtr<InElementType>, InLinkAccessor>
 {
-public:
-	using ElementType = InElementType;
-	using LinkAccessor = InLinkAccessor;
+	using Super = TLinkedListBuilderBase<InElementType, TObjectPtr<InElementType>, InLinkAccessor>;
+	using ElementType = Super::ElementType;
+	using PointerType = Super::PointerType;
+	using Super::Super;
 
 	UE_NONCOPYABLE(TObjectPtrLinkedListBuilder);
 
-private:
-	UE_FORCEINLINE_HINT void WriteEndPtr(InElementType* NewValue)
-	{
-		// Do not overwrite the same value to avoid dirtying the cache and
-		// also prevent TSAN from thinking we are messing around with existing data.
-		if (*EndPtr != NewValue)
-		{
-			*EndPtr = NewValue;
-		}
-	}
-public:
-
-    // Allow construction from raw ptr because call sites were handling the type mismatch between
-    // the head and next ptrs in lists by converting the TObjectPtr to a raw ptr with MutableView.
-	[[nodiscard]] explicit TObjectPtrLinkedListBuilder(ElementType** ListStartPtr) 
-        : StartPtr(reinterpret_cast<TObjectPtr<ElementType>*>(ListStartPtr))
-        , EndPtr(StartPtr)
-	{
-		check(ListStartPtr);
-	}
-
-	[[nodiscard]] explicit TObjectPtrLinkedListBuilder(TObjectPtr<ElementType>* ListStartPtr) 
-        : StartPtr(ListStartPtr)
-        , EndPtr(ListStartPtr)
-	{
-		check(ListStartPtr);
-	}
-
-	// Move builder back to start and prepare for overwriting
-	// It only changes state of builder, use NullTerminate() to mark list as empty!
-	UE_FORCEINLINE_HINT void Restart()
-	{
-		EndPtr = StartPtr;
-	}
-
-	// Append element, don't touch next link
-	UE_FORCEINLINE_HINT void AppendNoTerminate(ElementType& Element)
-	{
-		WriteEndPtr(&Element);
-		EndPtr = LinkAccessor::GetNextPtr(Element);
-	}
-
-	// Append element and mark it as last
-	UE_FORCEINLINE_HINT void AppendTerminated(ElementType& Element)
+	// Allow construction from raw ptr because call sites were handling the type mismatch between
+	// the head and next ptrs in lists by converting the TObjectPtr to a raw ptr with MutableView.
+	[[nodiscard]] explicit TObjectPtrLinkedListBuilder(ElementType** ListStartPtr)
+		: Super(reinterpret_cast<PointerType*>(ListStartPtr))
 	{
-		AppendNoTerminate(Element);
-		NullTerminate();
 	}
-
-	UE_FORCEINLINE_HINT void Remove(ElementType& Element)
-	{
-		TObjectPtr<ElementType>* PrevIt = StartPtr;
-		for (TObjectPtr<ElementType> It = *StartPtr; It; It = GetNext(*It))
-		{
-			if (It == &Element)
-			{
-				TObjectPtr<ElementType>& ElementNextRef = GetNextRef(Element);
-				*PrevIt = ElementNextRef;
-				ElementNextRef = nullptr;
-				break;
-			}
-
-			PrevIt = &It;
-		}
-	}
-
-	// Mark end of the list
-	UE_FORCEINLINE_HINT void NullTerminate()
-	{
-		WriteEndPtr(nullptr);
-	}
-
-	UE_FORCEINLINE_HINT void MoveToEnd()
-	{
-		for (TObjectPtr<ElementType> It = *StartPtr; It; It = GetNext(*It))
-		{
-			EndPtr = LinkAccessor::GetNextPtr(*It);
-		}
-	}
-
-	UE_FORCEINLINE_HINT bool MoveToNext()
-	{
-		if (*EndPtr)
-		{
-			EndPtr = LinkAccessor::GetNextPtr(**EndPtr);
-			return true;
-		}
-
-		return false;
-	}
-
-	[[nodiscard]] UE_FORCEINLINE_HINT ElementType* GetNext(ElementType& Element) const
-	{
-		return GetNextRef(Element);
-	}
-
-	[[nodiscard]] UE_FORCEINLINE_HINT ElementType* GetListStart() const
-	{
-		return *StartPtr;
-	}
-
-	[[nodiscard]] UE_FORCEINLINE_HINT ElementType* GetListEnd() const
-	{
-		return *EndPtr;
-	}
-
-private:
-	[[nodiscard]] UE_FORCEINLINE_HINT TObjectPtr<ElementType>& GetNextRef(ElementType& Element) const
-	{
-		return *LinkAccessor::GetNextPtr(Element);
-	}
-
-	TObjectPtr<ElementType>* StartPtr;
-	TObjectPtr<ElementType>* EndPtr;
 };
\ No newline at end of file
