Make some core container types constexpr default constructible even when using inline allocator.
TArray, TBitArray, TSparseArray, TCompactSet, TSparseSet, TSet, TMap

diff --git a/Engine/Source/Runtime/Core/Public/Containers/Array.h b/Engine/Source/Runtime/Core/Public/Containers/Array.h
--- a/Engine/Source/Runtime/Core/Public/Containers/Array.h
+++ b/Engine/Source/Runtime/Core/Public/Containers/Array.h
@@ -691,7 +691,7 @@ public:
 	/**
 	 * Constructor, initializes element number counters.
 	 */
-	[[nodiscard]] UE_FORCEINLINE_HINT TArray()
+	[[nodiscard]] UE_FORCEINLINE_HINT constexpr TArray()
 		: ArrayNum(0)
 		, ArrayMax(AllocatorInstance.GetInitialCapacity())
 	{
diff --git a/Engine/Source/Runtime/Core/Public/Containers/BitArray.h b/Engine/Source/Runtime/Core/Public/Containers/BitArray.h
--- a/Engine/Source/Runtime/Core/Public/Containers/BitArray.h
+++ b/Engine/Source/Runtime/Core/Public/Containers/BitArray.h
@@ -370,7 +370,7 @@ public:
 	template<typename,typename,bool>
 	friend class TConstDualSetBitIterator;
 
-	[[nodiscard]] TBitArray()
+	[[nodiscard]] constexpr TBitArray()
 	:	NumBits(0)
 	,	MaxBits(AllocatorInstance.GetInitialCapacity() * NumBitsPerDWORD)
 	{
diff --git a/Engine/Source/Runtime/Core/Public/Containers/CompactSet.h.inl b/Engine/Source/Runtime/Core/Public/Containers/CompactSet.h.inl
--- a/Engine/Source/Runtime/Core/Public/Containers/CompactSet.h.inl
+++ b/Engine/Source/Runtime/Core/Public/Containers/CompactSet.h.inl
@@ -303,7 +303,7 @@ public:
 
 	/* Constructors */
 
-	[[nodiscard]] UE_FORCEINLINE_HINT UE_TCOMPACT_SET() = default;
+	[[nodiscard]] UE_FORCEINLINE_HINT constexpr UE_TCOMPACT_SET() = default;
 
 	[[nodiscard]] explicit consteval UE_TCOMPACT_SET(EConstEval)
 		: Super(ConstEval)
diff --git a/Engine/Source/Runtime/Core/Public/Containers/ContainerAllocationPolicies.h b/Engine/Source/Runtime/Core/Public/Containers/ContainerAllocationPolicies.h
--- a/Engine/Source/Runtime/Core/Public/Containers/ContainerAllocationPolicies.h
+++ b/Engine/Source/Runtime/Core/Public/Containers/ContainerAllocationPolicies.h
@@ -634,7 +634,7 @@ public:
 
 	public:
 		/** Default constructor. */
-		ForAnyElementType()
+		constexpr ForAnyElementType()
 			: Data(nullptr)
 		{}
 
@@ -907,8 +907,15 @@ public:
 	public:
 
 		/** Default constructor. */
-		ForElementType()
+		constexpr ForElementType()
 		{
+			UE_IF_CONSTEVAL
+			{
+				for (uint8& Byte : InlineData)
+				{
+					Byte = 0;
+				}
+			}
 		}
 
 		explicit consteval ForElementType(EConstEval)
diff --git a/Engine/Source/Runtime/Core/Public/Containers/Map.h.inl b/Engine/Source/Runtime/Core/Public/Containers/Map.h.inl
--- a/Engine/Source/Runtime/Core/Public/Containers/Map.h.inl
+++ b/Engine/Source/Runtime/Core/Public/Containers/Map.h.inl
@@ -37,7 +37,7 @@ public:
 	typedef TPair<KeyType, ValueType> ElementType;
 
 protected:
-	[[nodiscard]] TMAPBASE() = default;
+	[[nodiscard]] constexpr TMAPBASE() = default;
 	[[nodiscard]] explicit consteval TMAPBASE(EConstEval)
 	: Pairs(ConstEval)
 	{
@@ -1054,7 +1054,7 @@ class TSORTABLEMAPBASE : public TMAPBASE<KeyType, ValueType, SetAllocator, KeyFu
 protected:
 	typedef TMAPBASE<KeyType, ValueType, SetAllocator, KeyFuncs> Super;
 
-	[[nodiscard]] TSORTABLEMAPBASE() = default;
+	[[nodiscard]] constexpr TSORTABLEMAPBASE() = default;
 	[[nodiscard]] explicit consteval TSORTABLEMAPBASE(EConstEval)
 	: Super(ConstEval)
 	{
@@ -1223,7 +1223,7 @@ public:
 	typedef typename Super::KeyInitType KeyInitType;
 	typedef typename Super::KeyConstPointerType KeyConstPointerType;
 
-	[[nodiscard]] TMAP() = default;
+	[[nodiscard]] constexpr TMAP() = default;
 	[[nodiscard]] explicit consteval TMAP(EConstEval)
 	: Super(ConstEval)
 	{
@@ -1450,7 +1450,7 @@ public:
 	typedef typename Super::KeyInitType KeyInitType;
 	typedef typename Super::ValueInitType ValueInitType;
 
-	[[nodiscard]] TMULTIMAP() = default;
+	[[nodiscard]] constexpr TMULTIMAP() = default;
 	[[nodiscard]] TMULTIMAP(TMULTIMAP&&) = default;
 	[[nodiscard]] TMULTIMAP(const TMULTIMAP&) = default;
 	TMULTIMAP& operator=(TMULTIMAP&&) = default;
diff --git a/Engine/Source/Runtime/Core/Public/Containers/SparseArray.h b/Engine/Source/Runtime/Core/Public/Containers/SparseArray.h
--- a/Engine/Source/Runtime/Core/Public/Containers/SparseArray.h
+++ b/Engine/Source/Runtime/Core/Public/Containers/SparseArray.h
@@ -70,7 +70,7 @@ protected:
 	}
 
 	/** Default constructor. */
-	[[nodiscard]] TSparseArrayBase()
+	[[nodiscard]] constexpr TSparseArrayBase()
 		: FirstFreeIndex(-1)
 		, NumFreeIndices(0)
 	{
@@ -863,7 +863,7 @@ public:
 	}
 
 	/** Default constructor. */
-	[[nodiscard]] TSparseArray() = default;
+	[[nodiscard]] constexpr TSparseArray() = default;
 
 	[[nodiscard]] explicit consteval TSparseArray(EConstEval)
 		: SuperType(ConstEval)
diff --git a/Engine/Source/Runtime/Core/Public/Containers/SparseSet.h.inl b/Engine/Source/Runtime/Core/Public/Containers/SparseSet.h.inl
--- a/Engine/Source/Runtime/Core/Public/Containers/SparseSet.h.inl
+++ b/Engine/Source/Runtime/Core/Public/Containers/SparseSet.h.inl
@@ -73,7 +73,7 @@ private:
 
 public:
 	/** Initialization constructor. */
-	[[nodiscard]] UE_FORCEINLINE_HINT UE_TSPARSE_SET() = default;
+	[[nodiscard]] UE_FORCEINLINE_HINT constexpr UE_TSPARSE_SET() = default;
 
 	[[nodiscard]] explicit consteval UE_TSPARSE_SET(EConstEval)
 	: Elements(ConstEval)
