Increase maximum FName storage in the editor to 4gb by aligning stored names to 8 bytes rather than 4.
Adds define UE_FNAME_ENTRY_ALIGNMENT, defaulted to 8 in editor and 0 otherwise.
Fix return type of name memory static functions to be int64 to avoid overflow.
Remove use of FNameEntry alignment from FNameStore and make that the responsibility of FNameAllocator.
Store amount of space wasted by alignment.

#FailedValidations Preflight

diff --git a/Engine/Source/Runtime/Core/Private/UObject/UnrealNames.cpp b/Engine/Source/Runtime/Core/Private/UObject/UnrealNames.cpp
--- a/Engine/Source/Runtime/Core/Private/UObject/UnrealNames.cpp
+++ b/Engine/Source/Runtime/Core/Private/UObject/UnrealNames.cpp
@@ -253,6 +253,7 @@ static constexpr uint32 FNameEntryIdMask = (1 << FNameEntryIdBits ) - 1;
 struct FDefaultNameStore
 {
 	// Decides if entry header + Name data fits in OutSize. Called inside allocator and shard lock.
+	// Implementatations should not consider alignment of FNameEntry, that's the responsibility of the caller
 	//
 	// @return Total entry aligned size if Name fits, 0 otherwise
 	static uint32 TryPlace(uint8* OutEntry, uint32 OutSize, FNameStringView Name)
@@ -297,7 +298,7 @@ private:
 	{
 		uint32 HeaderSize = FNameEntry::GetDataOffset();
 		uint32 FooterSize =  Len * (bIsWide ? sizeof(WIDECHAR) : sizeof(ANSICHAR));
-		return Align(HeaderSize + FooterSize, alignof(FNameEntry));
+		return HeaderSize + FooterSize;
 	}
 };
 
@@ -584,13 +585,15 @@ public:
 	// @param Bytes is aligned and fits within current block
 	inline FNameEntryHandle Allocate(uint32 Bytes)
 	{
-		checkSlow(Bytes % Stride == 0);
 		checkSlow(CurrentByteCursor % Stride == 0);
 		checkSlow(CurrentByteCursor + Bytes <= BlockSizeBytes);
  
 		// Use current cursor position for this allocation and increment cursor for next allocation
 		uint32 ByteOffset = CurrentByteCursor;
-		CurrentByteCursor += Bytes;
+		// Align up so byte offset can be encoded in FNameEntryId and store waste 
+		uint32 Step = Align(Bytes, alignof(FNameEntry));
+		Waste += Step - Bytes;
+		CurrentByteCursor += Step;
 		return FNameEntryHandle(CurrentBlock, ByteOffset / Stride);
 	}
  
@@ -617,7 +620,7 @@ public:
 	{
 		ScopeLock _(Lock);
  
-		uint32 Bytes = FNameStore::TryPlace(Blocks[CurrentBlock] + CurrentByteCursor, BlockSizeBytes - CurrentByteCursor, Name);
+		uint32 Bytes = BlockSizeBytes < CurrentByteCursor ? 0 : FNameStore::TryPlace(Blocks[CurrentBlock] + CurrentByteCursor, BlockSizeBytes - CurrentByteCursor, Name);
  
 		// Allocate a new block if current one is exhausted. We don't worry about a little bit
 		// of waste at the end given the relative size of block to average and max allocation.
@@ -695,6 +698,11 @@ public:
 		return CurrentBlock + 1;
 	}
 
+	int64 GetWaste() const 
+	{
+		return Waste;
+	}
+
 	uint8** GetBlocksForDebugVisualizer() { return (uint8**)Blocks; }
 
 	void DebugDump(TArray<const FNameEntry*>& Out) const
@@ -849,6 +857,11 @@ private:
 		}
 #endif
 
+		if (CurrentByteCursor < BlockSizeBytes)
+		{
+			Waste += (BlockSizeBytes - CurrentByteCursor);
+		}
+
 		++CurrentBlock;
 		CurrentByteCursor = 0;
 
@@ -872,6 +885,7 @@ private:
 	TSAN_ATOMIC(uint32) CurrentBlock = 0;
 	uint32 CurrentByteCursor = 0;
 	TSAN_ATOMIC(uint8*) Blocks[FNameMaxBlocks] = {};
+	TSAN_ATOMIC(int64) Waste = 0; 
 #if USE_FNAME_MMAP
 	TArray<uint8*> MmappedAddresses;
 	int32 MmapFile = -1;
@@ -2135,7 +2149,14 @@ void FNamePool::RetraceAll() const
 
 void FNamePool::LogStats(FOutputDevice& Ar) const
 {
-	Ar.Logf(TEXT("%i FNames using in %ikB + %ikB"), NumEntries(), sizeof(FNamePool) / 1024, Entries.NumBlocks() * FNameEntryAllocator::BlockSizeBytes / 1024);
+	Ar.Logf(TEXT("%i FNames using in %iKB static memory, %.2fMB / %.2fMB in block storage, %.2fMB in lookup"), 
+		NumEntries(), 
+		sizeof(FNamePool) / 1024, 
+		Entries.NumBlocks() * (int64)FNameEntryAllocator::BlockSizeBytes / 1024.0 / 1024.0,
+		FNameMaxBlocks * (int64)FNameEntryAllocator::BlockSizeBytes / 1024.0 / 1024.0,
+		NumSlots() * sizeof(FNameSlot) / 1024.0 / 1024.0
+	);
+	Ar.Logf(TEXT("%.2fMB of waste from aligning entries"), Entries.GetWaste() / 1024.0 / 1024.0);
 	Ar.Logf(TEXT("%d ansi FNames"), NumAnsiEntries());
 	Ar.Logf(TEXT("%d wide FNames"), NumWideEntries());
 #if UE_FNAME_OUTLINE_NUMBER
@@ -2575,7 +2596,7 @@ int32 FNameEntry::GetSizeInBytes() const
 	else
 #endif
 	{
-		return FNameStore::GetSize(*this);
+		return Align(FNameStore::GetSize(*this), alignof(FNameEntry));
 	}
 }
 
@@ -2624,18 +2645,18 @@ void FName::Reserve(uint32 NumBytes, uint32 NumNames)
 	GetNamePoolPostInit().Reserve(NumBytes, NumNames);
 }
 
-int32 FName::GetNameEntryMemorySize()
+int64 FName::GetNameEntryMemorySize()
 {
-	return GetNamePool().NumBlocks() * FNameEntryAllocator::BlockSizeBytes;
+	return GetNamePool().NumBlocks() * (int64)FNameEntryAllocator::BlockSizeBytes;
 }
 
-int32 FName::GetNameEntryMemoryEstimatedAvailable()
+int64 FName::GetNameEntryMemoryEstimatedAvailable()
 {
 	FNamePool& Pool = GetNamePool();
-	return (FNameMaxBlocks - Pool.NumBlocks()) * FNameEntryAllocator::BlockSizeBytes;
+	return (FNameMaxBlocks - Pool.NumBlocks()) * (int64)FNameEntryAllocator::BlockSizeBytes;
 }
 
-int32 FName::GetNameTableMemorySize()
+int64 FName::GetNameTableMemorySize()
 {
 	return GetNameEntryMemorySize() + sizeof(FNamePool) + GetNamePool().NumSlots() * sizeof(FNameSlot);
 }
diff --git a/Engine/Source/Runtime/Core/Public/UObject/NameTypes.h b/Engine/Source/Runtime/Core/Public/UObject/NameTypes.h
--- a/Engine/Source/Runtime/Core/Public/UObject/NameTypes.h
+++ b/Engine/Source/Runtime/Core/Public/UObject/NameTypes.h
@@ -39,6 +39,18 @@
 	#define UE_FNAME_OUTLINE_NUMBER 0
 #endif // UE_FNAME_OUTLINE_NUMBER
 
+// Allows overriding the alignment of FNameEntry. This decreases the granularity of memory allocation for FNames and 
+// therefore allows more data to be stored as FNames with the potential of increased waste.
+// See the details of FNameEntryId encoding.
+// Setting this to a value less than the natural alignment of FNameEntry will cause a compiler error.
+#ifndef UE_FNAME_ENTRY_ALIGNMENT
+	#if WITH_EDITOR
+		#define UE_FNAME_ENTRY_ALIGNMENT 8
+	#else
+		#define UE_FNAME_ENTRY_ALIGNMENT 0
+	#endif
+#endif // UE_FNAME_ENTRY_ALIGNMENT
+
 class FText;
 
 /** Maximum size of name, including the null terminator. */
@@ -263,7 +275,7 @@ struct FNameEntryHeader
 /**
  * A global deduplicated name stored in the global name table.
  */
-struct FNameEntry
+struct alignas(UE_FNAME_ENTRY_ALIGNMENT) FNameEntry
 {
 private:
 #if WITH_CASE_PRESERVING_NAME
@@ -1126,17 +1138,17 @@ public:
 	/**
 	 * @return Size of all name entries.
 	 */
-	[[nodiscard]] CORE_API static int32 GetNameEntryMemorySize();
+	[[nodiscard]] CORE_API static int64 GetNameEntryMemorySize();
 
 	/**
 	 * @return Estimated remaining size the name entry table is willing to allocate.
 	 */
-	[[nodiscard]] CORE_API static int32 GetNameEntryMemoryEstimatedAvailable();
+	[[nodiscard]] CORE_API static int64 GetNameEntryMemoryEstimatedAvailable();
 
 	/**
 	* @return Size of Name Table object as a whole
 	*/
-	[[nodiscard]] CORE_API static int32 GetNameTableMemorySize();
+	[[nodiscard]] CORE_API static int64 GetNameTableMemorySize();
 
 	/**
 	 * @return number of ansi names in name table
diff --git a/Engine/Source/Runtime/CoreUObject/Private/UObject/UObjectGlobals.cpp b/Engine/Source/Runtime/CoreUObject/Private/UObject/UObjectGlobals.cpp
--- a/Engine/Source/Runtime/CoreUObject/Private/UObject/UObjectGlobals.cpp
+++ b/Engine/Source/Runtime/CoreUObject/Private/UObject/UObjectGlobals.cpp
@@ -855,11 +855,11 @@ void StaticTick( float DeltaTime, bool bUseFullTimeLimit, float AsyncLoadingTime
 	int32 NameTableAnsiEntries = FName::GetNumAnsiNames();
 	int32 NameTableWideEntries = FName::GetNumWideNames();
 	int32 NameTableEntries = NameTableAnsiEntries + NameTableWideEntries;
-	int32 NameTableMemorySize = FName::GetNameTableMemorySize();
+	int64 NameTableMemorySize = FName::GetNameTableMemorySize();
 	SET_DWORD_STAT( STAT_NameTableEntries, NameTableEntries );
 	SET_DWORD_STAT( STAT_NameTableAnsiEntries, NameTableAnsiEntries );
 	SET_DWORD_STAT( STAT_NameTableWideEntries, NameTableWideEntries);
-	SET_DWORD_STAT( STAT_NameTableMemorySize, NameTableMemorySize );
+	SET_MEMORY_STAT( STAT_NameTableMemorySize, NameTableMemorySize );
 
 #if 0 // can't read stats with the new stats system
 	// Set async I/O bandwidth stats.
