Setting NetCL version to CL 46580545

#rnx

diff --git a/Engine/Source/Runtime/Core/Public/Misc/NetworkVersion.h b/Engine/Source/Runtime/Core/Public/Misc/NetworkVersion.h
--- a/Engine/Source/Runtime/Core/Public/Misc/NetworkVersion.h
+++ b/Engine/Source/Runtime/Core/Public/Misc/NetworkVersion.h
@@ -9,7 +9,7 @@
 #include "Misc/EngineNetworkCustomVersion.h"
 
 // The version number used for determining network compatibility. If zero, uses the engine compatible version.
-#define ENGINE_NET_VERSION  0
+#define ENGINE_NET_VERSION  46580545
 
 // The version number used for determining replay compatibility
 #define ENGINE_REPLAY_VERSION  ENGINE_NET_VERSION
