Bughawk comments

diff --git a/Engine/Source/Runtime/Core/Private/UObject/UnrealNames.cpp b/Engine/Source/Runtime/Core/Private/UObject/UnrealNames.cpp
--- a/Engine/Source/Runtime/Core/Private/UObject/UnrealNames.cpp
+++ b/Engine/Source/Runtime/Core/Private/UObject/UnrealNames.cpp
@@ -255,7 +255,7 @@ struct FDefaultNameStore
 	// Decides if entry header + Name data fits in OutSize. Called inside allocator and shard lock.
 	// Implementatations should not consider alignment of FNameEntry, that's the responsibility of the caller
 	//
-	// @return Total entry aligned size if Name fits, 0 otherwise
+	// @return Entry size desired by this name store if Name fits, 0 otherwise. Allocator will align up if necessary.
 	static uint32 TryPlace(uint8* OutEntry, uint32 OutSize, FNameStringView Name)
 	{
 		uint32 Bytes = GetDefaultSize(Name.Len, Name.bIsWide);
@@ -582,7 +582,7 @@ public:
 		}
 	}
  
-	// @param Bytes is aligned and fits within current block
+	// @param Bytes fits within current block
 	inline FNameEntryHandle Allocate(uint32 Bytes)
 	{
 		checkSlow(CurrentByteCursor % Stride == 0);
diff --git a/Engine/Source/Runtime/Core/Public/UObject/NameTypes.h b/Engine/Source/Runtime/Core/Public/UObject/NameTypes.h
--- a/Engine/Source/Runtime/Core/Public/UObject/NameTypes.h
+++ b/Engine/Source/Runtime/Core/Public/UObject/NameTypes.h
@@ -42,7 +42,7 @@
 // Allows overriding the alignment of FNameEntry. This decreases the granularity of memory allocation for FNames and 
 // therefore allows more data to be stored as FNames with the potential of increased waste.
 // See the details of FNameEntryId encoding.
-// Setting this to a value less than the natural alignment of FNameEntry will cause a compiler error.
+// Setting this to a nonzero value less than the natural alignment of FNameEntry will cause a compiler error.
 #ifndef UE_FNAME_ENTRY_ALIGNMENT
 	#if WITH_EDITOR
 		#define UE_FNAME_ENTRY_ALIGNMENT 8
