[Core] More fixing visualizers
* Added compact/sparse maps/set to container tests natvis section
* Fixed natvis to properly iterate set items and resolve values
* Major cleanup to gdb visualier, fixed maps and tpair presentations, added primary asset id
#tests Texted navits using containers test. Tested GDB with a server crash and

diff --git a/Engine/Extras/GDBPrinters/UEPrinters.py b/Engine/Extras/GDBPrinters/UEPrinters.py
--- a/Engine/Extras/GDBPrinters/UEPrinters.py
+++ b/Engine/Extras/GDBPrinters/UEPrinters.py
@@ -37,10 +37,46 @@ else:
 		def next(self):
 			return type(self).__next__(self)
 
+# ------------------------------------------------------------------------------
+# Helper functions
+#
+
 def default_iterator(val):
 	for field in val.type.fields():
 		yield field.name, val[field.name]
 
+def container_iterator(Array):
+	for i in range(Array.ArrayNum):
+		if Array.is_valid(i):
+			yield ('[%d]' % i, Array.get_value(i))
+
+def get_allocation_data(Allocator, ElementType):
+	Data = None
+	try:
+		Data = Allocator['Data']
+	except:
+		pass
+
+	if (Data == None):
+		try:
+			SecondaryData = Allocator['SecondaryData']
+			if SecondaryData != None:
+				SecondaryDataData = SecondaryData['Data']
+				if (SecondaryDataData > 0):
+					Data = SecondaryDataData
+		except:
+			pass
+
+	if (Data == None):
+		try:
+			Data = Allocator['InlineData']
+		except:
+			pass
+
+	if (Data == None):
+		return 0
+
+	return Data.cast(ElementType.pointer())
 
 # ------------------------------------------------------------------------------
 #
@@ -67,60 +103,31 @@ class FBitReferencePrinter:
 class TBitArrayPrinter:
 	"Print TBitArray"
 
-	class _iterator(Iterator):
-		def __init__(self, val):
-			self.Value = val
-			self.Counter = -1
-
-			try:
-				self.NumBits = self.Value['NumBits']
-				if self.NumBits.is_optimized_out:
-					self.NumBits = 0
-				else:
-					self.AllocatorInstance = self.Value['AllocatorInstance']
-					self.InlineData = self.AllocatorInstance['InlineData']
-					self.SecondaryData = self.AllocatorInstance['SecondaryData']
-					self.SecondaryDataData = None
-					if self.SecondaryData != None:
-						self.SecondaryDataData = self.SecondaryData['Data']
-			except:
-				raise
-
-		def __iter__(self):
-			return self
-
-		def __next__(self):
-			if self.NumBits == 0:
-				raise StopIteration
-
-			self.Counter = self.Counter + 1
-
-			if self.Counter >= self.NumBits:
-				raise StopIteration
-
-			if self.SecondaryDataData > 0:
-				data = self.SecondaryDataData.cast(gdb.lookup_type("uint32").pointer())
-			else:
-				data = self.InlineData.cast(gdb.lookup_type("uint32").pointer())
-
-			return ('[%d]' % self.Counter, (data[self.Counter/32] >> self.Counter) & 1)
-
 	def __init__(self, val):
 		self.Value = val
-		self.NumBits = self.Value['NumBits']
+		self.ArrayNum = self.Value['NumBits']
+		self.ArrayMax = self.Value['MaxBits']
+		self.Data = get_allocation_data(self.Value['AllocatorInstance'], gdb.lookup_type("uint32"))
 
 	def to_string(self):
-		if self.NumBits.is_optimized_out:
+		if self.ArrayNum.is_optimized_out:
 			pass
-		if self.NumBits == 0:
+		if self.ArrayNum == 0:
 			return 'empty'
-		pass
+		return 'Num='+str(self.ArrayNum)+' Max='+str(self.ArrayMax)
 
 	def children(self):
-		return self._iterator(self.Value)
+		return container_iterator(self)
 
 	def display_hint(self):
 		return 'array'
+	
+	def get_value(self, index):
+		return ((self.Data[index/32] >> index) & 1) != 0
+		
+	def is_valid(self, index):
+		return index < self.ArrayMax
+
 # ------------------------------------------------------------------------------
 # TIndirectArray
 #
@@ -189,233 +196,99 @@ class TChunkedArrayPrinter:
 	def display_hint(self):
 		return 'array'
 
+
 # ------------------------------------------------------------------------------
-# TSparseArray
+# TArray
 #
-class TSparseArrayPrinter:
-	"Print TSparseArray"
-
-	class _iterator(Iterator):
-		def __init__(self, val):
-
-			self.Value = val
-			self.Counter = -1
-			self.ElementType = self.Value.type.template_argument(0)
-
-			try:
-				self.NumFreeIndices = self.Value['NumFreeIndices']
-				self.Data = self.Value['Data']
-				self.InternalElementType = self.Data.type.template_argument(0)
-				self.ArrayNum = self.Data['ArrayNum']
-				if self.ArrayNum.is_optimized_out:
-					self.ArrayNum = 0
-				else:
-					self.AllocatorInstance = self.Data['AllocatorInstance']
-					self.AllocatorData = self.AllocatorInstance['Data']
-					self.AllocationFlags = self.Value['AllocationFlags']
-					self.AllocationFlagsInstance = self.AllocationFlags['AllocatorInstance']
-					self.InlineData = self.AllocationFlagsInstance['InlineData']
-					self.SecondaryData = self.AllocationFlagsInstance['SecondaryData']
-					self.SecondaryDataData = None
-					if self.SecondaryData != None:
-						self.SecondaryData['Data']
-			except:
-				raise
-
-		def __iter__(self):
-			return self
-
-		def __next__(self):
-			return self.next()
-
-		def __next__(self):
-			if self.ArrayNum == 0:
-				raise StopIteration
-
-			self.Counter = self.Counter + 1
-
-			if self.Counter >= self.ArrayNum:
-				raise StopIteration
-			else:
-				Data = None
-				if self.SecondaryDataData > 0:
-					Data = (self.SecondaryDataData.address.cast(gdb.lookup_type("int").pointer())[self.Counter/32] >> self.Counter) & 1
-				else:
-					Data = (self.InlineData.address.cast(gdb.lookup_type("int").pointer())[self.Counter/32] >> self.Counter) & 1
-
-				if Data != 0:
-					ElementOrFreeListValue = (self.AllocatorData.cast(self.InternalElementType.pointer()) + self.Counter).dereference()
-					Value = ElementOrFreeListValue['ElementData'].reinterpret_cast(self.ElementType.reference())
-					return ('[%d]' % self.Counter, Value)
-				else:
-					return self.__next__()
+class TArrayPrinter:
+	"Print TArray"
 
 	def __init__(self, val):
 		self.Value = val
-		self.ArrayNum = self.Value['Data']['ArrayNum']
+		self.ArrayNum = self.Value['ArrayNum']
+		self.ArrayMax = self.Value['ArrayMax']
+		self.Data = get_allocation_data(self.Value['AllocatorInstance'], self.Value.type.template_argument(0))
 
 	def to_string(self):
 		if self.ArrayNum.is_optimized_out:
 			pass
 		if self.ArrayNum == 0:
 			return 'empty'
-		pass
+		return 'Num='+str(self.ArrayNum)+' Max='+str(self.ArrayMax)
 
 	def children(self):
-		return self._iterator(self.Value)
+		return container_iterator(self)
 
 	def display_hint(self):
-		return 'string'
+		return 'array'
+
+	def get_value(self, index):
+		return self.Data[index]
+		
+	def is_valid(self, index):
+		return index < self.ArrayMax
 
 # ------------------------------------------------------------------------------
-# TSet
+# TSparseArray
 #
-class TSetPrinter:
-	"Print TSet"
-
-	class _iterator(Iterator):
-		def __init__(self, val):
-			self.Value = val
-			self.Counter = -1
-			self.ElementType = self.Value.type.template_argument(0)
-			self.TSetElementType = gdb.lookup_type("TSetElement<" + self.ElementType.name + ">")
-			self.ElementsArrayNum = 0
-
-			try:
-				self.Elements = self.Value["Elements"]
-				self.ElementsData = self.Elements["Data"]
-				self.ElementsArrayNum = self.ElementsData['ArrayNum']
-				self.NumFreeIndices = self.Elements['NumFreeIndices']
-				self.AllocatorInstance = self.ElementsData['AllocatorInstance']
-				self.AllocationFlags = self.Elements['AllocationFlags']
-				self.AllocatorInstanceData = None
-				self.InlineData = None
-				self.SecondaryDataData = None
-				try:
-					self.AllocatorInstanceData = self.AllocatorInstance['Data']
-				except:
-					pass
-				try:
-					self.InlineData = self.AllocationFlags['AllocatorInstance']['InlineData']
-				except:
-					pass
-				try:
-					SecondaryData = self.AllocationFlags['AllocatorInstance']['SecondaryData']
-					if SecondaryData != None:
-						self.SecondaryDataData = SecondaryData['Data']
-				except:
-					pass
-			except:
-				self.ElementsArrayNum = 0
-
-		def __iter__(self):
-			return self
-
-		def __next__(self):
-			self.Counter = self.Counter + 1
-
-			if self.Counter >= self.ElementsArrayNum:
-				raise StopIteration()
-			else:
-				IsUsed = 0
-				if self.SecondaryDataData > 0:
-					IsUsed = (self.SecondaryDataData.address.cast(gdb.lookup_type("uint32").pointer())[self.Counter/32] >> self.Counter) & 1
-				else:
-					IsUsed = (self.InlineData.address.cast(gdb.lookup_type("uint32").pointer())[self.Counter/32] >> self.Counter) & 1
-
-				if IsUsed == 1:
-					Value = self.AllocatorInstanceData.cast(self.TSetElementType.pointer())[self.Counter]
-					return ('[%d]' % self.Counter, Value)
-				else:
-					return self.__next__()
+class TSparseArrayPrinter:
+	"Print TSparseArray"
 
 	def __init__(self, val):
 		self.Value = val
-		self.ArrayNum = self.Value["Elements"]["Data"]['ArrayNum']
+		self.Array = TArrayPrinter(self.Value['Data'])
+		self.AllocationFlags = TBitArrayPrinter(self.Value['AllocationFlags'])
+
+		self.NumFreeIndices = self.Value['NumFreeIndices']
+		self.ArrayNum = self.Array.ArrayNum
+		self.ArrayMax = self.Array.ArrayMax
+		self.Data = self.Array.Data.cast(self.Value.type.template_argument(0).pointer())
 
 	def to_string(self):
 		if self.ArrayNum.is_optimized_out:
 			pass
 		if self.ArrayNum == 0:
 			return 'empty'
-		pass
+		return 'Num='+str(self.ArrayNum - self.NumFreeIndices)+' Max='+str(self.ArrayMax)
 
 	def children(self):
-		return self._iterator(self.Value)
+		return container_iterator(self)
 
 	def display_hint(self):
 		return 'array'
 
+	def get_value(self, index):
+		return self.Data[index]
+
+	def is_valid(self, index):
+		return self.AllocationFlags.get_value(index)
+
 # ------------------------------------------------------------------------------
-# TSetElementPrinter
+# TSet
 #
-
-class TSetElementPrinter:
-	"Print TSetElement"
+class TSparseSetPrinter(TSparseArrayPrinter):
+	"Print TSparseSet"
 
 	def __init__(self, val):
-		self.Value = val
+		super().__init__(val['Elements'])
 
-	def to_string(self):
-		if self.Value.is_optimized_out:
-			return '<optimized out>'
-
-		return self.Value["Value"]
-
-	def display_hint(self):
-		return "string"
+	def children(self):
+		for Str, Element in super().children():
+			yield (Str, Element['Value'])
 
 # ------------------------------------------------------------------------------
-# TMap
+# TSparseMap
 #
-class TMapPrinter:
-	"Print TMap"
 
-	class _iterator(Iterator):
-		def __init__(self, val):
-			self.Value = val
-			self.Counter = -1
-			try:
-				self.Pairs = self.Value['Pairs']
-				if self.Pairs.is_optimized_out:
-					self.ArrayNum = 0
-				else:
-					self.Elements = self.Pairs['Elements']
-					self.ElementsData = self.Elements['Data']
-					self.ArrayNum = self.ElementsData['ArrayNum']
-			except:
-				raise
-
-		def __iter__(self):
-			return self
-
-		def __next__(self):
-			if self.ArrayNum == 0:
-				raise StopIteration
-
-			self.Counter = self.Counter + 1
-
-			if self.Counter > 0:
-				raise StopIteration
-
-			return ('Pairs', self.Pairs)
+class TSparseMapPrinter(TSparseSetPrinter):
+	"Print TSparseMap"
 
 	def __init__(self, val):
-		self.Value = val
-		self.ArrayNum = self.Value['Pairs']['Elements']['Data']['ArrayNum']
+		super().__init__(val['Pairs'])
 
 	def children(self):
-		return self._iterator(self.Value)
-
-	def to_string(self):
-		if self.ArrayNum.is_optimized_out:
-			pass
-		if self.ArrayNum == 0:
-			return 'empty'
-		pass
-
-	def display_hint(self):
-		return 'map'
+		for Str, Pair in super().children():
+			yield (f'[{Pair['Key']}] = {Pair['Value'].type}', Pair['Value'])
 
 # ------------------------------------------------------------------------------
 # TWeakObjectPtr
@@ -607,103 +480,32 @@ def make_fminimalname_printer(val):
 		return FNameEntryIdPrinter(val['Index'])
 
 # ------------------------------------------------------------------------------
-# TTuple
+# FPrimaryAssetId
 #
-class TTuplePrinter:
-	"Print TTuple"
+def remove_quotes(name):
+	return str(name).strip("'").strip('"')
 
+class FPrimaryAssetIdPrinter:
+	"Print FPrimaryAssetId"
+	
 	def __init__(self, val):
 		self.Value = val
 
-		try:
-			self.TKey = self.Value["Key"];
-			self.TValue = self.Value["Value"];
-		except:
-			pass
-
 	def to_string(self):
-		return '(%s, %s)' % (self.TKey, self.TValue)
-
+		return remove_quotes(self.Value["PrimaryAssetType"]["Name"]) + ':' + remove_quotes(self.Value["PrimaryAssetName"])
 
 # ------------------------------------------------------------------------------
-# TArray
+# TPair (TTuple<Key,Value>)
 #
-class TArrayPrinter:
-	"Print TArray"
-
-	class _iterator(Iterator):
-		def __init__(self, val):
-			self.Value = val
-			self.Counter = -1
-			self.TType = self.Value.type.template_argument(0)
-
-			try:
-				self.ArrayNum = self.Value['ArrayNum']
-				if self.ArrayNum.is_optimized_out:
-					self.ArrayNum = 0
-
-				if self.ArrayNum > 0:
-					self.AllocatorInstance = self.Value['AllocatorInstance']
-					self.AllocatorInstanceData = None
-					self.InlineData = None
-					self.SecondaryDataData = None
-					try:
-						self.AllocatorInstanceData = self.AllocatorInstance['Data']
-					except:
-						pass
-					try:
-						self.InlineData = self.AllocatorInstance['InlineData']
-					except:
-						pass
-					try:
-						SecondaryData = self.AllocatorInstance['SecondaryData']
-						if SecondaryData != None:
-							self.SecondaryDataData = SecondaryData['Data']
-					except:
-						pass
-			except:
-				raise
-
-		def __iter__(self):
-			return self
-
-		def __next__(self):
-			if self.ArrayNum == 0:
-				raise StopIteration
-
-			self.Counter = self.Counter + 1
-
-			if self.Counter >= self.ArrayNum:
-				raise StopIteration
-
-			try:
-			    if self.AllocatorInstanceData != None:
-				    data = self.AllocatorInstanceData.cast(self.TType.pointer())
-			    elif self.SecondaryDataData != None:
-				    data = self.SecondaryDataData.cast(self.TType.pointer())
-			    else:
-				    data = self.InlineData.cast(self.TType.pointer())
-			except:
-				return ('[%d]' % self.Counter, "optimized")
-
-			return ('[%d]' % self.Counter, data[self.Counter])
-
+class TPairPrinter:
+	"Print TPair"
+	
 	def __init__(self, val):
-		self.Value = val;
-		self.ArrayNum = self.Value['ArrayNum']
-
-	def to_string(self):
-		if self.ArrayNum.is_optimized_out:
-			pass
-		if self.ArrayNum == 0:
-			return 'empty'
-		pass
+		self.Value = val
 
 	def children(self):
-		return self._iterator(self.Value)
-
-	def display_hint(self):
-		return 'array'
+		yield ("Key", self.Value["Key"])
+		yield ("Value", self.Value["Value"])
 
 #
 # Register our lookup function. If no objfile is passed use all globally.
@@ -725,13 +527,19 @@ def build_ue_pretty_printer():
 	pp.add_printer("TBitArray", '^TBitArray<.+>$', TBitArrayPrinter)
 	pp.add_printer("TChunkedArray", '^TChunkedArray<.+>$', TChunkedArrayPrinter)
 	pp.add_printer("TSparseArray", '^TSparseArray<.+>$', TSparseArrayPrinter)
-	pp.add_printer("TSetElement", '^TSetElement<.+>$', TSetElementPrinter)
-	pp.add_printer("TSet", '^TSet<.+>$', TSetPrinter)
+	pp.add_printer("TSet", '^TSet<.+>$', TSparseSetPrinter)
+	pp.add_printer("TSparseSet", '^TSparseSet<.+>$', TSparseSetPrinter)
 	pp.add_printer("FBitReference", '^FBitReference$', FBitReferencePrinter)
-#	pp.add_printer("TMap", '^TMap<.+,.+,.+>$', TMapPrinter)
-	pp.add_printer("TPair", '^TPair<.+,.+>$', TTuplePrinter)
-	pp.add_printer("TTuple", '^TTuple<.+,.+>$', TTuplePrinter)
+	pp.add_printer("TMapBase", '^TMapBase<.+,.+,.+,.+>$', TSparseMapPrinter)
+	pp.add_printer("TSortableMapBase", '^TSortableMapBase<.+,.+,.+,.+>$', TSparseMapPrinter)
+	pp.add_printer("TMap", '^TMap<.+,.+,.+,.+>$', TSparseMapPrinter)
+	pp.add_printer("TSparseMapBase", '^TSparseMapBase<.+,.+,.+,.+>$', TSparseMapPrinter)
+	pp.add_printer("TSortableSparseMapBase", '^TSortableSparseMapBase<.+,.+,.+,.+>$', TSparseMapPrinter)
+	pp.add_printer("TSparseMap", '^TSparseMap<.+,.+,.+,.+>$', TSparseMapPrinter)
+	pp.add_printer("TPair", '^TPair<.+,.+>$', TPairPrinter)
+	pp.add_printer("TTuple", '^TTuple<.+,.+>$', TPairPrinter)
 	pp.add_printer("TWeakObjectPtr", '^TWeakObjectPtr<.+>$', TWeakObjectPtrPrinter)
+	pp.add_printer("FPrimaryAssetId", '^FPrimaryAssetId$', FPrimaryAssetIdPrinter)
 	return pp
 
 register_ue_printers(None)
diff --git a/Engine/Source/Runtime/Core/Private/Containers/ContainersTest.cpp b/Engine/Source/Runtime/Core/Private/Containers/ContainersTest.cpp
--- a/Engine/Source/Runtime/Core/Private/Containers/ContainersTest.cpp
+++ b/Engine/Source/Runtime/Core/Private/Containers/ContainersTest.cpp
@@ -6,6 +6,8 @@
 #include "Containers/ArrayView.h"
 #include "Containers/Map.h"
 #include "Containers/Set.h"
+#include "Containers/SparseSet.h"
+#include "Containers/CompactSet.h"
 #include "Containers/SortedMap.h"
 #include "Containers/SortedSet.h"
 #include "Math/RandomStream.h"
@@ -23,6 +25,33 @@
 
 namespace
 {
+	struct FSetResolver
+	{
+		template<typename InElementType>
+		using TSetType = TSet<InElementType>;
+
+		template<typename InKey, typename InValueType>
+		using TMapType = TMap<InKey, InValueType>;
+	};
+
+	struct FSparseSetResolver
+	{
+		template<typename InElementType>
+		using TSetType = TSparseSet<InElementType>;
+
+		template<typename InKey, typename InValueType>
+		using TMapType = TSparseMap<InKey, InValueType>;
+	};
+
+	struct FCompactSetResolver
+	{
+		template<typename InElementType>
+		using TSetType = TCompactSet<InElementType>;
+
+		template<typename InKey, typename InValueType>
+		using TMapType = TCompactMap<InKey, InValueType>;
+	};
+
 	struct FContainerTestStats
 	{
 		int32  NextId;
@@ -1721,49 +1750,60 @@
 	FMemory::Free(Temp);
 }
 
-IMPLEMENT_SIMPLE_AUTOMATION_TEST(FContainersNatvisTest, "System.Core.Containers.Natvis", EAutomationTestFlags_ApplicationContextMask | EAutomationTestFlags::SmokeFilter)
-bool FContainersNatvisTest::RunTest(const FString& Parameters)
+template <typename Resolver>
+struct TSetNatvisTest
 {
-	// This results of this test are intended to be inspected visually in a platform's debugger, to ensure that the .natvis is correct.
-	// A breakpoint should be placed on the last line of each block, the block variables observed in the debugger and compared with the
-	// textual description.
+	template <typename Type>
+	using SetType = typename Resolver::template TSetType<Type>;
 
-	// TArray
+	TSetNatvisTest()
 	{
 		// Name       | Value | Type
-		// -----------+-------+---------------------------------------
-		// > EmptyArr | Empty | TArray<int,TSizedDefaultAllocator<32>>
-		TArray<int32> EmptyArr;
+		// -----------+-------+------------------------------------------------------
+		// > EmptySet | Empty | TSet<int,DefaultKeyFuncs<int,0>,FDefaultSetAllocator>
+		SetType<int32> EmptySet;
 
 		// Name       | Value         | Type
-		// -----------+---------------+---------------------------------------
-		// > SlackArr | Empty, Max=16 | TArray<int,TSizedDefaultAllocator<32>>
-		TArray<int32> SlackArr;
-		SlackArr.Reserve(16);
+		// -----------+---------------+------------------------------------------------------
+		// > SlackSet | Empty, Max=16 | TSet<int,DefaultKeyFuncs<int,0>,FDefaultSetAllocator>
+		SetType<int32> SlackSet;
+		SlackSet.Reserve(16);
 
 		// Name      | Value        | Type
-		// ----------+--------------+---------------------------------------
-		// > ArrVals | Num=5, Max=5 | TArray<int,TSizedDefaultAllocator<32>>
+		// ----------+--------------+------------------------------------------------------
+		// > SetVals | Num=5, Max=5 | TSet<int,DefaultKeyFuncs<int,0>,FDefaultSetAllocator>
 		//     [0]   | 4            | int
 		//     [1]   | 1            | int
 		//     [2]   | 2            | int
 		//     [3]   | 5            | int
 		//     [4]   | 3            | int
-		TArray<int32> ArrVals = { 4, 1, 2, 5, 3 };
+		SetType<int32> SetVals = { 4, 1, 2, 5, 3 };
 
-		// Name           | Value        | Type
-		// ---------------+--------------+---------------------------------------
-		// > SlackArrVals | Num=3, Max=5 | TArray<int,TSizedDefaultAllocator<32>>
-		//     [0]        | 4            | int
-		//     [1]        | 2            | int
-		//     [2]        | 3            | int
-		TArray<int32> SlackArrVals = { 4, 1, 2, 5, 3 };
-		SlackArrVals.RemoveAt(3);
-		SlackArrVals.RemoveAt(1);
+		// Name           | Value         | Type
+		// ---------------+---------------+------------------------------------------------------
+		// > SlackSetVals | Num=5, Max=16 | TSet<int,DefaultKeyFuncs<int,0>,FDefaultSetAllocator>
+		//     [0]        | 4             | int
+		//     [1]        | 1             | int
+		//     [2]        | 2             | int
+		//     [3]        | 5             | int
+		//     [4]        | 3             | int
+		SetType<int32> SlackSetVals = { 4, 1, 2, 5, 3 };
+		SlackSetVals.Reserve(16);
+
+		// Name            | Value        | Type
+		// ----------------+--------------+------------------------------------------------------
+		// > SparseSetVals | Num=3, Max=5 | TSet<int,DefaultKeyFuncs<int,0>,FDefaultSetAllocator>
+		//     [0]         | 4            | int
+		//     [1]         | 2            | int
+		//     [2]         | 3            | int
+		// Note: Order different for compact set
+		SetType<int32> SparseSetVals = { 4, 1, 2, 5, 3 };
+		SparseSetVals.Remove(1);
+		SparseSetVals.Remove(5);
 
 		// Name      | Value        | Type
-		// ----------+--------------+-------------------------------------------
-		// > ArrStrs | Num=5, Max=5 | TArray<FString,TSizedDefaultAllocator<32>>
+		// ----------+--------------+--------------------------------------------------------------
+		// > SetStrs | Num=5, Max=5 | TSet<FString,DefaultKeyFuncs<FString,0>,FDefaultSetAllocator>
 		//   > [0]   | L"d1"        | FString
 		//       [0] | 100 'd'      | wchar_t
 		//       [1] | 49 '1'       | wchar_t
@@ -1779,11 +1819,11 @@
 		//   > [4]   | L"c5"        | FString
 		//       [0] | 99 'c'       | wchar_t
 		//       [1] | 53 '5'       | wchar_t
-		TArray<FString> ArrStrs = { "d1", "a2", "b3", "e4", "c5" };
+		SetType<FString> SetStrs = { "d1", "a2", "b3", "e4", "c5" };
 
 		// Name            | Value                | Type
-		// ----------------+----------------------+-------------------------------------------------------------------------------
-		// > ArrUniqueStrs | Num=5, Max=5         | TArray<TUniquePtr<FString,TDefaultDelete<FString>>,TSizedDefaultAllocator<32>>
+		// ----------------+----------------------+--------------------------------------------------------------------------------------------------------------------------------------
+		// > SetUniqueStrs | Num=5, Max=5         | TSet<TUniquePtr<FString,TDefaultDelete<FString>>,DefaultKeyFuncs<TUniquePtr<FString,TDefaultDelete<FString>>,0>,FDefaultSetAllocator>
 		//   > [0]         | Ptr=0x01234567 L"d1" | TUniquePtr<FString,TDefaultDelete<FString>>
 		//       [0]       | 100 'd'              | wchar_t
 		//       [1]       | 49 '1'               | wchar_t
@@ -1799,64 +1839,132 @@
 		//   > [4]         | Ptr=0x01234567 L"c5" | TUniquePtr<FString,TDefaultDelete<FString>>
 		//       [0]       | 99 'c'               | wchar_t
 		//       [1]       | 53 '5'               | wchar_t
-		TArray<TUniquePtr<FString>> ArrUniqueStrs;
-		ArrUniqueStrs.Reserve(5);
-		ArrUniqueStrs.Add(MakeUnique<FString>("d1"));
-		ArrUniqueStrs.Add(MakeUnique<FString>("a2"));
-		ArrUniqueStrs.Add(MakeUnique<FString>("b3"));
-		ArrUniqueStrs.Add(MakeUnique<FString>("e4"));
-		ArrUniqueStrs.Add(MakeUnique<FString>("c5"));
+		SetType < TUniquePtr<FString> > SetUniqueStrs;
+		SetUniqueStrs.Reserve(5);
+		SetUniqueStrs.Add(MakeUnique<FString>("d1"));
+		SetUniqueStrs.Add(MakeUnique<FString>("a2"));
+		SetUniqueStrs.Add(MakeUnique<FString>("b3"));
+		SetUniqueStrs.Add(MakeUnique<FString>("e4"));
+		SetUniqueStrs.Add(MakeUnique<FString>("c5"));
 
 		PlaceBreakpointHere();
 	}
+};
 
-	// TSet
+template <typename Resolver>
+struct TMapNatvisTest
+{
+	template <typename Key, typename Value>
+	using MapType = typename Resolver::template TMapType<Key, Value>;
+
+	void Run()
 	{
 		// Name       | Value | Type
-		// -----------+-------+------------------------------------------------------
-		// > EmptySet | Empty | TSet<int,DefaultKeyFuncs<int,0>,FDefaultSetAllocator>
-		TSet<int32> EmptySet;
+		// -----------+-------+----------------------------------------------------------------------------------
+		// > EmptyMap | Empty | TMap<FString,int,FDefaultSetAllocator,TDefaultMapHashableKeyFuncs<FString,int,0>>
+		MapType<FString, int32> EmptyMap;
 
 		// Name       | Value         | Type
 		// -----------+---------------+------------------------------------------------------
-		// > SlackSet | Empty, Max=16 | TSet<int,DefaultKeyFuncs<int,0>,FDefaultSetAllocator>
-		TSet<int32> SlackSet;
-		SlackSet.Reserve(16);
+		// > SlackMap | Empty, Max=16 | TMap<FString,int,FDefaultSetAllocator,TDefaultMapHashableKeyFuncs<FString,int,0>>
+		MapType<FString, int32> SlackMap;
+		SlackMap.Reserve(16);
+
+		// Name       | Value        | Type
+		// -----------+--------------+------------------------------------------------------
+		// > MapVals  | Num=5, Max=5 | TMap<FString,int,FDefaultSetAllocator,TDefaultMapHashableKeyFuncs<FString,int,0>>
+		//     [L"d"] | 1            | int
+		//     [L"a"] | 2            | int
+		//     [L"b"] | 3            | int
+		//     [L"e"] | 4            | int
+		//     [L"c"] | 5            | int
+		MapType<FString, int32> MapVals = { { "d", 1 }, { "a", 2 }, { "b", 3 }, { "e", 4 }, { "c", 5 } };
+
+		// Name           | Value         | Type
+		// ---------------+---------------+------------------------------------------------------
+		// > SlackMapVals | Num=5, Max=16 | TMap<FString,int,FDefaultSetAllocator,TDefaultMapHashableKeyFuncs<FString,int,0>>
+		//     [L"d"]     | 1             | int
+		//     [L"a"]     | 2             | int
+		//     [L"b"]     | 3             | int
+		//     [L"e"]     | 4             | int
+		//     [L"c"]     | 5             | int
+		MapType<FString, int32> SlackMapVals = { { "d", 1 }, { "a", 2 }, { "b", 3 }, { "e", 4 }, { "c", 5 } };
+		SlackMapVals.Reserve(16);
+
+		// Name            | Value        | Type
+		// ----------------+--------------+------------------------------------------------------
+		// > SparseMapVals | Num=3, Max=5 | TMap<FString,int,FDefaultSetAllocator,TDefaultMapHashableKeyFuncs<FString,int,0>>
+		//     [L"d"]      | 1            | int
+		//     [L"b"]      | 3            | int
+		//     [L"c"]      | 5            | int
+		// Note: Order different for compact map
+		MapType<FString, int32> SparseMapVals = { { "d", 1 }, { "a", 2 }, { "b", 3 }, { "e", 4 }, { "c", 5 } };
+		SparseMapVals.Remove("a");
+		SparseMapVals.Remove("e");
+
+		// Name                      | Value        | Type
+		// --------------------------+--------------+----------------------------------------------------------------------------------------------------------------------------------------------------------
+		// > MapUniqueVals           | Num=5, Max=5 | TMap<TUniquePtr<FString,TDefaultDelete<FString>>,int,FDefaultSetAllocator,TDefaultMapHashableKeyFuncs<TUniquePtr<FString,TDefaultDelete<FString>>,int,0>>
+		//     [Ptr=0x12345678 L"d"] | 1            | int
+		//     [Ptr=0x12345678 L"a"] | 2            | int
+		//     [Ptr=0x12345678 L"b"] | 3            | int
+		//     [Ptr=0x12345678 L"e"] | 4            | int
+		//     [Ptr=0x12345678 L"c"] | 5            | int
+		MapType<TUniquePtr<FString>, int32> MapUniqueVals;
+		MapUniqueVals.Reserve(5);
+		MapUniqueVals.Emplace(MakeUnique<FString>("d"), 1);
+		MapUniqueVals.Emplace(MakeUnique<FString>("a"), 2);
+		MapUniqueVals.Emplace(MakeUnique<FString>("b"), 3);
+		MapUniqueVals.Emplace(MakeUnique<FString>("e"), 4);
+		MapUniqueVals.Emplace(MakeUnique<FString>("c"), 5);
+
+		PlaceBreakpointHere();
+	}
+};
+
+IMPLEMENT_SIMPLE_AUTOMATION_TEST(FContainersNatvisTest, "System.Core.Containers.Natvis", EAutomationTestFlags_ApplicationContextMask | EAutomationTestFlags::SmokeFilter)
+bool FContainersNatvisTest::RunTest(const FString& Parameters)
+{
+	// This results of this test are intended to be inspected visually in a platform's debugger, to ensure that the .natvis is correct.
+	// A breakpoint should be placed on the last line of each block, the block variables observed in the debugger and compared with the
+	// textual description.
+
+	// TArray
+	{
+		// Name       | Value | Type
+		// -----------+-------+---------------------------------------
+		// > EmptyArr | Empty | TArray<int,TSizedDefaultAllocator<32>>
+		TArray<int32> EmptyArr;
+
+		// Name       | Value         | Type
+		// -----------+---------------+---------------------------------------
+		// > SlackArr | Empty, Max=16 | TArray<int,TSizedDefaultAllocator<32>>
+		TArray<int32> SlackArr;
+		SlackArr.Reserve(16);
 
 		// Name      | Value        | Type
-		// ----------+--------------+------------------------------------------------------
-		// > SetVals | Num=5, Max=5 | TSet<int,DefaultKeyFuncs<int,0>,FDefaultSetAllocator>
+		// ----------+--------------+---------------------------------------
+		// > ArrVals | Num=5, Max=5 | TArray<int,TSizedDefaultAllocator<32>>
 		//     [0]   | 4            | int
 		//     [1]   | 1            | int
 		//     [2]   | 2            | int
 		//     [3]   | 5            | int
 		//     [4]   | 3            | int
-		TSet<int32> SetVals = { 4, 1, 2, 5, 3 };
+		TArray<int32> ArrVals = { 4, 1, 2, 5, 3 };
 
-		// Name           | Value         | Type
-		// ---------------+---------------+------------------------------------------------------
-		// > SlackSetVals | Num=5, Max=16 | TSet<int,DefaultKeyFuncs<int,0>,FDefaultSetAllocator>
-		//     [0]        | 4             | int
-		//     [1]        | 1             | int
-		//     [2]        | 2             | int
-		//     [3]        | 5             | int
-		//     [4]        | 3             | int
-		TSet<int32> SlackSetVals = { 4, 1, 2, 5, 3 };
-		SlackSetVals.Reserve(16);
-
-		// Name            | Value        | Type
-		// ----------------+--------------+------------------------------------------------------
-		// > SparseSetVals | Num=3, Max=5 | TSet<int,DefaultKeyFuncs<int,0>,FDefaultSetAllocator>
-		//     [0]         | 4            | int
-		//     [1]         | 2            | int
-		//     [2]         | 3            | int
-		TSet<int32> SparseSetVals = { 4, 1, 2, 5, 3 };
-		SparseSetVals.Remove(1);
-		SparseSetVals.Remove(5);
+		// Name           | Value        | Type
+		// ---------------+--------------+---------------------------------------
+		// > SlackArrVals | Num=3, Max=5 | TArray<int,TSizedDefaultAllocator<32>>
+		//     [0]        | 4            | int
+		//     [1]        | 2            | int
+		//     [2]        | 3            | int
+		TArray<int32> SlackArrVals = { 4, 1, 2, 5, 3 };
+		SlackArrVals.RemoveAt(3);
+		SlackArrVals.RemoveAt(1);
 
 		// Name      | Value        | Type
-		// ----------+--------------+--------------------------------------------------------------
-		// > SetStrs | Num=5, Max=5 | TSet<FString,DefaultKeyFuncs<FString,0>,FDefaultSetAllocator>
+		// ----------+--------------+-------------------------------------------
+		// > ArrStrs | Num=5, Max=5 | TArray<FString,TSizedDefaultAllocator<32>>
 		//   > [0]   | L"d1"        | FString
 		//       [0] | 100 'd'      | wchar_t
 		//       [1] | 49 '1'       | wchar_t
@@ -1872,11 +1980,11 @@
 		//   > [4]   | L"c5"        | FString
 		//       [0] | 99 'c'       | wchar_t
 		//       [1] | 53 '5'       | wchar_t
-		TSet<FString> SetStrs = { "d1", "a2", "b3", "e4", "c5" };
+		TArray<FString> ArrStrs = { "d1", "a2", "b3", "e4", "c5" };
 
 		// Name            | Value                | Type
-		// ----------------+----------------------+--------------------------------------------------------------------------------------------------------------------------------------
-		// > SetUniqueStrs | Num=5, Max=5         | TSet<TUniquePtr<FString,TDefaultDelete<FString>>,DefaultKeyFuncs<TUniquePtr<FString,TDefaultDelete<FString>>,0>,FDefaultSetAllocator>
+		// ----------------+----------------------+-------------------------------------------------------------------------------
+		// > ArrUniqueStrs | Num=5, Max=5         | TArray<TUniquePtr<FString,TDefaultDelete<FString>>,TSizedDefaultAllocator<32>>
 		//   > [0]         | Ptr=0x01234567 L"d1" | TUniquePtr<FString,TDefaultDelete<FString>>
 		//       [0]       | 100 'd'              | wchar_t
 		//       [1]       | 49 '1'               | wchar_t
@@ -1892,17 +2000,22 @@
 		//   > [4]         | Ptr=0x01234567 L"c5" | TUniquePtr<FString,TDefaultDelete<FString>>
 		//       [0]       | 99 'c'               | wchar_t
 		//       [1]       | 53 '5'               | wchar_t
-		TSet<TUniquePtr<FString>> SetUniqueStrs;
-		SetUniqueStrs.Reserve(5);
-		SetUniqueStrs.Add(MakeUnique<FString>("d1"));
-		SetUniqueStrs.Add(MakeUnique<FString>("a2"));
-		SetUniqueStrs.Add(MakeUnique<FString>("b3"));
-		SetUniqueStrs.Add(MakeUnique<FString>("e4"));
-		SetUniqueStrs.Add(MakeUnique<FString>("c5"));
+		TArray<TUniquePtr<FString>> ArrUniqueStrs;
+		ArrUniqueStrs.Reserve(5);
+		ArrUniqueStrs.Add(MakeUnique<FString>("d1"));
+		ArrUniqueStrs.Add(MakeUnique<FString>("a2"));
+		ArrUniqueStrs.Add(MakeUnique<FString>("b3"));
+		ArrUniqueStrs.Add(MakeUnique<FString>("e4"));
+		ArrUniqueStrs.Add(MakeUnique<FString>("c5"));
 
 		PlaceBreakpointHere();
 	}
 
+	// TSet
+	TSetNatvisTest<FSetResolver>();
+	TSetNatvisTest<FSparseSetResolver>();
+	TSetNatvisTest<FCompactSetResolver>();
+
 	// TSortedSet
 	{
 		// Name        | Value | Type
@@ -1997,67 +2110,9 @@
 	}
 
 	// TMap
-	{
-		// Name       | Value | Type
-		// -----------+-------+----------------------------------------------------------------------------------
-		// > EmptyMap | Empty | TMap<FString,int,FDefaultSetAllocator,TDefaultMapHashableKeyFuncs<FString,int,0>>
-		TMap<FString, int32> EmptyMap;
+	TMapNatvisTest<FSetResolver>().Run();
+	TMapNatvisTest<FSparseSetResolver>().Run();
+	TMapNatvisTest<FCompactSetResolver>().Run();
 
-		// Name       | Value         | Type
-		// -----------+---------------+------------------------------------------------------
-		// > SlackMap | Empty, Max=16 | TMap<FString,int,FDefaultSetAllocator,TDefaultMapHashableKeyFuncs<FString,int,0>>
-		TMap<FString, int32> SlackMap;
-		SlackMap.Reserve(16);
-
-		// Name       | Value        | Type
-		// -----------+--------------+------------------------------------------------------
-		// > MapVals  | Num=5, Max=5 | TMap<FString,int,FDefaultSetAllocator,TDefaultMapHashableKeyFuncs<FString,int,0>>
-		//     [L"d"] | 1            | int
-		//     [L"a"] | 2            | int
-		//     [L"b"] | 3            | int
-		//     [L"e"] | 4            | int
-		//     [L"c"] | 5            | int
-		TMap<FString, int32> MapVals = { { "d", 1 }, { "a", 2 }, { "b", 3 }, { "e", 4 }, { "c", 5 } };
-
-		// Name           | Value         | Type
-		// ---------------+---------------+------------------------------------------------------
-		// > SlackMapVals | Num=5, Max=16 | TMap<FString,int,FDefaultSetAllocator,TDefaultMapHashableKeyFuncs<FString,int,0>>
-		//     [L"d"]     | 1             | int
-		//     [L"a"]     | 2             | int
-		//     [L"b"]     | 3             | int
-		//     [L"e"]     | 4             | int
-		//     [L"c"]     | 5             | int
-		TMap<FString, int32> SlackMapVals = { { "d", 1 }, { "a", 2 }, { "b", 3 }, { "e", 4 }, { "c", 5 } };
-		SlackMapVals.Reserve(16);
-
-		// Name            | Value        | Type
-		// ----------------+--------------+------------------------------------------------------
-		// > SparseMapVals | Num=3, Max=5 | TMap<FString,int,FDefaultSetAllocator,TDefaultMapHashableKeyFuncs<FString,int,0>>
-		//     [L"d"]      | 1            | int
-		//     [L"b"]      | 3            | int
-		//     [L"c"]      | 5            | int
-		TMap<FString, int32> SparseMapVals = { { "d", 1 }, { "a", 2 }, { "b", 3 }, { "e", 4 }, { "c", 5 } };
-		SparseMapVals.Remove("a");
-		SparseMapVals.Remove("e");
-
-		// Name                      | Value        | Type
-		// --------------------------+--------------+----------------------------------------------------------------------------------------------------------------------------------------------------------
-		// > MapUniqueVals           | Num=5, Max=5 | TMap<TUniquePtr<FString,TDefaultDelete<FString>>,int,FDefaultSetAllocator,TDefaultMapHashableKeyFuncs<TUniquePtr<FString,TDefaultDelete<FString>>,int,0>>
-		//     [Ptr=0x12345678 L"d"] | 1            | int
-		//     [Ptr=0x12345678 L"a"] | 2            | int
-		//     [Ptr=0x12345678 L"b"] | 3            | int
-		//     [Ptr=0x12345678 L"e"] | 4            | int
-		//     [Ptr=0x12345678 L"c"] | 5            | int
-		TMap<TUniquePtr<FString>, int32> MapUniqueVals;
-		MapUniqueVals.Reserve(5);
-		MapUniqueVals.Emplace(MakeUnique<FString>("d"), 1);
-		MapUniqueVals.Emplace(MakeUnique<FString>("a"), 2);
-		MapUniqueVals.Emplace(MakeUnique<FString>("b"), 3);
-		MapUniqueVals.Emplace(MakeUnique<FString>("e"), 4);
-		MapUniqueVals.Emplace(MakeUnique<FString>("c"), 5);
-
-		PlaceBreakpointHere();
-	}
-
 	// TSortedMap
 	{
