Fix ensure triggered when GetDefaultTypedValue is called from RT with valid History.

diff --git a/Engine/Source/Runtime/Core/Private/HAL/ConsoleManager.cpp b/Engine/Source/Runtime/Core/Private/HAL/ConsoleManager.cpp
--- a/Engine/Source/Runtime/Core/Private/HAL/ConsoleManager.cpp
+++ b/Engine/Source/Runtime/Core/Private/HAL/ConsoleManager.cpp
@@ -1054,7 +1054,7 @@ public:
 		if (PriorityHistory != nullptr)
 		{
 			// constructor will never have more than one value
-			return PriorityHistory->History[ECVF_SetByConstructor][0].Value.GetValueOnAnyThread(true);
+			return PriorityHistory->History[ECVF_SetByConstructor][0].Value.GetValueOnAnyThread();
 		}
 		
 		// if we have no history at all, that means we never called Set, so the current value must be the value we were constructed with
