[AutoRTFM] Fix crash with FUnifiedError::GetErrorMessage inside a Verse transaction.

TCbWriter uses atomics in some cases; we can avoid them by passing it a memory view instead of having it manage the buffer and the buffer's refcounts.

Fixing the buffers would be ideal but this is more complex than I was able to solve in an evening, because they often live on the stack.

diff --git a/Engine/Source/Programs/AutoRTFMTests/Private/UECoreTests.cpp b/Engine/Source/Programs/AutoRTFMTests/Private/UECoreTests.cpp
--- a/Engine/Source/Programs/AutoRTFMTests/Private/UECoreTests.cpp
+++ b/Engine/Source/Programs/AutoRTFMTests/Private/UECoreTests.cpp
@@ -9,6 +9,7 @@
 #include "AutoRTFMTesting.h"
 #include "Blueprint/BlueprintExceptionInfo.h"
 #include "Delegates/IDelegateInstance.h"
+#include "Experimental/UnifiedError/CoreErrorTypes.h"
 #include "HAL/IConsoleManager.h"
 #include "HAL/MallocLeakDetection.h"
 #include "HAL/ThreadHeartBeat.h"
@@ -2351,4 +2352,24 @@ TEST_CASE("UECore.FInlineLogTemplate")
 			AutoRTFM::AbortTransaction();
 		});
 	}
-}
\ No newline at end of file
+}
+
+TEST_CASE("UECore.UnifiedError")
+{
+	SECTION("Commit")
+	{
+		AutoRTFM::Testing::Commit([&]
+		{
+			UE::UnifiedError::FError Err = UE::UnifiedError::Core::CancellationError::MakeError();
+			FText Text = Err.GetErrorMessage(/*bIncludeContext*/ false);
+		});
+	}
+	SECTION("Abort")
+	{
+		AutoRTFM::Testing::Commit([&]
+		{
+			UE::UnifiedError::FError Err = UE::UnifiedError::Core::CancellationError::MakeError();
+			FText Text = Err.GetErrorMessage(/*bIncludeContext*/ false);
+		});
+	}
+}
diff --git a/Engine/Source/Runtime/Core/Private/Experimental/UnifiedError/UnifiedError.cpp b/Engine/Source/Runtime/Core/Private/Experimental/UnifiedError/UnifiedError.cpp
--- a/Engine/Source/Runtime/Core/Private/Experimental/UnifiedError/UnifiedError.cpp
+++ b/Engine/Source/Runtime/Core/Private/Experimental/UnifiedError/UnifiedError.cpp
@@ -102,10 +102,10 @@ FText FError::GetErrorMessage(bool bIncludeContext) const
 	{
 		LogDetailFilter |= EDetailFilter::IncludeInContextLogMessage;
 	}
-	
+
 	TArray<FText> AdditionalContext;
 	AdditionalContext.Add(FText::Join(FText::FromString(": "), NSLOCTEXT("UnifiedError", "DefaultLogPrepend", "{Root/ModuleIdString}.{Root/ErrorCodeString}"), GetFormatErrorText()));
-	
+
 	const IErrorDetails* DetailsIt = ErrorDetails.GetReference();
 	while (DetailsIt != nullptr)
 	{
@@ -113,11 +113,11 @@ FText FError::GetErrorMessage(bool bIncludeContext) const
 		{
 			AdditionalContext.Add(FText::FromString(FUtf8String::Printf("(%s:{%s})", DetailsIt->GetErrorDetailsTypeNameForLog().GetData(), DetailsIt->GetErrorDetailsTypeNameForLog().GetData())));
 		}
-		const IErrorDetails* Next = DetailsIt->GetInnerErrorDetails().GetReference(); 
+		const IErrorDetails* Next = DetailsIt->GetInnerErrorDetails().GetReference();
 		DetailsIt = Next;
 	}
 	FText FormatText = FText::Join(FText::FromString(", "), AdditionalContext);
-	
+
 #if USE_LOCALIZED_STRUCTURED_LOG_FOR_FERRORMESSAGE
 	UE::FLogTemplateOptions TemplateOptions;
 	TemplateOptions.bAllowSubObjectReferences = true;
@@ -129,12 +129,13 @@ FText FError::GetErrorMessage(bool bIncludeContext) const
 	TemplateOptions.bAllowSubObjectReferences = true;
 	FInlineLogTemplate Template(*FormatString, TemplateOptions);
 #endif
-	TCbWriter<1024> Writer;
-
+	FCbWriter Writer;
 	SerializeDetails(Writer, EDetailFilter::IncludeInSerialize);
-	FUtf8StringBuilderBase OutputMessage;
 
-	return Template.FormatToText(Writer.Save());
+	TArray<char, TInlineAllocator<1024>> ScratchBuffer;
+	ScratchBuffer.AddUninitialized((int32)Writer.GetSaveSize());
+
+	return Template.FormatToText(Writer.Save(MakeMemoryView(ScratchBuffer)));
 }
 
 void FError::SerializeDetails(FCbWriter& Writer, const EDetailFilter DetailFilter, bool bIncludeRoot) const
@@ -176,7 +177,7 @@ FString FError::SerializeToJsonString(const EDetailFilter DetailFilter) const
 
 FString FError::SerializeToJsonForAnalytics() const
 {
-	return SerializeToJsonString(EDetailFilter::IncludeInAnalytics);	
+	return SerializeToJsonString(EDetailFilter::IncludeInAnalytics);
 }
 
 int32 FError::GetErrorCode() const
@@ -235,7 +236,7 @@ void SerializeForLog(FCbWriter& Writer, const FError& Error)
 	Writer.BeginObject();
 	Writer.AddString(ANSITEXTVIEW("$type"), ANSITEXTVIEW("FError"));
 
-	// what I have to do 
+	// what I have to do
 	TCbWriter<1024> InnerWriter;
 
 	Error.SerializeDetails(InnerWriter, EDetailFilter::IncludeInSerialize);
@@ -273,7 +274,7 @@ FDynamicErrorDetails::FDynamicErrorDetails(TRefCountPtr<const IErrorDetails> InI
 {
 	InnerErrorDetails = InInnerErrorDetails;
 }
-FDynamicErrorDetails::~FDynamicErrorDetails() 
+FDynamicErrorDetails::~FDynamicErrorDetails()
 {
 }
 
@@ -305,7 +306,7 @@ FAnsiString FStaticErrorDetails::GetErrorCodeString(const FError& Error) const
 	return FAnsiString(ErrorName);
 }
 
-FAnsiString FStaticErrorDetails::GetModuleIdString(const FError& Error) const 
+FAnsiString FStaticErrorDetails::GetModuleIdString(const FError& Error) const
 {
 	return FAnsiString(ModuleName);
 }
@@ -337,7 +338,7 @@ namespace ErrorRegistry
 		TMap<TPair<int32, int32>, FAnsiString> ErrorCodeNameMap;
 	};
 
-	
+
 	uint32 FErrorRegistry::RegisterModule(const FAnsiStringView ModuleName)
 	{
 		// todo: need to replace this with a stable hashing function
