[Backout] - CL46473968
We hit a case where FNamePoolShard::Grow triggered a recursive call to Grow, causing a deadlock.
#rnx
Original CL Desc
-----------------------------------------------------------------
Added stat scope for FNamePoolShard::Grow.

#rnx

diff --git a/Engine/Source/Runtime/Core/Private/UObject/UnrealNames.cpp b/Engine/Source/Runtime/Core/Private/UObject/UnrealNames.cpp
--- a/Engine/Source/Runtime/Core/Private/UObject/UnrealNames.cpp
+++ b/Engine/Source/Runtime/Core/Private/UObject/UnrealNames.cpp
@@ -1602,7 +1602,6 @@ private:
 		LLM_TAGSET_SCOPE_CLEAR(ELLMTagSet::Assets);
 		LLM_TAGSET_SCOPE_CLEAR(ELLMTagSet::AssetClasses);
 		UE_TRACE_METADATA_CLEAR_SCOPE();
-		QUICK_SCOPE_CYCLE_COUNTER(STAT_FNamePoolShard_Grow);
 		TArrayView<FNameSlot> OldSlots(Slots, Capacity());
 		const uint32 OldUsedSlots = UsedSlots;
 
