Added workaround for getting an array reference element from an rvalue tuple under MSVC.

diff --git a/Engine/Source/Runtime/Core/Private/Templates/TuplesTest.cpp b/Engine/Source/Runtime/Core/Private/Templates/TuplesTest.cpp
--- a/Engine/Source/Runtime/Core/Private/Templates/TuplesTest.cpp
+++ b/Engine/Source/Runtime/Core/Private/Templates/TuplesTest.cpp
@@ -1,6 +1,7 @@
 // Copyright Epic Games, Inc. All Rights Reserved.
 
 #include <type_traits>
+#include "Misc/AutomationTest.h"
 #include "Templates/Tuple.h"
 
 // Pair keys have a different Get implementation, so test a pair and a non-pair
@@ -267,3 +268,46 @@ static_assert(TIsTuple<const volatile TTuple<double, float&, char&&>>::Value, ""
 	// Multiple arguments
 	static_assert(TIsTuple<double, double>::Value, "");
 #endif
+
+IMPLEMENT_SIMPLE_AUTOMATION_TEST(FTuplesTest, "System.Core.Tuples.Smoke", EAutomationTestFlags_ApplicationContextMask | EAutomationTestFlags::SmokeFilter)
+bool FTuplesTest::RunTest(const FString& Parameters)
+{
+	// Test getting references from tuples - this catches an issue found when MSVC returns the wrong address of a reference-to
+	// array element in an rvalue tuple.  See UE_TUPLE_RVALUE_WORKAROUND.
+	{
+		auto TestTupleValuesImpl = [](auto&& TupleRef, const auto& Addresses)
+		{
+			VisitTupleElements(
+				[](const auto& Element, const void* Address)
+				{
+					check(&Forward<decltype(Element)>(Element) == Address);
+				},
+				Forward<decltype(TupleRef)>(TupleRef),
+				Addresses
+			);
+		};
+
+		auto TestTupleValues = [&](auto&& TupleRef, const auto& Addresses)
+		{
+			using TupleRefType = decltype(TupleRef);
+
+			// Test as an lvalue tuple and as an rvalue tuple
+			TestTupleValuesImpl(TupleRef, Addresses);
+			TestTupleValuesImpl(Forward<decltype(TupleRef)>(TupleRef), Addresses);
+		};
+
+		int Int = 5;
+		float Float = 3.14f;
+		char Char = 'x';
+		int Array[] = { 1, 2, 3 };
+		auto Addresses = MakeTuple(&Int, &Float, &Char, &Array[0]);
+		TestTupleValues(TTuple<const int&, const float&, const char&, const int(&)[3]>(Int, Float, Char, Array), Addresses);
+		TestTupleValues(TTuple<int&&, float&&, char&&, const int(&&)[3]>(MoveTemp(Int), MoveTemp(Float), MoveTemp(Char), MoveTemp(Array)), Addresses);
+
+		// Test a pair separately as they have specializations
+		auto AddressesPair = MakeTuple(&Float, &Array[0]);
+		TestTupleValues(TTuple<const float&, const int(&)[3]>(Float, Array), AddressesPair);
+		TestTupleValues(TTuple<float&&, const int(&&)[3]>( MoveTemp(Float), MoveTemp(Array)), AddressesPair);
+	}
+	return true;
+}
diff --git a/Engine/Source/Runtime/Core/Public/Templates/Tuple.h b/Engine/Source/Runtime/Core/Public/Templates/Tuple.h
--- a/Engine/Source/Runtime/Core/Public/Templates/Tuple.h
+++ b/Engine/Source/Runtime/Core/Public/Templates/Tuple.h
@@ -38,6 +38,14 @@
 	#define UE_TUPLE_STATIC_ANALYSIS_WORKAROUND 0
 #endif
 
+// MSVC generates bad codegen when getting a reference element from an rvalue tuple.  So, if the element type is
+// an rvalue reference, get it via an lvalue - it will be the same reference value in either case.
+#if defined(_MSC_VER) && !defined(__clang__)
+	#define UE_TUPLE_REFERENCE_WORKAROUND 1
+#else
+	#define UE_TUPLE_REFERENCE_WORKAROUND 0
+#endif
+
 class FArchive;
 
 template <typename... Types>
@@ -108,13 +116,22 @@ namespace UE::Core::Private::Tuple
 	struct TTupleElementGetterByIndex
 	{
 		template <typename DeducedType, typename TupleType>
-		static inline decltype(auto) GetImpl(const volatile TTupleBaseElement<DeducedType, Index, TupleSize>&, TupleType&& Tuple)
+		static inline decltype(auto) GetImpl(const volatile TTupleBaseElement<DeducedType, Index, TupleSize>& Element, TupleType&& Tuple)
 		{
-			// Brackets are important here - we want a reference type to be returned, not object type
-			decltype(auto) Result = (ForwardAsBase<TupleType, TTupleBaseElement<DeducedType, Index, TupleSize>>(Tuple).Value);
-
-			// Keep tuple rvalue references to rvalue reference elements as rvalues, because that's how std::get() works, not how C++ struct member access works.
-			return static_cast<std::conditional_t<!std::is_reference_v<TupleType> && std::is_rvalue_reference_v<DeducedType>, DeducedType, decltype(Result)>>(Result);
+#if UE_TUPLE_REFERENCE_WORKAROUND
+			if constexpr (!std::is_reference_v<TupleType> && std::is_reference_v<DeducedType>)
+			{
+				return (DeducedType)GetImpl(Element, Tuple);
+			}
+			else
+#endif
+			{
+				// Brackets are important here - we want a reference type to be returned, not object type
+				decltype(auto) Result = (ForwardAsBase<TupleType, TTupleBaseElement<DeducedType, Index, TupleSize>>(Tuple).Value);
+	
+				// Keep tuple rvalue references to rvalue reference elements as rvalues, because that's how std::get() works, not how C++ struct member access works.
+				return static_cast<std::conditional_t<!std::is_reference_v<TupleType> && std::is_rvalue_reference_v<DeducedType>, DeducedType, decltype(Result)>>(Result);
+			}
 		}
 
 		template <typename TupleType>
@@ -131,11 +148,21 @@ namespace UE::Core::Private::Tuple
 		template <typename TupleType>
 		static inline decltype(auto) Get(TupleType&& Tuple)
 		{
-			// Brackets are important here - we want a reference type to be returned, not object type
-			decltype(auto) Result = (Forward<TupleType>(Tuple).Key);
-
-			// Keep tuple rvalue references to rvalue reference elements as rvalues, because that's how std::get() works, not how C++ struct member access works.
-			return static_cast<std::conditional_t<!std::is_reference_v<TupleType> && std::is_rvalue_reference_v<decltype(Tuple.Key)>, decltype(Tuple.Key), decltype(Result)>>(Result);
+#if UE_TUPLE_REFERENCE_WORKAROUND
+			using KeyType = decltype(std::decay_t<TupleType>::Key);
+			if constexpr (!std::is_reference_v<TupleType> && std::is_reference_v<KeyType>)
+			{
+				return (KeyType)Get(Tuple);
+			}
+			else
+#endif
+			{
+				// Brackets are important here - we want a reference type to be returned, not object type
+				decltype(auto) Result = (Forward<TupleType>(Tuple).Key);
+	
+				// Keep tuple rvalue references to rvalue reference elements as rvalues, because that's how std::get() works, not how C++ struct member access works.
+				return static_cast<std::conditional_t<!std::is_reference_v<TupleType> && std::is_rvalue_reference_v<decltype(Tuple.Key)>, decltype(Tuple.Key), decltype(Result)>>(Result);
+			}
 		}
 	};
 	template <>
