- Fixed an uninitialized copy-constructor value

diff --git a/Engine/Source/Runtime/Core/Private/Serialization/Archive.cpp b/Engine/Source/Runtime/Core/Private/Serialization/Archive.cpp
--- a/Engine/Source/Runtime/Core/Private/Serialization/Archive.cpp
+++ b/Engine/Source/Runtime/Core/Private/Serialization/Archive.cpp
@@ -256,6 +256,7 @@ PRAGMA_ENABLE_DEPRECATION_WARNINGS
 	ArShouldSkipCompilingAssets          = ArchiveToCopy.ArShouldSkipCompilingAssets;
 	ArMaxSerializeSize                   = ArchiveToCopy.ArMaxSerializeSize;
 	ArIsFilterEditorOnly                 = ArchiveToCopy.ArIsFilterEditorOnly;
+	ArIsSavingOptionalObject             = ArchiveToCopy.ArIsSavingOptionalObject;
 	ArIsSaveGame                         = ArchiveToCopy.ArIsSaveGame;
 	ArIsNetArchive                       = ArchiveToCopy.ArIsNetArchive;
 	ArCustomPropertyList                 = ArchiveToCopy.ArCustomPropertyList;
