Fix invalid arguments to proc_pid_rusage on iOS 15 devices. rusave_v6 isn't implemented bt the iOS 15 runtime, so Fallback to rusage_v5

diff --git a/Engine/Source/Runtime/Core/Private/Apple/ApplePlatformMemory.cpp b/Engine/Source/Runtime/Core/Private/Apple/ApplePlatformMemory.cpp
--- a/Engine/Source/Runtime/Core/Private/Apple/ApplePlatformMemory.cpp
+++ b/Engine/Source/Runtime/Core/Private/Apple/ApplePlatformMemory.cpp
@@ -442,10 +442,10 @@ FPlatformMemoryStats FApplePlatformMemory::GetStats()
 	// considers "free" pages a waste of resources.
 	MemoryStats.AvailableVirtual = (Stats.free_count + Stats.inactive_count) * MemoryConstants.PageSize;
 	
-	rusage_info_current rusage_payload;
+	struct rusage_info_v5 rusage_payload;
 		 
 	int ret = proc_pid_rusage(getpid(),
-							  RUSAGE_INFO_CURRENT, // I.e., new RUSAGE_INFO_V6 this year.
+							  RUSAGE_INFO_V5, // I.e., new RUSAGE_INFO_V5 as RUSAGE_INFO_V6 is unsupported on iOS 15.
 							  (rusage_info_t *)&rusage_payload);
 	NSCAssert(ret == 0, @"Could not get rusage: %i.", errno); // Look up in `man errno`.
 			
