Mac: Fix issue with shared cache layout giving the wrong Module during portable stackwalking.

diff --git a/Engine/Source/Runtime/Core/Private/Apple/ApplePlatformStackWalk.cpp b/Engine/Source/Runtime/Core/Private/Apple/ApplePlatformStackWalk.cpp
--- a/Engine/Source/Runtime/Core/Private/Apple/ApplePlatformStackWalk.cpp
+++ b/Engine/Source/Runtime/Core/Private/Apple/ApplePlatformStackWalk.cpp
@@ -276,31 +276,80 @@ int32 FApplePlatformStackWalk::GetProcessModuleSignatures(FStackWalkModuleInfo *
 {
 	int32 ModuleCount = GetProcessModuleCount();
 
-	int32 SignatureIndex = 0;
-
-	for (int32 ModuleIndex = 0; ModuleIndex < ModuleCount && SignatureIndex < ModuleSignaturesSize; ModuleIndex++)
+	TArray<FStackWalkModuleInfo> TempModules;
+	TempModules.Reserve(ModuleCount);
+
+	// The shared cache breaks our module size calculations. When we sum segment vmsizes from the mach_header,
+	// we get inflated sizes because the cache splits segments apart - they're not contiguous. 
+	// This causes wrong module attribution when checking address ranges.
+
+	// So we use dladdr() to make it work as _dyld_get_image_header() gives us the header pointer, but
+	// for cached frameworks that's not where the code actually lives. dladdr's dli_fbase gives
+	// us the real base address and what we need.
+	// However the order in which we get these aren't sorted by base address.
+	// So we've got to collect them all.
+	// Then we calculate size as the distance to the next module instead of summing segments.
+	for (int32 ModuleIndex = 0; ModuleIndex < ModuleCount; ModuleIndex++)
 	{
 		const struct mach_header* Header = _dyld_get_image_header(ModuleIndex);
 		const ANSICHAR* ImageName = _dyld_get_image_name(ModuleIndex);
+
 		if (Header && ImageName)
 		{
 			FStackWalkModuleInfo Info;
-			Info.BaseOfImage = (uint64)Header;
+			Dl_info DylibInfo;
+			if (dladdr((const void*)Header, &DylibInfo) && DylibInfo.dli_fbase)
+			{
+				Info.BaseOfImage = (uint64)DylibInfo.dli_fbase;
+			}
+			else
+			{
+				// Fallback to the busted way in case things go wrong
+				// so at least we have something
+				Info.BaseOfImage = (uint64)Header;
+			}
+
+			Info.ImageSize = 0;
+
 			FCString::Strcpy(Info.ImageName, ANSI_TO_TCHAR(ImageName));
-			Info.ImageSize = GetModuleImageSize( Header );
 			FCString::Strcpy(Info.LoadedImageName, ANSI_TO_TCHAR(ImageName));
 			FCString::Strcpy(Info.ModuleName, ANSI_TO_TCHAR(ImageName));
 			Info.PdbAge = 0;
 			Info.PdbSig = 0;
 			FMemory::Memzero(&Info.PdbSig70, sizeof(Info.PdbSig70));
-			Info.TimeDateStamp = GetModuleTimeStamp( Header );
+			Info.TimeDateStamp = GetModuleTimeStamp(Header);
 
-			ModuleSignatures[SignatureIndex] = Info;
-			++SignatureIndex;
+			TempModules.Add(Info);
 		}
 	}
 	
-	return SignatureIndex;
+	TempModules.Sort([](const FStackWalkModuleInfo& A, const FStackWalkModuleInfo& B)
+	{
+		return A.BaseOfImage < B.BaseOfImage;
+	});
+
+	for (int32 Idx = 0; Idx < TempModules.Num(); Idx++)
+	{
+		if (Idx + 1 < TempModules.Num())
+		{
+			uint64 CurrentBase = TempModules[Idx].BaseOfImage;
+			uint64 NextBase = TempModules[Idx + 1].BaseOfImage;
+
+			TempModules[Idx].ImageSize = (uint32)(NextBase - CurrentBase);
+		}
+		else
+		{
+			TempModules[Idx].ImageSize = 0xFFFFFFFF;
+		}
+	}
+
+	int32 SignatureCount = FMath::Min(TempModules.Num(), ModuleSignaturesSize);
+	for (int32 Idx = 0; Idx < SignatureCount; Idx++)
+	{
+		ModuleSignatures[Idx] = TempModules[Idx];
+	}
+
+	return SignatureCount;
 }
 
 // Consider this for iOS
