Fixed "ignoring return value of function declared with const attribute" error in ensures.

diff --git a/Engine/Source/Runtime/Core/Public/Misc/AssertionMacros.h b/Engine/Source/Runtime/Core/Public/Misc/AssertionMacros.h
--- a/Engine/Source/Runtime/Core/Public/Misc/AssertionMacros.h
+++ b/Engine/Source/Runtime/Core/Public/Misc/AssertionMacros.h
@@ -461,10 +461,18 @@ RetType UE_COLD UE_DEBUG_SECTION DispatchCheckVerify (InnerType&& Inner, ArgType
 
 #else // DO_ENSURE && !USING_CODE_ANALYSIS
 
-	#define ensure(           InExpression                ) (LIKELY(!!(InExpression)))
-	#define ensureMsgf(       InExpression, InFormat, ... ) (LIKELY(!!(InExpression)))
-	#define ensureAlways(     InExpression                ) (LIKELY(!!(InExpression)))
-#define ensureAlwaysMsgf( InExpression, InFormat, ... ) (LIKELY(!!(InExpression)))
+	// The LIKELY macro can be implemented in terms of an intrinsic which doesn't want its return value discarded,
+	// so force the value through a function which isn't marked [[nodiscard]].  This allows the macros to be used
+	// as statements.
+	UE_REWRITE constexpr bool SkipNoDiscardBool(bool bBool)
+	{
+		return bBool;
+	}
+
+	#define ensure(           InExpression                ) ::SkipNoDiscardBool(LIKELY(!!(InExpression)))
+	#define ensureMsgf(       InExpression, InFormat, ... ) ::SkipNoDiscardBool(LIKELY(!!(InExpression)))
+	#define ensureAlways(     InExpression                ) ::SkipNoDiscardBool(LIKELY(!!(InExpression)))
+	#define ensureAlwaysMsgf( InExpression, InFormat, ... ) ::SkipNoDiscardBool(LIKELY(!!(InExpression)))
 
 #endif	// DO_ENSURE && !USING_CODE_ANALYSIS
 
