Change ConsoleManager command registration to non-inline

diff --git a/Engine/Source/Runtime/Core/Private/HAL/ConsoleManager.cpp b/Engine/Source/Runtime/Core/Private/HAL/ConsoleManager.cpp
--- a/Engine/Source/Runtime/Core/Private/HAL/ConsoleManager.cpp
+++ b/Engine/Source/Runtime/Core/Private/HAL/ConsoleManager.cpp
@@ -249,7 +249,7 @@ public:
 	
 	}
 	
-	/** Legacy funciton to add old single delegates to the new multicast delegate. */
+	/** Legacy function to add old single delegates to the new multicast delegate. */
 	virtual void SetOnChangedCallback(const FConsoleVariableDelegate& Callback) 
 	{
 		OnChangedCallback.Remove(LegacyDelegateHandle);
@@ -5364,6 +5364,31 @@ FAutoConsoleCommand::FAutoConsoleCommand(const TCHAR* Name, const TCHAR* Help, c
 {
 }
 
+FAutoConsoleCommandWithWorld::FAutoConsoleCommandWithWorld(const TCHAR* Name, const TCHAR* Help, const FConsoleCommandWithWorldDelegate& Command, uint32 Flags)
+	: FAutoConsoleObject(IConsoleManager::Get().RegisterConsoleCommand(Name, Help, Command, Flags))
+{
+}
+
+FAutoConsoleCommandWithWorldAndArgs::FAutoConsoleCommandWithWorldAndArgs(const TCHAR* Name, const TCHAR* Help, const FConsoleCommandWithWorldAndArgsDelegate& Command, uint32 Flags)
+	: FAutoConsoleObject(IConsoleManager::Get().RegisterConsoleCommand(Name, Help, Command, Flags))
+{
+}
+
+FAutoConsoleCommandWithArgsAndOutputDevice::FAutoConsoleCommandWithArgsAndOutputDevice(const TCHAR* Name, const TCHAR* Help, const FConsoleCommandWithArgsAndOutputDeviceDelegate& Command, uint32 Flags)
+	: FAutoConsoleObject(IConsoleManager::Get().RegisterConsoleCommand(Name, Help, Command, Flags))
+{
+}
+
+FAutoConsoleCommandWithOutputDevice::FAutoConsoleCommandWithOutputDevice(const TCHAR* Name, const TCHAR* Help, const FConsoleCommandWithOutputDeviceDelegate& Command, uint32 Flags)
+	: FAutoConsoleObject(IConsoleManager::Get().RegisterConsoleCommand(Name, Help, Command, Flags))
+{
+}
+
+FAutoConsoleCommandWithWorldArgsAndOutputDevice::FAutoConsoleCommandWithWorldArgsAndOutputDevice(const TCHAR* Name, const TCHAR* Help, const FConsoleCommandWithWorldArgsAndOutputDeviceDelegate& Command, uint32 Flags)
+	: FAutoConsoleObject(IConsoleManager::Get().RegisterConsoleCommand(Name, Help, Command, Flags))
+{
+}
+
 #endif
 
 #undef LOCTEXT_NAMESPACE
diff --git a/Engine/Source/Runtime/Core/Public/HAL/IConsoleManager.h b/Engine/Source/Runtime/Core/Public/HAL/IConsoleManager.h
--- a/Engine/Source/Runtime/Core/Public/HAL/IConsoleManager.h
+++ b/Engine/Source/Runtime/Core/Public/HAL/IConsoleManager.h
@@ -2125,12 +2125,7 @@ public:
 	 * @param	Command		The user function to call when this command is executed
 	 * @param	Flags		Optional flags bitmask
 	 */
-	FAutoConsoleCommandWithWorld(const TCHAR* Name, const TCHAR* Help, const FConsoleCommandWithWorldDelegate& Command, uint32 Flags = ECVF_Default)
-		: FAutoConsoleObject(IConsoleManager::Get().RegisterConsoleCommand(Name, Help, Command, Flags))
-	{
-	}
-
-	
+	CORE_API FAutoConsoleCommandWithWorld(const TCHAR* Name, const TCHAR* Help, const FConsoleCommandWithWorldDelegate& Command, uint32 Flags = ECVF_Default);
 };
 
 /**
@@ -2147,10 +2142,7 @@ public:
 	 * @param	Command		The user function to call when this command is executed
 	 * @param	Flags		Optional flags bitmask
 	 */
-	FAutoConsoleCommandWithWorldAndArgs(const TCHAR* Name, const TCHAR* Help, const FConsoleCommandWithWorldAndArgsDelegate& Command, uint32 Flags = ECVF_Default)
-		: FAutoConsoleObject(IConsoleManager::Get().RegisterConsoleCommand(Name, Help, Command, Flags))
-	{
-	}
+	CORE_API FAutoConsoleCommandWithWorldAndArgs(const TCHAR* Name, const TCHAR* Help, const FConsoleCommandWithWorldAndArgsDelegate& Command, uint32 Flags = ECVF_Default);
 };
 
 /**
@@ -2160,10 +2152,7 @@ class FAutoConsoleCommandWithArgsAndOutputDevice : private FAutoConsoleObject
 {
 public:
 	
-	FAutoConsoleCommandWithArgsAndOutputDevice(const TCHAR* Name, const TCHAR* Help, const FConsoleCommandWithArgsAndOutputDeviceDelegate& Command, uint32 Flags = ECVF_Default)
-		: FAutoConsoleObject(IConsoleManager::Get().RegisterConsoleCommand(Name, Help, Command, Flags))
-	{
-	}
+	CORE_API FAutoConsoleCommandWithArgsAndOutputDevice(const TCHAR* Name, const TCHAR* Help, const FConsoleCommandWithArgsAndOutputDeviceDelegate& Command, uint32 Flags = ECVF_Default);
 };
 
 /**
@@ -2180,10 +2169,7 @@ public:
 	 * @param	Command		The user function to call when this command is executed
 	 * @param	Flags		Optional flags bitmask
 	 */
-	FAutoConsoleCommandWithOutputDevice(const TCHAR* Name, const TCHAR* Help, const FConsoleCommandWithOutputDeviceDelegate& Command, uint32 Flags = ECVF_Default)
-		: FAutoConsoleObject(IConsoleManager::Get().RegisterConsoleCommand(Name, Help, Command, Flags))
-	{
-	}
+	CORE_API FAutoConsoleCommandWithOutputDevice(const TCHAR* Name, const TCHAR* Help, const FConsoleCommandWithOutputDeviceDelegate& Command, uint32 Flags = ECVF_Default);
 };
 
 /**
@@ -2200,10 +2186,7 @@ public:
 	 * @param	Command		The user function to call when this command is executed
 	 * @param	Flags		Optional flags bitmask
 	 */
-	FAutoConsoleCommandWithWorldArgsAndOutputDevice(const TCHAR* Name, const TCHAR* Help, const FConsoleCommandWithWorldArgsAndOutputDeviceDelegate& Command, uint32 Flags = ECVF_Default)
-		: FAutoConsoleObject(IConsoleManager::Get().RegisterConsoleCommand(Name, Help, Command, Flags))
-	{
-	}
+	CORE_API FAutoConsoleCommandWithWorldArgsAndOutputDevice(const TCHAR* Name, const TCHAR* Help, const FConsoleCommandWithWorldArgsAndOutputDeviceDelegate& Command, uint32 Flags = ECVF_Default);
 };
 
 #else
