Make FPixelFormatInfo ctor constexpr

For whatever reason Clang wasn't compiling the pixel formats array into pure data
without on some targets, but does with.

diff --git a/Engine/Source/Runtime/Core/Private/Misc/PixelFormat.cpp b/Engine/Source/Runtime/Core/Private/Misc/PixelFormat.cpp
--- a/Engine/Source/Runtime/Core/Private/Misc/PixelFormat.cpp
+++ b/Engine/Source/Runtime/Core/Private/Misc/PixelFormat.cpp
@@ -7,7 +7,7 @@
 //	Pixel format information.
 //
 
-FPixelFormatInfo::FPixelFormatInfo(
+constexpr FPixelFormatInfo::FPixelFormatInfo(
 	EPixelFormat InUnrealFormat,
 	const TCHAR* InName,
 	int32 InBlockSizeX,
diff --git a/Engine/Source/Runtime/Core/Public/PixelFormat.h b/Engine/Source/Runtime/Core/Public/PixelFormat.h
--- a/Engine/Source/Runtime/Core/Public/PixelFormat.h
+++ b/Engine/Source/Runtime/Core/Public/PixelFormat.h
@@ -452,7 +452,7 @@ UE_FORCEINLINE_HINT bool RequiresBlock4Alignment(EPixelFormat PixelFormat)
 struct FPixelFormatInfo
 {
 	FPixelFormatInfo() = delete;
-	FPixelFormatInfo(
+	constexpr FPixelFormatInfo(
 		EPixelFormat InUnrealFormat,
 		const TCHAR* InName,
 		int32 InBlockSizeX,
