Fix for build failiure.

diff --git a/Engine/Source/Runtime/Core/Public/Windows/WindowsPlatform.h b/Engine/Source/Runtime/Core/Public/Windows/WindowsPlatform.h
--- a/Engine/Source/Runtime/Core/Public/Windows/WindowsPlatform.h
+++ b/Engine/Source/Runtime/Core/Public/Windows/WindowsPlatform.h
@@ -210,8 +210,12 @@ typedef FWindowsPlatformTypes FPlatformTypes;
 #define PLATFORM_CACHE_LINE_SIZE	64
 
 // DLL export and import definitions
+#ifndef DLLEXPORT
 #define DLLEXPORT __declspec(dllexport)
+#endif
+#ifndef DLLIMPORT
 #define DLLIMPORT __declspec(dllimport)
+#endif
 
 // disable this now as it is annoying for generic platform implementations
 #pragma warning(disable : 4100) // unreferenced formal parameter
