Fix for build failure.

diff --git a/Engine/Source/Runtime/Core/Public/Unix/UnixPlatform.h b/Engine/Source/Runtime/Core/Public/Unix/UnixPlatform.h
--- a/Engine/Source/Runtime/Core/Public/Unix/UnixPlatform.h
+++ b/Engine/Source/Runtime/Core/Public/Unix/UnixPlatform.h
@@ -123,8 +123,12 @@ typedef FUnixPlatformTypes FPlatformTypes;
 #define ABSTRACT abstract
 
 // DLL export and import definitions
+#ifndef DLLEXPORT
 #define DLLEXPORT			__attribute__((visibility("default")))
+#endif
+#ifndef DLLIMPORT
 #define DLLIMPORT			__attribute__((visibility("default")))
+#endif
 
 // Alignment.
 #define GCC_PACK(n)			__attribute__((packed,aligned(n)))
