Revert platform header file changes.

diff --git a/Engine/Source/Runtime/Core/Public/Unix/UnixPlatform.h b/Engine/Source/Runtime/Core/Public/Unix/UnixPlatform.h
--- a/Engine/Source/Runtime/Core/Public/Unix/UnixPlatform.h
+++ b/Engine/Source/Runtime/Core/Public/Unix/UnixPlatform.h
@@ -123,12 +123,8 @@ typedef FUnixPlatformTypes FPlatformTypes;
 #define ABSTRACT abstract
 
 // DLL export and import definitions
-#ifndef DLLEXPORT
 #define DLLEXPORT			__attribute__((visibility("default")))
-#endif
-#ifndef DLLIMPORT
 #define DLLIMPORT			__attribute__((visibility("default")))
-#endif
 
 // Alignment.
 #define GCC_PACK(n)			__attribute__((packed,aligned(n)))
diff --git a/Engine/Source/Runtime/Core/Public/Windows/WindowsPlatform.h b/Engine/Source/Runtime/Core/Public/Windows/WindowsPlatform.h
--- a/Engine/Source/Runtime/Core/Public/Windows/WindowsPlatform.h
+++ b/Engine/Source/Runtime/Core/Public/Windows/WindowsPlatform.h
@@ -210,12 +210,8 @@ typedef FWindowsPlatformTypes FPlatformTypes;
 #define PLATFORM_CACHE_LINE_SIZE	64
 
 // DLL export and import definitions
-#ifndef DLLEXPORT
 #define DLLEXPORT __declspec(dllexport)
-#endif
-#ifndef DLLIMPORT
 #define DLLIMPORT __declspec(dllimport)
-#endif
 
 // disable this now as it is annoying for generic platform implementations
 #pragma warning(disable : 4100) // unreferenced formal parameter
