FBox2D now always has its bIsValid member initialized to false to stay consistent with FBox.

diff --git a/Engine/Source/Runtime/Core/Public/Math/Box2D.h b/Engine/Source/Runtime/Core/Public/Math/Box2D.h
--- a/Engine/Source/Runtime/Core/Public/Math/Box2D.h
+++ b/Engine/Source/Runtime/Core/Public/Math/Box2D.h
@@ -39,7 +39,7 @@ public:
 	TVector2<T> Max;
 
 	/** Holds a flag indicating whether this box is valid. */
-	bool bIsValid;
+	bool bIsValid = false;
 
 public:
 
