Remove UE_WITH_CONSTINIT_UOBJECT from Build.h where it may not be defined.
A future CL will handle these mutually exclusive flags in UBT.

diff --git a/Engine/Source/Runtime/Core/Public/Misc/Build.h b/Engine/Source/Runtime/Core/Public/Misc/Build.h
--- a/Engine/Source/Runtime/Core/Public/Misc/Build.h
+++ b/Engine/Source/Runtime/Core/Public/Misc/Build.h
@@ -147,7 +147,7 @@
 * Whether we support hot-reload. Currently requires a non-monolithic build and non-shipping configuration.
 */
 #ifndef WITH_HOT_RELOAD
-	#define WITH_HOT_RELOAD (!IS_MONOLITHIC && !UE_BUILD_SHIPPING && !UE_BUILD_TEST && !UE_GAME && !UE_SERVER && !UE_WITH_CONSTINIT_UOBJECT)
+	#define WITH_HOT_RELOAD (!IS_MONOLITHIC && !UE_BUILD_SHIPPING && !UE_BUILD_TEST && !UE_GAME && !UE_SERVER)
 #endif
 
 /**
