Fix casting errors for w4388.

#rnx

- Change index variables to the same type that they will be compared against to prevent the compiler raising w4388 warnings.

diff --git a/Engine/Source/Runtime/Core/Public/Experimental/UnifiedError/UnifiedError.h b/Engine/Source/Runtime/Core/Public/Experimental/UnifiedError/UnifiedError.h
--- a/Engine/Source/Runtime/Core/Public/Experimental/UnifiedError/UnifiedError.h
+++ b/Engine/Source/Runtime/Core/Public/Experimental/UnifiedError/UnifiedError.h
@@ -68,8 +68,8 @@ struct TStringLiteralWithColonsReplaced
 {
 	constexpr TStringLiteralWithColonsReplaced(const char(&InStr)[N])
 	{
-		int OutIndex = 0;
-		int i = 0;
+		unsigned int OutIndex = 0;
+		unsigned int i = 0;
 		for (; i < N; ++i)
 		{
 			if (InStr[i] == '\0')
@@ -103,8 +103,8 @@ struct TStringLiteralWithSpacesReplaced
 {
 	constexpr TStringLiteralWithSpacesReplaced(const char(&InStr)[N])
 	{
-		int OutIndex = 0;
-		int i = 0;
+		unsigned int OutIndex = 0;
+		unsigned int i = 0;
 		for (; i < N; ++i)
 		{
 			if (InStr[i] == '\0')
