Delegate CreateCopy code size improvement.

CreateCopy is only ever called wih non-escaped local vars as the target, meaning no locking is necessary.
This in turn means we don't need three CreateCopy overloads for the different ThreadSafetyMode variants,
a single one suffices.

This gets rid of 2 (often non-trivial) methods in every delegate instance and 2 vtable entries per delegate.

The mechanism used to implement non-locked allocations is also applicable to the various Create methods
on delegates, which likewise construct their return value into a private local variable and hence don't need
to control write access to it with a lock.

diff --git a/Engine/Source/Runtime/Core/Public/Delegates/DelegateBase.h b/Engine/Source/Runtime/Core/Public/Delegates/DelegateBase.h
--- a/Engine/Source/Runtime/Core/Public/Delegates/DelegateBase.h
+++ b/Engine/Source/Runtime/Core/Public/Delegates/DelegateBase.h
@@ -99,9 +99,11 @@ struct FDelegateAllocation
 
 template <typename ThreadSafetyMode>
 struct TWriteLockedDelegateAllocation;
+struct FPrivateDelegateAllocation;
 
 template <typename ThreadSafetyMode>
 void* operator new(size_t Size, const TWriteLockedDelegateAllocation<ThreadSafetyMode>& LockedAllocation);
+void* operator new(size_t Size, const FPrivateDelegateAllocation& Allocation);
 
 template <typename ThreadSafetyMode>
 struct TWriteLockedDelegateAllocation
@@ -121,6 +123,27 @@ private:
 	FDelegateAllocation& Allocation;
 };
 
+/**
+ * Used for delegate allocations when the underlying allocation is known to be not (yet)
+ * visible to other threads so locking is unnecessary. (Mainly when the target is
+ * a temporary local variable.)
+ */
+struct FPrivateDelegateAllocation
+{
+	UE_NONCOPYABLE(FPrivateDelegateAllocation);
+
+	friend void* operator new(size_t Size, const FPrivateDelegateAllocation& Allocation);
+
+	template <typename ThreadSafetyMode>
+	explicit FPrivateDelegateAllocation(TDelegateBase<ThreadSafetyMode>& Delegate)
+		: Allocation(Delegate)
+	{
+	}
+
+private:
+	FDelegateAllocation& Allocation;
+};
+
 /**
  * Base class for unicast delegates.
  */
@@ -128,6 +151,7 @@ template<typename ThreadSafetyMode>
 class TDelegateBase : public TDelegateAccessHandlerBase<ThreadSafetyMode>, private FDelegateAllocation
 {
 	friend struct TWriteLockedDelegateAllocation<ThreadSafetyMode>;
+	friend struct FPrivateDelegateAllocation;
 
 private:
 	using Super = TDelegateAccessHandlerBase<ThreadSafetyMode>;
@@ -416,3 +440,12 @@ template <typename ThreadSafetyMode>
 UE_FORCEINLINE_HINT void operator delete(void*, const TWriteLockedDelegateAllocation<ThreadSafetyMode>& LockedAllocation)
 {
 }
+
+UE_FORCEINLINE_HINT void* operator new(size_t Size, const FPrivateDelegateAllocation& Allocation)
+{
+	return UE::Core::Private::DelegateAllocate(Size, Allocation.Allocation);
+}
+
+UE_FORCEINLINE_HINT void operator delete(void*, const FPrivateDelegateAllocation& Allocation)
+{
+}
diff --git a/Engine/Source/Runtime/Core/Public/Delegates/DelegateInstanceInterface.h b/Engine/Source/Runtime/Core/Public/Delegates/DelegateInstanceInterface.h
--- a/Engine/Source/Runtime/Core/Public/Delegates/DelegateInstanceInterface.h
+++ b/Engine/Source/Runtime/Core/Public/Delegates/DelegateInstanceInterface.h
@@ -18,15 +18,16 @@ struct FNotThreadSafeNotCheckedDelegateMode;
 struct FThreadSafeDelegateMode;
 struct FNotThreadSafeDelegateMode;
 
+struct FPrivateDelegateAllocation;
+
 template <typename RetType, typename... ArgTypes, typename UserPolicy>
 struct IBaseDelegateInstance<RetType(ArgTypes...), UserPolicy> : public UserPolicy::FDelegateInstanceExtras
 {
 	/**
-	 * Emplaces a copy of the delegate instance into the given base delegate.
+	 * Emplaces a copy of the delegate instance into the given target delegate.
+	 * Target MUST be private since no locks are taken! (CreateCopy into a temporary, then move.)
 	 */
-	virtual void CreateCopy(TDelegateBase<FThreadSafeDelegateMode>& Base) const = 0;
-	virtual void CreateCopy(TDelegateBase<FNotThreadSafeDelegateMode>& Base) const = 0;
-	virtual void CreateCopy(TDelegateBase<FNotThreadSafeNotCheckedDelegateMode>& Base) const = 0;
+	virtual void CreateCopy(const FPrivateDelegateAllocation& Target) const = 0;
 
 	/**
 	 * Execute the delegate.  If the function pointer is not valid, an error will occur.
diff --git a/Engine/Source/Runtime/Core/Public/Delegates/DelegateInstancesImpl.h b/Engine/Source/Runtime/Core/Public/Delegates/DelegateInstancesImpl.h
--- a/Engine/Source/Runtime/Core/Public/Delegates/DelegateInstancesImpl.h
+++ b/Engine/Source/Runtime/Core/Public/Delegates/DelegateInstancesImpl.h
@@ -136,19 +136,9 @@ public:
 
 	// IBaseDelegateInstance interface
 
-	void CreateCopy(TDelegateBase<FThreadSafeDelegateMode>& Base) const final
+	void CreateCopy(const FPrivateDelegateAllocation& Target) const final
 	{
-		new (TWriteLockedDelegateAllocation{Base}) TBaseUFunctionDelegateInstance(*this);
-	}
-
-	void CreateCopy(TDelegateBase<FNotThreadSafeDelegateMode>& Base) const final
-	{
-		new (TWriteLockedDelegateAllocation{Base}) TBaseUFunctionDelegateInstance(*this);
-	}
-
-	void CreateCopy(TDelegateBase<FNotThreadSafeNotCheckedDelegateMode>& Base) const final
-	{
-		new (TWriteLockedDelegateAllocation{Base}) TBaseUFunctionDelegateInstance(*this);
+		new (Target) TBaseUFunctionDelegateInstance(*this);
 	}
 
 	RetValType Execute(ParamTypes... Params) const final
@@ -258,19 +248,9 @@ public:
 
 	// IBaseDelegateInstance interface
 
-	void CreateCopy(TDelegateBase<FThreadSafeDelegateMode>& Base) const final
-	{
-		new (TWriteLockedDelegateAllocation{Base}) TBaseSPMethodDelegateInstance(*this);
-	}
-
-	void CreateCopy(TDelegateBase<FNotThreadSafeDelegateMode>& Base) const final
-	{
-		new (TWriteLockedDelegateAllocation{Base}) TBaseSPMethodDelegateInstance(*this);
-	}
-
-	void CreateCopy(TDelegateBase<FNotThreadSafeNotCheckedDelegateMode>& Base) const final
+	void CreateCopy(const FPrivateDelegateAllocation& Target) const final
 	{
-		new (TWriteLockedDelegateAllocation{Base}) TBaseSPMethodDelegateInstance(*this);
+		new (Target) TBaseSPMethodDelegateInstance(*this);
 	}
 
 	RetValType Execute(ParamTypes... Params) const final
@@ -384,19 +364,9 @@ public:
 
 public:
 	// IBaseDelegateInstance interface
-	void CreateCopy(TDelegateBase<FThreadSafeDelegateMode>& Base) const final
+	void CreateCopy(const FPrivateDelegateAllocation& Target) const final
 	{
-		new (TWriteLockedDelegateAllocation{Base}) TBaseSPLambdaDelegateInstance(*this);
-	}
-
-	void CreateCopy(TDelegateBase<FNotThreadSafeDelegateMode>& Base) const final
-	{
-		new (TWriteLockedDelegateAllocation{Base}) TBaseSPLambdaDelegateInstance(*this);
-	}
-
-	void CreateCopy(TDelegateBase<FNotThreadSafeNotCheckedDelegateMode>& Base) const final
-	{
-		new (TWriteLockedDelegateAllocation{Base}) TBaseSPLambdaDelegateInstance(*this);
+		new (Target) TBaseSPLambdaDelegateInstance(*this);
 	}
 
 	RetValType Execute(ParamTypes... Params) const final
@@ -500,19 +470,9 @@ public:
 
 	// IBaseDelegateInstance interface
 
-	void CreateCopy(TDelegateBase<FThreadSafeDelegateMode>& Base) const final
-	{
-		new (TWriteLockedDelegateAllocation{Base}) TBaseRawMethodDelegateInstance(*this);
-	}
-
-	void CreateCopy(TDelegateBase<FNotThreadSafeDelegateMode>& Base) const final
-	{
-		new (TWriteLockedDelegateAllocation{Base}) TBaseRawMethodDelegateInstance(*this);
-	}
-
-	void CreateCopy(TDelegateBase<FNotThreadSafeNotCheckedDelegateMode>& Base) const final
+	void CreateCopy(const FPrivateDelegateAllocation& Target) const final
 	{
-		new (TWriteLockedDelegateAllocation{Base}) TBaseRawMethodDelegateInstance(*this);
+		new (Target) TBaseRawMethodDelegateInstance(*this);
 	}
 
 	RetValType Execute(ParamTypes... Params) const final
@@ -639,19 +599,9 @@ public:
 
 	// IBaseDelegateInstance interface
 
-	void CreateCopy(TDelegateBase<FThreadSafeDelegateMode>& Base) const final
+	void CreateCopy(const FPrivateDelegateAllocation& Target) const final
 	{
-		new (TWriteLockedDelegateAllocation{Base}) TBaseUObjectMethodDelegateInstance(*this);
-	}
-
-	void CreateCopy(TDelegateBase<FNotThreadSafeDelegateMode>& Base) const final
-	{
-		new (TWriteLockedDelegateAllocation{Base}) TBaseUObjectMethodDelegateInstance(*this);
-	}
-
-	void CreateCopy(TDelegateBase<FNotThreadSafeNotCheckedDelegateMode>& Base) const final
-	{
-		new (TWriteLockedDelegateAllocation{Base}) TBaseUObjectMethodDelegateInstance(*this);
+		new (Target) TBaseUObjectMethodDelegateInstance(*this);
 	}
 
 	RetValType Execute(ParamTypes... Params) const final
@@ -767,19 +717,9 @@ public:
 
 	// IBaseDelegateInstance interface
 
-	void CreateCopy(TDelegateBase<FThreadSafeDelegateMode>& Base) const final
-	{
-		new (TWriteLockedDelegateAllocation{Base}) TBaseStaticDelegateInstance(*this);
-	}
-
-	void CreateCopy(TDelegateBase<FNotThreadSafeDelegateMode>& Base) const final
-	{
-		new (TWriteLockedDelegateAllocation{Base}) TBaseStaticDelegateInstance(*this);
-	}
-
-	void CreateCopy(TDelegateBase<FNotThreadSafeNotCheckedDelegateMode>& Base) const final
+	void CreateCopy(const FPrivateDelegateAllocation& Target) const final
 	{
-		new (TWriteLockedDelegateAllocation{Base}) TBaseStaticDelegateInstance(*this);
+		new (Target) TBaseStaticDelegateInstance(*this);
 	}
 
 	RetValType Execute(ParamTypes... Params) const final
@@ -865,19 +805,9 @@ public:
 
 public:
 	// IBaseDelegateInstance interface
-	void CreateCopy(TDelegateBase<FThreadSafeDelegateMode>& Base) const final
+	void CreateCopy(const FPrivateDelegateAllocation& Target) const final
 	{
-		new (TWriteLockedDelegateAllocation{Base}) TBaseFunctorDelegateInstance(*this);
-	}
-
-	void CreateCopy(TDelegateBase<FNotThreadSafeDelegateMode>& Base) const final
-	{
-		new (TWriteLockedDelegateAllocation{Base}) TBaseFunctorDelegateInstance(*this);
-	}
-
-	void CreateCopy(TDelegateBase<FNotThreadSafeNotCheckedDelegateMode>& Base) const final
-	{
-		new (TWriteLockedDelegateAllocation{Base}) TBaseFunctorDelegateInstance(*this);
+		new (Target) TBaseFunctorDelegateInstance(*this);
 	}
 
 	RetValType Execute(ParamTypes... Params) const final
@@ -978,19 +908,9 @@ public:
 
 public:
 	// IBaseDelegateInstance interface
-	void CreateCopy(TDelegateBase<FThreadSafeDelegateMode>& Base) const final
-	{
-		new (TWriteLockedDelegateAllocation{Base}) TWeakBaseFunctorDelegateInstance(*this);
-	}
-
-	void CreateCopy(TDelegateBase<FNotThreadSafeDelegateMode>& Base) const final
-	{
-		new (TWriteLockedDelegateAllocation{Base}) TWeakBaseFunctorDelegateInstance(*this);
-	}
-
-	void CreateCopy(TDelegateBase<FNotThreadSafeNotCheckedDelegateMode>& Base) const final
+	void CreateCopy(const FPrivateDelegateAllocation& Target) const final
 	{
-		new (TWriteLockedDelegateAllocation{Base}) TWeakBaseFunctorDelegateInstance(*this);
+		new (Target) TWeakBaseFunctorDelegateInstance(*this);
 	}
 
 	RetValType Execute(ParamTypes... Params) const final
diff --git a/Engine/Source/Runtime/Core/Public/Delegates/DelegateSignatureImpl.inl b/Engine/Source/Runtime/Core/Public/Delegates/DelegateSignatureImpl.inl
--- a/Engine/Source/Runtime/Core/Public/Delegates/DelegateSignatureImpl.inl
+++ b/Engine/Source/Runtime/Core/Public/Delegates/DelegateSignatureImpl.inl
@@ -365,7 +365,7 @@ public:
 	[[nodiscard]] inline static TDelegate<RetValType(ParamTypes...), UserPolicy> CreateStatic(TIdentity_T<RetValType (*)(ParamTypes..., std::decay_t<VarTypes>...)> InFunc, VarTypes&&... Vars)
 	{
 		TDelegate<RetValType(ParamTypes...), UserPolicy> Result;
-		new (TWriteLockedDelegateAllocation{Result}) TBaseStaticDelegateInstance<FuncType, UserPolicy, std::decay_t<VarTypes>...>(InFunc, Forward<VarTypes>(Vars)...);
+		new (FPrivateDelegateAllocation{Result}) TBaseStaticDelegateInstance<FuncType, UserPolicy, std::decay_t<VarTypes>...>(InFunc, Forward<VarTypes>(Vars)...);
 		return Result;
 	}
 
@@ -377,7 +377,7 @@ public:
 	[[nodiscard]] inline static TDelegate<RetValType(ParamTypes...), UserPolicy> CreateLambda(FunctorType&& InFunctor, VarTypes&&... Vars)
 	{
 		TDelegate<RetValType(ParamTypes...), UserPolicy> Result;
-		new (TWriteLockedDelegateAllocation{Result}) TBaseFunctorDelegateInstance<FuncType, UserPolicy, std::remove_reference_t<FunctorType>, std::decay_t<VarTypes>...>(Forward<FunctorType>(InFunctor), Forward<VarTypes>(Vars)...);
+		new (FPrivateDelegateAllocation{Result}) TBaseFunctorDelegateInstance<FuncType, UserPolicy, std::remove_reference_t<FunctorType>, std::decay_t<VarTypes>...>(Forward<FunctorType>(InFunctor), Forward<VarTypes>(Vars)...);
 		return Result;
 	}
 
@@ -389,14 +389,14 @@ public:
 	[[nodiscard]] inline static TDelegate<RetValType(ParamTypes...), UserPolicy> CreateSPLambda(const TSharedRef<UserClass, Mode>& InUserObjectRef, FunctorType&& InFunctor, VarTypes&&... Vars)
 	{
 		TDelegate<RetValType(ParamTypes...), UserPolicy> Result;
-		new (TWriteLockedDelegateAllocation{Result}) TBaseSPLambdaDelegateInstance<Mode, FuncType, UserPolicy, std::remove_reference_t<FunctorType>, std::decay_t<VarTypes>...>(InUserObjectRef, Forward<FunctorType>(InFunctor), Forward<VarTypes>(Vars)...);
+		new (FPrivateDelegateAllocation{Result}) TBaseSPLambdaDelegateInstance<Mode, FuncType, UserPolicy, std::remove_reference_t<FunctorType>, std::decay_t<VarTypes>...>(InUserObjectRef, Forward<FunctorType>(InFunctor), Forward<VarTypes>(Vars)...);
 		return Result;
 	}
 	template <typename UserClass, typename FunctorType, typename... VarTypes>
 	[[nodiscard]] inline static TDelegate<RetValType(ParamTypes...), UserPolicy> CreateSPLambda(UserClass* InUserObject, FunctorType&& InFunctor, VarTypes&&... Vars)
 	{
 		TDelegate<RetValType(ParamTypes...), UserPolicy> Result;
-		new (TWriteLockedDelegateAllocation{Result}) TBaseSPLambdaDelegateInstance<decltype(InUserObject->AsWeak())::Mode, FuncType, UserPolicy, std::remove_reference_t<FunctorType>, std::decay_t<VarTypes>...>(TWeakPtr<const void>(InUserObject->AsWeak()), Forward<FunctorType>(InFunctor), Forward<VarTypes>(Vars)...);
+		new (FPrivateDelegateAllocation{Result}) TBaseSPLambdaDelegateInstance<decltype(InUserObject->AsWeak())::Mode, FuncType, UserPolicy, std::remove_reference_t<FunctorType>, std::decay_t<VarTypes>...>(TWeakPtr<const void>(InUserObject->AsWeak()), Forward<FunctorType>(InFunctor), Forward<VarTypes>(Vars)...);
 		return Result;
 	}
 
@@ -408,7 +408,7 @@ public:
 	[[nodiscard]] inline static TDelegate<RetValType(ParamTypes...), UserPolicy> CreateWeakLambda(const UObject* InUserObject, FunctorType&& InFunctor, VarTypes&&... Vars)
 	{
 		TDelegate<RetValType(ParamTypes...), UserPolicy> Result;
-		new (TWriteLockedDelegateAllocation{Result}) TWeakBaseFunctorDelegateInstance<FuncType, UserPolicy, std::remove_reference_t<FunctorType>, std::decay_t<VarTypes>...>(InUserObject, Forward<FunctorType>(InFunctor), Forward<VarTypes>(Vars)...);
+		new (FPrivateDelegateAllocation{Result}) TWeakBaseFunctorDelegateInstance<FuncType, UserPolicy, std::remove_reference_t<FunctorType>, std::decay_t<VarTypes>...>(InUserObject, Forward<FunctorType>(InFunctor), Forward<VarTypes>(Vars)...);
 		return Result;
 	}
 
@@ -424,14 +424,14 @@ public:
 		static_assert(!std::is_const_v<UserClass>, "Attempting to bind a delegate with a const object pointer and non-const member function.");
 
 		TDelegate<RetValType(ParamTypes...), UserPolicy> Result;
-		new (TWriteLockedDelegateAllocation{Result}) TBaseRawMethodDelegateInstance<false, UserClass, FuncType, UserPolicy, std::decay_t<VarTypes>...>(InUserObject, InFunc, Forward<VarTypes>(Vars)...);
+		new (FPrivateDelegateAllocation{Result}) TBaseRawMethodDelegateInstance<false, UserClass, FuncType, UserPolicy, std::decay_t<VarTypes>...>(InUserObject, InFunc, Forward<VarTypes>(Vars)...);
 		return Result;
 	}
 	template <typename UserClass, typename... VarTypes>
 	[[nodiscard]] inline static TDelegate<RetValType(ParamTypes...), UserPolicy> CreateRaw(const UserClass* InUserObject, typename TMemFunPtrType<true, UserClass, RetValType (ParamTypes..., std::decay_t<VarTypes>...)>::Type InFunc, VarTypes&&... Vars)
 	{
 		TDelegate<RetValType(ParamTypes...), UserPolicy> Result;
-		new (TWriteLockedDelegateAllocation{Result}) TBaseRawMethodDelegateInstance<true, const UserClass, FuncType, UserPolicy, std::decay_t<VarTypes>...>(InUserObject, InFunc, Forward<VarTypes>(Vars)...);
+		new (FPrivateDelegateAllocation{Result}) TBaseRawMethodDelegateInstance<true, const UserClass, FuncType, UserPolicy, std::decay_t<VarTypes>...>(InUserObject, InFunc, Forward<VarTypes>(Vars)...);
 		return Result;
 	}
 
@@ -447,14 +447,14 @@ public:
 		static_assert(!std::is_const_v<UserClass>, "Attempting to bind a delegate with a const object pointer and non-const member function.");
 
 		TDelegate<RetValType(ParamTypes...), UserPolicy> Result;
-		new (TWriteLockedDelegateAllocation{Result}) TBaseSPMethodDelegateInstance<false, UserClass, Mode, FuncType, UserPolicy, std::decay_t<VarTypes>...>(InUserObjectRef, InFunc, Forward<VarTypes>(Vars)...);
+		new (FPrivateDelegateAllocation{Result}) TBaseSPMethodDelegateInstance<false, UserClass, Mode, FuncType, UserPolicy, std::decay_t<VarTypes>...>(InUserObjectRef, InFunc, Forward<VarTypes>(Vars)...);
 		return Result;
 	}
 	template <typename UserClass, ESPMode Mode, typename... VarTypes>
 	[[nodiscard]] inline static TDelegate<RetValType(ParamTypes...), UserPolicy> CreateSP(const TSharedRef<UserClass, Mode>& InUserObjectRef, typename TMemFunPtrType<true, UserClass, RetValType (ParamTypes..., std::decay_t<VarTypes>...)>::Type InFunc, VarTypes&&... Vars)
 	{
 		TDelegate<RetValType(ParamTypes...), UserPolicy> Result;
-		new (TWriteLockedDelegateAllocation{Result}) TBaseSPMethodDelegateInstance<true, const UserClass, Mode, FuncType, UserPolicy, std::decay_t<VarTypes>...>(InUserObjectRef, InFunc, Forward<VarTypes>(Vars)...);
+		new (FPrivateDelegateAllocation{Result}) TBaseSPMethodDelegateInstance<true, const UserClass, Mode, FuncType, UserPolicy, std::decay_t<VarTypes>...>(InUserObjectRef, InFunc, Forward<VarTypes>(Vars)...);
 		return Result;
 	}
 
@@ -470,14 +470,14 @@ public:
 		static_assert(!std::is_const_v<UserClass>, "Attempting to bind a delegate with a const object pointer and non-const member function.");
 
 		TDelegate<RetValType(ParamTypes...), UserPolicy> Result;
-		new (TWriteLockedDelegateAllocation{Result}) TBaseSPMethodDelegateInstance<false, UserClass, decltype(InUserObject->AsShared())::Mode, FuncType, UserPolicy, std::decay_t<VarTypes>...>(StaticCastSharedRef<UserClass>(InUserObject->AsShared()), InFunc, Forward<VarTypes>(Vars)...);
+		new (FPrivateDelegateAllocation{Result}) TBaseSPMethodDelegateInstance<false, UserClass, decltype(InUserObject->AsShared())::Mode, FuncType, UserPolicy, std::decay_t<VarTypes>...>(StaticCastSharedRef<UserClass>(InUserObject->AsShared()), InFunc, Forward<VarTypes>(Vars)...);
 		return Result;
 	}
 	template <typename UserClass, typename... VarTypes>
 	[[nodiscard]] inline static TDelegate<RetValType(ParamTypes...), UserPolicy> CreateSP(const UserClass* InUserObject, typename TMemFunPtrType<true, UserClass, RetValType (ParamTypes..., std::decay_t<VarTypes>...)>::Type InFunc, VarTypes&&... Vars)
 	{
 		TDelegate<RetValType(ParamTypes...), UserPolicy> Result;
-		new (TWriteLockedDelegateAllocation{Result}) TBaseSPMethodDelegateInstance<true, const UserClass, decltype(InUserObject->AsShared())::Mode, FuncType, UserPolicy, std::decay_t<VarTypes>...>(StaticCastSharedRef<const UserClass>(InUserObject->AsShared()), InFunc, Forward<VarTypes>(Vars)...);
+		new (FPrivateDelegateAllocation{Result}) TBaseSPMethodDelegateInstance<true, const UserClass, decltype(InUserObject->AsShared())::Mode, FuncType, UserPolicy, std::decay_t<VarTypes>...>(StaticCastSharedRef<const UserClass>(InUserObject->AsShared()), InFunc, Forward<VarTypes>(Vars)...);
 		return Result;
 	}
 
@@ -495,14 +495,14 @@ public:
 		static_assert(!std::is_const_v<UserClass>, "Attempting to bind a delegate with a const object pointer and non-const member function.");
 
 		TDelegate<RetValType(ParamTypes...), UserPolicy> Result;
-		new (TWriteLockedDelegateAllocation{Result}) TBaseSPMethodDelegateInstance<false, UserClass, ESPMode::ThreadSafe, FuncType, UserPolicy, std::decay_t<VarTypes>...>(InUserObjectRef, InFunc, Forward<VarTypes>(Vars)...);
+		new (FPrivateDelegateAllocation{Result}) TBaseSPMethodDelegateInstance<false, UserClass, ESPMode::ThreadSafe, FuncType, UserPolicy, std::decay_t<VarTypes>...>(InUserObjectRef, InFunc, Forward<VarTypes>(Vars)...);
 		return Result;
 	}
 	template <typename UserClass, typename... VarTypes>
 	[[nodiscard]] inline static TDelegate<RetValType(ParamTypes...), UserPolicy> CreateThreadSafeSP(const TSharedRef<UserClass, ESPMode::ThreadSafe>& InUserObjectRef, typename TMemFunPtrType<true, UserClass, RetValType (ParamTypes..., std::decay_t<VarTypes>...)>::Type InFunc, VarTypes&&... Vars)
 	{
 		TDelegate<RetValType(ParamTypes...), UserPolicy> Result;
-		new (TWriteLockedDelegateAllocation{Result}) TBaseSPMethodDelegateInstance<true, const UserClass, ESPMode::ThreadSafe, FuncType, UserPolicy, std::decay_t<VarTypes>...>(InUserObjectRef, InFunc, Forward<VarTypes>(Vars)...);
+		new (FPrivateDelegateAllocation{Result}) TBaseSPMethodDelegateInstance<true, const UserClass, ESPMode::ThreadSafe, FuncType, UserPolicy, std::decay_t<VarTypes>...>(InUserObjectRef, InFunc, Forward<VarTypes>(Vars)...);
 		return Result;
 	}
 
@@ -520,14 +520,14 @@ public:
 		static_assert(!std::is_const_v<UserClass>, "Attempting to bind a delegate with a const object pointer and non-const member function.");
 
 		TDelegate<RetValType(ParamTypes...), UserPolicy> Result;
-		new (TWriteLockedDelegateAllocation{Result}) TBaseSPMethodDelegateInstance<false, UserClass, ESPMode::ThreadSafe, FuncType, UserPolicy, std::decay_t<VarTypes>...>(StaticCastSharedRef<UserClass>(InUserObject->AsShared()), InFunc, Forward<VarTypes>(Vars)...);
+		new (FPrivateDelegateAllocation{Result}) TBaseSPMethodDelegateInstance<false, UserClass, ESPMode::ThreadSafe, FuncType, UserPolicy, std::decay_t<VarTypes>...>(StaticCastSharedRef<UserClass>(InUserObject->AsShared()), InFunc, Forward<VarTypes>(Vars)...);
 		return Result;
 	}
 	template <typename UserClass, typename... VarTypes>
 	[[nodiscard]] inline static TDelegate<RetValType(ParamTypes...), UserPolicy> CreateThreadSafeSP(const UserClass* InUserObject, typename TMemFunPtrType<true, UserClass, RetValType (ParamTypes..., std::decay_t<VarTypes>...)>::Type InFunc, VarTypes&&... Vars)
 	{
 		TDelegate<RetValType(ParamTypes...), UserPolicy> Result;
-		new (TWriteLockedDelegateAllocation{Result}) TBaseSPMethodDelegateInstance<true, const UserClass, ESPMode::ThreadSafe, FuncType, UserPolicy, std::decay_t<VarTypes>...>(StaticCastSharedRef<const UserClass>(InUserObject->AsShared()), InFunc, Forward<VarTypes>(Vars)...);
+		new (FPrivateDelegateAllocation{Result}) TBaseSPMethodDelegateInstance<true, const UserClass, ESPMode::ThreadSafe, FuncType, UserPolicy, std::decay_t<VarTypes>...>(StaticCastSharedRef<const UserClass>(InUserObject->AsShared()), InFunc, Forward<VarTypes>(Vars)...);
 		return Result;
 	}
 
@@ -541,14 +541,14 @@ public:
 	[[nodiscard]] inline static TDelegate<RetValType(ParamTypes...), UserPolicy> CreateUFunction(UObjectTemplate* InUserObject, const FName& InFunctionName, VarTypes&&... Vars)
 	{
 		TDelegate<RetValType(ParamTypes...), UserPolicy> Result;
-		new (TWriteLockedDelegateAllocation{Result}) TBaseUFunctionDelegateInstance<UObjectTemplate, FuncType, UserPolicy, std::decay_t<VarTypes>...>(InUserObject, InFunctionName, Forward<VarTypes>(Vars)...);
+		new (FPrivateDelegateAllocation{Result}) TBaseUFunctionDelegateInstance<UObjectTemplate, FuncType, UserPolicy, std::decay_t<VarTypes>...>(InUserObject, InFunctionName, Forward<VarTypes>(Vars)...);
 		return Result;
 	}
 	template <typename UObjectTemplate, typename... VarTypes>
 	[[nodiscard]] inline static TDelegate<RetValType(ParamTypes...), UserPolicy> CreateUFunction(TObjectPtr<UObjectTemplate> InUserObject, const FName& InFunctionName, VarTypes&&... Vars)
 	{
 		TDelegate<RetValType(ParamTypes...), UserPolicy> Result;
-		new (TWriteLockedDelegateAllocation{Result}) TBaseUFunctionDelegateInstance<UObjectTemplate, FuncType, UserPolicy, std::decay_t<VarTypes>...>(ToRawPtr(InUserObject), InFunctionName, Forward<VarTypes>(Vars)...);
+		new (FPrivateDelegateAllocation{Result}) TBaseUFunctionDelegateInstance<UObjectTemplate, FuncType, UserPolicy, std::decay_t<VarTypes>...>(ToRawPtr(InUserObject), InFunctionName, Forward<VarTypes>(Vars)...);
 		return Result;
 	}
 
@@ -564,14 +564,14 @@ public:
 		static_assert(!std::is_const_v<UserClass>, "Attempting to bind a delegate with a const object pointer and non-const member function.");
 
 		TDelegate<RetValType(ParamTypes...), UserPolicy> Result;
-		new (TWriteLockedDelegateAllocation{Result}) TBaseUObjectMethodDelegateInstance<false, UserClass, FuncType, UserPolicy, std::decay_t<VarTypes>...>(InUserObject, InFunc, Forward<VarTypes>(Vars)...);
+		new (FPrivateDelegateAllocation{Result}) TBaseUObjectMethodDelegateInstance<false, UserClass, FuncType, UserPolicy, std::decay_t<VarTypes>...>(InUserObject, InFunc, Forward<VarTypes>(Vars)...);
 		return Result;
 	}
 	template <typename UserClass, typename... VarTypes>
 	[[nodiscard]] inline static TDelegate<RetValType(ParamTypes...), UserPolicy> CreateUObject(const UserClass* InUserObject, typename TMemFunPtrType<true, UserClass, RetValType (ParamTypes..., std::decay_t<VarTypes>...)>::Type InFunc, VarTypes&&... Vars)
 	{
 		TDelegate<RetValType(ParamTypes...), UserPolicy> Result;
-		new (TWriteLockedDelegateAllocation{Result}) TBaseUObjectMethodDelegateInstance<true, const UserClass, FuncType, UserPolicy, std::decay_t<VarTypes>...>(InUserObject, InFunc, Forward<VarTypes>(Vars)...);
+		new (FPrivateDelegateAllocation{Result}) TBaseUObjectMethodDelegateInstance<true, const UserClass, FuncType, UserPolicy, std::decay_t<VarTypes>...>(InUserObject, InFunc, Forward<VarTypes>(Vars)...);
 		return Result;
 	}
 	template <typename UserClass, typename... VarTypes>
@@ -580,14 +580,14 @@ public:
 		static_assert(!std::is_const_v<UserClass>, "Attempting to bind a delegate with a const object pointer and non-const member function.");
 
 		TDelegate<RetValType(ParamTypes...), UserPolicy> Result;
-		new (TWriteLockedDelegateAllocation{Result}) TBaseUObjectMethodDelegateInstance<false, UserClass, FuncType, UserPolicy, std::decay_t<VarTypes>...>(ToRawPtr(InUserObject), InFunc, Forward<VarTypes>(Vars)...);
+		new (FPrivateDelegateAllocation{Result}) TBaseUObjectMethodDelegateInstance<false, UserClass, FuncType, UserPolicy, std::decay_t<VarTypes>...>(ToRawPtr(InUserObject), InFunc, Forward<VarTypes>(Vars)...);
 		return Result;
 	}
 	template <typename UserClass, typename... VarTypes>
 	[[nodiscard]] inline static TDelegate<RetValType(ParamTypes...), UserPolicy> CreateUObject(TObjectPtr<UserClass> InUserObject, typename TMemFunPtrType<true, UserClass, RetValType (ParamTypes..., std::decay_t<VarTypes>...)>::Type InFunc, VarTypes&&... Vars)
 	{
 		TDelegate<RetValType(ParamTypes...), UserPolicy> Result;
-		new (TWriteLockedDelegateAllocation{Result}) TBaseUObjectMethodDelegateInstance<true, const UserClass, FuncType, UserPolicy, std::decay_t<VarTypes>...>(ToRawPtr(InUserObject), InFunc, Forward<VarTypes>(Vars)...);
+		new (FPrivateDelegateAllocation{Result}) TBaseUObjectMethodDelegateInstance<true, const UserClass, FuncType, UserPolicy, std::decay_t<VarTypes>...>(ToRawPtr(InUserObject), InFunc, Forward<VarTypes>(Vars)...);
 		return Result;
 	}
 
@@ -666,10 +666,12 @@ private:
 
 			if (OtherDelegateInstance != nullptr)
 			{
-				OtherDelegateInstance->CreateCopy(LocalCopy);
+				OtherDelegateInstance->CreateCopy(FPrivateDelegateAllocation{LocalCopy});
 			}
 		}
 
+		// Actually move-assign from our temporary into this instance. This
+		// is implemented in TDelegateBase and acquires a WriteAcessScope.
 		*this = MoveTemp(LocalCopy);
 	}
 
diff --git a/Engine/Source/Runtime/Core/Public/Delegates/MulticastDelegateBase.h b/Engine/Source/Runtime/Core/Public/Delegates/MulticastDelegateBase.h
--- a/Engine/Source/Runtime/Core/Public/Delegates/MulticastDelegateBase.h
+++ b/Engine/Source/Runtime/Core/Public/Delegates/MulticastDelegateBase.h
@@ -222,7 +222,7 @@ protected:
 				if (const IDelegateInstance* OtherInstance = OtherDelegateRef.GetDelegateInstanceProtected())
 			    {
 				    UnicastDelegateType TempDelegate;
-				    static_cast<const DelegateInstanceInterfaceType*>(OtherInstance)->CreateCopy(TempDelegate);
+					static_cast<const DelegateInstanceInterfaceType*>(OtherInstance)->CreateCopy(FPrivateDelegateAllocation{TempDelegate});
 				    TempInvocationList.Add(MoveTemp(TempDelegate));
 			    }
 		    }
