[core] modify array serialization to not re-set ArrayNum unless Ar.IsLoading; fixes some (innocuous) TSAN warnings.

diff --git a/Engine/Source/Runtime/Core/Public/Containers/Array.h b/Engine/Source/Runtime/Core/Public/Containers/Array.h
--- a/Engine/Source/Runtime/Core/Public/Containers/Array.h
+++ b/Engine/Source/Runtime/Core/Public/Containers/Array.h
@@ -4048,7 +4048,10 @@ struct TArrayPrivateFriend
 		// if we don't need to perform per-item serialization, just read it in bulk
 		if constexpr (sizeof(ElementType) == 1 || TCanBulkSerialize<ElementType>::Value)
 		{
-			A.ArrayNum = SerializeNum;
+			if (Ar.IsLoading())
+			{
+				A.ArrayNum = SerializeNum;
+			}
 
 			// Serialize simple bytes which require no construction or destruction.
 			if ((A.ArrayNum || A.ArrayMax) && Ar.IsLoading())
@@ -4098,8 +4101,6 @@ struct TArrayPrivateFriend
 		}
 		else
 		{
-			A.ArrayNum = SerializeNum;
-
 			for (SizeType i=0; i<A.ArrayNum; i++)
 			{
 				Ar << A[i];
