Exposing Apple Contacts in the UE5 Core module

diff --git a/Engine/Source/Programs/UnrealBuildTool/Platform/Apple/UEBuildApple.cs b/Engine/Source/Programs/UnrealBuildTool/Platform/Apple/UEBuildApple.cs
--- a/Engine/Source/Programs/UnrealBuildTool/Platform/Apple/UEBuildApple.cs
+++ b/Engine/Source/Programs/UnrealBuildTool/Platform/Apple/UEBuildApple.cs
@@ -55,6 +55,14 @@ namespace UnrealBuildTool
 		[CommandLine("-EnableUBSan")]
 		[XmlConfigFile(Category = "BuildConfiguration")]
 		public bool bEnableUndefinedBehaviorSanitizer = false;
+
+		/// <summary>
+		/// Enables FAppleContactManager support by linking to the Contacts framework and adding the addressbook entitlement if needed. 
+		/// NOTE: simply turning this on may introduce more paperwork during the App Store review process, so only enable this if you 
+		/// need to use that class!
+		/// </summary>
+		[ConfigFile(ConfigHierarchyType.Game, "/Script/AppleTargetPlatform.AppleTargetSettings", "bEnableContactsManager")]
+		public bool bEnableContactsManager = false;
 	}
 
 	/// <summary>
@@ -88,6 +96,7 @@ namespace UnrealBuildTool
 		public bool bEnableAddressSanitizer => Inner.bEnableAddressSanitizer;
 		public bool bEnableThreadSanitizer => Inner.bEnableThreadSanitizer;
 		public bool bEnableUndefinedBehaviorSanitizer => Inner.bEnableUndefinedBehaviorSanitizer;
+		public bool bEnableContactsManager => Inner.bEnableContactsManager;
 #pragma warning restore CS1591
 		#endregion
 	}
diff --git a/Engine/Source/Programs/UnrealBuildTool/ProjectFiles/Xcode/XcodeProject.cs b/Engine/Source/Programs/UnrealBuildTool/ProjectFiles/Xcode/XcodeProject.cs
--- a/Engine/Source/Programs/UnrealBuildTool/ProjectFiles/Xcode/XcodeProject.cs
+++ b/Engine/Source/Programs/UnrealBuildTool/ProjectFiles/Xcode/XcodeProject.cs
@@ -6,6 +6,7 @@ using System.IO;
 using System.Linq;
 using System.Text;
 using System.Text.Json;
+using System.Diagnostics;
 using EpicGames.Core;
 using Microsoft.CodeAnalysis;
 using Microsoft.Extensions.Logging;
@@ -1918,6 +1919,42 @@ namespace UnrealBuildTool.XcodeProjectXcconfig
 			Content.WriteLine($"\t\t\tproductName = \"{UnrealData.ProductName}\";");
 		}
 
+		private string? CopyAndWriteEntitlements(ref MetadataItem? EntitlementsMetadata, string ConfigDisplayName, bool bUseEntitlementsForPrograms, DirectoryReference ProjectOrEngineDir)
+		{
+			bool bEnableContactsManager = false;
+			ConfigHierarchy PlatformGameConfig = ConfigCache.ReadHierarchy(ConfigHierarchyType.Game, UnrealData.ConfigDirectory, Platform);
+			PlatformGameConfig.GetBool("/Script/AppleTargetPlatform.AppleTargetSettings", "bEnableContactsManager", out bEnableContactsManager);
+
+			// programs almost never want entitlements, so handle them specially
+			if ((UnrealData.bNoEntitlements) || (UnrealData.TargetRules.Type == TargetType.Program && !bUseEntitlementsForPrograms))
+			{
+				EntitlementsMetadata = null;
+			}
+			else if(EntitlementsMetadata != null && EntitlementsMetadata.File != null && bEnableContactsManager && Platform == UnrealTargetPlatform.Mac)
+			{
+				// if we are on iOS, IOSExports.cs already handles adding the addressbook entitlement if needed for us
+				string BuildIntermediateDir = $"{DirectoryReference.Combine(ProjectOrEngineDir, "Intermediate/Mac")}";
+				string TargetEntitlements = $"{BuildIntermediateDir}/{ConfigDisplayName}-{Path.GetFileName(EntitlementsMetadata.XcodeProjectRelative)}";
+
+				Directory.CreateDirectory(BuildIntermediateDir);
+
+				File.Copy(EntitlementsMetadata.File.FullName, TargetEntitlements, true);
+
+				FileInfo fileInfo = new FileInfo(TargetEntitlements);
+				fileInfo.IsReadOnly = false;
+
+				Process PLUtilProcess = new Process();
+				PLUtilProcess.StartInfo.FileName = "plutil";
+				PLUtilProcess.StartInfo.Arguments = $"-insert \"com\\.apple\\.security\\.personal-information\\.addressbook\" -bool true \"{TargetEntitlements}\"";
+
+				Utils.RunLocalProcess(PLUtilProcess);
+
+				return TargetEntitlements;
+			}
+
+			return EntitlementsMetadata == null ? null : EntitlementsMetadata.XcodeProjectRelative;
+		}
+
 		public override void WriteXcconfigFile(ILogger Logger)
 		{
 			// gather general, all-platform, data we are doing to put into the configs
@@ -2238,11 +2275,7 @@ namespace UnrealBuildTool.XcodeProjectXcconfig
 					}
 				}
 
-				// programs almost never want entitlements, so handle them specially
-				if ((UnrealData.bNoEntitlements) || (UnrealData.TargetRules.Type == TargetType.Program && !bUseEntitlementsForPrograms))
-				{
-					EntitlementsMetadata = null;
-				}
+				string? TargetEntitlementsFilePath = CopyAndWriteEntitlements(ref EntitlementsMetadata, Config.DisplayName, bUseEntitlementsForPrograms, ProjectOrEngineDir);
 
 				ConfigXcconfig.AppendLine("// pull in the shared settings for all configs for this target");
 				ConfigXcconfig.AppendLine($"#include \"{Xcconfig.Name}.xcconfig\"");
@@ -2251,9 +2284,9 @@ namespace UnrealBuildTool.XcodeProjectXcconfig
 				ConfigXcconfig.AppendLine($"UE_TARGET_CONFIG = {Config.BuildConfig}");
 				ConfigXcconfig.AppendLine($"UE_UBT_BINARY_SUBPATH = {ExetuableSubPath}");
 				ConfigXcconfig.AppendLine($"{ExecutableKey} = {ExecutableName}");
-				if (EntitlementsMetadata != null && EntitlementsMetadata.Mode == MetadataMode.UsePremade)
+				if (EntitlementsMetadata != null && TargetEntitlementsFilePath != null && EntitlementsMetadata.Mode == MetadataMode.UsePremade)
 				{
-					ConfigXcconfig.AppendLine($"CODE_SIGN_ENTITLEMENTS = {EntitlementsMetadata.XcodeProjectRelative}");
+					ConfigXcconfig.AppendLine($"CODE_SIGN_ENTITLEMENTS = {TargetEntitlementsFilePath}");
 				}
 
 				// debug settings
diff --git a/Engine/Source/Runtime/Core/Core.Build.cs b/Engine/Source/Runtime/Core/Core.Build.cs
--- a/Engine/Source/Runtime/Core/Core.Build.cs
+++ b/Engine/Source/Runtime/Core/Core.Build.cs
@@ -299,6 +299,20 @@ public class Core : ModuleRules
 		}
 
 		PublicDefinitions.Add("UE_MERGED_MODULES=" + (Target.bMergeModules ? "1" : "0"));
+		
+		if (Target.IsInPlatformGroup(UnrealPlatformGroup.Apple) && Target.Platform != UnrealTargetPlatform.TVOS)
+		{
+			bool bEnableContactsManager = 
+				Target.Platform == UnrealTargetPlatform.Mac ? (Target.bBuildEditor || Target.MacPlatform.bEnableContactsManager) : 
+				Target.IOSPlatform.bEnableContactsManager;
+
+			if(bEnableContactsManager)
+			{
+				PublicFrameworks.Add("Contacts");
+			}
+
+			PublicDefinitions.Add("WITH_APPLE_CONTACTS=" + (bEnableContactsManager ? "1" : "0"));
+		}
 	}
 
 	protected virtual bool SupportsBinaryConfig(ReadOnlyTargetRules Target)
diff --git a/Engine/Source/Runtime/Core/Private/Apple/AppleContactManager.cpp b/Engine/Source/Runtime/Core/Private/Apple/AppleContactManager.cpp
new file mode 100644
--- /dev/null
+++ b/Engine/Source/Runtime/Core/Private/Apple/AppleContactManager.cpp
@@ -0,0 +1,136 @@
+// Copyright Epic Games, Inc. All Rights Reserved.
+
+#if !PLATFORM_TVOS && WITH_APPLE_CONTACTS
+
+#include "Apple/AppleContactManager.h"
+
+DEFINE_LOG_CATEGORY_STATIC(LogAppleContactManager, Log, All);
+
+FAppleContactManager::FAppleContactManager()
+{
+	ContactStore = [[CNContactStore alloc] init];
+}
+
+bool FAppleContactManager::HaveContactAccessAuthorization()
+{
+	CNAuthorizationStatus AuthStatus = [CNContactStore authorizationStatusForEntityType: CNEntityTypeContacts];
+	
+#if PLATFORM_MAC
+	return AuthStatus == CNAuthorizationStatusAuthorized;
+#else
+	return AuthStatus == CNAuthorizationStatusLimited || AuthStatus == CNAuthorizationStatusAuthorized;
+#endif
+}
+
+void FAppleContactManager::RequestContactAccess()
+{
+	[ContactStore requestAccessForEntityType:CNEntityTypeContacts completionHandler:^(BOOL Granted, NSError* _Nullable Error) 
+	 {
+		if(!Granted)
+		{
+			const char* DomainCString = [Error.domain UTF8String];
+			
+			UE_LOG(LogAppleContactManager, Error, TEXT("Couldn't get access to contacts due to NSError %hs:%d"), DomainCString, Error.code)
+		}
+		
+		OnContactsPermissionChanged.Broadcast(static_cast<bool>(Granted));       
+	}];
+}
+
+FAppleContactManager::ELabeledContactType FAppleContactManager::GetContactTypeForString(NSString* Label, bool bPhoneNumber)
+{
+	if([Label isEqualToString: CNLabelHome])
+	{
+		return ELabeledContactType::Home;
+	}
+	else if([Label isEqualToString: CNLabelWork])
+	{
+		return ELabeledContactType::Work;
+	}
+	else if(bPhoneNumber && ([Label isEqualToString: CNLabelPhoneNumberMobile] || [Label isEqualToString: CNLabelPhoneNumberiPhone]))
+	{
+		return ELabeledContactType::Mobile;
+	}
+	else
+	{
+		return ELabeledContactType::Other;
+	}
+}
+
+// This will fail silently or crash if permissions were not granted AND/OR
+// you don't have a valid NSContactsUsageDescription string in your plist.
+TArray<FAppleContactManager::FAppleContact> FAppleContactManager::GetAllContacts()
+{
+	if(!HaveContactAccessAuthorization())
+	{
+		UE_LOG(LogAppleContactManager, Warning, 
+			   TEXT("HaveContactAccessAuthorization() returned false and thus GetAllContacts will likely fail. Did you call RequestContactAccess before GetAllContacts? Do you have a NSContactsUsageDescription string in your plist?"));
+	}
+	
+	// Change keys for which contact details to grab here
+	NSArray* Keys = @[CNContactFamilyNameKey, CNContactGivenNameKey, CNContactNicknameKey, 
+					  CNContactEmailAddressesKey, CNContactPhoneNumbersKey];
+	
+	
+	NSPredicate* Predicate = [NSPredicate predicateWithValue: true];
+	NSError* Error = nil;
+	
+	NSArray* CNContacts = [ContactStore unifiedContactsMatchingPredicate:Predicate keysToFetch:Keys error:&Error];
+	TArray<FAppleContact> ReturnValue;
+	
+	if(CNContacts) 
+	{
+		for (CNContact* Contact in CNContacts) 
+		{
+			FAppleContact AppleContact 
+			{
+				.FamilyName = Contact.familyName,
+				.GivenName = Contact.givenName,
+				.NickName = Contact.nickname
+			};
+
+			for (CNLabeledValue* Label in Contact.phoneNumbers) 
+			{
+				NSString* Phone = [Label.value stringValue];
+				if ([Phone length] > 0) 
+				{					
+					FLabeledContactValue LabeledPhone
+					{
+						.Type = GetContactTypeForString(Label.label, true),
+						.Value = Phone
+					};
+					
+					AppleContact.PhoneNumbers.Add(LabeledPhone);
+				}
+			}
+			
+			for (CNLabeledValue* Label in Contact.emailAddresses) 
+			{
+				NSString* EmailAddress = [Label.value stringValue];
+				if ([EmailAddress length] > 0) 
+				{					
+					FLabeledContactValue LabeledEmail 
+					{
+						.Type = GetContactTypeForString(Label.label, false),
+						.Value = EmailAddress
+					};
+					
+					AppleContact.EmailAddresses.Add(LabeledEmail);
+				}
+			}
+			
+			ReturnValue.Add(AppleContact);
+		}
+	} 
+	else if (Error) 
+	{
+		const char* ErrorDomain = [Error.domain UTF8String];
+		
+		UE_LOG(LogAppleContactManager, Error, 
+			   TEXT("Couldn't fetch contacts! NSError domain %hs, code %d"), ErrorDomain, Error.code);
+	}
+	
+	return ReturnValue;
+}
+
+#endif
diff --git a/Engine/Source/Runtime/Core/Public/Apple/AppleContactManager.h b/Engine/Source/Runtime/Core/Public/Apple/AppleContactManager.h
new file mode 100644
--- /dev/null
+++ b/Engine/Source/Runtime/Core/Public/Apple/AppleContactManager.h
@@ -0,0 +1,57 @@
+// Copyright Epic Games, Inc. All Rights Reserved.
+
+#pragma once
+
+#if !PLATFORM_TVOS && WITH_APPLE_CONTACTS
+
+#include "CoreMinimal.h"
+
+#include <Contacts/Contacts.h>
+
+/**
+* Apple implementation of accessing device Contacts.
+* @warning You will not be able to use this class (access will be silently denied, or worse, your app may crash), unless you have a valid NSContactsUsageDescription string in your plist.
+**/
+class CORE_API FAppleContactManager 
+{
+public:
+	enum class ELabeledContactType
+	{
+		Home,
+		Work,
+		Mobile,
+		Other
+	};
+	
+	struct FLabeledContactValue
+	{
+		ELabeledContactType Type;
+		FString Value;
+	};
+	
+	struct FAppleContact 
+	{
+		FString FamilyName;
+		FString GivenName;
+		FString NickName;
+		TArray<FLabeledContactValue> PhoneNumbers;
+		TArray<FLabeledContactValue> EmailAddresses;
+	};
+	
+	FAppleContactManager();
+	
+	bool HaveContactAccessAuthorization();
+	
+	void RequestContactAccess();
+	
+	TArray<FAppleContact> GetAllContacts();
+	
+	DECLARE_MULTICAST_DELEGATE_OneParam(FOnContactsPermissionChanged, bool);
+	FOnContactsPermissionChanged OnContactsPermissionChanged;
+private:
+	ELabeledContactType GetContactTypeForString(NSString* Label, bool bPhoneNumber);
+	
+	mutable __strong CNContactStore* ContactStore;
+};
+
+#endif
