Consolidate GetVarArgs implementations and update to handle UTF-8/UTF-16 string conversions.

[RN] minor

diff --git a/Engine/Source/Runtime/Core/Private/GenericPlatform/GenericPlatformString.cpp b/Engine/Source/Runtime/Core/Private/GenericPlatform/GenericPlatformString.cpp
--- a/Engine/Source/Runtime/Core/Private/GenericPlatform/GenericPlatformString.cpp
+++ b/Engine/Source/Runtime/Core/Private/GenericPlatform/GenericPlatformString.cpp
@@ -1,8 +1,11 @@
 // Copyright Epic Games, Inc. All Rights Reserved.
 
 #include "GenericPlatform/GenericPlatformString.h"
+#include "GenericPlatform/GenericWidePlatformString.h"
+#include "GenericPlatform/StandardPlatformString.h"
 #include "HAL/UnrealMemory.h"
 #include "Misc/Char.h"
+#include "Containers/StringConv.h"
 #include "Containers/UnrealString.h"
 #include "Logging/LogMacros.h"
 
@@ -22,6 +25,671 @@ void* FGenericPlatformString::Memcpy(void* Dest, const void* Src, SIZE_T Count)
 	return FMemory::Memcpy(Dest, Src, Count);
 }
 
+namespace
+{
+	// Output iterator which will not overflow the destination buffer but will keep track of how many characters have been written.
+	// It will also null terminate the output when destructed.
+	template <typename CharType>
+	struct TSafeDestIterator
+	{
+		// Non-copyable
+		TSafeDestIterator(const TSafeDestIterator&) = delete;
+		TSafeDestIterator& operator=(const TSafeDestIterator&) = delete;
+
+		explicit TSafeDestIterator(CharType* InPtr, SIZE_T InEnd)
+			: NumCharsWritten(0)
+			, Ptr(InPtr)
+			, EndMinusOne(InPtr + InEnd - 1)
+		{
+			check(InPtr);
+		}
+
+		FORCENOINLINE ~TSafeDestIterator()
+		{
+			// Auto-terminate buffer
+			*Ptr = 0;
+		}
+
+		// Writes Count instances of Ch to the destination buffer
+		FORCENOINLINE TSafeDestIterator& Write(CharType Ch, int32 Count = 1)
+		{
+			// Ensure we're not in the error state
+			if (EndMinusOne)
+			{
+				NumCharsWritten += Count;
+
+				int32 NumToWrite = FPlatformMath::Min(Count, (int32)(EndMinusOne - Ptr));
+				for (int32 Val = NumToWrite; Val; --Val)
+				{
+					*Ptr++ = Ch;
+				}
+
+				if (NumToWrite != Count)
+				{
+					EndMinusOne = 0;
+				}
+			}
+
+			return *this;
+		}
+
+		// Writes Count characters from Src to the destination buffer
+		template <typename OtherCharType>
+		FORCENOINLINE TSafeDestIterator& Write(const OtherCharType* Src, int32 Count)
+		{
+			// Ensure we're not in the error state
+			if (EndMinusOne)
+			{
+				NumCharsWritten += Count;
+
+				if constexpr (std::is_same_v<ANSICHAR, OtherCharType> && std::is_same_v<WIDECHAR, CharType>)
+				{
+					int32 ConvertedLength = FPlatformString::ConvertedLength<WIDECHAR, UTF8CHAR>((const UTF8CHAR*)Src);
+					int32 NumToWrite = FPlatformMath::Min(ConvertedLength, (int32)(EndMinusOne - Ptr));
+					Ptr = FPlatformString::Convert(Ptr, NumToWrite, (const UTF8CHAR*)Src, Count);
+
+					if (NumToWrite != ConvertedLength)
+					{
+						EndMinusOne = 0;
+					}
+				}
+				else if constexpr (std::is_same_v<WIDECHAR, OtherCharType> && std::is_same_v<ANSICHAR, CharType>)
+				{
+					int32 ConvertedLength = FPlatformString::ConvertedLength<UTF8CHAR, WIDECHAR>(Src);
+					int32 NumToWrite = FPlatformMath::Min(ConvertedLength, (int32)(EndMinusOne - Ptr));
+					Ptr = (CharType*)FPlatformString::Convert((UTF8CHAR*)Ptr, NumToWrite, Src, Count);
+
+					if (NumToWrite != ConvertedLength)
+					{
+						EndMinusOne = 0;
+					}
+				}
+				else
+				{
+					int32 NumToWrite = FPlatformMath::Min(Count, (int32)(EndMinusOne - Ptr));
+
+					for (int32 Val = NumToWrite; Val; --Val)
+					{
+						*Ptr++ = *Src++;
+					}
+
+					if (NumToWrite != Count)
+					{
+						EndMinusOne = 0;
+					}
+				}
+			}
+
+			return *this;
+		}
+
+		explicit operator bool() const
+		{
+			return !!EndMinusOne;
+		}
+
+		int32 GetNumCharsWritten() const
+		{
+			return NumCharsWritten;
+		}
+
+	private:
+		int32     NumCharsWritten;
+		CharType* Ptr;
+		CharType* EndMinusOne; // when null, this means the iterator has already moved past the writable area of the buffer
+	};
+
+	static inline bool CharIsIntegerFormatSpecifier(TCHAR Char)
+	{
+		return (Char == 'i') | (Char == 'd') | (Char == 'u') | (Char == 'X') | (Char == 'x');
+	}
+
+	template <typename FormatCharType, typename ArgCharType>
+	void WriteStringArg(TSafeDestIterator<FormatCharType>& DestIter, const ArgCharType* Val, int FieldLen, int PrecisionLen)
+	{
+		int RetCnt = PrecisionLen < 0 ? FPlatformString::Strlen(Val) : FPlatformString::Strnlen(Val, PrecisionLen);
+
+		int Spaces = FPlatformMath::Max(FPlatformMath::Abs(FieldLen) - RetCnt, 0);
+		if (Spaces > 0 && FieldLen > 0)
+		{
+			DestIter.Write(TEXT(' '), Spaces);
+		}
+		DestIter.Write(Val, RetCnt);
+		if (Spaces > 0 && FieldLen < 0)
+		{
+			DestIter.Write(TEXT(' '), Spaces);
+		}
+	}
+
+	template <typename ArgCharType, typename FormatCharType, typename VaListType>
+	void ProcessStringArg(TSafeDestIterator<FormatCharType>& DestIter, const FormatCharType*& Src, int FieldLen, int PrecisionLen, VaListType& ArgPtr)
+	{
+		// ArgPtr is taken as a templated parameter, despite being a va_list, because on some
+		// architectures it could be an array type, and when an array is used as a function parameter,
+		// the named value is treated as a pointer and won't bind to an array reference.
+
+		Src++;
+		static const ArgCharType* Null = LITERAL(ArgCharType, "(null)");
+		const ArgCharType* Val = va_arg(ArgPtr, ArgCharType*);
+		if (Val == nullptr)
+		{
+			Val = Null;
+		}
+
+		WriteStringArg(DestIter, Val, FieldLen, PrecisionLen);
+	}
+
+	static void LogUnrecognizedFormatCharacters(const char* MsgFormat, WIDECHAR Char, const WIDECHAR* InputFormat)													{ printf(MsgFormat, (wint_t)Char, TCHAR_TO_ANSI(InputFormat)); }
+	static void LogUnrecognizedFormatCharacters(const char* MsgFormat, WIDECHAR Char1, WIDECHAR Char2, const WIDECHAR* InputFormat)									{ printf(MsgFormat, (wint_t)Char1, (wint_t)Char2, TCHAR_TO_ANSI(InputFormat)); }
+	static void LogUnrecognizedFormatCharacters(const char* MsgFormat, WIDECHAR Char1, WIDECHAR Char2, WIDECHAR Char3, const WIDECHAR* InputFormat)					{ printf(MsgFormat, (wint_t)Char1, (wint_t)Char2, (wint_t)Char3, TCHAR_TO_ANSI(InputFormat)); }
+	static void LogUnrecognizedFormatCharacters(const char* MsgFormat, WIDECHAR Char1, WIDECHAR Char2, WIDECHAR Char3, WIDECHAR Char4, const WIDECHAR* InputFormat)	{ printf(MsgFormat, (wint_t)Char1, (wint_t)Char2, (wint_t)Char3, (wint_t)Char4, TCHAR_TO_ANSI(InputFormat)); }
+
+	static void LogUnrecognizedFormatCharacters(const char* MsgFormat, ANSICHAR Char, const ANSICHAR* InputFormat)													{ printf(MsgFormat, Char, InputFormat); }
+	static void LogUnrecognizedFormatCharacters(const char* MsgFormat, ANSICHAR Char1, ANSICHAR Char2, const ANSICHAR* InputFormat)									{ printf(MsgFormat, Char1, Char2, InputFormat); }
+	static void LogUnrecognizedFormatCharacters(const char* MsgFormat, ANSICHAR Char1, ANSICHAR Char2, ANSICHAR Char3, const ANSICHAR* InputFormat)					{ printf(MsgFormat, Char1, Char2, Char3, InputFormat); }
+	static void LogUnrecognizedFormatCharacters(const char* MsgFormat, ANSICHAR Char1, ANSICHAR Char2, ANSICHAR Char3, ANSICHAR Char4, const ANSICHAR* InputFormat) { printf(MsgFormat, Char1, Char2, Char3, Char4, InputFormat); }
+
+	template <typename CharType, typename VaListType>
+	static int32 GetVarArgsImpl(CharType* Dest, SIZE_T DestSize, const CharType*& Fmt, VaListType& ArgPtr)
+	{
+		if (Fmt == nullptr)
+		{
+			if ((DestSize > 0) && (Dest != nullptr))
+			{
+				*Dest = 0;
+			}
+			return 0;
+		}
+
+		const CharType* Src = Fmt;
+
+		TSafeDestIterator<CharType> DestIter(Dest, DestSize);
+		while (*Src)
+		{
+			if (*Src != '%')
+			{
+				if (!DestIter.Write(*Src++))
+				{
+					return -1;
+				}
+				continue;
+			}
+
+			const CharType* Percent = Src;
+			int FieldLen = 0;
+			bool bVariablePrecision = false;
+			int PrecisionLen = -1;
+
+			Src++; // skip the '%' char...
+
+			// Skip modifier flags that don't need additional processing;
+			// they still get passed to snprintf() below based on the conversion.
+			while (*Src == '+' || *Src == '#' || *Src == TEXT(' ') || *Src == TEXT('0'))
+			{
+				Src++;
+			}
+
+			// check for field width requests...
+			if ((*Src == '-') || ((*Src >= '0') && (*Src <= '9')))
+			{
+				const CharType* Cur = Src + 1;
+				while ((*Cur >= '0') && (*Cur <= '9'))
+				{
+					Cur++;
+				}
+
+				FieldLen = FPlatformString::Atoi(Src);
+				Src = Cur;
+			}
+
+			// check for dynamic field requests
+			if (*Src == '*')
+			{
+				FieldLen = va_arg(ArgPtr, int32);
+				Src++;
+			}
+
+			if (*Src == '.')
+			{
+				const CharType* Cur = Src + 1;
+				if (*Cur == '*')
+				{
+					bVariablePrecision = true;
+					PrecisionLen = va_arg(ArgPtr, int32);
+					Cur++;
+				}
+				else
+				{
+					while ((*Cur >= '0') && (*Cur <= '9'))
+					{
+						Cur++;
+					}
+
+					PrecisionLen = FPlatformString::Atoi(Src + 1);
+				}
+				Src = Cur;
+			}
+
+			bool bIsNarrowString = std::is_same_v<CharType, ANSICHAR>;
+			// Check for 'ls' field, change to 's' and remember
+			if ((Src[0] == 'l' && Src[1] == 's'))
+			{
+				bIsNarrowString = false;
+				Src++;
+			}
+
+			// Check for 'hs' field, change to 's' and remember
+			if ((Src[0] == 'h' && Src[1] == 's'))
+			{
+				bIsNarrowString = true;
+				Src++;
+			}
+
+			switch (*Src)
+			{
+			case '%':
+			{
+				Src++;
+				if (!DestIter.Write('%'))
+				{
+					return -1;
+				}
+				break;
+			}
+
+			case 'c':
+			{
+				CharType Val = (CharType)va_arg(ArgPtr, int);
+				Src++;
+				bool bSuccess = true;
+				if (FieldLen > 1)
+				{
+					bSuccess = !!DestIter.Write(CharType(' '), FieldLen - 1);
+				}
+				bSuccess = bSuccess && !!DestIter.Write(Val);
+				if (FieldLen < -1)
+				{
+					bSuccess = bSuccess && !!DestIter.Write(CharType(' '), FPlatformMath::Abs(FieldLen + 1));
+				}
+				if (!bSuccess)
+				{
+					return -1;
+				}
+				break;
+			}
+
+			case 'd':
+			case 'i':
+			case 'X':
+			case 'x':
+			case 'u':
+			{
+				Src++;
+				int Val = va_arg(ArgPtr, int);
+				ANSICHAR AnsiNum[64];
+				ANSICHAR FmtBuf[30];
+
+				// limit width to the buffer size
+				FieldLen = FMath::Min(static_cast<int>(sizeof(AnsiNum)) - 1, FieldLen);
+
+				// Yes, this is lame.
+				int CpyIdx = 0;
+				while (Percent < Src && CpyIdx < UE_ARRAY_COUNT(FmtBuf))
+				{
+					if (LIKELY(*Percent != '*'))
+					{
+						FmtBuf[CpyIdx] = (ANSICHAR)*Percent;
+						++CpyIdx;
+					}
+					else
+					{
+						snprintf(&FmtBuf[CpyIdx], sizeof(FmtBuf) - CpyIdx, "%d", FieldLen);
+						while (CpyIdx < UE_ARRAY_COUNT(FmtBuf) && FmtBuf[CpyIdx] != 0)
+						{
+							++CpyIdx;
+						}
+					}
+					++Percent;
+				}
+				FmtBuf[CpyIdx] = 0;
+
+				int RetCnt = snprintf(AnsiNum, sizeof(AnsiNum), FmtBuf, Val);
+				if (RetCnt < 0 || !DestIter.Write(AnsiNum, RetCnt))
+				{
+					return -1;
+				}
+				break;
+			}
+
+			case 'z':
+			case 'Z':
+			{
+				if (Src[1] == 0)
+				{
+					LogUnrecognizedFormatCharacters("Unknown percent [%lc] in FGenericWidePlatformString::GetVarArgs() [%s]\n.", Src[0], Fmt);
+					Src++;  // skip it, I guess.
+					break;
+				}
+
+				Src += 2;
+
+				size_t Val = va_arg(ArgPtr, size_t);
+
+				ANSICHAR AnsiNum[64];
+				ANSICHAR FmtBuf[30];
+
+				// limit width to the buffer size
+				FieldLen = FMath::Min(static_cast<int>(sizeof(AnsiNum)) - 1, FieldLen);
+
+				// Yes, this is lame.
+				int CpyIdx = 0;
+				while (Percent < Src && CpyIdx < UE_ARRAY_COUNT(FmtBuf))
+				{
+					if (LIKELY(*Percent != '*'))
+					{
+						FmtBuf[CpyIdx] = (ANSICHAR)*Percent;
+						++CpyIdx;
+					}
+					else
+					{
+						snprintf(&FmtBuf[CpyIdx], sizeof(FmtBuf) - CpyIdx, "%d", FieldLen);
+						while (CpyIdx < UE_ARRAY_COUNT(FmtBuf) && FmtBuf[CpyIdx] != 0)
+						{
+							++CpyIdx;
+						}
+					}
+					++Percent;
+				}
+				FmtBuf[CpyIdx] = 0;
+
+				int RetCnt = snprintf(AnsiNum, sizeof(AnsiNum), FmtBuf, Val);
+				if (RetCnt < 0 || !DestIter.Write(AnsiNum, RetCnt))
+				{
+					return -1;
+				}
+				break;
+			}
+
+			case 'p':
+			{
+				Src++;
+				void* Val = va_arg(ArgPtr, void*);
+				ANSICHAR AnsiNum[30];
+				ANSICHAR FmtBuf[30];
+
+				// Yes, this is lame.
+				int CpyIdx = 0;
+				while (Percent < Src && CpyIdx < UE_ARRAY_COUNT(FmtBuf))
+				{
+					FmtBuf[CpyIdx] = (ANSICHAR)*Percent;
+					Percent++;
+					CpyIdx++;
+				}
+				FmtBuf[CpyIdx] = 0;
+
+				int RetCnt = snprintf(AnsiNum, sizeof(AnsiNum), FmtBuf, Val);
+				if (RetCnt < 0 || !DestIter.Write(AnsiNum, RetCnt))
+				{
+					return -1;
+				}
+				break;
+			}
+
+			case 'l':
+			case 'I':
+			case 'h':
+			{
+				int RemainingSize = FPlatformString::Strlen(Src);
+
+				if (RemainingSize < 2)
+				{
+					LogUnrecognizedFormatCharacters("Unknown percent [%lc] in FGenericPlatformString::GetVarArgs() [%s]\n.", Src[0], Fmt);
+					Src++;  // skip it, I guess.
+					break;
+				}
+
+				// treat %ld as %d. Also shorts for %h will be promoted to ints. This path also handles %li, %lu, %lx and %lX.
+				if ((Src[0] == 'l' || Src[0] == 'h') && CharIsIntegerFormatSpecifier(Src[1]))
+				{
+					Src += 2;
+					long int Val = va_arg(ArgPtr, long int);
+					ANSICHAR AnsiNum[30];
+					ANSICHAR FmtBuf[30];
+
+					// Yes, this is lame.
+					int CpyIdx = 0;
+					while (Percent < Src && CpyIdx < UE_ARRAY_COUNT(FmtBuf))
+					{
+						FmtBuf[CpyIdx] = (ANSICHAR)*Percent;
+						Percent++;
+						CpyIdx++;
+					}
+					FmtBuf[CpyIdx] = 0;
+
+					int RetCnt = snprintf(AnsiNum, sizeof(AnsiNum), FmtBuf, Val);
+					if (RetCnt < 0 || !DestIter.Write(AnsiNum, RetCnt))
+					{
+						return -1;
+					}
+					break;
+				}
+				// Treat %lf as a %f
+				else if (Src[0] == 'l' && Src[1] == 'f')
+				{
+					Src += 2;
+					double Val = va_arg(ArgPtr, double);
+					ANSICHAR AnsiNum[340]; // make sure buffer is large enough to handle the longest double values, such as the largest number which is 316 characters including decimal point and 6 decimals
+					ANSICHAR FmtBuf[30];
+
+					// Yes, this is lame.
+					int CpyIdx = 0;
+					while (Percent < Src && CpyIdx < UE_ARRAY_COUNT(FmtBuf))
+					{
+						FmtBuf[CpyIdx] = (ANSICHAR)*Percent;
+						Percent++;
+						CpyIdx++;
+					}
+					FmtBuf[CpyIdx] = 0;
+
+					int RetCnt = snprintf(AnsiNum, sizeof(AnsiNum), FmtBuf, Val);
+					if (RetCnt < 0 || !DestIter.Write(AnsiNum, RetCnt))
+					{
+						return -1;
+					}
+					break;
+				}
+
+				// At this point we expect RemainingSize >= 3, but we know it's at least 2.
+				if (RemainingSize < 3)
+				{
+					LogUnrecognizedFormatCharacters("Unknown percent [%lc%lc] in FGenericPlatformString::GetVarArgs() [%s]\n.", Src[0], Src[1], Fmt);
+					Src++;  // skip it, I guess.
+					break;
+				}
+
+				// Deal with half-half formats - %hhd, %hhi, %hhu, %hhx and %hhX.
+				if (Src[0] == 'h')
+				{
+					if (Src[1] != 'h' || !CharIsIntegerFormatSpecifier(Src[2]))
+					{
+						LogUnrecognizedFormatCharacters("Unknown percent [%lc%lc%lc] in FGenericPlatformString::GetVarArgs() [%s]\n.", Src[0], Src[1], Src[2], Fmt);
+						Src++;  // skip it, I guess.
+						break;
+					}
+					Src += 3;
+					int Val = va_arg(ArgPtr, int);
+					ANSICHAR AnsiNum[8];
+					ANSICHAR FmtBuf[30];
+
+					// Yes, this is lame.
+					int CpyIdx = 0;
+					while (Percent < Src && CpyIdx < UE_ARRAY_COUNT(FmtBuf))
+					{
+						FmtBuf[CpyIdx] = (ANSICHAR)*Percent;
+						Percent++;
+						CpyIdx++;
+					}
+					FmtBuf[CpyIdx] = 0;
+
+					int RetCnt = snprintf(AnsiNum, sizeof(AnsiNum), FmtBuf, Val);
+					if (RetCnt < 0 || !DestIter.Write(AnsiNum, RetCnt))
+					{
+						return -1;
+					}
+					break;
+				}
+				else if (Src[0] == 'l')
+				{
+					if (Src[1] != 'l' || !CharIsIntegerFormatSpecifier(Src[2]))
+					{
+						LogUnrecognizedFormatCharacters("Unknown percent [%lc%lc%lc] in FGenericPlatformString::GetVarArgs() [%s]\n.", Src[0], Src[1], Src[2], Fmt);
+						Src++;  // skip it, I guess.
+						break;
+					}
+				}
+				else if (Src[0] == 'I')
+				{
+					if (RemainingSize < 4)
+					{
+						LogUnrecognizedFormatCharacters("Unknown percent [%lc%lc%lc] in FGenericPlatformString::GetVarArgs() [%s]\n.", Src[0], Src[1], Src[2], Fmt);
+						Src++;  // skip it, I guess.
+						break;
+					}
+					else if (Src[1] != '6' || Src[2] != '4' || !CharIsIntegerFormatSpecifier(Src[3]))
+					{
+						LogUnrecognizedFormatCharacters("Unknown percent [%lc%lc%lc%lc] in FGenericPlatformString::GetVarArgs() [%s]\n.", Src[0], Src[1], Src[2], Src[3], Fmt);
+						Src++;  // skip it, I guess.
+						break;
+					}
+				}
+				else
+				{
+					LogUnrecognizedFormatCharacters("Unknown percent [%lc%lc%lc] in FGenericPlatformString::GetVarArgs() [%s]\n.", Src[0], Src[1], Src[2], Fmt);
+					Src++;  // skip it, I guess.
+					break;
+				}
+
+				// Yes, this is lame.
+				unsigned long long Val = va_arg(ArgPtr, unsigned long long);
+				ANSICHAR AnsiNum[60];
+				ANSICHAR FmtBuf[30];
+
+				int CpyIdx;
+				if (Src[0] == 'I')
+				{
+					// Convert I64 syntax to ll 
+					Src += 4;
+					Percent += 4;
+					CpyIdx = 3;
+					FmtBuf[0] = '%';
+					FmtBuf[1] = 'l';
+					FmtBuf[2] = 'l';
+				}
+				else
+				{
+					Src += 3;
+					CpyIdx = 0;
+				}
+
+				while (Percent < Src && CpyIdx < UE_ARRAY_COUNT(FmtBuf))
+				{
+					FmtBuf[CpyIdx] = (ANSICHAR)*Percent;
+					Percent++;
+					CpyIdx++;
+				}
+				FmtBuf[CpyIdx] = 0;
+
+				int RetCnt = snprintf(AnsiNum, sizeof(AnsiNum), FmtBuf, Val);
+				if (RetCnt < 0 || !DestIter.Write(AnsiNum, RetCnt))
+				{
+					return -1;
+				}
+				break;
+			}
+
+			case 'f':
+			case 'e':
+			case 'g':
+			{
+				Src++;
+				double Val = va_arg(ArgPtr, double);
+				ANSICHAR AnsiNum[340]; // make sure buffer is large enough to handle the longest double values, such as the largest number which is 316 characters including decimal point and 6 decimals
+				ANSICHAR FmtBuf[30];
+
+				// Yes, this is lame.
+				int CpyIdx = 0;
+				while (Percent < Src && CpyIdx < UE_ARRAY_COUNT(FmtBuf))
+				{
+					FmtBuf[CpyIdx] = (ANSICHAR)*Percent;
+					Percent++;
+					CpyIdx++;
+				}
+				FmtBuf[CpyIdx] = 0;
+
+				int RetCnt;
+				if (bVariablePrecision)
+				{
+					RetCnt = snprintf(AnsiNum, sizeof(AnsiNum), FmtBuf, PrecisionLen, Val);
+				}
+				else
+				{
+					RetCnt = snprintf(AnsiNum, sizeof(AnsiNum), FmtBuf, Val);
+				}
+				if (RetCnt >= UE_ARRAY_COUNT(AnsiNum))
+				{
+					// We should print what we have written into AnsiNum but ensure we null terminate before printing
+					AnsiNum[UE_ARRAY_COUNT(AnsiNum) - 1] = '\0';
+					checkf(0, TEXT("Attempting to read past the size our buffer. Buffer Size: %" SIZE_T_FMT " Size to read: %d. Current contents: '%s'\n"), UE_ARRAY_COUNT(AnsiNum), RetCnt, UTF8_TO_TCHAR(AnsiNum));
+				}
+				if (RetCnt < 0 || !DestIter.Write(AnsiNum, RetCnt))
+				{
+					return -1;
+				}
+				break;
+			}
+
+			case 's':
+			{
+				if (bIsNarrowString)
+				{
+					ProcessStringArg<ANSICHAR>(DestIter, Src, FieldLen, PrecisionLen, ArgPtr);
+				}
+				else
+				{
+					ProcessStringArg<TCHAR>(DestIter, Src, FieldLen, PrecisionLen, ArgPtr);
+				}
+				if (!DestIter)
+				{
+					return -1;
+				}
+				break;
+			}
+
+			case 'S':
+			{
+				// The %S format represents a string which is the opposite of %s - wide if CharType is narrow, or narrow if CharType is wide
+				using OtherCharType = std::conditional_t<std::is_same_v<CharType, ANSICHAR>, WIDECHAR, ANSICHAR>;
+
+				ProcessStringArg<OtherCharType>(DestIter, Src, FieldLen, PrecisionLen, ArgPtr);
+				if (!DestIter)
+				{
+					return -1;
+				}
+				break;
+			}
+
+			default:
+				printf("Unknown percent [%%%c] in FGenericPlatformString::GetVarArgs().\n", *Src);
+				Src++;  // skip char, I guess.
+				break;
+			}
+		}
+
+		int32 Result = DestIter.GetNumCharsWritten();
+		return Result;
+	}
+}
+
 namespace UE::Core::Private
 {
 	/**
@@ -716,6 +1384,16 @@ int32 StrncmpImpl(const CharType1* String1, const CharType2* String2, SIZE_T Cou
 
 }
 
+int32 FGenericPlatformString::GetVarArgs(ANSICHAR* Dest, SIZE_T DestSize, const ANSICHAR*& Fmt, va_list ArgPtr)
+{
+	return GetVarArgsImpl(Dest, DestSize, Fmt, ArgPtr);
+}
+
+int32 FGenericPlatformString::GetVarArgs(WIDECHAR* Dest, SIZE_T DestSize, const WIDECHAR*& Fmt, va_list ArgPtr)
+{
+	return GetVarArgsImpl(Dest, DestSize, Fmt, ArgPtr);
+}
+
 int32 FGenericPlatformString::Strncmp(const ANSICHAR* Str1, const ANSICHAR* Str2, SIZE_T Count) { return GenericPlatformStringPrivate::StrncmpImpl(Str1, Str2, Count); }
 int32 FGenericPlatformString::Strncmp(const WIDECHAR* Str1, const ANSICHAR* Str2, SIZE_T Count) { return GenericPlatformStringPrivate::StrncmpImpl(Str1, Str2, Count); }
 int32 FGenericPlatformString::Strncmp(const UTF8CHAR* Str1, const ANSICHAR* Str2, SIZE_T Count) { return GenericPlatformStringPrivate::StrncmpImpl(Str1, Str2, Count); }
@@ -725,3 +1403,159 @@ int32 FGenericPlatformString::Strncmp(const UTF8CHAR* Str1, const WIDECHAR* Str2
 int32 FGenericPlatformString::Strncmp(const ANSICHAR* Str1, const UTF8CHAR* Str2, SIZE_T Count) { return GenericPlatformStringPrivate::StrncmpImpl(Str1, Str2, Count); }
 int32 FGenericPlatformString::Strncmp(const WIDECHAR* Str1, const UTF8CHAR* Str2, SIZE_T Count) { return GenericPlatformStringPrivate::StrncmpImpl(Str1, Str2, Count); }
 int32 FGenericPlatformString::Strncmp(const UTF8CHAR* Str1, const UTF8CHAR* Str2, SIZE_T Count) { return GenericPlatformStringPrivate::StrncmpImpl(Str1, Str2, Count); }
+
+#if WITH_LOW_LEVEL_TESTS
+
+#include "TestHarness.h"
+#include <catch2/generators/catch_generators.hpp>
+
+constexpr int OUTPUT_SIZE = 256;
+static int32 TestGetVarArgs(WIDECHAR* OutputString, const WIDECHAR* Format, ...)
+{
+	va_list ArgPtr;
+	va_start(ArgPtr, Format);
+	int32 Result = FGenericPlatformString::GetVarArgs(OutputString, OUTPUT_SIZE, Format, ArgPtr);
+	va_end(ArgPtr);
+	return Result;
+}
+
+static int32 TestGetVarArgs(ANSICHAR* OutputString, const ANSICHAR* Format, ...)
+{
+	va_list ArgPtr;
+	va_start(ArgPtr, Format);
+	int32 Result = FGenericPlatformString::GetVarArgs(OutputString, OUTPUT_SIZE, Format, ArgPtr);
+	va_end(ArgPtr);
+	return Result;
+}
+
+TEST_CASE("Core::GenericPlatformString::GetVarArgs", "[Core][String][Smoke]")
+{
+	{
+		WIDECHAR OutputString[OUTPUT_SIZE];
+
+		TestGetVarArgs(OutputString, TEXT("Test A|%-20s|%20s|%10.2f|%-10.2f|"), TEXT("LEFT"), TEXT("RIGHT"), 33.333333, 66.666666);
+		CHECK_MESSAGE(OutputString, FString(OutputString) == FString(TEXT("Test A|LEFT                |               RIGHT|     33.33|66.67     |")));
+
+		TestGetVarArgs(OutputString, TEXT("Test B|Percents:%%%%%%%d|"), 3);
+		CHECK_MESSAGE(OutputString, FString(OutputString) == FString(TEXT("Test B|Percents:%%%3|")));
+
+		TestGetVarArgs(OutputString, TEXT("Test C|%d|%i|%X|%x|%u|"), 12345, 54321, 0x123AbC, 15, 99);
+		CHECK_MESSAGE(OutputString, FString(OutputString) == FString(TEXT("Test C|12345|54321|123ABC|f|99|")));
+
+		TestGetVarArgs(OutputString, TEXT("Test D|%p|"), 0x12345);
+#ifdef _MSC_VER
+		// MSVC's standard library formats pointers differently
+		CHECK_MESSAGE(OutputString, FString(OutputString) == FString(TEXT("Test D|0000000000012345|")));
+#else
+		// Pointer format (which we get from snprintf) varies from platform to platform.
+		// Call it to make sure it puts the output in the proper place, but don't test for the exact format
+		//	CHECK_MESSAGE(OutputString, FString(OutputString) == FString(TEXT("Test D|0x12345|")));
+		{
+			FString OutputStringStr(OutputString);
+			CHECK_MESSAGE(OutputString, OutputStringStr.StartsWith(TEXT("Test D|"))
+				&& OutputStringStr.EndsWith(TEXT("|")) && OutputStringStr.Len() < TEXTVIEW("Test D||").Len() + 20);
+		}
+#endif
+
+		TestGetVarArgs(OutputString, TEXT("Test E|%" INT64_FMT "|"), int64(12345678912345LL));
+		CHECK_MESSAGE(OutputString, FString(OutputString) == FString(TEXT("Test E|12345678912345|")));
+
+		TestGetVarArgs(OutputString, TEXT("Test F|%f|%e|%g|"), 123.456, 123.456, 123.456);
+		CHECK_MESSAGE(OutputString, FString(OutputString) == FString(TEXT("Test F|123.456000|1.234560e+02|123.456|")));
+
+		TestGetVarArgs(OutputString, TEXT("Test G|%" UPTRINT_X_FMT"|"), UPTRINT(49374));
+		CHECK_MESSAGE(OutputString, FString(OutputString) == FString(TEXT("Test G|C0DE|")));
+
+		TestGetVarArgs(OutputString, TEXT("Test H|%" UPTRINT_x_FMT "|"), UPTRINT(49374));
+		CHECK_MESSAGE(OutputString, FString(OutputString) == FString(TEXT("Test H|c0de|")));
+
+		TestGetVarArgs(OutputString, TEXT("Test I|%" UINT64_FMT "|"), MAX_uint64);
+		CHECK_MESSAGE(OutputString, FString(OutputString) == FString(TEXT("Test I|18446744073709551615|")));
+
+		TestGetVarArgs(OutputString, TEXT("Test J|%*c|"), 10, 'J');
+		CHECK_MESSAGE(OutputString, FString(OutputString) == FString(TEXT("Test J|         J|")));
+
+		TestGetVarArgs(OutputString, TEXT("Test K|%-5c|"), 'K');
+		CHECK_MESSAGE(OutputString, FString(OutputString) == FString(TEXT("Test K|K    |")));
+
+		TestGetVarArgs(OutputString, TEXT("Test L|%hs|"), UTF8TEXT("Hello World"));
+		CHECK_MESSAGE(OutputString, FString(OutputString) == FString(TEXT("Test L|Hello World|")));
+
+		TestGetVarArgs(OutputString, TEXT("Test M|%hs|"), "\xF0\x9F\x98\x86\xF0\x9F\x80\x9C");
+		CHECK_MESSAGE(OutputString, FString(OutputString) == FString(UTF16_TO_TCHAR(u"Test M|\xD83D\xDE06\xD83C\xDC1C|")));
+
+		TestGetVarArgs(OutputString, TEXT("Test N|%ls|"), UTF16_TO_TCHAR(u"\xD83D\xDE06\xD83C\xDC1C"));
+		CHECK_MESSAGE(OutputString, FString(OutputString) == FString(UTF16_TO_TCHAR(u"Test N|\xD83D\xDE06\xD83C\xDC1C|")));
+
+		TestGetVarArgs(OutputString, TEXT("Test O|%s|"), UTF16_TO_TCHAR(u"\xD83D\xDE06\xD83C\xDC1C"));
+		CHECK_MESSAGE(OutputString, FString(OutputString) == FString(UTF16_TO_TCHAR(u"Test O|\xD83D\xDE06\xD83C\xDC1C|")));
+
+		TestGetVarArgs(OutputString, TEXT("Test P|%S|"), "\xF0\x9F\x98\x86\xF0\x9F\x80\x9C");
+		CHECK_MESSAGE(OutputString, FString(OutputString) == FString(UTF16_TO_TCHAR(u"Test P|\xD83D\xDE06\xD83C\xDC1C|")));
+	}
+
+	{
+		ANSICHAR OutputString[OUTPUT_SIZE];
+
+		TestGetVarArgs(OutputString, "Test A|%-20s|%20s|%10.2f|%-10.2f|", "LEFT", "RIGHT", 33.333333, 66.666666);
+		CHECK_MESSAGE(OutputString, FString(OutputString) == FString("Test A|LEFT                |               RIGHT|     33.33|66.67     |"));
+
+		TestGetVarArgs(OutputString, "Test B|Percents:%%%%%%%d|", 3);
+		CHECK_MESSAGE(OutputString, FString(OutputString) == FString("Test B|Percents:%%%3|"));
+
+		TestGetVarArgs(OutputString, "Test C|%d|%i|%X|%x|%u|", 12345, 54321, 0x123AbC, 15, 99);
+		CHECK_MESSAGE(OutputString, FString(OutputString) == FString("Test C|12345|54321|123ABC|f|99|"));
+
+		TestGetVarArgs(OutputString, "Test D|%p|", 0x12345);
+#ifdef _MSC_VER
+		// MSVC's standard library formats pointers differently
+		CHECK_MESSAGE(OutputString, FString(OutputString) == FString("Test D|0000000000012345|"));
+#else
+		// Pointer format (which we get from snprintf) varies from platform to platform.
+		// Call it to make sure it puts the output in the proper place, but don't test for the exact format
+		//	CHECK_MESSAGE(OutputString, FString(OutputString) == FString(TEXT("Test D|0x12345|")));
+		{
+			FString OutputStringStr(OutputString);
+			CHECK_MESSAGE(OutputString, OutputStringStr.StartsWith("Test D|")
+				&& OutputStringStr.EndsWith("|") && OutputStringStr.Len() < TEXTVIEW("Test D||").Len() + 20);
+		}
+#endif
+
+		TestGetVarArgs(OutputString, "Test E|%" INT64_FMT "|", int64(12345678912345LL));
+		CHECK_MESSAGE(OutputString, FString(OutputString) == FString("Test E|12345678912345|"));
+
+		TestGetVarArgs(OutputString, "Test F|%f|%e|%g|", 123.456, 123.456, 123.456);
+		CHECK_MESSAGE(OutputString, FString(OutputString) == FString("Test F|123.456000|1.234560e+02|123.456|"));
+
+		TestGetVarArgs(OutputString, "Test G|%" UPTRINT_X_FMT"|", UPTRINT(49374));
+		CHECK_MESSAGE(OutputString, FString(OutputString) == FString("Test G|C0DE|"));
+
+		TestGetVarArgs(OutputString, "Test H|%" UPTRINT_x_FMT "|", UPTRINT(49374));
+		CHECK_MESSAGE(OutputString, FString(OutputString) == FString("Test H|c0de|"));
+
+		TestGetVarArgs(OutputString, "Test I|%" UINT64_FMT "|", MAX_uint64);
+		CHECK_MESSAGE(OutputString, FString(OutputString) == FString("Test I|18446744073709551615|"));
+
+		TestGetVarArgs(OutputString, "Test J|%*c|", 10, 'J');
+		CHECK_MESSAGE(OutputString, FString(OutputString) == FString("Test J|         J|"));
+
+		TestGetVarArgs(OutputString, "Test K|%-5c|", 'K');
+		CHECK_MESSAGE(OutputString, FString(OutputString) == FString("Test K|K    |"));
+
+		TestGetVarArgs(OutputString, "Test L|%hs|", UTF8TEXT("Hello World"));
+		CHECK_MESSAGE(OutputString, FString(OutputString) == FString("Test L|Hello World|"));
+
+		TestGetVarArgs(OutputString, "Test M|%ls|", UTF16_TO_TCHAR(u"\xD83D\xDE06\xD83C\xDC1C"));
+		CHECK_MESSAGE(OutputString, FString(OutputString) == FString("Test M|\xF0\x9F\x98\x86\xF0\x9F\x80\x9C|"));
+
+		TestGetVarArgs(OutputString, "Test N|%hs|", "\xF0\x9F\x98\x86\xF0\x9F\x80\x9C");
+		CHECK_MESSAGE(OutputString, FString(OutputString) == FString("Test N|\xF0\x9F\x98\x86\xF0\x9F\x80\x9C|"));
+
+		TestGetVarArgs(OutputString, "Test O|%s|", "\xF0\x9F\x98\x86\xF0\x9F\x80\x9C");
+		CHECK_MESSAGE(OutputString, FString(OutputString) == FString("Test O|\xF0\x9F\x98\x86\xF0\x9F\x80\x9C|"));
+
+		TestGetVarArgs(OutputString, "Test P|%S|", UTF16_TO_TCHAR(u"\xD83D\xDE06\xD83C\xDC1C"));
+		CHECK_MESSAGE(OutputString, FString(OutputString) == FString("Test P|\xF0\x9F\x98\x86\xF0\x9F\x80\x9C|"));
+	}
+}
+#endif // WITH_LOW_LEVEL_TESTS
\ No newline at end of file
diff --git a/Engine/Source/Runtime/Core/Private/GenericPlatform/GenericWidePlatformString.cpp b/Engine/Source/Runtime/Core/Private/GenericPlatform/GenericWidePlatformString.cpp
--- a/Engine/Source/Runtime/Core/Private/GenericPlatform/GenericWidePlatformString.cpp
+++ b/Engine/Source/Runtime/Core/Private/GenericPlatform/GenericWidePlatformString.cpp
@@ -233,692 +233,4 @@ int iswspace(wint_t wc)
 }
 #endif
 
-
-#if WITH_LOW_LEVEL_TESTS
-constexpr int OUTPUT_SIZE = 256;
-int32 TestGetVarArgs(WIDECHAR* OutputString, const WIDECHAR* Format, ...)
-{
-	va_list ArgPtr;
-	va_start( ArgPtr, Format );
-	int32 Result = FGenericWidePlatformString::GetVarArgs(OutputString, OUTPUT_SIZE, Format, ArgPtr);
-	va_end(ArgPtr);
-	return Result;
-}
-
-TEST_CASE("Core::PlatformString::GetVarArgs", "[Core][String][Smoke]")
-{
-	WIDECHAR OutputString[OUTPUT_SIZE];
-
-	TestGetVarArgs(OutputString, TEXT("Test A|%-20s|%20s|%10.2f|%-10.2f|"), TEXT("LEFT"), TEXT("RIGHT"), 33.333333, 66.666666);
-	CHECK_MESSAGE(OutputString, FString(OutputString) == FString(TEXT("Test A|LEFT                |               RIGHT|     33.33|66.67     |")));
-
-	TestGetVarArgs(OutputString, TEXT("Test B|Percents:%%%%%%%d|"), 3);
-	CHECK_MESSAGE(OutputString, FString(OutputString) == FString(TEXT("Test B|Percents:%%%3|")));
-
-	TestGetVarArgs(OutputString, TEXT("Test C|%d|%i|%X|%x|%u|"), 12345, 54321, 0x123AbC, 15, 99);
-	CHECK_MESSAGE(OutputString, FString(OutputString) == FString(TEXT("Test C|12345|54321|123ABC|f|99|")));
-
-	TestGetVarArgs(OutputString, TEXT("Test D|%p|"), 0x12345);
-#ifdef _MSC_VER
-	// MSVC's standard library formats pointers differently
-	CHECK_MESSAGE(OutputString, FString(OutputString) == FString(TEXT("Test D|0000000000012345|")));
-#else
-	// Pointer format (which we get from snprintf) varies from platform to platform.
-	// Call it to make sure it puts the output in the proper place, but don't test for the exact format
-//	CHECK_MESSAGE(OutputString, FString(OutputString) == FString(TEXT("Test D|0x12345|")));
-	{
-		FString OutputStringStr(OutputString);
-		CHECK_MESSAGE(OutputString, OutputStringStr.StartsWith(TEXT("Test D|"))
-			&& OutputStringStr.EndsWith(TEXT("|")) && OutputStringStr.Len() < TEXTVIEW("Test D||").Len() + 20);
-	}
-#endif
-
-	TestGetVarArgs(OutputString, TEXT("Test E|%" INT64_FMT "|"), int64(12345678912345LL));
-	CHECK_MESSAGE(OutputString, FString(OutputString) == FString(TEXT("Test E|12345678912345|")));
-
-	TestGetVarArgs(OutputString, TEXT("Test F|%f|%e|%g|"), 123.456, 123.456, 123.456);
-	CHECK_MESSAGE(OutputString, FString(OutputString) == FString(TEXT("Test F|123.456000|1.234560e+02|123.456|")));
-
-	TestGetVarArgs(OutputString, TEXT("Test G|%" UPTRINT_X_FMT"|"), UPTRINT(49374));
-	CHECK_MESSAGE(OutputString, FString(OutputString) == FString(TEXT("Test G|C0DE|")));
-
-	TestGetVarArgs(OutputString, TEXT("Test H|%" UPTRINT_x_FMT "|"), UPTRINT(49374));
-	CHECK_MESSAGE(OutputString, FString(OutputString) == FString(TEXT("Test H|c0de|")));
-
-	TestGetVarArgs(OutputString, TEXT("Test I|%" UINT64_FMT "|"), MAX_uint64);
-	CHECK_MESSAGE(OutputString, FString(OutputString) == FString(TEXT("Test I|18446744073709551615|")));
-
-	TestGetVarArgs(OutputString, TEXT("Test J|%*c|"), 10, 'J');
-	CHECK_MESSAGE(OutputString, FString(OutputString) == FString(TEXT("Test J|         J|")));
-
-	TestGetVarArgs(OutputString, TEXT("Test K|%-5c|"), 'K');
-	CHECK_MESSAGE(OutputString, FString(OutputString) == FString(TEXT("Test K|K    |")));
-}
-#endif // WITH_LOW_LEVEL_TESTS
-
-namespace
-{
-	// Output iterator which will not overflow the destination buffer but will keep track of how many characters have been written.
-	// It will also null terminate the output when destructed.
-	struct FSafeDestIterator
-	{
-		// Non-copyable
-		FSafeDestIterator(const FSafeDestIterator&) = delete;
-		FSafeDestIterator& operator=(const FSafeDestIterator&) = delete;
-
-		explicit FSafeDestIterator(WIDECHAR* InPtr, SIZE_T InEnd)
-			: NumCharsWritten(0)
-			, Ptr            (InPtr)
-			, EndMinusOne    (InPtr + InEnd - 1)
-		{
-			check(InPtr);
-		}
-
-		FORCENOINLINE ~FSafeDestIterator()
-		{
-			// Auto-terminate buffer
-			*Ptr = 0;
-		}
-
-		// Writes Count instances of Ch to the destination buffer
-		FORCENOINLINE FSafeDestIterator& Write(TCHAR Ch, int32 Count = 1)
-		{
-			// Ensure we're not in the error state
-			if (EndMinusOne)
-			{
-				NumCharsWritten += Count;
-
-				int32 NumToWrite = FPlatformMath::Min(Count, (int32)(EndMinusOne - Ptr));
-				for (int32 Val = NumToWrite; Val; --Val)
-				{
-					*Ptr++ = Ch;
-				}
-
-				if (NumToWrite != Count)
-				{
-					EndMinusOne = 0;
-				}
-			}
-
-			return *this;
-		}
-
-		// Writes Count characters from Src to the destination buffer
-		template <typename CharType>
-		FORCENOINLINE FSafeDestIterator& Write(const CharType* Src, int32 Count)
-		{
-			// Ensure we're not in the error state
-			if (EndMinusOne)
-			{
-				NumCharsWritten += Count;
-
-				int32 NumToWrite = FPlatformMath::Min(Count, (int32)(EndMinusOne - Ptr));
-				for (int32 Val = NumToWrite; Val; --Val)
-				{
-					*Ptr++ = *Src++;
-				}
-
-				if (NumToWrite != Count)
-				{
-					EndMinusOne = 0;
-				}
-			}
-
-			return *this;
-		}
-
-		explicit operator bool() const
-		{
-			return !!EndMinusOne;
-		}
-
-		int32 GetNumCharsWritten() const
-		{
-			return NumCharsWritten;
-		}
-
-	private:
-		int32     NumCharsWritten;
-		WIDECHAR* Ptr;
-		WIDECHAR* EndMinusOne; // when null, this means the iterator has already moved past the writable area of the buffer
-	};
-
-	static inline bool CharIsIntegerFormatSpecifier(TCHAR Char)
-	{
-		return (Char == 'i') | (Char == 'd') | (Char == 'u') | (Char == 'X') | (Char == 'x');
-	}
-
-	template <typename CharType, typename VaListType>
-	void ProcessStringArg(FSafeDestIterator& DestIter, const TCHAR*& Src, int FieldLen, int PrecisionLen, VaListType& ArgPtr)
-	{
-		// ArgPtr is taken as a templated parameter, despite being a va_list, because on some
-		// architectures it could be an array type, and when an array is used as a function parameter,
-		// the named value is treated as a pointer and won't bind to an array reference.
-
-		Src++;
-		static const CharType* Null = LITERAL(CharType, "(null)");
-		const CharType* Val = va_arg(ArgPtr, CharType*);
-		if (Val == nullptr)
-		{
-			Val = Null;
-		}
-
-		int RetCnt = PrecisionLen < 0 ? FGenericWidePlatformString::Strlen(Val) : FGenericWidePlatformString::Strnlen(Val, PrecisionLen);
-		int Spaces = FPlatformMath::Max(FPlatformMath::Abs(FieldLen) - RetCnt, 0);
-		if (Spaces > 0 && FieldLen > 0)
-		{
-			DestIter.Write(TEXT(' '), Spaces);
-		}
-		DestIter.Write(Val, RetCnt);
-		if (Spaces > 0 && FieldLen < 0)
-		{
-			DestIter.Write(TEXT(' '), Spaces);
-		}
-	}
-}
-
-int32 FGenericWidePlatformString::GetVarArgs( WIDECHAR* Dest, SIZE_T DestSize, const WIDECHAR*& Fmt, va_list ArgPtr )
-{
-#if PLATFORM_TCHAR_IS_4_BYTES || PLATFORM_TCHAR_IS_UTF8CHAR
-	unimplemented();
-#else
-	static_assert(sizeof(TCHAR) == 2, "TCHAR is expected to be 16-bit");
-#endif
-
-	if (Fmt == nullptr)
-	{
-		if ((DestSize > 0) && (Dest != nullptr))
-		{
-			*Dest = 0;
-		}
-		return 0;
-	}
-
-	const TCHAR *Src = Fmt;
-
-	FSafeDestIterator DestIter(Dest, DestSize);
-	while (*Src)
-	{
-		if (*Src != '%')
-		{
-			if (!DestIter.Write(*Src++))
-			{
-				return -1;
-			}
-			continue;
-		}
-
-		const TCHAR *Percent = Src;
-		int FieldLen = 0;
-		bool bVariablePrecision = false;
-		int PrecisionLen = -1;
-
-		Src++; // skip the '%' char...
-
-		// Skip modifier flags that don't need additional processing;
-		// they still get passed to snprintf() below based on the conversion.
-		while (*Src == '+' || *Src == '#' || *Src == TEXT(' ') || *Src == TEXT('0'))
-		{
-			Src++;
-		}
-
-		// check for field width requests...
-		if ((*Src == '-') || ((*Src >= '0') && (*Src <= '9')))
-		{
-			const TCHAR *Cur = Src + 1;
-			while ((*Cur >= '0') && (*Cur <= '9'))
-			{
-				Cur++;
-			}
-
-			FieldLen = Atoi(Src);
-			Src = Cur;
-		}
-
-		// check for dynamic field requests
-		if (*Src == '*')
-		{
-			FieldLen = va_arg(ArgPtr, int32);
-			Src++;
-		}
-
-		if (*Src == '.')
-		{
-			const TCHAR *Cur = Src + 1;
-			if (*Cur == '*')
-			{
-				bVariablePrecision = true;
-				PrecisionLen = va_arg(ArgPtr, int32);
-				Cur++;
-			}
-			else
-			{
-				while ((*Cur >= '0') && (*Cur <= '9'))
-				{
-					Cur++;
-				}
-
-				PrecisionLen = Atoi(Src + 1);
-			}
-			Src = Cur;
-		}
-
-		// Check for 'ls' field, change to 's'
-		if ((Src[0] == 'l' && Src[1] == 's'))
-		{
-			Src++;
-		}
-
-		// Check for 'hs' field, change to 's' and remember
-		bool bIsNarrowString = false;
-		if ((Src[0] == 'h' && Src[1] == 's'))
-		{
-			bIsNarrowString = true;
-			Src++;
-		}
-
-		switch (*Src)
-		{
-			case '%':
-			{
-				Src++;
-				if (!DestIter.Write('%'))
-				{
-					return -1;
-				}
-				break;
-			}
-
-			case 'c':
-			{
-				TCHAR Val = (TCHAR) va_arg(ArgPtr, int);
-				Src++;
-				bool bSuccess = true;
-				if (FieldLen > 1)
-				{
-					bSuccess = !!DestIter.Write(TCHAR(' '), FieldLen - 1);
-				}
-				bSuccess = bSuccess && !!DestIter.Write(Val);
-				if (FieldLen < -1)
-				{
-					bSuccess = bSuccess && !!DestIter.Write(TCHAR(' '), FPlatformMath::Abs(FieldLen + 1));
-				}
-				if (!bSuccess)
-				{
-					return -1;
-				}
-				break;
-			}
-
-			case 'd':
-			case 'i':
-			case 'X':
-			case 'x':
-			case 'u':
-			{
-				Src++;
-				int Val = va_arg(ArgPtr, int);
-				ANSICHAR AnsiNum[64];
-				ANSICHAR FmtBuf[30];
-
-				// limit width to the buffer size
-				FieldLen = FMath::Min(static_cast<int>(sizeof(AnsiNum)) - 1, FieldLen);
-
-				// Yes, this is lame.
-				int CpyIdx = 0;
-				while (Percent < Src && CpyIdx < UE_ARRAY_COUNT(FmtBuf))
-				{
-					if (LIKELY(*Percent != '*'))
-					{
-						FmtBuf[CpyIdx] = (ANSICHAR) *Percent;
-						++CpyIdx;
-					}
-					else
-					{
-						snprintf(&FmtBuf[CpyIdx], sizeof(FmtBuf) - CpyIdx, "%d", FieldLen);
-						while (CpyIdx < UE_ARRAY_COUNT(FmtBuf) && FmtBuf[CpyIdx] != 0)
-						{
-							++CpyIdx;
-						}
-					}
-					++Percent;
-				}
-				FmtBuf[CpyIdx] = 0;
-
-				int RetCnt = snprintf(AnsiNum, sizeof(AnsiNum), FmtBuf, Val);
-				if (RetCnt < 0 || !DestIter.Write(AnsiNum, RetCnt))
-				{
-					return -1;
-				}
-				break;
-			}
-
-			case 'z':
-			case 'Z':
-			{
-				if (Src[1] == 0)
-				{
-					printf("Unknown percent [%lc] in FGenericWidePlatformString::GetVarArgs() [%s]\n.", (wint_t)Src[0], TCHAR_TO_ANSI(Fmt));
-					Src++;  // skip it, I guess.
-					break;
-				}
-
-				Src += 2;
-
-				size_t Val = va_arg(ArgPtr, size_t);
-
-				ANSICHAR AnsiNum[64];
-				ANSICHAR FmtBuf[30];
-
-				// limit width to the buffer size
-				FieldLen = FMath::Min(static_cast<int>(sizeof(AnsiNum)) - 1, FieldLen);
-
-				// Yes, this is lame.
-				int CpyIdx = 0;
-				while (Percent < Src && CpyIdx < UE_ARRAY_COUNT(FmtBuf))
-				{
-					if (LIKELY(*Percent != '*'))
-					{
-						FmtBuf[CpyIdx] = (ANSICHAR) *Percent;
-						++CpyIdx;
-					}
-					else
-					{
-						snprintf(&FmtBuf[CpyIdx], sizeof(FmtBuf) - CpyIdx, "%d", FieldLen);
-						while (CpyIdx < UE_ARRAY_COUNT(FmtBuf) && FmtBuf[CpyIdx] != 0)
-						{
-							++CpyIdx;
-						}
-					}
-					++Percent;
-				}
-				FmtBuf[CpyIdx] = 0;
-
-				int RetCnt = snprintf(AnsiNum, sizeof(AnsiNum), FmtBuf, Val);
-				if (RetCnt < 0 || !DestIter.Write(AnsiNum, RetCnt))
-				{
-					return -1;
-				}
-				break;
-			}
-
-			case 'p':
-			{
-				Src++;
-				void* Val = va_arg(ArgPtr, void*);
-				ANSICHAR AnsiNum[30];
-				ANSICHAR FmtBuf[30];
-
-				// Yes, this is lame.
-				int CpyIdx = 0;
-				while (Percent < Src && CpyIdx < UE_ARRAY_COUNT(FmtBuf))
-				{
-					FmtBuf[CpyIdx] = (ANSICHAR)*Percent;
-					Percent++;
-					CpyIdx++;
-				}
-				FmtBuf[CpyIdx] = 0;
-
-				int RetCnt = snprintf(AnsiNum, sizeof(AnsiNum), FmtBuf, Val);
-				if (RetCnt < 0 || !DestIter.Write(AnsiNum, RetCnt))
-				{
-					return -1;
-				}
-				break;
-			}
-
-			case 'l':
-			case 'I':
-			case 'h':
-			{
-				int RemainingSize = Strlen(Src);
-
-				if (RemainingSize < 2)
-				{
-					printf("Unknown percent [%lc] in FGenericWidePlatformString::GetVarArgs() [%s]\n.", (wint_t)Src[0], TCHAR_TO_ANSI(Fmt));
-					Src++;  // skip it, I guess.
-					break;
-				}
-
-				// treat %ld as %d. Also shorts for %h will be promoted to ints. This path also handles %li, %lu, %lx and %lX.
-				if ((Src[0] == 'l' || Src[0] == 'h') && CharIsIntegerFormatSpecifier(Src[1]))
-				{
-					Src+=2;
-					long int Val = va_arg(ArgPtr, long int);
-					ANSICHAR AnsiNum[30];
-					ANSICHAR FmtBuf[30];
-
-					// Yes, this is lame.
-					int CpyIdx = 0;
-					while (Percent < Src && CpyIdx < UE_ARRAY_COUNT(FmtBuf))
-					{
-						FmtBuf[CpyIdx] = (ANSICHAR)*Percent;
-						Percent++;
-						CpyIdx++;
-					}
-					FmtBuf[CpyIdx] = 0;
-
-					int RetCnt = snprintf(AnsiNum, sizeof(AnsiNum), FmtBuf, Val);
-					if (RetCnt < 0 || !DestIter.Write(AnsiNum, RetCnt))
-					{
-						return -1;
-					}
-					break;
-				}
-				// Treat %lf as a %f
-				else if (Src[0] == 'l' && Src[1] == 'f')
-				{
-					Src += 2;
-					double Val = va_arg(ArgPtr, double);
-					ANSICHAR AnsiNum[340]; // make sure buffer is large enough to handle the longest double values, such as the largest number which is 316 characters including decimal point and 6 decimals
-					ANSICHAR FmtBuf[30];
-
-					// Yes, this is lame.
-					int CpyIdx = 0;
-					while (Percent < Src && CpyIdx < UE_ARRAY_COUNT(FmtBuf))
-					{
-						FmtBuf[CpyIdx] = (ANSICHAR) *Percent;
-						Percent++;
-						CpyIdx++;
-					}
-					FmtBuf[CpyIdx] = 0;
-
-					int RetCnt = snprintf(AnsiNum, sizeof(AnsiNum), FmtBuf, Val);
-					if (RetCnt < 0 || !DestIter.Write(AnsiNum, RetCnt))
-					{
-						return -1;
-					}
-					break;
-				}
-
-				// At this point we expect RemainingSize >= 3, but we know it's at least 2.
-				if (RemainingSize < 3)
-				{
-					printf("Unknown percent [%lc%lc] in FGenericWidePlatformString::GetVarArgs() [%s]\n.", (wint_t)Src[0], (wint_t)Src[1], TCHAR_TO_ANSI(Fmt));
-					Src++;  // skip it, I guess.
-					break;
-				}
-
-				// Deal with half-half formats - %hhd, %hhi, %hhu, %hhx and %hhX.
-				if (Src[0] == 'h')
-				{
-					if (Src[1] != 'h' || !CharIsIntegerFormatSpecifier(Src[2]))
-					{
-						printf("Unknown percent [%lc%lc%lc] in FGenericWidePlatformString::GetVarArgs() [%s]\n.", (wint_t)Src[0], (wint_t)Src[1], (wint_t)Src[2], TCHAR_TO_ANSI(Fmt));
-						Src++;  // skip it, I guess.
-						break;
-					}
-					Src += 3;
-					int Val = va_arg(ArgPtr, int);
-					ANSICHAR AnsiNum[8];
-					ANSICHAR FmtBuf[30];
-
-					// Yes, this is lame.
-					int CpyIdx = 0;
-					while (Percent < Src && CpyIdx < UE_ARRAY_COUNT(FmtBuf))
-					{
-						FmtBuf[CpyIdx] = (ANSICHAR)*Percent;
-						Percent++;
-						CpyIdx++;
-					}
-					FmtBuf[CpyIdx] = 0;
-
-					int RetCnt = snprintf(AnsiNum, sizeof(AnsiNum), FmtBuf, Val);
-					if (RetCnt < 0 || !DestIter.Write(AnsiNum, RetCnt))
-					{
-						return -1;
-					}
-					break;
-				}
-				else if (Src[0] == 'l')
-				{
-					if (Src[1] != 'l' || !CharIsIntegerFormatSpecifier(Src[2]))
-					{
-						printf("Unknown percent [%lc%lc%lc] in FGenericWidePlatformString::GetVarArgs() [%s]\n.", (wint_t)Src[0], (wint_t)Src[1], (wint_t)Src[2], TCHAR_TO_ANSI(Fmt));
-						Src++;  // skip it, I guess.
-						break;
-					}
-				}
-				else if (Src[0] == 'I')
-				{
-					if (RemainingSize < 4)
-					{
-						printf("Unknown percent [%lc%lc%lc] in FGenericWidePlatformString::GetVarArgs() [%s]\n.", (wint_t)Src[0], (wint_t)Src[1], (wint_t)Src[2], TCHAR_TO_ANSI(Fmt));
-						Src++;  // skip it, I guess.
-						break;
-					}
-					else if (Src[1] != '6' || Src[2] != '4' || !CharIsIntegerFormatSpecifier(Src[3]))
-					{
-						printf("Unknown percent [%lc%lc%lc%lc] in FGenericWidePlatformString::GetVarArgs() [%s]\n.", (wint_t)Src[0], (wint_t)Src[1], (wint_t)Src[2], (wint_t)Src[3], TCHAR_TO_ANSI(Fmt));
-						Src++;  // skip it, I guess.
-						break;
-					}
-				}
-				else
-				{
-					printf("Unknown percent [%lc%lc%lc] in FGenericWidePlatformString::GetVarArgs() [%s]\n.", (wint_t)Src[0], (wint_t)Src[1], (wint_t)Src[2], TCHAR_TO_ANSI(Fmt));
-					Src++;  // skip it, I guess.
-					break;
-				}
-
-				// Yes, this is lame.
-				unsigned long long Val = va_arg(ArgPtr, unsigned long long);
-				ANSICHAR AnsiNum[60];
-				ANSICHAR FmtBuf[30];
-
-				int CpyIdx;
-				if (Src[0] == 'I')
-				{
-					// Convert I64 syntax to ll 
-					Src += 4;
-					Percent += 4;
-					CpyIdx = 3;
-					FmtBuf[0] = '%';
-					FmtBuf[1] = 'l';
-					FmtBuf[2] = 'l';
-				}
-				else
-				{
-					Src += 3;
-					CpyIdx = 0;
-				}
-
-				while (Percent < Src && CpyIdx < UE_ARRAY_COUNT(FmtBuf))
-				{
-					FmtBuf[CpyIdx] = (ANSICHAR) *Percent;
-					Percent++;
-					CpyIdx++;
-				}
-				FmtBuf[CpyIdx] = 0;
-
-				int RetCnt = snprintf(AnsiNum, sizeof(AnsiNum), FmtBuf, Val);
-				if (RetCnt < 0 || !DestIter.Write(AnsiNum, RetCnt))
-				{
-					return -1;
-				}
-				break;
-			}
-
-			case 'f':
-			case 'e':
-			case 'g':
-			{
-				Src++;
-				double Val = va_arg(ArgPtr, double);
-				ANSICHAR AnsiNum[340]; // make sure buffer is large enough to handle the longest double values, such as the largest number which is 316 characters including decimal point and 6 decimals
-				ANSICHAR FmtBuf[30];
-
-				// Yes, this is lame.
-				int CpyIdx = 0;
-				while (Percent < Src && CpyIdx < UE_ARRAY_COUNT(FmtBuf))
-				{
-					FmtBuf[CpyIdx] = (ANSICHAR) *Percent;
-					Percent++;
-					CpyIdx++;
-				}
-				FmtBuf[CpyIdx] = 0;
-
-				int RetCnt;
-				if (bVariablePrecision)
-				{
-					RetCnt = snprintf(AnsiNum, sizeof (AnsiNum), FmtBuf, PrecisionLen, Val);
-				}
-				else
-				{
-					RetCnt = snprintf(AnsiNum, sizeof (AnsiNum), FmtBuf, Val);
-				}
-				if (RetCnt >= UE_ARRAY_COUNT(AnsiNum))
-				{
-					// We should print what we have written into AnsiNum but ensure we null terminate before printing
-					AnsiNum[UE_ARRAY_COUNT(AnsiNum) - 1] = '\0';
-					checkf(0, TEXT("Attempting to read past the size our buffer. Buffer Size: %" SIZE_T_FMT " Size to read: %d. Current contents: '%s'\n"), UE_ARRAY_COUNT(AnsiNum), RetCnt, UTF8_TO_TCHAR(AnsiNum));
-				}
-				if (RetCnt < 0 || !DestIter.Write(AnsiNum, RetCnt))
-				{
-					return -1;
-				}
-				break;
-			}
-
-			case 's':
-			{
-				if (bIsNarrowString)
-				{
-					ProcessStringArg<ANSICHAR>(DestIter, Src, FieldLen, PrecisionLen, ArgPtr);
-				}
-				else
-				{
-					ProcessStringArg<TCHAR>(DestIter, Src, FieldLen, PrecisionLen, ArgPtr);
-				}
-				if (!DestIter)
-				{
-					return -1;
-				}
-				break;
-			}
-
-			case 'S':
-			{
-				// The %S format represents a string which is the opposite of %s - wide if TCHAR is narrow, or narrow if TCHAR is wide
-				using OtherCharType = std::conditional_t<std::is_same_v<TCHAR, ANSICHAR>, WIDECHAR, ANSICHAR>;
-
-				ProcessStringArg<OtherCharType>(DestIter, Src, FieldLen, PrecisionLen, ArgPtr);
-				if (!DestIter)
-				{
-					return -1;
-				}
-				break;
-			}
-
-			default:
-				printf("Unknown percent [%%%c] in FGenericWidePlatformString::GetVarArgs().\n", *Src);
-				Src++;  // skip char, I guess.
-				break;
-		}
-	}
-
-	int32 Result = DestIter.GetNumCharsWritten();
-	return Result;
-}
-
 #endif // PLATFORM_USE_GENERIC_STRING_IMPLEMENTATION
diff --git a/Engine/Source/Runtime/Core/Private/GenericPlatform/StandardPlatformString.cpp b/Engine/Source/Runtime/Core/Private/GenericPlatform/StandardPlatformString.cpp
--- a/Engine/Source/Runtime/Core/Private/GenericPlatform/StandardPlatformString.cpp
+++ b/Engine/Source/Runtime/Core/Private/GenericPlatform/StandardPlatformString.cpp
@@ -11,333 +11,4 @@
 #include "Logging/LogCategory.h"
 #include "Logging/LogMacros.h"
 
-
-DEFINE_LOG_CATEGORY_STATIC(LogStandardPlatformString, Log, All);
-
-#if PLATFORM_IOS || PLATFORM_WINDOWS
-	#define VA_LIST_REF va_list&
-#else
-	#define VA_LIST_REF va_list
-#endif
-
-
-struct FFormatInfo
-{
-	WIDECHAR Format[32];
-	WIDECHAR LengthModifier;
-	WIDECHAR Type;
-	bool HasDynamicWidth;
-};
-
-static int32 GetFormattingInfo(const WIDECHAR* Format, FFormatInfo& OutInfo)
-{
-	const WIDECHAR* FormatStart = Format++;
-
-	// Skip flags
-	while (*Format == LITERAL(WIDECHAR, '#') || *Format == LITERAL(WIDECHAR, '0') || *Format == LITERAL(WIDECHAR, '-')
-		|| *Format == LITERAL(WIDECHAR, ' ') || *Format == LITERAL(WIDECHAR, '+') || *Format == LITERAL(WIDECHAR, '\''))
-	{
-		Format++;
-	}
-
-	OutInfo.HasDynamicWidth = false;
-
-	// Skip width
-	while ((*Format >= LITERAL(WIDECHAR, '0') && *Format <= LITERAL(WIDECHAR, '9')) || *Format == LITERAL(WIDECHAR, '*'))
-	{
-		if (*Format == LITERAL(WIDECHAR, '*'))
-		{
-			OutInfo.HasDynamicWidth = true;
-		}
-		Format++;
-	}
-
-	// Skip precision
-	if (*Format == LITERAL(WIDECHAR, '.'))
-	{
-		Format++;
-		while ((*Format >= LITERAL(WIDECHAR, '0') && *Format <= LITERAL(WIDECHAR, '9')) || *Format == LITERAL(WIDECHAR, '*'))
-		{
-			if (*Format == LITERAL(WIDECHAR, '*'))
-			{
-				OutInfo.HasDynamicWidth = true;
-			}
-			Format++;
-		}
-	}
-
-	OutInfo.LengthModifier = 0;
-	if (*Format == LITERAL(WIDECHAR, 'h') || *Format == LITERAL(WIDECHAR, 'l') || *Format == LITERAL(WIDECHAR, 'j')
-		|| *Format == LITERAL(WIDECHAR, 'q') || *Format == LITERAL(WIDECHAR, 'L'))
-	{
-		OutInfo.LengthModifier = *Format++;
-		if (*Format == LITERAL(WIDECHAR, 'h'))
-		{
-			OutInfo.LengthModifier = 'H';
-			Format++;
-		}
-		else if (*Format == LITERAL(WIDECHAR, 'l'))
-		{
-			OutInfo.LengthModifier = 'L';
-			Format++;
-		}
-	}
-	else if (*Format == LITERAL(WIDECHAR, 't') || *Format == LITERAL(WIDECHAR, 'z'))
-	{
-#if PLATFORM_64BITS
-		OutInfo.LengthModifier = LITERAL(WIDECHAR, 'l');
-		++Format;
-#else
-		OutInfo.LengthModifier = *Format++;
-#endif
-	}
-
-	OutInfo.Type = *Format++;
-
-	// The only valid length modifier for floating point types is L, all other modifiers should be ignored. Length modifier for void pointers should also be ignored.
-	if (OutInfo.LengthModifier != LITERAL(WIDECHAR, 'L') &&
-		(OutInfo.Type == LITERAL(WIDECHAR, 'f') || OutInfo.Type == LITERAL(WIDECHAR, 'F')
-			|| OutInfo.Type == LITERAL(WIDECHAR, 'e') || OutInfo.Type == LITERAL(WIDECHAR, 'E')
-			|| OutInfo.Type == LITERAL(WIDECHAR, 'g') || OutInfo.Type == LITERAL(WIDECHAR, 'G')
-			|| OutInfo.Type == LITERAL(WIDECHAR, 'a') || OutInfo.Type == LITERAL(WIDECHAR, 'A')
-			|| OutInfo.Type == LITERAL(WIDECHAR, 'p')))
-	{
-		OutInfo.LengthModifier = 0;
-	}
-
-	const int32 FormatLength = Format - FormatStart;
-
-	FMemory::Memcpy(OutInfo.Format, FormatStart, FormatLength * sizeof(WIDECHAR));
-	int32 OutInfoFormatLength = FormatLength;
-	if (OutInfo.HasDynamicWidth)
-	{
-		if (OutInfo.Type == LITERAL(WIDECHAR, 's'))
-		{
-			OutInfo.Format[OutInfoFormatLength - 1] = 'l';
-			OutInfo.Format[OutInfoFormatLength++] = 's';
-		}
-		else if (OutInfo.Type == LITERAL(WIDECHAR, 'S'))
-		{
-			OutInfo.Format[OutInfoFormatLength - 1] = 'h';
-			OutInfo.Format[OutInfoFormatLength++] = 's';
-		}
-	}
-	OutInfo.Format[OutInfoFormatLength] = 0;
-
-	// HACKHACKHACK
-	// This formatting function expects to understand %s as a string no matter which char width.
-	// On mac (and possibly others) this must be fixed up to %S for widechars.
-	// So we will do the fixup ONLY if this is a widechar system and the format is given as %s.
-	// BUG: This function still doesn't handle char16_t correctly.
-	if (sizeof(WIDECHAR) == sizeof(wchar_t))
-	{
-		if (OutInfo.Type == LITERAL(WIDECHAR, 's'))
-		{
-			checkSlow(OutInfo.Format[OutInfoFormatLength - 1] == LITERAL(WIDECHAR, 's'));
-			OutInfo.Format[OutInfoFormatLength - 1] = LITERAL(WIDECHAR, 'S');
-		}
-		else if (OutInfo.Type == LITERAL(WIDECHAR, 'S'))
-		{
-			OutInfo.Format[OutInfoFormatLength - 1] = LITERAL(WIDECHAR, 's');
-		}
-	}
-
-	return FormatLength;
-}
-
-template <typename T1, typename T2>
-static int32 FormatString(const FFormatInfo& Info, VA_LIST_REF ArgPtr, WIDECHAR* Formatted, int32 Length)
-{
-	if (Info.HasDynamicWidth)
-	{
-		int32 Width = va_arg(ArgPtr, int32);
-		if (FChar::ToLower(Info.LengthModifier) == LITERAL(WIDECHAR, 'l'))
-		{
-			T1 Value = va_arg(ArgPtr, T1);
-			return swprintf(Formatted, Length, Info.Format, Width, Value);
-		}
-		else
-		{
-			T2 Value = va_arg(ArgPtr, T2);
-			return swprintf(Formatted, Length, Info.Format, Width, Value);
-		}
-	}
-	else
-	{
-		if (FChar::ToLower(Info.LengthModifier) == LITERAL(WIDECHAR, 'l'))
-		{
-			T1 Value = va_arg(ArgPtr, T1);
-			return swprintf(Formatted, Length, Info.Format, Value);
-		}
-		else
-		{
-			T2 Value = va_arg(ArgPtr, T2);
-			return swprintf(Formatted, Length, Info.Format, Value);
-		}
-	}
-}
-
-template <typename CallableType>
-static void FormatArgument(const FFormatInfo& Info, VA_LIST_REF ArgPtr, CallableType Callable)
-{
-	WIDECHAR Formatted[1024];
-
-	if (FChar::ToLower(Info.Type) == LITERAL(WIDECHAR, 's'))
-	{
-		if (Info.HasDynamicWidth)
-		{
-			int32 Width = va_arg(ArgPtr, int32);
-			const WIDECHAR* String = va_arg(ArgPtr, WIDECHAR*);
-			if (String)
-			{
-				int32 Length = swprintf(Formatted, UE_ARRAY_COUNT(Formatted), Info.Format, Width, String);
-				Callable(Formatted, Length);
-			}
-			else
-			{
-				Callable(TEXT("(null)"), UE_ARRAY_COUNT(Formatted));
-			}
-			return;
-		}
-		// Is it a plain string?
-		else if (Info.Type == LITERAL(WIDECHAR, 's') && Info.Format[1] == LITERAL(WIDECHAR, 'S'))
-		{
-			const WIDECHAR* String = va_arg(ArgPtr, WIDECHAR*);
-			int32 Length = FCString::Strlen(String);
-			Callable(String ? String : TEXT("(null)"), Length);
-		}
-		// Is it a wide string?
-		else if (Info.Type == LITERAL(WIDECHAR, 'S') && Info.Format[1] == LITERAL(WIDECHAR, 's'))
-		{
-			const ANSICHAR* String = va_arg(ArgPtr, ANSICHAR*);
-			int32 FormattedLength = FCStringAnsi::Strlen(String);
-			Callable(String ? String : "(null)", FormattedLength);
-		}
-		// Some form of string requiring formatting, such as a left- or right-justified string
-		else
-		{
-			// We call swprintf directly which may expect %S for a widechar string. This will be fixed up
-			// by the time we get here (See above in GetFormattingInfo).
-			const WIDECHAR* String = va_arg(ArgPtr, WIDECHAR*);
-			int32 Length = swprintf(Formatted, UE_ARRAY_COUNT(Formatted), Info.Format, String);
-			Callable(Formatted, Length);
-		}
-		return;
-	}
-
-	if (FChar::ToLower(Info.Type) == LITERAL(WIDECHAR, 'c'))
-	{
-		const WIDECHAR Char = (WIDECHAR)va_arg(ArgPtr, int32);
-		Formatted[0] = Char;
-		Formatted[1] = 0;
-		Callable(Formatted, 1);
-		return;
-	}
-
-	int32 Length = 0;
-	if (FChar::ToLower(Info.Type) == LITERAL(WIDECHAR, 'a') || FChar::ToLower(Info.Type) == LITERAL(WIDECHAR, 'e')
-		|| FChar::ToLower(Info.Type) == LITERAL(WIDECHAR, 'f') || FChar::ToLower(Info.Type) == LITERAL(WIDECHAR, 'g'))
-	{
-		Length = FormatString<long double, double>(Info, ArgPtr, Formatted, UE_ARRAY_COUNT(Formatted));
-	}
-	else if (Info.Type == LITERAL(WIDECHAR, 'p'))
-	{
-		void* Value = va_arg(ArgPtr, void*);
-		Length = swprintf(Formatted, UE_ARRAY_COUNT(Formatted), Info.Format, Value);
-	}
-	else if (FChar::ToLower(Info.Type) == LITERAL(WIDECHAR, 'd') || FChar::ToLower(Info.Type) == LITERAL(WIDECHAR, 'i'))
-	{
-		Length = FormatString<int64, int32>(Info, ArgPtr, Formatted, UE_ARRAY_COUNT(Formatted));
-	}
-	else if (FChar::ToLower(Info.Type) == LITERAL(WIDECHAR, 'o') || FChar::ToLower(Info.Type) == LITERAL(WIDECHAR, 'u') || FChar::ToLower(Info.Type) == LITERAL(WIDECHAR, 'x'))
-	{
-		Length = FormatString<uint64, uint32>(Info, ArgPtr, Formatted, UE_ARRAY_COUNT(Formatted));
-	}
-
-	check(Length != -1);
-
-	Callable(Formatted, Length);
-}
-
-int32 FStandardPlatformString::GetVarArgs( WIDECHAR* Dest, SIZE_T DestSize, const WIDECHAR*& Fmt, va_list ArgPtr )
-{
-	const WIDECHAR* Format = Fmt;
-	const WIDECHAR* DestStart = Dest;
-
-	if (DestSize == 0)
-	{
-		return -1;
-	}
-
-	--DestSize;
-	while (*Format)
-	{
-		if (*Format == LITERAL(WIDECHAR, '%'))
-		{
-			if (*(Format + 1) == LITERAL(WIDECHAR, '%'))
-			{
-				if (DestSize == 0)
-				{
-					*Dest = 0;
-					return -1;
-				}
-				*Dest++ = *Format;
-				Format += 2;
-				DestSize--;
-				continue;
-			}
-
-			FFormatInfo Info;
-			Format += GetFormattingInfo(Format, Info);
-
-			bool bLengthExceeded = false;
-			FormatArgument(Info, ArgPtr, [&Dest, &DestSize, &bLengthExceeded](const auto* FormattedArg, int32 Length)
-			{
-				if (FormattedArg && Length > 0)
-				{
-					if (Length > DestSize)
-					{
-						Length = DestSize;
-						bLengthExceeded = true;
-					}
-
-					if (sizeof(*FormattedArg) == sizeof(WIDECHAR))
-					{
-						FMemory::Memcpy(Dest, FormattedArg, Length * sizeof(WIDECHAR));
-					}
-					else
-					{
-						for (int32 Index = 0; Index != Length; ++Index)
-						{
-							Dest[Index] = FormattedArg[Index];
-						}
-					}
-
-					Dest += Length;
-					DestSize -= Length;
-				}
-			});
-			if (bLengthExceeded)
-			{
-				*Dest = 0;
-				return -1;
-			}
-		}
-		else
-		{
-			if (DestSize == 0)
-			{
-				*Dest = 0;
-				return -1;
-			}
-			*Dest++ = *Format++;
-			DestSize--;
-		}
-	}
-
-	*Dest = 0;
-
-	return Dest - DestStart;
-}
-
-#endif // !PLATFORM_USE_SYSTEM_VSWPRINTF
+#endif // !PLATFORM_USE_SYSTEM_VSWPRINTF
\ No newline at end of file
diff --git a/Engine/Source/Runtime/Core/Public/GenericPlatform/GenericPlatformString.h b/Engine/Source/Runtime/Core/Public/GenericPlatform/GenericPlatformString.h
--- a/Engine/Source/Runtime/Core/Public/GenericPlatform/GenericPlatformString.h
+++ b/Engine/Source/Runtime/Core/Public/GenericPlatform/GenericPlatformString.h
@@ -11,6 +11,7 @@
 #include "Traits/IsFixedWidthCharEncoding.h"
 
 #include <type_traits>
+#include <stdarg.h>
 
 namespace UE::Core::Private
 {
@@ -339,6 +340,13 @@ struct FGenericPlatformString : public FGenericPlatformStricmp
 	CORE_API static int32 Strncmp(const WIDECHAR* String1, const UTF8CHAR* String2, SIZE_T Count);
 	CORE_API static int32 Strncmp(const UTF8CHAR* String1, const UTF8CHAR* String2, SIZE_T Count);
 
+	static CORE_API int32 GetVarArgs(ANSICHAR* Dest, SIZE_T DestSize, const ANSICHAR*& Fmt, va_list ArgPtr);
+	static CORE_API int32 GetVarArgs(WIDECHAR* Dest, SIZE_T DestSize, const WIDECHAR*& Fmt, va_list ArgPtr);
+	UE_REWRITE static int32 GetVarArgs(UTF8CHAR* Dest, SIZE_T DestSize, const UTF8CHAR*& Fmt, va_list ArgPtr)
+	{
+		return GetVarArgs((ANSICHAR*)Dest, DestSize, (const ANSICHAR*&)(Fmt), ArgPtr);
+	}
+
 private:
 	/**
 	 * Forwarding function because we can't call FMemory::Memcpy directly due to #include ordering issues.
diff --git a/Engine/Source/Runtime/Core/Public/GenericPlatform/GenericWidePlatformString.h b/Engine/Source/Runtime/Core/Public/GenericPlatform/GenericWidePlatformString.h
--- a/Engine/Source/Runtime/Core/Public/GenericPlatform/GenericWidePlatformString.h
+++ b/Engine/Source/Runtime/Core/Public/GenericPlatform/GenericWidePlatformString.h
@@ -207,8 +207,6 @@ public:
 	
 	CORE_API static WIDECHAR* Strtok( WIDECHAR* StrToken, const WIDECHAR* Delim, WIDECHAR** Context );
 
-	CORE_API static int32 GetVarArgs(WIDECHAR* Dest, SIZE_T DestSize, const WIDECHAR*& Fmt, va_list ArgPtr);
-
 	/**
 	 * Ansi implementation
 	 **/
@@ -344,41 +342,6 @@ PRAGMA_DISABLE_DEPRECATION_WARNINGS
 PRAGMA_ENABLE_DEPRECATION_WARNINGS
 	}
 
-	static int32 GetVarArgs( ANSICHAR* Dest, SIZE_T DestSize, const ANSICHAR*& Fmt, va_list ArgPtr )
-	{
-#if PLATFORM_USE_S_SPEC_FOR_NARROWCHAR_IN_VSPRINTF
-		// fix up the Fmt string, as fast as possible, without using an FAnsiString
-		ANSICHAR* NewFormat = (ANSICHAR*)alloca(Strlen(Fmt) + 1);
-
-		const ANSICHAR* From = Fmt;
-		ANSICHAR* To = NewFormat;
-		for (;;)
-		{
-			ANSICHAR FromCh = *From++;
-			*To++ = FromCh;
-			if (FromCh != '%')
-			{
-				if (!FromCh)
-				{
-					break;
-				}
-
-				continue;
-			}
-
-			if (From[0] == 'h' && From[1] == 's')
-			{
-				// Skip the h and copy the s next time to give %s
-				++From;
-			}
-		}
-		int32 Result = vsnprintf(Dest, DestSize, NewFormat, ArgPtr);
-#else
-		int32 Result = vsnprintf(Dest, DestSize, Fmt, ArgPtr);
-#endif
-		return (Result != -1 && Result < (int32)DestSize) ? Result : -1;
-	}
-
 	/**
 	 * UTF8CHAR implementation.
 	 */
@@ -497,11 +460,6 @@ PRAGMA_ENABLE_DEPRECATION_WARNINGS
 		return (UTF8CHAR*)Strtok((ANSICHAR*)StrToken, (const ANSICHAR*)Delim, (ANSICHAR**)Context);
 	}
 
-	static UE_FORCEINLINE_HINT int32 GetVarArgs(UTF8CHAR* Dest, SIZE_T DestSize, const UTF8CHAR*& Fmt, va_list ArgPtr)
-	{
-		return GetVarArgs((ANSICHAR*)Dest, DestSize, *(const ANSICHAR**)&Fmt, ArgPtr);
-	}
-
 	/**
 	 * UCS2 implementation
 	 **/
diff --git a/Engine/Source/Runtime/Core/Public/GenericPlatform/StandardPlatformString.h b/Engine/Source/Runtime/Core/Public/GenericPlatform/StandardPlatformString.h
--- a/Engine/Source/Runtime/Core/Public/GenericPlatform/StandardPlatformString.h
+++ b/Engine/Source/Runtime/Core/Public/GenericPlatform/StandardPlatformString.h
@@ -182,70 +182,6 @@ PRAGMA_DISABLE_DEPRECATION_WARNINGS
 PRAGMA_ENABLE_DEPRECATION_WARNINGS
 	}
 
-#if PLATFORM_USE_SYSTEM_VSWPRINTF
-	static int32 GetVarArgs( WIDECHAR* Dest, SIZE_T DestSize, const WIDECHAR*& Fmt, va_list ArgPtr )
-	{
-#if PLATFORM_USE_LS_SPEC_FOR_WIDECHAR
-		// fix up the Fmt string, as fast as possible, without using an FString
-		const WIDECHAR* OldFormat = Fmt;
-		WIDECHAR* NewFormat = (WIDECHAR*)alloca((Strlen(Fmt) * 2 + 1) * sizeof(WIDECHAR));
-
-		int32 NewIndex = 0;
-
-		for (; *OldFormat != 0; NewIndex++, OldFormat++)
-		{
-			// fix up %s -> %ls and %c -> %lc
-			if (OldFormat[0] == LITERAL(WIDECHAR, '%'))
-			{
-				NewFormat[NewIndex++] = *OldFormat++;
-
-				if (*OldFormat == LITERAL(WIDECHAR, '%'))
-				{
-					NewFormat[NewIndex] = *OldFormat;
-				}
-				else
-				{
-					const WIDECHAR* NextChar = OldFormat;
-
-					while(*NextChar != 0 && !FChar::IsAlpha(*NextChar))
-					{
-						NewFormat[NewIndex++] = *NextChar;
-						++NextChar;
-					};
-
-					if (*NextChar == LITERAL(WIDECHAR, 's') || *NextChar == LITERAL(WIDECHAR, 'c'))
-					{
-						NewFormat[NewIndex++] = LITERAL(WIDECHAR, 'l');
-						NewFormat[NewIndex] = *NextChar;
-					}
-					else if (*NextChar == LITERAL(WIDECHAR, 'S'))
-					{
-						NewFormat[NewIndex] = LITERAL(WIDECHAR, 's');
-					}
-					else
-					{
-						NewFormat[NewIndex] = *NextChar;
-					}
-
-					OldFormat = NextChar;
-				}
-			}
-			else
-			{
-				NewFormat[NewIndex] = *OldFormat;
-			}
-		}
-		NewFormat[NewIndex] = 0;
-		int32 Result = vswprintf( Dest, DestSize, NewFormat, ArgPtr);
-#else
-		int32 Result = vswprintf( Dest, DestSize, Fmt, ArgPtr);
-#endif
-		return Result;
-	}
-#else // PLATFORM_USE_SYSTEM_VSWPRINTF
-	static CORE_API int32 GetVarArgs( WIDECHAR* Dest, SIZE_T DestSize, const WIDECHAR*& Fmt, va_list ArgPtr );
-#endif // PLATFORM_USE_SYSTEM_VSWPRINTF
-
 	/**
 	 * Ansi implementation
 	 **/
@@ -379,12 +315,6 @@ PRAGMA_DISABLE_DEPRECATION_WARNINGS
 PRAGMA_ENABLE_DEPRECATION_WARNINGS
 	}
 
-	static int32 GetVarArgs( ANSICHAR* Dest, SIZE_T DestSize, const ANSICHAR*& Fmt, va_list ArgPtr )
-	{
-		int32 Result = vsnprintf(Dest, DestSize, Fmt, ArgPtr);
-		return (Result != -1 && Result < (int32)DestSize) ? Result : -1;
-	}
-
 	/**
 	 * UTF8CHAR implementation.
 	 */
@@ -502,11 +432,6 @@ PRAGMA_ENABLE_DEPRECATION_WARNINGS
 		return (UTF8CHAR*)Strtok((ANSICHAR*)StrToken, (const ANSICHAR*)Delim, (ANSICHAR**)Context);
 	}
 
-	static UE_FORCEINLINE_HINT int32 GetVarArgs(UTF8CHAR* Dest, SIZE_T DestSize, const UTF8CHAR*& Fmt, va_list ArgPtr)
-	{
-		return GetVarArgs((ANSICHAR*)Dest, DestSize, *(const ANSICHAR**)&Fmt, ArgPtr);
-	}
-
 	/**
 	 * UCS2 implementation
 	 **/
diff --git a/Engine/Source/Runtime/Core/Public/Microsoft/MicrosoftPlatformString.h b/Engine/Source/Runtime/Core/Public/Microsoft/MicrosoftPlatformString.h
--- a/Engine/Source/Runtime/Core/Public/Microsoft/MicrosoftPlatformString.h
+++ b/Engine/Source/Runtime/Core/Public/Microsoft/MicrosoftPlatformString.h
@@ -184,17 +184,6 @@ struct FMicrosoftPlatformString :
 		return _tcstok_s(StrToken, Delim, Context);
 	}
 
-// Allow fallback to FGenericWidePlatformString::GetVarArgs when PLATFORM_USE_GENERIC_STRING_IMPLEMENTATION is set.
-#if PLATFORM_USE_GENERIC_STRING_IMPLEMENTATION
-	using Super::GetVarArgs;
-#else
-	static inline int32 GetVarArgs( WIDECHAR* Dest, SIZE_T DestSize, const WIDECHAR*& Fmt, va_list ArgPtr )
-	{
-		int32 Result = vswprintf(Dest, DestSize, Fmt, ArgPtr);
-		return Result;
-	}
-#endif
-
 	/** 
 	 * Ansi implementation 
 	 **/
@@ -310,12 +299,6 @@ struct FMicrosoftPlatformString :
 		return strtok_s(StrToken, Delim, Context);
 	}
 
-	static inline int32 GetVarArgs( ANSICHAR* Dest, SIZE_T DestSize, const ANSICHAR*& Fmt, va_list ArgPtr )
-	{
-		int32 Result = vsnprintf( Dest, DestSize, Fmt, ArgPtr );
-		return (Result != -1 && Result < (int32)DestSize) ? Result : -1;
-	}
-
 	/**
 	 * UCS2CHAR implementation - this is identical to WIDECHAR for Windows platforms
 	 */
@@ -446,11 +429,6 @@ struct FMicrosoftPlatformString :
 	{
 		return (UTF8CHAR*)Strtok((ANSICHAR*)StrToken, (const ANSICHAR*)Delim, (ANSICHAR**)Context);
 	}
-
-	static UE_FORCEINLINE_HINT int32 GetVarArgs(UTF8CHAR* Dest, SIZE_T DestSize, const UTF8CHAR*& Fmt, va_list ArgPtr)
-	{
-		return GetVarArgs((ANSICHAR*)Dest, DestSize, *(const ANSICHAR**)&Fmt, ArgPtr);
-	}
 };
 
 #pragma warning(pop) // 'function' was was declared deprecated  (needed for the secure string functions)
