Automation Tests Framework: don't timeout tests when a debugger is present to avoid skipping steps while debugging.

diff --git a/Engine/Source/Runtime/Core/Public/Misc/AutomationTest.h b/Engine/Source/Runtime/Core/Public/Misc/AutomationTest.h
--- a/Engine/Source/Runtime/Core/Public/Misc/AutomationTest.h
+++ b/Engine/Source/Runtime/Core/Public/Misc/AutomationTest.h
@@ -2950,7 +2950,7 @@ private:
 				Reset();
 				return true;
 			}
-			else if (FDateTime::UtcNow() >= StartedRunning + Timeout)
+			else if ((FDateTime::UtcNow() >= StartedRunning + Timeout) && !FPlatformMisc::IsDebuggerPresent())
 			{
 				Reset();
 				Spec->AddError(TEXT("Latent command timed out."), 0);
@@ -3017,7 +3017,8 @@ private:
 				UE::TScopeLock CriticalSection(ActionCS);
 				bDone = false;
 				StartedRunning = FDateTime::UtcNow();
-				Future = Async(Execution, [this]() {
+				Future = Async(Execution, [this]()
+				{
 					Predicate(FDoneDelegate::CreateRaw(this, &FAsyncUntilDoneLatentCommand::Done));
 				});
 			}
@@ -3027,7 +3028,7 @@ private:
 				Reset();
 				return true;
 			}
-			else if (FDateTime::UtcNow() >= StartedRunning + Timeout)
+			else if ((FDateTime::UtcNow() >= StartedRunning + Timeout) && !FPlatformMisc::IsDebuggerPresent())
 			{
 				Reset();
 				Spec->AddError(TEXT("Latent command timed out."), 0);
@@ -3098,7 +3099,8 @@ private:
 				UE::TScopeLock CriticalSection(ActionCS);
 				bDone = false;
 				StartedRunning = FDateTime::UtcNow();
-				Future = Async(Execution, [this]() {
+				Future = Async(Execution, [this]()
+				{
 					Predicate();
 					Done();
 				});
@@ -3109,7 +3111,7 @@ private:
 				Reset();
 				return true;
 			}
-			else if (FDateTime::UtcNow() >= StartedRunning + Timeout)
+			else if ((FDateTime::UtcNow() >= StartedRunning + Timeout) && !FPlatformMisc::IsDebuggerPresent())
 			{
 				Reset();
 				Spec->AddError(TEXT("Latent command timed out."), 0);
