Added more C++20 concepts.

diff --git a/Engine/Source/Runtime/Core/Public/Concepts/Enum.h b/Engine/Source/Runtime/Core/Public/Concepts/Enum.h
new file mode 100644
--- /dev/null
+++ b/Engine/Source/Runtime/Core/Public/Concepts/Enum.h
@@ -0,0 +1,14 @@
+// Copyright Epic Games, Inc. All Rights Reserved.
+
+#pragma once
+
+#include <type_traits>
+
+namespace UE
+{
+	/**
+	 * Concept which describes an enum type.
+	 */
+	template <typename T>
+	concept CEnum = std::is_enum_v<T>;
+}
diff --git a/Engine/Source/Runtime/Core/Public/Concepts/ScopedEnum.h b/Engine/Source/Runtime/Core/Public/Concepts/ScopedEnum.h
new file mode 100644
--- /dev/null
+++ b/Engine/Source/Runtime/Core/Public/Concepts/ScopedEnum.h
@@ -0,0 +1,14 @@
+// Copyright Epic Games, Inc. All Rights Reserved.
+
+#pragma once
+
+#include <type_traits>
+
+namespace UE
+{
+	/**
+	 * Concept which describes a scoped enum type (aka. enum class).
+	 */
+	template <typename T>
+	concept CScopedEnum = std::is_enum_v<T> && !std::is_convertible_v<T, int>;
+}
diff --git a/Engine/Source/Runtime/Core/Public/Concepts/Signed.h b/Engine/Source/Runtime/Core/Public/Concepts/Signed.h
new file mode 100644
--- /dev/null
+++ b/Engine/Source/Runtime/Core/Public/Concepts/Signed.h
@@ -0,0 +1,14 @@
+// Copyright Epic Games, Inc. All Rights Reserved.
+
+#pragma once
+
+#include <type_traits>
+
+namespace UE
+{
+	/**
+	 * Concept which describes a signed type.
+	 */
+	template <typename T>
+	concept CSigned = std::is_signed_v<T>;
+}
diff --git a/Engine/Source/Runtime/Core/Public/Concepts/SignedIntegral.h b/Engine/Source/Runtime/Core/Public/Concepts/SignedIntegral.h
new file mode 100644
--- /dev/null
+++ b/Engine/Source/Runtime/Core/Public/Concepts/SignedIntegral.h
@@ -0,0 +1,14 @@
+// Copyright Epic Games, Inc. All Rights Reserved.
+
+#pragma once
+
+#include <type_traits>
+
+namespace UE
+{
+	/**
+	 * Concept which describes a signed integral type.  We use this instead of std::signed_integral because <concepts> isn't a well supported header yet.
+	 */
+	template <typename T>
+	concept CSignedIntegral = std::is_integral_v<T> && std::is_signed_v<T>;
+}
diff --git a/Engine/Source/Runtime/Core/Public/Concepts/Unsigned.h b/Engine/Source/Runtime/Core/Public/Concepts/Unsigned.h
new file mode 100644
--- /dev/null
+++ b/Engine/Source/Runtime/Core/Public/Concepts/Unsigned.h
@@ -0,0 +1,14 @@
+// Copyright Epic Games, Inc. All Rights Reserved.
+
+#pragma once
+
+#include <type_traits>
+
+namespace UE
+{
+	/**
+	 * Concept which describes an unsigned type.
+	 */
+	template <typename T>
+	concept CUnsigned = std::is_unsigned_v<T>;
+}
diff --git a/Engine/Source/Runtime/Core/Public/Concepts/UnsignedIntegral.h b/Engine/Source/Runtime/Core/Public/Concepts/UnsignedIntegral.h
new file mode 100644
--- /dev/null
+++ b/Engine/Source/Runtime/Core/Public/Concepts/UnsignedIntegral.h
@@ -0,0 +1,14 @@
+// Copyright Epic Games, Inc. All Rights Reserved.
+
+#pragma once
+
+#include <type_traits>
+
+namespace UE
+{
+	/**
+	 * Concept which describes an unsigned integral type.  We use this instead of std::unsigned_integral because <concepts> isn't a well supported header yet.
+	 */
+	template <typename T>
+	concept CUnsignedIntegral = std::is_integral_v<T> && std::is_unsigned_v<T>;
+}
