Corrected a comment about trivial relocatability on TArray.

#rnx

diff --git a/Engine/Source/Runtime/Core/Public/Containers/Array.h b/Engine/Source/Runtime/Core/Public/Containers/Array.h
--- a/Engine/Source/Runtime/Core/Public/Containers/Array.h
+++ b/Engine/Source/Runtime/Core/Public/Containers/Array.h
@@ -640,7 +640,7 @@ namespace UE::Core::Private
 /**
  * Templated dynamic array
  *
- * A dynamically sized array of typed elements.  Makes the assumption that your elements are relocate-able; 
+ * A dynamically sized array of typed elements.  Makes the assumption that your elements are trivially relocatable; 
  * i.e. that they can be transparently moved to new memory without a copy constructor.  The main implication 
  * is that pointers to elements in the TArray may be invalidated by adding or removing other elements to the array. 
  * Removal of elements is O(N) and invalidates the indices of subsequent elements.
