Removed unused TTuple workaround code for static analysis.

diff --git a/Engine/Source/Runtime/Core/Public/Templates/Tuple.h b/Engine/Source/Runtime/Core/Public/Templates/Tuple.h
--- a/Engine/Source/Runtime/Core/Public/Templates/Tuple.h
+++ b/Engine/Source/Runtime/Core/Public/Templates/Tuple.h
@@ -14,30 +14,6 @@
 #include <tuple>
 #include <type_traits>
 
-// This workaround exists because Visual Studio causes false positives for code like this during static analysis:
-//
-// void TestFoo(TMap<int*, int>& Map)
-// {
-//     for (const TPair<int*, int>& Pair : Map)
-//     {
-//         if (*Pair.Key == 15 && Pair.Value != 15) // warning C6011: Dereferencing NULL pointer 'Pair.Key'
-//         {
-//         }
-//     }
-// }
-//
-// This seems to be a combination of the following:
-// - Dereferencing an iterator in a loop
-// - Iterator type is not a pointer.
-// - Key is a pointer type.
-// - Key and Value are members of a base class.
-// - Dereferencing is done as part of a compound boolean expression (removing '&& Pair.Value != 15' removes the warning)
-#if (defined(_MSC_VER) && !defined(__clang__)) && USING_CODE_ANALYSIS
-	#define UE_TUPLE_STATIC_ANALYSIS_WORKAROUND 0 // disable this to see if it's still needed
-#else
-	#define UE_TUPLE_STATIC_ANALYSIS_WORKAROUND 0
-#endif
-
 class FArchive;
 
 template <typename... Types>
@@ -124,34 +100,6 @@ namespace UE::Core::Private::Tuple
 		}
 	};
 
-#if UE_TUPLE_STATIC_ANALYSIS_WORKAROUND
-	template <>
-	struct TTupleElementGetterByIndex<0, 2>
-	{
-		template <typename TupleType>
-		static inline decltype(auto) Get(TupleType&& Tuple)
-		{
-			// Brackets are important here - we want a reference type to be returned, not object type
-			decltype(auto) Result = (Forward<TupleType>(Tuple).Key);
-
-			// Keep tuple rvalue references to rvalue reference elements as rvalues, because that's how std::get() works, not how C++ struct member access works.
-			return static_cast<std::conditional_t<!std::is_reference_v<TupleType> && std::is_rvalue_reference_v<decltype(Tuple.Key)>, decltype(Tuple.Key), decltype(Result)>>(Result);
-		}
-	};
-	template <>
-	struct TTupleElementGetterByIndex<1, 2>
-	{
-		template <typename TupleType>
-		static inline decltype(auto) Get(TupleType&& Tuple)
-		{
-			// Brackets are important here - we want a reference type to be returned, not object type
-			decltype(auto) Result = (Forward<TupleType>(Tuple).Value);
-
-			// Keep tuple rvalue references to rvalue reference elements as rvalues, because that's how std::get() works, not how C++ struct member access works.
-			return static_cast<std::conditional_t<!std::is_reference_v<TupleType> && std::is_rvalue_reference_v<decltype(Tuple.Value)>, decltype(Tuple.Value), decltype(Result)>>(Result);
-		}
-	};
-#else
 	template <>
 	struct TTupleElementGetterByIndex<0, 2>
 	{
@@ -165,7 +113,6 @@ namespace UE::Core::Private::Tuple
 			return static_cast<std::conditional_t<!std::is_reference_v<TupleType> && std::is_rvalue_reference_v<decltype(Tuple.Key)>, decltype(Tuple.Key), decltype(Result)>>(Result);
 		}
 	};
-#endif
 
 	template <typename Type, uint32 TupleSize>
 	struct TTupleElementGetterByType
@@ -183,26 +130,6 @@ namespace UE::Core::Private::Tuple
 		}
 	};
 
-#if UE_TUPLE_STATIC_ANALYSIS_WORKAROUND
-	template <typename Type>
-	struct TTupleElementGetterByType<Type, 2>
-	{
-		template <typename TupleType>
-		static inline decltype(auto) Get(TupleType&& Tuple)
-		{
-			if constexpr (std::is_same_v<Type, decltype(Tuple.Key)>)
-			{
-				return TTupleElementGetterByIndex<0, 2>::Get(Forward<TupleType>(Tuple));
-			}
-			else
-			{
-				static_assert(std::is_same_v<Type, decltype(Tuple.Value)>, "This has to be true - check the call site");
-				return TTupleElementGetterByIndex<1, 2>::Get(Forward<TupleType>(Tuple));
-			}
-		}
-	};
-#endif
-
 	template <uint32 ArgCount, uint32 ArgToCompare>
 	struct FEqualityHelper
 	{
@@ -360,117 +287,6 @@ namespace UE::Core::Private::Tuple
 		}
 	};
 
-#if UE_TUPLE_STATIC_ANALYSIS_WORKAROUND
-	template <typename KeyType, typename ValueType>
-	struct TTupleBase<TIntegerSequence<uint32, 0, 1>, KeyType, ValueType>
-	{
-		KeyType   Key;
-		ValueType Value;
-
-		using DummyPairIdentifier = void;
-
-		template <typename KeyArgType, typename ValueArgType>
-		explicit TTupleBase(EForwardingConstructor, KeyArgType&& KeyArg, ValueArgType&& ValueArg)
-			: Key  (Forward<KeyArgType  >(KeyArg  ))
-			, Value(Forward<ValueArgType>(ValueArg))
-		{
-		}
-
-		template <typename TupleType>
-		explicit TTupleBase(EOtherTupleConstructor, TupleType&& Other)
-			: Key  (Forward<TupleType>(Other).Get<0>())
-			, Value(Forward<TupleType>(Other).Get<1>())
-		{
-		}
-
-		TTupleBase() = default;
-		TTupleBase(TTupleBase&& Other) = default;
-		TTupleBase(const TTupleBase& Other) = default;
-		TTupleBase& operator=(TTupleBase&& Other) = default;
-		TTupleBase& operator=(const TTupleBase& Other) = default;
-
-		template <uint32 Index> UE_FORCEINLINE_HINT decltype(auto) Get()               &  { static_assert(Index < 2, "Invalid index passed to TTuple::Get"); return TTupleElementGetterByIndex<Index, 2>::Get(static_cast<               TTupleBase& >(*this)); }
-		template <uint32 Index> UE_FORCEINLINE_HINT decltype(auto) Get() const         &  { static_assert(Index < 2, "Invalid index passed to TTuple::Get"); return TTupleElementGetterByIndex<Index, 2>::Get(static_cast<const          TTupleBase& >(*this)); }
-		template <uint32 Index> UE_FORCEINLINE_HINT decltype(auto) Get()       volatile&  { static_assert(Index < 2, "Invalid index passed to TTuple::Get"); return TTupleElementGetterByIndex<Index, 2>::Get(static_cast<      volatile TTupleBase& >(*this)); }
-		template <uint32 Index> UE_FORCEINLINE_HINT decltype(auto) Get() const volatile&  { static_assert(Index < 2, "Invalid index passed to TTuple::Get"); return TTupleElementGetterByIndex<Index, 2>::Get(static_cast<const volatile TTupleBase& >(*this)); }
-		template <uint32 Index> UE_FORCEINLINE_HINT decltype(auto) Get()               && { static_assert(Index < 2, "Invalid index passed to TTuple::Get"); return TTupleElementGetterByIndex<Index, 2>::Get(static_cast<               TTupleBase&&>(*this)); }
-		template <uint32 Index> UE_FORCEINLINE_HINT decltype(auto) Get() const         && { static_assert(Index < 2, "Invalid index passed to TTuple::Get"); return TTupleElementGetterByIndex<Index, 2>::Get(static_cast<const          TTupleBase&&>(*this)); }
-		template <uint32 Index> UE_FORCEINLINE_HINT decltype(auto) Get()       volatile&& { static_assert(Index < 2, "Invalid index passed to TTuple::Get"); return TTupleElementGetterByIndex<Index, 2>::Get(static_cast<      volatile TTupleBase&&>(*this)); }
-		template <uint32 Index> UE_FORCEINLINE_HINT decltype(auto) Get() const volatile&& { static_assert(Index < 2, "Invalid index passed to TTuple::Get"); return TTupleElementGetterByIndex<Index, 2>::Get(static_cast<const volatile TTupleBase&&>(*this)); }
-
-		template <typename T> UE_FORCEINLINE_HINT decltype(auto) Get()               &  { static_assert(TTypeCountInParameterPack_V<T, KeyType, ValueType> == 1, "Invalid type passed to TTuple::Get"); return TTupleElementGetterByType<T, 2>::Get(static_cast<               TTupleBase& >(*this)); }
-		template <typename T> UE_FORCEINLINE_HINT decltype(auto) Get() const         &  { static_assert(TTypeCountInParameterPack_V<T, KeyType, ValueType> == 1, "Invalid type passed to TTuple::Get"); return TTupleElementGetterByType<T, 2>::Get(static_cast<const          TTupleBase& >(*this)); }
-		template <typename T> UE_FORCEINLINE_HINT decltype(auto) Get()       volatile&  { static_assert(TTypeCountInParameterPack_V<T, KeyType, ValueType> == 1, "Invalid type passed to TTuple::Get"); return TTupleElementGetterByType<T, 2>::Get(static_cast<      volatile TTupleBase& >(*this)); }
-		template <typename T> UE_FORCEINLINE_HINT decltype(auto) Get() const volatile&  { static_assert(TTypeCountInParameterPack_V<T, KeyType, ValueType> == 1, "Invalid type passed to TTuple::Get"); return TTupleElementGetterByType<T, 2>::Get(static_cast<const volatile TTupleBase& >(*this)); }
-		template <typename T> UE_FORCEINLINE_HINT decltype(auto) Get()               && { static_assert(TTypeCountInParameterPack_V<T, KeyType, ValueType> == 1, "Invalid type passed to TTuple::Get"); return TTupleElementGetterByType<T, 2>::Get(static_cast<               TTupleBase&&>(*this)); }
-		template <typename T> UE_FORCEINLINE_HINT decltype(auto) Get() const         && { static_assert(TTypeCountInParameterPack_V<T, KeyType, ValueType> == 1, "Invalid type passed to TTuple::Get"); return TTupleElementGetterByType<T, 2>::Get(static_cast<const          TTupleBase&&>(*this)); }
-		template <typename T> UE_FORCEINLINE_HINT decltype(auto) Get()       volatile&& { static_assert(TTypeCountInParameterPack_V<T, KeyType, ValueType> == 1, "Invalid type passed to TTuple::Get"); return TTupleElementGetterByType<T, 2>::Get(static_cast<      volatile TTupleBase&&>(*this)); }
-		template <typename T> UE_FORCEINLINE_HINT decltype(auto) Get() const volatile&& { static_assert(TTypeCountInParameterPack_V<T, KeyType, ValueType> == 1, "Invalid type passed to TTuple::Get"); return TTupleElementGetterByType<T, 2>::Get(static_cast<const volatile TTupleBase&&>(*this)); }
-
-		template <typename FuncType, typename... ArgTypes> decltype(auto) ApplyAfter(FuncType&& Func, ArgTypes&&... Args)               &  { return ::Invoke(Func, Forward<ArgTypes>(Args)..., static_cast<               TTupleBase& >(*this).template Get<0>(), static_cast<               TTupleBase& >(*this).template Get<1>()); }
-		template <typename FuncType, typename... ArgTypes> decltype(auto) ApplyAfter(FuncType&& Func, ArgTypes&&... Args) const         &  { return ::Invoke(Func, Forward<ArgTypes>(Args)..., static_cast<const          TTupleBase& >(*this).template Get<0>(), static_cast<const          TTupleBase& >(*this).template Get<1>()); }
-		template <typename FuncType, typename... ArgTypes> decltype(auto) ApplyAfter(FuncType&& Func, ArgTypes&&... Args)       volatile&  { return ::Invoke(Func, Forward<ArgTypes>(Args)..., static_cast<      volatile TTupleBase& >(*this).template Get<0>(), static_cast<      volatile TTupleBase& >(*this).template Get<1>()); }
-		template <typename FuncType, typename... ArgTypes> decltype(auto) ApplyAfter(FuncType&& Func, ArgTypes&&... Args) const volatile&  { return ::Invoke(Func, Forward<ArgTypes>(Args)..., static_cast<const volatile TTupleBase& >(*this).template Get<0>(), static_cast<const volatile TTupleBase& >(*this).template Get<1>()); }
-		template <typename FuncType, typename... ArgTypes> decltype(auto) ApplyAfter(FuncType&& Func, ArgTypes&&... Args)               && { return ::Invoke(Func, Forward<ArgTypes>(Args)..., static_cast<               TTupleBase&&>(*this).template Get<0>(), static_cast<               TTupleBase&&>(*this).template Get<1>()); }
-		template <typename FuncType, typename... ArgTypes> decltype(auto) ApplyAfter(FuncType&& Func, ArgTypes&&... Args) const         && { return ::Invoke(Func, Forward<ArgTypes>(Args)..., static_cast<const          TTupleBase&&>(*this).template Get<0>(), static_cast<const          TTupleBase&&>(*this).template Get<1>()); }
-		template <typename FuncType, typename... ArgTypes> decltype(auto) ApplyAfter(FuncType&& Func, ArgTypes&&... Args)       volatile&& { return ::Invoke(Func, Forward<ArgTypes>(Args)..., static_cast<      volatile TTupleBase&&>(*this).template Get<0>(), static_cast<      volatile TTupleBase&&>(*this).template Get<1>()); }
-		template <typename FuncType, typename... ArgTypes> decltype(auto) ApplyAfter(FuncType&& Func, ArgTypes&&... Args) const volatile&& { return ::Invoke(Func, Forward<ArgTypes>(Args)..., static_cast<const volatile TTupleBase&&>(*this).template Get<0>(), static_cast<const volatile TTupleBase&&>(*this).template Get<1>()); }
-
-		template <typename FuncType, typename... ArgTypes> decltype(auto) ApplyBefore(FuncType&& Func, ArgTypes&&... Args)               &  { return ::Invoke(Func, static_cast<               TTupleBase& >(*this).template Get<0>(), static_cast<               TTupleBase& >(*this).template Get<1>(), Forward<ArgTypes>(Args)...); }
-		template <typename FuncType, typename... ArgTypes> decltype(auto) ApplyBefore(FuncType&& Func, ArgTypes&&... Args) const         &  { return ::Invoke(Func, static_cast<const          TTupleBase& >(*this).template Get<0>(), static_cast<const          TTupleBase& >(*this).template Get<1>(), Forward<ArgTypes>(Args)...); }
-		template <typename FuncType, typename... ArgTypes> decltype(auto) ApplyBefore(FuncType&& Func, ArgTypes&&... Args)       volatile&  { return ::Invoke(Func, static_cast<      volatile TTupleBase& >(*this).template Get<0>(), static_cast<      volatile TTupleBase& >(*this).template Get<1>(), Forward<ArgTypes>(Args)...); }
-		template <typename FuncType, typename... ArgTypes> decltype(auto) ApplyBefore(FuncType&& Func, ArgTypes&&... Args) const volatile&  { return ::Invoke(Func, static_cast<const volatile TTupleBase& >(*this).template Get<0>(), static_cast<const volatile TTupleBase& >(*this).template Get<1>(), Forward<ArgTypes>(Args)...); }
-		template <typename FuncType, typename... ArgTypes> decltype(auto) ApplyBefore(FuncType&& Func, ArgTypes&&... Args)               && { return ::Invoke(Func, static_cast<               TTupleBase&&>(*this).template Get<0>(), static_cast<               TTupleBase&&>(*this).template Get<1>(), Forward<ArgTypes>(Args)...); }
-		template <typename FuncType, typename... ArgTypes> decltype(auto) ApplyBefore(FuncType&& Func, ArgTypes&&... Args) const         && { return ::Invoke(Func, static_cast<const          TTupleBase&&>(*this).template Get<0>(), static_cast<const          TTupleBase&&>(*this).template Get<1>(), Forward<ArgTypes>(Args)...); }
-		template <typename FuncType, typename... ArgTypes> decltype(auto) ApplyBefore(FuncType&& Func, ArgTypes&&... Args)       volatile&& { return ::Invoke(Func, static_cast<      volatile TTupleBase&&>(*this).template Get<0>(), static_cast<      volatile TTupleBase&&>(*this).template Get<1>(), Forward<ArgTypes>(Args)...); }
-		template <typename FuncType, typename... ArgTypes> decltype(auto) ApplyBefore(FuncType&& Func, ArgTypes&&... Args) const volatile&& { return ::Invoke(Func, static_cast<const volatile TTupleBase&&>(*this).template Get<0>(), static_cast<const volatile TTupleBase&&>(*this).template Get<1>(), Forward<ArgTypes>(Args)...); }
-
-		inline static void Serialize(FArchive& Ar, TTupleBase& Tuple)
-		{
-			Ar << Tuple.Key;
-			Ar << Tuple.Value;
-		}
-
-		inline static void SerializeStructured(FStructuredArchive::FSlot Slot, TTupleBase& Tuple)
-		{
-			FStructuredArchive::FStream Stream = Slot.EnterStream();
-			Stream.EnterElement() << Tuple.Key;
-			Stream.EnterElement() << Tuple.Value;
-		}
-
-		UE_FORCEINLINE_HINT friend bool operator==(const TTupleBase& Lhs, const TTupleBase& Rhs)
-		{
-			// This could be implemented with a fold expression when our compilers support it
-			return FEqualityHelper<2, 0>::Compare(Lhs, Rhs);
-		}
-
-		UE_FORCEINLINE_HINT friend bool operator!=(const TTupleBase& Lhs, const TTupleBase& Rhs)
-		{
-			return !(Lhs == Rhs);
-		}
-
-		UE_FORCEINLINE_HINT friend bool operator<(const TTupleBase& Lhs, const TTupleBase& Rhs)
-		{
-			return TLessThanHelper<2>::Do(Lhs, Rhs);
-		}
-
-		UE_FORCEINLINE_HINT friend bool operator<=(const TTupleBase& Lhs, const TTupleBase& Rhs)
-		{
-			return !(Rhs < Lhs);
-		}
-
-		UE_FORCEINLINE_HINT friend bool operator>(const TTupleBase& Lhs, const TTupleBase& Rhs)
-		{
-			return Rhs < Lhs;
-		}
-
-		UE_FORCEINLINE_HINT friend bool operator>=(const TTupleBase& Lhs, const TTupleBase& Rhs)
-		{
-			return !(Lhs < Rhs);
-		}
-	};
-#endif
-
 	template <typename LhsType, typename RhsType, uint32... Indices>
 	static void Assign(LhsType& Lhs, RhsType&& Rhs, TIntegerSequence<uint32, Indices...>)
 	{
@@ -553,17 +369,6 @@ namespace UE::Core::Private::Tuple
 		static constexpr uint32 Value = sizeof(Resolve(DeclVal<TTuple<TupleTypes...>*>())) - 1;
 	};
 
-#if UE_TUPLE_STATIC_ANALYSIS_WORKAROUND
-	template <typename Type, typename KeyType, typename ValueType>
-	struct TCVTupleIndex<Type, const volatile TTuple<KeyType, ValueType>>
-	{
-		static_assert(TTypeCountInParameterPack_V<Type, KeyType, ValueType> >= 1, "TTupleIndex instantiated with a tuple which does not contain the type");
-		static_assert(TTypeCountInParameterPack_V<Type, KeyType, ValueType> <= 1, "TTupleIndex instantiated with a tuple which contains multiple instances of the type");
-
-		static constexpr uint32 Value = std::is_same_v<Type, ValueType> ? 1 : 0;
-	};
-#endif
-
 	template <uint32 Index, typename TupleType>
 	struct TCVTupleElement
 	{
@@ -589,16 +394,6 @@ namespace UE::Core::Private::Tuple
 #endif
 	};
 
-#if UE_TUPLE_STATIC_ANALYSIS_WORKAROUND
-	template <uint32 Index, typename KeyType, typename ValueType>
-	struct TCVTupleElement<Index, const volatile TTuple<KeyType, ValueType>>
-	{
-		static_assert(Index < 2, "TTupleElement instantiated with an invalid index");
-
-		using Type = std::conditional_t<Index == 0, KeyType, ValueType>;
-	};
-#endif
-
 	template <uint32 ArgToCombine, uint32 ArgCount>
 	struct TGetTupleHashHelper
 	{
