Optimized TSet::UnhashElements

diff --git a/Engine/Source/Runtime/Core/Public/Containers/Set.h b/Engine/Source/Runtime/Core/Public/Containers/Set.h
--- a/Engine/Source/Runtime/Core/Public/Containers/Set.h
+++ b/Engine/Source/Runtime/Core/Public/Containers/Set.h
@@ -1478,11 +1478,8 @@ private:
 		}
 		else
 		{
-			/** Reset the whole hash to invalid */
-			for (int32 I = 0; I < HashSize; ++I)
-			{
-				HashPtr[I] = FSetElementId();
-			}
+			static_assert(FSetElementId().AsInteger() == -1);
+			FMemory::Memset(HashPtr, 0xFF, HashSize * sizeof(FSetElementId));
 		}
 	}
 
diff --git a/Engine/Source/Runtime/Core/Public/Containers/SetUtilities.h b/Engine/Source/Runtime/Core/Public/Containers/SetUtilities.h
--- a/Engine/Source/Runtime/Core/Public/Containers/SetUtilities.h
+++ b/Engine/Source/Runtime/Core/Public/Containers/SetUtilities.h
@@ -89,7 +89,7 @@ class FSetElementId
 {
 public:
 	/** Default constructor. */
-	UE_FORCEINLINE_HINT FSetElementId() = default;
+	constexpr UE_FORCEINLINE_HINT FSetElementId() = default;
 
 	/** @return a boolean value representing whether the id is NULL. */
 	[[nodiscard]] UE_FORCEINLINE_HINT bool IsValidId() const
@@ -107,7 +107,7 @@ public:
 		return A.Index != B.Index;
 	}
 
-	[[nodiscard]] UE_FORCEINLINE_HINT int32 AsInteger() const
+	[[nodiscard]] constexpr UE_FORCEINLINE_HINT int32 AsInteger() const
 	{
 		return Index;
 	}
