Replaced FString and FName constructions with TEXT() and empty string literals.

diff --git a/Engine/Source/Runtime/Core/Private/Android/AndroidJavaMessageBox.cpp b/Engine/Source/Runtime/Core/Private/Android/AndroidJavaMessageBox.cpp
--- a/Engine/Source/Runtime/Core/Private/Android/AndroidJavaMessageBox.cpp
+++ b/Engine/Source/Runtime/Core/Private/Android/AndroidJavaMessageBox.cpp
@@ -47,7 +47,7 @@ FName FJavaAndroidMessageBox::GetClassName()
 	}
 	else
 	{
-		return FName("");
+		return FName();
 	}
 }
 
diff --git a/Engine/Source/Runtime/Core/Private/Android/AndroidPlatformFile.cpp b/Engine/Source/Runtime/Core/Private/Android/AndroidPlatformFile.cpp
--- a/Engine/Source/Runtime/Core/Private/Android/AndroidPlatformFile.cpp
+++ b/Engine/Source/Runtime/Core/Private/Android/AndroidPlatformFile.cpp
@@ -895,7 +895,7 @@ public:
 		}
 		else
 		{
-			GLastOBBComment = FString("");
+			GLastOBBComment = FString();
 		}
 
 
diff --git a/Engine/Source/Runtime/Core/Private/Android/AndroidPlatformMisc.cpp b/Engine/Source/Runtime/Core/Private/Android/AndroidPlatformMisc.cpp
--- a/Engine/Source/Runtime/Core/Private/Android/AndroidPlatformMisc.cpp
+++ b/Engine/Source/Runtime/Core/Private/Android/AndroidPlatformMisc.cpp
@@ -923,7 +923,7 @@ bool FAndroidMisc::UseRenderThread()
 	// tested with lg optimus 2x and motorola xoom
 	// come back and revisit this later
 	// https://code.google.com/p/android/issues/detail?id=32636
-	if (FAndroidMisc::GetGPUFamily() == FString(TEXT("NVIDIA Tegra")) && FPlatformMisc::NumberOfCores() <= 2 && FAndroidMisc::GetGLVersion().StartsWith(TEXT("OpenGL ES 2.")))
+	if (FAndroidMisc::GetGPUFamily() == FString("NVIDIA Tegra") && FPlatformMisc::NumberOfCores() <= 2 && FAndroidMisc::GetGLVersion().StartsWith(TEXT("OpenGL ES 2.")))
 	{
 		return false;
 	}
@@ -935,13 +935,13 @@ bool FAndroidMisc::UseRenderThread()
 	}
 
 	// there is an issue with presenting the buffer on kindle fire (1st gen) with multiple threads using opengl
-	if (FAndroidMisc::GetDeviceModel() == FString(TEXT("Kindle Fire")))
+	if (FAndroidMisc::GetDeviceModel() == FString("Kindle Fire"))
 	{
 		return false;
 	}
 
 	// there is an issue with swapbuffer ordering on startup on samsung s3 mini with multiple threads using opengl
-	if (FAndroidMisc::GetDeviceModel() == FString(TEXT("GT-I8190L")))
+	if (FAndroidMisc::GetDeviceModel() == FString("GT-I8190L"))
 	{
 		return false;
 	}
@@ -1758,7 +1758,7 @@ TArray<uint8> FAndroidMisc::GetSystemFontBytes()
 {
 #if USE_ANDROID_FILE
 	TArray<uint8> FontBytes;
-	static FString FullFontPath = GFontPathBase + FString(TEXT("DroidSans.ttf"));
+	static FString FullFontPath = GFontPathBase + FString("DroidSans.ttf");
 	FFileHelper::LoadFileToArray(FontBytes, *FullFontPath);
 	return FontBytes;
 #else 
@@ -2671,7 +2671,7 @@ const TCHAR* FAndroidMisc::GamePersistentDownloadDir()
 
 FString FAndroidMisc::GetLoginId()
 {
-	static FString LoginId = TEXT("");
+	static FString LoginId;
 
 	// Return already loaded or generated Id
 	if (!LoginId.IsEmpty())
diff --git a/Engine/Source/Runtime/Core/Private/Apple/ApplePlatformStackWalk.cpp b/Engine/Source/Runtime/Core/Private/Apple/ApplePlatformStackWalk.cpp
--- a/Engine/Source/Runtime/Core/Private/Apple/ApplePlatformStackWalk.cpp
+++ b/Engine/Source/Runtime/Core/Private/Apple/ApplePlatformStackWalk.cpp
@@ -395,7 +395,7 @@ uint32 FApplePlatformStackWalk::CaptureThreadStackBackTrace(uint64 ThreadId, uin
 
 void CreateExceptionInfoString(int32 Signal, struct __siginfo* Info)
 {
-	FString ErrorString = TEXT("Unhandled Exception: ");
+	FString ErrorString = "Unhandled Exception: ";
 	
 #define HANDLE_CASE(a,b) case a: ErrorString += TEXT(#a #b); break;
 	
diff --git a/Engine/Source/Runtime/Core/Private/Compression/OodleDataCompressionFormat.cpp b/Engine/Source/Runtime/Core/Private/Compression/OodleDataCompressionFormat.cpp
--- a/Engine/Source/Runtime/Core/Private/Compression/OodleDataCompressionFormat.cpp
+++ b/Engine/Source/Runtime/Core/Private/Compression/OodleDataCompressionFormat.cpp
@@ -428,8 +428,8 @@ ICompressionFormat * CreateOodleDataCompressionFormat()
 		};
 
 		// override from command line :
-		FString Method = "";
-		FString Level = "";
+		FString Method;
+		FString Level;
 
 		// let commandline override
 		//FParse::Value does not change output fields if they are not found
@@ -464,7 +464,7 @@ ICompressionFormat * CreateOodleDataCompressionFormat()
 
 			// NOTE : we get OodleCompressDLL from Engine.ini for the platform we are running *on* not the platform we are packaging *for*
 			// the ProjectPackaging settings we get on the command line from Game.ini come from the *target* platform
-			FString OodleDLL = TEXT("");
+			FString OodleDLL;
 
 			// do NOT read OodleCompressDLL from GConfig here; it is passed on command line from CopyBuildToStagingDirectory.Automation.cs : 
 			FParse::Value(FCommandLine::Get(), TEXT("OodleCompressDLL="), OodleDLL);
diff --git a/Engine/Source/Runtime/Core/Private/Containers/Algo/AlgosTest.cpp b/Engine/Source/Runtime/Core/Private/Containers/Algo/AlgosTest.cpp
--- a/Engine/Source/Runtime/Core/Private/Containers/Algo/AlgosTest.cpp
+++ b/Engine/Source/Runtime/Core/Private/Containers/Algo/AlgosTest.cpp
@@ -720,9 +720,9 @@ public:
 			};
 
 			FString Empty;
-			FString DataA = TEXT("HeLlO wOrLd");
-			FString DataB = TEXT("HELLO GOODBYE");
-			FString DataC = TEXT("hello");
+			FString DataA = "HeLlO wOrLd";
+			FString DataB = "HELLO GOODBYE";
+			FString DataC = "hello";
 
 			// Test empty ranges with custom equality
 			check(Algo::Mismatch(Empty, Empty, CompareCaseInsensitive) == 0);
diff --git a/Engine/Source/Runtime/Core/Private/Containers/ContainersTest.cpp b/Engine/Source/Runtime/Core/Private/Containers/ContainersTest.cpp
--- a/Engine/Source/Runtime/Core/Private/Containers/ContainersTest.cpp
+++ b/Engine/Source/Runtime/Core/Private/Containers/ContainersTest.cpp
@@ -756,7 +756,7 @@ namespace
 		// Add elements in stably randomized order
 		for (int32 i = 0; i < NumObjects; i++)
 		{
-			Map.Add(KeyArray[i], FString(TEXT("New Value")));
+			Map.Add(KeyArray[i], FString("New Value"));
 		}
 
 		// Now do searches
diff --git a/Engine/Source/Runtime/Core/Private/Containers/StringTest.cpp b/Engine/Source/Runtime/Core/Private/Containers/StringTest.cpp
--- a/Engine/Source/Runtime/Core/Private/Containers/StringTest.cpp
+++ b/Engine/Source/Runtime/Core/Private/Containers/StringTest.cpp
@@ -245,7 +245,7 @@ IMPLEMENT_SIMPLE_AUTOMATION_TEST(FStringParseIntoArrayTest, "System.Core.String.
 bool FStringParseIntoArrayTest::RunTest(const FString& Parameters)
 {
 	FString Empty;
-	FString EmptyButAllocated = TEXT("This|is|a|long|string");
+	FString EmptyButAllocated = "This|is|a|long|string";
 	EmptyButAllocated.GetCharArray().Reset();
 
 	const TCHAR* Delim[] = { TEXT("|") };
@@ -254,30 +254,30 @@ bool FStringParseIntoArrayTest::RunTest(const FString& Parameters)
 	{
 		TArray<FString> Parsed;
 
-		Parsed.Reset(); FString(TEXT("|"))  .ParseIntoArray(Parsed, Delim[0], /*cull = */true);  TestTrue(TEXT("parse a single delimiter character with a single delimiter, culling empties"),            Parsed == TArray<FString>{});
-		Parsed.Reset(); FString(TEXT("|"))  .ParseIntoArray(Parsed, Delim, 1, /*cull = */true);  TestTrue(TEXT("parse a single delimiter character with an array of 1 delimiter, culling empties"),       Parsed == TArray<FString>{});
-		Parsed.Reset(); FString{}           .ParseIntoArray(Parsed, Delim[0], /*cull = */true);  TestTrue(TEXT("parse an unallocated string with a single delimiter, culling empties"),                   Parsed == TArray<FString>{});
-		Parsed.Reset(); FString{}           .ParseIntoArray(Parsed, Delim, 1, /*cull = */true);  TestTrue(TEXT("parse an unallocated string with an array of 1 delimiter, culling empties"),              Parsed == TArray<FString>{});
-		Parsed.Reset(); EmptyButAllocated   .ParseIntoArray(Parsed, Delim[0], /*cull = */true);  TestTrue(TEXT("parse an allocated empty string with a single delimiter, culling empties"),               Parsed == TArray<FString>{});
-		Parsed.Reset(); EmptyButAllocated   .ParseIntoArray(Parsed, Delim, 1, /*cull = */true);  TestTrue(TEXT("parse an allocated empty string with an array of 1 delimiter, culling empties"),          Parsed == TArray<FString>{});
-		Parsed.Reset(); FString(TEXT("a|")) .ParseIntoArray(Parsed, Delim[0], /*cull = */true);  TestTrue(TEXT("parse a string with a trailing delimiter with a single delimiter, culling empties"),      Parsed == TArray<FString>{TEXT("a")});
-		Parsed.Reset(); FString(TEXT("a|")) .ParseIntoArray(Parsed, Delim, 1, /*cull = */true);  TestTrue(TEXT("parse a string with a trailing delimiter with an array of 1 delimiter, culling empties"), Parsed == TArray<FString>{TEXT("a")});
-		Parsed.Reset(); FString(TEXT("|b")) .ParseIntoArray(Parsed, Delim[0], /*cull = */true);  TestTrue(TEXT("parse a string with a leading delimiter with a single delimiter, culling empties"),       Parsed == TArray<FString>{TEXT("b")});
-		Parsed.Reset(); FString(TEXT("|b")) .ParseIntoArray(Parsed, Delim, 1, /*cull = */true);  TestTrue(TEXT("parse a string with a leading delimiter with an array of 1 delimiter, culling empties"),  Parsed == TArray<FString>{TEXT("b")});
-		Parsed.Reset(); FString(TEXT("a|b")).ParseIntoArray(Parsed, Delim[0], /*cull = */true);  TestTrue(TEXT("parse a string with a middle delimiter with a single delimiter, culling empties"),        Parsed == TArray<FString>{TEXT("a"), TEXT("b")});
-		Parsed.Reset(); FString(TEXT("a|b")).ParseIntoArray(Parsed, Delim, 1, /*cull = */true);  TestTrue(TEXT("parse a string with a middle delimiter with an array of 1 delimiter, culling empties"),   Parsed == TArray<FString>{TEXT("a"), TEXT("b")});
-		Parsed.Reset(); FString(TEXT("|"))  .ParseIntoArray(Parsed, Delim[0], /*cull = */false); TestTrue(TEXT("parse a single delimiter character with a single delimiter, keeping empties"),            Parsed == TArray<FString>{TEXT(""), TEXT("")});
-		Parsed.Reset(); FString(TEXT("|"))  .ParseIntoArray(Parsed, Delim, 1, /*cull = */false); TestTrue(TEXT("parse a single delimiter character with an array of 1 delimiter, keeping empties"),       Parsed == TArray<FString>{TEXT(""), TEXT("")});
-		Parsed.Reset(); FString{}           .ParseIntoArray(Parsed, Delim[0], /*cull = */false); TestTrue(TEXT("parse an unallocated string with a single delimiter, keeping empties"),                   Parsed == TArray<FString>{TEXT("")});
-		Parsed.Reset(); FString{}           .ParseIntoArray(Parsed, Delim, 1, /*cull = */false); TestTrue(TEXT("parse an unallocated string with an array of 1 delimiter, keeping empties"),              Parsed == TArray<FString>{TEXT("")});
-		Parsed.Reset(); EmptyButAllocated   .ParseIntoArray(Parsed, Delim[0], /*cull = */false); TestTrue(TEXT("parse an allocated empty string with a single delimiter, keeping empties"),               Parsed == TArray<FString>{TEXT("")});
-		Parsed.Reset(); EmptyButAllocated   .ParseIntoArray(Parsed, Delim, 1, /*cull = */false); TestTrue(TEXT("parse an allocated empty string with an array of 1 delimiter, keeping empties"),          Parsed == TArray<FString>{TEXT("")});
-		Parsed.Reset(); FString(TEXT("a|")) .ParseIntoArray(Parsed, Delim[0], /*cull = */false); TestTrue(TEXT("parse a string with a trailing delimiter with a single delimiter, keeping empties"),      Parsed == TArray<FString>{TEXT("a"), TEXT("")});
-		Parsed.Reset(); FString(TEXT("a|")) .ParseIntoArray(Parsed, Delim, 1, /*cull = */false); TestTrue(TEXT("parse a string with a trailing delimiter with an array of 1 delimiter, keeping empties"), Parsed == TArray<FString>{TEXT("a"), TEXT("")});
-		Parsed.Reset(); FString(TEXT("|b")) .ParseIntoArray(Parsed, Delim[0], /*cull = */false); TestTrue(TEXT("parse a string with a leading delimiter with a single delimiter, keeping empties"),       Parsed == TArray<FString>{TEXT(""), TEXT("b")});
-		Parsed.Reset(); FString(TEXT("|b")) .ParseIntoArray(Parsed, Delim, 1, /*cull = */false); TestTrue(TEXT("parse a string with a leading delimiter with an array of 1 delimiter, keeping empties"),  Parsed == TArray<FString>{TEXT(""), TEXT("b")});
-		Parsed.Reset(); FString(TEXT("a|b")).ParseIntoArray(Parsed, Delim[0], /*cull = */false); TestTrue(TEXT("parse a string with a middle delimiter with a single delimiter, keeping empties"),        Parsed == TArray<FString>{TEXT("a"), TEXT("b")});
-		Parsed.Reset(); FString(TEXT("a|b")).ParseIntoArray(Parsed, Delim, 1, /*cull = */false); TestTrue(TEXT("parse a string with a middle delimiter with an array of 1 delimiter, keeping empties"),   Parsed == TArray<FString>{TEXT("a"), TEXT("b")});
+		Parsed.Reset(); FString("|")     .ParseIntoArray(Parsed, Delim[0], /*cull = */true);  TestTrue(TEXT("parse a single delimiter character with a single delimiter, culling empties"),            Parsed == TArray<FString>{});
+		Parsed.Reset(); FString("|")     .ParseIntoArray(Parsed, Delim, 1, /*cull = */true);  TestTrue(TEXT("parse a single delimiter character with an array of 1 delimiter, culling empties"),       Parsed == TArray<FString>{});
+		Parsed.Reset(); FString{}        .ParseIntoArray(Parsed, Delim[0], /*cull = */true);  TestTrue(TEXT("parse an unallocated string with a single delimiter, culling empties"),                   Parsed == TArray<FString>{});
+		Parsed.Reset(); FString{}        .ParseIntoArray(Parsed, Delim, 1, /*cull = */true);  TestTrue(TEXT("parse an unallocated string with an array of 1 delimiter, culling empties"),              Parsed == TArray<FString>{});
+		Parsed.Reset(); EmptyButAllocated.ParseIntoArray(Parsed, Delim[0], /*cull = */true);  TestTrue(TEXT("parse an allocated empty string with a single delimiter, culling empties"),               Parsed == TArray<FString>{});
+		Parsed.Reset(); EmptyButAllocated.ParseIntoArray(Parsed, Delim, 1, /*cull = */true);  TestTrue(TEXT("parse an allocated empty string with an array of 1 delimiter, culling empties"),          Parsed == TArray<FString>{});
+		Parsed.Reset(); FString("a|")    .ParseIntoArray(Parsed, Delim[0], /*cull = */true);  TestTrue(TEXT("parse a string with a trailing delimiter with a single delimiter, culling empties"),      Parsed == TArray<FString>{TEXT("a")});
+		Parsed.Reset(); FString("a|")    .ParseIntoArray(Parsed, Delim, 1, /*cull = */true);  TestTrue(TEXT("parse a string with a trailing delimiter with an array of 1 delimiter, culling empties"), Parsed == TArray<FString>{TEXT("a")});
+		Parsed.Reset(); FString("|b")    .ParseIntoArray(Parsed, Delim[0], /*cull = */true);  TestTrue(TEXT("parse a string with a leading delimiter with a single delimiter, culling empties"),       Parsed == TArray<FString>{TEXT("b")});
+		Parsed.Reset(); FString("|b")    .ParseIntoArray(Parsed, Delim, 1, /*cull = */true);  TestTrue(TEXT("parse a string with a leading delimiter with an array of 1 delimiter, culling empties"),  Parsed == TArray<FString>{TEXT("b")});
+		Parsed.Reset(); FString("a|b")   .ParseIntoArray(Parsed, Delim[0], /*cull = */true);  TestTrue(TEXT("parse a string with a middle delimiter with a single delimiter, culling empties"),        Parsed == TArray<FString>{TEXT("a"), TEXT("b")});
+		Parsed.Reset(); FString("a|b")   .ParseIntoArray(Parsed, Delim, 1, /*cull = */true);  TestTrue(TEXT("parse a string with a middle delimiter with an array of 1 delimiter, culling empties"),   Parsed == TArray<FString>{TEXT("a"), TEXT("b")});
+		Parsed.Reset(); FString("|")     .ParseIntoArray(Parsed, Delim[0], /*cull = */false); TestTrue(TEXT("parse a single delimiter character with a single delimiter, keeping empties"),            Parsed == TArray<FString>{TEXT(""), TEXT("")});
+		Parsed.Reset(); FString("|")     .ParseIntoArray(Parsed, Delim, 1, /*cull = */false); TestTrue(TEXT("parse a single delimiter character with an array of 1 delimiter, keeping empties"),       Parsed == TArray<FString>{TEXT(""), TEXT("")});
+		Parsed.Reset(); FString{}        .ParseIntoArray(Parsed, Delim[0], /*cull = */false); TestTrue(TEXT("parse an unallocated string with a single delimiter, keeping empties"),                   Parsed == TArray<FString>{TEXT("")});
+		Parsed.Reset(); FString{}        .ParseIntoArray(Parsed, Delim, 1, /*cull = */false); TestTrue(TEXT("parse an unallocated string with an array of 1 delimiter, keeping empties"),              Parsed == TArray<FString>{TEXT("")});
+		Parsed.Reset(); EmptyButAllocated.ParseIntoArray(Parsed, Delim[0], /*cull = */false); TestTrue(TEXT("parse an allocated empty string with a single delimiter, keeping empties"),               Parsed == TArray<FString>{TEXT("")});
+		Parsed.Reset(); EmptyButAllocated.ParseIntoArray(Parsed, Delim, 1, /*cull = */false); TestTrue(TEXT("parse an allocated empty string with an array of 1 delimiter, keeping empties"),          Parsed == TArray<FString>{TEXT("")});
+		Parsed.Reset(); FString("a|")    .ParseIntoArray(Parsed, Delim[0], /*cull = */false); TestTrue(TEXT("parse a string with a trailing delimiter with a single delimiter, keeping empties"),      Parsed == TArray<FString>{TEXT("a"), TEXT("")});
+		Parsed.Reset(); FString("a|")    .ParseIntoArray(Parsed, Delim, 1, /*cull = */false); TestTrue(TEXT("parse a string with a trailing delimiter with an array of 1 delimiter, keeping empties"), Parsed == TArray<FString>{TEXT("a"), TEXT("")});
+		Parsed.Reset(); FString("|b")    .ParseIntoArray(Parsed, Delim[0], /*cull = */false); TestTrue(TEXT("parse a string with a leading delimiter with a single delimiter, keeping empties"),       Parsed == TArray<FString>{TEXT(""), TEXT("b")});
+		Parsed.Reset(); FString("|b")    .ParseIntoArray(Parsed, Delim, 1, /*cull = */false); TestTrue(TEXT("parse a string with a leading delimiter with an array of 1 delimiter, keeping empties"),  Parsed == TArray<FString>{TEXT(""), TEXT("b")});
+		Parsed.Reset(); FString("a|b")   .ParseIntoArray(Parsed, Delim[0], /*cull = */false); TestTrue(TEXT("parse a string with a middle delimiter with a single delimiter, keeping empties"),        Parsed == TArray<FString>{TEXT("a"), TEXT("b")});
+		Parsed.Reset(); FString("a|b")   .ParseIntoArray(Parsed, Delim, 1, /*cull = */false); TestTrue(TEXT("parse a string with a middle delimiter with an array of 1 delimiter, keeping empties"),   Parsed == TArray<FString>{TEXT("a"), TEXT("b")});
 	}
 
 	return true;
@@ -294,7 +294,7 @@ bool FStringSubstringTest::RunTest(const FString& Parameters)
 		}
 	};
 
-	const FString TestString(TEXT("0123456789"));
+	const FString TestString("0123456789");
 
 #define SUBSTRINGTEST(TestName, ExpectedResult, Operation, ...) \
 	FString TestName = TestString.Operation(__VA_ARGS__); \
@@ -358,7 +358,7 @@ bool FStringFromStringViewTest::RunTest(const FString& Parameters)
 		const ANSICHAR* AnsiLiteral = "Literal";
 		TestEqual(TEXT("String(StringView)"), FString(FStringView(Literal)), Literal);
 		TestEqual(TEXT("String(AnsiStringView)"), FString(FAnsiStringView(AnsiLiteral)), Literal);
-		TestEqual(TEXT("String = StringView"), FString(TEXT("Temp")) = FStringView(Literal), Literal);
+		TestEqual(TEXT("String = StringView"), FString("Temp") = FStringView(Literal), Literal);
 
 		FStringView EmptyStringView;
 		FString EmptyString(EmptyStringView);
@@ -373,19 +373,19 @@ bool FStringFromStringViewTest::RunTest(const FString& Parameters)
 
 	// Verify assignment from a view of itself.
 	{
-		FString AssignEntireString(TEXT("AssignEntireString"));
+		FString AssignEntireString("AssignEntireString");
 		AssignEntireString = FStringView(AssignEntireString);
 		TestEqual(TEXT("String = StringView(String)"), AssignEntireString, TEXT("AssignEntireString"));
 
-		FString AssignStartOfString(TEXT("AssignStartOfString"));
+		FString AssignStartOfString("AssignStartOfString");
 		AssignStartOfString = FStringView(AssignStartOfString).Left(11);
 		TestEqual(TEXT("String = StringView(String).Left"), AssignStartOfString, TEXT("AssignStart"));
 
-		FString AssignEndOfString(TEXT("AssignEndOfString"));
+		FString AssignEndOfString("AssignEndOfString");
 		AssignEndOfString = FStringView(AssignEndOfString).Right(11);
 		TestEqual(TEXT("String = StringView(String).Right"), AssignEndOfString, TEXT("EndOfString"));
 
-		FString AssignMiddleOfString(TEXT("AssignMiddleOfString"));
+		FString AssignMiddleOfString("AssignMiddleOfString");
 		AssignMiddleOfString = FStringView(AssignMiddleOfString).Mid(6, 6);
 		TestEqual(TEXT("String = StringView(String).Mid"), AssignMiddleOfString, TEXT("Middle"));
 	}
@@ -393,28 +393,28 @@ bool FStringFromStringViewTest::RunTest(const FString& Parameters)
 	// Verify operators taking string views and character arrays
 	{
 		FStringView RhsStringView = FStringView(TEXT("RhsNotSZ"), 3);
-		FString MovePlusSVResult = FString(TEXT("Lhs")) + RhsStringView;
+		FString MovePlusSVResult = FString("Lhs") + RhsStringView;
 		TestEqual(TEXT("Move String + StringView"), MovePlusSVResult, TEXT("LhsRhs"));
 
-		FString CopyLhs(TEXT("Lhs"));
+		FString CopyLhs("Lhs");
 		FString CopyPlusSVResult = CopyLhs + RhsStringView;
 		TestEqual(TEXT("Copy String + StringView"), CopyPlusSVResult, TEXT("LhsRhs"));
 
-		FString MovePlusTCHARsResult = FString(TEXT("Lhs")) + TEXT("Rhs");
+		FString MovePlusTCHARsResult = FString("Lhs") + TEXT("Rhs");
 		TestEqual(TEXT("Move String + TCHAR*"), MovePlusTCHARsResult, TEXT("LhsRhs"));
 
 		FString CopyPlusTCHARsResult = CopyLhs + TEXT("Rhs");
 		TestEqual(TEXT("Copy String + TCHAR*"), CopyPlusTCHARsResult, TEXT("LhsRhs"));
 
 		FStringView LhsStringView = FStringView(TEXT("LhsNotSZ"), 3);
-		FString SVPlusMoveResult = LhsStringView + FString(TEXT("Rhs"));
+		FString SVPlusMoveResult = LhsStringView + FString("Rhs");
 		TestEqual(TEXT("StringView + Move String"), SVPlusMoveResult, TEXT("LhsRhs"));
 
-		FString CopyRhs(TEXT("Rhs"));
+		FString CopyRhs("Rhs");
 		FString SVPlusCopyResult = LhsStringView + CopyRhs;
 		TestEqual(TEXT("StringView + Copy String"), SVPlusCopyResult, TEXT("LhsRhs"));
 
-		FString TCHARsPlusMoveResult = TEXT("Lhs") + FString(TEXT("Rhs"));
+		FString TCHARsPlusMoveResult = TEXT("Lhs") + FString("Rhs");
 		TestEqual(TEXT("TCHAR* + Move String"), TCHARsPlusMoveResult, TEXT("LhsRhs"));
 
 		FString TCHARsPlusCopyResult = TEXT("Lhs") + CopyRhs;
@@ -593,7 +593,7 @@ bool FStringPathConcatCompoundOperatorTest::RunTest(const FString& Parameters)
 
 	// Test empty path /= empty file
 	TEST_EMPTYPATH_EMPTYFILE("NullString", FString());
-	TEST_EMPTYPATH_EMPTYFILE("EmptyString", FString(TEXT("")));
+	TEST_EMPTYPATH_EMPTYFILE("EmptyString", FString());
 	TEST_EMPTYPATH_EMPTYFILE("EmptyAnsiLiteralString", "");
 	TEST_EMPTYPATH_EMPTYFILE("EmptyLiteralString", TEXT(""));
 	TEST_EMPTYPATH_EMPTYFILE("NullStringView", FStringView());
@@ -602,7 +602,7 @@ bool FStringPathConcatCompoundOperatorTest::RunTest(const FString& Parameters)
 
 	// Test valid path /= empty file
 	TEST_VALIDPATH_EMPTYFILE("NullString", FString());
-	TEST_VALIDPATH_EMPTYFILE("EmptyString", FString(TEXT("")));
+	TEST_VALIDPATH_EMPTYFILE("EmptyString", FString());
 	TEST_VALIDPATH_EMPTYFILE("EmptyAnsiLiteralString", "");
 	TEST_VALIDPATH_EMPTYFILE("EmptyLiteralString", TEXT(""));
 	TEST_VALIDPATH_EMPTYFILE("NullStringView", FStringView());
diff --git a/Engine/Source/Runtime/Core/Private/GenericPlatform/GenericPlatformCrashContext.cpp b/Engine/Source/Runtime/Core/Private/GenericPlatform/GenericPlatformCrashContext.cpp
--- a/Engine/Source/Runtime/Core/Private/GenericPlatform/GenericPlatformCrashContext.cpp
+++ b/Engine/Source/Runtime/Core/Private/GenericPlatform/GenericPlatformCrashContext.cpp
@@ -472,7 +472,7 @@ void FGenericCrashContext::Initialize()
 	{
 		if (InParams.bSetGameNameSuffix)
 		{
-			NCached::Set(NCached::Session.GameName, *(FString(TEXT("UE-")) + FApp::GetProjectName() + InParams.GameNameSuffix));
+			NCached::Set(NCached::Session.GameName, *(FString("UE-") + FApp::GetProjectName() + InParams.GameNameSuffix));
 		}
 		if (InParams.SendUnattendedBugReports.IsSet())
 		{
diff --git a/Engine/Source/Runtime/Core/Private/GenericPlatform/GenericPlatformFileTest.cpp b/Engine/Source/Runtime/Core/Private/GenericPlatform/GenericPlatformFileTest.cpp
--- a/Engine/Source/Runtime/Core/Private/GenericPlatform/GenericPlatformFileTest.cpp
+++ b/Engine/Source/Runtime/Core/Private/GenericPlatform/GenericPlatformFileTest.cpp
@@ -22,8 +22,8 @@ IMPLEMENT_SIMPLE_AUTOMATION_TEST(FMMapFileReadTest, "System.Engine.Files.MMapFil
 bool FMMapFileReadTest::RunTest(const FString& InParameter)
 {
 	const FString TempDir = FPaths::AutomationTransientDir();
-	const FString ReadPrefix = TEXT("MMap_FileToRead");
-	const FString TxtExtension = TEXT(".txt");
+	const FString ReadPrefix = "MMap_FileToRead";
+	const FString TxtExtension = ".txt";
 	const FString TempFileToRead = FPaths::ConvertRelativePathToFull(FPaths::CreateTempFilename(*TempDir, *ReadPrefix, *TxtExtension));
 	const FString TestDirectory = FPaths::GetPath(TempFileToRead);
 
@@ -79,8 +79,8 @@ IMPLEMENT_SIMPLE_AUTOMATION_TEST(FMMapFileWriteTest, "System.Engine.Files.MMapFi
 bool FMMapFileWriteTest::RunTest(const FString& InParameter)
 {
 	const FString TempDir = FPaths::AutomationTransientDir();
-	const FString WritePrefix = TEXT("MMap_FileToWrite");
-	const FString TxtExtension = TEXT(".txt");
+	const FString WritePrefix = "MMap_FileToWrite";
+	const FString TxtExtension = ".txt";
 	const FString TempFileToWrite = FPaths::ConvertRelativePathToFull(FPaths::CreateTempFilename(*TempDir, *WritePrefix, *TxtExtension));
 	const int64	  MapAddressSpaceSize = 1024;
 	const FString TestDirectory = FPaths::GetPath(TempFileToWrite);
diff --git a/Engine/Source/Runtime/Core/Private/GenericPlatform/GenericPlatformMemory.cpp b/Engine/Source/Runtime/Core/Private/GenericPlatform/GenericPlatformMemory.cpp
--- a/Engine/Source/Runtime/Core/Private/GenericPlatform/GenericPlatformMemory.cpp
+++ b/Engine/Source/Runtime/Core/Private/GenericPlatform/GenericPlatformMemory.cpp
@@ -349,7 +349,7 @@ void FGenericPlatformMemory::DumpStats( class FOutputDevice& Ar )
 #if !NO_LOGGING
 	const FName CategoryName(LogMemory.GetCategoryName());
 #else
-	const FName CategoryName(TEXT("LogMemory"));
+	const FName CategoryName("LogMemory");
 #endif
 	Ar.CategorizedLogf(CategoryName, ELogVerbosity::Log, TEXT("Platform Memory Stats for %s"), ANSI_TO_TCHAR(FPlatformProperties::PlatformName()));
 	Ar.CategorizedLogf(CategoryName, ELogVerbosity::Log, TEXT("Process Physical Memory: %.2f MB used, %.2f MB peak"), (float)MemoryStats.UsedPhysical*InvMB, (float)MemoryStats.PeakUsedPhysical*InvMB);
diff --git a/Engine/Source/Runtime/Core/Private/GenericPlatform/GenericPlatformMisc.cpp b/Engine/Source/Runtime/Core/Private/GenericPlatform/GenericPlatformMisc.cpp
--- a/Engine/Source/Runtime/Core/Private/GenericPlatform/GenericPlatformMisc.cpp
+++ b/Engine/Source/Runtime/Core/Private/GenericPlatform/GenericPlatformMisc.cpp
@@ -224,23 +224,23 @@ static FAutoConsoleCommand GLogNamedEventsCmd(
 
 			if (Args.Num() >= 1)
 			{
-				if (Args[0].Compare(FString(TEXT("Reset")), ESearchCase::IgnoreCase) == 0)
+				if (Args[0].Compare(FString("Reset"), ESearchCase::IgnoreCase) == 0)
 				{
 					LogNameEventStats.StartFrame = GFrameCounter;
 					LogNameEventStats.Reset();
 				}
-				else if (Args[0].Compare(FString(TEXT("Dump")), ESearchCase::IgnoreCase) == 0)
+				else if (Args[0].Compare(FString("Dump"), ESearchCase::IgnoreCase) == 0)
 				{
 					LogNameEventStats.Dump();
 				}
-				else if (Args[0].Compare(FString(TEXT("Sort")), ESearchCase::IgnoreCase) == 0)
+				else if (Args[0].Compare(FString("Sort"), ESearchCase::IgnoreCase) == 0)
 				{
 					if (Args.Num() == 2)
 					{
 						LogNameEventStats.ParseSortType(Args[1]);
 					}
 				}
-				else if (Args[0].Compare(FString(TEXT("MaxOutput")), ESearchCase::IgnoreCase) == 0)
+				else if (Args[0].Compare(FString("MaxOutput"), ESearchCase::IgnoreCase) == 0)
 				{
 					if (Args.Num() == 2)
 					{
@@ -644,13 +644,13 @@ EProcessDiagnosticFlags FGenericPlatformMisc::GetProcessDiagnostics()
 FString FGenericPlatformMisc::GetCPUVendor()
 {
 	// Not implemented cross-platform. Each platform may or may not choose to implement this.
-	return FString( TEXT( "GenericCPUVendor" ) );
+	return FString( "GenericCPUVendor" );
 }
 
 FString FGenericPlatformMisc::GetCPUBrand()
 {
 	// Not implemented cross-platform. Each platform may or may not choose to implement this.
-	return FString( TEXT( "GenericCPUBrand" ) );
+	return FString( "GenericCPUBrand" );
 }
 
 uint32 FGenericPlatformMisc::GetCPUModelID()
@@ -667,7 +667,7 @@ uint32 FGenericPlatformMisc::GetCPUMicrocodeRevision()
 
 FString FGenericPlatformMisc::GetCPUChipset()
 {
-	return FString( TEXT( "Unknown" ) );
+	return FString( "Unknown" );
 }
 
 uint32 FGenericPlatformMisc::GetCPUInfo()
@@ -696,7 +696,7 @@ FString FGenericPlatformMisc::GetHostArchitecture()
 FString FGenericPlatformMisc::GetPrimaryGPUBrand()
 {
 	// Not implemented cross-platform. Each platform may or may not choose to implement this.
-	return FString( TEXT( "GenericGPUBrand" ) );
+	return FString( "GenericGPUBrand" );
 }
 
 FString FGenericPlatformMisc::GetDeviceMakeAndModel()
@@ -724,8 +724,8 @@ FGPUDriverInfo FGenericPlatformMisc::GetGPUDriverInfo(const FString& DeviceDescr
 void FGenericPlatformMisc::GetOSVersions( FString& out_OSVersionLabel, FString& out_OSSubVersionLabel )
 {
 	// Not implemented cross-platform. Each platform may or may not choose to implement this.
-	out_OSVersionLabel = FString( TEXT( "GenericOSVersionLabel" ) );
-	out_OSSubVersionLabel = FString( TEXT( "GenericOSSubVersionLabel" ) );
+	out_OSVersionLabel = FString( "GenericOSVersionLabel" );
+	out_OSSubVersionLabel = FString( "GenericOSSubVersionLabel" );
 }
 
 
@@ -868,7 +868,7 @@ bool FGenericPlatformMisc::SetStoredValue(const FString& InStoreId, const FStrin
 	check(!InKeyName.IsEmpty());
 
 	// This assumes that FPlatformProcess::ApplicationSettingsDir() returns a user-specific directory; it doesn't on Windows, but Windows overrides this behavior to use the registry
-	const FString ConfigPath = FString(FPlatformProcess::ApplicationSettingsDir()) / InStoreId / FString(TEXT("KeyValueStore.ini"));
+	const FString ConfigPath = FString(FPlatformProcess::ApplicationSettingsDir()) / InStoreId / FString("KeyValueStore.ini");
 		
 	FConfigFile ConfigFile;
 	ConfigFile.Read(ConfigPath);
@@ -884,7 +884,7 @@ bool FGenericPlatformMisc::GetStoredValue(const FString& InStoreId, const FStrin
 	check(!InKeyName.IsEmpty());
 
 	// This assumes that FPlatformProcess::ApplicationSettingsDir() returns a user-specific directory; it doesn't on Windows, but Windows overrides this behavior to use the registry
-	const FString ConfigPath = FString(FPlatformProcess::ApplicationSettingsDir()) / InStoreId / FString(TEXT("KeyValueStore.ini"));
+	const FString ConfigPath = FString(FPlatformProcess::ApplicationSettingsDir()) / InStoreId / FString("KeyValueStore.ini");
 		
 	FConfigFile ConfigFile;
 	ConfigFile.Read(ConfigPath);
@@ -899,7 +899,7 @@ bool FGenericPlatformMisc::DeleteStoredValue(const FString& InStoreId, const FSt
 	check(!InKeyName.IsEmpty());
 
 	// This assumes that FPlatformProcess::ApplicationSettingsDir() returns a user-specific directory; it doesn't on Windows, but Windows overrides this behavior to use the registry
-	const FString ConfigPath = FString(FPlatformProcess::ApplicationSettingsDir()) / InStoreId / FString(TEXT("KeyValueStore.ini"));
+	const FString ConfigPath = FString(FPlatformProcess::ApplicationSettingsDir()) / InStoreId / FString("KeyValueStore.ini");
 
 	FConfigFile ConfigFile;
 	ConfigFile.Read(ConfigPath);
@@ -918,7 +918,7 @@ bool FGenericPlatformMisc::DeleteStoredSection(const FString& InStoreId, const F
 	check(!InSectionName.IsEmpty());
 
 	// This assumes that FPlatformProcess::ApplicationSettingsDir() returns a user-specific directory; it doesn't on Windows, but Windows overrides this behavior to use the registry
-	const FString ConfigPath = FString(FPlatformProcess::ApplicationSettingsDir()) / InStoreId / FString(TEXT("KeyValueStore.ini"));
+	const FString ConfigPath = FString(FPlatformProcess::ApplicationSettingsDir()) / InStoreId / FString("KeyValueStore.ini");
 
 	FConfigFile ConfigFile;
 	ConfigFile.Read(ConfigPath);
@@ -1215,7 +1215,7 @@ void FGenericPlatformMisc::AddAdditionalRootDirectory(const FString& RootDir)
 static void MakeEngineDir(FString& OutEngineDir)
 {
 	// See if we are a root-level project
-	FString DefaultEngineDir = TEXT("../../../Engine/");
+	FString DefaultEngineDir = "../../../Engine/";
 #if PLATFORM_DESKTOP
 #if !defined(DISABLE_CWD_CHANGES) || DISABLE_CWD_CHANGES == 0
 	FPlatformProcess::SetCurrentWorkingDirectoryToBaseDir();
@@ -1905,7 +1905,7 @@ FGuid FGenericPlatformMisc::GetMachineId()
 
 FString FGenericPlatformMisc::GetDeviceTag()
 	{
-		FString DeviceTag = TEXT("");
+		FString DeviceTag;
 		FParse::Value( FCommandLine::Get(), TEXT("DeviceTag="), DeviceTag );
 
 		return DeviceTag;
diff --git a/Engine/Source/Runtime/Core/Private/GenericPlatform/GenericPlatformProcess.cpp b/Engine/Source/Runtime/Core/Private/GenericPlatform/GenericPlatformProcess.cpp
--- a/Engine/Source/Runtime/Core/Private/GenericPlatform/GenericPlatformProcess.cpp
+++ b/Engine/Source/Runtime/Core/Private/GenericPlatform/GenericPlatformProcess.cpp
@@ -375,7 +375,7 @@ bool FGenericPlatformProcess::IsApplicationRunning( const TCHAR* ProcName )
 FString FGenericPlatformProcess::GetApplicationName( uint32 ProcessId )
 {
 	UE_LOG(LogHAL, Fatal, TEXT("FGenericPlatformProcess::GetApplicationName not implemented on this platform"));
-	return FString(TEXT(""));
+	return FString();
 }
 
 bool FGenericPlatformProcess::ExecProcess(const TCHAR* URL, const TCHAR* Params, int32* OutReturnCode, FString* OutStdOut, FString* OutStdErr, const TCHAR* OptionalWorkingDirectory, bool bShouldEndWithParentProcess)
diff --git a/Engine/Source/Runtime/Core/Private/GenericPlatform/GenericWidePlatformString.cpp b/Engine/Source/Runtime/Core/Private/GenericPlatform/GenericWidePlatformString.cpp
--- a/Engine/Source/Runtime/Core/Private/GenericPlatform/GenericWidePlatformString.cpp
+++ b/Engine/Source/Runtime/Core/Private/GenericPlatform/GenericWidePlatformString.cpp
@@ -250,22 +250,22 @@ TEST_CASE("Core::PlatformString::GetVarArgs", "[Core][String][Smoke]")
 	WIDECHAR OutputString[OUTPUT_SIZE];
 
 	TestGetVarArgs(OutputString, TEXT("Test A|%-20s|%20s|%10.2f|%-10.2f|"), TEXT("LEFT"), TEXT("RIGHT"), 33.333333, 66.666666);
-	CHECK_MESSAGE(OutputString, FString(OutputString) == FString(TEXT("Test A|LEFT                |               RIGHT|     33.33|66.67     |")));
+	CHECK_MESSAGE(OutputString, FString(OutputString) == FString("Test A|LEFT                |               RIGHT|     33.33|66.67     |"));
 
 	TestGetVarArgs(OutputString, TEXT("Test B|Percents:%%%%%%%d|"), 3);
-	CHECK_MESSAGE(OutputString, FString(OutputString) == FString(TEXT("Test B|Percents:%%%3|")));
+	CHECK_MESSAGE(OutputString, FString(OutputString) == FString("Test B|Percents:%%%3|"));
 
 	TestGetVarArgs(OutputString, TEXT("Test C|%d|%i|%X|%x|%u|"), 12345, 54321, 0x123AbC, 15, 99);
-	CHECK_MESSAGE(OutputString, FString(OutputString) == FString(TEXT("Test C|12345|54321|123ABC|f|99|")));
+	CHECK_MESSAGE(OutputString, FString(OutputString) == FString("Test C|12345|54321|123ABC|f|99|"));
 
 	TestGetVarArgs(OutputString, TEXT("Test D|%p|"), 0x12345);
 #ifdef _MSC_VER
 	// MSVC's standard library formats pointers differently
-	CHECK_MESSAGE(OutputString, FString(OutputString) == FString(TEXT("Test D|0000000000012345|")));
+	CHECK_MESSAGE(OutputString, FString(OutputString) == FString("Test D|0000000000012345|"));
 #else
 	// Pointer format (which we get from snprintf) varies from platform to platform.
 	// Call it to make sure it puts the output in the proper place, but don't test for the exact format
-//	CHECK_MESSAGE(OutputString, FString(OutputString) == FString(TEXT("Test D|0x12345|")));
+//	CHECK_MESSAGE(OutputString, FString(OutputString) == FString("Test D|0x12345|"));
 	{
 		FString OutputStringStr(OutputString);
 		CHECK_MESSAGE(OutputString, OutputStringStr.StartsWith(TEXT("Test D|"))
@@ -274,25 +274,25 @@ TEST_CASE("Core::PlatformString::GetVarArgs", "[Core][String][Smoke]")
 #endif
 
 	TestGetVarArgs(OutputString, TEXT("Test E|%" INT64_FMT "|"), int64(12345678912345LL));
-	CHECK_MESSAGE(OutputString, FString(OutputString) == FString(TEXT("Test E|12345678912345|")));
+	CHECK_MESSAGE(OutputString, FString(OutputString) == FString("Test E|12345678912345|"));
 
 	TestGetVarArgs(OutputString, TEXT("Test F|%f|%e|%g|"), 123.456, 123.456, 123.456);
-	CHECK_MESSAGE(OutputString, FString(OutputString) == FString(TEXT("Test F|123.456000|1.234560e+02|123.456|")));
+	CHECK_MESSAGE(OutputString, FString(OutputString) == FString("Test F|123.456000|1.234560e+02|123.456|"));
 
 	TestGetVarArgs(OutputString, TEXT("Test G|%" UPTRINT_X_FMT"|"), UPTRINT(49374));
-	CHECK_MESSAGE(OutputString, FString(OutputString) == FString(TEXT("Test G|C0DE|")));
+	CHECK_MESSAGE(OutputString, FString(OutputString) == FString("Test G|C0DE|"));
 
 	TestGetVarArgs(OutputString, TEXT("Test H|%" UPTRINT_x_FMT "|"), UPTRINT(49374));
-	CHECK_MESSAGE(OutputString, FString(OutputString) == FString(TEXT("Test H|c0de|")));
+	CHECK_MESSAGE(OutputString, FString(OutputString) == FString("Test H|c0de|"));
 
 	TestGetVarArgs(OutputString, TEXT("Test I|%" UINT64_FMT "|"), MAX_uint64);
-	CHECK_MESSAGE(OutputString, FString(OutputString) == FString(TEXT("Test I|18446744073709551615|")));
+	CHECK_MESSAGE(OutputString, FString(OutputString) == FString("Test I|18446744073709551615|"));
 
 	TestGetVarArgs(OutputString, TEXT("Test J|%*c|"), 10, 'J');
-	CHECK_MESSAGE(OutputString, FString(OutputString) == FString(TEXT("Test J|         J|")));
+	CHECK_MESSAGE(OutputString, FString(OutputString) == FString("Test J|         J|"));
 
 	TestGetVarArgs(OutputString, TEXT("Test K|%-5c|"), 'K');
-	CHECK_MESSAGE(OutputString, FString(OutputString) == FString(TEXT("Test K|K    |")));
+	CHECK_MESSAGE(OutputString, FString(OutputString) == FString("Test K|K    |"));
 }
 #endif // WITH_LOW_LEVEL_TESTS
 
diff --git a/Engine/Source/Runtime/Core/Private/HAL/ConsoleManager.cpp b/Engine/Source/Runtime/Core/Private/HAL/ConsoleManager.cpp
--- a/Engine/Source/Runtime/Core/Private/HAL/ConsoleManager.cpp
+++ b/Engine/Source/Runtime/Core/Private/HAL/ConsoleManager.cpp
@@ -1497,7 +1497,7 @@ FConsoleVariableBase* FindOrCreatePlatformCVar(IConsoleVariable* CVar, FName Pla
 	return nullptr;
 }
 
-static const FString GSpecialDPNameForPremadePlatformKey(TEXT("/"));
+static const FString GSpecialDPNameForPremadePlatformKey("/");
 /**
  * Helper function to make a single key used in the PlatformValues set, that combines Platform and a DP name. If the DP name is empty, this will use
  * Platform name as the DP name. This matches up with how GetPlatformValueVariable() was generally used in the past - get the CVar using the platform-named
@@ -2491,8 +2491,8 @@ void FConsoleManager::LoadHistoryIfNeeded()
 	FConfigFile Ini;
 	FConfigContext::ReadSingleIntoLocalFile(Ini).Load(*(FPaths::GeneratedConfigDir() + TEXT("ConsoleHistory.ini")));
 
-	const FString SectionName = TEXT("ConsoleHistory");
-	const FName KeyName = TEXT("History");
+	const FString SectionName = "ConsoleHistory";
+	const FName KeyName = "History";
 
 	for (const auto& ConfigPair : AsConst(Ini))
 	{
@@ -2525,8 +2525,8 @@ void FConsoleManager::SaveHistory()
 {
 	FConfigFile Ini;
 	
-	const FString SectionName = TEXT("ConsoleHistory");
-	const FName KeyName = TEXT("History");
+	const FString SectionName = "ConsoleHistory";
+	const FName KeyName = "History";
 
 	for (const auto& HistoryPair : HistoryEntriesMap)
 	{
diff --git a/Engine/Source/Runtime/Core/Private/HAL/IPlatformFileOpenLogWrapper.cpp b/Engine/Source/Runtime/Core/Private/HAL/IPlatformFileOpenLogWrapper.cpp
--- a/Engine/Source/Runtime/Core/Private/HAL/IPlatformFileOpenLogWrapper.cpp
+++ b/Engine/Source/Runtime/Core/Private/HAL/IPlatformFileOpenLogWrapper.cpp
@@ -60,7 +60,7 @@ bool FPlatformFileOpenLog::Initialize(IPlatformFile* Inner, const TCHAR* Command
 #if WITH_EDITOR
 			LogFilePath = FPaths::Combine(*LogFileDirectory, TEXT("EditorOpenOrder.log"));
 #else 
-			FString OrderSuffix = TEXT("");
+			FString OrderSuffix;
 			FParse::Value(CommandLineParam, TEXT("FooSuffix="), OrderSuffix);
 
 			FString GameOpenOrderFileName = FString::Printf(TEXT("GameOpenOrder%s.log"), *OrderSuffix);
diff --git a/Engine/Source/Runtime/Core/Private/HAL/LowLevelMemTracker.cpp b/Engine/Source/Runtime/Core/Private/HAL/LowLevelMemTracker.cpp
--- a/Engine/Source/Runtime/Core/Private/HAL/LowLevelMemTracker.cpp
+++ b/Engine/Source/Runtime/Core/Private/HAL/LowLevelMemTracker.cpp
@@ -1703,13 +1703,13 @@ static FName GetTagName_CustomName()
 {
 	// This FName can be read before c++ global constructors are complete, by LLM_SCOPE_BY_BOOTSTRAP_TAG
 	// so it needs to be a function static rather than a global.
-	static FName TagName_CustomName(TEXT("CustomName"));
+	static FName TagName_CustomName("CustomName");
 	return TagName_CustomName;
 }
-static FName TagName_Untagged(TEXT("Untagged"));
-static FName TagName_UntaggedAsset(TEXT("UntaggedAsset"));
-static FName TagName_UntaggedAssetClass(TEXT("UntaggedAssetClass"));
-static FName TagName_UntaggedUObjectClasses(TEXT("UntaggedUObjectClasses"));
+static FName TagName_Untagged("Untagged");
+static FName TagName_UntaggedAsset("UntaggedAsset");
+static FName TagName_UntaggedAssetClass("UntaggedAssetClass");
+static FName TagName_UntaggedUObjectClasses("UntaggedUObjectClasses");
 
 FName LLMGetTagUniqueName(ELLMTag Tag)
 {
@@ -5565,7 +5565,7 @@ bool FLLMTracker::DumpForkedAllocationInfo()
 
 	uint16 ForkId = FForkProcessHelper::GetForkedChildProcessIndex();
 
-	FString ForkString = ForkId == 0 ? FString(TEXT("Parent")) : FString::Printf(TEXT("Child%d"), ForkId);
+	FString ForkString = ForkId == 0 ? FString("Parent") : FString::Printf(TEXT("Child%d"), ForkId);
 	const TCHAR* TrackerString = TEXT("Default");
 	if (Tracker == ELLMTracker::Platform)
 	{
diff --git a/Engine/Source/Runtime/Core/Private/HAL/LowLevelMemTrackerTest.cpp b/Engine/Source/Runtime/Core/Private/HAL/LowLevelMemTrackerTest.cpp
--- a/Engine/Source/Runtime/Core/Private/HAL/LowLevelMemTrackerTest.cpp
+++ b/Engine/Source/Runtime/Core/Private/HAL/LowLevelMemTrackerTest.cpp
@@ -30,8 +30,8 @@ bool FLLMScopeTest::RunTest(const FString& Parameters)
 		AllocationByNameExisting = new int[10];
 		if (bIsEnabled)
 		{
-			TestEqual(TEXT("NameExisting has proper name"), LocalLLM.GetActiveTagData(ELLMTracker::Default)->GetName(), FName(TEXT("LLMTestParentTag_LLMTestChildTag1")));
-			TestEqual(TEXT("NameExisting has proper parent"), LocalLLM.GetActiveTagData(ELLMTracker::Default)->GetParent()->GetName(), FName(TEXT("LLMTestParentTag")));
+			TestEqual(TEXT("NameExisting has proper name"), LocalLLM.GetActiveTagData(ELLMTracker::Default)->GetName(), FName("LLMTestParentTag_LLMTestChildTag1"));
+			TestEqual(TEXT("NameExisting has proper parent"), LocalLLM.GetActiveTagData(ELLMTracker::Default)->GetParent()->GetName(), FName("LLMTestParentTag"));
 		}
 	}
 	{
@@ -39,8 +39,8 @@ bool FLLMScopeTest::RunTest(const FString& Parameters)
 		AllocationByNameNew = new int[10];
 		if (bIsEnabled)
 		{
-			TestEqual(TEXT("NameNew has proper name"), LocalLLM.GetActiveTagData(ELLMTracker::Default)->GetName(), FName(TEXT("LLMTestParentTag_LLMTestChildTag4")));
-			TestEqual(TEXT("NameNew has proper parent"), LocalLLM.GetActiveTagData(ELLMTracker::Default)->GetParent()->GetName(), FName(TEXT("LLMTestParentTag")));
+			TestEqual(TEXT("NameNew has proper name"), LocalLLM.GetActiveTagData(ELLMTracker::Default)->GetName(), FName("LLMTestParentTag_LLMTestChildTag4"));
+			TestEqual(TEXT("NameNew has proper parent"), LocalLLM.GetActiveTagData(ELLMTracker::Default)->GetParent()->GetName(), FName("LLMTestParentTag"));
 		}
 	}
 	{
@@ -48,8 +48,8 @@ bool FLLMScopeTest::RunTest(const FString& Parameters)
 		AllocationByTag = new int[10];
 		if (bIsEnabled)
 		{
-			TestEqual(TEXT("ByTag has proper name"), LocalLLM.GetActiveTagData(ELLMTracker::Default)->GetName(), FName(TEXT("LLMTestParentTag_LLMTestChildTag2")));
-			TestEqual(TEXT("ByTag has proper parent"), LocalLLM.GetActiveTagData(ELLMTracker::Default)->GetParent()->GetName(), FName(TEXT("LLMTestParentTag")));
+			TestEqual(TEXT("ByTag has proper name"), LocalLLM.GetActiveTagData(ELLMTracker::Default)->GetName(), FName("LLMTestParentTag_LLMTestChildTag2"));
+			TestEqual(TEXT("ByTag has proper parent"), LocalLLM.GetActiveTagData(ELLMTracker::Default)->GetParent()->GetName(), FName("LLMTestParentTag"));
 		}
 	}
 	{
@@ -57,8 +57,8 @@ bool FLLMScopeTest::RunTest(const FString& Parameters)
 		AllocationByTag3 = new int[10];
 		if (bIsEnabled)
 		{
-			TestEqual(TEXT("Parsed parent tag has proper name"), LocalLLM.GetActiveTagData(ELLMTracker::Default)->GetName(), FName(TEXT("LLMTestParentTag_LLMTestChildTag3")));
-			TestEqual(TEXT("Parsed parent tag has proper parent"), LocalLLM.GetActiveTagData(ELLMTracker::Default)->GetParent()->GetName(), FName(TEXT("LLMTestParentTag")));
+			TestEqual(TEXT("Parsed parent tag has proper name"), LocalLLM.GetActiveTagData(ELLMTracker::Default)->GetName(), FName("LLMTestParentTag_LLMTestChildTag3"));
+			TestEqual(TEXT("Parsed parent tag has proper parent"), LocalLLM.GetActiveTagData(ELLMTracker::Default)->GetParent()->GetName(), FName("LLMTestParentTag"));
 		}
 	}
 	delete[] AllocationByNameExisting;
diff --git a/Engine/Source/Runtime/Core/Private/IOS/IOSPlatformMisc.cpp b/Engine/Source/Runtime/Core/Private/IOS/IOSPlatformMisc.cpp
--- a/Engine/Source/Runtime/Core/Private/IOS/IOSPlatformMisc.cpp
+++ b/Engine/Source/Runtime/Core/Private/IOS/IOSPlatformMisc.cpp
@@ -337,7 +337,7 @@ void FIOSPlatformMisc::PlatformHandleSplashScreen(bool ShowSplashScreen)
 
 const TCHAR* FIOSPlatformMisc::GamePersistentDownloadDir()
 {
-    static FString GamePersistentDownloadDir = TEXT("");
+    static FString GamePersistentDownloadDir;
     
     if (GamePersistentDownloadDir.Len() == 0)
     {
@@ -2080,25 +2080,25 @@ static TMap<FString, FString> GetPredefinedConfigRuleVars()
 	FString OSVersionLabel,OSSubVersionLabel;
 	FIOSPlatformMisc::GetOSVersions(OSVersionLabel,OSSubVersionLabel);
 	
-	PredefinedConfigRuleVars.Add(FString(TEXT("SRC_OSVersion")), OSSubVersionLabel);
-	PredefinedConfigRuleVars.Add(FString(TEXT("SRC_OSLanguage")), *FString([[NSLocale preferredLanguages] objectAtIndex:0]));
+	PredefinedConfigRuleVars.Add(FString("SRC_OSVersion"), OSSubVersionLabel);
+	PredefinedConfigRuleVars.Add(FString("SRC_OSLanguage"), *FString([[NSLocale preferredLanguages] objectAtIndex:0]));
 	
-	PredefinedConfigRuleVars.Add(FString(TEXT("SRC_GPUFamily")), CpuVendor);
-	PredefinedConfigRuleVars.Add(FString(TEXT("SRC_DeviceMake")), CpuVendor);
-	PredefinedConfigRuleVars.Add(FString(TEXT("SRC_DeviceModel")), CPUBrand);
-	PredefinedConfigRuleVars.Add(FString(TEXT("SRC_TotalPhysicalGB")), FString::FromInt(totalMemoryGB));
-	PredefinedConfigRuleVars.Add(FString(TEXT("memory")), FString::FromInt(totalMemoryMB));
-
-	PredefinedConfigRuleVars.Add(FString(TEXT("versionCode")), BuildNumber);
-	PredefinedConfigRuleVars.Add(FString(TEXT("versionName")), ProjectVersion);
+	PredefinedConfigRuleVars.Add(FString("SRC_GPUFamily"), CpuVendor);
+	PredefinedConfigRuleVars.Add(FString("SRC_DeviceMake"), CpuVendor);
+	PredefinedConfigRuleVars.Add(FString("SRC_DeviceModel"), CPUBrand);
+	PredefinedConfigRuleVars.Add(FString("SRC_TotalPhysicalGB"), FString::FromInt(totalMemoryGB));
+	PredefinedConfigRuleVars.Add(FString("memory"), FString::FromInt(totalMemoryMB));
+
+	PredefinedConfigRuleVars.Add(FString("versionCode"), BuildNumber);
+	PredefinedConfigRuleVars.Add(FString("versionName"), ProjectVersion);
 	
-	PredefinedConfigRuleVars.Add(FString(TEXT("processorCount")), FString::FromInt(FIOSPlatformMisc::NumberOfCores()));
+	PredefinedConfigRuleVars.Add(FString("processorCount"), FString::FromInt(FIOSPlatformMisc::NumberOfCores()));
 #if !PLATFORM_TVOS && !PLATFORM_VISIONOS
 	id<MTLDevice> DefaultMetalDevice = MTLCreateSystemDefaultDevice();
 	if(DefaultMetalDevice)
 	{
 		FString GpuName([DefaultMetalDevice name]);
-		PredefinedConfigRuleVars.Add(FString(TEXT("SRC_GPUFamily")), GpuName);
+		PredefinedConfigRuleVars.Add(FString("SRC_GPUFamily"), GpuName);
 		MTLGPUFamily Families[] = {
 			MTLGPUFamilyApple1,
 			MTLGPUFamilyApple2,
diff --git a/Engine/Source/Runtime/Core/Private/IOS/IOSPlatformProcess.cpp b/Engine/Source/Runtime/Core/Private/IOS/IOSPlatformProcess.cpp
--- a/Engine/Source/Runtime/Core/Private/IOS/IOSPlatformProcess.cpp
+++ b/Engine/Source/Runtime/Core/Private/IOS/IOSPlatformProcess.cpp
@@ -218,7 +218,7 @@ FString FIOSPlatformProcess::GenerateApplicationPath( const FString& AppName, EB
 {
     SCOPED_AUTORELEASE_POOL;
     
-    FString PlatformName = TEXT("IOS");
+    FString PlatformName = "IOS";
     FString ExecutableName = AppName;
     if (BuildConfiguration != EBuildConfiguration::Development && BuildConfiguration != EBuildConfiguration::DebugGame)
     {
diff --git a/Engine/Source/Runtime/Core/Private/Internationalization/Culture.cpp b/Engine/Source/Runtime/Core/Private/Internationalization/Culture.cpp
--- a/Engine/Source/Runtime/Core/Private/Internationalization/Culture.cpp
+++ b/Engine/Source/Runtime/Core/Private/Internationalization/Culture.cpp
@@ -370,7 +370,7 @@ void FCulture::RefreshCultureDisplayNames(const TArray<FString>& InPrioritizedDi
 	if (bFullRefresh)
 	{
 		{
-			static const FString EnglishCultureName = TEXT("en");
+			static const FString EnglishCultureName = "en";
 			CachedEnglishName = Implementation->GetEnglishName();
 			ApplyCultureDisplayNameSubstitutes(MakeArrayView(&EnglishCultureName, 1), CachedEnglishName);
 		}
diff --git a/Engine/Source/Runtime/Core/Private/Internationalization/Cultures/InvariantCulture.h b/Engine/Source/Runtime/Core/Private/Internationalization/Cultures/InvariantCulture.h
--- a/Engine/Source/Runtime/Core/Private/Internationalization/Cultures/InvariantCulture.h
+++ b/Engine/Source/Runtime/Core/Private/Internationalization/Cultures/InvariantCulture.h
@@ -50,14 +50,14 @@ public:
 
 		FCultureRef Culture = FCulture::Create(MakeUnique<FLegacyCultureImplementation>(
 			LOCTEXT("InvariantCultureDisplayName", "Invariant Language (Invariant Country)"),	//const FText DisplayName
-			FString(TEXT("Invariant Language (Invariant Country)")),							//const FString EnglishName
+			FString("Invariant Language (Invariant Country)"),									//const FString EnglishName
 			1033,																				//const int KeyboardLayoutId
 			1033,																				//const int LCID
-			FString(TEXT("")),																	//const FString Name
-			FString(TEXT("Invariant Language (Invariant Country)")),							//const FString NativeName
-			FString(TEXT("INT")),																//const FString UnrealLegacyThreeLetterISOLanguageName
-			FString(TEXT("ivl")),																//const FString ThreeLetterISOLanguageName
-			FString(TEXT("iv")),																//const FString TwoLetterISOLanguageName
+			FString{},																			//const FString Name
+			FString("Invariant Language (Invariant Country)"),									//const FString NativeName
+			FString("INT"),																		//const FString UnrealLegacyThreeLetterISOLanguageName
+			FString("ivl"),																		//const FString ThreeLetterISOLanguageName
+			FString("iv"),																		//const FString TwoLetterISOLanguageName
 			DecimalNumberFormattingRules,														//const FDecimalNumberFormattingRules InDecimalNumberFormattingRules
 			PercentFormattingRules,																//const FDecimalNumberFormattingRules InPercentFormattingRules
 			BaseCurrencyFormattingRules,														//const FDecimalNumberFormattingRules InBaseCurrencyFormattingRules
diff --git a/Engine/Source/Runtime/Core/Private/Internationalization/Cultures/KeysCulture.cpp b/Engine/Source/Runtime/Core/Private/Internationalization/Cultures/KeysCulture.cpp
--- a/Engine/Source/Runtime/Core/Private/Internationalization/Cultures/KeysCulture.cpp
+++ b/Engine/Source/Runtime/Core/Private/Internationalization/Cultures/KeysCulture.cpp
@@ -11,7 +11,7 @@ FKeysCulture::FKeysCulture(const FCultureRef& InInvariantCulture)
 
 const FString& FKeysCulture::StaticGetName()
 {
-	static const FString KeysCultureName = TEXT("keys");
+	static const FString KeysCultureName = "keys";
 	return KeysCultureName;
 }
 
diff --git a/Engine/Source/Runtime/Core/Private/Internationalization/Cultures/LeetCulture.cpp b/Engine/Source/Runtime/Core/Private/Internationalization/Cultures/LeetCulture.cpp
--- a/Engine/Source/Runtime/Core/Private/Internationalization/Cultures/LeetCulture.cpp
+++ b/Engine/Source/Runtime/Core/Private/Internationalization/Cultures/LeetCulture.cpp
@@ -11,7 +11,7 @@ FLeetCulture::FLeetCulture(const FCultureRef& InInvariantCulture)
 
 const FString& FLeetCulture::StaticGetName()
 {
-	static const FString LeetCultureName = TEXT("LEET");
+	static const FString LeetCultureName = "LEET";
 	return LeetCultureName;
 }
 
diff --git a/Engine/Source/Runtime/Core/Private/Internationalization/ICUInternationalization.cpp b/Engine/Source/Runtime/Core/Private/Internationalization/ICUInternationalization.cpp
--- a/Engine/Source/Runtime/Core/Private/Internationalization/ICUInternationalization.cpp
+++ b/Engine/Source/Runtime/Core/Private/Internationalization/ICUInternationalization.cpp
@@ -115,7 +115,7 @@ bool FICUInternationalization::Initialize()
 
 	u_setMemoryFunctions(nullptr, &FICUOverrides::Malloc, &FICUOverrides::Realloc, &FICUOverrides::Free, &ICUStatus);
 
-	const FString DataDirectoryRelativeToContent = TEXT("Internationalization");
+	const FString DataDirectoryRelativeToContent = "Internationalization";
 	const FString PotentialDataDirectories[] =
 	{
 		FPaths::ProjectContentDir() / DataDirectoryRelativeToContent, // Try game content directory.
@@ -234,13 +234,13 @@ void FICUInternationalization::LoadDLLs()
 
 #if PLATFORM_WINDOWS
 #if PLATFORM_64BITS
-	const FString PlatformFolderName = TEXT("Win64");
+	const FString PlatformFolderName = "Win64";
 #elif PLATFORM_32BITS
-	const FString PlatformFolderName = TEXT("Win32");
+	const FString PlatformFolderName = "Win32";
 #endif //PLATFORM_*BITS
 
 #if _MSC_VER >= 1900
-	const FString VSVersionFolderName = TEXT("VS2015");
+	const FString VSVersionFolderName = "VS2015";
 #else
 	#error "FICUInternationalization::LoadDLLs - Unknown _MSC_VER! Please update this code for this version of MSVC."
 #endif //_MSVC_VER
@@ -277,9 +277,9 @@ void FICUInternationalization::LoadDLLs()
 #endif //PLATFORM_*
 
 #if UE_BUILD_DEBUG && !defined(NDEBUG)
-	const FString LibraryNamePostfix = TEXT("d");
+	const FString LibraryNamePostfix = "d";
 #else
-	const FString LibraryNamePostfix = TEXT("");
+	const FString LibraryNamePostfix;
 #endif //DEBUG
 
 	for(FString Stem : LibraryNameStems)
diff --git a/Engine/Source/Runtime/Core/Private/Internationalization/StringTableCore.cpp b/Engine/Source/Runtime/Core/Private/Internationalization/StringTableCore.cpp
--- a/Engine/Source/Runtime/Core/Private/Internationalization/StringTableCore.cpp
+++ b/Engine/Source/Runtime/Core/Private/Internationalization/StringTableCore.cpp
@@ -83,7 +83,7 @@ FTextId FStringTableEntry::GetDisplayStringId() const
 
 const FString& FStringTableEntry::GetPlaceholderSourceString()
 {
-	static const FString MissingSourceString = TEXT("<MISSING STRING TABLE ENTRY>");
+	static const FString MissingSourceString = "<MISSING STRING TABLE ENTRY>";
 	return MissingSourceString;
 }
 
@@ -592,7 +592,7 @@ void FStringTableRedirects::InitStringTableRedirects()
 	{
 		for (FConfigSection::TConstIterator It(*CoreStringTableSection); It; ++It)
 		{
-			static const FName StringTableRedirectsName = TEXT("StringTableRedirects");
+			static const FName StringTableRedirectsName = "StringTableRedirects";
 			if (It.Key() == StringTableRedirectsName)
 			{
 				const FString& ConfigValue = It.Value().GetValue();
diff --git a/Engine/Source/Runtime/Core/Private/Internationalization/Text.cpp b/Engine/Source/Runtime/Core/Private/Internationalization/Text.cpp
--- a/Engine/Source/Runtime/Core/Private/Internationalization/Text.cpp
+++ b/Engine/Source/Runtime/Core/Private/Internationalization/Text.cpp
@@ -804,7 +804,7 @@ FText FText::AsMemory(uint64 NumBytes, const FNumberFormattingOptions* const Opt
 
 	static const TCHAR* Prefixes = TEXT("KMGTPEZYkMGTPEZY");
 	int32 Prefix = 0;
-	FString Suffix = TEXT("iB");
+	FString Suffix = "iB";
 	uint64 Unit = 1024;
 	if (UnitStandard == EMemoryUnitStandard::SI)
 	{
@@ -817,7 +817,7 @@ FText FText::AsMemory(uint64 NumBytes, const FNumberFormattingOptions* const Opt
 	if (NumBytes < 1000)
 	{
 		Args.Add( TEXT("Number"), FText::AsNumber( NumBytes, Options, TargetCulture) );
-		Args.Add( TEXT("Unit"), FText::FromString( FString( TEXT("B") ) ) );
+		Args.Add( TEXT("Unit"), FText::FromString( FString( "B" ) ) );
 		return FText::Format( NSLOCTEXT("Internationalization", "ComputerMemoryFormatting", "{Number} {Unit}"), Args );
 	}
 	while (NumBytes >= (1000000))
@@ -1611,7 +1611,7 @@ const TCHAR* FFormatArgumentValue::FromExportedString(const TCHAR* Buffer)
 {
 	// Is this a text gender?
 	{
-		static const FString TextGenderMarker = TEXT("ETextGender::");
+		static const FString TextGenderMarker = "ETextGender::";
 
 		ETextGender LocalGender = ETextGender::Masculine;
 		const TCHAR* Result = TextStringificationUtil::ReadScopedEnumFromBuffer(Buffer, TextGenderMarker, LocalGender);
diff --git a/Engine/Source/Runtime/Core/Private/Internationalization/TextHistory.cpp b/Engine/Source/Runtime/Core/Private/Internationalization/TextHistory.cpp
--- a/Engine/Source/Runtime/Core/Private/Internationalization/TextHistory.cpp
+++ b/Engine/Source/Runtime/Core/Private/Internationalization/TextHistory.cpp
@@ -318,7 +318,7 @@ const TCHAR* ReadNumberFormattingOptionsFromBuffer(const TCHAR* Buffer, FNumberF
 
 	auto ReadRoundingModeOption = [](const TCHAR* InValueBuffer, ERoundingMode& OutValue) -> const TCHAR*
 	{
-		static const FString RoundingModeMarker = TEXT("ERoundingMode::");
+		static const FString RoundingModeMarker = "ERoundingMode::";
 		return ReadScopedEnumFromBuffer(InValueBuffer, RoundingModeMarker, OutValue);
 	};
 
@@ -524,7 +524,7 @@ const TCHAR* ReadDateTimeFromBuffer(const TCHAR* Buffer, const FString& TokenMar
 {
 	auto ReadDateTimeStyle = [](const TCHAR* InValueBuffer, EDateTimeStyle::Type& OutValue) -> const TCHAR*
 	{
-		static const FString DateTimeStyleMarker = TEXT("EDateTimeStyle::");
+		static const FString DateTimeStyleMarker = "EDateTimeStyle::";
 		return ReadScopedEnumFromBuffer(InValueBuffer, DateTimeStyleMarker, OutValue);
 	};
 
diff --git a/Engine/Source/Runtime/Core/Private/Internationalization/TextLocalizationManager.cpp b/Engine/Source/Runtime/Core/Private/Internationalization/TextLocalizationManager.cpp
--- a/Engine/Source/Runtime/Core/Private/Internationalization/TextLocalizationManager.cpp
+++ b/Engine/Source/Runtime/Core/Private/Internationalization/TextLocalizationManager.cpp
@@ -430,7 +430,7 @@ void ApplyDefaultCultureSettings(const ELocalizationLoadFlags LocLoadFlags)
 		return TargetCultureName;
 	};
 
-	FString FallbackLanguage = TEXT("en");
+	FString FallbackLanguage = "en";
 	if (EnumHasAnyFlags(LocLoadFlags, ELocalizationLoadFlags::Game))
 	{
 		// If this is a game, use the native culture of the game as the fallback
diff --git a/Engine/Source/Runtime/Core/Private/Logging/LogSuppressionInterface.cpp b/Engine/Source/Runtime/Core/Private/Logging/LogSuppressionInterface.cpp
--- a/Engine/Source/Runtime/Core/Private/Logging/LogSuppressionInterface.cpp
+++ b/Engine/Source/Runtime/Core/Private/Logging/LogSuppressionInterface.cpp
@@ -86,7 +86,7 @@ class FLogSuppressionImplementation: public FLogSuppressionInterface, private FS
 			return false;
 		}
 		
-		static FName NAME_Only(TEXT("only"));
+		static FName NAME_Only("only");
 		if (NAME_Only != FName(*CommandParts[1]))
 		{
 			return false;
@@ -126,8 +126,8 @@ class FLogSuppressionImplementation: public FLogSuppressionInterface, private FS
 		// e.g., Set specific filter  : `log logshaders verbose
 		// e.g., Combo command        : `log global none, log logshaders verbose
 
-		static FName NAME_BootGlobal(TEXT("BootGlobal"));
-		static FName NAME_Reset(TEXT("Reset"));
+		static FName NAME_BootGlobal("BootGlobal");
+		static FName NAME_Reset("Reset");
 		FString Cmds = CmdString;
 		Cmds = Cmds.TrimStart().TrimQuotes().TrimStart();
 		TArray<FString> SubCmds;
@@ -231,14 +231,14 @@ class FLogSuppressionImplementation: public FLogSuppressionInterface, private FS
 				for (int32 PartIndex = 1; PartIndex < CommandParts.Num(); PartIndex++)
 				{
 					FName CmdToken = FName(*CommandParts[PartIndex]);
-					static FName NAME_Verbose(TEXT("Verbose"));
-					static FName NAME_VeryVerbose(TEXT("VeryVerbose"));
-					static FName NAME_On(TEXT("On"));
-					static FName NAME_Off(TEXT("Off"));
-					static FName NAME_Break(TEXT("Break"));
-					static FName NAME_Fatal(TEXT("Fatal"));
-					static FName NAME_Log(TEXT("Log"));
-					static FName NAME_Display(TEXT("Display"));
+					static FName NAME_Verbose("Verbose");
+					static FName NAME_VeryVerbose("VeryVerbose");
+					static FName NAME_On("On");
+					static FName NAME_Off("Off");
+					static FName NAME_Break("Break");
+					static FName NAME_Fatal("Fatal");
+					static FName NAME_Log("Log");
+					static FName NAME_Display("Display");
 					if (CmdToken == NAME_None)
 					{
 						Value &= ~ELogVerbosity::VerbosityMask;
@@ -411,7 +411,7 @@ class FLogSuppressionImplementation: public FLogSuppressionInterface, private FS
 		else
 		{
 			// see if we had a boot global override
-			static FName NAME_BootGlobal(TEXT("BootGlobal"));
+			static FName NAME_BootGlobal("BootGlobal");
 			Boot = BootAssociations.Find(NAME_BootGlobal);
 			if (Boot)
 			{
@@ -516,7 +516,7 @@ class FLogSuppressionImplementation: public FLogSuppressionInterface, private FS
 #if !UE_BUILD_SHIPPING
 		// and the command line overrides the config values
 		FString CmdLine(FCommandLine::Get());
-		FString LogCmds(TEXT("-LogCmds="));
+		FString LogCmds("-LogCmds=");
 		int32 IndexOfEnv = CmdLine.Find(TEXT("-EnvAfterHere"));
 		if (IndexOfEnv != INDEX_NONE)
 		{
diff --git a/Engine/Source/Runtime/Core/Private/Mac/MacPlatformMisc.cpp b/Engine/Source/Runtime/Core/Private/Mac/MacPlatformMisc.cpp
--- a/Engine/Source/Runtime/Core/Private/Mac/MacPlatformMisc.cpp
+++ b/Engine/Source/Runtime/Core/Private/Mac/MacPlatformMisc.cpp
@@ -171,7 +171,7 @@ static const char* GetClassCode()
 // e.g. 13.2.1
 static void MacPlatformGetOSProductVersion(FString &OutOSVersion)
 {
-	static FString OSVersion = TEXT("");
+	static FString OSVersion;
 	if (OSVersion.IsEmpty())
 	{
 		SIZE_T OSProductVersionStringSize = 0;
@@ -191,7 +191,7 @@ static void MacPlatformGetOSProductVersion(FString &OutOSVersion)
 // e.g. 22D68
 static void MacPlatformGetOSVersion(FString &OutOSBuild)
 {
-	static FString OSBuild = TEXT("");
+	static FString OSBuild;
 	if (OSBuild.IsEmpty())
 	{
 		SIZE_T OSVersionStringSize = 0;
@@ -1718,7 +1718,7 @@ bool FMacPlatformMisc::HasSeparateChannelForDebugOutput()
 
 FString FMacPlatformMisc::GetBonjourName()
 {
-	static FString BonjourName = TEXT("");
+	static FString BonjourName;
 	if (BonjourName.IsEmpty())
 	{
 		CFStringRef LocalHostNameRef = SCDynamicStoreCopyLocalHostName(NULL);
diff --git a/Engine/Source/Runtime/Core/Private/Mac/MacPlatformProcess.cpp b/Engine/Source/Runtime/Core/Private/Mac/MacPlatformProcess.cpp
--- a/Engine/Source/Runtime/Core/Private/Mac/MacPlatformProcess.cpp
+++ b/Engine/Source/Runtime/Core/Private/Mac/MacPlatformProcess.cpp
@@ -97,7 +97,7 @@ FString FMacPlatformProcess::GenerateApplicationPath( const FString& AppName, EB
 {
 	SCOPED_AUTORELEASE_POOL;
 	
-	FString PlatformName = TEXT("Mac");
+	FString PlatformName = "Mac";
 	FString ExecutableName = AppName;
 	if (BuildConfiguration != EBuildConfiguration::Development)
 	{
@@ -993,7 +993,7 @@ bool FMacPlatformProcess::IsApplicationRunning( uint32 ProcessId )
 
 FString FMacPlatformProcess::GetApplicationName( uint32 ProcessId )
 {
-	FString Output = TEXT("");
+	FString Output;
 
 	char Buffer[MAC_MAX_PATH];
 	int32 Ret = proc_pidpath(ProcessId, Buffer, sizeof(Buffer));
@@ -1115,7 +1115,7 @@ static TCHAR* UserLibrarySubDirectory()
 	static TCHAR Result[MAC_MAX_PATH] = TEXT("");
 	if (!Result[0])
 	{
-		FString SubDirectory = IsRunningGame() ? FString(FApp::GetProjectName()) : FString(TEXT("Unreal Engine")) / FApp::GetProjectName();
+		FString SubDirectory = IsRunningGame() ? FString(FApp::GetProjectName()) : FString("Unreal Engine") / FApp::GetProjectName();
 		if (IsRunningDedicatedServer())
 		{
 			SubDirectory += TEXT("Server");
diff --git a/Engine/Source/Runtime/Core/Private/Math/UnrealMath.cpp b/Engine/Source/Runtime/Core/Private/Math/UnrealMath.cpp
--- a/Engine/Source/Runtime/Core/Private/Math/UnrealMath.cpp
+++ b/Engine/Source/Runtime/Core/Private/Math/UnrealMath.cpp
@@ -3292,7 +3292,7 @@ PrecLoop:
 	}
 	else if( c == TEXT(")") )
 	{
-		*pStr = FString(TEXT(")")) + *pStr;
+		*pStr = FString(")") + *pStr;
 		*pResult = V;
 		return 1;
 	}
diff --git a/Engine/Source/Runtime/Core/Private/Microsoft/MicrosoftPlatformCrashContext.cpp b/Engine/Source/Runtime/Core/Private/Microsoft/MicrosoftPlatformCrashContext.cpp
--- a/Engine/Source/Runtime/Core/Private/Microsoft/MicrosoftPlatformCrashContext.cpp
+++ b/Engine/Source/Runtime/Core/Private/Microsoft/MicrosoftPlatformCrashContext.cpp
@@ -84,7 +84,7 @@ void FMicrosoftPlatformCrashContext::ConvertProgramCountersToStackFrames(
 	OutStackFrames.Reset(NumPCs);
 
 	TCHAR Buffer[PLATFORM_MAX_FILEPATH_LENGTH_DEPRECATED];
-	FString Unknown = TEXT("Unknown");
+	FString Unknown = "Unknown";
 
 	HANDLE Process = ProcessHandle.IsValid() ? ProcessHandle.Get() : GetCurrentProcess();
 	// Create the crash context
diff --git a/Engine/Source/Runtime/Core/Private/Misc/AutomationTest.cpp b/Engine/Source/Runtime/Core/Private/Misc/AutomationTest.cpp
--- a/Engine/Source/Runtime/Core/Private/Misc/AutomationTest.cpp
+++ b/Engine/Source/Runtime/Core/Private/Misc/AutomationTest.cpp
@@ -68,7 +68,7 @@ namespace AutomationTest
 		bLightweightStereoTestVariants,
 		TEXT("Whether to skip variants when the baseline test fails, and skip saving screenshots for successful variants"));
 
-	FString TestTagGlobalFilter = "";
+	FString TestTagGlobalFilter;
 	static FAutoConsoleVariableRef CVarAutomationTestTagGlobalFilter(
 		TEXT("Automation.TestTagGlobalFilter"),
 		TestTagGlobalFilter,
@@ -258,7 +258,7 @@ void FAutomationTestFramework::FAutomationTestOutputDevice::Serialize( const TCH
 			// using this log interception stuff.
 
 			FString LogString = FString(V);
-			FString AnalyticsString = TEXT("AUTOMATIONANALYTICS");
+			FString AnalyticsString = "AUTOMATIONANALYTICS";
 			if (LogString.StartsWith(*AnalyticsString))
 			{
 				//Remove "analytics" from the string
@@ -330,7 +330,7 @@ FAutomationTestFramework& FAutomationTestFramework::Get()
 
 FString FAutomationTestFramework::GetUserAutomationDirectory() const
 {
-	const FString DefaultAutomationSubFolder = TEXT("Unreal Automation");
+	const FString DefaultAutomationSubFolder = "Unreal Automation";
 	return FString(FPlatformProcess::UserDir()) + DefaultAutomationSubFolder;
 }
 
diff --git a/Engine/Source/Runtime/Core/Private/Misc/ConfigCacheIni.cpp b/Engine/Source/Runtime/Core/Private/Misc/ConfigCacheIni.cpp
--- a/Engine/Source/Runtime/Core/Private/Misc/ConfigCacheIni.cpp
+++ b/Engine/Source/Runtime/Core/Private/Misc/ConfigCacheIni.cpp
@@ -39,9 +39,9 @@
 
 namespace
 {
-	const FString CurrentIniVersionStr = TEXT("CurrentIniVersion");
-	const FName VersionSectionName = TEXT("Version");
-	const FString SectionsToSaveStr = TEXT("SectionsToSave");
+	const FString CurrentIniVersionStr = "CurrentIniVersion";
+	const FName VersionSectionName = "Version";
+	const FString SectionsToSaveStr = "SectionsToSave";
 
 	TMap<FString, FString> SectionRemap;
 	TMap<FString, TMap<FString, FString>> KeyRemap;
@@ -299,7 +299,7 @@ static void CheckLongSectionNames(const TCHAR* Section, const FConfigFile* File)
 		else
 		{
 			// Section is a short name
-			FString LongName = FString(TEXT("/Script/")) + Section;
+			FString LongName = FString("/Script/") + Section;
 			if (File->FindSection(*LongName))
 			{
 				UE_LOG(LogConfig, Fatal, TEXT("Short config section used instead of long %s"), Section);
@@ -1824,7 +1824,7 @@ void FillFileFromBuffer(FileType* File, FStringView Buffer, bool bHandleSymbolCo
 {
 	LLM_SCOPE(ELLMTag::ConfigSystem);
 
-	static const FName ConfigFileClassName = TEXT("ConfigFile");
+	static const FName ConfigFileClassName = "ConfigFile";
 	const FName FileName = FName(*FileHint);
 	LLM_SCOPE_DYNAMIC_STAT_OBJECTPATH_FNAME(FileName, ELLMTagSet::Assets);
 	LLM_SCOPE_DYNAMIC_STAT_OBJECTPATH_FNAME(ConfigFileClassName, ELLMTagSet::AssetClasses);
diff --git a/Engine/Source/Runtime/Core/Private/Misc/ConfigContext.cpp b/Engine/Source/Runtime/Core/Private/Misc/ConfigContext.cpp
--- a/Engine/Source/Runtime/Core/Private/Misc/ConfigContext.cpp
+++ b/Engine/Source/Runtime/Core/Private/Misc/ConfigContext.cpp
@@ -20,11 +20,11 @@ extern double GPerformLoadTime;
 
 namespace
 {
-	FString VersionName = TEXT("Version");
-	FString PreserveName = TEXT("Preserve");
-	FString LegacyIniVersionString = TEXT("IniVersion");
-	FString LegacyEngineString = TEXT("Engine.Engine");
-	FString CurrentIniVersionString = TEXT("CurrentIniVersion");
+	FString VersionName = "Version";
+	FString PreserveName = "Preserve";
+	FString LegacyIniVersionString = "IniVersion";
+	FString LegacyEngineString = "Engine.Engine";
+	FString CurrentIniVersionString = "CurrentIniVersion";
 	const TCHAR* SectionsToSaveString = TEXT("SectionsToSave");
 	const TCHAR* SaveAllSectionsKey = TEXT("bCanSaveAllSections");
 
@@ -164,7 +164,7 @@ const FConfigContext::FPerPlatformDirs& FConfigContext::GetPerPlatformDirs(const
 	FConfigContext::FPerPlatformDirs* Dirs = FConfigContext::PerPlatformDirs.Find(PlatformName);
 	if (Dirs == nullptr)
 	{
-		FString PluginExtDir = TEXT("<skip>");
+		FString PluginExtDir = "<skip>";
 		if (bIsForPlugin)
 		{
 			// look if there's a plugin extension for this platform, it will have the platform name in the path
diff --git a/Engine/Source/Runtime/Core/Private/Misc/EmbeddedCommunication.cpp b/Engine/Source/Runtime/Core/Private/Misc/EmbeddedCommunication.cpp
--- a/Engine/Source/Runtime/Core/Private/Misc/EmbeddedCommunication.cpp
+++ b/Engine/Source/Runtime/Core/Private/Misc/EmbeddedCommunication.cpp
@@ -445,7 +445,7 @@ FString FEmbeddedCommunication::GetDebugInfo()
 #if BUILD_EMBEDDED_APP
 	FScopeLock Lock(&GEmbeddedLock);
 
-	FString Str = "";
+	FString Str;
 	for (auto It : GRenderingWakeMap)
 	{
 		Str += FString::Printf(TEXT("%s:%d "), *It.Key.ToString(), It.Value);
diff --git a/Engine/Source/Runtime/Core/Private/Misc/MonitoredProcess.cpp b/Engine/Source/Runtime/Core/Private/Misc/MonitoredProcess.cpp
--- a/Engine/Source/Runtime/Core/Private/Misc/MonitoredProcess.cpp
+++ b/Engine/Source/Runtime/Core/Private/Misc/MonitoredProcess.cpp
@@ -382,11 +382,11 @@ bool FSerializedUATProcess::LaunchInternal()
 FString FSerializedUATProcess::GetUATPath()
 {
 #if PLATFORM_WINDOWS
-	FString RunUATScriptName = TEXT("RunUAT.bat");
+	FString RunUATScriptName = "RunUAT.bat";
 #elif PLATFORM_LINUX
-	FString RunUATScriptName = TEXT("RunUAT.sh");
+	FString RunUATScriptName = "RunUAT.sh";
 #else
-	FString RunUATScriptName = TEXT("RunUAT.command");
+	FString RunUATScriptName = "RunUAT.command";
 #endif
 
 	return FPaths::ConvertRelativePathToFull(FPaths::EngineDir() / TEXT("Build/BatchFiles") / RunUATScriptName);
diff --git a/Engine/Source/Runtime/Core/Private/Misc/NetworkVersion.cpp b/Engine/Source/Runtime/Core/Private/Misc/NetworkVersion.cpp
--- a/Engine/Source/Runtime/Core/Private/Misc/NetworkVersion.cpp
+++ b/Engine/Source/Runtime/Core/Private/Misc/NetworkVersion.cpp
@@ -87,7 +87,7 @@ FCustomVersionRegistration GRegisterGameNetworkCustomVersion(FGameNetworkCustomV
 
 FString& FNetworkVersion::GetProjectVersion_Internal()
 {
-	static FString ProjectVersion = TEXT("1.0.0");
+	static FString ProjectVersion = "1.0.0";
 	return ProjectVersion;
 }
 
diff --git a/Engine/Source/Runtime/Core/Private/Misc/PathViews.cpp b/Engine/Source/Runtime/Core/Private/Misc/PathViews.cpp
--- a/Engine/Source/Runtime/Core/Private/Misc/PathViews.cpp
+++ b/Engine/Source/Runtime/Core/Private/Misc/PathViews.cpp
@@ -1128,7 +1128,7 @@ FStringView FPathViews::GetMountPointNameFromPath(const FStringView InPath, bool
 			MountPointStringView = FStringView(InPath.GetData() + WithoutSlashes, InPath.Len() - WithoutSlashes);
 		}
 
-		static const FString ClassesPrefix = TEXT("Classes_");
+		static const FString ClassesPrefix = "Classes_";
 		const bool bHasClassesPrefix = MountPointStringView.StartsWith(ClassesPrefix);
 		if (bHasClassesPrefix)
 		{
diff --git a/Engine/Source/Runtime/Core/Private/Misc/Paths.cpp b/Engine/Source/Runtime/Core/Private/Misc/Paths.cpp
--- a/Engine/Source/Runtime/Core/Private/Misc/Paths.cpp
+++ b/Engine/Source/Runtime/Core/Private/Misc/Paths.cpp
@@ -806,7 +806,7 @@ bool FPaths::IsRestrictedPath(const FString& InPath)
 	if (!StaticData.bRestrictedSlashedFolderNamesInitialized)
 	{
 		// Add leading and trailing slashes to restricted folder names.
-		FString LeadingSlash(TEXT("/"));
+		FString LeadingSlash("/");
 
 		for (const FString& FolderStr : GetRestrictedFolderNames())
 		{
@@ -1161,7 +1161,7 @@ bool FPaths::IsDrive(const FString& InPath)
 			TempPath += TEXT("\\");
 		}
 
-		FString CheckPath = TEXT("");
+		FString CheckPath;
 		int32 ColonSlashIndex = TempPath.Find(TEXT(":\\"), ESearchCase::CaseSensitive);
 		if (ColonSlashIndex != INDEX_NONE)
 		{
diff --git a/Engine/Source/Runtime/Core/Private/Misc/Timespan.cpp b/Engine/Source/Runtime/Core/Private/Misc/Timespan.cpp
--- a/Engine/Source/Runtime/Core/Private/Misc/Timespan.cpp
+++ b/Engine/Source/Runtime/Core/Private/Misc/Timespan.cpp
@@ -161,7 +161,7 @@ bool FTimespan::Parse(const FString& TimespanString, FTimespan& OutTimespan)
 		}
 		else if (FractionalLen < 9)
 		{
-			Tokens[4] += FString(TEXT("000000000")).Left(9 - FractionalLen);
+			Tokens[4] += FString("000000000").Left(9 - FractionalLen);
 		}
 	}
 
diff --git a/Engine/Source/Runtime/Core/Private/Misc/UProjectInfo.cpp b/Engine/Source/Runtime/Core/Private/Misc/UProjectInfo.cpp
--- a/Engine/Source/Runtime/Core/Private/Misc/UProjectInfo.cpp
+++ b/Engine/Source/Runtime/Core/Private/Misc/UProjectInfo.cpp
@@ -22,7 +22,7 @@ void FUProjectDictionary::Refresh()
 
 	// Find all the .uprojectdirs files contained in the root folder
 	TArray<FString> ProjectDirsFiles;
-	IFileManager::Get().FindFiles(ProjectDirsFiles, *(RootDir / FString(TEXT("*.uprojectdirs"))), true, false);
+	IFileManager::Get().FindFiles(ProjectDirsFiles, *(RootDir / FString("*.uprojectdirs")), true, false);
 
 	// Get the normalized path to the root directory
 	FString NormalizedRootDir = FPaths::ConvertRelativePathToFull(RootDir);
@@ -60,7 +60,7 @@ void FUProjectDictionary::Refresh()
 	{
 		// Enumerate the subdirectories
 		TArray<FString> ProjectDirs;
-		IFileManager::Get().FindFiles(ProjectDirs, *(ProjectRootDir / FString(TEXT("*"))), false, true);
+		IFileManager::Get().FindFiles(ProjectDirs, *(ProjectRootDir / FString("*")), false, true);
 
 		// Check each one for project files
 		for(const FString& ProjectDir: ProjectDirs)
diff --git a/Engine/Source/Runtime/Core/Private/Modules/ModuleManager.cpp b/Engine/Source/Runtime/Core/Private/Modules/ModuleManager.cpp
--- a/Engine/Source/Runtime/Core/Private/Modules/ModuleManager.cpp
+++ b/Engine/Source/Runtime/Core/Private/Modules/ModuleManager.cpp
@@ -966,7 +966,7 @@ IModuleInterface* FModuleManager::LoadModuleWithFailureReason(const FName InModu
 
 #if	STATS
 	// This is fine here, we only load a handful of modules.
-	static FString Module = TEXT( "Module" );
+	static FString Module = "Module";
 	const FString LongName = Module / InModuleName.GetPlainNameString();
 	const TStatId StatId = FDynamicStats::CreateStatId<FStatGroup_STATGROUP_UObjects>( LongName );
 	FScopeCycleCounter CycleCounter( StatId );
@@ -2064,7 +2064,7 @@ bool FModuleManager::DoesLoadedModuleHaveUObjects( const FName ModuleName ) cons
 {
 	if (IsModuleLoaded(ModuleName) && IsPackageLoaded.IsBound())
 	{
-		return IsPackageLoaded.Execute(*FString(FString(TEXT("/Script/")) + ModuleName.ToString()));
+		return IsPackageLoaded.Execute(*FString(FString("/Script/") + ModuleName.ToString()));
 	}
 
 	return false;
diff --git a/Engine/Source/Runtime/Core/Private/PosixOS/PosixOSPlatformProcess.cpp b/Engine/Source/Runtime/Core/Private/PosixOS/PosixOSPlatformProcess.cpp
--- a/Engine/Source/Runtime/Core/Private/PosixOS/PosixOSPlatformProcess.cpp
+++ b/Engine/Source/Runtime/Core/Private/PosixOS/PosixOSPlatformProcess.cpp
@@ -24,7 +24,7 @@ bool FPosixOSPlatformProcess::IsFirstInstance()
 		// create the file if it doesn't exist
 		if (GFileLockDescriptor == -1)
 		{
-			FString LockFileName(TEXT("/tmp/"));
+			FString LockFileName("/tmp/");
 			FString ExecPath(FPlatformProcess::ExecutableName());
 			ExecPath.ReplaceInline(TEXT("/"), TEXT("-"), ESearchCase::CaseSensitive);
 			// [RCL] 2015-09-20: can run out of filename limits (256 bytes) due to a long path, be conservative and assume 4-char UTF-8 name like e.g. Japanese
diff --git a/Engine/Source/Runtime/Core/Private/ProfilingDebugging/ABTesting.cpp b/Engine/Source/Runtime/Core/Private/ProfilingDebugging/ABTesting.cpp
--- a/Engine/Source/Runtime/Core/Private/ProfilingDebugging/ABTesting.cpp
+++ b/Engine/Source/Runtime/Core/Private/ProfilingDebugging/ABTesting.cpp
@@ -249,18 +249,18 @@ FABTest& FABTest::Get()
 void FABTest::ABTestCmdFunc(const TArray<FString>& Args)
 {
 	FString ABTestCmds[2];
-	if (Args.Num() == 1 && Args[0].Compare(FString(TEXT("framelog")), ESearchCase::IgnoreCase) == 0)
+	if (Args.Num() == 1 && Args[0].Compare(FString("framelog"), ESearchCase::IgnoreCase) == 0)
 	{
 		Get().Stop();
 		Get().StartFrameLog();
 		return;
 	}
-	if (Args.Num() == 1 && Args[0].Compare(FString(TEXT("stop")), ESearchCase::IgnoreCase) == 0)
+	if (Args.Num() == 1 && Args[0].Compare(FString("stop"), ESearchCase::IgnoreCase) == 0)
 	{
 		Get().Stop();		
 		return;
 	}
-	if (Args.Num() == 1 && Args[0].Compare(FString(TEXT("scope")), ESearchCase::IgnoreCase) == 0)
+	if (Args.Num() == 1 && Args[0].Compare(FString("scope"), ESearchCase::IgnoreCase) == 0)
 	{
 		ABTestCmds[0] = FString("ScopeA");
 		ABTestCmds[1] = FString("ScopeB");
@@ -339,8 +339,8 @@ void FABTest::Start(FString* InABTestCmds, bool bScopeTest)
 	}
 	else
 	{
-		ABTestCmds[0] = FString("");
-		ABTestCmds[1] = FString("");
+		ABTestCmds[0] = FString();
+		ABTestCmds[1] = FString();
 	}
 	
 	bABScopeTestActive = bScopeTest;
diff --git a/Engine/Source/Runtime/Core/Private/ProfilingDebugging/ProfilingHelpers.cpp b/Engine/Source/Runtime/Core/Private/ProfilingDebugging/ProfilingHelpers.cpp
--- a/Engine/Source/Runtime/Core/Private/ProfilingDebugging/ProfilingHelpers.cpp
+++ b/Engine/Source/Runtime/Core/Private/ProfilingDebugging/ProfilingHelpers.cpp
@@ -156,7 +156,7 @@ FString CreateProfileFilename( const FString& InFilename, const FString& InFileE
 
 	/** This is meant to hold the name of the "sessions" that is occurring **/
 	static bool bSetProfilingSessionFolderName = false;
-	static FString ProfilingSessionFolderName = TEXT(""); 
+	static FString ProfilingSessionFolderName; 
 
 	// here we want to have just the same profiling session name so all of the files will go into that folder over the course of the run so you don't just have a ton of folders
 	FString FolderName;
diff --git a/Engine/Source/Runtime/Core/Private/Stats/StatsCommand.cpp b/Engine/Source/Runtime/Core/Private/Stats/StatsCommand.cpp
--- a/Engine/Source/Runtime/Core/Private/Stats/StatsCommand.cpp
+++ b/Engine/Source/Runtime/Core/Private/Stats/StatsCommand.cpp
@@ -167,7 +167,7 @@ struct FBudgetData
 	/** Builds any extra meta data from the stats provided **/
 	void Process()
 	{
-		FString ChildPrefix(TEXT("-"));
+		FString ChildPrefix("-");
 		for (FString& Stat : Stats)
 		{
 			if (Stat.RemoveFromStart(ChildPrefix))
@@ -249,7 +249,7 @@ struct FStatSlowParams : public FStatParams
 	FStatSlowParams( const TCHAR* Cmd = nullptr )
 		: FStatParams( Cmd )
 	{
-		static FName NAME_Slow = TEXT( "Slow" );
+		static FName NAME_Slow = "Slow";
 		Group = TParsedValueWithDefault<FName>( nullptr, nullptr, NAME_Slow );
 		CullMs = TParsedValueWithDefault<float>( Cmd, TEXT( "ms=" ), 1.0f );
 		MaxHierarchyDepth = TParsedValueWithDefault<int32>( Cmd, TEXT( "maxdepth=" ), 4 );
@@ -753,7 +753,7 @@ static bool HandleToggleCommandBroadcast(const FName& InStatName, bool& bOutCurr
 	bOutOthersEnabled = false;
 
 	// Check to see if all stats have been disabled... 
-	static const FName NAME_NoGroup = FName(TEXT("STATGROUP_None"));
+	static const FName NAME_NoGroup = FName("STATGROUP_None");
 	if (InStatName == NAME_NoGroup)
 	{
 		// Iterate through all enabled groups.
@@ -955,7 +955,7 @@ struct FHUDGroupManager
 
 		ResizeFramesHistory( Params.MaxHistoryFrames.Get() );
 
-		const FName MaybeGroupFName = FName(*(FString(TEXT("STATGROUP_")) + Params.Group.Get().GetPlainNameString()));
+		const FName MaybeGroupFName = FName(*(FString("STATGROUP_") + Params.Group.Get().GetPlainNameString()));
 		const bool bResults = HandleToggleCommandBroadcast( MaybeGroupFName, bCurrentEnabled, bOthersEnabled );
 		if (!bResults)
 		{
@@ -1235,7 +1235,7 @@ struct FHUDGroupManager
 				
 				if(MergedTaskGraphThreads.Num())
 				{
-					NewFrame.InclusiveAggregateThreadBreakdown.Add(FName(TEXT("MergedTaskGraphThreads")), MergedTaskGraphThreads);
+					NewFrame.InclusiveAggregateThreadBreakdown.Add(FName("MergedTaskGraphThreads"), MergedTaskGraphThreads);
 				}
 			}
 		}
@@ -1944,7 +1944,7 @@ static void StatCmd(FString InCmd, bool bStatCommand, FOutputDevice* Ar /*= null
 			FString MaybeGroup;
 			FParse::Token(Cmd, MaybeGroup, false);
 
-			DumpNonFrame(Stats, MaybeGroup.Len() == 0 ? NAME_None : FName(*(FString(TEXT("STATGROUP_")) + MaybeGroup)));
+			DumpNonFrame(Stats, MaybeGroup.Len() == 0 ? NAME_None : FName(*(FString("STATGROUP_") + MaybeGroup)));
 		}
 		else if (FParse::Command(&Cmd, TEXT("DUMPCPU")))
 		{
@@ -2393,7 +2393,7 @@ namespace UE::Stats
 							MaybeGroup.RemoveAt(PlusPos, EAllowShrinking::No);
 						}
 
-						const FName MaybeGroupFName = FName(*(FString(TEXT("STATGROUP_")) + MaybeGroup));
+						const FName MaybeGroupFName = FName(*(FString("STATGROUP_") + MaybeGroup));
 						bResult = FStatGroupGameThreadNotifier::Get().StatGroupNames.Contains(MaybeGroupFName);
 					}
 				}
@@ -2417,7 +2417,7 @@ namespace UE::Stats
 							if (!Line.Contains(TEXT("+Stats=")))	//ignore stats array
 							{
 								FString ThreadName;
-								Line.Split(FString(TEXT("=")), &ThreadName, nullptr);
+								Line.Split(FString("="), &ThreadName, nullptr);
 								float Budget = -1.f;
 								if (GConfig->GetFloat(*BudgetSection, *ThreadName, Budget, GEngineIni))
 								{
diff --git a/Engine/Source/Runtime/Core/Private/Stats/StatsData.cpp b/Engine/Source/Runtime/Core/Private/Stats/StatsData.cpp
--- a/Engine/Source/Runtime/Core/Private/Stats/StatsData.cpp
+++ b/Engine/Source/Runtime/Core/Private/Stats/StatsData.cpp
@@ -27,7 +27,7 @@ const FName FStatConstants::NAME_ThreadRoot = "ThreadRoot";
 const char* FStatConstants::ThreadGroupName = STAT_GROUP_TO_FStatGroup( STATGROUP_Threads )::GetGroupName();
 const FName FStatConstants::NAME_ThreadGroup = FStatConstants::ThreadGroupName;
 const FName FStatConstants::RAW_SecondsPerCycle = FStatNameAndInfo( GET_STATFNAME( STAT_SecondsPerCycle ), true ).GetRawName();
-const FName FStatConstants::NAME_NoCategory = FName(TEXT("STATCAT_None"));
+const FName FStatConstants::NAME_NoCategory = FName("STATCAT_None");
 
 const FString FStatConstants::StatsFileExtension = TEXT( ".uestats" );
 const FString FStatConstants::StatsFileRawExtension = TEXT( ".uestatsraw" );
@@ -265,7 +265,7 @@ void FRawStatStackNode::AddNameHierarchy(int32 CurrentPrefixDepth)
 						MaxCommonality = TestDepth + 1;
 					}
 				}
-				FString NewName(TEXT("NameFolder//"));
+				FString NewName("NameFolder//");
 				for (int32 TestDepth = 0; TestDepth < MaxCommonality; TestDepth++)
 				{
 					NewName += ChildNames[StartIndex][TestDepth].ToString();
@@ -1679,7 +1679,7 @@ FName FStatsThreadState::GetStatThreadName( const FStatPacket& Packet ) const
 	}
 	else if ( Packet.ThreadType == EThreadType::EndOfPipe )
 	{
-		static const FName NAME_EndOfPipe = TEXT("EndOfPipe");
+		static const FName NAME_EndOfPipe = "EndOfPipe";
 		ThreadName = NAME_EndOfPipe;
 	}
 	else if( Packet.ThreadType == EThreadType::Other )
@@ -1690,7 +1690,7 @@ FName FStatsThreadState::GetStatThreadName( const FStatPacket& Packet ) const
 		{
 			UE_LOG( LogStats, Display, TEXT( "There is no thread with id: %u. Please add thread metadata for this thread." ), Packet.ThreadId );
 
-			static const FName NAME_UnknownThread = TEXT( "UnknownThread" );
+			static const FName NAME_UnknownThread = "UnknownThread";
 			NewThreadName = FName( *FStatsUtils::BuildUniqueThreadName( Packet.ThreadId ) );
 
 			UE_AUTORTFM_OPEN
@@ -1863,7 +1863,7 @@ void FStatsUtils::GetNameAndGroup(FStatMessage const& Item, FString& OutName, FS
 
 FString FStatsUtils::DebugPrint(FStatMessage const& Item)
 {
-	FString Result(TEXT("Invalid"));
+	FString Result("Invalid");
 	switch (Item.NameAndInfo.GetField<EStatDataType>())
 	{
 	case EStatDataType::ST_int64:
diff --git a/Engine/Source/Runtime/Core/Private/Stats/StatsSystemTypes.cpp b/Engine/Source/Runtime/Core/Private/Stats/StatsSystemTypes.cpp
--- a/Engine/Source/Runtime/Core/Private/Stats/StatsSystemTypes.cpp
+++ b/Engine/Source/Runtime/Core/Private/Stats/StatsSystemTypes.cpp
@@ -587,7 +587,7 @@ public:
 	{
 		FString MaybeGroup;
 		FParse::Token(Cmd, MaybeGroup, false);
-		MaybeGroup = FString(TEXT("STATGROUP_")) + MaybeGroup;
+		MaybeGroup = FString("STATGROUP_") + MaybeGroup;
 		FName MaybeGroupFName(*MaybeGroup);
 
 		FGroupEnable* Found = HighPerformanceEnable.Find(MaybeGroupFName);
@@ -1205,7 +1205,7 @@ void FThreadStats::FlushRawStats(bool bHasBrokenCallstacks /*= false*/, bool bFo
 void FThreadStats::CheckForCollectingStartupStats()
 {
 	FString CmdLine(FCommandLine::Get());
-	FString StatCmds(TEXT("-StatCmds="));
+	FString StatCmds("-StatCmds=");
 	while (1)
 	{
 		FString Cmds;
diff --git a/Engine/Source/Runtime/Core/Private/Unix/UnixPlatformCrashContext.cpp b/Engine/Source/Runtime/Core/Private/Unix/UnixPlatformCrashContext.cpp
--- a/Engine/Source/Runtime/Core/Private/Unix/UnixPlatformCrashContext.cpp
+++ b/Engine/Source/Runtime/Core/Private/Unix/UnixPlatformCrashContext.cpp
@@ -168,7 +168,7 @@ void GracefulTerminationHandler(int32 Signal, siginfo_t* Info, void* Context)
 
 void CreateExceptionInfoString(int32 Signal, siginfo_t* Info, ucontext_t *Context)
 {
-	FString ErrorString = TEXT("Unhandled Exception: ");
+	FString ErrorString = "Unhandled Exception: ";
 	ErrorString += DescribeSignal(Signal, Info, Context);
 	FCString::Strncpy(GErrorExceptionDescription, *ErrorString, FMath::Min(ErrorString.Len() + 1, (int32)UE_ARRAY_COUNT(GErrorExceptionDescription)));
 }
@@ -696,13 +696,13 @@ void FUnixCrashContext::GenerateCrashInfoAndLaunchReporter() const
 					}
 					else
 					{
-						FFileHelper::SaveStringToFile(FString(TEXT("Log not available, too large for submission to crash reporting backend")), *LogDstAbsolute);
+						FFileHelper::SaveStringToFile(FString("Log not available, too large for submission to crash reporting backend"), *LogDstAbsolute);
 					}
 				}
 			}
 			else if (TotalDiskFreeSpace >= MaxFileSizeForCrashLog)
 			{
-				FFileHelper::SaveStringToFile(FString(TEXT("Log not available, server has low available disk space")), *LogDstAbsolute);
+				FFileHelper::SaveStringToFile(FString("Log not available, server has low available disk space"), *LogDstAbsolute);
 			}
 #endif // !NO_LOGGING
 
@@ -721,7 +721,7 @@ void FUnixCrashContext::GenerateCrashInfoAndLaunchReporter() const
 #endif
 			FString CrashReportLogFilename = LogBaseFilename + TEXT("-CRC") + LogExtension;
 			FString CrashReportLogFilepath = FPaths::Combine(*LogFolder, *CrashReportLogFilename);
-			FString CrashReportClientArguments = TEXT(" -Abslog=");
+			FString CrashReportClientArguments = " -Abslog=";
 			CrashReportClientArguments += TEXT("\"\"") + CrashReportLogFilepath + TEXT("\"\"");
 			CrashReportClientArguments += TEXT(" ");
 
diff --git a/Engine/Source/Runtime/Core/Private/Unix/UnixPlatformFile.cpp b/Engine/Source/Runtime/Core/Private/Unix/UnixPlatformFile.cpp
--- a/Engine/Source/Runtime/Core/Private/Unix/UnixPlatformFile.cpp
+++ b/Engine/Source/Runtime/Core/Private/Unix/UnixPlatformFile.cpp
@@ -700,7 +700,7 @@ public:
 				return FString();
 			}
 
-			FString FoundFilename(TEXT("/"));	// Start from the root.
+			FString FoundFilename("/");	// Start from the root.
 			bFound = MapFileRecursively(PossiblyWrongFilename, 0, MaxPathComponents, FoundFilename);
 			if (bFound)
 			{
@@ -764,7 +764,7 @@ public:
 				int MaxPathComponents = CountPathComponents(Filename);
 				if (MaxPathComponents > 0)
 				{
-					FString FoundFilename(TEXT("/"));	// start with root
+					FString FoundFilename("/");	// start with root
 					if (MapFileRecursively(Filename, 0, MaxPathComponents, FoundFilename))
 					{
 						Handle = open(TCHAR_TO_UTF8(*FoundFilename), O_RDONLY | O_CLOEXEC);
diff --git a/Engine/Source/Runtime/Core/Private/Unix/UnixPlatformMisc.cpp b/Engine/Source/Runtime/Core/Private/Unix/UnixPlatformMisc.cpp
--- a/Engine/Source/Runtime/Core/Private/Unix/UnixPlatformMisc.cpp
+++ b/Engine/Source/Runtime/Core/Private/Unix/UnixPlatformMisc.cpp
@@ -509,7 +509,7 @@ namespace
 void FUnixPlatformMisc::GetOSVersions(FString& out_OSVersionLabel, FString& out_OSSubVersionLabel)
 {
 	// Set up Fallback Details
-	out_OSVersionLabel = FString(TEXT("GenericLinuxVersion"));
+	out_OSVersionLabel = FString("GenericLinuxVersion");
 	// Get Kernel Version
 	out_OSSubVersionLabel = GetKernelVersion();
 
@@ -1234,7 +1234,7 @@ uint32 FUnixPlatformMisc::GetCPUInfo()
 FString FUnixPlatformMisc::GetCPUBrand()
 {
 #if PLATFORM_HAS_CPUID
-	static FString Result = TEXT("NonX86CPUBrand");
+	static FString Result = "NonX86CPUBrand";
 	static bool bHaveResult = false;
 
 	if (!bHaveResult)
@@ -1818,7 +1818,7 @@ bool FUnixPlatformMisc::SetStoredValues(const FString& InStoreId, const FString&
 	check(!InStoreId.IsEmpty());
 	check(!InSectionName.IsEmpty());
 
-	const FString ConfigPath = FString(FPlatformProcess::ApplicationSettingsDir()) / InStoreId / FString(TEXT("KeyValueStore.ini"));
+	const FString ConfigPath = FString(FPlatformProcess::ApplicationSettingsDir()) / InStoreId / FString("KeyValueStore.ini");
 
 	FConfigFile ConfigFile;
 	ConfigFile.Read(ConfigPath);
@@ -2128,7 +2128,7 @@ bool FUnixPlatformMisc::StopPGOCollectionAndCloseFile()
 	fflush(stdout);
 
 	// write out a fake file to make sure the previous file is no longer kept open by the LLVM machinery
-	FString DummyProfileFileName = TEXT("/tmp/");
+	FString DummyProfileFileName = "/tmp/";
 	DummyProfileFileName += FGuid::NewGuid().ToString();
 	FPlatformMisc::SetEnvironmentVar(TEXT("LLVM_PROFILE_FILE"), *DummyProfileFileName);
 	__llvm_profile_reset_counters();
diff --git a/Engine/Source/Runtime/Core/Private/Unix/UnixPlatformProcess.cpp b/Engine/Source/Runtime/Core/Private/Unix/UnixPlatformProcess.cpp
--- a/Engine/Source/Runtime/Core/Private/Unix/UnixPlatformProcess.cpp
+++ b/Engine/Source/Runtime/Core/Private/Unix/UnixPlatformProcess.cpp
@@ -489,7 +489,7 @@ FString FUnixPlatformProcess::GenerateApplicationPath( const FString& AppName, E
 
 FString FUnixPlatformProcess::GetApplicationName( uint32 ProcessId )
 {
-	FString Output = TEXT("");
+	FString Output;
 
 	const int32 ReadLinkSize = 1024;	
 	char ReadLinkCmd[ReadLinkSize] = {0};
@@ -2015,7 +2015,7 @@ bool FUnixPlatformProcess::IsApplicationRunning( uint32 ProcessId )
 
 bool FUnixPlatformProcess::IsApplicationRunning( const TCHAR* ProcName )
 {
-	FString Commandline = TEXT("pidof '");
+	FString Commandline = "pidof '";
 	Commandline += ProcName;
 	Commandline += TEXT("'  > /dev/null");
 	return !system(TCHAR_TO_UTF8(*Commandline));
diff --git a/Engine/Source/Runtime/Core/Private/Windows/WindowsPlatformCrashContext.cpp b/Engine/Source/Runtime/Core/Private/Windows/WindowsPlatformCrashContext.cpp
--- a/Engine/Source/Runtime/Core/Private/Windows/WindowsPlatformCrashContext.cpp
+++ b/Engine/Source/Runtime/Core/Private/Windows/WindowsPlatformCrashContext.cpp
@@ -1041,13 +1041,13 @@ int32 ReportCrashUsingCrashReportClient(FWindowsPlatformCrashContext& InContext,
 				CrashReportClientArguments += TEXT(" -nullrhi");
 			}
 
-			CrashReportClientArguments += FString(TEXT(" -AppName=")) + AppName;
-			CrashReportClientArguments += FString(TEXT(" -CrashGUID=")) + CrashGUID;
+			CrashReportClientArguments += FString(" -AppName=") + AppName;
+			CrashReportClientArguments += FString(" -CrashGUID=") + CrashGUID;
 
 			const FString DownstreamStorage = FWindowsPlatformStackWalk::GetDownstreamStorage();
 			if (!DownstreamStorage.IsEmpty())
 			{
-				CrashReportClientArguments += FString(TEXT(" -DebugSymbols=")) + DownstreamStorage;
+				CrashReportClientArguments += FString(" -DebugSymbols=") + DownstreamStorage;
 			}
 
 			// CrashReportClient.exe should run without dragging in binaries from an inherited dll directory
@@ -1146,7 +1146,7 @@ THIRD_PARTY_INCLUDES_END
 void CreateExceptionInfoString(EXCEPTION_RECORD* ExceptionRecord, TCHAR* OutErrorString, int32 ErrorStringBufSize)
 {
 	// #CrashReport: 2014-08-18 Fix FString usage?
-	FString ErrorString = TEXT("Unhandled Exception: ");
+	FString ErrorString = "Unhandled Exception: ";
 
 #define HANDLE_CASE(x) case x: ErrorString += TEXT(#x); break;
 
diff --git a/Engine/Source/Runtime/Core/Private/Windows/WindowsPlatformMisc.cpp b/Engine/Source/Runtime/Core/Private/Windows/WindowsPlatformMisc.cpp
--- a/Engine/Source/Runtime/Core/Private/Windows/WindowsPlatformMisc.cpp
+++ b/Engine/Source/Runtime/Core/Private/Windows/WindowsPlatformMisc.cpp
@@ -132,8 +132,8 @@ int32 GetOSVersionsHelper( TCHAR* OutOSVersionLabel, int32 OSVersionLabelLength,
 
 	OSVERSIONINFOEX OsVersionInfo = {0};
 	OsVersionInfo.dwOSVersionInfoSize = sizeof( OSVERSIONINFOEX );
-	FString OSVersionLabel = TEXT("Windows (unknown version)");
-	FString OSSubVersionLabel = TEXT("");
+	FString OSVersionLabel = "Windows (unknown version)";
+	FString OSSubVersionLabel;
 #ifdef __clang__
 #pragma clang diagnostic push
 #pragma clang diagnostic ignored "-Wdeprecated-declarations"
@@ -2603,7 +2603,7 @@ const TCHAR* FWindowsPlatformMisc::GetPlatformFeaturesModuleName()
 	// If running a dedicated server then we use the default PlatformFeatures
 	if (!IsRunningDedicatedServer())
 	{
-		static FString PlatformFeaturesName = TEXT("WindowsPlatformFeatures");
+		static FString PlatformFeaturesName = "WindowsPlatformFeatures";
 		static bool bIniChecked = false;
 		if (!bIniChecked && !GEngineIni.IsEmpty())
 		{
@@ -2710,7 +2710,7 @@ bool FWindowsPlatformMisc::SetStoredValue(const FString& InStoreId, const FStrin
 	check(!InSectionName.IsEmpty());
 	check(!InKeyName.IsEmpty());
 
-	FString FullRegistryKey = FString(TEXT("Software")) / InStoreId / InSectionName;
+	FString FullRegistryKey = FString("Software") / InStoreId / InSectionName;
 	FullRegistryKey = FullRegistryKey.Replace(TEXT("/"), TEXT("\\")); // we use forward slashes, but the registry needs back slashes
 
 	HKEY hKey;
@@ -2738,7 +2738,7 @@ bool FWindowsPlatformMisc::GetStoredValue(const FString& InStoreId, const FStrin
 	check(!InSectionName.IsEmpty());
 	check(!InKeyName.IsEmpty());
 
-	FString FullRegistryKey = FString(TEXT("Software")) / InStoreId / InSectionName;
+	FString FullRegistryKey = FString("Software") / InStoreId / InSectionName;
 	FullRegistryKey = FullRegistryKey.Replace(TEXT("/"), TEXT("\\")); // we use forward slashes, but the registry needs back slashes
 
 	return QueryRegKey(HKEY_CURRENT_USER, *FullRegistryKey, *InKeyName, OutValue);
@@ -2752,7 +2752,7 @@ bool FWindowsPlatformMisc::DeleteStoredValue(const FString& InStoreId, const FSt
 	check(!InSectionName.IsEmpty());
 	check(!InKeyName.IsEmpty());
 
-	FString FullRegistryKey = FString(TEXT("Software")) / InStoreId / InSectionName;
+	FString FullRegistryKey = FString("Software") / InStoreId / InSectionName;
 	FullRegistryKey = FullRegistryKey.Replace(TEXT("/"), TEXT("\\")); // we use forward slashes, but the registry needs back slashes
 
 	HKEY hKey;
@@ -2790,7 +2790,7 @@ bool FWindowsPlatformMisc::DeleteStoredSection(const FString& InStoreId, const F
 	check(!InStoreId.IsEmpty());
 	check(!InSectionName.IsEmpty());
 
-	FString FullRegistryKey = FString(TEXT("Software")) / InStoreId / InSectionName;
+	FString FullRegistryKey = FString("Software") / InStoreId / InSectionName;
 	FullRegistryKey = FullRegistryKey.Replace(TEXT("/"), TEXT("\\")); // we use forward slashes, but the registry needs back slashes
 
 	return ::RegDeleteTree(HKEY_CURRENT_USER, *FullRegistryKey) == ERROR_SUCCESS;
@@ -3599,7 +3599,7 @@ static void GetVideoDriverDetailsFromSetup(const FString& DeviceName, bool bVerb
 	
 	HDEVINFO hDevInfo = SetupDiGetClassDevs(&GUID_DEVCLASS_DISPLAY, NULL, NULL, DIGCF_PRESENT);
 
-	FString RegistryKey = "";
+	FString RegistryKey;
 	
 	if (hDevInfo != INVALID_HANDLE_VALUE)
 	{
@@ -3728,7 +3728,7 @@ static void GetVideoDriverDetailsFromSetup(const FString& DeviceName, bool bVerb
 		
 		if(FWindowsPlatformMisc::QueryRegKey(HKEY_LOCAL_MACHINE, *Key, TEXT("Catalyst_Version"), Out.UserDriverVersion))
 		{
-			Out.UserDriverVersion = FString(TEXT("Catalyst ")) + Out.UserDriverVersion;
+			Out.UserDriverVersion = FString("Catalyst ") + Out.UserDriverVersion;
 		}
 
 		FString Edition;
@@ -3810,7 +3810,7 @@ static void GetVideoDriverDetails(const FString& Key, FGPUDriverInfo& Out)
 	{
 		if(FWindowsPlatformMisc::QueryRegKey(HKEY_LOCAL_MACHINE, *Key, TEXT("Catalyst_Version"), Out.UserDriverVersion))
 		{
-			Out.UserDriverVersion = FString(TEXT("Catalyst ")) + Out.UserDriverVersion;
+			Out.UserDriverVersion = FString("Catalyst ") + Out.UserDriverVersion;
 		}
 
 		FString Edition;
@@ -3929,7 +3929,7 @@ FGPUDriverInfo FWindowsPlatformMisc::GetGPUDriverInfo(const FString& DeviceDescr
 
 			if(DriverLocation.Left(18) == TEXT("\\Registry\\Machine\\"))		// not case sensitive
 			{
-				DriverLocation = FString(TEXT("\\HKEY_LOCAL_MACHINE\\")) + DriverLocation.RightChop(18);
+				DriverLocation = FString("\\HKEY_LOCAL_MACHINE\\") + DriverLocation.RightChop(18);
 			}
 			if(DriverLocation.Left(20) == TEXT("\\HKEY_LOCAL_MACHINE\\"))		// not case sensitive
 			{
@@ -4030,7 +4030,7 @@ FGPUDriverInfo FWindowsPlatformMisc::GetGPUDriverInfo(const FString& DeviceDescr
 		{
 			if(DriverLocation.Left(18) == TEXT("\\Registry\\Machine\\"))		// not case sensitive
 			{
-				DriverLocation = FString(TEXT("\\HKEY_LOCAL_MACHINE\\")) + DriverLocation.RightChop(18);
+				DriverLocation = FString("\\HKEY_LOCAL_MACHINE\\") + DriverLocation.RightChop(18);
 			}
 			if(DriverLocation.Left(20) == TEXT("\\HKEY_LOCAL_MACHINE\\"))		// not case sensitive
 			{
diff --git a/Engine/Source/Runtime/Core/Private/Windows/WindowsPlatformProcess.cpp b/Engine/Source/Runtime/Core/Private/Windows/WindowsPlatformProcess.cpp
--- a/Engine/Source/Runtime/Core/Private/Windows/WindowsPlatformProcess.cpp
+++ b/Engine/Source/Runtime/Core/Private/Windows/WindowsPlatformProcess.cpp
@@ -1027,7 +1027,7 @@ bool FWindowsPlatformProcess::IsApplicationRunning( const TCHAR* ProcName )
 
 FString FWindowsPlatformProcess::GetApplicationName( uint32 ProcessId )
 {
-	FString Output = TEXT("");
+	FString Output;
 	HANDLE ProcessHandle = ::OpenProcess(PROCESS_QUERY_INFORMATION, false, ProcessId);
 	if (ProcessHandle != NULL)
 	{
diff --git a/Engine/Source/Runtime/Core/Public/Containers/StringView.h b/Engine/Source/Runtime/Core/Public/Containers/StringView.h
--- a/Engine/Source/Runtime/Core/Public/Containers/StringView.h
+++ b/Engine/Source/Runtime/Core/Public/Containers/StringView.h
@@ -90,7 +90,7 @@ namespace UE::Core::Private
  * @code
  *	void UseString(FStringView InString);
  *
- *	FString String(TEXT("ABC"));
+ *	FString String("ABC");
  *	const TCHAR* CString = *String;
  *	TStringBuilder<16> StringBuilder;
  *	StringBuilder.Append(TEXT("ABC"));
diff --git a/Engine/Source/Runtime/Core/Public/HAL/IConsoleManager.h b/Engine/Source/Runtime/Core/Public/HAL/IConsoleManager.h
--- a/Engine/Source/Runtime/Core/Public/HAL/IConsoleManager.h
+++ b/Engine/Source/Runtime/Core/Public/HAL/IConsoleManager.h
@@ -879,7 +879,7 @@ public:
 	 */
 	static FName ModularFeatureName()
 	{
-		static const FName Name = TEXT("ConsoleCommandExecutor");
+		static const FName Name = "ConsoleCommandExecutor";
 		return Name;
 	}
 
diff --git a/Engine/Source/Runtime/Core/Public/HAL/ThreadManager.h b/Engine/Source/Runtime/Core/Public/HAL/ThreadManager.h
--- a/Engine/Source/Runtime/Core/Public/HAL/ThreadManager.h
+++ b/Engine/Source/Runtime/Core/Public/HAL/ThreadManager.h
@@ -61,8 +61,8 @@ public:
 	/** Returns the name of a thread given its TLS id */
 	inline static const FString& GetThreadName(uint32 ThreadId)
 	{
-		static FString GameThreadName(TEXT("GameThread"));
-		static FString RenderThreadName(TEXT("RenderThread"));
+		static FString GameThreadName("GameThread");
+		static FString RenderThreadName("RenderThread");
 		PRAGMA_DISABLE_DEPRECATION_WARNINGS
 		if (ThreadId == GGameThreadId)
 		{
diff --git a/Engine/Source/Runtime/Core/Public/IO/IoChunkEncoding.h b/Engine/Source/Runtime/Core/Public/IO/IoChunkEncoding.h
--- a/Engine/Source/Runtime/Core/Public/IO/IoChunkEncoding.h
+++ b/Engine/Source/Runtime/Core/Public/IO/IoChunkEncoding.h
@@ -24,7 +24,7 @@ enum class EIoEncryptionMethod : uint8
 /** Defines how the I/O chunk gets encoded into a set of compressed and encrypted block(s). */
 struct FIoChunkEncodingParams
 {
-	FName CompressionFormat = TEXT("Oodle");
+	FName CompressionFormat = "Oodle";
 	FMemoryView EncryptionKey;
 	uint32 BlockSize = (64 << 10);
 };
diff --git a/Engine/Source/Runtime/Core/Public/Math/BigInt.h b/Engine/Source/Runtime/Core/Public/Math/BigInt.h
--- a/Engine/Source/Runtime/Core/Public/Math/BigInt.h
+++ b/Engine/Source/Runtime/Core/Public/Math/BigInt.h
@@ -969,7 +969,7 @@ public:
 	 */
 	[[nodiscard]] FString ToString() const
 	{
-		FString Text(TEXT("0x"));
+		FString Text("0x");
 		int32 WordIndex;
 		for (WordIndex = NumWords - 1; WordIndex > 0; --WordIndex)
 		{
diff --git a/Engine/Source/Runtime/Core/Public/Math/Rotator.h b/Engine/Source/Runtime/Core/Public/Math/Rotator.h
--- a/Engine/Source/Runtime/Core/Public/Math/Rotator.h
+++ b/Engine/Source/Runtime/Core/Public/Math/Rotator.h
@@ -826,10 +826,10 @@ inline FString TRotator<T>::ToCompactString() const
 {
 	if( IsNearlyZero() )
 	{
-		return FString(TEXT("R(0)"));
+		return FString("R(0)");
 	}
 
-	FString ReturnString(TEXT("R("));
+	FString ReturnString("R(");
 	bool bIsEmptyString = true;
 	if( !FMath::IsNearlyZero(Pitch) )
 	{
@@ -840,7 +840,7 @@ inline FString TRotator<T>::ToCompactString() const
 	{
 		if( !bIsEmptyString )
 		{
-			ReturnString += FString(TEXT(", "));
+			ReturnString += FString(", ");
 		}
 		ReturnString += FString::Printf(TEXT("Y=%.2f"), Yaw);
 		bIsEmptyString = false;
@@ -849,12 +849,12 @@ inline FString TRotator<T>::ToCompactString() const
 	{
 		if( !bIsEmptyString )
 		{
-			ReturnString += FString(TEXT(", "));
+			ReturnString += FString(", ");
 		}
 		ReturnString += FString::Printf(TEXT("R=%.2f"), Roll);
 		bIsEmptyString = false;
 	}
-	ReturnString += FString(TEXT(")"));
+	ReturnString += FString(")");
 	return ReturnString;
 }
 
diff --git a/Engine/Source/Runtime/Core/Public/Math/Vector.h b/Engine/Source/Runtime/Core/Public/Math/Vector.h
--- a/Engine/Source/Runtime/Core/Public/Math/Vector.h
+++ b/Engine/Source/Runtime/Core/Public/Math/Vector.h
@@ -2375,10 +2375,10 @@ inline FString TVector<T>::ToCompactString() const
 {
     if(IsNearlyZero())
     {
-        return FString(TEXT("V(0)"));
+        return FString("V(0)");
     }
 
-    FString ReturnString(TEXT("V("));
+    FString ReturnString("V(");
     bool bIsEmptyString = true;
     if(!FMath::IsNearlyZero(X))
     {
@@ -2389,7 +2389,7 @@ inline FString TVector<T>::ToCompactString() const
     {
         if(!bIsEmptyString)
         {
-            ReturnString += FString(TEXT(", "));
+            ReturnString += FString(", ");
         }
         ReturnString += FString::Printf(TEXT("Y=%.2f"), Y);
         bIsEmptyString = false;
@@ -2398,12 +2398,12 @@ inline FString TVector<T>::ToCompactString() const
     {
         if(!bIsEmptyString)
         {
-            ReturnString += FString(TEXT(", "));
+            ReturnString += FString(", ");
         }
         ReturnString += FString::Printf(TEXT("Z=%.2f"), Z);
         bIsEmptyString = false;
     }
-    ReturnString += FString(TEXT(")"));
+    ReturnString += FString(")");
     return ReturnString;
 }
 
diff --git a/Engine/Source/Runtime/Core/Public/Misc/AssertionMacros.h b/Engine/Source/Runtime/Core/Public/Misc/AssertionMacros.h
--- a/Engine/Source/Runtime/Core/Public/Misc/AssertionMacros.h
+++ b/Engine/Source/Runtime/Core/Public/Misc/AssertionMacros.h
@@ -509,7 +509,7 @@ namespace UEAsserts_Private
 #define GET_FUNCTION_NAME_STRING_VIEW_CHECKED(ClassName, FunctionName) \
 	((void)sizeof(&ClassName::FunctionName), TEXTVIEW(#FunctionName))
 
-// Returns FName(TEXT("FunctionName")), while statically verifying that the function exists in ClassName
+// Returns FName("FunctionName"), while statically verifying that the function exists in ClassName
 // Handles overloaded functions by specifying the argument of the overload to use
 #define GET_FUNCTION_NAME_CHECKED_OneParam(ClassName, FunctionName, ArgType) \
 	(FName(GET_FUNCTION_NAME_STRING_CHECKED_OneParam(ClassName, FunctionName, ArgType)))
diff --git a/Engine/Source/Runtime/Core/Public/Misc/AxisDisplayInfo.h b/Engine/Source/Runtime/Core/Public/Misc/AxisDisplayInfo.h
--- a/Engine/Source/Runtime/Core/Public/Misc/AxisDisplayInfo.h
+++ b/Engine/Source/Runtime/Core/Public/Misc/AxisDisplayInfo.h
@@ -83,7 +83,7 @@ public:
 
 	static FName GetModularFeatureName()
 	{
-		static const FName Name = TEXT("AxisDisplayInfo");
+		static const FName Name = "AxisDisplayInfo";
 		return Name;
 	}
 
diff --git a/Engine/Source/Runtime/Core/Public/Misc/DataDrivenPlatformInfoRegistry.h b/Engine/Source/Runtime/Core/Public/Misc/DataDrivenPlatformInfoRegistry.h
--- a/Engine/Source/Runtime/Core/Public/Misc/DataDrivenPlatformInfoRegistry.h
+++ b/Engine/Source/Runtime/Core/Public/Misc/DataDrivenPlatformInfoRegistry.h
@@ -252,7 +252,7 @@ public:
 		default:
 			break;
 		}
-		static const FString EmptyString = TEXT("");
+		static const FString EmptyString;
 		return EmptyString;
 	}
 
diff --git a/Engine/Source/Runtime/Core/Public/Misc/IEngineCrypto.h b/Engine/Source/Runtime/Core/Public/Misc/IEngineCrypto.h
--- a/Engine/Source/Runtime/Core/Public/Misc/IEngineCrypto.h
+++ b/Engine/Source/Runtime/Core/Public/Misc/IEngineCrypto.h
@@ -16,7 +16,7 @@ struct IEngineCrypto : public IModularFeature
 	*/
 	static inline FName GetFeatureName()
 	{
-		static const FName Name(TEXT("EngineCryptoFeature"));
+		static const FName Name("EngineCryptoFeature");
 		return Name;
 	}
 
diff --git a/Engine/Source/Runtime/Core/Public/Stats/StatsFile.h b/Engine/Source/Runtime/Core/Public/Stats/StatsFile.h
--- a/Engine/Source/Runtime/Core/Public/Stats/StatsFile.h
+++ b/Engine/Source/Runtime/Core/Public/Stats/StatsFile.h
@@ -716,7 +716,7 @@ public:
 				}
 				else
 				{
-					TheFName = FName( TEXT( "Unknown FName" ) );
+					TheFName = FName( "Unknown FName" );
 					Number = 0;
 					UE_LOG( LogTemp, Warning, TEXT( "Missing FName Indexed: %d, %d" ), Index, Number );
 				}
@@ -737,7 +737,7 @@ public:
 			}
 			else
 			{
-				TheFName = FName( TEXT( "Unknown FName" ) );
+				TheFName = FName( "Unknown FName" );
 				Number = 0;
 				UE_LOG( LogTemp, Warning, TEXT( "Missing FName Indexed: %d, %d" ), Index, Number );
 			}
diff --git a/Engine/Source/Runtime/Core/Public/UObject/ScriptDelegates.h b/Engine/Source/Runtime/Core/Public/UObject/ScriptDelegates.h
--- a/Engine/Source/Runtime/Core/Public/UObject/ScriptDelegates.h
+++ b/Engine/Source/Runtime/Core/Public/UObject/ScriptDelegates.h
@@ -828,7 +828,7 @@ public:
 
 		if( IsBound() )
 		{
-			FString AllDelegatesString = TEXT( "[" );
+			FString AllDelegatesString = "[";
 			bool bAddComma = false;
 			for( typename InvocationListType::TConstIterator CurDelegate( InvocationList ); CurDelegate; ++CurDelegate )
 			{
diff --git a/Engine/Source/Runtime/Core/Tests/Algo/UniqueTest.cpp b/Engine/Source/Runtime/Core/Tests/Algo/UniqueTest.cpp
--- a/Engine/Source/Runtime/Core/Tests/Algo/UniqueTest.cpp
+++ b/Engine/Source/Runtime/Core/Tests/Algo/UniqueTest.cpp
@@ -34,10 +34,10 @@ TEST_CASE_NAMED(FUniqueTest, "System::Core::Algo::Unique", "[ApplicationContextM
 		CHECK_MESSAGE(TEXT("`Unique` with duplicates and unique items must return correct result"), (Array == TArray<int32>{1, 2, 3}));
 	}
 	{
-		FString Str = TEXT("aa");
+		FString Str = "aa";
 		Str = Str.Mid(0, Unique(Str));
 		CHECK_MESSAGE(TEXT("`Unique` on `FString` as an example of arbitrary random-access container must compile and return correct result"),
-			Str == FString(TEXT("a")));
+			Str == FString("a"));
 	}
 	{
 		int32 Array[] = { 1 };
@@ -78,10 +78,10 @@ TEST_CASE_NAMED(FUniqueByTest, "System::Core::Algo::UniqueBy", "[ApplicationCont
 		CHECK_MESSAGE(TEXT("`UniqueBy` with duplicates and unique items must return correct result"), (Array == TArray<FIntPair>{ { 1, 0 }, { 2, 0 }, { 3, 0 } }));
 	}
 	{
-		FString Str = TEXT("aa");
+		FString Str = "aa";
 		Str = Str.Mid(0, UniqueBy(Str, FIdentityFunctor()));
 		CHECK_MESSAGE(TEXT("`UniqueBy` on `FString` as an example of arbitrary random-access container must compile and return correct result"),
-			Str == FString(TEXT("a")));
+			Str == FString("a"));
 	}
 	{
 		FIntPair Array[] = { { 1, 0 } };
diff --git a/Engine/Source/Runtime/Core/Tests/Containers/StringViewTest.cpp b/Engine/Source/Runtime/Core/Tests/Containers/StringViewTest.cpp
--- a/Engine/Source/Runtime/Core/Tests/Containers/StringViewTest.cpp
+++ b/Engine/Source/Runtime/Core/Tests/Containers/StringViewTest.cpp
@@ -75,7 +75,7 @@ TEST_CASE_NAMED(FStringViewConstructorTest, "System::Core::String::StringView::C
 
 	SECTION("Construct from FString")
 	{
-		FString String(TEXT("String Object"));
+		FString String("String Object");
 		FStringView View(String);
 
 		CHECK(View.Len() == String.Len());
@@ -132,7 +132,7 @@ TEST_CASE_NAMED(FStringViewConstructorTest, "System::Core::String::StringView::C
 		STATIC_CHECK(std::is_same_v<decltype(ViewWideStringBuilder)::ElementType, WIDECHAR>);
 		STATIC_CHECK(std::is_same_v<decltype(ViewUtf8StringBuilder)::ElementType, UTF8CHAR>);
 
-		FString String(TEXT("Test"));
+		FString String("Test");
 		TStringView ViewString(String);
 		STATIC_CHECK(std::is_same_v<decltype(ViewString)::ElementType, WIDECHAR>);
 	}
diff --git a/Engine/Source/Runtime/Core/Tests/HAL/ConsoleManagerTest.cpp b/Engine/Source/Runtime/Core/Tests/HAL/ConsoleManagerTest.cpp
--- a/Engine/Source/Runtime/Core/Tests/HAL/ConsoleManagerTest.cpp
+++ b/Engine/Source/Runtime/Core/Tests/HAL/ConsoleManagerTest.cpp
@@ -108,21 +108,21 @@ TEST_CASE_NAMED(FConsoleManagerTest, "System::Core::HAL::ConsoleManager", "[Appl
 		{
 			CHECK(VarA->GetInt() == 1);
 			CHECK(VarA->GetFloat() == (float)1);
-			CHECK(VarA->GetString() == FString(TEXT("1")));
+			CHECK(VarA->GetString() == FString("1"));
 
 			CHECK(VarB->GetInt() == 1);
 			CHECK(FMath::IsNearlyEqual(VarB->GetFloat(), 1.2f, UE_KINDA_SMALL_NUMBER));
-			CHECK(VarB->GetString() == FString(TEXT("1.2")));
+			CHECK(VarB->GetString() == FString("1.2"));
 
 			CHECK(RefD == 2);
 			CHECK(VarD->GetInt() == 2);
 			CHECK(VarD->GetFloat() == (float)2);
-			CHECK(VarD->GetString() == FString(TEXT("2")));
+			CHECK(VarD->GetString() == FString("2"));
 
 			CHECK(FMath::IsNearlyEqual(RefE, 2.1f, UE_KINDA_SMALL_NUMBER));
 			CHECK(VarE->GetInt() == (int32)RefE);
 			CHECK(VarE->GetFloat() == RefE);
-			CHECK(VarE->GetString() == FString(TEXT("2.1")));
+			CHECK(VarE->GetString() == FString("2.1"));
 		}
 
 		SECTION("Set variable values (string)")
@@ -133,10 +133,10 @@ TEST_CASE_NAMED(FConsoleManagerTest, "System::Core::HAL::ConsoleManager", "[Appl
 			VarE->Set(TEXT("3.1"), ECVF_SetByConsoleVariablesIni);
 			CHECK(ChangeCounter == 1);
 
-			CHECK(VarA->GetString() == FString(TEXT("3")));
-			CHECK(VarB->GetString() == FString(TEXT("3.1")));
-			CHECK(VarD->GetString() == FString(TEXT("3")));
-			CHECK(VarE->GetString() == FString(TEXT("3.1")));
+			CHECK(VarA->GetString() == FString("3"));
+			CHECK(VarB->GetString() == FString("3.1"));
+			CHECK(VarD->GetString() == FString("3"));
+			CHECK(VarE->GetString() == FString("3.1"));
 			CHECK(RefD == 3);
 			CHECK(RefE == 3.1f);
 
@@ -152,9 +152,9 @@ TEST_CASE_NAMED(FConsoleManagerTest, "System::Core::HAL::ConsoleManager", "[Appl
 			CHECK(VarC->GetInt() == 1);
 			// note: exact comparison fails in Win32 release
 			CHECK(FMath::IsNearlyEqual(VarC->GetFloat(), 1.23f, UE_KINDA_SMALL_NUMBER));
-			CHECK(VarC->GetString() == FString(TEXT("1.23")));
+			CHECK(VarC->GetString() == FString("1.23"));
 			VarC->Set(TEXT("3.1"), ECVF_SetByConsole);
-			CHECK(VarC->GetString() == FString(TEXT("3.1")));
+			CHECK(VarC->GetString() == FString("3.1"));
 
 			Manager.UnregisterConsoleObject(TEXT("TestNameC"), false);
 			CHECK(!Manager.FindConsoleVariable(TEXT("TestNameC")));
@@ -196,22 +196,22 @@ TEST_CASE_NAMED(FConsoleManagerTest, "System::Core::HAL::ConsoleManager", "[Appl
 
 			// lower should fail
 			VarX->Set(TEXT("111"), ECVF_SetByScalability);
-			CHECK(VarX->GetString() == FString(TEXT("3")));
+			CHECK(VarX->GetString() == FString("3"));
 			CHECK(((uint32)VarX->GetFlags() & ECVF_SetByMask) == ECVF_SetByConsoleVariablesIni);
 
 			// higher should work
 			VarX->Set(TEXT("222"), ECVF_SetByCommandline);
-			CHECK(VarX->GetString() == FString(TEXT("222")));
+			CHECK(VarX->GetString() == FString("222"));
 			CHECK(((uint32)VarX->GetFlags() & ECVF_SetByMask) == ECVF_SetByCommandline);
 
 			// lower should fail
 			VarX->Set(TEXT("333"), ECVF_SetByConsoleVariablesIni);
-			CHECK(VarX->GetString() == FString(TEXT("222")));
+			CHECK(VarX->GetString() == FString("222"));
 			CHECK(((uint32)VarX->GetFlags() & ECVF_SetByMask) == ECVF_SetByCommandline);
 
 			// higher should work
 			VarX->Set(TEXT("444"), ECVF_SetByConsole);
-			CHECK(VarX->GetString() == FString(TEXT("444")));
+			CHECK(VarX->GetString() == FString("444"));
 			CHECK(((uint32)VarX->GetFlags() & ECVF_SetByMask) == ECVF_SetByConsole);
 		}
 	}
diff --git a/Engine/Source/Runtime/Core/Tests/IO/IoChunkEncodingTest.cpp b/Engine/Source/Runtime/Core/Tests/IO/IoChunkEncodingTest.cpp
--- a/Engine/Source/Runtime/Core/Tests/IO/IoChunkEncodingTest.cpp
+++ b/Engine/Source/Runtime/Core/Tests/IO/IoChunkEncodingTest.cpp
@@ -60,7 +60,7 @@ TEST_CASE("Core::IO::IoChunk", "[Core][IO]")
 		};
 
 		const FAES::FAESKey Key = GenerateEncryptionKey();
-		const FName Method = TEXT("Oodle");
+		const FName Method = "Oodle";
 		constexpr uint64 BlockSize = 64 * sizeof(uint64);
 		constexpr int32 N = 5000;
 		const TArray<uint64> ExpectedValues = GenerateData(N);
diff --git a/Engine/Source/Runtime/Core/Tests/Internationalization/ArchiveTest.cpp b/Engine/Source/Runtime/Core/Tests/Internationalization/ArchiveTest.cpp
--- a/Engine/Source/Runtime/Core/Tests/Internationalization/ArchiveTest.cpp
+++ b/Engine/Source/Runtime/Core/Tests/Internationalization/ArchiveTest.cpp
@@ -46,8 +46,8 @@ TEST_CASE_NAMED(FArchiveTest, "System::Core::Misc::Internationalization Archive"
 	FLocItem Translation = SourceA;
 	Translation.Text = TEXT("TranslatedTextA");
 
-	FString TestNamespace = TEXT("TestNamespace");
-	FString SourceAKey = TEXT("TextA");
+	FString TestNamespace = "TestNamespace";
+	FString SourceAKey = "TextA";
 
 	// Test entry add
 	{
diff --git a/Engine/Source/Runtime/Core/Tests/Internationalization/KeysCultureTest.cpp b/Engine/Source/Runtime/Core/Tests/Internationalization/KeysCultureTest.cpp
--- a/Engine/Source/Runtime/Core/Tests/Internationalization/KeysCultureTest.cpp
+++ b/Engine/Source/Runtime/Core/Tests/Internationalization/KeysCultureTest.cpp
@@ -49,8 +49,8 @@ TEST_CASE_NAMED(FKeysCultureTest, "Development::System::Core::Misc::KeysCulture"
 {
 	FRestoreCultureStateGuard RestoreCultureStateGuard;
 	
-	static const FString NativeCulture = TEXT("en");
-	static const FString TargetCulture = TEXT("ja");
+	static const FString NativeCulture = "en";
+	static const FString TargetCulture = "ja";
 
 	// Test Overview
 	// 1. We set up the test environment and change the culture to the native culture
diff --git a/Engine/Source/Runtime/Core/Tests/Internationalization/ManifestTest.cpp b/Engine/Source/Runtime/Core/Tests/Internationalization/ManifestTest.cpp
--- a/Engine/Source/Runtime/Core/Tests/Internationalization/ManifestTest.cpp
+++ b/Engine/Source/Runtime/Core/Tests/Internationalization/ManifestTest.cpp
@@ -392,7 +392,7 @@ TEST_CASE_NAMED(FManifestTest, "System::Core::Misc::Internationalization Manifes
 		FLocItem Source( TEXT("TestText") );
 		Source.MetadataObj = MakeShareable( new FLocMetadataObject(*SourceMetadataA) );
 
-		FString TestNamespace = TEXT("TestNamespace");
+		FString TestNamespace = "TestNamespace";
 
 
 		// Adding entries with exactly matching Source and matching Context
diff --git a/Engine/Source/Runtime/Core/Tests/Internationalization/MetadataTest.cpp b/Engine/Source/Runtime/Core/Tests/Internationalization/MetadataTest.cpp
--- a/Engine/Source/Runtime/Core/Tests/Internationalization/MetadataTest.cpp
+++ b/Engine/Source/Runtime/Core/Tests/Internationalization/MetadataTest.cpp
@@ -18,8 +18,8 @@ TEST_CASE_NAMED(FMetadataTest, "System::Core::Misc::Internationalization Metadat
 	TSharedPtr<FLocMetadataValue> MetadataValueBoolTrue = MakeShareable( new FLocMetadataValueBoolean( BoolTrue ) );
 
 	// StringValue metadata
-	FString StringA = TEXT("A");
-	FString StringB = TEXT("B");
+	FString StringA = "A";
+	FString StringB = "B";
 	TSharedPtr<FLocMetadataValue> MetadataValueStringA = MakeShareable( new FLocMetadataValueString( StringA ) );
 	TSharedPtr<FLocMetadataValue> MetadataValueStringB = MakeShareable( new FLocMetadataValueString( StringB ) );
 
diff --git a/Engine/Source/Runtime/Core/Tests/Internationalization/TextTest.cpp b/Engine/Source/Runtime/Core/Tests/Internationalization/TextTest.cpp
--- a/Engine/Source/Runtime/Core/Tests/Internationalization/TextTest.cpp
+++ b/Engine/Source/Runtime/Core/Tests/Internationalization/TextTest.cpp
@@ -502,7 +502,7 @@ public:
 				TArray<uint8> FormattedHistoryAsFrenchCanadian;
 				TArray<uint8> InvariantFTextData;
 
-				FString InvariantString = TEXT("This is a culture invariant string.");
+				FString InvariantString = "This is a culture invariant string.";
 				FString FormattedTestLayer2_OriginalLanguageSourceString;
 				FText FormattedTestLayer2;
 
@@ -1503,8 +1503,8 @@ TEST_CASE_NAMED(FICUTextTest, "System::Core::Misc::ICUText", "[.][EditorContext]
 {
 	// Test to make sure that ICUUtilities converts strings correctly
 
-	const FString SourceString(TEXT("This is a test"));
-	const FString SourceString2(TEXT("This is another test"));
+	const FString SourceString("This is a test");
+	const FString SourceString2("This is another test");
 	icu::UnicodeString ICUString;
 	FString ConversionBackStr;
 
diff --git a/Engine/Source/Runtime/Core/Tests/Misc/ParseTests.cpp b/Engine/Source/Runtime/Core/Tests/Misc/ParseTests.cpp
--- a/Engine/Source/Runtime/Core/Tests/Misc/ParseTests.cpp
+++ b/Engine/Source/Runtime/Core/Tests/Misc/ParseTests.cpp
@@ -125,7 +125,7 @@ TEST_CASE("Parse::InitFromString", "[Parse][Smoke]")
 	SECTION("FVector2D")
 	{
 		FVector2D Value(0.0, 1.0);
-		FString Expected(TEXT("X=0.000 Y=1.000"));
+		FString Expected("X=0.000 Y=1.000");
 		CHECK(Value.ToString() == Expected);
 
 		// Back-and-forth conversion should work :
@@ -162,7 +162,7 @@ TEST_CASE("Parse::InitFromString", "[Parse][Smoke]")
 	SECTION("FVector")
 	{
 		FVector Value(0.0, 1.0, 2.0);
-		FString Expected(TEXT("X=0.000 Y=1.000 Z=2.000"));
+		FString Expected("X=0.000 Y=1.000 Z=2.000");
 		CHECK(Value.ToString() == Expected);
 
 		// Back-and-forth conversion should work :
@@ -200,7 +200,7 @@ TEST_CASE("Parse::InitFromString", "[Parse][Smoke]")
 	SECTION("FVector4")
 	{
 		FVector4 Value(0.0, 1.0, 2.0, 3.0);
-		FString Expected(TEXT("X=0.000 Y=1.000 Z=2.000 W=3.000"));
+		FString Expected("X=0.000 Y=1.000 Z=2.000 W=3.000");
 		CHECK(Value.ToString() == Expected);
 
 		// Back-and-forth conversion should work :
@@ -238,7 +238,7 @@ TEST_CASE("Parse::InitFromString", "[Parse][Smoke]")
 	SECTION("FQuat")
 	{
 		FQuat Value(0.0, 1.0, 2.0, 3.0);
-		FString Expected(TEXT("X=0.000000000 Y=1.000000000 Z=2.000000000 W=3.000000000"));
+		FString Expected("X=0.000000000 Y=1.000000000 Z=2.000000000 W=3.000000000");
 		CHECK(Value.ToString() == Expected);
 
 		// Back-and-forth conversion should work :
@@ -276,7 +276,7 @@ TEST_CASE("Parse::InitFromString", "[Parse][Smoke]")
 	SECTION("FLinearColor")
 	{
 		FLinearColor Value(0.0, 1.0, 2.0, 3.0);
-		FString Expected(TEXT("(R=0.000000,G=1.000000,B=2.000000,A=3.000000)"));
+		FString Expected("(R=0.000000,G=1.000000,B=2.000000,A=3.000000)");
 		CHECK(Value.ToString() == Expected);
 
 		// Back-and-forth conversion should work :
@@ -315,7 +315,7 @@ TEST_CASE("Parse::InitFromString", "[Parse][Smoke]")
 	SECTION("FColor")
 	{
 		FColor Value(0, 1, 2, 3);
-		FString Expected(TEXT("(R=0,G=1,B=2,A=3)"));
+		FString Expected("(R=0,G=1,B=2,A=3)");
 		CHECK(Value.ToString() == Expected);
 
 		// Back-and-forth conversion should work :
@@ -352,7 +352,7 @@ TEST_CASE("Parse::InitFromString", "[Parse][Smoke]")
 	SECTION("FRotator")
 	{
 		FRotator Value(0.0, 1.0, 2.0);
-		FString Expected(TEXT("P=0.000000 Y=1.000000 R=2.000000"));
+		FString Expected("P=0.000000 Y=1.000000 R=2.000000");
 		CHECK(Value.ToString() == Expected);
 
 		// Back-and-forth conversion should work :
@@ -389,7 +389,7 @@ TEST_CASE("Parse::InitFromString", "[Parse][Smoke]")
 	SECTION("FIntPoint")
 	{
 		FIntPoint Value(0, 1);
-		FString Expected(TEXT("X=0 Y=1"));
+		FString Expected("X=0 Y=1");
 		CHECK(Value.ToString() == Expected);
 
 		// Back-and-forth conversion should work :
@@ -429,7 +429,7 @@ TEST_CASE("Parse::InitFromString", "[Parse][Smoke]")
 	SECTION("FIntVector")
 	{
 		FIntVector Value(0, 1, 2);
-		FString Expected(TEXT("X=0 Y=1 Z=2"));
+		FString Expected("X=0 Y=1 Z=2");
 		CHECK(Value.ToString() == Expected);
 
 		// Back-and-forth conversion should work :
diff --git a/Engine/Source/Runtime/Core/Tests/Misc/StringBuilderTest.cpp b/Engine/Source/Runtime/Core/Tests/Misc/StringBuilderTest.cpp
--- a/Engine/Source/Runtime/Core/Tests/Misc/StringBuilderTest.cpp
+++ b/Engine/Source/Runtime/Core/Tests/Misc/StringBuilderTest.cpp
@@ -157,7 +157,7 @@ TEST_CASE("Core::String::StringBuilder", "[Core][String][Smoke]")
 	SECTION("Append FString")
 	{
 		TStringBuilder<4> Builder;
-		Builder << FString(TEXT("ABC"));
+		Builder << FString("ABC");
 		CHECK(Builder == TEXTVIEW("ABC"));
 	}
 
diff --git a/Engine/Source/Runtime/Core/Tests/Misc/TCharToUTF8.cpp b/Engine/Source/Runtime/Core/Tests/Misc/TCharToUTF8.cpp
--- a/Engine/Source/Runtime/Core/Tests/Misc/TCharToUTF8.cpp
+++ b/Engine/Source/Runtime/Core/Tests/Misc/TCharToUTF8.cpp
@@ -17,15 +17,15 @@ TEST_CASE_NAMED(FTCharToUTF8Test, "System::Core::Misc::TCharToUtf8", "[Applicati
 		{
 			const FTCHARToUTF8 ConvertedValue(*TestString);
 			const FUTF8ToTCHAR ReverseConvertedValue(ConvertedValue.Get());
-			CHECK_MESSAGE(FString(TEXT("Expected test string 1 to have different encoded value")), FCStringAnsi::Strcmp(ExpectedResult, (const ANSICHAR*)ConvertedValue.Get()) == 0);
-			CHECK_MESSAGE(FString(TEXT("Expected test string 1 to have different decoded value")), FCString::Strcmp(*TestString, ReverseConvertedValue.Get()) == 0);
+			CHECK_MESSAGE(FString("Expected test string 1 to have different encoded value"), FCStringAnsi::Strcmp(ExpectedResult, (const ANSICHAR*)ConvertedValue.Get()) == 0);
+			CHECK_MESSAGE(FString("Expected test string 1 to have different decoded value"), FCString::Strcmp(*TestString, ReverseConvertedValue.Get()) == 0);
 		}
 
 		{
 			const FUTF8ToTCHAR ConvertedValue(ExpectedResult);
 			const FTCHARToUTF8 ReverseConvertedValue(ConvertedValue.Get());
-			CHECK_MESSAGE(FString(TEXT("Expected test string 2 to have different encoded value")), FCString::Strcmp(*TestString, ConvertedValue.Get()) == 0);
-			CHECK_MESSAGE(FString(TEXT("Expected test string 2 to have different decoded value")), FCStringAnsi::Strcmp(ExpectedResult, (const ANSICHAR*)ReverseConvertedValue.Get()) == 0);
+			CHECK_MESSAGE(FString("Expected test string 2 to have different encoded value"), FCString::Strcmp(*TestString, ConvertedValue.Get()) == 0);
+			CHECK_MESSAGE(FString("Expected test string 2 to have different decoded value"), FCStringAnsi::Strcmp(ExpectedResult, (const ANSICHAR*)ReverseConvertedValue.Get()) == 0);
 		}
 	}
 
@@ -36,15 +36,15 @@ TEST_CASE_NAMED(FTCharToUTF8Test, "System::Core::Misc::TCharToUtf8", "[Applicati
 		{
 			const FTCHARToUTF8 ConvertedValue(*TestString);
 			const FUTF8ToTCHAR ReverseConvertedValue(ConvertedValue.Get());
-			CHECK_MESSAGE(FString(TEXT("Expected test string 3 to have different encoded value")), FCStringAnsi::Strcmp(ExpectedResult, (const ANSICHAR*)ConvertedValue.Get()) == 0);
-			CHECK_MESSAGE(FString(TEXT("Expected test string 3 to have different decoded value")), FCString::Strcmp(*TestString, ReverseConvertedValue.Get()) == 0);
+			CHECK_MESSAGE(FString("Expected test string 3 to have different encoded value"), FCStringAnsi::Strcmp(ExpectedResult, (const ANSICHAR*)ConvertedValue.Get()) == 0);
+			CHECK_MESSAGE(FString("Expected test string 3 to have different decoded value"), FCString::Strcmp(*TestString, ReverseConvertedValue.Get()) == 0);
 		}
 
 		{
 			const FUTF8ToTCHAR ConvertedValue(ExpectedResult);
 			const FTCHARToUTF8 ReverseConvertedValue(ConvertedValue.Get());
-			CHECK_MESSAGE(FString(TEXT("Expected test string 4 to have different encoded value")), FCString::Strcmp(*TestString, ConvertedValue.Get()) == 0);
-			CHECK_MESSAGE(FString(TEXT("Expected test string 4 to have different decoded value")), FCStringAnsi::Strcmp(ExpectedResult, (const ANSICHAR*)ReverseConvertedValue.Get()) == 0);
+			CHECK_MESSAGE(FString("Expected test string 4 to have different encoded value"), FCString::Strcmp(*TestString, ConvertedValue.Get()) == 0);
+			CHECK_MESSAGE(FString("Expected test string 4 to have different decoded value"), FCStringAnsi::Strcmp(ExpectedResult, (const ANSICHAR*)ReverseConvertedValue.Get()) == 0);
 		}
 	}
 
@@ -55,15 +55,15 @@ TEST_CASE_NAMED(FTCharToUTF8Test, "System::Core::Misc::TCharToUtf8", "[Applicati
 		{
 			const FTCHARToUTF8 ConvertedValue(*TestString);
 			const FUTF8ToTCHAR ReverseConvertedValue(ConvertedValue.Get());
-			CHECK_MESSAGE(FString(TEXT("Expected test string 5 to have different encoded value")), FCStringAnsi::Strcmp(ExpectedResult, (const ANSICHAR*)ConvertedValue.Get()) == 0);
-			CHECK_MESSAGE(FString(TEXT("Expected test string 5 to have different decoded value")), FCString::Strcmp(*TestString, ReverseConvertedValue.Get()) == 0);
+			CHECK_MESSAGE(FString("Expected test string 5 to have different encoded value"), FCStringAnsi::Strcmp(ExpectedResult, (const ANSICHAR*)ConvertedValue.Get()) == 0);
+			CHECK_MESSAGE(FString("Expected test string 5 to have different decoded value"), FCString::Strcmp(*TestString, ReverseConvertedValue.Get()) == 0);
 		}
 
 		{
 			const FUTF8ToTCHAR ConvertedValue(ExpectedResult);
 			const FTCHARToUTF8 ReverseConvertedValue(ConvertedValue.Get());
-			CHECK_MESSAGE(FString(TEXT("Expected test string 6 to have different encoded value")), FCString::Strcmp(*TestString, ConvertedValue.Get()) == 0);
-			CHECK_MESSAGE(FString(TEXT("Expected test string 6 to have different decoded value")), FCStringAnsi::Strcmp(ExpectedResult, (const ANSICHAR*)ReverseConvertedValue.Get()) == 0);
+			CHECK_MESSAGE(FString("Expected test string 6 to have different encoded value"), FCString::Strcmp(*TestString, ConvertedValue.Get()) == 0);
+			CHECK_MESSAGE(FString("Expected test string 6 to have different decoded value"), FCStringAnsi::Strcmp(ExpectedResult, (const ANSICHAR*)ReverseConvertedValue.Get()) == 0);
 		}
 	}
 
@@ -74,21 +74,21 @@ TEST_CASE_NAMED(FTCharToUTF8Test, "System::Core::Misc::TCharToUtf8", "[Applicati
 		{
 			const FTCHARToUTF8 ConvertedValue(*TestString);
 			const FUTF8ToTCHAR ReverseConvertedValue(ConvertedValue.Get());
-			CHECK_MESSAGE(FString(TEXT("Expected test string 7 to have different encoded value")), FCStringAnsi::Strcmp(ExpectedResult, (const ANSICHAR*)ConvertedValue.Get()) == 0);
-			CHECK_MESSAGE(FString(TEXT("Expected test string 7 to have different decoded value")), FCString::Strcmp(*TestString, ReverseConvertedValue.Get()) == 0);
+			CHECK_MESSAGE(FString("Expected test string 7 to have different encoded value"), FCStringAnsi::Strcmp(ExpectedResult, (const ANSICHAR*)ConvertedValue.Get()) == 0);
+			CHECK_MESSAGE(FString("Expected test string 7 to have different decoded value"), FCString::Strcmp(*TestString, ReverseConvertedValue.Get()) == 0);
 		}
 
 		{
 			const FUTF8ToTCHAR ConvertedValue(ExpectedResult);
 			const FTCHARToUTF8 ReverseConvertedValue(ConvertedValue.Get());
-			CHECK_MESSAGE(FString(TEXT("Expected test string 8 to have different encoded value")), FCString::Strcmp(*TestString, ConvertedValue.Get()) == 0);
-			CHECK_MESSAGE(FString(TEXT("Expected test string 8 to have different decoded value")), FCStringAnsi::Strcmp(ExpectedResult, (const ANSICHAR*)ReverseConvertedValue.Get()) == 0);
+			CHECK_MESSAGE(FString("Expected test string 8 to have different encoded value"), FCString::Strcmp(*TestString, ConvertedValue.Get()) == 0);
+			CHECK_MESSAGE(FString("Expected test string 8 to have different decoded value"), FCStringAnsi::Strcmp(ExpectedResult, (const ANSICHAR*)ReverseConvertedValue.Get()) == 0);
 		}
 	}
 
 	{
 		// Null terminator
-		const FString TestString = TEXT("Test");
+		const FString TestString = "Test";
 		{
 			const FTCHARToUTF8 ConvertedValue(*TestString);
 			const bool bIsNullTerminated = *(ConvertedValue.Get() + ConvertedValue.Length()) == 0;
diff --git a/Engine/Source/Runtime/Core/Tests/Misc/TVariantTest.cpp b/Engine/Source/Runtime/Core/Tests/Misc/TVariantTest.cpp
--- a/Engine/Source/Runtime/Core/Tests/Misc/TVariantTest.cpp
+++ b/Engine/Source/Runtime/Core/Tests/Misc/TVariantTest.cpp
@@ -309,7 +309,7 @@ TEST_CASE_NAMED(FTVariantEmplaceTest, "System::Core::Misc::TVariant::Emplace", "
 
 	Var.Emplace<FString>(TEXT("FOO"));
 	check(Var.IsType<FString>());
-	check(Var.Get<FString>() == FString{ TEXT("FOO") });
+	check(Var.Get<FString>() == FString{ "FOO" });
 
 }
 
@@ -527,7 +527,7 @@ TEST_CASE_NAMED(FTVariantSerializationTest, "System::Core::Misc::TVariant::Seria
 
 	Reader << VarData;
 	check(VarData.IsType<FString>());
-	check(VarData.Get<FString>() == FString(TEXT("Foo")));
+	check(VarData.Get<FString>() == FString("Foo"));
 
 }
 
diff --git a/Engine/Source/Runtime/Core/Tests/Misc/URLRequestFilterTest.cpp b/Engine/Source/Runtime/Core/Tests/Misc/URLRequestFilterTest.cpp
--- a/Engine/Source/Runtime/Core/Tests/Misc/URLRequestFilterTest.cpp
+++ b/Engine/Source/Runtime/Core/Tests/Misc/URLRequestFilterTest.cpp
@@ -435,7 +435,7 @@ TEST_CASE_NAMED(FMiscURLRequestFilterTest, "System::Core::Misc::URLRequestFilter
 
 	SECTION("Config initialization")
 	{
-		const FString TestConfigFilename = TEXT("TestConfigFile");
+		const FString TestConfigFilename = "TestConfigFile";
 		const TCHAR* ConfigContents = TEXT(R"(
 [NoScheme]
 
diff --git a/Engine/Source/Runtime/Core/Tests/Serialization/Base64Test.cpp b/Engine/Source/Runtime/Core/Tests/Serialization/Base64Test.cpp
--- a/Engine/Source/Runtime/Core/Tests/Serialization/Base64Test.cpp
+++ b/Engine/Source/Runtime/Core/Tests/Serialization/Base64Test.cpp
@@ -60,13 +60,13 @@ TEST_CASE_NAMED(FBase64EncodeDecodeRFCTestVectorsTest, "System::Core::Serializat
 	// This test does a good job of verifying the correct padding of various lengths by going through 2 cycles of the padding modulus.
 	const FBase64StringTest TestVectors[] =
 	{
-		{ FString(TEXT("")), FString(TEXT("")) },
-		{ FString(TEXT("f")), FString(TEXT("Zg==")) },
-		{ FString(TEXT("fo")), FString(TEXT("Zm8=")) },
-		{ FString(TEXT("foo")), FString(TEXT("Zm9v")) },
-		{ FString(TEXT("foob")), FString(TEXT("Zm9vYg==")) },
-		{ FString(TEXT("fooba")), FString(TEXT("Zm9vYmE=")) },
-		{ FString(TEXT("foobar")), FString(TEXT("Zm9vYmFy")) },
+		{ FString(), FString() },
+		{ FString("f"), FString("Zg==") },
+		{ FString("fo"), FString("Zm8=") },
+		{ FString("foo"), FString("Zm9v") },
+		{ FString("foob"), FString("Zm9vYg==") },
+		{ FString("fooba"), FString("Zm9vYmE=") },
+		{ FString("foobar"), FString("Zm9vYmFy") },
 	};
 	RunRoundTripSuccessfulBase64Tests(TestVectors, EBase64Mode::Standard);
 
@@ -79,7 +79,7 @@ TEST_CASE_NAMED(FBase64StandartEncodeDecodeEverySixBitInputTest, "System::Core::
 {
 	const TBase64Test<48> TestVectors[] =
 	{
-		{ { 0,16,131,16,81,135,32,146,139,48,211,143,65,20,147,81,85,151,97,150,155,113,215,159,130,24,163,146,89,167,162,154,171,178,219,175,195,28,179,211,93,183,227,158,187,243,223,191 }, FString(TEXT("ABCDEFGHIJKLMNOPQRSTUVWXYZabcdefghijklmnopqrstuvwxyz0123456789+/")) }
+		{ { 0,16,131,16,81,135,32,146,139,48,211,143,65,20,147,81,85,151,97,150,155,113,215,159,130,24,163,146,89,167,162,154,171,178,219,175,195,28,179,211,93,183,227,158,187,243,223,191 }, FString("ABCDEFGHIJKLMNOPQRSTUVWXYZabcdefghijklmnopqrstuvwxyz0123456789+/") }
 	};
 
 	RunRoundTripSuccessfulBase64Tests(TestVectors, EBase64Mode::Standard);
@@ -90,7 +90,7 @@ TEST_CASE_NAMED(FBase64URLSafeTest, "System::Core::Serialization::Base64::URLSaf
 {
 	const TBase64Test<48> TestVectors[] = 
 	{
-		{ { 0,16,131,16,81,135,32,146,139,48,211,143,65,20,147,81,85,151,97,150,155,113,215,159,130,24,163,146,89,167,162,154,171,178,219,175,195,28,179,211,93,183,227,158,187,243,223,191 }, FString(TEXT("ABCDEFGHIJKLMNOPQRSTUVWXYZabcdefghijklmnopqrstuvwxyz0123456789-_")) }
+		{ { 0,16,131,16,81,135,32,146,139,48,211,143,65,20,147,81,85,151,97,150,155,113,215,159,130,24,163,146,89,167,162,154,171,178,219,175,195,28,179,211,93,183,227,158,187,243,223,191 }, FString("ABCDEFGHIJKLMNOPQRSTUVWXYZabcdefghijklmnopqrstuvwxyz0123456789-_") }
 	};
 
 	RunRoundTripSuccessfulBase64Tests(TestVectors, EBase64Mode::UrlSafe);
diff --git a/Engine/Source/Runtime/Core/Tests/Serialization/MemoryReaderTest.cpp b/Engine/Source/Runtime/Core/Tests/Serialization/MemoryReaderTest.cpp
--- a/Engine/Source/Runtime/Core/Tests/Serialization/MemoryReaderTest.cpp
+++ b/Engine/Source/Runtime/Core/Tests/Serialization/MemoryReaderTest.cpp
@@ -40,7 +40,7 @@ TEST_CASE_NAMED(FMemoryReaderTest, "System::Core::Serialization::MemoryReader",
 		const double TestValueD    = 256.5;
 		const bool   TestValueB    = true;
 		const WIDECHAR TestValueCh = 0xF2;
-		const FString TestAnsiStr  = TEXT("Joe");
+		const FString TestAnsiStr  = "Joe";
 		const FString TestUtf16Str = TEXT("\uC11C\uC6B8\uC0AC\uB78C"); // Must contains at least one non-ANSI codepoint to ensure UTF16 is used.
 
 		// The values to write.
diff --git a/Engine/Source/Runtime/Core/Tests/String/FormatTest.cpp b/Engine/Source/Runtime/Core/Tests/String/FormatTest.cpp
--- a/Engine/Source/Runtime/Core/Tests/String/FormatTest.cpp
+++ b/Engine/Source/Runtime/Core/Tests/String/FormatTest.cpp
@@ -182,7 +182,7 @@ TEST_CASE_NAMED(FStringFormattingTestOrderedAllTypes, "System::Core::String Form
 	const uint64 Value4 = 400;
 	const float Value5 = 500.0;
 	const double Value6 = 600.0;
-	const FString Value7 = TEXT("Text");
+	const FString Value7 = "Text";
 	const TCHAR* Value8 = TEXT("Text");
 
 	TArray<FStringFormatArg> Args;
diff --git a/Engine/Source/Runtime/Core/Tests/String/JoinTest.cpp b/Engine/Source/Runtime/Core/Tests/String/JoinTest.cpp
--- a/Engine/Source/Runtime/Core/Tests/String/JoinTest.cpp
+++ b/Engine/Source/Runtime/Core/Tests/String/JoinTest.cpp
@@ -28,22 +28,22 @@ TEST_CASE_NAMED(FStringJoinTest, "System::Core::String::Join", "[Core][String][S
 	{
 		// TArrayView<T>&&
 		CHECK(WriteToString<128>(Join(MakeArrayView<FStringView>({TEXTVIEW("ABC"), TEXTVIEW("DEF")}), TEXT(", "))) == TEXTVIEW("ABC, DEF"));
-		CHECK(WriteToString<128>(Join(MakeArrayView<FString>({FString(TEXT("ABC")), FString(TEXT("DEF"))}), TEXT(", "))) == TEXTVIEW("ABC, DEF"));
+		CHECK(WriteToString<128>(Join(MakeArrayView<FString>({FString("ABC"), FString("DEF")}), TEXT(", "))) == TEXTVIEW("ABC, DEF"));
 
 		// const T(&)[N]
 		const FStringView ArrayStringView[] = {TEXTVIEW("ABC"), TEXTVIEW("DEF")};
 		CHECK(WriteToString<128>(Join(ArrayStringView, TEXT(", "))) == TEXTVIEW("ABC, DEF"));
-		const FString ArrayString[] = {FString(TEXT("ABC")), FString(TEXT("DEF"))};
+		const FString ArrayString[] = {FString("ABC"), FString("DEF")};
 		CHECK(WriteToString<128>(Join(ArrayString, TEXT(", "))) == TEXTVIEW("ABC, DEF"));
 	}
 
 	SECTION("JoinBy")
 	{
 		// TArrayView<FString>&&
-		CHECK(WriteToString<128>(JoinBy(MakeArrayView<FString>({FString(TEXT("ABC")), FString(TEXT("DEF"))}), UE_PROJECTION_MEMBER(FString, ToLower), TEXT(", "))) == TEXTVIEW("abc, def"));
+		CHECK(WriteToString<128>(JoinBy(MakeArrayView<FString>({FString("ABC"), FString("DEF")}), UE_PROJECTION_MEMBER(FString, ToLower), TEXT(", "))) == TEXTVIEW("abc, def"));
 
 		// const FString(&)[N]
-		const FString ArrayString[] = { FString(TEXT("ABC")), FString(TEXT("DEF")) };
+		const FString ArrayString[] = { FString("ABC"), FString("DEF") };
 		CHECK(WriteToString<128>(JoinBy(ArrayString, UE_PROJECTION_MEMBER(FString, ToLower), TEXT(", "))) == TEXTVIEW("abc, def"));
 	}
 
@@ -51,22 +51,22 @@ TEST_CASE_NAMED(FStringJoinTest, "System::Core::String::Join", "[Core][String][S
 	{
 		// TArrayView<T>&&
 		CHECK(WriteToString<128>(JoinQuoted(MakeArrayView<FStringView>({TEXTVIEW("ABC"), TEXTVIEW("DEF")}), TEXT(", "), TEXT("|"))) == TEXTVIEW("|ABC|, |DEF|"));
-		CHECK(WriteToString<128>(JoinQuoted(MakeArrayView<FString>({FString(TEXT("ABC")), FString(TEXT("DEF"))}), TEXT(", "), TEXT("|"))) == TEXTVIEW("|ABC|, |DEF|"));
+		CHECK(WriteToString<128>(JoinQuoted(MakeArrayView<FString>({FString("ABC"), FString("DEF")}), TEXT(", "), TEXT("|"))) == TEXTVIEW("|ABC|, |DEF|"));
 
 		// const T(&)[N]
 		const FStringView ArrayStringView[] = { TEXTVIEW("ABC"), TEXTVIEW("DEF") };
 		CHECK(WriteToString<128>(JoinQuoted(ArrayStringView, TEXT(", "), TEXT("|"))) == TEXTVIEW("|ABC|, |DEF|"));
-		const FString ArrayString[] = { FString(TEXT("ABC")), FString(TEXT("DEF")) };
+		const FString ArrayString[] = { FString("ABC"), FString("DEF") };
 		CHECK(WriteToString<128>(JoinQuoted(ArrayString, TEXT(", "), TEXT("|"))) == TEXTVIEW("|ABC|, |DEF|"));
 	}
 
 	SECTION("JoinQuotedBy")
 	{
 		// TArrayView<FString>&&
-		CHECK(WriteToString<128>(JoinQuotedBy(MakeArrayView<FString>({FString(TEXT("ABC")), FString(TEXT("DEF"))}), UE_PROJECTION_MEMBER(FString, ToLower), TEXT(", "), TEXT("|"))) == TEXTVIEW("|abc|, |def|"));
+		CHECK(WriteToString<128>(JoinQuotedBy(MakeArrayView<FString>({FString("ABC"), FString("DEF")}), UE_PROJECTION_MEMBER(FString, ToLower), TEXT(", "), TEXT("|"))) == TEXTVIEW("|abc|, |def|"));
 
 		// const FString(&)[N]
-		const FString ArrayString[] = { FString(TEXT("ABC")), FString(TEXT("DEF")) };
+		const FString ArrayString[] = { FString("ABC"), FString("DEF") };
 		CHECK(WriteToString<128>(JoinQuotedBy(ArrayString, UE_PROJECTION_MEMBER(FString, ToLower), TEXT(", "), TEXT("|"))) == TEXTVIEW("|abc|, |def|"));
 	}
 
@@ -74,22 +74,22 @@ TEST_CASE_NAMED(FStringJoinTest, "System::Core::String::Join", "[Core][String][S
 	{
 		// TTuple<Ts...>&&
 		CHECK(WriteToString<128>(JoinTuple(MakeTuple(TEXTVIEW("ABC"), 123), TEXT(", "))) == TEXTVIEW("ABC, 123"));
-		CHECK(WriteToString<128>(JoinTuple(MakeTuple(FString(TEXT("ABC")), 123), TEXT(", "))) == TEXTVIEW("ABC, 123"));
+		CHECK(WriteToString<128>(JoinTuple(MakeTuple(FString("ABC"), 123), TEXT(", "))) == TEXTVIEW("ABC, 123"));
 
 		// const TTuple<Ts...>&
 		const TTuple<FStringView, int32> TupleStringInt = MakeTuple(TEXTVIEW("ABC"), 123);
 		CHECK(WriteToString<128>(JoinTuple(TupleStringInt, TEXT(", "))) == TEXTVIEW("ABC, 123"));
-		const TTuple<FString, int32> TupleStringViewInt = MakeTuple(FString(TEXT("ABC")), 123);
+		const TTuple<FString, int32> TupleStringViewInt = MakeTuple(FString("ABC"), 123);
 		CHECK(WriteToString<128>(JoinTuple(TupleStringViewInt, TEXT(", "))) == TEXTVIEW("ABC, 123"));
 	}
 
 	SECTION("JoinTupleBy")
 	{
 		// TTuple<FString, int32>&&
-		CHECK(WriteToString<128>(JoinTupleBy(MakeTuple(FString(TEXT("ABC")), 123), ConditionalToLower, TEXT(", "))) == TEXTVIEW("abc, 123"));
+		CHECK(WriteToString<128>(JoinTupleBy(MakeTuple(FString("ABC"), 123), ConditionalToLower, TEXT(", "))) == TEXTVIEW("abc, 123"));
 
 		// const TTuple<FString, int32>&
-		const TTuple<FString, int32> TupleStringInt = MakeTuple(FString(TEXT("ABC")), 123);
+		const TTuple<FString, int32> TupleStringInt = MakeTuple(FString("ABC"), 123);
 		CHECK(WriteToString<128>(JoinTupleBy(TupleStringInt, ConditionalToLower, TEXT(", "))) == TEXTVIEW("abc, 123"));
 	}
 
@@ -97,22 +97,22 @@ TEST_CASE_NAMED(FStringJoinTest, "System::Core::String::Join", "[Core][String][S
 	{
 		// TTuple<Ts...>&&
 		CHECK(WriteToString<128>(JoinTupleQuoted(MakeTuple(TEXTVIEW("ABC"), 123), TEXT(", "), TEXT("|"))) == TEXTVIEW("|ABC|, |123|"));
-		CHECK(WriteToString<128>(JoinTupleQuoted(MakeTuple(FString(TEXT("ABC")), 123), TEXT(", "), TEXT("|"))) == TEXTVIEW("|ABC|, |123|"));
+		CHECK(WriteToString<128>(JoinTupleQuoted(MakeTuple(FString("ABC"), 123), TEXT(", "), TEXT("|"))) == TEXTVIEW("|ABC|, |123|"));
 
 		// const TTuple<Ts...>&
 		const TTuple<FStringView, int32> TupleStringInt = MakeTuple(TEXTVIEW("ABC"), 123);
 		CHECK(WriteToString<128>(JoinTupleQuoted(TupleStringInt, TEXT(", "), TEXT("|"))) == TEXTVIEW("|ABC|, |123|"));
-		const TTuple<FString, int32> TupleStringViewInt = MakeTuple(FString(TEXT("ABC")), 123);
+		const TTuple<FString, int32> TupleStringViewInt = MakeTuple(FString("ABC"), 123);
 		CHECK(WriteToString<128>(JoinTupleQuoted(TupleStringViewInt, TEXT(", "), TEXT("|"))) == TEXTVIEW("|ABC|, |123|"));
 	}
 
 	SECTION("JoinTupleQuotedBy")
 	{
 		// TTuple<FString, int32>&&
-		CHECK(WriteToString<128>(JoinTupleQuotedBy(MakeTuple(FString(TEXT("ABC")), 123), ConditionalToLower, TEXT(", "), TEXT("|"))) == TEXTVIEW("|abc|, |123|"));
+		CHECK(WriteToString<128>(JoinTupleQuotedBy(MakeTuple(FString("ABC"), 123), ConditionalToLower, TEXT(", "), TEXT("|"))) == TEXTVIEW("|abc|, |123|"));
 
 		// const TTuple<FString, int32>&
-		const TTuple<FString, int32> TupleStringInt = MakeTuple(FString(TEXT("ABC")), 123);
+		const TTuple<FString, int32> TupleStringInt = MakeTuple(FString("ABC"), 123);
 		CHECK(WriteToString<128>(JoinTupleQuotedBy(TupleStringInt, ConditionalToLower, TEXT(", "), TEXT("|"))) == TEXTVIEW("|abc|, |123|"));
 	}
 }
diff --git a/Engine/Source/Runtime/Core/Tests/Templates/MemoryOpsTest.cpp b/Engine/Source/Runtime/Core/Tests/Templates/MemoryOpsTest.cpp
--- a/Engine/Source/Runtime/Core/Tests/Templates/MemoryOpsTest.cpp
+++ b/Engine/Source/Runtime/Core/Tests/Templates/MemoryOpsTest.cpp
@@ -93,9 +93,9 @@ TEST_CASE_NAMED(FRelocateConstructItemsTest, "UE::Core::RelocateConstructItems",
 
 		FString* TypedBuffer = (FString*)Buffer;
 
-		new ((void*)(TypedBuffer + 2)) FString(TEXT("String A1"));
-		new ((void*)(TypedBuffer + 3)) FString(TEXT("String B1"));
-		new ((void*)(TypedBuffer + 4)) FString(TEXT("String C1"));
+		new ((void*)(TypedBuffer + 2)) FString("String A1");
+		new ((void*)(TypedBuffer + 3)) FString("String B1");
+		new ((void*)(TypedBuffer + 4)) FString("String C1");
 
 		RelocateConstructItems<FString>((void*)Buffer, TypedBuffer + 2, 3);
 
@@ -114,9 +114,9 @@ TEST_CASE_NAMED(FRelocateConstructItemsTest, "UE::Core::RelocateConstructItems",
 
 		FString* TypedBuffer = (FString*)Buffer;
 
-		new ((void*)(TypedBuffer + 0)) FString(TEXT("String A2"));
-		new ((void*)(TypedBuffer + 1)) FString(TEXT("String B2"));
-		new ((void*)(TypedBuffer + 2)) FString(TEXT("String C2"));
+		new ((void*)(TypedBuffer + 0)) FString("String A2");
+		new ((void*)(TypedBuffer + 1)) FString("String B2");
+		new ((void*)(TypedBuffer + 2)) FString("String C2");
 
 		RelocateConstructItems<FString>((void*)(TypedBuffer + 2), TypedBuffer, 3);
 
