Fix edge-case behavior in MatchesWildcard, and make the comments more explicit about the wildcard patterns used.

The wildcard syntax here is somewhat non-standard (? in wildcards is not usually allowed to match 0 characters), so note that in the docs.

Also add tests for MatchesWildcard behavior.

diff --git a/Engine/Source/Runtime/Core/Private/Containers/String.cpp.inl b/Engine/Source/Runtime/Core/Private/Containers/String.cpp.inl
--- a/Engine/Source/Runtime/Core/Private/Containers/String.cpp.inl
+++ b/Engine/Source/Runtime/Core/Private/Containers/String.cpp.inl
@@ -43,7 +43,7 @@ namespace UE::Core::Private
 				break;
 			}
 
-			if (!Compare(*Target, WCh))
+			if (TargetLength == 0 || !Compare(*Target, WCh))
 			{
 				return false;
 			}
@@ -54,14 +54,22 @@ namespace UE::Core::Private
 			--WildcardLength;
 		}
 
-		// Test for common suffix
+		// Test for common suffix. If we make it here, WildcardLength != 0, and *Wildcard
+		// is an actual wildcard character. (Postcondition of previous loop.)
 		const UE_STRING_CHARTYPE* TPtr = Target   + TargetLength;
 		const UE_STRING_CHARTYPE* WPtr = Wildcard + WildcardLength;
 		for (;;)
 		{
+			if (TargetLength == 0)
+			{
+				break;
+			}
+
 			--TPtr;
 			--WPtr;
 
+			// Since *Wildcard is a wildcard character, which we will dereference here when
+			// WildcardLength is 1, we know we won't decrement WildcardLength past 1.
 			UE_STRING_CHARTYPE WCh = *WPtr;
 			if (WCh == CHARTEXT(UE_STRING_CHARTYPE, '*') || WCh == CHARTEXT(UE_STRING_CHARTYPE, '?'))
 			{
@@ -75,11 +83,6 @@ namespace UE::Core::Private
 
 			--TargetLength;
 			--WildcardLength;
-
-			if (TargetLength == 0)
-			{
-				break;
-			}
 		}
 
 		// Match * against anything and ? against single (and zero?) chars
diff --git a/Engine/Source/Runtime/Core/Public/Containers/UnrealString.h.inl b/Engine/Source/Runtime/Core/Public/Containers/UnrealString.h.inl
--- a/Engine/Source/Runtime/Core/Public/Containers/UnrealString.h.inl
+++ b/Engine/Source/Runtime/Core/Public/Containers/UnrealString.h.inl
@@ -1634,10 +1634,10 @@ public:
 	/**
 	 * Searches this string for a given wild card
 	 *
-	 * @param Wildcard		*?-type wildcard
+	 * @param Wildcard		*?-type wildcard. '*' matches 0 or more arbitrary characters, '?' matches either an arbitrary character or nothing, all other characters match themselves.
 	 * @param SearchCase	Indicates whether the search is case sensitive or not (defaults to ESearchCase::IgnoreCase)
 	 * @return true if this string matches the *?-type wildcard given. 
-	 * @warning This is a simple, SLOW routine. Use with caution
+	 * @warning This is a simple, SLOW routine. Use with caution. Also note that '?' can match nothing (not common for wildcard patterns).
 	 */
 	template <
 		typename CharRangeType,
@@ -1658,10 +1658,10 @@ public:
 	/**
 	 * Searches this string for a given wild card
 	 *
-	 * @param Wildcard		*?-type wildcard
+	 * @param Wildcard		*?-type wildcard. '*' matches 0 or more arbitrary characters, '?' matches either an arbitrary character or nothing, all other characters match themselves.
 	 * @param SearchCase	Indicates whether the search is case sensitive or not (defaults to ESearchCase::IgnoreCase)
 	 * @return true if this string matches the *?-type wildcard given.
-	 * @warning This is a simple, SLOW routine. Use with caution
+	 * @warning This is a simple, SLOW routine. Use with caution. Also note that '?' can match nothing (not common for wildcard patterns).
 	 */
 	[[nodiscard]] bool MatchesWildcard(const ElementType* Wildcard, ESearchCase::Type SearchCase = ESearchCase::IgnoreCase) const
 	{
@@ -1671,10 +1671,10 @@ public:
 	/**
 	 * Searches this string for a given wild card
 	 *
-	 * @param Wildcard		*?-type wildcard
+	 * @param Wildcard		*?-type wildcard. '*' matches 0 or more arbitrary characters, '?' matches either an arbitrary character or nothing, all other characters match themselves.
 	 * @param SearchCase	Indicates whether the search is case sensitive or not (defaults to ESearchCase::IgnoreCase)
 	 * @return true if this string matches the *?-type wildcard given.
-	 * @warning This is a simple, SLOW routine. Use with caution
+	 * @warning This is a simple, SLOW routine. Use with caution. Also note that '?' can match nothing (not common for wildcard patterns).
 	 */
 	[[nodiscard]] bool MatchesWildcard(const UE_STRING_CLASS& Wildcard, ESearchCase::Type SearchCase = ESearchCase::IgnoreCase) const
 	{
@@ -1684,11 +1684,11 @@ public:
 	/**
 	 * Searches this string for a given wild card
 	 *
-	 * @param Wildcard     *?-type wildcard
+	 * @param Wildcard		*?-type wildcard. '*' matches 0 or more arbitrary characters, '?' matches either an arbitrary character or nothing, all other characters match themselves.
 	 * @param WildcardLen  The length of Wildcard
 	 * @param SearchCase   Indicates whether the search is case sensitive or not (defaults to ESearchCase::IgnoreCase)
 	 * @return true if this string matches the *?-type wildcard given.
-	 * @warning This is a simple, SLOW routine. Use with caution
+	 * @warning This is a simple, SLOW routine. Use with caution. Also note that '?' can match nothing (not common for wildcard patterns).
 	 */
 	[[nodiscard]] CORE_API bool MatchesWildcard(const ElementType* Wildcard, int32 WildcardLen, ESearchCase::Type SearchCase = ESearchCase::IgnoreCase) const;
 
diff --git a/Engine/Source/Runtime/Core/Tests/String/MatchesWildcardTest.cpp b/Engine/Source/Runtime/Core/Tests/String/MatchesWildcardTest.cpp
new file mode 100644
--- /dev/null
+++ b/Engine/Source/Runtime/Core/Tests/String/MatchesWildcardTest.cpp
@@ -0,0 +1,99 @@
+// Copyright Epic Games, Inc. All Rights Reserved.
+
+#if WITH_TESTS
+
+#include "Containers/AnsiString.h"
+#include "Containers/UnrealString.h"
+#include "Containers/Utf8String.h"
+#include "Tests/TestHarnessAdapter.h"
+
+namespace UE::String
+{
+
+TEST_CASE_NAMED(FStringMatchesWildcardTest, "System::Core::String::MatchesWildcard", "[Core][String][SmokeFilter]")
+{
+	struct FTestCase
+	{
+		const ANSICHAR* Target;
+		const ANSICHAR* Pattern;
+		bool bResultCaseSensitive;
+		bool bResultIgnoreCase;
+	};
+
+	static const FTestCase TestCases[] =
+	{
+		{ "", "", true, true }, // Trivial: empty matches empty
+		{ "", "*", true, true }, // Trivial: * matches empty
+		{ "a", "*a*", true, true }, // Trivial: * matches empty on both sides
+		{ "abcd", "", false, false }, // Empty pattern doesn't match non-empty string
+		{ "abcd", "abcd", true, true }, // Literal match
+		{ "Abcd", "abcd", false, true }, // Literal match w/ case diff
+		{ "abcd", "aBcd", false, true }, // Literal match w/ case diff
+		{ "Abcd", "aBcd", false, true }, // Literal match w/ case diff
+		{ "abcd", "ab?d", true, true }, // Trivial use of ?
+		{ "abd", "ab?d", true, true }, // our ? CAN match nothing (for some reason)
+		{ "abcde", "ab?d", false, false }, // ? should _not_ be able to match 2 chars, though
+		{ "abd", "ab???d", true, true }, // multiple ?: with UE semantics, this should match 0, 1, 2, or 3 chars between ab and d
+		{ "abcd", "ab???d", true, true }, // multiple ? (1 char between) - should match
+		{ "abccd", "ab???d", true, true }, // multiple ? (2 chars between) - should match
+		{ "abcccd", "ab???d", true, true }, // multiple ? (3 chars between) - should match
+		{ "abccccd", "ab???d", false, false }, // multiple ? (4 chars between) - should not match
+		{ "abcd", "*", true, true }, // Trivial use of *
+		{ "abc", "a*", true, true }, // Trivial use of *
+		{ "abbbbccccd", "a*d", true, true }, // one * in the middle, exact prefix/suffix (hits fast paths)
+		{ "abcd", "*.*", false, false }, // DOS-ism: for us, "*.*" isn't the same as "*", demand an actual . if that is given
+		{ "abc.def", "*.*", true, true }, // actual match for *.*
+		{ "abc...def.g", "*.???.g", true, true }, // first . match after * will fail, need to retry to find real match
+		{ "abc/d", "*/D", false, true }, // suffix match
+		{ "axbxcxdxexxf", "a*b*c*d*e*f", true, true }, // many * wildcards but deterministic match
+		{ "ababababbababaczzx", "a*b*a*b?c??x", true, true }, // match after many wildcards
+		{ "ababababbababaczzy", "a*b*a*b?c??x", false, false }, // no match after many wildcards
+		{ ":::::::::a a", "*:*:*:*:*:*:*:*:*:*:*", false, false }, // Exponential-time (each candidate fails to find the last missing ':', then backtracks). This particular test is not bad, but adding more : to the pattern and string shows the behavior.
+		{ "a", "a*a", false, false }, // Regression test: MatchesWildcard used to get this wrong
+		{ "a", "*a*a", false, false }, // Regression test: MatchesWildcard used to get this wrong
+	};
+
+	auto RunTests = [](auto EmptyStringWithCorrectType)
+	{
+		auto Target = EmptyStringWithCorrectType;
+		auto Pattern = EmptyStringWithCorrectType;
+
+		for (const FTestCase& Test : TestCases)
+		{
+			Target = Test.Target;
+			Pattern = Test.Pattern;
+			bool bResultCaseSensitive = Target.MatchesWildcard(Pattern, ESearchCase::CaseSensitive);
+			bool bResultIgnoreCase = Target.MatchesWildcard(Pattern, ESearchCase::IgnoreCase);
+
+			if (bResultCaseSensitive != Test.bResultCaseSensitive ||
+				bResultIgnoreCase != Test.bResultIgnoreCase)
+			{
+				CAPTURE(Test.Target, Test.Pattern);
+				CAPTURE(bResultCaseSensitive, Test.bResultCaseSensitive);
+				CAPTURE(bResultIgnoreCase, Test.bResultIgnoreCase);
+			}
+
+			bool bResultsMatch = (bResultCaseSensitive == Test.bResultCaseSensitive) && (bResultIgnoreCase == Test.bResultIgnoreCase);
+			CHECK(bResultsMatch);
+		}
+	};
+
+	SECTION("MatchesWildcard with ANSICHAR")
+	{
+		RunTests(FAnsiString{});
+	}
+
+	SECTION("MatchesWildcard with WIDECHAR")
+	{
+		RunTests(FString{});
+	}
+
+	SECTION("MatchesWildcard with UTF8CHAR")
+	{
+		RunTests(FUtf8String{});
+	}
+}
+
+} // UE::String
+
+#endif //WITH_TESTS
