Implement GetTypeHash and comparison operator for FBoxCenterAndExtent

diff --git a/Engine/Source/Runtime/Core/Public/Math/GenericOctree.h b/Engine/Source/Runtime/Core/Public/Math/GenericOctree.h
--- a/Engine/Source/Runtime/Core/Public/Math/GenericOctree.h
+++ b/Engine/Source/Runtime/Core/Public/Math/GenericOctree.h
@@ -86,6 +86,13 @@ public:
 		return FBox(Center - Extent,Center + Extent);
 	}
 
+
+	/** Comparison operator for FBoxCenterAndExtra*/
+	[[nodiscard]] bool operator==(const FBoxCenterAndExtent& Other) const
+	{
+		return Center == Other.Center && Extent == Other.Extent;
+	}
+
 	/**
 	 * Determines whether two boxes intersect.
 	 * Warning: this operates on the W of the bounds positions!
@@ -139,6 +146,11 @@ public:
 		// extents.  If the boxes don't intersect on any of the axes, they don't intersect.
 		return VectorAnyGreaterThan(CenterDifference,CompositeExtent) == false;
 	}
+
+	[[nodiscard]] friend inline uint32 GetTypeHash(const FBoxCenterAndExtent& Box)
+	{
+		return HashCombine(GetTypeHash(Box.Center), GetTypeHash(Box.Extent));
+	}
 };
 
 /** A reference to a child of an octree node. */
