[TaskGraph]
- Allow to pass task flags from the old task graph

diff --git a/Engine/Source/Runtime/Core/Public/Async/TaskGraphInterfaces.h b/Engine/Source/Runtime/Core/Public/Async/TaskGraphInterfaces.h
--- a/Engine/Source/Runtime/Core/Public/Async/TaskGraphInterfaces.h
+++ b/Engine/Source/Runtime/Core/Public/Async/TaskGraphInterfaces.h
@@ -636,6 +636,14 @@ public:
 		}
 
 	private:
+		struct CGetTaskFlagsProvider
+		{
+			template <typename TaskType>
+			auto Requires(TaskType& Task) -> decltype(
+				Task.GetTaskFlags()
+			);
+		};
+
 		template<typename...T>
 		inline TGraphTask* ConstructAndHoldImpl(T&&... Args)
 		{
@@ -646,7 +654,15 @@ public:
 			UE::Tasks::EExtendedTaskPriority ExtPri;
 			UE::Tasks::Private::TranslatePriority(TaskObject->GetDesiredThread(), Pri, ExtPri);
 
-			Task->Init(Pri, ExtPri);
+			// Only use this when TTask implements GetTaskFlags().
+			if constexpr (TModels_V<CGetTaskFlagsProvider, TTask>)
+			{
+				Task->Init(Pri, ExtPri, TaskObject->GetTaskFlags());
+			}
+			else
+			{
+				Task->Init(Pri, ExtPri);
+			}
 
 			return Task;
 		}
@@ -672,9 +688,9 @@ private:
 	{
 	}
 
-	void Init(UE::Tasks::ETaskPriority InPriority, UE::Tasks::EExtendedTaskPriority InExtendedPriority)
+	void Init(UE::Tasks::ETaskPriority InPriority, UE::Tasks::EExtendedTaskPriority InExtendedPriority, UE::Tasks::ETaskFlags InTaskFlags = UE::Tasks::ETaskFlags::None)
 	{
-		FBaseGraphTask::Init(TEXT("GraphTask"), InPriority, InExtendedPriority);
+		FBaseGraphTask::Init(TEXT("GraphTask"), InPriority, InExtendedPriority, InTaskFlags);
 	}
 
 	struct CGetStatIdProvider
