Fix typo in FBoxCenterAndExtent

diff --git a/Engine/Source/Runtime/Core/Public/Math/GenericOctree.h b/Engine/Source/Runtime/Core/Public/Math/GenericOctree.h
--- a/Engine/Source/Runtime/Core/Public/Math/GenericOctree.h
+++ b/Engine/Source/Runtime/Core/Public/Math/GenericOctree.h
@@ -87,7 +87,7 @@ public:
 	}
 
 
-	/** Comparison operator for FBoxCenterAndExtra*/
+	/** Comparison operator for FBoxCenterAndExtent */
 	[[nodiscard]] bool operator==(const FBoxCenterAndExtent& Other) const
 	{
 		return Center == Other.Center && Extent == Other.Extent;
