* Added CHARTEXTVIEW(CharType, x) which expands out to ANSITEXTVIEW(x), TEXTVIEW(x), WIDETEXTVIEW(x) or UTF8TEXTVIEW(x) depending on CharType.
* Fixed CHARTEXT(CharType, x) to support WIDECHAR when PLATFORM_TCHAR_IS_UTF8CHAR; it expands to WIDETEXT(x).

diff --git a/Engine/Source/Runtime/Core/Public/Containers/StringView.h b/Engine/Source/Runtime/Core/Public/Containers/StringView.h
--- a/Engine/Source/Runtime/Core/Public/Containers/StringView.h
+++ b/Engine/Source/Runtime/Core/Public/Containers/StringView.h
@@ -556,6 +556,47 @@ template <UE::CCharType CharType>
 #define WIDETEXTVIEW(str) PREPROCESSOR_JOIN(WIDETEXT(str), _PrivateWSV)
 #define UTF8TEXTVIEW(str) (str##_PrivateU8SV)
 
+namespace UE::Core::Private
+{
+	template <typename CharType>
+	void CharTextViewStaticAssert()
+	{
+		static_assert(sizeof(CharType) == 0, "Unsupported char type passed to CHARTEXTVIEW");
+	}
+}
+
+// Expands out to ANSITEXTVIEW(x), TEXTVIEW(x), WIDETEXTVIEW(x) or UTF8TEXTVIEW(x) depending on CharType
+#define CHARTEXTVIEW(CharType, x) \
+	( \
+		[]() -> decltype(auto) \
+		{ \
+			/* We expect <type_traits> to already have been included any time the CHARTEXTVIEW macro is used */ \
+			using UnqualifiedCharType = std::remove_cv_t<CharType>; \
+			if constexpr (std::is_same_v<UnqualifiedCharType, ANSICHAR>) \
+			{ \
+				return ANSITEXTVIEW(x); \
+			} \
+			else if constexpr (std::is_same_v<UnqualifiedCharType, TCHAR>) \
+			{ \
+				return TEXTVIEW(x); \
+			} \
+			else if constexpr (std::is_same_v<UnqualifiedCharType, WIDECHAR>) \
+			{ \
+				return WIDETEXTVIEW(x); \
+			} \
+			else if constexpr (std::is_same_v<UnqualifiedCharType, UTF8CHAR>) \
+			{ \
+				return UTF8TEXTVIEW(x); \
+			} \
+			else \
+			{ \
+				/* We want a compile error, so forward to a templated function because of the static_assert(false) problem */ \
+				UE::Core::Private::CharTextViewStaticAssert<CharType>(); \
+				return x; \
+			} \
+		}() \
+	)
+
 //////////////////////////////////////////////////////////////////////////
 
 template <UE::CCharType CharType>
diff --git a/Engine/Source/Runtime/Core/Public/HAL/Platform.h b/Engine/Source/Runtime/Core/Public/HAL/Platform.h
--- a/Engine/Source/Runtime/Core/Public/HAL/Platform.h
+++ b/Engine/Source/Runtime/Core/Public/HAL/Platform.h
@@ -1287,7 +1287,7 @@ namespace UE::Core::Private
 #define UTF16TEXT(x) UTF16TEXT_PASTE(x)
 #define WIDETEXT(str) WIDETEXT_PASTE(str)
 
-// Expands out to x, TEXT(x) or UTF8TEXT(x) depending on CharType
+// Expands out to x, TEXT(x), WIDETEXT(x) or UTF8TEXT(x) depending on CharType
 #define CHARTEXT(CharType, x) \
 	( \
 		[]() -> decltype(auto) \
@@ -1302,6 +1302,10 @@ namespace UE::Core::Private
 			{ \
 				return TEXT(x); \
 			} \
+			else if constexpr (std::is_same_v<UnqualifiedCharType, WIDECHAR>) \
+			{ \
+				return WIDETEXT(x); \
+			} \
 			else if constexpr (std::is_same_v<UnqualifiedCharType, UTF8CHAR>) \
 			{ \
 				return UTF8TEXT(x); \
