[Insights] Fixed initialization (command line parsing) for memory trace.
* Fixed the bug where the last channel was not correctly parsed if multiple space chars were used to separate cmd line arguments.
* Also fixed issue where parsing for "memory" and "memory_light" channel presets was ignoring other channels (ex. "memeory_light,probing" was ignoring "probing").

diff --git a/Engine/Source/Runtime/Core/Private/ProfilingDebugging/MemoryTrace.cpp b/Engine/Source/Runtime/Core/Private/ProfilingDebugging/MemoryTrace.cpp
--- a/Engine/Source/Runtime/Core/Private/ProfilingDebugging/MemoryTrace.cpp
+++ b/Engine/Source/Runtime/Core/Private/ProfilingDebugging/MemoryTrace.cpp
@@ -305,6 +305,7 @@ template <typename ArgCharType>
 static EMemoryTraceInit MemoryTrace_ShouldEnable(int32 ArgC, const ArgCharType* const* ArgV)
 {
 	EMemoryTraceInit Mode = EMemoryTraceInit::Disabled;
+
 	for (int32 ArgIndex = 1; ArgIndex < ArgC; ArgIndex++)
 	{
 		const ArgCharType* Arg = ArgV[ArgIndex];
@@ -313,54 +314,52 @@ static EMemoryTraceInit MemoryTrace_ShouldEnable(int32 ArgC, const ArgCharType*
 		// In certain scenarios (for example when launching a packaged title on some platforms)
 		// the multiple arguments are squashed into one single argument. Due to this we also need
 		// to parse each argument as if it contains multiple sub arguments.
-		int32 ArgStart = -1;
+		int32 ArgStart = INDEX_NONE;
 		while (ArgView.FindChar(CHARTEXT(ArgCharType, '-'), ArgStart))
 		{
 			ArgView.RightChopInline(ArgStart);
 
-			int32 ArgEnd = INDEX_NONE;
-			ArgView.RightChop(1).FindChar('-', ArgEnd);
-
-			TStringView<ArgCharType> SubArg = ArgView.SubStr(0, ArgEnd == INDEX_NONE ? ArgView.Len() : ArgEnd);
+			int32 SpacePos = INDEX_NONE;
+			ArgView.FindChar(CHARTEXT(ArgCharType, ' '), SpacePos);
+			TStringView<ArgCharType> SubArg = (SpacePos == INDEX_NONE) ? ArgView : ArgView.Left(SpacePos);
 
-			if (SubArg.SubStr(0, 7).Equals("-trace=", ESearchCase::IgnoreCase))
+			if (SubArg.Left(7).Equals(CHARTEXTVIEW(ArgCharType, "-trace="), ESearchCase::IgnoreCase))
 			{
-				TStringView<ArgCharType> ChannelsOrPresets = SubArg.SubStr(7, SubArg.Len() - 7);
+				TStringView<ArgCharType> ChannelsOrPresets = SubArg.RightChop(7);
 
-				int32 CommaPos = -1;
 				while (!ChannelsOrPresets.IsEmpty())
 				{
-					ChannelsOrPresets.FindChar(',', CommaPos);
-					TStringView<ArgCharType> Channel = ChannelsOrPresets.SubStr(0, CommaPos == INDEX_NONE ? ChannelsOrPresets.Len() : CommaPos);
-					if (Channel.Equals("memalloc", ESearchCase::IgnoreCase))
+					int32 CommaPos = INDEX_NONE;
+					ChannelsOrPresets.FindChar(CHARTEXT(ArgCharType, ','), CommaPos);
+					TStringView<ArgCharType> Channel = (CommaPos == INDEX_NONE) ? ChannelsOrPresets : ChannelsOrPresets.Left(CommaPos);
+
+					if (Channel.Equals(CHARTEXTVIEW(ArgCharType, "memalloc"), ESearchCase::IgnoreCase))
 					{
 						Mode |= EMemoryTraceInit::AllocEvents;
 						// The tags and memory scopes are also emitted on the "memalloc" channel.
 						Mode |= EMemoryTraceInit::Tags;
 					}
-					else if (Channel.Equals("callstack", ESearchCase::IgnoreCase))
+					else if (Channel.Equals(CHARTEXTVIEW(ArgCharType, "callstack"), ESearchCase::IgnoreCase))
 					{
 						Mode |= EMemoryTraceInit::Callstacks;
 					}
-					else if (Channel.Equals("memtag", ESearchCase::IgnoreCase))
+					else if (Channel.Equals(CHARTEXTVIEW(ArgCharType, "memtag"), ESearchCase::IgnoreCase))
 					{
 						// The "memtag" channel is only used to enable LLM tracing.
 					}
-					else if (Channel.Equals("probing", ESearchCase::IgnoreCase))
+					else if (Channel.Equals(CHARTEXTVIEW(ArgCharType, "probing"), ESearchCase::IgnoreCase))
 					{
 						Mode |= EMemoryTraceInit::Probing;
 					}
-					else if (Channel.Left(6).Equals("memory", ESearchCase::IgnoreCase))
+					else if (Channel.Equals(CHARTEXTVIEW(ArgCharType, "memory"), ESearchCase::IgnoreCase))
 					{
-						if (Channel.Equals("memory", ESearchCase::IgnoreCase))
-						{
-							return EMemoryTraceInit::Full;
-						}
-						else if (Channel.Equals("memory_light", ESearchCase::IgnoreCase))
-						{
-							return EMemoryTraceInit::Light;
-						}
+						Mode |= EMemoryTraceInit::Full;
 					}
+					else if (Channel.Equals(CHARTEXTVIEW(ArgCharType, "memory_light"), ESearchCase::IgnoreCase))
+					{
+						Mode |= EMemoryTraceInit::Light;
+					}
+
 					ChannelsOrPresets.RightChopInline(Channel.Len() + 1);
 				}
 			}
@@ -387,7 +386,7 @@ FMalloc* MemoryTrace_CreateInternal(FMalloc* InMalloc, EMemoryTraceInit Mode)
 	// one as part of static deinit. However we may receive more memory
 	// trace events that would get lost as Trace's worker thread has been
 	// terminated. So flush the last remaining memory events trace needs
-	// to be updated which we will do that in response to to memory events.
+	// to be updated which we will do that in response to memory events.
 	// We'll use an atexit can to know when Trace is probably no longer
 	// getting ticked.
 	atexit([]() { MemoryTrace_EnableTracePump(); });
diff --git a/Engine/Source/Runtime/Core/Public/ProfilingDebugging/MemoryTrace.h b/Engine/Source/Runtime/Core/Public/ProfilingDebugging/MemoryTrace.h
--- a/Engine/Source/Runtime/Core/Public/ProfilingDebugging/MemoryTrace.h
+++ b/Engine/Source/Runtime/Core/Public/ProfilingDebugging/MemoryTrace.h
@@ -97,7 +97,7 @@ class FMalloc* MemoryTrace_Create(class FMalloc* InMalloc);
 void MemoryTrace_Initialize();
 
 /**
- * Returns the untracked allocator used to initalize memory tracing. Only internally used.
+ * Returns the untracked allocator used to initialize memory tracing. Only internally used.
  * @return A valid allocator if exists, nullptr otherwise.
  */
 class FMalloc* MemoryTrace_GetAllocator();
