[RaceDetector]
- Fixes to avoid ambiguous FRWScopeLockType

diff --git a/Engine/Source/Runtime/Core/Internal/Sanitizer/RaceDetectorTypes.h b/Engine/Source/Runtime/Core/Internal/Sanitizer/RaceDetectorTypes.h
--- a/Engine/Source/Runtime/Core/Internal/Sanitizer/RaceDetectorTypes.h
+++ b/Engine/Source/Runtime/Core/Internal/Sanitizer/RaceDetectorTypes.h
@@ -125,7 +125,7 @@ namespace UE::Sanitizer::RaceDetector {
 		MutexType& Mutex;
 	};
 
-	enum FRWScopeLockType
+	enum class ERWScopeLockType
 	{
 		SLT_ReadOnly = 0,
 		SLT_Write,
@@ -137,11 +137,11 @@ namespace UE::Sanitizer::RaceDetector {
 	public:
 		UE_NONCOPYABLE(TRWScopeLock);
 
-		INSTRUMENTATION_FUNCTION_ATTRIBUTES FORCEINLINE TRWScopeLock(MutexType& InMutex, FRWScopeLockType InLockType)
+		INSTRUMENTATION_FUNCTION_ATTRIBUTES FORCEINLINE TRWScopeLock(MutexType& InMutex, ERWScopeLockType InLockType)
 			: Mutex(InMutex)
 			, LockType(InLockType)
 		{
-			if (LockType == SLT_ReadOnly)
+			if (LockType == ERWScopeLockType::SLT_ReadOnly)
 			{
 				Mutex.ReadLock();
 			}
@@ -153,7 +153,7 @@ namespace UE::Sanitizer::RaceDetector {
 
 		INSTRUMENTATION_FUNCTION_ATTRIBUTES FORCEINLINE ~TRWScopeLock()
 		{
-			if (LockType == SLT_ReadOnly)
+			if (LockType == ERWScopeLockType::SLT_ReadOnly)
 			{
 				Mutex.ReadUnlock();
 			}
@@ -165,7 +165,7 @@ namespace UE::Sanitizer::RaceDetector {
 
 	private:
 		MutexType& Mutex;
-		FRWScopeLockType LockType;
+		ERWScopeLockType LockType;
 	};
 
 	// ------------------------------------------------------------------------------
@@ -793,7 +793,7 @@ namespace UE::Sanitizer::RaceDetector {
 			// #17 Requires: The failure argument shall not be memory_order_release nor memory_order_acq_rel.
 			// 
 			// Also for the time being MSVC atomic implementation always combines both order so we have no way to test our support of different SuccessOrder and FailureOrder until they fix their implementation.
-			TRWScopeLock Scope(Lock, IsAtomicOrderRelease(SuccessOrder) ? SLT_Write : SLT_ReadOnly);
+			TRWScopeLock Scope(Lock, IsAtomicOrderRelease(SuccessOrder) ? ERWScopeLockType::SLT_Write : ERWScopeLockType::SLT_ReadOnly);
 			bool bSucceeded = AtomicOp();
 
 			FAtomicMemoryOrder Order = bSucceeded ? SuccessOrder : FailureOrder;
