Fix: Debug Ensure does not send from Mac project to Crash Reporter from Editor

#rnx

diff --git a/Engine/Source/Developer/CrashDebugHelper/Private/Mac/CrashDebugHelperMac.cpp b/Engine/Source/Developer/CrashDebugHelper/Private/Mac/CrashDebugHelperMac.cpp
--- a/Engine/Source/Developer/CrashDebugHelper/Private/Mac/CrashDebugHelperMac.cpp
+++ b/Engine/Source/Developer/CrashDebugHelper/Private/Mac/CrashDebugHelperMac.cpp
@@ -287,13 +287,21 @@ static int32 ParseExceptionCode(TCHAR const* CrashLog, uint32& OutExceptionCode)
 			{
 				OutExceptionCode = SIGTRAP;
 			}
+			else if(Buffer.Find(TEXT("SIGUSR1")) != INDEX_NONE)
+			{
+				OutExceptionCode = SIGUSR1;
+			}
+			else if(Buffer.Find(TEXT("SIGUSR2")) != INDEX_NONE)
+			{
+				OutExceptionCode = SIGUSR2;
+			}
 			else if(Buffer.IsNumeric())
 			{
 				Found = swscanf(TCHAR_TO_WCHAR(*Buffer), WTEXT("%u"), &OutExceptionCode);
 			}
 			else
 			{
-				ensure(false);
+				UE_LOG(LogCrashDebugHelper, Warning, TEXT("Unrecognized exception code: %s"), *Buffer);
 				OutExceptionCode = SIGUSR1;
 			}
 		}
@@ -587,6 +595,7 @@ bool FCrashDebugHelperMac::ParseCrashDump(const FString& InCrashDumpName, FCrash
 			// Replace the binary PLCrashReporter file with an easily readable text dump
 			FFileHelper::SaveStringToFile(CrashDump, *InCrashDumpName);
 		}
+		[CrashLog release];
 	}
 	
 	if ( !CrashDump.IsEmpty() || FFileHelper::LoadFileToString( CrashDump, *InCrashDumpName ) )
@@ -647,6 +656,7 @@ bool FCrashDebugHelperMac::CreateMinidumpDiagnosticReport( const FString& InCras
 				// Replace the binary PLCrashReporter file with an easily readable text dump
 				FFileHelper::SaveStringToFile(CrashDump, *InCrashDumpName);
             }
+			[CrashLog release];
         }
     }
 	
diff --git a/Engine/Source/Runtime/Core/Private/Mac/MacPlatformMisc.cpp b/Engine/Source/Runtime/Core/Private/Mac/MacPlatformMisc.cpp
--- a/Engine/Source/Runtime/Core/Private/Mac/MacPlatformMisc.cpp
+++ b/Engine/Source/Runtime/Core/Private/Mac/MacPlatformMisc.cpp
@@ -2405,7 +2405,7 @@ void FMacCrashContext::GenerateInfoInFolder(char const* const InfoFolder) const
 	}
 }
 
-void FMacCrashContext::GenerateCrashInfoAndLaunchReporter() const
+void FMacCrashContext::GenerateInfoAndLaunchCrashReporter(bool IsCrash /* false means ensure */) const
 {
 	// Prevent CrashReportClient from spawning another CrashReportClient.
 	bool bCanRunCrashReportClient = FCString::Stristr( *(GMacAppInfo.ExecutableName), TEXT( "CrashReportClient" ) ) == nullptr;
@@ -2447,7 +2447,7 @@ void FMacCrashContext::GenerateCrashInfoAndLaunchReporter() const
 	if(bCanRunCrashReportClient)
 	{
 		// create a crash-specific directory
-		FString CrashInfoFolder = FString::Printf(TEXT("%s/CrashReport-UE-%s-pid-%d-%s"), UTF8_TO_TCHAR(GMacAppInfo.CrashReportPath), UTF8_TO_TCHAR(GMacAppInfo.AppNameUTF8), (int32)getpid(), *GMacAppInfo.RunUUID.ToString());
+		FString CrashInfoFolder = FString::Printf(TEXT("%s/%s-UE-%s-pid-%d-%s"), UTF8_TO_TCHAR(GMacAppInfo.CrashReportPath), IsCrash ? TEXT("CrashReport") : TEXT("EnsureReport"), UTF8_TO_TCHAR(GMacAppInfo.AppNameUTF8), (int32)getpid(), *GMacAppInfo.RunUUID.ToString());
 
 		// Do not inline this! The lifetime of this object needs to extend over the usage of Argv in posix_spawn() call below.
 		auto CrashInfoFolderUTF8 = TStringConversion<FTCHARToUTF8_Convert>(*CrashInfoFolder);
@@ -2465,17 +2465,31 @@ void FMacCrashContext::GenerateCrashInfoAndLaunchReporter() const
 		Argv[Argc++] = "CrashReportClient";
 		Argv[Argc++] = CrashInfoFolderUTF8.Get();
 
-		if (bImplicitSend)
-		{
-			Argv[Argc++] = "-Unattended";
-			Argv[Argc++] = "-ImplicitSend";
-		}
-		else if(GMacAppInfo.bIsUnattended)
+		if (IsCrash)
 		{
-			Argv[Argc++] ="-Unattended";
+			if (bImplicitSend)
+			{
+				Argv[Argc++] = "-Unattended";
+				Argv[Argc++] = "-ImplicitSend";
+			}
+			else if(GMacAppInfo.bIsUnattended)
+			{
+				Argv[Argc++] ="-Unattended";
+			}
+			else
+			{
+				if (!bSendUsageData)
+				{
+					Argv[Argc++] = "-NoAnalytics";
+				}
+			}
 		}
-		else
+		else	/* ensure */
 		{
+			if (!IsInteractiveEnsureMode())
+			{
+				Argv[Argc++] = "-Unattended";
+			}
 			if (!bSendUsageData)
 			{
 				Argv[Argc++] = "-NoAnalytics";
@@ -2523,97 +2537,41 @@ void FMacCrashContext::GenerateCrashInfoAndLaunchReporter() const
 		}
 	}
 
-	// Sandboxed applications re-raise the signal to trampoline into the system crash reporter as suppressing it may fall foul of Apple's Mac App Store rules.
-	// @todo Submit an application to the MAS & see whether Apple's reviewers highlight our crash reporting or trampolining to the system reporter.
-	if(GMacAppInfo.bIsSandboxed)
-	{
-		struct sigaction Action;
-		FMemory::Memzero(&Action, sizeof(struct sigaction));
-		Action.sa_handler = SIG_DFL;
-		sigemptyset(&Action.sa_mask);
-		sigaction(SIGQUIT, &Action, NULL);
-		sigaction(SIGILL, &Action, NULL);
-		sigaction(SIGEMT, &Action, NULL);
-		sigaction(SIGFPE, &Action, NULL);
-		sigaction(SIGBUS, &Action, NULL);
-		sigaction(SIGSEGV, &Action, NULL);
-		sigaction(SIGSYS, &Action, NULL);
-		sigaction(SIGABRT, &Action, NULL);
-		sigaction(SIGTRAP, &Action, NULL);
-		
-		raise(Signal);
-	}
-	
-	_Exit(1);
-}
-
-void FMacCrashContext::GenerateEnsureInfoAndLaunchReporter() const
-{
-	// Prevent CrashReportClient from spawning another CrashReportClient.
-	bool bCanRunCrashReportClient = FCString::Stristr( *(GMacAppInfo.ExecutableName), TEXT( "CrashReportClient" ) ) == nullptr;
-	
-	bool bSendUnattendedBugReports = true;
-	if (GConfig)
-	{
-		GConfig->GetBool(TEXT("/Script/UnrealEd.CrashReportsPrivacySettings"), TEXT("bSendUnattendedBugReports"), bSendUnattendedBugReports, GEditorSettingsIni);
-	}
-
-	bool bSendUsageData = true;
-	if (GConfig)
-	{
-		GConfig->GetBool(TEXT("/Script/UnrealEd.AnalyticsPrivacySettings"), TEXT("bSendUsageData"), bSendUsageData, GEditorSettingsIni);
-	}
-
-	// NOTE: A blueprint-only game packaged from a vanilla engine downloaded from Epic Game Store isn't considered a 'Licensee' version because the engine was built by Epic.
-	//       There is no way at the moment do distinguish this case properly.
-	if (BuildSettings::IsLicenseeVersion() && !UE_EDITOR)
+	if(IsCrash)
 	{
-		// do not send unattended reports in licensees' builds except for the editor, where it is governed by the above setting
-		bSendUnattendedBugReports = false;
-		bSendUsageData = false;
-	}
-
-	const bool bUnattended = GMacAppInfo.bIsUnattended || !IsInteractiveEnsureMode() || IsRunningDedicatedServer();
-	if (bUnattended && !bSendUnattendedBugReports)
-	{
-		bCanRunCrashReportClient = false;
-	}
-
-	if(bCanRunCrashReportClient)
-	{
-		SCOPED_AUTORELEASE_POOL;
-		
-		// Write the PLCrashReporter report to the expected location
-		NSData* CrashReport = [FMacApplicationInfo::CrashReporter generateLiveReport];
-		[CrashReport writeToFile:[NSString stringWithUTF8String:GMacAppInfo.PLCrashReportPath] atomically:YES];
-
-		// Use a slightly different output folder name to not conflict with a subequent crash
-		const FGuid Guid = FGuid::NewGuid();
-		FString GameName = FApp::GetProjectName();
-		FString EnsureLogFolder = FString(GMacAppInfo.CrashReportPath) / FString::Printf(TEXT("EnsureReport-%s-%s"), *GameName, *Guid.ToString(EGuidFormats::Digits));
-		
-		CaptureAllThreadContext(0);
-		GenerateInfoInFolder(TCHAR_TO_UTF8(*EnsureLogFolder));
-		
-		FString Arguments;
-		if (IsInteractiveEnsureMode())
-		{
-			Arguments = FString::Printf(TEXT("\"%s/\""), *EnsureLogFolder);
-		}
-		else
+		// Sandboxed applications re-raise the signal to trampoline into the system crash reporter as suppressing it may fall foul of Apple's Mac App Store rules.
+		// @todo Submit an application to the MAS & see whether Apple's reviewers highlight our crash reporting or trampolining to the system reporter.
+		if (GMacAppInfo.bIsSandboxed)
 		{
-			Arguments = FString::Printf(TEXT("\"%s/\" -Unattended"), *EnsureLogFolder);
+			struct sigaction Action;
+			FMemory::Memzero(&Action, sizeof(struct sigaction));
+			Action.sa_handler = SIG_DFL;
+			sigemptyset(&Action.sa_mask);
+			sigaction(SIGQUIT, &Action, NULL);
+			sigaction(SIGILL, &Action, NULL);
+			sigaction(SIGEMT, &Action, NULL);
+			sigaction(SIGFPE, &Action, NULL);
+			sigaction(SIGBUS, &Action, NULL);
+			sigaction(SIGSEGV, &Action, NULL);
+			sigaction(SIGSYS, &Action, NULL);
+			sigaction(SIGABRT, &Action, NULL);
+			sigaction(SIGTRAP, &Action, NULL);
+			
+			raise(Signal);
 		}
+		
+		_Exit(1);
+	} 
+}
 
-		// If the editor setting has been disabled to not send analytics extend this to the CRC
-		if (!bSendUsageData)
-		{
-			Arguments += TEXT(" -NoAnalytics");
-		}
+void FMacCrashContext::GenerateCrashInfoAndLaunchReporter() const
+{
+	GenerateInfoAndLaunchCrashReporter(true);
+}
 
-		FString ReportClient = FPaths::ConvertRelativePathToFull(FPlatformProcess::GenerateApplicationPath(TEXT("CrashReportClient"), EBuildConfiguration::Development));
-		FPlatformProcess::ExecProcess(*ReportClient, *Arguments, nullptr, nullptr, nullptr);
-	}
+void FMacCrashContext::GenerateEnsureInfoAndLaunchReporter() const
+{
+	GenerateInfoAndLaunchCrashReporter(false);
 }
 
 static FString AllThreadContexts;
diff --git a/Engine/Source/Runtime/Core/Public/Mac/MacPlatformCrashContext.h b/Engine/Source/Runtime/Core/Public/Mac/MacPlatformCrashContext.h
--- a/Engine/Source/Runtime/Core/Public/Mac/MacPlatformCrashContext.h
+++ b/Engine/Source/Runtime/Core/Public/Mac/MacPlatformCrashContext.h
@@ -15,6 +15,9 @@ struct CORE_API FMacCrashContext : public FApplePlatformCrashContext
 	/** Generates the ensure/crash info into the given folder */
 	void GenerateInfoInFolder(char const* const InfoFolder) const;
 	
+	/** Generates crash or ensure information for reporter */
+	void GenerateInfoAndLaunchCrashReporter(bool IsCrash /* false means ensure */) const;
+	
 	/** Generates information for crash reporter */
 	void GenerateCrashInfoAndLaunchReporter() const;
 	
