Added a type parameter to UE_OPS_NAMESPACE_VISIBLE for clarity and future compatibility.
Moved UE_OPS_NAMESPACE_VISIBLE comments to the top of UEOps.h for clarity.

#rnx

diff --git a/Engine/Source/Runtime/Core/Public/Misc/UEOps.h b/Engine/Source/Runtime/Core/Public/Misc/UEOps.h
--- a/Engine/Source/Runtime/Core/Public/Misc/UEOps.h
+++ b/Engine/Source/Runtime/Core/Public/Misc/UEOps.h
@@ -21,6 +21,34 @@
  *
  * bool FMyType::UEOpGreaterThan(const FMyType&) const;
  *   Will never be called if FMyType::UEOpLessThan is defined, so is redundant.
+ *
+ * When a type in a namespace uses UEOps functions, comparisons can fail when calling the operators from an
+ * unrelated namespace, where the global operators are hidden and argument-dependent lookup cannot find the
+ * operators because they are not in the same namespace as the type.
+ *
+ * UE_OPS_NAMESPACE_VISIBLE makes the operators visible in that type's namespace, allowing ADL to work as
+ * intended. Example:
+ *
+ * namespace N1
+ * {
+ *     struct FMyType
+ *     {
+ *         bool UEOpEquals(const FMyType& Other) const;
+ *     };
+ *
+ *     UE_OPS_NAMESPACE_VISIBLE(FMyType)
+ * }
+ *
+ * namespace N2
+ * {
+ *     struct FOtherType;
+ *     bool operator==(const FOtherType& Lhs, const FOtherType& Rhs);
+ *
+ *     bool AreDifferent(const N1::FMyType& Lhs, const N1::FMyType& Rhs)
+ *     {
+ *         return Lhs != Rhs; // without UE_OPS_NAMESPACE_VISIBLE above, this will not compile.
+ *     }
+ * }
  */
 
 
@@ -176,35 +204,10 @@ template <typename LhsType, typename RhsType>
 	}
 }
 
-// When a type in a namespace uses UEOps functions, comparisons can fail when calling the operators from an
-// unrelated namespace, where the global operators are hidden and argument-dependent lookup cannot find the
-// operators because they are not in the same namespace as the type.
-//
-// UE_OPS_NAMESPACE_VISIBLE makes the operators visible in that type's namespace, allowing ADL to work as
-// intended. Example:
-//
-// namespace N1
-// {
-//     struct FMyType
-//     {
-//         bool UEOpEquals(const FMyType& Other) const;
-//     };
-//
-//     UE_OPS_NAMESPACE_VISIBLE
-// }
-//
-// namespace N2
-// {
-//     struct FOtherType;
-//     bool operator==(const FOtherType& Lhs, const FOtherType& Rhs);
-//
-//     bool AreDifferent(const N1::FMyType& Lhs, const N1::FMyType& Rhs)
-//     {
-//         return Lhs != Rhs; // without UE_OPS_NAMESPACE_VISIBLE above, this will not compile.
-//     }
-// }
 #if UE_OPS_WORKAROUND_PRIVATE
-	#define UE_OPS_NAMESPACE_VISIBLE \
+	// Allows UEOps to be used on the given type when called from an unrelated namespace.
+	#define UE_OPS_NAMESPACE_VISIBLE(Type) \
+		static_assert(sizeof(Type) > 0); /* We don't use the Type parameter right now - just check that the type exists as a basic usability check */ \
 		using ::operator==; \
 		using ::operator!=; \
 		using ::operator<; \
@@ -212,7 +215,9 @@ template <typename LhsType, typename RhsType>
 		using ::operator<=; \
 		using ::operator>=;
 #else
-	#define UE_OPS_NAMESPACE_VISIBLE \
+	// Allows UEOps to be used on the given type when called from an unrelated namespace.
+	#define UE_OPS_NAMESPACE_VISIBLE(Type) \
+		static_assert(sizeof(Type) > 0); /* We don't use the Type parameter right now - just check that the type exists as a basic usability check */ \
 		using ::operator==; \
 		using ::operator<; \
 		using ::operator>; \
diff --git a/Engine/Source/Runtime/CoreUObject/Public/UObject/ObjectHandle.h b/Engine/Source/Runtime/CoreUObject/Public/UObject/ObjectHandle.h
--- a/Engine/Source/Runtime/CoreUObject/Public/UObject/ObjectHandle.h
+++ b/Engine/Source/Runtime/CoreUObject/Public/UObject/ObjectHandle.h
@@ -126,7 +126,7 @@ namespace UE::CoreUObject::Private
 	};
 
 	// Allow generated operators to be accessible from comparisons in other namespaces.
-	UE_OPS_NAMESPACE_VISIBLE
+	UE_OPS_NAMESPACE_VISIBLE(FObjectHandlePrivate)
 
 	/* Returns the packed object ref for this object IF one exists otherwise returns a null PackedObjectRef */
 	COREUOBJECT_API FPackedObjectRef FindExistingPackedObjectRef(const UObject* Object);
@@ -171,7 +171,7 @@ namespace UE::CoreUObject::Private
 	};
 
 	// Allow generated operators to be accessible from comparisons in other namespaces.
-	UE_OPS_NAMESPACE_VISIBLE
+	UE_OPS_NAMESPACE_VISIBLE(FRemoteObjectHandlePrivate)
 
 #endif
 
