Fix bug with FUnitConversion converting values that did not need converting.
Resulted in values that were slightly off due to conversion precision limits.
#FailedValidations Preflight

diff --git a/Engine/Source/Runtime/Core/Public/Math/UnitConversion.inl b/Engine/Source/Runtime/Core/Public/Math/UnitConversion.inl
--- a/Engine/Source/Runtime/Core/Public/Math/UnitConversion.inl
+++ b/Engine/Source/Runtime/Core/Public/Math/UnitConversion.inl
@@ -83,6 +83,10 @@ T FUnitConversion::Convert(T InValue, EUnit From, EUnit To)
 	{
 		return InValue;
 	}
+	else if (From == To)
+	{
+		return InValue;
+	}
 
 	switch(FUnitConversion::GetUnitType(From))
 	{
