Added static verification of expected definitions of UEOps functions.  Fixed TStringView which wasn't const-correct.

diff --git a/Engine/Source/Runtime/Core/Public/Containers/StringView.h b/Engine/Source/Runtime/Core/Public/Containers/StringView.h
--- a/Engine/Source/Runtime/Core/Public/Containers/StringView.h
+++ b/Engine/Source/Runtime/Core/Public/Containers/StringView.h
@@ -468,13 +468,13 @@ public:
 	}
 
 	template <UE::CConvertibleTo<TStringView> CharRangeType>
-	[[nodiscard]] UE_REWRITE auto UEOpLessThan(CharRangeType&& Rhs)
+	[[nodiscard]] UE_REWRITE auto UEOpLessThan(CharRangeType&& Rhs) const
 	{
 		return TStringView::PrivateLess(*this, ImplicitConv<TStringView>(Forward<CharRangeType>(Rhs)));
 	}
 
 	template <UE::CConvertibleTo<TStringView> CharRangeType>
-	[[nodiscard]] UE_REWRITE auto UEOpGreaterThan(CharRangeType&& Rhs)
+	[[nodiscard]] UE_REWRITE auto UEOpGreaterThan(CharRangeType&& Rhs) const
 	{
 		return TStringView::PrivateGreater(*this, ImplicitConv<TStringView>(Forward<CharRangeType>(Rhs)));
 	}
diff --git a/Engine/Source/Runtime/Core/Public/Misc/UEOps.h b/Engine/Source/Runtime/Core/Public/Misc/UEOps.h
--- a/Engine/Source/Runtime/Core/Public/Misc/UEOps.h
+++ b/Engine/Source/Runtime/Core/Public/Misc/UEOps.h
@@ -78,29 +78,109 @@ namespace UE::Core::Private
 	struct FIncomplete;
 
 	template <typename LhsType, typename RhsType>
-	concept CWithUEOpEquals = requires(const LhsType & Lhs, const RhsType & Rhs)
+	concept CWithUEOpEquals = requires(LhsType& Lhs, RhsType& Rhs)
 	{
 		{ Lhs.UEOpEquals(Rhs) } -> UE::CSameAs<bool>;
 	};
 
 	template <typename LhsType, typename RhsType>
-	concept CWithUEOpLessThan = requires(const LhsType & Lhs, const RhsType & Rhs)
+	concept CWithConstUEOpEquals = requires(const LhsType& Lhs, const RhsType& Rhs)
+	{
+		{ Lhs.UEOpEquals(Rhs) } -> UE::CSameAs<bool>;
+	};
+
+	template <typename LhsType, typename RhsType>
+	concept CWithUEOpLessThan = requires(LhsType& Lhs, RhsType& Rhs)
+	{
+		{ Lhs.UEOpLessThan(Rhs) } -> UE::CSameAs<bool>;
+	};
+
+	template <typename LhsType, typename RhsType>
+	concept CWithConstUEOpLessThan = requires(const LhsType& Lhs, const RhsType& Rhs)
 	{
 		{ Lhs.UEOpLessThan(Rhs) } -> UE::CSameAs<bool>;
 	};
 
 	template <typename LhsType, typename RhsType>
-	concept CWithUEOpGreaterThan = requires(const LhsType & Lhs, const RhsType & Rhs)
+	concept CWithUEOpGreaterThan = requires(LhsType& Lhs, RhsType& Rhs)
 	{
 		{ Lhs.UEOpGreaterThan(Rhs) } -> UE::CSameAs<bool>;
 	};
+
+	template <typename LhsType, typename RhsType>
+	concept CWithConstUEOpGreaterThan = requires(const LhsType& Lhs, const RhsType& Rhs)
+	{
+		{ Lhs.UEOpGreaterThan(Rhs) } -> UE::CSameAs<bool>;
+	};
+
+	template <typename LhsType, typename RhsType>
+	concept CWithMismatchedUEOpLess = CWithUEOpLessThan<LhsType, RhsType> && !CWithUEOpLessThan<RhsType, LhsType> && !CWithUEOpGreaterThan<LhsType, RhsType>;
+
+	template <typename LhsType, typename RhsType>
+	concept CWithMismatchedUEOpGreater = CWithUEOpGreaterThan<LhsType, RhsType> && !CWithUEOpGreaterThan   <RhsType, LhsType> && !CWithUEOpLessThan<LhsType, RhsType>;
+
+	template <typename LhsType, typename RhsType>
+	concept CWithMismatchedConstUEOpEquals = CWithUEOpEquals<LhsType, RhsType> && !CWithConstUEOpEquals<LhsType, RhsType>;
+
+	template <typename LhsType, typename RhsType>
+	concept CWithMismatchedConstUEOpLess = CWithUEOpLessThan<LhsType, RhsType> && !CWithConstUEOpLessThan<LhsType, RhsType>;
+
+	template <typename LhsType, typename RhsType>
+	concept CWithMismatchedConstUEOpGreater = CWithUEOpGreaterThan<LhsType, RhsType> && !CWithConstUEOpGreaterThan<LhsType, RhsType>;
+
+	template <typename LhsType, typename RhsType>
+	constexpr bool ValidateUEOpsEquals()
+	{
+		if constexpr (CWithMismatchedConstUEOpEquals<LhsType, RhsType> || CWithMismatchedConstUEOpEquals<RhsType, LhsType>)
+		{
+			static_assert(sizeof(LhsType*) == 0, "UEOpEquals should be const and accept a const argument");
+			return false;
+		}
+		else
+		{
+			return true;
+		}
+	}
+
+	template <typename LhsType, typename RhsType>
+	constexpr bool ValidateUEOpsLessAndGreater()
+	{
+		if constexpr (CWithMismatchedConstUEOpLess<LhsType, RhsType> || CWithMismatchedConstUEOpLess<RhsType, LhsType>)
+		{
+			static_assert(sizeof(LhsType*) == 0, "UEOpLessThan should be const and accept a const argument");
+			return false;
+		}
+		else if constexpr (CWithMismatchedConstUEOpGreater<LhsType, RhsType> || CWithMismatchedConstUEOpGreater<RhsType, LhsType>)
+		{
+			static_assert(sizeof(LhsType*) == 0, "UEOpGreaterThan should be const and accept a const argument");
+			return false;
+		}
+		else if constexpr (CWithMismatchedUEOpLess<LhsType, RhsType> || CWithMismatchedUEOpLess<RhsType, LhsType>)
+		{
+			static_assert(sizeof(LhsType*) == 0, "UEOpLessThan should be matched with UEOpGreaterThan of the same argument type");
+			return false;
+		}
+		else if constexpr (CWithMismatchedUEOpGreater<LhsType, RhsType> || CWithMismatchedUEOpGreater<RhsType, LhsType>)
+		{
+			static_assert(sizeof(LhsType*) == 0, "UEOpGreaterThan should be matched with UEOpLessThan of the same argument type");
+			return false;
+		}
+		else
+		{
+			return true;
+		}
+	}
 }
 
 template <typename LhsType, typename RhsType>
 	requires UE::Core::Private::CWithUEOpEquals<LhsType, RhsType> || UE::Core::Private::CWithUEOpEquals<RhsType, LhsType>
 [[nodiscard]] UE_REWRITE constexpr bool operator==(const LhsType& Lhs, const RhsType& Rhs)
 {
-	if constexpr (UE::Core::Private::CWithUEOpEquals<LhsType, RhsType>)
+	if constexpr (!UE::Core::Private::ValidateUEOpsEquals<LhsType, RhsType>())
+	{
+		return false;
+	}
+	else if constexpr (UE::Core::Private::CWithUEOpEquals<LhsType, RhsType>)
 	{
 		return Lhs.UEOpEquals(Rhs);
 	}
@@ -115,7 +195,14 @@ template <typename LhsType>
 	requires UE::Core::Private::CWithUEOpEquals<LhsType, decltype(nullptr)>
 [[nodiscard]] UE_REWRITE constexpr bool operator==(const LhsType& Lhs, int UE::Core::Private::FIncomplete::* Rhs)
 {
-	return Lhs.UEOpEquals(nullptr);
+	if constexpr (!UE::Core::Private::ValidateUEOpsEquals<LhsType, decltype(nullptr)>())
+	{
+		return false;
+	}
+	else
+	{
+		return Lhs.UEOpEquals(nullptr);
+	}
 }
 
 #if UE_OPS_WORKAROUND_PRIVATE
@@ -123,7 +210,11 @@ template <typename LhsType>
 		requires UE::Core::Private::CWithUEOpEquals<LhsType, RhsType> || UE::Core::Private::CWithUEOpEquals<RhsType, LhsType>
 	[[nodiscard]] UE_REWRITE constexpr bool operator!=(const LhsType& Lhs, const RhsType& Rhs)
 	{
-		if constexpr (UE::Core::Private::CWithUEOpEquals<LhsType, RhsType>)
+		if constexpr (!UE::Core::Private::ValidateUEOpsEquals<LhsType, RhsType>())
+		{
+			return false;
+		}
+		else if constexpr (UE::Core::Private::CWithUEOpEquals<LhsType, RhsType>)
 		{
 			return !Lhs.UEOpEquals(Rhs);
 		}
@@ -137,68 +228,98 @@ template <typename LhsType>
 		requires UE::Core::Private::CWithUEOpEquals<LhsType, decltype(nullptr)>
 	[[nodiscard]] UE_REWRITE constexpr bool operator!=(const LhsType& Lhs, int UE::Core::Private::FIncomplete::* Rhs)
 	{
-		return !Lhs.UEOpEquals(nullptr);
+		if constexpr (!UE::Core::Private::ValidateUEOpsEquals<LhsType, decltype(nullptr)>())
+		{
+			return false;
+		}
+		else
+		{
+			return !Lhs.UEOpEquals(nullptr);
+		}
 	}
 
 	template <typename RhsType>
 		requires UE::Core::Private::CWithUEOpEquals<RhsType, decltype(nullptr)>
 	[[nodiscard]] UE_REWRITE constexpr bool operator!=(int UE::Core::Private::FIncomplete::* Lhs, const RhsType& Rhs)
 	{
-		return !Rhs.UEOpEquals(nullptr);
+		if constexpr (!UE::Core::Private::ValidateUEOpsEquals<decltype(nullptr), RhsType>())
+		{
+			return false;
+		}
+		else
+		{
+			return !Rhs.UEOpEquals(nullptr);
+		}
 	}
 #endif
 
 template <typename LhsType, typename RhsType>
-	requires UE::Core::Private::CWithUEOpLessThan<LhsType, RhsType> || UE::Core::Private::CWithUEOpGreaterThan<RhsType, LhsType>
+	requires UE::Core::Private::CWithUEOpLessThan<LhsType, RhsType> || UE::Core::Private::CWithUEOpLessThan<RhsType, LhsType> || UE::Core::Private::CWithUEOpGreaterThan<LhsType, RhsType> || UE::Core::Private::CWithUEOpGreaterThan<RhsType, LhsType>
 [[nodiscard]] UE_REWRITE constexpr bool operator<(const LhsType& Lhs, const RhsType& Rhs)
 {
-	if constexpr (UE::Core::Private::CWithUEOpLessThan<LhsType, RhsType>)
+	if constexpr (!UE::Core::Private::ValidateUEOpsLessAndGreater<LhsType, RhsType>())
+	{
+		return false;
+	}
+	else if constexpr (UE::Core::Private::CWithUEOpLessThan<LhsType, RhsType>)
 	{
 		return Lhs.UEOpLessThan(Rhs);
 	}
-	else
+	else // if constexpr (UE::Core::Private::CWithUEOpGreaterThan<RhsType, LhsType>)
 	{
 		return Rhs.UEOpGreaterThan(Lhs);
 	}
 }
 
 template <typename LhsType, typename RhsType>
-	requires UE::Core::Private::CWithUEOpLessThan<RhsType, LhsType> || UE::Core::Private::CWithUEOpGreaterThan<LhsType, RhsType>
+	requires UE::Core::Private::CWithUEOpLessThan<LhsType, RhsType> || UE::Core::Private::CWithUEOpLessThan<RhsType, LhsType> || UE::Core::Private::CWithUEOpGreaterThan<LhsType, RhsType> || UE::Core::Private::CWithUEOpGreaterThan<RhsType, LhsType>
 [[nodiscard]] UE_REWRITE constexpr bool operator>(const LhsType& Lhs, const RhsType& Rhs)
 {
-	if constexpr (UE::Core::Private::CWithUEOpLessThan<RhsType, LhsType>)
+	if constexpr (!UE::Core::Private::ValidateUEOpsLessAndGreater<LhsType, RhsType>())
+	{
+		return false;
+	}
+	else if constexpr (UE::Core::Private::CWithUEOpLessThan<RhsType, LhsType>)
 	{
 		return Rhs.UEOpLessThan(Lhs);
 	}
-	else
+	else // if constexpr (UE::Core::Private::CWithUEOpGreaterThan<LhsType, RhsType>)
 	{
 		return Lhs.UEOpGreaterThan(Rhs);
 	}
 }
 
 template <typename LhsType, typename RhsType>
-	requires UE::Core::Private::CWithUEOpLessThan<LhsType, RhsType> || UE::Core::Private::CWithUEOpGreaterThan<RhsType, LhsType>
+	requires UE::Core::Private::CWithUEOpLessThan<LhsType, RhsType> || UE::Core::Private::CWithUEOpLessThan<RhsType, LhsType> || UE::Core::Private::CWithUEOpGreaterThan<LhsType, RhsType> || UE::Core::Private::CWithUEOpGreaterThan<RhsType, LhsType>
 [[nodiscard]] UE_REWRITE constexpr bool operator>=(const LhsType& Lhs, const RhsType& Rhs)
 {
-	if constexpr (UE::Core::Private::CWithUEOpLessThan<LhsType, RhsType>)
+	if constexpr (!UE::Core::Private::ValidateUEOpsLessAndGreater<LhsType, RhsType>())
+	{
+		return false;
+	}
+	else if constexpr (UE::Core::Private::CWithUEOpLessThan<LhsType, RhsType>)
 	{
 		return !Lhs.UEOpLessThan(Rhs);
 	}
-	else
+	else // if constexpr (UE::Core::Private::CWithUEOpGreaterThan<RhsType, LhsType>)
 	{
 		return !Rhs.UEOpGreaterThan(Lhs);
 	}
 }
 
 template <typename LhsType, typename RhsType>
-	requires UE::Core::Private::CWithUEOpLessThan<RhsType, LhsType> || UE::Core::Private::CWithUEOpGreaterThan<LhsType, RhsType>
+	requires UE::Core::Private::CWithUEOpLessThan<LhsType, RhsType> || UE::Core::Private::CWithUEOpLessThan<RhsType, LhsType> || UE::Core::Private::CWithUEOpGreaterThan<LhsType, RhsType> || UE::Core::Private::CWithUEOpGreaterThan<RhsType, LhsType>
 [[nodiscard]] UE_REWRITE constexpr bool operator<=(const LhsType& Lhs, const RhsType& Rhs)
 {
-	if constexpr (UE::Core::Private::CWithUEOpLessThan<RhsType, LhsType>)
+	if constexpr (!UE::Core::Private::ValidateUEOpsLessAndGreater<LhsType, RhsType>())
+	{
+		return false;
+	}
+	else if constexpr (UE::Core::Private::CWithUEOpLessThan<RhsType, LhsType>)
 	{
 		return !Rhs.UEOpLessThan(Lhs);
 	}
-	else
+	else // if constexpr (UE::Core::Private::CWithUEOpGreaterThan<LhsType, RhsType>)
 	{
 		return !Lhs.UEOpGreaterThan(Rhs);
 	}
