Clarified a comment in UE_OPS_NAMESPACE_VISIBLE about the type needing to be complete as well as existing.

#rnx

diff --git a/Engine/Source/Runtime/Core/Public/Misc/UEOps.h b/Engine/Source/Runtime/Core/Public/Misc/UEOps.h
--- a/Engine/Source/Runtime/Core/Public/Misc/UEOps.h
+++ b/Engine/Source/Runtime/Core/Public/Misc/UEOps.h
@@ -328,7 +328,7 @@ template <typename LhsType, typename RhsType>
 #if UE_OPS_WORKAROUND_PRIVATE
 	// Allows UEOps to be used on the given type when called from an unrelated namespace.
 	#define UE_OPS_NAMESPACE_VISIBLE(Type) \
-		static_assert(sizeof(Type) > 0); /* We don't use the Type parameter right now - just check that the type exists as a basic usability check */ \
+		static_assert(sizeof(Type) > 0); /* We don't use the Type parameter right now - just check that the name exists and the macro is placed after the full type definition as a basic usability check */ \
 		using ::operator==; \
 		using ::operator!=; \
 		using ::operator<; \
@@ -338,7 +338,7 @@ template <typename LhsType, typename RhsType>
 #else
 	// Allows UEOps to be used on the given type when called from an unrelated namespace.
 	#define UE_OPS_NAMESPACE_VISIBLE(Type) \
-		static_assert(sizeof(Type) > 0); /* We don't use the Type parameter right now - just check that the type exists as a basic usability check */ \
+		static_assert(sizeof(Type) > 0); /* We don't use the Type parameter right now - just check that the name exists and the macro is placed after the full type definition as a basic usability check */ \
 		using ::operator==; \
 		using ::operator<; \
 		using ::operator>; \
