Fixed undefined behavior when serializing packed ints.
The fix is a more correct shift operation when dealing with int64 values saved as int32.

#rnx

diff --git a/Engine/Source/Runtime/Core/Private/Serialization/Archive.cpp b/Engine/Source/Runtime/Core/Private/Serialization/Archive.cpp
--- a/Engine/Source/Runtime/Core/Private/Serialization/Archive.cpp
+++ b/Engine/Source/Runtime/Core/Private/Serialization/Archive.cpp
@@ -1346,7 +1346,7 @@ void FArchive::SerializeIntPacked(uint32& Value)
 
 			more = NextByte & 1;				// Check 1 bit to see if theres more after this
 			NextByte = NextByte >> 1;			// Shift to get actual 7 bit value
-			TempValue += (uint32)NextByte << (7 * cnt++);	// Add to total value
+			TempValue += (uint32)((uint64)NextByte << (7 * cnt++));	// Add to total value
 		}
 
 		Value = TempValue;
