Fix for decltype on angle delta functions.

diff --git a/Engine/Source/Runtime/Core/Public/Math/UnrealMathUtility.h b/Engine/Source/Runtime/Core/Public/Math/UnrealMathUtility.h
--- a/Engine/Source/Runtime/Core/Public/Math/UnrealMathUtility.h
+++ b/Engine/Source/Runtime/Core/Public/Math/UnrealMathUtility.h
@@ -894,7 +894,7 @@ public:
 		typename T2
 		UE_REQUIRES(std::is_floating_point_v<T> || std::is_floating_point_v<T2>)
 	>
-	[[nodiscard]] static constexpr auto FindDeltaAngleDegrees(T A1, T2 A2) -> decltype(A1 - A2)
+	[[nodiscard]] static constexpr auto FindDeltaAngleDegrees(T A1, T2 A2) -> decltype(A2 - A1)
 	{
 		// Find the difference
 		auto Delta = A2 - A1;
@@ -911,7 +911,7 @@ public:
 		typename T2
 		UE_REQUIRES(std::is_floating_point_v<T> || std::is_floating_point_v<T2>)
 	>
-	[[nodiscard]] static constexpr auto FindDeltaAngleRadians(T A1, T2 A2) -> decltype(A1 - A2)
+	[[nodiscard]] static constexpr auto FindDeltaAngleRadians(T A1, T2 A2) -> decltype(A2 - A1)
 	{
 		// Find the difference
 		auto Delta = A2 - A1;
