[RaceDetector]
- Notify visual studio as soon as it attaches itself to not show exceptions

#FailedValidations Preflight

diff --git a/Engine/Source/Runtime/Core/Private/Sanitizer/RaceDetectorWindows.cpp b/Engine/Source/Runtime/Core/Private/Sanitizer/RaceDetectorWindows.cpp
--- a/Engine/Source/Runtime/Core/Private/Sanitizer/RaceDetectorWindows.cpp
+++ b/Engine/Source/Runtime/Core/Private/Sanitizer/RaceDetectorWindows.cpp
@@ -1300,6 +1300,21 @@ INSTRUMENTATION_FUNCTION_ATTRIBUTES BOOL WINAPI DllMain(
 	return TRUE;
 }
 
+typedef NTSTATUS(NTAPI DBGUIREMOTEBREAKIN)
+(
+	IN HANDLE hProcess
+);
+
+NTSTATUS(NTAPI* TrueDbgUiRemoteBreakin)(HANDLE hProcess);
+INSTRUMENTATION_FUNCTION_ATTRIBUTES NTSTATUS NTAPI DetouredDbgUiRemoteBreakin(HANDLE hProcess)
+{
+        // Notify visual studio as soon as it attaches itself to avoid seeing unhandled exceptions
+	// for late attachments.
+	Platform::HideFirstChanceExceptionInVisualStudio();
+
+	return TrueDbgUiRemoteBreakin(hProcess);
+}
+
 INSTRUMENTATION_FUNCTION_ATTRIBUTES int DetourBootFunctions()
 {
 	DetourTransactionBegin();
@@ -1311,6 +1326,8 @@ INSTRUMENTATION_FUNCTION_ATTRIBUTES int DetourBootFunctions()
 	DetourAttach(&True_beginthreadex, Detoured_beginthreadex);
 	TrueExitThread = ::ExitThread;
 	DetourAttach(&TrueExitThread, DetouredExitThread);
+	TrueDbgUiRemoteBreakin = (DBGUIREMOTEBREAKIN*)GetProcAddress(GetModuleHandleA("ntdll"), "DbgUiRemoteBreakin");
+	DetourAttach(&TrueDbgUiRemoteBreakin, DetouredDbgUiRemoteBreakin);
 	TrueFreeLibraryAndExitThread = ::FreeLibraryAndExitThread;
 	DetourAttach(&TrueFreeLibraryAndExitThread, DetouredFreeLibraryAndExitThread);
 	True_Cnd_do_broadcast_at_thread_exit = ::_Cnd_do_broadcast_at_thread_exit;
