Replaced PREPROCESSOR_JOIN with UE_JOIN.

diff --git a/Engine/Plugins/Compression/OodleNetwork/Source/Private/OodleNetworkHandlerComponent.cpp b/Engine/Plugins/Compression/OodleNetwork/Source/Private/OodleNetworkHandlerComponent.cpp
--- a/Engine/Plugins/Compression/OodleNetwork/Source/Private/OodleNetworkHandlerComponent.cpp
+++ b/Engine/Plugins/Compression/OodleNetwork/Source/Private/OodleNetworkHandlerComponent.cpp
@@ -98,28 +98,28 @@ namespace UE::Oodle
 
 /** To be placed outside and before the scope of the measured code */
 #define UE_OODLE_LIGHTWEIGHT_TIME_GUARD_DECLARE(Name, ThresholdMS) \
-	const double PREPROCESSOR_JOIN(__TimeGuard_ThresholdMS_, Name) = ThresholdMS; \
-	double PREPROCESSOR_JOIN(__TimeGuard_MSElapsed_, Name) = 0.0;
+	const double UE_JOIN(__TimeGuard_ThresholdMS_, Name) = ThresholdMS; \
+	double UE_JOIN(__TimeGuard_MSElapsed_, Name) = 0.0;
 
 /** To be placed inside the scope, directly before the measured code */
 #define UE_OODLE_LIGHTWEIGHT_TIME_GUARD_BEGIN(Name) \
-	uint64 PREPROCESSOR_JOIN(__TimeGuard_StartCycles_, Name) = \
-		(PREPROCESSOR_JOIN(__TimeGuard_ThresholdMS_, Name) > 0.0 && UE::Oodle::GOodleTimeguardLimit > 0) ? FPlatformTime::Cycles64() : 0;
+	uint64 UE_JOIN(__TimeGuard_StartCycles_, Name) = \
+		(UE_JOIN(__TimeGuard_ThresholdMS_, Name) > 0.0 && UE::Oodle::GOodleTimeguardLimit > 0) ? FPlatformTime::Cycles64() : 0;
 
 /** To be placed inside the scope, directly after the measured code */
 #define UE_OODLE_LIGHTWEIGHT_TIME_GUARD_END(Name) \
-	if (PREPROCESSOR_JOIN(__TimeGuard_ThresholdMS_, Name) > 0.0 && UE::Oodle::GOodleTimeguardLimit > 0) \
+	if (UE_JOIN(__TimeGuard_ThresholdMS_, Name) > 0.0 && UE::Oodle::GOodleTimeguardLimit > 0) \
 	{ \
-		PREPROCESSOR_JOIN(__TimeGuard_MSElapsed_, Name) = \
-			FPlatformTime::ToMilliseconds64(FPlatformTime::Cycles64() - PREPROCESSOR_JOIN(__TimeGuard_StartCycles_, Name)); \
+		UE_JOIN(__TimeGuard_MSElapsed_, Name) = \
+			FPlatformTime::ToMilliseconds64(FPlatformTime::Cycles64() - UE_JOIN(__TimeGuard_StartCycles_, Name)); \
 	}
 
 /** To be placed outside and after the scope of the measured code */
 #define UE_OODLE_LIGHTWEIGHT_TIME_GUARD_REPORT(Name) \
-	if (PREPROCESSOR_JOIN(__TimeGuard_MSElapsed_, Name) > PREPROCESSOR_JOIN(__TimeGuard_ThresholdMS_, Name)) \
+	if (UE_JOIN(__TimeGuard_MSElapsed_, Name) > UE_JOIN(__TimeGuard_ThresholdMS_, Name)) \
 	{ \
 		UE_LOG(OodleNetworkHandlerComponentLog, Warning, TEXT("OodleNetworkHandlerComponent '%s' took %.2fms!"), TEXT(#Name), \
-				PREPROCESSOR_JOIN(__TimeGuard_MSElapsed_, Name)); \
+				UE_JOIN(__TimeGuard_MSElapsed_, Name)); \
 		UE::Oodle::GOodleTimeguardLimit--; \
 	}
 
diff --git a/Engine/Plugins/Developer/Concert/ConcertMain/Source/ConcertTransport/Public/Trace/Private/ConcertProtocolTraceInternal.h b/Engine/Plugins/Developer/Concert/ConcertMain/Source/ConcertTransport/Public/Trace/Private/ConcertProtocolTraceInternal.h
--- a/Engine/Plugins/Developer/Concert/ConcertMain/Source/ConcertTransport/Public/Trace/Private/ConcertProtocolTraceInternal.h
+++ b/Engine/Plugins/Developer/Concert/ConcertMain/Source/ConcertTransport/Public/Trace/Private/ConcertProtocolTraceInternal.h
@@ -29,7 +29,7 @@ namespace UE::ConcertTrace
 
 #define INTERNAL_CONCERT_TRACE_INIT(InEndpointIdTOptional, InEndpointDisplayFString, InIsServer) UE::ConcertTrace::Init(InEndpointIdTOptional, InEndpointDisplayFString, InIsServer);
 
-#define INTERNAL_CONCERT_TRACE_OBJECT_SCOPE(InEventName, InObjectPath, InSequenceId, InProtocolId) UE::ConcertTrace::FConcertScopedObjectTrace PREPROCESSOR_JOIN(TheScope, __LINE__)(InProtocolId, InSequenceId, #InEventName, InObjectPath);
+#define INTERNAL_CONCERT_TRACE_OBJECT_SCOPE(InEventName, InObjectPath, InSequenceId, InProtocolId) UE::ConcertTrace::FConcertScopedObjectTrace UE_JOIN(TheScope, __LINE__)(InProtocolId, InSequenceId, #InEventName, InObjectPath);
 
 #define INTERNAL_CONCERT_TRACE_OBJECT_TRANSMISSION_START(InDestinationEndpointId, InObjectPath, InSequenceId, InProtocolId) UE::ConcertTrace::LogTransmissionStart(InProtocolId, InSequenceId, InObjectPath, InDestinationEndpointId);
 #define INTERNAL_CONCERT_TRACE_OBJECT_TRANSMISSION_RECEIVE(InObjectPath, InSequenceId, InProtocolId) UE::ConcertTrace::LogTransmissionReceive(InProtocolId, InSequenceId, InObjectPath);
diff --git a/Engine/Plugins/MovieScene/MovieRenderPipeline/Source/MovieRenderPipelineCore/Public/MoviePipelineUtils.h b/Engine/Plugins/MovieScene/MovieRenderPipeline/Source/MovieRenderPipelineCore/Public/MoviePipelineUtils.h
--- a/Engine/Plugins/MovieScene/MovieRenderPipeline/Source/MovieRenderPipelineCore/Public/MoviePipelineUtils.h
+++ b/Engine/Plugins/MovieScene/MovieRenderPipeline/Source/MovieRenderPipelineCore/Public/MoviePipelineUtils.h
@@ -167,23 +167,23 @@ MOVIERENDERPIPELINECORE_API void GetPassCompositeData(FMoviePipelineMergerOutput
  * render if it appears in different scopes).
  */
 #define MRG_LOG_ONCE_PER_RENDER(Identifier, LogCategory, LogVerbosity, LogFormat, ...)									\
-	static UWorld* PREPROCESSOR_JOIN(__LastRenderWorld__, Identifier) = nullptr;										\
-	static TArray<FString> PREPROCESSOR_JOIN(__LoggedMessages__, Identifier);											\
+	static UWorld* UE_JOIN(__LastRenderWorld__, Identifier) = nullptr;													\
+	static TArray<FString> UE_JOIN(__LoggedMessages__, Identifier);														\
 																														\
 	/* If this world has not been encountered yet (ie, no renders yet, or this is a new render), clear out */			\
 	/* the messages that have been recorded as already logged. */														\
-	if (PREPROCESSOR_JOIN(__LastRenderWorld__, Identifier) != MoviePipeline::FindCurrentWorld())						\
+	if (UE_JOIN(__LastRenderWorld__, Identifier) != MoviePipeline::FindCurrentWorld())									\
 	{																													\
-		PREPROCESSOR_JOIN(__LastRenderWorld__, Identifier) = MoviePipeline::FindCurrentWorld();							\
-		PREPROCESSOR_JOIN(__LoggedMessages__, Identifier).Empty();														\
+		UE_JOIN(__LastRenderWorld__, Identifier) = MoviePipeline::FindCurrentWorld();									\
+		UE_JOIN(__LoggedMessages__, Identifier).Empty();																\
 	}																													\
 																														\
 	/* If this message hasn't been logged, record that it has been logged so we can skip it for other frames. */		\
-	const FString PREPROCESSOR_JOIN(__FinalLogMsg__, Identifier) = FString::Printf(LogFormat, ##__VA_ARGS__);			\
-	if (!PREPROCESSOR_JOIN(__LoggedMessages__, Identifier).Contains(PREPROCESSOR_JOIN(__FinalLogMsg__, Identifier)))	\
+	const FString UE_JOIN(__FinalLogMsg__, Identifier) = FString::Printf(LogFormat, ##__VA_ARGS__);						\
+	if (!UE_JOIN(__LoggedMessages__, Identifier).Contains(UE_JOIN(__FinalLogMsg__, Identifier)))						\
 	{																													\
-		UE_LOG(LogCategory, LogVerbosity, TEXT("%s"), *PREPROCESSOR_JOIN(__FinalLogMsg__, Identifier));					\
-		PREPROCESSOR_JOIN(__LoggedMessages__, Identifier).Add(PREPROCESSOR_JOIN(__FinalLogMsg__, Identifier));			\
+		UE_LOG(LogCategory, LogVerbosity, TEXT("%s"), *UE_JOIN(__FinalLogMsg__, Identifier));							\
+		UE_JOIN(__LoggedMessages__, Identifier).Add(UE_JOIN(__FinalLogMsg__, Identifier));								\
 	}																													\
 
 namespace UE
diff --git a/Engine/Plugins/Online/EOSShared/Source/EOSShared/Public/EOSShared.h b/Engine/Plugins/Online/EOSShared/Source/EOSShared/Public/EOSShared.h
--- a/Engine/Plugins/Online/EOSShared/Source/EOSShared/Public/EOSShared.h
+++ b/Engine/Plugins/Online/EOSShared/Source/EOSShared/Public/EOSShared.h
@@ -19,12 +19,12 @@
  */
 #if UE_WITH_EOS_SDK_APIVERSION_WARNINGS
 #define UE_EOS_CHECK_API_MISMATCH(Actual, Expected) \
-	struct PREPROCESSOR_JOIN(FEosApiMismatchMsg_, __LINE__) { \
+	struct UE_JOIN(FEosApiMismatchMsg_, __LINE__) { \
 		[[deprecated(#Actual " updated from " #Expected " to " PREPROCESSOR_TO_STRING(Actual) ", behaviour/params may have changed")]] \
 		static constexpr int condition(TStaticDeprecateExpression<true>) { return 1; } \
 		static constexpr int condition(TStaticDeprecateExpression<false>) { return 1; } \
 	}; \
-	enum class PREPROCESSOR_JOIN(EEosApiMismatchMsg_, __LINE__) { Value = PREPROCESSOR_JOIN(FEosApiMismatchMsg_, __LINE__)::condition(TStaticDeprecateExpression<Actual != Expected>()) }
+	enum class UE_JOIN(EEosApiMismatchMsg_, __LINE__) { Value = UE_JOIN(FEosApiMismatchMsg_, __LINE__)::condition(TStaticDeprecateExpression<Actual != Expected>()) }
 #else
 #define UE_EOS_CHECK_API_MISMATCH(Actual, Expected)
 #endif
diff --git a/Engine/Plugins/Runtime/Metasound/Source/MetasoundGraphCore/Public/MetasoundTrace.h b/Engine/Plugins/Runtime/Metasound/Source/MetasoundGraphCore/Public/MetasoundTrace.h
--- a/Engine/Plugins/Runtime/Metasound/Source/MetasoundGraphCore/Public/MetasoundTrace.h
+++ b/Engine/Plugins/Runtime/Metasound/Source/MetasoundGraphCore/Public/MetasoundTrace.h
@@ -19,7 +19,7 @@
 // Copied from SCOPED_NAMED_EVENT but modified
 // to accommodate event names containing ::
 #define METASOUND_TRACE_CPUPROFILER_EVENT_SCOPE_COND(Name, Cond)\
-	FScopedNamedEventConditionalStatic PREPROCESSOR_JOIN(MetaSound_NamedEvent_,__LINE__)(FColor::Green, NAMED_EVENT_STR(#Name), GCycleStatsShouldEmitNamedEvents > 0 && (Cond));\
+	FScopedNamedEventConditionalStatic UE_JOIN(MetaSound_NamedEvent_,__LINE__)(FColor::Green, NAMED_EVENT_STR(#Name), GCycleStatsShouldEmitNamedEvents > 0 && (Cond));\
 	TRACE_CPUPROFILER_EVENT_SCOPE_CONDITIONAL(Name, Cond);
 
 #define METASOUND_TRACE_CPUPROFILER_EVENT_SCOPE(Name)\
@@ -35,7 +35,7 @@
 // Uses cached Insights SpecId to avoid string lookup
 #define METASOUND_TRACE_CPUPROFILER_EVENT_SCOPE_FAST(TraceSpecId, Name) \
 	FScopedNamedEventConditional ANONYMOUS_VARIABLE(NamedEvent_)(FColor::Green, Name, GCycleStatsShouldEmitNamedEvents > 0);\
-	TRACE_CPUPROFILER_EVENT_SCOPE_USE(TraceSpecId, Name, PREPROCESSOR_JOIN(MetaSound_NamedEventScope_, __LINE__), true)
+	TRACE_CPUPROFILER_EVENT_SCOPE_USE(TraceSpecId, Name, UE_JOIN(MetaSound_NamedEventScope_, __LINE__), true)
 
 #else
 // Metasound CPU profiler trace *not* enabled
diff --git a/Engine/Plugins/Runtime/ReplicationSystemTestPlugin/Source/Private/NetworkAutomationTestMacros.h b/Engine/Plugins/Runtime/ReplicationSystemTestPlugin/Source/Private/NetworkAutomationTestMacros.h
--- a/Engine/Plugins/Runtime/ReplicationSystemTestPlugin/Source/Private/NetworkAutomationTestMacros.h
+++ b/Engine/Plugins/Runtime/ReplicationSystemTestPlugin/Source/Private/NetworkAutomationTestMacros.h
@@ -124,9 +124,9 @@ Name(const T& Value, U ExpectedValue, const char* ValueText, const char* Expecte
 #if UE_NET_WITH_LOW_LEVEL_TESTS
 
 #define UE_NET_TEST_MSG_INTERNAL(Message) \
-	UE::Net::FTestMessage PREPROCESSOR_JOIN(TestMessage, __LINE__); \
-	PREPROCESSOR_JOIN(TestMessage, __LINE__) << Message; \
-	UNSCOPED_INFO(PREPROCESSOR_JOIN(TestMessage, __LINE__).C_Str())
+	UE::Net::FTestMessage UE_JOIN(TestMessage, __LINE__); \
+	UE_JOIN(TestMessage, __LINE__) << Message; \
+	UNSCOPED_INFO(UE_JOIN(TestMessage, __LINE__).C_Str())
 
 #define UE_NET_ASSERT_EQ(V1, V2) REQUIRE(V1 == V2)
 #define UE_NET_ASSERT_NE(V1, V2) REQUIRE(V1 != V2)
diff --git a/Engine/Plugins/Runtime/StateTree/Source/StateTreeModule/Public/Debugger/StateTreeDebug.h b/Engine/Plugins/Runtime/StateTree/Source/StateTreeModule/Public/Debugger/StateTreeDebug.h
--- a/Engine/Plugins/Runtime/StateTree/Source/StateTreeModule/Public/Debugger/StateTreeDebug.h
+++ b/Engine/Plugins/Runtime/StateTree/Source/StateTreeModule/Public/Debugger/StateTreeDebug.h
@@ -53,7 +53,7 @@ enum class EStateTreeTraceEventType : uint8;
 #define UE_STATETREE_DEBUG_TASK_EVENT(InCtx, InIndex, InDataView, InEventType, InStatus) TRACE_STATETREE_TASK_EVENT((InCtx), FStateTreeIndex16(InIndex), (InDataView), (InEventType), (InStatus));
 #define UE_STATETREE_DEBUG_EVALUATOR_EVENT(InCtx, InIndex, InDataView, InEventType) TRACE_STATETREE_EVALUATOR_EVENT((InCtx), FStateTreeIndex16(InIndex), (InDataView), (InEventType));
 
-#define UE_STATETREE_ID_NAME PREPROCESSOR_JOIN(InstanceId,__LINE__) \
+#define UE_STATETREE_ID_NAME UE_JOIN(InstanceId,__LINE__) \
 
 /** Scope based macros captures the instance ID since it might not be accessible when exiting the scope */
 #define UE_STATETREE_DEBUG_SCOPED_PHASE(InCtx, InPhase) \
diff --git a/Engine/Source/Editor/UnrealEd/Private/Cooker/CookProfiling.h b/Engine/Source/Editor/UnrealEd/Private/Cooker/CookProfiling.h
--- a/Engine/Source/Editor/UnrealEd/Private/Cooker/CookProfiling.h
+++ b/Engine/Source/Editor/UnrealEd/Private/Cooker/CookProfiling.h
@@ -54,9 +54,9 @@ private:
 };
 
 #define UE_CREATE_HIERARCHICAL_COOKTIMER(name, incrementScope) \
-			FScopeTimer PREPROCESSOR_JOIN(__HierarchicalCookTimerScope, __LINE__)(__COUNTER__, #name, incrementScope); 
+			FScopeTimer UE_JOIN(__HierarchicalCookTimerScope, __LINE__)(__COUNTER__, #name, incrementScope); 
 #define UE_CREATE_TEXT_HIERARCHICAL_COOKTIMER(name, incrementScope) \
-			FScopeTimer PREPROCESSOR_JOIN(__HierarchicalCookTimerScope, __LINE__)(__COUNTER__, name, incrementScope);
+			FScopeTimer UE_JOIN(__HierarchicalCookTimerScope, __LINE__)(__COUNTER__, name, incrementScope);
 
 // Emits trace events denoting scope/lifetime of an activity on the cooking channel.
 #define UE_SCOPED_COOKTIMER(name)						TRACE_CPUPROFILER_EVENT_SCOPE_ON_CHANNEL(name, CookChannel)
@@ -64,30 +64,30 @@ private:
 #define UE_SCOPED_TEXT_COOKTIMER(name)					TRACE_CPUPROFILER_EVENT_SCOPE_TEXT_ON_CHANNEL(name, CookChannel)
 
 #define UE_SCOPED_COOKTIMER_AND_DURATION(name, durationStorage) \
-			FScopedDurationTimer PREPROCESSOR_JOIN(__CookTimerScope, __LINE__)(durationStorage); \
+			FScopedDurationTimer UE_JOIN(__CookTimerScope, __LINE__)(durationStorage); \
 			UE_SCOPED_COOKTIMER(name)
 // More expensive version that takes a dynamic TCHAR*
 #define UE_SCOPED_TEXT_COOKTIMER_AND_DURATION(name, durationStorage) \
-			FScopedDurationTimer PREPROCESSOR_JOIN(__CookTimerScope, __LINE__)(durationStorage); \
+			FScopedDurationTimer UE_JOIN(__CookTimerScope, __LINE__)(durationStorage); \
 			UE_SCOPED_TEXT_COOKTIMER(name)
 
 // Emits trace events denoting scope/lifetime of an activity on the cooking channel.  Also creates a named hierarchical timer that can be aggregated and reported at cook completion.
 #define UE_SCOPED_HIERARCHICAL_COOKTIMER(name) \
 			UE_CREATE_HIERARCHICAL_COOKTIMER(name, true); \
-			PREPROCESSOR_JOIN(__HierarchicalCookTimerScope, __LINE__).Start(); \
+			UE_JOIN(__HierarchicalCookTimerScope, __LINE__).Start(); \
 			UE_SCOPED_COOKTIMER(name)
 // More expensive version that takes a dynamic TCHAR*
 #define UE_SCOPED_TEXT_HIERARCHICAL_COOKTIMER(name) \
 			UE_CREATE_TEXT_HIERARCHICAL_COOKTIMER(name, true); \
-			PREPROCESSOR_JOIN(__HierarchicalCookTimerScope, __LINE__).Start(); \
+			UE_JOIN(__HierarchicalCookTimerScope, __LINE__).Start(); \
 			UE_SCOPED_TEXT_COOKTIMER(name)
 // Version that takes a separate variable for where the duration should be stored
 #define UE_SCOPED_HIERARCHICAL_COOKTIMER_AND_DURATION(name, durationStorage) \
-			FScopedDurationTimer PREPROCESSOR_JOIN(__CookTimerScope, __LINE__)(durationStorage); \
+			FScopedDurationTimer UE_JOIN(__CookTimerScope, __LINE__)(durationStorage); \
 			UE_SCOPED_HIERARCHICAL_COOKTIMER(name)
 // More expensive version that takes a separate variable for where the duration should be stored and a dynamic TCHAR*
 #define UE_SCOPED_TEXT_HIERARCHICAL_COOKTIMER_AND_DURATION(name, durationStorage) \
-			FScopedDurationTimer PREPROCESSOR_JOIN(__CookTimerScope, __LINE__)(durationStorage); \
+			FScopedDurationTimer UE_JOIN(__CookTimerScope, __LINE__)(durationStorage); \
 			UE_SCOPED_TEXT_HIERARCHICAL_COOKTIMER(name)
 
 #define UE_CUSTOM_COOKTIMER_LOG Cpu
@@ -95,16 +95,16 @@ private:
 // Emits trace events denoting scope/lifetime of an activity on the cooking channel.
 #define UE_SCOPED_CUSTOM_COOKTIMER(name)				UE_TRACE_LOG_SCOPED_T(UE_CUSTOM_COOKTIMER_LOG, name, CookChannel)
 #define UE_SCOPED_CUSTOM_COOKTIMER_AND_DURATION(name, durationStorage) \
-			FScopedDurationTimer PREPROCESSOR_JOIN(__CookTimerScope, __LINE__)(durationStorage); \
+			FScopedDurationTimer UE_JOIN(__CookTimerScope, __LINE__)(durationStorage); \
 			UE_SCOPED_CUSTOM_COOKTIMER(name)
 
 // Emits trace events denoting scope/lifetime of an activity on the cooking channel.  Also creates a named hierarchical timer that can be aggregated and reported at cook completion.
 #define UE_SCOPED_HIERARCHICAL_CUSTOM_COOKTIMER(name) \
 			UE_CREATE_HIERARCHICAL_COOKTIMER(name, true); \
-			PREPROCESSOR_JOIN(__HierarchicalCookTimerScope, __LINE__).Start(); \
+			UE_JOIN(__HierarchicalCookTimerScope, __LINE__).Start(); \
 			UE_SCOPED_CUSTOM_COOKTIMER(name)
 #define UE_SCOPED_HIERARCHICAL_CUSTOM_COOKTIMER_AND_DURATION(name, durationStorage) \
-			FScopedDurationTimer PREPROCESSOR_JOIN(__CookTimerScope, __LINE__)(durationStorage); \
+			FScopedDurationTimer UE_JOIN(__CookTimerScope, __LINE__)(durationStorage); \
 			UE_SCOPED_HIERARCHICAL_CUSTOM_COOKTIMER(name)
 #define UE_ADD_CUSTOM_COOKTIMER_META(name, key, value) << name.key(value)
 
diff --git a/Engine/Source/Editor/UnrealEd/Public/CookPackageSplitter.h b/Engine/Source/Editor/UnrealEd/Public/CookPackageSplitter.h
--- a/Engine/Source/Editor/UnrealEd/Public/CookPackageSplitter.h
+++ b/Engine/Source/Editor/UnrealEd/Public/CookPackageSplitter.h
@@ -441,10 +441,10 @@ private:
  * REGISTER_COOKPACKAGE_SPLITTER(FMyCookPackageSplitter, UMySplitDataClass);
  */
 #define DEFINE_COOKPACKAGE_SPLITTER(SplitterClass, SplitDataClass) \
-friend class PREPROCESSOR_JOIN(PREPROCESSOR_JOIN(SplitterClass, SplitDataClass), _Register)
+friend class UE_JOIN(UE_JOIN(SplitterClass, SplitDataClass), _Register)
 
 #define REGISTER_COOKPACKAGE_SPLITTER(SplitterClass, SplitDataClass) \
-class PREPROCESSOR_JOIN(PREPROCESSOR_JOIN(SplitterClass, SplitDataClass), _Register) \
+class UE_JOIN(UE_JOIN(SplitterClass, SplitDataClass), _Register) \
 	: public UE::Cook::Private::FRegisteredCookPackageSplitter \
 { \
 	virtual UClass* GetSplitDataClass() const override \
@@ -468,9 +468,9 @@ class PREPROCESSOR_JOIN(PREPROCESSOR_JOIN(SplitterClass, SplitDataClass), _Regis
 		return SplitterClass::GetSplitterDebugName(); \
 	} \
 }; \
-namespace PREPROCESSOR_JOIN(SplitterClass, SplitDataClass) \
+namespace UE_JOIN(SplitterClass, SplitDataClass) \
 { \
-	static PREPROCESSOR_JOIN(PREPROCESSOR_JOIN(SplitterClass, SplitDataClass), _Register) DefaultObject; \
+	static UE_JOIN(UE_JOIN(SplitterClass, SplitDataClass), _Register) DefaultObject; \
 }
 
 #endif
diff --git a/Engine/Source/Programs/Online/OSSTestsCore/OnlineSubsystemCatchHelper.h b/Engine/Source/Programs/Online/OSSTestsCore/OnlineSubsystemCatchHelper.h
--- a/Engine/Source/Programs/Online/OSSTestsCore/OnlineSubsystemCatchHelper.h
+++ b/Engine/Source/Programs/Online/OSSTestsCore/OnlineSubsystemCatchHelper.h
@@ -188,18 +188,18 @@ TArray<TFunction<void()>>* GetGlobalInitalizers();
 
 #define INTERNAL_ONLINESUBSYSTEM_TEST_CASE_NAMED_FIXTURE(RegName, ClassName, Name, Tags, ...)\
 namespace {\
-class PREPROCESSOR_JOIN(OnlineSubsystemTest_,RegName) : public ClassName {\
+class UE_JOIN(OnlineSubsystemTest_,RegName) : public ClassName {\
 public:\
 	virtual void invoke() const override;\
 	void test();\
 };\
-	FOnlineSubsystemAutoReg RegName( ::MakeTestInvokerFixture( PREPROCESSOR_JOIN(&OnlineSubsystemTest_,RegName)::test ), CATCH_INTERNAL_LINEINFO, Name, Tags, "");\
+	FOnlineSubsystemAutoReg RegName( ::MakeTestInvokerFixture( UE_JOIN(&OnlineSubsystemTest_,RegName)::test ), CATCH_INTERNAL_LINEINFO, Name, Tags, "");\
 }\
-void PREPROCESSOR_JOIN(OnlineSubsystemTest_,RegName)::test()\
+void UE_JOIN(OnlineSubsystemTest_,RegName)::test()\
 {\
 	invoke();\
 }\
-void PREPROCESSOR_JOIN(OnlineSubsystemTest_,RegName)::invoke() const\
+void UE_JOIN(OnlineSubsystemTest_,RegName)::invoke() const\
 
 #define ONLINESUBSYSTEM_TEST_CASE_FIXTURE(ClassName, Name, Tags, ...) \
 	INTERNAL_ONLINESUBSYSTEM_TEST_CASE_NAMED_FIXTURE(INTERNAL_CATCH_UNIQUE_NAME(OnlineSubsystemRegistrar), ClassName, Name, Tags, __VA_ARGS__)
diff --git a/Engine/Source/Programs/Online/OnlineTestsCore/Source/Public/OnlineCatchHelper.h b/Engine/Source/Programs/Online/OnlineTestsCore/Source/Public/OnlineCatchHelper.h
--- a/Engine/Source/Programs/Online/OnlineTestsCore/Source/Public/OnlineCatchHelper.h
+++ b/Engine/Source/Programs/Online/OnlineTestsCore/Source/Public/OnlineCatchHelper.h
@@ -170,16 +170,16 @@ TArray<TFunction<void()>>* GetGlobalInitalizers();
 
 #define INTERNAL_ONLINE_TEST_CASE_NAMED(RegName, Name, Tags, ...)\
 namespace {\
-class PREPROCESSOR_JOIN(OnlineTest_,RegName) : public OnlineTestBase {\
+class UE_JOIN(OnlineTest_,RegName) : public OnlineTestBase {\
 protected:\
 	virtual void invoke() const override;\
 };\
-	OnlineTestBase* PREPROCESSOR_JOIN(Construct_,RegName) (){\
-		return new PREPROCESSOR_JOIN(OnlineTest_,RegName) ();\
+	OnlineTestBase* UE_JOIN(Construct_,RegName) (){\
+		return new UE_JOIN(OnlineTest_,RegName) ();\
 	}\
-	OnlineAutoReg RegName( PREPROCESSOR_JOIN(&Construct_,RegName) , CATCH_INTERNAL_LINEINFO , Name, Tags, "");\
+	OnlineAutoReg RegName( UE_JOIN(&Construct_,RegName) , CATCH_INTERNAL_LINEINFO , Name, Tags, "");\
 }\
-void PREPROCESSOR_JOIN(OnlineTest_,RegName)::invoke() const\
+void UE_JOIN(OnlineTest_,RegName)::invoke() const\
 
 #define ONLINE_TEST_CASE(Name, Tags, ...) \
 	INTERNAL_ONLINE_TEST_CASE_NAMED(INTERNAL_CATCH_UNIQUE_NAME(OnlineRegistrar), Name, Tags, __VA_ARGS__)
diff --git a/Engine/Source/Runtime/AssetRegistry/Public/AssetRegistry/AssetDependencyGatherer.h b/Engine/Source/Runtime/AssetRegistry/Public/AssetRegistry/AssetDependencyGatherer.h
--- a/Engine/Source/Runtime/AssetRegistry/Public/AssetRegistry/AssetDependencyGatherer.h
+++ b/Engine/Source/Runtime/AssetRegistry/Public/AssetRegistry/AssetDependencyGatherer.h
@@ -94,7 +94,7 @@ namespace UE::AssetDependencyGatherer::Private
  * REGISTER_ASSETDEPENDENCY_GATHERER(FMyAssetDependencyGatherer, UMyAssetClass);
  */
 #define REGISTER_ASSETDEPENDENCY_GATHERER(GathererClass, AssetClass) \
-class PREPROCESSOR_JOIN(PREPROCESSOR_JOIN(GathererClass, AssetClass), _Register) : public UE::AssetDependencyGatherer::Private::FRegisteredAssetDependencyGatherer \
+class UE_JOIN(UE_JOIN(GathererClass, AssetClass), _Register) : public UE::AssetDependencyGatherer::Private::FRegisteredAssetDependencyGatherer \
 { \
 	virtual UClass* GetAssetClass() const override { return AssetClass::StaticClass(); } \
 	virtual void GatherDependencies(const FAssetData& AssetData, const FAssetRegistryState& AssetRegistryState, TFunctionRef<FARCompiledFilter(const FARFilter&)> CompileFilterFunc, TArray<IAssetDependencyGatherer::FGathereredDependency>& OutDependencies, TArray<FString>& OutDependencyDirectories) const override \
@@ -103,9 +103,9 @@ class PREPROCESSOR_JOIN(PREPROCESSOR_JOIN(GathererClass, AssetClass), _Register)
 		Gatherer.GatherDependencies(AssetData, AssetRegistryState, CompileFilterFunc, OutDependencies, OutDependencyDirectories); \
 	} \
 }; \
-namespace PREPROCESSOR_JOIN(GathererClass, AssetClass) \
+namespace UE_JOIN(GathererClass, AssetClass) \
 { \
-	static PREPROCESSOR_JOIN(PREPROCESSOR_JOIN(GathererClass, AssetClass), _Register) DefaultObject; \
+	static UE_JOIN(UE_JOIN(GathererClass, AssetClass), _Register) DefaultObject; \
 }
 
 #endif // WITH_EDITOR
diff --git a/Engine/Source/Runtime/Core/Private/Containers/String.cpp.inl b/Engine/Source/Runtime/Core/Private/Containers/String.cpp.inl
--- a/Engine/Source/Runtime/Core/Private/Containers/String.cpp.inl
+++ b/Engine/Source/Runtime/Core/Private/Containers/String.cpp.inl
@@ -717,7 +717,7 @@ namespace UE::String::Private
 {
 
 template <typename LhsType, typename RhsType>
-[[nodiscard]] FORCEINLINE UE_STRING_CLASS PREPROCESSOR_JOIN(ConcatStrings_, UE_STRING_CLASS)(LhsType&& Lhs, RhsType&& Rhs)
+[[nodiscard]] FORCEINLINE UE_STRING_CLASS UE_JOIN(ConcatStrings_, UE_STRING_CLASS)(LhsType&& Lhs, RhsType&& Rhs)
 {
 	Lhs.CheckInvariants();
 	Rhs.CheckInvariants();
@@ -736,7 +736,7 @@ template <typename LhsType, typename RhsType>
 }
 
 template <typename LhsCharType, typename RhsType>
-[[nodiscard]] FORCEINLINE UE_STRING_CLASS PREPROCESSOR_JOIN(ConcatRangeString_, UE_STRING_CLASS)(const LhsCharType* Lhs, int32 LhsLen, RhsType&& Rhs)
+[[nodiscard]] FORCEINLINE UE_STRING_CLASS UE_JOIN(ConcatRangeString_, UE_STRING_CLASS)(const LhsCharType* Lhs, int32 LhsLen, RhsType&& Rhs)
 {
 	using ElementType = UE_STRING_CLASS::ElementType;
 
@@ -765,7 +765,7 @@ template <typename LhsCharType, typename RhsType>
 }
 
 template <typename LhsType, typename RhsCharType>
-[[nodiscard]] FORCEINLINE UE_STRING_CLASS PREPROCESSOR_JOIN(ConcatStringRange_, UE_STRING_CLASS)(LhsType&& Lhs, const RhsCharType* Rhs, int32 RhsLen)
+[[nodiscard]] FORCEINLINE UE_STRING_CLASS UE_JOIN(ConcatStringRange_, UE_STRING_CLASS)(LhsType&& Lhs, const RhsCharType* Rhs, int32 RhsLen)
 {
 	Lhs.CheckInvariants();
 	checkSlow(RhsLen >= 0);
@@ -781,7 +781,7 @@ template <typename LhsType, typename RhsCharType>
 }
 
 template <typename LhsCharType, typename RhsType>
-[[nodiscard]] FORCEINLINE UE_STRING_CLASS PREPROCESSOR_JOIN(ConcatPtrString_, UE_STRING_CLASS)(const LhsCharType* Lhs, RhsType&& Rhs)
+[[nodiscard]] FORCEINLINE UE_STRING_CLASS UE_JOIN(ConcatPtrString_, UE_STRING_CLASS)(const LhsCharType* Lhs, RhsType&& Rhs)
 {
 	checkSlow(Lhs);
 	if (!Lhs)
@@ -789,34 +789,34 @@ template <typename LhsCharType, typename RhsType>
 		return Forward<RhsType>(Rhs);
 	}
 
-	return PREPROCESSOR_JOIN(ConcatRangeString_, UE_STRING_CLASS)(Lhs, TCString<LhsCharType>::Strlen(Lhs), Forward<RhsType>(Rhs));
+	return UE_JOIN(ConcatRangeString_, UE_STRING_CLASS)(Lhs, TCString<LhsCharType>::Strlen(Lhs), Forward<RhsType>(Rhs));
 }
 
 template <typename LhsType, typename RhsCharType>
-[[nodiscard]] FORCEINLINE UE_STRING_CLASS PREPROCESSOR_JOIN(ConcatStringPtr_, UE_STRING_CLASS)(LhsType&& Lhs, const RhsCharType* Rhs)
+[[nodiscard]] FORCEINLINE UE_STRING_CLASS UE_JOIN(ConcatStringPtr_, UE_STRING_CLASS)(LhsType&& Lhs, const RhsCharType* Rhs)
 {
 	checkSlow(Rhs);
 	if (!Rhs)
 	{
 		return Forward<LhsType>(Lhs);
 	}
-	return PREPROCESSOR_JOIN(ConcatStringRange_, UE_STRING_CLASS)(Forward<LhsType>(Lhs), Rhs, TCString<RhsCharType>::Strlen(Rhs));
+	return UE_JOIN(ConcatStringRange_, UE_STRING_CLASS)(Forward<LhsType>(Lhs), Rhs, TCString<RhsCharType>::Strlen(Rhs));
 }
 
 } // namespace UE::String::Private
 
-UE_STRING_CLASS UE_STRING_CLASS::ConcatFF(const UE_STRING_CLASS& Lhs, const UE_STRING_CLASS& Rhs)				{ return UE::String::Private::PREPROCESSOR_JOIN(ConcatStrings_,     UE_STRING_CLASS)(Lhs, Rhs); }
-UE_STRING_CLASS UE_STRING_CLASS::ConcatFF(UE_STRING_CLASS&& Lhs, const UE_STRING_CLASS& Rhs)					{ return UE::String::Private::PREPROCESSOR_JOIN(ConcatStrings_,     UE_STRING_CLASS)(MoveTemp(Lhs), Rhs); }
-UE_STRING_CLASS UE_STRING_CLASS::ConcatFF(const UE_STRING_CLASS& Lhs, UE_STRING_CLASS&& Rhs)					{ return UE::String::Private::PREPROCESSOR_JOIN(ConcatStrings_,     UE_STRING_CLASS)(Lhs, MoveTemp(Rhs)); }
-UE_STRING_CLASS UE_STRING_CLASS::ConcatFF(UE_STRING_CLASS&& Lhs, UE_STRING_CLASS&& Rhs)							{ return UE::String::Private::PREPROCESSOR_JOIN(ConcatStrings_,     UE_STRING_CLASS)(MoveTemp(Lhs), MoveTemp(Rhs)); }
-UE_STRING_CLASS UE_STRING_CLASS::ConcatFC(const UE_STRING_CLASS& Lhs, const ElementType* Rhs)					{ return UE::String::Private::PREPROCESSOR_JOIN(ConcatStringPtr_,   UE_STRING_CLASS)(Lhs, Rhs); }
-UE_STRING_CLASS UE_STRING_CLASS::ConcatFC(UE_STRING_CLASS&& Lhs, const ElementType* Rhs)						{ return UE::String::Private::PREPROCESSOR_JOIN(ConcatStringPtr_,   UE_STRING_CLASS)(MoveTemp(Lhs), Rhs); }
-UE_STRING_CLASS UE_STRING_CLASS::ConcatCF(const ElementType* Lhs,	const UE_STRING_CLASS& Rhs)					{ return UE::String::Private::PREPROCESSOR_JOIN(ConcatPtrString_,   UE_STRING_CLASS)(Lhs, Rhs); }
-UE_STRING_CLASS UE_STRING_CLASS::ConcatCF(const ElementType* Lhs, UE_STRING_CLASS&& Rhs)						{ return UE::String::Private::PREPROCESSOR_JOIN(ConcatPtrString_,   UE_STRING_CLASS)(Lhs, MoveTemp(Rhs)); }
-UE_STRING_CLASS UE_STRING_CLASS::ConcatFR(const UE_STRING_CLASS& Lhs, const ElementType* Rhs, int32 RhsLen)		{ return UE::String::Private::PREPROCESSOR_JOIN(ConcatStringRange_, UE_STRING_CLASS)(Lhs, Rhs, RhsLen); }
-UE_STRING_CLASS UE_STRING_CLASS::ConcatFR(UE_STRING_CLASS&& Lhs, const ElementType* Rhs, int32 RhsLen)			{ return UE::String::Private::PREPROCESSOR_JOIN(ConcatStringRange_, UE_STRING_CLASS)(MoveTemp(Lhs), Rhs, RhsLen); }
-UE_STRING_CLASS UE_STRING_CLASS::ConcatRF(const ElementType* Lhs, int32 LhsLen, const UE_STRING_CLASS& Rhs)		{ return UE::String::Private::PREPROCESSOR_JOIN(ConcatRangeString_, UE_STRING_CLASS)(Lhs, LhsLen, Rhs); }
-UE_STRING_CLASS UE_STRING_CLASS::ConcatRF(const ElementType* Lhs, int32 LhsLen, UE_STRING_CLASS&& Rhs)			{ return UE::String::Private::PREPROCESSOR_JOIN(ConcatRangeString_, UE_STRING_CLASS)(Lhs, LhsLen, MoveTemp(Rhs)); }
+UE_STRING_CLASS UE_STRING_CLASS::ConcatFF(const UE_STRING_CLASS& Lhs, const UE_STRING_CLASS& Rhs)				{ return UE::String::Private::UE_JOIN(ConcatStrings_,     UE_STRING_CLASS)(Lhs, Rhs); }
+UE_STRING_CLASS UE_STRING_CLASS::ConcatFF(UE_STRING_CLASS&& Lhs, const UE_STRING_CLASS& Rhs)					{ return UE::String::Private::UE_JOIN(ConcatStrings_,     UE_STRING_CLASS)(MoveTemp(Lhs), Rhs); }
+UE_STRING_CLASS UE_STRING_CLASS::ConcatFF(const UE_STRING_CLASS& Lhs, UE_STRING_CLASS&& Rhs)					{ return UE::String::Private::UE_JOIN(ConcatStrings_,     UE_STRING_CLASS)(Lhs, MoveTemp(Rhs)); }
+UE_STRING_CLASS UE_STRING_CLASS::ConcatFF(UE_STRING_CLASS&& Lhs, UE_STRING_CLASS&& Rhs)							{ return UE::String::Private::UE_JOIN(ConcatStrings_,     UE_STRING_CLASS)(MoveTemp(Lhs), MoveTemp(Rhs)); }
+UE_STRING_CLASS UE_STRING_CLASS::ConcatFC(const UE_STRING_CLASS& Lhs, const ElementType* Rhs)					{ return UE::String::Private::UE_JOIN(ConcatStringPtr_,   UE_STRING_CLASS)(Lhs, Rhs); }
+UE_STRING_CLASS UE_STRING_CLASS::ConcatFC(UE_STRING_CLASS&& Lhs, const ElementType* Rhs)						{ return UE::String::Private::UE_JOIN(ConcatStringPtr_,   UE_STRING_CLASS)(MoveTemp(Lhs), Rhs); }
+UE_STRING_CLASS UE_STRING_CLASS::ConcatCF(const ElementType* Lhs,	const UE_STRING_CLASS& Rhs)					{ return UE::String::Private::UE_JOIN(ConcatPtrString_,   UE_STRING_CLASS)(Lhs, Rhs); }
+UE_STRING_CLASS UE_STRING_CLASS::ConcatCF(const ElementType* Lhs, UE_STRING_CLASS&& Rhs)						{ return UE::String::Private::UE_JOIN(ConcatPtrString_,   UE_STRING_CLASS)(Lhs, MoveTemp(Rhs)); }
+UE_STRING_CLASS UE_STRING_CLASS::ConcatFR(const UE_STRING_CLASS& Lhs, const ElementType* Rhs, int32 RhsLen)		{ return UE::String::Private::UE_JOIN(ConcatStringRange_, UE_STRING_CLASS)(Lhs, Rhs, RhsLen); }
+UE_STRING_CLASS UE_STRING_CLASS::ConcatFR(UE_STRING_CLASS&& Lhs, const ElementType* Rhs, int32 RhsLen)			{ return UE::String::Private::UE_JOIN(ConcatStringRange_, UE_STRING_CLASS)(MoveTemp(Lhs), Rhs, RhsLen); }
+UE_STRING_CLASS UE_STRING_CLASS::ConcatRF(const ElementType* Lhs, int32 LhsLen, const UE_STRING_CLASS& Rhs)		{ return UE::String::Private::UE_JOIN(ConcatRangeString_, UE_STRING_CLASS)(Lhs, LhsLen, Rhs); }
+UE_STRING_CLASS UE_STRING_CLASS::ConcatRF(const ElementType* Lhs, int32 LhsLen, UE_STRING_CLASS&& Rhs)			{ return UE::String::Private::UE_JOIN(ConcatRangeString_, UE_STRING_CLASS)(Lhs, LhsLen, MoveTemp(Rhs)); }
 
 /**
  * Concatenate this path with given path ensuring the / character is used between them
@@ -1628,7 +1628,7 @@ UE_STRING_CLASS UE_STRING_CLASS::ReplaceQuotesWithEscapedQuotes() &&
 	return MoveTemp(*this);
 }
 
-static const UE_STRING_CHARTYPE* PREPROCESSOR_JOIN(CharToEscapeSeqMap_, UE_STRING_CLASS)[6][2] =
+static const UE_STRING_CHARTYPE* UE_JOIN(CharToEscapeSeqMap_, UE_STRING_CLASS)[6][2] =
 {
 	// Always replace \\ first to avoid double-escaping characters
 	{ CHARTEXT(UE_STRING_CHARTYPE, "\\"), CHARTEXT(UE_STRING_CHARTYPE, "\\\\") },
@@ -1643,7 +1643,7 @@ void UE_STRING_CLASS::ReplaceCharWithEscapedCharInline(const TArray<ElementType>
 {
 	if (Len() > 0 && (Chars == nullptr || Chars->Num() > 0))
 	{
-		const auto& CharToEscapeSeqMap = PREPROCESSOR_JOIN(CharToEscapeSeqMap_, UE_STRING_CLASS);
+		const auto& CharToEscapeSeqMap = UE_JOIN(CharToEscapeSeqMap_, UE_STRING_CLASS);
 
 		for (uint32 ChIdx = 0; ChIdx < UE_ARRAY_COUNT(CharToEscapeSeqMap); ChIdx++)
 		{
@@ -1660,7 +1660,7 @@ void UE_STRING_CLASS::ReplaceEscapedCharWithCharInline(const TArray<ElementType>
 {
 	if (Len() > 0 && (Chars == nullptr || Chars->Num() > 0))
 	{
-		const auto& CharToEscapeSeqMap = PREPROCESSOR_JOIN(CharToEscapeSeqMap_, UE_STRING_CLASS);
+		const auto& CharToEscapeSeqMap = UE_JOIN(CharToEscapeSeqMap_, UE_STRING_CLASS);
 
 		// Spin CharToEscapeSeqMap backwards to ensure we're doing the inverse of ReplaceCharWithEscapedChar
 		for (int32 ChIdx = UE_ARRAY_COUNT(CharToEscapeSeqMap); ChIdx > 0; )
diff --git a/Engine/Source/Runtime/Core/Public/Apple/ApplePlatformMisc.h b/Engine/Source/Runtime/Core/Public/Apple/ApplePlatformMisc.h
--- a/Engine/Source/Runtime/Core/Public/Apple/ApplePlatformMisc.h
+++ b/Engine/Source/Runtime/Core/Public/Apple/ApplePlatformMisc.h
@@ -44,7 +44,7 @@ private:
 	NSAutoreleasePool*	Pool;
 };
 
-#define SCOPED_AUTORELEASE_POOL const FScopeAutoreleasePool PREPROCESSOR_JOIN(Pool,__LINE__);
+#define SCOPED_AUTORELEASE_POOL const FScopeAutoreleasePool UE_JOIN(Pool,__LINE__);
 
 #endif // !__has_feature(objc_arc)
 #endif // __OBJC__
diff --git a/Engine/Source/Runtime/Core/Public/Containers/CompactSet.h.inl b/Engine/Source/Runtime/Core/Public/Containers/CompactSet.h.inl
--- a/Engine/Source/Runtime/Core/Public/Containers/CompactSet.h.inl
+++ b/Engine/Source/Runtime/Core/Public/Containers/CompactSet.h.inl
@@ -25,7 +25,7 @@ namespace UE::Core::Private
 	[[noreturn]] CORE_API void OnInvalidSetNum(unsigned long long NewNum);
 }
 
-#define TSETPRIVATEFRIEND PREPROCESSOR_JOIN(UE_TCOMPACT_SET, PrivateFriend)
+#define TSETPRIVATEFRIEND UE_JOIN(UE_TCOMPACT_SET, PrivateFriend)
 
 /**
  * A set with an optional KeyFuncs parameters for customizing how the elements are compared and searched.  
diff --git a/Engine/Source/Runtime/Core/Public/Containers/Map.h.inl b/Engine/Source/Runtime/Core/Public/Containers/Map.h.inl
--- a/Engine/Source/Runtime/Core/Public/Containers/Map.h.inl
+++ b/Engine/Source/Runtime/Core/Public/Containers/Map.h.inl
@@ -4,12 +4,12 @@
 #error "Map.h.inl should only be included after defining UE_TMAP_PREFIX"
 #endif
 
-#define TMAPBASE PREPROCESSOR_JOIN(UE_TMAP_PREFIX, MapBase)
-#define TSORTABLEMAPBASE PREPROCESSOR_JOIN(UE_TMAP_PREFIX, SortableMapBase)
-#define TMAP PREPROCESSOR_JOIN(UE_TMAP_PREFIX, Map)
-#define TMAPPRIVATEFRIEND PREPROCESSOR_JOIN(TMAP,PrivateFriend)
-#define TMULTIMAP PREPROCESSOR_JOIN(UE_TMAP_PREFIX, MultiMap)
-#define TSET PREPROCESSOR_JOIN(UE_TMAP_PREFIX, Set)
+#define TMAPBASE UE_JOIN(UE_TMAP_PREFIX, MapBase)
+#define TSORTABLEMAPBASE UE_JOIN(UE_TMAP_PREFIX, SortableMapBase)
+#define TMAP UE_JOIN(UE_TMAP_PREFIX, Map)
+#define TMAPPRIVATEFRIEND UE_JOIN(TMAP,PrivateFriend)
+#define TMULTIMAP UE_JOIN(UE_TMAP_PREFIX, MultiMap)
+#define TSET UE_JOIN(UE_TMAP_PREFIX, Set)
 #define TMAP_STRINGIFY(name) #name
 
 /** 
diff --git a/Engine/Source/Runtime/Core/Public/Containers/SparseSet.h.inl b/Engine/Source/Runtime/Core/Public/Containers/SparseSet.h.inl
--- a/Engine/Source/Runtime/Core/Public/Containers/SparseSet.h.inl
+++ b/Engine/Source/Runtime/Core/Public/Containers/SparseSet.h.inl
@@ -27,7 +27,7 @@
 template <typename AllocatorType>
 class TScriptSparseSet;
 
-#define TSETPRIVATEFRIEND PREPROCESSOR_JOIN(UE_TSPARSE_SET, PrivateFriend)
+#define TSETPRIVATEFRIEND UE_JOIN(UE_TSPARSE_SET, PrivateFriend)
 
 /**
  * A set with an optional KeyFuncs parameters for customizing how the elements are compared and searched.
diff --git a/Engine/Source/Runtime/Core/Public/Containers/StringView.h b/Engine/Source/Runtime/Core/Public/Containers/StringView.h
--- a/Engine/Source/Runtime/Core/Public/Containers/StringView.h
+++ b/Engine/Source/Runtime/Core/Public/Containers/StringView.h
@@ -553,7 +553,7 @@ template <UE::CCharType CharType>
 	#define TEXTVIEW(str) TEXT(str##_PrivateSV)
 #endif
 #define ANSITEXTVIEW(str) (str##_PrivateASV)
-#define WIDETEXTVIEW(str) PREPROCESSOR_JOIN(WIDETEXT(str), _PrivateWSV)
+#define WIDETEXTVIEW(str) UE_JOIN(WIDETEXT(str), _PrivateWSV)
 #define UTF8TEXTVIEW(str) (str##_PrivateU8SV)
 
 namespace UE::Core::Private
diff --git a/Engine/Source/Runtime/Core/Public/Containers/UnrealString.h.inl b/Engine/Source/Runtime/Core/Public/Containers/UnrealString.h.inl
--- a/Engine/Source/Runtime/Core/Public/Containers/UnrealString.h.inl
+++ b/Engine/Source/Runtime/Core/Public/Containers/UnrealString.h.inl
@@ -33,7 +33,7 @@
 struct FStringFormatArg;
 template<typename InKeyType,typename InValueType,typename SetAllocator ,typename KeyFuncs > class TMap;
 
-using PREPROCESSOR_JOIN(UE_STRING_CLASS, FormatNamedArguments) = TMap<UE_STRING_CLASS, FStringFormatArg>;
+using UE_JOIN(UE_STRING_CLASS, FormatNamedArguments) = TMap<UE_STRING_CLASS, FStringFormatArg>;
 using FStringFormatOrderedArguments = TArray<FStringFormatArg>;
 
 [[nodiscard]] UE_STRING_CHARTYPE*       GetData(UE_STRING_CLASS&);
@@ -1462,12 +1462,12 @@ public:
 	 * @param InNamedArguments		A map of named arguments that match the tokens specified in InExpression
 	 * @return A string containing the formatted text
 	 */
-	[[nodiscard]] static UE_STRING_CLASS Format(const FmtCharType* InFormatString, const PREPROCESSOR_JOIN(UE_STRING_CLASS, FormatNamedArguments)& InNamedArguments)
+	[[nodiscard]] static UE_STRING_CLASS Format(const FmtCharType* InFormatString, const UE_JOIN(UE_STRING_CLASS, FormatNamedArguments)& InNamedArguments)
 	{
 		return FormatImpl((const ElementType*)InFormatString, InNamedArguments);
 	}
 #ifdef UE_STRING_PRINTF_FMT_CHARTYPE2
-	[[nodiscard]] static UE_STRING_CLASS Format(const FmtCharType2* InFormatString, const PREPROCESSOR_JOIN(UE_STRING_CLASS, FormatNamedArguments)& InNamedArguments)
+	[[nodiscard]] static UE_STRING_CLASS Format(const FmtCharType2* InFormatString, const UE_JOIN(UE_STRING_CLASS, FormatNamedArguments)& InNamedArguments)
 	{
 		return FormatImpl((const ElementType*)InFormatString, InNamedArguments);
 	}
@@ -1491,7 +1491,7 @@ public:
 #endif
 
 private:
-	static CORE_API UE_STRING_CLASS FormatImpl(const ElementType* InFormatString, const PREPROCESSOR_JOIN(UE_STRING_CLASS, FormatNamedArguments)& InNamedArguments);
+	static CORE_API UE_STRING_CLASS FormatImpl(const ElementType* InFormatString, const UE_JOIN(UE_STRING_CLASS, FormatNamedArguments)& InNamedArguments);
 	static CORE_API UE_STRING_CLASS FormatImpl(const ElementType* InFormatString, const FStringFormatOrderedArguments& InOrderedArguments);
 public:
 
diff --git a/Engine/Source/Runtime/Core/Public/Delegates/Delegate.h b/Engine/Source/Runtime/Core/Public/Delegates/Delegate.h
--- a/Engine/Source/Runtime/Core/Public/Delegates/Delegate.h
+++ b/Engine/Source/Runtime/Core/Public/Delegates/Delegate.h
@@ -407,7 +407,7 @@ class DynamicMulticastDelegateClassName : public TBaseDynamicMulticastDelegate<F
 		return {};
 	}
 
-	#define STATIC_FUNCTION_FNAME(str) NStrAfterLastDoubleColon_Private::TStaticFNameFromCharSequence<typename NStrAfterLastDoubleColon_Private::TStrAfterLastDoubleColon<decltype(PREPROCESSOR_JOIN(str, _intseq))>::Type>::Get()
+	#define STATIC_FUNCTION_FNAME(str) NStrAfterLastDoubleColon_Private::TStaticFNameFromCharSequence<typename NStrAfterLastDoubleColon_Private::TStrAfterLastDoubleColon<decltype(UE_JOIN(str, _intseq))>::Type>::Get()
 
 #else
 
diff --git a/Engine/Source/Runtime/Core/Public/HAL/LowLevelMemTracker.h b/Engine/Source/Runtime/Core/Public/HAL/LowLevelMemTracker.h
--- a/Engine/Source/Runtime/Core/Public/HAL/LowLevelMemTracker.h
+++ b/Engine/Source/Runtime/Core/Public/HAL/LowLevelMemTracker.h
@@ -328,7 +328,7 @@ extern CORE_API const FName LLMGetUntaggedTagName(ELLMTagSet TagSet);
 #define LLM(x) x
 #define LLM_IF_ENABLED(x) if (FLowLevelMemTracker::IsEnabled()) { x; }
 #define LLM_IS_ENABLED() FLowLevelMemTracker::IsEnabled()
-#define SCOPE_NAME PREPROCESSOR_JOIN(LLMScope,__LINE__)
+#define SCOPE_NAME UE_JOIN(LLMScope,__LINE__)
 
 ///////////////////////////////////////////////////////////////////////////////////////
 // These are the main macros to use externally when tracking memory
@@ -367,18 +367,18 @@ extern CORE_API const FName LLMGetUntaggedTagName(ELLMTagSet TagSet);
 	} while (false) /* do/while is added to require a semicolon */
 
 #define LLM_TAGSET_SCOPE(Tag, TagSet)								FLLMScope SCOPE_NAME(Tag, false /* bIsStaTag */, TagSet, ELLMTracker::Default);
-#define LLM_SCOPE_BYNAME(Tag) 										static FName PREPROCESSOR_JOIN(LLMScope_Name,__LINE__)(Tag);\
-																	FLLMScope SCOPE_NAME(PREPROCESSOR_JOIN(LLMScope_Name,__LINE__), false /* bIsStatTag */, ELLMTagSet::None, ELLMTracker::Default);\
-																	UE_MEMSCOPE(PREPROCESSOR_JOIN(LLMScope_Name,__LINE__));
-#define LLM_SCOPE_BYTAG(TagDeclName)								FLLMScope SCOPE_NAME(PREPROCESSOR_JOIN(LLMTagDeclaration_, TagDeclName).GetUniqueName(), false /* bIsStatTag */, ELLMTagSet::None, ELLMTracker::Default);\
-																	UE_MEMSCOPE(PREPROCESSOR_JOIN(LLMTagDeclaration_,TagDeclName).GetUniqueName());
-#define LLM_SCOPE_RENDER_RESOURCE(Tag)								static const FString PREPROCESSOR_JOIN(LLMScope_NamePrefix,__LINE__)(TEXT("RenderResources.")); \
-																	FName PREPROCESSOR_JOIN(LLMScope_Name,__LINE__)(PREPROCESSOR_JOIN(LLMScope_NamePrefix, __LINE__) + (Tag ? Tag : TEXT("Unknown"))); \
-																	FLLMScope SCOPE_NAME(PREPROCESSOR_JOIN(LLMScope_Name,__LINE__), false /* bIsStatTag */, ELLMTagSet::Assets, ELLMTracker::Default, false /* bOverride */);
+#define LLM_SCOPE_BYNAME(Tag) 										static FName UE_JOIN(LLMScope_Name,__LINE__)(Tag);\
+																	FLLMScope SCOPE_NAME(UE_JOIN(LLMScope_Name,__LINE__), false /* bIsStatTag */, ELLMTagSet::None, ELLMTracker::Default);\
+																	UE_MEMSCOPE(UE_JOIN(LLMScope_Name,__LINE__));
+#define LLM_SCOPE_BYTAG(TagDeclName)								FLLMScope SCOPE_NAME(UE_JOIN(LLMTagDeclaration_, TagDeclName).GetUniqueName(), false /* bIsStatTag */, ELLMTagSet::None, ELLMTracker::Default);\
+																	UE_MEMSCOPE(UE_JOIN(LLMTagDeclaration_,TagDeclName).GetUniqueName());
+#define LLM_SCOPE_RENDER_RESOURCE(Tag)								static const FString UE_JOIN(LLMScope_NamePrefix,__LINE__)(TEXT("RenderResources.")); \
+																	FName UE_JOIN(LLMScope_Name,__LINE__)(UE_JOIN(LLMScope_NamePrefix, __LINE__) + (Tag ? Tag : TEXT("Unknown"))); \
+																	FLLMScope SCOPE_NAME(UE_JOIN(LLMScope_Name,__LINE__), false /* bIsStatTag */, ELLMTagSet::Assets, ELLMTracker::Default, false /* bOverride */);
 #define LLM_PLATFORM_SCOPE(Tag)										FLLMScope SCOPE_NAME(Tag, false /* bIsStatTag */, ELLMTagSet::None, ELLMTracker::Platform);
-#define LLM_PLATFORM_SCOPE_BYNAME(Tag) 								static FName PREPROCESSOR_JOIN(LLMScope_Name,__LINE__)(Tag);\
-																	FLLMScope SCOPE_NAME(PREPROCESSOR_JOIN(LLMScope_Name,__LINE__), false /* bIsStatTag */, ELLMTagSet::None, ELLMTracker::Platform);
-#define LLM_PLATFORM_SCOPE_BYTAG(TagDeclName)						FLLMScope SCOPE_NAME(PREPROCESSOR_JOIN(LLMTagDeclaration_, TagDeclName).GetUniqueName(), false /* bIsStatTag */, ELLMTagSet::None, ELLMTracker::Platform);
+#define LLM_PLATFORM_SCOPE_BYNAME(Tag) 								static FName UE_JOIN(LLMScope_Name,__LINE__)(Tag);\
+																	FLLMScope SCOPE_NAME(UE_JOIN(LLMScope_Name,__LINE__), false /* bIsStatTag */, ELLMTagSet::None, ELLMTracker::Platform);
+#define LLM_PLATFORM_SCOPE_BYTAG(TagDeclName)						FLLMScope SCOPE_NAME(UE_JOIN(LLMTagDeclaration_, TagDeclName).GetUniqueName(), false /* bIsStatTag */, ELLMTagSet::None, ELLMTracker::Platform);
 #define LLM_SCOPE_CLEAR()											FLLMClearScope SCOPE_NAME(ELLMTagSet::None, ELLMTracker::Default);\
 																	UE_MEMSCOPE(0)
 #define LLM_TAGSET_SCOPE_CLEAR(TagSet)								FLLMClearScope SCOPE_NAME(TagSet, ELLMTracker::Default);
@@ -389,7 +389,7 @@ extern CORE_API const FName LLMGetUntaggedTagName(ELLMTagSet TagSet);
 #define LLM_SCOPED_PAUSE_TRACKING(AllocType) 						FLLMPauseScope SCOPE_NAME(ELLMTag::Untagged, false /* bIsStatTag */, 0, ELLMTracker::Max, AllocType);
 #define LLM_SCOPED_PAUSE_TRACKING_FOR_TRACKER(Tracker, AllocType) 	FLLMPauseScope SCOPE_NAME(ELLMTag::Untagged, false /* bIsStatTag */, 0, Tracker, AllocType);
 #define LLM_SCOPED_PAUSE_TRACKING_WITH_ENUM_AND_AMOUNT(Tag, Amount, Tracker, AllocType) FLLMPauseScope SCOPE_NAME(Tag, false /* bIsStatTag */, Amount, Tracker, AllocType);
-#define LLM_SCOPED_PAUSE_TRACKING_WITH_ENUM_AND_AMOUNT_BYTAG(TagDeclName, Amount, Tracker, AllocType) FLLMPauseScope SCOPE_NAME(PREPROCESSOR_JOIN(LLMTagDeclaration_, TagDeclName).GetUniqueName(), false /* bIsStatTag */, Amount, Tracker, AllocType);
+#define LLM_SCOPED_PAUSE_TRACKING_WITH_ENUM_AND_AMOUNT_BYTAG(TagDeclName, Amount, Tracker, AllocType) FLLMPauseScope SCOPE_NAME(UE_JOIN(LLMTagDeclaration_, TagDeclName).GetUniqueName(), false /* bIsStatTag */, Amount, Tracker, AllocType);
 
 /**
  * LLM realloc scope macros. Used when reallocating a pointer and you wish to retain the tagging from the source pointer
@@ -416,7 +416,7 @@ extern CORE_API const FName LLMGetUntaggedTagName(ELLMTagSet TagSet);
  * @param SummaryStatName - (Optional) - The name of the stat group to add on this tag's amount when publishing LLM
  *        data each frame, or NAME_None if no stat group should be added to.
  */
-#define LLM_DEFINE_TAG(UniqueNameWithUnderscores, ...) FLLMTagDeclaration PREPROCESSOR_JOIN(LLMTagDeclaration_, UniqueNameWithUnderscores)(TEXT(#UniqueNameWithUnderscores), ##__VA_ARGS__)
+#define LLM_DEFINE_TAG(UniqueNameWithUnderscores, ...) FLLMTagDeclaration UE_JOIN(LLMTagDeclaration_, UniqueNameWithUnderscores)(TEXT(#UniqueNameWithUnderscores), ##__VA_ARGS__)
 
 /**
  * Statically define a tag which can be used in LLM_SCOPE_BYTAG or referenced by name in other LLM_SCOPEs.
@@ -431,7 +431,7 @@ extern CORE_API const FName LLMGetUntaggedTagName(ELLMTagSet TagSet);
  * @param SummaryStatName - (Optional) - The name of the stat group to add on this tag's amount when publishing LLM
  *        data each frame, or NAME_None if no stat group should be added to.
  */
-#define LLM_DEFINE_STATIC_TAG(UniqueNameWithUnderscores, ...) static FLLMTagDeclaration PREPROCESSOR_JOIN(LLMTagDeclaration_, UniqueNameWithUnderscores)(TEXT(#UniqueNameWithUnderscores), ##__VA_ARGS__)
+#define LLM_DEFINE_STATIC_TAG(UniqueNameWithUnderscores, ...) static FLLMTagDeclaration UE_JOIN(LLMTagDeclaration_, UniqueNameWithUnderscores)(TEXT(#UniqueNameWithUnderscores), ##__VA_ARGS__)
 
  /**
   * Declare a tag defined by LLM_DEFINE_TAG. It is used in LLM_SCOPE_BYTAG or referenced by name in other LLM_SCOPEs.
@@ -441,11 +441,11 @@ extern CORE_API const FName LLMGetUntaggedTagName(ELLMTagSet TagSet);
   *        linkage, aka ENGINE_API. If omitted, no module linkage will be used and the tag will create link errors if
   *        used from another module.
   */
-#define LLM_DECLARE_TAG(UniqueNameWithUnderscores) extern FLLMTagDeclaration PREPROCESSOR_JOIN(LLMTagDeclaration_, UniqueNameWithUnderscores)
-#define LLM_DECLARE_TAG_API(UniqueNameWithUnderscores, ModuleAPI) extern ModuleAPI FLLMTagDeclaration PREPROCESSOR_JOIN(LLMTagDeclaration_, UniqueNameWithUnderscores)
+#define LLM_DECLARE_TAG(UniqueNameWithUnderscores) extern FLLMTagDeclaration UE_JOIN(LLMTagDeclaration_, UniqueNameWithUnderscores)
+#define LLM_DECLARE_TAG_API(UniqueNameWithUnderscores, ModuleAPI) extern ModuleAPI FLLMTagDeclaration UE_JOIN(LLMTagDeclaration_, UniqueNameWithUnderscores)
 
 /** Get the unique Name of a Tag. It can be passed to functions such as OnLowLevelAlloc that take a Tag UniqueName. */
-#define LLM_TAG_NAME(UniqueNameWithUnderscores) (PREPROCESSOR_JOIN(LLMTagDeclaration_, UniqueNameWithUnderscores).GetUniqueName())
+#define LLM_TAG_NAME(UniqueNameWithUnderscores) (UE_JOIN(LLMTagDeclaration_, UniqueNameWithUnderscores).GetUniqueName())
 
 /**
  * The BootStrap versions of LLM_DEFINE_TAG, LLM_SCOPE_BYTAG, and LLM_DECLARE_TAG support use in scopes during global
@@ -453,15 +453,15 @@ extern CORE_API const FName LLMGetUntaggedTagName(ELLMTagSet TagSet);
  * configurations other than shipping, because they use a function static rather than a global variable.
  */
 #define LLM_DEFINE_BOOTSTRAP_TAG(UniqueNameWithUnderscores, ...) \
-	FLLMTagDeclaration& PREPROCESSOR_JOIN(GetLLMTagDeclaration_, UniqueNameWithUnderscores)() \
+	FLLMTagDeclaration& UE_JOIN(GetLLMTagDeclaration_, UniqueNameWithUnderscores)() \
 	{ \
-		static FLLMTagDeclaration PREPROCESSOR_JOIN(LLMTagDeclaration_, UniqueNameWithUnderscores)(TEXT(#UniqueNameWithUnderscores), ##__VA_ARGS__); \
-		return PREPROCESSOR_JOIN(LLMTagDeclaration_, UniqueNameWithUnderscores); \
+		static FLLMTagDeclaration UE_JOIN(LLMTagDeclaration_, UniqueNameWithUnderscores)(TEXT(#UniqueNameWithUnderscores), ##__VA_ARGS__); \
+		return UE_JOIN(LLMTagDeclaration_, UniqueNameWithUnderscores); \
 	}
-#define LLM_SCOPE_BY_BOOTSTRAP_TAG(TagDeclName) FLLMScope SCOPE_NAME(PREPROCESSOR_JOIN(GetLLMTagDeclaration_, TagDeclName)().GetUniqueName(), false /* bIsStatTag */, ELLMTagSet::None, ELLMTracker::Default);\
-												UE_MEMSCOPE(PREPROCESSOR_JOIN(GetLLMTagDeclaration_,TagDeclName)().GetUniqueName());
-#define LLM_DECLARE_BOOTSTRAP_TAG(UniqueNameWithUnderscores) extern FLLMTagDeclaration& PREPROCESSOR_JOIN(GetLLMTagDeclaration_, UniqueNameWithUnderscores)();
-#define LLM_DECLARE_BOOTSTRAP_TAG_API(UniqueNameWithUnderscores, ModuleAPI) extern ModuleAPI FLLMTagDeclaration& PREPROCESSOR_JOIN(GetLLMTagDeclaration_, UniqueNameWithUnderscores)();
+#define LLM_SCOPE_BY_BOOTSTRAP_TAG(TagDeclName) FLLMScope SCOPE_NAME(UE_JOIN(GetLLMTagDeclaration_, TagDeclName)().GetUniqueName(), false /* bIsStatTag */, ELLMTagSet::None, ELLMTracker::Default);\
+												UE_MEMSCOPE(UE_JOIN(GetLLMTagDeclaration_,TagDeclName)().GetUniqueName());
+#define LLM_DECLARE_BOOTSTRAP_TAG(UniqueNameWithUnderscores) extern FLLMTagDeclaration& UE_JOIN(GetLLMTagDeclaration_, UniqueNameWithUnderscores)();
+#define LLM_DECLARE_BOOTSTRAP_TAG_API(UniqueNameWithUnderscores, ModuleAPI) extern ModuleAPI FLLMTagDeclaration& UE_JOIN(GetLLMTagDeclaration_, UniqueNameWithUnderscores)();
 
 typedef void*(*LLMAllocFunction)(size_t);
 typedef void(*LLMFreeFunction)(void*, size_t);
diff --git a/Engine/Source/Runtime/Core/Public/HAL/PreprocessorHelpers.h b/Engine/Source/Runtime/Core/Public/HAL/PreprocessorHelpers.h
--- a/Engine/Source/Runtime/Core/Public/HAL/PreprocessorHelpers.h
+++ b/Engine/Source/Runtime/Core/Public/HAL/PreprocessorHelpers.h
@@ -77,26 +77,26 @@
 #if PLATFORM_IS_EXTENSION
 // Creates a string that can be used to include a header in the platform extension form "PlatformHeader.h", not like
 // below form. When using this you should add "// IWYU pragma: export" at the end of the line.
-#define COMPILED_PLATFORM_HEADER(Suffix) PREPROCESSOR_TO_STRING(PREPROCESSOR_JOIN(PLATFORM_HEADER_NAME, Suffix))
+#define COMPILED_PLATFORM_HEADER(Suffix) PREPROCESSOR_TO_STRING(UE_JOIN(PLATFORM_HEADER_NAME, Suffix))
 #else
 // Creates a string that can be used to include a header in the form "Platform/PlatformHeader.h", like
 // "Windows/WindowsPlatformFile.h". When using this you should add "// IWYU pragma: export" at the end of the line.
-#define COMPILED_PLATFORM_HEADER(Suffix) PREPROCESSOR_TO_STRING(PREPROCESSOR_JOIN(PLATFORM_HEADER_NAME/PLATFORM_HEADER_NAME, Suffix))
+#define COMPILED_PLATFORM_HEADER(Suffix) PREPROCESSOR_TO_STRING(UE_JOIN(PLATFORM_HEADER_NAME/PLATFORM_HEADER_NAME, Suffix))
 #endif
 
 // Creates a string that can be used to include a header in the platform extension form "PlatformHeader.h", but will
 // not include a directory like COMPILED_PLATFORM_HEADER does, generally for UBT generated headers.
-#define COMPILED_PLATFORM_HEADER_GENERATED(Suffix) PREPROCESSOR_TO_STRING(PREPROCESSOR_JOIN(PLATFORM_HEADER_NAME, Suffix))
+#define COMPILED_PLATFORM_HEADER_GENERATED(Suffix) PREPROCESSOR_TO_STRING(UE_JOIN(PLATFORM_HEADER_NAME, Suffix))
 
 #if PLATFORM_IS_EXTENSION
 // Creates a string that can be used to include a header with the platform in its name, like
 // "Prefix/PlatformNameSuffix.h". When using this you should add "// IWYU pragma: export" at the end of the line.
-#define COMPILED_PLATFORM_HEADER_WITH_PREFIX(Prefix, Suffix) PREPROCESSOR_TO_STRING(Prefix/PREPROCESSOR_JOIN(PLATFORM_HEADER_NAME, Suffix))
+#define COMPILED_PLATFORM_HEADER_WITH_PREFIX(Prefix, Suffix) PREPROCESSOR_TO_STRING(Prefix/UE_JOIN(PLATFORM_HEADER_NAME, Suffix))
 #else
 // Creates a string that can be used to include a header with the platform in its name, like
 // "Prefix/PlatformName/PlatformNameSuffix.h". When using this you should add "// IWYU pragma: export" at the end of the
 // line.
-#define COMPILED_PLATFORM_HEADER_WITH_PREFIX(Prefix, Suffix) PREPROCESSOR_TO_STRING(Prefix/PLATFORM_HEADER_NAME/PREPROCESSOR_JOIN(PLATFORM_HEADER_NAME, Suffix))
+#define COMPILED_PLATFORM_HEADER_WITH_PREFIX(Prefix, Suffix) PREPROCESSOR_TO_STRING(Prefix/PLATFORM_HEADER_NAME/UE_JOIN(PLATFORM_HEADER_NAME, Suffix))
 #endif
 
 // These macros should be regarded as deprecated - use the UE_ macros they map to instead.
diff --git a/Engine/Source/Runtime/Core/Public/Logging/LogTrace.h b/Engine/Source/Runtime/Core/Public/Logging/LogTrace.h
--- a/Engine/Source/Runtime/Core/Public/Logging/LogTrace.h
+++ b/Engine/Source/Runtime/Core/Public/Logging/LogTrace.h
@@ -62,26 +62,26 @@ private:
 #define TRACE_LOG_MESSAGE(Category, Verbosity, Format, ...) \
 	if (UE_TRACE_CHANNELEXPR_IS_ENABLED(LogChannel)) \
 	{ \
-		static bool PREPROCESSOR_JOIN(__LogPoint, __LINE__); \
-		if (!PREPROCESSOR_JOIN(__LogPoint, __LINE__)) \
+		static bool UE_JOIN(__LogPoint, __LINE__); \
+		if (!UE_JOIN(__LogPoint, __LINE__)) \
 		{ \
-			FLogTrace::OutputLogMessageSpec(&PREPROCESSOR_JOIN(__LogPoint, __LINE__), &Category, ELogVerbosity::Verbosity, __FILE__, __LINE__, TEXT("%s")); \
-			PREPROCESSOR_JOIN(__LogPoint, __LINE__) = true; \
+			FLogTrace::OutputLogMessageSpec(&UE_JOIN(__LogPoint, __LINE__), &Category, ELogVerbosity::Verbosity, __FILE__, __LINE__, TEXT("%s")); \
+			UE_JOIN(__LogPoint, __LINE__) = true; \
 		} \
-		FLogTrace::OutputLogMessageSimple(&PREPROCESSOR_JOIN(__LogPoint, __LINE__), Format, ##__VA_ARGS__); \
+		FLogTrace::OutputLogMessageSimple(&UE_JOIN(__LogPoint, __LINE__), Format, ##__VA_ARGS__); \
 	}
 #else
 #define TRACE_LOG_MESSAGE(Category, Verbosity, Format, ...) \
 	static_assert(TIsArrayOrRefOfTypeByPredicate<decltype(Format), TIsCharEncodingCompatibleWithTCHAR>::Value, "Formatting string must be a TCHAR array."); \
 	if (UE_TRACE_CHANNELEXPR_IS_ENABLED(LogChannel)) \
 	{ \
-		static bool PREPROCESSOR_JOIN(__LogPoint, __LINE__); \
-		if (!PREPROCESSOR_JOIN(__LogPoint, __LINE__)) \
+		static bool UE_JOIN(__LogPoint, __LINE__); \
+		if (!UE_JOIN(__LogPoint, __LINE__)) \
 		{ \
-			FLogTrace::OutputLogMessageSpec(&PREPROCESSOR_JOIN(__LogPoint, __LINE__), &Category, ELogVerbosity::Verbosity, __FILE__, __LINE__, (const TCHAR*)Format); \
-			PREPROCESSOR_JOIN(__LogPoint, __LINE__) = true; \
+			FLogTrace::OutputLogMessageSpec(&UE_JOIN(__LogPoint, __LINE__), &Category, ELogVerbosity::Verbosity, __FILE__, __LINE__, (const TCHAR*)Format); \
+			UE_JOIN(__LogPoint, __LINE__) = true; \
 		} \
-		FLogTrace::OutputLogMessage(&PREPROCESSOR_JOIN(__LogPoint, __LINE__), ##__VA_ARGS__); \
+		FLogTrace::OutputLogMessage(&UE_JOIN(__LogPoint, __LINE__), ##__VA_ARGS__); \
 	}
 #endif
 
diff --git a/Engine/Source/Runtime/Core/Public/Misc/AutomationTest.h b/Engine/Source/Runtime/Core/Public/Misc/AutomationTest.h
--- a/Engine/Source/Runtime/Core/Public/Misc/AutomationTest.h
+++ b/Engine/Source/Runtime/Core/Public/Misc/AutomationTest.h
@@ -77,7 +77,7 @@ class FTextFilterExpressionEvaluator;
 // This macro allows for early exit of the executing unit test function when the condition is false
 // It explicitly uses the condition to ensure static analysis is happy with nullptr checks
 #ifndef UE_RETURN_ON_ERROR
-#define UE_RETURN_ON_ERROR(Condition, Message) const bool PREPROCESSOR_JOIN(UE____bCondition_Line_, __LINE__) = (Condition); AddErrorIfFalse(PREPROCESSOR_JOIN(UE____bCondition_Line_, __LINE__), (Message)); if(!PREPROCESSOR_JOIN(UE____bCondition_Line_, __LINE__)) return false
+#define UE_RETURN_ON_ERROR(Condition, Message) const bool UE_JOIN(UE____bCondition_Line_, __LINE__) = (Condition); AddErrorIfFalse(UE_JOIN(UE____bCondition_Line_, __LINE__), (Message)); if(!UE_JOIN(UE____bCondition_Line_, __LINE__)) return false
 #endif
 
 /**
diff --git a/Engine/Source/Runtime/Core/Public/Misc/CoreMiscDefines.h b/Engine/Source/Runtime/Core/Public/Misc/CoreMiscDefines.h
--- a/Engine/Source/Runtime/Core/Public/Misc/CoreMiscDefines.h
+++ b/Engine/Source/Runtime/Core/Public/Misc/CoreMiscDefines.h
@@ -175,11 +175,11 @@ enum EPerElement {PerElement};
 
 #ifdef __COUNTER__
 	// Created a variable with a unique name
-	#define ANONYMOUS_VARIABLE( Name ) PREPROCESSOR_JOIN(Name, __COUNTER__)
+	#define ANONYMOUS_VARIABLE( Name ) UE_JOIN(Name, __COUNTER__)
 #else
 	// Created a variable with a unique name.
 	// Less reliable than the __COUNTER__ version.
-	#define ANONYMOUS_VARIABLE( Name ) PREPROCESSOR_JOIN(Name, __LINE__)
+	#define ANONYMOUS_VARIABLE( Name ) UE_JOIN(Name, __LINE__)
 #endif
 
 /** Thread-safe call once helper for void functions, similar to std::call_once without the std::once_flag */
@@ -418,23 +418,23 @@ struct TStaticDeprecateExpression
  * Can be used in the same contexts as static_assert but gives a warning rather than an error, and 'fails' if the expression is true rather than false.
  */
 #define UE_STATIC_DEPRECATE(Version, bExpression, Message) \
-	struct PREPROCESSOR_JOIN(FDeprecationMsg_, __LINE__) { \
+	struct UE_JOIN(FDeprecationMsg_, __LINE__) { \
 		[[deprecated(Message " - Please update your code to the new API before upgrading to the next release, otherwise your project will no longer compile.")]] \
 		static constexpr int condition(TStaticDeprecateExpression<true>) { return 1; } \
 		static constexpr int condition(TStaticDeprecateExpression<false>) { return 1; } \
 	}; \
-	enum class PREPROCESSOR_JOIN(EDeprecationMsg_, __LINE__) { Value = PREPROCESSOR_JOIN(FDeprecationMsg_, __LINE__)::condition(TStaticDeprecateExpression<!!(bExpression)>()) }
+	enum class UE_JOIN(EDeprecationMsg_, __LINE__) { Value = UE_JOIN(FDeprecationMsg_, __LINE__)::condition(TStaticDeprecateExpression<!!(bExpression)>()) }
 
 /**
  * Can be used in the same contexts as static_assert but gives a warning rather than an error
  */
 #define UE_STATIC_ASSERT_WARN(bExpression, Message) \
-	struct PREPROCESSOR_JOIN(FStaticWarningMsg_, __LINE__) { \
+	struct UE_JOIN(FStaticWarningMsg_, __LINE__) { \
 		[[deprecated(Message)]] \
 		static constexpr int condition(TStaticDeprecateExpression<true>) { return 1; } \
 		static constexpr int condition(TStaticDeprecateExpression<false>) { return 1; } \
 	}; \
-	enum class PREPROCESSOR_JOIN(EStaticWarningMsg_, __LINE__) { Value = PREPROCESSOR_JOIN(FStaticWarningMsg_, __LINE__)::condition(TStaticDeprecateExpression<!(bExpression)>()) }
+	enum class UE_JOIN(EStaticWarningMsg_, __LINE__) { Value = UE_JOIN(FStaticWarningMsg_, __LINE__)::condition(TStaticDeprecateExpression<!(bExpression)>()) }
 
 // These defines are used to mark a difference between two pointers as expected to fit into the specified range
 // while still leaving something searchable if the surrounding code is updated to work with a 64 bit count/range
diff --git a/Engine/Source/Runtime/Core/Public/Misc/MTAccessDetector.h b/Engine/Source/Runtime/Core/Public/Misc/MTAccessDetector.h
--- a/Engine/Source/Runtime/Core/Public/Misc/MTAccessDetector.h
+++ b/Engine/Source/Runtime/Core/Public/Misc/MTAccessDetector.h
@@ -732,8 +732,8 @@ public:
 #define UE_MT_DECLARE_RW_FULLY_RECURSIVE_ACCESS_DETECTOR(AccessDetector) FRWFullyRecursiveAccessDetector AccessDetector
 #define UE_MT_DECLARE_MRSW_RECURSIVE_ACCESS_DETECTOR(AccessDetector) FMRSWRecursiveAccessDetector AccessDetector
 
-#define UE_MT_SCOPED_READ_ACCESS(AccessDetector) const FBaseScopedAccessDetector& PREPROCESSOR_JOIN(ScopedMTAccessDetector_,__LINE__) = MakeScopedReaderAccessDetector(AccessDetector)
-#define UE_MT_SCOPED_WRITE_ACCESS(AccessDetector) const FBaseScopedAccessDetector& PREPROCESSOR_JOIN(ScopedMTAccessDetector_,__LINE__) = MakeScopedWriterAccessDetector(AccessDetector)
+#define UE_MT_SCOPED_READ_ACCESS(AccessDetector) const FBaseScopedAccessDetector& UE_JOIN(ScopedMTAccessDetector_,__LINE__) = MakeScopedReaderAccessDetector(AccessDetector)
+#define UE_MT_SCOPED_WRITE_ACCESS(AccessDetector) const FBaseScopedAccessDetector& UE_JOIN(ScopedMTAccessDetector_,__LINE__) = MakeScopedWriterAccessDetector(AccessDetector)
 
 #define UE_MT_ACQUIRE_READ_ACCESS(AccessDetector) (AccessDetector).AcquireReadAccess()
 #define UE_MT_RELEASE_READ_ACCESS(AccessDetector) (AccessDetector).ReleaseReadAccess()
diff --git a/Engine/Source/Runtime/Core/Public/Misc/TimeGuard.h b/Engine/Source/Runtime/Core/Public/Misc/TimeGuard.h
--- a/Engine/Source/Runtime/Core/Public/Misc/TimeGuard.h
+++ b/Engine/Source/Runtime/Core/Public/Misc/TimeGuard.h
@@ -177,17 +177,17 @@ typedef FTimeGuard FLightweightTimeGuard;
 #if DO_LIGHTWEIGHT_TIMEGUARD
 
   #define LIGHTWEIGHT_TIME_GUARD_BEGIN( Name, ThresholdMS ) \
-	float PREPROCESSOR_JOIN(__TimeGuard_ThresholdMS_, Name) = ThresholdMS; \
-	uint64 PREPROCESSOR_JOIN(__TimeGuard_StartCycles_, Name) = ( ThresholdMS > 0.0f ) ? FPlatformTime::Cycles64() : 0;
+	float UE_JOIN(__TimeGuard_ThresholdMS_, Name) = ThresholdMS; \
+	uint64 UE_JOIN(__TimeGuard_StartCycles_, Name) = ( ThresholdMS > 0.0f ) ? FPlatformTime::Cycles64() : 0;
 
   #define LIGHTWEIGHT_TIME_GUARD_END( Name, NameStringCode ) \
-	if ( PREPROCESSOR_JOIN(__TimeGuard_ThresholdMS_, Name) > 0.0f ) \
+	if ( UE_JOIN(__TimeGuard_ThresholdMS_, Name) > 0.0f ) \
 	{\
-		float PREPROCESSOR_JOIN(__TimeGuard_MSElapsed_,Name) = FPlatformTime::ToMilliseconds64( FPlatformTime::Cycles64() - PREPROCESSOR_JOIN(__TimeGuard_StartCycles_,Name) ); \
-		if ( PREPROCESSOR_JOIN(__TimeGuard_MSElapsed_,Name) > PREPROCESSOR_JOIN(__TimeGuard_ThresholdMS_, Name) ) \
+		float UE_JOIN(__TimeGuard_MSElapsed_,Name) = FPlatformTime::ToMilliseconds64( FPlatformTime::Cycles64() - UE_JOIN(__TimeGuard_StartCycles_,Name) ); \
+		if ( UE_JOIN(__TimeGuard_MSElapsed_,Name) > UE_JOIN(__TimeGuard_ThresholdMS_, Name) ) \
 		{ \
 			FString ReportName = NameStringCode; \
-			UE_LOG(LogCore, Warning, TEXT("LIGHTWEIGHT_TIME_GUARD: %s - %s took %.2fms!"), TEXT(#Name), *ReportName, PREPROCESSOR_JOIN(__TimeGuard_MSElapsed_,Name)); \
+			UE_LOG(LogCore, Warning, TEXT("LIGHTWEIGHT_TIME_GUARD: %s - %s took %.2fms!"), TEXT(#Name), *ReportName, UE_JOIN(__TimeGuard_MSElapsed_,Name)); \
 		} \
 	}
 #else
diff --git a/Engine/Source/Runtime/Core/Public/ProfilingDebugging/CallstackTrace.h b/Engine/Source/Runtime/Core/Public/ProfilingDebugging/CallstackTrace.h
--- a/Engine/Source/Runtime/Core/Public/ProfilingDebugging/CallstackTrace.h
+++ b/Engine/Source/Runtime/Core/Public/ProfilingDebugging/CallstackTrace.h
@@ -60,7 +60,7 @@ CORE_API uint32 CallstackTrace_GetCurrentId();
   * preventing us to resolve it without crashing. Instead the callstack will be 
   * only the caller address.
   */
-#define CALLSTACK_TRACE_LIMIT_CALLSTACKRESOLVE_SCOPE() 						FCallStackTraceLimitResolveScope PREPROCESSOR_JOIN(FCTLMScope,__LINE__)
+#define CALLSTACK_TRACE_LIMIT_CALLSTACKRESOLVE_SCOPE() 						FCallStackTraceLimitResolveScope UE_JOIN(FCTLMScope,__LINE__)
 
 extern CORE_API uint32 GCallStackTracingTlsSlotIndex;
 
diff --git a/Engine/Source/Runtime/Core/Public/ProfilingDebugging/CookStats.h b/Engine/Source/Runtime/Core/Public/ProfilingDebugging/CookStats.h
--- a/Engine/Source/Runtime/Core/Public/ProfilingDebugging/CookStats.h
+++ b/Engine/Source/Runtime/Core/Public/ProfilingDebugging/CookStats.h
@@ -510,7 +510,7 @@ private:
 #endif
 
 #define UE_SCOPED_COOK_STAT(Name, StatType) \
-	FScopedCookStat  PREPROCESSOR_JOIN(__CookTimerScope, __LINE__)(Name.ToUnstableInt(), StatType);
+	FScopedCookStat  UE_JOIN(__CookTimerScope, __LINE__)(Name.ToUnstableInt(), StatType);
 
 #define UE_MULTI_SCOPED_COOK_STAT_INIT() \
 	TOptional<FScopedCookStat> __CookStatScope; \
diff --git a/Engine/Source/Runtime/Core/Public/ProfilingDebugging/CountersTrace.h b/Engine/Source/Runtime/Core/Public/ProfilingDebugging/CountersTrace.h
--- a/Engine/Source/Runtime/Core/Public/ProfilingDebugging/CountersTrace.h
+++ b/Engine/Source/Runtime/Core/Public/ProfilingDebugging/CountersTrace.h
@@ -263,11 +263,11 @@ struct FCountersTrace
 
 #define __TRACE_DECLARE_INLINE_COUNTER(CounterDisplayName, CounterType, CounterDisplayHint) \
 	__TRACE_CHECK_COUNTER_NAME(CounterDisplayName) \
-	static FCountersTrace::CounterType PREPROCESSOR_JOIN(__TraceCounter, __LINE__)(CounterDisplayName, CounterDisplayHint);
+	static FCountersTrace::CounterType UE_JOIN(__TraceCounter, __LINE__)(CounterDisplayName, CounterDisplayHint);
 
 #define __TRACE_INLINE_COUNTER_SET(CounterDisplayName, Value, CounterType, CounterDisplayHint) \
 	__TRACE_DECLARE_INLINE_COUNTER(CounterDisplayName, CounterType, CounterDisplayHint) \
-	PREPROCESSOR_JOIN(__TraceCounter, __LINE__).Set(Value);
+	UE_JOIN(__TraceCounter, __LINE__).Set(Value);
 
 #define TRACE_INT_VALUE(CounterDisplayName, Value)                      __TRACE_INLINE_COUNTER_SET(CounterDisplayName, Value, FCounterInt, TraceCounterDisplayHint_None)
 #define TRACE_ATOMIC_INT_VALUE(CounterDisplayName, Value)               __TRACE_INLINE_COUNTER_SET(CounterDisplayName, Value, FCounterAtomicInt, TraceCounterDisplayHint_None)
@@ -287,7 +287,7 @@ struct FCountersTrace
 
 #define TRACE_DECLARE_COUNTER(CounterType, CounterName, CounterDisplayName, CounterDisplayHint) \
 	__TRACE_CHECK_COUNTER_NAME(CounterDisplayName) \
-	FCountersTrace::CounterType PREPROCESSOR_JOIN(__GTraceCounter, CounterName)(CounterDisplayName, CounterDisplayHint);
+	FCountersTrace::CounterType UE_JOIN(__GTraceCounter, CounterName)(CounterDisplayName, CounterDisplayHint);
 
 //////////////////////////////////////////////////
 // Declare Int Counters
@@ -299,10 +299,10 @@ struct FCountersTrace
 	TRACE_DECLARE_COUNTER(FCounterUncheckedInt, CounterName, CounterDisplayName, TraceCounterDisplayHint_None)
 
 #define TRACE_DECLARE_INT_COUNTER_EXTERN(CounterName) \
-	extern FCountersTrace::FCounterInt PREPROCESSOR_JOIN(__GTraceCounter, CounterName);
+	extern FCountersTrace::FCounterInt UE_JOIN(__GTraceCounter, CounterName);
 
 #define TRACE_DECLARE_UNCHECKED_INT_COUNTER_EXTERN(CounterName) \
-	extern FCountersTrace::FCounterUncheckedInt PREPROCESSOR_JOIN(__GTraceCounter, CounterName);
+	extern FCountersTrace::FCounterUncheckedInt UE_JOIN(__GTraceCounter, CounterName);
 
 //////////////////////////////////////////////////
 // Declare Atomic Int Counters
@@ -314,10 +314,10 @@ struct FCountersTrace
 	TRACE_DECLARE_COUNTER(FCounterUncheckedAtomicInt, CounterName, CounterDisplayName, TraceCounterDisplayHint_None)
 
 #define TRACE_DECLARE_ATOMIC_INT_COUNTER_EXTERN(CounterName) \
-	extern FCountersTrace::FCounterAtomicInt PREPROCESSOR_JOIN(__GTraceCounter, CounterName);
+	extern FCountersTrace::FCounterAtomicInt UE_JOIN(__GTraceCounter, CounterName);
 
 #define TRACE_DECLARE_UNCHECKED_ATOMIC_INT_COUNTER_EXTERN(CounterName) \
-	extern FCountersTrace::FCounterUncheckedAtomicInt PREPROCESSOR_JOIN(__GTraceCounter, CounterName);
+	extern FCountersTrace::FCounterUncheckedAtomicInt UE_JOIN(__GTraceCounter, CounterName);
 
 //////////////////////////////////////////////////
 // Declare Float Counters
@@ -329,10 +329,10 @@ struct FCountersTrace
 	TRACE_DECLARE_COUNTER(FCounterUncheckedFloat, CounterName, CounterDisplayName, TraceCounterDisplayHint_None)
 
 #define TRACE_DECLARE_FLOAT_COUNTER_EXTERN(CounterName) \
-	extern FCountersTrace::FCounterFloat PREPROCESSOR_JOIN(__GTraceCounter, CounterName);
+	extern FCountersTrace::FCounterFloat UE_JOIN(__GTraceCounter, CounterName);
 
 #define TRACE_DECLARE_UNCHECKED_FLOAT_COUNTER_EXTERN(CounterName) \
-	extern FCountersTrace::FCounterUncheckedFloat PREPROCESSOR_JOIN(__GTraceCounter, CounterName);
+	extern FCountersTrace::FCounterUncheckedFloat UE_JOIN(__GTraceCounter, CounterName);
 
 //////////////////////////////////////////////////
 // Declare Atomic Float Counters
@@ -344,10 +344,10 @@ struct FCountersTrace
 	TRACE_DECLARE_COUNTER(FCounterUncheckedAtomicFloat, CounterName, CounterDisplayName, TraceCounterDisplayHint_None)
 
 #define TRACE_DECLARE_ATOMIC_FLOAT_COUNTER_EXTERN(CounterName) \
-	extern FCountersTrace::FCounterAtomicFloat PREPROCESSOR_JOIN(__GTraceCounter, CounterName);
+	extern FCountersTrace::FCounterAtomicFloat UE_JOIN(__GTraceCounter, CounterName);
 
 #define TRACE_DECLARE_UNCHECKED_ATOMIC_FLOAT_COUNTER_EXTERN(CounterName) \
-	extern FCountersTrace::FCounterUncheckedAtomicFloat PREPROCESSOR_JOIN(__GTraceCounter, CounterName);
+	extern FCountersTrace::FCounterUncheckedAtomicFloat UE_JOIN(__GTraceCounter, CounterName);
 
 //////////////////////////////////////////////////
 // Declare Memory Counters
@@ -385,46 +385,46 @@ struct FCountersTrace
 // A value that does not change will be traced (or not) depending on how the counter is created.
 
 #define TRACE_COUNTER_SET(CounterName, Value) \
-	PREPROCESSOR_JOIN(__GTraceCounter, CounterName).Set(Value);
+	UE_JOIN(__GTraceCounter, CounterName).Set(Value);
 
 #define TRACE_COUNTER_GET(CounterName) \
-	PREPROCESSOR_JOIN(__GTraceCounter, CounterName).Get();
+	UE_JOIN(__GTraceCounter, CounterName).Get();
 
 #define TRACE_COUNTER_ADD(CounterName, Value) \
-	PREPROCESSOR_JOIN(__GTraceCounter, CounterName).Add(Value);
+	UE_JOIN(__GTraceCounter, CounterName).Add(Value);
 
 #define TRACE_COUNTER_SUBTRACT(CounterName, Value) \
-	PREPROCESSOR_JOIN(__GTraceCounter, CounterName).Subtract(Value);
+	UE_JOIN(__GTraceCounter, CounterName).Subtract(Value);
 
 #define TRACE_COUNTER_INCREMENT(CounterName) \
-	PREPROCESSOR_JOIN(__GTraceCounter, CounterName).Increment();
+	UE_JOIN(__GTraceCounter, CounterName).Increment();
 
 #define TRACE_COUNTER_DECREMENT(CounterName) \
-	PREPROCESSOR_JOIN(__GTraceCounter, CounterName).Decrement();
+	UE_JOIN(__GTraceCounter, CounterName).Decrement();
 
 // _IF_DIFFERENT / _IF_NOT_ZERO
 // It will not trace a value that doesn't change (no matter how the counter was created).
 
 #define TRACE_COUNTER_SET_IF_DIFFERENT(CounterName, Value) \
-	PREPROCESSOR_JOIN(__GTraceCounter, CounterName).SetIfDifferent(Value);
+	UE_JOIN(__GTraceCounter, CounterName).SetIfDifferent(Value);
 
 #define TRACE_COUNTER_ADD_IF_NOT_ZERO(CounterName, Value) \
-	PREPROCESSOR_JOIN(__GTraceCounter, CounterName).AddIfNotZero(Value);
+	UE_JOIN(__GTraceCounter, CounterName).AddIfNotZero(Value);
 
 #define TRACE_COUNTER_SUBTRACT_IF_NOT_ZERO(CounterName, Value) \
-	PREPROCESSOR_JOIN(__GTraceCounter, CounterName).SubtractIfNotZero(Value);
+	UE_JOIN(__GTraceCounter, CounterName).SubtractIfNotZero(Value);
 
 // _ALWAYS
 // It will trace even if value doesn't change (no matter how the counter was created).
 
 #define TRACE_COUNTER_SET_ALWAYS(CounterName, Value) \
-	PREPROCESSOR_JOIN(__GTraceCounter, CounterName).SetAlways(Value);
+	UE_JOIN(__GTraceCounter, CounterName).SetAlways(Value);
 
 #define TRACE_COUNTER_ADD_ALWAYS(CounterName, Value) \
-	PREPROCESSOR_JOIN(__GTraceCounter, CounterName).AddAlways(Value);
+	UE_JOIN(__GTraceCounter, CounterName).AddAlways(Value);
 
 #define TRACE_COUNTER_SUBTRACT_ALWAYS(CounterName, Value) \
-	PREPROCESSOR_JOIN(__GTraceCounter, CounterName).SubtractAlways(Value);
+	UE_JOIN(__GTraceCounter, CounterName).SubtractAlways(Value);
 
 #else // COUNTERSTRACE_ENABLED
 
diff --git a/Engine/Source/Runtime/Core/Public/ProfilingDebugging/CpuProfilerTrace.h b/Engine/Source/Runtime/Core/Public/ProfilingDebugging/CpuProfilerTrace.h
--- a/Engine/Source/Runtime/Core/Public/ProfilingDebugging/CpuProfilerTrace.h
+++ b/Engine/Source/Runtime/Core/Public/ProfilingDebugging/CpuProfilerTrace.h
@@ -366,8 +366,8 @@ private:
 #define TRACE_CPUPROFILER_EVENT_MANUAL_START(EventNameStr) \
 	if (CpuChannel) \
 	{ \
-		static uint32 PREPROCESSOR_JOIN(__CpuProfilerEventSpecId, __LINE__) = FCpuProfilerTrace::OutputEventType(EventNameStr, __FILE__, __LINE__); \
-		FCpuProfilerTrace::OutputBeginEvent(PREPROCESSOR_JOIN(__CpuProfilerEventSpecId, __LINE__)); \
+		static uint32 UE_JOIN(__CpuProfilerEventSpecId, __LINE__) = FCpuProfilerTrace::OutputEventType(EventNameStr, __FILE__, __LINE__); \
+		FCpuProfilerTrace::OutputBeginEvent(UE_JOIN(__CpuProfilerEventSpecId, __LINE__)); \
 	}
 
 // Advanced macro that will check if CpuChannel is enabled and, if so, ends the previously started CPU event.
@@ -393,8 +393,8 @@ private:
 //       the name should be guaranteed to not change between calls.
 // Example: TRACE_CPUPROFILER_EVENT_SCOPE_STR_CONDITIONAL("My Scoped Timer A", Condition)
 #define TRACE_CPUPROFILER_EVENT_SCOPE_STR_CONDITIONAL(NameStr, Condition) \
-	TRACE_CPUPROFILER_EVENT_DECLARE(PREPROCESSOR_JOIN(__CpuProfilerEventSpecId, __LINE__)); \
-	TRACE_CPUPROFILER_EVENT_SCOPE_USE(PREPROCESSOR_JOIN(__CpuProfilerEventSpecId, __LINE__), NameStr, PREPROCESSOR_JOIN(__CpuProfilerEventScope, __LINE__), (Condition));
+	TRACE_CPUPROFILER_EVENT_DECLARE(UE_JOIN(__CpuProfilerEventSpecId, __LINE__)); \
+	TRACE_CPUPROFILER_EVENT_SCOPE_USE(UE_JOIN(__CpuProfilerEventSpecId, __LINE__), NameStr, UE_JOIN(__CpuProfilerEventScope, __LINE__), (Condition));
 
 // Trace a scoped CPU timing event providing a static/constant string (const ANSICHAR* or
 // const TCHAR* or FName) as the scope name. It will use the CPU trace channel.
@@ -411,8 +411,8 @@ private:
 // Note: The event will be emitted only if both the given channel and the CpuChannel are enabled.
 // Example: TRACE_CPUPROFILER_EVENT_SCOPE_STR_ON_CHANNEL_CONDITIONAL("My Scoped Timer A", CustomChannel, Condition)
 #define TRACE_CPUPROFILER_EVENT_SCOPE_STR_ON_CHANNEL_CONDITIONAL(NameStr, Channel, Condition) \
-	TRACE_CPUPROFILER_EVENT_DECLARE(PREPROCESSOR_JOIN(__CpuProfilerEventSpecId, __LINE__)); \
-	TRACE_CPUPROFILER_EVENT_SCOPE_USE_ON_CHANNEL(PREPROCESSOR_JOIN(__CpuProfilerEventSpecId, __LINE__), NameStr, PREPROCESSOR_JOIN(__CpuProfilerEventScope, __LINE__), Channel, (Condition));
+	TRACE_CPUPROFILER_EVENT_DECLARE(UE_JOIN(__CpuProfilerEventSpecId, __LINE__)); \
+	TRACE_CPUPROFILER_EVENT_SCOPE_USE_ON_CHANNEL(UE_JOIN(__CpuProfilerEventSpecId, __LINE__), NameStr, UE_JOIN(__CpuProfilerEventScope, __LINE__), Channel, (Condition));
 #define TRACE_CPUPROFILER_EVENT_SCOPE_ON_CHANNEL_STR_CONDITIONAL(NameStr, Channel, Condition) \
 	TRACE_CPUPROFILER_EVENT_SCOPE_STR_ON_CHANNEL_CONDITIONAL(NameStr, Channel, Condition)
 #define TRACE_CPUPROFILER_EVENT_SCOPE_ON_CHANNEL_CONDITIONAL_STR(NameStr, Channel, Condition) \
@@ -475,7 +475,7 @@ private:
 //       or a static string. Use it only if scope name really needs to be a dynamic string.
 // Example: TRACE_CPUPROFILER_EVENT_SCOPE_TEXT_CONDITIONAL(*MyScopedTimerNameString, Condition)
 #define TRACE_CPUPROFILER_EVENT_SCOPE_TEXT_CONDITIONAL(Name, Condition) \
-	FCpuProfilerTrace::FDynamicEventScope PREPROCESSOR_JOIN(__CpuProfilerEventScope, __LINE__)(Name, (Condition), __FILE__, __LINE__);
+	FCpuProfilerTrace::FDynamicEventScope UE_JOIN(__CpuProfilerEventScope, __LINE__)(Name, (Condition), __FILE__, __LINE__);
 
 // Trace a scoped CPU timing event providing a dynamic string (const ANSICHAR* or
 // const TCHAR* or FName) as the scope name. It will use the CPU trace channel.
@@ -492,7 +492,7 @@ private:
 // Note: The event will be emitted only if both the given channel and the CpuChannel are enabled.
 // Example: TRACE_CPUPROFILER_EVENT_SCOPE_TEXT_ON_CHANNEL_CONDITIONAL(*MyScopedTimerNameString, CustomChannel, Condition)
 #define TRACE_CPUPROFILER_EVENT_SCOPE_TEXT_ON_CHANNEL_CONDITIONAL(Name, Channel, Condition) \
-	FCpuProfilerTrace::FDynamicEventScope PREPROCESSOR_JOIN(__CpuProfilerEventScope, __LINE__)(Name, Channel, (Condition), __FILE__, __LINE__);
+	FCpuProfilerTrace::FDynamicEventScope UE_JOIN(__CpuProfilerEventScope, __LINE__)(Name, Channel, (Condition), __FILE__, __LINE__);
 
 // Trace a scoped CPU timing event providing a dynamic string (const ANSICHAR* or const TCHAR* or FName)
 // as the scope name and a trace channel.
diff --git a/Engine/Source/Runtime/Core/Public/ProfilingDebugging/CsvProfilerTrace.h b/Engine/Source/Runtime/Core/Public/ProfilingDebugging/CsvProfilerTrace.h
--- a/Engine/Source/Runtime/Core/Public/ProfilingDebugging/CsvProfilerTrace.h
+++ b/Engine/Source/Runtime/Core/Public/ProfilingDebugging/CsvProfilerTrace.h
@@ -47,17 +47,17 @@ struct FCsvProfilerTrace
 	FCsvProfilerTrace::OutputRegisterCategory(Index, Name);
 
 #define TRACE_CSV_PROFILER_INLINE_STAT(StatName, CategoryIndex) \
-	static bool PREPROCESSOR_JOIN(__CsvProfilerStat, __LINE__); \
-	if (!PREPROCESSOR_JOIN(__CsvProfilerStat, __LINE__)) { \
+	static bool UE_JOIN(__CsvProfilerStat, __LINE__); \
+	if (!UE_JOIN(__CsvProfilerStat, __LINE__)) { \
 		FCsvProfilerTrace::OutputInlineStat(StatName, CategoryIndex); \
-		PREPROCESSOR_JOIN(__CsvProfilerStat, __LINE__) = true; \
+		UE_JOIN(__CsvProfilerStat, __LINE__) = true; \
 	}
 
 #define TRACE_CSV_PROFILER_INLINE_STAT_EXCLUSIVE(StatName) \
-	static bool PREPROCESSOR_JOIN(__CsvProfilerStat, __LINE__); \
-	if (!PREPROCESSOR_JOIN(__CsvProfilerStat, __LINE__)) { \
+	static bool UE_JOIN(__CsvProfilerStat, __LINE__); \
+	if (!UE_JOIN(__CsvProfilerStat, __LINE__)) { \
 		FCsvProfilerTrace::OutputInlineStatExclusive(StatName); \
-		PREPROCESSOR_JOIN(__CsvProfilerStat, __LINE__) = true; \
+		UE_JOIN(__CsvProfilerStat, __LINE__) = true; \
 	}
 
 #define TRACE_CSV_PROFILER_DECLARED_STAT(StatName, CategoryIndex) \
diff --git a/Engine/Source/Runtime/Core/Public/ProfilingDebugging/MetadataTrace.h b/Engine/Source/Runtime/Core/Public/ProfilingDebugging/MetadataTrace.h
--- a/Engine/Source/Runtime/Core/Public/ProfilingDebugging/MetadataTrace.h
+++ b/Engine/Source/Runtime/Core/Public/ProfilingDebugging/MetadataTrace.h
@@ -89,7 +89,7 @@ private:
  * When this scope is exited, the previous metadata stack is restored.
  */
 #define UE_TRACE_METADATA_RESTORE_SCOPE(SavedMetadataIdentifier) \
-	FMetadataRestoreScope PREPROCESSOR_JOIN(__MetadataRestoreScope, __LINE__)(SavedMetadataIdentifier);
+	FMetadataRestoreScope UE_JOIN(__MetadataRestoreScope, __LINE__)(SavedMetadataIdentifier);
 
 #else // UE_TRACE_METADATA_ENABLED
 
diff --git a/Engine/Source/Runtime/Core/Public/ProfilingDebugging/MiscTrace.h b/Engine/Source/Runtime/Core/Public/ProfilingDebugging/MiscTrace.h
--- a/Engine/Source/Runtime/Core/Public/ProfilingDebugging/MiscTrace.h
+++ b/Engine/Source/Runtime/Core/Public/ProfilingDebugging/MiscTrace.h
@@ -116,31 +116,31 @@ private:
 #define TRACE_BOOKMARK(Format, ...) \
 if (UE_TRACE_CHANNELEXPR_IS_ENABLED(BookmarkChannel)) \
 { \
-	static bool PREPROCESSOR_JOIN(__BookmarkPoint, __LINE__); \
-	if (!PREPROCESSOR_JOIN(__BookmarkPoint, __LINE__)) \
+	static bool UE_JOIN(__BookmarkPoint, __LINE__); \
+	if (!UE_JOIN(__BookmarkPoint, __LINE__)) \
 	{ \
 		static_assert(std::is_const_v<std::remove_reference_t<decltype(Format)>>, "Formatting string must be a const TCHAR array."); \
 		static_assert(TIsArrayOrRefOfTypeByPredicate<decltype(Format), TIsCharEncodingCompatibleWithTCHAR>::Value, "Formatting string must be a TCHAR array."); \
 		UE_VALIDATE_FORMAT_STRING(Format, ##__VA_ARGS__); \
-		FMiscTrace::OutputBookmarkSpec(&PREPROCESSOR_JOIN(__BookmarkPoint, __LINE__), __FILE__, __LINE__, Format); \
-		PREPROCESSOR_JOIN(__BookmarkPoint, __LINE__) = true; \
+		FMiscTrace::OutputBookmarkSpec(&UE_JOIN(__BookmarkPoint, __LINE__), __FILE__, __LINE__, Format); \
+		UE_JOIN(__BookmarkPoint, __LINE__) = true; \
 	} \
-	FMiscTrace::OutputBookmark(CallstackTrace_GetCurrentId(), &PREPROCESSOR_JOIN(__BookmarkPoint, __LINE__), ##__VA_ARGS__); \
+	FMiscTrace::OutputBookmark(CallstackTrace_GetCurrentId(), &UE_JOIN(__BookmarkPoint, __LINE__), ##__VA_ARGS__); \
 }
 
 #define TRACE_BOOKMARK_CYCLES(Cycles, Format, ...) \
 if (UE_TRACE_CHANNELEXPR_IS_ENABLED(BookmarkChannel)) \
 { \
-	static bool PREPROCESSOR_JOIN(__BookmarkPoint, __LINE__); \
-	if (!PREPROCESSOR_JOIN(__BookmarkPoint, __LINE__)) \
+	static bool UE_JOIN(__BookmarkPoint, __LINE__); \
+	if (!UE_JOIN(__BookmarkPoint, __LINE__)) \
 	{ \
 		static_assert(std::is_const_v<std::remove_reference_t<decltype(Format)>>, "Formatting string must be a const TCHAR array."); \
 		static_assert(TIsArrayOrRefOfTypeByPredicate<decltype(Format), TIsCharEncodingCompatibleWithTCHAR>::Value, "Formatting string must be a TCHAR array."); \
 		UE_VALIDATE_FORMAT_STRING(Format, ##__VA_ARGS__); \
-		FMiscTrace::OutputBookmarkSpec(&PREPROCESSOR_JOIN(__BookmarkPoint, __LINE__), __FILE__, __LINE__, Format); \
-		PREPROCESSOR_JOIN(__BookmarkPoint, __LINE__) = true; \
+		FMiscTrace::OutputBookmarkSpec(&UE_JOIN(__BookmarkPoint, __LINE__), __FILE__, __LINE__, Format); \
+		UE_JOIN(__BookmarkPoint, __LINE__) = true; \
 	} \
-	FMiscTrace::OutputBookmarkCycles(CallstackTrace_GetCurrentId(), Cycles, &PREPROCESSOR_JOIN(__BookmarkPoint, __LINE__), ##__VA_ARGS__); \
+	FMiscTrace::OutputBookmarkCycles(CallstackTrace_GetCurrentId(), Cycles, &UE_JOIN(__BookmarkPoint, __LINE__), ##__VA_ARGS__); \
 }
 
 /**
diff --git a/Engine/Source/Runtime/Core/Public/ProfilingDebugging/StallDetector.h b/Engine/Source/Runtime/Core/Public/ProfilingDebugging/StallDetector.h
--- a/Engine/Source/Runtime/Core/Public/ProfilingDebugging/StallDetector.h
+++ b/Engine/Source/Runtime/Core/Public/ProfilingDebugging/StallDetector.h
@@ -32,16 +32,16 @@
  **/
 #define SCOPE_STALL_COUNTER(InName, InBudgetSeconds)                                                                                                                                           \
 	static_assert(UE_ARRAY_COUNT(PREPROCESSOR_TO_STRING(InName)) <= UE::StallDetectorScopeMaxNameLength, "Scope name must be less than UE::StallDetectorScopeMaxNameLength characters long "); \
-	static UE::FStallDetectorStats PREPROCESSOR_JOIN(StallDetectorStats, __LINE__)(TEXT(#InName), InBudgetSeconds, UE::EStallDetectorReportingMode::Disabled);                                 \
-	UE::FStallDetector PREPROCESSOR_JOIN(ScopeStallDetector, __LINE__)(PREPROCESSOR_JOIN(StallDetectorStats, __LINE__));
+	static UE::FStallDetectorStats UE_JOIN(StallDetectorStats, __LINE__)(TEXT(#InName), InBudgetSeconds, UE::EStallDetectorReportingMode::Disabled);                                 \
+	UE::FStallDetector UE_JOIN(ScopeStallDetector, __LINE__)(UE_JOIN(StallDetectorStats, __LINE__));
 
 /**
  * Counts and sends a report of the first occurence of a stall.
  **/
 #define SCOPE_STALL_REPORTER(InName, InBudgetSeconds)                                                                                                                                          \
 	static_assert(UE_ARRAY_COUNT(PREPROCESSOR_TO_STRING(InName)) <= UE::StallDetectorScopeMaxNameLength, "Scope name must be less than UE::StallDetectorScopeMaxNameLength characters long "); \
-	static UE::FStallDetectorStats PREPROCESSOR_JOIN(StallDetectorStats, __LINE__)(TEXT(#InName), InBudgetSeconds, UE::EStallDetectorReportingMode::First);                                    \
-	UE::FStallDetector PREPROCESSOR_JOIN(ScopeStallDetector, __LINE__)(PREPROCESSOR_JOIN(StallDetectorStats, __LINE__));
+	static UE::FStallDetectorStats UE_JOIN(StallDetectorStats, __LINE__)(TEXT(#InName), InBudgetSeconds, UE::EStallDetectorReportingMode::First);                                    \
+	UE::FStallDetector UE_JOIN(ScopeStallDetector, __LINE__)(UE_JOIN(StallDetectorStats, __LINE__));
 
 /**
  * Counts and reports every occurence of a stall.
@@ -50,14 +50,14 @@
  **/
 #define SCOPE_STALL_REPORTER_ALWAYS(InName, InBudgetSeconds)                                                                                                                                   \
 	static_assert(UE_ARRAY_COUNT(PREPROCESSOR_TO_STRING(InName)) <= UE::StallDetectorScopeMaxNameLength, "Scope name must be less than UE::StallDetectorScopeMaxNameLength characters long "); \
-	static UE::FStallDetectorStats PREPROCESSOR_JOIN(StallDetectorStats, __LINE__)(TEXT(#InName), InBudgetSeconds, UE::EStallDetectorReportingMode::Always);                                   \
-	UE::FStallDetector PREPROCESSOR_JOIN(ScopeStallDetector, __LINE__)(PREPROCESSOR_JOIN(StallDetectorStats, __LINE__));
+	static UE::FStallDetectorStats UE_JOIN(StallDetectorStats, __LINE__)(TEXT(#InName), InBudgetSeconds, UE::EStallDetectorReportingMode::Always);                                   \
+	UE::FStallDetector UE_JOIN(ScopeStallDetector, __LINE__)(UE_JOIN(StallDetectorStats, __LINE__));
 
 /**
  * Pauses stall detectors on the current thread so that time in the current thread spent in this scope will not be counted as a stall.
  **/
 #define SCOPE_STALL_DETECTOR_PAUSE() \
-	UE::FStallDetectorPause PREPROCESSOR_JOIN(ScopeStallDetectorPause, __LINE__);
+	UE::FStallDetectorPause UE_JOIN(ScopeStallDetectorPause, __LINE__);
 
 #else // STALL_DETECTOR
 
diff --git a/Engine/Source/Runtime/Core/Public/ProfilingDebugging/TagTrace.h b/Engine/Source/Runtime/Core/Public/ProfilingDebugging/TagTrace.h
--- a/Engine/Source/Runtime/Core/Public/ProfilingDebugging/TagTrace.h
+++ b/Engine/Source/Runtime/Core/Public/ProfilingDebugging/TagTrace.h
@@ -89,12 +89,12 @@ private:
 };
 
 ////////////////////////////////////////////////////////////////////////////////
-#define UE_MEMSCOPE(InTag)					FMemScope PREPROCESSOR_JOIN(MemScope,__LINE__)(InTag);
+#define UE_MEMSCOPE(InTag)					FMemScope UE_JOIN(MemScope,__LINE__)(InTag);
 // UE_MEMSCOPE_PTR(InPtr) or UE_MEMSCOPE_PTR(InPtr, InRootHeap)
-#define UE_MEMSCOPE_PTR(InPtr, ...)			FMemScopePtr PREPROCESSOR_JOIN(MemPtrScope,__LINE__)((uint64)InPtr, ##__VA_ARGS__);
-#define UE_MEMSCOPE_DEFAULT(InTag)			FDefaultMemScope PREPROCESSOR_JOIN(MemScope,__LINE__)(InTag);
-#define UE_MEMSCOPE_UNINITIALIZED(Line)		FMemScope PREPROCESSOR_JOIN(MemScope,Line);
-#define UE_MEMSCOPE_ACTIVATE(Line, InTag)	PREPROCESSOR_JOIN(MemScope,Line).SetTagAndActivate(InTag);
+#define UE_MEMSCOPE_PTR(InPtr, ...)			FMemScopePtr UE_JOIN(MemPtrScope,__LINE__)((uint64)InPtr, ##__VA_ARGS__);
+#define UE_MEMSCOPE_DEFAULT(InTag)			FDefaultMemScope UE_JOIN(MemScope,__LINE__)(InTag);
+#define UE_MEMSCOPE_UNINITIALIZED(Line)		FMemScope UE_JOIN(MemScope,Line);
+#define UE_MEMSCOPE_ACTIVATE(Line, InTag)	UE_JOIN(MemScope,Line).SetTagAndActivate(InTag);
 
 #else // UE_MEMORY_TAGS_TRACE_ENABLED
 
diff --git a/Engine/Source/Runtime/Core/Public/Serialization/Archive.h b/Engine/Source/Runtime/Core/Public/Serialization/Archive.h
--- a/Engine/Source/Runtime/Core/Public/Serialization/Archive.h
+++ b/Engine/Source/Runtime/Core/Public/Serialization/Archive.h
@@ -1699,9 +1699,9 @@ public:
 
 #define SerializeBitfield(Ar, BitField) \
 	{ \
-		bool PREPROCESSOR_JOIN(TempBitfield, __LINE__) = BitField; \
-		Ar.SerializeBits(&PREPROCESSOR_JOIN(TempBitfield, __LINE__), 1); \
-		BitField = PREPROCESSOR_JOIN(TempBitfield, __LINE__); \
+		bool UE_JOIN(TempBitfield, __LINE__) = BitField; \
+		Ar.SerializeBits(&UE_JOIN(TempBitfield, __LINE__), 1); \
+		BitField = UE_JOIN(TempBitfield, __LINE__); \
 	}
 
 	virtual void SerializeBits(void* V, int64 LengthBits)
diff --git a/Engine/Source/Runtime/Core/Public/Serialization/MemoryLayout.h b/Engine/Source/Runtime/Core/Public/Serialization/MemoryLayout.h
--- a/Engine/Source/Runtime/Core/Public/Serialization/MemoryLayout.h
+++ b/Engine/Source/Runtime/Core/Public/Serialization/MemoryLayout.h
@@ -493,17 +493,17 @@ PRAGMA_ENABLE_DEPRECATION_WARNINGS
 #define INTERNAL_LAYOUT_INTERFACE_PREFIX_NonVirtual(...) __VA_ARGS__
 #define INTERNAL_LAYOUT_INTERFACE_PREFIX_Virtual(...) __VA_ARGS__ virtual
 #define INTERNAL_LAYOUT_INTERFACE_PREFIX_Abstract(...) virtual
-#define INTERNAL_LAYOUT_INTERFACE_PREFIX(Type) PREPROCESSOR_JOIN(INTERNAL_LAYOUT_INTERFACE_PREFIX_, Type)
+#define INTERNAL_LAYOUT_INTERFACE_PREFIX(Type) UE_JOIN(INTERNAL_LAYOUT_INTERFACE_PREFIX_, Type)
 
 #define INTERNAL_LAYOUT_INTERFACE_SUFFIX_NonVirtual ;
 #define INTERNAL_LAYOUT_INTERFACE_SUFFIX_Virtual ;
 #define INTERNAL_LAYOUT_INTERFACE_SUFFIX_Abstract { return FTypeLayoutDesc::GetInvalidTypeLayout(); }
-#define INTERNAL_LAYOUT_INTERFACE_SUFFIX(Type) PREPROCESSOR_JOIN(INTERNAL_LAYOUT_INTERFACE_SUFFIX_, Type)
+#define INTERNAL_LAYOUT_INTERFACE_SUFFIX(Type) UE_JOIN(INTERNAL_LAYOUT_INTERFACE_SUFFIX_, Type)
 
 #define INTERNAL_LAYOUT_INTERFACE_INLINE_IMPL_NonVirtual { return StaticGetTypeLayout(); }
 #define INTERNAL_LAYOUT_INTERFACE_INLINE_IMPL_Virtual { return StaticGetTypeLayout(); }
 #define INTERNAL_LAYOUT_INTERFACE_INLINE_IMPL_Abstract { return FTypeLayoutDesc::GetInvalidTypeLayout(); }
-#define INTERNAL_LAYOUT_INTERFACE_INLINE_IMPL(Type) PREPROCESSOR_JOIN(INTERNAL_LAYOUT_INTERFACE_INLINE_IMPL_, Type)
+#define INTERNAL_LAYOUT_INTERFACE_INLINE_IMPL(Type) UE_JOIN(INTERNAL_LAYOUT_INTERFACE_INLINE_IMPL_, Type)
 
 #define INTERNAL_DECLARE_TYPE_LAYOUT_COMMON(T, InInterface) \
 	static constexpr int CounterBase = __COUNTER__; \
diff --git a/Engine/Source/Runtime/Core/Public/Stats/Stats.h b/Engine/Source/Runtime/Core/Public/Stats/Stats.h
--- a/Engine/Source/Runtime/Core/Public/Stats/Stats.h
+++ b/Engine/Source/Runtime/Core/Public/Stats/Stats.h
@@ -574,8 +574,8 @@ struct FStat_##StatName\
 // ----------------------------------------------------------------
 
 #define SCOPE_CYCLE_COUNTER_TO_TRACE(StatString, StatName, Condition) \
-	TRACE_CPUPROFILER_EVENT_DECLARE(PREPROCESSOR_JOIN(PREPROCESSOR_JOIN(__Decl_, StatName), __LINE__)); \
-	TRACE_CPUPROFILER_EVENT_SCOPE_USE(PREPROCESSOR_JOIN(PREPROCESSOR_JOIN(__Decl_, StatName), __LINE__), StatString, PREPROCESSOR_JOIN(PREPROCESSOR_JOIN(__Scope_, StatName), __LINE__), Condition && GCycleStatsShouldEmitNamedEvents);
+	TRACE_CPUPROFILER_EVENT_DECLARE(UE_JOIN(UE_JOIN(__Decl_, StatName), __LINE__)); \
+	TRACE_CPUPROFILER_EVENT_SCOPE_USE(UE_JOIN(UE_JOIN(__Decl_, StatName), __LINE__), StatString, UE_JOIN(UE_JOIN(__Scope_, StatName), __LINE__), Condition && GCycleStatsShouldEmitNamedEvents);
 
 // Note: Since these are defining the stat inline, no custom stat data will exist.
 #define DECLARE_SCOPE_CYCLE_COUNTER(CounterName,Stat,GroupId) \
diff --git a/Engine/Source/Runtime/Core/Public/Stats/StatsHierarchical.h b/Engine/Source/Runtime/Core/Public/Stats/StatsHierarchical.h
--- a/Engine/Source/Runtime/Core/Public/Stats/StatsHierarchical.h
+++ b/Engine/Source/Runtime/Core/Public/Stats/StatsHierarchical.h
@@ -105,9 +105,9 @@ protected:
 // Note: You need to call FStatsHierarhical::BeginMeasurements && FStatsHierarhical::EndMeasurements
 // for this to have an effect.
 #define DECLARE_SCOPE_HIERARCHICAL_COUNTER(CounterName) \
-	FStatsHierarchical::FScope PREPROCESSOR_JOIN(StatsHierarchicalScope, __LINE__)(#CounterName);
+	FStatsHierarchical::FScope UE_JOIN(StatsHierarchicalScope, __LINE__)(#CounterName);
 #define DECLARE_SCOPE_HIERARCHICAL_COUNTER_FUNC() \
-	FStatsHierarchical::FScope PREPROCESSOR_JOIN(StatsHierarchicalScope, __LINE__)(__FUNCTION__);
+	FStatsHierarchical::FScope UE_JOIN(StatsHierarchicalScope, __LINE__)(__FUNCTION__);
 
 #else
 
diff --git a/Engine/Source/Runtime/Core/Public/Stats/StatsMisc.h b/Engine/Source/Runtime/Core/Public/Stats/StatsMisc.h
--- a/Engine/Source/Runtime/Core/Public/Stats/StatsMisc.h
+++ b/Engine/Source/Runtime/Core/Public/Stats/StatsMisc.h
@@ -197,25 +197,25 @@ struct FScopeLogTime : public FConditionalScopeLogTime
 };
 
 #define SCOPE_LOG_TIME(Name,CumulativePtr) \
-	FConditionalScopeLogTime PREPROCESSOR_JOIN(ScopeLogTime,__LINE__)(true, Name, CumulativePtr);
+	FConditionalScopeLogTime UE_JOIN(ScopeLogTime,__LINE__)(true, Name, CumulativePtr);
 
 #define SCOPE_LOG_TIME_IN_SECONDS(Name,CumulativePtr) \
-	FConditionalScopeLogTime PREPROCESSOR_JOIN(ScopeLogTime,__LINE__)(true, Name, CumulativePtr, FScopeLogTime::ScopeLog_Seconds);
+	FConditionalScopeLogTime UE_JOIN(ScopeLogTime,__LINE__)(true, Name, CumulativePtr, FScopeLogTime::ScopeLog_Seconds);
 
 #define SCOPE_LOG_TIME_FUNC() \
-	FConditionalScopeLogTime PREPROCESSOR_JOIN(ScopeLogTime,__LINE__)(true, __FUNCTION__);
+	FConditionalScopeLogTime UE_JOIN(ScopeLogTime,__LINE__)(true, __FUNCTION__);
 
 #define SCOPE_LOG_TIME_FUNC_WITH_GLOBAL(CumulativePtr) \
-	FConditionalScopeLogTime PREPROCESSOR_JOIN(ScopeLogTime,__LINE__)(true, __FUNCTION__,CumulativePtr);
+	FConditionalScopeLogTime UE_JOIN(ScopeLogTime,__LINE__)(true, __FUNCTION__,CumulativePtr);
 
 #define CONDITIONAL_SCOPE_LOG_TIME(Condition, Name,CumulativePtr) \
-	FConditionalScopeLogTime PREPROCESSOR_JOIN(ScopeLogTime,__LINE__)(Condition, Name, CumulativePtr);
+	FConditionalScopeLogTime UE_JOIN(ScopeLogTime,__LINE__)(Condition, Name, CumulativePtr);
 
 #define CONDITIONAL_SCOPE_LOG_TIME_IN_SECONDS(Condition, Name,CumulativePtr) \
-	FConditionalScopeLogTime PREPROCESSOR_JOIN(ScopeLogTime,__LINE__)(Condition, Name, CumulativePtr, FScopeLogTime::ScopeLog_Seconds);
+	FConditionalScopeLogTime UE_JOIN(ScopeLogTime,__LINE__)(Condition, Name, CumulativePtr, FScopeLogTime::ScopeLog_Seconds);
 
 #define CONDITIONAL_SCOPE_LOG_TIME_FUNC(Condition) \
-	FConditionalScopeLogTime PREPROCESSOR_JOIN(ScopeLogTime,__LINE__)(Condition, __FUNCTION__);
+	FConditionalScopeLogTime UE_JOIN(ScopeLogTime,__LINE__)(Condition, __FUNCTION__);
 
 #define CONDITIONAL_SCOPE_LOG_TIME_FUNC_WITH_GLOBAL(Condition, CumulativePtr) \
-	FConditionalScopeLogTime PREPROCESSOR_JOIN(ScopeLogTime,__LINE__)(Condition, __FUNCTION__,CumulativePtr);
+	FConditionalScopeLogTime UE_JOIN(ScopeLogTime,__LINE__)(Condition, __FUNCTION__,CumulativePtr);
diff --git a/Engine/Source/Runtime/Core/Public/Templates/UnrealTemplate.h b/Engine/Source/Runtime/Core/Public/Templates/UnrealTemplate.h
--- a/Engine/Source/Runtime/Core/Public/Templates/UnrealTemplate.h
+++ b/Engine/Source/Runtime/Core/Public/Templates/UnrealTemplate.h
@@ -427,9 +427,9 @@ private:
  * Macro variant on TGuardValue<bool> that can deal with bitfields which cannot be passed by reference in to TGuardValue
  */
 #define FGuardValue_Bitfield(ReferenceValue, NewValue) \
-	const bool PREPROCESSOR_JOIN(TempBitfield, __LINE__) = ReferenceValue; \
+	const bool UE_JOIN(TempBitfield, __LINE__) = ReferenceValue; \
 	ReferenceValue = NewValue; \
-	const TGuardValue_Bitfield_Cleanup<TFunction<void()>> PREPROCESSOR_JOIN(TempBitfieldCleanup, __LINE__)([&](){ ReferenceValue = PREPROCESSOR_JOIN(TempBitfield, __LINE__); });
+	const TGuardValue_Bitfield_Cleanup<TFunction<void()>> UE_JOIN(TempBitfieldCleanup, __LINE__)([&](){ ReferenceValue = UE_JOIN(TempBitfield, __LINE__); });
 
 /** 
  * Commonly used to make sure a value is incremented, and then decremented anyway the function can terminate.
diff --git a/Engine/Source/Runtime/CoreUObject/Private/UObject/ObjectSerializeAccessScope.h b/Engine/Source/Runtime/CoreUObject/Private/UObject/ObjectSerializeAccessScope.h
--- a/Engine/Source/Runtime/CoreUObject/Private/UObject/ObjectSerializeAccessScope.h
+++ b/Engine/Source/Runtime/CoreUObject/Private/UObject/ObjectSerializeAccessScope.h
@@ -42,8 +42,8 @@ namespace UE::CoreUObject::Private
 		bool bSuspended;
 	};
 }
-#define UE_SERIALIZE_ACCCESS_SCOPE(Object) UE::CoreUObject::Private::FObjectSerializeAccessScope PREPROCESSOR_JOIN(_Scope, __LINE__)(*(Object))
-#define UE_SERIALIZE_ACCCESS_SCOPE_SUSPEND() UE::CoreUObject::Private::FObjectSerializeAccessScope::FSuspendScope PREPROCESSOR_JOIN(_Suspend, __LINE__)
+#define UE_SERIALIZE_ACCCESS_SCOPE(Object) UE::CoreUObject::Private::FObjectSerializeAccessScope UE_JOIN(_Scope, __LINE__)(*(Object))
+#define UE_SERIALIZE_ACCCESS_SCOPE_SUSPEND() UE::CoreUObject::Private::FObjectSerializeAccessScope::FSuspendScope UE_JOIN(_Suspend, __LINE__)
 
 #else
 
diff --git a/Engine/Source/Runtime/CoreUObject/Public/Cooker/CookDependency.h b/Engine/Source/Runtime/CoreUObject/Public/Cooker/CookDependency.h
--- a/Engine/Source/Runtime/CoreUObject/Public/Cooker/CookDependency.h
+++ b/Engine/Source/Runtime/CoreUObject/Public/Cooker/CookDependency.h
@@ -374,13 +374,13 @@ struct FCookDependencyFunctionRegistration
 */
 #define UE_COOK_DEPENDENCY_FUNCTION(Name, Function) \
 	UE::Cook::Dependency::Private::FCookDependencyFunctionRegistration \
-	PREPROCESSOR_JOIN(FCookDependencyFunctionRegistration_,Name)(TEXT(#Name), Function)
+	UE_JOIN(FCookDependencyFunctionRegistration_,Name)(TEXT(#Name), Function)
 /**
  * Return the FName to use to call a function that was registered via UE_COOK_DEPENDENCY_FUNCTION(Name, Function).
  * Name should be the same bare cpptoken that was passed into UE_COOK_DEPENDENCY_FUNCTION.
  */
 #define UE_COOK_DEPENDENCY_FUNCTION_CALL(Name) \
-	PREPROCESSOR_JOIN(FCookDependencyFunctionRegistration_,Name).GetFName()
+	UE_JOIN(FCookDependencyFunctionRegistration_,Name).GetFName()
 
 #else // WITH_EDITOR
 
diff --git a/Engine/Source/Runtime/Engine/Public/SceneRenderBuilderInterface.h b/Engine/Source/Runtime/Engine/Public/SceneRenderBuilderInterface.h
--- a/Engine/Source/Runtime/Engine/Public/SceneRenderBuilderInterface.h
+++ b/Engine/Source/Runtime/Engine/Public/SceneRenderBuilderInterface.h
@@ -127,8 +127,8 @@ public:
 
 // A scope macro used to construct a scene render group scope.
 #define SCENE_RENDER_GROUP_SCOPE(SceneRenderBuilder, Name, Flags) \
-	FSceneRenderGroupScope PREPROCESSOR_JOIN(__SceneRenderGroup_ScopeRef_,__LINE__)( \
-		  (SceneRenderBuilder)                                                       \
-		, Forward<FString&&>(Name)                                                   \
-		, Flags                                                                      \
+	FSceneRenderGroupScope UE_JOIN(__SceneRenderGroup_ScopeRef_,__LINE__)( \
+		  (SceneRenderBuilder)                                             \
+		, Forward<FString&&>(Name)                                         \
+		, Flags                                                            \
 	)
\ No newline at end of file
diff --git a/Engine/Source/Runtime/Net/Core/Public/Net/Core/Trace/Private/NetTraceInternal.h b/Engine/Source/Runtime/Net/Core/Public/Net/Core/Trace/Private/NetTraceInternal.h
--- a/Engine/Source/Runtime/Net/Core/Public/Net/Core/Trace/Private/NetTraceInternal.h
+++ b/Engine/Source/Runtime/Net/Core/Public/Net/Core/Trace/Private/NetTraceInternal.h
@@ -658,39 +658,39 @@ void FNetTraceCollector::Reset()
 #define UE_NET_TRACE_INTERNAL_POP_SEND_BUNCH(Collector) UE_NET_TRACE_DO_IF(Collector, FNetTrace::PopSendBunch(*Collector))
 #define UE_NET_TRACE_INTERNAL_EVENTS(Collector, SrcCollector, Stream) UE_NET_TRACE_DO_IF(Collector, FNetTrace::FoldTraceCollector(Collector, SrcCollector, GetBitStreamPositionForNetTrace(Stream)))
 #define UE_NET_TRACE_INTERNAL_END_BUNCH(Collector, Bunch, ...) UE_NET_TRACE_DO_IF(Collector, FNetTrace::TraceBunch(*Collector, MakeBunchInfo(Bunch), __VA_ARGS__))
-#define UE_NET_TRACE_INTERNAL_BUNCH_SCOPE(Collector, Bunch, ...) FNetTraceBunchScope<decltype(Bunch)> PREPROCESSOR_JOIN(NetTraceBunchScope, __LINE__)(Bunch, __VA_ARGS__, Collector)
+#define UE_NET_TRACE_INTERNAL_BUNCH_SCOPE(Collector, Bunch, ...) FNetTraceBunchScope<decltype(Bunch)> UE_JOIN(NetTraceBunchScope, __LINE__)(Bunch, __VA_ARGS__, Collector)
 
 #define UE_NET_TRACE_INTERNAL_SCOPE(Name, Stream, Collector, Verbosity) \
-	auto PREPROCESSOR_JOIN(NetTraceNameFunc_, __LINE__) = []() { static const UE::Net::FNetDebugName* Name = UE::Net::CreatePersistentNetDebugName(TEXT(#Name)); return FNetTrace::TraceName(Name); }; \
-	FNetTrace::TChooseTraceEventScope<Verbosity, decltype(Stream)>::Type PREPROCESSOR_JOIN(NetTraceScope, __LINE__)(PREPROCESSOR_JOIN(NetTraceNameFunc_, __LINE__), Stream, Collector, Verbosity)
+	auto UE_JOIN(NetTraceNameFunc_, __LINE__) = []() { static const UE::Net::FNetDebugName* Name = UE::Net::CreatePersistentNetDebugName(TEXT(#Name)); return FNetTrace::TraceName(Name); }; \
+	FNetTrace::TChooseTraceEventScope<Verbosity, decltype(Stream)>::Type UE_JOIN(NetTraceScope, __LINE__)(UE_JOIN(NetTraceNameFunc_, __LINE__), Stream, Collector, Verbosity)
 
 #define UE_NET_TRACE_INTERNAL_OBJECT_SCOPE(HandleOrNetGUID, Stream, Collector, Verbosity) \
-	FNetTrace::TChooseTraceEventScope<Verbosity, decltype(Stream)>::Type PREPROCESSOR_JOIN(NetTraceScope, __LINE__)(GetObjectIdForNetTrace(HandleOrNetGUID), Stream, Collector, Verbosity)
+	FNetTrace::TChooseTraceEventScope<Verbosity, decltype(Stream)>::Type UE_JOIN(NetTraceScope, __LINE__)(GetObjectIdForNetTrace(HandleOrNetGUID), Stream, Collector, Verbosity)
 
 #define UE_NET_TRACE_INTERNAL_DYNAMIC_NAME_SCOPE(Name, Stream, Collector, Verbosity) \
-	auto PREPROCESSOR_JOIN(NetTraceDynamicNameFunc_, __LINE__) = [&]() { return FNetTrace::TraceName(Name); }; \
-	FNetTrace::TChooseTraceEventScope<Verbosity, decltype(Stream)>::Type PREPROCESSOR_JOIN(NetTraceScope, __LINE__)(PREPROCESSOR_JOIN(NetTraceDynamicNameFunc_, __LINE__), Stream, Collector, Verbosity)
+	auto UE_JOIN(NetTraceDynamicNameFunc_, __LINE__) = [&]() { return FNetTrace::TraceName(Name); }; \
+	FNetTrace::TChooseTraceEventScope<Verbosity, decltype(Stream)>::Type UE_JOIN(NetTraceScope, __LINE__)(UE_JOIN(NetTraceDynamicNameFunc_, __LINE__), Stream, Collector, Verbosity)
 
 #define UE_NET_TRACE_INTERNAL_NAMED_SCOPE(ScopeName, EventName, Stream, Collector, Verbosity) \
-	auto PREPROCESSOR_JOIN(NetTraceNameFunc_, __LINE__) = []() { static const UE::Net::FNetDebugName* Name = UE::Net::CreatePersistentNetDebugName(TEXT(#EventName)); return FNetTrace::TraceName(Name); }; \
-	FNetTrace::TChooseTraceEventScope<Verbosity, decltype(Stream)>::Type ScopeName(PREPROCESSOR_JOIN(NetTraceNameFunc_, __LINE__), Stream, Collector, Verbosity)
+	auto UE_JOIN(NetTraceNameFunc_, __LINE__) = []() { static const UE::Net::FNetDebugName* Name = UE::Net::CreatePersistentNetDebugName(TEXT(#EventName)); return FNetTrace::TraceName(Name); }; \
+	FNetTrace::TChooseTraceEventScope<Verbosity, decltype(Stream)>::Type ScopeName(UE_JOIN(NetTraceNameFunc_, __LINE__), Stream, Collector, Verbosity)
 
 #define UE_NET_TRACE_INTERNAL_NAMED_OBJECT_SCOPE(ScopeName, HandleOrNetGUID, Stream, Collector, Verbosity) \
 	FNetTrace::TChooseTraceEventScope<Verbosity, decltype(Stream)>::Type ScopeName(GetObjectIdForNetTrace(HandleOrNetGUID), Stream, Collector, Verbosity)
 
 #define UE_NET_TRACE_INTERNAL_NAMED_DYNAMIC_NAME_SCOPE(ScopeName, EventName, Stream, Collector, Verbosity) \
-	auto PREPROCESSOR_JOIN(NetTraceDynamicNameFunc_, __LINE__) = [&]() { UE::Net::FNetDebugNameId NameId = FNetTrace::TraceName(EventName); return NameId; }; \
-	FNetTrace::TChooseTraceEventScope<Verbosity, decltype(Stream)>::Type ScopeName(PREPROCESSOR_JOIN(NetTraceDynamicNameFunc_, __LINE__), Stream, Collector, Verbosity)
+	auto UE_JOIN(NetTraceDynamicNameFunc_, __LINE__) = [&]() { UE::Net::FNetDebugNameId NameId = FNetTrace::TraceName(EventName); return NameId; }; \
+	FNetTrace::TChooseTraceEventScope<Verbosity, decltype(Stream)>::Type ScopeName(UE_JOIN(NetTraceDynamicNameFunc_, __LINE__), Stream, Collector, Verbosity)
 
 #define UE_NET_TRACE_INTERNAL_SET_SCOPE_NAME(ScopeName, EventName) UE_NET_TRACE_DO_IF(ScopeName.IsValid(), ScopeName.SetEventName(FNetTrace::TraceName(EventName)))
 #define UE_NET_TRACE_INTERNAL_SET_SCOPE_OBJECTID(ScopeName, HandleOrNetGUID) UE_NET_TRACE_DO_IF(ScopeName.IsValid(), ScopeName.SetEventObjectId(GetObjectIdForNetTrace(HandleOrNetGUID)))
 #define UE_NET_TRACE_INTERNAL_EXIT_NAMED_SCOPE(ScopeName) ScopeName.ExitScope()
-#define UE_NET_TRACE_INTERNAL_OFFSET_SCOPE(Offset, Collector) FNetTraceOffsetScope PREPROCESSOR_JOIN(NetTraceOffsetScope, __LINE__)(Offset, Collector);
+#define UE_NET_TRACE_INTERNAL_OFFSET_SCOPE(Offset, Collector) FNetTraceOffsetScope UE_JOIN(NetTraceOffsetScope, __LINE__)(Offset, Collector);
 
 #define UE_NET_TRACE_INTERNAL(Name, Collector, StartPos, EndPos, Verbosity) \
 	do { if (FNetTrace::GetCollectorAtVerbosity(Collector, Verbosity)) { \
-			static const UE::Net::FNetDebugName* PREPROCESSOR_JOIN(NetTraceNameId_, __LINE__) = UE::Net::CreatePersistentNetDebugName(TEXT(#Name)); \
-			FNetTrace::TracePacketContentEvent(*Collector, FNetTrace::TraceName(PREPROCESSOR_JOIN(NetTraceNameId_, __LINE__)), StartPos, EndPos, Verbosity); \
+			static const UE::Net::FNetDebugName* UE_JOIN(NetTraceNameId_, __LINE__) = UE::Net::CreatePersistentNetDebugName(TEXT(#Name)); \
+			FNetTrace::TracePacketContentEvent(*Collector, FNetTrace::TraceName(UE_JOIN(NetTraceNameId_, __LINE__)), StartPos, EndPos, Verbosity); \
 	} } while (0)
 
 #define UE_NET_TRACE_INTERNAL_DYNAMIC_NAME(Name, Collector, StartPos, EndPos, Verbosity) UE_NET_TRACE_DO_IF(FNetTrace::GetCollectorAtVerbosity(Collector, Verbosity), FNetTrace::TracePacketContentEvent(*Collector, FNetTrace::TraceName(Name), StartPos, EndPos, Verbosity))
@@ -707,14 +707,14 @@ void FNetTraceCollector::Reset()
 
 #define UE_NET_TRACE_INTERNAL_PACKET_STATSCOUNTER(GameInstanceId, ConnectionId, Name, StatValue, Verbosity) \
 	do { if (FNetTrace::GetNetTraceVerbosityEnabled(Verbosity)) { \
-			static const UE::Net::FNetDebugName* PREPROCESSOR_JOIN(NetTraceNameId_, __LINE__) = UE::Net::CreatePersistentNetDebugName(TEXT(#Name)); \
-			FNetTrace::TracePacketStatsCounter(GameInstanceId, ConnectionId, FNetTrace::TraceName(PREPROCESSOR_JOIN(NetTraceNameId_, __LINE__)), StatValue); \
+			static const UE::Net::FNetDebugName* UE_JOIN(NetTraceNameId_, __LINE__) = UE::Net::CreatePersistentNetDebugName(TEXT(#Name)); \
+			FNetTrace::TracePacketStatsCounter(GameInstanceId, ConnectionId, FNetTrace::TraceName(UE_JOIN(NetTraceNameId_, __LINE__)), StatValue); \
 	} } while (0)
 
 #define UE_NET_TRACE_INTERNAL_FRAME_STATSCOUNTER(GameInstanceId, Name, StatValue, Verbosity) \
 	do { if (FNetTrace::GetNetTraceVerbosityEnabled(Verbosity)) { \
-			static const UE::Net::FNetDebugName* PREPROCESSOR_JOIN(NetTraceNameId_, __LINE__) = UE::Net::CreatePersistentNetDebugName(TEXT(#Name)); \
-			FNetTrace::TraceFrameStatsCounter(GameInstanceId, FNetTrace::TraceName(PREPROCESSOR_JOIN(NetTraceNameId_, __LINE__)), StatValue); \
+			static const UE::Net::FNetDebugName* UE_JOIN(NetTraceNameId_, __LINE__) = UE::Net::CreatePersistentNetDebugName(TEXT(#Name)); \
+			FNetTrace::TraceFrameStatsCounter(GameInstanceId, FNetTrace::TraceName(UE_JOIN(NetTraceNameId_, __LINE__)), StatValue); \
 	} } while (0)
 
 #define UE_NET_TRACE_INTERNAL_END_SESSION(GameInstanceId) FNetTrace::TraceEndSession(GameInstanceId);
@@ -722,8 +722,8 @@ void FNetTraceCollector::Reset()
 
 	#define UE_NET_TRACE_INTERNAL_POLL_OBJECT_SCOPE(Handle, Timer) FNetTraceObjectPollTimer __NetTracePollTimer(Timer, (GNetTraceRuntimeVerbosity >= ENetTraceVerbosity::VeryVerbose && Handle.IsValid()) ? Handle.GetReplicationSystemId() : -1, Handle.GetId());
 #define UE_NET_TRACE_INTERNAL_POLL_OBJECT_IS_WASTE __NetTracePollTimer.bIsWaste = true;
-#define UE_NET_TRACE_INTERNAL_QUANTIZE_OBJECT_SCOPE(Handle, Timer) FNetTraceObjectQuantizeTimer PREPROCESSOR_JOIN(__NetTraceQuantizeTimer, __LINE__)(Timer, (GNetTraceRuntimeVerbosity >= ENetTraceVerbosity::VeryVerbose && Handle.IsValid()) ? Handle.GetReplicationSystemId() : -1, Handle.GetId());
-#define UE_NET_TRACE_INTERNAL_WRITE_OBJECT_SCOPE(Handle, Timer) FNetTraceObjectWriteTimer PREPROCESSOR_JOIN(__NetTraceWriteTimer, __LINE__)(Timer, (GNetTraceRuntimeVerbosity >= ENetTraceVerbosity::VeryVerbose && Handle.IsValid()) ? Handle.GetReplicationSystemId() : -1, Handle.GetId());	
+#define UE_NET_TRACE_INTERNAL_QUANTIZE_OBJECT_SCOPE(Handle, Timer) FNetTraceObjectQuantizeTimer UE_JOIN(__NetTraceQuantizeTimer, __LINE__)(Timer, (GNetTraceRuntimeVerbosity >= ENetTraceVerbosity::VeryVerbose && Handle.IsValid()) ? Handle.GetReplicationSystemId() : -1, Handle.GetId());
+#define UE_NET_TRACE_INTERNAL_WRITE_OBJECT_SCOPE(Handle, Timer) FNetTraceObjectWriteTimer UE_JOIN(__NetTraceWriteTimer, __LINE__)(Timer, (GNetTraceRuntimeVerbosity >= ENetTraceVerbosity::VeryVerbose && Handle.IsValid()) ? Handle.GetReplicationSystemId() : -1, Handle.GetId());	
 	
 #define UE_NET_TRACE_INTERNAL_SUBOBJECT(ParentHandle, SubObjectHandle) FNetTrace::TraceSubObject(ParentHandle.GetReplicationSystemId(), GetObjectIdForNetTrace(ParentHandle), GetObjectIdForNetTrace(SubObjectHandle))
 					
diff --git a/Engine/Source/Runtime/Net/Iris/Public/Iris/Core/IrisProfiler.h b/Engine/Source/Runtime/Net/Iris/Public/Iris/Core/IrisProfiler.h
--- a/Engine/Source/Runtime/Net/Iris/Public/Iris/Core/IrisProfiler.h
+++ b/Engine/Source/Runtime/Net/Iris/Public/Iris/Core/IrisProfiler.h
@@ -37,8 +37,8 @@
 #		pragma comment (lib, "c:/Program Files/Superluminal/Performance/API/lib/x64/PerformanceAPI_MD.lib")
 #		define IRIS_PROFILER_SCOPE(X) PERFORMANCEAPI_INSTRUMENT(PREPROCESSOR_TO_STRING(X))
 #		define IRIS_PROFILER_SCOPE_CONDITIONAL(X,Cond) PERFORMANCEAPI_INSTRUMENT(PREPROCESSOR_TO_STRING(X))
-#		define IRIS_PROFILER_SCOPE_TEXT(X) PERFORMANCEAPI_INSTRUMENT_DATA(PREPROCESSOR_JOIN(IrisProfilerScope, __LINE__), X)
-#		define IRIS_PROFILER_SCOPE_TEXT_CONDITIONAL(X, Cond) PERFORMANCEAPI_INSTRUMENT_DATA(PREPROCESSOR_JOIN(IrisProfilerScope, __LINE__), X)
+#		define IRIS_PROFILER_SCOPE_TEXT(X) PERFORMANCEAPI_INSTRUMENT_DATA(UE_JOIN(IrisProfilerScope, __LINE__), X)
+#		define IRIS_PROFILER_SCOPE_TEXT_CONDITIONAL(X, Cond) PERFORMANCEAPI_INSTRUMENT_DATA(UE_JOIN(IrisProfilerScope, __LINE__), X)
 #	else
 #		include "ProfilingDebugging/CpuProfilerTrace.h"
 #		define IRIS_PROFILER_SCOPE(X) TRACE_CPUPROFILER_EVENT_SCOPE(X)
diff --git a/Engine/Source/Runtime/PacketHandlers/PacketHandler/Private/PacketHandler.cpp b/Engine/Source/Runtime/PacketHandlers/PacketHandler/Private/PacketHandler.cpp
--- a/Engine/Source/Runtime/PacketHandlers/PacketHandler/Private/PacketHandler.cpp
+++ b/Engine/Source/Runtime/PacketHandlers/PacketHandler/Private/PacketHandler.cpp
@@ -56,18 +56,18 @@ static FAutoConsoleVariableRef CVarNetPacketHandlerTimeguardLimit(
 
 // Lightweight time guard. Note: Threshold of 0 disables the timeguard
 #define NET_LIGHTWEIGHT_TIME_GUARD_BEGIN( Name, ThresholdMS ) \
-	double PREPROCESSOR_JOIN(__TimeGuard_ThresholdMS_, Name) = ThresholdMS; \
-	uint64 PREPROCESSOR_JOIN(__TimeGuard_StartCycles_, Name) = ( ThresholdMS > 0.0f && GPacketHandlerTimeguardLimit > 0 ) ? FPlatformTime::Cycles64() : 0; \
+	double UE_JOIN(__TimeGuard_ThresholdMS_, Name) = ThresholdMS; \
+	uint64 UE_JOIN(__TimeGuard_StartCycles_, Name) = ( ThresholdMS > 0.0f && GPacketHandlerTimeguardLimit > 0 ) ? FPlatformTime::Cycles64() : 0; \
 	GPacketHandlerDiscardTimeguardMeasurement = false;
 
 #define NET_LIGHTWEIGHT_TIME_GUARD_END( Name, NameStringCode ) \
-	if ( PREPROCESSOR_JOIN(__TimeGuard_ThresholdMS_, Name) > 0.0f && GPacketHandlerTimeguardLimit > 0 && !GPacketHandlerDiscardTimeguardMeasurement ) \
+	if ( UE_JOIN(__TimeGuard_ThresholdMS_, Name) > 0.0f && GPacketHandlerTimeguardLimit > 0 && !GPacketHandlerDiscardTimeguardMeasurement ) \
 	{\
-		double PREPROCESSOR_JOIN(__TimeGuard_MSElapsed_,Name) = FPlatformTime::ToMilliseconds64( FPlatformTime::Cycles64() - PREPROCESSOR_JOIN(__TimeGuard_StartCycles_,Name) ); \
-		if ( PREPROCESSOR_JOIN(__TimeGuard_MSElapsed_,Name) > PREPROCESSOR_JOIN(__TimeGuard_ThresholdMS_, Name) ) \
+		double UE_JOIN(__TimeGuard_MSElapsed_,Name) = FPlatformTime::ToMilliseconds64( FPlatformTime::Cycles64() - UE_JOIN(__TimeGuard_StartCycles_,Name) ); \
+		if ( UE_JOIN(__TimeGuard_MSElapsed_,Name) > UE_JOIN(__TimeGuard_ThresholdMS_, Name) ) \
 		{ \
 			FString ReportName = NameStringCode; \
-			UE_LOG(PacketHandlerLog, Warning, TEXT("PacketHandler: %s - %s took %.2fms!"), TEXT(#Name), *ReportName, PREPROCESSOR_JOIN(__TimeGuard_MSElapsed_,Name)); \
+			UE_LOG(PacketHandlerLog, Warning, TEXT("PacketHandler: %s - %s took %.2fms!"), TEXT(#Name), *ReportName, UE_JOIN(__TimeGuard_MSElapsed_,Name)); \
 			GPacketHandlerTimeguardLimit--; \
 		} \
 	}
diff --git a/Engine/Source/Runtime/PakFile/Internal/IoDispatcherFileBackendTypes.h b/Engine/Source/Runtime/PakFile/Internal/IoDispatcherFileBackendTypes.h
--- a/Engine/Source/Runtime/PakFile/Internal/IoDispatcherFileBackendTypes.h
+++ b/Engine/Source/Runtime/PakFile/Internal/IoDispatcherFileBackendTypes.h
@@ -14,7 +14,7 @@
 #endif
 
 #ifndef PLATFORM_IODISPATCHER_MODULE
-#define PLATFORM_IODISPATCHER_MODULE PREPROCESSOR_TO_STRING(PREPROCESSOR_JOIN(PLATFORM_HEADER_NAME, PlatformIoDispatcher))
+#define PLATFORM_IODISPATCHER_MODULE PREPROCESSOR_TO_STRING(UE_JOIN(PLATFORM_HEADER_NAME, PlatformIoDispatcher))
 #endif
 
 #define UE_FILEIOSTORE_STATS_ENABLED (COUNTERSTRACE_ENABLED || CSV_PROFILER_STATS)
diff --git a/Engine/Source/Runtime/RHI/Public/RHIBreadcrumbs.h b/Engine/Source/Runtime/RHI/Public/RHIBreadcrumbs.h
--- a/Engine/Source/Runtime/RHI/Public/RHIBreadcrumbs.h
+++ b/Engine/Source/Runtime/RHI/Public/RHIBreadcrumbs.h
@@ -1363,12 +1363,12 @@ namespace UE::RHI::GPUProfiler
 	#define RHI_BREADCRUMB_FIELD(Name, Value) std::forward_as_tuple(TEXT(Name), Value)
 
 	#define RHI_BREADCRUMB_EVENT_PRIVATE_IMPL(RHICmdList_Or_RHIContext, Stat, Condition, StaticName, Format, ...)\
-		TOptional<FRHIBreadcrumbScope> PREPROCESSOR_JOIN(BreadcrumbScope, __LINE__);							 \
+		TOptional<FRHIBreadcrumbScope> UE_JOIN(BreadcrumbScope, __LINE__);										 \
 		do																										 \
 		{																										 \
 			if (Condition)																						 \
 			{																									 \
-				PREPROCESSOR_JOIN(BreadcrumbScope, __LINE__).Emplace(											 \
+				UE_JOIN(BreadcrumbScope, __LINE__).Emplace(														 \
 					UE::RHI::Breadcrumbs::Private::GetRHICmdList(RHICmdList_Or_RHIContext),						 \
 					RHI_BREADCRUMB_DESC_FORWARD_VALUES(															 \
 						  StaticName																			 \
diff --git a/Engine/Source/Runtime/RHI/Public/RHICommandList.h b/Engine/Source/Runtime/RHI/Public/RHICommandList.h
--- a/Engine/Source/Runtime/RHI/Public/RHICommandList.h
+++ b/Engine/Source/Runtime/RHI/Public/RHICommandList.h
@@ -1630,19 +1630,19 @@ struct FRHICommand : public FRHICommandBase
 	struct CommandName final : public FRHICommand<CommandName<TemplateType>, FUnnamedRhiCommand>
 
 #define FRHICOMMAND_MACRO(CommandName)								\
-	struct PREPROCESSOR_JOIN(CommandName##String, __LINE__)			\
+	struct UE_JOIN(CommandName##String, __LINE__)					\
 	{																\
 		static const TCHAR* TStr() { return TEXT(#CommandName); }	\
 	};																\
-	struct CommandName final : public FRHICommand<CommandName, PREPROCESSOR_JOIN(CommandName##String, __LINE__)>
+	struct CommandName final : public FRHICommand<CommandName, UE_JOIN(CommandName##String, __LINE__)>
 
 #define FRHICOMMAND_MACRO_TPL(TemplateType, CommandName)			\
-	struct PREPROCESSOR_JOIN(CommandName##String, __LINE__)			\
+	struct UE_JOIN(CommandName##String, __LINE__)					\
 	{																\
 		static const TCHAR* TStr() { return TEXT(#CommandName); }	\
 	};																\
 	template<typename TemplateType>									\
-	struct CommandName final : public FRHICommand<CommandName<TemplateType>, PREPROCESSOR_JOIN(CommandName##String, __LINE__)>
+	struct CommandName final : public FRHICommand<CommandName<TemplateType>, UE_JOIN(CommandName##String, __LINE__)>
 
 FRHICOMMAND_MACRO(FRHICommandBeginUpdateMultiFrameResource)
 {
@@ -5112,7 +5112,7 @@ struct FRHIScopedGPUMask
 };
 
 #if WITH_MGPU
-	#define SCOPED_GPU_MASK(RHICmdList, GPUMask) FRHIScopedGPUMask PREPROCESSOR_JOIN(ScopedGPUMask, __LINE__){ RHICmdList, GPUMask }
+	#define SCOPED_GPU_MASK(RHICmdList, GPUMask) FRHIScopedGPUMask UE_JOIN(ScopedGPUMask, __LINE__){ RHICmdList, GPUMask }
 #else
 	#define SCOPED_GPU_MASK(RHICmdList, GPUMask)
 #endif // WITH_MGPU
@@ -5152,7 +5152,7 @@ private:
 #endif
 };
 
-#define SCOPED_UNIFORM_BUFFER_STATIC_BINDINGS(RHICmdList, UniformBuffers) FScopedUniformBufferStaticBindings PREPROCESSOR_JOIN(UniformBuffers, __LINE__){ RHICmdList, UniformBuffers }
+#define SCOPED_UNIFORM_BUFFER_STATIC_BINDINGS(RHICmdList, UniformBuffers) FScopedUniformBufferStaticBindings UE_JOIN(UniformBuffers, __LINE__){ RHICmdList, UniformBuffers }
 
 // Helper to enable the use of graphics RHI command lists from within platform RHI implementations.
 // Recorded commands are dispatched when the command list is destructed. Intended for use on the stack / in a scope block.
diff --git a/Engine/Source/Runtime/RenderCore/Public/ProfilingDebugging/RealtimeGPUProfiler.h b/Engine/Source/Runtime/RenderCore/Public/ProfilingDebugging/RealtimeGPUProfiler.h
--- a/Engine/Source/Runtime/RenderCore/Public/ProfilingDebugging/RealtimeGPUProfiler.h
+++ b/Engine/Source/Runtime/RenderCore/Public/ProfilingDebugging/RealtimeGPUProfiler.h
@@ -59,12 +59,12 @@
 	};
 
 	#define RHI_BREADCRUMB_EVENT_GAMETHREAD_PRIVATE_IMPL(Stat, Condition, StaticName, Format, ...)	\
-		TOptional<FRHIBreadcrumbScope_GameThread> PREPROCESSOR_JOIN(BreadcrumbScope, __LINE__);		\
+		TOptional<FRHIBreadcrumbScope_GameThread> UE_JOIN(BreadcrumbScope, __LINE__);				\
 		do																							\
 		{																							\
 			if (Condition)																			\
 			{																						\
-				PREPROCESSOR_JOIN(BreadcrumbScope, __LINE__).Emplace(								\
+				UE_JOIN(BreadcrumbScope, __LINE__).Emplace(											\
 					RHI_BREADCRUMB_DESC_COPY_VALUES(												\
 						  StaticName																\
 						, Format																	\
@@ -231,15 +231,15 @@ CSV_DECLARE_CATEGORY_MODULE_EXTERN(RENDERCORE_API, GPU);
 	#define DEFINE_GPU_STAT(StatName)                           DEFINE_STAT(Stat_GPU_##StatName);                                                    CSV_DEFINE_STAT(GPU,StatName);                FRHIDrawStatsCategory DrawcallCountCategory_##StatName;
 	#define DEFINE_GPU_DRAWCALL_STAT(StatName)                  DEFINE_STAT(Stat_GPU_##StatName);                                                    CSV_DEFINE_STAT(GPU,StatName);                FRHIDrawStatsCategory DrawcallCountCategory_##StatName((TCHAR*)TEXT(#StatName));
 
-	#define SCOPED_GPU_STAT_VERBOSE(RHICmdList, StatName, Description)                 \
-		FScopedGPUStatEvent PREPROCESSOR_JOIN(GPUStatEvent_##StatName,__LINE__)(       \
-			  RHICmdList														       \
-			, CSV_STAT_FNAME(StatName)											       \
-			, GET_STATID(Stat_GPU_##StatName)                                          \
-			, Description														       \
-		);																			   \
-		FScopedDrawStatCategory PREPROCESSOR_JOIN(DrawCallScope_##StatName, __LINE__)( \
-			RHICmdList, DrawcallCountCategory_##StatName							   \
+	#define SCOPED_GPU_STAT_VERBOSE(RHICmdList, StatName, Description)			\
+		FScopedGPUStatEvent UE_JOIN(GPUStatEvent_##StatName,__LINE__)(			\
+			  RHICmdList														\
+			, CSV_STAT_FNAME(StatName)											\
+			, GET_STATID(Stat_GPU_##StatName)									\
+			, Description														\
+		);																		\
+		FScopedDrawStatCategory UE_JOIN(DrawCallScope_##StatName, __LINE__)(	\
+			RHICmdList, DrawcallCountCategory_##StatName						\
 		);
 
 	#define SCOPED_GPU_STAT(RHICmdList, StatName) SCOPED_GPU_STAT_VERBOSE(RHICmdList, StatName, nullptr)
diff --git a/Engine/Source/Runtime/RenderCore/Public/RenderGraphBuilder.h b/Engine/Source/Runtime/RenderCore/Public/RenderGraphBuilder.h
--- a/Engine/Source/Runtime/RenderCore/Public/RenderGraphBuilder.h
+++ b/Engine/Source/Runtime/RenderCore/Public/RenderGraphBuilder.h
@@ -1197,7 +1197,7 @@ private:
 };
 
 #define RDG_ASYNC_COMPUTE_BUDGET_SCOPE(GraphBuilder, AsyncComputeBudget) \
-	FRDGAsyncComputeBudgetScopeGuard PREPROCESSOR_JOIN(FRDGAsyncComputeBudgetScope, __LINE__)(GraphBuilder, AsyncComputeBudget)
+	FRDGAsyncComputeBudgetScopeGuard UE_JOIN(FRDGAsyncComputeBudgetScope, __LINE__)(GraphBuilder, AsyncComputeBudget)
 
 #if WITH_MGPU
 	#define RDG_GPU_MASK_SCOPE(GraphBuilder, GPUMask) SCOPED_GPU_MASK(GraphBuilder.RHICmdList, GPUMask)
diff --git a/Engine/Source/Runtime/RenderCore/Public/RenderGraphEvent.h b/Engine/Source/Runtime/RenderCore/Public/RenderGraphEvent.h
--- a/Engine/Source/Runtime/RenderCore/Public/RenderGraphEvent.h
+++ b/Engine/Source/Runtime/RenderCore/Public/RenderGraphEvent.h
@@ -473,9 +473,9 @@ private:
 		} while (false)
 
 	#define RDG_EVENT_SCOPE_IMPL(GraphBuilder, Condition, ScopeFlags, GPUStatArgs, StaticName, FormatString, ...)	\
-		TOptional<TRDGEventScopeGuard<FRDGScope_RHI>> PREPROCESSOR_JOIN(__RDG_ScopeRef_, __LINE__);					\
+		TOptional<TRDGEventScopeGuard<FRDGScope_RHI>> UE_JOIN(__RDG_ScopeRef_, __LINE__);							\
 		RDG_EVENT_SCOPE_CONSTRUCT(																					\
-			  PREPROCESSOR_JOIN(__RDG_ScopeRef_, __LINE__)															\
+			  UE_JOIN(__RDG_ScopeRef_, __LINE__)																	\
 			, GraphBuilder																							\
 			, Condition																								\
 			, ScopeFlags																							\
@@ -538,15 +538,15 @@ private:
 #endif
 
 #if HAS_GPU_STATS && (RHI_NEW_GPU_PROFILER == 0)
-	#define RDG_GPU_STAT_SCOPE(GraphBuilder, StatName)                      TRDGEventScopeGuard<FRDGScope_GPU> PREPROCESSOR_JOIN(__RDG_GPUStatEvent_##StatName,__LINE__) ((GraphBuilder), ERDGScopeFlags::AlwaysEnable, (GraphBuilder).RHICmdList.GetGPUMask(), CSV_STAT_FNAME(StatName), GET_STATID(Stat_GPU_##StatName), nullptr    , DrawcallCountCategory_##StatName);
-	#define RDG_GPU_STAT_SCOPE_VERBOSE(GraphBuilder, StatName, Description) TRDGEventScopeGuard<FRDGScope_GPU> PREPROCESSOR_JOIN(__RDG_GPUStatEvent_##StatName,__LINE__) ((GraphBuilder), ERDGScopeFlags::AlwaysEnable, (GraphBuilder).RHICmdList.GetGPUMask(), CSV_STAT_FNAME(StatName), GET_STATID(Stat_GPU_##StatName), Description, DrawcallCountCategory_##StatName);
+	#define RDG_GPU_STAT_SCOPE(GraphBuilder, StatName)                      TRDGEventScopeGuard<FRDGScope_GPU> UE_JOIN(__RDG_GPUStatEvent_##StatName,__LINE__) ((GraphBuilder), ERDGScopeFlags::AlwaysEnable, (GraphBuilder).RHICmdList.GetGPUMask(), CSV_STAT_FNAME(StatName), GET_STATID(Stat_GPU_##StatName), nullptr    , DrawcallCountCategory_##StatName);
+	#define RDG_GPU_STAT_SCOPE_VERBOSE(GraphBuilder, StatName, Description) TRDGEventScopeGuard<FRDGScope_GPU> UE_JOIN(__RDG_GPUStatEvent_##StatName,__LINE__) ((GraphBuilder), ERDGScopeFlags::AlwaysEnable, (GraphBuilder).RHICmdList.GetGPUMask(), CSV_STAT_FNAME(StatName), GET_STATID(Stat_GPU_##StatName), Description, DrawcallCountCategory_##StatName);
 #else
 	#define RDG_GPU_STAT_SCOPE(GraphBuilder, StatName)
 	#define RDG_GPU_STAT_SCOPE_VERBOSE(GraphBuilder, StatName, Description)
 #endif
 
 #if CSV_PROFILER_STATS
-	#define RDG_CSV_STAT_EXCLUSIVE_SCOPE(GraphBuilder, StatName) TRDGEventScopeGuard<FRDGScope_CSVExclusive> PREPROCESSOR_JOIN(__RDG_CSVStat_##StatName,__LINE__) ((GraphBuilder), ERDGScopeFlags::AlwaysEnable, #StatName);
+	#define RDG_CSV_STAT_EXCLUSIVE_SCOPE(GraphBuilder, StatName) TRDGEventScopeGuard<FRDGScope_CSVExclusive> UE_JOIN(__RDG_CSVStat_##StatName,__LINE__) ((GraphBuilder), ERDGScopeFlags::AlwaysEnable, #StatName);
 #else
 	#define RDG_CSV_STAT_EXCLUSIVE_SCOPE(GraphBuilder, StatName)
 #endif
diff --git a/Engine/Source/Runtime/RenderCore/Public/RenderingThread.h b/Engine/Source/Runtime/RenderCore/Public/RenderingThread.h
--- a/Engine/Source/Runtime/RenderCore/Public/RenderingThread.h
+++ b/Engine/Source/Runtime/RenderCore/Public/RenderingThread.h
@@ -238,12 +238,12 @@ private:
 
 /** Declares a new render command tag type from a name. */
 #define DECLARE_RENDER_COMMAND_TAG(Type, Name) \
-	struct PREPROCESSOR_JOIN(TSTR_, PREPROCESSOR_JOIN(Name, __LINE__)) \
+	struct UE_JOIN(TSTR_, UE_JOIN(Name, __LINE__)) \
 	{  \
 		static const char* CStr() { return #Name; } \
 		static const TCHAR* TStr() { return TEXT(#Name); } \
 	}; \
-	using Type = TRenderCommandTag<PREPROCESSOR_JOIN(TSTR_, PREPROCESSOR_JOIN(Name, __LINE__))>;
+	using Type = TRenderCommandTag<UE_JOIN(TSTR_, UE_JOIN(Name, __LINE__))>;
 
 /** Describes which pipes are configured to use the render command pipe system. */
 enum class ERenderCommandPipeMode : uint8
@@ -1165,8 +1165,8 @@ public:
  *  pipe, if enabled.
  */
 #define ENQUEUE_RENDER_COMMAND(Type) \
-	DECLARE_RENDER_COMMAND_TAG(PREPROCESSOR_JOIN(FRenderCommandTag_, PREPROCESSOR_JOIN(Type, __LINE__)), Type) \
-	FRenderCommandDispatcher::Enqueue<PREPROCESSOR_JOIN(FRenderCommandTag_, PREPROCESSOR_JOIN(Type, __LINE__))>
+	DECLARE_RENDER_COMMAND_TAG(UE_JOIN(FRenderCommandTag_, UE_JOIN(Type, __LINE__)), Type) \
+	FRenderCommandDispatcher::Enqueue<UE_JOIN(FRenderCommandTag_, UE_JOIN(Type, __LINE__))>
 
 ////////////////////////////////////
 // RenderThread scoped work
diff --git a/Engine/Source/Runtime/Renderer/Private/SceneCulling/SceneCulling.cpp b/Engine/Source/Runtime/Renderer/Private/SceneCulling/SceneCulling.cpp
--- a/Engine/Source/Runtime/Renderer/Private/SceneCulling/SceneCulling.cpp
+++ b/Engine/Source/Runtime/Renderer/Private/SceneCulling/SceneCulling.cpp
@@ -240,7 +240,7 @@ struct FIHLoggerListScopeHelper
 	FString LogStr;
 };
 
-#define BUILDER_LOG_SCOPE(Fmt, ...) if (GBuilderForLogging) { GBuilderForLogging->AddLog(FString::Printf(TEXT(Fmt), ##__VA_ARGS__)); } FIHLoggerScopeHelper PREPROCESSOR_JOIN(IHLoggerScopeHelper, __LINE__)
+#define BUILDER_LOG_SCOPE(Fmt, ...) if (GBuilderForLogging) { GBuilderForLogging->AddLog(FString::Printf(TEXT(Fmt), ##__VA_ARGS__)); } FIHLoggerScopeHelper UE_JOIN(IHLoggerScopeHelper, __LINE__)
 #define BUILDER_LOG_LIST(Fmt, ...) FIHLoggerListScopeHelper IHLoggerListScopeHelper(FString::Printf(TEXT(Fmt), ##__VA_ARGS__))
 #define BUILDER_LOG_LIST_APPEND(Fmt, ...) IHLoggerListScopeHelper.Add(FString::Printf(TEXT(Fmt), ##__VA_ARGS__))
 
diff --git a/Engine/Source/Runtime/SignalProcessing/Public/DSP/BufferDiagnostics.h b/Engine/Source/Runtime/SignalProcessing/Public/DSP/BufferDiagnostics.h
--- a/Engine/Source/Runtime/SignalProcessing/Public/DSP/BufferDiagnostics.h
+++ b/Engine/Source/Runtime/SignalProcessing/Public/DSP/BufferDiagnostics.h
@@ -129,24 +129,24 @@ namespace Audio
 
 #if WITH_AUDIO_BUFFERDIAGNOSTICS
 #define AUDIO_CHECK_BUFFER(BUFFER)\
-	static FCheckedBufferState PREPROCESSOR_JOIN(BufferCheck, __LINE__)(__LINE__,TEXT(__FILE__),TEXT(PREPROCESSOR_TO_STRING(BUFFER)));\
-	if (!PREPROCESSOR_JOIN(BufferCheck, __LINE__).DoCheck(BUFFER) )\
+	static FCheckedBufferState UE_JOIN(BufferCheck, __LINE__)(__LINE__,TEXT(__FILE__),TEXT(PREPROCESSOR_TO_STRING(BUFFER)));\
+	if (!UE_JOIN(BufferCheck, __LINE__).DoCheck(BUFFER) )\
 	{\
-		PREPROCESSOR_JOIN(BufferCheck, __LINE__).FailedBufferCheckImpl(TEXT(""));\
+		UE_JOIN(BufferCheck, __LINE__).FailedBufferCheckImpl(TEXT(""));\
 	}
 
 #define AUDIO_CHECK_BUFFER_NAMED(BUFFER,NAME)\
-	static FCheckedBufferState PREPROCESSOR_JOIN(BufferCheck, __LINE__)(__LINE__,TEXT(__FILE__),NAME);\
-	if (!PREPROCESSOR_JOIN(BufferCheck, __LINE__).DoCheck(BUFFER) )\
+	static FCheckedBufferState UE_JOIN(BufferCheck, __LINE__)(__LINE__,TEXT(__FILE__),NAME);\
+	if (!UE_JOIN(BufferCheck, __LINE__).DoCheck(BUFFER) )\
 	{\
-		PREPROCESSOR_JOIN(BufferCheck, __LINE__).FailedBufferCheckImpl(TEXT(""));\
+		UE_JOIN(BufferCheck, __LINE__).FailedBufferCheckImpl(TEXT(""));\
 	}
 
 #define AUDIO_CHECK_BUFFER_NAMED_MSG(BUFFER,NAME,MSG,...)\
-	static FCheckedBufferState PREPROCESSOR_JOIN(BufferCheck, __LINE__)(__LINE__,TEXT(__FILE__),NAME);\
-	if (!PREPROCESSOR_JOIN(BufferCheck, __LINE__).DoCheck(BUFFER))\
+	static FCheckedBufferState UE_JOIN(BufferCheck, __LINE__)(__LINE__,TEXT(__FILE__),NAME);\
+	if (!UE_JOIN(BufferCheck, __LINE__).DoCheck(BUFFER))\
 	{\
-		PREPROCESSOR_JOIN(BufferCheck, __LINE__).FailedBufferCheckImpl(MSG,__VA_ARGS__);\
+		UE_JOIN(BufferCheck, __LINE__).FailedBufferCheckImpl(MSG,__VA_ARGS__);\
 	}
 #else //WITH_AUDIO_BUFFERDIAGNOSTICS
 	#define AUDIO_CHECK_BUFFER_NAMED(BUFFER,NAME)
diff --git a/Engine/Source/Runtime/StudioTelemetry/Public/StudioTelemetry.h b/Engine/Source/Runtime/StudioTelemetry/Public/StudioTelemetry.h
--- a/Engine/Source/Runtime/StudioTelemetry/Public/StudioTelemetry.h
+++ b/Engine/Source/Runtime/StudioTelemetry/Public/StudioTelemetry.h
@@ -158,7 +158,7 @@ private:
 };
 
 // Useful macros for scoped sessions and spans
-#define STUDIO_TELEMETRY_SESSION_SCOPE FStudioTelemetry::FSessionScope PREPROCESSOR_JOIN(FSessionScope, __LINE__);
-#define STUDIO_TELEMETRY_SPAN_SCOPE(Name) FStudioTelemetry::FSpanScope PREPROCESSOR_JOIN(FSpanScope, __LINE__)(TEXT(#Name));
+#define STUDIO_TELEMETRY_SESSION_SCOPE FStudioTelemetry::FSessionScope UE_JOIN(FSessionScope, __LINE__);
+#define STUDIO_TELEMETRY_SPAN_SCOPE(Name) FStudioTelemetry::FSpanScope UE_JOIN(FSpanScope, __LINE__)(TEXT(#Name));
 #define STUDIO_TELEMETRY_START_SPAN(Name) if (FStudioTelemetry::Get().IsAvailable()) { FStudioTelemetry::Get().StartSpan(TEXT(#Name));}
 #define STUDIO_TELEMETRY_END_SPAN(Name) if (FStudioTelemetry::Get().IsAvailable()) { FStudioTelemetry::Get().EndSpan(TEXT(#Name));}
diff --git a/Engine/Source/Runtime/TraceLog/Public/Trace/Detail/Trace.h b/Engine/Source/Runtime/TraceLog/Public/Trace/Detail/Trace.h
--- a/Engine/Source/Runtime/TraceLog/Public/Trace/Detail/Trace.h
+++ b/Engine/Source/Runtime/TraceLog/Public/Trace/Detail/Trace.h
@@ -187,27 +187,27 @@ class FChannel;
 		TRACE_PRIVATE_LOG_EPILOG()
 
 #define TRACE_PRIVATE_LOG_SCOPED(LoggerName, EventName, ChannelsExpr, ...) \
-	UE::Trace::Private::FScopedLogScope PREPROCESSOR_JOIN(TheScope, __LINE__); \
+	UE::Trace::Private::FScopedLogScope UE_JOIN(TheScope, __LINE__); \
 	TRACE_PRIVATE_LOG_PRELUDE(ScopedEnter, LoggerName, EventName, ChannelsExpr, ##__VA_ARGS__) \
-		PREPROCESSOR_JOIN(TheScope, __LINE__).SetActive(), \
+		UE_JOIN(TheScope, __LINE__).SetActive(), \
 		TRACE_PRIVATE_LOG_EPILOG()
 
 #define TRACE_PRIVATE_LOG_SCOPED_CONDITIONAL(LoggerName, EventName, ChannelsExpr, Condition, ...) \
-	UE::Trace::Private::FScopedLogScope PREPROCESSOR_JOIN(TheScope, __LINE__); \
+	UE::Trace::Private::FScopedLogScope UE_JOIN(TheScope, __LINE__); \
 	TRACE_PRIVATE_LOG_PRELUDE_CONDITIONAL(ScopedEnter, LoggerName, EventName, ChannelsExpr, Condition, ##__VA_ARGS__) \
-		PREPROCESSOR_JOIN(TheScope, __LINE__).SetActive(), \
+		UE_JOIN(TheScope, __LINE__).SetActive(), \
 		TRACE_PRIVATE_LOG_EPILOG()
 
 #define TRACE_PRIVATE_LOG_SCOPED_T(LoggerName, EventName, ChannelsExpr, ...) \
-	UE::Trace::Private::FScopedStampedLogScope PREPROCESSOR_JOIN(TheScope, __LINE__); \
+	UE::Trace::Private::FScopedStampedLogScope UE_JOIN(TheScope, __LINE__); \
 	TRACE_PRIVATE_LOG_PRELUDE(ScopedStampedEnter, LoggerName, EventName, ChannelsExpr, ##__VA_ARGS__) \
-		PREPROCESSOR_JOIN(TheScope, __LINE__).SetActive(), \
+		UE_JOIN(TheScope, __LINE__).SetActive(), \
 		TRACE_PRIVATE_LOG_EPILOG()
 
 #define TRACE_PRIVATE_LOG_SCOPED_T_CONDITIONAL(LoggerName, EventName, ChannelsExpr, Condition, ...) \
-	UE::Trace::Private::FScopedStampedLogScope PREPROCESSOR_JOIN(TheScope, __LINE__); \
+	UE::Trace::Private::FScopedStampedLogScope UE_JOIN(TheScope, __LINE__); \
 	TRACE_PRIVATE_LOG_PRELUDE_CONDITIONAL(ScopedStampedEnter, LoggerName, EventName, ChannelsExpr, Condition, ##__VA_ARGS__) \
-		PREPROCESSOR_JOIN(TheScope, __LINE__).SetActive(), \
+		UE_JOIN(TheScope, __LINE__).SetActive(), \
 		TRACE_PRIVATE_LOG_EPILOG()
 
 #define TRACE_PRIVATE_GET_DEFINITION_TYPE_ID(LoggerName, EventName) \
