[SpatialTrace] SpatialTrace API v1
* Introduces a new API for tracing out spatial point data for analysis.
* Point components include Position, Direction and Velocity (all traced as FVector3d).
* Current implementation traces out all components if any changes.
* Some tests to do some basic validation of the API.

diff --git a/Engine/Source/Runtime/Core/Private/ProfilingDebugging/SpatialTrace.cpp b/Engine/Source/Runtime/Core/Private/ProfilingDebugging/SpatialTrace.cpp
new file mode 100644
--- /dev/null
+++ b/Engine/Source/Runtime/Core/Private/ProfilingDebugging/SpatialTrace.cpp
@@ -0,0 +1,182 @@
+// Copyright Epic Games, Inc. All Rights Reserved.
+
+#include "ProfilingDebugging/SpatialTrace.h"
+
+#if UE_TRACE_SPATIAL_ENABLED
+#include "HAL/PlatformTime.h"
+#include "Trace/Trace.inl"
+
+// Configuration for whether to trace only components which have changed.
+#ifndef UE_TRACE_SPATIAL_COMPONENT_DELTAS
+#define UE_TRACE_SPATIAL_COMPONENT_DELTAS 0
+#endif // UE_TRACE_SPATIAL_COMPONENT_DELTAS
+
+UE_TRACE_MINIMAL_CHANNEL_DEFINE(SpatialChannel);
+
+UE_TRACE_MINIMAL_EVENT_BEGIN(Spatial, PointSpec, NoSync | Important)
+	UE_TRACE_MINIMAL_EVENT_FIELD(uint32, Id)
+	UE_TRACE_MINIMAL_EVENT_FIELD(UE::Trace::WideString, Name)
+	UE_TRACE_MINIMAL_EVENT_FIELD(uint8, Mask)
+UE_TRACE_MINIMAL_EVENT_END()
+
+// Helper macro containing component event boilerplate.
+#define UE_TRACE_MINIMAL_SPATIAL_EVENT_BEGIN(LoggerName, EventName, ...) \
+	UE_TRACE_MINIMAL_EVENT_BEGIN(Spatial, EventName, ##__VA_ARGS__) \
+		UE_TRACE_MINIMAL_EVENT_FIELD(uint64, Cycle) \
+		UE_TRACE_MINIMAL_EVENT_FIELD(uint32, SpecId)
+
+// 2^N - 1 subsets (N = number of components, no empty subset).
+
+// 1-component events.
+
+UE_TRACE_MINIMAL_SPATIAL_EVENT_BEGIN(Spatial, SetPointPos)
+	UE_TRACE_MINIMAL_EVENT_FIELD(double[], Pos)
+UE_TRACE_MINIMAL_EVENT_END()
+
+UE_TRACE_MINIMAL_SPATIAL_EVENT_BEGIN(Spatial, SetPointVel)
+	UE_TRACE_MINIMAL_EVENT_FIELD(double[], Vel)
+UE_TRACE_MINIMAL_EVENT_END()
+
+UE_TRACE_MINIMAL_SPATIAL_EVENT_BEGIN(Spatial, SetPointDir)
+	UE_TRACE_MINIMAL_EVENT_FIELD(double[], Dir)
+UE_TRACE_MINIMAL_EVENT_END()
+
+// 2-component events.
+
+UE_TRACE_MINIMAL_SPATIAL_EVENT_BEGIN(Spatial, SetPointPosVel)
+	UE_TRACE_MINIMAL_EVENT_FIELD(double[], Pos)
+	UE_TRACE_MINIMAL_EVENT_FIELD(double[], Vel)
+UE_TRACE_MINIMAL_EVENT_END()
+
+UE_TRACE_MINIMAL_SPATIAL_EVENT_BEGIN(Spatial, SetPointPosDir)
+	UE_TRACE_MINIMAL_EVENT_FIELD(double[], Pos)
+	UE_TRACE_MINIMAL_EVENT_FIELD(double[], Dir)
+UE_TRACE_MINIMAL_EVENT_END()
+
+UE_TRACE_MINIMAL_SPATIAL_EVENT_BEGIN(Spatial, SetPointVelDir)
+	UE_TRACE_MINIMAL_EVENT_FIELD(double[], Vel)
+	UE_TRACE_MINIMAL_EVENT_FIELD(double[], Dir)
+UE_TRACE_MINIMAL_EVENT_END()
+
+// 3-component events.
+
+UE_TRACE_MINIMAL_SPATIAL_EVENT_BEGIN(Spatial, SetPointPosVelDir)
+	UE_TRACE_MINIMAL_EVENT_FIELD(double[], Pos)
+	UE_TRACE_MINIMAL_EVENT_FIELD(double[], Vel)
+	UE_TRACE_MINIMAL_EVENT_FIELD(double[], Dir)
+UE_TRACE_MINIMAL_EVENT_END()
+
+#undef UE_TRACE_MINIMAL_SPATIAL_EVENT_BEGIN
+
+namespace UE::SpatialTrace::Private
+{
+
+	uint32 GetNextPointSpecId()
+	{
+		static std::atomic<uint32> NextPointSpecId = 1;
+		return NextPointSpecId.fetch_add(1);
+	}
+
+	uint32 TracePointSpec(FStringView Name, EPointComponentFlags ComponentMask)
+	{
+		if (!UE_TRACE_CHANNELEXPR_IS_ENABLED(SpatialChannel))
+		{
+			return InvalidPointSpecId;
+		}
+
+		const uint32 SpecId = GetNextPointSpecId();
+
+		UE_TRACE_MINIMAL_LOG(Spatial, PointSpec, SpatialChannel, static_cast<uint32>(Name.NumBytes()))
+			<< PointSpec.Id(SpecId)
+			<< PointSpec.Name(GetData(Name), GetNum(Name))
+			<< PointSpec.Mask((uint8)ComponentMask);
+
+		return SpecId;
+	}
+
+	void TracePointSet(uint32 SpecId, EPointComponentFlags SpecComponentMask, EPointComponentFlags ChangedComponentMask, const FPointDataView& ComponentData)
+	{
+		if (!UE_TRACE_CHANNELEXPR_IS_ENABLED(SpatialChannel))
+		{
+			return;
+		}
+
+#define UE_TRACE_MINIMAL_SPATIAL_LOG(LoggerName, EventName, ChannelsExpr, ...) \
+			UE_TRACE_MINIMAL_LOG(LoggerName, EventName, ChannelsExpr, ##__VA_ARGS__) \
+				<< EventName.Cycle(FPlatformTime::Cycles64()) \
+				<< EventName.SpecId(SpecId)
+
+		const EPointComponentFlags EventComponentMask = UE_TRACE_SPATIAL_COMPONENT_DELTAS ? ChangedComponentMask : SpecComponentMask;
+		switch (EventComponentMask)
+		{
+			case EPointComponentFlags::Position:
+				UE_TRACE_MINIMAL_SPATIAL_LOG(Spatial, SetPointPos, SpatialChannel)
+					<< SetPointPos.Pos(&ComponentData.Position->X, 3);
+				break;
+
+			case EPointComponentFlags::Velocity:
+				UE_TRACE_MINIMAL_SPATIAL_LOG(Spatial, SetPointVel, SpatialChannel)
+					<< SetPointVel.Vel(&ComponentData.Velocity->X, 3);
+				break;
+
+			case EPointComponentFlags::Direction:
+				UE_TRACE_MINIMAL_SPATIAL_LOG(Spatial, SetPointDir, SpatialChannel)
+					<< SetPointDir.Dir(&ComponentData.Direction->X, 3);
+				break;
+
+			case EPointComponentFlags::Position | EPointComponentFlags::Velocity:
+				UE_TRACE_MINIMAL_SPATIAL_LOG(Spatial, SetPointPosVel, SpatialChannel)
+					<< SetPointPosVel.Pos(&ComponentData.Position->X, 3)
+					<< SetPointPosVel.Vel(&ComponentData.Velocity->X, 3);
+				break;
+
+			case EPointComponentFlags::Position | EPointComponentFlags::Direction:
+				UE_TRACE_MINIMAL_SPATIAL_LOG(Spatial, SetPointPosDir, SpatialChannel)
+					<< SetPointPosDir.Pos(&ComponentData.Position->X, 3)
+					<< SetPointPosDir.Dir(&ComponentData.Direction->X, 3);
+				break;
+
+			case EPointComponentFlags::Velocity | EPointComponentFlags::Direction:
+				UE_TRACE_MINIMAL_SPATIAL_LOG(Spatial, SetPointVelDir, SpatialChannel)
+					<< SetPointVelDir.Vel(&ComponentData.Velocity->X, 3)
+					<< SetPointVelDir.Dir(&ComponentData.Direction->X, 3);
+				break;
+
+			case EPointComponentFlags::Position | EPointComponentFlags::Velocity | EPointComponentFlags::Direction:
+				UE_TRACE_MINIMAL_SPATIAL_LOG(Spatial, SetPointPosVelDir, SpatialChannel)
+					<< SetPointPosVelDir.Pos(&ComponentData.Position->X, 3)
+					<< SetPointPosVelDir.Vel(&ComponentData.Velocity->X, 3)
+					<< SetPointPosVelDir.Dir(&ComponentData.Direction->X, 3);
+				break;
+
+			default:
+				checkNoEntry();
+				break;
+		}
+
+#undef UE_TRACE_MINIMAL_SPATIAL_LOG
+	}
+
+	TTuple<uint32, bool> GetOrCreatePointSpecId(std::atomic<uint32>& InOutSpecId, FStringView Name, EPointComponentFlags ComponentMask)
+	{
+		const uint32 CurrentId = InOutSpecId.load(std::memory_order_relaxed);
+		if (CurrentId != InvalidPointSpecId)
+		{
+			return MakeTuple(CurrentId, false);
+		}
+
+		const uint32 NewSpecId = TracePointSpec(Name, ComponentMask);
+		uint32 Expected = InvalidPointSpecId;
+		if (InOutSpecId.compare_exchange_strong(Expected, NewSpecId, std::memory_order_relaxed))
+		{
+			return MakeTuple(NewSpecId, true);
+		}
+		else
+		{
+			return MakeTuple(Expected, false);
+		}
+	}
+
+} // namespace UE::SpatialTrace::Private
+
+#endif // UE_TRACE_SPATIAL_ENABLED
\ No newline at end of file
diff --git a/Engine/Source/Runtime/Core/Public/ProfilingDebugging/SpatialTrace.h b/Engine/Source/Runtime/Core/Public/ProfilingDebugging/SpatialTrace.h
new file mode 100644
--- /dev/null
+++ b/Engine/Source/Runtime/Core/Public/ProfilingDebugging/SpatialTrace.h
@@ -0,0 +1,334 @@
+// Copyright Epic Games, Inc. All Rights Reserved.
+
+#pragma once
+
+#include "Containers/StringView.h"
+#include "Containers/UnrealString.h"
+#include "Math/Vector.h"
+#include "Misc/TVariant.h"
+#include "Trace/Trace.h"
+
+#ifndef UE_TRACE_SPATIAL_ENABLED
+#define UE_TRACE_SPATIAL_ENABLED UE_TRACE_MINIMAL_ENABLED
+#endif // UE_TRACE_SPATIAL_ENABLED
+
+#if UE_TRACE_SPATIAL_ENABLED
+
+UE_TRACE_MINIMAL_CHANNEL_EXTERN(SpatialChannel, CORE_API);
+
+namespace UE::SpatialTrace::Private
+{
+	// Utility class containing either a dynamically allocated string it owns, or a reference to a string literal.
+	class FStringOrView
+	{
+	public:
+		using StringType = FString;
+		using StringViewType = FStringView;
+
+		explicit FStringOrView(const StringViewType& InView)
+			: Storage(TInPlaceType<StringViewType>(), InView)
+		{}
+		explicit FStringOrView(const StringType& InString)
+			: Storage(TInPlaceType<StringType>(), InString)
+		{}
+		explicit FStringOrView(StringType&& InString)
+			: Storage(TInPlaceType<StringType>(), MoveTemp(InString))
+		{}
+
+		static FStringOrView FromView(const StringViewType& InView)
+		{
+			return FStringOrView(InView);
+		}
+
+		static FStringOrView FromString(const StringType& InString)
+		{
+			return FStringOrView(InString);
+		}
+
+		static FStringOrView FromString(StringType&& InString)
+		{
+			return FStringOrView(MoveTemp(InString));
+		}
+
+		StringViewType GetView() const
+		{
+			return Visit([](const auto& Value) -> StringViewType
+			{
+				return StringViewType(Value);
+			}, Storage);
+		}
+
+	private:
+		TVariant<StringViewType, StringType> Storage;
+	};
+
+	enum class EPointComponentFlags : uint8
+	{
+		None = 0,
+		Position = 1 << 0,
+		Direction = 1 << 1,
+		Velocity = 1 << 2,
+		All = 0xFF
+	};
+	ENUM_CLASS_FLAGS(EPointComponentFlags);
+
+	inline static constexpr uint32 InvalidPointSpecId = 0;
+
+	// Data view structure
+	struct FPointDataView
+	{
+		const FVector* Position = nullptr;
+		const FVector* Velocity = nullptr;
+		const FVector* Direction = nullptr;
+	};
+
+	CORE_API uint32 TracePointSpec(FStringView Name, EPointComponentFlags ComponentMask);
+	CORE_API void TracePointSet(uint32 SpecId, EPointComponentFlags SpecComponentMask, EPointComponentFlags ChangedComponentMask, const FPointDataView& ComponentData);
+	CORE_API TTuple<uint32, bool> GetOrCreatePointSpecId(std::atomic<uint32>& InOutSpecId, FStringView Name, EPointComponentFlags ComponentMask);
+
+	template <EPointComponentFlags InComponentMask>
+	class TPointSpec
+	{
+		static constexpr EPointComponentFlags ComponentMask = InComponentMask;
+		static_assert(EnumHasAnyFlags(ComponentMask, EPointComponentFlags::Position), "Point spec must have Position component.");
+
+		struct FData
+		{
+			FVector Position;
+
+			struct FEmpty {};
+			UE_NO_UNIQUE_ADDRESS std::conditional_t<EnumHasAnyFlags(ComponentMask, EPointComponentFlags::Velocity), FVector, FEmpty> Velocity;
+			UE_NO_UNIQUE_ADDRESS std::conditional_t<EnumHasAnyFlags(ComponentMask, EPointComponentFlags::Direction), FVector, FEmpty> Direction;
+		};
+
+		FStringOrView Name;
+		std::atomic<uint32> Id = InvalidPointSpecId;
+		FData Data;
+
+		explicit TPointSpec(FStringOrView&& InName)
+			: Name(MoveTemp(InName))
+		{
+		}
+
+	public:
+
+		TPointSpec(TPointSpec&& Other)
+			: Name(MoveTemp(Other.Name))
+			, Id(Other.Id.exchange(Private::InvalidPointSpecId))
+			, Data(MoveTemp(Other.Data))
+		{
+		}
+
+		TPointSpec& operator=(TPointSpec&& Other)
+		{
+			if (this != &Other)
+			{
+				Name = MoveTemp(Other.Name);
+				Id = Other.Id.exchange(Private::InvalidPointSpecId);
+				Data = MoveTemp(Other.Data);
+			}
+			return *this;
+		}
+
+		template<typename LiteralType>
+		[[nodiscard]] static TPointSpec FromLiteral(const LiteralType& InName)
+		{
+			static_assert(TIsArrayOrRefOfTypeByPredicate<LiteralType, TIsCharEncodingCompatibleWithTCHAR>::Value, "DisplayName must be a TCHAR array.");
+			return TPointSpec(FStringOrView::FromView(InName));
+		}
+
+		[[nodiscard]] static TPointSpec FromView(const FStringView& InName)
+		{
+			return TPointSpec(FStringOrView::FromView(InName));
+		}
+
+		[[nodiscard]] static TPointSpec FromString(const FString& InName)
+		{
+			return TPointSpec(FStringOrView::FromString(InName));
+		}
+
+		[[nodiscard]] static TPointSpec FromString(FString&& InName)
+		{
+			return TPointSpec(FStringOrView::FromString(MoveTemp(InName)));
+		}
+
+		void Set(const FVector& Position)
+			requires(ComponentMask == EPointComponentFlags::Position)
+		{
+			SetInternal({
+				.Position = Position,
+			});
+		}
+
+		void Set(const FVector& Position, const FVector& Velocity)
+			requires(ComponentMask == (EPointComponentFlags::Position | EPointComponentFlags::Velocity))
+		{
+			SetInternal({
+				.Position = Position,
+				.Velocity = Velocity,
+			});
+		}
+
+		void Set(const FVector& Position, const FVector& Direction)
+			requires(ComponentMask == (EPointComponentFlags::Position | EPointComponentFlags::Direction))
+		{
+			SetInternal({
+				.Position = Position,
+				.Direction = Direction,
+			});
+		}
+
+		void Set(const FVector& Position, const FVector& Velocity, const FVector& Direction)
+			requires(ComponentMask == (EPointComponentFlags::Position | EPointComponentFlags::Velocity | EPointComponentFlags::Direction))
+		{
+			SetInternal({
+				.Position = Position,
+				.Velocity = Velocity,
+				.Direction = Direction,
+			});
+		}
+
+	private:
+
+		void SetInternal(const FData& NewData)
+		{
+			EPointComponentFlags ChangedMask = EPointComponentFlags::None;
+			if constexpr (EnumHasAnyFlags(ComponentMask, EPointComponentFlags::Position))
+			{
+				if (Data.Position != NewData.Position)
+				{
+					Data.Position = NewData.Position;
+					ChangedMask |= EPointComponentFlags::Position;
+				}
+			}
+			if constexpr (EnumHasAnyFlags(ComponentMask, EPointComponentFlags::Velocity))
+			{
+				if (Data.Velocity != NewData.Velocity)
+				{
+					Data.Velocity = NewData.Velocity;
+					ChangedMask |= EPointComponentFlags::Velocity;
+				}
+			}
+			if constexpr (EnumHasAnyFlags(ComponentMask, EPointComponentFlags::Direction))
+			{
+				if (Data.Direction != NewData.Direction)
+				{
+					Data.Direction = NewData.Direction;
+					ChangedMask |= EPointComponentFlags::Direction;
+				}
+			}
+
+			if (!UE_TRACE_CHANNELEXPR_IS_ENABLED(SpatialChannel) || ChangedMask == EPointComponentFlags::None)
+			{
+				return;
+			}
+
+			uint32 SpecId;
+			bool bSpecCreated;
+			Tie(SpecId, bSpecCreated) = GetOrCreatePointSpecId(Id, Name.GetView(), ComponentMask);
+			TracePointSet(SpecId, ComponentMask, bSpecCreated ? ComponentMask : ChangedMask, GetDataView());
+		}
+
+		FPointDataView GetDataView() const
+		{
+			FPointDataView DataView;
+			if constexpr (EnumHasAnyFlags(ComponentMask, EPointComponentFlags::Position))
+			{
+				DataView.Position = &Data.Position;
+			}
+			if constexpr (EnumHasAnyFlags(ComponentMask, EPointComponentFlags::Velocity))
+			{
+				DataView.Velocity = &Data.Velocity;
+			}
+			if constexpr (EnumHasAnyFlags(ComponentMask, EPointComponentFlags::Direction))
+			{
+				DataView.Direction = &Data.Direction;
+			}
+			return DataView;
+		}
+	};
+
+} // namespace UE::SpatialTrace::Private
+
+namespace UE::SpatialTrace
+{
+	using FPointPositionSpec                          = Private::TPointSpec<Private::EPointComponentFlags::Position>;
+	using FPointPositionVelocitySpec                  = Private::TPointSpec<Private::EPointComponentFlags::Position | Private::EPointComponentFlags::Velocity>;
+	using FPointPositionDirectionSpec                 = Private::TPointSpec<Private::EPointComponentFlags::Position | Private::EPointComponentFlags::Direction>;
+	using FPointPositionVelocityDirectionSpec         = Private::TPointSpec<Private::EPointComponentFlags::Position | Private::EPointComponentFlags::Velocity  | Private::EPointComponentFlags::Direction>;
+} // namespace UE::SpatialTrace
+
+//////////////////////////////////////////////////
+// Internal helpers
+
+#define UE_PRIVATE_TRACE_SPATIAL_POINT_SPEC_VAR(Suffix) \
+	PREPROCESSOR_JOIN(__SpatialTracePointSpec, Suffix)
+
+#define UE_PRIVATE_TRACE_SPATIAL_POINT_INLINE_SPEC_DECLARE(SpecType, DisplayName) \
+	static SpecType UE_PRIVATE_TRACE_SPATIAL_POINT_SPEC_VAR(__LINE__) = SpecType::FromLiteral(DisplayName);
+
+#define UE_PRIVATE_TRACE_SPATIAL_POINT_INLINE_SPEC(SpecType, DisplayName, ...)   \
+	UE_PRIVATE_TRACE_SPATIAL_POINT_INLINE_SPEC_DECLARE(SpecType, DisplayName)    \
+	UE_PRIVATE_TRACE_SPATIAL_POINT_SPEC_VAR(__LINE__).Set(__VA_ARGS__);
+
+#define UE_PRIVATE_TRACE_SPATIAL_POINT_SPEC_DECLARE(SpecType, Name, DisplayName) \
+	SpecType UE_PRIVATE_TRACE_SPATIAL_POINT_SPEC_VAR(Name) = SpecType::FromLiteral(DisplayName)
+
+#define UE_PRIVATE_TRACE_SPATIAL_POINT_SPEC_EXTERN(SpecType, Name) \
+	extern SpecType UE_PRIVATE_TRACE_SPATIAL_POINT_SPEC_VAR(Name)
+
+//////////////////////////////////////////////////
+// Runtime log macros
+
+#define UE_TRACE_SPATIAL_POINT_POS_LOG(Spec, Position)                                            UE_PRIVATE_TRACE_SPATIAL_POINT_SPEC_VAR(Spec).Set(Position);
+#define UE_TRACE_SPATIAL_POINT_POS_VEL_LOG(Spec, Position, Velocity)                              UE_PRIVATE_TRACE_SPATIAL_POINT_SPEC_VAR(Spec).Set(Position, Velocity);
+#define UE_TRACE_SPATIAL_POINT_POS_DIR_LOG(Spec, Position, Direction)                             UE_PRIVATE_TRACE_SPATIAL_POINT_SPEC_VAR(Spec).Set(Position, Direction);
+#define UE_TRACE_SPATIAL_POINT_POS_VEL_DIR_LOG(Spec, Position, Velocity, Direction)               UE_PRIVATE_TRACE_SPATIAL_POINT_SPEC_VAR(Spec).Set(Position, Velocity, Direction);
+
+//////////////////////////////////////////////////
+// Inline declare & log macros
+
+#define UE_TRACE_SPATIAL_POINT_POS_INLINE_LOG(DisplayName, Position)                              UE_PRIVATE_TRACE_SPATIAL_POINT_INLINE_SPEC(::UE::SpatialTrace::FPointPositionSpec,                   DisplayName, Position);
+#define UE_TRACE_SPATIAL_POINT_POS_VEL_INLINE_LOG(DisplayName, Position, Velocity)                UE_PRIVATE_TRACE_SPATIAL_POINT_INLINE_SPEC(::UE::SpatialTrace::FPointPositionVelocitySpec,           DisplayName, Position, Velocity);
+#define UE_TRACE_SPATIAL_POINT_POS_DIR_INLINE_LOG(DisplayName, Position, Direction)               UE_PRIVATE_TRACE_SPATIAL_POINT_INLINE_SPEC(::UE::SpatialTrace::FPointPositionDirectionSpec,          DisplayName, Position, Direction);
+#define UE_TRACE_SPATIAL_POINT_POS_VEL_DIR_INLINE_LOG(DisplayName, Position, Velocity, Direction) UE_PRIVATE_TRACE_SPATIAL_POINT_INLINE_SPEC(::UE::SpatialTrace::FPointPositionVelocityDirectionSpec,  DisplayName, Position, Velocity, Direction);
+
+//////////////////////////////////////////////////
+// Spec‐object declarations
+
+#define UE_TRACE_SPATIAL_POINT_POS_SPEC_DECLARE(Name, DisplayName)         UE_PRIVATE_TRACE_SPATIAL_POINT_SPEC_DECLARE(::UE::SpatialTrace::FPointPositionSpec,                  Name, DisplayName);
+#define UE_TRACE_SPATIAL_POINT_POS_SPEC_EXTERN(Name)                       UE_PRIVATE_TRACE_SPATIAL_POINT_SPEC_EXTERN(::UE::SpatialTrace::FPointPositionSpec,                   Name);
+
+#define UE_TRACE_SPATIAL_POINT_POS_VEL_SPEC_DECLARE(Name, DisplayName)     UE_PRIVATE_TRACE_SPATIAL_POINT_SPEC_DECLARE(::UE::SpatialTrace::FPointPositionVelocitySpec,          Name, DisplayName);
+#define UE_TRACE_SPATIAL_POINT_POS_VEL_SPEC_EXTERN(Name)                   UE_PRIVATE_TRACE_SPATIAL_POINT_SPEC_EXTERN(::UE::SpatialTrace::FPointPositionVelocitySpec,           Name);
+
+#define UE_TRACE_SPATIAL_POINT_POS_DIR_SPEC_DECLARE(Name, DisplayName)     UE_PRIVATE_TRACE_SPATIAL_POINT_SPEC_DECLARE(::UE::SpatialTrace::FPointPositionDirectionSpec,         Name, DisplayName);
+#define UE_TRACE_SPATIAL_POINT_POS_DIR_SPEC_EXTERN(Name)                   UE_PRIVATE_TRACE_SPATIAL_POINT_SPEC_EXTERN(::UE::SpatialTrace::FPointPositionDirectionSpec,          Name);
+
+#define UE_TRACE_SPATIAL_POINT_POS_VEL_DIR_SPEC_DECLARE(Name, DisplayName) UE_PRIVATE_TRACE_SPATIAL_POINT_SPEC_DECLARE(::UE::SpatialTrace::FPointPositionVelocityDirectionSpec, Name, DisplayName);
+#define UE_TRACE_SPATIAL_POINT_POS_VEL_DIR_SPEC_EXTERN(Name)               UE_PRIVATE_TRACE_SPATIAL_POINT_SPEC_EXTERN(::UE::SpatialTrace::FPointPositionVelocityDirectionSpec,  Name);
+
+#else // UE_TRACE_SPATIAL_ENABLED == 0
+
+//–– No-op stubs ––
+
+#define UE_TRACE_SPATIAL_POINT_POS_LOG(...)
+#define UE_TRACE_SPATIAL_POINT_POS_VEL_LOG(...)
+#define UE_TRACE_SPATIAL_POINT_POS_DIR_LOG(...)
+#define UE_TRACE_SPATIAL_POINT_POS_VEL_DIR_LOG(...)
+
+#define UE_TRACE_SPATIAL_POINT_POS_INLINE_LOG(...)
+#define UE_TRACE_SPATIAL_POINT_POS_VEL_INLINE_LOG(...)
+#define UE_TRACE_SPATIAL_POINT_POS_DIR_INLINE_LOG(...)
+#define UE_TRACE_SPATIAL_POINT_POS_VEL_DIR_INLINE_LOG(...)
+
+#define UE_TRACE_SPATIAL_POINT_POS_SPEC_DECLARE(...)
+#define UE_TRACE_SPATIAL_POINT_POS_SPEC_EXTERN(...)
+#define UE_TRACE_SPATIAL_POINT_POS_VEL_SPEC_DECLARE(...)
+#define UE_TRACE_SPATIAL_POINT_POS_VEL_SPEC_EXTERN(...)
+#define UE_TRACE_SPATIAL_POINT_POS_DIR_SPEC_DECLARE(...)
+#define UE_TRACE_SPATIAL_POINT_POS_DIR_SPEC_EXTERN(...)
+#define UE_TRACE_SPATIAL_POINT_POS_VEL_DIR_SPEC_DECLARE(...)
+#define UE_TRACE_SPATIAL_POINT_POS_VEL_DIR_SPEC_EXTERN(...)
+
+#endif // UE_TRACE_SPATIAL_ENABLED
diff --git a/Engine/Source/Runtime/Core/Tests/ProfilingDebugging/SpatialTraceTest.cpp b/Engine/Source/Runtime/Core/Tests/ProfilingDebugging/SpatialTraceTest.cpp
new file mode 100644
--- /dev/null
+++ b/Engine/Source/Runtime/Core/Tests/ProfilingDebugging/SpatialTraceTest.cpp
@@ -0,0 +1,98 @@
+// Copyright Epic Games, Inc. All Rights Reserved.
+
+#if WITH_TESTS
+#include "ProfilingDebugging/SpatialTrace.h"
+#include "Tests/TestHarnessAdapter.h"
+
+TEST_CASE_NAMED(FSpatialTraceMacroTest, "System::Core::SpatialTrace::Macros", "[ApplicationContextMask][SmokeFilter]")
+{
+	SECTION("Position logging macro must compile")
+	{
+		UE_TRACE_SPATIAL_POINT_POS_SPEC_DECLARE(TestPositionSpec, TEXT("MacroTestPosition"));
+		UE_TRACE_SPATIAL_POINT_POS_LOG(TestPositionSpec, FVector(1.0, 2.0, 3.0));
+	}
+
+	SECTION("Position/direction logging macro must compile")
+	{
+		UE_TRACE_SPATIAL_POINT_POS_DIR_SPEC_DECLARE(TestPosDirectionSpec, TEXT("MacroTestPosDirection"));
+		UE_TRACE_SPATIAL_POINT_POS_DIR_LOG(TestPosDirectionSpec, FVector(1.0, 2.0, 3.0), FVector::ForwardVector);
+	}
+
+	SECTION("Position/velocity logging macro must compile")
+	{
+		UE_TRACE_SPATIAL_POINT_POS_VEL_SPEC_DECLARE(TestPosVelocitySpec, TEXT("MacroTestPosVelocity"));
+		UE_TRACE_SPATIAL_POINT_POS_VEL_LOG(TestPosVelocitySpec, FVector(0.0, 0.0, 0.0), FVector(10.0, 0.0, 0.0));
+	}
+
+	SECTION("Position/velocity/direction logging macro must compile")
+	{
+		UE_TRACE_SPATIAL_POINT_POS_VEL_DIR_SPEC_DECLARE(TestPosVelDirSpec, TEXT("MacroTestPosVelDir"));
+		UE_TRACE_SPATIAL_POINT_POS_VEL_DIR_LOG(TestPosVelDirSpec, FVector(5.0, 10.0, 15.0), FVector(1.0, 0.0, 0.0), FVector::UpVector);
+	}
+
+	SECTION("Inline logging macros must compile")
+	{
+		UE_TRACE_SPATIAL_POINT_POS_INLINE_LOG(TEXT("InlinePosition"), FVector(100.0, 200.0, 300.0));
+		UE_TRACE_SPATIAL_POINT_POS_DIR_INLINE_LOG(TEXT("InlinePosDirection"), FVector(1.0, 1.0, 1.0), FVector::UpVector);
+		UE_TRACE_SPATIAL_POINT_POS_VEL_INLINE_LOG(TEXT("InlinePosVelocity"), FVector::ZeroVector, FVector(0.0, 0.0, 9.8));
+		UE_TRACE_SPATIAL_POINT_POS_VEL_DIR_INLINE_LOG(TEXT("InlinePosVelDir"), FVector(50.0, 60.0, 70.0), FVector(0.0, 5.0, 0.0), FVector::ForwardVector);
+	}
+}
+
+#if UE_TRACE_SPATIAL_ENABLED
+
+TEST_CASE_NAMED(FSpatialTraceStringOrViewTest, "System::Core::SpatialTrace::StringOrView", "[ApplicationContextMask][SmokeFilter]")
+{
+	using FStringOrView = UE::SpatialTrace::Private::FStringOrView;
+
+	SECTION("View construction and access")
+	{
+		const TCHAR* Literal = TEXT("TestLiteral");
+		FStringOrView ViewWrapper = FStringOrView::FromView(Literal);
+
+		TStringView<TCHAR> Retrieved = ViewWrapper.GetView();
+		CHECK_MESSAGE(TEXT("StringOrView must correctly store and retrieve string views"), Retrieved.GetData() == Literal);
+	}
+
+	SECTION("String construction and access")
+	{
+		FString OriginalString = TEXT("TestString");
+		FStringOrView StringWrapper = FStringOrView::FromString(OriginalString);
+
+		TStringView<TCHAR> Retrieved = StringWrapper.GetView();
+		CHECK_MESSAGE(TEXT("StringOrView must correctly store and retrieve owned strings"), Retrieved == TStringView<TCHAR>(OriginalString));
+	}
+
+	SECTION("Moved string construction")
+	{
+		FString MovedString = TEXT("MovedString");
+		const TCHAR* OrigDataPtr = *MovedString;
+		FStringOrView MovedWrapper = FStringOrView::FromString(MoveTemp(MovedString));
+
+		TStringView<TCHAR> Retrieved = MovedWrapper.GetView();
+		CHECK_MESSAGE(TEXT("StringOrView must correctly handle moved strings"), Retrieved.GetData() == OrigDataPtr);
+	}
+}
+
+/*
+ * The following test contains commented out code which would trigger compilation errors if they were uncommented.
+ */
+TEST_CASE_NAMED(FSpatialTraceCompileErrorTest, "System::Core::SpatialTrace::CompileErrors", "[ApplicationContextMask][SmokeFilter]")
+{
+	// Error: No position component.
+	/*using FInvalidSpec1 = UE::SpatialTrace::Private::TPointSpec<UE::SpatialTrace::Private::EPointComponentFlags::Velocity>;
+	FInvalidSpec1 InvalidSpec1 = FInvalidSpec1::FromLiteral(TEXT("InvalidSpec1"));*/
+
+	// Error: Invalid display name type (not const TCHAR array)
+	/*FString NonConstName = TEXT("Invalid");
+	UE_TRACE_SPATIAL_POINT_POS_SPEC_DECLARE(TestSpec, NonConstName);*/
+
+	// Error: Wrong parameter count
+	/*UE_TRACE_SPATIAL_POINT_POS_DIR_SPEC_DECLARE(Spec, TEXT("Test"));
+	UE_TRACE_SPATIAL_POINT_POS_LOG(Spec, FVector3d::ZeroVector);
+	UE_TRACE_SPATIAL_POINT_POS_VEL_DIR_LOG(Spec, FVector3d::ZeroVector, FVector3d::ForwardVector, FVector3d::UpVector);*/
+}
+
+#endif // UE_TRACE_SPATIAL_ENABLED
+
+#endif // WITH_TESTS
