Replaced the member function name trimming of *Dynamic delegate functions with a C++20 solution that works on all compilers and should be less work for the compiler.

diff --git a/Engine/Source/Runtime/Core/Public/Delegates/Delegate.h b/Engine/Source/Runtime/Core/Public/Delegates/Delegate.h
--- a/Engine/Source/Runtime/Core/Public/Delegates/Delegate.h
+++ b/Engine/Source/Runtime/Core/Public/Delegates/Delegate.h
@@ -315,105 +315,65 @@ class DynamicMulticastDelegateClassName : public TBaseDynamicMulticastDelegate<F
 		} \
 	};
 
-
-#define ENABLE_STATIC_FUNCTION_FNAMES (PLATFORM_COMPILER_CLANG && (__clang_major__ > 3 || (__clang_major__ == 3 && __clang_minor__ >= 5)))
-
-#if ENABLE_STATIC_FUNCTION_FNAMES
-
-	namespace NStrAfterLastDoubleColon_Private
+namespace UE::Core::Private
+{
+	template <SIZE_T Len>
+	struct TCharArray
 	{
-		template <typename T>
-		struct TTypedImpl
-		{
-			// Dummy bool parameter needed so that we can ensure the last specialization is partial
-			template <bool Dummy, T... Chars>
-			struct Inner;
-
-			template <typename SecondMatch, T... Chars>
-			struct Inner2
-			{
-				typedef typename TTypedImpl<T>::template Inner<true, Chars...>::Type Type;
-			};
-
-			template <T... Chars>
-			struct Inner2<TIntegerSequence<T>, Chars...>
-			{
-				typedef TIntegerSequence<T, Chars...> Type;
-			};
+		TCharArray() = default;
 
-			template <bool Dummy, T Char, T... Chars>
-			struct Inner<Dummy, Char, Chars...>
+		constexpr TCharArray(const char(&InStr)[Len])
+		{
+			for (SIZE_T Index = 0; Index != Len; ++Index)
 			{
-				typedef typename Inner<Dummy, Chars...>::Type Type;
-			};
+				this->Str[Index] = InStr[Index];
+			}
+		}
 
-			template <bool Dummy, T... Chars>
-			struct Inner<Dummy, ':', ':', Chars...>
-			{
-				typedef typename Inner2<typename Inner<Dummy, Chars...>::Type, Chars...>::Type Type;
-			};
+		char Str[Len];
+	};
 
-			// Without the dummy bool parameter, this would not compile, as full specializations
-			// of nested class templates is not allowed.
-			template <bool Dummy>
-			struct Inner<Dummy>
+	// Given a char array as a template argument, returns a new char array that is the substring after the ::
+	template <TCharArray CharArray>
+	constexpr auto GetTrimmedMemberFunctionName()
+	{
+		// Returns the substring after the ::
+		auto FindFuncName = [](const char* Str)
+		{
+			while (Str[0] != ':' || Str[1] != ':')
 			{
-				typedef TIntegerSequence<T> Type;
-			};
+				++Str;
+			}
+			return Str + 2;
 		};
 
-
-		template <typename IntSeq>
-		struct TStaticFNameFromCharSequence;
-
-		template <typename T, T... Chars>
-		struct TStaticFNameFromCharSequence<TIntegerSequence<T, Chars...>>
+		// Inline version of StrLen that is also constexpr
+		auto StrLen = [](const char* Str)
 		{
-			static FName Get()
+			SIZE_T Result = 0;
+			while (*Str != 0)
 			{
-				static const T Str[sizeof...(Chars) + 1] = { Chars..., 0 };
-				FName* Result = nullptr;
-				UE_AUTORTFM_OPEN
-				{
-					static FName StaticResult{Str};
-					Result = &StaticResult;
-				};
-				return *Result;
+				++Str;
+				++Result;
 			}
+			return Result;
 		};
 
-		/**
-		 * Metafunction which evaluates to a TIntegerSequence of chars containing only the function name.
-		 *
-		 * Example:
-		 * TStrAfterLastDoubleColon<TEXT("&SomeClass::SomeNestedClass::SomeFunc")_intseq>::Type == TIntegerSequence<TCHAR, 'S', 'o', 'm', 'e', 'F', 'u', 'n', 'c'>
-		 */
-		template <typename T>
-		struct TStrAfterLastDoubleColon;
-
-		template <typename T, T... Chars>
-		struct TStrAfterLastDoubleColon<TIntegerSequence<T, Chars...>>
+		constexpr const char* FuncName = FindFuncName(CharArray.Str);
+		constexpr SIZE_T FuncNameArrayLen = StrLen(FuncName) + 1;
+		TCharArray<FuncNameArrayLen> Result;
+		for (SIZE_T Index = 0; Index != FuncNameArrayLen; ++Index)
 		{
-			typedef typename NStrAfterLastDoubleColon_Private::TTypedImpl<T>::template Inner<true, Chars...>::Type Type;
-		};
-	}
-
-	/**
-	 * Custom literal operator which converts a string into a TIntegerSequence of chars.
-	 */
-	template <typename T, T... Chars>
-	UE_FORCEINLINE_HINT constexpr TIntegerSequence<T, Chars...> operator""_intseq()
-	{
-		return {};
+			Result.Str[Index] = FuncName[Index];
+		}
+		return Result;
 	}
 
-	#define STATIC_FUNCTION_FNAME(str) NStrAfterLastDoubleColon_Private::TStaticFNameFromCharSequence<typename NStrAfterLastDoubleColon_Private::TStrAfterLastDoubleColon<decltype(UE_JOIN(str, _intseq))>::Type>::Get()
-
-#else
-
-	#define STATIC_FUNCTION_FNAME(str) UE::Delegates::Private::GetTrimmedMemberFunctionName(str)
+	template <TCharArray CharArray>
+	constexpr inline auto TrimmedMemberFunctionName = GetTrimmedMemberFunctionName<CharArray>();
+}
 
-#endif
+#define STATIC_FUNCTION_FNAME(MemberFunctionName) FName(UE::Core::Private::TrimmedMemberFunctionName<MemberFunctionName>.Str)
 
 
 /**
@@ -457,26 +417,6 @@ class DynamicMulticastDelegateClassName : public TBaseDynamicMulticastDelegate<F
 #define IsAlreadyBound( UserObject, FuncName ) __Internal_IsAlreadyBound( UserObject, FuncName, STATIC_FUNCTION_FNAME( #FuncName ) )
 
 
-namespace UE::Delegates::Private
-{
-	/**
-	 * Returns the root function name from a string representing a member function pointer.
-	 * Note: this function only returns a pointer to the substring and doesn't create a new string.
-	 *
-	 * @param  InMacroFunctionName  The string containing the member function name.
-	 * @return An FName of the member function name.
-	 */
-	inline FName GetTrimmedMemberFunctionName(const ANSICHAR* InMacroFunctionName)
-	{
-		// We strip off the class prefix and just return the function name by itself.
-		check(InMacroFunctionName);
-		const ANSICHAR* Result = FCStringAnsi::Strrstr( InMacroFunctionName,  "::" );
-		checkf(Result && Result[2] != '0', TEXT("'%hs' does not look like a member function"), InMacroFunctionName);
-		return FName(Result + 2);
-	}
-}
-
-
 /*********************************************************************************************************************/
 
 // We define this as a guard to prevent DelegateSignatureImpl.inl being included outside of this file
