Address PVS warning as it doesn't seem to understand the Tie function. Remove some redundant qualifiers.

diff --git a/Engine/Source/Runtime/Core/Public/ProfilingDebugging/SpatialTrace.h b/Engine/Source/Runtime/Core/Public/ProfilingDebugging/SpatialTrace.h
--- a/Engine/Source/Runtime/Core/Public/ProfilingDebugging/SpatialTrace.h
+++ b/Engine/Source/Runtime/Core/Public/ProfilingDebugging/SpatialTrace.h
@@ -114,7 +114,7 @@ namespace UE::SpatialTrace::Private
 
 		TPointSpec(TPointSpec&& Other)
 			: Name(MoveTemp(Other.Name))
-			, Id(Other.Id.exchange(Private::InvalidPointSpecId))
+			, Id(Other.Id.exchange(InvalidPointSpecId))
 			, Data(MoveTemp(Other.Data))
 		{
 		}
@@ -124,7 +124,7 @@ namespace UE::SpatialTrace::Private
 			if (this != &Other)
 			{
 				Name = MoveTemp(Other.Name);
-				Id = Other.Id.exchange(Private::InvalidPointSpecId);
+				Id = Other.Id.exchange(InvalidPointSpecId);
 				Data = MoveTemp(Other.Data);
 			}
 			return *this;
@@ -223,8 +223,8 @@ namespace UE::SpatialTrace::Private
 				return;
 			}
 
-			uint32 SpecId;
-			bool bSpecCreated;
+			uint32 SpecId = InvalidPointSpecId;
+			bool bSpecCreated = false;
 			Tie(SpecId, bSpecCreated) = GetOrCreatePointSpecId(Id, Name.GetView(), ComponentMask);
 			TracePointSet(SpecId, ComponentMask, bSpecCreated ? ComponentMask : ChangedMask, GetDataView());
 		}
