[RaceDetector]
- Add support to set a specific name for the log file via the command-line

#FailedValidations Preflight

diff --git a/Engine/Source/Runtime/Core/Private/Sanitizer/RaceDetector.cpp b/Engine/Source/Runtime/Core/Private/Sanitizer/RaceDetector.cpp
--- a/Engine/Source/Runtime/Core/Private/Sanitizer/RaceDetector.cpp
+++ b/Engine/Source/Runtime/Core/Private/Sanitizer/RaceDetector.cpp
@@ -1465,7 +1465,14 @@ namespace UE::Sanitizer::RaceDetector {
 		HookAlwaysOnInstrumentation();
 		
 		const TCHAR* CommandLine = Platform::GetCommandLine();
-		if (FCString::Stristr(CommandLine, TEXT("-racedetector")))
+
+		FString RaceDetectorLogFile;
+		if (FParse::Value(CommandLine, TEXT("-racedetectorlog="), RaceDetectorLogFile))
+		{
+			RaceDetectorLog = TUniquePtr<FArchive>(IFileManager::Get().CreateFileWriter(*RaceDetectorLogFile, FILEWRITE_AllowRead));
+		}
+
+		if (FParse::Param(CommandLine, TEXT("-racedetector")))
 		{
 			ToggleRaceDetection(true);
 		}
