Replaced PREPROCESSOR_TO_STRING with UE_STRINGIZE.

diff --git a/Engine/Build/Windows/Resources/Default.rc2 b/Engine/Build/Windows/Resources/Default.rc2
--- a/Engine/Build/Windows/Resources/Default.rc2
+++ b/Engine/Build/Windows/Resources/Default.rc2
@@ -12,31 +12,31 @@
 
 // Various strings used for project resources
 #ifdef PROJECT_COMPANY_NAME
-	#define BUILD_PROJECT_COMPANY_NAME PREPROCESSOR_TO_STRING(PROJECT_COMPANY_NAME)
+	#define BUILD_PROJECT_COMPANY_NAME UE_STRINGIZE(PROJECT_COMPANY_NAME)
 #else
 	#define BUILD_PROJECT_COMPANY_NAME EPIC_COMPANY_NAME
 #endif
 
 #ifdef PROJECT_COPYRIGHT_STRING
-	#define BUILD_PROJECT_COPYRIGHT_STRING PREPROCESSOR_TO_STRING(PROJECT_COPYRIGHT_STRING)
+	#define BUILD_PROJECT_COPYRIGHT_STRING UE_STRINGIZE(PROJECT_COPYRIGHT_STRING)
 #else
 	#define BUILD_PROJECT_COPYRIGHT_STRING EPIC_COPYRIGHT_STRING
 #endif
 
 #ifdef PROJECT_PRODUCT_NAME
-	#define BUILD_PROJECT_PRODUCT_NAME PREPROCESSOR_TO_STRING(PROJECT_PRODUCT_NAME)
+	#define BUILD_PROJECT_PRODUCT_NAME UE_STRINGIZE(PROJECT_PRODUCT_NAME)
 #else
-	#define BUILD_PROJECT_PRODUCT_NAME PREPROCESSOR_TO_STRING(UE_APP_NAME)
+	#define BUILD_PROJECT_PRODUCT_NAME UE_STRINGIZE(UE_APP_NAME)
 #endif
 
 #ifdef BUILD_VERSION
-	#define BUILD_PROJECT_PRODUCT_VERSION PREPROCESSOR_TO_STRING(BUILD_VERSION)
+	#define BUILD_PROJECT_PRODUCT_VERSION UE_STRINGIZE(BUILD_VERSION)
 #else
 	#define BUILD_PROJECT_PRODUCT_VERSION ENGINE_VERSION_STRING
 #endif
 
 #ifdef PROJECT_PRODUCT_IDENTIFIER
-	#define BUILD_PROJECT_PRODUCT_IDENTIFIER PREPROCESSOR_TO_STRING(PROJECT_PRODUCT_IDENTIFIER)
+	#define BUILD_PROJECT_PRODUCT_IDENTIFIER UE_STRINGIZE(PROJECT_PRODUCT_IDENTIFIER)
 #else
 	#define BUILD_PROJECT_PRODUCT_IDENTIFIER EPIC_PRODUCT_IDENTIFIER
 #endif
@@ -78,7 +78,7 @@ BEGIN
 			VALUE "FileDescription", BUILD_PROJECT_PRODUCT_NAME
 			VALUE "InternalName", BUILD_PROJECT_PRODUCT_IDENTIFIER
 #ifdef ORIGINAL_FILE_NAME
-			VALUE "OriginalFilename", PREPROCESSOR_TO_STRING(ORIGINAL_FILE_NAME)
+			VALUE "OriginalFilename", UE_STRINGIZE(ORIGINAL_FILE_NAME)
 #endif
 		END
 	END
diff --git a/Engine/Plugins/Editor/SpeedTreeImporter/Source/SpeedTreeImporter/Private/SpeedTreeImportFactory.cpp b/Engine/Plugins/Editor/SpeedTreeImporter/Source/SpeedTreeImporter/Private/SpeedTreeImportFactory.cpp
--- a/Engine/Plugins/Editor/SpeedTreeImporter/Source/SpeedTreeImporter/Private/SpeedTreeImportFactory.cpp
+++ b/Engine/Plugins/Editor/SpeedTreeImporter/Source/SpeedTreeImporter/Private/SpeedTreeImportFactory.cpp
@@ -1998,7 +1998,7 @@ UObject* USpeedTreeImportFactory::FactoryCreateBinary7(UClass* InClass, UObject*
 	if ( !bOutOperationCanceled )
 	{
 #ifdef SPEEDTREE_KEY
-		SpeedTree::CCore::Authorize(PREPROCESSOR_TO_STRING(SPEEDTREE_KEY));
+		SpeedTree::CCore::Authorize(UE_STRINGIZE(SPEEDTREE_KEY));
 #endif
 		
 		SpeedTree::CCore SpeedTree;
diff --git a/Engine/Plugins/Enterprise/DatasmithCADImporter/Source/WireInterface/Public/WireInterfaceModule.h b/Engine/Plugins/Enterprise/DatasmithCADImporter/Source/WireInterface/Public/WireInterfaceModule.h
--- a/Engine/Plugins/Enterprise/DatasmithCADImporter/Source/WireInterface/Public/WireInterfaceModule.h
+++ b/Engine/Plugins/Enterprise/DatasmithCADImporter/Source/WireInterface/Public/WireInterfaceModule.h
@@ -18,12 +18,12 @@ public:
 
     static FDatasmithWireTranslatorModule& Get()
     {
-        return FModuleManager::LoadModuleChecked< FDatasmithWireTranslatorModule >(PREPROCESSOR_TO_STRING(UE_DATASMITHWIRETRANSLATOR_MODULE_NAME));
+        return FModuleManager::LoadModuleChecked< FDatasmithWireTranslatorModule >(UE_STRINGIZE(UE_DATASMITHWIRETRANSLATOR_MODULE_NAME));
 	}
 
     static bool IsAvailable()
     {
-        return FModuleManager::Get().IsModuleLoaded(PREPROCESSOR_TO_STRING(UE_DATASMITHWIRETRANSLATOR_MODULE_NAME));
+        return FModuleManager::Get().IsModuleLoaded(UE_STRINGIZE(UE_DATASMITHWIRETRANSLATOR_MODULE_NAME));
     }
 
 	virtual void StartupModule() override;
diff --git a/Engine/Plugins/Experimental/NNERuntimeIREE/Source/NNERuntimeIREE/Private/NNERuntimeIREEModule.cpp b/Engine/Plugins/Experimental/NNERuntimeIREE/Source/NNERuntimeIREE/Private/NNERuntimeIREEModule.cpp
--- a/Engine/Plugins/Experimental/NNERuntimeIREE/Source/NNERuntimeIREE/Private/NNERuntimeIREEModule.cpp
+++ b/Engine/Plugins/Experimental/NNERuntimeIREE/Source/NNERuntimeIREE/Private/NNERuntimeIREEModule.cpp
@@ -80,7 +80,7 @@ void FNNERuntimeIREEModule::StartupModule()
 
 #if WITH_EDITOR
 	const FString PluginDir = IPluginManager::Get().FindPlugin("NNERuntimeIREE")->GetBaseDir();
-	const FString MlirToolsSharedLibPath = FPaths::Combine(PluginDir, TEXT(PREPROCESSOR_TO_STRING(NNEMLIRTOOLS_SHAREDLIB_PATH)));
+	const FString MlirToolsSharedLibPath = FPaths::Combine(PluginDir, TEXT(UE_STRINGIZE(NNEMLIRTOOLS_SHAREDLIB_PATH)));
 
 	DllHandle = DllHelper::GetDllHandle(MlirToolsSharedLibPath);
 	if (!DllHandle)
diff --git a/Engine/Plugins/Experimental/NNERuntimeRDG/Source/NNERuntimeRDGUtils/Private/NNERuntimeRDGUtilsModule.cpp b/Engine/Plugins/Experimental/NNERuntimeRDG/Source/NNERuntimeRDGUtils/Private/NNERuntimeRDGUtilsModule.cpp
--- a/Engine/Plugins/Experimental/NNERuntimeRDG/Source/NNERuntimeRDGUtils/Private/NNERuntimeRDGUtilsModule.cpp
+++ b/Engine/Plugins/Experimental/NNERuntimeRDG/Source/NNERuntimeRDGUtils/Private/NNERuntimeRDGUtilsModule.cpp
@@ -13,7 +13,7 @@ public:
 	virtual void StartupModule() override
 	{
 		const FString PluginDir = IPluginManager::Get().FindPlugin("NNERuntimeRDG")->GetBaseDir();
-		const FString OrtSharedLibPath = FPaths::Combine(PluginDir, TEXT(PREPROCESSOR_TO_STRING(ONNXRUNTIME_SHAREDLIB_PATH)));
+		const FString OrtSharedLibPath = FPaths::Combine(PluginDir, TEXT(UE_STRINGIZE(ONNXRUNTIME_SHAREDLIB_PATH)));
 
 		OrtDllHandle = FPlatformProcess::GetDllHandle(*OrtSharedLibPath);
 		if (!OrtDllHandle)
diff --git a/Engine/Plugins/Media/MfMedia/Source/MfMedia/Private/MfMediaPrivate.h b/Engine/Plugins/Media/MfMedia/Source/MfMedia/Private/MfMediaPrivate.h
--- a/Engine/Plugins/Media/MfMedia/Source/MfMedia/Private/MfMediaPrivate.h
+++ b/Engine/Plugins/Media/MfMedia/Source/MfMedia/Private/MfMediaPrivate.h
@@ -43,7 +43,7 @@
 	#include "Microsoft/HideMicrosoftPlatformTypes.h"
 
 #elif PLATFORM_WINDOWS && !UE_SERVER
-	#pragma message("Skipping MfMedia (requires WINVER >= 0x0601, but WINVER is " PREPROCESSOR_TO_STRING(WINVER) ")")
+	#pragma message("Skipping MfMedia (requires WINVER >= 0x0601, but WINVER is " UE_STRINGIZE(WINVER) ")")
 
 #endif //MFMEDIA_SUPPORTED_PLATFORM
 
diff --git a/Engine/Plugins/Media/PixelStreaming2/Source/PixelStreaming2RTC/Private/EpicRtcVideoDecoder.cpp b/Engine/Plugins/Media/PixelStreaming2/Source/PixelStreaming2RTC/Private/EpicRtcVideoDecoder.cpp
--- a/Engine/Plugins/Media/PixelStreaming2/Source/PixelStreaming2RTC/Private/EpicRtcVideoDecoder.cpp
+++ b/Engine/Plugins/Media/PixelStreaming2/Source/PixelStreaming2RTC/Private/EpicRtcVideoDecoder.cpp
@@ -143,7 +143,7 @@ namespace UE::PixelStreaming2
 				}
 				else
 				{
-					UE_LOGFMT(LogPixelStreaming2RTC, Error, "VideoResource isn't a compatible type! Expected either a FVideoResourceRHI or FVideoResourceCPU. Received: {0}", PREPROCESSOR_TO_STRING(TVideoResource));
+					UE_LOGFMT(LogPixelStreaming2RTC, Error, "VideoResource isn't a compatible type! Expected either a FVideoResourceRHI or FVideoResourceCPU. Received: {0}", UE_STRINGIZE(TVideoResource));
 					return EpicRtcMediaResult::Error;
 				}	
 				
diff --git a/Engine/Plugins/Media/PixelStreaming2/Source/PixelStreaming2RTC/Private/EpicRtcVideoEncoder.cpp b/Engine/Plugins/Media/PixelStreaming2/Source/PixelStreaming2RTC/Private/EpicRtcVideoEncoder.cpp
--- a/Engine/Plugins/Media/PixelStreaming2/Source/PixelStreaming2RTC/Private/EpicRtcVideoEncoder.cpp
+++ b/Engine/Plugins/Media/PixelStreaming2/Source/PixelStreaming2RTC/Private/EpicRtcVideoEncoder.cpp
@@ -249,7 +249,7 @@ namespace UE::PixelStreaming2
 
 			if (!Encoder)
 			{
-				UE_LOGFMT(LogPixelStreaming2RTC, Error, "TEpicRtcVideoEncoder<{0}>::SetConfig: Unable to create {1} encoder!", PREPROCESSOR_TO_STRING(TVideoResource), ToString(CodecInfo->GetCodec()));
+				UE_LOGFMT(LogPixelStreaming2RTC, Error, "TEpicRtcVideoEncoder<{0}>::SetConfig: Unable to create {1} encoder!", UE_STRINGIZE(TVideoResource), ToString(CodecInfo->GetCodec()));
 				return EpicRtcMediaResult::Error;
 			}
 
@@ -315,7 +315,7 @@ namespace UE::PixelStreaming2
 		}
 		else
 		{
-			UE_LOGFMT(LogPixelStreaming2RTC, Error, "VideoResource isn't a compatible type! Expected either a FVideoResourceRHI or FVideoResourceCPU. Received: {0}", PREPROCESSOR_TO_STRING(TVideoResource));
+			UE_LOGFMT(LogPixelStreaming2RTC, Error, "VideoResource isn't a compatible type! Expected either a FVideoResourceRHI or FVideoResourceCPU. Received: {0}", UE_STRINGIZE(TVideoResource));
 			return EpicRtcMediaResult::Error;
 		}
 
@@ -363,7 +363,7 @@ namespace UE::PixelStreaming2
 		}
 		else
 		{
-			UE_LOGFMT(LogPixelStreaming2RTC, Error, "VideoResource isn't a compatible type! Expected either a FVideoResourceRHI or FVideoResourceCPU. Received: {0}", PREPROCESSOR_TO_STRING(TVideoResource));
+			UE_LOGFMT(LogPixelStreaming2RTC, Error, "VideoResource isn't a compatible type! Expected either a FVideoResourceRHI or FVideoResourceCPU. Received: {0}", UE_STRINGIZE(TVideoResource));
 			return EpicRtcMediaResult::Error;
 		}
 
diff --git a/Engine/Plugins/Media/WmfMedia/Source/WmfMedia/Public/WmfMediaCommon.h b/Engine/Plugins/Media/WmfMedia/Source/WmfMedia/Public/WmfMediaCommon.h
--- a/Engine/Plugins/Media/WmfMedia/Source/WmfMedia/Public/WmfMediaCommon.h
+++ b/Engine/Plugins/Media/WmfMedia/Source/WmfMedia/Public/WmfMediaCommon.h
@@ -50,7 +50,7 @@
 	#include "Windows/HideWindowsPlatformTypes.h"
 
 #elif PLATFORM_WINDOWS && !UE_SERVER
-	#pragma message("Skipping WmfMedia (requires WINVER >= 0x0600, but WINVER is " PREPROCESSOR_TO_STRING(WINVER) ")")
+	#pragma message("Skipping WmfMedia (requires WINVER >= 0x0600, but WINVER is " UE_STRINGIZE(WINVER) ")")
 
 #endif //WMFMEDIA_SUPPORTED_PLATFORM
 
diff --git a/Engine/Plugins/Mutable/Source/CustomizableObjectEditor/Private/MuCOE/GenerateMutableSource/GenerateMutableSource.h b/Engine/Plugins/Mutable/Source/CustomizableObjectEditor/Private/MuCOE/GenerateMutableSource/GenerateMutableSource.h
--- a/Engine/Plugins/Mutable/Source/CustomizableObjectEditor/Private/MuCOE/GenerateMutableSource/GenerateMutableSource.h
+++ b/Engine/Plugins/Mutable/Source/CustomizableObjectEditor/Private/MuCOE/GenerateMutableSource/GenerateMutableSource.h
@@ -456,7 +456,7 @@ private:
 
 /** Return the default value if there is a cycle. */
 #define RETURN_ON_CYCLE(Pin, GenerationContext) \
-	FGraphCycle GraphCycle(FGraphCycleKey(Pin, TEXT(__FILE__ PREPROCESSOR_TO_STRING(__LINE__)),GenerationContext.MacroNodesStack.Num() ? GenerationContext.MacroNodesStack.Top() : nullptr), GenerationContext); \
+	FGraphCycle GraphCycle(FGraphCycleKey(Pin, TEXT(__FILE__ UE_STRINGIZE(__LINE__)),GenerationContext.MacroNodesStack.Num() ? GenerationContext.MacroNodesStack.Top() : nullptr), GenerationContext); \
 	if (GraphCycle.FoundCycle()) \
 	{ \
 		return {}; \
diff --git a/Engine/Plugins/NNE/NNERuntimeORT/Source/NNERuntimeORT/Private/NNERuntimeORTModule.cpp b/Engine/Plugins/NNE/NNERuntimeORT/Source/NNERuntimeORT/Private/NNERuntimeORTModule.cpp
--- a/Engine/Plugins/NNE/NNERuntimeORT/Source/NNERuntimeORT/Private/NNERuntimeORTModule.cpp
+++ b/Engine/Plugins/NNE/NNERuntimeORT/Source/NNERuntimeORT/Private/NNERuntimeORTModule.cpp
@@ -66,7 +66,7 @@ void FNNERuntimeORTModule::StartupModule()
 	using namespace UE::NNERuntimeORT::Private;
 
 	const FString PluginDir = IPluginManager::Get().FindPlugin("NNERuntimeORT")->GetBaseDir();
-	const FString OrtSharedLibPath = FPaths::Combine(PluginDir, TEXT(PREPROCESSOR_TO_STRING(ONNXRUNTIME_SHAREDLIB_PATH)));
+	const FString OrtSharedLibPath = FPaths::Combine(PluginDir, TEXT(UE_STRINGIZE(ONNXRUNTIME_SHAREDLIB_PATH)));
 
 	if (!DllHelper::GetDllHandle(OrtSharedLibPath, DllHandles))
 	{
@@ -80,7 +80,7 @@ void FNNERuntimeORTModule::StartupModule()
 
 #if PLATFORM_WINDOWS
 	const FString ModuleDir = FPlatformProcess::GetModulesDirectory();
-	const FString DirectMLSharedLibPath = FPaths::Combine(ModuleDir, TEXT(PREPROCESSOR_TO_STRING(DIRECTML_PATH)), TEXT("DirectML.dll"));
+	const FString DirectMLSharedLibPath = FPaths::Combine(ModuleDir, TEXT(UE_STRINGIZE(DIRECTML_PATH)), TEXT("DirectML.dll"));
 
 	bDirectMLDllLoaded = DllHelper::GetDllHandle(DirectMLSharedLibPath, DllHandles);
 	if (!bDirectMLDllLoaded)
diff --git a/Engine/Plugins/NetcodeUnitTest/NetcodeUnitTest/Source/NetcodeUnitTest/Public/NetcodeUnitTest.h b/Engine/Plugins/NetcodeUnitTest/NetcodeUnitTest/Source/NetcodeUnitTest/Public/NetcodeUnitTest.h
--- a/Engine/Plugins/NetcodeUnitTest/NetcodeUnitTest/Source/NetcodeUnitTest/Public/NetcodeUnitTest.h
+++ b/Engine/Plugins/NetcodeUnitTest/NetcodeUnitTest/Source/NetcodeUnitTest/Public/NetcodeUnitTest.h
@@ -52,9 +52,9 @@ DECLARE_LOG_CATEGORY_EXTERN(NetCodeTestNone, Log, All);
 #define UNIT_ASSERT(Condition) \
 	if (!(Condition)) \
 	{ \
-		UE_LOG(LogUnitTest, Error, TEXT(PREPROCESSOR_TO_STRING(__FILE__)) TEXT(": Line: ") TEXT(PREPROCESSOR_TO_STRING(__LINE__)) \
+		UE_LOG(LogUnitTest, Error, TEXT(UE_STRINGIZE(__FILE__)) TEXT(": Line: ") TEXT(UE_STRINGIZE(__LINE__)) \
 				TEXT(":")); \
-		UE_LOG(LogUnitTest, Error, TEXT("Condition '(") TEXT(PREPROCESSOR_TO_STRING(Condition)) TEXT(")' failed")); \
+		UE_LOG(LogUnitTest, Error, TEXT("Condition '(") TEXT(UE_STRINGIZE(Condition)) TEXT(")' failed")); \
 		check(false); /* Try to get a meaningful stack trace. */ \
 		FPlatformMisc::RequestExit(true); \
 		CA_ASSUME(false); \
diff --git a/Engine/Plugins/Online/EOSShared/Source/EOSShared/Public/EOSShared.h b/Engine/Plugins/Online/EOSShared/Source/EOSShared/Public/EOSShared.h
--- a/Engine/Plugins/Online/EOSShared/Source/EOSShared/Public/EOSShared.h
+++ b/Engine/Plugins/Online/EOSShared/Source/EOSShared/Public/EOSShared.h
@@ -20,7 +20,7 @@
 #if UE_WITH_EOS_SDK_APIVERSION_WARNINGS
 #define UE_EOS_CHECK_API_MISMATCH(Actual, Expected) \
 	struct UE_JOIN(FEosApiMismatchMsg_, __LINE__) { \
-		[[deprecated(#Actual " updated from " #Expected " to " PREPROCESSOR_TO_STRING(Actual) ", behaviour/params may have changed")]] \
+		[[deprecated(#Actual " updated from " #Expected " to " UE_STRINGIZE(Actual) ", behaviour/params may have changed")]] \
 		static constexpr int condition(TStaticDeprecateExpression<true>) { return 1; } \
 		static constexpr int condition(TStaticDeprecateExpression<false>) { return 1; } \
 	}; \
diff --git a/Engine/Plugins/Runtime/CommonUI/Source/CommonInput/Private/CommonInputBaseTypes.cpp b/Engine/Plugins/Runtime/CommonUI/Source/CommonInput/Private/CommonInputBaseTypes.cpp
--- a/Engine/Plugins/Runtime/CommonUI/Source/CommonInput/Private/CommonInputBaseTypes.cpp
+++ b/Engine/Plugins/Runtime/CommonUI/Source/CommonInput/Private/CommonInputBaseTypes.cpp
@@ -524,7 +524,7 @@ const TArray<FName>& FCommonInputPlatformBaseData::GetRegisteredPlatforms()
 FName FCommonInputBase::GetCurrentPlatformName()
 {
 #if defined(UE_COMMONINPUT_PLATFORM_TYPE)
-	return FName(PREPROCESSOR_TO_STRING(UE_COMMONINPUT_PLATFORM_TYPE));
+	return FName(UE_STRINGIZE(UE_COMMONINPUT_PLATFORM_TYPE));
 #else
 	return FCommonInputDefaults::PlatformPC;
 #endif
diff --git a/Engine/Plugins/Runtime/OpenCV/Source/OpenCVHelper/Private/OpenCVHelperModule.cpp b/Engine/Plugins/Runtime/OpenCV/Source/OpenCVHelper/Private/OpenCVHelperModule.cpp
--- a/Engine/Plugins/Runtime/OpenCV/Source/OpenCVHelper/Private/OpenCVHelperModule.cpp
+++ b/Engine/Plugins/Runtime/OpenCV/Source/OpenCVHelper/Private/OpenCVHelperModule.cpp
@@ -77,8 +77,8 @@ void FOpenCVHelperModule::StartupModule()
 #if WITH_OPENCV
 
 #if defined(OPENCV_DLL_NAME)
-	const FString OpenCvBinPath = PluginDir / TEXT(PREPROCESSOR_TO_STRING(OPENCV_PLATFORM_PATH));
-	const FString DLLPath = OpenCvBinPath / TEXT(PREPROCESSOR_TO_STRING(OPENCV_DLL_NAME));
+	const FString OpenCvBinPath = PluginDir / TEXT(UE_STRINGIZE(OPENCV_PLATFORM_PATH));
+	const FString DLLPath = OpenCvBinPath / TEXT(UE_STRINGIZE(OPENCV_DLL_NAME));
 
 	FPlatformProcess::PushDllDirectory(*OpenCvBinPath);
 	OpenCvDllHandle = FPlatformProcess::GetDllHandle(*DLLPath);
diff --git a/Engine/Plugins/Runtime/ReplicationSystemTestPlugin/Source/Private/NetworkAutomationTest.h b/Engine/Plugins/Runtime/ReplicationSystemTestPlugin/Source/Private/NetworkAutomationTest.h
--- a/Engine/Plugins/Runtime/ReplicationSystemTestPlugin/Source/Private/NetworkAutomationTest.h
+++ b/Engine/Plugins/Runtime/ReplicationSystemTestPlugin/Source/Private/NetworkAutomationTest.h
@@ -118,7 +118,7 @@ public:
 
 #define UE_NET_TEST_CATCH_INTERNAL(TestSuite, TestCase, BaseClass) \
 	constexpr TCHAR TestSuite ## _ ## TestCase ## Name[] = TEXT("Net." #TestSuite "." #TestCase); \
-	TEST_CASE_METHOD((TCatchFixtureWrapper<BaseClass, TestSuite ## _ ## TestCase ## Name>), PREPROCESSOR_TO_STRING(TestSuite ## _ ## TestCase))
+	TEST_CASE_METHOD((TCatchFixtureWrapper<BaseClass, TestSuite ## _ ## TestCase ## Name>), UE_STRINGIZE(TestSuite ## _ ## TestCase))
 
 #define UE_NET_TEST_FIXTURE(TestFixture, TestCase) UE_NET_TEST_CATCH_INTERNAL(TestFixture, TestCase, TestFixture)
 #define UE_NET_TEST(TestSuite, TestCase) UE_NET_TEST_CATCH_INTERNAL(TestSuite, TestCase, UE::Net::FNetworkAutomationTestSuiteFixture)
diff --git a/Engine/Plugins/VirtualProduction/Rivermax/RivermaxCore/Source/RivermaxCore/Private/RivermaxManager.cpp b/Engine/Plugins/VirtualProduction/Rivermax/RivermaxCore/Source/RivermaxCore/Private/RivermaxManager.cpp
--- a/Engine/Plugins/VirtualProduction/Rivermax/RivermaxCore/Source/RivermaxCore/Private/RivermaxManager.cpp
+++ b/Engine/Plugins/VirtualProduction/Rivermax/RivermaxCore/Source/RivermaxCore/Private/RivermaxManager.cpp
@@ -118,7 +118,7 @@ namespace UE::RivermaxCore::Private
 	{
 #if defined(RIVERMAX_LIBRARY_PLATFORM_PATH) && defined(RIVERMAX_LIBRARY_NAME)
 
-		FString LibraryPath = TEXT(PREPROCESSOR_TO_STRING(RIVERMAX_LIBRARY_PLATFORM_PATH));
+		FString LibraryPath = TEXT(UE_STRINGIZE(RIVERMAX_LIBRARY_PLATFORM_PATH));
 
 		// Look for environment variable defining where the library version we are looking for is located
 
@@ -151,7 +151,7 @@ namespace UE::RivermaxCore::Private
 		if (FPaths::DirectoryExists(LibraryPath))
 		{
 			FPlatformProcess::PushDllDirectory(*LibraryPath);
-			const FString LibraryName = TEXT(PREPROCESSOR_TO_STRING(RIVERMAX_LIBRARY_NAME));
+			const FString LibraryName = TEXT(UE_STRINGIZE(RIVERMAX_LIBRARY_NAME));
 			LibraryHandle = FPlatformProcess::GetDllHandle(*LibraryName);
 			FPlatformProcess::PopDllDirectory(*LibraryPath);
 		}
diff --git a/Engine/Source/Developer/LowLevelTestsRunner/Private/TestCommon/Initialization.cpp b/Engine/Source/Developer/LowLevelTestsRunner/Private/TestCommon/Initialization.cpp
--- a/Engine/Source/Developer/LowLevelTestsRunner/Private/TestCommon/Initialization.cpp
+++ b/Engine/Source/Developer/LowLevelTestsRunner/Private/TestCommon/Initialization.cpp
@@ -60,7 +60,7 @@ void SetProjectNameAndDirectory()
 	// and project files. Resolves issues with UE_PROJECT_NAME being set to the parent project instead of the program's project.
 	// This will likely break if we ever have one LLT project file have mutliple targets eg: LyraGameEOS vs LyraGame
 #ifdef UE_TARGET_NAME
-	FCString::Strncpy(GInternalProjectName, TEXT(PREPROCESSOR_TO_STRING(UE_TARGET_NAME)), UE_ARRAY_COUNT(GInternalProjectName));
+	FCString::Strncpy(GInternalProjectName, TEXT(UE_STRINGIZE(UE_TARGET_NAME)), UE_ARRAY_COUNT(GInternalProjectName));
 #endif
 
 	// There may be cases on some platforms that you need to verify files are in the filesystem, 
diff --git a/Engine/Source/Developer/LowLevelTestsRunner/Private/TestExterns.cpp b/Engine/Source/Developer/LowLevelTestsRunner/Private/TestExterns.cpp
--- a/Engine/Source/Developer/LowLevelTestsRunner/Private/TestExterns.cpp
+++ b/Engine/Source/Developer/LowLevelTestsRunner/Private/TestExterns.cpp
@@ -8,7 +8,7 @@
 #if !defined(UE_PROJECT_NAME)
 TCHAR GInternalProjectName[64] = TEXT("LowLevelTests");
 #else
-TCHAR GInternalProjectName[64] = TEXT( PREPROCESSOR_TO_STRING(UE_PROJECT_NAME) );
+TCHAR GInternalProjectName[64] = TEXT( UE_STRINGIZE(UE_PROJECT_NAME) );
 #endif 
 IMPLEMENT_FOREIGN_ENGINE_DIR()
 
diff --git a/Engine/Source/Editor/EditorConfig/Private/EditorConfigSubystem.cpp b/Engine/Source/Editor/EditorConfig/Private/EditorConfigSubystem.cpp
--- a/Engine/Source/Editor/EditorConfig/Private/EditorConfigSubystem.cpp
+++ b/Engine/Source/Editor/EditorConfig/Private/EditorConfigSubystem.cpp
@@ -22,7 +22,7 @@ void UEditorConfigSubsystem::Initialize(FSubsystemCollectionBase& Collection)
 	AddSearchDirectory(ESearchDirectoryType::Engine, FPaths::Combine(FPaths::EngineConfigDir(), TEXT("Editor"))); // Engine
 	AddSearchDirectory(ESearchDirectoryType::Project, FPaths::Combine(FPaths::ProjectConfigDir(), TEXT("Editor"))); // ProjectName
 #ifdef UE_SAVED_DIR_OVERRIDE
-	AddSearchDirectory(ESearchDirectoryType::User, FPaths::Combine(FPlatformProcess::UserSettingsDir(), TEXT(PREPROCESSOR_TO_STRING(UE_SAVED_DIR_OVERRIDE)), TEXT("Editor"))); // AppData
+	AddSearchDirectory(ESearchDirectoryType::User, FPaths::Combine(FPlatformProcess::UserSettingsDir(), TEXT(UE_STRINGIZE(UE_SAVED_DIR_OVERRIDE)), TEXT("Editor"))); // AppData
 #else
 	AddSearchDirectory(ESearchDirectoryType::User, FPaths::Combine(FPlatformProcess::UserSettingsDir(), *FApp::GetEpicProductIdentifier(), TEXT("Editor"))); // AppData
 #endif
diff --git a/Engine/Source/Editor/MainFrame/Private/Menus/LayoutsMenu.cpp b/Engine/Source/Editor/MainFrame/Private/Menus/LayoutsMenu.cpp
--- a/Engine/Source/Editor/MainFrame/Private/Menus/LayoutsMenu.cpp
+++ b/Engine/Source/Editor/MainFrame/Private/Menus/LayoutsMenu.cpp
@@ -1048,7 +1048,7 @@ static FString GetLayoutJsonFileName(const FString& InConfigFileName)
 {
 	const FString JsonFileName = FPaths::GetBaseFilename(InConfigFileName) + TEXT(".json");
 #ifdef UE_SAVED_DIR_OVERRIDE
-	const FString UserSettingsPath = FPaths::Combine(FPlatformProcess::UserSettingsDir(), TEXT(PREPROCESSOR_TO_STRING(UE_SAVED_DIR_OVERRIDE)), TEXT("Editor"), JsonFileName);
+	const FString UserSettingsPath = FPaths::Combine(FPlatformProcess::UserSettingsDir(), TEXT(UE_STRINGIZE(UE_SAVED_DIR_OVERRIDE)), TEXT("Editor"), JsonFileName);
 #else
 	const FString UserSettingsPath = FPaths::Combine(FPlatformProcess::UserSettingsDir(), FApp::GetEpicProductIdentifier(), TEXT("Editor"), JsonFileName);
 #endif
diff --git a/Engine/Source/Editor/NNEEditor/Private/NNEEditorModule.cpp b/Engine/Source/Editor/NNEEditor/Private/NNEEditorModule.cpp
--- a/Engine/Source/Editor/NNEEditor/Private/NNEEditorModule.cpp
+++ b/Engine/Source/Editor/NNEEditor/Private/NNEEditorModule.cpp
@@ -43,7 +43,7 @@ public:
 #else
 		const FString ModuleDir = FPlatformProcess::GetModulesDirectory();
 #endif
-		const FString NNEEditorOnnxToolsSharedLibPath = FPaths::Combine(ModuleDir, TEXT(PREPROCESSOR_TO_STRING(NNEEDITORONNXTOOLS_SHAREDLIB_FILENAME)));
+		const FString NNEEditorOnnxToolsSharedLibPath = FPaths::Combine(ModuleDir, TEXT(UE_STRINGIZE(NNEEDITORONNXTOOLS_SHAREDLIB_FILENAME)));
 		NNEEditorOnnxToolsSharedLibHandle = UE::NNEEditor::Private::OnnxHelper::GetSharedLibHandle(NNEEditorOnnxToolsSharedLibPath);
 		UE::NNEEditor::Private::OnnxModelInspectorHelper::SetupSharedLibFunctionPointer(NNEEditorOnnxToolsSharedLibHandle);
 #endif // ifdef NNEEDITORONNXTOOLS_SUPPORTED
diff --git a/Engine/Source/Programs/SwitchboardListener/SblCore/Private/Windows/Resources/Resource.rc2 b/Engine/Source/Programs/SwitchboardListener/SblCore/Private/Windows/Resources/Resource.rc2
--- a/Engine/Source/Programs/SwitchboardListener/SblCore/Private/Windows/Resources/Resource.rc2
+++ b/Engine/Source/Programs/SwitchboardListener/SblCore/Private/Windows/Resources/Resource.rc2
@@ -12,31 +12,31 @@
 
 // Various strings used for project resources
 #ifdef PROJECT_COMPANY_NAME
-	#define BUILD_PROJECT_COMPANY_NAME PREPROCESSOR_TO_STRING(PROJECT_COMPANY_NAME)
+	#define BUILD_PROJECT_COMPANY_NAME UE_STRINGIZE(PROJECT_COMPANY_NAME)
 #else
 	#define BUILD_PROJECT_COMPANY_NAME EPIC_COMPANY_NAME
 #endif
 
 #ifdef PROJECT_COPYRIGHT_STRING
-	#define BUILD_PROJECT_COPYRIGHT_STRING PREPROCESSOR_TO_STRING(PROJECT_COPYRIGHT_STRING)
+	#define BUILD_PROJECT_COPYRIGHT_STRING UE_STRINGIZE(PROJECT_COPYRIGHT_STRING)
 #else
 	#define BUILD_PROJECT_COPYRIGHT_STRING EPIC_COPYRIGHT_STRING
 #endif
 
 #ifdef PROJECT_PRODUCT_NAME
-	#define BUILD_PROJECT_PRODUCT_NAME PREPROCESSOR_TO_STRING(PROJECT_PRODUCT_NAME)
+	#define BUILD_PROJECT_PRODUCT_NAME UE_STRINGIZE(PROJECT_PRODUCT_NAME)
 #else
-	#define BUILD_PROJECT_PRODUCT_NAME PREPROCESSOR_TO_STRING(UE_APP_NAME)
+	#define BUILD_PROJECT_PRODUCT_NAME UE_STRINGIZE(UE_APP_NAME)
 #endif
 
 #ifdef BUILD_VERSION
-	#define BUILD_PROJECT_PRODUCT_VERSION PREPROCESSOR_TO_STRING(BUILD_VERSION)
+	#define BUILD_PROJECT_PRODUCT_VERSION UE_STRINGIZE(BUILD_VERSION)
 #else
 	#define BUILD_PROJECT_PRODUCT_VERSION ENGINE_VERSION_STRING
 #endif
 
 #ifdef PROJECT_PRODUCT_IDENTIFIER
-	#define BUILD_PROJECT_PRODUCT_IDENTIFIER PREPROCESSOR_TO_STRING(PROJECT_PRODUCT_IDENTIFIER)
+	#define BUILD_PROJECT_PRODUCT_IDENTIFIER UE_STRINGIZE(PROJECT_PRODUCT_IDENTIFIER)
 #else
 	#define BUILD_PROJECT_PRODUCT_IDENTIFIER EPIC_PRODUCT_IDENTIFIER
 #endif
@@ -77,7 +77,7 @@ BEGIN
 			VALUE "FileDescription", BUILD_PROJECT_PRODUCT_NAME
 			VALUE "InternalName", BUILD_PROJECT_PRODUCT_IDENTIFIER
 #ifdef ORIGINAL_FILE_NAME
-			VALUE "OriginalFilename", PREPROCESSOR_TO_STRING(ORIGINAL_FILE_NAME)
+			VALUE "OriginalFilename", UE_STRINGIZE(ORIGINAL_FILE_NAME)
 #endif
 		END
 	END
diff --git a/Engine/Source/Runtime/Core/Private/Containers/String.cpp.inl b/Engine/Source/Runtime/Core/Private/Containers/String.cpp.inl
--- a/Engine/Source/Runtime/Core/Private/Containers/String.cpp.inl
+++ b/Engine/Source/Runtime/Core/Private/Containers/String.cpp.inl
@@ -1777,7 +1777,7 @@ void UE_STRING_CLASS::AppendfImpl(UE_STRING_CLASS& AppendToMe, const ElementType
 	}
 }
 
-static_assert(PLATFORM_LITTLE_ENDIAN, PREPROCESSOR_TO_STRING(UE_STRING_CLASS) " serialization needs updating to support big-endian platforms!");
+static_assert(PLATFORM_LITTLE_ENDIAN, UE_STRINGIZE(UE_STRING_CLASS) " serialization needs updating to support big-endian platforms!");
 
 FArchive& operator<<(FArchive& Ar, UE_STRING_CLASS& Str)
 {
diff --git a/Engine/Source/Runtime/Core/Private/GenericPlatform/GenericPlatformMisc.cpp b/Engine/Source/Runtime/Core/Private/GenericPlatform/GenericPlatformMisc.cpp
--- a/Engine/Source/Runtime/Core/Private/GenericPlatform/GenericPlatformMisc.cpp
+++ b/Engine/Source/Runtime/Core/Private/GenericPlatform/GenericPlatformMisc.cpp
@@ -1516,17 +1516,17 @@ const TCHAR* FGenericPlatformMisc::GeneratedConfigDir()
 
 const TCHAR* FGenericPlatformMisc::GetUBTPlatform()
 {
-	return TEXT(PREPROCESSOR_TO_STRING(UBT_COMPILED_PLATFORM));
+	return TEXT(UE_STRINGIZE(UBT_COMPILED_PLATFORM));
 }
 
 const TCHAR* FGenericPlatformMisc::GetUBTTarget()
 {
-	return TEXT(PREPROCESSOR_TO_STRING(UBT_COMPILED_TARGET));
+	return TEXT(UE_STRINGIZE(UBT_COMPILED_TARGET));
 }
 
 const TCHAR* FGenericPlatformMisc::GetUBTArchitecture()
 {
-	return TEXT(PREPROCESSOR_TO_STRING(UBT_COMPILED_ARCHITECTURE));
+	return TEXT(UE_STRINGIZE(UBT_COMPILED_ARCHITECTURE));
 }
 
 
@@ -1544,7 +1544,7 @@ using FUBTTargetNameArrayType = TCHAR[128];
 			{
 				/** The name of the UBT target that the current executable was built from. Defaults to the UE default target for this type to make content only projects work,
 					but will be overridden by the primary game module if it exists */
-				FCString::Strcpy(UBTTargetName, TEXT("Unreal" PREPROCESSOR_TO_STRING(UBT_COMPILED_TARGET)));
+				FCString::Strcpy(UBTTargetName, TEXT("Unreal" UE_STRINGIZE(UBT_COMPILED_TARGET)));
 			}
 
 			TCHAR UBTTargetName[128];
@@ -1559,7 +1559,7 @@ using FUBTTargetNameArrayType = TCHAR[128];
 	{
 		/** The name of the UBT target that the current executable was built from. Defaults to the UE default target for this type to make content only projects work,
 			but will be overridden by the primary game module if it exists */
-		static FUBTTargetNameArrayType GUBTTargetName = TEXT("Unreal" PREPROCESSOR_TO_STRING(UBT_COMPILED_TARGET));
+		static FUBTTargetNameArrayType GUBTTargetName = TEXT("Unreal" UE_STRINGIZE(UBT_COMPILED_TARGET));
 
 		return GUBTTargetName;
 	}
diff --git a/Engine/Source/Runtime/Core/Private/Misc/App.cpp b/Engine/Source/Runtime/Core/Private/Misc/App.cpp
--- a/Engine/Source/Runtime/Core/Private/Misc/App.cpp
+++ b/Engine/Source/Runtime/Core/Private/Misc/App.cpp
@@ -422,7 +422,7 @@ void FApp::PrintStartupLogMessages()
 #elif PLATFORM_CPU_X86_FAMILY
 	UE_LOG(LogInit, Log, TEXT("Architecture: x64"));
 #elif defined(UE_ARCH_NAME)
-	UE_LOG(LogInit, Log, TEXT("Architecture: %hs"), PREPROCESSOR_TO_STRING(UE_ARCH_NAME));
+	UE_LOG(LogInit, Log, TEXT("Architecture: %hs"), UE_STRINGIZE(UE_ARCH_NAME));
 #else
 #error No architecture name defined!
 #endif // x64/arm
diff --git a/Engine/Source/Runtime/Core/Private/Misc/ConfigCacheIni.cpp b/Engine/Source/Runtime/Core/Private/Misc/ConfigCacheIni.cpp
--- a/Engine/Source/Runtime/Core/Private/Misc/ConfigCacheIni.cpp
+++ b/Engine/Source/Runtime/Core/Private/Misc/ConfigCacheIni.cpp
@@ -6537,7 +6537,7 @@ const FString& FConfigCacheIni::GetCustomConfigString()
 #endif
 
 #ifdef UE_USE_COMMAND_LINE_PARAM_FOR_CUSTOM_CONFIG
-		FString CustomName = PREPROCESSOR_TO_STRING(UE_USE_COMMAND_LINE_PARAM_FOR_CUSTOM_CONFIG);
+		FString CustomName = UE_STRINGIZE(UE_USE_COMMAND_LINE_PARAM_FOR_CUSTOM_CONFIG);
 		if (!bCustomConfigOverrideApplied && FParse::Param(FCommandLine::Get(), *CustomName))
 		{
 			bCustomConfigOverrideApplied = true;
diff --git a/Engine/Source/Runtime/Core/Private/Misc/CoreGlobals.cpp b/Engine/Source/Runtime/Core/Private/Misc/CoreGlobals.cpp
--- a/Engine/Source/Runtime/Core/Private/Misc/CoreGlobals.cpp
+++ b/Engine/Source/Runtime/Core/Private/Misc/CoreGlobals.cpp
@@ -281,7 +281,7 @@ UE::CoreUObject::Private::FObjectHandlePackageDebugData*& CORE_API GObjectHandle
 /** Game name, used for base game directory and ini among other things										*/
 #if (!IS_MONOLITHIC && !IS_PROGRAM)
 // In modular game builds, the game name will be set when the application launches
-TCHAR					GInternalProjectName[64]					= TEXT(PREPROCESSOR_TO_STRING(UE_PROJECT_NAME));
+TCHAR					GInternalProjectName[64]					= TEXT(UE_STRINGIZE(UE_PROJECT_NAME));
 #elif !IS_MONOLITHIC && IS_PROGRAM
 // In non-monolithic programs builds, the game name will be set by the module, but not just yet, so we need to NOT initialize it!
 TCHAR					GInternalProjectName[64];
diff --git a/Engine/Source/Runtime/Core/Private/Misc/Paths.cpp b/Engine/Source/Runtime/Core/Private/Misc/Paths.cpp
--- a/Engine/Source/Runtime/Core/Private/Misc/Paths.cpp
+++ b/Engine/Source/Runtime/Core/Private/Misc/Paths.cpp
@@ -447,7 +447,7 @@ FString FPaths::ProjectUserDir()
 	{
 		// if defined, this will override both saveddirsuffix and enginesaveddirsuffix
 #ifdef UE_SAVED_DIR_OVERRIDE
-		return FPaths::Combine(FPlatformProcess::UserSettingsDir(), TEXT(PREPROCESSOR_TO_STRING(UE_SAVED_DIR_OVERRIDE))) + TEXT("/");
+		return FPaths::Combine(FPlatformProcess::UserSettingsDir(), TEXT(UE_STRINGIZE(UE_SAVED_DIR_OVERRIDE))) + TEXT("/");
 #else
 		return FPaths::Combine(FPlatformProcess::UserSettingsDir(), FApp::GetProjectName()) + TEXT("/");
 #endif
diff --git a/Engine/Source/Runtime/Core/Private/ProfilingDebugging/TraceAuxiliary.cpp b/Engine/Source/Runtime/Core/Private/ProfilingDebugging/TraceAuxiliary.cpp
--- a/Engine/Source/Runtime/Core/Private/ProfilingDebugging/TraceAuxiliary.cpp
+++ b/Engine/Source/Runtime/Core/Private/ProfilingDebugging/TraceAuxiliary.cpp
@@ -2017,7 +2017,7 @@ void FTraceAuxiliary::Initialize(const TCHAR* CommandLine)
 	// so that it is always sent (all channels are enabled prior to initialization).
 	const TCHAR* BranchName = BuildSettings::GetBranchName();
 	const TCHAR* BuildVersion = BuildSettings::GetBuildVersion();
-	constexpr uint32 PlatformLen = UE_ARRAY_COUNT(PREPROCESSOR_TO_STRING(UBT_COMPILED_PLATFORM)) - 1;
+	constexpr uint32 PlatformLen = UE_ARRAY_COUNT(UE_STRINGIZE(UBT_COMPILED_PLATFORM)) - 1;
 	const uint32 AppNameLen = FCString::Strlen(AppName);
 	const uint32 ProjectNameLen = FCString::Strlen(ProjectName);
 	const uint32 CommandLineLen = FCString::Strlen(CommandLine);
@@ -2043,7 +2043,7 @@ void FTraceAuxiliary::Initialize(const TCHAR* CommandLine)
 	}
 
 	UE_TRACE_LOG(Diagnostics, Session2, UE::Trace::TraceLogChannel, DataSize)
-		<< Session2.Platform(PREPROCESSOR_TO_STRING(UBT_COMPILED_PLATFORM), PlatformLen)
+		<< Session2.Platform(UE_STRINGIZE(UBT_COMPILED_PLATFORM), PlatformLen)
 		<< Session2.AppName(AppName, AppNameLen)
 		<< Session2.ProjectName(ProjectName, ProjectNameLen)
 		<< Session2.Branch(BranchName, BranchNameLen)
diff --git a/Engine/Source/Runtime/Core/Private/UObject/UnrealNames.cpp b/Engine/Source/Runtime/Core/Private/UObject/UnrealNames.cpp
--- a/Engine/Source/Runtime/Core/Private/UObject/UnrealNames.cpp
+++ b/Engine/Source/Runtime/Core/Private/UObject/UnrealNames.cpp
@@ -121,7 +121,7 @@ const TCHAR* LexToString(EName Ename)
 {
 	switch (Ename)
 	{
-#define REGISTER_NAME(num,name) case EName::name: return TEXT(PREPROCESSOR_TO_STRING(name));
+#define REGISTER_NAME(num,name) case EName::name: return TEXT(UE_STRINGIZE(name));
 #include "UObject/UnrealNames.inl"
 #undef REGISTER_NAME
 		default:
diff --git a/Engine/Source/Runtime/Core/Public/Clang/ClangPlatformCompilerPreSetup.h b/Engine/Source/Runtime/Core/Public/Clang/ClangPlatformCompilerPreSetup.h
--- a/Engine/Source/Runtime/Core/Public/Clang/ClangPlatformCompilerPreSetup.h
+++ b/Engine/Source/Runtime/Core/Public/Clang/ClangPlatformCompilerPreSetup.h
@@ -310,12 +310,12 @@
 
 #ifndef EMIT_CUSTOM_WARNING_AT_LINE
 	#define EMIT_CUSTOM_WARNING_AT_LINE(Line, Warning) \
-		_Pragma(PREPROCESSOR_TO_STRING(message(Warning)))
+		_Pragma(UE_STRINGIZE(message(Warning)))
 #endif // EMIT_CUSTOM_WARNING_AT_LINE
 
 #ifndef EMIT_CUSTOM_ERROR_AT_LINE
 	#define EMIT_CUSTOM_ERROR_AT_LINE(Line, Error) \
-		_Pragma(PREPROCESSOR_TO_STRING(message(Error)))
+		_Pragma(UE_STRINGIZE(message(Error)))
 #endif // EMIT_CUSTOM_ERROR_AT_LINE
 
 #ifndef UE_COMPILER_THIRD_PARTY_INCLUDES_START
diff --git a/Engine/Source/Runtime/Core/Public/Containers/UnrealString.h.inl b/Engine/Source/Runtime/Core/Public/Containers/UnrealString.h.inl
--- a/Engine/Source/Runtime/Core/Public/Containers/UnrealString.h.inl
+++ b/Engine/Source/Runtime/Core/Public/Containers/UnrealString.h.inl
@@ -109,16 +109,16 @@ public:
 	[[nodiscard]] CORE_API UE_STRING_CLASS(const UTF32CHAR* Str);
 
 	/** Construct from null-terminated C substring or nullptr */
-	UE_STRING_DEPRECATED(5.4, "This constructor has been deprecated - please use " PREPROCESSOR_TO_STRING(UE_STRING_CLASS) "::ConstructFromPtrSize(Ptr, Size) instead.") [[nodiscard]] CORE_API UE_STRING_CLASS(int32 Len, const ANSICHAR* Str);
-	UE_STRING_DEPRECATED(5.4, "This constructor has been deprecated - please use " PREPROCESSOR_TO_STRING(UE_STRING_CLASS) "::ConstructFromPtrSize(Ptr, Size) instead.") [[nodiscard]] CORE_API UE_STRING_CLASS(int32 Len, const WIDECHAR* Str);
-	UE_STRING_DEPRECATED(5.4, "This constructor has been deprecated - please use " PREPROCESSOR_TO_STRING(UE_STRING_CLASS) "::ConstructFromPtrSize(Ptr, Size) instead.") [[nodiscard]] CORE_API UE_STRING_CLASS(int32 Len, const UTF8CHAR* Str);
-	UE_STRING_DEPRECATED(5.4, "This constructor has been deprecated - please use " PREPROCESSOR_TO_STRING(UE_STRING_CLASS) "::ConstructFromPtrSize(Ptr, Size) instead.") [[nodiscard]] CORE_API UE_STRING_CLASS(int32 Len, const UCS2CHAR* Str);
+	UE_STRING_DEPRECATED(5.4, "This constructor has been deprecated - please use " UE_STRINGIZE(UE_STRING_CLASS) "::ConstructFromPtrSize(Ptr, Size) instead.") [[nodiscard]] CORE_API UE_STRING_CLASS(int32 Len, const ANSICHAR* Str);
+	UE_STRING_DEPRECATED(5.4, "This constructor has been deprecated - please use " UE_STRINGIZE(UE_STRING_CLASS) "::ConstructFromPtrSize(Ptr, Size) instead.") [[nodiscard]] CORE_API UE_STRING_CLASS(int32 Len, const WIDECHAR* Str);
+	UE_STRING_DEPRECATED(5.4, "This constructor has been deprecated - please use " UE_STRINGIZE(UE_STRING_CLASS) "::ConstructFromPtrSize(Ptr, Size) instead.") [[nodiscard]] CORE_API UE_STRING_CLASS(int32 Len, const UTF8CHAR* Str);
+	UE_STRING_DEPRECATED(5.4, "This constructor has been deprecated - please use " UE_STRINGIZE(UE_STRING_CLASS) "::ConstructFromPtrSize(Ptr, Size) instead.") [[nodiscard]] CORE_API UE_STRING_CLASS(int32 Len, const UCS2CHAR* Str);
 
 	/** Construct from null-terminated C string or nullptr with extra slack on top of original string length */
-	UE_STRING_DEPRECATED(5.4, "This constructor has been deprecated - please use " PREPROCESSOR_TO_STRING(UE_STRING_CLASS) "::ConstructWithSlack(Ptr, Size) instead.") [[nodiscard]] CORE_API UE_STRING_CLASS(const ANSICHAR* Str, int32 ExtraSlack);
-	UE_STRING_DEPRECATED(5.4, "This constructor has been deprecated - please use " PREPROCESSOR_TO_STRING(UE_STRING_CLASS) "::ConstructWithSlack(Ptr, Size) instead.") [[nodiscard]] CORE_API UE_STRING_CLASS(const WIDECHAR* Str, int32 ExtraSlack);
-	UE_STRING_DEPRECATED(5.4, "This constructor has been deprecated - please use " PREPROCESSOR_TO_STRING(UE_STRING_CLASS) "::ConstructWithSlack(Ptr, Size) instead.") [[nodiscard]] CORE_API UE_STRING_CLASS(const UTF8CHAR* Str, int32 ExtraSlack);
-	UE_STRING_DEPRECATED(5.4, "This constructor has been deprecated - please use " PREPROCESSOR_TO_STRING(UE_STRING_CLASS) "::ConstructWithSlack(Ptr, Size) instead.") [[nodiscard]] CORE_API UE_STRING_CLASS(const UCS2CHAR* Str, int32 ExtraSlack);
+	UE_STRING_DEPRECATED(5.4, "This constructor has been deprecated - please use " UE_STRINGIZE(UE_STRING_CLASS) "::ConstructWithSlack(Ptr, Size) instead.") [[nodiscard]] CORE_API UE_STRING_CLASS(const ANSICHAR* Str, int32 ExtraSlack);
+	UE_STRING_DEPRECATED(5.4, "This constructor has been deprecated - please use " UE_STRINGIZE(UE_STRING_CLASS) "::ConstructWithSlack(Ptr, Size) instead.") [[nodiscard]] CORE_API UE_STRING_CLASS(const WIDECHAR* Str, int32 ExtraSlack);
+	UE_STRING_DEPRECATED(5.4, "This constructor has been deprecated - please use " UE_STRINGIZE(UE_STRING_CLASS) "::ConstructWithSlack(Ptr, Size) instead.") [[nodiscard]] CORE_API UE_STRING_CLASS(const UTF8CHAR* Str, int32 ExtraSlack);
+	UE_STRING_DEPRECATED(5.4, "This constructor has been deprecated - please use " UE_STRINGIZE(UE_STRING_CLASS) "::ConstructWithSlack(Ptr, Size) instead.") [[nodiscard]] CORE_API UE_STRING_CLASS(const UCS2CHAR* Str, int32 ExtraSlack);
 
 	/** Construct from null-terminated C string with extra slack on top of original string length. */
 	[[nodiscard]] static CORE_API UE_STRING_CLASS ConstructWithSlack(const ANSICHAR* Str, int32 ExtraSlack);
@@ -2193,7 +2193,7 @@ struct TNameOf<UE_STRING_CLASS>
 {
 	[[nodiscard]] UE_FORCEINLINE_HINT static TCHAR const* GetName()
 	{
-		return TEXT(PREPROCESSOR_TO_STRING(UE_STRING_CLASS));
+		return TEXT(UE_STRINGIZE(UE_STRING_CLASS));
 	}
 };
 
diff --git a/Engine/Source/Runtime/Core/Public/HAL/PreprocessorHelpers.h b/Engine/Source/Runtime/Core/Public/HAL/PreprocessorHelpers.h
--- a/Engine/Source/Runtime/Core/Public/HAL/PreprocessorHelpers.h
+++ b/Engine/Source/Runtime/Core/Public/HAL/PreprocessorHelpers.h
@@ -77,26 +77,26 @@
 #if PLATFORM_IS_EXTENSION
 // Creates a string that can be used to include a header in the platform extension form "PlatformHeader.h", not like
 // below form. When using this you should add "// IWYU pragma: export" at the end of the line.
-#define COMPILED_PLATFORM_HEADER(Suffix) PREPROCESSOR_TO_STRING(UE_JOIN(PLATFORM_HEADER_NAME, Suffix))
+#define COMPILED_PLATFORM_HEADER(Suffix) UE_STRINGIZE(UE_JOIN(PLATFORM_HEADER_NAME, Suffix))
 #else
 // Creates a string that can be used to include a header in the form "Platform/PlatformHeader.h", like
 // "Windows/WindowsPlatformFile.h". When using this you should add "// IWYU pragma: export" at the end of the line.
-#define COMPILED_PLATFORM_HEADER(Suffix) PREPROCESSOR_TO_STRING(UE_JOIN(PLATFORM_HEADER_NAME/PLATFORM_HEADER_NAME, Suffix))
+#define COMPILED_PLATFORM_HEADER(Suffix) UE_STRINGIZE(UE_JOIN(PLATFORM_HEADER_NAME/PLATFORM_HEADER_NAME, Suffix))
 #endif
 
 // Creates a string that can be used to include a header in the platform extension form "PlatformHeader.h", but will
 // not include a directory like COMPILED_PLATFORM_HEADER does, generally for UBT generated headers.
-#define COMPILED_PLATFORM_HEADER_GENERATED(Suffix) PREPROCESSOR_TO_STRING(UE_JOIN(PLATFORM_HEADER_NAME, Suffix))
+#define COMPILED_PLATFORM_HEADER_GENERATED(Suffix) UE_STRINGIZE(UE_JOIN(PLATFORM_HEADER_NAME, Suffix))
 
 #if PLATFORM_IS_EXTENSION
 // Creates a string that can be used to include a header with the platform in its name, like
 // "Prefix/PlatformNameSuffix.h". When using this you should add "// IWYU pragma: export" at the end of the line.
-#define COMPILED_PLATFORM_HEADER_WITH_PREFIX(Prefix, Suffix) PREPROCESSOR_TO_STRING(Prefix/UE_JOIN(PLATFORM_HEADER_NAME, Suffix))
+#define COMPILED_PLATFORM_HEADER_WITH_PREFIX(Prefix, Suffix) UE_STRINGIZE(Prefix/UE_JOIN(PLATFORM_HEADER_NAME, Suffix))
 #else
 // Creates a string that can be used to include a header with the platform in its name, like
 // "Prefix/PlatformName/PlatformNameSuffix.h". When using this you should add "// IWYU pragma: export" at the end of the
 // line.
-#define COMPILED_PLATFORM_HEADER_WITH_PREFIX(Prefix, Suffix) PREPROCESSOR_TO_STRING(Prefix/PLATFORM_HEADER_NAME/UE_JOIN(PLATFORM_HEADER_NAME, Suffix))
+#define COMPILED_PLATFORM_HEADER_WITH_PREFIX(Prefix, Suffix) UE_STRINGIZE(Prefix/PLATFORM_HEADER_NAME/UE_JOIN(PLATFORM_HEADER_NAME, Suffix))
 #endif
 
 // These macros should be regarded as deprecated - use the UE_ macros they map to instead.
diff --git a/Engine/Source/Runtime/Core/Public/Misc/CoreMiscDefines.h b/Engine/Source/Runtime/Core/Public/Misc/CoreMiscDefines.h
--- a/Engine/Source/Runtime/Core/Public/Misc/CoreMiscDefines.h
+++ b/Engine/Source/Runtime/Core/Public/Misc/CoreMiscDefines.h
@@ -6,15 +6,13 @@
 #include "HAL/PreprocessorHelpers.h"
 
 // When passed to pragma message will result in clickable warning in VS
-#define WARNING_LOCATION(Line) __FILE__ "(" PREPROCESSOR_TO_STRING(Line) ")"
+#define WARNING_LOCATION(Line) __FILE__ "(" UE_STRINGIZE(Line) ")"
 
 // This file is included in some resource files, which issue a warning:
 //
 // warning RC4011: identifier truncated to 'PLATFORM_CAN_SUPPORT_EDITORONLY'
 //
-// due to limitations of resource compiler. The only thing needed from this file
-// for resource compilation is PREPROCESSOR_TO_STRING macro at the end, so we take
-// rest of code out for resource compilation.
+// due to limitations of resource compiler.
 #ifndef RC_INVOKED
 
 #define PLAYWORLD_PACKAGE_PREFIX TEXT("UEDPIE")
@@ -166,8 +164,8 @@ enum EPerElement {PerElement};
 
 // Push and pop macro definitions
 #ifdef __clang__
-	#define UE_PUSH_MACRO(name) _Pragma(PREPROCESSOR_TO_STRING(push_macro(name)))
-	#define UE_POP_MACRO(name) _Pragma(PREPROCESSOR_TO_STRING(pop_macro(name)))
+	#define UE_PUSH_MACRO(name) _Pragma(UE_STRINGIZE(push_macro(name)))
+	#define UE_POP_MACRO(name) _Pragma(UE_STRINGIZE(pop_macro(name)))
 #else
 	#define UE_PUSH_MACRO(name) __pragma(push_macro(name))
 	#define UE_POP_MACRO(name) __pragma(pop_macro(name))
@@ -390,12 +388,12 @@ enum EPerElement {PerElement};
 	#define UE_DEPRECATED_HEADER(Version, Message)
 #elif defined(_MSC_VER)
     #if UE_WARNINGS_AS_ERRORS
-    	#define UE_DEPRECATED_HEADER(Version, Message) __pragma(message(__FILE__ "(" PREPROCESSOR_TO_STRING(__LINE__) "): error C4996: " Message " Please update your code to the new API before upgrading to the next release, otherwise your project will no longer compile."))
+    	#define UE_DEPRECATED_HEADER(Version, Message) __pragma(message(__FILE__ "(" UE_STRINGIZE(__LINE__) "): error C4996: " Message " Please update your code to the new API before upgrading to the next release, otherwise your project will no longer compile."))
     #else
-    	#define UE_DEPRECATED_HEADER(Version, Message) __pragma(message(__FILE__ "(" PREPROCESSOR_TO_STRING(__LINE__) "): warning C4996: " Message " Please update your code to the new API before upgrading to the next release, otherwise your project will no longer compile."))
+    	#define UE_DEPRECATED_HEADER(Version, Message) __pragma(message(__FILE__ "(" UE_STRINGIZE(__LINE__) "): warning C4996: " Message " Please update your code to the new API before upgrading to the next release, otherwise your project will no longer compile."))
     #endif
 #else
-	#define UE_DEPRECATED_HEADER(Version, Message) _Pragma(PREPROCESSOR_TO_STRING(message(Message " Please update your code to the new API before upgrading to the next release, otherwise your project will no longer compile.")))
+	#define UE_DEPRECATED_HEADER(Version, Message) _Pragma(UE_STRINGIZE(message(Message " Please update your code to the new API before upgrading to the next release, otherwise your project will no longer compile.")))
 #endif
 
 
diff --git a/Engine/Source/Runtime/Core/Public/Modules/ModuleManager.h b/Engine/Source/Runtime/Core/Public/Modules/ModuleManager.h
--- a/Engine/Source/Runtime/Core/Public/Modules/ModuleManager.h
+++ b/Engine/Source/Runtime/Core/Public/Modules/ModuleManager.h
@@ -1044,7 +1044,7 @@ namespace UE::Core::Private
 	{ \
 		FTargetNameRegistration() \
 		{ \
-			FPlatformMisc::SetUBTTargetName(TEXT(PREPROCESSOR_TO_STRING(UE_TARGET_NAME))); \
+			FPlatformMisc::SetUBTTargetName(TEXT(UE_STRINGIZE(UE_TARGET_NAME))); \
 		} \
 	} GTargetNameRegistration;
 
@@ -1086,7 +1086,7 @@ namespace UE::Core::Private
 
 		#define IMPLEMENT_PRIMARY_GAME_MODULE( ModuleImplClass, ModuleName, DEPRECATED_GameName ) \
 			/* For monolithic builds, we must statically define the game's name string (See Core.h) */ \
-			TCHAR GInternalProjectName[64] = TEXT( PREPROCESSOR_TO_STRING(UE_PROJECT_NAME) ); \
+			TCHAR GInternalProjectName[64] = TEXT( UE_STRINGIZE(UE_PROJECT_NAME) ); \
 			/* Implement the GIsGameAgnosticExe variable (See Core.h). */ \
 			bool GIsGameAgnosticExe = false; \
 			IMPLEMENT_FOREIGN_ENGINE_DIR() \
@@ -1100,7 +1100,7 @@ namespace UE::Core::Private
 
 		#define IMPLEMENT_PRIMARY_GAME_MODULE( ModuleImplClass, ModuleName, DEPRECATED_GameName ) \
 			/* For monolithic builds, we must statically define the game's name string (See Core.h) */ \
-			TCHAR GInternalProjectName[64] = TEXT( PREPROCESSOR_TO_STRING(UE_PROJECT_NAME) ); \
+			TCHAR GInternalProjectName[64] = TEXT( UE_STRINGIZE(UE_PROJECT_NAME) ); \
 			PER_MODULE_BOILERPLATE \
 			IMPLEMENT_FOREIGN_ENGINE_DIR() \
 			IMPLEMENT_SIGNING_KEY_REGISTRATION() \
diff --git a/Engine/Source/Runtime/Core/Public/ProfilingDebugging/CsvProfiler.h b/Engine/Source/Runtime/Core/Public/ProfilingDebugging/CsvProfiler.h
--- a/Engine/Source/Runtime/Core/Public/ProfilingDebugging/CsvProfiler.h
+++ b/Engine/Source/Runtime/Core/Public/ProfilingDebugging/CsvProfiler.h
@@ -56,7 +56,7 @@ struct FCsvDeclaredStat;
 	// Events
 	#define CSV_EVENT(Category, Format, ...) \
 		FCsvProfiler::RecordEventf( CSV_CATEGORY_INDEX(Category), Format, ##__VA_ARGS__ ); \
-		TRACE_BOOKMARK(TEXT(PREPROCESSOR_TO_STRING(Category)) TEXT("/") Format, ##__VA_ARGS__)
+		TRACE_BOOKMARK(TEXT(UE_STRINGIZE(Category)) TEXT("/") Format, ##__VA_ARGS__)
 
 	#define CSV_EVENT_GLOBAL(Format, ...) \
 		FCsvProfiler::RecordEventf( CSV_CATEGORY_INDEX_GLOBAL, Format, ##__VA_ARGS__ ); \
diff --git a/Engine/Source/Runtime/Core/Public/ProfilingDebugging/InstanceCounter.h b/Engine/Source/Runtime/Core/Public/ProfilingDebugging/InstanceCounter.h
--- a/Engine/Source/Runtime/Core/Public/ProfilingDebugging/InstanceCounter.h
+++ b/Engine/Source/Runtime/Core/Public/ProfilingDebugging/InstanceCounter.h
@@ -100,13 +100,13 @@ protected:
 	 *	Add to class body of TypeName to track counts of "TypeName". Counts can be accessed via code or console
 	 */
 	#define COUNT_INSTANCES(TypeName) \
-		struct FCount##TypeName : FInstanceCountingObject { FCount##TypeName() : FInstanceCountingObject(TEXT(PREPROCESSOR_TO_STRING(TypeName))) {} }; FCount##TypeName TypeName##Count; 
+		struct FCount##TypeName : FInstanceCountingObject { FCount##TypeName() : FInstanceCountingObject(TEXT(UE_STRINGIZE(TypeName))) {} }; FCount##TypeName TypeName##Count; 
 
 	 /**
 	 *	Identical to COUNT_INSTANCES, but construction/destruction will be written to the log along with the address of this object (us, not the instance that contains us)
 	 */
 	#define COUNT_INSTANCES_AND_LOG(TypeName) \
-		struct FCount##TypeName : FInstanceCountingObject { FCount##TypeName() : FInstanceCountingObject(TEXT(PREPROCESSOR_TO_STRING(TypeName)), true) {} }; FCount##TypeName TypeName##Count;
+		struct FCount##TypeName : FInstanceCountingObject { FCount##TypeName() : FInstanceCountingObject(TEXT(UE_STRINGIZE(TypeName)), true) {} }; FCount##TypeName TypeName##Count;
 
 
 #else
diff --git a/Engine/Source/Runtime/Core/Public/ProfilingDebugging/StallDetector.h b/Engine/Source/Runtime/Core/Public/ProfilingDebugging/StallDetector.h
--- a/Engine/Source/Runtime/Core/Public/ProfilingDebugging/StallDetector.h
+++ b/Engine/Source/Runtime/Core/Public/ProfilingDebugging/StallDetector.h
@@ -30,16 +30,16 @@
 /**
  * Counts, but does not report, a stall.
  **/
-#define SCOPE_STALL_COUNTER(InName, InBudgetSeconds)                                                                                                                                           \
-	static_assert(UE_ARRAY_COUNT(PREPROCESSOR_TO_STRING(InName)) <= UE::StallDetectorScopeMaxNameLength, "Scope name must be less than UE::StallDetectorScopeMaxNameLength characters long "); \
+#define SCOPE_STALL_COUNTER(InName, InBudgetSeconds)                                                                                                                                 \
+	static_assert(UE_ARRAY_COUNT(UE_STRINGIZE(InName)) <= UE::StallDetectorScopeMaxNameLength, "Scope name must be less than UE::StallDetectorScopeMaxNameLength characters long "); \
 	static UE::FStallDetectorStats UE_JOIN(StallDetectorStats, __LINE__)(TEXT(#InName), InBudgetSeconds, UE::EStallDetectorReportingMode::Disabled);                                 \
 	UE::FStallDetector UE_JOIN(ScopeStallDetector, __LINE__)(UE_JOIN(StallDetectorStats, __LINE__));
 
 /**
  * Counts and sends a report of the first occurence of a stall.
  **/
-#define SCOPE_STALL_REPORTER(InName, InBudgetSeconds)                                                                                                                                          \
-	static_assert(UE_ARRAY_COUNT(PREPROCESSOR_TO_STRING(InName)) <= UE::StallDetectorScopeMaxNameLength, "Scope name must be less than UE::StallDetectorScopeMaxNameLength characters long "); \
+#define SCOPE_STALL_REPORTER(InName, InBudgetSeconds)                                                                                                                                \
+	static_assert(UE_ARRAY_COUNT(UE_STRINGIZE(InName)) <= UE::StallDetectorScopeMaxNameLength, "Scope name must be less than UE::StallDetectorScopeMaxNameLength characters long "); \
 	static UE::FStallDetectorStats UE_JOIN(StallDetectorStats, __LINE__)(TEXT(#InName), InBudgetSeconds, UE::EStallDetectorReportingMode::First);                                    \
 	UE::FStallDetector UE_JOIN(ScopeStallDetector, __LINE__)(UE_JOIN(StallDetectorStats, __LINE__));
 
@@ -48,8 +48,8 @@
  * Use with extreme caution as you could overload the report backend!
  * This mainly exists to help test and exercise the stall detector system.
  **/
-#define SCOPE_STALL_REPORTER_ALWAYS(InName, InBudgetSeconds)                                                                                                                                   \
-	static_assert(UE_ARRAY_COUNT(PREPROCESSOR_TO_STRING(InName)) <= UE::StallDetectorScopeMaxNameLength, "Scope name must be less than UE::StallDetectorScopeMaxNameLength characters long "); \
+#define SCOPE_STALL_REPORTER_ALWAYS(InName, InBudgetSeconds)                                                                                                                         \
+	static_assert(UE_ARRAY_COUNT(UE_STRINGIZE(InName)) <= UE::StallDetectorScopeMaxNameLength, "Scope name must be less than UE::StallDetectorScopeMaxNameLength characters long "); \
 	static UE::FStallDetectorStats UE_JOIN(StallDetectorStats, __LINE__)(TEXT(#InName), InBudgetSeconds, UE::EStallDetectorReportingMode::Always);                                   \
 	UE::FStallDetector UE_JOIN(ScopeStallDetector, __LINE__)(UE_JOIN(StallDetectorStats, __LINE__));
 
diff --git a/Engine/Source/Runtime/Core/Tests/Misc/ParseTests.cpp b/Engine/Source/Runtime/Core/Tests/Misc/ParseTests.cpp
--- a/Engine/Source/Runtime/Core/Tests/Misc/ParseTests.cpp
+++ b/Engine/Source/Runtime/Core/Tests/Misc/ParseTests.cpp
@@ -879,8 +879,8 @@ TEST_CASE("Parse::Value::Numbers", "[Smoke]")
 				{ TEXT("a=-3.1415"), -3.1415 },
 				{ TEXT("a=340282346638528859811704183484516925440.0"), 340282346638528859811704183484516925440.0 },
 				{ TEXT("a=-340282346638528859811704183484516925440.0"), -340282346638528859811704183484516925440.0 },
-				{ TEXT("a=" PREPROCESSOR_TO_STRING(DBL_MAX)), DBL_MAX },
-				{ TEXT("a=-" PREPROCESSOR_TO_STRING(DBL_MAX)), -DBL_MAX },
+				{ TEXT("a=" UE_STRINGIZE(DBL_MAX)), DBL_MAX },
+				{ TEXT("a=-" UE_STRINGIZE(DBL_MAX)), -DBL_MAX },
 				// quoted
 				{ TEXT("a=\"0.0\""), 0.0 },
 				{ TEXT("a=\"0.5\""), 0.5 },
@@ -890,8 +890,8 @@ TEST_CASE("Parse::Value::Numbers", "[Smoke]")
 				{ TEXT("a=\"-3.1415\""), -3.1415 },
 				{ TEXT("a=\"340282346638528859811704183484516925440.0\""), 340282346638528859811704183484516925440.0 },
 				{ TEXT("a=\"-340282346638528859811704183484516925440.0\""), -340282346638528859811704183484516925440.0 },
-				{ TEXT("a=\"" PREPROCESSOR_TO_STRING(DBL_MAX) "\""), DBL_MAX},
-				{ TEXT("a=\"-" PREPROCESSOR_TO_STRING(DBL_MAX) "\""), -DBL_MAX },
+				{ TEXT("a=\"" UE_STRINGIZE(DBL_MAX) "\""), DBL_MAX},
+				{ TEXT("a=\"-" UE_STRINGIZE(DBL_MAX) "\""), -DBL_MAX },
 				// broken quotes takes the first number found
 				{ TEXT("a=123\"456\""), 123.0 },
 				{ TEXT("a=\"123\"456"), 123.0 },
diff --git a/Engine/Source/Runtime/CoreUObject/Public/UObject/Field.h b/Engine/Source/Runtime/CoreUObject/Public/UObject/Field.h
--- a/Engine/Source/Runtime/CoreUObject/Public/UObject/Field.h
+++ b/Engine/Source/Runtime/CoreUObject/Public/UObject/Field.h
@@ -240,7 +240,7 @@ struct FRegisterFieldClass
 #define UE_DEFINE_FIELD_CLASS_STATIC_CONSTINIT(TClass)                                               \
 	constinit TNoDestroy<FFieldClass> TClass::ConstInitClass{ NoDestroyConstEval,                    \
 		                                                      ConstEval,                             \
-		                                                      TEXT(PREPROCESSOR_TO_STRING(TClass)),  \
+		                                                      TEXT(UE_STRINGIZE(TClass)),            \
 		                                                      TClass::StaticClassCastFlagsPrivate(), \
 		                                                      TClass::StaticClassCastFlags(),        \
 		                                                      &TClass::Super::ConstInitClass,        \
@@ -312,19 +312,19 @@ public: \
 #endif
 
 #if UE_WITH_CONSTINIT_UOBJECT
-#define IMPLEMENT_FIELD(TClass)                                                                        \
-	FField* TClass::Construct(const FFieldVariant& InOwner, const FName& InName, EObjectFlags InFlags) \
-	{                                                                                                  \
-		IMPLEMENT_FIELD_CONSTRUCT_IMPLEMENTATION(TClass)                                               \
-	}                                                                                                  \
-	UE_DEFINE_FIELD_CLASS_STATIC_CONSTINIT(TClass) 													   \
-	static FRegisterFieldClass Register_FieldClass_##TClass( 										   \
-		&TClass::ConstInitClass, TEXT(UE_STRINGIZE(TClass))); 				   						   \
-	FFieldClass* TClass::StaticClass()                                                                 \
-	{                                                                                                  \
-		checkf(ConstInitClass->IsFullyConstructed(), 				                				   \
+#define IMPLEMENT_FIELD(TClass)                                                                           \
+	FField* TClass::Construct(const FFieldVariant& InOwner, const FName& InName, EObjectFlags InFlags)    \
+	{                                                                                                     \
+		IMPLEMENT_FIELD_CONSTRUCT_IMPLEMENTATION(TClass)                                                  \
+	}                                                                                                     \
+	UE_DEFINE_FIELD_CLASS_STATIC_CONSTINIT(TClass)                                                        \
+	static FRegisterFieldClass Register_FieldClass_##TClass(                                              \
+		&TClass::ConstInitClass, TEXT(UE_STRINGIZE(TClass)));                                             \
+	FFieldClass* TClass::StaticClass()                                                                    \
+	{                                                                                                     \
+		checkf(ConstInitClass->IsFullyConstructed(),                                                      \
 			UE_STRINGIZE(TClass) TEXT(": FieldClass not fully constructed before ::StaticClass called")); \
-		return &ConstInitClass;                                                                        \
+		return &ConstInitClass;                                                                           \
 	}
 #else // UE_WITH_CONSTINIT_UOBJECT
 #define IMPLEMENT_FIELD(TClass) \
@@ -334,7 +334,7 @@ FField* TClass::Construct(const FFieldVariant& InOwner, const FName& InName, EOb
 } \
 FFieldClass* TClass::StaticClass() \
 { \
-	static FFieldClass StaticFieldClass(TEXT(PREPROCESSOR_TO_STRING(TClass)), TClass::StaticClassCastFlagsPrivate(), TClass::StaticClassCastFlags(), TClass::Super::StaticClass(), &TClass::Construct); \
+	static FFieldClass StaticFieldClass(TEXT(UE_STRINGIZE(TClass)), TClass::StaticClassCastFlagsPrivate(), TClass::StaticClassCastFlags(), TClass::Super::StaticClass(), &TClass::Construct); \
 	return &StaticFieldClass; \
 } 
 #endif // UE_WITH_CONSTINIT_UOBJECT
diff --git a/Engine/Source/Runtime/Engine/Private/Kismet/BlueprintTypeConversions.cpp b/Engine/Source/Runtime/Engine/Private/Kismet/BlueprintTypeConversions.cpp
--- a/Engine/Source/Runtime/Engine/Private/Kismet/BlueprintTypeConversions.cpp
+++ b/Engine/Source/Runtime/Engine/Private/Kismet/BlueprintTypeConversions.cpp
@@ -38,8 +38,8 @@
 		&TVariantStructure<BASETYPE>::Get,																	\
 		&TVariantStructure<DOUBLEVARIANT>::Get,																\
 		&TVariantStructure<FLOATVARIANT>::Get,																\
-		TEXT(PREPROCESSOR_TO_STRING(MAKE_CONVERSION_FUNCTION_NAME(DOUBLEVARIANT, FLOATVARIANT))),			\
-		TEXT(PREPROCESSOR_TO_STRING(MAKE_CONVERSION_FUNCTION_NAME(FLOATVARIANT, DOUBLEVARIANT))),			\
+		TEXT(UE_STRINGIZE(MAKE_CONVERSION_FUNCTION_NAME(DOUBLEVARIANT, FLOATVARIANT))),						\
+		TEXT(UE_STRINGIZE(MAKE_CONVERSION_FUNCTION_NAME(FLOATVARIANT, DOUBLEVARIANT))),						\
 		&UE::Kismet::BlueprintTypeConversions::Internal::ConvertType<DOUBLEVARIANT, FLOATVARIANT>,			\
 		&UE::Kismet::BlueprintTypeConversions::Internal::ConvertType<FLOATVARIANT, DOUBLEVARIANT>			\
 	);
@@ -411,8 +411,8 @@ UE::Kismet::BlueprintTypeConversions::Internal::FStructConversionEntry FBox2DEnt
 	&ReturnNullVariant,
 	&ReturnNullVariant,
 	&TVariantStructure<FBox2f>::Get,
-	TEXT(PREPROCESSOR_TO_STRING(MAKE_CONVERSION_FUNCTION_NAME(FBox2d, FBox2f))),
-	TEXT(PREPROCESSOR_TO_STRING(MAKE_CONVERSION_FUNCTION_NAME(FBox2f, FBox2d))),
+	TEXT(UE_STRINGIZE(MAKE_CONVERSION_FUNCTION_NAME(FBox2d, FBox2f))),
+	TEXT(UE_STRINGIZE(MAKE_CONVERSION_FUNCTION_NAME(FBox2f, FBox2d))),
 	&UE::Kismet::BlueprintTypeConversions::Internal::ConvertType<FBox2d, FBox2f>,
 	&UE::Kismet::BlueprintTypeConversions::Internal::ConvertType<FBox2f, FBox2d>
 );
@@ -446,8 +446,8 @@ UE::Kismet::BlueprintTypeConversions::Internal::FStructConversionEntry FVector2D
 	&ReturnNullVariant,
 	&ReturnNullVariant,
 	&TVariantStructure<FVector2f>::Get,
-	TEXT(PREPROCESSOR_TO_STRING(MAKE_CONVERSION_FUNCTION_NAME(FVector2d, FVector2f))),
-	TEXT(PREPROCESSOR_TO_STRING(MAKE_CONVERSION_FUNCTION_NAME(FVector2f, FVector2d))),
+	TEXT(UE_STRINGIZE(MAKE_CONVERSION_FUNCTION_NAME(FVector2d, FVector2f))),
+	TEXT(UE_STRINGIZE(MAKE_CONVERSION_FUNCTION_NAME(FVector2f, FVector2d))),
 	&UE::Kismet::BlueprintTypeConversions::Internal::ConvertType<FVector2d, FVector2f>,
 	&UE::Kismet::BlueprintTypeConversions::Internal::ConvertType<FVector2f, FVector2d>
 );
@@ -457,8 +457,8 @@ UE::Kismet::BlueprintTypeConversions::Internal::FStructConversionEntry FDeprecat
 	&TVariantStructure<FVector2f>::Get,
 	&TBaseStructure<FVector2D>::Get,
 	&TVariantStructure<FVector2f>::Get,
-	TEXT(PREPROCESSOR_TO_STRING(MAKE_CONVERSION_FUNCTION_NAME(FVector2d, FVector2f))),
-	TEXT(PREPROCESSOR_TO_STRING(MAKE_CONVERSION_FUNCTION_NAME(FVector2f, FVector2d))),
+	TEXT(UE_STRINGIZE(MAKE_CONVERSION_FUNCTION_NAME(FVector2d, FVector2f))),
+	TEXT(UE_STRINGIZE(MAKE_CONVERSION_FUNCTION_NAME(FVector2f, FVector2d))),
 	&UE::Kismet::BlueprintTypeConversions::Internal::ConvertType<FVector2d, FVector2f>,
 	&UE::Kismet::BlueprintTypeConversions::Internal::ConvertType<FVector2f, FVector2d>
 );
\ No newline at end of file
diff --git a/Engine/Source/Runtime/Engine/Private/ShowFlags.cpp b/Engine/Source/Runtime/Engine/Private/ShowFlags.cpp
--- a/Engine/Source/Runtime/Engine/Private/ShowFlags.cpp
+++ b/Engine/Source/Runtime/Engine/Private/ShowFlags.cpp
@@ -211,7 +211,7 @@ int32 FEngineShowFlags::FindIndexByName(const TCHAR* Name, const TCHAR* CommaSep
 		// search through all defined show flags.
 		FString Search = Name;
 
-		#define SHOWFLAG_ALWAYS_ACCESSIBLE(a,...) if(Search == PREPROCESSOR_TO_STRING(a)) { return (int32)SF_##a; }
+		#define SHOWFLAG_ALWAYS_ACCESSIBLE(a,...) if(Search == UE_STRINGIZE(a)) { return (int32)SF_##a; }
 
 		#include "ShowFlagsValues.inl" // IWYU pragma: keep
 
@@ -259,7 +259,7 @@ FString FEngineShowFlags::FindNameByIndex(uint32 InIndex)
 {
 	FString Name;
 
-	#define SHOWFLAG_ALWAYS_ACCESSIBLE(a,...) case SF_##a: Name = PREPROCESSOR_TO_STRING(a); break;
+	#define SHOWFLAG_ALWAYS_ACCESSIBLE(a,...) case SF_##a: Name = UE_STRINGIZE(a); break;
 
 	switch (InIndex)
 	{
@@ -275,7 +275,7 @@ END_FUNCTION_BUILD_OPTIMIZATION
 
 void FEngineShowFlags::AddNameByIndex(uint32 InIndex, FString& Out)
 {
-	#define SHOWFLAG_ALWAYS_ACCESSIBLE(a,...) case SF_##a: Out += PREPROCESSOR_TO_STRING(a); break;
+	#define SHOWFLAG_ALWAYS_ACCESSIBLE(a,...) case SF_##a: Out += UE_STRINGIZE(a); break;
 	switch (InIndex)
 	{
 		#include "ShowFlagsValues.inl" // IWYU pragma: keep
diff --git a/Engine/Source/Runtime/Engine/Public/Net/UnrealNetwork.h b/Engine/Source/Runtime/Engine/Public/Net/UnrealNetwork.h
--- a/Engine/Source/Runtime/Engine/Public/Net/UnrealNetwork.h
+++ b/Engine/Source/Runtime/Engine/Public/Net/UnrealNetwork.h
@@ -516,15 +516,15 @@ ENGINE_API void ResetReplicatedLifetimeProperty(
 
 // This macro is for RPC parameter validation.
 // It handles the details of what should happen if a validation expression fails
-#define RPC_VALIDATE( expression )						\
-	if ( !( expression ) )								\
-	{													\
-		UE_LOG( LogNet, Warning,						\
-		TEXT("RPC_VALIDATE Failed: ")					\
-		TEXT( PREPROCESSOR_TO_STRING( expression ) )	\
-		TEXT(" File: ")									\
-		TEXT( PREPROCESSOR_TO_STRING( __FILE__ ) )		\
-		TEXT(" Line: ")									\
-		TEXT( PREPROCESSOR_TO_STRING( __LINE__ ) ) );	\
-		return false;									\
+#define RPC_VALIDATE( expression )			\
+	if ( !( expression ) )					\
+	{										\
+		UE_LOG( LogNet, Warning,			\
+		TEXT("RPC_VALIDATE Failed: ")		\
+		TEXT( UE_STRINGIZE( expression ) )	\
+		TEXT(" File: ")						\
+		TEXT( UE_STRINGIZE( __FILE__ ) )	\
+		TEXT(" Line: ")						\
+		TEXT( UE_STRINGIZE( __LINE__ ) ) );	\
+		return false;						\
 	}
diff --git a/Engine/Source/Runtime/Engine/Public/ShowFlags.h b/Engine/Source/Runtime/Engine/Public/ShowFlags.h
--- a/Engine/Source/Runtime/Engine/Public/ShowFlags.h
+++ b/Engine/Source/Runtime/Engine/Public/ShowFlags.h
@@ -109,8 +109,8 @@ struct FEngineShowFlags
 
 		if (LocNames.Num() == 0)
 		{
-			#define SHOWFLAG_ALWAYS_ACCESSIBLE(a,b,c) LocNames.Add( TEXT(PREPROCESSOR_TO_STRING(a)), c);
-			#define SHOWFLAG_FIXED_IN_SHIPPING(v,a,b,c) LocNames.Add( TEXT(PREPROCESSOR_TO_STRING(a)), c);
+			#define SHOWFLAG_ALWAYS_ACCESSIBLE(a,b,c) LocNames.Add( TEXT(UE_STRINGIZE(a)), c);
+			#define SHOWFLAG_FIXED_IN_SHIPPING(v,a,b,c) LocNames.Add( TEXT(UE_STRINGIZE(a)), c);
 			#include "ShowFlagsValues.inl"
 
 			// Additional strings that don't correspond to a showflag variable:
@@ -363,7 +363,7 @@ struct FEngineShowFlags
 	 */
 	template <class T> static void IterateAllFlags(T& Sink)
 	{
-		#define SHOWFLAG_ALWAYS_ACCESSIBLE(a,...) if( !Sink.OnEngineShowFlag((uint32)SF_##a, TEXT(PREPROCESSOR_TO_STRING(a))) ) { return; }
+		#define SHOWFLAG_ALWAYS_ACCESSIBLE(a,...) if( !Sink.OnEngineShowFlag((uint32)SF_##a, TEXT(UE_STRINGIZE(a))) ) { return; }
 		#include "ShowFlagsValues.inl"
 
 		IterateCustomFlags([&](uint32 Index, const FString& Name) { return Sink.OnCustomShowFlag(Index, Name); });
diff --git a/Engine/Source/Runtime/MovieScene/Public/MovieSceneFwd.h b/Engine/Source/Runtime/MovieScene/Public/MovieSceneFwd.h
--- a/Engine/Source/Runtime/MovieScene/Public/MovieSceneFwd.h
+++ b/Engine/Source/Runtime/MovieScene/Public/MovieSceneFwd.h
@@ -160,7 +160,7 @@ MOVIESCENE_API FFrameNumber UpgradeLegacyMovieSceneTime(UObject* ErrorContext, F
 #endif
 
 #ifdef UE_MOVIESCENE_TODOS
-	#define UE_MOVIESCENE_TODO(MSG) UE_MOVIESCENE_TODO_IMPL(message("" __FILE__ "(" PREPROCESSOR_TO_STRING(__LINE__) "): warning TODO: " # MSG))
+	#define UE_MOVIESCENE_TODO(MSG) UE_MOVIESCENE_TODO_IMPL(message("" __FILE__ "(" UE_STRINGIZE(__LINE__) "): warning TODO: " # MSG))
 #else
 	#define UE_MOVIESCENE_TODO(MSG)
 #endif
\ No newline at end of file
diff --git a/Engine/Source/Runtime/Net/Iris/Public/Iris/Core/IrisProfiler.h b/Engine/Source/Runtime/Net/Iris/Public/Iris/Core/IrisProfiler.h
--- a/Engine/Source/Runtime/Net/Iris/Public/Iris/Core/IrisProfiler.h
+++ b/Engine/Source/Runtime/Net/Iris/Public/Iris/Core/IrisProfiler.h
@@ -35,8 +35,8 @@
 #		include "c:/Program Files/Superluminal/Performance/API/include/Superluminal/PerformanceAPI.h"
 #		include "HAL/PreprocessorHelpers.h"
 #		pragma comment (lib, "c:/Program Files/Superluminal/Performance/API/lib/x64/PerformanceAPI_MD.lib")
-#		define IRIS_PROFILER_SCOPE(X) PERFORMANCEAPI_INSTRUMENT(PREPROCESSOR_TO_STRING(X))
-#		define IRIS_PROFILER_SCOPE_CONDITIONAL(X,Cond) PERFORMANCEAPI_INSTRUMENT(PREPROCESSOR_TO_STRING(X))
+#		define IRIS_PROFILER_SCOPE(X) PERFORMANCEAPI_INSTRUMENT(UE_STRINGIZE(X))
+#		define IRIS_PROFILER_SCOPE_CONDITIONAL(X,Cond) PERFORMANCEAPI_INSTRUMENT(UE_STRINGIZE(X))
 #		define IRIS_PROFILER_SCOPE_TEXT(X) PERFORMANCEAPI_INSTRUMENT_DATA(UE_JOIN(IrisProfilerScope, __LINE__), X)
 #		define IRIS_PROFILER_SCOPE_TEXT_CONDITIONAL(X, Cond) PERFORMANCEAPI_INSTRUMENT_DATA(UE_JOIN(IrisProfilerScope, __LINE__), X)
 #	else
diff --git a/Engine/Source/Runtime/Net/Iris/Public/Iris/ReplicationState/PropertyNetSerializerInfoRegistry.h b/Engine/Source/Runtime/Net/Iris/Public/Iris/ReplicationState/PropertyNetSerializerInfoRegistry.h
--- a/Engine/Source/Runtime/Net/Iris/Public/Iris/ReplicationState/PropertyNetSerializerInfoRegistry.h
+++ b/Engine/Source/Runtime/Net/Iris/Public/Iris/ReplicationState/PropertyNetSerializerInfoRegistry.h
@@ -263,7 +263,7 @@ static F##Name##NetSerializerRegistryDelegates Name##NetSerializerRegistryDelega
 
 // Utility that can be used to forward serialization of a Struct to a specific NetSerializer
 #define UE_NET_IMPLEMENT_FORWARDING_NETSERIALIZER_AND_REGISTRY_DELEGATES(Name, SerializerName) \
-	static const FName PropertyNetSerializerRegistry_NAME_##Name( PREPROCESSOR_TO_STRING(Name) ); \
+	static const FName PropertyNetSerializerRegistry_NAME_##Name( UE_STRINGIZE(Name) ); \
 	UE_NET_IMPLEMENT_NAMED_STRUCT_NETSERIALIZER_INFO(PropertyNetSerializerRegistry_NAME_##Name, SerializerName); \
 	UE_NET_IMPLEMENT_NETSERIALIZER_REGISTRY_DELEGATES(Name)
 
@@ -271,13 +271,13 @@ static F##Name##NetSerializerRegistryDelegates Name##NetSerializerRegistryDelega
 // Utility that can be used to forward serialization of a Struct to a last resort net serializer
 #define UE_NET_IMPLEMENT_NAMED_STRUCT_LASTRESORT_NETSERIALIZER_AND_REGISTRY_DELEGATES(Name) \
 	UE_NET_IS_DECLARED_TYPE(Name); \
-	static const FName PropertyNetSerializerRegistry_NAME_##Name( PREPROCESSOR_TO_STRING(Name) ); \
+	static const FName PropertyNetSerializerRegistry_NAME_##Name( UE_STRINGIZE(Name) ); \
 	UE_NET_IMPLEMENT_NAMED_STRUCT_LASTRESORT_NETSERIALIZER_INFO(PropertyNetSerializerRegistry_NAME_##Name); \
 	UE_NET_IMPLEMENT_NETSERIALIZER_REGISTRY_DELEGATES(Name)
 
 // Utility that can be used to forward serialization of a Struct to a last resort net serializer with a maximum quantized size
 #define UE_NET_IMPLEMENT_NAMED_STRUCT_LASTRESORT_NETSERIALIZER_AND_REGISTRY_DELEGATES_WITH_SIZE_OVERRIDE(Name, MaxBits) \
 	UE_NET_IS_DECLARED_TYPE(Name); \
-	static const FName PropertyNetSerializerRegistry_NAME_##Name( PREPROCESSOR_TO_STRING(Name) ); \
+	static const FName PropertyNetSerializerRegistry_NAME_##Name( UE_STRINGIZE(Name) ); \
 	UE_NET_IMPLEMENT_NAMED_STRUCT_LASTRESORT_NETSERIALIZER_INFO_WITH_SIZE_OVERRIDE(PropertyNetSerializerRegistry_NAME_##Name, MaxBits); \
 	UE_NET_IMPLEMENT_NETSERIALIZER_REGISTRY_DELEGATES(Name)
diff --git a/Engine/Source/Runtime/Net/Iris/Public/Iris/ReplicationSystem/NetTokenStructDefines.h b/Engine/Source/Runtime/Net/Iris/Public/Iris/ReplicationSystem/NetTokenStructDefines.h
--- a/Engine/Source/Runtime/Net/Iris/Public/Iris/ReplicationSystem/NetTokenStructDefines.h
+++ b/Engine/Source/Runtime/Net/Iris/Public/Iris/ReplicationSystem/NetTokenStructDefines.h
@@ -30,7 +30,7 @@ struct TNetTokenStructOpsTypeTraits : public TStructOpsTypeTraitsBase2<T>
 
 // Optional, declares a default native NetSerializer method, GetTokenStoreName method and Equality operators (for equality via identity) for a NetToken Struct type
 #define UE_NET_NETTOKEN_GENERATED_BODY(NAME, API) \
-	inline static FName TokenStoreName = TEXT( PREPROCESSOR_TO_STRING(F##NAME) );\
+	inline static FName TokenStoreName = TEXT( UE_STRINGIZE(F##NAME) );\
 	static FName GetTokenStoreName() \
 	{ \
 		return TokenStoreName; \
diff --git a/Engine/Source/Runtime/PakFile/Internal/IoDispatcherFileBackendTypes.h b/Engine/Source/Runtime/PakFile/Internal/IoDispatcherFileBackendTypes.h
--- a/Engine/Source/Runtime/PakFile/Internal/IoDispatcherFileBackendTypes.h
+++ b/Engine/Source/Runtime/PakFile/Internal/IoDispatcherFileBackendTypes.h
@@ -14,7 +14,7 @@
 #endif
 
 #ifndef PLATFORM_IODISPATCHER_MODULE
-#define PLATFORM_IODISPATCHER_MODULE PREPROCESSOR_TO_STRING(UE_JOIN(PLATFORM_HEADER_NAME, PlatformIoDispatcher))
+#define PLATFORM_IODISPATCHER_MODULE UE_STRINGIZE(UE_JOIN(PLATFORM_HEADER_NAME, PlatformIoDispatcher))
 #endif
 
 #define UE_FILEIOSTORE_STATS_ENABLED (COUNTERSTRACE_ENABLED || CSV_PROFILER_STATS)
diff --git a/Engine/Source/Runtime/PhysicsCore/Private/PhysXLibs.cpp b/Engine/Source/Runtime/PhysicsCore/Private/PhysXLibs.cpp
--- a/Engine/Source/Runtime/PhysicsCore/Private/PhysXLibs.cpp
+++ b/Engine/Source/Runtime/PhysicsCore/Private/PhysXLibs.cpp
@@ -66,26 +66,26 @@ namespace PhysDLLHelper
 #endif
 
 #ifdef UE_PHYSX_SUFFIX
-	FString PhysXSuffix(TEXT(PREPROCESSOR_TO_STRING(UE_PHYSX_SUFFIX)) + ArchName + TEXT(".dll"));
+	FString PhysXSuffix(TEXT(UE_STRINGIZE(UE_PHYSX_SUFFIX)) + ArchName + TEXT(".dll"));
 #else
 	FString PhysXSuffix(ArchName + TEXT(".dll"));
 #endif
 
 #ifdef UE_APEX_SUFFIX
-	FString APEXSuffix(TEXT(PREPROCESSOR_TO_STRING(UE_APEX_SUFFIX)) + ArchName + TEXT(".dll"));
+	FString APEXSuffix(TEXT(UE_STRINGIZE(UE_APEX_SUFFIX)) + ArchName + TEXT(".dll"));
 #else
 	FString APEXSuffix(ArchName + TEXT(".dll"));
 #endif
 #elif PLATFORM_MAC
 	FString PhysXBinariesRoot = TEXT("Binaries/ThirdParty/PhysX3/Mac/");
 #ifdef UE_PHYSX_SUFFIX
-	FString PhysXSuffix = FString(TEXT(PREPROCESSOR_TO_STRING(UE_PHYSX_SUFFIX))) + TEXT(".dylib");
+	FString PhysXSuffix = FString(TEXT(UE_STRINGIZE(UE_PHYSX_SUFFIX))) + TEXT(".dylib");
 #else
 	FString PhysXSuffix(TEXT(".dylib"));
 #endif
 
 #ifdef UE_APEX_SUFFIX
-	FString APEXSuffix = FString(TEXT(PREPROCESSOR_TO_STRING(UE_APEX_SUFFIX))) + TEXT(".dylib");
+	FString APEXSuffix = FString(TEXT(UE_STRINGIZE(UE_APEX_SUFFIX))) + TEXT(".dylib");
 #else
 	FString APEXSuffix(TEXT(".dylib"));
 #endif
diff --git a/Engine/Source/Runtime/RenderCore/Public/Shader.h b/Engine/Source/Runtime/RenderCore/Public/Shader.h
--- a/Engine/Source/Runtime/RenderCore/Public/Shader.h
+++ b/Engine/Source/Runtime/RenderCore/Public/Shader.h
@@ -2005,23 +2005,23 @@ protected:
 #if !UE_BUILD_DOCS
 // Vertex+Pixel
 #define IMPLEMENT_SHADERPIPELINE_TYPE_VSPS(PipelineName, VertexShaderType, PixelShaderType, bRemoveUnused)	\
-	static FShaderPipelineType PipelineName(TEXT(PREPROCESSOR_TO_STRING(PipelineName)), &VertexShaderType::GetStaticType(), nullptr, &PixelShaderType::GetStaticType(), false, bRemoveUnused);
+	static FShaderPipelineType PipelineName(TEXT(UE_STRINGIZE(PipelineName)), &VertexShaderType::GetStaticType(), nullptr, &PixelShaderType::GetStaticType(), false, bRemoveUnused);
 // Only VS
 #define IMPLEMENT_SHADERPIPELINE_TYPE_VS(PipelineName, VertexShaderType, bRemoveUnused)	\
-	static FShaderPipelineType PipelineName(TEXT(PREPROCESSOR_TO_STRING(PipelineName)), &VertexShaderType::GetStaticType(), nullptr, nullptr, false, bRemoveUnused);
+	static FShaderPipelineType PipelineName(TEXT(UE_STRINGIZE(PipelineName)), &VertexShaderType::GetStaticType(), nullptr, nullptr, false, bRemoveUnused);
 // Vertex+Geometry+Pixel
 #define IMPLEMENT_SHADERPIPELINE_TYPE_VSGSPS(PipelineName, VertexShaderType, GeometryShaderType, PixelShaderType, bRemoveUnused)	\
-	static FShaderPipelineType PipelineName(TEXT(PREPROCESSOR_TO_STRING(PipelineName)), &VertexShaderType::GetStaticType(), &GeometryShaderType::GetStaticType(), &PixelShaderType::GetStaticType(), false, bRemoveUnused);
+	static FShaderPipelineType PipelineName(TEXT(UE_STRINGIZE(PipelineName)), &VertexShaderType::GetStaticType(), &GeometryShaderType::GetStaticType(), &PixelShaderType::GetStaticType(), false, bRemoveUnused);
 // Vertex+Geometry
 #define IMPLEMENT_SHADERPIPELINE_TYPE_VSGS(PipelineName, VertexShaderType, GeometryShaderType, bRemoveUnused)	\
-	static FShaderPipelineType PipelineName(TEXT(PREPROCESSOR_TO_STRING(PipelineName)), &VertexShaderType::GetStaticType(), &GeometryShaderType::GetStaticType(), nullptr, false, bRemoveUnused);
+	static FShaderPipelineType PipelineName(TEXT(UE_STRINGIZE(PipelineName)), &VertexShaderType::GetStaticType(), &GeometryShaderType::GetStaticType(), nullptr, false, bRemoveUnused);
 
 // Mesh+Pixel
 #define IMPLEMENT_SHADERPIPELINE_TYPE_MSPS(PipelineName, MeshShaderType, PixelShaderType, bRemoveUnused)	\
-	static FShaderPipelineType PipelineName(TEXT(PREPROCESSOR_TO_STRING(PipelineName)), &MeshShaderType::GetStaticType(), nullptr, &PixelShaderType::GetStaticType(), true, bRemoveUnused);
+	static FShaderPipelineType PipelineName(TEXT(UE_STRINGIZE(PipelineName)), &MeshShaderType::GetStaticType(), nullptr, &PixelShaderType::GetStaticType(), true, bRemoveUnused);
 // Mesh+Amplification+Pixel
 #define IMPLEMENT_SHADERPIPELINE_TYPE_MSASPS(PipelineName, MeshShaderType, AmplificationShaderType, PixelShaderType, bRemoveUnused)	\
-	static FShaderPipelineType PipelineName(TEXT(PREPROCESSOR_TO_STRING(PipelineName)), &MeshShaderType::GetStaticType(), &AmplificationShaderType::GetStaticType(), &PixelShaderType::GetStaticType(), true, bRemoveUnused);
+	static FShaderPipelineType PipelineName(TEXT(UE_STRINGIZE(PipelineName)), &MeshShaderType::GetStaticType(), &AmplificationShaderType::GetStaticType(), &PixelShaderType::GetStaticType(), true, bRemoveUnused);
 #endif
 
 /** Encapsulates a dependency on a shader type and saved state from that shader type. */
diff --git a/Engine/Source/Runtime/SignalProcessing/Public/DSP/BufferDiagnostics.h b/Engine/Source/Runtime/SignalProcessing/Public/DSP/BufferDiagnostics.h
--- a/Engine/Source/Runtime/SignalProcessing/Public/DSP/BufferDiagnostics.h
+++ b/Engine/Source/Runtime/SignalProcessing/Public/DSP/BufferDiagnostics.h
@@ -129,7 +129,7 @@ namespace Audio
 
 #if WITH_AUDIO_BUFFERDIAGNOSTICS
 #define AUDIO_CHECK_BUFFER(BUFFER)\
-	static FCheckedBufferState UE_JOIN(BufferCheck, __LINE__)(__LINE__,TEXT(__FILE__),TEXT(PREPROCESSOR_TO_STRING(BUFFER)));\
+	static FCheckedBufferState UE_JOIN(BufferCheck, __LINE__)(__LINE__,TEXT(__FILE__),TEXT(UE_STRINGIZE(BUFFER)));\
 	if (!UE_JOIN(BufferCheck, __LINE__).DoCheck(BUFFER) )\
 	{\
 		UE_JOIN(BufferCheck, __LINE__).FailedBufferCheckImpl(TEXT(""));\
diff --git a/Engine/Source/Runtime/Slate/Private/Framework/Docking/LayoutService.cpp b/Engine/Source/Runtime/Slate/Private/Framework/Docking/LayoutService.cpp
--- a/Engine/Source/Runtime/Slate/Private/Framework/Docking/LayoutService.cpp
+++ b/Engine/Source/Runtime/Slate/Private/Framework/Docking/LayoutService.cpp
@@ -170,7 +170,7 @@ static FString GetLayoutJsonFileName(const FString& InConfigFileName)
 {
 	const FString JsonFileName = FPaths::GetBaseFilename(InConfigFileName) + TEXT(".json");
 #ifdef UE_SAVED_DIR_OVERRIDE
-	const FString UserSettingsPath = FPaths::Combine(FPlatformProcess::UserSettingsDir(), TEXT(PREPROCESSOR_TO_STRING(UE_SAVED_DIR_OVERRIDE)), TEXT("Editor"), JsonFileName);
+	const FString UserSettingsPath = FPaths::Combine(FPlatformProcess::UserSettingsDir(), TEXT(UE_STRINGIZE(UE_SAVED_DIR_OVERRIDE)), TEXT("Editor"), JsonFileName);
 #else
 	const FString UserSettingsPath = FPaths::Combine(FPlatformProcess::UserSettingsDir(), FApp::GetEpicProductIdentifier(), TEXT("Editor"), JsonFileName);
 #endif
diff --git a/Engine/Source/Runtime/StorageServerClient/Private/StorageServerPlatformFile.cpp b/Engine/Source/Runtime/StorageServerClient/Private/StorageServerPlatformFile.cpp
--- a/Engine/Source/Runtime/StorageServerClient/Private/StorageServerPlatformFile.cpp
+++ b/Engine/Source/Runtime/StorageServerClient/Private/StorageServerPlatformFile.cpp
@@ -887,7 +887,7 @@ FStorageServerPlatformFile::FStorageServerPlatformFile()
 	if (UE::IsUsingZenPakFileStreaming())
 	{
 		ServerEngineDirView = FStringView(TEXT("Engine/"));
-		ServerProjectDirView = FStringView(TEXT(PREPROCESSOR_TO_STRING(UE_PROJECT_NAME)) TEXT("/"));
+		ServerProjectDirView = FStringView(TEXT(UE_STRINGIZE(UE_PROJECT_NAME)) TEXT("/"));
 	}
 }
 
