Replaced PREPROCESSOR_REMOVE_OPTIONAL_PARENS with UE_REMOVE_OPTIONAL_PARENS.

diff --git a/Engine/Source/Runtime/Core/Public/Misc/AssertionMacros.h b/Engine/Source/Runtime/Core/Public/Misc/AssertionMacros.h
--- a/Engine/Source/Runtime/Core/Public/Misc/AssertionMacros.h
+++ b/Engine/Source/Runtime/Core/Public/Misc/AssertionMacros.h
@@ -524,16 +524,16 @@ namespace UEAsserts_Private
 	(FName(GET_FUNCTION_NAME_STRING_CHECKED_FourParams(ClassName, FunctionName, ArgType1, ArgType2, ArgType3, ArgType4)))
 
 #define GET_FUNCTION_NAME_STRING_CHECKED_OneParam(ClassName, FunctionName, ArgType) \
-	((void)sizeof((std::declval<ClassName&>().FunctionName(std::declval<PREPROCESSOR_REMOVE_OPTIONAL_PARENS(ArgType)>()), (int)0)), TEXT(#FunctionName))
+	((void)sizeof((std::declval<ClassName&>().FunctionName(std::declval<UE_REMOVE_OPTIONAL_PARENS(ArgType)>()), (int)0)), TEXT(#FunctionName))
 
 #define GET_FUNCTION_NAME_STRING_CHECKED_TwoParams(ClassName, FunctionName, ArgType1, ArgType2) \
-	((void)sizeof((std::declval<ClassName&>().FunctionName(std::declval<PREPROCESSOR_REMOVE_OPTIONAL_PARENS(ArgType1)>(), std::declval<PREPROCESSOR_REMOVE_OPTIONAL_PARENS(ArgType2)>()), (int)0)), TEXT(#FunctionName))
+	((void)sizeof((std::declval<ClassName&>().FunctionName(std::declval<UE_REMOVE_OPTIONAL_PARENS(ArgType1)>(), std::declval<UE_REMOVE_OPTIONAL_PARENS(ArgType2)>()), (int)0)), TEXT(#FunctionName))
 
 #define GET_FUNCTION_NAME_STRING_CHECKED_ThreeParams(ClassName, FunctionName, ArgType1, ArgType2, ArgType3) \
-	((void)sizeof((std::declval<ClassName&>().FunctionName(std::declval<PREPROCESSOR_REMOVE_OPTIONAL_PARENS(ArgType1)>(), std::declval<PREPROCESSOR_REMOVE_OPTIONAL_PARENS(ArgType2)>(), std::declval<PREPROCESSOR_REMOVE_OPTIONAL_PARENS(ArgType3)>()), (int)0)), TEXT(#FunctionName))
+	((void)sizeof((std::declval<ClassName&>().FunctionName(std::declval<UE_REMOVE_OPTIONAL_PARENS(ArgType1)>(), std::declval<UE_REMOVE_OPTIONAL_PARENS(ArgType2)>(), std::declval<UE_REMOVE_OPTIONAL_PARENS(ArgType3)>()), (int)0)), TEXT(#FunctionName))
 
 #define GET_FUNCTION_NAME_STRING_CHECKED_FourParams(ClassName, FunctionName, ArgType1, ArgType2, ArgType3, ArgType4) \
-	((void)sizeof((std::declval<ClassName&>().FunctionName(std::declval<PREPROCESSOR_REMOVE_OPTIONAL_PARENS(ArgType1)>(), std::declval<PREPROCESSOR_REMOVE_OPTIONAL_PARENS(ArgType2)>(), std::declval<PREPROCESSOR_REMOVE_OPTIONAL_PARENS(ArgType3)>(), std::declval<PREPROCESSOR_REMOVE_OPTIONAL_PARENS(ArgType4)>()), (int)0)), TEXT(#FunctionName))
+	((void)sizeof((std::declval<ClassName&>().FunctionName(std::declval<UE_REMOVE_OPTIONAL_PARENS(ArgType1)>(), std::declval<UE_REMOVE_OPTIONAL_PARENS(ArgType2)>(), std::declval<UE_REMOVE_OPTIONAL_PARENS(ArgType3)>(), std::declval<UE_REMOVE_OPTIONAL_PARENS(ArgType4)>()), (int)0)), TEXT(#FunctionName))
 
 /*----------------------------------------------------------------------------
 	Low level error macros
diff --git a/Engine/Source/Runtime/Core/Public/Misc/DefinePrivateMemberPtr.h b/Engine/Source/Runtime/Core/Public/Misc/DefinePrivateMemberPtr.h
--- a/Engine/Source/Runtime/Core/Public/Misc/DefinePrivateMemberPtr.h
+++ b/Engine/Source/Runtime/Core/Public/Misc/DefinePrivateMemberPtr.h
@@ -56,5 +56,5 @@ namespace UE::Core::Private
 // (Stuff.*GPrivateStuffLogVal)(); // Logs: "Val: 7"
 //
 #define UE_DEFINE_PRIVATE_MEMBER_PTR(Type, Name, Class, Member) \
-    TIdentity_T<PREPROCESSOR_REMOVE_OPTIONAL_PARENS(Type)> PREPROCESSOR_REMOVE_OPTIONAL_PARENS(Class)::* Name; \
-    template struct UE::Core::Private::TPrivateAccess<&Name, &PREPROCESSOR_REMOVE_OPTIONAL_PARENS(Class)::Member>
+    TIdentity_T<UE_REMOVE_OPTIONAL_PARENS(Type)> UE_REMOVE_OPTIONAL_PARENS(Class)::* Name; \
+    template struct UE::Core::Private::TPrivateAccess<&Name, &UE_REMOVE_OPTIONAL_PARENS(Class)::Member>
diff --git a/Engine/Source/Runtime/Core/Public/Serialization/MemoryLayout.h b/Engine/Source/Runtime/Core/Public/Serialization/MemoryLayout.h
--- a/Engine/Source/Runtime/Core/Public/Serialization/MemoryLayout.h
+++ b/Engine/Source/Runtime/Core/Public/Serialization/MemoryLayout.h
@@ -349,7 +349,7 @@ static void InternalInitializeBaseHelper(FTypeLayoutDesc& TypeDesc)
 	FieldDesc.Name = TEXT("BASE");
 	FieldDesc.UFieldNameLength = 4;
 	FieldDesc.Type = &StaticGetTypeLayoutDesc<BaseType>();
-	FieldDesc.WriteFrozenMemoryImageFunc = TGetFreezeImageFieldHelper<PREPROCESSOR_REMOVE_OPTIONAL_PARENS(T)>::Do();
+	FieldDesc.WriteFrozenMemoryImageFunc = TGetFreezeImageFieldHelper<UE_REMOVE_OPTIONAL_PARENS(T)>::Do();
 	FieldDesc.Offset = GetBaseOffset<T, BaseType>();
 	FieldDesc.NumArray = 1u;
 	FieldDesc.Next = TypeDesc.Fields;
@@ -408,8 +408,8 @@ namespace Freeze
 			FFieldLayoutDesc& FieldDesc = *(FFieldLayoutDesc*)FieldBuffer; \
 			FieldDesc.Name = TEXT(#InName); \
 			FieldDesc.UFieldNameLength = Freeze::FindFieldNameLength(FieldDesc.Name); \
-			FieldDesc.Type = &StaticGetTypeLayoutDesc<PREPROCESSOR_REMOVE_OPTIONAL_PARENS(T)>(); \
-			FieldDesc.WriteFrozenMemoryImageFunc = TGetFreezeImageFieldHelper<PREPROCESSOR_REMOVE_OPTIONAL_PARENS(T)>::Do(); \
+			FieldDesc.Type = &StaticGetTypeLayoutDesc<UE_REMOVE_OPTIONAL_PARENS(T)>(); \
+			FieldDesc.WriteFrozenMemoryImageFunc = TGetFreezeImageFieldHelper<UE_REMOVE_OPTIONAL_PARENS(T)>::Do(); \
 			FieldDesc.Offset = InOffset; \
 			FieldDesc.NumArray = InNumArray; \
 			FieldDesc.Flags = InFlags; \
@@ -425,7 +425,7 @@ namespace Freeze
 	UE_DECLARE_INTERNAL_LINK_SPECIALIZATION(InternalLinkType, Counter - CounterBase) { \
 		UE_STATIC_ONLY(InternalLinkType); \
 		static void CallWriteFrozenField(FMemoryImageWriter& Writer, const void* Object, const void* FieldObject, const FTypeLayoutDesc& TypeDesc, const FTypeLayoutDesc& DerivedTypeDesc) { \
-			static_cast<const DerivedType*>(Object)->InFunc(Writer, *(PREPROCESSOR_REMOVE_OPTIONAL_PARENS(T)*)FieldObject); \
+			static_cast<const DerivedType*>(Object)->InFunc(Writer, *(UE_REMOVE_OPTIONAL_PARENS(T)*)FieldObject); \
 		} \
 		static void Initialize(FTypeLayoutDesc& TypeDesc) { \
 			InternalLinkType<Counter - CounterBase + 1>::Initialize(TypeDesc); \
@@ -433,7 +433,7 @@ namespace Freeze
 			FFieldLayoutDesc& FieldDesc = *(FFieldLayoutDesc*)FieldBuffer; \
 			FieldDesc.Name = TEXT(#InName); \
 			FieldDesc.UFieldNameLength = Freeze::FindFieldNameLength(FieldDesc.Name); \
-			FieldDesc.Type = &StaticGetTypeLayoutDesc<PREPROCESSOR_REMOVE_OPTIONAL_PARENS(T)>(); \
+			FieldDesc.Type = &StaticGetTypeLayoutDesc<UE_REMOVE_OPTIONAL_PARENS(T)>(); \
 			FieldDesc.WriteFrozenMemoryImageFunc = &CallWriteFrozenField; \
 			FieldDesc.Offset = STRUCT_OFFSET(DerivedType, InName); \
 			FieldDesc.NumArray = 1u; \
@@ -468,22 +468,22 @@ PRAGMA_ENABLE_DEPRECATION_WARNINGS
 		} \
 	};
 
-#define LAYOUT_FIELD(T, Name, ...) PREPROCESSOR_REMOVE_OPTIONAL_PARENS(T) Name; INTERNAL_LAYOUT_FIELD(T, Name, STRUCT_OFFSET(DerivedType, Name), EFieldLayoutFlags::MakeFlags(__VA_ARGS__), 1u, 0u, __COUNTER__)
-#define LAYOUT_MUTABLE_FIELD(T, Name, ...) mutable PREPROCESSOR_REMOVE_OPTIONAL_PARENS(T) Name; INTERNAL_LAYOUT_FIELD(T, Name, STRUCT_OFFSET(DerivedType, Name), EFieldLayoutFlags::MakeFlags(__VA_ARGS__), 1u, 0u, __COUNTER__)
-#define LAYOUT_FIELD_INITIALIZED(T, Name, Value, ...) PREPROCESSOR_REMOVE_OPTIONAL_PARENS(T) Name = Value; INTERNAL_LAYOUT_FIELD(T, Name, STRUCT_OFFSET(DerivedType, Name), EFieldLayoutFlags::MakeFlags(__VA_ARGS__), 1u, 0u, __COUNTER__)
-#define LAYOUT_MUTABLE_FIELD_INITIALIZED(T, Name, Value, ...) mutable PREPROCESSOR_REMOVE_OPTIONAL_PARENS(T) Name = Value; INTERNAL_LAYOUT_FIELD(T, Name, STRUCT_OFFSET(DerivedType, Name), EFieldLayoutFlags::MakeFlags(__VA_ARGS__), 1u, 0u, __COUNTER__)
-#define LAYOUT_ARRAY(T, Name, NumArray, ...) PREPROCESSOR_REMOVE_OPTIONAL_PARENS(T) Name[NumArray]; INTERNAL_LAYOUT_FIELD(T, Name, STRUCT_OFFSET(DerivedType, Name), EFieldLayoutFlags::MakeFlags(__VA_ARGS__), NumArray, 0u, __COUNTER__)
-#define LAYOUT_MUTABLE_BITFIELD(T, Name, BitFieldSize, ...) mutable PREPROCESSOR_REMOVE_OPTIONAL_PARENS(T) Name : BitFieldSize; INTERNAL_LAYOUT_FIELD(T, Name, ~0u, EFieldLayoutFlags::MakeFlags(__VA_ARGS__), 1u, BitFieldSize, __COUNTER__)
-#define LAYOUT_BITFIELD(T, Name, BitFieldSize, ...) PREPROCESSOR_REMOVE_OPTIONAL_PARENS(T) Name : BitFieldSize; INTERNAL_LAYOUT_FIELD(T, Name, ~0u, EFieldLayoutFlags::MakeFlags(__VA_ARGS__), 1u, BitFieldSize, __COUNTER__)
-#define LAYOUT_FIELD_WITH_WRITER(T, Name, Func) PREPROCESSOR_REMOVE_OPTIONAL_PARENS(T) Name; INTERNAL_LAYOUT_FIELD_WITH_WRITER(T, Name, Func, __COUNTER__)
-#define LAYOUT_MUTABLE_FIELD_WITH_WRITER(T, Name, Func) mutable PREPROCESSOR_REMOVE_OPTIONAL_PARENS(T) Name; INTERNAL_LAYOUT_FIELD_WITH_WRITER(T, Name, Func, __COUNTER__)
+#define LAYOUT_FIELD(T, Name, ...) UE_REMOVE_OPTIONAL_PARENS(T) Name; INTERNAL_LAYOUT_FIELD(T, Name, STRUCT_OFFSET(DerivedType, Name), EFieldLayoutFlags::MakeFlags(__VA_ARGS__), 1u, 0u, __COUNTER__)
+#define LAYOUT_MUTABLE_FIELD(T, Name, ...) mutable UE_REMOVE_OPTIONAL_PARENS(T) Name; INTERNAL_LAYOUT_FIELD(T, Name, STRUCT_OFFSET(DerivedType, Name), EFieldLayoutFlags::MakeFlags(__VA_ARGS__), 1u, 0u, __COUNTER__)
+#define LAYOUT_FIELD_INITIALIZED(T, Name, Value, ...) UE_REMOVE_OPTIONAL_PARENS(T) Name = Value; INTERNAL_LAYOUT_FIELD(T, Name, STRUCT_OFFSET(DerivedType, Name), EFieldLayoutFlags::MakeFlags(__VA_ARGS__), 1u, 0u, __COUNTER__)
+#define LAYOUT_MUTABLE_FIELD_INITIALIZED(T, Name, Value, ...) mutable UE_REMOVE_OPTIONAL_PARENS(T) Name = Value; INTERNAL_LAYOUT_FIELD(T, Name, STRUCT_OFFSET(DerivedType, Name), EFieldLayoutFlags::MakeFlags(__VA_ARGS__), 1u, 0u, __COUNTER__)
+#define LAYOUT_ARRAY(T, Name, NumArray, ...) UE_REMOVE_OPTIONAL_PARENS(T) Name[NumArray]; INTERNAL_LAYOUT_FIELD(T, Name, STRUCT_OFFSET(DerivedType, Name), EFieldLayoutFlags::MakeFlags(__VA_ARGS__), NumArray, 0u, __COUNTER__)
+#define LAYOUT_MUTABLE_BITFIELD(T, Name, BitFieldSize, ...) mutable UE_REMOVE_OPTIONAL_PARENS(T) Name : BitFieldSize; INTERNAL_LAYOUT_FIELD(T, Name, ~0u, EFieldLayoutFlags::MakeFlags(__VA_ARGS__), 1u, BitFieldSize, __COUNTER__)
+#define LAYOUT_BITFIELD(T, Name, BitFieldSize, ...) UE_REMOVE_OPTIONAL_PARENS(T) Name : BitFieldSize; INTERNAL_LAYOUT_FIELD(T, Name, ~0u, EFieldLayoutFlags::MakeFlags(__VA_ARGS__), 1u, BitFieldSize, __COUNTER__)
+#define LAYOUT_FIELD_WITH_WRITER(T, Name, Func) UE_REMOVE_OPTIONAL_PARENS(T) Name; INTERNAL_LAYOUT_FIELD_WITH_WRITER(T, Name, Func, __COUNTER__)
+#define LAYOUT_MUTABLE_FIELD_WITH_WRITER(T, Name, Func) mutable UE_REMOVE_OPTIONAL_PARENS(T) Name; INTERNAL_LAYOUT_FIELD_WITH_WRITER(T, Name, Func, __COUNTER__)
 #define LAYOUT_WRITE_MEMORY_IMAGE(Func) INTERNAL_LAYOUT_WRITE_MEMORY_IMAGE(Func, __COUNTER__)
 #define LAYOUT_TOSTRING(Func) INTERNAL_LAYOUT_TOSTRING(Func, __COUNTER__)
 
 #if WITH_EDITORONLY_DATA
-#define LAYOUT_FIELD_EDITORONLY(T, Name, ...) PREPROCESSOR_REMOVE_OPTIONAL_PARENS(T) Name; INTERNAL_LAYOUT_FIELD(T, Name, STRUCT_OFFSET(DerivedType, Name), EFieldLayoutFlags::MakeFlagsEditorOnly(__VA_ARGS__), 1u, 0u, __COUNTER__)
-#define LAYOUT_ARRAY_EDITORONLY(T, Name, NumArray, ...) PREPROCESSOR_REMOVE_OPTIONAL_PARENS(T) Name[NumArray]; INTERNAL_LAYOUT_FIELD(T, Name, STRUCT_OFFSET(DerivedType, Name), EFieldLayoutFlags::MakeFlagsEditorOnly(__VA_ARGS__), NumArray, 0u, __COUNTER__)
-#define LAYOUT_BITFIELD_EDITORONLY(T, Name, BitFieldSize, ...) PREPROCESSOR_REMOVE_OPTIONAL_PARENS(T) Name : BitFieldSize; INTERNAL_LAYOUT_FIELD(T, Name, ~0u, EFieldLayoutFlags::MakeFlagsEditorOnly(__VA_ARGS__), 1u, BitFieldSize, __COUNTER__)
+#define LAYOUT_FIELD_EDITORONLY(T, Name, ...) UE_REMOVE_OPTIONAL_PARENS(T) Name; INTERNAL_LAYOUT_FIELD(T, Name, STRUCT_OFFSET(DerivedType, Name), EFieldLayoutFlags::MakeFlagsEditorOnly(__VA_ARGS__), 1u, 0u, __COUNTER__)
+#define LAYOUT_ARRAY_EDITORONLY(T, Name, NumArray, ...) UE_REMOVE_OPTIONAL_PARENS(T) Name[NumArray]; INTERNAL_LAYOUT_FIELD(T, Name, STRUCT_OFFSET(DerivedType, Name), EFieldLayoutFlags::MakeFlagsEditorOnly(__VA_ARGS__), NumArray, 0u, __COUNTER__)
+#define LAYOUT_BITFIELD_EDITORONLY(T, Name, BitFieldSize, ...) UE_REMOVE_OPTIONAL_PARENS(T) Name : BitFieldSize; INTERNAL_LAYOUT_FIELD(T, Name, ~0u, EFieldLayoutFlags::MakeFlagsEditorOnly(__VA_ARGS__), 1u, BitFieldSize, __COUNTER__)
 #else
 #define LAYOUT_FIELD_EDITORONLY(T, Name, ...)
 #define LAYOUT_ARRAY_EDITORONLY(T, Name, NumArray, ...)
@@ -507,34 +507,34 @@ PRAGMA_ENABLE_DEPRECATION_WARNINGS
 
 #define INTERNAL_DECLARE_TYPE_LAYOUT_COMMON(T, InInterface) \
 	static constexpr int CounterBase = __COUNTER__; \
-	public: using DerivedType = PREPROCESSOR_REMOVE_OPTIONAL_PARENS(T); \
+	public: using DerivedType = UE_REMOVE_OPTIONAL_PARENS(T); \
 	static constexpr ETypeLayoutInterface::Type InterfaceType = ETypeLayoutInterface::InInterface; \
 	UE_DECLARE_INTERNAL_LINK_BASE(InternalLinkType) { UE_STATIC_ONLY(InternalLinkType); static UE_FORCEINLINE_HINT void Initialize(FTypeLayoutDesc& TypeDesc) {} }
 
 #define INTERNAL_DECLARE_INLINE_TYPE_LAYOUT(T, InInterface) \
 	private: static void InternalDestroy(void* Object, const FTypeLayoutDesc&, const FPointerTableBase* PtrTable, bool bIsFrozen) { \
-		Freeze::DestroyObject(static_cast<PREPROCESSOR_REMOVE_OPTIONAL_PARENS(T)*>(Object), PtrTable, bIsFrozen); \
+		Freeze::DestroyObject(static_cast<UE_REMOVE_OPTIONAL_PARENS(T)*>(Object), PtrTable, bIsFrozen); \
 	} \
 	public: static FTypeLayoutDesc& StaticGetTypeLayout() { \
-		static_assert(TValidateInterfaceHelper<PREPROCESSOR_REMOVE_OPTIONAL_PARENS(T), ETypeLayoutInterface::InInterface>::Value, #InInterface " is invalid interface for " #T); \
+		static_assert(TValidateInterfaceHelper<UE_REMOVE_OPTIONAL_PARENS(T), ETypeLayoutInterface::InInterface>::Value, #InInterface " is invalid interface for " #T); \
 		alignas(FTypeLayoutDesc) static uint8 TypeBuffer[sizeof(FTypeLayoutDesc)] = { 0 }; \
 		FTypeLayoutDesc& TypeDesc = *(FTypeLayoutDesc*)TypeBuffer; \
 		if (!TypeDesc.IsInitialized) { \
 			TypeDesc.IsInitialized = true; \
 			TypeDesc.Name = TEXT(#T); \
-			TypeDesc.WriteFrozenMemoryImageFunc = TGetFreezeImageHelper<PREPROCESSOR_REMOVE_OPTIONAL_PARENS(T)>::Do(); \
+			TypeDesc.WriteFrozenMemoryImageFunc = TGetFreezeImageHelper<UE_REMOVE_OPTIONAL_PARENS(T)>::Do(); \
 			TypeDesc.UnfrozenCopyFunc = &Freeze::DefaultUnfrozenCopy; \
 			TypeDesc.AppendHashFunc = &Freeze::DefaultAppendHash; \
 			TypeDesc.GetTargetAlignmentFunc = &Freeze::DefaultGetTargetAlignment; \
 			TypeDesc.ToStringFunc = &Freeze::DefaultToString; \
 			TypeDesc.DestroyFunc = &InternalDestroy; \
-			TypeDesc.Size = sizeof(PREPROCESSOR_REMOVE_OPTIONAL_PARENS(T)); \
-			TypeDesc.Alignment = alignof(PREPROCESSOR_REMOVE_OPTIONAL_PARENS(T)); \
+			TypeDesc.Size = sizeof(UE_REMOVE_OPTIONAL_PARENS(T)); \
+			TypeDesc.Alignment = alignof(UE_REMOVE_OPTIONAL_PARENS(T)); \
 			TypeDesc.Interface = ETypeLayoutInterface::InInterface; \
 			TypeDesc.SizeFromFields = ~0u; \
-			TypeDesc.GetDefaultObjectFunc = &TGetDefaultObjectHelper<PREPROCESSOR_REMOVE_OPTIONAL_PARENS(T), ETypeLayoutInterface::InInterface>::Do; \
+			TypeDesc.GetDefaultObjectFunc = &TGetDefaultObjectHelper<UE_REMOVE_OPTIONAL_PARENS(T), ETypeLayoutInterface::InInterface>::Do; \
 			InternalLinkType<1>::Initialize(TypeDesc); \
-			InternalInitializeBases<PREPROCESSOR_REMOVE_OPTIONAL_PARENS(T)>(TypeDesc); \
+			InternalInitializeBases<UE_REMOVE_OPTIONAL_PARENS(T)>(TypeDesc); \
 			FTypeLayoutDesc::Initialize(TypeDesc); \
 		} \
 		return TypeDesc; } \
@@ -548,7 +548,7 @@ PRAGMA_ENABLE_DEPRECATION_WARNINGS
 	INTERNAL_DECLARE_TYPE_LAYOUT_COMMON(T, InInterface)
 
 #define INTERNAL_DECLARE_LAYOUT_BASE(T) \
-	private: using InternalBaseType = typename TGetBaseTypeHelper<PREPROCESSOR_REMOVE_OPTIONAL_PARENS(T)>::Type; \
+	private: using InternalBaseType = typename TGetBaseTypeHelper<UE_REMOVE_OPTIONAL_PARENS(T)>::Type; \
 	template<typename InternalType> static void InternalInitializeBases(FTypeLayoutDesc& TypeDesc) { TInitializeBaseHelper<InternalType, InternalBaseType>::Initialize(TypeDesc); }
 
 #define INTERNAL_DECLARE_LAYOUT_EXPLICIT_BASES(T, ...) \
@@ -563,39 +563,39 @@ PRAGMA_ENABLE_DEPRECATION_WARNINGS
 #define DECLARE_EXPORTED_TYPE_LAYOUT_EXPLICIT_BASES(T, RequiredAPI, Interface, ...) INTERNAL_DECLARE_LAYOUT_EXPLICIT_BASES(T, __VA_ARGS__); INTERNAL_DECLARE_TYPE_LAYOUT(T, Interface, RequiredAPI)
 
 #define INTERNAL_IMPLEMENT_TYPE_LAYOUT_COMMON(TemplatePrefix, T) \
-	PREPROCESSOR_REMOVE_OPTIONAL_PARENS(TemplatePrefix) void PREPROCESSOR_REMOVE_OPTIONAL_PARENS(T)::InternalDestroy(void* Object, const FTypeLayoutDesc&, const FPointerTableBase* PtrTable, bool bIsFrozen) { \
-		Freeze::DestroyObject(static_cast<PREPROCESSOR_REMOVE_OPTIONAL_PARENS(T)*>(Object), PtrTable, bIsFrozen); \
+	UE_REMOVE_OPTIONAL_PARENS(TemplatePrefix) void UE_REMOVE_OPTIONAL_PARENS(T)::InternalDestroy(void* Object, const FTypeLayoutDesc&, const FPointerTableBase* PtrTable, bool bIsFrozen) { \
+		Freeze::DestroyObject(static_cast<UE_REMOVE_OPTIONAL_PARENS(T)*>(Object), PtrTable, bIsFrozen); \
 	} \
-	PREPROCESSOR_REMOVE_OPTIONAL_PARENS(TemplatePrefix) FTypeLayoutDesc& PREPROCESSOR_REMOVE_OPTIONAL_PARENS(T)::StaticGetTypeLayout() { \
-		static_assert(TValidateInterfaceHelper<PREPROCESSOR_REMOVE_OPTIONAL_PARENS(T), InterfaceType>::Value, "Invalid interface for " #T); \
+	UE_REMOVE_OPTIONAL_PARENS(TemplatePrefix) FTypeLayoutDesc& UE_REMOVE_OPTIONAL_PARENS(T)::StaticGetTypeLayout() { \
+		static_assert(TValidateInterfaceHelper<UE_REMOVE_OPTIONAL_PARENS(T), InterfaceType>::Value, "Invalid interface for " #T); \
 		alignas(FTypeLayoutDesc) static uint8 TypeBuffer[sizeof(FTypeLayoutDesc)] = { 0 }; \
 		FTypeLayoutDesc& TypeDesc = *(FTypeLayoutDesc*)TypeBuffer; \
 		if (!TypeDesc.IsInitialized) { \
 			TypeDesc.IsInitialized = true; \
 			TypeDesc.Name = TEXT(#T); \
-			TypeDesc.WriteFrozenMemoryImageFunc = TGetFreezeImageHelper<PREPROCESSOR_REMOVE_OPTIONAL_PARENS(T)>::Do(); \
+			TypeDesc.WriteFrozenMemoryImageFunc = TGetFreezeImageHelper<UE_REMOVE_OPTIONAL_PARENS(T)>::Do(); \
 			TypeDesc.UnfrozenCopyFunc = &Freeze::DefaultUnfrozenCopy; \
 			TypeDesc.AppendHashFunc = &Freeze::DefaultAppendHash; \
 			TypeDesc.GetTargetAlignmentFunc = &Freeze::DefaultGetTargetAlignment; \
 			TypeDesc.ToStringFunc = &Freeze::DefaultToString; \
 			TypeDesc.DestroyFunc = &InternalDestroy; \
-			TypeDesc.Size = sizeof(PREPROCESSOR_REMOVE_OPTIONAL_PARENS(T)); \
-			TypeDesc.Alignment = alignof(PREPROCESSOR_REMOVE_OPTIONAL_PARENS(T)); \
+			TypeDesc.Size = sizeof(UE_REMOVE_OPTIONAL_PARENS(T)); \
+			TypeDesc.Alignment = alignof(UE_REMOVE_OPTIONAL_PARENS(T)); \
 			TypeDesc.Interface = InterfaceType; \
 			TypeDesc.SizeFromFields = ~0u; \
-			TypeDesc.GetDefaultObjectFunc = &TGetDefaultObjectHelper<PREPROCESSOR_REMOVE_OPTIONAL_PARENS(T), InterfaceType>::Do; \
+			TypeDesc.GetDefaultObjectFunc = &TGetDefaultObjectHelper<UE_REMOVE_OPTIONAL_PARENS(T), InterfaceType>::Do; \
 			InternalLinkType<1>::Initialize(TypeDesc); \
-			InternalInitializeBases<PREPROCESSOR_REMOVE_OPTIONAL_PARENS(T)>(TypeDesc); \
+			InternalInitializeBases<UE_REMOVE_OPTIONAL_PARENS(T)>(TypeDesc); \
 			FTypeLayoutDesc::Initialize(TypeDesc); \
 		} \
 		return TypeDesc; }
 
-//#define INTERNAL_REGISTER_TYPE_LAYOUT(T) static const FRegisterTypeLayoutDesc ANONYMOUS_VARIABLE(RegisterTypeLayoutDesc)(PREPROCESSOR_REMOVE_OPTIONAL_PARENS(T)::StaticGetTypeLayout())
+//#define INTERNAL_REGISTER_TYPE_LAYOUT(T) static const FRegisterTypeLayoutDesc ANONYMOUS_VARIABLE(RegisterTypeLayoutDesc)(UE_REMOVE_OPTIONAL_PARENS(T)::StaticGetTypeLayout())
 #define INTERNAL_REGISTER_TYPE_LAYOUT(T) static const FDelayedAutoRegisterHelper ANONYMOUS_VARIABLE(DelayedAutoRegisterHelper)(EDelayedRegisterRunPhase::ShaderTypesReady, []{ FTypeLayoutDesc::Register(T::StaticGetTypeLayout()); });
 
 #define IMPLEMENT_UNREGISTERED_TEMPLATE_TYPE_LAYOUT(TemplatePrefix, T) \
 	INTERNAL_IMPLEMENT_TYPE_LAYOUT_COMMON(TemplatePrefix, T); \
-	PREPROCESSOR_REMOVE_OPTIONAL_PARENS(TemplatePrefix) const FTypeLayoutDesc& PREPROCESSOR_REMOVE_OPTIONAL_PARENS(T)::GetTypeLayout() const { return StaticGetTypeLayout(); }
+	UE_REMOVE_OPTIONAL_PARENS(TemplatePrefix) const FTypeLayoutDesc& UE_REMOVE_OPTIONAL_PARENS(T)::GetTypeLayout() const { return StaticGetTypeLayout(); }
 
 #define IMPLEMENT_TEMPLATE_TYPE_LAYOUT(TemplatePrefix, T) \
 	IMPLEMENT_UNREGISTERED_TEMPLATE_TYPE_LAYOUT(TemplatePrefix, T); \
@@ -610,8 +610,8 @@ PRAGMA_ENABLE_DEPRECATION_WARNINGS
 #define REGISTER_INLINE_TYPE_LAYOUT(T) \
 	static const FDelayedAutoRegisterHelper ANONYMOUS_VARIABLE(DelayedAutoRegisterHelper)(EDelayedRegisterRunPhase::ShaderTypesReady, [] \
 		{ \
-			PREPROCESSOR_REMOVE_OPTIONAL_PARENS(T)::StaticGetTypeLayout().Name = TEXT(#T); \
-			FTypeLayoutDesc::Register(PREPROCESSOR_REMOVE_OPTIONAL_PARENS(T)::StaticGetTypeLayout()); \
+			UE_REMOVE_OPTIONAL_PARENS(T)::StaticGetTypeLayout().Name = TEXT(#T); \
+			FTypeLayoutDesc::Register(UE_REMOVE_OPTIONAL_PARENS(T)::StaticGetTypeLayout()); \
 		} \
 	);
 
@@ -659,26 +659,26 @@ inline void DeleteObjectFromLayout(T* Object, const FPointerTableBase* PtrTable
 }
 
 #define DECLARE_TEMPLATE_INTRINSIC_TYPE_LAYOUT(TemplatePrefix, T) \
-	PREPROCESSOR_REMOVE_OPTIONAL_PARENS(TemplatePrefix) struct THasTypeLayout<PREPROCESSOR_REMOVE_OPTIONAL_PARENS(T)> { UE_STATIC_ONLY(THasTypeLayout); static const bool Value = true; }; \
-	PREPROCESSOR_REMOVE_OPTIONAL_PARENS(TemplatePrefix) struct TStaticGetTypeLayoutHelper<PREPROCESSOR_REMOVE_OPTIONAL_PARENS(T)> { \
+	UE_REMOVE_OPTIONAL_PARENS(TemplatePrefix) struct THasTypeLayout<UE_REMOVE_OPTIONAL_PARENS(T)> { UE_STATIC_ONLY(THasTypeLayout); static const bool Value = true; }; \
+	UE_REMOVE_OPTIONAL_PARENS(TemplatePrefix) struct TStaticGetTypeLayoutHelper<UE_REMOVE_OPTIONAL_PARENS(T)> { \
 		UE_STATIC_ONLY(TStaticGetTypeLayoutHelper); \
 		static void CallWriteMemoryImage(FMemoryImageWriter& Writer, const void* Object, const FTypeLayoutDesc& TypeDesc, const FTypeLayoutDesc& DerivedTypeDesc) { \
-			Freeze::IntrinsicWriteMemoryImage(Writer, *static_cast<PREPROCESSOR_REMOVE_OPTIONAL_PARENS(T) const*>(Object), TypeDesc); \
+			Freeze::IntrinsicWriteMemoryImage(Writer, *static_cast<UE_REMOVE_OPTIONAL_PARENS(T) const*>(Object), TypeDesc); \
 		} \
 		static uint32 CallUnfrozenCopy(const FMemoryUnfreezeContent& Context, const void* Object, const FTypeLayoutDesc& TypeDesc, void* OutDst) { \
-			return Freeze::IntrinsicUnfrozenCopy(Context, *static_cast<PREPROCESSOR_REMOVE_OPTIONAL_PARENS(T) const*>(Object), OutDst); \
+			return Freeze::IntrinsicUnfrozenCopy(Context, *static_cast<UE_REMOVE_OPTIONAL_PARENS(T) const*>(Object), OutDst); \
 		} \
 		static uint32 CallAppendHash(const FTypeLayoutDesc& TypeDesc, const FPlatformTypeLayoutParameters& LayoutParams, FSHA1& Hasher) { \
-			return Freeze::IntrinsicAppendHash(static_cast<PREPROCESSOR_REMOVE_OPTIONAL_PARENS(T) const*>(nullptr), TypeDesc, LayoutParams, Hasher); \
+			return Freeze::IntrinsicAppendHash(static_cast<UE_REMOVE_OPTIONAL_PARENS(T) const*>(nullptr), TypeDesc, LayoutParams, Hasher); \
 		} \
 		static uint32 CallGetTargetAlignment(const FTypeLayoutDesc& TypeDesc, const FPlatformTypeLayoutParameters& LayoutParams) { \
-			return Freeze::IntrinsicGetTargetAlignment(static_cast<PREPROCESSOR_REMOVE_OPTIONAL_PARENS(T) const*>(nullptr), TypeDesc, LayoutParams); \
+			return Freeze::IntrinsicGetTargetAlignment(static_cast<UE_REMOVE_OPTIONAL_PARENS(T) const*>(nullptr), TypeDesc, LayoutParams); \
 		} \
 		static void CallToString(const void* Object, const FTypeLayoutDesc& TypeDesc, const FPlatformTypeLayoutParameters& LayoutParams, FMemoryToStringContext& OutContext) { \
-			return Freeze::IntrinsicToString(*static_cast<PREPROCESSOR_REMOVE_OPTIONAL_PARENS(T) const*>(Object), TypeDesc, LayoutParams, OutContext); \
+			return Freeze::IntrinsicToString(*static_cast<UE_REMOVE_OPTIONAL_PARENS(T) const*>(Object), TypeDesc, LayoutParams, OutContext); \
 		} \
 		static void CallDestroy(void* Object, const FTypeLayoutDesc&, const FPointerTableBase* PtrTable, bool bIsFrozen) { \
-			Freeze::DestroyObject(static_cast<PREPROCESSOR_REMOVE_OPTIONAL_PARENS(T)*>(Object), PtrTable, bIsFrozen); \
+			Freeze::DestroyObject(static_cast<UE_REMOVE_OPTIONAL_PARENS(T)*>(Object), PtrTable, bIsFrozen); \
 		} \
 		static const FTypeLayoutDesc& Do() { \
 			alignas(FTypeLayoutDesc) static uint8 TypeBuffer[sizeof(FTypeLayoutDesc)] = { 0 }; \
@@ -693,19 +693,19 @@ inline void DeleteObjectFromLayout(T* Object, const FPointerTableBase* PtrTable
 				TypeDesc.GetTargetAlignmentFunc = &CallGetTargetAlignment; \
 				TypeDesc.ToStringFunc = &CallToString; \
 				TypeDesc.DestroyFunc = &CallDestroy; \
-				TypeDesc.Size = sizeof(PREPROCESSOR_REMOVE_OPTIONAL_PARENS(T)); \
-				TypeDesc.Alignment = alignof(PREPROCESSOR_REMOVE_OPTIONAL_PARENS(T)); \
+				TypeDesc.Size = sizeof(UE_REMOVE_OPTIONAL_PARENS(T)); \
+				TypeDesc.Alignment = alignof(UE_REMOVE_OPTIONAL_PARENS(T)); \
 				TypeDesc.Interface = ETypeLayoutInterface::NonVirtual; \
-				TypeDesc.SizeFromFields = sizeof(PREPROCESSOR_REMOVE_OPTIONAL_PARENS(T)); \
+				TypeDesc.SizeFromFields = sizeof(UE_REMOVE_OPTIONAL_PARENS(T)); \
 			} \
 			return TypeDesc; } }; \
-	PREPROCESSOR_REMOVE_OPTIONAL_PARENS(TemplatePrefix) struct TGetTypeLayoutHelper<PREPROCESSOR_REMOVE_OPTIONAL_PARENS(T)> { \
+	UE_REMOVE_OPTIONAL_PARENS(TemplatePrefix) struct TGetTypeLayoutHelper<UE_REMOVE_OPTIONAL_PARENS(T)> { \
 		UE_STATIC_ONLY(TGetTypeLayoutHelper); \
-		static const FTypeLayoutDesc& Do(const PREPROCESSOR_REMOVE_OPTIONAL_PARENS(T)&) { return TStaticGetTypeLayoutHelper<PREPROCESSOR_REMOVE_OPTIONAL_PARENS(T)>::Do(); }}
+		static const FTypeLayoutDesc& Do(const UE_REMOVE_OPTIONAL_PARENS(T)&) { return TStaticGetTypeLayoutHelper<UE_REMOVE_OPTIONAL_PARENS(T)>::Do(); }}
 
 #define DECLARE_EXPORTED_TEMPLATE_INTRINSIC_TYPE_LAYOUT(TemplatePrefix, T, RequiredAPI) \
-	PREPROCESSOR_REMOVE_OPTIONAL_PARENS(TemplatePrefix) struct THasTypeLayout<PREPROCESSOR_REMOVE_OPTIONAL_PARENS(T)> { UE_STATIC_ONLY(THasTypeLayout); static const bool Value = true; }; \
-	PREPROCESSOR_REMOVE_OPTIONAL_PARENS(TemplatePrefix) struct TStaticGetTypeLayoutHelper<PREPROCESSOR_REMOVE_OPTIONAL_PARENS(T)> { \
+	UE_REMOVE_OPTIONAL_PARENS(TemplatePrefix) struct THasTypeLayout<UE_REMOVE_OPTIONAL_PARENS(T)> { UE_STATIC_ONLY(THasTypeLayout); static const bool Value = true; }; \
+	UE_REMOVE_OPTIONAL_PARENS(TemplatePrefix) struct TStaticGetTypeLayoutHelper<UE_REMOVE_OPTIONAL_PARENS(T)> { \
 		UE_STATIC_ONLY(TStaticGetTypeLayoutHelper); \
 		static void CallWriteMemoryImage(FMemoryImageWriter& Writer, const void* Object, const FTypeLayoutDesc& TypeDesc, const FTypeLayoutDesc& DerivedTypeDesc); \
 		static uint32 CallUnfrozenCopy(const FMemoryUnfreezeContent& Context, const void* Object, const FTypeLayoutDesc& TypeDesc, void* OutDst); \
@@ -714,30 +714,30 @@ inline void DeleteObjectFromLayout(T* Object, const FPointerTableBase* PtrTable
 		static void CallToString(const void* Object, const FTypeLayoutDesc& TypeDesc, const FPlatformTypeLayoutParameters& LayoutParams, FMemoryToStringContext& OutContext); \
 		static void CallDestroy(void* Object, const FTypeLayoutDesc&, const FPointerTableBase* PtrTable, bool bIsFrozen); \
 		RequiredAPI static const FTypeLayoutDesc& Do(); }; \
-	PREPROCESSOR_REMOVE_OPTIONAL_PARENS(TemplatePrefix) struct TGetTypeLayoutHelper<PREPROCESSOR_REMOVE_OPTIONAL_PARENS(T)> { \
+	UE_REMOVE_OPTIONAL_PARENS(TemplatePrefix) struct TGetTypeLayoutHelper<UE_REMOVE_OPTIONAL_PARENS(T)> { \
 		UE_STATIC_ONLY(TGetTypeLayoutHelper); \
-		static const FTypeLayoutDesc& Do(const PREPROCESSOR_REMOVE_OPTIONAL_PARENS(T)&) { return TStaticGetTypeLayoutHelper<PREPROCESSOR_REMOVE_OPTIONAL_PARENS(T)>::Do(); }}
+		static const FTypeLayoutDesc& Do(const UE_REMOVE_OPTIONAL_PARENS(T)&) { return TStaticGetTypeLayoutHelper<UE_REMOVE_OPTIONAL_PARENS(T)>::Do(); }}
 
 #define IMPLEMENT_EXPORTED_INTRINSIC_TYPE_LAYOUT(T) \
-		void TStaticGetTypeLayoutHelper<PREPROCESSOR_REMOVE_OPTIONAL_PARENS(T)>::CallWriteMemoryImage(FMemoryImageWriter& Writer, const void* Object, const FTypeLayoutDesc& TypeDesc, const FTypeLayoutDesc& DerivedTypeDesc) { \
-			Freeze::IntrinsicWriteMemoryImage(Writer, *static_cast<PREPROCESSOR_REMOVE_OPTIONAL_PARENS(T) const*>(Object), TypeDesc); \
+		void TStaticGetTypeLayoutHelper<UE_REMOVE_OPTIONAL_PARENS(T)>::CallWriteMemoryImage(FMemoryImageWriter& Writer, const void* Object, const FTypeLayoutDesc& TypeDesc, const FTypeLayoutDesc& DerivedTypeDesc) { \
+			Freeze::IntrinsicWriteMemoryImage(Writer, *static_cast<UE_REMOVE_OPTIONAL_PARENS(T) const*>(Object), TypeDesc); \
 		} \
-		uint32 TStaticGetTypeLayoutHelper<PREPROCESSOR_REMOVE_OPTIONAL_PARENS(T)>::CallUnfrozenCopy(const FMemoryUnfreezeContent& Context, const void* Object, const FTypeLayoutDesc& TypeDesc, void* OutDst) { \
-			return Freeze::IntrinsicUnfrozenCopy(Context, *static_cast<PREPROCESSOR_REMOVE_OPTIONAL_PARENS(T) const*>(Object), OutDst); \
+		uint32 TStaticGetTypeLayoutHelper<UE_REMOVE_OPTIONAL_PARENS(T)>::CallUnfrozenCopy(const FMemoryUnfreezeContent& Context, const void* Object, const FTypeLayoutDesc& TypeDesc, void* OutDst) { \
+			return Freeze::IntrinsicUnfrozenCopy(Context, *static_cast<UE_REMOVE_OPTIONAL_PARENS(T) const*>(Object), OutDst); \
 		} \
-		uint32 TStaticGetTypeLayoutHelper<PREPROCESSOR_REMOVE_OPTIONAL_PARENS(T)>::CallAppendHash(const FTypeLayoutDesc& TypeDesc, const FPlatformTypeLayoutParameters& LayoutParams, FSHA1& Hasher) { \
-			return Freeze::IntrinsicAppendHash(static_cast<PREPROCESSOR_REMOVE_OPTIONAL_PARENS(T) const*>(nullptr), TypeDesc, LayoutParams, Hasher); \
+		uint32 TStaticGetTypeLayoutHelper<UE_REMOVE_OPTIONAL_PARENS(T)>::CallAppendHash(const FTypeLayoutDesc& TypeDesc, const FPlatformTypeLayoutParameters& LayoutParams, FSHA1& Hasher) { \
+			return Freeze::IntrinsicAppendHash(static_cast<UE_REMOVE_OPTIONAL_PARENS(T) const*>(nullptr), TypeDesc, LayoutParams, Hasher); \
 		} \
-		uint32 TStaticGetTypeLayoutHelper<PREPROCESSOR_REMOVE_OPTIONAL_PARENS(T)>::CallGetTargetAlignment(const FTypeLayoutDesc& TypeDesc, const FPlatformTypeLayoutParameters& LayoutParams) { \
-			return Freeze::IntrinsicGetTargetAlignment(static_cast<PREPROCESSOR_REMOVE_OPTIONAL_PARENS(T) const*>(nullptr), TypeDesc, LayoutParams); \
+		uint32 TStaticGetTypeLayoutHelper<UE_REMOVE_OPTIONAL_PARENS(T)>::CallGetTargetAlignment(const FTypeLayoutDesc& TypeDesc, const FPlatformTypeLayoutParameters& LayoutParams) { \
+			return Freeze::IntrinsicGetTargetAlignment(static_cast<UE_REMOVE_OPTIONAL_PARENS(T) const*>(nullptr), TypeDesc, LayoutParams); \
 		} \
-		void TStaticGetTypeLayoutHelper<PREPROCESSOR_REMOVE_OPTIONAL_PARENS(T)>::CallToString(const void* Object, const FTypeLayoutDesc& TypeDesc, const FPlatformTypeLayoutParameters& LayoutParams, FMemoryToStringContext& OutContext) { \
-			return Freeze::IntrinsicToString(*static_cast<PREPROCESSOR_REMOVE_OPTIONAL_PARENS(T) const*>(Object), TypeDesc, LayoutParams, OutContext); \
+		void TStaticGetTypeLayoutHelper<UE_REMOVE_OPTIONAL_PARENS(T)>::CallToString(const void* Object, const FTypeLayoutDesc& TypeDesc, const FPlatformTypeLayoutParameters& LayoutParams, FMemoryToStringContext& OutContext) { \
+			return Freeze::IntrinsicToString(*static_cast<UE_REMOVE_OPTIONAL_PARENS(T) const*>(Object), TypeDesc, LayoutParams, OutContext); \
 		} \
-		void TStaticGetTypeLayoutHelper<PREPROCESSOR_REMOVE_OPTIONAL_PARENS(T)>::CallDestroy(void* Object, const FTypeLayoutDesc&, const FPointerTableBase* PtrTable, bool bIsFrozen) { \
-			Freeze::DestroyObject(static_cast<PREPROCESSOR_REMOVE_OPTIONAL_PARENS(T)*>(Object), PtrTable, bIsFrozen); \
+		void TStaticGetTypeLayoutHelper<UE_REMOVE_OPTIONAL_PARENS(T)>::CallDestroy(void* Object, const FTypeLayoutDesc&, const FPointerTableBase* PtrTable, bool bIsFrozen) { \
+			Freeze::DestroyObject(static_cast<UE_REMOVE_OPTIONAL_PARENS(T)*>(Object), PtrTable, bIsFrozen); \
 		} \
-		const FTypeLayoutDesc& TStaticGetTypeLayoutHelper<PREPROCESSOR_REMOVE_OPTIONAL_PARENS(T)>::Do() { \
+		const FTypeLayoutDesc& TStaticGetTypeLayoutHelper<UE_REMOVE_OPTIONAL_PARENS(T)>::Do() { \
 			alignas(FTypeLayoutDesc) static uint8 TypeBuffer[sizeof(FTypeLayoutDesc)] = { 0 }; \
 			FTypeLayoutDesc& TypeDesc = *(FTypeLayoutDesc*)TypeBuffer; \
 			if (!TypeDesc.IsInitialized) { \
@@ -750,21 +750,21 @@ inline void DeleteObjectFromLayout(T* Object, const FPointerTableBase* PtrTable
 				TypeDesc.GetTargetAlignmentFunc = &CallGetTargetAlignment; \
 				TypeDesc.ToStringFunc = &CallToString; \
 				TypeDesc.DestroyFunc = &CallDestroy; \
-				TypeDesc.Size = sizeof(PREPROCESSOR_REMOVE_OPTIONAL_PARENS(T)); \
-				TypeDesc.Alignment = alignof(PREPROCESSOR_REMOVE_OPTIONAL_PARENS(T)); \
+				TypeDesc.Size = sizeof(UE_REMOVE_OPTIONAL_PARENS(T)); \
+				TypeDesc.Alignment = alignof(UE_REMOVE_OPTIONAL_PARENS(T)); \
 				TypeDesc.Interface = ETypeLayoutInterface::NonVirtual; \
-				TypeDesc.SizeFromFields = sizeof(PREPROCESSOR_REMOVE_OPTIONAL_PARENS(T)); \
+				TypeDesc.SizeFromFields = sizeof(UE_REMOVE_OPTIONAL_PARENS(T)); \
 			} \
 			return TypeDesc; }
 
 #define DECLARE_INTRINSIC_TYPE_LAYOUT(T) DECLARE_TEMPLATE_INTRINSIC_TYPE_LAYOUT(template<>, T)
 
 #define ALIAS_TEMPLATE_TYPE_LAYOUT(TemplatePrefix, T, Alias) \
-	PREPROCESSOR_REMOVE_OPTIONAL_PARENS(TemplatePrefix) struct TStaticGetTypeLayoutHelper<PREPROCESSOR_REMOVE_OPTIONAL_PARENS(T)> : public TStaticGetTypeLayoutHelper<PREPROCESSOR_REMOVE_OPTIONAL_PARENS(Alias)> { UE_STATIC_ONLY(TStaticGetTypeLayoutHelper); }; \
-	PREPROCESSOR_REMOVE_OPTIONAL_PARENS(TemplatePrefix) struct TGetTypeLayoutHelper<PREPROCESSOR_REMOVE_OPTIONAL_PARENS(T)> : public TGetTypeLayoutHelper<PREPROCESSOR_REMOVE_OPTIONAL_PARENS(Alias)> { UE_STATIC_ONLY(TGetTypeLayoutHelper); }
+	UE_REMOVE_OPTIONAL_PARENS(TemplatePrefix) struct TStaticGetTypeLayoutHelper<UE_REMOVE_OPTIONAL_PARENS(T)> : public TStaticGetTypeLayoutHelper<UE_REMOVE_OPTIONAL_PARENS(Alias)> { UE_STATIC_ONLY(TStaticGetTypeLayoutHelper); }; \
+	UE_REMOVE_OPTIONAL_PARENS(TemplatePrefix) struct TGetTypeLayoutHelper<UE_REMOVE_OPTIONAL_PARENS(T)> : public TGetTypeLayoutHelper<UE_REMOVE_OPTIONAL_PARENS(Alias)> { UE_STATIC_ONLY(TGetTypeLayoutHelper); }
 
 #define ALIAS_TYPE_LAYOUT(Type, Alias) \
-	static_assert(sizeof(PREPROCESSOR_REMOVE_OPTIONAL_PARENS(Type)) == sizeof(PREPROCESSOR_REMOVE_OPTIONAL_PARENS(Alias)), "Using a type alias but the sizes don't match!"); \
+	static_assert(sizeof(UE_REMOVE_OPTIONAL_PARENS(Type)) == sizeof(UE_REMOVE_OPTIONAL_PARENS(Alias)), "Using a type alias but the sizes don't match!"); \
 	ALIAS_TEMPLATE_TYPE_LAYOUT(template<>, Type, Alias)
 
 DECLARE_INTRINSIC_TYPE_LAYOUT(char);
diff --git a/Engine/Source/Runtime/CoreUObject/Public/UObject/Field.h b/Engine/Source/Runtime/CoreUObject/Public/UObject/Field.h
--- a/Engine/Source/Runtime/CoreUObject/Public/UObject/Field.h
+++ b/Engine/Source/Runtime/CoreUObject/Public/UObject/Field.h
@@ -260,7 +260,7 @@ private: \
 	TClass& operator=(const TClass&) = delete;   \
 	UE_DEFINE_GET_FIELD_CLASS_PRIVATE \
 public: \
-	using Super = PREPROCESSOR_REMOVE_OPTIONAL_PARENS(TSuperClass);\
+	using Super = UE_REMOVE_OPTIONAL_PARENS(TSuperClass);\
 	using ThisClass = TClass;\
 	TClass(EInternal InInernal, FFieldClass* InClass) \
 		: Super(EC_InternalUseOnlyConstructor, InClass) \
diff --git a/Engine/Source/Runtime/RenderCore/Public/VertexFactory.h b/Engine/Source/Runtime/RenderCore/Public/VertexFactory.h
--- a/Engine/Source/Runtime/RenderCore/Public/VertexFactory.h
+++ b/Engine/Source/Runtime/RenderCore/Public/VertexFactory.h
@@ -586,13 +586,13 @@ extern RENDERCORE_API FVertexFactoryType* FindVertexFactoryType(const FHashedNam
 * Templated macro to define a FVertexFactoryType. Can be used for templated Vertex Factories.
 */
 #define IMPLEMENT_TEMPLATE_VERTEX_FACTORY_TYPE(TemplatePrefix, FactoryClass, ShaderFilename, Flags) \
-	PREPROCESSOR_REMOVE_OPTIONAL_PARENS(TemplatePrefix) FVertexFactoryType FactoryClass::StaticType( \
+	UE_REMOVE_OPTIONAL_PARENS(TemplatePrefix) FVertexFactoryType FactoryClass::StaticType( \
 		TEXT(#FactoryClass), \
 		TEXT(ShaderFilename), \
 		Flags, \
 		IMPLEMENT_VERTEX_FACTORY_VTABLE(FactoryClass) \
 		); \
-		PREPROCESSOR_REMOVE_OPTIONAL_PARENS(TemplatePrefix) FVertexFactoryType* FactoryClass::GetType() const { return &StaticType; }
+		UE_REMOVE_OPTIONAL_PARENS(TemplatePrefix) FVertexFactoryType* FactoryClass::GetType() const { return &StaticType; }
 
 /** Encapsulates a dependency on a vertex factory type and saved state from that vertex factory type. */
 class FVertexFactoryTypeDependency
