Made UE_JOIN variadic.

diff --git a/Engine/Plugins/Runtime/Metasound/Source/MetasoundFrontend/Public/MetasoundFrontendModuleRegistrationMacros.h b/Engine/Plugins/Runtime/Metasound/Source/MetasoundFrontend/Public/MetasoundFrontendModuleRegistrationMacros.h
--- a/Engine/Plugins/Runtime/Metasound/Source/MetasoundFrontend/Public/MetasoundFrontendModuleRegistrationMacros.h
+++ b/Engine/Plugins/Runtime/Metasound/Source/MetasoundFrontend/Public/MetasoundFrontendModuleRegistrationMacros.h
@@ -184,7 +184,7 @@ namespace Metasound::Frontend
 #define METASOUND_UNREGISTER_ON_MODULE_SHUTDOWN 1
 
 // Module registration list class name
-#define METASOUND_MODULE_REGISTRATION_LIST_CLASS UE_JOIN(F, UE_JOIN(METASOUND_PLUGIN, UE_JOIN(_, UE_JOIN(METASOUND_MODULE, _Registration))))
+#define METASOUND_MODULE_REGISTRATION_LIST_CLASS UE_JOIN(F, METASOUND_PLUGIN, _, METASOUND_MODULE, _Registration)
 
 // Declare a module registration list. This needs to be visible wherever METASOUND_REGISTER_NODE
 // is called. Generally this can be a header file within the module's Private directory. 
diff --git a/Engine/Source/Editor/UnrealEd/Public/CookPackageSplitter.h b/Engine/Source/Editor/UnrealEd/Public/CookPackageSplitter.h
--- a/Engine/Source/Editor/UnrealEd/Public/CookPackageSplitter.h
+++ b/Engine/Source/Editor/UnrealEd/Public/CookPackageSplitter.h
@@ -441,10 +441,10 @@ private:
  * REGISTER_COOKPACKAGE_SPLITTER(FMyCookPackageSplitter, UMySplitDataClass);
  */
 #define DEFINE_COOKPACKAGE_SPLITTER(SplitterClass, SplitDataClass) \
-friend class UE_JOIN(UE_JOIN(SplitterClass, SplitDataClass), _Register)
+friend class UE_JOIN(SplitterClass, SplitDataClass, _Register)
 
 #define REGISTER_COOKPACKAGE_SPLITTER(SplitterClass, SplitDataClass) \
-class UE_JOIN(UE_JOIN(SplitterClass, SplitDataClass), _Register) \
+class UE_JOIN(SplitterClass, SplitDataClass, _Register) \
 	: public UE::Cook::Private::FRegisteredCookPackageSplitter \
 { \
 	virtual UClass* GetSplitDataClass() const override \
@@ -470,7 +470,7 @@ class UE_JOIN(UE_JOIN(SplitterClass, SplitDataClass), _Register) \
 }; \
 namespace UE_JOIN(SplitterClass, SplitDataClass) \
 { \
-	static UE_JOIN(UE_JOIN(SplitterClass, SplitDataClass), _Register) DefaultObject; \
+	static UE_JOIN(SplitterClass, SplitDataClass, _Register) DefaultObject; \
 }
 
 #endif
diff --git a/Engine/Source/Runtime/AssetRegistry/Public/AssetRegistry/AssetDependencyGatherer.h b/Engine/Source/Runtime/AssetRegistry/Public/AssetRegistry/AssetDependencyGatherer.h
--- a/Engine/Source/Runtime/AssetRegistry/Public/AssetRegistry/AssetDependencyGatherer.h
+++ b/Engine/Source/Runtime/AssetRegistry/Public/AssetRegistry/AssetDependencyGatherer.h
@@ -94,7 +94,7 @@ namespace UE::AssetDependencyGatherer::Private
  * REGISTER_ASSETDEPENDENCY_GATHERER(FMyAssetDependencyGatherer, UMyAssetClass);
  */
 #define REGISTER_ASSETDEPENDENCY_GATHERER(GathererClass, AssetClass) \
-class UE_JOIN(UE_JOIN(GathererClass, AssetClass), _Register) : public UE::AssetDependencyGatherer::Private::FRegisteredAssetDependencyGatherer \
+class UE_JOIN(GathererClass, AssetClass, _Register) : public UE::AssetDependencyGatherer::Private::FRegisteredAssetDependencyGatherer \
 { \
 	virtual UClass* GetAssetClass() const override { return AssetClass::StaticClass(); } \
 	virtual void GatherDependencies(const FAssetData& AssetData, const FAssetRegistryState& AssetRegistryState, TFunctionRef<FARCompiledFilter(const FARFilter&)> CompileFilterFunc, TArray<IAssetDependencyGatherer::FGathereredDependency>& OutDependencies, TArray<FString>& OutDependencyDirectories) const override \
@@ -105,7 +105,7 @@ class UE_JOIN(UE_JOIN(GathererClass, AssetClass), _Register) : public UE::AssetD
 }; \
 namespace UE_JOIN(GathererClass, AssetClass) \
 { \
-	static UE_JOIN(UE_JOIN(GathererClass, AssetClass), _Register) DefaultObject; \
+	static UE_JOIN(GathererClass, AssetClass, _Register) DefaultObject; \
 }
 
 #endif // WITH_EDITOR
diff --git a/Engine/Source/Runtime/Core/Public/HAL/PreprocessorHelpers.h b/Engine/Source/Runtime/Core/Public/HAL/PreprocessorHelpers.h
--- a/Engine/Source/Runtime/Core/Public/HAL/PreprocessorHelpers.h
+++ b/Engine/Source/Runtime/Core/Public/HAL/PreprocessorHelpers.h
@@ -6,9 +6,13 @@
 #define UE_STRINGIZE(Token) UE_PRIVATE_STRINGIZE(Token)
 #define UE_PRIVATE_STRINGIZE(Token) #Token
 
-// Concatenates two preprocessor tokens, performing macro expansion on them first
-#define UE_JOIN(TokenA, TokenB) UE_PRIVATE_JOIN(TokenA, TokenB)
-#define UE_PRIVATE_JOIN(TokenA, TokenB) TokenA##TokenB
+// Concatenates two or more preprocessor tokens, performing macro expansion on them first
+#define UE_JOIN(TokenA, TokenB, ...) UE_APPEND_VA_ARG_COUNT(UE_PRIVATE_JOIN, __VA_ARGS__)(TokenA, TokenB, __VA_ARGS__)
+#define UE_PRIVATE_JOIN0(TokenA, TokenB) TokenA##TokenB
+#define UE_PRIVATE_JOIN1(TokenA, TokenB, TokenC) TokenA##TokenB##TokenC
+#define UE_PRIVATE_JOIN2(TokenA, TokenB, TokenC, TokenD) TokenA##TokenB##TokenC##TokenD
+#define UE_PRIVATE_JOIN3(TokenA, TokenB, TokenC, TokenD, TokenE) TokenA##TokenB##TokenC##TokenD##TokenE
+#define UE_PRIVATE_JOIN4(TokenA, TokenB, TokenC, TokenD, TokenE, TokenF) TokenA##TokenB##TokenC##TokenD##TokenE##TokenF
 
 // Concatenates the first two preprocessor tokens of a variadic list, after performing macro expansion on them
 #define UE_JOIN_FIRST(Token, ...) UE_PRIVATE_JOIN_FIRST(Token, __VA_ARGS__)
diff --git a/Engine/Source/Runtime/Core/Public/Stats/Stats.h b/Engine/Source/Runtime/Core/Public/Stats/Stats.h
--- a/Engine/Source/Runtime/Core/Public/Stats/Stats.h
+++ b/Engine/Source/Runtime/Core/Public/Stats/Stats.h
@@ -574,8 +574,8 @@ struct FStat_##StatName\
 // ----------------------------------------------------------------
 
 #define SCOPE_CYCLE_COUNTER_TO_TRACE(StatString, StatName, Condition) \
-	TRACE_CPUPROFILER_EVENT_DECLARE(UE_JOIN(UE_JOIN(__Decl_, StatName), __LINE__)); \
-	TRACE_CPUPROFILER_EVENT_SCOPE_USE(UE_JOIN(UE_JOIN(__Decl_, StatName), __LINE__), StatString, UE_JOIN(UE_JOIN(__Scope_, StatName), __LINE__), Condition && GCycleStatsShouldEmitNamedEvents);
+	TRACE_CPUPROFILER_EVENT_DECLARE(UE_JOIN(__Decl_, StatName, __LINE__)); \
+	TRACE_CPUPROFILER_EVENT_SCOPE_USE(UE_JOIN(__Decl_, StatName, __LINE__), StatString, UE_JOIN(__Scope_, StatName, __LINE__), Condition && GCycleStatsShouldEmitNamedEvents);
 
 // Note: Since these are defining the stat inline, no custom stat data will exist.
 #define DECLARE_SCOPE_CYCLE_COUNTER(CounterName,Stat,GroupId) \
diff --git a/Engine/Source/Runtime/Core/Tests/HAL/PreprocessorHelperTests.cpp b/Engine/Source/Runtime/Core/Tests/HAL/PreprocessorHelperTests.cpp
--- a/Engine/Source/Runtime/Core/Tests/HAL/PreprocessorHelperTests.cpp
+++ b/Engine/Source/Runtime/Core/Tests/HAL/PreprocessorHelperTests.cpp
@@ -21,7 +21,7 @@ TEST_CASE_NAMED(FPreprocessorHelperArgCountTest, "System::Core::HAL::UE_VA_ARG_C
 	STATIC_CHECK(UE_VA_ARG_COUNT(a, b, c, d, e, f, g, h, i, j) == 10);
 }
 
-TEST_CASE_NAMED(FPreprocessorHelperArgAppendCountTest, "System::Core::HAL::UE_APPEND_VA_ARG_COUNT", "[ApplicationContextMask][EngineFilter]")
+TEST_CASE_NAMED(FPreprocessorHelperJoinTest, "System::Core::HAL::UE_APPEND_VA_ARG_COUNT", "[ApplicationContextMask][EngineFilter]")
 {
 	#define TESTPREFIX_0 22
 	#define TESTPREFIX_1 32
@@ -46,6 +46,41 @@ TEST_CASE_NAMED(FPreprocessorHelperArgAppendCountTest, "System::Core::HAL::UE_AP
 	STATIC_CHECK(UE_APPEND_VA_ARG_COUNT(TESTPREFIX_, a, b, c, d, e, f, g, h) == TESTPREFIX_8);
 	STATIC_CHECK(UE_APPEND_VA_ARG_COUNT(TESTPREFIX_, a, b, c, d, e, f, g, h, i) == TESTPREFIX_9);
 	STATIC_CHECK(UE_APPEND_VA_ARG_COUNT(TESTPREFIX_, a, b, c, d, e, f, g, h, i, j) == TESTPREFIX_10);
+
+	#undef TESTPREFIX_10
+	#undef TESTPREFIX_9
+	#undef TESTPREFIX_8
+	#undef TESTPREFIX_7
+	#undef TESTPREFIX_6
+	#undef TESTPREFIX_5
+	#undef TESTPREFIX_4
+	#undef TESTPREFIX_3
+	#undef TESTPREFIX_2
+	#undef TESTPREFIX_1
+	#undef TESTPREFIX_0
+}
+
+TEST_CASE_NAMED(FPreprocessorHelperArgAppendCountTest, "System::Core::HAL::UE_JOIN", "[ApplicationContextMask][EngineFilter]")
+{
+	#define TESTMACRO_0 abc
+	#define TESTMACRO_1 defgh
+	#define TESTMACRO_2 ij
+	#define TESTMACRO_3 klmn
+	#define TESTMACRO_4 o
+	#define TESTMACRO_5 pqrstuvwxyz
+
+	CHECK(FCStringAnsi::Strcmp(UE_STRINGIZE(UE_JOIN(TESTMACRO_0, TESTMACRO_1)), "abcdefgh") == 0);
+	CHECK(FCStringAnsi::Strcmp(UE_STRINGIZE(UE_JOIN(TESTMACRO_0, TESTMACRO_1, TESTMACRO_2)), "abcdefghij") == 0);
+	CHECK(FCStringAnsi::Strcmp(UE_STRINGIZE(UE_JOIN(TESTMACRO_0, TESTMACRO_1, TESTMACRO_2, TESTMACRO_3)), "abcdefghijklmn") == 0);
+	CHECK(FCStringAnsi::Strcmp(UE_STRINGIZE(UE_JOIN(TESTMACRO_0, TESTMACRO_1, TESTMACRO_2, TESTMACRO_3, TESTMACRO_4)), "abcdefghijklmno") == 0);
+	CHECK(FCStringAnsi::Strcmp(UE_STRINGIZE(UE_JOIN(TESTMACRO_0, TESTMACRO_1, TESTMACRO_2, TESTMACRO_3, TESTMACRO_4, TESTMACRO_5)), "abcdefghijklmnopqrstuvwxyz") == 0);
+
+	#undef TESTMACRO_5
+	#undef TESTMACRO_4
+	#undef TESTMACRO_3
+	#undef TESTMACRO_2
+	#undef TESTMACRO_1
+	#undef TESTMACRO_0
 }
 
 #endif //WITH_TESTS
diff --git a/Engine/Source/Runtime/RenderCore/Public/RenderingThread.h b/Engine/Source/Runtime/RenderCore/Public/RenderingThread.h
--- a/Engine/Source/Runtime/RenderCore/Public/RenderingThread.h
+++ b/Engine/Source/Runtime/RenderCore/Public/RenderingThread.h
@@ -238,12 +238,12 @@ private:
 
 /** Declares a new render command tag type from a name. */
 #define DECLARE_RENDER_COMMAND_TAG(Type, Name) \
-	struct UE_JOIN(TSTR_, UE_JOIN(Name, __LINE__)) \
+	struct UE_JOIN(TSTR_, Name, __LINE__) \
 	{  \
 		static const char* CStr() { return #Name; } \
 		static const TCHAR* TStr() { return TEXT(#Name); } \
 	}; \
-	using Type = TRenderCommandTag<UE_JOIN(TSTR_, UE_JOIN(Name, __LINE__))>;
+	using Type = TRenderCommandTag<UE_JOIN(TSTR_, Name, __LINE__)>;
 
 /** Describes which pipes are configured to use the render command pipe system. */
 enum class ERenderCommandPipeMode : uint8
@@ -1165,8 +1165,8 @@ public:
  *  pipe, if enabled.
  */
 #define ENQUEUE_RENDER_COMMAND(Type) \
-	DECLARE_RENDER_COMMAND_TAG(UE_JOIN(FRenderCommandTag_, UE_JOIN(Type, __LINE__)), Type) \
-	FRenderCommandDispatcher::Enqueue<UE_JOIN(FRenderCommandTag_, UE_JOIN(Type, __LINE__))>
+	DECLARE_RENDER_COMMAND_TAG(UE_JOIN(FRenderCommandTag_, Type, __LINE__), Type) \
+	FRenderCommandDispatcher::Enqueue<UE_JOIN(FRenderCommandTag_, Type, __LINE__)>
 
 ////////////////////////////////////
 // RenderThread scoped work
