Suppressed PVS warnings about Strcmp's arguments being equal.

#rnx

diff --git a/Engine/Source/Runtime/Core/Tests/HAL/PreprocessorHelperTests.cpp b/Engine/Source/Runtime/Core/Tests/HAL/PreprocessorHelperTests.cpp
--- a/Engine/Source/Runtime/Core/Tests/HAL/PreprocessorHelperTests.cpp
+++ b/Engine/Source/Runtime/Core/Tests/HAL/PreprocessorHelperTests.cpp
@@ -69,11 +69,11 @@ TEST_CASE_NAMED(FPreprocessorHelperArgAppendCountTest, "System::Core::HAL::UE_JO
 	#define TESTMACRO_4 o
 	#define TESTMACRO_5 pqrstuvwxyz
 
-	CHECK(FCStringAnsi::Strcmp(UE_STRINGIZE(UE_JOIN(TESTMACRO_0, TESTMACRO_1)), "abcdefgh") == 0);
-	CHECK(FCStringAnsi::Strcmp(UE_STRINGIZE(UE_JOIN(TESTMACRO_0, TESTMACRO_1, TESTMACRO_2)), "abcdefghij") == 0);
-	CHECK(FCStringAnsi::Strcmp(UE_STRINGIZE(UE_JOIN(TESTMACRO_0, TESTMACRO_1, TESTMACRO_2, TESTMACRO_3)), "abcdefghijklmn") == 0);
-	CHECK(FCStringAnsi::Strcmp(UE_STRINGIZE(UE_JOIN(TESTMACRO_0, TESTMACRO_1, TESTMACRO_2, TESTMACRO_3, TESTMACRO_4)), "abcdefghijklmno") == 0);
-	CHECK(FCStringAnsi::Strcmp(UE_STRINGIZE(UE_JOIN(TESTMACRO_0, TESTMACRO_1, TESTMACRO_2, TESTMACRO_3, TESTMACRO_4, TESTMACRO_5)), "abcdefghijklmnopqrstuvwxyz") == 0);
+	CHECK(FCStringAnsi::Strcmp(UE_STRINGIZE(UE_JOIN(TESTMACRO_0, TESTMACRO_1)), "abcdefgh") == 0); //-V549
+	CHECK(FCStringAnsi::Strcmp(UE_STRINGIZE(UE_JOIN(TESTMACRO_0, TESTMACRO_1, TESTMACRO_2)), "abcdefghij") == 0); //-V549
+	CHECK(FCStringAnsi::Strcmp(UE_STRINGIZE(UE_JOIN(TESTMACRO_0, TESTMACRO_1, TESTMACRO_2, TESTMACRO_3)), "abcdefghijklmn") == 0); //-V549
+	CHECK(FCStringAnsi::Strcmp(UE_STRINGIZE(UE_JOIN(TESTMACRO_0, TESTMACRO_1, TESTMACRO_2, TESTMACRO_3, TESTMACRO_4)), "abcdefghijklmno") == 0); //-V549
+	CHECK(FCStringAnsi::Strcmp(UE_STRINGIZE(UE_JOIN(TESTMACRO_0, TESTMACRO_1, TESTMACRO_2, TESTMACRO_3, TESTMACRO_4, TESTMACRO_5)), "abcdefghijklmnopqrstuvwxyz") == 0); //-V549
 
 	#undef TESTMACRO_5
 	#undef TESTMACRO_4
