Removed UE_ALLOWSHRINKING_BOOL_DEPRECATED.

diff --git a/Engine/Plugins/Experimental/UAF/UAF/Source/UAF/Public/TransformArray.h b/Engine/Plugins/Experimental/UAF/UAF/Source/UAF/Public/TransformArray.h
--- a/Engine/Plugins/Experimental/UAF/UAF/Source/UAF/Public/TransformArray.h
+++ b/Engine/Plugins/Experimental/UAF/UAF/Source/UAF/Public/TransformArray.h
@@ -44,21 +44,11 @@ struct TTransformArrayAoS
 	{
 		Transforms.SetNum(NumTransforms, AllowShrinking);
 	}
-	UE_ALLOWSHRINKING_BOOL_DEPRECATED("SetNum")
-	inline void SetNum(int32 NumTransforms, bool bAllowShrinking)
-	{
-		SetNum(NumTransforms, bAllowShrinking ? EAllowShrinking::Yes : EAllowShrinking::No);
-	}
 
 	void SetNumUninitialized(int32 NumTransforms, EAllowShrinking AllowShrinking = UE::Core::Private::AllowShrinkingByDefault<AllocatorType>())
 	{
 		Transforms.SetNumUninitialized(NumTransforms, AllowShrinking);
 	}
-	UE_ALLOWSHRINKING_BOOL_DEPRECATED("SetNumUninitialized")
-	inline void SetNumUninitialized(int32 NumTransforms, bool bAllowShrinking)
-	{
-		SetNumUninitialized(NumTransforms, bAllowShrinking ? EAllowShrinking::Yes : EAllowShrinking::No);
-	}
 
 	inline void SetIdentity(bool bAdditiveIdentity = false)
 	{
diff --git a/Engine/Plugins/Runtime/ReplicationGraph/Source/Public/ReplicationGraphTypes.h b/Engine/Plugins/Runtime/ReplicationGraph/Source/Public/ReplicationGraphTypes.h
--- a/Engine/Plugins/Runtime/ReplicationGraph/Source/Public/ReplicationGraphTypes.h
+++ b/Engine/Plugins/Runtime/ReplicationGraph/Source/Public/ReplicationGraphTypes.h
@@ -516,11 +516,6 @@ struct FActorRepListRefView
 	{
 		return RepList.RemoveSingleSwap(ElementToRemove, AllowShrink) > 0;
 	}
-	UE_ALLOWSHRINKING_BOOL_DEPRECATED("RemoveFast")
-	inline bool RemoveFast(const FActorRepListType& ElementToRemove, bool bAllowShrink)
-	{
-		return RemoveFast(ElementToRemove, bAllowShrink ? EAllowShrinking::Yes : EAllowShrinking::No);
-	}
 
 	/** Removes the element but keeps the order intact. Generally not recommended for large lists. */
 	bool RemoveSlow(const FActorRepListType& ElementToRemove)
diff --git a/Engine/Source/Runtime/Core/Public/Containers/AllowShrinking.h b/Engine/Source/Runtime/Core/Public/Containers/AllowShrinking.h
--- a/Engine/Source/Runtime/Core/Public/Containers/AllowShrinking.h
+++ b/Engine/Source/Runtime/Core/Public/Containers/AllowShrinking.h
@@ -28,4 +28,3 @@ namespace UE::Core::Private
 	}
 }
 
-#define UE_ALLOWSHRINKING_BOOL_DEPRECATED(FunctionName) UE_DEPRECATED(5.6, FunctionName " with a boolean bAllowShrinking has been deprecated - please use the EAllowShrinking enum instead")
diff --git a/Engine/Source/Runtime/Core/Public/Containers/Array.h b/Engine/Source/Runtime/Core/Public/Containers/Array.h
--- a/Engine/Source/Runtime/Core/Public/Containers/Array.h
+++ b/Engine/Source/Runtime/Core/Public/Containers/Array.h
@@ -402,9 +402,9 @@ namespace UE::Core::Private
 		return Result;
 	}
 
-
+	
 	// Called only when we KNOW we are going to do a realloc increasing by 1.
-	// In this case, we know that max == num and can simplify things in a very
+	// In this case, we know that max == num and can simplify things in a very 
 	// hot location in the code.
 	// This returns the old ArrayMax in order to save a register clobber/reload.
 	template <uint32 Flags, typename AllocatorInstanceType>
@@ -414,7 +414,7 @@ namespace UE::Core::Private
 		AllocatorInstanceType&                       AllocatorInstance,
 		TAllocatorSizeType_T<AllocatorInstanceType>& ArrayMax
 	)
-	{
+	{		
 		using SizeType  = TAllocatorSizeType_T<AllocatorInstanceType>;
 		using USizeType = std::make_unsigned_t<SizeType>;
 
@@ -422,7 +422,7 @@ namespace UE::Core::Private
 		const USizeType UNewNum = UOldMax + 1U;
 		const SizeType  OldMax  = (SizeType)UOldMax;
 		const SizeType  NewNum  = (SizeType)UNewNum;
-
+		
 		// This should only happen when we've underflowed or overflowed SizeType
 		if (NewNum < OldMax)
 		{
@@ -460,7 +460,7 @@ namespace UE::Core::Private
 		TAllocatorSizeType_T<AllocatorInstanceType>& ArrayMax
 	)
 	{
-		return ReallocGrow1_DoAlloc_Impl<Flags, AllocatorInstanceType>(ElementSizeAndAlignment & 0xff, ElementSizeAndAlignment >> 8, AllocatorInstance, ArrayMax);
+		return ReallocGrow1_DoAlloc_Impl<Flags, AllocatorInstanceType>(ElementSizeAndAlignment & 0xff, ElementSizeAndAlignment >> 8, AllocatorInstance, ArrayMax);		
 	}
 
 	template <uint32 Flags, typename AllocatorInstanceType>
@@ -471,7 +471,7 @@ namespace UE::Core::Private
 		TAllocatorSizeType_T<AllocatorInstanceType>& ArrayMax
 	)
 	{
-		return ReallocGrow1_DoAlloc_Impl<Flags, AllocatorInstanceType>(ElementSize, ElementAlignment, AllocatorInstance, ArrayMax);
+		return ReallocGrow1_DoAlloc_Impl<Flags, AllocatorInstanceType>(ElementSize, ElementAlignment, AllocatorInstance, ArrayMax);		
 	}
 
 	// This should be used for repeated growing operations when reallocations are to be amortized over multiple inserts.
@@ -657,9 +657,9 @@ namespace UE::Core::Private
 /**
  * Templated dynamic array
  *
- * A dynamically sized array of typed elements.  Makes the assumption that your elements are trivially relocatable;
- * i.e. that they can be transparently moved to new memory without a copy constructor.  The main implication
- * is that pointers to elements in the TArray may be invalidated by adding or removing other elements to the array.
+ * A dynamically sized array of typed elements.  Makes the assumption that your elements are trivially relocatable; 
+ * i.e. that they can be transparently moved to new memory without a copy constructor.  The main implication 
+ * is that pointers to elements in the TArray may be invalidated by adding or removing other elements to the array. 
  * Removal of elements is O(N) and invalidates the indices of subsequent elements.
  *
  * Caution: as noted below some methods are not safe for element types that require constructors.
@@ -1039,7 +1039,7 @@ public:
 		return const_cast<TArray*>(this)->GetData();
 	}
 
-	/**
+	/** 
 	 * Helper function returning the size of the inner type.
 	 *
 	 * @returns Size in bytes of array type.
@@ -1049,7 +1049,7 @@ public:
 		return sizeof(ElementType);
 	}
 
-	/**
+	/** 
 	 * Helper function to return the amount of memory allocated by this
 	 * container.
 	 * Only returns the size of allocations made directly by the container, not the elements themselves.
@@ -1125,7 +1125,7 @@ public:
 	}
 
 	/**
-	 * Returns true if the array is empty and contains no elements.
+	 * Returns true if the array is empty and contains no elements. 
 	 *
 	 * @returns True if the array is empty.
 	 * @see Num
@@ -1209,11 +1209,6 @@ public:
 		}
 		return Result;
 	}
-	UE_ALLOWSHRINKING_BOOL_DEPRECATED("Pop")
-	UE_NODEBUG UE_FORCEINLINE_HINT ElementType Pop(bool bAllowShrinking)
-	{
-		return Pop(bAllowShrinking ? EAllowShrinking::Yes : EAllowShrinking::No);
-	}
 
 	/**
 	 * Pushes element into the array.
@@ -1568,7 +1563,7 @@ public:
 	/**
 	 * Bulk serialize array as a single memory blob when loading. Uses regular serialization code for saving
 	 * and doesn't serialize at all otherwise (e.g. transient, garbage collection, ...).
-	 *
+	 * 
 	 * Requirements:
 	 *   - T's << operator needs to serialize ALL member variables in the SAME order they are layed out in memory.
 	 *   - T's << operator can NOT perform any fixup operations. This limitation can be lifted by manually copying
@@ -1595,11 +1590,11 @@ public:
 		// Serialize element size to detect mismatch across platforms.
 		int32 SerializedElementSize = ElementSize;
 		Ar << SerializedElementSize;
-
+	
 		if (bForcePerElementSerialization
 			|| (Ar.IsSaving()			// if we are saving, we always do the ordinary serialize as a way to make sure it matches up with bulk serialization
 			&& !Ar.IsCooking()			// but cooking and transacting is performance critical, so we skip that
-			&& !Ar.IsTransacting())
+			&& !Ar.IsTransacting())		
 			|| Ar.IsByteSwapping()		// if we are byteswapping, we need to do that per-element
 			)
 		{
@@ -1667,7 +1662,7 @@ public:
 		// here until we get in to the NOINLINE functions that actually handle the growth.
 
 		// Modulo some register allocation and pointer offsets, inlined sites on x64 for the "Tiny" path should look roughly like:
-		//
+		// 
 		// mov eax, dword ptr [array + 8h] // load arraynum
 		// cmp eax, dword ptr [array + 0Ch] // load arraymax
 		// jne [to ArrayNum++]
@@ -1683,14 +1678,14 @@ public:
 		if (ArrayNum == ArrayMax)
 		{
 			// Both branches here write the return in to ArrayNum. This is because the function call clobbers the registers
-			// and if we assign as part of the return into something we need, the compiler doesn't have to reload the data
+			// and if we assign as part of the return into something we need, the compiler doesn't have to reload the data 
 			// into the clobbered register.
 
 			// When we can pack size and alignment into a single 16 bit load, we save a parameter setup instruction
 			// for the function call. The 16 bit part is load-bearing for arm codegen.
 			// PVS believes this is overwrought because technically sizeof() <= 255 implies alignof <= 255. I think this is far more
 			// clear and costs nothing.
-			if constexpr (sizeof(ElementType) <= 255 && alignof(ElementType) <= 255) // -V590
+			if constexpr (sizeof(ElementType) <= 255 && alignof(ElementType) <= 255) // -V590 
 			{
 				// Note that the realloc functions are templated ONLY on allocator instance so they are not duplicated
 				// in the code for every type!
@@ -1705,7 +1700,7 @@ public:
 
 		SizeType OldArrayNum = ArrayNum;
 		ArrayNum++;
-		return OldArrayNum;
+		return OldArrayNum;	
 	}
 	UE_FORCEINLINE_HINT SizeType AddUninitialized(SizeType Count)
 	{
@@ -1726,7 +1721,7 @@ private:
 		// Begin sensitive code! See comments in AddUninitialized() before touching!
 		if (ArrayNum == ArrayMax)
 		{
-			if constexpr (sizeof(ElementType) <= 255 && alignof(ElementType) <= 255) // -V590
+			if constexpr (sizeof(ElementType) <= 255 && alignof(ElementType) <= 255) // -V590 
 			{
 				ArrayNum = UE::Core::Private::ReallocGrow1_DoAlloc_Tiny<UE::Core::Private::GetAllocatorFlags<AllocatorType>()>(sizeof(ElementType) | (alignof(ElementType) << 8), AllocatorInstance, ArrayMax);
 			}
@@ -1737,7 +1732,7 @@ private:
 		}
 		// End sensitive code!
 		SizeType OldNum = ArrayNum;
-		ArrayNum++;
+		ArrayNum++;		
 		ElementType* Data = GetData() + Index;
 		RelocateConstructItems<ElementType>((void*)(Data + 1), Data, OldNum - Index);
 	}
@@ -1969,7 +1964,7 @@ public:
 	{
 		CheckAddress(&Item);
 
-		// construct a copy in place at Index (this new operator will insert at
+		// construct a copy in place at Index (this new operator will insert at 
 		// Index, then construct that memory with Item)
 		InsertUninitializedImpl(Index);
 		::new((void*)(GetData() + Index)) ElementType(MoveTempIfPossible(Item));
@@ -1988,7 +1983,7 @@ public:
 	{
 		CheckAddress(&Item);
 
-		// construct a copy in place at Index (this new operator will insert at
+		// construct a copy in place at Index (this new operator will insert at 
 		// Index, then construct that memory with Item)
 		InsertUninitializedImpl(Index);
 		::new((void*)(GetData() + Index)) ElementType(Item);
@@ -2008,7 +2003,7 @@ public:
 	{
 		CheckAddress(&Item);
 
-		// construct a copy in place at Index (this new operator will insert at
+		// construct a copy in place at Index (this new operator will insert at 
 		// Index, then construct that memory with Item)
 		InsertUninitializedImpl(Index);
 		ElementType* Ptr = GetData() + Index;
@@ -2028,7 +2023,7 @@ public:
 	{
 		CheckAddress(&Item);
 
-		// construct a copy in place at Index (this new operator will insert at
+		// construct a copy in place at Index (this new operator will insert at 
 		// Index, then construct that memory with Item)
 		InsertUninitializedImpl(Index);
 		ElementType* Ptr = GetData() + Index;
@@ -2103,10 +2098,8 @@ public:
 	 * @param Count (Optional) Number of elements to remove. Default is 1.
 	 * @param AllowShrinking (Optional) Tells if this call can shrink array if suitable after remove. Default is yes.
 	 */
-	template <UE::CIntegral CountType>
-	UE_FORCEINLINE_HINT void RemoveAt(SizeType Index, CountType Count, EAllowShrinking AllowShrinking = UE::Core::Private::AllowShrinkingByDefault<AllocatorType>())
+	UE_FORCEINLINE_HINT void RemoveAt(SizeType Index, SizeType Count, EAllowShrinking AllowShrinking = UE::Core::Private::AllowShrinkingByDefault<AllocatorType>())
 	{
-		static_assert(!std::is_same_v<CountType, bool>, "TArray::RemoveAt: unexpected bool passed as the Count argument");
 		RangeCheck(Index, Count);
 		if (Count)
 		{
@@ -2123,12 +2116,6 @@ public:
 			}
 		}
 	}
-	template <typename CountType>
-	UE_ALLOWSHRINKING_BOOL_DEPRECATED("RemoveAt")
-	UE_FORCEINLINE_HINT void RemoveAt(SizeType Index, CountType Count, bool bAllowShrinking)
-	{
-		RemoveAt(Index, Count, bAllowShrinking ? EAllowShrinking::Yes : EAllowShrinking::No);
-	}
 
 private:
 	void RemoveAtSwapImpl(SizeType Index)
@@ -2209,10 +2196,8 @@ public:
 	 * @param AllowShrinking (Optional) Tells if this call can shrink array if
 	 *                        suitable after remove. Default is yes.
 	 */
-	template <UE::CIntegral CountType>
-	UE_FORCEINLINE_HINT void RemoveAtSwap(SizeType Index, CountType Count, EAllowShrinking AllowShrinking = UE::Core::Private::AllowShrinkingByDefault<AllocatorType>())
+	UE_FORCEINLINE_HINT void RemoveAtSwap(SizeType Index, SizeType Count, EAllowShrinking AllowShrinking = UE::Core::Private::AllowShrinkingByDefault<AllocatorType>())
 	{
-		static_assert(!std::is_same_v<CountType, bool>, "TArray::RemoveAtSwap: unexpected bool passed as the Count argument");
 		RangeCheck(Index, Count);
 		if (Count)
 		{
@@ -2229,12 +2214,6 @@ public:
 			}
 		}
 	}
-	template <typename CountType>
-	UE_ALLOWSHRINKING_BOOL_DEPRECATED("RemoveAtSwap")
-	UE_NODEBUG UE_FORCEINLINE_HINT void RemoveAtSwap(SizeType Index, CountType Count, bool bAllowShrinking)
-	{
-		RemoveAtSwap(Index, Count, bAllowShrinking ? EAllowShrinking::Yes : EAllowShrinking::No);
-	}
 
 	/**
 	 * Same as empty, but doesn't change memory allocations, unless the new size is larger than
@@ -2321,11 +2300,6 @@ public:
 			RemoveAt(NewNum, Num() - NewNum, AllowShrinking);
 		}
 	}
-	UE_ALLOWSHRINKING_BOOL_DEPRECATED("SetNum")
-	UE_NODEBUG UE_FORCEINLINE_HINT void SetNum(SizeType NewNum, bool bAllowShrinking)
-	{
-		SetNum(NewNum, bAllowShrinking ? EAllowShrinking::Yes : EAllowShrinking::No);
-	}
 
 	/**
 	 * Resizes array to given number of elements, optionally shrinking it.
@@ -2350,11 +2324,6 @@ public:
 			RemoveAt(NewNum, Num() - NewNum, AllowShrinking);
 		}
 	}
-	UE_ALLOWSHRINKING_BOOL_DEPRECATED("SetNumZeroed")
-	UE_NODEBUG UE_FORCEINLINE_HINT void SetNumZeroed(SizeType NewNum, bool bAllowShrinking)
-	{
-		SetNumZeroed(NewNum, bAllowShrinking ? EAllowShrinking::Yes : EAllowShrinking::No);
-	}
 
 	/**
 	 * Resizes array to given number of elements. New elements will be uninitialized.
@@ -2378,11 +2347,6 @@ public:
 			RemoveAt(NewNum, Num() - NewNum, AllowShrinking);
 		}
 	}
-	UE_ALLOWSHRINKING_BOOL_DEPRECATED("SetNumUninitialized")
-	UE_NODEBUG UE_FORCEINLINE_HINT void SetNumUninitialized(SizeType NewNum, bool bAllowShrinking)
-	{
-		SetNumUninitialized(NewNum, bAllowShrinking ? EAllowShrinking::Yes : EAllowShrinking::No);
-	}
 
 	/**
 	 * Does nothing except setting the new number of elements in the array. Does not destruct items, does not de-allocate memory.
@@ -2582,7 +2546,7 @@ public:
 		// Begin sensitive code! See comments in AddUninitialized() before touching!
 		if (ArrayNum == ArrayMax)
 		{
-			if constexpr (sizeof(ElementType) <= 255 && alignof(ElementType) <= 255) // -V590
+			if constexpr (sizeof(ElementType) <= 255 && alignof(ElementType) <= 255) // -V590 
 			{
 				ArrayNum = UE::Core::Private::ReallocGrow1_DoAlloc_Tiny<UE::Core::Private::GetAllocatorFlags<AllocatorType>()>(sizeof(ElementType) | (alignof(ElementType) << 8), AllocatorInstance, ArrayMax);
 			}
@@ -2593,7 +2557,7 @@ public:
 		}
 		// End sensitive code!
 		SizeType OldArrayNum = ArrayNum;
-		ArrayNum++;
+		ArrayNum++;		
 		void* Ptr = (char*)AllocatorInstance.GetAllocation() + sizeof(ElementType) * OldArrayNum;
 		(void)new (Ptr) ElementType(Forward<ArgsType>(Args)...);
 		return OldArrayNum;
@@ -2633,7 +2597,7 @@ public:
 		// Begin sensitive code! See comments in AddUninitialized() before touching!
 		if (ArrayNum == ArrayMax)
 		{
-			if constexpr (sizeof(ElementType) <= 255 && alignof(ElementType) <= 255) // -V590
+			if constexpr (sizeof(ElementType) <= 255 && alignof(ElementType) <= 255) // -V590 
 			{
 				// Looks weird but saves the register reload
 				ArrayNum = UE::Core::Private::ReallocGrow1_DoAlloc_Tiny<UE::Core::Private::GetAllocatorFlags<AllocatorType>()>(sizeof(ElementType) | (alignof(ElementType) << 8), AllocatorInstance, ArrayMax);
@@ -2829,7 +2793,7 @@ public:
 		return *reinterpret_cast<TArray<typename ElementCompat::ReinterpretType>*>(this);
 	}
 #endif
-
+  
 	/** Immutable implicit conversion operator to constant container of compatible element type. */
 	template <
 		typename AliasElementType = ElementType
@@ -3185,12 +3149,6 @@ public:
 
 		return OriginalNum - ArrayNum;
 	}
-	template <class PREDICATE_CLASS>
-	UE_ALLOWSHRINKING_BOOL_DEPRECATED("RemoveAllSwap")
-	UE_NODEBUG UE_FORCEINLINE_HINT SizeType RemoveAllSwap(const PREDICATE_CLASS& Predicate, bool bAllowShrinking)
-	{
-		return RemoveAllSwap(Predicate, bAllowShrinking ? EAllowShrinking::Yes : EAllowShrinking::No);
-	}
 
 	/**
 	 * Removes the first occurrence of the specified item in the array. This version is much more efficient
@@ -3215,11 +3173,6 @@ public:
 		// Removed one item
 		return 1;
 	}
-	UE_ALLOWSHRINKING_BOOL_DEPRECATED("RemoveSingleSwap")
-	UE_NODEBUG UE_FORCEINLINE_HINT SizeType RemoveSingleSwap(const ElementType& Item, bool bAllowShrinking)
-	{
-		return RemoveSingleSwap(Item, bAllowShrinking ? EAllowShrinking::Yes : EAllowShrinking::No);
-	}
 
 	/**
 	 * Removes all instances of a given item from the array.
@@ -3262,11 +3215,6 @@ public:
 
 		return OriginalNum - ArrayNum;
 	}
-	UE_ALLOWSHRINKING_BOOL_DEPRECATED("RemoveSwap")
-	UE_NODEBUG UE_FORCEINLINE_HINT SizeType RemoveSwap(const ElementType& Item, bool bAllowShrinking)
-	{
-		return RemoveSwap(Item, bAllowShrinking ? EAllowShrinking::Yes : EAllowShrinking::No);
-	}
 
 	/**
 	 * Element-wise array memory swap.
@@ -3467,7 +3415,7 @@ public:
 #if defined(_MSC_VER) && !defined(__clang__)	// Relies on MSVC-specific lazy template instantiation to support arrays of incomplete types
 private:
 	/**
-	 * Helper function that can be used inside the debuggers watch window to debug TArrays. E.g. "*Class->Defaults.DebugGet(5)".
+	 * Helper function that can be used inside the debuggers watch window to debug TArrays. E.g. "*Class->Defaults.DebugGet(5)". 
 	 *
 	 * @param Index Position to get.
 	 * @returns Reference to the element at given position.
@@ -3620,7 +3568,7 @@ public:
 	* Implicit heaps
 	*/
 public:
-	/**
+	/** 
 	 * Builds an implicit heap from the array.
 	 *
 	 * @param Predicate Predicate class instance.
@@ -3649,7 +3597,7 @@ public:
 		Heapify(TLess<ElementType>());
 	}
 
-	/**
+	/** 
 	 * Adds a new element to the heap.
 	 *
 	 * @param InItem Item to be added.
@@ -3671,7 +3619,7 @@ public:
 		return Result;
 	}
 
-	/**
+	/** 
 	 * Adds a new element to the heap.
 	 *
 	 * @param InItem Item to be added.
@@ -3693,7 +3641,7 @@ public:
 		return Result;
 	}
 
-	/**
+	/** 
 	 * Adds a new element to the heap. Assumes < operator is defined for the
 	 * template type.
 	 *
@@ -3709,7 +3657,7 @@ public:
 		return HeapPush(MoveTempIfPossible(InItem), TLess<ElementType>());
 	}
 
-	/**
+	/** 
 	 * Adds a new element to the heap. Assumes < operator is defined for the
 	 * template type.
 	 *
@@ -3725,7 +3673,7 @@ public:
 		return HeapPush(InItem, TLess<ElementType>());
 	}
 
-	/**
+	/** 
 	 * Removes the top element from the heap.
 	 *
 	 * @param OutItem The removed item.
@@ -3745,14 +3693,8 @@ public:
 		TDereferenceWrapper< ElementType, PREDICATE_CLASS> PredicateWrapper(Predicate);
 		AlgoImpl::HeapSiftDown(GetData(), (SizeType)0, Num(), FIdentityFunctor(), PredicateWrapper);
 	}
-	template <class PREDICATE_CLASS>
-	UE_ALLOWSHRINKING_BOOL_DEPRECATED("HeapPop")
-	UE_NODEBUG UE_FORCEINLINE_HINT void HeapPop(ElementType& OutItem, const PREDICATE_CLASS& Predicate, bool bAllowShrinking)
-	{
-		HeapPop(OutItem, Predicate, bAllowShrinking ? EAllowShrinking::Yes : EAllowShrinking::No);
-	}
 
-	/**
+	/** 
 	 * Removes the top element from the heap. Assumes < operator is defined for
 	 * the template type.
 	 *
@@ -3767,11 +3709,6 @@ public:
 	{
 		HeapPop(OutItem, TLess<ElementType>(), AllowShrinking);
 	}
-	UE_ALLOWSHRINKING_BOOL_DEPRECATED("HeapPop")
-	UE_NODEBUG UE_FORCEINLINE_HINT void HeapPop(ElementType& OutItem, bool bAllowShrinking)
-	{
-		HeapPop(OutItem, bAllowShrinking ? EAllowShrinking::Yes : EAllowShrinking::No);
-	}
 
 	/**
 	 * Verifies the heap.
@@ -3784,7 +3721,7 @@ public:
 		check(Algo::IsHeap(*this, Predicate));
 	}
 
-	/**
+	/** 
 	 * Removes the top element from the heap.
 	 *
 	 * @param Predicate Predicate class instance.
@@ -3801,14 +3738,8 @@ public:
 		TDereferenceWrapper< ElementType, PREDICATE_CLASS> PredicateWrapper(Predicate);
 		AlgoImpl::HeapSiftDown(GetData(), (SizeType)0, Num(), FIdentityFunctor(), PredicateWrapper);
 	}
-	template <class PREDICATE_CLASS>
-	UE_ALLOWSHRINKING_BOOL_DEPRECATED("HeapPopDiscard")
-	UE_NODEBUG UE_FORCEINLINE_HINT void HeapPopDiscard(const PREDICATE_CLASS& Predicate, bool bAllowShrinking)
-	{
-		HeapPopDiscard(Predicate, bAllowShrinking ? EAllowShrinking::Yes : EAllowShrinking::No);
-	}
 
-	/**
+	/** 
 	 * Removes the top element from the heap. Assumes < operator is defined for the template type.
 	 *
 	 * @param AllowShrinking (Optional) Tells if this call can shrink the array
@@ -3822,13 +3753,8 @@ public:
 	{
 		HeapPopDiscard(TLess<ElementType>(), AllowShrinking);
 	}
-	UE_ALLOWSHRINKING_BOOL_DEPRECATED("HeapPopDiscard")
-	UE_NODEBUG UE_FORCEINLINE_HINT void HeapPopDiscard(bool bAllowShrinking)
-	{
-		HeapPopDiscard(TLess<ElementType>(), bAllowShrinking ? EAllowShrinking::Yes : EAllowShrinking::No);
-	}
 
-	/**
+	/** 
 	 * Returns the top element from the heap (does not remove the element).
 	 *
 	 * Const version.
@@ -3840,7 +3766,7 @@ public:
 		return (*this)[0];
 	}
 
-	/**
+	/** 
 	 * Returns the top element from the heap (does not remove the element).
 	 *
 	 * @returns The reference to the top element from the heap.
@@ -3871,12 +3797,6 @@ public:
 		AlgoImpl::HeapSiftDown(GetData(), Index, Num(), FIdentityFunctor(), PredicateWrapper);
 		AlgoImpl::HeapSiftUp(GetData(), (SizeType)0, FPlatformMath::Min(Index, Num() - 1), FIdentityFunctor(), PredicateWrapper);
 	}
-	template <class PREDICATE_CLASS>
-	UE_ALLOWSHRINKING_BOOL_DEPRECATED("HeapRemoveAt")
-	UE_NODEBUG void HeapRemoveAt(SizeType Index, const PREDICATE_CLASS& Predicate, bool bAllowShrinking)
-	{
-		HeapRemoveAt(Index, Predicate, bAllowShrinking ? EAllowShrinking::Yes : EAllowShrinking::No);
-	}
 
 	/**
 	 * Removes an element from the heap. Assumes < operator is defined for the template type.
@@ -3893,11 +3813,6 @@ public:
 	{
 		HeapRemoveAt(Index, TLess< ElementType >(), AllowShrinking);
 	}
-	UE_ALLOWSHRINKING_BOOL_DEPRECATED("HeapRemoveAt")
-	UE_NODEBUG UE_FORCEINLINE_HINT void HeapRemoveAt(SizeType Index, bool bAllowShrinking)
-	{
-		HeapRemoveAt(Index, bAllowShrinking ? EAllowShrinking::Yes : EAllowShrinking::No);
-	}
 
 	/**
 	 * Performs heap sort on the array.
diff --git a/Engine/Source/Runtime/Core/Public/Containers/IndirectArray.h b/Engine/Source/Runtime/Core/Public/Containers/IndirectArray.h
--- a/Engine/Source/Runtime/Core/Public/Containers/IndirectArray.h
+++ b/Engine/Source/Runtime/Core/Public/Containers/IndirectArray.h
@@ -282,11 +282,6 @@ public:
 		}
 		Array.RemoveAt(Index, Count, AllowShrinking);
 	}
-	UE_ALLOWSHRINKING_BOOL_DEPRECATED("RemoveAt")
-	UE_FORCEINLINE_HINT void RemoveAt(int32 Index, int32 Count, bool bAllowShrinking)
-	{
-		RemoveAt(Index, Count, bAllowShrinking ? EAllowShrinking::Yes : EAllowShrinking::No);
-	}
 
 	/**
 	 * Removes an element at given location optionally shrinking
@@ -333,11 +328,6 @@ public:
 		}
 		Array.RemoveAtSwap(Index, Count, AllowShrinking);
 	}
-	UE_ALLOWSHRINKING_BOOL_DEPRECATED("RemoveAtSwap")
-	UE_FORCEINLINE_HINT void RemoveAtSwap(int32 Index, int32 Count, bool bAllowShrinking)
-	{
-		RemoveAtSwap(Index, Count, bAllowShrinking ? EAllowShrinking::Yes : EAllowShrinking::No);
-	}
 
 	/**
 	 * Element-wise array element swap.
diff --git a/Engine/Source/Runtime/Core/Public/Containers/PagedArray.h b/Engine/Source/Runtime/Core/Public/Containers/PagedArray.h
--- a/Engine/Source/Runtime/Core/Public/Containers/PagedArray.h
+++ b/Engine/Source/Runtime/Core/Public/Containers/PagedArray.h
@@ -355,11 +355,6 @@ public:
 		}
 		Count = NewNum;
 	}
-	UE_ALLOWSHRINKING_BOOL_DEPRECATED("SetNum")
-	UE_FORCEINLINE_HINT void SetNum(SizeType NewNum, bool bAllowShrinking)
-	{
-		SetNum(NewNum, bAllowShrinking ? EAllowShrinking::Yes : EAllowShrinking::No);
-	}
 
 	/*
 	 * Fills the memory used by the container elements with the parameter byte value.
@@ -540,11 +535,6 @@ public:
 		}
 		--Count;
 	}
-	UE_ALLOWSHRINKING_BOOL_DEPRECATED("Pop")
-	UE_FORCEINLINE_HINT void Pop(bool bAllowShrinking)
-	{
-		Pop(bAllowShrinking ? EAllowShrinking::Yes : EAllowShrinking::No);
-	}
 
 	/**
 	 * Removes the element at the parameter index position and swaps the last element if existent to the same position
@@ -574,11 +564,6 @@ public:
 		}
 		--Count;
 	}
-	UE_ALLOWSHRINKING_BOOL_DEPRECATED("RemoveAtSwap")
-	UE_FORCEINLINE_HINT void RemoveAtSwap(SizeType Index, bool bAllowShrinking)
-	{
-		RemoveAtSwap(Index, bAllowShrinking ? EAllowShrinking::Yes : EAllowShrinking::No);
-	}
 
 	/*
 	 * Empties the container effectively destroying all contained elements and releases any acquired storage.
diff --git a/Engine/Source/Runtime/Core/Public/Containers/ScriptArray.h b/Engine/Source/Runtime/Core/Public/Containers/ScriptArray.h
--- a/Engine/Source/Runtime/Core/Public/Containers/ScriptArray.h
+++ b/Engine/Source/Runtime/Core/Public/Containers/ScriptArray.h
@@ -123,11 +123,6 @@ public:
 			Remove(NewNum, OldNum - NewNum, NumBytesPerElement, AlignmentOfElement, AllowShrinking);
 		}
 	}
-	UE_ALLOWSHRINKING_BOOL_DEPRECATED("SetNumUninitialized")
-	FORCEINLINE void SetNumUninitialized(int32 NewNum, int32 NumBytesPerElement, uint32 AlignmentOfElement, bool bAllowShrinking)
-	{
-		SetNumUninitialized(NewNum, NumBytesPerElement, AlignmentOfElement, bAllowShrinking ? EAllowShrinking::Yes : EAllowShrinking::No);
-	}
 	void MoveAssign(TScriptArray& Other, int32 NumBytesPerElement, uint32 AlignmentOfElement)
 	{
 		checkSlow(this != &Other);
@@ -225,11 +220,6 @@ public:
 			checkSlow(ArrayMax >= ArrayNum);
 		}
 	}
-	UE_ALLOWSHRINKING_BOOL_DEPRECATED("Remove")
-	FORCEINLINE void Remove(int32 Index, int32 Count, int32 NumBytesPerElement, uint32 AlignmentOfElement, bool bAllowShrinking)
-	{
-		Remove(Index, Count, NumBytesPerElement, AlignmentOfElement, bAllowShrinking ? EAllowShrinking::Yes : EAllowShrinking::No);
-	}
 	[[nodiscard]] SIZE_T GetAllocatedSize(int32 NumBytesPerElement) const
 	{
 		return ((const typename AllocatorType::ForAnyElementType*)this)->GetAllocatedSize(ArrayMax, NumBytesPerElement);
diff --git a/Engine/Source/Runtime/Core/Public/Containers/UnrealString.h.inl b/Engine/Source/Runtime/Core/Public/Containers/UnrealString.h.inl
--- a/Engine/Source/Runtime/Core/Public/Containers/UnrealString.h.inl
+++ b/Engine/Source/Runtime/Core/Public/Containers/UnrealString.h.inl
@@ -494,11 +494,6 @@ public:
 	 * @param AllowShrinking Whether or not to reallocate to shrink the storage after removal.
 	 */
 	CORE_API void RemoveAt(int32 Index, int32 Count, EAllowShrinking AllowShrinking = EAllowShrinking::Default);
-	UE_ALLOWSHRINKING_BOOL_DEPRECATED("RemoveAt")
-	UE_FORCEINLINE_HINT void RemoveAt(int32 Index, int32 Count, bool bAllowShrinking)
-	{
-		RemoveAt(Index, Count, bAllowShrinking ? EAllowShrinking::Yes : EAllowShrinking::No);
-	}
 
 	/**
 	 * Removes the text from the start of the string if it exists.
@@ -987,11 +982,6 @@ public:
 		Count = FMath::Clamp(Count, 0, Length);
 		RemoveAt(Count, Length-Count, AllowShrinking);
 	}
-	UE_ALLOWSHRINKING_BOOL_DEPRECATED("LeftInline")
-	UE_FORCEINLINE_HINT void LeftInline(int32 Count, bool bAllowShrinking)
-	{
-		LeftInline(Count, bAllowShrinking ? EAllowShrinking::Yes : EAllowShrinking::No);
-	}
 
 	/** Returns the left most characters from the string chopping the given number of characters from the end */
 	[[nodiscard]] inline UE_STRING_CLASS LeftChop(int32 Count) const &
@@ -1012,11 +1002,6 @@ public:
 		const int32 Length = Len();
 		RemoveAt(FMath::Clamp(Length-Count, 0, Length), Count, AllowShrinking);
 	}
-	UE_ALLOWSHRINKING_BOOL_DEPRECATED("LeftChopInline")
-	UE_FORCEINLINE_HINT void LeftChopInline(int32 Count, bool bAllowShrinking)
-	{
-		LeftChopInline(Count, bAllowShrinking ? EAllowShrinking::Yes : EAllowShrinking::No);
-	}
 
 	/** Returns the string to the right of the specified location, counting back from the right (end of the word). */
 	[[nodiscard]] inline UE_STRING_CLASS Right(int32 Count) const &
@@ -1037,11 +1022,6 @@ public:
 		const int32 Length = Len();
 		RemoveAt(0, Length-FMath::Clamp(Count,0,Length), AllowShrinking);
 	}
-	UE_ALLOWSHRINKING_BOOL_DEPRECATED("RightInline")
-	UE_FORCEINLINE_HINT void RightInline(int32 Count, bool bAllowShrinking)
-	{
-		RightInline(Count, bAllowShrinking ? EAllowShrinking::Yes : EAllowShrinking::No);
-	}
 
 	/** Returns the string to the right of the specified location, counting forward from the left (from the beginning of the word). */
 	[[nodiscard]] CORE_API UE_STRING_CLASS RightChop(int32 Count) const &;
@@ -1057,11 +1037,6 @@ public:
 	{
 		RemoveAt(0, Count, AllowShrinking);
 	}
-	UE_ALLOWSHRINKING_BOOL_DEPRECATED("RightChopInline")
-	UE_FORCEINLINE_HINT void RightChopInline(int32 Count, bool bAllowShrinking)
-	{
-		RightChopInline(Count, bAllowShrinking ? EAllowShrinking::Yes : EAllowShrinking::No);
-	}
 
 	/** Returns the substring from Start position for Count characters. */
 	[[nodiscard]] CORE_API UE_STRING_CLASS Mid(int32 Start, int32 Count) const &;
@@ -1080,11 +1055,6 @@ public:
 		}
 		RightChopInline(Start, AllowShrinking);
 	}
-	UE_ALLOWSHRINKING_BOOL_DEPRECATED("MidInline")
-	UE_FORCEINLINE_HINT void MidInline(int32 Start, int32 Count, bool bAllowShrinking)
-	{
-		MidInline(Start, Count, bAllowShrinking ? EAllowShrinking::Yes : EAllowShrinking::No);
-	}
 
 	/**
 	 * Searches the string for a substring. Can search from beginning or end, and ignore case or not.
diff --git a/Engine/Source/Runtime/CoreUObject/Private/Serialization/ArchiveStackTrace.cpp b/Engine/Source/Runtime/CoreUObject/Private/Serialization/ArchiveStackTrace.cpp
--- a/Engine/Source/Runtime/CoreUObject/Private/Serialization/ArchiveStackTrace.cpp
+++ b/Engine/Source/Runtime/CoreUObject/Private/Serialization/ArchiveStackTrace.cpp
@@ -244,7 +244,7 @@ FString FArchiveCallstacks::FCallstackData::ToString(const TCHAR* CallstackCutof
 			int32 CutoffIndex = StackLine.Find(TEXT(" "), ESearchCase::CaseSensitive);
 			if (CutoffIndex >= -1 && CutoffIndex < StackLine.Len() - 2)
 			{
-				StackLine.MidInline(CutoffIndex + 1, MAX_int32, false);
+				StackLine.MidInline(CutoffIndex + 1, MAX_int32, EAllowShrinking::No);
 			}
 		}
 		HumanReadableString += Indent;
diff --git a/Engine/Source/Runtime/CoreUObject/Private/Serialization/AsyncLoading.h b/Engine/Source/Runtime/CoreUObject/Private/Serialization/AsyncLoading.h
--- a/Engine/Source/Runtime/CoreUObject/Private/Serialization/AsyncLoading.h
+++ b/Engine/Source/Runtime/CoreUObject/Private/Serialization/AsyncLoading.h
@@ -330,11 +330,6 @@ class FImportOrImportIndexArray : public TArray<int32>
 {
 public:
 	void HeapPop(int32& OutItem, EAllowShrinking AllowShrinking = EAllowShrinking::Default);
-	UE_ALLOWSHRINKING_BOOL_DEPRECATED("HeapPop")
-	FORCEINLINE void HeapPop(int32& OutItem, bool bAllowShrinking)
-	{
-		HeapPop(OutItem, bAllowShrinking ? EAllowShrinking::Yes : EAllowShrinking::No);
-	}
 };
 
 /**
diff --git a/Engine/Source/Runtime/Engine/Classes/Animation/PreviewAssetAttachComponent.h b/Engine/Source/Runtime/Engine/Classes/Animation/PreviewAssetAttachComponent.h
--- a/Engine/Source/Runtime/Engine/Classes/Animation/PreviewAssetAttachComponent.h
+++ b/Engine/Source/Runtime/Engine/Classes/Animation/PreviewAssetAttachComponent.h
@@ -141,11 +141,6 @@ public:
 	 */
 	ENGINE_API void RemoveAtSwap( int32 Index, EAllowShrinking AllowShrinking = EAllowShrinking::Default);
 	ENGINE_API void RemoveAtSwap( int32 Index, int32 Count, EAllowShrinking AllowShrinking = EAllowShrinking::Default);
-	UE_ALLOWSHRINKING_BOOL_DEPRECATED("RemoveAtSwap")
-	inline void RemoveAtSwap(int32 Index, int32 Count, bool bAllowShrinking)
-	{
-		RemoveAtSwap(Index, Count, bAllowShrinking ? EAllowShrinking::Yes : EAllowShrinking::No);
-	}
 
 	/**
 	 * Helper function to fix up attached objects after property deprecation
diff --git a/Engine/Source/Runtime/Experimental/Chaos/Public/RewindData.h b/Engine/Source/Runtime/Experimental/Chaos/Public/RewindData.h
--- a/Engine/Source/Runtime/Experimental/Chaos/Public/RewindData.h
+++ b/Engine/Source/Runtime/Experimental/Chaos/Public/RewindData.h
@@ -1541,11 +1541,6 @@ public:
 			KeyToIdx.Remove(Key);
 		}
 	}
-	UE_ALLOWSHRINKING_BOOL_DEPRECATED("Remove")
-	FORCEINLINE void Remove(const TKey Key, const bool bAllowShrinking)
-	{
-		Remove(Key, bAllowShrinking ? EAllowShrinking::Yes : EAllowShrinking::No);
-	}
 
 	void Shrink()
 	{
@@ -1626,21 +1621,11 @@ public:
 	{
 		DirtyParticles.Remove(Particle, AllowShrinking);
 	}
-	UE_ALLOWSHRINKING_BOOL_DEPRECATED("RemoveObject")
-	FORCEINLINE void RemoveObject(const FGeometryParticleHandle* Particle, const bool bAllowShrinking)
-	{
-		RemoveObject(Particle, bAllowShrinking ? EAllowShrinking::Yes : EAllowShrinking::No);
-	}
 
 	void RemoveObject(const FPBDJointConstraintHandle* Joint, const EAllowShrinking AllowShrinking = EAllowShrinking::Default)
 	{
 		DirtyJoints.Remove(Joint, AllowShrinking);
 	}
-	UE_ALLOWSHRINKING_BOOL_DEPRECATED("RemoveObject")
-	FORCEINLINE void RemoveObject(const FPBDJointConstraintHandle* Joint, const bool bAllowShrinking)
-	{
-		RemoveObject(Joint, bAllowShrinking ? EAllowShrinking::Yes : EAllowShrinking::No);
-	}
 
 	int32 GetEarliestFrame_Internal() const { return CurFrame - FramesSaved; }
 
diff --git a/Engine/Source/Runtime/GeometryCore/Public/DynamicMesh/ColliderMesh.h b/Engine/Source/Runtime/GeometryCore/Public/DynamicMesh/ColliderMesh.h
--- a/Engine/Source/Runtime/GeometryCore/Public/DynamicMesh/ColliderMesh.h
+++ b/Engine/Source/Runtime/GeometryCore/Public/DynamicMesh/ColliderMesh.h
@@ -66,11 +66,6 @@ public:
 	GEOMETRYCORE_API void Initialize(const FDynamicMesh3& SourceMesh, const FBuildOptions& BuildOptions = FBuildOptions());
 
 	GEOMETRYCORE_API void Reset(EAllowShrinking AllowShrinking);
-	UE_ALLOWSHRINKING_BOOL_DEPRECATED("Reset")
-	FORCEINLINE void Reset(bool bReleaseMemory)
-	{
-		Reset(bReleaseMemory ? EAllowShrinking::Yes : EAllowShrinking::No);
-	}
 
 	// mesh API required for TMeshAABBTree3
 	bool IsTriangle(int32 TriangleID) const { return TriangleID >= 0 && TriangleID < Triangles.Num(); }
diff --git a/Engine/Source/Runtime/GeometryCore/Public/Spatial/DenseGrid2.h b/Engine/Source/Runtime/GeometryCore/Public/Spatial/DenseGrid2.h
--- a/Engine/Source/Runtime/GeometryCore/Public/Spatial/DenseGrid2.h
+++ b/Engine/Source/Runtime/GeometryCore/Public/Spatial/DenseGrid2.h
@@ -69,11 +69,6 @@ public:
 		DimensionY = (int64)DimY;
 		Buffer.SetNumUninitialized( DimensionX*DimensionY, AllowShrinking);
 	}
-	UE_ALLOWSHRINKING_BOOL_DEPRECATED("Resize")
-	FORCEINLINE void Resize(int32 DimX, int32 DimY, bool bAllowShrinking)
-	{
-		Resize(DimX, DimY, bAllowShrinking ? EAllowShrinking::Yes : EAllowShrinking::No);
-	}
 
 	void AssignAll(ElemType Value)
 	{
diff --git a/Engine/Source/Runtime/GeometryCore/Public/Spatial/DenseGrid3.h b/Engine/Source/Runtime/GeometryCore/Public/Spatial/DenseGrid3.h
--- a/Engine/Source/Runtime/GeometryCore/Public/Spatial/DenseGrid3.h
+++ b/Engine/Source/Runtime/GeometryCore/Public/Spatial/DenseGrid3.h
@@ -64,11 +64,6 @@ public:
 		Buffer.SetNumUninitialized(DimX * DimY * DimZ, AllowShrinking);
 		Dimensions = FVector3i(DimX, DimY, DimZ);
 	}
-	UE_ALLOWSHRINKING_BOOL_DEPRECATED("Resize")
-	FORCEINLINE void Resize(int DimX, int DimY, int DimZ, bool bAllowShrinking)
-	{
-		Resize(DimX, DimY, DimZ, bAllowShrinking ? EAllowShrinking::Yes : EAllowShrinking::No);
-	}
 
 	void Assign(ElemType Value)
 	{
diff --git a/Engine/Source/Runtime/RenderCore/Public/ShaderSource.h b/Engine/Source/Runtime/RenderCore/Public/ShaderSource.h
--- a/Engine/Source/Runtime/RenderCore/Public/ShaderSource.h
+++ b/Engine/Source/Runtime/RenderCore/Public/ShaderSource.h
@@ -79,11 +79,6 @@ public:
 		checkf(Num <= Len(), TEXT("Trying to shrink to %d characters but existing allocation is smaller (%d characters)"), Num, Len());
 		SetLen(Num, AllowShrinking);
 	}
-	UE_ALLOWSHRINKING_BOOL_DEPRECATED("ShrinkToLen")
-	inline void ShrinkToLen(int32 Num, bool bShrink)
-	{
-		ShrinkToLen(Num, bShrink ? EAllowShrinking::Yes : EAllowShrinking::No);
-	}
 
 	/* String view accessor. Will decompress data if it was previously compressed.
 	 * @return a string view pointing to the source contents, excluding padding.
diff --git a/Engine/Source/Runtime/Renderer/Private/ScenePrivateBase.h b/Engine/Source/Runtime/Renderer/Private/ScenePrivateBase.h
--- a/Engine/Source/Runtime/Renderer/Private/ScenePrivateBase.h
+++ b/Engine/Source/Runtime/Renderer/Private/ScenePrivateBase.h
@@ -94,11 +94,6 @@ public:
 			}
 		}
 	}
-	UE_ALLOWSHRINKING_BOOL_DEPRECATED("Remove")
-	FORCEINLINE void Remove(uint32 Count, bool bAllowShrinking)
-	{
-		Remove(Count, bAllowShrinking ? EAllowShrinking::Yes : EAllowShrinking::No);
-	}
 
 	void Reserve(int32 Count)
 	{
diff --git a/Engine/Source/Runtime/SlateCore/Public/Containers/ObservableArray.h b/Engine/Source/Runtime/SlateCore/Public/Containers/ObservableArray.h
--- a/Engine/Source/Runtime/SlateCore/Public/Containers/ObservableArray.h
+++ b/Engine/Source/Runtime/SlateCore/Public/Containers/ObservableArray.h
@@ -310,11 +310,6 @@ public:
 		RemoveAt(Index, 1, AllowShrinking);
 		return 1;
 	}
-	UE_ALLOWSHRINKING_BOOL_DEPRECATED("RemoveSingle")
-	inline SizeType RemoveSingle(const ElementType& Item, bool bAllowShrinking)
-	{
-		return RemoveSingle(Item, bAllowShrinking ? EAllowShrinking::Yes : EAllowShrinking::No);
-	}
 
 	SizeType RemoveSingleSwap(const ElementType& Item, EAllowShrinking AllowShrinking = EAllowShrinking::Default)
 	{
@@ -326,11 +321,6 @@ public:
 		RemoveAtSwap(Index, 1, AllowShrinking);
 		return 1;
 	}
-	UE_ALLOWSHRINKING_BOOL_DEPRECATED("RemoveSingleSwap")
-	inline SizeType RemoveSingleSwap(const ElementType& Item, bool bAllowShrinking)
-	{
-		return RemoveSingleSwap(Item, bAllowShrinking ? EAllowShrinking::Yes : EAllowShrinking::No);
-	}
 
 	void RemoveAt(SizeType Index, EAllowShrinking AllowShrinking = EAllowShrinking::Default)
 	{
@@ -370,11 +360,6 @@ public:
 			}
 		}
 	}
-	UE_ALLOWSHRINKING_BOOL_DEPRECATED("RemoveAt")
-	inline void RemoveAt(SizeType Index, SizeType NumToRemove, bool bAllowShrinking)
-	{
-		RemoveAt(Index, NumToRemove, bAllowShrinking ? EAllowShrinking::Yes : EAllowShrinking::No);
-	}
 
 	void RemoveAtSwap(SizeType Index, EAllowShrinking AllowShrinking = EAllowShrinking::Default)
 	{
@@ -418,11 +403,6 @@ public:
 			}
 		}
 	}
-	UE_ALLOWSHRINKING_BOOL_DEPRECATED("RemoveAtSwap")
-	inline void RemoveAtSwap(SizeType Index, SizeType NumToRemove, bool bAllowShrinking)
-	{
-		RemoveAtSwap(Index, NumToRemove, bAllowShrinking ? EAllowShrinking::Yes : EAllowShrinking::No);
-	}
 
 	void Swap(SizeType FirstIndexToSwap, SizeType SecondIndexToSwap)
 	{
