Removed misleading (Optional) from RemoveAt and RemoveAtSwap comments.

#rnx

diff --git a/Engine/Source/Runtime/Core/Public/Containers/Array.h b/Engine/Source/Runtime/Core/Public/Containers/Array.h
--- a/Engine/Source/Runtime/Core/Public/Containers/Array.h
+++ b/Engine/Source/Runtime/Core/Public/Containers/Array.h
@@ -2095,7 +2095,7 @@ public:
 	 * the array.
 	 *
 	 * @param Index Location in array of the element to remove.
-	 * @param Count (Optional) Number of elements to remove. Default is 1.
+	 * @param Count Number of elements to remove. Default is 1.
 	 * @param AllowShrinking (Optional) Tells if this call can shrink array if suitable after remove. Default is yes.
 	 */
 	UE_FORCEINLINE_HINT void RemoveAt(SizeType Index, SizeType Count, EAllowShrinking AllowShrinking = UE::Core::Private::AllowShrinkingByDefault<AllocatorType>())
@@ -2192,7 +2192,7 @@ public:
 	 * O(ArrayNum)), but does not preserve the order.
 	 *
 	 * @param Index Location in array of the element to remove.
-	 * @param Count (Optional) Number of elements to remove. Default is 1.
+	 * @param Count Number of elements to remove. Default is 1.
 	 * @param AllowShrinking (Optional) Tells if this call can shrink array if
 	 *                        suitable after remove. Default is yes.
 	 */
