Removed deprecated CHAR32.

diff --git a/Engine/Source/Runtime/Core/Public/GenericPlatform/GenericPlatform.h b/Engine/Source/Runtime/Core/Public/GenericPlatform/GenericPlatform.h
--- a/Engine/Source/Runtime/Core/Public/GenericPlatform/GenericPlatform.h
+++ b/Engine/Source/Runtime/Core/Public/GenericPlatform/GenericPlatform.h
@@ -73,10 +73,6 @@ struct FGenericPlatformTypes
 	// A 16-bit character type. In-memory only.  16-bit representation. Should really be char16_t but making this the generic option is easier for compilers which don't fully support C++11 yet (i.e. MSVC).
 	typedef uint16				CHAR16;		
 
-	// A 32-bit character type. In-memory only. 32-bit representation. Should really be char32_t but making this the generic option is easier for compilers which don't fully support C++11 yet (i.e. MSVC).
-	[[deprecated("CHAR32 has been deprecated - please use UTF32CHAR instead.")]] // not using UE_DEPRECATED(5.6, "...") for dependency reasons
-	typedef uint32				CHAR32;
-
 	// A 32-bit character type. In-memory only. 32-bit representation.
 	using UTF32CHAR = char32_t;
 
