Removed deprecated DEPRECATED_MACRO.

diff --git a/Engine/Source/Runtime/Core/Public/GenericPlatform/GenericPlatformCompilerPreSetup.h b/Engine/Source/Runtime/Core/Public/GenericPlatform/GenericPlatformCompilerPreSetup.h
--- a/Engine/Source/Runtime/Core/Public/GenericPlatform/GenericPlatformCompilerPreSetup.h
+++ b/Engine/Source/Runtime/Core/Public/GenericPlatform/GenericPlatformCompilerPreSetup.h
@@ -38,10 +38,6 @@
 	#endif
 #endif
 
-#ifndef DEPRECATED_MACRO
-	#define DEPRECATED_MACRO(Version, Message) UE_DEPRECATED_MACRO(5.1, "The DEPRECATED_MACRO macro has been deprecated in favor of UE_DEPRECATED_MACRO.") UE_DEPRECATED_MACRO(Version, Message)
-#endif
-
 #ifndef PRAGMA_DEFAULT_VISIBILITY_START
 	#define PRAGMA_DEFAULT_VISIBILITY_START
 #endif
