Removed obsolete includes in Core and CoreUObject.

#rnx

diff --git a/Engine/Source/Runtime/Core/Public/Algo/Impl/BinaryHeap.h b/Engine/Source/Runtime/Core/Public/Algo/Impl/BinaryHeap.h
--- a/Engine/Source/Runtime/Core/Public/Algo/Impl/BinaryHeap.h
+++ b/Engine/Source/Runtime/Core/Public/Algo/Impl/BinaryHeap.h
@@ -201,7 +201,3 @@ namespace AlgoImpl
 		}
 	}
 }
-
-#if UE_ENABLE_INCLUDE_ORDER_DEPRECATED_IN_5_6
-#include "Templates/Invoke.h"
-#endif
diff --git a/Engine/Source/Runtime/Core/Public/Algo/IntroSort.h b/Engine/Source/Runtime/Core/Public/Algo/IntroSort.h
--- a/Engine/Source/Runtime/Core/Public/Algo/IntroSort.h
+++ b/Engine/Source/Runtime/Core/Public/Algo/IntroSort.h
@@ -198,7 +198,3 @@ namespace Algo
 #endif
 	}
 }
-
-#if UE_ENABLE_INCLUDE_ORDER_DEPRECATED_IN_5_6
-#include "Templates/Invoke.h"
-#endif
diff --git a/Engine/Source/Runtime/Core/Public/Async/Fundamental/Scheduler.h b/Engine/Source/Runtime/Core/Public/Async/Fundamental/Scheduler.h
--- a/Engine/Source/Runtime/Core/Public/Async/Fundamental/Scheduler.h
+++ b/Engine/Source/Runtime/Core/Public/Async/Fundamental/Scheduler.h
@@ -222,7 +222,3 @@ namespace LowLevelTasks
 		StopWorkers();
 	}
 }
-
-#if UE_ENABLE_INCLUDE_ORDER_DEPRECATED_IN_5_6
-#include "HAL/CriticalSection.h"
-#endif
diff --git a/Engine/Source/Runtime/Core/Public/Containers/Array.h b/Engine/Source/Runtime/Core/Public/Containers/Array.h
--- a/Engine/Source/Runtime/Core/Public/Containers/Array.h
+++ b/Engine/Source/Runtime/Core/Public/Containers/Array.h
@@ -4063,10 +4063,3 @@ template<typename InElementType, typename InAllocatorType>
 	}
 	return Hash;
 }
-
-#if UE_ENABLE_INCLUDE_ORDER_DEPRECATED_IN_5_4
-#include "Templates/IsSigned.h"
-#include "Templates/AndOrNot.h"
-#include "Templates/IsConstructible.h"
-#include "Templates/MakeUnsigned.h"
-#endif
diff --git a/Engine/Source/Runtime/Core/Public/Containers/BinaryHeap.h b/Engine/Source/Runtime/Core/Public/Containers/BinaryHeap.h
--- a/Engine/Source/Runtime/Core/Public/Containers/BinaryHeap.h
+++ b/Engine/Source/Runtime/Core/Public/Containers/BinaryHeap.h
@@ -401,7 +401,3 @@ inline void FBinaryHeap< KeyType, IndexType >::ResetInternal()
 	Keys = nullptr;
 	HeapIndexes = nullptr;
 }
-
-#if UE_ENABLE_INCLUDE_ORDER_DEPRECATED_IN_5_4
-#include "Templates/IsSigned.h"
-#endif
diff --git a/Engine/Source/Runtime/Core/Public/Containers/EnumAsByte.h b/Engine/Source/Runtime/Core/Public/Containers/EnumAsByte.h
--- a/Engine/Source/Runtime/Core/Public/Containers/EnumAsByte.h
+++ b/Engine/Source/Runtime/Core/Public/Containers/EnumAsByte.h
@@ -124,7 +124,3 @@ template<class T>
 
 
 template<class T> struct TIsPODType<TEnumAsByte<T>> { enum { Value = true }; };
-
-#if UE_ENABLE_INCLUDE_ORDER_DEPRECATED_IN_5_6
-#include "Traits/IsTEnumAsByte.h"
-#endif
diff --git a/Engine/Source/Runtime/Core/Public/Containers/RingBuffer.h b/Engine/Source/Runtime/Core/Public/Containers/RingBuffer.h
--- a/Engine/Source/Runtime/Core/Public/Containers/RingBuffer.h
+++ b/Engine/Source/Runtime/Core/Public/Containers/RingBuffer.h
@@ -1035,9 +1035,3 @@ private:
 	 */
 	StorageModuloType AfterBack;
 };
-
-#if UE_ENABLE_INCLUDE_ORDER_DEPRECATED_IN_5_4
-#include "Templates/IsTriviallyDestructible.h"
-#include "Templates/MakeSigned.h"
-#include "Templates/MakeUnsigned.h"
-#endif
diff --git a/Engine/Source/Runtime/Core/Public/Containers/Set.h b/Engine/Source/Runtime/Core/Public/Containers/Set.h
--- a/Engine/Source/Runtime/Core/Public/Containers/Set.h
+++ b/Engine/Source/Runtime/Core/Public/Containers/Set.h
@@ -40,7 +40,3 @@ using TScriptSet = TScriptSparseSet<Allocator>;
 #endif
 
 using FScriptSet = TScriptSet<FDefaultSetAllocator>;
-
-#if UE_ENABLE_INCLUDE_ORDER_DEPRECATED_IN_5_4
-#include "Templates/Decay.h"
-#endif
diff --git a/Engine/Source/Runtime/Core/Public/Containers/SparseArray.h b/Engine/Source/Runtime/Core/Public/Containers/SparseArray.h
--- a/Engine/Source/Runtime/Core/Public/Containers/SparseArray.h
+++ b/Engine/Source/Runtime/Core/Public/Containers/SparseArray.h
@@ -1768,8 +1768,3 @@ void operator<<(FStructuredArchive::FSlot Slot, TSparseArray<ElementType, Alloca
 		}
 	}
 }
-
-#if UE_ENABLE_INCLUDE_ORDER_DEPRECATED_IN_5_4
-#include "Templates/IsTriviallyCopyConstructible.h"
-#include "Templates/IsTriviallyDestructible.h"
-#endif
diff --git a/Engine/Source/Runtime/Core/Public/Containers/StringConv.h b/Engine/Source/Runtime/Core/Public/Containers/StringConv.h
--- a/Engine/Source/Runtime/Core/Public/Containers/StringConv.h
+++ b/Engine/Source/Runtime/Core/Public/Containers/StringConv.h
@@ -1363,9 +1363,3 @@ template <typename ToType, typename FromType>
 {
 	return StringToArray<ToType>(Str, TCString<FromType>::Strlen(Str) + 1);
 }
-
-#if UE_ENABLE_INCLUDE_ORDER_DEPRECATED_IN_5_4
-#include "Templates/AndOrNot.h"
-#include "Templates/EnableIf.h"
-#include "Templates/RemoveReference.h"
-#endif
diff --git a/Engine/Source/Runtime/Core/Public/Core.h b/Engine/Source/Runtime/Core/Public/Core.h
--- a/Engine/Source/Runtime/Core/Public/Core.h
+++ b/Engine/Source/Runtime/Core/Public/Core.h
@@ -313,7 +313,3 @@ MONOLITHIC_HEADER_BOILERPLATE()
 #include "Misc/EnumClassFlags.h"
 #include "Misc/EnumRange.h"
 // IWYU pragma: end_exports
-
-#if UE_ENABLE_INCLUDE_ORDER_DEPRECATED_IN_5_5
-#include "Templates/IsTriviallyDestructible.h"
-#endif
diff --git a/Engine/Source/Runtime/Core/Public/CoreMinimal.h b/Engine/Source/Runtime/Core/Public/CoreMinimal.h
--- a/Engine/Source/Runtime/Core/Public/CoreMinimal.h
+++ b/Engine/Source/Runtime/Core/Public/CoreMinimal.h
@@ -52,9 +52,6 @@
 #include "Containers/ContainerAllocationPolicies.h"
 #include "Templates/IsEnumClass.h"
 #include "HAL/PlatformProperties.h"
-#if UE_ENABLE_INCLUDE_ORDER_DEPRECATED_IN_5_4
-#include "Misc/Compression.h"
-#endif
 #include "Misc/EngineVersionBase.h"
 #include "Internationalization/TextNamespaceFwd.h"
 #include "Serialization/Archive.h"
@@ -71,16 +68,10 @@
 #include "Containers/Array.h"
 #include "Misc/FrameNumber.h"
 #include "Misc/Timespan.h"
-#if UE_ENABLE_INCLUDE_ORDER_DEPRECATED_IN_5_4
-#include "HAL/CriticalSection.h"
-#endif
 #include "Containers/StringConv.h"
 #include "UObject/UnrealNames.h"
 #include "UObject/NameTypes.h"
 #include "Misc/Parse.h"
-#if UE_ENABLE_INCLUDE_ORDER_DEPRECATED_IN_5_4
-#include "Math/Color.h"
-#endif
 #include "Templates/AlignmentTemplates.h"
 #include "Misc/StructBuilder.h"
 #include "Templates/Decay.h"
@@ -94,9 +85,6 @@
 #include "Containers/Set.h"
 #include "Algo/Reverse.h"
 #include "Containers/Map.h"
-#if UE_ENABLE_INCLUDE_ORDER_DEPRECATED_IN_5_4
-#include "Math/ColorList.h"
-#endif
 #include "Math/IntPoint.h"
 #include "Math/IntVector.h"
 #include "Logging/LogCategory.h"
@@ -132,10 +120,6 @@
 #include "Math/TwoVectors.h"
 #include "Math/Edge.h"
 #include "UObject/ObjectVersion.h"
-#if UE_ENABLE_INCLUDE_ORDER_DEPRECATED_IN_5_4
-#include "Math/Plane.h"
-#include "Math/Sphere.h"
-#endif
 #include "Math/CapsuleShape.h"
 #include "Math/Rotator.h"
 #include "Misc/DateTime.h"
@@ -163,21 +147,8 @@
 #include "Math/ScaleMatrix.h"
 #include "Math/MirrorMatrix.h"
 #include "Math/ClipProjectionMatrix.h"
-#if UE_ENABLE_INCLUDE_ORDER_DEPRECATED_IN_5_4
-#include "Math/InterpCurvePoint.h"
-#include "Math/InterpCurve.h"
-#include "Math/CurveEdInterface.h"
-#endif
 #include "Math/Float32.h"
 #include "Math/Float16.h"
-#if UE_ENABLE_INCLUDE_ORDER_DEPRECATED_IN_5_4
-#include "Math/Float16Color.h"
-#include "Math/Vector2DHalf.h"
-#endif
 #include "Math/Transform.h"
 #include "Math/ConvexHull2d.h"
 #include "Math/UnrealMath.h"
-
-#if UE_ENABLE_INCLUDE_ORDER_DEPRECATED_IN_5_5
-#include "Templates/IsTriviallyDestructible.h"
-#endif
diff --git a/Engine/Source/Runtime/Core/Public/CoreSharedPCH.h b/Engine/Source/Runtime/Core/Public/CoreSharedPCH.h
--- a/Engine/Source/Runtime/Core/Public/CoreSharedPCH.h
+++ b/Engine/Source/Runtime/Core/Public/CoreSharedPCH.h
@@ -252,12 +252,3 @@
 #include "UObject/UnrealNames.h"
 #include "UObject/UObjectHierarchyFwd.h"
 #include "UObject/WeakObjectPtrTemplates.h"
-
-#if UE_ENABLE_INCLUDE_ORDER_DEPRECATED_IN_5_5
-#include "Templates/IsTriviallyDestructible.h"
-#endif
-
-#if UE_ENABLE_INCLUDE_ORDER_DEPRECATED_IN_5_6
-#include "Math/BasicMathExpressionEvaluator.h"
-#include "Misc/ExpressionParserTypes.h"
-#endif
diff --git a/Engine/Source/Runtime/Core/Public/Delegates/DelegateSignatureImpl.inl b/Engine/Source/Runtime/Core/Public/Delegates/DelegateSignatureImpl.inl
--- a/Engine/Source/Runtime/Core/Public/Delegates/DelegateSignatureImpl.inl
+++ b/Engine/Source/Runtime/Core/Public/Delegates/DelegateSignatureImpl.inl
@@ -1294,10 +1294,6 @@ public:
 	// NOTE:  Broadcast() method must be defined in derived classes
 };
 
-#if UE_ENABLE_INCLUDE_ORDER_DEPRECATED_IN_5_5
-#include "Misc/Crc.h"
-#endif
-
 #if UE_ENABLE_INCLUDE_ORDER_DEPRECATED_IN_5_7
 #include "Templates/IsConst.h"
 #include "Templates/RemoveReference.h"
diff --git a/Engine/Source/Runtime/Core/Public/Internationalization/FastDecimalFormat.h b/Engine/Source/Runtime/Core/Public/Internationalization/FastDecimalFormat.h
--- a/Engine/Source/Runtime/Core/Public/Internationalization/FastDecimalFormat.h
+++ b/Engine/Source/Runtime/Core/Public/Internationalization/FastDecimalFormat.h
@@ -202,10 +202,3 @@ CORE_API const FDecimalNumberFormattingRules& GetCultureAgnosticFormattingRules(
 CORE_API uint64 Pow10(const int32 InExponent);
 
 } // namespace FastDecimalFormat
-
-#if UE_ENABLE_INCLUDE_ORDER_DEPRECATED_IN_5_4
-#include "Templates/EnableIf.h"
-#include "Templates/IsFloatingPoint.h"
-#include "Templates/IsIntegral.h"
-#include "Templates/IsSigned.h"
-#endif
diff --git a/Engine/Source/Runtime/Core/Public/Math/UnrealMathNeon.h b/Engine/Source/Runtime/Core/Public/Math/UnrealMathNeon.h
--- a/Engine/Source/Runtime/Core/Public/Math/UnrealMathNeon.h
+++ b/Engine/Source/Runtime/Core/Public/Math/UnrealMathNeon.h
@@ -2930,8 +2930,3 @@ FORCEINLINE VectorRegister4i VectorIntExpandLow16To32(VectorRegister4i V)
 #endif // #if PLATFORM_ENABLE_VECTORINTRINSICS_NEON
 
 PRAGMA_ENABLE_SHADOW_VARIABLE_WARNINGS
-
-#if UE_ENABLE_INCLUDE_ORDER_DEPRECATED_IN_5_4
-#include <type_traits>
-#endif
-
diff --git a/Engine/Source/Runtime/Core/Public/Math/Vector2D.h b/Engine/Source/Runtime/Core/Public/Math/Vector2D.h
--- a/Engine/Source/Runtime/Core/Public/Math/Vector2D.h
+++ b/Engine/Source/Runtime/Core/Public/Math/Vector2D.h
@@ -10,9 +10,6 @@
 #include "Containers/UnrealString.h"
 #include "Misc/Parse.h"
 #include "Misc/LargeWorldCoordinatesSerializer.h"
-#if UE_ENABLE_INCLUDE_ORDER_DEPRECATED_IN_5_4
-#include "Misc/NetworkVersion.h"
-#endif
 #include "Misc/EngineNetworkCustomVersion.h"
 #include "Math/IntPoint.h"
 #include "Logging/LogMacros.h"
diff --git a/Engine/Source/Runtime/Core/Public/Misc/AutomationTest.h b/Engine/Source/Runtime/Core/Public/Misc/AutomationTest.h
--- a/Engine/Source/Runtime/Core/Public/Misc/AutomationTest.h
+++ b/Engine/Source/Runtime/Core/Public/Misc/AutomationTest.h
@@ -4976,7 +4976,3 @@ private:
 	double MaxTotalRunTimeInSeconds = 300.0;
 	
 };
-
-#if UE_ENABLE_INCLUDE_ORDER_DEPRECATED_IN_5_6
-# include "TextFilterExpressionEvaluator.h"
-#endif
diff --git a/Engine/Source/Runtime/Core/Public/Misc/Crc.h b/Engine/Source/Runtime/Core/Public/Misc/Crc.h
--- a/Engine/Source/Runtime/Core/Public/Misc/Crc.h
+++ b/Engine/Source/Runtime/Core/Public/Misc/Crc.h
@@ -331,7 +331,3 @@ struct TStringPointerMapKeyFuncs_DEPRECATED
 		return FCrc::Strihash_DEPRECATED(ToCStr(Key));
 	}
 };
-
-#if UE_ENABLE_INCLUDE_ORDER_DEPRECATED_IN_5_5
-#include "Templates/EnableIf.h"
-#endif
diff --git a/Engine/Source/Runtime/Core/Public/Misc/ExpressionParserTypes.h b/Engine/Source/Runtime/Core/Public/Misc/ExpressionParserTypes.h
--- a/Engine/Source/Runtime/Core/Public/Misc/ExpressionParserTypes.h
+++ b/Engine/Source/Runtime/Core/Public/Misc/ExpressionParserTypes.h
@@ -637,8 +637,4 @@ extern template class TTokenDefinitions<ANSICHAR>;
 extern template class TTokenDefinitions<WIDECHAR>;
 extern template class TTokenDefinitions<UTF8CHAR>;
 
-#if UE_ENABLE_INCLUDE_ORDER_DEPRECATED_IN_5_6
-#include "Templates/PointerIsConvertibleFromTo.h"
-#endif
-
 #include "Misc/ExpressionParserTypes.inl" // IWYU pragma: export
diff --git a/Engine/Source/Runtime/Core/Public/Misc/FrameRate.h b/Engine/Source/Runtime/Core/Public/Misc/FrameRate.h
--- a/Engine/Source/Runtime/Core/Public/Misc/FrameRate.h
+++ b/Engine/Source/Runtime/Core/Public/Misc/FrameRate.h
@@ -332,7 +332,3 @@ CORE_API TValueOrError<FFrameRate, FExpressionError> ParseFrameRate(const TCHAR*
  * Common Lex::TryParseString overload for FFrameRate
  */
 CORE_API bool TryParseString(FFrameRate& OutFrameRate, const TCHAR* InString);
-
-#if UE_ENABLE_INCLUDE_ORDER_DEPRECATED_IN_5_6
-# include "ExpressionParserTypes.h"
-#endif
diff --git a/Engine/Source/Runtime/Core/Public/Serialization/CompactBinary.h b/Engine/Source/Runtime/Core/Public/Serialization/CompactBinary.h
--- a/Engine/Source/Runtime/Core/Public/Serialization/CompactBinary.h
+++ b/Engine/Source/Runtime/Core/Public/Serialization/CompactBinary.h
@@ -1463,7 +1463,3 @@ inline FSharedBuffer FCbField::AsBinary(const FSharedBuffer& Default) &&
 }
 
 ///////////////////////////////////////////////////////////////////////////////////////////////////////////////////////
-
-#if UE_ENABLE_INCLUDE_ORDER_DEPRECATED_IN_5_5
-#include "Templates/IsTriviallyDestructible.h"
-#endif
diff --git a/Engine/Source/Runtime/Core/Public/Serialization/MemoryLayout.h b/Engine/Source/Runtime/Core/Public/Serialization/MemoryLayout.h
--- a/Engine/Source/Runtime/Core/Public/Serialization/MemoryLayout.h
+++ b/Engine/Source/Runtime/Core/Public/Serialization/MemoryLayout.h
@@ -877,7 +877,3 @@ private:
 	friend CORE_API bool LoadFromCompactBinary(FCbFieldView Field, FPlatformTypeLayoutParameters& OutValue);
 #endif
 };
-
-#if UE_ENABLE_INCLUDE_ORDER_DEPRECATED_IN_5_5
-#include "Templates/IsTriviallyDestructible.h"
-#endif
diff --git a/Engine/Source/Runtime/Core/Public/Stats/HitchTrackingStatScope.h b/Engine/Source/Runtime/Core/Public/Stats/HitchTrackingStatScope.h
--- a/Engine/Source/Runtime/Core/Public/Stats/HitchTrackingStatScope.h
+++ b/Engine/Source/Runtime/Core/Public/Stats/HitchTrackingStatScope.h
@@ -5,10 +5,6 @@
 #include "Misc/Build.h"
 #include "Stats/StatsCommon.h"
 
-#if UE_ENABLE_INCLUDE_ORDER_DEPRECATED_IN_5_6
-#include "ProfilingDebugging/UMemoryDefines.h"
-#endif
-
 #ifndef USE_LIGHTWEIGHT_STATS_FOR_HITCH_DETECTION
 #define USE_LIGHTWEIGHT_STATS_FOR_HITCH_DETECTION 1
 #endif
diff --git a/Engine/Source/Runtime/Core/Public/Stats/Stats.h b/Engine/Source/Runtime/Core/Public/Stats/Stats.h
--- a/Engine/Source/Runtime/Core/Public/Stats/Stats.h
+++ b/Engine/Source/Runtime/Core/Public/Stats/Stats.h
@@ -10,22 +10,6 @@
 #include "Stats/LightweightStats.h"
 #include "Stats/StatsSystemTypes.h"
 
-#if UE_ENABLE_INCLUDE_ORDER_DEPRECATED_IN_5_6
-// These includes are no longer required for this file.
-#include "Delegates/Delegate.h"
-#include "HAL/PlatformTime.h"
-#include "HAL/ThreadSingleton.h"
-#include "ProfilingDebugging/UMemoryDefines.h"
-#include "AutoRTFM.h"
-#include "Math/Color.h"
-#include "UObject/NameTypes.h"
-
-// Include these directly where needed instead.
-#include "Stats/StatsCommand.h"
-#include "Stats/StatsSystem.h"
-#include "Stats/ThreadIdleStats.h"
-#endif // UE_ENABLE_INCLUDE_ORDER_DEPRECATED_IN_5_6
-
 // used by the profiler
 enum EStatType
 {
@@ -727,53 +711,3 @@ struct FStat_##StatName\
 #endif // !STATS
 
 #include "Stats/GlobalStats.inl"
-
-// Deprecated includes from Stats2.h
-#if UE_ENABLE_INCLUDE_ORDER_DEPRECATED_IN_5_6
-// These should be included directly instead of through this header.
-#include "Stats/HitchTrackingStatScope.h"
-#include "Stats/StatsCommand.h"
-#include "Stats/StatsSystem.h"
-#include "Stats/StatsSystemTypes.h"
-#include "Stats/ThreadIdleStats.h"
-
-// Includes from Stats2.h
-#include "Containers/Array.h"
-#include "Containers/ChunkedArray.h"
-#include "Containers/ContainerAllocationPolicies.h"
-#include "Containers/LockFreeList.h"
-#include "Containers/UnrealString.h"
-#include "CoreGlobals.h"
-#include "CoreTypes.h"
-#include "Delegates/Delegate.h"
-#include "Delegates/DelegateBase.h"
-#include "HAL/CriticalSection.h"
-#include "HAL/LowLevelMemTracker.h"
-#include "HAL/PlatformCrt.h"
-#include "HAL/PlatformMemory.h"
-#include "HAL/PlatformMisc.h"
-#include "HAL/PlatformTLS.h"
-#include "HAL/PlatformTime.h"
-#include "HAL/ThreadSafeCounter.h"
-#include "HAL/ThreadSingleton.h"
-#include "HAL/UnrealMemory.h"
-#include "Math/Color.h"
-#include "Math/NumericLimits.h"
-#include "Misc/AssertionMacros.h"
-#include "Misc/Build.h"
-#include "Misc/CString.h"
-#include "Misc/EnumClassFlags.h"
-#include "ProfilingDebugging/CpuProfilerTrace.h"
-#include "ProfilingDebugging/MiscTrace.h"
-#include "StatsCommon.h"
-#include "StatsTrace.h"
-#include "Templates/Atomic.h"
-#include "Templates/TypeCompatibleBytes.h"
-#include "Templates/UniquePtr.h"
-#include "Templates/UnrealTemplate.h"
-#include "Trace/Detail/Channel.h"
-#include "Trace/Detail/Channel.inl"
-#include "Trace/Trace.h"
-#include "UObject/NameTypes.h"
-#include "UObject/UnrealNames.h"
-#endif // UE_ENABLE_INCLUDE_ORDER_DEPRECATED_IN_5_6
diff --git a/Engine/Source/Runtime/Core/Public/Templates/Function.h b/Engine/Source/Runtime/Core/Public/Templates/Function.h
--- a/Engine/Source/Runtime/Core/Public/Templates/Function.h
+++ b/Engine/Source/Runtime/Core/Public/Templates/Function.h
@@ -859,13 +859,3 @@ FORCEINLINE bool operator!=(TYPE_OF_NULLPTR, const TFunction<FuncType>& Func)
 	return (bool)Func;
 }
 #endif
-
-#if UE_ENABLE_INCLUDE_ORDER_DEPRECATED_IN_5_4
-#include "Templates/AndOrNot.h"
-#include "Templates/Decay.h"
-#include "Templates/IsConstructible.h"
-#include "Templates/IsInvocable.h"
-#include "Templates/IsPointer.h"
-#include "Templates/IsMemberPointer.h"
-#include "Templates/RemoveReference.h"
-#endif
diff --git a/Engine/Source/Runtime/Core/Public/Templates/IsTrivial.h b/Engine/Source/Runtime/Core/Public/Templates/IsTrivial.h
--- a/Engine/Source/Runtime/Core/Public/Templates/IsTrivial.h
+++ b/Engine/Source/Runtime/Core/Public/Templates/IsTrivial.h
@@ -15,7 +15,3 @@ struct TIsTrivial
 {
 	enum { Value = TAndValue<std::is_trivially_destructible_v<T>, TIsTriviallyCopyConstructible<T>, TIsTriviallyCopyAssignable<T>>::Value };
 };
-
-#if UE_ENABLE_INCLUDE_ORDER_DEPRECATED_IN_5_5
-#include "Templates/IsTriviallyDestructible.h"
-#endif
diff --git a/Engine/Source/Runtime/Core/Public/Templates/NonNullPointer.h b/Engine/Source/Runtime/Core/Public/Templates/NonNullPointer.h
--- a/Engine/Source/Runtime/Core/Public/Templates/NonNullPointer.h
+++ b/Engine/Source/Runtime/Core/Public/Templates/NonNullPointer.h
@@ -259,8 +259,3 @@ inline ObjectType* GetRawPointerOrNull(const TOptional<TNonNullPtr<ObjectType>>&
 {
 	return Optional.IsSet() ? Optional->Get() : nullptr;
 }
-
-#if UE_ENABLE_INCLUDE_ORDER_DEPRECATED_IN_5_4
-#include "Templates/EnableIf.h"
-#include "Templates/PointerIsConvertibleFromTo.h"
-#endif
diff --git a/Engine/Source/Runtime/Core/Public/Templates/Tuple.h b/Engine/Source/Runtime/Core/Public/Templates/Tuple.h
--- a/Engine/Source/Runtime/Core/Public/Templates/Tuple.h
+++ b/Engine/Source/Runtime/Core/Public/Templates/Tuple.h
@@ -725,8 +725,3 @@ UE_FORCEINLINE_HINT void operator<<(FStructuredArchive::FSlot Slot, TTuple<Types
 {
 	TTuple<Types...>::SerializeStructured(Slot, Tuple);
 }
-
-#if UE_ENABLE_INCLUDE_ORDER_DEPRECATED_IN_5_4
-#include "Templates/Decay.h"
-#include "Templates/IsConstructible.h"
-#endif
\ No newline at end of file
diff --git a/Engine/Source/Runtime/Core/Public/Templates/TypeHash.h b/Engine/Source/Runtime/Core/Public/Templates/TypeHash.h
--- a/Engine/Source/Runtime/Core/Public/Templates/TypeHash.h
+++ b/Engine/Source/Runtime/Core/Public/Templates/TypeHash.h
@@ -220,7 +220,3 @@ inline void DispatchAppendHash(BuilderType& Builder, const ArgType& Arg)
 {
 	AppendHash(Builder, Arg);
 }
-
-#if UE_ENABLE_INCLUDE_ORDER_DEPRECATED_IN_5_6
-#include "Templates/IsEnum.h"
-#endif
diff --git a/Engine/Source/Runtime/Core/Public/Templates/UnrealTemplate.h b/Engine/Source/Runtime/Core/Public/Templates/UnrealTemplate.h
--- a/Engine/Source/Runtime/Core/Public/Templates/UnrealTemplate.h
+++ b/Engine/Source/Runtime/Core/Public/Templates/UnrealTemplate.h
@@ -768,9 +768,3 @@ UE_INTRINSIC_CAST UE_REWRITE decltype(auto) ForwardAsBase(std::remove_reference_
 {
 	return (TCopyQualifiersAndRefsFromTo_T<T&&, Base>)Obj;
 }
-
-#if UE_ENABLE_INCLUDE_ORDER_DEPRECATED_IN_5_4
-#include "Templates/AndOrNot.h"
-#include "Templates/EnableIf.h"
-#include "Templates/IsArithmetic.h"
-#endif
diff --git a/Engine/Source/Runtime/Core/Public/Templates/UnrealTypeTraits.h b/Engine/Source/Runtime/Core/Public/Templates/UnrealTypeTraits.h
--- a/Engine/Source/Runtime/Core/Public/Templates/UnrealTypeTraits.h
+++ b/Engine/Source/Runtime/Core/Public/Templates/UnrealTypeTraits.h
@@ -484,7 +484,3 @@ public:																											\
 #undef IS_EMPTY
 #undef IS_POD
 #undef HAS_TRIVIAL_CONSTRUCTOR
-
-#if UE_ENABLE_INCLUDE_ORDER_DEPRECATED_IN_5_5
-#include "Templates/Requires.h"
-#endif
diff --git a/Engine/Source/Runtime/Core/Public/UObject/WeakObjectPtrTemplates.h b/Engine/Source/Runtime/Core/Public/UObject/WeakObjectPtrTemplates.h
--- a/Engine/Source/Runtime/Core/Public/UObject/WeakObjectPtrTemplates.h
+++ b/Engine/Source/Runtime/Core/Public/UObject/WeakObjectPtrTemplates.h
@@ -479,9 +479,3 @@ FArchive& operator<<( FArchive& Ar, TWeakObjectPtr<T, TWeakObjectPtrBase>& WeakO
 	WeakObjectPtr.Serialize(Ar);
 	return Ar;
 }
-
-#if UE_ENABLE_INCLUDE_ORDER_DEPRECATED_IN_5_4
-#include "Templates/AndOrNot.h"
-#include "Templates/IsPointer.h"
-#include "Templates/PointerIsConvertibleFromTo.h"
-#endif
diff --git a/Engine/Source/Runtime/Core/Public/Windows/WindowsPlatformMutex.h b/Engine/Source/Runtime/Core/Public/Windows/WindowsPlatformMutex.h
--- a/Engine/Source/Runtime/Core/Public/Windows/WindowsPlatformMutex.h
+++ b/Engine/Source/Runtime/Core/Public/Windows/WindowsPlatformMutex.h
@@ -148,7 +148,3 @@ using FPlatformSharedMutex = FWindowsSharedMutex;
 using FPlatformSystemWideMutex = FWindowsSystemWideMutex;
 
 } // UE
-
-#if UE_ENABLE_INCLUDE_ORDER_DEPRECATED_IN_5_6
-#include "HAL/PlatformMemory.h"
-#endif
diff --git a/Engine/Source/Runtime/CoreUObject/Public/StructUtils/UserDefinedStruct.h b/Engine/Source/Runtime/CoreUObject/Public/StructUtils/UserDefinedStruct.h
--- a/Engine/Source/Runtime/CoreUObject/Public/StructUtils/UserDefinedStruct.h
+++ b/Engine/Source/Runtime/CoreUObject/Public/StructUtils/UserDefinedStruct.h
@@ -2,10 +2,6 @@
 
 #pragma once
 
-#if UE_ENABLE_INCLUDE_ORDER_DEPRECATED_IN_5_5
-#include "CoreMinimal.h"
-#include "UObject/ObjectMacros.h"
-#endif // UE_ENABLE_INCLUDE_ORDER_DEPRECATED_IN_5_5
 #include "Misc/Guid.h"
 #include "UObject/Class.h"
 #include "UObject/ObjectMacros.h"
diff --git a/Engine/Source/Runtime/CoreUObject/Public/UObject/Class.h b/Engine/Source/Runtime/CoreUObject/Public/UObject/Class.h
--- a/Engine/Source/Runtime/CoreUObject/Public/UObject/Class.h
+++ b/Engine/Source/Runtime/CoreUObject/Public/UObject/Class.h
@@ -86,10 +86,6 @@
 #include <type_traits>
 #include <atomic>
 
-#if UE_ENABLE_INCLUDE_ORDER_DEPRECATED_IN_5_5
-#include "Templates/IsTriviallyDestructible.h"
-#endif
-
 #include "Class.generated.h"
 
 class FArchive;
diff --git a/Engine/Source/Runtime/CoreUObject/Public/UObject/Script.h b/Engine/Source/Runtime/CoreUObject/Public/UObject/Script.h
--- a/Engine/Source/Runtime/CoreUObject/Public/UObject/Script.h
+++ b/Engine/Source/Runtime/CoreUObject/Public/UObject/Script.h
@@ -8,9 +8,6 @@
 
 #include "Delegates/Delegate.h"
 #include "HAL/ThreadSingleton.h"
-#if UE_ENABLE_INCLUDE_ORDER_DEPRECATED_IN_5_4
-#include "Internationalization/Text.h"
-#endif
 #include "Stats/Stats.h"
 #include "Misc/EnumClassFlags.h"
 #include "Misc/CoreMisc.h"
