Removed deprecated Core text globals.

#rnx

diff --git a/Engine/Source/Programs/UnrealBuildAccelerator/UbaCoordinatorHorde.Target.cs b/Engine/Source/Programs/UnrealBuildAccelerator/UbaCoordinatorHorde.Target.cs
--- a/Engine/Source/Programs/UnrealBuildAccelerator/UbaCoordinatorHorde.Target.cs
+++ b/Engine/Source/Programs/UnrealBuildAccelerator/UbaCoordinatorHorde.Target.cs
@@ -61,7 +61,6 @@ public class UbaCoordinatorHordeTarget : TargetRules
 		}
 
 		GlobalDefinitions.Add("UE_EXTERNAL_PROFILING_ENABLED=0");
-		GlobalDefinitions.Add("DISABLE_LEGACY_CORE_TEXTS=1");
 		GlobalDefinitions.Add("STATS=0");
 		GlobalDefinitions.Add("ENABLE_STATNAMEDEVENTS=0");
 		GlobalDefinitions.Add("ALLOW_HITCH_DETECTION=0");
diff --git a/Engine/Source/Runtime/Core/Private/Misc/CoreGlobals.cpp b/Engine/Source/Runtime/Core/Private/Misc/CoreGlobals.cpp
--- a/Engine/Source/Runtime/Core/Private/Misc/CoreGlobals.cpp
+++ b/Engine/Source/Runtime/Core/Private/Misc/CoreGlobals.cpp
@@ -109,16 +109,6 @@ void FCoreTexts::TearDown()
 	TLazySingleton<FCoreTextsSingleton>::TearDown();
 }
 
-#if !defined(DISABLE_LEGACY_CORE_TEXTS) || DISABLE_LEGACY_CORE_TEXTS == 0
-PRAGMA_DISABLE_DEPRECATION_WARNINGS
-CORE_API const FText GYes	= LOCTEXT("Yes",	"Yes");
-CORE_API const FText GNo	= LOCTEXT("No",		"No");
-CORE_API const FText GTrue	= LOCTEXT("True",	"True");
-CORE_API const FText GFalse	= LOCTEXT("False",	"False");
-CORE_API const FText GNone	= LOCTEXT("None",	"None");
-PRAGMA_ENABLE_DEPRECATION_WARNINGS
-#endif
-
 /** If true, this executable is able to run all games (which are loaded as DLL's) **/
 #if UE_GAME || UE_SERVER
 	// In monolithic builds, implemented by the IMPLEMENT_GAME_MODULE macro or by UnrealGame module.
diff --git a/Engine/Source/Runtime/Core/Public/CoreGlobals.h b/Engine/Source/Runtime/Core/Public/CoreGlobals.h
--- a/Engine/Source/Runtime/Core/Public/CoreGlobals.h
+++ b/Engine/Source/Runtime/Core/Public/CoreGlobals.h
@@ -128,19 +128,6 @@ struct FCoreTexts
 	}
 };
 
-#if !defined(DISABLE_LEGACY_CORE_TEXTS) || DISABLE_LEGACY_CORE_TEXTS == 0
-UE_DEPRECATED(4.23, "GTrue has been deprecated in favor of FCoreTexts::Get().True.")
-extern CORE_API const FText GTrue;
-UE_DEPRECATED(4.23, "GFalse has been deprecated in favor of FCoreTexts::Get().False.")
-extern CORE_API const FText GFalse;
-UE_DEPRECATED(4.23, "GYes has been deprecated in favor of FCoreTexts::Get().Yes.")
-extern CORE_API const FText GYes;
-UE_DEPRECATED(4.23, "GNo has been deprecated in favor of FCoreTexts::Get().No.")
-extern CORE_API const FText GNo;
-UE_DEPRECATED(4.23, "GNone has been deprecated in favor of FCoreTexts::Get().None.")
-extern CORE_API const FText GNone;
-#endif
-
 /** If true, this executable is able to run all games (which are loaded as DLL's). */
 extern CORE_API bool GIsGameAgnosticExe;
 
