Removed deprecated access to GRenderingThread.

#rnx

diff --git a/Engine/Source/Runtime/Core/Private/HAL/ThreadingBase.cpp b/Engine/Source/Runtime/Core/Private/HAL/ThreadingBase.cpp
--- a/Engine/Source/Runtime/Core/Private/HAL/ThreadingBase.cpp
+++ b/Engine/Source/Runtime/Core/Private/HAL/ThreadingBase.cpp
@@ -253,8 +253,6 @@ CORE_API bool IsInAudioThread()
 	return (GIsAudioThreadRunning.load(std::memory_order_acquire) && !GIsAudioThreadSuspended.load(std::memory_order_acquire)) ? GAudioPipe.IsInContext() : IsInGameThread();
 }
 
-CORE_API FRunnableThread* GRenderingThread = nullptr;
-
 CORE_API bool IsInActualRenderingThread()
 {
 PRAGMA_DISABLE_DEPRECATION_WARNINGS
@@ -282,7 +280,7 @@ PRAGMA_DISABLE_DEPRECATION_WARNINGS
 	{
 		const uint32 CurrentThreadId = FPlatformTLS::GetCurrentThreadId();
 		bool oldValue = (GRenderThreadId == 0) ? (CurrentThreadId == GGameThreadId) || FTaskTagScope::IsRunningDuringStaticInit() : (CurrentThreadId == GRenderThreadId);
-		ensureMsgf(oldValue == newValue, TEXT("oldValue(%i) newValue(%i) If this check fails make sure that there is a FTaskTagScope(ETaskTag::ERenderingThread) as deep as possible on the current callstack, you can see the current value in ActiveNamedThreads(%x), GRenderingThread(%p)"), oldValue, newValue, FTaskTagScope::GetCurrentTag(), GRenderingThread);
+		ensureMsgf(oldValue == newValue, TEXT("oldValue(%i) newValue(%i) If this check fails make sure that there is a FTaskTagScope(ETaskTag::ERenderingThread) as deep as possible on the current callstack, you can see the current value in ActiveNamedThreads(%x)"), oldValue, newValue, FTaskTagScope::GetCurrentTag());
 		newValue = oldValue;
 	}
 #endif
@@ -319,7 +317,7 @@ PRAGMA_DISABLE_DEPRECATION_WARNINGS
 	{
 		const uint32 CurrentThreadId = FPlatformTLS::GetCurrentThreadId();
 		bool oldValue = (GRenderThreadId == 0) ?  true : CurrentThreadId != GGameThreadId;
-		ensureMsgf(oldValue == newValue, TEXT("oldValue(%i) newValue(%i) If this check fails make sure that there is a FTaskTagScope(ETaskTag::EParallelRenderingThread) as deep as possible on the current callstack, you can see the current value in ActiveNamedThreads(%x), GRenderingThread(%p)"), oldValue, newValue, FTaskTagScope::GetCurrentTag(), GRenderingThread);
+		ensureMsgf(oldValue == newValue, TEXT("oldValue(%i) newValue(%i) If this check fails make sure that there is a FTaskTagScope(ETaskTag::EParallelRenderingThread) as deep as possible on the current callstack, you can see the current value in ActiveNamedThreads(%x)"), oldValue, newValue, FTaskTagScope::GetCurrentTag());
 		newValue = oldValue;
 	}
 #endif
diff --git a/Engine/Source/Runtime/Core/Public/CoreGlobals.h b/Engine/Source/Runtime/Core/Public/CoreGlobals.h
--- a/Engine/Source/Runtime/Core/Public/CoreGlobals.h
+++ b/Engine/Source/Runtime/Core/Public/CoreGlobals.h
@@ -755,10 +755,6 @@ extern CORE_API bool (*IsInAsyncLoadingThread)();
 /** @return True only if called from the async loading thread if it's enabled. */
 extern CORE_API bool IsInActualLoadingThread();
 
-/** Thread used for rendering */
-UE_DEPRECATED(4.26, "Please use `GIsThreadedRendering` or `IsInActualRenderingThread()`")
-extern CORE_API FRunnableThread* GRenderingThread;
-
 /** @return True if RHI thread is running */
 extern CORE_API bool IsRHIThreadRunning();
 
diff --git a/Engine/Source/Runtime/RenderCore/Private/RenderingThread.cpp b/Engine/Source/Runtime/RenderCore/Private/RenderingThread.cpp
--- a/Engine/Source/Runtime/RenderCore/Private/RenderingThread.cpp
+++ b/Engine/Source/Runtime/RenderCore/Private/RenderingThread.cpp
@@ -559,6 +559,8 @@ static FString BuildRenderingThreadName( uint32 ThreadIndex )
 	return FString::Printf( TEXT( "%s %u" ), *FName( NAME_RenderThread ).GetPlainNameString(), ThreadIndex );
 }
 
+static FRunnableThread* GRenderingThread = nullptr;
+
 static void StartRenderingThread()
 {
 	check(IsInGameThread());
