Removed 32-bit support from mallocs.

diff --git a/Engine/Source/Runtime/Core/Private/HAL/MallocBinned.cpp b/Engine/Source/Runtime/Core/Private/HAL/MallocBinned.cpp
--- a/Engine/Source/Runtime/Core/Private/HAL/MallocBinned.cpp
+++ b/Engine/Source/Runtime/Core/Private/HAL/MallocBinned.cpp
@@ -60,10 +60,6 @@ struct FMallocBinned::FPoolInfo
 	FFreeMem*		FirstMem;   // 12/16
 	FPoolInfo*		Next;		// 16/24
 	FPoolInfo**		PrevLink;	// 20/32
-#if PLATFORM_32BITS
-	/** Explicit padding for 32 bit builds */
-	uint8 Padding[12]; // 32
-#endif
 
 	void SetAllocationSizes( uint32 InBytes, UPTRINT InOsBytes, uint32 InTableIndex, uint32 SmallAllocLimt )
 	{
diff --git a/Engine/Source/Runtime/Core/Private/HAL/MallocBinnedGPU.cpp b/Engine/Source/Runtime/Core/Private/HAL/MallocBinnedGPU.cpp
--- a/Engine/Source/Runtime/Core/Private/HAL/MallocBinnedGPU.cpp
+++ b/Engine/Source/Runtime/Core/Private/HAL/MallocBinnedGPU.cpp
@@ -4,7 +4,7 @@
 
 PRAGMA_DISABLE_UNSAFE_TYPECAST_WARNINGS
 
-#if PLATFORM_64BITS && PLATFORM_HAS_FPlatformVirtualMemoryBlock
+#if PLATFORM_HAS_FPlatformVirtualMemoryBlock
 #include "Logging/LogMacros.h"
 #include "Templates/Function.h"
 #include "GenericPlatform/GenericPlatformProcess.h"
@@ -526,7 +526,6 @@ FMallocBinnedGPU::FMallocBinnedGPU()
 	MBG_STAT(ConsolidatedMemory = 0;)
 	MBG_STAT(GPUProxyMemory = 0;)
 
-	check(!PLATFORM_32BITS);
 	ArenaParams.BasePageSize = 4096;
 	ArenaParams.AllocationGranularity = FPlatformMemory::FPlatformVirtualMemoryBlock::GetCommitAlignment();
 
diff --git a/Engine/Source/Runtime/Core/Private/HAL/MallocStomp.cpp b/Engine/Source/Runtime/Core/Private/HAL/MallocStomp.cpp
--- a/Engine/Source/Runtime/Core/Private/HAL/MallocStomp.cpp
+++ b/Engine/Source/Runtime/Core/Private/HAL/MallocStomp.cpp
@@ -26,12 +26,8 @@ thread_local SIZE_T VirtualAddressMax = 0;
 #define MALLOC_STOMP_KEEP_VIRTUAL_MEMORY 0
 #endif
 
-#if PLATFORM_64BITS
 // 64-bit ABIs on x86_64 expect a 16-byte alignment
 #define STOMPALIGNMENT 16U
-#else
-#define STOMPALIGNMENT 0U
-#endif
 
 static void MallocStompOverrunTest()
 {
@@ -93,14 +89,12 @@ void* FMallocStomp::TryMalloc(SIZE_T Size, uint32 Alignment)
 		Size = 1U;
 	}
 
-#if PLATFORM_64BITS
 	// 64-bit ABIs on x86_64 expect a 16-byte alignment
 	Alignment = FMath::Max<uint32>(Alignment, STOMPALIGNMENT);
-#endif
 
 	constexpr static SIZE_T AllocationDataSize = sizeof(FAllocationData);
 
-	const SIZE_T AlignedSize = Alignment ? ((Size + Alignment - 1) & -(int32)Alignment) : Size;
+	const SIZE_T AlignedSize = ((Size + Alignment - 1) & -(int32)Alignment);
 	const SIZE_T AlignmentSize = Alignment > PageSize ? Alignment - PageSize : 0;
 	const SIZE_T AllocFullPageSize = (AlignedSize + AlignmentSize + AllocationDataSize + PageSize - 1) & -(SSIZE_T)PageSize;
 	const SIZE_T TotalAllocationSize = AllocFullPageSize + PageSize;
diff --git a/Engine/Source/Runtime/Core/Private/HAL/MallocStomp2.cpp b/Engine/Source/Runtime/Core/Private/HAL/MallocStomp2.cpp
--- a/Engine/Source/Runtime/Core/Private/HAL/MallocStomp2.cpp
+++ b/Engine/Source/Runtime/Core/Private/HAL/MallocStomp2.cpp
@@ -61,10 +61,8 @@ void* FMallocStomp2::Malloc(SIZE_T Size, uint32 Alignment)
 
 void* FMallocStomp2::TryMalloc(SIZE_T Size, uint32 Alignment)
 {
-#if PLATFORM_64BITS
 	// 64-bit ABIs on x86_64 expect a 16-byte alignment
 	Alignment = FMath::Max<uint32>(Alignment, 16U);
-#endif
 
 	void* FullAllocationPointer = nullptr;
 
@@ -102,7 +100,7 @@ void* FMallocStomp2::TryMalloc(SIZE_T Size, uint32 Alignment)
 
 	if(bUseUnderrunMode)
 	{
-		const SIZE_T AlignedAllocationData = (Alignment > 0U) ? ((AllocationDataSize + Alignment - 1U) & -static_cast<int32>(Alignment)) : AllocationDataSize;
+		const SIZE_T AlignedAllocationData = ((AllocationDataSize + Alignment - 1U) & -static_cast<int32>(Alignment));
 		ReturnedPointer = reinterpret_cast<void*>(reinterpret_cast<uint8*>(FullAllocationPointer) + PageSize + AlignedAllocationData);
 	}
 	else
diff --git a/Engine/Source/Runtime/Core/Public/HAL/MallocBinned.h b/Engine/Source/Runtime/Core/Public/HAL/MallocBinned.h
--- a/Engine/Source/Runtime/Core/Public/HAL/MallocBinned.h
+++ b/Engine/Source/Runtime/Core/Public/HAL/MallocBinned.h
@@ -42,22 +42,14 @@ struct FGenericMemoryStats;
 
 #if defined CACHE_FREED_OS_ALLOCS
 	#define MAX_CACHED_OS_FREES (64)
-	#if PLATFORM_64BITS
-		#define MAX_CACHED_OS_FREES_BYTE_LIMIT (64*1024*1024)
-	#else
-		#define MAX_CACHED_OS_FREES_BYTE_LIMIT (16*1024*1024)
-	#endif
+	#define MAX_CACHED_OS_FREES_BYTE_LIMIT (64*1024*1024)
 #endif
 
 #if defined USE_INTERNAL_LOCKS && !defined USE_COARSE_GRAIN_LOCKS
 #	define USE_FINE_GRAIN_LOCKS
 #endif
 
-#if PLATFORM_64BITS
 typedef int64 BINNED_STAT_TYPE;
-#else
-typedef int32 BINNED_STAT_TYPE;
-#endif
 
 //when modifying the global allocator stats, if we are using COARSE locks, then all callsites for stat modification are covered by the allocator-wide access guard. Thus the stats can be modified directly.
 //If we are using FINE locks, then we must modify the stats through atomics as the locks are either not actually covering the stat callsites, or are locking specific table locks which is not sufficient for stats.
diff --git a/Engine/Source/Runtime/Core/Public/HAL/MallocBinned2.h b/Engine/Source/Runtime/Core/Public/HAL/MallocBinned2.h
--- a/Engine/Source/Runtime/Core/Public/HAL/MallocBinned2.h
+++ b/Engine/Source/Runtime/Core/Public/HAL/MallocBinned2.h
@@ -19,11 +19,7 @@
 
 
 #define UE_MB2_MAX_CACHED_OS_FREES (64)
-#if PLATFORM_64BITS
-#	define UE_MB2_MAX_CACHED_OS_FREES_BYTE_LIMIT (64*1024*1024)
-#else
-#	define UE_MB2_MAX_CACHED_OS_FREES_BYTE_LIMIT (16*1024*1024)
-#endif
+#define UE_MB2_MAX_CACHED_OS_FREES_BYTE_LIMIT (64*1024*1024)
 
 #define UE_MB2_LARGE_ALLOC					65536		// Alignment of OS-allocated pointer - pool-allocated pointers will have a non-aligned pointer
 
diff --git a/Engine/Source/Runtime/Core/Public/HAL/MallocBinnedGPU.h b/Engine/Source/Runtime/Core/Public/HAL/MallocBinnedGPU.h
--- a/Engine/Source/Runtime/Core/Public/HAL/MallocBinnedGPU.h
+++ b/Engine/Source/Runtime/Core/Public/HAL/MallocBinnedGPU.h
@@ -10,7 +10,7 @@
 #include "Templates/Atomic.h"
 #include "Templates/MemoryOps.h"
 
-#if PLATFORM_64BITS && PLATFORM_HAS_FPlatformVirtualMemoryBlock
+#if PLATFORM_HAS_FPlatformVirtualMemoryBlock
 #include "Async/UniqueLock.h"
 #include "HAL/Allocators/CachedOSPageAllocator.h"
 #include "HAL/Allocators/PooledVirtualMemoryAllocator.h"
diff --git a/Engine/Source/Runtime/Core/Public/HAL/MallocStomp2.h b/Engine/Source/Runtime/Core/Public/HAL/MallocStomp2.h
--- a/Engine/Source/Runtime/Core/Public/HAL/MallocStomp2.h
+++ b/Engine/Source/Runtime/Core/Public/HAL/MallocStomp2.h
@@ -33,14 +33,8 @@ private:
 		return (((InSize)+(InAlignment - 1)) & ~(InAlignment - 1));
 	}
 
-
-#if PLATFORM_64BITS
 	/** Expected value to be found in the sentinel. */
 	static constexpr SIZE_T SentinelExpectedValue = 0xdeadbeefdeadbeef;
-#else
-	/** Expected value to be found in the sentinel. */
-	static constexpr SIZE_T SentinelExpectedValue = 0xdeadbeef;
-#endif
 
 	const SIZE_T PageSize;
 
