Mark TExternalMutex as deprecated.

TExternalMutex is just a thin wrapper around TIntrusiveMutex. Please switch to calling into TIntrusiveMutex directly.

diff --git a/Engine/Source/Runtime/Core/Public/Async/ExternalMutex.h b/Engine/Source/Runtime/Core/Public/Async/ExternalMutex.h
--- a/Engine/Source/Runtime/Core/Public/Async/ExternalMutex.h
+++ b/Engine/Source/Runtime/Core/Public/Async/ExternalMutex.h
@@ -20,7 +20,8 @@ namespace UE
  * It is valid to use TExternalMutex exclusively as a temporary, e.g., TExternalMutex(State).Lock();
  */
 template <CIntrusiveMutexParams ParamsType>
-class TExternalMutex final
+class UE_DEPRECATED(5.8, "Use TIntrusiveMutex.")
+TExternalMutex final
 {
 	using StateType = TIntrusiveMutexStateType_T<ParamsType>;
 
@@ -68,6 +69,8 @@ struct FExternalMutexParams
 
 } // Core::Private
 
-using FExternalMutex UE_DEPRECATED(5.7, "Use TExternalMutex or TIntrusiveMutex.") = TExternalMutex<Core::Private::FExternalMutexParams>;
+PRAGMA_DISABLE_DEPRECATION_WARNINGS
+using FExternalMutex UE_DEPRECATED(5.7, "Use TIntrusiveMutex.") = TExternalMutex<Core::Private::FExternalMutexParams>;
+PRAGMA_ENABLE_DEPRECATION_WARNINGS
 
 } // UE
diff --git a/Engine/Source/Runtime/Core/Tests/Async/ExternalMutexTest.cpp b/Engine/Source/Runtime/Core/Tests/Async/ExternalMutexTest.cpp
--- a/Engine/Source/Runtime/Core/Tests/Async/ExternalMutexTest.cpp
+++ b/Engine/Source/Runtime/Core/Tests/Async/ExternalMutexTest.cpp
@@ -9,6 +9,8 @@
 #include "HAL/Thread.h"
 #include "TestHarness.h"
 
+PRAGMA_DISABLE_DEPRECATION_WARNINGS
+
 namespace UE
 {
 
@@ -112,4 +114,6 @@ TEST_CASE("Core::Async::ExternalMutex", "[Core][Async][Slow][LinuxSkip]")
 
 } // UE
 
+PRAGMA_ENABLE_DEPRECATION_WARNINGS
+
 #endif // WITH_LOW_LEVEL_TESTS
diff --git a/Engine/Source/Runtime/Core/Tests/Async/MutexLockRateTest.cpp b/Engine/Source/Runtime/Core/Tests/Async/MutexLockRateTest.cpp
--- a/Engine/Source/Runtime/Core/Tests/Async/MutexLockRateTest.cpp
+++ b/Engine/Source/Runtime/Core/Tests/Async/MutexLockRateTest.cpp
@@ -205,6 +205,7 @@ TEST_CASE("Core::Async::MutexLockRate", "[.][Core][Async][Slow]")
 
 	SECTION("ExternalMutex")
 	{
+		PRAGMA_DISABLE_DEPRECATION_WARNINGS
 		UE_LOG(LogCore, Display, TEXT("TExternalMutex(uint8)"));
 		struct FExternalMutexType
 		{
@@ -214,6 +215,7 @@ TEST_CASE("Core::Async::MutexLockRate", "[.][Core][Async][Slow]")
 			void Unlock() { Mutex.Unlock(); }
 		};
 		TestLockRate(*MakeUnique<FExternalMutexType>(), LockTarget, IterationCount);
+		PRAGMA_ENABLE_DEPRECATION_WARNINGS
 	}
 
 	SECTION("IntrusiveMutex")
