[HashTable] Add some more explicit checks, since we are getting failures.

diff --git a/Engine/Source/Runtime/Core/Public/Containers/HashTable.h b/Engine/Source/Runtime/Core/Public/Containers/HashTable.h
--- a/Engine/Source/Runtime/Core/Public/Containers/HashTable.h
+++ b/Engine/Source/Runtime/Core/Public/Containers/HashTable.h
@@ -543,8 +543,8 @@ inline THashTable<InAllocator>::THashTable(uint32 InHashSize, uint32 InIndexSize
 	: HashMask(InHashSize - 1u)
 	, IndexSize(InIndexSize)
 {
-	check(InHashSize > 0u && InHashSize <= 0x10000);
-	check(FMath::IsPowerOfTwo(InHashSize));
+	checkf(InHashSize > 0u && InHashSize <= 0x10000, TEXT("InHashSize out of range: %u, InIndexSize %u"), InHashSize, InIndexSize);
+	checkf(FMath::IsPowerOfTwo(InHashSize), TEXT("InHashSize (%u) is not a power of two, InIndexSize %u"), InHashSize, InIndexSize);
 
 	Hash.ResizeAllocation(0, InHashSize, sizeof(uint32));
 	FMemory::Memset(Hash.GetAllocation(), 0xff, InHashSize * 4);
