Deprecated 32-bit/64-bit options.

diff --git a/Engine/Source/Runtime/Core/Public/GenericPlatform/GenericPlatformMisc.h b/Engine/Source/Runtime/Core/Public/GenericPlatform/GenericPlatformMisc.h
--- a/Engine/Source/Runtime/Core/Public/GenericPlatform/GenericPlatformMisc.h
+++ b/Engine/Source/Runtime/Core/Public/GenericPlatform/GenericPlatformMisc.h
@@ -1114,7 +1114,7 @@ public:
 	 *
 	 * @return	true if we're running in a 64-bit operating system
 	 */
-	//UE_DEPRECATED(5.8, "FPlatformMisc::Is64bitOperatingSystem is deprecated as UE only supports 64-bit.")
+	UE_DEPRECATED(5.8, "FPlatformMisc::Is64bitOperatingSystem is deprecated as UE only supports 64-bit.")
 	FORCEINLINE static bool Is64bitOperatingSystem()
 	{
 		return true;
diff --git a/Engine/Source/Runtime/Core/Public/HAL/Platform.h b/Engine/Source/Runtime/Core/Public/HAL/Platform.h
--- a/Engine/Source/Runtime/Core/Public/HAL/Platform.h
+++ b/Engine/Source/Runtime/Core/Public/HAL/Platform.h
@@ -676,15 +676,13 @@
 #define PLATFORM_CONSOLE (#)
 
 #ifdef PLATFORM_32BITS
-	#undef PLATFORM_32BITS
-//	#error PLATFORM_32BITS should not be defined.
+	#error PLATFORM_32BITS should not be defined.
 #endif
 #ifdef PLATFORM_64BITS
-	#undef PLATFORM_64BITS
-//	#error PLATFORM_64BITS should not be defined.
+	#error PLATFORM_64BITS should not be defined.
 #endif
-#define PLATFORM_32BITS /*UE_DEPRECATED_MACRO(5.8, "PLATFORM_32BITS is deprecated as UE only supports 64-bit.")*/ 0
-#define PLATFORM_64BITS /*UE_DEPRECATED_MACRO(5.8, "PLATFORM_64BITS is deprecated as UE only supports 64-bit.")*/ 1
+#define PLATFORM_32BITS UE_DEPRECATED_MACRO(5.8, "PLATFORM_32BITS is deprecated as UE only supports 64-bit.") 0
+#define PLATFORM_64BITS UE_DEPRECATED_MACRO(5.8, "PLATFORM_64BITS is deprecated as UE only supports 64-bit.") 1
 
 // not supported by the platform system yet or maybe ever
 #define PLATFORM_VTABLE_AT_END_OF_CLASS 0
