- Removed unnecessary confighierarchy strings, and plugin branches that have no static layers
- Removed the "config xxx" console commands in SHIPPING

diff --git a/Engine/Plugins/Runtime/GameFeatures/Source/GameFeatures/Private/GameFeatureData.cpp b/Engine/Plugins/Runtime/GameFeatures/Source/GameFeatures/Private/GameFeatureData.cpp
--- a/Engine/Plugins/Runtime/GameFeatures/Source/GameFeatures/Private/GameFeatureData.cpp
+++ b/Engine/Plugins/Runtime/GameFeatures/Source/GameFeatures/Private/GameFeatureData.cpp
@@ -171,6 +171,9 @@ void UGameFeatureData::InitializeBasePluginIniFile(const FString& PluginInstalle
 
 			return;
 		}
+
+		// remove the branch so we don't dirty it up with empty branches
+		GConfig->Remove(PluginConfigFilename);
 	}
 
 	if (bUseNewDynamicLayers)
@@ -256,7 +259,7 @@ void UGameFeatureData::InitializeHierarchicalPluginIniFiles(const FString& Plugi
 		UE::DynamicConfig::PerformDynamicConfig(PluginName, [PluginName](FConfigModificationTracker* ChangeTracker)
 			{
 				{
-					// Hotfixes still use [ConsoleVariables], not Plugin_CVarsLow/HighPriority - since these are set at Hotfix priority, 
+					// Hotfixes still use [ConsoleVariables], not PluginCVars_Low/HighPriority - since these are set at Hotfix priority, 
 					// we use the HF style naming
 					ChangeTracker->CVars.Add(TEXT("ConsoleVariables")).CVarPriority = (int)ECVF_SetByHotfix;
 
@@ -682,7 +685,7 @@ void UGameFeatureData::InitializeHierarchicalPluginIniFiles(const TArrayView<FSt
 		UE::DynamicConfig::PerformDynamicConfig("InitializeHierarchicalPluginIniFiles", [&PluginNames](FConfigModificationTracker* ChangeTracker)
 			{
 				{
-					// Hotfixes still use [ConsoleVariables], not Plugin_CVarsLow/HighPriority - since these are set at Hotfix priority, 
+					// Hotfixes still use [ConsoleVariables], not PluginCVars_Low/HighPriority - since these are set at Hotfix priority, 
 					// we use the HF style naming
 					ChangeTracker->CVars.Add(TEXT("ConsoleVariables")).CVarPriority = (int)ECVF_SetByHotfix;
 
diff --git a/Engine/Source/Runtime/Core/Private/Misc/ConfigCacheIni.cpp b/Engine/Source/Runtime/Core/Private/Misc/ConfigCacheIni.cpp
--- a/Engine/Source/Runtime/Core/Private/Misc/ConfigCacheIni.cpp
+++ b/Engine/Source/Runtime/Core/Private/Misc/ConfigCacheIni.cpp
@@ -7208,6 +7208,7 @@ double GPrepareForLoadTime = 0;
 double GPerformLoadTime = 0;
 double GConfigShrinkTime = 0;
 
+#if !UE_BUILD_SHIPPING
 class FIniExec : public FSelfRegisteringExec
 {
 	virtual bool Exec_Dev(class UWorld* InWorld, const TCHAR* Cmd, FOutputDevice& Ar) override
@@ -7794,6 +7795,6 @@ class FIniExec : public FSelfRegisteringExec
 	}
 	
 } GConfigExec;
-
+#endif // !UE_BUILD_SHIPPING
 
 #undef LOCTEXT_NAMESPACE
diff --git a/Engine/Source/Runtime/Core/Private/Misc/ConfigContext.cpp b/Engine/Source/Runtime/Core/Private/Misc/ConfigContext.cpp
--- a/Engine/Source/Runtime/Core/Private/Misc/ConfigContext.cpp
+++ b/Engine/Source/Runtime/Core/Private/Misc/ConfigContext.cpp
@@ -509,6 +509,15 @@ bool FConfigContext::PrepareForLoad(bool& bPerformLoad)
 			check(!DestIniFilename.IsEmpty());
 			
 			Branch = &ConfigSystem->AddNewBranch(DestIniFilename);
+
+			// we allow plugin branches to get wiped if they had no static layers - non-plugins we keep just in case
+			// a plugin wants to modify it (a plugin modifying a plugin would require that the modified plugin has
+			// actual config data in it)
+			if (bIsForPlugin)
+			{
+				static bool bRemoveEmptyPluginBranches = !FParse::Param(FCommandLine::Get(), TEXT("NoIniPluginBranchClean"));
+				bCanRemoveBranchIfEmpty = bRemoveEmptyPluginBranches;
+			}
 		}
 	}
 	Branch->bIsHierarchical = bIsHierarchicalConfig;
@@ -590,6 +599,12 @@ bool FConfigContext::PerformLoad()
 	// just write out the exact same thing it read in!
 	bool bGeneratedFile = GenerateDestIniFile();
 
+	// if we wiped a branch that had no static files, we cannot continue
+	if (bCanRemoveBranchIfEmpty && !bGeneratedFile)
+	{
+		return false;
+	}
+
 	// we are done here!
 	if (bIsFixingUpAfterBinaryConfig)
 	{
@@ -1007,7 +1022,8 @@ void FConfigContext::AddStaticLayersToHierarchy(TArray<FString>* GatheredLayerFi
 bool FConfigContext::LoadIniFileHierarchy()
 {
 	static bool bDumpIniLoadInfo = FParse::Param(FCommandLine::Get(), TEXT("dumpiniloads"));
-	
+	static bool bRemoveMissingLayersFromHierarchy = !FParse::Param(FCommandLine::Get(), TEXT("NoIniHierarchyClean"));
+
 	// LogVariables(*BaseIniName, Platform);
 
 #if ALLOW_INI_OVERRIDE_FROM_COMMANDLINE
@@ -1021,9 +1037,8 @@ bool FConfigContext::LoadIniFileHierarchy()
 	}
 #endif
 
-
-
 	bool bReadAnyFile = false;
+	TSet<int32> KeysToRemove;
 
 	TRACE_CPUPROFILER_EVENT_SCOPE(LoadIniFileHierarchy);
 	// Traverse ini list back to front, merging along the way.
@@ -1042,6 +1057,10 @@ bool FConfigContext::LoadIniFileHierarchy()
 		// skip non-existant files
 		if (!bAssumeExists && IsUsingLocalIniFile(*IniFileName, nullptr) && !DoesConfigFileExistWrapper(*IniFileName, IniCacheSet, bUseStagedGlobalCache ? ConfigSystem->StagedGlobalConfigCache : nullptr, bUseStagedPluginCache ? StagedPluginConfigCache : nullptr))
 		{
+			if (bRemoveMissingLayersFromHierarchy)
+			{
+				KeysToRemove.Add(HierarchyIt.Key);
+			}
 			continue;
 		}
 		
@@ -1087,6 +1106,11 @@ bool FConfigContext::LoadIniFileHierarchy()
 		}
 	}
 
+	for (int32 Key : KeysToRemove)
+	{
+		Branch->Hierarchy.Remove(Key);
+	}
+
 	// if we had been reading into the Static cache, not InMemory, then start the InMemory from this point
 	if (Branch->ReplayMethod != EBranchReplayMethod::NoReplay)
 	{
@@ -1128,7 +1152,15 @@ bool FConfigContext::GenerateDestIniFile()
 		Branch->FinalCombinedLayers.Empty();
 
 		// read the static files into the branch
-		LoadIniFileHierarchy();
+		if (!LoadIniFileHierarchy())
+		{
+			if (bCanRemoveBranchIfEmpty)
+			{
+				ConfigSystem->Remove(DestIniFilename);
+				Branch = TemporaryBranch = nullptr;
+				return false;
+			}
+		}
 	}
 	
 #if !IS_PROGRAM
diff --git a/Engine/Source/Runtime/Core/Public/Misc/ConfigContext.h b/Engine/Source/Runtime/Core/Public/Misc/ConfigContext.h
--- a/Engine/Source/Runtime/Core/Public/Misc/ConfigContext.h
+++ b/Engine/Source/Runtime/Core/Public/Misc/ConfigContext.h
@@ -250,6 +250,7 @@ public:
 	bool bIncludeTagNameInBranchName = false;
 	bool bIsMakingBinaryConfig = false;
 	bool bIsFixingUpAfterBinaryConfig = false;
+	bool bCanRemoveBranchIfEmpty = false;
 
 	DynamicLayerPriority PluginModificationPriority;
 
