Apply Editor ini settings on top of Engine ones when creating CrashReportClient.ini from editor
Read from CrashReportClient.ini in FWindowsPlatformCrashContext::CopyPlatformSpecificFiles

diff --git a/Engine/Source/Runtime/Core/Private/GenericPlatform/GenericPlatformCrashContext.cpp b/Engine/Source/Runtime/Core/Private/GenericPlatform/GenericPlatformCrashContext.cpp
--- a/Engine/Source/Runtime/Core/Private/GenericPlatform/GenericPlatformCrashContext.cpp
+++ b/Engine/Source/Runtime/Core/Private/GenericPlatform/GenericPlatformCrashContext.cpp
@@ -696,18 +696,33 @@ void FGenericCrashContext::InitializeFromConfig()
 	PurgeOldCrashConfig();
 
 	const bool bForceGetSection = false;
-	const FConfigSection* CRCConfigSection = GConfig->GetSection(ConfigSectionName, bForceGetSection, GEngineIni);
+	FConfigFile CrashConfigFile;
 
-	if (CRCConfigSection != nullptr)
+	if (const FConfigSection* CRCConfigSection = GConfig->GetSection(ConfigSectionName, bForceGetSection, GEngineIni))
 	{
 		// Create a config file and save to a temp location. This file will be copied to
 		// the crash folder for all crash reports create by this session.
-		FConfigFile CrashConfigFile;
 
 		FConfigSection CRCConfigSectionCopy(*CRCConfigSection);
 		CrashConfigFile.Add(ConfigSectionName, CRCConfigSectionCopy);
 
 		CrashConfigFile.Dirty = true;
+	}
+
+#if WITH_EDITOR
+	if (const FConfigSection* CRCEConfigSection = GConfig->GetSection(ConfigSectionName, bForceGetSection, GEditorIni))
+	{
+		for (FConfigSection::TConstIterator ConfigValueIt(*CRCEConfigSection); ConfigValueIt; ++ConfigValueIt)
+		{
+			CrashConfigFile.SetInSection(ConfigSectionName, ConfigValueIt.Key(), ConfigValueIt.Value().GetValue());
+		}
+
+		CrashConfigFile.Dirty = true;
+	}
+#endif
+
+	if (CrashConfigFile.Dirty)
+	{
 		CrashConfigFile.Write(FString(GetCrashConfigFilePath()));
 	}
 
diff --git a/Engine/Source/Runtime/Core/Private/Windows/WindowsPlatformCrashContext.cpp b/Engine/Source/Runtime/Core/Private/Windows/WindowsPlatformCrashContext.cpp
--- a/Engine/Source/Runtime/Core/Private/Windows/WindowsPlatformCrashContext.cpp
+++ b/Engine/Source/Runtime/Core/Private/Windows/WindowsPlatformCrashContext.cpp
@@ -360,19 +360,42 @@ void FWindowsPlatformCrashContext::CopyPlatformSpecificFiles(const TCHAR* Output
 {
 	FGenericCrashContext::CopyPlatformSpecificFiles(OutputDirectory, Context);
 
+	auto GetRecordCrashDump = [this](const TCHAR* Key)
+		{
+			bool bRecordCrashDump = true;
+
+			static const TCHAR* SectionName = TEXT("CrashReportClient");
+
+			const TCHAR* CrashConfigFilePath = GetCrashConfigFilePath();
+			if (IFileManager::Get().FileExists(CrashConfigFilePath))
+			{
+				FConfigFile CrashConfigFile;
+				CrashConfigFile.Read(CrashConfigFilePath);
+				if (CrashConfigFile.GetBool(SectionName, Key, bRecordCrashDump))
+				{
+					return bRecordCrashDump;
+				}
+			}
+
+#if UE_EDITOR
+			if (GConfig->GetBool(SectionName, Key, bRecordCrashDump, GEditorIni))
+			{
+				return bRecordCrashDump;
+			}
+#endif
+
+			GConfig->GetBool(SectionName, Key, bRecordCrashDump, GEngineIni);
+			return bRecordCrashDump;
+		};
+
 	bool bRecordCrashDump = true;
 	switch (Type)
 	{
 		case ECrashContextType::Stall:
-			GConfig->GetBool(TEXT("CrashReportClient"), TEXT("Stall.RecordDump"), bRecordCrashDump, GEngineIni);
+			bRecordCrashDump = GetRecordCrashDump(TEXT("Stall.RecordDump"));
 			break;
 		case ECrashContextType::Ensure:
-#if UE_EDITOR
-			if (!GConfig->GetBool(TEXT("CrashReportClient"), TEXT("Ensure.RecordDump"), bRecordCrashDump, GEditorIni))
-#endif
-			{
-				GConfig->GetBool(TEXT("CrashReportClient"), TEXT("Ensure.RecordDump"), bRecordCrashDump, GEngineIni);
-			}
+			bRecordCrashDump = GetRecordCrashDump(TEXT("Ensure.RecordDump"));
 			break;
 		case ECrashContextType::AbnormalShutdown:
 		case ECrashContextType::Assert:
