Fixed GetVarArgs character count in the presence of encoding conversions

#rnx

diff --git a/Engine/Source/Runtime/Core/Private/GenericPlatform/GenericPlatformString.cpp b/Engine/Source/Runtime/Core/Private/GenericPlatform/GenericPlatformString.cpp
--- a/Engine/Source/Runtime/Core/Private/GenericPlatform/GenericPlatformString.cpp
+++ b/Engine/Source/Runtime/Core/Private/GenericPlatform/GenericPlatformString.cpp
@@ -80,29 +80,15 @@ namespace
 			// Ensure we're not in the error state
 			if (EndMinusOne)
 			{
-				NumCharsWritten += Count;
-
-				if constexpr (std::is_same_v<ANSICHAR, OtherCharType> && std::is_same_v<WIDECHAR, CharType>)
+				using DstType = std::conditional_t<std::is_same_v<ANSICHAR, CharType>, UTF8CHAR, CharType>;
+				using SrcType = std::conditional_t<std::is_same_v<ANSICHAR, OtherCharType>, UTF8CHAR, OtherCharType>;
+				if constexpr (!std::is_same_v<DstType, SrcType>)
 				{
-					int32 ConvertedLength = FPlatformString::ConvertedLength<WIDECHAR, UTF8CHAR>((const UTF8CHAR*)Src, Count);
+					int32 ConvertedLength = FPlatformString::ConvertedLength<DstType, SrcType>((const SrcType*)Src, Count);
 					int32 NumToWrite = FPlatformMath::Min(ConvertedLength, (int32)(EndMinusOne - Ptr));
-					CharType* Result = FPlatformString::Convert(Ptr, NumToWrite, (const UTF8CHAR*)Src, Count);
+					CharType* Result = (CharType*)FPlatformString::Convert((DstType*)Ptr, NumToWrite, (const SrcType*)Src, Count);
 
-					if (NumToWrite != ConvertedLength)
-					{
-						Ptr = Ptr + NumToWrite;
-						EndMinusOne = 0;
-					}
-					else
-					{
-						Ptr = Result;
-					}
-				}
-				else if constexpr (std::is_same_v<WIDECHAR, OtherCharType> && std::is_same_v<ANSICHAR, CharType>)
-				{
-					int32 ConvertedLength = FPlatformString::ConvertedLength<UTF8CHAR, WIDECHAR>(Src, Count);
-					int32 NumToWrite = FPlatformMath::Min(ConvertedLength, (int32)(EndMinusOne - Ptr));
-					CharType* Result = (CharType*)FPlatformString::Convert((UTF8CHAR*)Ptr, NumToWrite, Src, Count);
+					NumCharsWritten += ConvertedLength;
 
 					if (NumToWrite != ConvertedLength)
 					{
@@ -116,6 +102,8 @@ namespace
 				}
 				else
 				{
+					NumCharsWritten += Count;
+
 					int32 NumToWrite = FPlatformMath::Min(Count, (int32)(EndMinusOne - Ptr));
 
 					for (int32 Val = NumToWrite; Val; --Val)
@@ -1416,182 +1404,97 @@ int32 FGenericPlatformString::Strncmp(const UTF8CHAR* Str1, const UTF8CHAR* Str2
 
 #if WITH_LOW_LEVEL_TESTS
 
+#include "Containers/AnsiString.h"
+#include "Containers/Utf8String.h"
 #include "TestHarness.h"
 #include <catch2/generators/catch_generators.hpp>
 
-constexpr int OUTPUT_SIZE = 256;
-static int32 TestGetVarArgs(WIDECHAR* OutputString, const WIDECHAR* Format, ...)
+template <int32 OutputSize, typename ExpectedCharType, typename FormatCharType>
+static void TestGetVarArgs(const ExpectedCharType* Expected, const FormatCharType* Format, ...)
 {
-	va_list ArgPtr;
-	va_start(ArgPtr, Format);
-	int32 Result = FGenericPlatformString::GetVarArgs(OutputString, OUTPUT_SIZE, Format, ArgPtr);
-	va_end(ArgPtr);
-	return Result;
-}
+	FormatCharType OutputString[OutputSize];
 
-static int32 TestGetVarArgs(ANSICHAR* OutputString, const ANSICHAR* Format, ...)
-{
-	va_list ArgPtr;
-	va_start(ArgPtr, Format);
-	int32 Result = FGenericPlatformString::GetVarArgs(OutputString, OUTPUT_SIZE, Format, ArgPtr);
-	va_end(ArgPtr);
-	return Result;
-}
-
-static int32 TestGetVarArgs(WIDECHAR* OutputString, int OutputSize, const WIDECHAR* Format, ...)
-{
 	va_list ArgPtr;
 	va_start(ArgPtr, Format);
 	int32 Result = FGenericPlatformString::GetVarArgs(OutputString, OutputSize, Format, ArgPtr);
 	va_end(ArgPtr);
-	return Result;
-}
 
-static int32 TestGetVarArgs(ANSICHAR* OutputString, int OutputSize, const ANSICHAR* Format, ...)
-{
-	va_list ArgPtr;
-	va_start(ArgPtr, Format);
-	int32 Result = FGenericPlatformString::GetVarArgs(OutputString, OutputSize, Format, ArgPtr);
-	va_end(ArgPtr);
-	return Result;
+	TString<ExpectedCharType> ExpectedString(Expected);
+	CAPTURE(ExpectedString);
+	CHECK(ExpectedString == TString<ExpectedCharType>(OutputString));
+	CHECK(ExpectedString.Len() == Result);
 }
 
 TEST_CASE("Core::GenericPlatformString::GetVarArgs", "[Core][String][Smoke]")
 {
-	{
-		WIDECHAR OutputString[OUTPUT_SIZE];
+	constexpr int32 OutputSize = 256;
 
-		TestGetVarArgs(OutputString, TEXT("Test A|%-20s|%20s|%10.2f|%-10.2f|"), TEXT("LEFT"), TEXT("RIGHT"), 33.333333, 66.666666);
-		CHECK_MESSAGE(OutputString, FString(OutputString) == FString(TEXT("Test A|LEFT                |               RIGHT|     33.33|66.67     |")));
+	TestGetVarArgs<OutputSize>(TEXT("Test A|LEFT                |               RIGHT|     33.33|66.67     |"), TEXT("Test A|%-20s|%20s|%10.2f|%-10.2f|"), TEXT("LEFT"), TEXT("RIGHT"), 33.333333, 66.666666);
+	TestGetVarArgs<OutputSize>(TEXT("Test B|Percents:%%%3|"), TEXT("Test B|Percents:%%%%%%%d|"), 3);
+	TestGetVarArgs<OutputSize>(TEXT("Test C|12345|54321|123ABC|f|99|"), TEXT("Test C|%d|%i|%X|%x|%u|"), 12345, 54321, 0x123AbC, 15, 99);
 
-		TestGetVarArgs(OutputString, TEXT("Test B|Percents:%%%%%%%d|"), 3);
-		CHECK_MESSAGE(OutputString, FString(OutputString) == FString(TEXT("Test B|Percents:%%%3|")));
-
-		TestGetVarArgs(OutputString, TEXT("Test C|%d|%i|%X|%x|%u|"), 12345, 54321, 0x123AbC, 15, 99);
-		CHECK_MESSAGE(OutputString, FString(OutputString) == FString(TEXT("Test C|12345|54321|123ABC|f|99|")));
-
-		TestGetVarArgs(OutputString, TEXT("Test D|%p|"), 0x12345);
 #ifdef _MSC_VER
-		// MSVC's standard library formats pointers differently
-		CHECK_MESSAGE(OutputString, FString(OutputString) == FString(TEXT("Test D|0000000000012345|")));
+	// MSVC's standard library formats pointers differently
+	TestGetVarArgs<OutputSize>(TEXT("Test D|0000000000012345|"), TEXT("Test D|%p|"), 0x12345);
 #else
-		// Pointer format (which we get from snprintf) varies from platform to platform.
-		// Call it to make sure it puts the output in the proper place, but don't test for the exact format
-		//	CHECK_MESSAGE(OutputString, FString(OutputString) == FString(TEXT("Test D|0x12345|")));
-		{
-			FString OutputStringStr(OutputString);
-			CHECK_MESSAGE(OutputString, OutputStringStr.StartsWith(TEXT("Test D|"))
-				&& OutputStringStr.EndsWith(TEXT("|")) && OutputStringStr.Len() < TEXTVIEW("Test D||").Len() + 20);
-		}
-#endif
-
-		TestGetVarArgs(OutputString, TEXT("Test E|%" INT64_FMT "|"), int64(12345678912345LL));
-		CHECK_MESSAGE(OutputString, FString(OutputString) == FString(TEXT("Test E|12345678912345|")));
-
-		TestGetVarArgs(OutputString, TEXT("Test F|%f|%e|%g|"), 123.456, 123.456, 123.456);
-		CHECK_MESSAGE(OutputString, FString(OutputString) == FString(TEXT("Test F|123.456000|1.234560e+02|123.456|")));
-
-		TestGetVarArgs(OutputString, TEXT("Test G|%" UPTRINT_X_FMT"|"), UPTRINT(49374));
-		CHECK_MESSAGE(OutputString, FString(OutputString) == FString(TEXT("Test G|C0DE|")));
-
-		TestGetVarArgs(OutputString, TEXT("Test H|%" UPTRINT_x_FMT "|"), UPTRINT(49374));
-		CHECK_MESSAGE(OutputString, FString(OutputString) == FString(TEXT("Test H|c0de|")));
-
-		TestGetVarArgs(OutputString, TEXT("Test I|%" UINT64_FMT "|"), MAX_uint64);
-		CHECK_MESSAGE(OutputString, FString(OutputString) == FString(TEXT("Test I|18446744073709551615|")));
-
-		TestGetVarArgs(OutputString, TEXT("Test J|%*c|"), 10, 'J');
-		CHECK_MESSAGE(OutputString, FString(OutputString) == FString(TEXT("Test J|         J|")));
-
-		TestGetVarArgs(OutputString, TEXT("Test K|%-5c|"), 'K');
-		CHECK_MESSAGE(OutputString, FString(OutputString) == FString(TEXT("Test K|K    |")));
-
-		TestGetVarArgs(OutputString, TEXT("Test L|%hs|"), UTF8TEXT("Hello World"));
-		CHECK_MESSAGE(OutputString, FString(OutputString) == FString(TEXT("Test L|Hello World|")));
-
-		TestGetVarArgs(OutputString, TEXT("Test M|%hs|"), "\xF0\x9F\x98\x86\xF0\x9F\x80\x9C");
-		CHECK_MESSAGE(OutputString, FString(OutputString) == FString(UTF16_TO_TCHAR(u"Test M|\xD83D\xDE06\xD83C\xDC1C|")));
-
-		TestGetVarArgs(OutputString, TEXT("Test N|%ls|"), UTF16_TO_TCHAR(u"\xD83D\xDE06\xD83C\xDC1C"));
-		CHECK_MESSAGE(OutputString, FString(OutputString) == FString(UTF16_TO_TCHAR(u"Test N|\xD83D\xDE06\xD83C\xDC1C|")));
-
-		TestGetVarArgs(OutputString, TEXT("Test O|%s|"), UTF16_TO_TCHAR(u"\xD83D\xDE06\xD83C\xDC1C"));
-		CHECK_MESSAGE(OutputString, FString(OutputString) == FString(UTF16_TO_TCHAR(u"Test O|\xD83D\xDE06\xD83C\xDC1C|")));
-
-		TestGetVarArgs(OutputString, TEXT("Test P|%S|"), "\xF0\x9F\x98\x86\xF0\x9F\x80\x9C");
-		CHECK_MESSAGE(OutputString, FString(OutputString) == FString(UTF16_TO_TCHAR(u"Test P|\xD83D\xDE06\xD83C\xDC1C|")));
-
-		WIDECHAR SmallOutputString[9];
-		TestGetVarArgs(SmallOutputString, 9, TEXT("%08x"), 0x12345678);
-		CHECK_MESSAGE(SmallOutputString, FString(SmallOutputString) == FString(TEXT("12345678")));
-	}
-
+	// Pointer format (which we get from snprintf) varies from platform to platform.
+	// Call it to make sure it puts the output in the proper place, but don't test for the exact format
 	{
-		ANSICHAR OutputString[OUTPUT_SIZE];
-
-		TestGetVarArgs(OutputString, "Test A|%-20s|%20s|%10.2f|%-10.2f|", "LEFT", "RIGHT", 33.333333, 66.666666);
-		CHECK_MESSAGE(OutputString, FString(OutputString) == FString("Test A|LEFT                |               RIGHT|     33.33|66.67     |"));
-
-		TestGetVarArgs(OutputString, "Test B|Percents:%%%%%%%d|", 3);
-		CHECK_MESSAGE(OutputString, FString(OutputString) == FString("Test B|Percents:%%%3|"));
+		WIDECHAR Output[OutputSize];
+		TCString<WIDECHAR>::Snprintf(Output, OutputSize, WIDETEXT("Test D|%p|"), 0x12345);
+		FString OutputString(Output);
+		CHECK_MESSAGE(Output, OutputString.StartsWith(TEXT("Test D|"))
+			&& OutputString.EndsWith(TEXT("|")) && OutputString.Len() < TEXTVIEW("Test D||").Len() + 20);
+	}
+#endif
 
-		TestGetVarArgs(OutputString, "Test C|%d|%i|%X|%x|%u|", 12345, 54321, 0x123AbC, 15, 99);
-		CHECK_MESSAGE(OutputString, FString(OutputString) == FString("Test C|12345|54321|123ABC|f|99|"));
+	TestGetVarArgs<OutputSize>(TEXT("Test E|12345678912345|"), TEXT("Test E|%" INT64_FMT "|"), int64(12345678912345LL));
+	TestGetVarArgs<OutputSize>(TEXT("Test F|123.456000|1.234560e+02|123.456|"), TEXT("Test F|%f|%e|%g|"), 123.456, 123.456, 123.456);
+	TestGetVarArgs<OutputSize>(TEXT("Test G|C0DE|"), TEXT("Test G|%" UPTRINT_X_FMT"|"), UPTRINT(49374));
+	TestGetVarArgs<OutputSize>(TEXT("Test H|c0de|"), TEXT("Test H|%" UPTRINT_x_FMT "|"), UPTRINT(49374));
+	TestGetVarArgs<OutputSize>(TEXT("Test I|18446744073709551615|"), TEXT("Test I|%" UINT64_FMT "|"), MAX_uint64);
+	TestGetVarArgs<OutputSize>(TEXT("Test J|         J|"), TEXT("Test J|%*c|"), 10, 'J');
+	TestGetVarArgs<OutputSize>(TEXT("Test K|K    |"), TEXT("Test K|%-5c|"), 'K');
+	TestGetVarArgs<OutputSize>(TEXT("Test L|Hello World|"), TEXT("Test L|%hs|"), UTF8TEXT("Hello World"));
+	TestGetVarArgs<OutputSize>(TEXT("Test M|\U0001F606\U0001F41C|"), TEXT("Test M|%hs|"), UTF8TEXT("\U0001F606\U0001F41C"));
+	TestGetVarArgs<OutputSize>(TEXT("Test N|\U0001F606\U0001F41C|"), TEXT("Test N|%ls|"), TEXT("\U0001F606\U0001F41C"));
+	TestGetVarArgs<OutputSize>(TEXT("Test O|\U0001F606\U0001F41C|"), TEXT("Test O|%s|"), TEXT("\U0001F606\U0001F41C"));
+	TestGetVarArgs<OutputSize>(TEXT("Test P|\U0001F606\U0001F41C|"), TEXT("Test P|%S|"), UTF8TEXT("\U0001F606\U0001F41C"));
+
+	TestGetVarArgs<9>(TEXT("12345678"), TEXT("%08x"), 0x12345678);
+
+	TestGetVarArgs<OutputSize>("Test A|LEFT                |               RIGHT|     33.33|66.67     |", "Test A|%-20s|%20s|%10.2f|%-10.2f|", "LEFT", "RIGHT", 33.333333, 66.666666);
+	TestGetVarArgs<OutputSize>("Test B|Percents:%%%3|", "Test B|Percents:%%%%%%%d|", 3);
+	TestGetVarArgs<OutputSize>("Test C|12345|54321|123ABC|f|99|", "Test C|%d|%i|%X|%x|%u|", 12345, 54321, 0x123AbC, 15, 99);
 
-		TestGetVarArgs(OutputString, "Test D|%p|", 0x12345);
 #ifdef _MSC_VER
-		// MSVC's standard library formats pointers differently
-		CHECK_MESSAGE(OutputString, FString(OutputString) == FString("Test D|0000000000012345|"));
+	// MSVC's standard library formats pointers differently
+	TestGetVarArgs<OutputSize>("Test D|0000000000012345|", "Test D|%p|", 0x12345);
 #else
-		// Pointer format (which we get from snprintf) varies from platform to platform.
-		// Call it to make sure it puts the output in the proper place, but don't test for the exact format
-		//	CHECK_MESSAGE(OutputString, FString(OutputString) == FString(TEXT("Test D|0x12345|")));
-		{
-			FString OutputStringStr(OutputString);
-			CHECK_MESSAGE(OutputString, OutputStringStr.StartsWith("Test D|")
-				&& OutputStringStr.EndsWith("|") && OutputStringStr.Len() < TEXTVIEW("Test D||").Len() + 20);
-		}
+	// Pointer format (which we get from snprintf) varies from platform to platform.
+	// Call it to make sure it puts the output in the proper place, but don't test for the exact format
+	{
+		ANSICHAR Output[OutputSize];
+		TCString<ANSICHAR>::Snprintf(Output, OutputSize, "Test D|%p|", 0x12345);
+		FAnsiString OutputString(Output);
+		CHECK_MESSAGE(Output, OutputString.StartsWith(TEXT("Test D|"))
+			&& OutputString.EndsWith(TEXT("|")) && OutputString.Len() < TEXTVIEW("Test D||").Len() + 20);
+	}
 #endif
 
-		TestGetVarArgs(OutputString, "Test E|%" INT64_FMT "|", int64(12345678912345LL));
-		CHECK_MESSAGE(OutputString, FString(OutputString) == FString("Test E|12345678912345|"));
-
-		TestGetVarArgs(OutputString, "Test F|%f|%e|%g|", 123.456, 123.456, 123.456);
-		CHECK_MESSAGE(OutputString, FString(OutputString) == FString("Test F|123.456000|1.234560e+02|123.456|"));
-
-		TestGetVarArgs(OutputString, "Test G|%" UPTRINT_X_FMT"|", UPTRINT(49374));
-		CHECK_MESSAGE(OutputString, FString(OutputString) == FString("Test G|C0DE|"));
-
-		TestGetVarArgs(OutputString, "Test H|%" UPTRINT_x_FMT "|", UPTRINT(49374));
-		CHECK_MESSAGE(OutputString, FString(OutputString) == FString("Test H|c0de|"));
-
-		TestGetVarArgs(OutputString, "Test I|%" UINT64_FMT "|", MAX_uint64);
-		CHECK_MESSAGE(OutputString, FString(OutputString) == FString("Test I|18446744073709551615|"));
-
-		TestGetVarArgs(OutputString, "Test J|%*c|", 10, 'J');
-		CHECK_MESSAGE(OutputString, FString(OutputString) == FString("Test J|         J|"));
-
-		TestGetVarArgs(OutputString, "Test K|%-5c|", 'K');
-		CHECK_MESSAGE(OutputString, FString(OutputString) == FString("Test K|K    |"));
-
-		TestGetVarArgs(OutputString, "Test L|%hs|", UTF8TEXT("Hello World"));
-		CHECK_MESSAGE(OutputString, FString(OutputString) == FString("Test L|Hello World|"));
-
-		TestGetVarArgs(OutputString, "Test M|%ls|", UTF16_TO_TCHAR(u"\xD83D\xDE06\xD83C\xDC1C"));
-		CHECK_MESSAGE(OutputString, FString(OutputString) == FString("Test M|\xF0\x9F\x98\x86\xF0\x9F\x80\x9C|"));
-
-		TestGetVarArgs(OutputString, "Test N|%hs|", "\xF0\x9F\x98\x86\xF0\x9F\x80\x9C");
-		CHECK_MESSAGE(OutputString, FString(OutputString) == FString("Test N|\xF0\x9F\x98\x86\xF0\x9F\x80\x9C|"));
-
-		TestGetVarArgs(OutputString, "Test O|%s|", "\xF0\x9F\x98\x86\xF0\x9F\x80\x9C");
-		CHECK_MESSAGE(OutputString, FString(OutputString) == FString("Test O|\xF0\x9F\x98\x86\xF0\x9F\x80\x9C|"));
-
-		TestGetVarArgs(OutputString, "Test P|%S|", UTF16_TO_TCHAR(u"\xD83D\xDE06\xD83C\xDC1C"));
-		CHECK_MESSAGE(OutputString, FString(OutputString) == FString("Test P|\xF0\x9F\x98\x86\xF0\x9F\x80\x9C|"));
-
-		ANSICHAR SmallOutputString[9];
-		TestGetVarArgs(SmallOutputString, 9, "%08x", 0x12345678);
-		CHECK_MESSAGE(SmallOutputString, FString(SmallOutputString) == FString("12345678"));
-	}
+	TestGetVarArgs<OutputSize>("Test E|12345678912345|", "Test E|%" INT64_FMT "|", int64(12345678912345LL));
+	TestGetVarArgs<OutputSize>("Test F|123.456000|1.234560e+02|123.456|", "Test F|%f|%e|%g|", 123.456, 123.456, 123.456);
+	TestGetVarArgs<OutputSize>("Test G|C0DE|", "Test G|%" UPTRINT_X_FMT"|", UPTRINT(49374));
+	TestGetVarArgs<OutputSize>("Test H|c0de|", "Test H|%" UPTRINT_x_FMT "|", UPTRINT(49374));
+	TestGetVarArgs<OutputSize>("Test I|18446744073709551615|", "Test I|%" UINT64_FMT "|", MAX_uint64);
+	TestGetVarArgs<OutputSize>("Test J|         J|", "Test J|%*c|", 10, 'J');
+	TestGetVarArgs<OutputSize>("Test K|K    |", "Test K|%-5c|", 'K');
+	TestGetVarArgs<OutputSize>("Test L|Hello World|", "Test L|%hs|", UTF8TEXT("Hello World"));
+	TestGetVarArgs<OutputSize>(UTF8TEXT("Test M|\U0001F606\U0001F41C|"), "Test M|%ls|", TEXT("\U0001F606\U0001F41C"));
+	TestGetVarArgs<OutputSize>(UTF8TEXT("Test N|\U0001F606\U0001F41C|"), "Test N|%hs|", UTF8TEXT("\U0001F606\U0001F41C"));
+	TestGetVarArgs<OutputSize>(UTF8TEXT("Test O|\U0001F606\U0001F41C|"), "Test O|%s|", UTF8TEXT("\U0001F606\U0001F41C"));
+	TestGetVarArgs<OutputSize>(UTF8TEXT("Test P|\U0001F606\U0001F41C|"), "Test P|%S|", TEXT("\U0001F606\U0001F41C"));
+
+	TestGetVarArgs<9>("12345678", "%08x", 0x12345678);
 }
 #endif // WITH_LOW_LEVEL_TESTS
\ No newline at end of file
