Fixed crash in FMallocStomp2 with over-aligned allocations.
Replaced bespoke aligning code with AlignmentTemplates.h.

diff --git a/Engine/Source/Runtime/Core/Private/HAL/MallocStomp2.cpp b/Engine/Source/Runtime/Core/Private/HAL/MallocStomp2.cpp
--- a/Engine/Source/Runtime/Core/Private/HAL/MallocStomp2.cpp
+++ b/Engine/Source/Runtime/Core/Private/HAL/MallocStomp2.cpp
@@ -9,6 +9,7 @@
 #include "HAL/IConsoleManager.h"
 #include "GenericPlatform/GenericPlatformMemory.h"
 #include "Algo/BinarySearch.h"
+#include "Templates/AlignmentTemplates.h"
 
 CORE_API FMallocStomp2* GMallocStomp2 = nullptr;
 CORE_API bool GMallocStomp2Enabled = false;
@@ -66,13 +67,13 @@ void* FMallocStomp2::TryMalloc(SIZE_T Size, uint32 Alignment)
 
 	void* FullAllocationPointer = nullptr;
 
-	if (Size < MinSize || Size >= MaxSize || FPlatformTLS::GetTlsValue(FMallocStomp2::DisableTlsSlot) != 0 )
+	if (Size < MinSize || Size >= MaxSize || Alignment > PageSize || FPlatformTLS::GetTlsValue(FMallocStomp2::DisableTlsSlot) != 0)
 	{
 		return UsedMalloc->Malloc(Size, Alignment);
 	}
 
-	const SIZE_T AlignedSize = GetAlignedSize(Size, Alignment);
-	const SIZE_T AllocFullPageSize = GetAlignedSize(AlignedSize + sizeof(FAllocationData), PageSize);
+	const SIZE_T AlignedSize = Align(Size, Alignment);
+	const SIZE_T AllocFullPageSize = Align(AlignedSize + sizeof(FAllocationData), PageSize);
 	const SIZE_T TotalAllocationSize = AllocFullPageSize + PageSize;
 
 
@@ -100,7 +101,7 @@ void* FMallocStomp2::TryMalloc(SIZE_T Size, uint32 Alignment)
 
 	if(bUseUnderrunMode)
 	{
-		const SIZE_T AlignedAllocationData = ((AllocationDataSize + Alignment - 1U) & -static_cast<int32>(Alignment));
+		const SIZE_T AlignedAllocationData = Align(AllocationDataSize, Alignment);
 		ReturnedPointer = reinterpret_cast<void*>(reinterpret_cast<uint8*>(FullAllocationPointer) + PageSize + AlignedAllocationData);
 	}
 	else
diff --git a/Engine/Source/Runtime/Core/Public/HAL/MallocStomp2.h b/Engine/Source/Runtime/Core/Public/HAL/MallocStomp2.h
--- a/Engine/Source/Runtime/Core/Public/HAL/MallocStomp2.h
+++ b/Engine/Source/Runtime/Core/Public/HAL/MallocStomp2.h
@@ -28,11 +28,6 @@
 class FMallocStomp2 final : public FMalloc
 {
 private:
-	UE_FORCEINLINE_HINT SIZE_T GetAlignedSize(SIZE_T InSize, SIZE_T InAlignment)
-	{
-		return (((InSize)+(InAlignment - 1)) & ~(InAlignment - 1));
-	}
-
 	/** Expected value to be found in the sentinel. */
 	static constexpr SIZE_T SentinelExpectedValue = 0xdeadbeefdeadbeef;
 
