[GenericPlatform]

Don't enable the RenderThread if the application can never render. This alleviates unnecessary syncronization between the gamethread and the render thread for calls like FlushRenderingCommands

#rnx

#FailedValidations Preflight

diff --git a/Engine/Source/Runtime/Core/Private/GenericPlatform/GenericPlatformMisc.cpp b/Engine/Source/Runtime/Core/Private/GenericPlatform/GenericPlatformMisc.cpp
--- a/Engine/Source/Runtime/Core/Private/GenericPlatform/GenericPlatformMisc.cpp
+++ b/Engine/Source/Runtime/Core/Private/GenericPlatform/GenericPlatformMisc.cpp
@@ -1629,7 +1629,7 @@ bool FGenericPlatformMisc::UseRenderThread()
 #endif
 
 	// allow if not overridden
-	return true;
+	return FApp::CanEverRender();
 }
 
 bool FGenericPlatformMisc::AllowThreadHeartBeat()
