Switched to ANSI literals for GET_MEMBER* macros.

diff --git a/Engine/Source/Runtime/Core/Public/Misc/AssertionMacros.h b/Engine/Source/Runtime/Core/Public/Misc/AssertionMacros.h
--- a/Engine/Source/Runtime/Core/Public/Misc/AssertionMacros.h
+++ b/Engine/Source/Runtime/Core/Public/Misc/AssertionMacros.h
@@ -479,9 +479,9 @@ namespace UEAsserts_Private
 	bool GetMemberNameCheckedJunk(R(*)(Args...));
 }
 
-// Returns FName(TEXTVIEW("EnumeratorName")), while statically verifying that the enumerator exists in the enum
+// Returns FName(ANSITEXTVIEW("EnumeratorName")), while statically verifying that the enumerator exists in the enum
 #define GET_ENUMERATOR_NAME_CHECKED(EnumName, EnumeratorName) \
-	(FName(GET_ENUMERATOR_NAME_STRING_VIEW_CHECKED(EnumName, EnumeratorName)))
+	(FName(GET_ENUMERATOR_NAME_ANSI_STRING_VIEW_CHECKED(EnumName, EnumeratorName)))
 
 #define GET_ENUMERATOR_NAME_STRING_CHECKED(EnumName, EnumeratorName) \
 	((void)sizeof(UEAsserts_Private::GetMemberNameCheckedJunk(EnumName::EnumeratorName)), TEXT(#EnumeratorName))
@@ -489,9 +489,15 @@ namespace UEAsserts_Private
 #define GET_ENUMERATOR_NAME_STRING_VIEW_CHECKED(EnumName, EnumeratorName) \
 	((void)sizeof(UEAsserts_Private::GetMemberNameCheckedJunk(EnumName::EnumeratorName)), TEXTVIEW(#EnumeratorName))
 
-// Returns FName(TEXTVIEW("MemberName")), while statically verifying that the member exists in ClassName
+#define GET_ENUMERATOR_NAME_ANSI_STRING_CHECKED(EnumName, EnumeratorName) \
+	((void)sizeof(UEAsserts_Private::GetMemberNameCheckedJunk(EnumName::EnumeratorName)), #EnumeratorName)
+
+#define GET_ENUMERATOR_NAME_ANSI_STRING_VIEW_CHECKED(EnumName, EnumeratorName) \
+	((void)sizeof(UEAsserts_Private::GetMemberNameCheckedJunk(EnumName::EnumeratorName)), ANSITEXTVIEW(#EnumeratorName))
+
+// Returns FName(ANSITEXTVIEW("MemberName")), while statically verifying that the member exists in ClassName
 #define GET_MEMBER_NAME_CHECKED(ClassName, MemberName) \
-	(FName(GET_MEMBER_NAME_STRING_VIEW_CHECKED(ClassName, MemberName)))
+	(FName(GET_MEMBER_NAME_ANSI_STRING_VIEW_CHECKED(ClassName, MemberName)))
 
 #define GET_MEMBER_NAME_STRING_CHECKED(ClassName, MemberName) \
 	((void)sizeof(UEAsserts_Private::GetMemberNameCheckedJunk(((ClassName*)0)->MemberName)), TEXT(#MemberName))
@@ -499,9 +505,15 @@ namespace UEAsserts_Private
 #define GET_MEMBER_NAME_STRING_VIEW_CHECKED(ClassName, MemberName) \
 	((void)sizeof(UEAsserts_Private::GetMemberNameCheckedJunk(((ClassName*)0)->MemberName)), TEXTVIEW(#MemberName))
 
-// Returns FName(TEXTVIEW("FunctionName")), while statically verifying that the function exists in ClassName
+#define GET_MEMBER_NAME_ANSI_STRING_CHECKED(ClassName, MemberName) \
+	((void)sizeof(UEAsserts_Private::GetMemberNameCheckedJunk(((ClassName*)0)->MemberName)), #MemberName)
+
+#define GET_MEMBER_NAME_ANSI_STRING_VIEW_CHECKED(ClassName, MemberName) \
+	((void)sizeof(UEAsserts_Private::GetMemberNameCheckedJunk(((ClassName*)0)->MemberName)), ANSITEXTVIEW(#MemberName))
+
+// Returns FName(ANSITEXTVIEW("FunctionName")), while statically verifying that the function exists in ClassName
 #define GET_FUNCTION_NAME_CHECKED(ClassName, FunctionName) \
-	(FName(GET_FUNCTION_NAME_STRING_VIEW_CHECKED(ClassName, FunctionName)))
+	(FName(GET_FUNCTION_NAME_ANSI_STRING_VIEW_CHECKED(ClassName, FunctionName)))
 
 #define GET_FUNCTION_NAME_STRING_CHECKED(ClassName, FunctionName) \
 	((void)sizeof(&ClassName::FunctionName), TEXT(#FunctionName))
@@ -509,19 +521,25 @@ namespace UEAsserts_Private
 #define GET_FUNCTION_NAME_STRING_VIEW_CHECKED(ClassName, FunctionName) \
 	((void)sizeof(&ClassName::FunctionName), TEXTVIEW(#FunctionName))
 
+#define GET_FUNCTION_NAME_ANSI_STRING_CHECKED(ClassName, FunctionName) \
+	((void)sizeof(&ClassName::FunctionName), #FunctionName)
+
+#define GET_FUNCTION_NAME_ANSI_STRING_VIEW_CHECKED(ClassName, FunctionName) \
+	((void)sizeof(&ClassName::FunctionName), ANSITEXTVIEW(#FunctionName))
+
 // Returns FName("FunctionName"), while statically verifying that the function exists in ClassName
 // Handles overloaded functions by specifying the argument of the overload to use
 #define GET_FUNCTION_NAME_CHECKED_OneParam(ClassName, FunctionName, ArgType) \
-	(FName(GET_FUNCTION_NAME_STRING_CHECKED_OneParam(ClassName, FunctionName, ArgType)))
+	(FName(GET_FUNCTION_NAME_ANSI_STRING_CHECKED_OneParam(ClassName, FunctionName, ArgType)))
 
 #define GET_FUNCTION_NAME_CHECKED_TwoParams(ClassName, FunctionName, ArgType1, ArgType2) \
-	(FName(GET_FUNCTION_NAME_STRING_CHECKED_TwoParams(ClassName, FunctionName, ArgType1, ArgType2)))
+	(FName(GET_FUNCTION_NAME_ANSI_STRING_CHECKED_TwoParams(ClassName, FunctionName, ArgType1, ArgType2)))
 
 #define GET_FUNCTION_NAME_CHECKED_ThreeParams(ClassName, FunctionName, ArgType1, ArgType2, ArgType3) \
-	(FName(GET_FUNCTION_NAME_STRING_CHECKED_ThreeParams(ClassName, FunctionName, ArgType1, ArgType2, ArgType3)))
+	(FName(GET_FUNCTION_NAME_ANSI_STRING_CHECKED_ThreeParams(ClassName, FunctionName, ArgType1, ArgType2, ArgType3)))
 
 #define GET_FUNCTION_NAME_CHECKED_FourParams(ClassName, FunctionName, ArgType1, ArgType2, ArgType3, ArgType4) \
-	(FName(GET_FUNCTION_NAME_STRING_CHECKED_FourParams(ClassName, FunctionName, ArgType1, ArgType2, ArgType3, ArgType4)))
+	(FName(GET_FUNCTION_NAME_ANSI_STRING_CHECKED_FourParams(ClassName, FunctionName, ArgType1, ArgType2, ArgType3, ArgType4)))
 
 #define GET_FUNCTION_NAME_STRING_CHECKED_OneParam(ClassName, FunctionName, ArgType) \
 	((void)sizeof((std::declval<ClassName&>().FunctionName(std::declval<UE_REMOVE_OPTIONAL_PARENS(ArgType)>()), (int)0)), TEXT(#FunctionName))
@@ -535,6 +553,18 @@ namespace UEAsserts_Private
 #define GET_FUNCTION_NAME_STRING_CHECKED_FourParams(ClassName, FunctionName, ArgType1, ArgType2, ArgType3, ArgType4) \
 	((void)sizeof((std::declval<ClassName&>().FunctionName(std::declval<UE_REMOVE_OPTIONAL_PARENS(ArgType1)>(), std::declval<UE_REMOVE_OPTIONAL_PARENS(ArgType2)>(), std::declval<UE_REMOVE_OPTIONAL_PARENS(ArgType3)>(), std::declval<UE_REMOVE_OPTIONAL_PARENS(ArgType4)>()), (int)0)), TEXT(#FunctionName))
 
+#define GET_FUNCTION_NAME_ANSI_STRING_CHECKED_OneParam(ClassName, FunctionName, ArgType) \
+	((void)sizeof((std::declval<ClassName&>().FunctionName(std::declval<UE_REMOVE_OPTIONAL_PARENS(ArgType)>()), (int)0)), #FunctionName)
+
+#define GET_FUNCTION_NAME_ANSI_STRING_CHECKED_TwoParams(ClassName, FunctionName, ArgType1, ArgType2) \
+	((void)sizeof((std::declval<ClassName&>().FunctionName(std::declval<UE_REMOVE_OPTIONAL_PARENS(ArgType1)>(), std::declval<UE_REMOVE_OPTIONAL_PARENS(ArgType2)>()), (int)0)), #FunctionName)
+
+#define GET_FUNCTION_NAME_ANSI_STRING_CHECKED_ThreeParams(ClassName, FunctionName, ArgType1, ArgType2, ArgType3) \
+	((void)sizeof((std::declval<ClassName&>().FunctionName(std::declval<UE_REMOVE_OPTIONAL_PARENS(ArgType1)>(), std::declval<UE_REMOVE_OPTIONAL_PARENS(ArgType2)>(), std::declval<UE_REMOVE_OPTIONAL_PARENS(ArgType3)>()), (int)0)), #FunctionName)
+
+#define GET_FUNCTION_NAME_ANSI_STRING_CHECKED_FourParams(ClassName, FunctionName, ArgType1, ArgType2, ArgType3, ArgType4) \
+	((void)sizeof((std::declval<ClassName&>().FunctionName(std::declval<UE_REMOVE_OPTIONAL_PARENS(ArgType1)>(), std::declval<UE_REMOVE_OPTIONAL_PARENS(ArgType2)>(), std::declval<UE_REMOVE_OPTIONAL_PARENS(ArgType3)>(), std::declval<UE_REMOVE_OPTIONAL_PARENS(ArgType4)>()), (int)0)), #FunctionName)
+
 /*----------------------------------------------------------------------------
 	Low level error macros
 ----------------------------------------------------------------------------*/
