Replaced PREPROCESSOR_JOIN.

diff --git a/Engine/Source/Runtime/Core/Public/ProfilingDebugging/SpatialTrace.h b/Engine/Source/Runtime/Core/Public/ProfilingDebugging/SpatialTrace.h
--- a/Engine/Source/Runtime/Core/Public/ProfilingDebugging/SpatialTrace.h
+++ b/Engine/Source/Runtime/Core/Public/ProfilingDebugging/SpatialTrace.h
@@ -262,7 +262,7 @@ namespace UE::SpatialTrace
 // Internal helpers
 
 #define UE_PRIVATE_TRACE_SPATIAL_POINT_SPEC_VAR(Suffix) \
-	PREPROCESSOR_JOIN(__SpatialTracePointSpec, Suffix)
+	UE_JOIN(__SpatialTracePointSpec, Suffix)
 
 #define UE_PRIVATE_TRACE_SPATIAL_POINT_INLINE_SPEC_DECLARE(SpecType, DisplayName) \
 	static SpecType UE_PRIVATE_TRACE_SPATIAL_POINT_SPEC_VAR(__LINE__) = SpecType::FromLiteral(DisplayName);
