[WorldPartition] Added a unittest for IsStreamingCompleted & GetIntersectingCells
#FailedValidations Preflight
#virtualized

diff --git a/Engine/Source/Runtime/Core/Public/Misc/AutomationTest.h b/Engine/Source/Runtime/Core/Public/Misc/AutomationTest.h
--- a/Engine/Source/Runtime/Core/Public/Misc/AutomationTest.h
+++ b/Engine/Source/Runtime/Core/Public/Misc/AutomationTest.h
@@ -4768,6 +4768,12 @@ public:
 		, Timeout(InTimeout)
 	{}
 
+	FUntilCommand(TFunction<bool()> InCallback, float InTimeout = 5.0f)
+		: Callback(MoveTemp(InCallback))
+		, TimeoutCallback([]() { return true; })
+		, Timeout(InTimeout)
+	{}
+
 	virtual bool Update() override
 	{
 		if ( !Callback() )
