Firefight - restore version bump to fix existing data validation

diff --git a/Engine/Source/Runtime/Core/Public/UObject/FortniteMainBranchObjectVersions.inl b/Engine/Source/Runtime/Core/Public/UObject/FortniteMainBranchObjectVersions.inl
--- a/Engine/Source/Runtime/Core/Public/UObject/FortniteMainBranchObjectVersions.inl
+++ b/Engine/Source/Runtime/Core/Public/UObject/FortniteMainBranchObjectVersions.inl
@@ -694,5 +694,8 @@ FORTNITE_MAIN_VERSION(FastGeoTransformerSettingAssets, 226)
 // UAF can now generate procedural systems for assets
 FORTNITE_MAIN_VERSION(UAFProceduralSystems, 227)
 
+// Firefight: version 228 was introduced and backed out, adding it back to allow assets created in the meantime to load
+FORTNITE_MAIN_VERSION(Firefight_228, 228)
+
 // -----<new versions can be added above this line>-------------------------------------------------
 #undef FORTNITE_MAIN_VERSION
\ No newline at end of file
