[ColorManagement]
* Fixed header unit compile error

diff --git a/Engine/Source/Runtime/Core/Public/ColorManagement/ColorManagementDefines.h b/Engine/Source/Runtime/Core/Public/ColorManagement/ColorManagementDefines.h
--- a/Engine/Source/Runtime/Core/Public/ColorManagement/ColorManagementDefines.h
+++ b/Engine/Source/Runtime/Core/Public/ColorManagement/ColorManagementDefines.h
@@ -77,7 +77,7 @@ enum class EChromaticAdaptationMethod : uint8
 };
 
 /** Default method used across the engine for chromatic adaptation. */
-constexpr EChromaticAdaptationMethod DEFAULT_CHROMATIC_ADAPTATION_METHOD = EChromaticAdaptationMethod::Bradford;
+inline constexpr EChromaticAdaptationMethod DEFAULT_CHROMATIC_ADAPTATION_METHOD = EChromaticAdaptationMethod::Bradford;
 
 /** List of standard white points. */
 enum class EWhitePoint : uint8
