Deprecated PREPROCESSOR_* macros.

diff --git a/Engine/Source/Runtime/Core/Public/HAL/PreprocessorHelpers.h b/Engine/Source/Runtime/Core/Public/HAL/PreprocessorHelpers.h
--- a/Engine/Source/Runtime/Core/Public/HAL/PreprocessorHelpers.h
+++ b/Engine/Source/Runtime/Core/Public/HAL/PreprocessorHelpers.h
@@ -104,13 +104,13 @@
 #endif
 
 // These macros should be regarded as deprecated - use the UE_ macros they map to instead.
-#define PREPROCESSOR_TO_STRING(Token)                 UE_STRINGIZE(Token)
-#define PREPROCESSOR_JOIN(TokenA, TokenB)             UE_JOIN(TokenA, TokenB)
-#define PREPROCESSOR_JOIN_FIRST(Token, ...)           UE_JOIN_FIRST(Token, ##__VA_ARGS__)
-#define PREPROCESSOR_IF(OneOrZero, Token1, Token0)    UE_IF(OneOrZero, Token1, Token0)
-#define PREPROCESSOR_COMMA_SEPARATED(First, ...)      UE_COMMA_SEPARATED(First, ##__VA_ARGS__)
-#define PREPROCESSOR_VA_ARG_COUNT(...)                UE_VA_ARG_COUNT(__VA_ARGS__)
-#define PREPROCESSOR_APPEND_VA_ARG_COUNT(Prefix, ...) UE_APPEND_VA_ARG_COUNT(Prefix, ##__VA_ARGS__)
-#define PREPROCESSOR_NOTHING                          UE_EMPTY
-#define PREPROCESSOR_NOTHING_FUNCTION(...)            UE_EMPTY_FUNCTION(__VA_ARGS__)
-#define PREPROCESSOR_REMOVE_OPTIONAL_PARENS(...)      UE_REMOVE_OPTIONAL_PARENS(__VA_ARGS__)
+#define PREPROCESSOR_TO_STRING(Token)                 UE_DEPRECATED_MACRO(5.8, "PREPROCESSOR_TO_STRING has been deprecated - please use UE_STRINGIZE instead.") UE_STRINGIZE(Token)
+#define PREPROCESSOR_JOIN(TokenA, TokenB)             UE_DEPRECATED_MACRO(5.8, "PREPROCESSOR_JOIN has been deprecated - please use UE_JOIN instead.") UE_JOIN(TokenA, TokenB)
+#define PREPROCESSOR_JOIN_FIRST(Token, ...)           UE_DEPRECATED_MACRO(5.8, "PREPROCESSOR_JOIN_FIRST has been deprecated - please use UE_JOIN_FIRST instead.") UE_JOIN_FIRST(Token, ##__VA_ARGS__)
+#define PREPROCESSOR_IF(OneOrZero, Token1, Token0)    UE_DEPRECATED_MACRO(5.8, "PREPROCESSOR_IF has been deprecated - please use UE_IF instead.") UE_IF(OneOrZero, Token1, Token0)
+#define PREPROCESSOR_COMMA_SEPARATED(First, ...)      UE_DEPRECATED_MACRO(5.8, "PREPROCESSOR_COMMA_SEPARATED has been deprecated - please use UE_COMMA_SEPARATED instead.") UE_COMMA_SEPARATED(First, ##__VA_ARGS__)
+#define PREPROCESSOR_VA_ARG_COUNT(...)                UE_DEPRECATED_MACRO(5.8, "PREPROCESSOR_VA_ARG_COUNT has been deprecated - please use UE_VA_ARG_COUNT instead.") UE_VA_ARG_COUNT(__VA_ARGS__)
+#define PREPROCESSOR_APPEND_VA_ARG_COUNT(Prefix, ...) UE_DEPRECATED_MACRO(5.8, "PREPROCESSOR_APPEND_VA_ARG_COUNT has been deprecated - please use UE_APPEND_VA_ARG_COUNT instead.") UE_APPEND_VA_ARG_COUNT(Prefix, ##__VA_ARGS__)
+#define PREPROCESSOR_NOTHING                          UE_DEPRECATED_MACRO(5.8, "PREPROCESSOR_NOTHING has been deprecated - please use UE_EMPTY instead.") UE_EMPTY
+#define PREPROCESSOR_NOTHING_FUNCTION(...)            UE_DEPRECATED_MACRO(5.8, "PREPROCESSOR_NOTHING_FUNCTION has been deprecated - please use UE_EMPTY_FUNCTION instead.") UE_EMPTY_FUNCTION(__VA_ARGS__)
+#define PREPROCESSOR_REMOVE_OPTIONAL_PARENS(...)      UE_DEPRECATED_MACRO(5.8, "PREPROCESSOR_REMOVE_OPTIONAL_PARENS has been deprecated - please use UE_REMOVE_OPTIONAL_PARENS instead.") UE_REMOVE_OPTIONAL_PARENS(__VA_ARGS__)
