Reverted some needless FString and FName constructions in tests where we want the redundancy.

diff --git a/Engine/Source/Runtime/Core/Private/Serialization/Csv/CsvParserTests.cpp b/Engine/Source/Runtime/Core/Private/Serialization/Csv/CsvParserTests.cpp
--- a/Engine/Source/Runtime/Core/Private/Serialization/Csv/CsvParserTests.cpp
+++ b/Engine/Source/Runtime/Core/Private/Serialization/Csv/CsvParserTests.cpp
@@ -123,7 +123,7 @@ IMPLEMENT_SIMPLE_AUTOMATION_TEST(FEmptyStringTest, "System.Core.CSV Parser.Empty
 bool FEmptyStringTest::RunTest(const FString& Parameters)
 {
 	return	FCsvParser(FString()).GetRows().Num() == 0 &&
-			FCsvParser(FString()).GetRows().Num() == 0;
+			FCsvParser(FString("")).GetRows().Num() == 0;
 }
 
 IMPLEMENT_SIMPLE_AUTOMATION_TEST(FOnlyNewLinesTest, "System.Core.CSV Parser.Only Newlines", EAutomationTestFlags::EditorContext | EAutomationTestFlags::EngineFilter)
diff --git a/Engine/Source/Runtime/Core/Private/UObject/UnrealNames.cpp b/Engine/Source/Runtime/Core/Private/UObject/UnrealNames.cpp
--- a/Engine/Source/Runtime/Core/Private/UObject/UnrealNames.cpp
+++ b/Engine/Source/Runtime/Core/Private/UObject/UnrealNames.cpp
@@ -3109,7 +3109,7 @@ struct FNameHelper
 	{
 		check(Len >= 0);
 
-		// Make NAME_None == TEXTVIEW("") or TEXTVIEW(nullptr) consistent with NAME_None == FName() or FName(nullptr)
+		// Make NAME_None == TEXTVIEW("") or TEXTVIEW(nullptr) consistent with NAME_None == FName(TEXT("")) or FName(nullptr)
 		if (Len == 0)
 		{
 			return Name.IsNone();
@@ -4075,8 +4075,8 @@ void FName::AutoTest()
 	const FName AutoTest_1("AutoTest_1");
 	const FName autoTest_1("autoTest_1");
 	const FName autoTeSt_1("autoTeSt_1");
-	const FName AutoTest_2("AutoTest_2");
-	const FName AutoTestB_2("AutoTestB_2");
+	const FName AutoTest_2(TEXT("AutoTest_2"));
+	const FName AutoTestB_2(TEXT("AutoTestB_2"));
 
 	check(AutoTest_1.GetNumber() == NAME_EXTERNAL_TO_INTERNAL(1));
 	check(autoTest_1.GetNumber() == NAME_EXTERNAL_TO_INTERNAL(1));
@@ -4101,7 +4101,7 @@ void FName::AutoTest()
 	
 	check(FName("ABC").ToString(Buffer) == 3 &&			!FCString::Strcmp(Buffer, TEXT("ABC")));
 	check(FName("abc").ToString(Buffer) == 3 &&			!FCString::Strcmp(Buffer, TEXT("abc")));
-	check(FName("abc").ToString(Buffer) == 3 &&	!FCString::Strcmp(Buffer, TEXT("abc")));
+	check(FName(TEXT("abc")).ToString(Buffer) == 3 &&	!FCString::Strcmp(Buffer, TEXT("abc")));
 	check(FName("ABC_0").ToString(Buffer) == 5 &&		!FCString::Strcmp(Buffer, TEXT("ABC_0")));
 	check(FName("ABC_10").ToString(Buffer) == 6 &&		!FCString::Strcmp(Buffer, TEXT("ABC_10")));	
 #endif
@@ -4127,12 +4127,12 @@ void FName::AutoTest()
 	check(NullName == FName(static_cast<WIDECHAR*>(nullptr)));
 	check(NullName == FName(NAME_None));
 	check(NullName == FName());
-	check(NullName == FName());
-	check(NullName == FName());
+	check(NullName == FName(""));
+	check(NullName == FName(TEXT("")));
 	check(NullName == FName("None"));
 	check(NullName == FName("none"));
 	check(NullName == FName("NONE"));
-	check(NullName == FName("None"));
+	check(NullName == FName(TEXT("None")));
 	check(FName().ToEName());
 	check(*FName().ToEName() == NAME_None);
 	check(NullName.GetComparisonIndex().ToUnstableInt() == 0);
@@ -4169,7 +4169,7 @@ void FName::AutoTest()
 	check(FName("Text_01") == FName("Text_01", NAME_NO_NUMBER_INTERNAL));
 
 	// Test unterminated strings
-	check(FName() == FName(0, "Unused"));
+	check(FName("") == FName(0, "Unused"));
 	check(FName("Used") == FName(4, "UsedUnused"));
 	check(FName("Used") == FName(4, "Used"));
 	check(FName("Used_0") == FName(6, "Used_01"));
diff --git a/Engine/Source/Runtime/Core/Public/Containers/UnrealString.h.inl b/Engine/Source/Runtime/Core/Public/Containers/UnrealString.h.inl
--- a/Engine/Source/Runtime/Core/Public/Containers/UnrealString.h.inl
+++ b/Engine/Source/Runtime/Core/Public/Containers/UnrealString.h.inl
@@ -1292,7 +1292,7 @@ public:
 
 		if (Num != OtherNum)
 		{
-			// Handle special case where FString() == FString()
+			// Handle special case where FString() == FString("")
 			return Num + OtherNum == 1;
 		}
 		else if (Num > 1)
