[Core]
* Fixed wrong comment in pragma warning disable 4291.

diff --git a/Engine/Source/Runtime/Core/Public/Windows/WindowsPlatformCompilerSetup.h b/Engine/Source/Runtime/Core/Public/Windows/WindowsPlatformCompilerSetup.h
--- a/Engine/Source/Runtime/Core/Public/Windows/WindowsPlatformCompilerSetup.h
+++ b/Engine/Source/Runtime/Core/Public/Windows/WindowsPlatformCompilerSetup.h
@@ -145,7 +145,7 @@ static_assert(_MSC_VER >= 1938, "Visual Studio 2022 v17.8 (MSVC 14.38) or greate
 #pragma warning(disable: 4251) // 'type' needs to have dll-interface to be used by clients of 'type'								https://docs.microsoft.com/en-us/cpp/error-messages/compiler-warnings/compiler-warning-level-1-c4251
 #pragma warning(disable: 4267) // 'var' : conversion from 'size_t' to 'type', possible loss of data									https://docs.microsoft.com/en-us/cpp/error-messages/compiler-warnings/compiler-warning-level-3-c4267
 #pragma warning(disable: 4275) // non - DLL-interface classkey 'identifier' used as base for DLL-interface classkey 'identifier'	https://docs.microsoft.com/en-us/cpp/error-messages/compiler-warnings/compiler-warning-level-2-c4275
-#pragma warning(disable: 4291) // typedef-name '' used as synonym for class-name ''													https://docs.microsoft.com/en-us/cpp/error-messages/compiler-warnings/compiler-warning-level-1-c4291
+#pragma warning(disable: 4291) // no matching operator delete found; memory will not be freed if initialization throws an exception https://docs.microsoft.com/en-us/cpp/error-messages/compiler-warnings/compiler-warning-level-1-c4291
 
 #pragma warning(disable: 4307) // '': integral constant overflow																	https://docs.microsoft.com/en-us/cpp/error-messages/compiler-warnings/compiler-warning-level-2-c4307
 #pragma warning(disable: 4315) // 'classname': 'this' pointer for member 'member' may not be aligned 'alignment' as expected by the constructor	https://docs.microsoft.com/en-us/cpp/error-messages/compiler-warnings/compiler-warnings-c4200-through-c4399
