Adding new delegates to provide notification of slate modal window creation/destruction with additional context (e.g., whether it is a slow task window)

diff --git a/Engine/Source/Runtime/ApplicationCore/Private/Mac/MacApplication.cpp b/Engine/Source/Runtime/ApplicationCore/Private/Mac/MacApplication.cpp
--- a/Engine/Source/Runtime/ApplicationCore/Private/Mac/MacApplication.cpp
+++ b/Engine/Source/Runtime/ApplicationCore/Private/Mac/MacApplication.cpp
@@ -186,8 +186,8 @@ FMacApplication::FMacApplication()
 	FMemory::Memzero(GestureUsage);
 	LastGestureUsed = EGestureEvent::None;
 
-	FCoreDelegates::PreSlateModal.AddRaw(this, &FMacApplication::StartScopedModalEvent);
-    FCoreDelegates::PostSlateModal.AddRaw(this, &FMacApplication::EndScopedModalEvent);
+	FCoreDelegates::PreSlateModalWithContext.AddRaw(this, &FMacApplication::StartScopedModalEvent);
+    FCoreDelegates::PostSlateModalWithContext.AddRaw(this, &FMacApplication::EndScopedModalEvent);
 #endif
 
 	FMacApplication::OnDisplayReconfiguration(kCGNullDirectDisplay, kCGDisplayDesktopShapeChangedFlag, this);
@@ -461,13 +461,13 @@ void FMacApplication::SendAnalytics(IAnalyticsProvider* Provider)
 	FMemory::Memzero(GestureUsage);
 	LastGestureUsed = EGestureEvent::None;
 }
-void FMacApplication::StartScopedModalEvent()
+void FMacApplication::StartScopedModalEvent(const FCoreDelegates::FModalWindowContext& Context)
 {
     FPlatformApplicationMisc::bMacApplicationModalMode = true;
     FPlatformApplicationMisc::bChachedMacMenuStateNeedsUpdate = true;
 }
 
-void FMacApplication::EndScopedModalEvent()
+void FMacApplication::EndScopedModalEvent(const FCoreDelegates::FModalWindowContext& Context)
 {
     FPlatformApplicationMisc::bMacApplicationModalMode = false;
     FPlatformApplicationMisc::bChachedMacMenuStateNeedsUpdate = true;
diff --git a/Engine/Source/Runtime/ApplicationCore/Public/Mac/MacApplication.h b/Engine/Source/Runtime/ApplicationCore/Public/Mac/MacApplication.h
--- a/Engine/Source/Runtime/ApplicationCore/Public/Mac/MacApplication.h
+++ b/Engine/Source/Runtime/ApplicationCore/Public/Mac/MacApplication.h
@@ -194,9 +194,9 @@ public:
 #if WITH_EDITOR
 	virtual void SendAnalytics(IAnalyticsProvider* Provider) override;
 
-	void StartScopedModalEvent();
+	void StartScopedModalEvent(const FCoreDelegates::FModalWindowContext& Context);
 
-	void EndScopedModalEvent();
+	void EndScopedModalEvent(const FCoreDelegates::FModalWindowContext& Context);
 #endif
 
 public:
diff --git a/Engine/Source/Runtime/Core/Private/Misc/CoreDelegates.cpp b/Engine/Source/Runtime/Core/Private/Misc/CoreDelegates.cpp
--- a/Engine/Source/Runtime/Core/Private/Misc/CoreDelegates.cpp
+++ b/Engine/Source/Runtime/Core/Private/Misc/CoreDelegates.cpp
@@ -113,8 +113,12 @@ TDelegate<void(TArray<uint8>&, TArray<uint8>&)>& FCoreDelegates::GetPakSigningKe
 #if WITH_EDITOR
 	FSimpleMulticastDelegate FCoreDelegates::PreModal;
 	FSimpleMulticastDelegate FCoreDelegates::PostModal;
+	FCoreDelegates::FPrePostSlateModalWithContextDelegate FCoreDelegates::PreSlateModalWithContext;
+	FCoreDelegates::FPrePostSlateModalWithContextDelegate FCoreDelegates::PostSlateModalWithContext;
+PRAGMA_DISABLE_DEPRECATION_WARNINGS
     FSimpleMulticastDelegate FCoreDelegates::PreSlateModal;
     FSimpleMulticastDelegate FCoreDelegates::PostSlateModal;
+PRAGMA_ENABLE_DEPRECATION_WARNINGS
 #endif	//WITH_EDITOR
 #if ALLOW_OTHER_PLATFORM_CONFIG
 	TDelegate<FCoreDelegates::FCVarKeyValueMap(const FString&)> FCoreDelegates::GatherDeviceProfileCVars;
diff --git a/Engine/Source/Runtime/Core/Public/Misc/CoreDelegates.h b/Engine/Source/Runtime/Core/Public/Misc/CoreDelegates.h
--- a/Engine/Source/Runtime/Core/Public/Misc/CoreDelegates.h
+++ b/Engine/Source/Runtime/Core/Public/Misc/CoreDelegates.h
@@ -188,17 +188,35 @@ public:
 	
 
 #if WITH_EDITOR
+	struct FModalWindowContext
+	{
+		// Only available in PreSlateModalWithContext
+		TOptional<bool> bIsSlowTaskWindow;
+
+		// Guaranteed to match between PreSlateModalWithContext and PostSlateModalWithContext
+		typedef std::intptr_t FWindowHandle;
+		FWindowHandle WindowIdentifier = 0;
+	};
+
+	DECLARE_MULTICAST_DELEGATE_OneParam(FPrePostSlateModalWithContextDelegate, const FModalWindowContext&);
+
 	// Called before the editor displays a modal window, allowing other windows the opportunity to disable themselves to avoid reentrant calls
 	static CORE_API FSimpleMulticastDelegate PreModal;
 
 	// Called after the editor dismisses a modal window, allowing other windows the opportunity to disable themselves to avoid reentrant calls
 	static CORE_API FSimpleMulticastDelegate PostModal;
     
-    // Called before the editor displays a Slate (non-platform) modal window, allowing other windows the opportunity to disable themselves to avoid reentrant calls
+	UE_DEPRECATED(5.8, "PreSlateModal has been deprecated. Please use PreSlateModalWithContext instead.")
     static CORE_API FSimpleMulticastDelegate PreSlateModal;
+    // Called before the editor displays a Slate (non-platform) modal window, allowing other windows the opportunity to disable themselves to avoid reentrant calls
+	static CORE_API FPrePostSlateModalWithContextDelegate PreSlateModalWithContext;
     
-    // Called after the editor dismisses a Slate (non-platform) modal window, allowing other windows the opportunity to disable themselves to avoid reentrant calls
+    // Called before the editor displays a Slate (non-platform) modal window, allowing other windows the opportunity to disable themselves to avoid reentrant calls
+	UE_DEPRECATED(5.8, "PostSlateModal has been deprecated. Please use PostSlateModalWithContext instead.")
     static CORE_API FSimpleMulticastDelegate PostSlateModal;
+
+    // Called after the editor dismisses a Slate (non-platform) modal window, allowing other windows the opportunity to disable themselves to avoid reentrant calls
+	static CORE_API FPrePostSlateModalWithContextDelegate PostSlateModalWithContext;
     
 #endif	//WITH_EDITOR
 
diff --git a/Engine/Source/Runtime/Slate/Private/Framework/Application/SlateApplication.cpp b/Engine/Source/Runtime/Slate/Private/Framework/Application/SlateApplication.cpp
--- a/Engine/Source/Runtime/Slate/Private/Framework/Application/SlateApplication.cpp
+++ b/Engine/Source/Runtime/Slate/Private/Framework/Application/SlateApplication.cpp
@@ -2109,7 +2109,13 @@ void FSlateApplication::AddModalWindow( TSharedRef<SWindow> InSlateWindow, const
 	}
 
 #if WITH_EDITOR
+	FCoreDelegates::PreSlateModalWithContext.Broadcast(
+		FCoreDelegates::FModalWindowContext{ 
+			.bIsSlowTaskWindow = bSlowTaskWindow,
+			.WindowIdentifier = (std::intptr_t)(&InSlateWindow.Get()) });
+PRAGMA_DISABLE_DEPRECATION_WARNINGS
     FCoreDelegates::PreSlateModal.Broadcast();
+PRAGMA_ENABLE_DEPRECATION_WARNINGS
 #endif
     // Push the active modal window onto the stack.
 	ActiveModalWindows.AddUnique( InSlateWindow );
diff --git a/Engine/Source/Runtime/SlateCore/Private/Widgets/SWindow.cpp b/Engine/Source/Runtime/SlateCore/Private/Widgets/SWindow.cpp
--- a/Engine/Source/Runtime/SlateCore/Private/Widgets/SWindow.cpp
+++ b/Engine/Source/Runtime/SlateCore/Private/Widgets/SWindow.cpp
@@ -1407,7 +1407,10 @@ void SWindow::NotifyWindowBeingDestroyed()
 #if WITH_EDITOR
 	if(bIsModalWindow)
 	{
+PRAGMA_DISABLE_DEPRECATION_WARNINGS
 		FCoreDelegates::PostSlateModal.Broadcast();
+PRAGMA_ENABLE_DEPRECATION_WARNINGS
+		FCoreDelegates::PostSlateModalWithContext.Broadcast({ .WindowIdentifier = (intptr_t)this });
 	}
 #endif
 
