Mark symbol lookup timing data as "very verbose" to reduce log-spam.

Previously, symbol lookup times would appear on Linux ahead of every stack crawl, one line per unique frame. This data is not generally useful unless you are actively working on stack crawl performance.

diff --git a/Engine/Source/Runtime/Core/Private/Unix/UnixPlatformStackWalk.cpp b/Engine/Source/Runtime/Core/Private/Unix/UnixPlatformStackWalk.cpp
--- a/Engine/Source/Runtime/Core/Private/Unix/UnixPlatformStackWalk.cpp
+++ b/Engine/Source/Runtime/Core/Private/Unix/UnixPlatformStackWalk.cpp
@@ -79,7 +79,7 @@ void CORE_API UnixPlatformStackWalk_PreloadModuleSymbolFile()
 
 			GModuleSymbolFileMemory += FPlatformMemory::GetConstants().PageSize;
 			UE_LOG(LogHAL, Warning, TEXT("UnixPlatformStackWalk_PreloadModuleSymbolFile: GModuleSymbolFileMemory = 0x%x (After adding additional memory page"), GModuleSymbolFileMemory);
-	
+
 			// On linux read() will transfer at most 2,147,479,552 bytes
 			const int32 MaxBytesToRead = 0x7FFFF000;
 
@@ -324,7 +324,7 @@ namespace
 			{
 				if (CheckingEnsureTime)
 				{
-					UE_LOG(LogCore, Log, TEXT("0x%016llx Dladdr: %lfms"), reinterpret_cast<void*>(ProgramCounter), (DladdrEndTime - StartTime) * 1000);
+					UE_LOG(LogCore, VeryVerbose, TEXT("Symbol lookup time for 0x%016llx Dladdr: %lfms"), reinterpret_cast<void*>(ProgramCounter), (DladdrEndTime - StartTime) * 1000);
 				}
 
 				// If we cannot find the module name or the module base return early
@@ -519,7 +519,7 @@ namespace
 
 		if (CheckingEnsureTime)
 		{
-			UE_LOG(LogCore, Log, TEXT("0x%016llx Dladdr: %lfms Open: %lfms Search: %lfms"),
+			UE_LOG(LogCore, VeryVerbose, TEXT("Symbol lookup time for 0x%016llx Dladdr: %lfms Open: %lfms Search: %lfms"),
 					ProgramCounter,
 					(DladdrEndTime - StartTime) * 1000.0,
 					(RecordReaderEndTime - DladdrEndTime) * 1000.0,
@@ -541,14 +541,14 @@ bool FUnixPlatformStackWalk::ProgramCounterToHumanReadableString( int32 CurrentC
 	// Callstack lines should be written in this standard format
 	//
 	//	0xaddress module!func [file]
-	// 
+	//
 	// E.g. 0x045C8D01 OrionClient.self!UEngine::PerformError() [D:\Epic\Orion\Engine\Source\Runtime\Engine\Private\UnrealEngine.cpp:6481]
 	//
 	// Module may be omitted, everything else should be present, or substituted with a string that conforms to the expected type
 	//
 	// E.g 0x00000000 UnknownFunction []
 	//
-	// 
+	//
 
 	if (HumanReadableString && HumanReadableStringSize > 0)
 	{
