Add Decibels to UE's list of units

diff --git a/Engine/Source/Runtime/Core/Private/Math/UnitConversion.cpp b/Engine/Source/Runtime/Core/Private/Math/UnitConversion.cpp
--- a/Engine/Source/Runtime/Core/Private/Math/UnitConversion.cpp
+++ b/Engine/Source/Runtime/Core/Private/Math/UnitConversion.cpp
@@ -84,6 +84,8 @@ constexpr FParseCandidate ParseCandidates[] = {
 	{ TEXT("KilogramCentimeters"),	EUnit::KilogramCentimeters },	{ TEXT("kgcm"),		EUnit::KilogramCentimeters },
 	{ TEXT("KilogramMeters"),		EUnit::KilogramMeters },		{ TEXT("kgm"),		EUnit::KilogramMeters },
 
+	{TEXT("Decibels"),			EUnit::Decibels},					{TEXT("dB"),		EUnit::Decibels},
+	
 	{ TEXT("Hertz"),				EUnit::Hertz },					{ TEXT("Hz"),		EUnit::Hertz },
 	{ TEXT("Kilohertz"),			EUnit::Kilohertz },				{ TEXT("KHz"),		EUnit::Kilohertz },
 	{ TEXT("Megahertz"),			EUnit::Megahertz },				{ TEXT("MHz"),		EUnit::Megahertz },
@@ -152,6 +154,8 @@ constexpr const TCHAR* const DisplayStrings[] = {
 
 	TEXT("Ns"), TEXT("kgcm"), TEXT("kgm"),
 
+	TEXT("dB"),
+	
 	TEXT("Hz"), TEXT("KHz"), TEXT("MHz"), TEXT("GHz"), TEXT("rpm"),
 
 	TEXT("B"), TEXT("KiB"), TEXT("MiB"), TEXT("GiB"), TEXT("TiB"),
@@ -224,6 +228,8 @@ constexpr const TCHAR* const SupportedUnitsStrings[] = {
 	TEXT("NewtonSeconds"),
 	TEXT("KilogramCentimeters"),
 	TEXT("KilogramMeters"),
+	
+	TEXT("Decibels"),
 
 	TEXT("Hertz"),				
 	TEXT("Kilohertz"),			
@@ -295,6 +301,8 @@ constexpr EUnitType UnitTypes[] = {
 
 	EUnitType::Impulse,		EUnitType::PositionalImpulse, EUnitType::PositionalImpulse,
 
+	EUnitType::Loudness,
+	
 	EUnitType::Frequency,	EUnitType::Frequency,	EUnitType::Frequency,	EUnitType::Frequency,	EUnitType::Frequency,
 
 	EUnitType::DataSize,	EUnitType::DataSize,	EUnitType::DataSize,	EUnitType::DataSize,	EUnitType::DataSize,
diff --git a/Engine/Source/Runtime/Core/Public/Math/UnitConversion.h b/Engine/Source/Runtime/Core/Public/Math/UnitConversion.h
--- a/Engine/Source/Runtime/Core/Public/Math/UnitConversion.h
+++ b/Engine/Source/Runtime/Core/Public/Math/UnitConversion.h
@@ -53,9 +53,12 @@ enum class EUnit : uint8
 	/** Impulse unit. */
 	NewtonSeconds, KilogramCentimeters, KilogramMeters,
 
+	/** Loudness unit. */
+	Decibels,
+	
 	/** Frequency unit. */
 	Hertz, Kilohertz, Megahertz, Gigahertz, RevolutionsPerMinute,
-
+	
 	/** Data Size unit. */
 	Bytes, Kilobytes, Megabytes, Gigabytes, Terabytes,
 
@@ -96,7 +99,7 @@ enum class EUnit : uint8
 /** Enumeration that specifies particular classes of unit */
 enum class EUnitType
 {
-	Distance, Angle, Speed, AngularSpeed, Acceleration, Temperature, Mass, Density, Force, Torque, Impulse, PositionalImpulse, Frequency, DataSize, LuminousFlux, LuminousIntensity, Illuminance, Luminance, Time, PixelDensity, Multipliers, ExposureValue, Stress,
+	Distance, Angle, Speed, AngularSpeed, Acceleration, Temperature, Mass, Density, Force, Torque, Impulse, PositionalImpulse, Loudness, Frequency, DataSize, LuminousFlux, LuminousIntensity, Illuminance, Luminance, Time, PixelDensity, Multipliers, ExposureValue, Stress,
 
 	// Symbolic entry - do not use directly
 	NumberOf,
diff --git a/Engine/Source/Runtime/Core/Public/Math/UnitConversion.inl b/Engine/Source/Runtime/Core/Public/Math/UnitConversion.inl
--- a/Engine/Source/Runtime/Core/Public/Math/UnitConversion.inl
+++ b/Engine/Source/Runtime/Core/Public/Math/UnitConversion.inl
@@ -134,6 +134,7 @@ T FUnitConversion::Convert(T InValue, EUnit From, EUnit To)
 			}
 		}
 
+		case EUnitType::Loudness:
 		case EUnitType::Impulse:
 		case EUnitType::LuminousIntensity:
 		case EUnitType::Illuminance:
