Fix delegate crashes in the cvar callback system when using merged modules

#rnx

diff --git a/Engine/Source/Runtime/Core/Private/HAL/ConsoleManager.cpp b/Engine/Source/Runtime/Core/Private/HAL/ConsoleManager.cpp
--- a/Engine/Source/Runtime/Core/Private/HAL/ConsoleManager.cpp
+++ b/Engine/Source/Runtime/Core/Private/HAL/ConsoleManager.cpp
@@ -250,10 +250,16 @@ public:
 	}
 	
 	/** Legacy function to add old single delegates to the new multicast delegate. */
-	virtual void SetOnChangedCallback(const FConsoleVariableDelegate& Callback) 
+	virtual void SetOnChangedCallback(const FConsoleVariableDelegate& Callback) override
+	{
+		OnChangedCallback.Remove(LegacyDelegateHandle);
+		LegacyDelegateHandle = OnChangedCallback.Add(Callback);
+	}
+	
+	/** Remove a single delegate to the new multicast delegate. */
+	virtual void ClearOnChangedCallback() override
 	{
 		OnChangedCallback.Remove(LegacyDelegateHandle);
-		LegacyDelegateHandle = OnChangedCallback.Add(Callback); 
 	}
 
 	/** Returns a multicast delegate with which to register. Called when this CVar changes. */
@@ -4989,6 +4995,14 @@ public:
 		}
 	}
 
+	virtual void ClearOnChangedCallback() override
+	{
+		if (Bind())
+		{
+			RealVariable->ClearOnChangedCallback();
+		}
+	}
+
 	virtual FConsoleVariableMulticastDelegate& OnChangedDelegate() override
 	{
 		if (Bind())
diff --git a/Engine/Source/Runtime/Core/Public/HAL/IConsoleManager.h b/Engine/Source/Runtime/Core/Public/HAL/IConsoleManager.h
--- a/Engine/Source/Runtime/Core/Public/HAL/IConsoleManager.h
+++ b/Engine/Source/Runtime/Core/Public/HAL/IConsoleManager.h
@@ -666,6 +666,11 @@ public:
 	 **/
 	virtual void SetOnChangedCallback(const FConsoleVariableDelegate& Callback) = 0;
 
+	/**
+	 * Remove the callback function set to this console variable.
+	 */
+	virtual void ClearOnChangedCallback() = 0;
+
 	virtual FConsoleVariableMulticastDelegate& OnChangedDelegate() = 0;
 
 	/**
@@ -1654,7 +1659,9 @@ public:
 
 	virtual ~FAutoConsoleVariableRef()
 	{
+		AsVariable()->ClearOnChangedCallback();
 	}
+
 	/** Dereference back to a variable**/
 	UE_FORCEINLINE_HINT IConsoleVariable& operator*()
 	{
@@ -1841,6 +1848,11 @@ public:
 		AsVariable()->SetOnChangedCallback(Callback);
 	}
 
+	virtual ~TAutoConsoleVariable()
+	{
+		AsVariable()->ClearOnChangedCallback();
+	}
+
 	T GetValueOnGameThread() const
 	{
 		return Ref->GetValueOnGameThread();
