Added UEOpEquals to TOptional.

diff --git a/Engine/Source/Runtime/Core/Public/Misc/Optional.h b/Engine/Source/Runtime/Core/Public/Misc/Optional.h
--- a/Engine/Source/Runtime/Core/Public/Misc/Optional.h
+++ b/Engine/Source/Runtime/Core/Public/Misc/Optional.h
@@ -6,6 +6,7 @@
 #include "Misc/AssertionMacros.h"
 #include "Misc/IntrusiveUnsetOptionalState.h"
 #include "Misc/OptionalFwd.h"
+#include "Misc/UEOps.h"
 #include "Templates/MemoryOps.h"
 #include "Templates/TypeHash.h" // Definitions of GetTypeHash for primitive types must be visible when GetTypeHash for TOptional is defined
 #include "Templates/UnrealTemplate.h"
@@ -372,9 +373,9 @@ public:
 		return *Result;
 	}
 
-	[[nodiscard]] friend constexpr bool operator==(const TOptional& Lhs, const TOptional& Rhs)
+	[[nodiscard]] constexpr bool UEOpEquals(const TOptional& Rhs) const
 	{
-		bool bIsLhsSet = Lhs.IsSet();
+		bool bIsLhsSet = this->IsSet();
 		bool bIsRhsSet = Rhs.IsSet();
 
 		if (bIsLhsSet != bIsRhsSet)
@@ -386,12 +387,7 @@ public:
 			return true;
 		}
 
-		return Lhs.TypedValue == Rhs.TypedValue;
-	}
-
-	[[nodiscard]] friend constexpr bool operator!=(const TOptional& Lhs, const TOptional& Rhs)
-	{
-		return !(Lhs == Rhs);
+		return this->TypedValue == Rhs.TypedValue;
 	}
 
 	void Serialize(FArchive& Ar)
