[Backout] - CL46630849
#rnx
Original CL Desc
-----------------------------------------------------------------
Added UEOpEquals to TOptional.

diff --git a/Engine/Source/Runtime/Core/Public/Misc/Optional.h b/Engine/Source/Runtime/Core/Public/Misc/Optional.h
--- a/Engine/Source/Runtime/Core/Public/Misc/Optional.h
+++ b/Engine/Source/Runtime/Core/Public/Misc/Optional.h
@@ -6,7 +6,6 @@
 #include "Misc/AssertionMacros.h"
 #include "Misc/IntrusiveUnsetOptionalState.h"
 #include "Misc/OptionalFwd.h"
-#include "Misc/UEOps.h"
 #include "Templates/MemoryOps.h"
 #include "Templates/TypeHash.h" // Definitions of GetTypeHash for primitive types must be visible when GetTypeHash for TOptional is defined
 #include "Templates/UnrealTemplate.h"
@@ -373,9 +372,9 @@ public:
 		return *Result;
 	}
 
-	[[nodiscard]] constexpr bool UEOpEquals(const TOptional& Rhs) const
+	[[nodiscard]] friend constexpr bool operator==(const TOptional& Lhs, const TOptional& Rhs)
 	{
-		bool bIsLhsSet = this->IsSet();
+		bool bIsLhsSet = Lhs.IsSet();
 		bool bIsRhsSet = Rhs.IsSet();
 
 		if (bIsLhsSet != bIsRhsSet)
@@ -387,7 +386,12 @@ public:
 			return true;
 		}
 
-		return this->TypedValue == Rhs.TypedValue;
+		return Lhs.TypedValue == Rhs.TypedValue;
+	}
+
+	[[nodiscard]] friend constexpr bool operator!=(const TOptional& Lhs, const TOptional& Rhs)
+	{
+		return !(Lhs == Rhs);
 	}
 
 	void Serialize(FArchive& Ar)
