[Core]
* Added missing #pragma once to ObjectThumbnail.inl

diff --git a/Engine/Source/Runtime/Core/Public/Misc/ObjectThumbnail.inl b/Engine/Source/Runtime/Core/Public/Misc/ObjectThumbnail.inl
--- a/Engine/Source/Runtime/Core/Public/Misc/ObjectThumbnail.inl
+++ b/Engine/Source/Runtime/Core/Public/Misc/ObjectThumbnail.inl
@@ -2,6 +2,8 @@
 
 #if defined(IMAGECORE_API)
 
+#pragma once
+
 #include "Misc/ObjectThumbnail.h"
 #include "ImageCore.h"
 
