Added new concepts for reference types.

diff --git a/Engine/Source/Runtime/Core/Public/Concepts/LValueReference.h b/Engine/Source/Runtime/Core/Public/Concepts/LValueReference.h
new file mode 100644
--- /dev/null
+++ b/Engine/Source/Runtime/Core/Public/Concepts/LValueReference.h
@@ -0,0 +1,14 @@
+// Copyright Epic Games, Inc. All Rights Reserved.
+
+#pragma once
+
+#include <type_traits>
+
+namespace UE
+{
+	/**
+	 * Concept which describes an lvalue reference type.
+	 */
+	template <typename T>
+	concept CLValueReference = std::is_lvalue_reference_v<T>;
+}
diff --git a/Engine/Source/Runtime/Core/Public/Concepts/RValueReference.h b/Engine/Source/Runtime/Core/Public/Concepts/RValueReference.h
new file mode 100644
--- /dev/null
+++ b/Engine/Source/Runtime/Core/Public/Concepts/RValueReference.h
@@ -0,0 +1,14 @@
+// Copyright Epic Games, Inc. All Rights Reserved.
+
+#pragma once
+
+#include <type_traits>
+
+namespace UE
+{
+	/**
+	 * Concept which describes an rvalue reference type.
+	 */
+	template <typename T>
+	concept CRValueReference = std::is_rvalue_reference_v<T>;
+}
diff --git a/Engine/Source/Runtime/Core/Public/Concepts/Reference.h b/Engine/Source/Runtime/Core/Public/Concepts/Reference.h
new file mode 100644
--- /dev/null
+++ b/Engine/Source/Runtime/Core/Public/Concepts/Reference.h
@@ -0,0 +1,14 @@
+// Copyright Epic Games, Inc. All Rights Reserved.
+
+#pragma once
+
+#include <type_traits>
+
+namespace UE
+{
+	/**
+	 * Concept which describes a reference type - either lvalue or rvalue.
+	 */
+	template <typename T>
+	concept CReference = std::is_reference_v<T>;
+}
