Added new concept for aggregate types.

diff --git a/Engine/Source/Runtime/Core/Public/Concepts/Aggregate.h b/Engine/Source/Runtime/Core/Public/Concepts/Aggregate.h
new file mode 100644
--- /dev/null
+++ b/Engine/Source/Runtime/Core/Public/Concepts/Aggregate.h
@@ -0,0 +1,14 @@
+// Copyright Epic Games, Inc. All Rights Reserved.
+
+#pragma once
+
+#include <type_traits>
+
+namespace UE
+{
+	/**
+	 * Concept which describes an aggregate type.
+	 */
+	template <typename T>
+	concept CAggregate = std::is_aggregate_v<T>;
+}
