Made FDelegateUserObject constructible from smart pointers when UE_WITH_REMOTE_OBJECT_HANDLE == 1.

diff --git a/Engine/Source/Runtime/Core/Public/Concepts/Dereferenceable.h b/Engine/Source/Runtime/Core/Public/Concepts/Dereferenceable.h
new file mode 100644
--- /dev/null
+++ b/Engine/Source/Runtime/Core/Public/Concepts/Dereferenceable.h
@@ -0,0 +1,18 @@
+// Copyright Epic Games, Inc. All Rights Reserved.
+
+#pragma once
+
+#include "Concepts/LValueReference.h"
+
+namespace UE
+{
+	/**
+	 * Concept which describes a type which is dereferenceable to obtain an lvalue reference to another object.
+	 */
+	template <typename T>
+	concept CDereferenceable =
+		requires(const T& Val)
+		{
+			{ *Val } -> UE::CLValueReference;
+		};
+}
diff --git a/Engine/Source/Runtime/Core/Public/Delegates/IDelegateInstance.h b/Engine/Source/Runtime/Core/Public/Delegates/IDelegateInstance.h
--- a/Engine/Source/Runtime/Core/Public/Delegates/IDelegateInstance.h
+++ b/Engine/Source/Runtime/Core/Public/Delegates/IDelegateInstance.h
@@ -3,6 +3,8 @@
 #pragma once
 
 #include "CoreTypes.h"
+#include "Concepts/Dereferenceable.h"
+#include "Concepts/Pointer.h"
 #include "Templates/TypeHash.h" 
 #include "UObject/NameTypes.h"
 #include "Delegates/DelegateSettings.h"
@@ -85,6 +87,12 @@ public:
 		UserObject = (const void*)InUserObject;
 		bIsUObject = std::is_convertible_v<UserObjectType*, UObject*>;
 	}
+	template<UE::CDereferenceable PtrWrapperType>
+		requires !UE::CPointer<PtrWrapperType>
+	UE_REWRITE FDelegateUserObject(const PtrWrapperType& InPtrWrapper)
+		: FDelegateUserObject(&*InPtrWrapper)
+	{
+	}
 	UE_FORCEINLINE_HINT bool IsUObject() const
 	{
 		return bIsUObject;
