Fixed potential null dereference in FDelegateUserObject when UE_WITH_REMOTE_OBJECT_HANDLE == 1.

diff --git a/Engine/Source/Runtime/Core/Public/Concepts/ConditionallyConvertibleToBool.h b/Engine/Source/Runtime/Core/Public/Concepts/ConditionallyConvertibleToBool.h
new file mode 100644
--- /dev/null
+++ b/Engine/Source/Runtime/Core/Public/Concepts/ConditionallyConvertibleToBool.h
@@ -0,0 +1,18 @@
+// Copyright Epic Games, Inc. All Rights Reserved.
+
+#pragma once
+
+#include "Concepts/SameAs.h"
+
+namespace UE
+{
+	/**
+	 * Concept which describes a type which is conditionally convertible to bool.
+	 */
+	template <typename T>
+	concept CConditionallyConvertibleToBool =
+		requires (const T& Val, bool bOther)
+		{
+			{ Val && bOther } -> UE::CSameAs<bool>;
+		};
+}
diff --git a/Engine/Source/Runtime/Core/Public/Concepts/PointerLike.h b/Engine/Source/Runtime/Core/Public/Concepts/PointerLike.h
new file mode 100644
--- /dev/null
+++ b/Engine/Source/Runtime/Core/Public/Concepts/PointerLike.h
@@ -0,0 +1,15 @@
+// Copyright Epic Games, Inc. All Rights Reserved.
+
+#pragma once
+
+#include "Concepts/ConditionallyConvertibleToBool.h"
+#include "Concepts/Dereferenceable.h"
+
+namespace UE
+{
+	/**
+	 * Concept which describes a type which acts like a pointer - notably is dereferenceable and is bool-testable for null.
+	 */
+	template <typename T>
+	concept CPointerLike = UE::CConditionallyConvertibleToBool<T> && UE::CDereferenceable<T>;
+}
diff --git a/Engine/Source/Runtime/Core/Public/Delegates/IDelegateInstance.h b/Engine/Source/Runtime/Core/Public/Delegates/IDelegateInstance.h
--- a/Engine/Source/Runtime/Core/Public/Delegates/IDelegateInstance.h
+++ b/Engine/Source/Runtime/Core/Public/Delegates/IDelegateInstance.h
@@ -3,8 +3,8 @@
 #pragma once
 
 #include "CoreTypes.h"
-#include "Concepts/Dereferenceable.h"
 #include "Concepts/Pointer.h"
+#include "Concepts/PointerLike.h"
 #include "Templates/TypeHash.h" 
 #include "UObject/NameTypes.h"
 #include "Delegates/DelegateSettings.h"
@@ -87,10 +87,10 @@ public:
 		UserObject = (const void*)InUserObject;
 		bIsUObject = std::is_convertible_v<UserObjectType*, UObject*>;
 	}
-	template<UE::CDereferenceable PtrWrapperType>
+	template<UE::CPointerLike PtrWrapperType>
 		requires !UE::CPointer<PtrWrapperType>
 	UE_REWRITE FDelegateUserObject(const PtrWrapperType& InPtrWrapper)
-		: FDelegateUserObject(&*InPtrWrapper)
+		: FDelegateUserObject(InPtrWrapper ? &*InPtrWrapper : nullptr)
 	{
 	}
 	UE_FORCEINLINE_HINT bool IsUObject() const
