Made UE::CConditionallyConvertibleToBool less permissive by avoiding potential operator overloads.

diff --git a/Engine/Source/Runtime/Core/Public/Concepts/ConditionallyConvertibleToBool.h b/Engine/Source/Runtime/Core/Public/Concepts/ConditionallyConvertibleToBool.h
--- a/Engine/Source/Runtime/Core/Public/Concepts/ConditionallyConvertibleToBool.h
+++ b/Engine/Source/Runtime/Core/Public/Concepts/ConditionallyConvertibleToBool.h
@@ -2,8 +2,6 @@
 
 #pragma once
 
-#include "Concepts/SameAs.h"
-
 namespace UE
 {
 	/**
@@ -11,8 +9,8 @@ namespace UE
 	 */
 	template <typename T>
 	concept CConditionallyConvertibleToBool =
-		requires (const T& Val, bool bOther)
+		requires (const T& Val)
 		{
-			{ Val && bOther } -> UE::CSameAs<bool>;
+			Val ? 1 : 0;
 		};
 }
