Replaced some trivial forwarding functions' use of UE_FORCEINLINE_HINT with UE_REWRITE.

#rnx

diff --git a/Engine/Source/Runtime/Core/Public/Delegates/IDelegateInstance.h b/Engine/Source/Runtime/Core/Public/Delegates/IDelegateInstance.h
--- a/Engine/Source/Runtime/Core/Public/Delegates/IDelegateInstance.h
+++ b/Engine/Source/Runtime/Core/Public/Delegates/IDelegateInstance.h
@@ -54,7 +54,7 @@ public:
 		return ID != Rhs.ID;
 	}
 
-	friend UE_FORCEINLINE_HINT uint32 GetTypeHash(const FDelegateHandle& Key)
+	friend UE_REWRITE uint32 GetTypeHash(const FDelegateHandle& Key)
 	{
 		return GetTypeHash(Key.ID);
 	}
@@ -93,19 +93,19 @@ public:
 		: FDelegateUserObject(InPtrWrapper ? &*InPtrWrapper : nullptr)
 	{
 	}
-	UE_FORCEINLINE_HINT bool IsUObject() const
+	UE_REWRITE bool IsUObject() const
 	{
 		return bIsUObject;
 	}
-	UE_FORCEINLINE_HINT const void* GetRaw() const
+	UE_REWRITE const void* GetRaw() const
 	{
 		return UserObject;
 	}
-	UE_FORCEINLINE_HINT const UObject* GetUObject() const
+	UE_REWRITE const UObject* GetUObject() const
 	{
 		return (const UObject*)UserObject;
 	}
-	UE_FORCEINLINE_HINT operator const void* () const
+	UE_REWRITE operator const void* () const
 	{
 		return UserObject;
 	}
