Fixed Clang compile error.

diff --git a/Engine/Source/Runtime/Core/Public/Delegates/IDelegateInstance.h b/Engine/Source/Runtime/Core/Public/Delegates/IDelegateInstance.h
--- a/Engine/Source/Runtime/Core/Public/Delegates/IDelegateInstance.h
+++ b/Engine/Source/Runtime/Core/Public/Delegates/IDelegateInstance.h
@@ -54,7 +54,7 @@ public:
 		return ID != Rhs.ID;
 	}
 
-	friend UE_REWRITE uint32 GetTypeHash(const FDelegateHandle& Key)
+	UE_REWRITE friend uint32 GetTypeHash(const FDelegateHandle& Key)
 	{
 		return GetTypeHash(Key.ID);
 	}
