- Allow plugins to have PluginName.ini along with Base/DefaultPluginName.ini (but deprecated)
- Also deprecated BasePluginName.ini so we can eventually look for fewer ini files during plugin mounting
- Show a message with GFP old style names as well (same as above, but it didn't go through the ConfigHierarchy checks for deprecation)

diff --git a/Engine/Plugins/Runtime/GameFeatures/Source/GameFeatures/Private/GameFeatureData.cpp b/Engine/Plugins/Runtime/GameFeatures/Source/GameFeatures/Private/GameFeatureData.cpp
--- a/Engine/Plugins/Runtime/GameFeatures/Source/GameFeatures/Private/GameFeatureData.cpp
+++ b/Engine/Plugins/Runtime/GameFeatures/Source/GameFeatures/Private/GameFeatureData.cpp
@@ -5,6 +5,7 @@
 #include "AssetRegistry/AssetData.h"
 #include "Engine/AssetManager.h"
 #include "GameFeaturesSubsystem.h"
+#include "HAL/FileManager.h"
 #include "InstallBundleUtils.h"
 #include "Misc/ConfigCacheIni.h"
 #include "Misc/ConfigContext.h"
@@ -15,6 +16,7 @@
 #include "DeviceProfiles/DeviceProfileFragment.h"
 #include "DeviceProfiles/DeviceProfileManager.h"
 #include "Interfaces/IPluginManager.h"
+#include "Logging/MessageLog.h"
 
 #if WITH_EDITOR
 #include "Settings/EditorExperimentalSettings.h"
@@ -130,6 +132,41 @@ namespace GameFeatureData
 #endif // #if WITH_EDITOR
 }
 
+// either show an editor warning, or just write to log for non-editor
+// NOTE: For tirst checkin, it will just log a warning, then once we have settled, we can checkin
+// more in-your-face behavior
+static void ProcessDeprecation(const FString& File)
+{
+	if (!IsInGameThread())
+	{
+		AsyncTask(ENamedThreads::GameThread, [File]() { ProcessDeprecation(File); });
+		return;
+	}
+
+	FFormatNamedArguments Arguments;
+	Arguments.Add(TEXT("File"), FText::FromString(File));
+	FText Msg = FText::Format(LOCTEXT("DeprecatedGFPConfig", "Found a deprecated, old-style, GFP config file '{File}'. Rename if to the new Default<PluginName>.ini style, and MAKE SURE to add + tokens before array properties (like other ini files would have). In the future, this file will not be loaded."), Arguments);
+
+	//if (GIsEditor)
+	//{
+	//	static TSet<FString> AlreadyWarnedKeys;
+
+	//	FString AbsPath = IFileManager::Get().ConvertToAbsolutePathForExternalAppForRead(*File);
+
+	//	FMessageLog EditorErrors("EditorErrors");
+	//	TSharedRef<FTokenizedMessage> Message = EditorErrors.Message(EMessageSeverity::Warning);
+	//	if (File.EndsWith(TEXT(".ini")))
+	//	{
+	//		Message->AddToken(FURLToken::Create(FString::Printf(TEXT("file://%s"), *AbsPath), LOCTEXT("DeprecatedGFPConfig_URLCLick", "Click to open file")));
+	//	}
+	//	Message->AddToken(FTextToken::Create(Msg));
+	//	EditorErrors.Notify();
+	//}
+
+	UE_LOG(LogConfig, Warning, TEXT("%s"), *Msg.ToString());
+	//	ensureAlwaysMsgf(false, TEXT("%s"), *Msg.ToString());
+}
+
 void UGameFeatureData::InitializeBasePluginIniFile(const FString& PluginInstalledFilename)
 {
 	const FString PluginName = FPaths::GetBaseFilename(PluginInstalledFilename);
@@ -137,8 +174,9 @@ void UGameFeatureData::InitializeBasePluginIniFile(const FString& PluginInstalle
 	static bool bUseNewDynamicLayers = IConsoleManager::Get().FindConsoleVariable(TEXT("ini.UseNewDynamicLayers"))->GetInt() != 0;
 	bool bIncludePluginNameInBranchName = true;
 
-	// DEPRECATED NAMING PATH - must keep because these files are read in as a single file, not in a hierarchical way, so 
-	// they don't have the + syntax for arrays
+	// UE_DEPRECATED(5.8, "Handle old-style named GFP .ini files - keeping for a bit to tell licensees, as well as any straggling merges of streams into Main")
+	// (Using the UE_DEPRECATED tag above just so we can find this later to delete the whole block)
+	// This old path will be removed. See ProcessDeprecation() has the error messaging
 	{
 		const FString PluginConfigDir = FPaths::GetPath(PluginInstalledFilename) / TEXT("Config/");
 		const FString EngineConfigDir = FPaths::EngineConfigDir();
@@ -154,8 +192,7 @@ void UGameFeatureData::InitializeBasePluginIniFile(const FString& PluginInstalle
 		FConfigFile& PluginConfig = GConfig->Add(PluginConfigFilename, FConfigFile());
 		if (FConfigCacheIni::LoadExternalIniFile(PluginConfig, *PluginName, *EngineConfigDir, *PluginConfigDir, bIsBaseIniName, nullptr, bForceReloadFromDisk, bWriteDestIni))
 		{
-			// This is the deprecated loading path that doesn't handle cases like + in arrays
-			UE_LOG(LogGameFeatures, Log, TEXT("Loaded deprecated config %s, rename to start with Default for normal parsing"), *PluginConfigFilename);
+			ProcessDeprecation(PluginConfigFilename);
 
 			// register this plugin, so the ConfigContext.Load, and future loads, know about it
 			if (bUseNewDynamicLayers)
diff --git a/Engine/Source/Runtime/Core/Private/Misc/ConfigContext.cpp b/Engine/Source/Runtime/Core/Private/Misc/ConfigContext.cpp
--- a/Engine/Source/Runtime/Core/Private/Misc/ConfigContext.cpp
+++ b/Engine/Source/Runtime/Core/Private/Misc/ConfigContext.cpp
@@ -13,11 +13,17 @@
 #include "Misc/ScopeLock.h"
 #include "ProfilingDebugging/CpuProfilerTrace.h"
 #include "ProfilingDebugging/AssetMetadataTrace.h"
+#include "HAL/FileManager.h"
 #include "HAL/LowLevelMemStats.h"
+#include "Logging/MessageLog.h"
+#include "Async/Async.h"
+#include "Async/TaskGraphInterfaces.h"
 
 extern double GPrepareForLoadTime;
 extern double GPerformLoadTime;
 
+#define LOCTEXT_NAMESPACE "ConfigSystem"
+
 namespace
 {
 	FString VersionName = "Version";
@@ -32,11 +38,30 @@ namespace
 	constexpr int32 KeyFlag_UseGlobalCache = 1;
 	constexpr int32 KeyFlag_UsePluginCache = 2;
 	constexpr int32 KeyFlag_AssumeExists = 4;
+	constexpr int32 KeyFlag_Deprecated_Plugin = 8;
+	constexpr int32 KeyFlag_Deprecated_Other = 16;
+	// if you add a KeyFlag, increase NumFlagsBits below!
+
+	constexpr int32 NumStaticLayerBits = 6; // needs to hold the length of GConfigLayers (or any custom layers, which is highly unlikely to ever be longer than GConfigLayers)
+	constexpr int32 NumExpansionBits = 4; // needs to hold the length of GConfigExpansions
+	constexpr int32 NumPlatformBits = 4; // needs to hold the length of the longest IniPlatformParent chains (which are generally going always be < 4)
+	constexpr int32 NumFlagsBits = 5; // needs to hold KeyFlags above
+	constexpr int32 NumDynamicLayerBits = 32 - (NumStaticLayerBits + NumExpansionBits + NumPlatformBits + NumFlagsBits); // needs to hold all the dynamic layers
+
+	constexpr int32 DynamicKeyOffset = NumStaticLayerBits + NumExpansionBits + NumPlatformBits + NumFlagsBits;
+	constexpr int32 LayerOffset = NumExpansionBits + NumPlatformBits + NumFlagsBits;
+	constexpr int32 ExpansionOffset = NumPlatformBits + NumFlagsBits;
+	constexpr int32 PlatformOffset = NumFlagsBits;
+	constexpr int32 FlagsOffset = 0;
 
-	constexpr int32 NumLayerBits = 6;
-	constexpr int32 NumExpansionBits = 6;
-	constexpr int32 NumPlatformBits = 6;
-	constexpr int32 NumFlagsBits = 3;
+	static_assert(UE_ARRAY_COUNT(GConfigLayers) < (1 << NumStaticLayerBits), "Need more NumStaticLayerBits");
+	static_assert(UE_ARRAY_COUNT(GConfigExpansions) < (1 << NumExpansionBits), "Need more NumExpansionBits");
+	static_assert(NumDynamicLayerBits > 4, "NumDynamicLayerBits is very small. Dynamic layers could overflow!");
+
+	constexpr int32 GetStaticKey(int32 LayerIndex, int32 ExpansionIndex, int32 PlatformIndex, int32 Flags)
+	{
+		return (LayerIndex << LayerOffset) + (ExpansionIndex << ExpansionOffset) + (PlatformIndex << PlatformOffset) + (Flags << FlagsOffset);
+	}
 }
 
 
@@ -876,6 +901,17 @@ void FConfigContext::AddStaticLayersToHierarchy(TArray<FString>* GatheredLayerFi
 			}
 			FlagPartOfKey |= KeyFlag_AssumeExists;
 		}
+		if (EnumHasAnyFlags(Layer.Flag, EConfigLayerFlags::Deprecated))
+		{
+			if (Layers == GPluginLayers)
+			{
+				FlagPartOfKey |= KeyFlag_Deprecated_Plugin;
+			}
+			else
+			{
+				FlagPartOfKey |= KeyFlag_Deprecated_Other;
+			}
+		}
 
 		// start replacing basic variables
 		FString LayerPath = PerformBasicReplacements(Layer.Path, *BaseIniName);
@@ -1010,6 +1046,41 @@ void FConfigContext::AddStaticLayersToHierarchy(TArray<FString>* GatheredLayerFi
 }
 
 
+// either show an editor warning, or just write to log for non-editor
+// NOTE: For tirst checkin, it will just log a warning, then once we have settled, we can checkin
+// more in-your-face behavior
+static void ProcessDeprecation(const FString& File, const FText& MessageText)
+{
+	if (!IsInGameThread())
+	{
+		AsyncTask(ENamedThreads::GameThread, [File, MessageText]() { ProcessDeprecation(File, MessageText); });
+		return;
+	}
+
+	FFormatNamedArguments Arguments;
+	Arguments.Add(TEXT("File"), FText::FromString(File));
+	Arguments.Add(TEXT("Message"), MessageText);
+	FText Msg = FText::Format(LOCTEXT("DeprecatedConfigFile", "Found a deprecated '{File}' - {Message}"), Arguments);
+
+	//if (GIsEditor)
+	//{
+	//	static TSet<FString> AlreadyWarnedKeys;
+
+	//	FString AbsPath = IFileManager::Get().ConvertToAbsolutePathForExternalAppForRead(*File);
+
+	//	FMessageLog EditorErrors("EditorErrors");
+	//	TSharedRef<FTokenizedMessage> Message = EditorErrors.Message(EMessageSeverity::Warning);
+	//	if (File.EndsWith(TEXT(".ini")))
+	//	{
+	//		Message->AddToken(FURLToken::Create(FString::Printf(TEXT("file://%s"), *AbsPath), LOCTEXT("DeprecatedGFPConfig_URLCLick", "Click to open file")));
+	//	}
+	//	Message->AddToken(FTextToken::Create(Msg));
+	//	EditorErrors.Notify();
+	//}
+
+	UE_LOG(LogConfig, Warning, TEXT("%s"), *Msg.ToString());
+	//ensureAlwaysMsgf(false, TEXT("%s"), *Msg.ToString());
+}
 
 /**
  * This will completely load .ini file hierarchy into the passed in FConfigFile. The passed in FConfigFile will then
@@ -1071,6 +1142,15 @@ bool FConfigContext::LoadIniFileHierarchy()
 
 		UE_CLOG(bDumpIniLoadInfo, LogConfig, Display, TEXT("   Found %s!"), *IniFileName);
 
+		if (HierarchyIt.Key & KeyFlag_Deprecated_Plugin)
+		{
+			ProcessDeprecation(IniFileName, LOCTEXT("RenameToDefault", "Rename to Default<PluginName>.ini'"));
+		}
+		if (HierarchyIt.Key & KeyFlag_Deprecated_Other)
+		{
+			ProcessDeprecation(IniFileName, LOCTEXT("RenameToDefault", "Unknown reason, check ConfigHierarchy.h for EConfigLayerFlags::Deprecated that might trigger this."));
+		}
+
 		bool bDoEmptyConfig = false;
 		//UE_LOG(LogConfig, Log,  TEXT( "Combining configFile: %s" ), *IniList(IniIndex) );
 
@@ -1280,20 +1360,6 @@ bool FConfigContext::GenerateDestIniFile()
 	FConfigFileHierarchy
 -----------------------------------------------------------------------------*/
 
-constexpr int32 DynamicKeyOffset = NumLayerBits + NumExpansionBits + NumPlatformBits + NumFlagsBits;
-constexpr int32 LayerOffset = NumExpansionBits + NumPlatformBits + NumFlagsBits;
-constexpr int32 ExpansionOffset = NumPlatformBits + NumFlagsBits;
-constexpr int32 PlatformOffset = NumFlagsBits;
-constexpr int32 FlagsOffset = 0;
-
-
-constexpr int32 GetStaticKey(int32 LayerIndex, int32 ExpansionIndex, int32 PlatformIndex, int32 Flags)
-{
-	return (LayerIndex << LayerOffset) + (ExpansionIndex << ExpansionOffset) + (PlatformIndex << PlatformOffset) + (Flags << FlagsOffset);
-}
-
-static_assert(UE_ARRAY_COUNT(GConfigLayers) < (1 << NumLayerBits), "Need more NumLayerBits");
-static_assert(UE_ARRAY_COUNT(GConfigExpansions) < (1 << NumExpansionBits), "Need more NumExpansionBits");
 
 FConfigFileHierarchy::FConfigFileHierarchy()
 	: KeyGen(1 << DynamicKeyOffset)
@@ -1303,6 +1369,7 @@ FConfigFileHierarchy::FConfigFileHierarchy()
 
 int32 FConfigFileHierarchy::GenerateDynamicKey()
 {
+	checkf((KeyGen + 1) < (1 << NumDynamicLayerBits), TEXT("The dynamic layer's index overflowed in FConfigFileHierarchy::GenerateDynamicKey(). We will need to fiddle with the number of assigned bits in the key, or extend to int64"));;
 	return ++KeyGen;
 }
 
@@ -1395,3 +1462,5 @@ void FConfigContext::VisualizeHierarchy(FOutputDevice& Ar, const TCHAR* IniName)
 
 	Ar.Logf(TEXT("======================================================="));
 }
+
+#undef LOCTEXT_NAMESPACE
diff --git a/Engine/Source/Runtime/Core/Public/Misc/ConfigHierarchy.h b/Engine/Source/Runtime/Core/Public/Misc/ConfigHierarchy.h
--- a/Engine/Source/Runtime/Core/Public/Misc/ConfigHierarchy.h
+++ b/Engine/Source/Runtime/Core/Public/Misc/ConfigHierarchy.h
@@ -54,10 +54,14 @@ inline FConfigLayer GPluginLayers[] =
 	// Engine/Base.ini
 	{ TEXT("AbsoluteBase"),				TEXT("{ENGINE}/Config/PluginBase.ini"), EConfigLayerFlags::NoExpand},
 
-	// Plugin/Base*.ini
-	{ TEXT("PluginBase"),				TEXT("{PLUGIN}/Config/Base{TYPE}.ini"), EConfigLayerFlags::UsePluginConfigCache },
-	// Plugin/Default*.ini (we use Base and Default as we can have both depending on Engine or Project plugin, but going forward we should stick with Default)
+	// Plugin/Base*.ini (DEPRECATED - Use DefaultPluginName.ini instead!)
+	// UE_DEPRECATED(5.8, "Use DefaultPluginName instead")
+	{ TEXT("PluginBase"),				TEXT("{PLUGIN}/Config/Base{TYPE}.ini"), EConfigLayerFlags::UsePluginConfigCache | EConfigLayerFlags::Deprecated},
+	// Plugin/Default*.ini
 	{ TEXT("PluginDefault"),			TEXT("{PLUGIN}/Config/Default{TYPE}.ini"), EConfigLayerFlags::UsePluginConfigCache },
+	// Plugin/*.ini (DEPRECATED - Use DefaultPluginName.ini instead!)
+	// UE_DEPRECATED(5.8, "Use DefaultPluginName instead")
+	{ TEXT("Plugin"),					TEXT("{PLUGIN}/Config/{TYPE}.ini"), EConfigLayerFlags::UsePluginConfigCache | EConfigLayerFlags::Deprecated },
 	// Plugin/Platform/Platform*.ini
 	{ TEXT("PluginPlatform"),			TEXT("{PLUGIN}/Config/{PLATFORM}/{PLATFORM}{TYPE}.ini"), EConfigLayerFlags::UsePluginConfigCache },
 	// Project/Default.ini
diff --git a/Engine/Source/Runtime/Core/Public/Misc/ConfigTypes.h b/Engine/Source/Runtime/Core/Public/Misc/ConfigTypes.h
--- a/Engine/Source/Runtime/Core/Public/Misc/ConfigTypes.h
+++ b/Engine/Source/Runtime/Core/Public/Misc/ConfigTypes.h
@@ -51,6 +51,7 @@ enum class EConfigLayerFlags : int32
 	RequiresCustomConfig = (1 << 5), // disabled if no custom config specified
 	UseGlobalConfigCache = (1 << 6),
 	UsePluginConfigCache = (1 << 7),
+	Deprecated = (1 << 8), // allowed but will ensure or show some other deprecated message
 };
 ENUM_CLASS_FLAGS(EConfigLayerFlags);
 
@@ -65,6 +66,8 @@ struct FConfigLayer
 	const TCHAR* Path;
 	// Special flag
 	EConfigLayerFlags Flag;
+	// Message to display if a deprecated file is found
+	const FString DeprecatedMessage;
 
 	// Will be cached early on to avoid threading issues
 	uint8 bHasCheckedExist = false;
