Add some additional code documentation to FIoRequestImpl, explaining the difference between ::SetFailed and ::SetLastBackendError.

#rnx

diff --git a/Engine/Source/Runtime/Core/Public/IO/IoDispatcherBackend.h b/Engine/Source/Runtime/Core/Public/IO/IoDispatcherBackend.h
--- a/Engine/Source/Runtime/Core/Public/IO/IoDispatcherBackend.h
+++ b/Engine/Source/Runtime/Core/Public/IO/IoDispatcherBackend.h
@@ -62,6 +62,10 @@ public:
 		}
 	}
 
+	/**
+	 * Mark the request as failed, just like calling ::SetFailed only this version allows the caller to
+	 * set the desired EIoErrorCode rather than accepting the default one.
+	 */
 	void SetLastBackendError(EIoErrorCode InError)
 	{
 		if (!IsCancelled())
