[LockFreeList]
- Fix benign race on relaxed read

diff --git a/Engine/Source/Runtime/Core/Public/Containers/LockFreeList.h b/Engine/Source/Runtime/Core/Public/Containers/LockFreeList.h
--- a/Engine/Source/Runtime/Core/Public/Containers/LockFreeList.h
+++ b/Engine/Source/Runtime/Core/Public/Containers/LockFreeList.h
@@ -94,8 +94,8 @@ public:
 		}
 		uint32 BlockIndex = Index / ItemsPerPage;
 		uint32 SubIndex = Index % ItemsPerPage;
-		checkLockFreePointerList(Index < (uint32)NextIndex.GetValue() && Index < MaxTotalItems && BlockIndex < MaxBlocks && Pages[BlockIndex]);
-		return Pages[BlockIndex] + SubIndex;
+		checkLockFreePointerList(Index < (uint32)NextIndex.GetValue() && Index < MaxTotalItems && BlockIndex < MaxBlocks && Pages[BlockIndex].load(std::memory_order_relaxed));
+		return Pages[BlockIndex].load(std::memory_order_relaxed) + SubIndex;
 	}
 
 private:
@@ -104,26 +104,28 @@ private:
 		uint32 BlockIndex = Index / ItemsPerPage;
 		uint32 SubIndex = Index % ItemsPerPage;
 		checkLockFreePointerList(Index && Index < (uint32)NextIndex.GetValue() && Index < MaxTotalItems && BlockIndex < MaxBlocks);
-		if (!Pages[BlockIndex])
+
+		if (!Pages[BlockIndex].load(std::memory_order_relaxed))
 		{
 			T* NewBlock = (T*)LockFreeAllocLinks(ItemsPerPage * sizeof(T));
 			checkLockFreePointerList(IsAligned(NewBlock, alignof(T)));
-			if (FPlatformAtomics::InterlockedCompareExchangePointer((void**)&Pages[BlockIndex], NewBlock, nullptr) != nullptr)
+			T* ExpectedValue = nullptr;
+			if (!Pages[BlockIndex].compare_exchange_strong(ExpectedValue, NewBlock))
 			{
 				// we lost discard block
-				checkLockFreePointerList(Pages[BlockIndex] && Pages[BlockIndex] != NewBlock);
+				checkLockFreePointerList(Pages[BlockIndex].load(std::memory_order_relaxed) && Pages[BlockIndex].load(std::memory_order_relaxed) != NewBlock);
 				LockFreeFreeLinks(ItemsPerPage * sizeof(T), NewBlock);
 			}
 			else
 			{
-				checkLockFreePointerList(Pages[BlockIndex]);
+				checkLockFreePointerList(Pages[BlockIndex].load(std::memory_order_relaxed));
 			}
 		}
-		return (void*)(Pages[BlockIndex] + SubIndex);
+		return (void*)(Pages[BlockIndex].load(std::memory_order_relaxed) + SubIndex);
 	}
 
 	alignas(PLATFORM_CACHE_LINE_SIZE) FThreadSafeCounter NextIndex;
-	alignas(PLATFORM_CACHE_LINE_SIZE) T* Pages[MaxBlocks];
+	alignas(PLATFORM_CACHE_LINE_SIZE) std::atomic<T*> Pages[MaxBlocks];
 };
 
 
