Fixed const incorrectness of FDelegateUserObjectConst when UE_WITH_REMOTE_OBJECT_HANDLE == 1.
Fixed some missing parentheses around requires conditions.

diff --git a/Engine/Source/Runtime/Core/Public/Delegates/IDelegateInstance.h b/Engine/Source/Runtime/Core/Public/Delegates/IDelegateInstance.h
--- a/Engine/Source/Runtime/Core/Public/Delegates/IDelegateInstance.h
+++ b/Engine/Source/Runtime/Core/Public/Delegates/IDelegateInstance.h
@@ -75,42 +75,64 @@ private:
 /*
  * Helper class that validates Delegate object type
  */
-class FDelegateUserObject
+template <bool bConst>
+class TDelegateUserObject
 {
-	const void* UserObject;
+	template <bool>
+	friend class TDelegateUserObject;
+
+	using VoidType = std::conditional_t<bConst, const void, void>;
+	using UObjectType = std::conditional_t<bConst, const UObject, UObject>;
+
+	VoidType* UserObject;
 	bool bIsUObject = false;
+
+private:
+	explicit TDelegateUserObject(VoidType* InUserObject, bool bInIsUObject)
+		: UserObject(InUserObject)
+		, bIsUObject(bInIsUObject)
+	{
+	}
+
 public:
 	template<typename UserObjectType>
-	FDelegateUserObject(const UserObjectType* InUserObject)
+		requires (bConst || !std::is_const_v<UserObjectType>)
+	TDelegateUserObject(UserObjectType* InUserObject)
 	{
-		UE_STATIC_ASSERT_COMPLETE_TYPE(UserObjectType, "FDelegateUserObject(UserObject) UserObject must not be an incomplete type");
-		UserObject = (const void*)InUserObject;
-		bIsUObject = std::is_convertible_v<UserObjectType*, UObject*>;
+		UE_STATIC_ASSERT_COMPLETE_TYPE(UserObjectType, "TDelegateUserObject(UserObject) UserObject must not be an incomplete type");
+		UserObject = (VoidType*)InUserObject;
+		bIsUObject = std::is_convertible_v<UserObjectType*, UObjectType*>;
 	}
 	template<UE::CPointerLike PtrWrapperType>
-		requires !UE::CPointer<PtrWrapperType>
-	UE_REWRITE FDelegateUserObject(const PtrWrapperType& InPtrWrapper)
-		: FDelegateUserObject(InPtrWrapper ? &*InPtrWrapper : nullptr)
+		requires (!UE::CPointer<PtrWrapperType>)
+	UE_REWRITE TDelegateUserObject(const PtrWrapperType& InPtrWrapper)
+		: TDelegateUserObject(InPtrWrapper ? &*InPtrWrapper : nullptr)
 	{
 	}
 	UE_REWRITE bool IsUObject() const
 	{
 		return bIsUObject;
 	}
-	UE_REWRITE const void* GetRaw() const
+	UE_REWRITE VoidType* GetRaw() const
 	{
 		return UserObject;
 	}
-	UE_REWRITE const UObject* GetUObject() const
+	UE_REWRITE UObjectType* GetUObject() const
 	{
-		return (const UObject*)UserObject;
+		return (UObjectType*)UserObject;
 	}
-	UE_REWRITE operator const void* () const
+	UE_REWRITE operator VoidType*() const
 	{
 		return UserObject;
 	}
+	UE_REWRITE operator TDelegateUserObject<true>() const
+		requires (!bConst)
+	{
+		return TDelegateUserObject<true>(UserObject, bIsUObject);
+	}
 };
-using FDelegateUserObjectConst = const FDelegateUserObject;
+using FDelegateUserObject = TDelegateUserObject<false>;
+using FDelegateUserObjectConst = TDelegateUserObject<true>;
 #else
 using FDelegateUserObject = void*;
 using FDelegateUserObjectConst = const void*;
