- Cleaned up the unused DynamicKey stuff

#FailedValidations Preflight

diff --git a/Engine/Source/Runtime/Core/Private/Misc/ConfigContext.cpp b/Engine/Source/Runtime/Core/Private/Misc/ConfigContext.cpp
--- a/Engine/Source/Runtime/Core/Private/Misc/ConfigContext.cpp
+++ b/Engine/Source/Runtime/Core/Private/Misc/ConfigContext.cpp
@@ -35,20 +35,18 @@ namespace
 	const TCHAR* SaveAllSectionsKey = TEXT("bCanSaveAllSections");
 
 	// some settings for hierarchy keys
-	constexpr int32 KeyFlag_UseGlobalCache = 1;
-	constexpr int32 KeyFlag_UsePluginCache = 2;
-	constexpr int32 KeyFlag_AssumeExists = 4;
-	constexpr int32 KeyFlag_Deprecated_Plugin = 8;
-	constexpr int32 KeyFlag_Deprecated_Other = 16;
-	// if you add a KeyFlag, increase NumFlagsBits below!
-
-	constexpr int32 NumStaticLayerBits = 6; // needs to hold the length of GConfigLayers (or any custom layers, which is highly unlikely to ever be longer than GConfigLayers)
-	constexpr int32 NumExpansionBits = 4; // needs to hold the length of GConfigExpansions
-	constexpr int32 NumPlatformBits = 4; // needs to hold the length of the longest IniPlatformParent chains (which are generally going always be < 4)
-	constexpr int32 NumFlagsBits = 5; // needs to hold KeyFlags above
-	constexpr int32 NumDynamicLayerBits = 32 - (NumStaticLayerBits + NumExpansionBits + NumPlatformBits + NumFlagsBits); // needs to hold all the dynamic layers
-
-	constexpr int32 DynamicKeyOffset = NumStaticLayerBits + NumExpansionBits + NumPlatformBits + NumFlagsBits;
+	constexpr int32 KeyFlag_UseGlobalCache = 1 << 0;
+	constexpr int32 KeyFlag_UsePluginCache = 1 << 1;
+	constexpr int32 KeyFlag_AssumeExists = 1 << 2;
+	constexpr int32 KeyFlag_Deprecated_Plugin = 1 << 3;
+	constexpr int32 KeyFlag_Deprecated_Other = 1 << 4;
+	// if you add a KeyFlag, make sure to increase NumFlagsBits below if needed
+
+	constexpr int32 NumStaticLayerBits = 8; // needs to hold the length of GConfigLayers (or any custom layers, which is highly unlikely to ever be longer than GConfigLayers)
+	constexpr int32 NumExpansionBits = 5; // needs to hold the length of GConfigExpansions
+	constexpr int32 NumPlatformBits = 5; // needs to hold the length of the longest IniPlatformParent chains (which are generally going always be < 4)
+	constexpr int32 NumFlagsBits = 8; // needs to hold KeyFlags above
+
 	constexpr int32 LayerOffset = NumExpansionBits + NumPlatformBits + NumFlagsBits;
 	constexpr int32 ExpansionOffset = NumPlatformBits + NumFlagsBits;
 	constexpr int32 PlatformOffset = NumFlagsBits;
@@ -56,7 +54,7 @@ namespace
 
 	static_assert(UE_ARRAY_COUNT(GConfigLayers) < (1 << NumStaticLayerBits), "Need more NumStaticLayerBits");
 	static_assert(UE_ARRAY_COUNT(GConfigExpansions) < (1 << NumExpansionBits), "Need more NumExpansionBits");
-	static_assert(NumDynamicLayerBits > 4, "NumDynamicLayerBits is very small. Dynamic layers could overflow!");
+	static_assert(NumStaticLayerBits + NumExpansionBits + NumPlatformBits + NumFlagsBits <= 32, "All bits combined is greater than 32");
 
 	constexpr int32 GetStaticKey(int32 LayerIndex, int32 ExpansionIndex, int32 PlatformIndex, int32 Flags)
 	{
@@ -1362,17 +1360,10 @@ bool FConfigContext::GenerateDestIniFile()
 
 
 FConfigFileHierarchy::FConfigFileHierarchy()
-	: KeyGen(1 << DynamicKeyOffset)
 {
 }
 
 
-int32 FConfigFileHierarchy::GenerateDynamicKey()
-{
-	checkf((KeyGen + 1) < (1 << NumDynamicLayerBits), TEXT("The dynamic layer's index overflowed in FConfigFileHierarchy::GenerateDynamicKey(). We will need to fiddle with the number of assigned bits in the key, or extend to int64"));;
-	return ++KeyGen;
-}
-
 int32 FConfigFileHierarchy::AddStaticLayer(const FString& Filename, int32 LayerIndex, int32 ExpansionIndex /*= 0*/, int32 PlatformIndex /*= 0*/, int32 Flags)
 {
 	int32 Key = GetStaticKey(LayerIndex, ExpansionIndex, PlatformIndex, Flags);
@@ -1380,13 +1371,6 @@ int32 FConfigFileHierarchy::AddStaticLayer(const FString& Filename, int32 LayerI
 	return Key;
 }
 
-int32 FConfigFileHierarchy::AddDynamicLayer(const FString& Filename)
-{
-	int32 Key = GenerateDynamicKey();
-	Emplace(Key, Filename);
-	return Key;
-}
-
 void FConfigContext::EnsureRequiredGlobalPathsHaveBeenInitialized()
 {
 	PerformBasicReplacements(TEXT(""), TEXT("")); // requests user directories and FConfigCacheIni::GetCustomConfigString
diff --git a/Engine/Source/Runtime/Core/Public/Misc/ConfigTypes.h b/Engine/Source/Runtime/Core/Public/Misc/ConfigTypes.h
--- a/Engine/Source/Runtime/Core/Public/Misc/ConfigTypes.h
+++ b/Engine/Source/Runtime/Core/Public/Misc/ConfigTypes.h
@@ -12,16 +12,12 @@
 
 class FConfigFileHierarchy : public TMap<int32, FUtf8String>
 {
-private:
-	int32 KeyGen = 0;
-
 public:
 	FConfigFileHierarchy();
 
 	friend FArchive& operator<<(FArchive& Ar, FConfigFileHierarchy& ConfigFileHierarchy)
 	{
 		Ar << static_cast<FConfigFileHierarchy::Super&>(ConfigFileHierarchy);
-		Ar << ConfigFileHierarchy.KeyGen;
 		return Ar;
 	}
 
