- Removed two unused function declarations

diff --git a/Engine/Source/Runtime/Core/Public/Misc/ConfigTypes.h b/Engine/Source/Runtime/Core/Public/Misc/ConfigTypes.h
--- a/Engine/Source/Runtime/Core/Public/Misc/ConfigTypes.h
+++ b/Engine/Source/Runtime/Core/Public/Misc/ConfigTypes.h
@@ -22,10 +22,8 @@ public:
 	}
 
 private:
-	int32 GenerateDynamicKey();
 
 	int32 AddStaticLayer(const FString& Filename, int32 LayerIndex, int32 ExpansionIndex, int32 PlatformIndex, int32 Flags);
-	int32 AddDynamicLayer(const FString& Filename);
 
 	friend class FConfigFile;
 	friend class FConfigContext;
