update logging associated with memmaping ios.

#rnx

diff --git a/Engine/Source/Runtime/Core/Private/Apple/ApplePlatformFile.cpp b/Engine/Source/Runtime/Core/Private/Apple/ApplePlatformFile.cpp
--- a/Engine/Source/Runtime/Core/Private/Apple/ApplePlatformFile.cpp
+++ b/Engine/Source/Runtime/Core/Private/Apple/ApplePlatformFile.cpp
@@ -473,7 +473,7 @@ public:
 		const uint8* AlignedMapPtr = (const uint8 *)mmap(NULL, AlignedSize, Protection, InternalFlags, FileHandle, AlignedOffset);
 		if (AlignedMapPtr == (const uint8*)-1 || AlignedMapPtr == nullptr)
 		{
-			UE_LOG(LogHAL, Warning, TEXT("Failed to map memory %s, error is %d"), *Filename, errno);
+			UE_LOG(LogHAL, Warning, TEXT("Failed to map memory %s, error is %d, aligned size = %lld, aligned offset = %lld, internal flags = %d"), *Filename, errno, AlignedSize, AlignedOffset, InternalFlags);
 			return nullptr;
 		}
 		LLM_IF_ENABLED(FLowLevelMemTracker::Get().OnLowLevelAlloc(ELLMTracker::Platform, AlignedMapPtr, AlignedSize));
