Fix TRangedConsolveVariable to return Base::GetFloat() in the GetFloat() function

diff --git a/Engine/Source/Runtime/Core/Private/HAL/ConsoleManager.cpp b/Engine/Source/Runtime/Core/Private/HAL/ConsoleManager.cpp
--- a/Engine/Source/Runtime/Core/Private/HAL/ConsoleManager.cpp
+++ b/Engine/Source/Runtime/Core/Private/HAL/ConsoleManager.cpp
@@ -1970,7 +1970,7 @@ public:
 	}
 	virtual float GetFloat() const override
 	{
-		return FMath::Clamp(Base::GetInt(), Min->GetFloat(), Max->GetFloat());
+		return FMath::Clamp(Base::GetFloat(), Min->GetFloat(), Max->GetFloat());
 	}
 	virtual FString GetString() const override
 	{
