Added an ensure to validate the MallocStomp2MinSize and MallocStomp2MaxSize variables.

diff --git a/Engine/Source/Runtime/Core/Private/HAL/MallocStomp2.cpp b/Engine/Source/Runtime/Core/Private/HAL/MallocStomp2.cpp
--- a/Engine/Source/Runtime/Core/Private/HAL/MallocStomp2.cpp
+++ b/Engine/Source/Runtime/Core/Private/HAL/MallocStomp2.cpp
@@ -357,6 +357,8 @@ FMalloc* FMallocStomp2::OverrideIfEnabled(FMalloc* InUsedAlloc)
 		FParse::Value(CommandLine, TEXT("MallocStomp2MinSize="), MinSize);
 		FParse::Value(CommandLine, TEXT("MallocStomp2MaxSize="), MaxSize);
 		bool bUseUnderrunMode = FParse::Param(CommandLine, TEXT("MallocStomp2UnderrunMode"));
+		ensureMsgf(MaxSize > 0, TEXT("MallocStomp2 has been enabled but MallocStomp2MaxSize is set to 0 - this will not produce meaningful results."));
+		ensureMsgf(MaxSize >= MinSize, TEXT("MallocStomp2 has been enabled but MallocStomp2MaxSize is smaller than MallocStomp2MinSize - this will not produce meaningful results."));
 #else
 		bool bUseUnderrunMode = false;
 #endif
