[RaceDetector]
- Fix false positive that could happen if a value was written between CreateThread(CREATE_SUSPENDED) and ResumeThread and then read by the thread after it was resumed

diff --git a/Engine/Source/Runtime/Core/Private/Sanitizer/RaceDetectorWindows.cpp b/Engine/Source/Runtime/Core/Private/Sanitizer/RaceDetectorWindows.cpp
--- a/Engine/Source/Runtime/Core/Private/Sanitizer/RaceDetectorWindows.cpp
+++ b/Engine/Source/Runtime/Core/Private/Sanitizer/RaceDetectorWindows.cpp
@@ -141,6 +141,16 @@ INSTRUMENTATION_FUNCTION_ATTRIBUTES DWORD WINAPI InstrumentedStartThread(LPVOID
 	{
 		FInstrumentationScope InstrumentationScope;
 		Context.ClockBank.Acquire(ThreadArgs->ClockBank, _ReturnAddress());
+
+		// Do not try to sync if we reach this before the handle is set from the creating thread
+		// it means that CREATE_SUSPENDED was not passed as a flag anyway so we don't need any barrier in that case.
+		if (ThreadArgs->ThreadHandle && bRuntimeInitialized)
+		{
+			// This allows us to sync with ResumeThread if the thread was created in the suspended state.
+			FSyncObjectRef Sync = GetSyncObject(Context, ThreadArgs->ThreadHandle);
+			Sync->SyncAcquireAsSoleOwnerOrReadOwner(Context, _ReturnAddress(), ThreadArgs->ThreadHandle, TEXT("InstrumentedStartThread"));
+		}
+
 		// We don't increment the clock here since we don't have any associated
 		// contextid yet. Clock is incremented when contextid is reserved on first
 		// memory access.
@@ -644,6 +654,29 @@ INSTRUMENTATION_FUNCTION_ATTRIBUTES PVOID NTAPI DetouredRtlReallocateHeap(PVOID
 	return NewPtr;
 }
 
+DWORD ResumeThread(
+  [in] HANDLE hThread
+);
+
+typedef NTSYSAPI DWORD(NTAPI RESUMETHREAD)(IN HANDLE ThreadHandle);
+
+RESUMETHREAD* TrueResumeThread = nullptr;
+INSTRUMENTATION_FUNCTION_ATTRIBUTES DWORD NTAPI DetouredResumeThread(HANDLE ThreadHandle)
+{
+	START_WINAPI_INSTRUMENTATION(return TrueResumeThread(ThreadHandle));
+
+	// Add a release barrier when thread is resumed to avoid false positive that could happen
+	// if something is written between the createthread and the resumethread and then 
+	// read by the thread once it is resumed.
+	FSyncObjectRef Sync = GetSyncObject(Context, ThreadHandle);
+	Sync->SyncReleaseAsSoleOwner(Context, _ReturnAddress(), ThreadHandle, TEXT("ResumeThread"));
+	DWORD result = TrueResumeThread(ThreadHandle);
+
+	FINISH_WINAPI_INSTRUMENTATION();
+
+	return result;
+}
+
 // Functions we need to detour for race detector.
 TArray<TPair<PVOID*, void*>> DetouredFunctions;
 
@@ -654,6 +687,8 @@ INSTRUMENTATION_FUNCTION_ATTRIBUTES void PopulateDetouredFunctions()
 		return;
 	}
 
+	TrueResumeThread = ::ResumeThread;
+	DetouredFunctions.Emplace(&(PVOID&)TrueResumeThread, DetouredResumeThread);
 	TrueWaitForSingleObject = ::WaitForSingleObject;
 	DetouredFunctions.Emplace(&(PVOID&)TrueWaitForSingleObject, DetouredWaitForSingleObject);
 	TrueWaitForSingleObjectEx = ::WaitForSingleObjectEx;
