- Fixed a duplication LOCTEXT tag

diff --git a/Engine/Source/Runtime/Core/Private/Misc/ConfigContext.cpp b/Engine/Source/Runtime/Core/Private/Misc/ConfigContext.cpp
--- a/Engine/Source/Runtime/Core/Private/Misc/ConfigContext.cpp
+++ b/Engine/Source/Runtime/Core/Private/Misc/ConfigContext.cpp
@@ -1146,7 +1146,7 @@ bool FConfigContext::LoadIniFileHierarchy()
 		}
 		if (HierarchyIt.Key & KeyFlag_Deprecated_Other)
 		{
-			ProcessDeprecation(IniFileName, LOCTEXT("RenameToDefault", "Unknown reason, check ConfigHierarchy.h for EConfigLayerFlags::Deprecated that might trigger this."));
+			ProcessDeprecation(IniFileName, LOCTEXT("UnknownConfigDeprecation", "Unknown reason, check ConfigHierarchy.h for EConfigLayerFlags::Deprecated that might trigger this."));
 		}
 
 		bool bDoEmptyConfig = false;
