* Support truncation of PerfMem CSVProfiler files via csv.MaxFileSizeMB cvar for maximum file size. default: 0 (unlimited).
* Add csv.DeleteDefaultDirectory command to clean up after specific dev usage of CSVProfiler that does not specify an explicit directory.
* Add FCSVProfiler::GetNumBytesWritten so that it can be sent to analytics before attempting to send CSV file with feedback  (TODO send withCore.BadMatchDetected analytics, to help understand CSV upload challenges)
* Picks up where Daniel Silver left off

#tests Trace through various cases in cooked client and PIE

#rnx

diff --git a/Engine/Source/Runtime/Core/Private/ProfilingDebugging/CsvProfiler.cpp b/Engine/Source/Runtime/Core/Private/ProfilingDebugging/CsvProfiler.cpp
--- a/Engine/Source/Runtime/Core/Private/ProfilingDebugging/CsvProfiler.cpp
+++ b/Engine/Source/Runtime/Core/Private/ProfilingDebugging/CsvProfiler.cpp
@@ -150,11 +150,28 @@ TAutoConsoleVariable<int32> CVarCsvTargetFrameRateOverride(
 TAutoConsoleVariable<int32> CVarCsvShippingContinuousWrites(
 	TEXT("csv.Shipping.ContinuousWrites"),
 	-1,
-	TEXT("Only applies in shipping buids. If set, overrides csv.ContinousWrites."),
+	TEXT("Only applies in shipping buids. If != -1, overrides csv.ContinousWrites."),
+	ECVF_Default
+);
+
+TAutoConsoleVariable<int32> CVarCsvShippingMaxFileSizeMB(
+	TEXT("csv.Shipping.MaxFileSizeMB"),
+	-1,
+	TEXT("Only applies in shipping buids. If != -1, overrides csv.MaxFileSizeMB."),
 	ECVF_Default
 );
 #endif
 
+TAutoConsoleVariable<int32> CVarCsvMaxFileSizeMB(
+	TEXT("csv.MaxFileSizeMB"),
+	0,
+	TEXT("Approximate maximum file size (in megabytes) for CSV profiler files.")
+	TEXT("The profiler will stop writing values to the output file after this size is exceeded.")
+	TEXT("Only supported when ContinuousWrites cvar is > 0.")
+	TEXT("If <= 0, size is not limited."),
+	ECVF_Default
+);
+
 TAutoConsoleVariable<int32> CVarCsvCompressionMode(
 	TEXT("csv.CompressionMode"),
 	-1,
@@ -457,7 +474,7 @@ static TArray<FEventExecCmds>* GCsvEventExecCmds = NULL;
 
 bool IsContinuousWriteEnabled(bool bGameThread)
 {
-	int CVarValue = -1;
+	int32 CVarValue = -1;
 #if UE_BUILD_SHIPPING
 	CVarValue = bGameThread ? CVarCsvShippingContinuousWrites.GetValueOnGameThread() : CVarCsvShippingContinuousWrites.GetValueOnAnyThread();
 	if (CVarValue == -1)
@@ -468,6 +485,18 @@ bool IsContinuousWriteEnabled(bool bGameThread)
 	return CVarValue > 0;
 } 
 
+int32 GetMaxFileSizeMB(bool bGameThread)
+{
+	int32 CVarValue = -1;
+#if UE_BUILD_SHIPPING
+	CVarValue = bGameThread ? CVarCsvShippingMaxFileSizeMB.GetValueOnGameThread() : CVarCsvShippingMaxFileSizeMB.GetValueOnAnyThread();
+	if (CVarValue == -1)
+#endif
+	{
+		CVarValue = bGameThread ? CVarCsvMaxFileSizeMB.GetValueOnGameThread() : CVarCsvMaxFileSizeMB.GetValueOnAnyThread();
+	}
+	return CVarValue;
+}
 
 static CSV_PROFILER_INLINE void ValidateFName(const FName& StatName)
 {
@@ -1069,6 +1098,24 @@ static FAutoConsoleCommandWithWorldArgsAndOutputDevice HandleCSVCategoryCmd(
 	FConsoleCommandWithWorldArgsAndOutputDeviceDelegate::CreateStatic(&HandleCSVCategoryCommand)
 );
 
+static void HandleCSVDeleteDefaultDirectoryCommand()
+{
+	if ((GCsvProfilerIsCapturing && IsContinuousWriteEnabled(true)) || GCsvProfilerIsWritingFile)
+	{
+		UE_LOG(LogCsvProfiler, Error, TEXT("CSV Profiler is writing file, ignoring command csv.DeleteDefaultDirectory"));
+	}
+	else
+	{
+		IFileManager::Get().DeleteDirectory(*FCsvProfiler::GetDefaultDirectory(), false, true);
+	}
+}
+
+static FAutoConsoleCommand HandleCSVDeleteDefaultDirectoryCmd(
+	TEXT("csv.DeleteDefaultDirectory"),
+	TEXT("Deletes the default directory for CSV Profiler files. Directory is currently only used by the CsvProfile command, and the command line argument: -ReplicationCSVCaptureFrames=X"),
+	FConsoleCommandDelegate::CreateStatic(&HandleCSVDeleteDefaultDirectoryCommand)
+);
+
 //-----------------------------------------------------------------------------
 //	TSingleProducerSingleConsumerList : fast lock-free single producer/single 
 //  consumer list implementation. 
@@ -1707,10 +1754,11 @@ struct FCsvStatSeriesValue
 class FCsvWriterHelper
 {
 public:
-	FCsvWriterHelper(const TSharedRef<FArchive>& InOutputFile, int32 InBufferSize, bool bInCompressOutput)
+	FCsvWriterHelper(const TSharedRef<FArchive>& InOutputFile, int32 InBufferSize, bool bInCompressOutput, bool bInContinuousWrites)
 		: OutputFile(InOutputFile)
 		, bIsLineStart(true)
 		, BytesInBuffer(0)
+		, MaxFileBytes(0)
 	{
 		if (InBufferSize > 0)
 		{
@@ -1720,6 +1768,28 @@ public:
 				GZipBuffer.SetNumUninitialized(InBufferSize);
 			}
 		}
+
+		BytesWritten.store(0);
+
+		int32 MaxFileSizeMB = GetMaxFileSizeMB(false);
+		if (MaxFileSizeMB > 0)
+		{
+			if (!bInContinuousWrites)
+			{
+				UE_LOG(LogCsvProfiler, Error, TEXT("cvar csv.MaxFileSizeMB not supported unless the appropriate ContinuousWrites cvar is true. Ignoring MaxFileSizeMB."));
+				MaxFileBytes = 0;
+			}
+			else
+			{
+				MaxFileBytes = MaxFileSizeMB;
+				MaxFileBytes *= 1024 * 1024;
+				if (MaxFileBytes > InBufferSize)
+				{
+					// reduce the threshold so that tuned value isn't exceeded quite as much once the file is done flushing
+					MaxFileBytes -= InBufferSize;
+				}
+			}
+		}
 	}
 
 	~FCsvWriterHelper()
@@ -1821,7 +1891,7 @@ private:
 	{
 		if (Buffer.Num() == 0)
 		{
-			OutputFile->Serialize(Src, NumBytes);
+			SerializeToOutputFile(Src, NumBytes);
 		}
 		else
 		{
@@ -1884,19 +1954,35 @@ private:
 				}
 
 				{
-					OutputFile->Serialize(GZipBuffer.GetData(), CompressedSize);
+					SerializeToOutputFile(GZipBuffer.GetData(), CompressedSize);
 				}
 			}
 			else
 			{
 				// No compression. Write directly to the output file
-				OutputFile->Serialize(Buffer.GetData(), BytesInBuffer);
+				SerializeToOutputFile(Buffer.GetData(), BytesInBuffer);
 			}
 
 			BytesInBuffer = 0;
 		}
 	}
 
+	void SerializeToOutputFile(void* Src, int32 NumBytes)
+	{
+		OutputFile->Serialize(Src, NumBytes);
+
+		uint64 PreviousBytesWritten = BytesWritten.fetch_add(NumBytes, std::memory_order_relaxed);
+		if (MaxFileBytes > 0)
+		{
+			uint64 CurrentBytesWritten = PreviousBytesWritten;
+			CurrentBytesWritten += NumBytes;
+			if (PreviousBytesWritten < MaxFileBytes && CurrentBytesWritten >= MaxFileBytes)
+			{
+				UE_LOG(LogCsvProfiler, Log, TEXT("CSV Profiler will no longer write values due to exceeding csv.MaxFileSizeMB. %llu written of %llu max bytes)"), CurrentBytesWritten, MaxFileBytes);
+			}
+		}
+	}
+
 	TSharedRef<FArchive> OutputFile;
 	bool bIsLineStart;
 
@@ -1904,11 +1990,24 @@ private:
 	TArray<uint8> Buffer;
 	TArray<uint8> GZipBuffer;
 
+	std::atomic<uint64> BytesWritten;
+	uint64 MaxFileBytes;
+
 public:
 	inline uint64 GetAllocatedSize() const
 	{
 		return Buffer.GetAllocatedSize() + GZipBuffer.GetAllocatedSize();
 	}
+
+	inline uint64 GetNumBytesWritten() const
+	{
+		return BytesWritten.load(std::memory_order_relaxed);
+	}
+
+	inline bool HasReachedMaxFileBytes() const
+	{
+		return MaxFileBytes <= 0 ? false : GetNumBytesWritten() >= MaxFileBytes;
+	}
 };
 
 struct FCsvProcessedEvent
@@ -2253,6 +2352,9 @@ public:
 	void Finalize(const TMap<FString, FString>& Metadata);
 
 	inline uint64 GetAllocatedSize() const;
+
+	inline uint64 GetNumBytesWritten() const;
+	inline bool HasReachedMaxFileBytes() const;
 };
 
 FCsvStatSeries::FCsvStatSeries(EType InSeriesType, const FCsvStatID& InStatID, FCsvStreamWriter* InWriter, FCsvStatRegister& StatRegister, const FString& ThreadName, FCsvAggregateStatSeries* InLinkedAggregateStatSeries)
@@ -2815,7 +2917,7 @@ private:
 
 
 FCsvStreamWriter::FCsvStreamWriter(const TSharedRef<FArchive>& InOutputFile, bool bInContinuousWrites, int32 InBufferSize, int64 InNumFramesToBuffer, bool bInCompressOutput, uint32 InRenderThreadId, uint32 InRHIThreadId, bool bAggregateTaskWorkerStats)
-	: Stream(InOutputFile, InBufferSize, bInCompressOutput)
+	: Stream(InOutputFile, InBufferSize, bInCompressOutput, bInContinuousWrites)
 	, NumFramesToBuffer(InNumFramesToBuffer)
 	, WriteFrameIndex(-1)
 	, ReadFrameIndex(-1)
@@ -2850,6 +2952,11 @@ void FCsvStreamWriter::AddSeries(FCsvStatSeries* Series, bool bIsVisible)
 void FCsvStreamWriter::PushValue(FCsvStatSeries* Series, int64 FrameNumber, const FCsvStatSeriesValue& Value)
 {
 	check(Series->ColumnIndex != -1);
+	
+	if (Stream.HasReachedMaxFileBytes())
+	{
+		return;
+	}
 
 	WriteFrameIndex = FMath::Max(FrameNumber, WriteFrameIndex);
 
@@ -2866,6 +2973,11 @@ void FCsvStreamWriter::PushValue(FCsvStatSeries* Series, int64 FrameNumber, cons
 
 void FCsvStreamWriter::PushEvent(const FCsvProcessedEvent& Event)
 {
+	if (Stream.HasReachedMaxFileBytes())
+	{
+		return;
+	}
+
 	Rows.FindOrAdd(Event.FrameNumber).Events.Add(Event);
 }
 
@@ -3046,6 +3158,17 @@ uint64 FCsvStreamWriter::GetAllocatedSize() const
 	return Size;
 }
 
+inline uint64 FCsvStreamWriter::GetNumBytesWritten() const
+{
+	return Stream.GetNumBytesWritten();
+}
+
+inline bool FCsvStreamWriter::HasReachedMaxFileBytes() const
+{
+	return Stream.HasReachedMaxFileBytes();
+}
+
+
 //-----------------------------------------------------------------------------
 //	FCsvProfilerProcessingThread class : low priority thread to process 
 //  profiling data
@@ -3656,7 +3779,7 @@ void FCsvProfiler::BeginCaptureInternal(const FCsvCaptureCommand& CurrentCommand
 	const TCHAR* CsvExtension = bCompressOutput ? TEXT(".csv.gz") : TEXT(".csv");
 
 	// Determine the output path and filename based on override params
-	FString DestinationFolder = CurrentCommand.DestinationFolder.IsEmpty() ? FPaths::ProfilingDir() + TEXT("CSV/") : CurrentCommand.DestinationFolder + TEXT("/");
+	FString DestinationFolder = CurrentCommand.DestinationFolder.IsEmpty() ? GetDefaultDirectory() : CurrentCommand.DestinationFolder + TEXT("/");
 	FString Filename = CurrentCommand.Filename.IsEmpty() ? FString::Printf(TEXT("Profile(%s)%s"), *FDateTime::Now().ToString(TEXT("%Y%m%d_%H%M%S")), CsvExtension) : CurrentCommand.Filename;
 	OutputFilename = DestinationFolder + Filename;
 
@@ -4071,6 +4194,9 @@ TSharedFuture<FString> FCsvProfiler::EndCapture(FGraphEventRef EventToSignal)
 	SetNonPersistentMetadata(TEXT("EndTimestamp"), *FString::Printf(TEXT("%lld"), FDateTime::UtcNow().ToUnixTimestamp()));
 	SetNonPersistentMetadata(TEXT("CaptureDuration"), *FString::SanitizeFloat(FPlatformTime::Seconds()-CaptureStartTime));
 
+	bool bHasReachedMaxFileBytes = CsvWriter == nullptr ? false : CsvWriter->HasReachedMaxFileBytes();
+	SetNonPersistentMetadata(TEXT("CsvMaxFileBytesReached"), bHasReachedMaxFileBytes ? TEXT("1") : TEXT("0"));
+
 	TPromise<FString>* Completion = new TPromise<FString>([EventToSignal]()
 	{
 		if (EventToSignal)
@@ -4342,6 +4468,10 @@ void FCsvProfiler::EndWait()
 #endif
 }
 
+FString FCsvProfiler::GetDefaultDirectory()
+{
+	return FPaths::ProfilingDir() + TEXT("CSV/");
+}
 
 void FCsvProfiler::RecordEventfInternal(int32 CategoryIndex, const TCHAR* Fmt, ...)
 {
@@ -4877,6 +5007,12 @@ const FGuid& FCsvProfiler::GetCsvID() const
 	return CsvGUID;
 }
 
+/*Get the current number of bytes written to file*/
+uint64 FCsvProfiler::GetNumBytesWritten() const
+{
+	return CsvWriter == nullptr ? 0 : CsvWriter->GetNumBytesWritten();
+}
+
 bool FCsvProfiler::EnableCategoryByString(const FString& CategoryName) const
 {
 	int32 Category = FCsvCategoryData::Get()->GetCategoryIndex(CategoryName);
diff --git a/Engine/Source/Runtime/Core/Public/ProfilingDebugging/CsvProfiler.h b/Engine/Source/Runtime/Core/Public/ProfilingDebugging/CsvProfiler.h
--- a/Engine/Source/Runtime/Core/Public/ProfilingDebugging/CsvProfiler.h
+++ b/Engine/Source/Runtime/Core/Public/ProfilingDebugging/CsvProfiler.h
@@ -423,6 +423,8 @@ public:
 	CORE_API static void BeginWait();
 	CORE_API static void EndWait();
 
+	CORE_API static FString GetDefaultDirectory();
+
 	/** Singleton interface */
 	CORE_API bool IsCapturing() const;
 	CORE_API bool IsCapturing_Renderthread() const;
@@ -433,6 +435,7 @@ public:
 	CORE_API int32 GetCaptureFrameNumberRT() const;
 	CORE_API int32 GetNumFrameToCaptureOnEvent() const;
 	CORE_API const FGuid& GetCsvID() const;
+	CORE_API uint64	GetNumBytesWritten() const;
 
 	CORE_API bool EnableCategoryByString(const FString& CategoryName) const;
 	CORE_API void EnableCategoryByIndex(uint32 CategoryIndex, bool bEnable) const;
