- Fixed console command arg parsting to support ""

diff --git a/Engine/Source/Runtime/Core/Private/HAL/ConsoleManager.cpp b/Engine/Source/Runtime/Core/Private/HAL/ConsoleManager.cpp
--- a/Engine/Source/Runtime/Core/Private/HAL/ConsoleManager.cpp
+++ b/Engine/Source/Runtime/Core/Private/HAL/ConsoleManager.cpp
@@ -3022,8 +3022,12 @@ bool FConsoleManager::ProcessUserConsoleInput(const TCHAR* InInput, FOutputDevic
 	{
 		// Process command
 		// Build up argument list
-		TArray< FString > Args;
-		FString( It ).ParseIntoArrayWS( Args );
+		TArray<FString> Args;
+		FString Arg;
+		while (FParse::Token(It, Arg, false))
+		{
+			Args.Add(Arg);
+		}
 
 		const bool bShowHelp = bCommandEndedInQuestion || ((Args.Num() == 1) && (Args[0] == TEXT("?")));
 		if( bShowHelp )
