Stop default-initing in some Math lib DiagnosticCheckNaN impls.

Moving towards the check being just a check and not mutating anything, for numerous reasons.
Starting with just Vector and Box, doing this incrementally in case this causes log spam or other
issues.

diff --git a/Engine/Source/Runtime/Core/Public/Math/Box.h b/Engine/Source/Runtime/Core/Public/Math/Box.h
--- a/Engine/Source/Runtime/Core/Public/Math/Box.h
+++ b/Engine/Source/Runtime/Core/Public/Math/Box.h
@@ -97,12 +97,11 @@ public:
 		{
 #if ENABLE_UNINITIALIZED_BOX_DIAGNOSTIC
 			// In case ENABLE_UNINITIALIZED_BOX_DIAGNOSTIC was explicitly enabled then report each encountered accesses to uninitialized members.
-			ensureAlwaysMsgf(false, TEXT("FBox contains NaN: %s. Initializing with zero extent and marking as invalid for safety."), *ToString());
+			ensureAlwaysMsgf(false, TEXT("FBox contains NaN: %s."), *ToString());
 #else
 			// In case diagnostic check is enable from NaN diagnostic we report the error using the NaN error reporting behavior.
-			logOrEnsureNanError(TEXT("FBox contains NaN: %s. Initializing with zero extent and marking as invalid for safety."), *ToString());
+			logOrEnsureNanError(TEXT("FBox contains NaN: %s."), *ToString());
 #endif
-			const_cast<TBox<T>*>(this)->Init();
 		}
 	}
 #else
diff --git a/Engine/Source/Runtime/Core/Public/Math/Vector.h b/Engine/Source/Runtime/Core/Public/Math/Vector.h
--- a/Engine/Source/Runtime/Core/Public/Math/Vector.h
+++ b/Engine/Source/Runtime/Core/Public/Math/Vector.h
@@ -131,7 +131,6 @@ public:
         if (ContainsNaN())
         {
             logOrEnsureNanError(TEXT("FVector contains NaN: %s"), *ToString());
-            *const_cast<TVector<T>*>(static_cast<const TVector<T>*>(this)) = ZeroVector;
         }
     }
 
@@ -140,7 +139,6 @@ public:
         if (ContainsNaN())
         {
             logOrEnsureNanError(TEXT("%s: FVector contains NaN: %s"), Message, *ToString());
-            *const_cast<TVector<T>*>(static_cast<const TVector<T>*>(this)) = ZeroVector;
         }
     }
 #else
