Fixed comment to refer to correct macro name.

#rnx

diff --git a/Engine/Source/Runtime/Core/Private/Templates/TuplesTest.cpp b/Engine/Source/Runtime/Core/Private/Templates/TuplesTest.cpp
--- a/Engine/Source/Runtime/Core/Private/Templates/TuplesTest.cpp
+++ b/Engine/Source/Runtime/Core/Private/Templates/TuplesTest.cpp
@@ -273,7 +273,7 @@ IMPLEMENT_SIMPLE_AUTOMATION_TEST(FTuplesTest, "System.Core.Tuples.Smoke", EAutom
 bool FTuplesTest::RunTest(const FString& Parameters)
 {
 	// Test getting references from tuples - this catches an issue found when MSVC returns the wrong address of a reference-to
-	// array element in an rvalue tuple.  See UE_TUPLE_RVALUE_WORKAROUND.
+	// array element in an rvalue tuple.  See UE_TUPLE_REFERENCE_WORKAROUND.
 	{
 		auto TestTupleValuesImpl = [](auto&& TupleRef, const auto& Addresses)
 		{
