[Insights] TraceAux: Added EngineVersion to Diagnostics.Session2 event.

diff --git a/Engine/Source/Runtime/Core/Private/ProfilingDebugging/TraceAuxiliary.cpp b/Engine/Source/Runtime/Core/Private/ProfilingDebugging/TraceAuxiliary.cpp
--- a/Engine/Source/Runtime/Core/Private/ProfilingDebugging/TraceAuxiliary.cpp
+++ b/Engine/Source/Runtime/Core/Private/ProfilingDebugging/TraceAuxiliary.cpp
@@ -26,10 +26,8 @@
 	#define  UE_TRACEAUX_CON_CMDS WITH_UNREAL_TRACE_CONSOLE_CMDS
 #endif
 
-#if UE_TRACE_MINIMAL_ENABLED && !UE_TRACE_ENABLED
-	#define UE_TRACEAUX_FULL 0
-#elif UE_TRACE_ENABLED
-	#define UE_TRACEAUX_FULL 1
+#if !defined(UE_TRACEAUX_FULL)
+	#define UE_TRACEAUX_FULL UE_TRACE_ENABLED
 #endif
 
 #if !defined(UE_TRACE_SERVER_LAUNCH_ENABLED)
@@ -81,7 +79,6 @@
 #include "Containers/Map.h"
 #include "Containers/StringView.h"
 #include "Containers/UnrealString.h"
-#include "Containers/UnrealString.h"
 #include "CoreGlobals.h"
 #include "HAL/FileManager.h"
 #include "HAL/IConsoleManager.h"
@@ -108,7 +105,7 @@
 
 
 ////////////////////////////////////////////////////////////////////////////////
-/// Stats and logging 
+/// Stats and logging
 ////////////////////////////////////////////////////////////////////////////////
 
 CSV_DEFINE_CATEGORY(Trace, true);
@@ -162,7 +159,7 @@
 #endif
 
 ////////////////////////////////////////////////////////////////////////////////
-/// Types and constants 
+/// Types and constants
 ////////////////////////////////////////////////////////////////////////////////
 
 enum class EWorkerThreadConfig : uint8 {
@@ -187,7 +184,7 @@
  * * Command line arguments for forking
  * * Misc state (like worker thread)
  */
-class FTraceAuxiliaryImpl  
+class FTraceAuxiliaryImpl
 {
 public:
 	FString GetDest() const;
@@ -780,10 +777,10 @@
 }
 
 ////////////////////////////////////////////////////////////////////////////////
-bool FTraceAuxiliaryImpl::HasCommandlineChannels() const 
-{ 
+bool FTraceAuxiliaryImpl::HasCommandlineChannels() const
+{
 #if UE_TRACEAUX_CMD_ARGS
-	return !CommandlineChannels.IsEmpty(); 
+	return !CommandlineChannels.IsEmpty();
 #else
 	return false;
 #endif
@@ -1690,20 +1687,5 @@
 
 
 ////////////////////////////////////////////////////////////////////////////////
-UE_TRACE_EVENT_BEGIN(Diagnostics, Session2, NoSync|Important)
-	UE_TRACE_EVENT_FIELD(UE::Trace::AnsiString, Platform)
-	UE_TRACE_EVENT_FIELD(UE::Trace::AnsiString, AppName)
-	UE_TRACE_EVENT_FIELD(UE::Trace::WideString, ProjectName)
-	UE_TRACE_EVENT_FIELD(UE::Trace::WideString, CommandLine)
-	UE_TRACE_EVENT_FIELD(UE::Trace::WideString, Branch)
-	UE_TRACE_EVENT_FIELD(UE::Trace::WideString, BuildVersion)
-	UE_TRACE_EVENT_FIELD(uint32, Changelist)
-	UE_TRACE_EVENT_FIELD(uint8, ConfigurationType)
-	UE_TRACE_EVENT_FIELD(uint8, TargetType)
-	UE_TRACE_EVENT_FIELD(uint32[], InstanceId)
-	UE_TRACE_EVENT_FIELD(UE::Trace::AnsiString, VFSPaths)
-UE_TRACE_EVENT_END()
-
-////////////////////////////////////////////////////////////////////////////////
 static bool StartFromCommandlineArguments(const TCHAR* CommandLine)
 {
@@ -1965,5 +1947,95 @@
 #endif // UE_TRACE_ENABLED
 
 ////////////////////////////////////////////////////////////////////////////////
+UE_TRACE_EVENT_BEGIN(Diagnostics, Session2, NoSync|Important)
+	UE_TRACE_EVENT_FIELD(uint8, ConfigurationType)
+	UE_TRACE_EVENT_FIELD(uint8, TargetType)
+	UE_TRACE_EVENT_FIELD(uint32, Changelist)
+	UE_TRACE_EVENT_FIELD(uint32[], InstanceId) // added in UE 5.5
+	UE_TRACE_EVENT_FIELD(UE::Trace::AnsiString, Platform)
+	UE_TRACE_EVENT_FIELD(UE::Trace::AnsiString, AppName)
+	UE_TRACE_EVENT_FIELD(UE::Trace::WideString, ProjectName)
+	UE_TRACE_EVENT_FIELD(UE::Trace::WideString, Branch)
+	UE_TRACE_EVENT_FIELD(UE::Trace::WideString, BuildVersion)
+	UE_TRACE_EVENT_FIELD(UE::Trace::WideString, EngineVersion) // added in UE 5.8
+	UE_TRACE_EVENT_FIELD(UE::Trace::WideString, CommandLine)
+	UE_TRACE_EVENT_FIELD(UE::Trace::AnsiString, VFSPaths) // added in UE 5.7
+UE_TRACE_EVENT_END()
+
+////////////////////////////////////////////////////////////////////////////////
+void OutputDiagnostics(const TCHAR* CommandLine)
+{
+	constexpr int InstanceIdSize = 4;
+	FGuid InstanceGuid = FApp::GetInstanceId();
+	uint32 InstanceId[InstanceIdSize];
+	for (int Index = 0; Index < InstanceIdSize; ++Index)
+	{
+		InstanceId[Index] = InstanceGuid[Index];
+	}
+
+	const ANSICHAR* Platform = UE_STRINGIZE(UBT_COMPILED_PLATFORM);
+	constexpr uint32 PlatformLen = UE_ARRAY_COUNT(UE_STRINGIZE(UBT_COMPILED_PLATFORM)) - 1;
+
+	const ANSICHAR* AppName = UE_APP_NAME;
+	constexpr uint32 AppNameLen = UE_ARRAY_COUNT(UE_APP_NAME) - 1;
+
+	const TCHAR* ProjectName = TEXT("");
+#if IS_MONOLITHIC && !IS_PROGRAM
+	if (FApp::HasProjectName())
+	{
+		ProjectName = FApp::GetProjectName();
+	}
+#endif
+	const uint32 ProjectNameLen = FCString::Strlen(ProjectName);
+
+	const TCHAR* BuildVersion = BuildSettings::GetBuildVersion();
+	const uint32 BuildVersionLen = FCString::Strlen(BuildVersion);
+
+	const TCHAR* BranchName = BuildSettings::GetBranchName();
+	const uint32 BranchNameLen = FCString::Strlen(BranchName);
+
+	const TCHAR* EngineVersion = BuildSettings::GetEngineVersionString();
+	const uint32 EngineVersionLen = FCString::Strlen(EngineVersion);
+
+#if UE_TRACEAUX_FULL
+	const uint32 CommandLineLen = FCString::Strlen(CommandLine);
+
+	const char* VFSPaths = BuildSettings::GetVfsPaths();
+	const uint32 VFSPathsLen = FCStringAnsi::Strlen(VFSPaths);
+#endif
+
+	const uint32 DataSize = 0
+		+ (InstanceIdSize * sizeof(uint32))
+		+ (PlatformLen * sizeof(ANSICHAR))
+		+ (AppNameLen * sizeof(ANSICHAR))
+		+ (ProjectNameLen * sizeof(TCHAR))
+		+ (BuildVersionLen * sizeof(TCHAR))
+		+ (BranchNameLen * sizeof(TCHAR))
+		+ (EngineVersionLen * sizeof(TCHAR))
+#if UE_TRACEAUX_FULL
+		+ (CommandLineLen * sizeof(TCHAR))
+		+ (VFSPathsLen * sizeof(ANSICHAR))
+#endif
+		;
+
+	UE_TRACE_LOG(Diagnostics, Session2, UE::Trace::TraceLogChannel, DataSize)
+		<< Session2.ConfigurationType(uint8(FApp::GetBuildConfiguration()))
+		<< Session2.TargetType(uint8(FApp::GetBuildTargetType()))
+		<< Session2.Changelist(BuildSettings::GetCurrentChangelist())
+		<< Session2.InstanceId(InstanceId, InstanceIdSize)
+		<< Session2.Platform(Platform, PlatformLen)
+		<< Session2.AppName(AppName, AppNameLen)
+		<< Session2.ProjectName(ProjectName, ProjectNameLen)
+		<< Session2.BuildVersion(BuildVersion, BuildVersionLen)
+		<< Session2.Branch(BranchName, BranchNameLen)
+		<< Session2.EngineVersion(EngineVersion, EngineVersionLen)
+#if UE_TRACEAUX_FULL
+		<< Session2.CommandLine(CommandLine, CommandLineLen)
+		<< Session2.VFSPaths(VFSPaths, VFSPathsLen)
+#endif
+		;
+}
+
+////////////////////////////////////////////////////////////////////////////////
 void FTraceAuxiliary::Initialize(const TCHAR* CommandLine)
 {
@@ -2010,61 +2082,9 @@
 
 	UE_LOG(LogTrace, Log, TEXT("Initializing trace..."));
 
-	const TCHAR* AppName = TEXT(UE_APP_NAME);
-	const TCHAR* ProjectName = TEXT("");
-#if IS_MONOLITHIC && !IS_PROGRAM
-	if (FApp::HasProjectName())
-	{
-		ProjectName = FApp::GetProjectName();
-	}
-#endif
-
-	constexpr int InstanceIdSize = 4;
-
 	// Trace out information about this session. This is done before initialization,
 	// so that it is always sent (all channels are enabled prior to initialization).
-	const TCHAR* BranchName = BuildSettings::GetBranchName();
-	const TCHAR* BuildVersion = BuildSettings::GetBuildVersion();
-	constexpr uint32 PlatformLen = UE_ARRAY_COUNT(UE_STRINGIZE(UBT_COMPILED_PLATFORM)) - 1;
-	const uint32 AppNameLen = FCString::Strlen(AppName);
-	const uint32 ProjectNameLen = FCString::Strlen(ProjectName);
-	const uint32 CommandLineLen = FCString::Strlen(CommandLine);
-	const uint32 BranchNameLen = FCString::Strlen(BranchName);
-	const uint32 BuildVersionLen = FCString::Strlen(BuildVersion);
-	const char* VFSPaths = BuildSettings::GetVfsPaths();
-	const uint32 VFSPathsLen = FCStringAnsi::Strlen(VFSPaths);
-	const uint32 DataSize =
-		(PlatformLen * sizeof(ANSICHAR)) +
-		(AppNameLen * sizeof(ANSICHAR)) +
-		(ProjectNameLen * sizeof(TCHAR)) +
-		(CommandLineLen * sizeof(TCHAR)) +
-		(BranchNameLen * sizeof(TCHAR)) +
-		(BuildVersionLen * sizeof(TCHAR)) +
-		(InstanceIdSize * sizeof(uint32)) +
-		(VFSPathsLen * sizeof(ANSICHAR));
+	OutputDiagnostics(CommandLine);
 
-	FGuid InstanceGuid = FApp::GetInstanceId();
-	uint32 InstanceId[InstanceIdSize];
-	for (int Index = 0; Index < InstanceIdSize; ++Index)
-	{
-		InstanceId[Index] = InstanceGuid[Index];
-	}
-
-	UE_TRACE_LOG(Diagnostics, Session2, UE::Trace::TraceLogChannel, DataSize)
-		<< Session2.Platform(UE_STRINGIZE(UBT_COMPILED_PLATFORM), PlatformLen)
-		<< Session2.AppName(AppName, AppNameLen)
-		<< Session2.ProjectName(ProjectName, ProjectNameLen)
-		<< Session2.Branch(BranchName, BranchNameLen)
-		<< Session2.BuildVersion(BuildVersion, BuildVersionLen)
-		<< Session2.Changelist(BuildSettings::GetCurrentChangelist())
-		<< Session2.ConfigurationType(uint8(FApp::GetBuildConfiguration()))
-		<< Session2.TargetType(uint8(FApp::GetBuildTargetType()))
-		<< Session2.InstanceId(InstanceId, InstanceIdSize)
-#if UE_TRACEAUX_FULL
-		<< Session2.CommandLine(CommandLine, CommandLineLen)
-		<< Session2.VFSPaths(VFSPaths, VFSPathsLen)
-#endif
-		;
-
 	// Attempt to send trace data somewhere from the command line. It perhaps
 	// seems odd to do this before initializing Trace, but it is done this way
