[TraceAux] Fixed NonUnity compile error.

#rnx

diff --git a/Engine/Source/Runtime/Core/Private/ProfilingDebugging/TraceAuxiliary.cpp b/Engine/Source/Runtime/Core/Private/ProfilingDebugging/TraceAuxiliary.cpp
--- a/Engine/Source/Runtime/Core/Private/ProfilingDebugging/TraceAuxiliary.cpp
+++ b/Engine/Source/Runtime/Core/Private/ProfilingDebugging/TraceAuxiliary.cpp
@@ -9,7 +9,6 @@
 #include "ProfilingDebugging/TagTrace.h"
 #include "Trace/Trace.h"
 
-
 ////////////////////////////////////////////////////////////////////////////////
 /// Configuration
 ////////////////////////////////////////////////////////////////////////////////
@@ -90,6 +89,7 @@
 #include "Misc/CoreDelegates.h"
 #include "Misc/DateTime.h"
 #include "Misc/Fork.h"
+#include "Misc/Guid.h"
 #include "Misc/PathViews.h"
 #include "Misc/Paths.h"
 #include "Misc/ScopeRWLock.h"
@@ -103,7 +103,6 @@
 #include "Templates/UnrealTemplate.h"
 #include "Trace/Trace.inl"
 
-
 ////////////////////////////////////////////////////////////////////////////////
 /// Stats and logging
 ////////////////////////////////////////////////////////////////////////////////
@@ -174,7 +173,6 @@ const FTraceAuxiliary::FChannelPreset GDefaultChannels(TEXT("Default"), TEXT("cp
 const FTraceAuxiliary::FChannelPreset GMemoryChannels(TEXT("Memory"), TEXT("memtag,memalloc,callstack,module"), true);
 const FTraceAuxiliary::FChannelPreset GMemoryLightChannels(TEXT("Memory_Light"), TEXT("memtag,memalloc"), true);
 
-
 ////////////////////////////////////////////////////////////////////////////////
 /**
  * FTraceAuxiliaryImpl is an implementation of the FTraceAuxiliary public interface.
@@ -340,7 +338,6 @@ void FTraceAuxiliaryImpl::EnableChannels(TConstArrayView<uint32> ChannelIds, TMa
 	}
 }
 
-
 ////////////////////////////////////////////////////////////////////////////////
 void FTraceAuxiliaryImpl::DisableChannels(const TCHAR* ChannelList, const FTraceAuxiliary::FLogCategoryAlias& LogCategory)
 {
@@ -372,7 +369,6 @@ void FTraceAuxiliaryImpl::DisableChannels(TConstArrayView<uint32> ChannelIds, TM
 }
 
 ////////////////////////////////////////////////////////////////////////////////
-
 /**
  * Parses channel tokens from a "chan1, chan2, ...." string and expands presets
  * to an output array.
@@ -1685,7 +1681,6 @@ static FAutoConsoleCommand TraceRegionEndCmd(
 
 #endif // UE_TRACE_ENABLED
 
-
 ////////////////////////////////////////////////////////////////////////////////
 static bool StartFromCommandlineArguments(const TCHAR* CommandLine)
 {
@@ -1763,12 +1758,14 @@ static bool StartFromCommandlineArguments(const TCHAR* CommandLine)
 		return false;
 	}
 
-
 	// Finally start tracing to the requested connection
 	return FTraceAuxiliary::Start(Type, Target, *Channels, &Opts);
-#else
+
+#else // UE_TRACE_ENABLED
+
 	return false;
-#endif
+
+#endif // UE_TRACE_ENABLED
 }
 
 ////////////////////////////////////////////////////////////////////////////////
@@ -1833,7 +1830,7 @@ bool FTraceAuxiliary::Start(EConnectionType Type, const TCHAR* Target, const TCH
 	uint16 SendFlags = (Options && Options->bExcludeTail) ? UE::Trace::FSendFlags::ExcludeTail : 0;
 
 	return GTraceAuxiliary.Connect(Type, Target, LogCategory, SendFlags);
-#else
+#else // UE_TRACE_MINIMAL_ENABLED
 	return false;
 #endif
 }
@@ -1852,7 +1849,7 @@ bool FTraceAuxiliary::Relay(UPTRINT Handle, UE::Trace::IoWriteFunc WriteFunc, UE
 	uint16 SendFlags = (Options && Options->bExcludeTail) ? UE::Trace::FSendFlags::ExcludeTail : 0;
 
 	return GTraceAuxiliary.ConnectRelay(Handle, WriteFunc, CloseFunc, SendFlags);
-#else
+#else // UE_TRACE_MINIMAL_ENABLED
 	return false;
 #endif
 }
@@ -1944,9 +1941,10 @@ static void TraceAuxiliaryAddPostForkCallback(const TCHAR* CommandLine)
 		}
 	});
 }
-#endif // UE_TRACE_ENABLED
+#endif // UE_TRACE_MINIMAL_ENABLED
 
 ////////////////////////////////////////////////////////////////////////////////
+#if UE_TRACE_MINIMAL_ENABLED
 UE_TRACE_EVENT_BEGIN(Diagnostics, Session2, NoSync|Important)
 	UE_TRACE_EVENT_FIELD(uint8, ConfigurationType)
 	UE_TRACE_EVENT_FIELD(uint8, TargetType)
@@ -1961,9 +1959,11 @@ UE_TRACE_EVENT_BEGIN(Diagnostics, Session2, NoSync|Important)
 	UE_TRACE_EVENT_FIELD(UE::Trace::WideString, CommandLine)
 	UE_TRACE_EVENT_FIELD(UE::Trace::AnsiString, VFSPaths) // added in UE 5.7
 UE_TRACE_EVENT_END()
+#endif // UE_TRACE_MINIMAL_ENABLED
 
 ////////////////////////////////////////////////////////////////////////////////
-void OutputDiagnostics(const TCHAR* CommandLine)
+#if UE_TRACE_MINIMAL_ENABLED
+static void OutputDiagnostics(const TCHAR* CommandLine)
 {
 	constexpr int InstanceIdSize = 4;
 	FGuid InstanceGuid = FApp::GetInstanceId();
@@ -2035,6 +2035,7 @@ void OutputDiagnostics(const TCHAR* CommandLine)
 #endif
 		;
 }
+#endif // UE_TRACE_MINIMAL_ENABLED
 
 ////////////////////////////////////////////////////////////////////////////////
 void FTraceAuxiliary::Initialize(const TCHAR* CommandLine)
