New -LogFileName UE argument to allow specifying just the filename portion of the log file and to avoid conflicting with -Log which is used as both a flag (for showing the log) and a parameter (for log file path).

diff --git a/Engine/Source/Runtime/Core/Private/GenericPlatform/GenericPlatformOutputDevices.cpp b/Engine/Source/Runtime/Core/Private/GenericPlatform/GenericPlatformOutputDevices.cpp
--- a/Engine/Source/Runtime/Core/Private/GenericPlatform/GenericPlatformOutputDevices.cpp
+++ b/Engine/Source/Runtime/Core/Private/GenericPlatform/GenericPlatformOutputDevices.cpp
@@ -81,7 +81,8 @@ FString FGenericPlatformOutputDevices::GetAbsoluteLogFilename()
 		FCString::Strncpy(CachedAbsoluteFilename, *FPaths::ProjectLogDir(), UE_ARRAY_COUNT(CachedAbsoluteFilename));
 		FString LogFilename;
 		const bool bShouldStopOnSeparator = false;
-		if (!FParse::Value(FCommandLine::Get(), TEXT("LOG="), LogFilename, bShouldStopOnSeparator))
+		if (!FParse::Value(FCommandLine::Get(), TEXT("LOG="), LogFilename, bShouldStopOnSeparator)
+			&& !FParse::Value(FCommandLine::Get(), TEXT("LogFileName="), LogFilename, bShouldStopOnSeparator))
 		{
 			if (FParse::Value(FCommandLine::Get(), TEXT("ABSLOG="), LogFilename, bShouldStopOnSeparator))
 			{
