Add ability to indent TokenizedMessage to allow for semi-visual grouping within SMessageLog widget

diff --git a/Engine/Source/Developer/MessageLog/Private/UserInterface/SMessageLogMessageListRow.cpp b/Engine/Source/Developer/MessageLog/Private/UserInterface/SMessageLogMessageListRow.cpp
--- a/Engine/Source/Developer/MessageLog/Private/UserInterface/SMessageLogMessageListRow.cpp
+++ b/Engine/Source/Developer/MessageLog/Private/UserInterface/SMessageLogMessageListRow.cpp
@@ -112,7 +112,7 @@ TSharedRef<SWidget> SMessageLogMessageListRow::GenerateWidget()
 		.VAlign(VAlign_Top)
 		[
 			SNew(SBox)
-			.Padding(2.0f)
+			.Padding((Message->GetIndentationLevel() * 20.f) + 2.f, 2.f, 2.f, 2.f)
 			[
 				(SeverityImageName == NAME_None)
 				? SNullWidget::NullWidget
diff --git a/Engine/Source/Runtime/Core/Private/Logging/TokenizedMessage.cpp b/Engine/Source/Runtime/Core/Private/Logging/TokenizedMessage.cpp
--- a/Engine/Source/Runtime/Core/Private/Logging/TokenizedMessage.cpp
+++ b/Engine/Source/Runtime/Core/Private/Logging/TokenizedMessage.cpp
@@ -118,6 +118,16 @@ void FTokenizedMessage::SetIdentifier(FName InIdentifier)
 	Identifier = InIdentifier;
 }
 
+int32 FTokenizedMessage::GetIndentationLevel() const
+{
+	return IndentationLevel;
+}
+
+void FTokenizedMessage::SetIndentationLevel(int32 Level)
+{
+	IndentationLevel = FMath::Max(Level, 0);
+}
+
 TSharedRef<FTokenizedMessage> FTokenizedMessage::AddToken( const TSharedRef<IMessageToken>& InToken )
 {
 	MessageTokens.Add( InToken );
diff --git a/Engine/Source/Runtime/Core/Public/Logging/TokenizedMessage.h b/Engine/Source/Runtime/Core/Public/Logging/TokenizedMessage.h
--- a/Engine/Source/Runtime/Core/Public/Logging/TokenizedMessage.h
+++ b/Engine/Source/Runtime/Core/Public/Logging/TokenizedMessage.h
@@ -224,6 +224,12 @@ public:
 	
 	/** Assigns Identifier for the message to the provided name */
 	CORE_API void SetIdentifier(FName InIdentifier);
+	
+	/** @returns Identation level for the message, defaults to 0 */
+	CORE_API int32 GetIndentationLevel() const;
+
+	/** Sets Identation level for the message, used for padding message within MessageLog widget */
+	CORE_API void SetIndentationLevel(int32 Level);
 
 	/** Private constructor - we want to only create these structures as shared references via Create() */
 	explicit FTokenizedMessage(FPrivateToken)
@@ -247,6 +253,9 @@ private:
 
 	/** Identifier for the message */
 	FName Identifier;
+	
+	/** Requested indentation level for this message, used for padding message visually */
+	int32 IndentationLevel = 0;
 };
 
 /** Basic message token with a localized text payload */
