Trivial 'Identation' typo fix

diff --git a/Engine/Source/Runtime/Core/Public/Logging/TokenizedMessage.h b/Engine/Source/Runtime/Core/Public/Logging/TokenizedMessage.h
--- a/Engine/Source/Runtime/Core/Public/Logging/TokenizedMessage.h
+++ b/Engine/Source/Runtime/Core/Public/Logging/TokenizedMessage.h
@@ -225,10 +225,10 @@ public:
 	/** Assigns Identifier for the message to the provided name */
 	CORE_API void SetIdentifier(FName InIdentifier);
 	
-	/** @returns Identation level for the message, defaults to 0 */
+	/** @returns Indentation level for the message, defaults to 0 */
 	CORE_API int32 GetIndentationLevel() const;
 
-	/** Sets Identation level for the message, used for padding message within MessageLog widget */
+	/** Sets Indentation level for the message, used for padding message within MessageLog widget */
 	CORE_API void SetIndentationLevel(int32 Level);
 
 	/** Private constructor - we want to only create these structures as shared references via Create() */
