Added link to bug report for TTuple reference workaround.

#rnx

diff --git a/Engine/Source/Runtime/Core/Public/Templates/Tuple.h b/Engine/Source/Runtime/Core/Public/Templates/Tuple.h
--- a/Engine/Source/Runtime/Core/Public/Templates/Tuple.h
+++ b/Engine/Source/Runtime/Core/Public/Templates/Tuple.h
@@ -16,6 +16,8 @@
 
 // MSVC generates bad codegen when getting a reference element from an rvalue tuple.  So, if the element type is
 // an rvalue reference, get it via an lvalue - it will be the same reference value in either case.
+//
+// The bug is reported here: https://developercommunity.visualstudio.com/t/Bad-codegen-when-a-function-returns-an-a/10984053
 #if defined(_MSC_VER) && !defined(__clang__)
 	#define UE_TUPLE_REFERENCE_WORKAROUND 1
 #else
