Refactor UE_SCOPED_TIMER to begin timing after it logs the start message.
This correctly isolate the measurement to the scoped code block and exclude the logging overhead which is several hundred milliseconds.

diff --git a/Engine/Source/Runtime/Core/Public/ProfilingDebugging/ScopedTimers.h b/Engine/Source/Runtime/Core/Public/ProfilingDebugging/ScopedTimers.h
--- a/Engine/Source/Runtime/Core/Public/ProfilingDebugging/ScopedTimers.h
+++ b/Engine/Source/Runtime/Core/Public/ProfilingDebugging/ScopedTimers.h
@@ -283,10 +283,11 @@ public:
 	explicit FScopedDurationTimeCustomLogger(const TCHAR* InTitle, double& InTotalTime, Func InLogFunc)
 		: Title(InTitle)
 		, LogFunc(InLogFunc)
-		, LocalTime(-FPlatformTime::Seconds())
 		, TotalTime(InTotalTime)
 	{
 		LogFunc(*FString::Printf(TEXT("%s started..."), InTitle));
+
+		LocalTime = -FPlatformTime::Seconds();
 	}
 
 	~FScopedDurationTimeCustomLogger()
