Added ForwardLike.

diff --git a/Engine/Source/Runtime/Core/Private/Templates/ForwardLike.cpp b/Engine/Source/Runtime/Core/Private/Templates/ForwardLike.cpp
new file mode 100644
--- /dev/null
+++ b/Engine/Source/Runtime/Core/Private/Templates/ForwardLike.cpp
@@ -0,0 +1,232 @@
+// Copyright Epic Games, Inc. All Rights Reserved.
+
+#include <type_traits>
+#include <utility>
+#include "Templates/UnrealTemplate.h"
+
+// Check that our ForwardLike has the same semantics as the example on https://en.cppreference.com/w/cpp/utility/forward_like.html
+namespace UE::Tests
+{
+	template <typename T>
+	constexpr std::add_volatile_t<T>& as_volatile(T& t) noexcept
+	{
+		return t;
+	}
+
+	template <class T, class U>
+	constexpr auto&& forward_like(U&& x) noexcept
+	{
+		// std::forward_like doesn't support volatile, so add it here.
+
+		constexpr bool is_adding_const = std::is_const_v<std::remove_reference_t<T>>;
+		constexpr bool is_adding_volatile = std::is_volatile_v<std::remove_reference_t<T>>;
+		if constexpr (std::is_lvalue_reference_v<T&&>)
+		{
+			if constexpr (is_adding_const)
+			{
+				if constexpr (is_adding_volatile)
+				{
+					return as_volatile(std::as_const(x));
+				}
+				else
+				{
+					return std::as_const(x);
+				}
+			}
+			else
+			{
+				if constexpr (is_adding_volatile)
+				{
+					return as_volatile(static_cast<U&>(x));
+				}
+				else
+				{
+					return static_cast<U&>(x);
+				}
+			}
+		}
+		else
+		{
+			if constexpr (is_adding_const)
+			{
+				if constexpr (is_adding_volatile)
+				{
+					return std::move(as_volatile(std::as_const(x)));
+				}
+				else
+				{
+					return std::move(std::as_const(x));
+				}
+			}
+			else
+			{
+				if constexpr (is_adding_volatile)
+				{
+					return std::move(as_volatile(x));
+				}
+				else
+				{
+					return std::move(x);
+				}
+			}
+		}
+	}
+}
+
+template <typename Type, typename ForwardLikeType>
+static inline constexpr bool bForwardLikeMatchesStd_V =
+	std::is_same_v<
+		decltype(UE::Tests::forward_like<ForwardLikeType>(std::declval<Type>())),
+		decltype(ForwardLike<ForwardLikeType>(std::declval<Type>()))
+	>;
+
+struct FForwardLikeInnerType
+{
+};
+
+struct FForwardLikeOuterType
+{
+	FForwardLikeInnerType Inner;
+};
+
+static_assert(bForwardLikeMatchesStd_V<               FForwardLikeInnerType  ,                FForwardLikeOuterType  >, "");
+static_assert(bForwardLikeMatchesStd_V<const          FForwardLikeInnerType  ,                FForwardLikeOuterType  >, "");
+static_assert(bForwardLikeMatchesStd_V<      volatile FForwardLikeInnerType  ,                FForwardLikeOuterType  >, "");
+static_assert(bForwardLikeMatchesStd_V<const volatile FForwardLikeInnerType  ,                FForwardLikeOuterType  >, "");
+static_assert(bForwardLikeMatchesStd_V<               FForwardLikeInnerType& ,                FForwardLikeOuterType  >, "");
+static_assert(bForwardLikeMatchesStd_V<const          FForwardLikeInnerType& ,                FForwardLikeOuterType  >, "");
+static_assert(bForwardLikeMatchesStd_V<      volatile FForwardLikeInnerType& ,                FForwardLikeOuterType  >, "");
+static_assert(bForwardLikeMatchesStd_V<const volatile FForwardLikeInnerType& ,                FForwardLikeOuterType  >, "");
+static_assert(bForwardLikeMatchesStd_V<               FForwardLikeInnerType&&,                FForwardLikeOuterType  >, "");
+static_assert(bForwardLikeMatchesStd_V<const          FForwardLikeInnerType&&,                FForwardLikeOuterType  >, "");
+static_assert(bForwardLikeMatchesStd_V<      volatile FForwardLikeInnerType&&,                FForwardLikeOuterType  >, "");
+static_assert(bForwardLikeMatchesStd_V<const volatile FForwardLikeInnerType&&,                FForwardLikeOuterType  >, "");
+
+static_assert(bForwardLikeMatchesStd_V<               FForwardLikeInnerType  , const          FForwardLikeOuterType  >, "");
+static_assert(bForwardLikeMatchesStd_V<const          FForwardLikeInnerType  , const          FForwardLikeOuterType  >, "");
+static_assert(bForwardLikeMatchesStd_V<      volatile FForwardLikeInnerType  , const          FForwardLikeOuterType  >, "");
+static_assert(bForwardLikeMatchesStd_V<const volatile FForwardLikeInnerType  , const          FForwardLikeOuterType  >, "");
+static_assert(bForwardLikeMatchesStd_V<               FForwardLikeInnerType& , const          FForwardLikeOuterType  >, "");
+static_assert(bForwardLikeMatchesStd_V<const          FForwardLikeInnerType& , const          FForwardLikeOuterType  >, "");
+static_assert(bForwardLikeMatchesStd_V<      volatile FForwardLikeInnerType& , const          FForwardLikeOuterType  >, "");
+static_assert(bForwardLikeMatchesStd_V<const volatile FForwardLikeInnerType& , const          FForwardLikeOuterType  >, "");
+static_assert(bForwardLikeMatchesStd_V<               FForwardLikeInnerType&&, const          FForwardLikeOuterType  >, "");
+static_assert(bForwardLikeMatchesStd_V<const          FForwardLikeInnerType&&, const          FForwardLikeOuterType  >, "");
+static_assert(bForwardLikeMatchesStd_V<      volatile FForwardLikeInnerType&&, const          FForwardLikeOuterType  >, "");
+static_assert(bForwardLikeMatchesStd_V<const volatile FForwardLikeInnerType&&, const          FForwardLikeOuterType  >, "");
+
+static_assert(bForwardLikeMatchesStd_V<               FForwardLikeInnerType  ,       volatile FForwardLikeOuterType  >, "");
+static_assert(bForwardLikeMatchesStd_V<const          FForwardLikeInnerType  ,       volatile FForwardLikeOuterType  >, "");
+static_assert(bForwardLikeMatchesStd_V<      volatile FForwardLikeInnerType  ,       volatile FForwardLikeOuterType  >, "");
+static_assert(bForwardLikeMatchesStd_V<const volatile FForwardLikeInnerType  ,       volatile FForwardLikeOuterType  >, "");
+static_assert(bForwardLikeMatchesStd_V<               FForwardLikeInnerType& ,       volatile FForwardLikeOuterType  >, "");
+static_assert(bForwardLikeMatchesStd_V<const          FForwardLikeInnerType& ,       volatile FForwardLikeOuterType  >, "");
+static_assert(bForwardLikeMatchesStd_V<      volatile FForwardLikeInnerType& ,       volatile FForwardLikeOuterType  >, "");
+static_assert(bForwardLikeMatchesStd_V<const volatile FForwardLikeInnerType& ,       volatile FForwardLikeOuterType  >, "");
+static_assert(bForwardLikeMatchesStd_V<               FForwardLikeInnerType&&,       volatile FForwardLikeOuterType  >, "");
+static_assert(bForwardLikeMatchesStd_V<const          FForwardLikeInnerType&&,       volatile FForwardLikeOuterType  >, "");
+static_assert(bForwardLikeMatchesStd_V<      volatile FForwardLikeInnerType&&,       volatile FForwardLikeOuterType  >, "");
+static_assert(bForwardLikeMatchesStd_V<const volatile FForwardLikeInnerType&&,       volatile FForwardLikeOuterType  >, "");
+
+static_assert(bForwardLikeMatchesStd_V<               FForwardLikeInnerType  , const volatile FForwardLikeOuterType  >, "");
+static_assert(bForwardLikeMatchesStd_V<const          FForwardLikeInnerType  , const volatile FForwardLikeOuterType  >, "");
+static_assert(bForwardLikeMatchesStd_V<      volatile FForwardLikeInnerType  , const volatile FForwardLikeOuterType  >, "");
+static_assert(bForwardLikeMatchesStd_V<const volatile FForwardLikeInnerType  , const volatile FForwardLikeOuterType  >, "");
+static_assert(bForwardLikeMatchesStd_V<               FForwardLikeInnerType& , const volatile FForwardLikeOuterType  >, "");
+static_assert(bForwardLikeMatchesStd_V<const          FForwardLikeInnerType& , const volatile FForwardLikeOuterType  >, "");
+static_assert(bForwardLikeMatchesStd_V<      volatile FForwardLikeInnerType& , const volatile FForwardLikeOuterType  >, "");
+static_assert(bForwardLikeMatchesStd_V<const volatile FForwardLikeInnerType& , const volatile FForwardLikeOuterType  >, "");
+static_assert(bForwardLikeMatchesStd_V<               FForwardLikeInnerType&&, const volatile FForwardLikeOuterType  >, "");
+static_assert(bForwardLikeMatchesStd_V<const          FForwardLikeInnerType&&, const volatile FForwardLikeOuterType  >, "");
+static_assert(bForwardLikeMatchesStd_V<      volatile FForwardLikeInnerType&&, const volatile FForwardLikeOuterType  >, "");
+static_assert(bForwardLikeMatchesStd_V<const volatile FForwardLikeInnerType&&, const volatile FForwardLikeOuterType  >, "");
+
+static_assert(bForwardLikeMatchesStd_V<               FForwardLikeInnerType  ,                FForwardLikeOuterType& >, "");
+static_assert(bForwardLikeMatchesStd_V<const          FForwardLikeInnerType  ,                FForwardLikeOuterType& >, "");
+static_assert(bForwardLikeMatchesStd_V<      volatile FForwardLikeInnerType  ,                FForwardLikeOuterType& >, "");
+static_assert(bForwardLikeMatchesStd_V<const volatile FForwardLikeInnerType  ,                FForwardLikeOuterType& >, "");
+static_assert(bForwardLikeMatchesStd_V<               FForwardLikeInnerType& ,                FForwardLikeOuterType& >, "");
+static_assert(bForwardLikeMatchesStd_V<const          FForwardLikeInnerType& ,                FForwardLikeOuterType& >, "");
+static_assert(bForwardLikeMatchesStd_V<      volatile FForwardLikeInnerType& ,                FForwardLikeOuterType& >, "");
+static_assert(bForwardLikeMatchesStd_V<const volatile FForwardLikeInnerType& ,                FForwardLikeOuterType& >, "");
+static_assert(bForwardLikeMatchesStd_V<               FForwardLikeInnerType&&,                FForwardLikeOuterType& >, "");
+static_assert(bForwardLikeMatchesStd_V<const          FForwardLikeInnerType&&,                FForwardLikeOuterType& >, "");
+static_assert(bForwardLikeMatchesStd_V<      volatile FForwardLikeInnerType&&,                FForwardLikeOuterType& >, "");
+static_assert(bForwardLikeMatchesStd_V<const volatile FForwardLikeInnerType&&,                FForwardLikeOuterType& >, "");
+
+static_assert(bForwardLikeMatchesStd_V<               FForwardLikeInnerType  , const          FForwardLikeOuterType& >, "");
+static_assert(bForwardLikeMatchesStd_V<const          FForwardLikeInnerType  , const          FForwardLikeOuterType& >, "");
+static_assert(bForwardLikeMatchesStd_V<      volatile FForwardLikeInnerType  , const          FForwardLikeOuterType& >, "");
+static_assert(bForwardLikeMatchesStd_V<const volatile FForwardLikeInnerType  , const          FForwardLikeOuterType& >, "");
+static_assert(bForwardLikeMatchesStd_V<               FForwardLikeInnerType& , const          FForwardLikeOuterType& >, "");
+static_assert(bForwardLikeMatchesStd_V<const          FForwardLikeInnerType& , const          FForwardLikeOuterType& >, "");
+static_assert(bForwardLikeMatchesStd_V<      volatile FForwardLikeInnerType& , const          FForwardLikeOuterType& >, "");
+static_assert(bForwardLikeMatchesStd_V<const volatile FForwardLikeInnerType& , const          FForwardLikeOuterType& >, "");
+static_assert(bForwardLikeMatchesStd_V<               FForwardLikeInnerType&&, const          FForwardLikeOuterType& >, "");
+static_assert(bForwardLikeMatchesStd_V<const          FForwardLikeInnerType&&, const          FForwardLikeOuterType& >, "");
+static_assert(bForwardLikeMatchesStd_V<      volatile FForwardLikeInnerType&&, const          FForwardLikeOuterType& >, "");
+static_assert(bForwardLikeMatchesStd_V<const volatile FForwardLikeInnerType&&, const          FForwardLikeOuterType& >, "");
+
+static_assert(bForwardLikeMatchesStd_V<               FForwardLikeInnerType  ,       volatile FForwardLikeOuterType& >, "");
+static_assert(bForwardLikeMatchesStd_V<const          FForwardLikeInnerType  ,       volatile FForwardLikeOuterType& >, "");
+static_assert(bForwardLikeMatchesStd_V<      volatile FForwardLikeInnerType  ,       volatile FForwardLikeOuterType& >, "");
+static_assert(bForwardLikeMatchesStd_V<const volatile FForwardLikeInnerType  ,       volatile FForwardLikeOuterType& >, "");
+static_assert(bForwardLikeMatchesStd_V<               FForwardLikeInnerType& ,       volatile FForwardLikeOuterType& >, "");
+static_assert(bForwardLikeMatchesStd_V<const          FForwardLikeInnerType& ,       volatile FForwardLikeOuterType& >, "");
+static_assert(bForwardLikeMatchesStd_V<      volatile FForwardLikeInnerType& ,       volatile FForwardLikeOuterType& >, "");
+static_assert(bForwardLikeMatchesStd_V<const volatile FForwardLikeInnerType& ,       volatile FForwardLikeOuterType& >, "");
+static_assert(bForwardLikeMatchesStd_V<               FForwardLikeInnerType&&,       volatile FForwardLikeOuterType& >, "");
+static_assert(bForwardLikeMatchesStd_V<const          FForwardLikeInnerType&&,       volatile FForwardLikeOuterType& >, "");
+static_assert(bForwardLikeMatchesStd_V<      volatile FForwardLikeInnerType&&,       volatile FForwardLikeOuterType& >, "");
+static_assert(bForwardLikeMatchesStd_V<const volatile FForwardLikeInnerType&&,       volatile FForwardLikeOuterType& >, "");
+
+static_assert(bForwardLikeMatchesStd_V<               FForwardLikeInnerType  ,                FForwardLikeOuterType&&>, "");
+static_assert(bForwardLikeMatchesStd_V<const          FForwardLikeInnerType  ,                FForwardLikeOuterType&&>, "");
+static_assert(bForwardLikeMatchesStd_V<      volatile FForwardLikeInnerType  ,                FForwardLikeOuterType&&>, "");
+static_assert(bForwardLikeMatchesStd_V<const volatile FForwardLikeInnerType  ,                FForwardLikeOuterType&&>, "");
+static_assert(bForwardLikeMatchesStd_V<               FForwardLikeInnerType& ,                FForwardLikeOuterType&&>, "");
+static_assert(bForwardLikeMatchesStd_V<const          FForwardLikeInnerType& ,                FForwardLikeOuterType&&>, "");
+static_assert(bForwardLikeMatchesStd_V<      volatile FForwardLikeInnerType& ,                FForwardLikeOuterType&&>, "");
+static_assert(bForwardLikeMatchesStd_V<const volatile FForwardLikeInnerType& ,                FForwardLikeOuterType&&>, "");
+static_assert(bForwardLikeMatchesStd_V<               FForwardLikeInnerType&&,                FForwardLikeOuterType&&>, "");
+static_assert(bForwardLikeMatchesStd_V<const          FForwardLikeInnerType&&,                FForwardLikeOuterType&&>, "");
+static_assert(bForwardLikeMatchesStd_V<      volatile FForwardLikeInnerType&&,                FForwardLikeOuterType&&>, "");
+static_assert(bForwardLikeMatchesStd_V<const volatile FForwardLikeInnerType&&,                FForwardLikeOuterType&&>, "");
+
+static_assert(bForwardLikeMatchesStd_V<               FForwardLikeInnerType  , const          FForwardLikeOuterType&&>, "");
+static_assert(bForwardLikeMatchesStd_V<const          FForwardLikeInnerType  , const          FForwardLikeOuterType&&>, "");
+static_assert(bForwardLikeMatchesStd_V<      volatile FForwardLikeInnerType  , const          FForwardLikeOuterType&&>, "");
+static_assert(bForwardLikeMatchesStd_V<const volatile FForwardLikeInnerType  , const          FForwardLikeOuterType&&>, "");
+static_assert(bForwardLikeMatchesStd_V<               FForwardLikeInnerType& , const          FForwardLikeOuterType&&>, "");
+static_assert(bForwardLikeMatchesStd_V<const          FForwardLikeInnerType& , const          FForwardLikeOuterType&&>, "");
+static_assert(bForwardLikeMatchesStd_V<      volatile FForwardLikeInnerType& , const          FForwardLikeOuterType&&>, "");
+static_assert(bForwardLikeMatchesStd_V<const volatile FForwardLikeInnerType& , const          FForwardLikeOuterType&&>, "");
+static_assert(bForwardLikeMatchesStd_V<               FForwardLikeInnerType&&, const          FForwardLikeOuterType&&>, "");
+static_assert(bForwardLikeMatchesStd_V<const          FForwardLikeInnerType&&, const          FForwardLikeOuterType&&>, "");
+static_assert(bForwardLikeMatchesStd_V<      volatile FForwardLikeInnerType&&, const          FForwardLikeOuterType&&>, "");
+static_assert(bForwardLikeMatchesStd_V<const volatile FForwardLikeInnerType&&, const          FForwardLikeOuterType&&>, "");
+
+static_assert(bForwardLikeMatchesStd_V<               FForwardLikeInnerType  ,       volatile FForwardLikeOuterType&&>, "");
+static_assert(bForwardLikeMatchesStd_V<const          FForwardLikeInnerType  ,       volatile FForwardLikeOuterType&&>, "");
+static_assert(bForwardLikeMatchesStd_V<      volatile FForwardLikeInnerType  ,       volatile FForwardLikeOuterType&&>, "");
+static_assert(bForwardLikeMatchesStd_V<const volatile FForwardLikeInnerType  ,       volatile FForwardLikeOuterType&&>, "");
+static_assert(bForwardLikeMatchesStd_V<               FForwardLikeInnerType& ,       volatile FForwardLikeOuterType&&>, "");
+static_assert(bForwardLikeMatchesStd_V<const          FForwardLikeInnerType& ,       volatile FForwardLikeOuterType&&>, "");
+static_assert(bForwardLikeMatchesStd_V<      volatile FForwardLikeInnerType& ,       volatile FForwardLikeOuterType&&>, "");
+static_assert(bForwardLikeMatchesStd_V<const volatile FForwardLikeInnerType& ,       volatile FForwardLikeOuterType&&>, "");
+static_assert(bForwardLikeMatchesStd_V<               FForwardLikeInnerType&&,       volatile FForwardLikeOuterType&&>, "");
+static_assert(bForwardLikeMatchesStd_V<const          FForwardLikeInnerType&&,       volatile FForwardLikeOuterType&&>, "");
+static_assert(bForwardLikeMatchesStd_V<      volatile FForwardLikeInnerType&&,       volatile FForwardLikeOuterType&&>, "");
+static_assert(bForwardLikeMatchesStd_V<const volatile FForwardLikeInnerType&&,       volatile FForwardLikeOuterType&&>, "");
+
+static_assert(bForwardLikeMatchesStd_V<               FForwardLikeInnerType  , const volatile FForwardLikeOuterType&&>, "");
+static_assert(bForwardLikeMatchesStd_V<const          FForwardLikeInnerType  , const volatile FForwardLikeOuterType&&>, "");
+static_assert(bForwardLikeMatchesStd_V<      volatile FForwardLikeInnerType  , const volatile FForwardLikeOuterType&&>, "");
+static_assert(bForwardLikeMatchesStd_V<const volatile FForwardLikeInnerType  , const volatile FForwardLikeOuterType&&>, "");
+static_assert(bForwardLikeMatchesStd_V<               FForwardLikeInnerType& , const volatile FForwardLikeOuterType&&>, "");
+static_assert(bForwardLikeMatchesStd_V<const          FForwardLikeInnerType& , const volatile FForwardLikeOuterType&&>, "");
+static_assert(bForwardLikeMatchesStd_V<      volatile FForwardLikeInnerType& , const volatile FForwardLikeOuterType&&>, "");
+static_assert(bForwardLikeMatchesStd_V<const volatile FForwardLikeInnerType& , const volatile FForwardLikeOuterType&&>, "");
+static_assert(bForwardLikeMatchesStd_V<const          FForwardLikeInnerType&&, const volatile FForwardLikeOuterType&&>, "");
+static_assert(bForwardLikeMatchesStd_V<      volatile FForwardLikeInnerType&&, const volatile FForwardLikeOuterType&&>, "");
+static_assert(bForwardLikeMatchesStd_V<const volatile FForwardLikeInnerType&&, const volatile FForwardLikeOuterType&&>, "");
diff --git a/Engine/Source/Runtime/Core/Public/Templates/UnrealTemplate.h b/Engine/Source/Runtime/Core/Public/Templates/UnrealTemplate.h
--- a/Engine/Source/Runtime/Core/Public/Templates/UnrealTemplate.h
+++ b/Engine/Source/Runtime/Core/Public/Templates/UnrealTemplate.h
@@ -768,3 +768,26 @@ UE_INTRINSIC_CAST UE_REWRITE decltype(auto) ForwardAsBase(std::remove_reference_
 {
 	return (TCopyQualifiersAndRefsFromTo_T<T&&, Base>)Obj;
 }
+
+/**
+ * ForwardLike will cast a reference to an object with the cv qualifiers and value category
+ * of another type.  This allows perfect forwarding of subobjects with the same cvrefs as
+ * the outer type.  Qualifiers and refs from each type are composed, so const will never
+ * be cast away, nor will an lvalue be cast to an rvalue.
+ *
+ * Example:
+ *
+ * template <typename T>
+ * void GetInnerA(T&& Outer)
+ * {
+ *     // Has type const TypeOfA& if Outer is a const lvalue.
+ *     // Has type TypeOfA&& if Outer is an rvalue.
+ *     // etc.
+ *     auto&& OuterA = ForwardLike<T>(Outer.A);
+ * }
+ */
+template <typename T, typename U>
+UE_INTRINSIC_CAST UE_REWRITE decltype(auto) ForwardLike(U&& Obj)
+{
+	return (TCopyQualifiersAndRefsFromTo_T<T&&, std::remove_reference_t<U>>)Obj;
+}
