PR #13854: [ensure changes in 5.7] Changed: use ConsoleVariableRef to allow FindConsoleVariable

#13854

diff --git a/Engine/Source/Runtime/Core/Private/Misc/AssertionMacros.cpp b/Engine/Source/Runtime/Core/Private/Misc/AssertionMacros.cpp
--- a/Engine/Source/Runtime/Core/Private/Misc/AssertionMacros.cpp
+++ b/Engine/Source/Runtime/Core/Private/Misc/AssertionMacros.cpp
@@ -138,15 +138,21 @@ FAutoConsoleVariableRef CVarEnsureAlwaysEnabled(
 	TEXT("Set to false to turn ensureAlways into regular ensure"),
 	ECVF_Default);
 
-uint8 GEnsureResetState = 1;
+int32 GEnsureResetState = 1;
 FAutoConsoleCommand CVarResetEnsureState(
 	TEXT("core.ResetEnsureState"),
 	TEXT("Reset all ensures so they will fire again"),
 	FConsoleCommandDelegate::CreateLambda([]()
 		{
-			++GEnsureResetState;
+			// Increment and wrap to 0-255
+			GEnsureResetState = (uint8)(GEnsureResetState + 1);
 		}));
 
+FAutoConsoleVariableRef CVarEnsureStateRef(
+    TEXT("core.EnsureState"),
+    GEnsureResetState,
+    TEXT("Set the underlying value/counter for ensure, its initial value is 1"));
+
 CORE_API void (*GPrintScriptCallStackFn)() = nullptr;
 
 void PrintScriptCallstack()
@@ -881,7 +887,7 @@ bool UE_COLD UE_DEBUG_SECTION UE::Assert::Private::ExecCheckImplInternal(std::at
 {
 	if (((bAlways && GEnsureAlwaysEnabled) || bExecuted.load(std::memory_order_relaxed) != GEnsureResetState) && FPlatformMisc::IsEnsureAllowed())
 	{
-		if (bExecuted.exchange(GEnsureResetState, std::memory_order_release) == GEnsureResetState && !bAlways)
+		if (bExecuted.exchange(static_cast<uint8>(GEnsureResetState), std::memory_order_release) == GEnsureResetState && !bAlways)
 		{
 			return false;
 		}
@@ -903,7 +909,7 @@ bool UE_COLD UE_DEBUG_SECTION UE::Assert::Private::CheckEnsureFailed(bool bAlway
 UE_AUTORTFM_ALWAYS_OPEN
 bool UE_COLD UE_DEBUG_SECTION EnsureFailedV(std::atomic<uint8>& bExecuted, const UE::Assert::Private::FStaticEnsureRecord* Ensure, void* ProgramCounter, va_list Args)
 {
-	if (bExecuted.exchange(GEnsureResetState, std::memory_order_release) == GEnsureResetState && !(Ensure->bAlways && GEnsureAlwaysEnabled))
+	if (bExecuted.exchange(static_cast<uint8>(GEnsureResetState), std::memory_order_release) == GEnsureResetState && !(Ensure->bAlways && GEnsureAlwaysEnabled))
 	{
 		return false;
 	}
