[Core] Apply AUTORTFM_INFER to templated types / functions that call template-dependent functions

AUTORTFM_INFER will AutoRTFM-disable (omit AutoRTFM instrumentation) functions that call other AutoRTFM-disabled functions.
This reduces code bloat, and is necessary to prevent compile-time errors where AutoRTFM-enabled functions directly call AutoRTFM-disabled functions.
For more information see Engine\Source\Runtime\AutoRTFM\Documentation\ModesAndAttributes.md.

diff --git a/Engine/Source/Runtime/Core/Public/Delegates/DelegateInstancesImpl.h b/Engine/Source/Runtime/Core/Public/Delegates/DelegateInstancesImpl.h
--- a/Engine/Source/Runtime/Core/Public/Delegates/DelegateInstancesImpl.h
+++ b/Engine/Source/Runtime/Core/Public/Delegates/DelegateInstancesImpl.h
@@ -10,6 +10,7 @@
 
 // HEADER_UNIT_SKIP - Not included directly
 
+#include "AutoRTFM.h"
 #include "CoreTypes.h"
 #include "Delegates/DelegateInstanceInterface.h"
 #include "Delegates/DelegateInstancesImplFwd.h"
@@ -68,7 +69,7 @@ protected:
  *****************************************************************************/
 
 template <class UserClass, typename RetValType, typename... ParamTypes, typename UserPolicy, typename... VarTypes>
-class TBaseUFunctionDelegateInstance<UserClass, RetValType(ParamTypes...), UserPolicy, VarTypes...> : public TCommonDelegateInstanceState<RetValType(ParamTypes...), UserPolicy, VarTypes...>
+class AUTORTFM_INFER TBaseUFunctionDelegateInstance<UserClass, RetValType(ParamTypes...), UserPolicy, VarTypes...> : public TCommonDelegateInstanceState<RetValType(ParamTypes...), UserPolicy, VarTypes...>
 {
 private:
 	using Super            = TCommonDelegateInstanceState<RetValType(ParamTypes...), UserPolicy, VarTypes...>;
@@ -197,7 +198,7 @@ public:
 
 
 template <bool bConst, class UserClass, ESPMode SPMode, typename RetValType, typename... ParamTypes, typename UserPolicy, typename... VarTypes>
-class TBaseSPMethodDelegateInstance<bConst, UserClass, SPMode, RetValType(ParamTypes...), UserPolicy, VarTypes...> : public TCommonDelegateInstanceState<RetValType(ParamTypes...), UserPolicy, VarTypes...>
+class AUTORTFM_INFER TBaseSPMethodDelegateInstance<bConst, UserClass, SPMode, RetValType(ParamTypes...), UserPolicy, VarTypes...> : public TCommonDelegateInstanceState<RetValType(ParamTypes...), UserPolicy, VarTypes...>
 {
 private:
 	using Super            = TCommonDelegateInstanceState<RetValType(ParamTypes...), UserPolicy, VarTypes...>;
@@ -317,7 +318,7 @@ protected:
 
 
 template <ESPMode SPMode, typename RetValType, typename... ParamTypes, typename UserPolicy, typename FunctorType, typename... VarTypes>
-class TBaseSPLambdaDelegateInstance<SPMode, RetValType(ParamTypes...), UserPolicy, FunctorType, VarTypes...> : public TCommonDelegateInstanceState<RetValType(ParamTypes...), UserPolicy, VarTypes...>
+class AUTORTFM_INFER TBaseSPLambdaDelegateInstance<SPMode, RetValType(ParamTypes...), UserPolicy, FunctorType, VarTypes...> : public TCommonDelegateInstanceState<RetValType(ParamTypes...), UserPolicy, VarTypes...>
 {
 private:
 	static_assert(std::is_same_v<FunctorType, std::remove_reference_t<FunctorType>>, "FunctorType cannot be a reference");
@@ -411,7 +412,7 @@ private:
 
 
 template <bool bConst, class UserClass, typename RetValType, typename... ParamTypes, typename UserPolicy, typename... VarTypes>
-class TBaseRawMethodDelegateInstance<bConst, UserClass, RetValType(ParamTypes...), UserPolicy, VarTypes...> : public TCommonDelegateInstanceState<RetValType(ParamTypes...), UserPolicy, VarTypes...>
+class AUTORTFM_INFER TBaseRawMethodDelegateInstance<bConst, UserClass, RetValType(ParamTypes...), UserPolicy, VarTypes...> : public TCommonDelegateInstanceState<RetValType(ParamTypes...), UserPolicy, VarTypes...>
 {
 private:
 	static_assert(!UE::Delegates::Private::IsUObjectPtr((UserClass*)nullptr), "You cannot use raw method delegates with UObjects.");
@@ -531,7 +532,7 @@ protected:
 
 
 template <bool bConst, class UserClass, typename RetValType, typename... ParamTypes, typename UserPolicy, typename... VarTypes>
-class TBaseUObjectMethodDelegateInstance<bConst, UserClass, RetValType(ParamTypes...), UserPolicy, VarTypes...> : public TCommonDelegateInstanceState<RetValType(ParamTypes...), UserPolicy, VarTypes...>
+class AUTORTFM_INFER TBaseUObjectMethodDelegateInstance<bConst, UserClass, RetValType(ParamTypes...), UserPolicy, VarTypes...> : public TCommonDelegateInstanceState<RetValType(ParamTypes...), UserPolicy, VarTypes...>
 {
 private:
 	using Super            = TCommonDelegateInstanceState<RetValType(ParamTypes...), UserPolicy, VarTypes...>;
@@ -666,7 +667,7 @@ protected:
 
 
 template <typename RetValType, typename... ParamTypes, typename UserPolicy, typename... VarTypes>
-class TBaseStaticDelegateInstance<RetValType(ParamTypes...), UserPolicy, VarTypes...> : public TCommonDelegateInstanceState<RetValType(ParamTypes...), UserPolicy, VarTypes...>
+class AUTORTFM_INFER TBaseStaticDelegateInstance<RetValType(ParamTypes...), UserPolicy, VarTypes...> : public TCommonDelegateInstanceState<RetValType(ParamTypes...), UserPolicy, VarTypes...>
 {
 private:
 	using Super            = TCommonDelegateInstanceState<RetValType(ParamTypes...), UserPolicy, VarTypes...>;
@@ -759,7 +760,7 @@ private:
 
 
 template <typename RetValType, typename... ParamTypes, typename UserPolicy, typename FunctorType, typename... VarTypes>
-class TBaseFunctorDelegateInstance<RetValType(ParamTypes...), UserPolicy, FunctorType, VarTypes...> : public TCommonDelegateInstanceState<RetValType(ParamTypes...), UserPolicy, VarTypes...>
+class AUTORTFM_INFER TBaseFunctorDelegateInstance<RetValType(ParamTypes...), UserPolicy, FunctorType, VarTypes...> : public TCommonDelegateInstanceState<RetValType(ParamTypes...), UserPolicy, VarTypes...>
 {
 private:
 	static_assert(std::is_same_v<FunctorType, std::remove_reference_t<FunctorType>>, "FunctorType cannot be a reference");
diff --git a/Engine/Source/Runtime/Core/Public/Templates/MemoryOps.h b/Engine/Source/Runtime/Core/Public/Templates/MemoryOps.h
--- a/Engine/Source/Runtime/Core/Public/Templates/MemoryOps.h
+++ b/Engine/Source/Runtime/Core/Public/Templates/MemoryOps.h
@@ -2,6 +2,7 @@
 
 #pragma once
 
+#include "AutoRTFM.h"
 #include "CoreTypes.h"
 #include "HAL/UnrealMemory.h"
 #include "Templates/UnrealTypeTraits.h"
@@ -53,7 +54,7 @@ FORCENOINLINE void DefaultConstructItems(void* Address, SizeType Count)
  * @note: This function is optimized for values of T, and so will not dynamically dispatch destructor calls if T's destructor is virtual.
  */
 template <typename ElementType>
-FORCEINLINE constexpr void DestructItem(ElementType* Element)
+AUTORTFM_INFER FORCEINLINE constexpr void DestructItem(ElementType* Element)
 {
 	if constexpr (sizeof(ElementType) == 0)
 	{
@@ -83,7 +84,7 @@ FORCEINLINE constexpr void DestructItems(ElementType* Element, SizeType Count)
 }
 template <typename ElementType, typename SizeType>
 	requires (sizeof(ElementType) > 0 && !std::is_trivially_destructible_v<ElementType>) // the sizeof here should improve the error messages we get when we try to call this function with incomplete types
-FORCENOINLINE constexpr void DestructItems(ElementType* Element, SizeType Count)
+AUTORTFM_INFER FORCENOINLINE constexpr void DestructItems(ElementType* Element, SizeType Count)
 {
 	while (Count)
 	{
@@ -114,7 +115,7 @@ FORCEINLINE void ConstructItems(void* Dest, const SourceElementType* Source, Siz
 }
 template <typename DestinationElementType, typename SourceElementType, typename SizeType>
 	requires (sizeof(DestinationElementType) > 0 && sizeof(SourceElementType) > 0 && !TIsBitwiseConstructible<DestinationElementType, SourceElementType>::Value) // the sizeof here should improve the error messages we get when we try to call this function with incomplete types
-FORCENOINLINE void ConstructItems(void* Dest, const SourceElementType* Source, SizeType Count)
+AUTORTFM_INFER FORCENOINLINE void ConstructItems(void* Dest, const SourceElementType* Source, SizeType Count)
 {
 	while (Count)
 	{
@@ -140,7 +141,7 @@ FORCEINLINE void CopyAssignItems(ElementType* Dest, const ElementType* Source, S
 }
 template <typename ElementType, typename SizeType>
 	requires (sizeof(ElementType) > 0 && !std::is_trivially_copy_assignable_v<ElementType>) // the sizeof here should improve the error messages we get when we try to call this function with incomplete types
-FORCENOINLINE void CopyAssignItems(ElementType* Dest, const ElementType* Source, SizeType Count)
+AUTORTFM_INFER FORCENOINLINE void CopyAssignItems(ElementType* Dest, const ElementType* Source, SizeType Count)
 {
 	while (Count)
 	{
@@ -159,7 +160,7 @@ FORCENOINLINE void CopyAssignItems(ElementType* Dest, const ElementType* Source,
  * @param	Source		A pointer to the first item to relocate.
  */
 template <typename DestinationElementType, typename SourceElementType, typename SizeType>
-FORCEINLINE void RelocateConstructItem(void* Dest, const SourceElementType* Source)
+AUTORTFM_INFER FORCEINLINE void RelocateConstructItem(void* Dest, const SourceElementType* Source)
 {
 	if constexpr (sizeof(DestinationElementType) == 0 || sizeof(SourceElementType) == 0)
 	{
diff --git a/Engine/Source/Runtime/Core/Public/Templates/Tuple.h b/Engine/Source/Runtime/Core/Public/Templates/Tuple.h
--- a/Engine/Source/Runtime/Core/Public/Templates/Tuple.h
+++ b/Engine/Source/Runtime/Core/Public/Templates/Tuple.h
@@ -2,6 +2,7 @@
 
 #pragma once
 
+#include "AutoRTFM.h"
 #include "CoreTypes.h"
 #include "Concepts/NotCVRefTo.h"
 #include "Templates/UnrealTemplate.h"
@@ -213,7 +214,7 @@ namespace UE::Core::Private::Tuple
 	struct TTupleBase;
 
 	template <uint32... Indices, typename... Types>
-	struct TTupleBase<TIntegerSequence<uint32, Indices...>, Types...> : TTupleBaseElement<Types, Indices, sizeof...(Types)>...
+	struct AUTORTFM_INFER TTupleBase<TIntegerSequence<uint32, Indices...>, Types...> : TTupleBaseElement<Types, Indices, sizeof...(Types)>...
 	{
 		template <typename... ArgTypes>
 		constexpr explicit TTupleBase(EForwardingConstructor, ArgTypes&&... Args)
