Removed deprecated TScriptDelegate<FWeakObjectPtr> and TMulticastScriptDelegate<FWeakObjectPtr> support.

diff --git a/Engine/Source/Runtime/Core/Public/UObject/ScriptDelegates.h b/Engine/Source/Runtime/Core/Public/UObject/ScriptDelegates.h
--- a/Engine/Source/Runtime/Core/Public/UObject/ScriptDelegates.h
+++ b/Engine/Source/Runtime/Core/Public/UObject/ScriptDelegates.h
@@ -25,47 +25,16 @@ namespace UE::Core::Private
 		using ThreadSafetyMode = InThreadSafetyMode;
 		using UnicastThreadSafetyModeForMulticasts = FNotThreadSafeNotCheckedDelegateMode;
 	};
-
-	template <>
-	struct UE_DEPRECATED(5.3, "TScriptDelegate<FWeakObjectPtr> and TMulticastScriptDelegate<FWeakObjectPtr> have been deprecated, please use FScriptDelegate or FMulticastScriptDelegate respectively.") TScriptDelegateTraits<FWeakObjectPtr>
-	{
-		// After this deprecated specialization has been removed, all of the functions inside
-		// TMulticastScriptDelegate which take OtherDummy parameters should also be removed,
-		// and also the TScriptDelegate(const TScriptDelegate<FWeakObjectPtr>&) constructor.
-
-		using WeakPtrType = FWeakObjectPtr;
-		using ThreadSafetyMode = FNotThreadSafeDelegateMode;
-		using UnicastThreadSafetyModeForMulticasts = FNotThreadSafeNotCheckedDelegateMode;
-	};
-
-	// This function only exists to allow compatibility between multicast and unicast delegate types which use an explicit FWeakObjectPtr template parameter
-	template <typename From, typename To>
-	/* UE_DEPRECATED(5.3, "Deprecated - remove after TScriptDelegateTraits<FWeakObjectPtr> is removed") */
-	inline constexpr bool BackwardCompatibilityCheck()
-	{
-		if constexpr (std::is_same_v<From, FNotThreadSafeDelegateMode>)
-		{
-			return std::is_same_v<To, FWeakObjectPtr>;
-		}
-		else if constexpr (std::is_same_v<From, FWeakObjectPtr>)
-		{
-			return std::is_same_v<To, FNotThreadSafeDelegateMode>;
-		}
-		else
-		{
-			return false;
-		}
-	}
 }
 
 /**
  * Script delegate base class.
  */
 template <typename InThreadSafetyMode>
-class TScriptDelegate : public TDelegateAccessHandlerBase<typename UE::Core::Private::TScriptDelegateTraits<InThreadSafetyMode>::ThreadSafetyMode>
+class TScriptDelegate : public TDelegateAccessHandlerBase<InThreadSafetyMode>
 {
 public:
-	using ThreadSafetyMode = typename UE::Core::Private::TScriptDelegateTraits<InThreadSafetyMode>::ThreadSafetyMode;
+	using ThreadSafetyMode = InThreadSafetyMode;
 	using WeakPtrType = typename UE::Core::Private::TScriptDelegateTraits<InThreadSafetyMode>::WeakPtrType;
 
 private:
@@ -96,19 +65,6 @@ public:
 		FunctionName = Other.FunctionName;
 	}
 
-	template <
-		typename OtherThreadSafetyMode
-		UE_REQUIRES(UE::Core::Private::BackwardCompatibilityCheck<InThreadSafetyMode, OtherThreadSafetyMode>())
-	>
-	/* UE_DEPRECATED(5.3, "Deprecated - remove after TScriptDelegateTraits<FWeakObjectPtr> is removed") */
-	TScriptDelegate(const TScriptDelegate<OtherThreadSafetyMode>& Other)
-	{
-		typename TScriptDelegate<OtherThreadSafetyMode>::FReadAccessScope OtherReadScope = Other.GetReadAccessScope();
-
-		Object = Other.Object;
-		FunctionName = Other.FunctionName;
-	}
-
 	TScriptDelegate& operator=(const TScriptDelegate& Other)
 	{
 		WeakPtrType OtherObject;
@@ -128,30 +84,6 @@ public:
 
 		return *this;
 	}
-	template <
-		typename OtherThreadSafetyMode
-		UE_REQUIRES(UE::Core::Private::BackwardCompatibilityCheck<InThreadSafetyMode, OtherThreadSafetyMode>())
-	>
-	/* UE_DEPRECATED(5.3, "Deprecated - remove after TScriptDelegateTraits<FWeakObjectPtr> is removed") */
-	TScriptDelegate& operator=(const TScriptDelegate<OtherThreadSafetyMode>& Other)
-	{
-		WeakPtrType OtherObject;
-		FName OtherFunctionName;
-
-		{
-			typename TScriptDelegate<OtherThreadSafetyMode>::FReadAccessScope OtherReadScope = Other.GetReadAccessScope();
-			OtherObject = Other.Object;
-			OtherFunctionName = Other.FunctionName;
-		}
-
-		{
-			FWriteAccessScope WriteScope = GetWriteAccessScope();
-			Object = OtherObject;
-			FunctionName = OtherFunctionName;
-		}
-
-		return *this;
-	}
 
 private:
 
@@ -308,45 +240,11 @@ public:
 
 		return bResult;
 	}
-	template <
-		typename OtherThreadSafetyMode
-		UE_REQUIRES(UE::Core::Private::BackwardCompatibilityCheck<InThreadSafetyMode, OtherThreadSafetyMode>())
-	>
-	/* UE_DEPRECATED(5.3, "Deprecated - remove after TScriptDelegateTraits<FWeakObjectPtr> is removed") */
-	inline bool operator==(const TScriptDelegate<OtherThreadSafetyMode>& Other) const
-	{
-		WeakPtrType OtherObject;
-		FName OtherFunctionName;
-
-		{
-			typename TScriptDelegate<OtherThreadSafetyMode>::FReadAccessScope OtherReadScope = Other.GetReadAccessScope();
-			OtherObject = Other.Object;
-			OtherFunctionName = Other.FunctionName;
-		}
-
-		bool bResult;
-
-		{
-			FReadAccessScope ThisReadScope = GetReadAccessScope();
-			bResult = Object == OtherObject && FunctionName == OtherFunctionName;
-		}
-
-		return bResult;
-	}
 
 	UE_FORCEINLINE_HINT bool operator!=( const TScriptDelegate& Other ) const
 	{
 		return !operator==(Other);
 	}
-	template <
-		typename OtherThreadSafetyMode
-		UE_REQUIRES(UE::Core::Private::BackwardCompatibilityCheck<InThreadSafetyMode, OtherThreadSafetyMode>())
-	>
-	/* UE_DEPRECATED(5.3, "Deprecated - remove after TScriptDelegateTraits<FWeakObjectPtr> is removed") */
-	UE_FORCEINLINE_HINT bool operator!=(const TScriptDelegate<OtherThreadSafetyMode>& Other) const
-	{
-		return !operator==(Other);
-	}
 
 	/** 
 	 * Gets the object bound to this delegate
@@ -505,7 +403,7 @@ struct TIsZeroConstructType<TScriptDelegate<ThreadSafetyMode>>
  * Script multi-cast delegate base class
  */
 template <typename InThreadSafetyMode>
-class TMulticastScriptDelegate : public TDelegateAccessHandlerBase<typename UE::Core::Private::TScriptDelegateTraits<InThreadSafetyMode>::ThreadSafetyMode>
+class TMulticastScriptDelegate : public TDelegateAccessHandlerBase<InThreadSafetyMode>
 {
 private:
 	using Super = TDelegateAccessHandlerBase<InThreadSafetyMode>;
@@ -612,24 +510,6 @@ public:
 
 		return Contains(Object, FunctionName);
 	}
-	template <
-		typename OtherThreadSafetyMode
-		UE_REQUIRES(UE::Core::Private::BackwardCompatibilityCheck<InThreadSafetyMode, OtherThreadSafetyMode>())
-	>
-	/* UE_DEPRECATED(5.3, "Deprecated - remove after TScriptDelegateTraits<FWeakObjectPtr> is removed") */
-	bool Contains(const TScriptDelegate<OtherThreadSafetyMode>& InDelegate) const
-	{
-		const UObject* Object;
-		FName FunctionName;
-
-		{
-			typename TScriptDelegate<OtherThreadSafetyMode>::FReadAccessScope OtherReadScope = InDelegate.GetReadAccessScope();
-			Object = InDelegate.Object.Get();
-			FunctionName = InDelegate.FunctionName;
-		}
-
-		return Contains(Object, FunctionName);
-	}
 
 	/**
 	 * Checks whether a function delegate is already a member of this multi-cast delegate's invocation list
@@ -666,25 +546,6 @@ public:
 			AddInternal(MoveTemp(LocalCopy));
 		}
 	}
-	template <
-		typename OtherThreadSafetyMode
-		UE_REQUIRES(UE::Core::Private::BackwardCompatibilityCheck<InThreadSafetyMode, OtherThreadSafetyMode>())
-	>
-	/* UE_DEPRECATED(5.3, "Deprecated - remove after TScriptDelegateTraits<FWeakObjectPtr> is removed") */
-	void Add(const TScriptDelegate<OtherThreadSafetyMode>& InDelegate)
-	{
-		UnicastDelegateType LocalCopy = UnicastDelegateType::CopyFrom(InDelegate);
-
-		{
-			FWriteAccessScope WriteScope = GetWriteAccessScope();
-
-			// First check for any objects that may have expired
-			CompactInvocationList();
-
-			// Add the delegate
-			AddInternal(MoveTemp(LocalCopy));
-		}
-	}
 
 	/**
 	 * Adds a function delegate to this multi-cast delegate's invocation list if a delegate with the same signature
@@ -706,25 +567,6 @@ public:
 			CompactInvocationList();
 		}
 	}
-	template <
-		typename OtherThreadSafetyMode
-		UE_REQUIRES(UE::Core::Private::BackwardCompatibilityCheck<InThreadSafetyMode, OtherThreadSafetyMode>())
-	>
-	/* UE_DEPRECATED(5.3, "Deprecated - remove after TScriptDelegateTraits<FWeakObjectPtr> is removed") */
-	void AddUnique(const TScriptDelegate<OtherThreadSafetyMode>& InDelegate)
-	{
-		UnicastDelegateType LocalCopy = UnicastDelegateType::CopyFrom(InDelegate);
-
-		{
-			FWriteAccessScope WriteScope = GetWriteAccessScope();
-
-			// Add the delegate
-			AddUniqueInternal(MoveTemp(LocalCopy));
-
-			// Then check for any objects that may have expired
-			CompactInvocationList();
-		}
-	}
 
 	/**
 	 * Removes a function from this multi-cast delegate's invocation list (performance is O(N)).  Note that the
@@ -746,25 +588,6 @@ public:
 			CompactInvocationList();
 		}
 	}
-	template <
-		typename OtherThreadSafetyMode
-		UE_REQUIRES(UE::Core::Private::BackwardCompatibilityCheck<InThreadSafetyMode, OtherThreadSafetyMode>())
-	>
-	/* UE_DEPRECATED(5.3, "Deprecated - remove after TScriptDelegateTraits<FWeakObjectPtr> is removed") */
-	void Remove(const TScriptDelegate<OtherThreadSafetyMode>& InDelegate)
-	{
-		UnicastDelegateType LocalCopy = UnicastDelegateType::CopyFrom(InDelegate);
-
-		{
-			FWriteAccessScope WriteScope = GetWriteAccessScope();
-
-			// Remove the delegate
-			RemoveInternal(LocalCopy);
-
-			// Check for any delegates that may have expired
-			CompactInvocationList();
-		}
-	}
 
 	/**
 	 * Removes a function from this multi-cast delegate's invocation list (performance is O(N)).  Note that the
