Expand CheckDelegatesLifetime to discover and remove stale delegates at module shutdown

- Remove previous lifetime checks that proved very brittle
- Add tracking of multicasts that have registered delegates in DLL-bound modules
- Remove live delegates when their module winds down and report a full error for each
- Report basic stats on usage

#rnx
#tests Tested on console with monolithic and merged modules builds

diff --git a/Engine/Source/Programs/UnrealBuildTool/Configuration/UEBuildTarget.cs b/Engine/Source/Programs/UnrealBuildTool/Configuration/UEBuildTarget.cs
--- a/Engine/Source/Programs/UnrealBuildTool/Configuration/UEBuildTarget.cs
+++ b/Engine/Source/Programs/UnrealBuildTool/Configuration/UEBuildTarget.cs
@@ -1397,6 +1397,12 @@ namespace UnrealBuildTool
 				RulesObject.bUseSharedPCHs = false;
 			}
 
+			// Force shared PCHs off with delegate lifetime, as this is tricky to find out otherwise
+			if (RulesObject.bCheckDelegateLifetime)
+			{
+				RulesObject.bUseSharedPCHs = false;
+			}
+
 			// Don't link if we're just preprocessing
 			if (RulesObject.bPreprocessOnly)
 			{
@@ -3852,6 +3858,22 @@ namespace UnrealBuildTool
 
 			return true;
 		}
+		bool CanPluginBeMerged(UEBuildPlugin plugin)
+		{
+			if (Rules.MergePluginsExcluded.Contains(plugin.Name))
+			{
+				return false;
+			}
+			else if (Rules.MergeOnlyGameFeatures)
+			{
+				DirectoryReference gameFeaturesDir = DirectoryReference.Combine(ProjectDirectory, "Plugins", "GameFeatures");
+				return plugin.File.IsUnderDirectory(gameFeaturesDir);
+			}
+			else
+			{
+				return true;
+			}
+		}
 
 		/// <summary>
 		/// Dynamically merge modules into libraries based on target settings
@@ -3947,23 +3969,6 @@ namespace UnrealBuildTool
 				}
 			}
 
-			bool CanPluginBeMerged(UEBuildPlugin plugin)
-			{
-				if (Rules.MergePluginsExcluded.Contains(plugin.Name))
-				{
-					return false;
-				}
-				else if (Rules.MergeOnlyGameFeatures)
-				{
-					DirectoryReference gameFeaturesDir = DirectoryReference.Combine(ProjectDirectory, "Plugins", "GameFeatures");
-					return plugin.File.IsUnderDirectory(gameFeaturesDir);
-				}
-				else
-				{
-					return true;
-				}
-			}
-
 			Dictionary<UEBuildPlugin, HashSet<UEBuildPlugin>> globalPluginUserMap = new();
 
 			void GetPluginDependencies(UEBuildPlugin plugin, ref HashSet<UEBuildPlugin> dependencies, bool bShouldOnlyFollowCode = false, bool bHasUsersWithSourceCode = false)
@@ -4453,6 +4458,36 @@ namespace UnrealBuildTool
 				}
 			}
 
+			if (Rules.bCheckDelegateLifetime)
+			{
+				// Add module name information for code in dynamic libraries.
+				if (Rules.bMergeModules)
+				{
+					foreach (UEBuildModuleCPP Module in Modules.Values.OfType<UEBuildModuleCPP>())
+					{
+						if (Module.Binary != null && Module.Binary.Type == UEBuildBinaryType.DynamicLinkLibrary)
+						{
+							Module.Rules.PrivateDefinitions.Add($"UE_DYNAMIC_MODULE_NAME=\"{Module.Name}\"");
+						}
+					}
+				}
+
+				// When merged modules are disabled, simulate the behavior, except we'll do it for all mergeable plugins as we don't know how the merging would have been done.
+				else
+				{
+					foreach (UEBuildPlugin Plugin in BuildPlugins!)
+					{
+						foreach (UEBuildModuleCPP Module in Plugin.Modules)
+						{
+							if (CanPluginBeMerged(Plugin))
+							{
+								Module.Rules.PrivateDefinitions.Add($"UE_DYNAMIC_MODULE_NAME=\"{Module.Name}\"");
+							}
+						}
+					}
+				}
+			}
+
 			bool anyValidationErrors = false;
 
 			// Allow each platform to run any needed validation on the plugins.
diff --git a/Engine/Source/Runtime/Core/Private/Delegates/DelegateBase.cpp b/Engine/Source/Runtime/Core/Private/Delegates/DelegateBase.cpp
--- a/Engine/Source/Runtime/Core/Private/Delegates/DelegateBase.cpp
+++ b/Engine/Source/Runtime/Core/Private/Delegates/DelegateBase.cpp
@@ -7,27 +7,159 @@
 
 DEFINE_LOG_CATEGORY_STATIC(LogDelegates, Log, All);
 
-bool FTrackedDelegateInstanceExtras::IsValid() const
+struct FDelegateRegistration
 {
-	return ModuleName == NAME_None || FModuleManager::Get().IsModuleSafeToUse(ModuleName);
+	FDelegateRegistration(FDelegateHandle InRegisteredHandle, FName InRegisteredFunctionName)
+		: RegisteredHandle(InRegisteredHandle), RegisteredFunctionName(InRegisteredFunctionName)
+	{}
+
+	FDelegateHandle RegisteredHandle;
+	FName RegisteredFunctionName;
+};
+
+using FMulticastRegistrations = TArray<FDelegateRegistration>;
+using FMulticastModuleUsers = TMap<FName, FMulticastRegistrations>;
+
+static FCriticalSection GLiveMulticastSetLock;
+static TMap<FMulticastDelegateSafetyLayer*, FMulticastModuleUsers> GLiveMulticastMap;
+static TAtomic<bool> GIsRemovingDelegates = false;
+
+static bool GEnableDelegateLifetimeChecks = !UE_BUILD_SHIPPING;
+static FAutoConsoleVariableRef EnableDelegateLifetimeChecksCVar(
+	TEXT("Delegates.EnableDelegateLifetimeChecks"),
+	GEnableDelegateLifetimeChecks,
+	TEXT("When set, enable the lifetime checking and automatic unregistration of delegates subscribed to multicast delegates."),
+	ECVF_Default);
+
+TMpscQueue<FMulticastDelegateSafetyLayer*>& GDestroyedMulticastQueue()
+{
+	static TMpscQueue<FMulticastDelegateSafetyLayer*> StaticQueue;
+	return StaticQueue;
 }
 
-bool FTrackedDelegateInstanceExtras::CheckValid() const
+FMulticastDelegateSafetyLayer::~FMulticastDelegateSafetyLayer()
 {
-	if (!IsValid())
+	if (GEnableDelegateLifetimeChecks)
 	{
-		UE_LOG(LogDelegates, Fatal, TEXT("FTrackedDelegateInstanceExtras::CheckValid: bad delegate from %s (%s)"),
-			*GetBoundFunctionName().ToString(), *ModuleName.ToString());
+		TRACE_CPUPROFILER_EVENT_SCOPE(~FMulticastDelegateSafetyLayer);
 
-		return false;
+		checkf(!GIsRemovingDelegates, TEXT("Multicast destruction during delegate lifetime checks is unsafe!"));
+		GDestroyedMulticastQueue().Enqueue(this);
 	}
-	else
+}
+
+void FMulticastDelegateSafetyLayer::RegisterDelegateInstance(FDelegateHandle DelegateHandle, FName ModuleName, FName CallsiteName)
+{
+	if (GEnableDelegateLifetimeChecks)
 	{
-		return true;
+		TRACE_CPUPROFILER_EVENT_SCOPE(RegisterDelegateInstance);
+
+		if (ModuleName != NAME_None)
+		{
+			FScopeLock Lock(&GLiveMulticastSetLock);
+			FMulticastModuleUsers& MulticastModuleUsers = GLiveMulticastMap.FindOrAdd(this);
+			FMulticastRegistrations& MulticastRegistrations = MulticastModuleUsers.FindOrAdd(ModuleName);
+			MulticastRegistrations.Emplace(FDelegateRegistration(DelegateHandle, CallsiteName));
+		}
+	}
+}
+
+static void RemoveDestroyedMulticastsFromLiveSet()
+{
+	TRACE_CPUPROFILER_EVENT_SCOPE(RemoveDestroyedMulticastsFromLiveSet);
+
+	FScopeLock Lock(&GLiveMulticastSetLock);
+
+	FMulticastDelegateSafetyLayer* MulticastToDequeue;
+	while (GDestroyedMulticastQueue().Dequeue(MulticastToDequeue))
+	{
+		GLiveMulticastMap.Remove(MulticastToDequeue);
+	}
+}
+
+void FMulticastDelegateSafetyLayer::RemoveAllDelegatesInModules(TConstArrayView<FName> ModuleNames)
+{
+	if (GEnableDelegateLifetimeChecks)
+	{
+		TRACE_CPUPROFILER_EVENT_SCOPE(RemoveAllDelegatesInModules);
+
+		int32 RemovedCount = 0;
+		int32 PendingDestroyedMulticastsCount = 0;
+
+		RemoveDestroyedMulticastsFromLiveSet();
+
+		FScopeLock Lock(&GLiveMulticastSetLock);
+
+		// First, remove all flushed multicasts. After this point, we won't check for more and assume no multicasts is destroyed in this time frame.
+		// This should be a safe assumption because almost all multicasts use the default unsafe policy, but it is checked in the destructor anyway.
+		GIsRemovingDelegates = true;
+
+		{
+			TRACE_CPUPROFILER_EVENT_SCOPE(RemoveAllDelegatesInModules/Flush);
+
+			FMulticastDelegateSafetyLayer* MulticastToDequeue;
+			while (GDestroyedMulticastQueue().Dequeue(MulticastToDequeue))
+			{
+				GLiveMulticastMap.Remove(MulticastToDequeue);
+				PendingDestroyedMulticastsCount++;
+			}
+		}
+
+		// Now find all live delegate registrations for each of those modules, report them, and suppress them if needed.
+		for (auto LiveMulticastIt = GLiveMulticastMap.CreateIterator(); LiveMulticastIt; ++LiveMulticastIt)
+		{
+			FMulticastDelegateSafetyLayer* Multicast = LiveMulticastIt.Key();
+			FMulticastModuleUsers& MulticastModuleUsers = LiveMulticastIt.Value();
+
+			for (FName ModuleName : ModuleNames)
+			{
+				if (const FMulticastRegistrations* MulticastRegistrations = MulticastModuleUsers.Find(ModuleName))
+				{
+					for (const FDelegateRegistration& DelegateRegistration : *MulticastRegistrations)
+					{
+						if (Multicast->RemoveDelegateInstanceWrapper(DelegateRegistration.RegisteredHandle))
+						{
+							const FString ReportedMessage = FString::Printf(TEXT("stale delegate in %s: %s"), *ModuleName.ToString(), *DelegateRegistration.RegisteredFunctionName.ToString());
+							UE_LOG(LogDelegates, Error, TEXT("RemoveAllDelegatesInModules: %s"), *ReportedMessage);
+							ensureAlwaysMsgf(false, TEXT("%s"), *ReportedMessage);
+							RemovedCount++;
+						}
+					}
+
+					MulticastModuleUsers.Remove(ModuleName);
+				}
+			}
+
+			if (MulticastModuleUsers.Num() == 0)
+			{
+				LiveMulticastIt.RemoveCurrent();
+			}
+		}
+
+		GIsRemovingDelegates = false;
+
+#if !UE_BUILD_SHIPPING
+
+		int32 TrackedDelegateInstances = 0;
+		size_t AllocatedSize = GLiveMulticastMap.GetAllocatedSize();
+		for (const TPair<FMulticastDelegateSafetyLayer*, FMulticastModuleUsers>& LiveMulticast : GLiveMulticastMap)
+		{
+			AllocatedSize += LiveMulticast.Value.GetAllocatedSize();
+			for (const TPair<FName, FMulticastRegistrations>& MulticastModuleUser : LiveMulticast.Value)
+			{
+				TrackedDelegateInstances += MulticastModuleUser.Value.Num();
+				AllocatedSize += MulticastModuleUser.Value.GetAllocatedSize();
+			}
+		}
+
+		UE_LOG(LogDelegates, Verbose, TEXT("RemoveAllDelegatesInModules: %d multicasts, %d in queue, %d instances, %dKB used for tracking"),
+			GLiveMulticastMap.Num(), PendingDestroyedMulticastsCount, TrackedDelegateInstances, AllocatedSize / 1024);
+
+#endif // !UE_BUILD_SHIPPING
 	}
 }
 
-#endif //UE_DELEGATE_CHECK_LIFETIME
+#endif // UE_DELEGATE_CHECK_LIFETIME
 
 void* UE::Core::Private::DelegateAllocate(size_t Size, FDelegateAllocation& Allocation)
 {
diff --git a/Engine/Source/Runtime/Core/Private/Modules/ModuleManager.cpp b/Engine/Source/Runtime/Core/Private/Modules/ModuleManager.cpp
--- a/Engine/Source/Runtime/Core/Private/Modules/ModuleManager.cpp
+++ b/Engine/Source/Runtime/Core/Private/Modules/ModuleManager.cpp
@@ -555,15 +555,13 @@ void* FModuleManager::InternalLoadLibrary(FName ModuleName, const FString& Libra
 #if UE_MERGED_MODULES
 
 	// Stop any delay unload for that library, and match the handle for the code module to the library path.
+	DelayUnloadLibraries.Remove(LibraryToLoad);
 	if (void** ExistingHandle = LibraryHandles.Find(LibraryToLoad))
 	{
 		UE_LOG(LogModuleManager, Verbose, TEXT("InternalLoadLibrary: already loaded library '%s'"), *LibraryToLoad);
 
-		DelayUnloadLibraries.Remove(LibraryToLoad);
-
 		return ExistingHandle;
 	}
-	DelayUnloadLibraries.Remove(LibraryToLoad);
 
 	// The plugin manager normally handles dependencies, but merged libraries introduce indirect dependencies.
 	// If we load plugin A that's merged with B, and B depends on C in another library...
@@ -591,7 +589,6 @@ void* FModuleManager::InternalLoadLibrary(FName ModuleName, const FString& Libra
 #if UE_MERGED_MODULES
 
 	// Stop any delay unload for that library, and match the handle for the code module to the library path.
-	DelayUnloadLibraries.Remove(LibraryToLoad);
 	LibraryHandles.Add(LibraryToLoad, Handle);
 
 #endif // UE_MERGED_MODULES
@@ -647,13 +644,18 @@ void FModuleManager::InternalFreeLibrary(FName ModuleName, void* Handle)
 	}
 
 	// Safety feature: allow disabling of merged library unloading, either wholly, or for specific libraries
-	const auto CanUnloadLibrary = [&LibraryToUnload]()
+	const auto CanUnloadLibrary = [&LibraryToUnload, this]()
 	{
 		if (!GEnableMergedLibraryUnloading)
 		{
 			return false;
 		}
 
+		if (DelayUnloadLibraries.Contains(LibraryToUnload))
+		{
+			return false;
+		}
+
 		if (GPersistentMergedLibraries.Len())
 		{
 			TArray<FString> PersistentLibraries;
@@ -890,6 +892,10 @@ void FModuleManager::OnObjectCleanup()
 			{
 				UE_LOG(LogModuleManager, Log, TEXT("OnObjectCleanup: unloading dynamic library '%s'"), *LibraryAndModulesIterator.Key());
 
+#if UE_DELEGATE_CHECK_LIFETIME
+				FMulticastDelegateSafetyLayer::RemoveAllDelegatesInModules(LibraryAndModulesIterator.Value());
+#endif // UE_DELEGATE_CHECK_LIFETIME
+
 				const uint64_t InitialMemUsed = FPlatformMemory::GetMemoryUsedFast();
 				FPlatformProcess::FreeDllHandle(*HandlePtr);
 				const uint64_t FinalMemUsed = FPlatformMemory::GetMemoryUsedFast();
@@ -1332,7 +1338,13 @@ bool FModuleManager::UnloadModule( const FName InModuleName, bool bIsShutdown, b
 			// Release reference to module interface.  This will actually destroy the module object.
 			ModuleInfo.Module.Reset();
 
-#if !IS_MONOLITHIC
+#if IS_MONOLITHIC
+
+#if UE_DELEGATE_CHECK_LIFETIME
+			FMulticastDelegateSafetyLayer::RemoveAllDelegatesInModules({ InModuleName });
+#endif // UE_DELEGATE_CHECK_LIFETIME
+
+#else
 
 #if UE_MERGED_MODULES
 			// Double-check that we don't have a valid handle for this...
@@ -1359,7 +1371,7 @@ bool FModuleManager::UnloadModule( const FName InModuleName, bool bIsShutdown, b
 				}
 				ModuleInfo.Handle = nullptr;
 			}
-#endif
+#endif // IS_MONOLITHIC
 
 			// If we're shutting down, then we never want this module to be "resurrected" in this session.
 			// It's gone for good!  So we'll mark it as such so that we can catch cases where a routine is
diff --git a/Engine/Source/Runtime/Core/Public/Delegates/DelegateBase.h b/Engine/Source/Runtime/Core/Public/Delegates/DelegateBase.h
--- a/Engine/Source/Runtime/Core/Public/Delegates/DelegateBase.h
+++ b/Engine/Source/Runtime/Core/Public/Delegates/DelegateBase.h
@@ -35,87 +35,6 @@ class TDelegateBase;
 
 ALIAS_TEMPLATE_TYPE_LAYOUT(template<typename ElementType>, FDelegateAllocatorType::ForElementType<ElementType>, void*);
 
-#if UE_DELEGATE_CHECK_LIFETIME
-
-// This tracking layer allows us to safely check if the code for virtual calls in IDelegateInstance are safe
-// This code will be called on instances that potentially have their code unloaded!
-class FTrackedDelegateInstanceExtras : public IDelegateInstance
-{
-public:
-	FTrackedDelegateInstanceExtras()
-#ifdef UE_MODULE_NAME
-		: ModuleName(UE_MODULE_NAME)
-#else
-		: ModuleName(NAME_None)
-#endif //UE_MODULE_NAME
-		, BoundFunctionName(NAME_None)
-	{}
-
-	/**
-	 * Set the name of the module this was compiled in.
-	 */
-	void SetModuleName(FName InModuleName)
-	{
-		ModuleName = InModuleName;
-	}
-
-	/**
-	 * Retrieves the name of the module this was compiled in.
-	 */
-	FName GetModuleName() const
-	{
-		return ModuleName;
-	}
-
-	/**
-	 * Stores the debug name for this delegate's call site - could be a context object, filename, etc.
-	 */
-	void SetBoundFunctionName(FName InBoundFunctionName)
-	{
-		BoundFunctionName = InBoundFunctionName;
-	}
-
-	/**
-	 * Retrieves the debug name of the bound function or call site.
-	 * This is similar to TryGetBoundFunctionName in IDelegateInstance but we need a non-virtual version with its own storage.
-	 * In turn, TryGetBoundFunctionName will call this if compiled with UE_DELEGATE_CHECK_LIFETIME.
-	 */
-	FName GetBoundFunctionName() const
-	{
-		return BoundFunctionName;
-	}
-
-	/**
-	 * Checks whether this delegate instance is safe to call virtual functions on (i.e. the module for it is loaded).
-	 * 
-	 * Returns whether the instance is valid.
-	 */
-	CORE_API bool IsValid() const;
-
-	/**
-	 * Checks whether this delegate instance is safe, and produce debug logging if not.
-	 * 
-	 * Returns whether the instance is valid.
-	 */
-	CORE_API bool CheckValid() const;
-
-protected:
-
-	FName ModuleName;
-	FName BoundFunctionName;
-};
-
-// Convenience macro to check for lifetime validity as widely as possible
-#define CHECK_DELEGATE_LIFETIME(DelegateInstance)                                           \
-    if (DelegateInstance != nullptr)                                                        \
-    {                                                                                       \
-    	static_cast<const FTrackedDelegateInstanceExtras*>(DelegateInstance)->CheckValid(); \
-    }
-
-#else
-#define CHECK_DELEGATE_LIFETIME(DelegateInstance)
-#endif // UE_DELEGATE_CHECK_LIFETIME
-
 // not thread-safe version, with automatic race detection in dev builds
 struct FDefaultDelegateUserPolicy
 {
@@ -142,11 +61,7 @@ struct FDefaultDelegateUserPolicy
 	//     holds an array of FDelegateExtras-derived objects which is the invocation list).
 	//   - Public data members and member functions are accessible directly through the TMulticastDelegate object.
 
-#if UE_DELEGATE_CHECK_LIFETIME
-	using FDelegateInstanceExtras = FTrackedDelegateInstanceExtras;
-#else
 	using FDelegateInstanceExtras = IDelegateInstance;
-#endif // UE_DELEGATE_CHECK_LIFETIME
 
 	using FThreadSafetyMode =
 #if UE_DETECT_DELEGATES_RACE_CONDITIONS
@@ -163,11 +78,7 @@ struct FDefaultTSDelegateUserPolicy
 {
 	// see `FDefaultDelegateUserPolicy` for documentation
 
-#if UE_DELEGATE_CHECK_LIFETIME
-	using FDelegateInstanceExtras = FTrackedDelegateInstanceExtras;
-#else
 	using FDelegateInstanceExtras = IDelegateInstance;
-#endif // UE_DELEGATE_CHECK_LIFETIME
 
 	using FThreadSafetyMode = FThreadSafeDelegateMode;
 	using FDelegateExtras = TDelegateBase<FThreadSafetyMode>;
@@ -177,11 +88,7 @@ struct FDefaultTSDelegateUserPolicy
 // not thread-safe version, no race detection. used primarily for deprecated unsafe delegates that must be kept running for backward compatibility
 struct FNotThreadSafeNotCheckedDelegateUserPolicy
 {
-#if UE_DELEGATE_CHECK_LIFETIME
-	using FDelegateInstanceExtras = FTrackedDelegateInstanceExtras;
-#else
 	using FDelegateInstanceExtras = IDelegateInstance;
-#endif // UE_DELEGATE_CHECK_LIFETIME
 
 	using FThreadSafetyMode = FNotThreadSafeNotCheckedDelegateMode;
 	using FDelegateExtras = TDelegateBase<FThreadSafetyMode>;
diff --git a/Engine/Source/Runtime/Core/Public/Delegates/DelegateInstancesImpl.h b/Engine/Source/Runtime/Core/Public/Delegates/DelegateInstancesImpl.h
--- a/Engine/Source/Runtime/Core/Public/Delegates/DelegateInstancesImpl.h
+++ b/Engine/Source/Runtime/Core/Public/Delegates/DelegateInstancesImpl.h
@@ -23,18 +23,6 @@
 #include "UObject/NameTypes.h"
 #include "UObject/WeakObjectPtrTemplates.h"
 
-#if UE_DELEGATE_CHECK_LIFETIME
-#  ifdef UE_MODULE_NAME
-#    define INSTRUMENT_DELEGATE_MODULE() FTrackedDelegateInstanceExtras::SetModuleName(UE_MODULE_NAME)
-#  else
-#    define INSTRUMENT_DELEGATE_MODULE()
-#  endif // UE_MODULE_NAME
-#  define INSTRUMENT_DELEGATE_FUNCTION(Function) FTrackedDelegateInstanceExtras::SetBoundFunctionName(Function)
-#else
-#  define INSTRUMENT_DELEGATE_MODULE()
-#  define INSTRUMENT_DELEGATE_FUNCTION(Function)
-#endif // UE_DELEGATE_CHECK_LIFETIME
-
 namespace UE::Delegates::Private
 {
 	constexpr bool IsUObjectPtr(const volatile UObjectBase*) { return true; }
@@ -86,9 +74,6 @@ public:
 	{
 		check(InFunctionName != NAME_None);
 
-		INSTRUMENT_DELEGATE_MODULE();
-		INSTRUMENT_DELEGATE_FUNCTION(*UserClass::StaticClass()->GetName());
-
 		if (InUserObject != nullptr)
 		{
 			CachedFunction = UserObjectPtr->FindFunctionChecked(InFunctionName);
@@ -216,9 +201,6 @@ public:
 		// NOTE: Shared pointer delegates are allowed to have a null incoming object pointer.  Weak pointers can expire,
 		//       an it is possible for a copy of a delegate instance to end up with a null pointer.
 		checkSlow(MethodPtr != nullptr);
-
-		INSTRUMENT_DELEGATE_MODULE();
-		INSTRUMENT_DELEGATE_FUNCTION(__FUNCTION__);
 	}
 
 	// IDelegateInstance interface
@@ -335,8 +317,6 @@ public:
 		, ContextObject(InContextObject)
 		, Functor(Forward<InFunctorType>(InFunctor))
 	{
-		INSTRUMENT_DELEGATE_MODULE();
-		INSTRUMENT_DELEGATE_FUNCTION(__FUNCTION__);
 	}
 
 	// IDelegateInstance interface
@@ -437,9 +417,6 @@ public:
 	{
 		// Non-expirable delegates must always have a non-null object pointer on creation (otherwise they could never execute.)
 		check(InUserObject != nullptr && MethodPtr != nullptr);
-
-		INSTRUMENT_DELEGATE_MODULE();
-		INSTRUMENT_DELEGATE_FUNCTION(__FUNCTION__);
 	}
 
 	// IDelegateInstance interface
@@ -552,9 +529,6 @@ public:
 		// NOTE: UObject delegates are allowed to have a null incoming object pointer.  UObject weak pointers can expire,
 		//       an it is possible for a copy of a delegate instance to end up with a null pointer.
 		checkSlow(MethodPtr != nullptr);
-
-		INSTRUMENT_DELEGATE_MODULE();
-		INSTRUMENT_DELEGATE_FUNCTION(*UserClass::StaticClass()->GetName());
 	}
 
 	// IDelegateInstance interface
@@ -682,8 +656,6 @@ public:
 		, StaticFuncPtr(InStaticFuncPtr)
 	{
 		check(StaticFuncPtr != nullptr);
-
-		INSTRUMENT_DELEGATE_MODULE();
 	}
 
 	// IDelegateInstance interface
@@ -774,7 +746,6 @@ public:
 		: Super  (Forward<InVarTypes>(Vars)...)
 		, Functor(Forward<InFunctorType>(InFunctor))
 	{
-		INSTRUMENT_DELEGATE_MODULE();
 	}
 
 	// IDelegateInstance interface
@@ -865,8 +836,6 @@ public:
 		, ContextObject(InContextObject)
 		, Functor      (Forward< InFunctorType>(InFunctor))
 	{
-		INSTRUMENT_DELEGATE_MODULE();
-		INSTRUMENT_DELEGATE_FUNCTION(*InContextObject->GetName());
 	}
 
 	// IDelegateInstance interface
diff --git a/Engine/Source/Runtime/Core/Public/Delegates/DelegateSignatureImpl.inl b/Engine/Source/Runtime/Core/Public/Delegates/DelegateSignatureImpl.inl
--- a/Engine/Source/Runtime/Core/Public/Delegates/DelegateSignatureImpl.inl
+++ b/Engine/Source/Runtime/Core/Public/Delegates/DelegateSignatureImpl.inl
@@ -34,6 +34,17 @@ TCopyQualifiersFromTo_T<From, To>* Cast(From* Src);
 
 template<typename UserPolicy> class TMulticastDelegateBase;
 
+#if UE_DELEGATE_CHECK_LIFETIME && defined(UE_DYNAMIC_MODULE_NAME)
+#include "Misc/SourceLocation.h"
+#include "Misc/SourceLocationUtils.h"
+#define UE_REGISTER_MULTICAST_INSTANCE(Handle) FMulticastDelegateSafetyLayer::RegisterDelegateInstance( \
+	Handle,                                                                                          \
+	UE_DYNAMIC_MODULE_NAME,                                                                          \
+	UE::FSourceLocation::Current().GetFunctionName())
+#else
+#define UE_REGISTER_MULTICAST_INSTANCE(Handle)
+#endif // UE_DELEGATE_CHECK_LIFETIME && defined(UE_DYNAMIC_MODULE_NAME)
+
 /**
  * Delegate registration class.  A function can expose this base class to users to let them bind
  * against, but without letting them execute.  This gives us the benefit of C# events.
@@ -325,6 +336,9 @@ private:
 	using typename Super::FWriteAccessScope;
 	using Super::GetWriteAccessScope;
 
+	template <typename>
+	friend class TMulticastDelegateBase;
+
 public:
 	using RegistrationType = Super;
 
@@ -605,7 +619,6 @@ public:
 		FReadAccessScope ReadScope = GetReadAccessScope();
 
 		const DelegateInstanceInterfaceType* LocalDelegateInstance = GetDelegateInstanceProtected();
-		CHECK_DELEGATE_LIFETIME(LocalDelegateInstance);
 
 		// If this assert goes off, Execute() was called before a function was bound to the delegate.
 		// Consider using ExecuteIfBound() instead.
@@ -627,7 +640,6 @@ public:
 
 		if (const DelegateInstanceInterfaceType* Ptr = GetDelegateInstanceProtected())
 		{
-			CHECK_DELEGATE_LIFETIME(Ptr);
 			return Ptr->ExecuteIfSafe(Forward<ParamTypes>(Params)...);
 		}
 
@@ -770,7 +782,9 @@ public:
 	 */
 	FDelegateHandle Add(FDelegate&& InNewDelegate)
 	{
-		return Super::AddDelegateInstance(MoveTemp(InNewDelegate));
+		FDelegateHandle Handle = Super::AddDelegateInstance(MoveTemp(InNewDelegate));
+		UE_REGISTER_MULTICAST_INSTANCE(Handle);
+		return Handle;
 	}
 
 	/**
@@ -780,7 +794,9 @@ public:
 	 */
 	FDelegateHandle Add(const FDelegate& InNewDelegate)
 	{
-		return Super::AddDelegateInstance(CopyTemp(InNewDelegate));
+		FDelegateHandle Handle = Super::AddDelegateInstance(CopyTemp(InNewDelegate));
+		UE_REGISTER_MULTICAST_INSTANCE(Handle);
+		return Handle;
 	}
 
 	/**
@@ -791,7 +807,9 @@ public:
 	template <typename... VarTypes>
 	inline FDelegateHandle AddStatic(typename TBaseStaticDelegateInstance<void (ParamTypes...), UserPolicy, std::decay_t<VarTypes>...>::FFuncPtr InFunc, VarTypes&&... Vars)
 	{
-		return Super::AddDelegateInstance(FDelegate::CreateStatic(InFunc, Forward<VarTypes>(Vars)...));
+		FDelegateHandle Handle = Super::AddDelegateInstance(FDelegate::CreateStatic(InFunc, Forward<VarTypes>(Vars)...));
+		UE_REGISTER_MULTICAST_INSTANCE(Handle);
+		return Handle;
 	}
 
 	/**
@@ -803,7 +821,9 @@ public:
 	template<typename FunctorType, typename... VarTypes>
 	inline FDelegateHandle AddLambda(FunctorType&& InFunctor, VarTypes&&... Vars)
 	{
-		return Super::AddDelegateInstance(FDelegate::CreateLambda(Forward<FunctorType>(InFunctor), Forward<VarTypes>(Vars)...));
+		FDelegateHandle Handle = Super::AddDelegateInstance(FDelegate::CreateLambda(Forward<FunctorType>(InFunctor), Forward<VarTypes>(Vars)...));
+		UE_REGISTER_MULTICAST_INSTANCE(Handle);
+		return Handle;
 	}
 
 	/**
@@ -816,7 +836,9 @@ public:
 	template <typename UserClass, typename FunctorType, typename... VarTypes>
 	inline FDelegateHandle AddSPLambda(const UserClass* InUserObject, FunctorType&& InFunctor, VarTypes&&... Vars)
 	{
-		return Super::AddDelegateInstance(FDelegate::CreateSPLambda(InUserObject, Forward<FunctorType>(InFunctor), Forward<VarTypes>(Vars)...));
+		FDelegateHandle Handle = Super::AddDelegateInstance(FDelegate::CreateSPLambda(InUserObject, Forward<FunctorType>(InFunctor), Forward<VarTypes>(Vars)...));
+		UE_REGISTER_MULTICAST_INSTANCE(Handle);
+		return Handle;
 	}
 
 	/**
@@ -829,7 +851,9 @@ public:
 	template<typename UserClass, typename FunctorType, typename... VarTypes>
 	inline FDelegateHandle AddWeakLambda(UserClass* InUserObject, FunctorType&& InFunctor, VarTypes&&... Vars)
 	{
-		return Super::AddDelegateInstance(FDelegate::CreateWeakLambda(InUserObject, Forward<FunctorType>(InFunctor), Forward<VarTypes>(Vars)...));
+		FDelegateHandle Handle = Super::AddDelegateInstance(FDelegate::CreateWeakLambda(InUserObject, Forward<FunctorType>(InFunctor), Forward<VarTypes>(Vars)...));
+		UE_REGISTER_MULTICAST_INSTANCE(Handle);
+		return Handle;
 	}
 
 	/**
@@ -846,12 +870,16 @@ public:
 	{
 		static_assert(!std::is_const_v<UserClass>, "Attempting to bind a delegate with a const object pointer and non-const member function.");
 
-		return Super::AddDelegateInstance(FDelegate::CreateRaw(InUserObject, InFunc, Forward<VarTypes>(Vars)...));
+		FDelegateHandle Handle = Super::AddDelegateInstance(FDelegate::CreateRaw(InUserObject, InFunc, Forward<VarTypes>(Vars)...));
+		UE_REGISTER_MULTICAST_INSTANCE(Handle);
+		return Handle;
 	}
 	template <typename UserClass, typename... VarTypes>
 	inline FDelegateHandle AddRaw(const UserClass* InUserObject, typename TMemFunPtrType<true, UserClass, void (ParamTypes..., std::decay_t<VarTypes>...)>::Type InFunc, VarTypes&&... Vars)
 	{
-		return Super::AddDelegateInstance(FDelegate::CreateRaw(InUserObject, InFunc, Forward<VarTypes>(Vars)...));
+		FDelegateHandle Handle = Super::AddDelegateInstance(FDelegate::CreateRaw(InUserObject, InFunc, Forward<VarTypes>(Vars)...));
+		UE_REGISTER_MULTICAST_INSTANCE(Handle);
+		return Handle;
 	}
 
 	/**
@@ -867,12 +895,16 @@ public:
 	{
 		static_assert(!std::is_const_v<UserClass>, "Attempting to bind a delegate with a const object pointer and non-const member function.");
 
-		return Super::AddDelegateInstance(FDelegate::CreateSP(InUserObjectRef, InFunc, Forward<VarTypes>(Vars)...));
+		FDelegateHandle Handle = Super::AddDelegateInstance(FDelegate::CreateSP(InUserObjectRef, InFunc, Forward<VarTypes>(Vars)...));
+		UE_REGISTER_MULTICAST_INSTANCE(Handle);
+		return Handle;
 	}
 	template <typename UserClass, ESPMode Mode, typename... VarTypes>
 	inline FDelegateHandle AddSP(const TSharedRef<UserClass, Mode>& InUserObjectRef, typename TMemFunPtrType<true, UserClass, void (ParamTypes..., std::decay_t<VarTypes>...)>::Type InFunc, VarTypes&&... Vars)
 	{
-		return Super::AddDelegateInstance(FDelegate::CreateSP(InUserObjectRef, InFunc, Forward<VarTypes>(Vars)...));
+		FDelegateHandle Handle = Super::AddDelegateInstance(FDelegate::CreateSP(InUserObjectRef, InFunc, Forward<VarTypes>(Vars)...));
+		UE_REGISTER_MULTICAST_INSTANCE(Handle);
+		return Handle;
 	}
 
 	/**
@@ -888,12 +920,16 @@ public:
 	{
 		static_assert(!std::is_const_v<UserClass>, "Attempting to bind a delegate with a const object pointer and non-const member function.");
 
-		return Super::AddDelegateInstance(FDelegate::CreateSP(InUserObject, InFunc, Forward<VarTypes>(Vars)...));
+		FDelegateHandle Handle = Super::AddDelegateInstance(FDelegate::CreateSP(InUserObject, InFunc, Forward<VarTypes>(Vars)...));
+		UE_REGISTER_MULTICAST_INSTANCE(Handle);
+		return Handle;
 	}
 	template <typename UserClass, typename... VarTypes>
 	inline FDelegateHandle AddSP(const UserClass* InUserObject, typename TMemFunPtrType<true, UserClass, void (ParamTypes..., std::decay_t<VarTypes>...)>::Type InFunc, VarTypes&&... Vars)
 	{
-		return Super::AddDelegateInstance(FDelegate::CreateSP(InUserObject, InFunc, Forward<VarTypes>(Vars)...));
+		FDelegateHandle Handle = Super::AddDelegateInstance(FDelegate::CreateSP(InUserObject, InFunc, Forward<VarTypes>(Vars)...));
+		UE_REGISTER_MULTICAST_INSTANCE(Handle);
+		return Handle;
 	}
 
 	/**
@@ -909,12 +945,16 @@ public:
 	{
 		static_assert(!std::is_const_v<UserClass>, "Attempting to bind a delegate with a const object pointer and non-const member function.");
 
-		return Super::AddDelegateInstance(FDelegate::CreateThreadSafeSP(InUserObjectRef, InFunc, Forward<VarTypes>(Vars)...));
+		FDelegateHandle Handle = Super::AddDelegateInstance(FDelegate::CreateThreadSafeSP(InUserObjectRef, InFunc, Forward<VarTypes>(Vars)...));
+		UE_REGISTER_MULTICAST_INSTANCE(Handle);
+		return Handle;
 	}
 	template <typename UserClass, typename... VarTypes>
 	inline FDelegateHandle AddThreadSafeSP(const TSharedRef<UserClass, ESPMode::ThreadSafe>& InUserObjectRef, typename TMemFunPtrType<true, UserClass, void (ParamTypes..., std::decay_t<VarTypes>...)>::Type InFunc, VarTypes&&... Vars)
 	{
-		return Super::AddDelegateInstance(FDelegate::CreateThreadSafeSP(InUserObjectRef, InFunc, Forward<VarTypes>(Vars)...));
+		FDelegateHandle Handle = Super::AddDelegateInstance(FDelegate::CreateThreadSafeSP(InUserObjectRef, InFunc, Forward<VarTypes>(Vars)...));
+		UE_REGISTER_MULTICAST_INSTANCE(Handle);
+		return Handle;
 	}
 
 	/**
@@ -932,12 +972,16 @@ public:
 	{
 		static_assert(!std::is_const_v<UserClass>, "Attempting to bind a delegate with a const object pointer and non-const member function.");
 
-		return Super::AddDelegateInstance(FDelegate::CreateThreadSafeSP(InUserObject, InFunc, Forward<VarTypes>(Vars)...));
+		FDelegateHandle Handle = Super::AddDelegateInstance(FDelegate::CreateThreadSafeSP(InUserObject, InFunc, Forward<VarTypes>(Vars)...));
+		UE_REGISTER_MULTICAST_INSTANCE(Handle);
+		return Handle;
 	}
 	template <typename UserClass, typename... VarTypes>
 	inline FDelegateHandle AddThreadSafeSP(const UserClass* InUserObject, typename TMemFunPtrType<true, UserClass, void (ParamTypes..., std::decay_t<VarTypes>...)>::Type InFunc, VarTypes&&... Vars)
 	{
-		return Super::AddDelegateInstance(FDelegate::CreateThreadSafeSP(InUserObject, InFunc, Forward<VarTypes>(Vars)...));
+		FDelegateHandle Handle = Super::AddDelegateInstance(FDelegate::CreateThreadSafeSP(InUserObject, InFunc, Forward<VarTypes>(Vars)...));
+		UE_REGISTER_MULTICAST_INSTANCE(Handle);
+		return Handle;
 	}
 
 	/**
@@ -951,12 +995,16 @@ public:
 	template <typename UObjectTemplate, typename... VarTypes>
 	inline FDelegateHandle AddUFunction(UObjectTemplate* InUserObject, const FName& InFunctionName, VarTypes&&... Vars)
 	{
-		return Super::AddDelegateInstance(FDelegate::CreateUFunction(InUserObject, InFunctionName, Forward<VarTypes>(Vars)...));
+		FDelegateHandle Handle = Super::AddDelegateInstance(FDelegate::CreateUFunction(InUserObject, InFunctionName, Forward<VarTypes>(Vars)...));
+		UE_REGISTER_MULTICAST_INSTANCE(Handle);
+		return Handle;
 	}
 	template <typename UObjectTemplate, typename... VarTypes>
 	inline FDelegateHandle AddUFunction(TObjectPtr<UObjectTemplate> InUserObject, const FName& InFunctionName, VarTypes&&... Vars)
 	{
-		return Super::AddDelegateInstance(FDelegate::CreateUFunction(InUserObject, InFunctionName, Forward<VarTypes>(Vars)...));
+		FDelegateHandle Handle = Super::AddDelegateInstance(FDelegate::CreateUFunction(InUserObject, InFunctionName, Forward<VarTypes>(Vars)...));
+		UE_REGISTER_MULTICAST_INSTANCE(Handle);
+		return Handle;
 	}
 
 	/**
@@ -972,24 +1020,32 @@ public:
 	{
 		static_assert(!std::is_const_v<UserClass>, "Attempting to bind a delegate with a const object pointer and non-const member function.");
 
-		return Super::AddDelegateInstance(FDelegate::CreateUObject(InUserObject, InFunc, Forward<VarTypes>(Vars)...));
+		FDelegateHandle Handle = Super::AddDelegateInstance(FDelegate::CreateUObject(InUserObject, InFunc, Forward<VarTypes>(Vars)...));
+		UE_REGISTER_MULTICAST_INSTANCE(Handle);
+		return Handle;
 	}
 	template <typename UserClass, typename... VarTypes>
 	inline FDelegateHandle AddUObject(const UserClass* InUserObject, typename TMemFunPtrType<true, UserClass, void (ParamTypes..., std::decay_t<VarTypes>...)>::Type InFunc, VarTypes&&... Vars)
 	{
-		return Super::AddDelegateInstance(FDelegate::CreateUObject(InUserObject, InFunc, Forward<VarTypes>(Vars)...));
+		FDelegateHandle Handle = Super::AddDelegateInstance(FDelegate::CreateUObject(InUserObject, InFunc, Forward<VarTypes>(Vars)...));
+		UE_REGISTER_MULTICAST_INSTANCE(Handle);
+		return Handle;
 	}
 	template <typename UserClass, typename... VarTypes>
 	inline FDelegateHandle AddUObject(TObjectPtr<UserClass> InUserObject, typename TMemFunPtrType<false, UserClass, void (ParamTypes..., std::decay_t<VarTypes>...)>::Type InFunc, VarTypes&&... Vars)
 	{
 		static_assert(!std::is_const_v<UserClass>, "Attempting to bind a delegate with a const object pointer and non-const member function.");
 
-		return Super::AddDelegateInstance(FDelegate::CreateUObject(InUserObject, InFunc, Forward<VarTypes>(Vars)...));
+		FDelegateHandle Handle = Super::AddDelegateInstance(FDelegate::CreateUObject(InUserObject, InFunc, Forward<VarTypes>(Vars)...));
+		UE_REGISTER_MULTICAST_INSTANCE(Handle);
+		return Handle;
 	}
 	template <typename UserClass, typename... VarTypes>
 	inline FDelegateHandle AddUObject(TObjectPtr<UserClass> InUserObject, typename TMemFunPtrType<true, UserClass, void (ParamTypes..., std::decay_t<VarTypes>...)>::Type InFunc, VarTypes&&... Vars)
 	{
-		return Super::AddDelegateInstance(FDelegate::CreateUObject(InUserObject, InFunc, Forward<VarTypes>(Vars)...));
+		FDelegateHandle Handle = Super::AddDelegateInstance(FDelegate::CreateUObject(InUserObject, InFunc, Forward<VarTypes>(Vars)...));
+		UE_REGISTER_MULTICAST_INSTANCE(Handle);
+		return Handle;
 	}
 
 public:
diff --git a/Engine/Source/Runtime/Core/Public/Delegates/MulticastDelegateBase.h b/Engine/Source/Runtime/Core/Public/Delegates/MulticastDelegateBase.h
--- a/Engine/Source/Runtime/Core/Public/Delegates/MulticastDelegateBase.h
+++ b/Engine/Source/Runtime/Core/Public/Delegates/MulticastDelegateBase.h
@@ -20,11 +20,38 @@
 
 #define UE_MULTICAST_DELEGATE_DEFAULT_COMPACTION_THRESHOLD 2
 
+#if UE_DELEGATE_CHECK_LIFETIME
+
+// Centralized registration layer for multicast delegates
+class FMulticastDelegateSafetyLayer
+{
+public:
+
+	virtual CORE_API ~FMulticastDelegateSafetyLayer();
+
+	// Register a delegate instance for a given module
+	CORE_API void RegisterDelegateInstance(FDelegateHandle DelegateHandle, FName ModuleName, FName CallsiteName);
+
+	// Remove every delegate instance for this module
+	static CORE_API void RemoveAllDelegatesInModules(TConstArrayView<FName> ModuleNames);
+
+protected:
+
+	// We need this virtual wrapper because we're in a parent class for the actual delegate
+	virtual bool RemoveDelegateInstanceWrapper(FDelegateHandle Handle) = 0;
+};
+
+#endif // UE_DELEGATE_CHECK_LIFETIME
+
 /**
  * Abstract base class for multicast delegates.
  */
 template<typename UserPolicy>
-class TMulticastDelegateBase : public TDelegateAccessHandlerBase<typename UserPolicy::FThreadSafetyMode>
+class TMulticastDelegateBase :
+#if UE_DELEGATE_CHECK_LIFETIME
+	protected FMulticastDelegateSafetyLayer,
+#endif // UE_DELEGATE_CHECK_LIFETIME
+	public TDelegateAccessHandlerBase<typename UserPolicy::FThreadSafetyMode>
 {
 protected:
 	using Super = TDelegateAccessHandlerBase<typename UserPolicy::FThreadSafetyMode>;
@@ -114,8 +141,6 @@ public:
 		for (const UnicastDelegateType& DelegateBaseRef : InvocationList)
 		{
 			const IDelegateInstance* DelegateInstance = DelegateBaseRef.GetDelegateInstanceProtected();
-			CHECK_DELEGATE_LIFETIME(DelegateInstance);
-
 			if ((DelegateInstance != nullptr) && DelegateInstance->HasSameObject(InUserObject))
 			{
 				return true;
@@ -142,8 +167,6 @@ public:
 			for (UnicastDelegateType& DelegateBaseRef : InvocationList)
 			{
 				IDelegateInstance* DelegateInstance = DelegateBaseRef.GetDelegateInstanceProtected();
-				CHECK_DELEGATE_LIFETIME(DelegateInstance);
-
 				if ((DelegateInstance != nullptr) && DelegateInstance->HasSameObject(InUserObject))
 				{
 					// Manually unbind the delegate here so the compaction will find and remove it.
@@ -166,7 +189,6 @@ public:
 				UnicastDelegateType& DelegateBaseRef = InvocationList[InvocationListIndex];
 
 				IDelegateInstance* DelegateInstance = DelegateBaseRef.GetDelegateInstanceProtected();
-				CHECK_DELEGATE_LIFETIME(DelegateInstance);
 
 				if (DelegateInstance == nullptr
 					|| DelegateInstance->HasSameObject(InUserObject)
@@ -205,39 +227,6 @@ public:
 		return Size;
 	}
 
-#if UE_DELEGATE_CHECK_LIFETIME
-
-	/**
-	 * Gets a list of all delegate call sites & module names registered to this multicast
-	 *
-	 * @return A list of pair of call site + module name, including for delegates that were unbound but not yet compacted.
-	 */
-	TArray<TPair<FName, FName>> GetInvocationNames() const
-	{
-		TArray<TPair<FName, FName>> FoundFunctionsAndModules;
-
-		FReadAccessScope ReadScope = GetReadAccessScope();
-
-		for (const UnicastDelegateType& DelegateBaseRef : InvocationList)
-		{
-			const IDelegateInstance* DelegateInstance = DelegateBaseRef.GetDelegateInstanceProtected();
-			CHECK_DELEGATE_LIFETIME(DelegateInstance);
-
-			if (DelegateInstance != nullptr)
-			{
-				const FTrackedDelegateInstanceExtras* TrackedDelegateInstance = static_cast<const FTrackedDelegateInstanceExtras*>(DelegateInstance);
-
-				FoundFunctionsAndModules.Add(TPair<FName, FName>(
-					TrackedDelegateInstance->GetBoundFunctionName(),
-					TrackedDelegateInstance->GetModuleName()));
-			}
-		}
-
-		return FoundFunctionsAndModules;
-	}
-
-#endif // UE_DELEGATE_CHECK_LIFETIME
-
 protected:
 
 	/** Hidden default constructor. */
@@ -260,8 +249,6 @@ protected:
 		    {
 				if (const IDelegateInstance* OtherInstance = OtherDelegateRef.GetDelegateInstanceProtected())
 			    {
-					CHECK_DELEGATE_LIFETIME(OtherInstance);
-
 				    UnicastDelegateType TempDelegate;
 					static_cast<const DelegateInstanceInterfaceType*>(OtherInstance)->CreateCopy(FPrivateDelegateAllocation{TempDelegate});
 				    TempInvocationList.Add(MoveTemp(TempDelegate));
@@ -296,8 +283,6 @@ protected:
 				const UnicastDelegateType& DelegateBase = LocalInvocationList[InvocationListIndex];
 
 				const IDelegateInstance* DelegateInstanceInterface = DelegateBase.GetDelegateInstanceProtected();
-				CHECK_DELEGATE_LIFETIME(DelegateInstanceInterface);
-
 				if (DelegateInstanceInterface == nullptr || !static_cast<const DelegateInstanceInterfaceType*>(DelegateInstanceInterface)->ExecuteIfSafe(Params...))
 				{
 					NeedsCompaction = true;
@@ -348,8 +333,6 @@ protected:
 			UnicastDelegateType& DelegateBase = InvocationList[InvocationListIndex];
 
 			IDelegateInstance* DelegateInstance = DelegateBase.GetDelegateInstanceProtected();
-			CHECK_DELEGATE_LIFETIME(DelegateInstance);
-
 			if (DelegateInstance && DelegateInstance->GetHandle() == Handle)
 			{
 				DelegateBase.Unbind();
@@ -361,6 +344,16 @@ protected:
 		return false;
 	}
 
+#if UE_DELEGATE_CHECK_LIFETIME
+
+protected:
+	virtual bool RemoveDelegateInstanceWrapper(FDelegateHandle Handle) override
+	{
+		return RemoveDelegateInstance(Handle);
+	}
+
+#endif // UE_DELEGATE_CHECK_LIFETIME
+
 private:
 
 	inline bool ShouldSkipCompactionUnderAutoRTFM() const
@@ -405,7 +398,6 @@ private:
 			auto& DelegateBaseRef = InvocationList[InvocationListIndex];
 
 			IDelegateInstance* DelegateInstance = DelegateBaseRef.GetDelegateInstanceProtected();
-			CHECK_DELEGATE_LIFETIME(DelegateInstance);
 
 			if (DelegateInstance == nullptr	|| DelegateInstance->IsCompactable())
 			{
@@ -474,11 +466,6 @@ private:
 	{
 		for (UnicastDelegateType& DelegateBaseRef : InvocationList)
 		{
-#if UE_DELEGATE_CHECK_LIFETIME
-			IDelegateInstance* DelegateInstance = DelegateBaseRef.GetDelegateInstanceProtected();
-			CHECK_DELEGATE_LIFETIME(DelegateInstance);
-#endif // UE_DELEGATE_CHECK_LIFETIME
-
 			DelegateBaseRef.Unbind();
 		}
 
