[Tasks]

Implements task retraction for scheduled tasks within the TaskConcurrencyLimiter. Should worker threads be blocked waiting on a task limiter, there are no workers available to allow the task limiter to proceed. We now resolve this by allowing tasks to run inline on the waiting thread by retracting the work from the job system.

- LowLevelTasks::FTask has been wrapped with a FLimiterTask to allow managing the task lifetime across worker and waiting threads when retracting
- A new TArray of cache aligned FLimiterTask tasks has been added to keep track of tasks that have been submitted to the scheduler that may have not started running yet. These types are padded to avoid false sharing as tasks execute and are submitted across various threads.

#rnx

diff --git a/Engine/Source/Runtime/Core/Private/Tasks/TaskConcurrencyLimiter.cpp b/Engine/Source/Runtime/Core/Private/Tasks/TaskConcurrencyLimiter.cpp
--- a/Engine/Source/Runtime/Core/Private/Tasks/TaskConcurrencyLimiter.cpp
+++ b/Engine/Source/Runtime/Core/Private/Tasks/TaskConcurrencyLimiter.cpp
@@ -19,7 +19,7 @@ namespace UE::Tasks::TaskConcurrencyLimiter_Private
 		}
 	}
 
-	void FPimpl::AddWorkItem(LowLevelTasks::FTask* Task)
+	void FPimpl::AddWorkItem(FLimiterTask* Task)
 	{
 		NumWorkItems.fetch_add(1, std::memory_order_acquire);
 
@@ -58,6 +58,29 @@ namespace UE::Tasks::TaskConcurrencyLimiter_Private
 			return true;
 		}
 
+		// We must drain all work we can find via retraction since 
+		// we can't be sure if the worker threads are able to proceed.
+		bool bDidSomething;
+		do
+		{
+			bDidSomething = false;
+			// Before we wait on our completion event, find tasks that are currently in the task system and attempt to run them inline
+			// since it's possible task threads are waiting on work that was submitted to the limiter and may not make progress otherwise
+			for (int32 SlotIndex = 0; SlotIndex < ScheduledTasks.Num(); ++SlotIndex)
+			{
+				// If the task is being executed then the value will likely be null. However, it's possible the task is 
+				// being executed but hasn't cleared this state yet. That is _fine_ since it means our expedite 
+				// below will simply fail but it means we are now responsible for destroying the task and must call Release.
+				if (FLimiterTask* LimiterTask = ScheduledTasks[SlotIndex].Task.exchange(nullptr, std::memory_order_acquire))
+				{
+					bDidSomething = true;
+					LimiterTask->Task.TryExpedite();
+					LimiterTask->Release(); 
+				}
+			}
+		} while (bDidSomething || NumWorkItems.load(std::memory_order_acquire) != 0);
+
+		// This should fall-through immediately
 		return LocalCompletionEvent->Wait(Timeout);
 	}
 
@@ -66,13 +89,18 @@ namespace UE::Tasks::TaskConcurrencyLimiter_Private
 		bool bWakeUpWorker = !bSkipFirstWakeUp;
 		do
 		{
-			if (LowLevelTasks::FTask* Task = WorkQueue.Pop())
+			if (FLimiterTask* LimiterTask = WorkQueue.Pop())
 			{
+				LowLevelTasks::FTask& Task = LimiterTask->Task;
+
 				// Now that we know the ConcurrencySlot, set it at launch time so
 				// the executor can retrieve it.
-				Task->SetUserData((void*)(UPTRINT)ConcurrencySlot);
+				Task.SetUserData((void*)(UPTRINT)ConcurrencySlot);
 
-				LowLevelTasks::TryLaunch(*Task, bWakeUpWorker ? LowLevelTasks::EQueuePreference::GlobalQueuePreference : LowLevelTasks::EQueuePreference::LocalQueuePreference, bWakeUpWorker);
+				// Before we launch the job, increment the task lifetime since the job may run and complete before we assign to ScheduledTasks
+				LimiterTask->AddRef();
+				LowLevelTasks::TryLaunch(Task, bWakeUpWorker ? LowLevelTasks::EQueuePreference::GlobalQueuePreference : LowLevelTasks::EQueuePreference::LocalQueuePreference, bWakeUpWorker);
+				ScheduledTasks[ConcurrencySlot].Task.store(LimiterTask, std::memory_order_release); // Signal this task has now been scheduled
 			}
 			else
 			{
@@ -90,9 +118,8 @@ namespace UE::Tasks::TaskConcurrencyLimiter_Private
 	{
 		// Once we are in a worker thread, we want to schedule on the local queue without waking up additional workers
 		// to allow our own worker to pick up the next item and avoid wake-up cost / context switch.
-		static constexpr bool bSkipFirstWakeUp = true;
-
-		ProcessQueue(ConcurrencySlot, bSkipFirstWakeUp);
+		// We must check IsInWorkerThread() since it's possible we have retracted the work to run outside of a worker thread
+		ProcessQueue(ConcurrencySlot, IsInWorkerThread());
 	}
 
 	void FPimpl::ProcessQueueFromPush(uint32 ConcurrencySlot)
diff --git a/Engine/Source/Runtime/Core/Public/Tasks/TaskConcurrencyLimiter.h b/Engine/Source/Runtime/Core/Public/Tasks/TaskConcurrencyLimiter.h
--- a/Engine/Source/Runtime/Core/Public/Tasks/TaskConcurrencyLimiter.h
+++ b/Engine/Source/Runtime/Core/Public/Tasks/TaskConcurrencyLimiter.h
@@ -53,11 +53,19 @@ namespace UE::Tasks
 		// an implementation details of FTaskConcurrenctyLimiter
 		class FPimpl : public TSharedFromThis<FPimpl>
 		{
+			// Reference counted wrapper for FTask to manage the task lifetime between worker threads and 
+			// any thread that might retract the task once it is submitted to the scheduler
+			struct FLimiterTask : FRefCountBase
+			{
+				LowLevelTasks::FTask Task;
+			};
+
 		public:
 			explicit FPimpl(uint32 InMaxConcurrency, ETaskPriority InTaskPriority)
 				: ConcurrencySlots(InMaxConcurrency)
 				, TaskPriority(InTaskPriority)
 			{
+				ScheduledTasks.SetNum(InMaxConcurrency);
 			}
 
 			CORE_API ~FPimpl();
@@ -65,45 +73,63 @@ namespace UE::Tasks
 			template<typename TaskFunctionType>
 			void Push(const TCHAR* DebugName, TaskFunctionType&& TaskFunction)
 			{
-				TSharedPtr<LowLevelTasks::FTask> Task = MakeShared<LowLevelTasks::FTask>();
-
-				Task->Init(
+				FLimiterTask* LimiterTask = new FLimiterTask();
+				LowLevelTasks::FTask& Task = LimiterTask->Task;
+				Task.Init(
 					DebugName,
 					TaskPriority,
 					[
 						TaskFunction = MoveTemp(TaskFunction),
 						this,
 						Pimpl = TSharedFromThis<FPimpl>::AsShared(), // to keep it alive
-						Task  // self-destruct
+						LimiterTask = TRefCountPtr<FLimiterTask>(LimiterTask) // self-destruct
 					]()
 					{
+						LowLevelTasks::FTask& Task = LimiterTask->Task;
+
 						// We can't pass the ConcurrencySlot in the lambda during creation as
 						// it's not actually acquired yet. The value will be passed using
 						// the user data when the task is launched.
-						uint32 ConcurrencySlot = (uint32)(UPTRINT)Task->GetUserData();
+						uint32 ConcurrencySlot = (uint32)(UPTRINT)Task.GetUserData();
+
+						// This task is now executing; remove it from the ScheduledTask array so it cannot be retracted.
+						// It's possible a thread waiting on the limiter might have retracted this task before we can clear it. 
+						bool bOwnLimiterTask = ScheduledTasks[ConcurrencySlot].Task.exchange(nullptr, std::memory_order_release) != nullptr;
+						if (bOwnLimiterTask)
+						{
+							// Whomever obtains the task from ScheduledTasks must release the reference
+							LimiterTask->Release();
+						}
 
 						TaskFunction(ConcurrencySlot);
 						CompleteWorkItem(ConcurrencySlot);
 					}
 				);
 
-				AddWorkItem(Task.Get());
+				AddWorkItem(LimiterTask);
 			}
 
 			bool CORE_API Wait(FTimespan Timeout);
 
 		private:
-			void CORE_API AddWorkItem(LowLevelTasks::FTask* Task);
+			void CORE_API AddWorkItem(FLimiterTask* Task);
 			void CORE_API ProcessQueue(uint32 ConcurrencySlot, bool bSkipFirstWakeUp);
 			void CORE_API ProcessQueueFromWorker(uint32 ConcurrencySlot);
 			void CORE_API ProcessQueueFromPush(uint32 ConcurrencySlot);
 			void CORE_API CompleteWorkItem(uint32 ConcurrencySlot);
 
+			// Padded struct to prevent false sharing as we read and write to the ScheduledTasks array			
+			struct alignas(PLATFORM_CACHE_LINE_SIZE) FPaddedSharedTask
+			{
+				std::atomic<FLimiterTask*> Task;
+			};
+
 			FConcurrencySlots ConcurrencySlots; // free slots queue. used also to limit concurrency
-			ETaskPriority TaskPriority;
-			TLockFreePointerListFIFO<LowLevelTasks::FTask, PLATFORM_CACHE_LINE_SIZE> WorkQueue; // a queue of user-provided task functions
-			std::atomic<uint32> NumWorkItems { 0 };
+			TLockFreePointerListFIFO<FLimiterTask, PLATFORM_CACHE_LINE_SIZE> WorkQueue; // a queue of user-provided task functions
+			TArray<FPaddedSharedTask> ScheduledTasks;
 			std::atomic<FEvent*> CompletionEvent { nullptr };
+			std::atomic<uint32> NumWorkItems { 0 };
+			ETaskPriority TaskPriority;
 		};
 
 	} // namespace TaskConcurrencyLimiter_Private
@@ -161,4 +187,4 @@ namespace UE::Tasks
 		TSharedRef<TaskConcurrencyLimiter_Private::FPimpl> Pimpl;
 	};
 
-} // namespace UE::Tasks
\ No newline at end of file
+} // namespace UE::Tasks
diff --git a/Engine/Source/Runtime/Core/Tests/Tasks/TasksTest.cpp b/Engine/Source/Runtime/Core/Tests/Tasks/TasksTest.cpp
--- a/Engine/Source/Runtime/Core/Tests/Tasks/TasksTest.cpp
+++ b/Engine/Source/Runtime/Core/Tests/Tasks/TasksTest.cpp
@@ -1671,6 +1671,59 @@ namespace UE { namespace TasksTests
 		UE_BENCHMARK(5, TaskConcurrencyLimiter_WaitingStressTest<10'000>);
 	}
 
+	template<uint32 RepeatCount>
+	void TaskConcurrencyLimiter_WaitingRetractionStressTest()
+	{
+		for (uint32 RepeatIndex = 0; RepeatIndex < RepeatCount; ++RepeatIndex)
+		{
+			// Schedule more blocking work than we have task threads to run them on so we are forced to retract work in the limiter to proceed
+			const uint32 NumWaitingTasks = LowLevelTasks::FScheduler::Get().GetMaxNumWorkers() * 2;
+			TArray<FEventRef> Blockers;
+			TArray<FTask> WaitingTasks;
+			Blockers.SetNum(NumWaitingTasks);
+			WaitingTasks.Reserve(NumWaitingTasks);
+
+			std::atomic<int32> WaitCount = 0;
+			for (uint32 WaitIndex = 0; WaitIndex < NumWaitingTasks; ++WaitIndex)
+			{
+				WaitingTasks.Add(Launch(UE_SOURCE_LOCATION,
+					[&Blockers, WaitIndex, &WaitCount]
+					{
+						TRACE_CPUPROFILER_EVENT_SCOPE(Waiting);
+						WaitCount.fetch_add(1, std::memory_order_relaxed);
+						Blockers[WaitIndex]->Wait();
+					}));
+			}
+
+			// Ensure we have waited long enough that we have started blocking tasks on all worker threads
+			while (WaitCount.load(std::memory_order_relaxed) < (int32)LowLevelTasks::FScheduler::Get().GetNumWorkers())
+			{
+				FPlatformProcess::Sleep(0.0f);
+			}
+
+			FTaskConcurrencyLimiter TaskConcurrencyLimiter{ 1 };
+			for (uint32 WaitIndex = 0; WaitIndex < NumWaitingTasks; ++WaitIndex)
+			{
+				TaskConcurrencyLimiter.Push(UE_SOURCE_LOCATION,
+					[&Blockers, WaitIndex](uint32)
+					{
+						TRACE_CPUPROFILER_EVENT_SCOPE(Blocked);
+						Blockers[WaitIndex]->Trigger();
+					}
+				);
+			}
+
+			// Without retractions we will never run any of the queued task work
+			TaskConcurrencyLimiter.Wait();
+			Wait(WaitingTasks);
+		}
+	}
+
+	TEST_CASE_NAMED(FTasksConcurrencyLimiterWaitingRetractionTest, "System::Core::Async::Tasks::TaskConcurrencyLimiter::WaitingRetraction", "[.][ApplicationContextMask][EngineFilter]")
+	{
+		UE_BENCHMARK(5, TaskConcurrencyLimiter_WaitingRetractionStressTest<10'000>);
+	}
+
 	void BenchmarkBlockUnblockWorkers(uint32 NumWorkers)
 	{
 		TRACE_CPUPROFILER_EVENT_SCOPE(BenchmarkBlockUnblockWorkers);
