[Core] Fix compact set to work with constinit

#tests Compiled client build

diff --git a/Engine/Source/Runtime/Core/Public/Containers/CompactSetBase.h b/Engine/Source/Runtime/Core/Public/Containers/CompactSetBase.h
--- a/Engine/Source/Runtime/Core/Public/Containers/CompactSetBase.h
+++ b/Engine/Source/Runtime/Core/Public/Containers/CompactSetBase.h
@@ -78,7 +78,7 @@ protected:
 	static_assert(std::is_same_v<SizeType, int32>, "TCompactSet currently only supports 32-bit allocators");
 	static_assert(sizeof(HashCountType) == UE::Core::CompactHashTable::GetMemoryAlignment(), "Hashtable alignment changed, need to make sure HashCountType is stored correctly");
 
-	[[nodiscard]] UE_FORCEINLINE_HINT TCompactSetBase() = default;
+	[[nodiscard]] UE_FORCEINLINE_HINT constexpr TCompactSetBase() = default;
 
 	[[nodiscard]] explicit consteval TCompactSetBase(EConstEval)
 	: Elements(ConstEval)
diff --git a/Engine/Source/Runtime/Core/Public/Containers/ScriptCompactSet.h b/Engine/Source/Runtime/Core/Public/Containers/ScriptCompactSet.h
--- a/Engine/Source/Runtime/Core/Public/Containers/ScriptCompactSet.h
+++ b/Engine/Source/Runtime/Core/Public/Containers/ScriptCompactSet.h
@@ -17,7 +17,7 @@ class TScriptCompactSet : public TCompactSetBase<typename Allocator::template El
 {
 	using Super = TCompactSetBase<typename Allocator::template ElementAllocator<sizeof(uint8)>>;
 public:
-	[[nodiscard]] static FScriptCompactSetLayout GetScriptLayout(int32 ElementSize, int32 ElementAlignment)
+	[[nodiscard]] static constexpr FScriptCompactSetLayout GetScriptLayout(int32 ElementSize, int32 ElementAlignment)
 	{
 		return { ElementSize, FGenericPlatformMath::Max<int32>(ElementAlignment, UE::Core::CompactHashTable::GetMemoryAlignment()) };
 	}
@@ -35,7 +35,7 @@ public:
 	// End - intrusive TOptional<TScriptCompactSet> state //
 	////////////////////////////////////////////////////////
 
-	[[nodiscard]] TScriptCompactSet() = default;
+	[[nodiscard]] constexpr TScriptCompactSet() = default;
 
 	[[nodiscard]] bool IsValidIndex(int32 Index) const
 	{
