[Core] Add vizualization support for compact sets/maps
* Fixup bad indexing in GDB visualizer that could break anything beyond the first 32 elements of a sparse array
* Added helper for referencing other containers using their visualizer
* Re-organized container visualizers into a separate function

#tests Compiled linux server with default sparse and compact sets and viewed all sets and maps in container tests.

diff --git a/Engine/Extras/GDBPrinters/UEPrinters.py b/Engine/Extras/GDBPrinters/UEPrinters.py
--- a/Engine/Extras/GDBPrinters/UEPrinters.py
+++ b/Engine/Extras/GDBPrinters/UEPrinters.py
@@ -78,6 +78,22 @@ def get_allocation_data(Allocator, ElementType):
 
 	return Data.cast(ElementType.pointer())
 
+# ------------------------------------------------------------------------------
+# ContainerPrinterHelper - used to wrap another container but just adjust the results slightly
+#
+class ContainerPrinterHelper:
+	def __init__(self, container):
+		self.Container = gdb.default_visualizer(container)
+
+	def to_string(self):
+		return self.Container.to_string()
+
+	def children(self):
+		return self.Container.children()
+
+	def display_hint(self):
+		return self.Container.display_hint()
+
 # ------------------------------------------------------------------------------
 #
 #  Custom pretty printers.
@@ -123,7 +139,7 @@ class TBitArrayPrinter:
 		return 'array'
 	
 	def get_value(self, index):
-		return ((self.Data[index/32] >> index) & 1) != 0
+		return (self.Data[index//32] >> (index % 32) & 1) != 0
 		
 	def is_valid(self, index):
 		return index < self.ArrayMax
@@ -264,11 +280,9 @@ class TSparseArrayPrinter:
 		return self.AllocationFlags.get_value(index)
 
 # ------------------------------------------------------------------------------
-# TSet
+# TSparseSet - Printer for only sparse type maps
 #
-class TSparseSetPrinter(TSparseArrayPrinter):
-	"Print TSparseSet"
-
+class TSparseSetPrinter(ContainerPrinterHelper):
 	def __init__(self, val):
 		super().__init__(val['Elements'])
 
@@ -277,12 +291,30 @@ class TSparseSetPrinter(TSparseArrayPrinter):
 			yield (Str, Element['Value'])
 
 # ------------------------------------------------------------------------------
-# TSparseMap
+# TCompactSet - Printer for only compact type maps
 #
+class TCompactSetPrinter(TArrayPrinter):
+	def __init__(self, val):
+		self.Value = val
+		self.ArrayNum = self.Value['NumElements']
+		self.ArrayMax = self.Value['MaxElements']
+		self.Data = get_allocation_data(self.Value['Elements'], self.Value.type.template_argument(0))
 
-class TSparseMapPrinter(TSparseSetPrinter):
-	"Print TSparseMap"
+# ------------------------------------------------------------------------------
+# TSet - Resolve to either sparse or compact depending on the definition of the type by determing what properties are available
+#
+class TSetPrinter(ContainerPrinterHelper):
+	def __init__(self, val):
+		if gdb.types.has_field(val.type, "NumElements"):
+			self.Container = TCompactSetPrinter(val)
+		else:
+			self.Container = TSparseSetPrinter(val)
 
+# ------------------------------------------------------------------------------
+# TMap - Printer for all formats of maps
+#
+
+class TMapPrinter(ContainerPrinterHelper):
 	def __init__(self, val):
 		super().__init__(val['Pairs'])
 
@@ -515,6 +547,27 @@ def register_ue_printers(objfile):
 	gdb.printing.register_pretty_printer(objfile, build_ue_pretty_printer(), True)
 	print("Registered pretty printers for UE classes")
 
+def add_container_printers(pp):
+	pp.add_printer("TArray", '^TArray<.+,.+>$', TArrayPrinter)
+	pp.add_printer("TBitArray", '^TBitArray<.+>$', TBitArrayPrinter)
+	pp.add_printer("TChunkedArray", '^TChunkedArray<.+>$', TChunkedArrayPrinter)
+	pp.add_printer("TSparseArray", '^TSparseArray<.+>$', TSparseArrayPrinter)
+
+	pp.add_printer("TSet", '^TSet<.+>$', TSetPrinter)
+	pp.add_printer("TSparseSet", '^TSparseSet<.+>$', TSparseSetPrinter)
+	pp.add_printer("TCompactSet", '^TCompactSet<.+>$', TCompactSetPrinter)
+
+	if (True): # useful to remove all map matching if you need to debug set printer or below
+		pp.add_printer("TMapBase", '^TMapBase<.+,.+,.+,.+>$', TMapPrinter)
+		pp.add_printer("TSortableMapBase", '^TSortableMapBase<.+,.+,.+,.+>$', TMapPrinter)
+		pp.add_printer("TMap", '^TMap<.+,.+,.+,.+>$', TMapPrinter)
+		pp.add_printer("TSparseMapBase", '^TSparseMapBase<.+,.+,.+,.+>$', TMapPrinter)
+		pp.add_printer("TSortableSparseMapBase", '^TSortableSparseMapBase<.+,.+,.+,.+>$', TMapPrinter)
+		pp.add_printer("TSparseMap", '^TSparseMap<.+,.+,.+,.+>$', TMapPrinter)
+		pp.add_printer("TCompactMapBase", '^TCompactMapBase<.+,.+,.+,.+>$', TMapPrinter)
+		pp.add_printer("TSortableCompactMapBase", '^TSortableCompactMapBase<.+,.+,.+,.+>$', TMapPrinter)
+		pp.add_printer("TCompactMap", '^TCompactMap<.+,.+,.+,.+>$', TMapPrinter)
+
 def build_ue_pretty_printer():
 	# add a random numeric suffix to the printer name so we can reload printers during the same session for iteration
 	pp = gdb.printing.RegexpCollectionPrettyPrinter("UnrealEngine")
@@ -523,23 +576,14 @@ def build_ue_pretty_printer():
 	pp.add_printer("FNameEntryId", '^FNameEntryId$', FNameEntryIdPrinter)
 	pp.add_printer("FName", '^FName$', make_fname_printer)
 	pp.add_printer("FMinimalName", '^FMinimalName$', make_fminimalname_printer)
-	pp.add_printer("TArray", '^TArray<.+,.+>$', TArrayPrinter)
-	pp.add_printer("TBitArray", '^TBitArray<.+>$', TBitArrayPrinter)
-	pp.add_printer("TChunkedArray", '^TChunkedArray<.+>$', TChunkedArrayPrinter)
-	pp.add_printer("TSparseArray", '^TSparseArray<.+>$', TSparseArrayPrinter)
-	pp.add_printer("TSet", '^TSet<.+>$', TSparseSetPrinter)
-	pp.add_printer("TSparseSet", '^TSparseSet<.+>$', TSparseSetPrinter)
 	pp.add_printer("FBitReference", '^FBitReference$', FBitReferencePrinter)
-	pp.add_printer("TMapBase", '^TMapBase<.+,.+,.+,.+>$', TSparseMapPrinter)
-	pp.add_printer("TSortableMapBase", '^TSortableMapBase<.+,.+,.+,.+>$', TSparseMapPrinter)
-	pp.add_printer("TMap", '^TMap<.+,.+,.+,.+>$', TSparseMapPrinter)
-	pp.add_printer("TSparseMapBase", '^TSparseMapBase<.+,.+,.+,.+>$', TSparseMapPrinter)
-	pp.add_printer("TSortableSparseMapBase", '^TSortableSparseMapBase<.+,.+,.+,.+>$', TSparseMapPrinter)
-	pp.add_printer("TSparseMap", '^TSparseMap<.+,.+,.+,.+>$', TSparseMapPrinter)
 	pp.add_printer("TPair", '^TPair<.+,.+>$', TPairPrinter)
 	pp.add_printer("TTuple", '^TTuple<.+,.+>$', TPairPrinter)
 	pp.add_printer("TWeakObjectPtr", '^TWeakObjectPtr<.+>$', TWeakObjectPtrPrinter)
 	pp.add_printer("FPrimaryAssetId", '^FPrimaryAssetId$', FPrimaryAssetIdPrinter)
+
+	add_container_printers(pp)
+
 	return pp
 
 register_ue_printers(None)
diff --git a/Engine/Source/Runtime/Core/Private/Containers/ContainersTest.cpp b/Engine/Source/Runtime/Core/Private/Containers/ContainersTest.cpp
--- a/Engine/Source/Runtime/Core/Private/Containers/ContainersTest.cpp
+++ b/Engine/Source/Runtime/Core/Private/Containers/ContainersTest.cpp
@@ -1756,98 +1756,97 @@ struct TSetNatvisTest
 	template <typename Type>
 	using SetType = typename Resolver::template TSetType<Type>;
 
+	// Name       | Value | Type
+	// -----------+-------+------------------------------------------------------
+	// > EmptySet | Empty | TSet<int,DefaultKeyFuncs<int,0>,FDefaultSetAllocator>
+	SetType<int32> EmptySet;
+
+	// Name       | Value         | Type
+	// -----------+---------------+------------------------------------------------------
+	// > SlackSet | Empty, Max=16 | TSet<int,DefaultKeyFuncs<int,0>,FDefaultSetAllocator>
+	SetType<int32> SlackSet;
+
+	// Name      | Value        | Type
+	// ----------+--------------+------------------------------------------------------
+	// > SetVals | Num=5, Max=5 | TSet<int,DefaultKeyFuncs<int,0>,FDefaultSetAllocator>
+	//     [0]   | 4            | int
+	//     [1]   | 1            | int
+	//     [2]   | 2            | int
+	//     [3]   | 5            | int
+	//     [4]   | 3            | int
+	SetType<int32> SetVals = { 4, 1, 2, 5, 3 };
+
+	// Name           | Value         | Type
+	// ---------------+---------------+------------------------------------------------------
+	// > SlackSetVals | Num=5, Max=16 | TSet<int,DefaultKeyFuncs<int,0>,FDefaultSetAllocator>
+	//     [0]        | 4             | int
+	//     [1]        | 1             | int
+	//     [2]        | 2             | int
+	//     [3]        | 5             | int
+	//     [4]        | 3             | int
+	SetType<int32> SlackSetVals = { 4, 1, 2, 5, 3 };
+
+	// Name            | Value        | Type
+	// ----------------+--------------+------------------------------------------------------
+	// > SparseSetVals | Num=3, Max=5 | TSet<int,DefaultKeyFuncs<int,0>,FDefaultSetAllocator>
+	//     [0]         | 4            | int
+	//     [1]         | 2            | int
+	//     [2]         | 3            | int
+	// Note: Order different for compact set
+	SetType<int32> SparseSetVals = { 4, 1, 2, 5, 3 };
+
+	// Name      | Value        | Type
+	// ----------+--------------+--------------------------------------------------------------
+	// > SetStrs | Num=5, Max=5 | TSet<FString,DefaultKeyFuncs<FString,0>,FDefaultSetAllocator>
+	//   > [0]   | L"d1"        | FString
+	//       [0] | 100 'd'      | wchar_t
+	//       [1] | 49 '1'       | wchar_t
+	//   > [1]   | L"a2"        | FString
+	//       [0] | 97 'a'       | wchar_t
+	//       [1] | 50 '2'       | wchar_t
+	//   > [2]   | L"b3"        | FString
+	//       [0] | 98 'b'       | wchar_t
+	//       [1] | 51 '3'       | wchar_t
+	//   > [3]   | L"e4"        | FString
+	//       [0] | 101 'e'      | wchar_t
+	//       [1] | 52 '4'       | wchar_t
+	//   > [4]   | L"c5"        | FString
+	//       [0] | 99 'c'       | wchar_t
+	//       [1] | 53 '5'       | wchar_t
+	SetType<FString> SetStrs = { "d1", "a2", "b3", "e4", "c5" };
+
+	// Name            | Value                | Type
+	// ----------------+----------------------+--------------------------------------------------------------------------------------------------------------------------------------
+	// > SetUniqueStrs | Num=5, Max=5         | TSet<TUniquePtr<FString,TDefaultDelete<FString>>,DefaultKeyFuncs<TUniquePtr<FString,TDefaultDelete<FString>>,0>,FDefaultSetAllocator>
+	//   > [0]         | Ptr=0x01234567 L"d1" | TUniquePtr<FString,TDefaultDelete<FString>>
+	//       [0]       | 100 'd'              | wchar_t
+	//       [1]       | 49 '1'               | wchar_t
+	//   > [1]         | Ptr=0x01234567 L"a2" | TUniquePtr<FString,TDefaultDelete<FString>>
+	//       [0]       | 97 'a'               | wchar_t
+	//       [1]       | 50 '2'               | wchar_t
+	//   > [2]         | Ptr=0x01234567 L"b3" | TUniquePtr<FString,TDefaultDelete<FString>>
+	//       [0]       | 98 'b'               | wchar_t
+	//       [1]       | 51 '3'               | wchar_t
+	//   > [3]         | Ptr=0x01234567 L"e4" | TUniquePtr<FString,TDefaultDelete<FString>>
+	//       [0]       | 101 'e'              | wchar_t
+	//       [1]       | 52 '4'               | wchar_t
+	//   > [4]         | Ptr=0x01234567 L"c5" | TUniquePtr<FString,TDefaultDelete<FString>>
+	//       [0]       | 99 'c'               | wchar_t
+	//       [1]       | 53 '5'               | wchar_t
+	SetType<TUniquePtr<FString>> SetUniqueStrs;
+
 	TSetNatvisTest()
 	{
-		// Name       | Value | Type
-		// -----------+-------+------------------------------------------------------
-		// > EmptySet | Empty | TSet<int,DefaultKeyFuncs<int,0>,FDefaultSetAllocator>
-		SetType<int32> EmptySet;
-
-		// Name       | Value         | Type
-		// -----------+---------------+------------------------------------------------------
-		// > SlackSet | Empty, Max=16 | TSet<int,DefaultKeyFuncs<int,0>,FDefaultSetAllocator>
-		SetType<int32> SlackSet;
 		SlackSet.Reserve(16);
-
-		// Name      | Value        | Type
-		// ----------+--------------+------------------------------------------------------
-		// > SetVals | Num=5, Max=5 | TSet<int,DefaultKeyFuncs<int,0>,FDefaultSetAllocator>
-		//     [0]   | 4            | int
-		//     [1]   | 1            | int
-		//     [2]   | 2            | int
-		//     [3]   | 5            | int
-		//     [4]   | 3            | int
-		SetType<int32> SetVals = { 4, 1, 2, 5, 3 };
-
-		// Name           | Value         | Type
-		// ---------------+---------------+------------------------------------------------------
-		// > SlackSetVals | Num=5, Max=16 | TSet<int,DefaultKeyFuncs<int,0>,FDefaultSetAllocator>
-		//     [0]        | 4             | int
-		//     [1]        | 1             | int
-		//     [2]        | 2             | int
-		//     [3]        | 5             | int
-		//     [4]        | 3             | int
-		SetType<int32> SlackSetVals = { 4, 1, 2, 5, 3 };
 		SlackSetVals.Reserve(16);
-
-		// Name            | Value        | Type
-		// ----------------+--------------+------------------------------------------------------
-		// > SparseSetVals | Num=3, Max=5 | TSet<int,DefaultKeyFuncs<int,0>,FDefaultSetAllocator>
-		//     [0]         | 4            | int
-		//     [1]         | 2            | int
-		//     [2]         | 3            | int
-		// Note: Order different for compact set
-		SetType<int32> SparseSetVals = { 4, 1, 2, 5, 3 };
 		SparseSetVals.Remove(1);
 		SparseSetVals.Remove(5);
-
-		// Name      | Value        | Type
-		// ----------+--------------+--------------------------------------------------------------
-		// > SetStrs | Num=5, Max=5 | TSet<FString,DefaultKeyFuncs<FString,0>,FDefaultSetAllocator>
-		//   > [0]   | L"d1"        | FString
-		//       [0] | 100 'd'      | wchar_t
-		//       [1] | 49 '1'       | wchar_t
-		//   > [1]   | L"a2"        | FString
-		//       [0] | 97 'a'       | wchar_t
-		//       [1] | 50 '2'       | wchar_t
-		//   > [2]   | L"b3"        | FString
-		//       [0] | 98 'b'       | wchar_t
-		//       [1] | 51 '3'       | wchar_t
-		//   > [3]   | L"e4"        | FString
-		//       [0] | 101 'e'      | wchar_t
-		//       [1] | 52 '4'       | wchar_t
-		//   > [4]   | L"c5"        | FString
-		//       [0] | 99 'c'       | wchar_t
-		//       [1] | 53 '5'       | wchar_t
-		SetType<FString> SetStrs = { "d1", "a2", "b3", "e4", "c5" };
-
-		// Name            | Value                | Type
-		// ----------------+----------------------+--------------------------------------------------------------------------------------------------------------------------------------
-		// > SetUniqueStrs | Num=5, Max=5         | TSet<TUniquePtr<FString,TDefaultDelete<FString>>,DefaultKeyFuncs<TUniquePtr<FString,TDefaultDelete<FString>>,0>,FDefaultSetAllocator>
-		//   > [0]         | Ptr=0x01234567 L"d1" | TUniquePtr<FString,TDefaultDelete<FString>>
-		//       [0]       | 100 'd'              | wchar_t
-		//       [1]       | 49 '1'               | wchar_t
-		//   > [1]         | Ptr=0x01234567 L"a2" | TUniquePtr<FString,TDefaultDelete<FString>>
-		//       [0]       | 97 'a'               | wchar_t
-		//       [1]       | 50 '2'               | wchar_t
-		//   > [2]         | Ptr=0x01234567 L"b3" | TUniquePtr<FString,TDefaultDelete<FString>>
-		//       [0]       | 98 'b'               | wchar_t
-		//       [1]       | 51 '3'               | wchar_t
-		//   > [3]         | Ptr=0x01234567 L"e4" | TUniquePtr<FString,TDefaultDelete<FString>>
-		//       [0]       | 101 'e'              | wchar_t
-		//       [1]       | 52 '4'               | wchar_t
-		//   > [4]         | Ptr=0x01234567 L"c5" | TUniquePtr<FString,TDefaultDelete<FString>>
-		//       [0]       | 99 'c'               | wchar_t
-		//       [1]       | 53 '5'               | wchar_t
-		SetType < TUniquePtr<FString> > SetUniqueStrs;
 		SetUniqueStrs.Reserve(5);
 		SetUniqueStrs.Add(MakeUnique<FString>("d1"));
 		SetUniqueStrs.Add(MakeUnique<FString>("a2"));
 		SetUniqueStrs.Add(MakeUnique<FString>("b3"));
 		SetUniqueStrs.Add(MakeUnique<FString>("e4"));
 		SetUniqueStrs.Add(MakeUnique<FString>("c5"));
-
-		PlaceBreakpointHere();
 	}
 };
 
@@ -1856,69 +1855,68 @@ struct TMapNatvisTest
 {
 	template <typename Key, typename Value>
 	using MapType = typename Resolver::template TMapType<Key, Value>;
-
-	void Run()
+	
+	// Name       | Value | Type
+	// -----------+-------+----------------------------------------------------------------------------------
+	// > EmptyMap | Empty | TMap<FString,int,FDefaultSetAllocator,TDefaultMapHashableKeyFuncs<FString,int,0>>
+	MapType<FString, int32> EmptyMap;
+
+	// Name       | Value         | Type
+	// -----------+---------------+------------------------------------------------------
+	// > SlackMap | Empty, Max=16 | TMap<FString,int,FDefaultSetAllocator,TDefaultMapHashableKeyFuncs<FString,int,0>>
+	MapType<FString, int32> SlackMap;
+
+	// Name       | Value        | Type
+	// -----------+--------------+------------------------------------------------------
+	// > MapVals  | Num=5, Max=5 | TMap<FString,int,FDefaultSetAllocator,TDefaultMapHashableKeyFuncs<FString,int,0>>
+	//     [L"d"] | 1            | int
+	//     [L"a"] | 2            | int
+	//     [L"b"] | 3            | int
+	//     [L"e"] | 4            | int
+	//     [L"c"] | 5            | int
+	MapType<FString, int32> MapVals = { { "d", 1 }, { "a", 2 }, { "b", 3 }, { "e", 4 }, { "c", 5 } };
+
+	// Name           | Value         | Type
+	// ---------------+---------------+------------------------------------------------------
+	// > SlackMapVals | Num=5, Max=16 | TMap<FString,int,FDefaultSetAllocator,TDefaultMapHashableKeyFuncs<FString,int,0>>
+	//     [L"d"]     | 1             | int
+	//     [L"a"]     | 2             | int
+	//     [L"b"]     | 3             | int
+	//     [L"e"]     | 4             | int
+	//     [L"c"]     | 5             | int
+	MapType<FString, int32> SlackMapVals = { { "d", 1 }, { "a", 2 }, { "b", 3 }, { "e", 4 }, { "c", 5 } };
+
+	// Name            | Value        | Type
+	// ----------------+--------------+------------------------------------------------------
+	// > SparseMapVals | Num=3, Max=5 | TMap<FString,int,FDefaultSetAllocator,TDefaultMapHashableKeyFuncs<FString,int,0>>
+	//     [L"d"]      | 1            | int
+	//     [L"b"]      | 3            | int
+	//     [L"c"]      | 5            | int
+	// Note: Order different for compact map
+	MapType<FString, int32> SparseMapVals = { { "d", 1 }, { "a", 2 }, { "b", 3 }, { "e", 4 }, { "c", 5 } };
+
+	// Name                      | Value        | Type
+	// --------------------------+--------------+----------------------------------------------------------------------------------------------------------------------------------------------------------
+	// > MapUniqueVals           | Num=5, Max=5 | TMap<TUniquePtr<FString,TDefaultDelete<FString>>,int,FDefaultSetAllocator,TDefaultMapHashableKeyFuncs<TUniquePtr<FString,TDefaultDelete<FString>>,int,0>>
+	//     [Ptr=0x12345678 L"d"] | 1            | int
+	//     [Ptr=0x12345678 L"a"] | 2            | int
+	//     [Ptr=0x12345678 L"b"] | 3            | int
+	//     [Ptr=0x12345678 L"e"] | 4            | int
+	//     [Ptr=0x12345678 L"c"] | 5            | int
+	MapType<TUniquePtr<FString>, int32> MapUniqueVals;
+
+	TMapNatvisTest()
 	{
-		// Name       | Value | Type
-		// -----------+-------+----------------------------------------------------------------------------------
-		// > EmptyMap | Empty | TMap<FString,int,FDefaultSetAllocator,TDefaultMapHashableKeyFuncs<FString,int,0>>
-		MapType<FString, int32> EmptyMap;
-
-		// Name       | Value         | Type
-		// -----------+---------------+------------------------------------------------------
-		// > SlackMap | Empty, Max=16 | TMap<FString,int,FDefaultSetAllocator,TDefaultMapHashableKeyFuncs<FString,int,0>>
-		MapType<FString, int32> SlackMap;
 		SlackMap.Reserve(16);
-
-		// Name       | Value        | Type
-		// -----------+--------------+------------------------------------------------------
-		// > MapVals  | Num=5, Max=5 | TMap<FString,int,FDefaultSetAllocator,TDefaultMapHashableKeyFuncs<FString,int,0>>
-		//     [L"d"] | 1            | int
-		//     [L"a"] | 2            | int
-		//     [L"b"] | 3            | int
-		//     [L"e"] | 4            | int
-		//     [L"c"] | 5            | int
-		MapType<FString, int32> MapVals = { { "d", 1 }, { "a", 2 }, { "b", 3 }, { "e", 4 }, { "c", 5 } };
-
-		// Name           | Value         | Type
-		// ---------------+---------------+------------------------------------------------------
-		// > SlackMapVals | Num=5, Max=16 | TMap<FString,int,FDefaultSetAllocator,TDefaultMapHashableKeyFuncs<FString,int,0>>
-		//     [L"d"]     | 1             | int
-		//     [L"a"]     | 2             | int
-		//     [L"b"]     | 3             | int
-		//     [L"e"]     | 4             | int
-		//     [L"c"]     | 5             | int
-		MapType<FString, int32> SlackMapVals = { { "d", 1 }, { "a", 2 }, { "b", 3 }, { "e", 4 }, { "c", 5 } };
 		SlackMapVals.Reserve(16);
-
-		// Name            | Value        | Type
-		// ----------------+--------------+------------------------------------------------------
-		// > SparseMapVals | Num=3, Max=5 | TMap<FString,int,FDefaultSetAllocator,TDefaultMapHashableKeyFuncs<FString,int,0>>
-		//     [L"d"]      | 1            | int
-		//     [L"b"]      | 3            | int
-		//     [L"c"]      | 5            | int
-		// Note: Order different for compact map
-		MapType<FString, int32> SparseMapVals = { { "d", 1 }, { "a", 2 }, { "b", 3 }, { "e", 4 }, { "c", 5 } };
 		SparseMapVals.Remove("a");
 		SparseMapVals.Remove("e");
-
-		// Name                      | Value        | Type
-		// --------------------------+--------------+----------------------------------------------------------------------------------------------------------------------------------------------------------
-		// > MapUniqueVals           | Num=5, Max=5 | TMap<TUniquePtr<FString,TDefaultDelete<FString>>,int,FDefaultSetAllocator,TDefaultMapHashableKeyFuncs<TUniquePtr<FString,TDefaultDelete<FString>>,int,0>>
-		//     [Ptr=0x12345678 L"d"] | 1            | int
-		//     [Ptr=0x12345678 L"a"] | 2            | int
-		//     [Ptr=0x12345678 L"b"] | 3            | int
-		//     [Ptr=0x12345678 L"e"] | 4            | int
-		//     [Ptr=0x12345678 L"c"] | 5            | int
-		MapType<TUniquePtr<FString>, int32> MapUniqueVals;
 		MapUniqueVals.Reserve(5);
 		MapUniqueVals.Emplace(MakeUnique<FString>("d"), 1);
 		MapUniqueVals.Emplace(MakeUnique<FString>("a"), 2);
 		MapUniqueVals.Emplace(MakeUnique<FString>("b"), 3);
 		MapUniqueVals.Emplace(MakeUnique<FString>("e"), 4);
 		MapUniqueVals.Emplace(MakeUnique<FString>("c"), 5);
-
-		PlaceBreakpointHere();
 	}
 };
 
@@ -2011,10 +2009,19 @@ bool FContainersNatvisTest::RunTest(const FString& Parameters)
 		PlaceBreakpointHere();
 	}
 
-	// TSet
-	TSetNatvisTest<FSetResolver>();
-	TSetNatvisTest<FSparseSetResolver>();
-	TSetNatvisTest<FCompactSetResolver>();
+	{
+		// TSet
+		TSetNatvisTest<FSetResolver> SetNatvisTest;
+		TSetNatvisTest<FSparseSetResolver> SparseSetNatvisTest;
+		TSetNatvisTest<FCompactSetResolver> CompactNatvisTest;
+
+		// TMap
+		TMapNatvisTest<FSetResolver> MapNatvisTest;
+		TMapNatvisTest<FSparseSetResolver> SparseMapNatvisTest;
+		TMapNatvisTest<FCompactSetResolver> CompactMapNatvisTest;
+
+		PlaceBreakpointHere();
+	}
 
 	// TSortedSet
 	{
@@ -2109,11 +2116,6 @@ bool FContainersNatvisTest::RunTest(const FString& Parameters)
 		PlaceBreakpointHere();
 	}
 
-	// TMap
-	TMapNatvisTest<FSetResolver>().Run();
-	TMapNatvisTest<FSparseSetResolver>().Run();
-	TMapNatvisTest<FCompactSetResolver>().Run();
-
 	// TSortedMap
 	{
 		// Name        | Value | Type
