TaskConcurrencyLimiter: Whitespace and comment changes for easier readability.

diff --git a/Engine/Source/Runtime/Core/Private/Tasks/TaskConcurrencyLimiter.cpp b/Engine/Source/Runtime/Core/Private/Tasks/TaskConcurrencyLimiter.cpp
--- a/Engine/Source/Runtime/Core/Private/Tasks/TaskConcurrencyLimiter.cpp
+++ b/Engine/Source/Runtime/Core/Private/Tasks/TaskConcurrencyLimiter.cpp
@@ -11,136 +11,148 @@
 
 namespace UE::Tasks::TaskConcurrencyLimiter_Private
 {
-	FPimpl::~FPimpl()
+
+FPimpl::~FPimpl()
+{
+	if (FEvent* Event = CompletionEvent.exchange(nullptr, std::memory_order_relaxed))
 	{
-		if (FEvent* Event = CompletionEvent.exchange(nullptr, std::memory_order_relaxed))
-		{
-			FPlatformProcess::ReturnSynchEventToPool(Event);
-		}
+		FPlatformProcess::ReturnSynchEventToPool(Event);
 	}
+}
 
-	void FPimpl::AddWorkItem(FLimiterTask* Task)
-	{
-		NumWorkItems.fetch_add(1, std::memory_order_acquire);
+void FPimpl::AddWorkItem(FLimiterTask* Task)
+{
+	NumWorkItems.fetch_add(1, std::memory_order_acquire);
 
-		WorkQueue.Push(Task);
+	WorkQueue.Push(Task);
 
-		uint32 ConcurrencySlot;
-		if (ConcurrencySlots.Alloc(ConcurrencySlot))
-		{
-			ProcessQueueFromPush(ConcurrencySlot);
-		}
+	uint32 ConcurrencySlot;
+	if (ConcurrencySlots.Alloc(ConcurrencySlot))
+	{
+		ProcessQueueFromPush(ConcurrencySlot);
 	}
+}
 
-	bool FPimpl::Wait(FTimespan Timeout)
+bool FPimpl::Wait(FTimespan Timeout)
+{
+	if (NumWorkItems.load(std::memory_order_relaxed) == 0)
 	{
-		if (NumWorkItems.load(std::memory_order_relaxed) == 0)
-		{
-			return true;
-		}
+		return true;
+	}
 
-		FEvent* LocalCompletionEvent = CompletionEvent.load(std::memory_order_acquire);
-		if (LocalCompletionEvent == nullptr)
+	FEvent* LocalCompletionEvent = CompletionEvent.load(std::memory_order_acquire);
+	if (LocalCompletionEvent == nullptr)
+	{
+		FEvent* NewEvent = FPlatformProcess::GetSynchEventFromPool(true);
+		if (!CompletionEvent.compare_exchange_strong(LocalCompletionEvent, NewEvent, std::memory_order_acq_rel, std::memory_order_acquire))
 		{
-			FEvent* NewEvent = FPlatformProcess::GetSynchEventFromPool(true);
-			if (!CompletionEvent.compare_exchange_strong(LocalCompletionEvent, NewEvent, std::memory_order_acq_rel, std::memory_order_acquire))
-			{
-				FPlatformProcess::ReturnSynchEventToPool(NewEvent);
-			}
-			else
-			{
-				LocalCompletionEvent = NewEvent;
-			}
+			FPlatformProcess::ReturnSynchEventToPool(NewEvent);
 		}
-
-		if (NumWorkItems.load(std::memory_order_acquire) == 0)
+		else
 		{
-			return true;
+			LocalCompletionEvent = NewEvent;
 		}
+	}
 
-		// We must drain all work we can find via retraction since 
-		// we can't be sure if the worker threads are able to proceed.
-		bool bDidSomething;
-		do
-		{
-			bDidSomething = false;
-			// Before we wait on our completion event, find tasks that are currently in the task system and attempt to run them inline
-			// since it's possible task threads are waiting on work that was submitted to the limiter and may not make progress otherwise
-			for (int32 SlotIndex = 0; SlotIndex < ScheduledTasks.Num(); ++SlotIndex)
-			{
-				// If the task is being executed then the value will likely be null. However, it's possible the task is 
-				// being executed but hasn't cleared this state yet. That is _fine_ since it means our expedite 
-				// below will simply fail but it means we are now responsible for destroying the task and must call Release.
-				if (FLimiterTask* LimiterTask = ScheduledTasks[SlotIndex].Task.exchange(nullptr, std::memory_order_acquire))
-				{
-					bDidSomething = true;
-					LimiterTask->Task.TryExpedite();
-					LimiterTask->Release(); 
-				}
-			}
-		} while (bDidSomething || NumWorkItems.load(std::memory_order_acquire) != 0);
-
-		// This should fall-through immediately
-		return LocalCompletionEvent->Wait(Timeout);
+	if (NumWorkItems.load(std::memory_order_acquire) == 0)
+	{
+		return true;
 	}
 
-	void FPimpl::ProcessQueue(uint32 ConcurrencySlot, bool bSkipFirstWakeUp)
+	// We must use the global TaskScheduler's retraction feature (via the call to TryExpedite) to drain all tasks we
+	// can find that are queued in the global TaskScheduler but not yet started, because at any time all of the
+	// workerthreads in the global TaskScheduler might start working on tasks from another system (e.g. one calling
+	// ParallelFor) and be unable to immediately work on the task that we launched. And if all of those tasks block on
+	// a resource that our waitthread is holding (e.g. the AssetRegistry's internal critical section), then we would
+	// deadlock. For simplicity we wait for a task to be launched into the global TaskScheduler before retracting the
+	// task; this is a slight performance wastage but prevents us from having to duplicate or generalize the
+	// ProcessQueue code to be called from here as well.
+	bool bDidSomething;
+	do
 	{
-		bool bWakeUpWorker = !bSkipFirstWakeUp;
-		do
+		bDidSomething = false;
+		for (int32 SlotIndex = 0; SlotIndex < ScheduledTasks.Num(); ++SlotIndex)
 		{
-			if (FLimiterTask* LimiterTask = WorkQueue.Pop())
-			{
-				LowLevelTasks::FTask& Task = LimiterTask->Task;
-
-				// Now that we know the ConcurrencySlot, set it at launch time so
-				// the executor can retrieve it.
-				Task.SetUserData((void*)(UPTRINT)ConcurrencySlot);
-
-				// Before we launch the job, increment the task lifetime since the job may run and complete before we assign to ScheduledTasks
-				LimiterTask->AddRef();
-				LowLevelTasks::TryLaunch(Task, bWakeUpWorker ? LowLevelTasks::EQueuePreference::GlobalQueuePreference : LowLevelTasks::EQueuePreference::LocalQueuePreference, bWakeUpWorker);
-				ScheduledTasks[ConcurrencySlot].Task.store(LimiterTask, std::memory_order_release); // Signal this task has now been scheduled
-			}
-			else
+			// If the task is being executed then the value will likely be null. However, it's possible the task is 
+			// being executed but hasn't cleared this state yet. That is _fine_ since it means our expedite 
+			// below will simply fail but it means we are now responsible for destroying the task and must call Release.
+			if (FLimiterTask* LimiterTask = ScheduledTasks[SlotIndex].Task.exchange(nullptr, std::memory_order_acquire))
 			{
-				ConcurrencySlots.Release(ConcurrencySlot);
-				break;
+				bDidSomething = true;
+				LimiterTask->Task.TryExpedite();
+				LimiterTask->Release(); 
 			}
+		}
+	// Continue looking for retraction until all pushed items have completed. Stopping our search for retractables
+	// as soon as we find nothing is retractable would potentially deadlock in the case that a workerthread afterwards
+	// finishes our task that it was working on, queues our next task, but then is assigned another system's before it
+	// can start working on the next task for us.
+	} while (bDidSomething || NumWorkItems.load(std::memory_order_acquire) != 0);
+
+	// This should fall-through immediately because NumWorkItems is 0 and so we called Trigger during the last
+	// CompleteTask.
+	return LocalCompletionEvent->Wait(Timeout);
+}
+
+void FPimpl::ProcessQueue(uint32 ConcurrencySlot, bool bSkipFirstWakeUp)
+{
+	bool bWakeUpWorker = !bSkipFirstWakeUp;
+	do
+	{
+		if (FLimiterTask* LimiterTask = WorkQueue.Pop())
+		{
+			LowLevelTasks::FTask& Task = LimiterTask->Task;
+
+			// Now that we know the ConcurrencySlot, set it at launch time so
+			// the executor can retrieve it.
+			Task.SetUserData((void*)(UPTRINT)ConcurrencySlot);
+
+			// Before we launch the job, increment the task lifetime since the job may run and complete before we assign to ScheduledTasks
+			LimiterTask->AddRef();
+			LowLevelTasks::TryLaunch(Task, bWakeUpWorker ? LowLevelTasks::EQueuePreference::GlobalQueuePreference : LowLevelTasks::EQueuePreference::LocalQueuePreference, bWakeUpWorker);
+			// TODO: Leaked pointer if the task completes before we call store. We never remove the task from ScheduledTasks in that case,
+			// and never Release it, and we either close the TaskConcurrencyLimiter or we overwrite it with the next task for the same ConcurrencySlot.
+			ScheduledTasks[ConcurrencySlot].Task.store(LimiterTask, std::memory_order_release); // Signal this task has now been scheduled
+		}
+		else
+		{
+			ConcurrencySlots.Release(ConcurrencySlot);
+			break;
+		}
 
-			// Don't skip wake-up if we launch any additional tasks.
-			bWakeUpWorker = true;
+		// Don't skip wake-up if we launch any additional tasks.
+		bWakeUpWorker = true;
 
-		} while (ConcurrencySlots.Alloc(ConcurrencySlot));
-	}
+	} while (ConcurrencySlots.Alloc(ConcurrencySlot));
+}
 
-	void FPimpl::ProcessQueueFromWorker(uint32 ConcurrencySlot)
-	{
-		// Once we are in a worker thread, we want to schedule on the local queue without waking up additional workers
-		// to allow our own worker to pick up the next item and avoid wake-up cost / context switch.
-		// We must check IsInWorkerThread() since it's possible we have retracted the work to run outside of a worker thread
-		ProcessQueue(ConcurrencySlot, IsInWorkerThread());
-	}
+void FPimpl::ProcessQueueFromWorker(uint32 ConcurrencySlot)
+{
+	// Once we are in a worker thread, we want to schedule on the local queue without waking up additional workers
+	// to allow our own worker to pick up the next item and avoid wake-up cost / context switch.
+	// We must check IsInWorkerThread() since it's possible we have retracted the work to run outside of a worker thread
+	ProcessQueue(ConcurrencySlot, IsInWorkerThread());
+}
 
-	void FPimpl::ProcessQueueFromPush(uint32 ConcurrencySlot)
-	{
-		// When we push new items, we don't want to skip any wake-up.
-		static constexpr bool bSkipFirstWakeUp = false;
+void FPimpl::ProcessQueueFromPush(uint32 ConcurrencySlot)
+{
+	// When we push new items, we don't want to skip any wake-up.
+	static constexpr bool bSkipFirstWakeUp = false;
 
-		ProcessQueue(ConcurrencySlot, bSkipFirstWakeUp);
-	}
+	ProcessQueue(ConcurrencySlot, bSkipFirstWakeUp);
+}
 
-	void FPimpl::CompleteWorkItem(uint32 ConcurrencySlot)
+void FPimpl::CompleteWorkItem(uint32 ConcurrencySlot)
+{
+	if (NumWorkItems.fetch_sub(1, std::memory_order_release) == 1)
 	{
-		if (NumWorkItems.fetch_sub(1, std::memory_order_release) == 1)
+		if (FEvent* LocalCompletionEvent = CompletionEvent.load(std::memory_order_acquire))
 		{
-			if (FEvent* LocalCompletionEvent = CompletionEvent.load(std::memory_order_acquire))
-			{
-				LocalCompletionEvent->Trigger();
-			}
+			LocalCompletionEvent->Trigger();
 		}
-
-		ProcessQueueFromWorker(ConcurrencySlot);
 	}
 
+	ProcessQueueFromWorker(ConcurrencySlot);
+}
+
 } // namespace UE::Tasks::TaskConcurrencyLimiter_Private
diff --git a/Engine/Source/Runtime/Core/Public/Tasks/TaskConcurrencyLimiter.h b/Engine/Source/Runtime/Core/Public/Tasks/TaskConcurrencyLimiter.h
--- a/Engine/Source/Runtime/Core/Public/Tasks/TaskConcurrencyLimiter.h
+++ b/Engine/Source/Runtime/Core/Public/Tasks/TaskConcurrencyLimiter.h
@@ -10,181 +10,210 @@
 
 #include <AtomicQueue.h>
 
-namespace UE::Tasks
+namespace UE::Tasks::TaskConcurrencyLimiter_Private
+{
+
+// a queue of free slots in range [0 .. max_concurrency). Initially contains all slots in the range.
+class FConcurrencySlots
 {
-	namespace TaskConcurrencyLimiter_Private
+public:
+	explicit FConcurrencySlots(uint32 MaxConcurrency)
+		: FreeSlots(MaxConcurrency)
 	{
-		// a queue of free slots in range [0 .. max_concurrency). Initially contains all slots in the range.
-		class FConcurrencySlots
+		for (uint32 Index = IndexOffset; Index < MaxConcurrency + IndexOffset; ++Index)
 		{
-		public:
-			explicit FConcurrencySlots(uint32 MaxConcurrency)
-				: FreeSlots(MaxConcurrency)
-			{
-				for (uint32 Index = IndexOffset; Index < MaxConcurrency + IndexOffset; ++Index)
-				{
-					FreeSlots.push(Index);
-				}
-			}
+			FreeSlots.push(Index);
+		}
+	}
 
-			bool Alloc(uint32& Slot)
-			{
-				if (FreeSlots.try_pop(Slot))
-				{
-					Slot -= IndexOffset;
-					return true;
-				}
+	bool Alloc(uint32& Slot)
+	{
+		if (FreeSlots.try_pop(Slot))
+		{
+			Slot -= IndexOffset;
+			return true;
+		}
 
-				return false;
-			}
+		return false;
+	}
 
-			void Release(uint32 Slot)
-			{
-				FreeSlots.push(Slot + IndexOffset);
-			}
+	void Release(uint32 Slot)
+	{
+		FreeSlots.push(Slot + IndexOffset);
+	}
+
+private:
+	// this queue uses 0 as a special "null" value. to work around this, slots are shifted by one for storage, thus ending up in 
+	// [1 .. max_concurrency] range
+	static constexpr int32 IndexOffset = 1;
+	atomic_queue::AtomicQueueB<uint32> FreeSlots; // a bounded lock-free FIFO queue
+};
+
+// an implementation details of FTaskConcurrenctyLimiter
+class FPimpl : public TSharedFromThis<FPimpl>
+{
+	// Reference counted wrapper for FTask to manage the task lifetime between worker threads and 
+	// any thread that might retract the task once it is submitted to the scheduler
+	struct FLimiterTask : FRefCountBase
+	{
+		LowLevelTasks::FTask Task;
+	};
+
+public:
+	explicit FPimpl(uint32 InMaxConcurrency, ETaskPriority InTaskPriority)
+		: ConcurrencySlots(InMaxConcurrency)
+		, TaskPriority(InTaskPriority)
+	{
+		ScheduledTasks.SetNum(InMaxConcurrency);
+	}
 
-		private:
-			// this queue uses 0 as a special "null" value. to work around this, slots are shifted by one for storage, thus ending up in 
-			// [1 .. max_concurrency] range
-			static constexpr int32 IndexOffset = 1;
-			atomic_queue::AtomicQueueB<uint32> FreeSlots; // a bounded lock-free FIFO queue
-		};
+	CORE_API ~FPimpl();
 
-		// an implementation details of FTaskConcurrenctyLimiter
-		class FPimpl : public TSharedFromThis<FPimpl>
-		{
-			// Reference counted wrapper for FTask to manage the task lifetime between worker threads and 
-			// any thread that might retract the task once it is submitted to the scheduler
-			struct FLimiterTask : FRefCountBase
-			{
-				LowLevelTasks::FTask Task;
-			};
+	template<typename TaskFunctionType>
+	void Push(const TCHAR* DebugName, TaskFunctionType&& TaskFunction);
 
-		public:
-			explicit FPimpl(uint32 InMaxConcurrency, ETaskPriority InTaskPriority)
-				: ConcurrencySlots(InMaxConcurrency)
-				, TaskPriority(InTaskPriority)
-			{
-				ScheduledTasks.SetNum(InMaxConcurrency);
-			}
+	bool CORE_API Wait(FTimespan Timeout);
 
-			CORE_API ~FPimpl();
+private:
+	void CORE_API AddWorkItem(FLimiterTask* Task);
+	void CORE_API ProcessQueue(uint32 ConcurrencySlot, bool bSkipFirstWakeUp);
+	void CORE_API ProcessQueueFromWorker(uint32 ConcurrencySlot);
+	void CORE_API ProcessQueueFromPush(uint32 ConcurrencySlot);
+	void CORE_API CompleteWorkItem(uint32 ConcurrencySlot);
 
-			template<typename TaskFunctionType>
-			void Push(const TCHAR* DebugName, TaskFunctionType&& TaskFunction)
+	// Padded struct to prevent false sharing as we read and write to the ScheduledTasks array			
+	struct alignas(PLATFORM_CACHE_LINE_SIZE) FPaddedSharedTask
+	{
+		std::atomic<FLimiterTask*> Task;
+	};
+
+	FConcurrencySlots ConcurrencySlots; // free slots queue. used also to limit concurrency
+	TLockFreePointerListFIFO<FLimiterTask, PLATFORM_CACHE_LINE_SIZE> WorkQueue; // a queue of user-provided task functions
+	TArray<FPaddedSharedTask> ScheduledTasks;
+	std::atomic<FEvent*> CompletionEvent { nullptr };
+	std::atomic<uint32> NumWorkItems { 0 };
+	ETaskPriority TaskPriority;
+};
+
+///////////////////////////////////////////////////////
+// FPimpl's template implementations
+///////////////////////////////////////////////////////
+
+template<typename TaskFunctionType>
+void FPimpl::Push(const TCHAR* DebugName, TaskFunctionType&& TaskFunction)
+{
+	FLimiterTask* LimiterTask = new FLimiterTask();
+	LowLevelTasks::FTask& Task = LimiterTask->Task;
+	Task.Init(
+		DebugName,
+		TaskPriority,
+		[
+			TaskFunction = MoveTemp(TaskFunction),
+			this,
+			// Assign Pimpl to keep the TaskConcurrenyLimiter alive as long as the Task we created is alive.
+		Pimpl = TSharedFromThis<FPimpl>::AsShared(),
+		// Assign FLimiterTask wrapper into a TRefCountPtr to keep the LimiterTask alive as long as its inner Task is alive.
+		LimiterTask = TRefCountPtr<FLimiterTask>(LimiterTask)
+		]()
+		{
+			LowLevelTasks::FTask& Task = LimiterTask->Task;
+
+			// We can't pass the ConcurrencySlot in the lambda during creation as
+			// it's not actually acquired yet. The value is assigned during ProcessQueue
+			// after we have called Alloc and received a ConcurrencySlot for it.
+			uint32 ConcurrencySlot = (uint32)(UPTRINT)Task.GetUserData();
+
+			// This task is now executing, either on a taskthread that received the task we created, or on
+			// the thread that called Wait if the waitthread reached the TryExpedite call before the
+			// taskthread started the Task.
+			// In the case that we are in the taskthread, we need to implement our two contracts:
+			//    1. Remove Task from ScheduledTasks as soon as it starts.
+			//		 This prevent waitthread wastefully calling TryExpedite in most cases.
+			//           Note a useless TryExpedite will still occur and be a noop if the threads are racing.
+			//       Remove the pointer so that we can later safely free the Tasks's memory.
+			//    2. Call Release on Task when removing it from ScheduledTasks.
+			//       This balances AddRef done in ProcessQueue.
+			// Note that if we reach here in the waitthread, then we have already removed the task from
+			// ScheduledTasks and this call to exchange is a noop and we set bOwnLimiterTask=false.
+			bool bOwnLimiterTask = ScheduledTasks[ConcurrencySlot].Task.exchange(nullptr, std::memory_order_release) != nullptr;
+			if (bOwnLimiterTask)
 			{
-				FLimiterTask* LimiterTask = new FLimiterTask();
-				LowLevelTasks::FTask& Task = LimiterTask->Task;
-				Task.Init(
-					DebugName,
-					TaskPriority,
-					[
-						TaskFunction = MoveTemp(TaskFunction),
-						this,
-						Pimpl = TSharedFromThis<FPimpl>::AsShared(), // to keep it alive
-						LimiterTask = TRefCountPtr<FLimiterTask>(LimiterTask) // self-destruct
-					]()
-					{
-						LowLevelTasks::FTask& Task = LimiterTask->Task;
-
-						// We can't pass the ConcurrencySlot in the lambda during creation as
-						// it's not actually acquired yet. The value will be passed using
-						// the user data when the task is launched.
-						uint32 ConcurrencySlot = (uint32)(UPTRINT)Task.GetUserData();
-
-						// This task is now executing; remove it from the ScheduledTask array so it cannot be retracted.
-						// It's possible a thread waiting on the limiter might have retracted this task before we can clear it. 
-						bool bOwnLimiterTask = ScheduledTasks[ConcurrencySlot].Task.exchange(nullptr, std::memory_order_release) != nullptr;
-						if (bOwnLimiterTask)
-						{
-							// Whomever obtains the task from ScheduledTasks must release the reference
-							LimiterTask->Release();
-						}
-
-						TaskFunction(ConcurrencySlot);
-						CompleteWorkItem(ConcurrencySlot);
-					}
-				);
-
-				AddWorkItem(LimiterTask);
+				// Whomever obtains the task from ScheduledTasks must release the reference
+				LimiterTask->Release();
 			}
 
-			bool CORE_API Wait(FTimespan Timeout);
+			TaskFunction(ConcurrencySlot);
 
-		private:
-			void CORE_API AddWorkItem(FLimiterTask* Task);
-			void CORE_API ProcessQueue(uint32 ConcurrencySlot, bool bSkipFirstWakeUp);
-			void CORE_API ProcessQueueFromWorker(uint32 ConcurrencySlot);
-			void CORE_API ProcessQueueFromPush(uint32 ConcurrencySlot);
-			void CORE_API CompleteWorkItem(uint32 ConcurrencySlot);
+			// Finish bookkeeping for the current task and look for more tasks to launch.
+			// If we find another task to launch, have it use the ConcurrencySlot we just finished
+			// with; note this means use of a ConcurrencySlot may shift around between threads.
+			// For the first task we find, bias it towards the current worker thread if possible.
+			// If we find more than one, kick off a new Task for the global scheduler to schedule,
+			// for every task after the first.
+			CompleteWorkItem(ConcurrencySlot);
+		}
+	);
 
-			// Padded struct to prevent false sharing as we read and write to the ScheduledTasks array			
-			struct alignas(PLATFORM_CACHE_LINE_SIZE) FPaddedSharedTask
-			{
-				std::atomic<FLimiterTask*> Task;
-			};
+	AddWorkItem(LimiterTask);
+}
+
+} // namespace UE::Tasks::TaskConcurrencyLimiter_Private
 
-			FConcurrencySlots ConcurrencySlots; // free slots queue. used also to limit concurrency
-			TLockFreePointerListFIFO<FLimiterTask, PLATFORM_CACHE_LINE_SIZE> WorkQueue; // a queue of user-provided task functions
-			TArray<FPaddedSharedTask> ScheduledTasks;
-			std::atomic<FEvent*> CompletionEvent { nullptr };
-			std::atomic<uint32> NumWorkItems { 0 };
-			ETaskPriority TaskPriority;
-		};
 
-	} // namespace TaskConcurrencyLimiter_Private
+namespace UE::Tasks
+{
 
+/**
+* A lightweight construct that limits the concurrency of tasks pushed into it. 
+*
+* @note This class supports being destroyed before the tasks it contains are finished.
+*/
+class FTaskConcurrencyLimiter
+{
+public:
 	/**
-	* A lightweight construct that limits the concurrency of tasks pushed into it. 
-	*
-	* @note This class supports being destroyed before the tasks it contains are finished.
-	*/
-	class FTaskConcurrencyLimiter
+		* Constructor.
+		*
+		* @param MaxConcurrency     How wide the processing can go.
+		* @param TaskPriority       Priority the tasks will be launched with.
+		*/
+	explicit FTaskConcurrencyLimiter(uint32 MaxConcurrency, ETaskPriority TaskPriority = ETaskPriority::Default)
+		: Pimpl(MakeShared<TaskConcurrencyLimiter_Private::FPimpl>(MaxConcurrency, TaskPriority))
 	{
-	public:
-		/**
-		 * Constructor.
-		 *
-		 * @param MaxConcurrency     How wide the processing can go.
-		 * @param TaskPriority       Priority the tasks will be launched with.
-		 */
-		explicit FTaskConcurrencyLimiter(uint32 MaxConcurrency, ETaskPriority TaskPriority = ETaskPriority::Default)
-			: Pimpl(MakeShared<TaskConcurrencyLimiter_Private::FPimpl>(MaxConcurrency, TaskPriority))
-		{
-		}
+	}
 
-		/**
-		 * Push a new task.
-		 *
-		 * @param DebugName    Helps to identify the task in debugger and profiler.
-		 * @param TaskFunction A callable with a slot parameter, usually a lambda but can be also a functor object 
-		 *                     or a pointer to a function. The slot parameter is an index in [0..max_concurrency) range, 
-		 *                     unique at any moment of time, that can be used in user code to index a fixed-size buffer. 
-		 *                     See `TaskConcurrencyLimiterStressTest()` for an example.
-		 */
-		template<typename TaskFunctionType>
-		void Push(const TCHAR* DebugName, TaskFunctionType&& TaskFunction)
-		{
-			Pimpl->Push(DebugName, MoveTemp(TaskFunction));
-		}
+	/**
+		* Push a new task.
+		*
+		* @param DebugName    Helps to identify the task in debugger and profiler.
+		* @param TaskFunction A callable with a slot parameter, usually a lambda but can be also a functor object 
+		*                     or a pointer to a function. The slot parameter is an index in [0..max_concurrency) range, 
+		*                     unique at any moment of time, that can be used in user code to index a fixed-size buffer. 
+		*                     See `TaskConcurrencyLimiterStressTest()` for an example.
+		*/
+	template<typename TaskFunctionType>
+	void Push(const TCHAR* DebugName, TaskFunctionType&& TaskFunction)
+	{
+		Pimpl->Push(DebugName, MoveTemp(TaskFunction));
+	}
 
-		/**
-		 * Waits for task's completion with timeout.
-		 *
-		 * @param  Timeout Maximum amount of time to wait for tasks to finish before returning.
-		 * @return true if all tasks are completed, false otherwise.
-		 * 
-		 * @note   A wait is satisfied once the internal task counter reaches 0 and is never reset
-		 *         afterward when more tasks are added. A new FTaskConcurrencyLimiter can be used
-		 *         for such a use case.
-		 */
-		bool Wait(FTimespan Timeout = FTimespan::MaxValue())
-		{
-			return Pimpl->Wait(Timeout);
-		}
+	/**
+		* Waits for task's completion with timeout.
+		*
+		* @param  Timeout Maximum amount of time to wait for tasks to finish before returning.
+		* @return true if all tasks are completed, false otherwise.
+		* 
+		* @note   A wait is satisfied once the internal task counter reaches 0 and is never reset
+		*         afterward when more tasks are added. A new FTaskConcurrencyLimiter can be used
+		*         for such a use case.
+		*/
+	bool Wait(FTimespan Timeout = FTimespan::MaxValue())
+	{
+		return Pimpl->Wait(Timeout);
+	}
 
-	private:
-		TSharedRef<TaskConcurrencyLimiter_Private::FPimpl> Pimpl;
-	};
+private:
+	TSharedRef<TaskConcurrencyLimiter_Private::FPimpl> Pimpl;
+};
 
 } // namespace UE::Tasks
