[Tasks]

Fix for a possible task ptr leak in the TaskLimiter ScheduledTasks array which could lead to simultaneous access to the same task slot.

Previously, in ProcessQueue we AddRef our task and then schedule the task before writing the task to the ScheduledTasks array. However if the task executes immediately on a task thread, the task thread won't take ownership of the Task and won't call Release() on the job to match the AddRef in ProcessQueue. ProcessQueue will then store the completed task ptr to the ScheduledTasks list for a task that has completed. It is then possible a new scheduled task can write on top of that ScheduledTask slot (since the slot was released when the task completed) meaning we could lose track of the ptr and leak the task.

We now always write the ScheduledTasks before launching to resolve this issue. Wait() now has the possibility of retracting a task that hasn't been scheduled and that is resolved by the retracted tasks now calling TryExecute and if that fails then calling TryExpedite.

#rnx

diff --git a/Engine/Source/Runtime/Core/Private/Tasks/TaskConcurrencyLimiter.cpp b/Engine/Source/Runtime/Core/Private/Tasks/TaskConcurrencyLimiter.cpp
--- a/Engine/Source/Runtime/Core/Private/Tasks/TaskConcurrencyLimiter.cpp
+++ b/Engine/Source/Runtime/Core/Private/Tasks/TaskConcurrencyLimiter.cpp
@@ -79,7 +79,13 @@ bool FPimpl::Wait(FTimespan Timeout)
 			if (FLimiterTask* LimiterTask = ScheduledTasks[SlotIndex].Task.exchange(nullptr, std::memory_order_acquire))
 			{
 				bDidSomething = true;
-				LimiterTask->Task.TryExpedite();
+				// We populate ScheduledTasks before launching the task, as such we may have retracted a task
+				// that hasn't launched yet. To handle that case we TryExecute, and if that fails 
+				// (because we are already scheduled) we then expedite the task
+				if (!LimiterTask->Task.TryExecute())
+				{
+					LimiterTask->Task.TryExpedite();
+				}
 				LimiterTask->Release(); 
 			}
 		}
@@ -107,12 +113,12 @@ void FPimpl::ProcessQueue(uint32 ConcurrencySlot, bool bSkipFirstWakeUp)
 			// the executor can retrieve it.
 			Task.SetUserData((void*)(UPTRINT)ConcurrencySlot);
 
-			// Before we launch the job, increment the task lifetime since the job may run and complete before we assign to ScheduledTasks
+			// Increment the refcount and set the ScheduledTask. As soon as ScheduledTasks is set the task might 
+			// be retracted, including before it has been scheduled here. As a result, TryLaunch may fail but that 
+			// will only occur when TryExecute on the retracting thread has succeeded ensuring our task has been run.
 			LimiterTask->AddRef();
-			LowLevelTasks::TryLaunch(Task, bWakeUpWorker ? LowLevelTasks::EQueuePreference::GlobalQueuePreference : LowLevelTasks::EQueuePreference::LocalQueuePreference, bWakeUpWorker);
-			// TODO: Leaked pointer if the task completes before we call store. We never remove the task from ScheduledTasks in that case,
-			// and never Release it, and we either close the TaskConcurrencyLimiter or we overwrite it with the next task for the same ConcurrencySlot.
 			ScheduledTasks[ConcurrencySlot].Task.store(LimiterTask, std::memory_order_release); // Signal this task has now been scheduled
+			LowLevelTasks::TryLaunch(Task, bWakeUpWorker ? LowLevelTasks::EQueuePreference::GlobalQueuePreference : LowLevelTasks::EQueuePreference::LocalQueuePreference, bWakeUpWorker);
 		}
 		else
 		{
