[ChunkedArray] AddElement with rvalue reference.

diff --git a/Engine/Source/Runtime/Core/Public/Containers/ChunkedArray.h b/Engine/Source/Runtime/Core/Public/Containers/ChunkedArray.h
--- a/Engine/Source/Runtime/Core/Public/Containers/ChunkedArray.h
+++ b/Engine/Source/Runtime/Core/Public/Containers/ChunkedArray.h
@@ -184,6 +184,18 @@ public:
 		return Index >= 0 && Index < NumElements;
 	}
 
+	/**
+	 * Adds a new item to the end of the chunked array.
+	 *
+	 * @param Item	The item to add
+	 * @return		Index to the new item
+	 */
+	int32 AddElement(ElementType&& Item)
+	{
+		new(*this) ElementType(MoveTempIfPossible(Item));
+		return this->NumElements - 1;
+	}
+
 	/**
 	 * Adds a new item to the end of the chunked array.
 	 *
