Fix support for iOS shared file containers
- fix method name (it's not a url, it's a path)
- give access to the shared container file location

#rnx

diff --git a/Engine/Source/Runtime/Core/Private/IOS/IOSPlatformFile.cpp b/Engine/Source/Runtime/Core/Private/IOS/IOSPlatformFile.cpp
--- a/Engine/Source/Runtime/Core/Private/IOS/IOSPlatformFile.cpp
+++ b/Engine/Source/Runtime/Core/Private/IOS/IOSPlatformFile.cpp
@@ -964,7 +964,9 @@ bool FIOSPlatformFile::IterateDirectoryCommon(const TCHAR* Directory, const TFun
 FString FIOSPlatformFile::ConvertToPlatformPath(const FString& Filename, bool bForWrite, bool bIsPublicWrite)
 {
 	FString Result = Filename;
-	if (Result.Contains(TEXT("/OnDemandResources/")) || Result.StartsWith(TEXT("/var/")))
+	if (Result.Contains(TEXT("/OnDemandResources/")) || 
+		Result.StartsWith(TEXT("/private/var/")) || 
+		Result.StartsWith(TEXT("/var/")))
 	{
 		return Result;
 	}
@@ -1038,7 +1040,7 @@ FString FIOSPlatformFile::ConvertToPlatformPath(const FString& Filename, bool bF
 	return Result;
 }
 
-FString FIOSPlatformFile::GetSharedContainerURL(const TCHAR* AppGroupID)
+FString FIOSPlatformFile::GetSharedContainerPath(const TCHAR* AppGroupID)
 {
 	SCOPED_AUTORELEASE_POOL;
 	
@@ -1048,11 +1050,10 @@ FString FIOSPlatformFile::GetSharedContainerURL(const TCHAR* AppGroupID)
 	
 	if (ContainerUrl != nil)
 	{
-		NSLog(@"SharedContainer:%@", [ContainerUrl path]);
 		return FString([ContainerUrl path]);
 	}
 	
-	NSLog(@"No SharedContainer with AppGroupID:%@", GroupID);
+	NSLog(@"[IOSPlatformFile] No SharedContainer with AppGroupID:%@", GroupID);
 	return ""; 
 }
 
diff --git a/Engine/Source/Runtime/Core/Public/IOS/IOSPlatformFile.h b/Engine/Source/Runtime/Core/Public/IOS/IOSPlatformFile.h
--- a/Engine/Source/Runtime/Core/Public/IOS/IOSPlatformFile.h
+++ b/Engine/Source/Runtime/Core/Public/IOS/IOSPlatformFile.h
@@ -54,8 +54,8 @@ public:
 	virtual void SetCreatePublicFiles(bool bCreatePublicFilesIn) override;
 	
 	// If app has a Shared File Container via the AppGroups capability, return the
-	// absolute url path or an empty string if it's not found
-	static FString GetSharedContainerURL(const TCHAR* AppGroupID);
+	// path or an empty string if it's not found
+	static FString GetSharedContainerPath(const TCHAR* AppGroupID);
 	
 private:
 	bool IterateDirectoryCommon(const TCHAR* Directory, const TFunctionRef<bool(struct dirent*)>& Visitor);
