Update to PVS 7.39
#rnx

diff --git a/Engine/Source/Runtime/Core/Public/Analysis/PvsStudioCodeAnalysis.h b/Engine/Source/Runtime/Core/Public/Analysis/PvsStudioCodeAnalysis.h
--- a/Engine/Source/Runtime/Core/Public/Analysis/PvsStudioCodeAnalysis.h
+++ b/Engine/Source/Runtime/Core/Public/Analysis/PvsStudioCodeAnalysis.h
@@ -8,16 +8,16 @@
 	// v7.38 currently crashes using the new parser, use the old parser until resolved
 	//V_USE_OLD_PARSER
 
-	//-V::505,542,677,704,719,720,730,735,780,1002,1055,1062,1100
-	//-V:TRYCOMPRESSION:519,547
+	// PVS 7.39 false positive
+	//-V:NumSampledKeys:609
+
+	//-V::505,542,630,677,704,719,720,730,735,780,1002,1055,1062,1100
 	//-V:check(:501,547,560,605
 	//-V:checkSlow(:547
 	//-V:checkf(:510
 	//-V:checkfSlow(:547
 	//-V:dtAssert(:568
 	//-V:rcAssert(:568
-	//-V:GET_FUNCTION_NAME_CHECKED:521
-	//-V:ENABLE_TEXT_ERROR_CHECKING_RESULTS:560
 	//-V:ENABLE_LOC_TESTING:560,617
 	//-V:WITH_EDITOR:560
 	//-V:UE_LOG_ACTIVE:560
@@ -27,19 +27,13 @@
 	//-V:TestTrueExpr:501
 	//-V:PLATFORM_:517,547
 	//-V:ensureMsgf:562
-	//-V:SelectedEmitter:519
 	//-V:MAX_VERTS_PER_POLY:512
-	//-V:127:547
-	//-V:0x7F:547
-	//-V:WARN_COLOR:547
 	//-V:<<:614
-	//-V:OPENGL_PERFORMANCE_DATA_INVALID:564
 	//-V:SPAWN_INIT:595
 	//-V:BEGIN_UPDATE_LOOP:595
-	//-V:OPENGL_PERFORMANCE_DATA_INVALID:560
+	//-V:OPENGL_PERFORMANCE_DATA_INVALID:560,564
 	//-V:UE_CLOG(:501,560
 	//-V:UE_LOG(:501,510,560
-	//-V:UGL_REQUIRED_VOID:501
 	//-V:ensure(:595
 	//-V:ALLOCATE_VERTEX_DATA_TEMPLATE:501
 	//-V:UGL_REQUIRED:501
@@ -51,25 +45,20 @@
 	//-V:INC_MEMORY_STAT_BY:568
 	//-V:DEC_MEMORY_STAT_BY:568
 	//-V:SELECT_STATIC_MESH_VERTEX_TYPE:622
-	//-V:GET_FUNCTION_NAME_CHECKED:685
 	//-V:This(:678
 	//-V:state->error:649
-	//-V:ProjModifiers:616
 	//-V:PERF_DETAILED_PER_CLASS_GC_STATS:686
 	//-V:FMath:656
 	//-V:->*:607
-	//-V:GENERATED_UCLASS_BODY:764
 	//-V:CalcSegmentCostOnPoly:764
 	//-V:DrawLine:764
-	//-V:vrapi_SubmitFrame:641
 	//-V:VertexData:773
 	//-V:Linker:678
 	//-V:self:678
 	//-V:FindChar:679
 	//-V:BytesToHex:530
 	//-V:FindOrAdd:530
-	//-V:GetPawn:623
-	//-V:FindComponentByClass:623
+	//-V:GetPawn<:623
 	//-V:Response.Response:1051
 
 	// The following classes retain a reference to data supplied in the constructor by the derived class which can not yet be initialized.
@@ -82,9 +71,6 @@
 	// Exclude all generated protobuf files
 	//V_EXCLUDE_PATH *.pb.cc
 
-	// warning V630: Instantiation of TRingBuffer < FString >: The 'Malloc' function is used to allocate memory for an array of objects which are classes containing constructors and destructors.
-	//-V::630
-
 	// Disabling because incorrectly flagging all TStaticArrays
 	// V557: Array overrun is possible
 	//-V::557
