TraceLog sends internal statistics when connection is active

* Moved traceing of TraceLog statistics from TraceAuxilliary using the "Counters" channel to a built in event always emitted whenever there is a connection. The counters are split into two events, the bandwidth values emitted every update (~1kb/s) and memory stats emitted at a lower frequency (once per second).
* Refactored stats updating and emitting from static functions to a member function of `FTraceAuxilliaryImpl`.

diff --git a/Engine/Source/Runtime/Core/Private/ProfilingDebugging/TraceAuxiliary.cpp b/Engine/Source/Runtime/Core/Private/ProfilingDebugging/TraceAuxiliary.cpp
--- a/Engine/Source/Runtime/Core/Private/ProfilingDebugging/TraceAuxiliary.cpp
+++ b/Engine/Source/Runtime/Core/Private/ProfilingDebugging/TraceAuxiliary.cpp
@@ -37,14 +37,6 @@
 	#define UE_TRACE_AUTOSTART 1
 #endif
 
-#if !defined(UE_TRACEAUX_EMIT_STAT_COUNTERS)
-	#define UE_TRACEAUX_EMIT_STAT_COUNTERS (0 && STATS)
-#endif
-
-#if !defined(UE_TRACEAUX_EMIT_TRACE_COUNTERS)
-	#define UE_TRACEAUX_EMIT_TRACE_COUNTERS (1 && COUNTERSTRACE_ENABLED)
-#endif
-
 #if !defined(UE_TRACEAUX_EMIT_LLM_STATS)
 	#define UE_TRACEAUX_EMIT_LLM_STATS (1 && ENABLE_LOW_LEVEL_MEM_TRACKER)
 #endif
@@ -109,40 +101,6 @@
 
 CSV_DEFINE_CATEGORY(Trace, true);
 
-#if UE_TRACEAUX_EMIT_STAT_COUNTERS
-#include "Stats/Stats.h"
-DECLARE_STATS_GROUP(TEXT("TraceLog"), STATGROUP_Trace, STATCAT_Advanced);
-DECLARE_MEMORY_STAT(TEXT("Total Memory Used"),             STAT_Trace_MemoryUsed,                 STATGROUP_Trace);
-DECLARE_MEMORY_STAT(TEXT("Block Pool Memory"),             STAT_Trace_Memory_BlockPool,           STATGROUP_Trace);
-DECLARE_MEMORY_STAT(TEXT("Fixed Buffers Memory"),          STAT_Trace_Memory_FixedBuffers,        STATGROUP_Trace);
-DECLARE_MEMORY_STAT(TEXT("Shared Buffers Memory"),         STAT_Trace_Memory_SharedBuffers,       STATGROUP_Trace);
-DECLARE_MEMORY_STAT(TEXT("Important Events Cache"),        STAT_Trace_Memory_Cache,               STATGROUP_Trace);
-DECLARE_MEMORY_STAT(TEXT("Important Events Cache Used"),   STAT_Trace_Memory_Cache_Used,          STATGROUP_Trace);
-DECLARE_MEMORY_STAT(TEXT("Important Events Cache Unused"), STAT_Trace_Memory_Cache_Unused,        STATGROUP_Trace);
-DECLARE_MEMORY_STAT(TEXT("Important Events Cache Wasted"), STAT_Trace_Memory_Cache_Wasted,        STATGROUP_Trace);
-DECLARE_MEMORY_STAT(TEXT("Bytes Emitted"),                 STAT_Trace_Bandwidth_Emitted,          STATGROUP_Trace);
-DECLARE_MEMORY_STAT(TEXT("Bytes Emitted Not Traced"),      STAT_Trace_Bandwidth_EmittedNotTraced, STATGROUP_Trace);
-DECLARE_MEMORY_STAT(TEXT("Bytes Traced"),                  STAT_Trace_Bandwidth_Traced,           STATGROUP_Trace);
-DECLARE_MEMORY_STAT(TEXT("Bytes Sent"),                    STAT_Trace_Bandwidth_Sent,             STATGROUP_Trace);
-DECLARE_MEMORY_STAT(TEXT("Memory Error"),                  STAT_Trace_MemoryError,                STATGROUP_Trace);
-#endif
-
-#if UE_TRACEAUX_EMIT_TRACE_COUNTERS
-TRACE_DECLARE_UNCHECKED_ATOMIC_MEMORY_COUNTER(Trace_MemoryUsed,                 TEXT("Trace/MemoryUsed"));
-TRACE_DECLARE_UNCHECKED_ATOMIC_MEMORY_COUNTER(Trace_Memory_BlockPool,           TEXT("Trace/Memory/BlockPool"));
-TRACE_DECLARE_UNCHECKED_ATOMIC_MEMORY_COUNTER(Trace_Memory_FixedBuffers,        TEXT("Trace/Memory/FixedBuffers"));
-TRACE_DECLARE_UNCHECKED_ATOMIC_MEMORY_COUNTER(Trace_Memory_SharedBuffers,       TEXT("Trace/Memory/SharedBuffers"));
-TRACE_DECLARE_UNCHECKED_ATOMIC_MEMORY_COUNTER(Trace_Memory_Cache,               TEXT("Trace/Memory/Cache"));
-TRACE_DECLARE_UNCHECKED_ATOMIC_MEMORY_COUNTER(Trace_Memory_Cache_Used,          TEXT("Trace/Memory/Cache/Used"));
-TRACE_DECLARE_UNCHECKED_ATOMIC_MEMORY_COUNTER(Trace_Memory_Cache_Unused,        TEXT("Trace/Memory/Cache/Unused"));
-TRACE_DECLARE_UNCHECKED_ATOMIC_MEMORY_COUNTER(Trace_Memory_Cache_Wasted,        TEXT("Trace/Memory/Cache/Wasted"));
-TRACE_DECLARE_UNCHECKED_ATOMIC_MEMORY_COUNTER(Trace_Bandwidth_Emitted,          TEXT("Trace/Bandwidth/Emitted"));
-TRACE_DECLARE_UNCHECKED_ATOMIC_MEMORY_COUNTER(Trace_Bandwidth_EmittedNotTraced, TEXT("Trace/Bandwidth/EmittedNotTraced"));
-TRACE_DECLARE_UNCHECKED_ATOMIC_MEMORY_COUNTER(Trace_Bandwidth_Traced,           TEXT("Trace/Bandwidth/Traced"));
-TRACE_DECLARE_UNCHECKED_ATOMIC_MEMORY_COUNTER(Trace_Bandwidth_Sent,             TEXT("Trace/Bandwidth/Sent"));
-TRACE_DECLARE_UNCHECKED_ATOMIC_MEMORY_COUNTER(Trace_MemoryError,                TEXT("Trace/MemoryError"));
-#endif
-
 #if UE_TRACEAUX_EMIT_LLM_STATS
 LLM_DEFINE_TAG(TraceLog);
 LLM_DEFINE_TAG(TraceLog_BlockPool);
@@ -212,6 +170,8 @@ public:
 	bool SendSnapshot(const TCHAR* InHost, uint32 InPort, const FTraceAuxiliary::FLogCategoryAlias& LogCategory);
 	void AddCommandlineChannels(const TCHAR* ChannelList);
 	void ResetCommandlineChannels();
+	void UpdateStats();
+	void EmitStats(const UE::Trace::FStatistics& Stats);
 
 	// True if this is parent process with forking requested before forking.
 	bool IsParentProcessAndPreFork();
@@ -235,6 +195,8 @@ private:
 	bool WriteToFile(const TCHAR* Path, const FTraceAuxiliary::FLogCategoryAlias& LogCategory, uint16 SendFlags);
 	bool FinalizeFilePath(const TCHAR* InPath, FString& OutPath, const FTraceAuxiliary::FLogCategoryAlias& LogCategory);
 	void AddCommandlineChannel(const TCHAR* Name, const FTraceAuxiliary::FLogCategoryAlias& LogCategory);
+	void EmitDiagnostics(const UE::Trace::FStatistics& Stats);
+	void OnGetOnScreenMessages(FCoreDelegates::FSeverityMessageMap& OutMessages);
 	static void FrameBasedUpdate();
 
 	template<typename ChannelType>
@@ -249,6 +211,9 @@ private:
 	bool bReadOnlyChannelsFrozen = false;
 	FString PausedPreset;
 	FDelegateHandle EndFrameUpdateHandle;
+	FDelegateHandle OnGetOnScreenMessageHandle;
+	FRWLock DiagnosticsLock;
+	FString DiagnosticMessage;
 
 	struct FCurrentTraceTarget
 	{
@@ -896,6 +861,127 @@ void FTraceAuxiliaryImpl::FrameBasedUpdate()
 	UE::Trace::Update();
 }
 
+////////////////////////////////////////////////////////////////////////////////
+void FTraceAuxiliaryImpl::UpdateStats() 
+{
+#if UE_TRACE_STATISTICS_ENABLED 
+	UE::Trace::FStatistics Stats;
+	UE::Trace::GetStatistics(Stats);
+	EmitStats(Stats);
+	EmitDiagnostics(Stats);
+#endif
+}
+
+////////////////////////////////////////////////////////////////////////////////
+void FTraceAuxiliaryImpl::EmitStats(const UE::Trace::FStatistics& Stats)
+{
+#if UE_TRACEAUX_EMIT_LLM_STATS && ENABLE_LOW_LEVEL_MEM_TRACKER
+	if (FLowLevelMemTracker::Get().IsEnabled())
+	{
+		TRACE_CPUPROFILER_EVENT_SCOPE(FTraceAux_Stats_LLM);
+		FLowLevelMemTracker::Get().SetTagAmountForTracker(ELLMTracker::Default, LLM_TAG_NAME(TraceLog), ELLMTagSet::None, Stats.MemoryUsed, true);
+		FLowLevelMemTracker::Get().SetTagAmountForTracker(ELLMTracker::Default, LLM_TAG_NAME(TraceLog_BlockPool), ELLMTagSet::None, Stats.BlockPoolAllocated, true);
+		FLowLevelMemTracker::Get().SetTagAmountForTracker(ELLMTracker::Default, LLM_TAG_NAME(TraceLog_FixedBuffers), ELLMTagSet::None, Stats.FixedBufferAllocated, true);
+		FLowLevelMemTracker::Get().SetTagAmountForTracker(ELLMTracker::Default, LLM_TAG_NAME(TraceLog_SharedBuffers), ELLMTagSet::None, Stats.SharedBufferAllocated, true);
+		FLowLevelMemTracker::Get().SetTagAmountForTracker(ELLMTracker::Default, LLM_TAG_NAME(TraceLog_Cache), ELLMTagSet::None, Stats.CacheAllocated, true);
+	}
+#endif
+
+#if UE_TRACEAUX_EMIT_CSV_STATS
+	// Only publish CSV stats if we have ever run tracing in order to reduce overhead in most runs.
+	static bool bDoCsvStats = false;
+	if (IsConnected() || bDoCsvStats)
+	{
+		bDoCsvStats = true;
+
+		TRACE_CPUPROFILER_EVENT_SCOPE(FTraceAux_Stats_CSV);
+		CSV_CUSTOM_STAT(Trace, MemoryUsedMb,    double(Stats.MemoryUsed)            / 1024.0 / 1024.0, ECsvCustomStatOp::Set);
+		CSV_CUSTOM_STAT(Trace, BlockPoolMb,     double(Stats.BlockPoolAllocated)    / 1024.0 / 1024.0, ECsvCustomStatOp::Set);
+		CSV_CUSTOM_STAT(Trace, FixedBuffersMb,  double(Stats.FixedBufferAllocated)  / 1024.0 / 1024.0, ECsvCustomStatOp::Set);
+		CSV_CUSTOM_STAT(Trace, SharedBuffersMb, double(Stats.SharedBufferAllocated) / 1024.0 / 1024.0, ECsvCustomStatOp::Set);
+		CSV_CUSTOM_STAT(Trace, CacheMb,         double(Stats.CacheAllocated)        / 1024.0 / 1024.0, ECsvCustomStatOp::Set);
+	}
+#endif
+}
+
+////////////////////////////////////////////////////////////////////////////////
+void FTraceAuxiliaryImpl::EmitDiagnostics(const UE::Trace::FStatistics& Stats)
+{
+	static double LastUpdate = FPlatformTime::Seconds();
+	const double CurrentTime = FPlatformTime::Seconds();
+	static uint64 LastSent = Stats.BytesSent;
+	static uint64 LastTraced = Stats.BytesTraced;
+	static uint64 LastEmitted = Stats.BytesEmitted;
+	const double Delta = FMath::Clamp(CurrentTime - LastUpdate, UE_DOUBLE_SMALL_NUMBER, UE_DOUBLE_BIG_NUMBER);
+	const double TracedPerSecond = double(Stats.BytesTraced - LastTraced) / Delta;
+	const double SentPerSecond = double(Stats.BytesSent - LastSent) / Delta;
+	const double EmittedPerSecond = double(Stats.BytesEmitted - LastEmitted) / Delta;
+
+	LastSent = Stats.BytesSent;
+	LastTraced = Stats.BytesTraced;
+	LastEmitted = Stats.BytesEmitted;
+	LastUpdate = CurrentTime;
+
+	// Generate a diagnostics message (can be rendered to screen) if the block
+	// pool is at the limit and utilization is above warning threshhold.
+	static uint32 Sticky = 0;
+	constexpr uint32 Stickyness = 50; // Number of frames the message will "stick" after
+	constexpr float BlockPoolUtilizationWarning = 0.8f; // Warn above 80% utilization
+	constexpr float BlockPoolGrowthWarning = 0.8f; // Warn above 80% growth
+	check(Stats.BlockPoolAllocatedBlocks >= Stats.BlockPoolFreeBlocks);
+	const float BlockPoolFreeFraction = Stats.BlockPoolAllocatedBlocks ? float(Stats.BlockPoolFreeBlocks)/float(Stats.BlockPoolAllocatedBlocks) : 1.0f;
+	const float BlockPoolUtilization = 1.0f - BlockPoolFreeFraction;
+	const float BlockPoolGrowth = GInitializeDesc.BlockPoolMaxSize ? float(Stats.BlockPoolAllocated) / float(GInitializeDesc.BlockPoolMaxSize) : 0.0f;
+
+	if (BlockPoolGrowth > BlockPoolGrowthWarning && BlockPoolUtilization > BlockPoolUtilizationWarning)
+	{
+		Sticky = Stickyness;
+	}
+
+	if (Sticky)
+	{
+		TStringBuilder<128> Diagnostic;
+		Diagnostic.Appendf(TEXT("Trace: Block pool %.2f %% full (Size: %.2f MiB, Max: %.2f MiB), Sending: %.2f MiB/s Emitting: %.2f Mib/s"),
+			100.0f * BlockPoolUtilization,
+			float(Stats.BlockPoolAllocated) / (1024*1024),
+			float(GInitializeDesc.BlockPoolMaxSize) / (1024*1024),
+			SentPerSecond / (1024 * 1024),
+			EmittedPerSecond / (1024 * 1024)
+		);
+		Sticky = FMath::Clamp(Sticky - 1, 0u, Stickyness);
+		
+		{
+			FWriteScopeLock _(DiagnosticsLock);
+			DiagnosticMessage = Diagnostic.ToString();
+		}
+
+		if (!OnGetOnScreenMessageHandle.IsValid())
+		{
+			OnGetOnScreenMessageHandle = FCoreDelegates::OnGetOnScreenMessages.AddRaw(this, &FTraceAuxiliaryImpl::OnGetOnScreenMessages);
+		}
+	}
+	else
+	{
+		if (!DiagnosticMessage.IsEmpty())
+		{
+			FWriteScopeLock _(DiagnosticsLock);
+			DiagnosticMessage.Empty();
+		}
+	}
+}
+
+////////////////////////////////////////////////////////////////////////////////
+void FTraceAuxiliaryImpl::OnGetOnScreenMessages(FCoreDelegates::FSeverityMessageMap& OutMessages)
+{
+	FReadScopeLock _(DiagnosticsLock);
+	if (DiagnosticMessage.IsEmpty())
+	{
+		return;
+	}
+	
+	OutMessages.Add(FCoreDelegates::EOnScreenMessageSeverity::Warning, FText::FromString(DiagnosticMessage));
+}
+
 ////////////////////////////////////////////////////////////////////////////////
 bool FTraceAuxiliaryImpl::WriteToFile(const TCHAR* Path, const FTraceAuxiliary::FLogCategoryAlias& LogCategory, uint16 SendFlags)
 {
@@ -1053,64 +1139,6 @@ void FTraceAuxiliaryImpl::StartEndFrameUpdate()
 	EndFrameUpdateHandle = FCoreDelegates::OnEndFrame.AddStatic(&FTraceAuxiliaryImpl::FrameBasedUpdate);
 }
 
-////////////////////////////////////////////////////////////////////////////////
-void TraceAuxiliaryEmitStats(const UE::Trace::FStatistics& Stats)
-{
-	uint64 TotalMemoryUsed = Stats.BlockPoolAllocated + Stats.FixedBufferAllocated + Stats.SharedBufferAllocated + Stats.CacheAllocated;
-
-#if UE_TRACEAUX_EMIT_STAT_COUNTERS
-	if (FThreadStats::IsCollectingData())
-	{
-		TRACE_CPUPROFILER_EVENT_SCOPE(FTraceAux_Stats_STAT);
-		SET_MEMORY_STAT(STAT_Trace_MemoryUsed, Stats.MemoryUsed);
-		SET_MEMORY_STAT(STAT_Trace_Memory_BlockPool, Stats.BlockPoolAllocated);
-		SET_MEMORY_STAT(STAT_Trace_Memory_FixedBuffers, Stats.FixedBufferAllocated);
-		SET_MEMORY_STAT(STAT_Trace_Memory_SharedBuffers, Stats.SharedBufferAllocated);
-		SET_MEMORY_STAT(STAT_Trace_Memory_Cache, Stats.CacheAllocated);
-		SET_MEMORY_STAT(STAT_Trace_Memory_Cache_Used, Stats.CacheUsed);
-		SET_MEMORY_STAT(STAT_Trace_Memory_Cache_Unused, (int64)Stats.CacheAllocated - (int64)Stats.CacheUsed);
-		SET_MEMORY_STAT(STAT_Trace_Memory_Cache_Wasted, Stats.CacheWaste);
-		SET_MEMORY_STAT(STAT_Trace_Bandwidth_Emitted, Stats.BytesEmitted);
-		SET_MEMORY_STAT(STAT_Trace_Bandwidth_EmittedNotTraced, (int64)Stats.BytesEmitted - (int64)Stats.BytesTraced);
-		SET_MEMORY_STAT(STAT_Trace_Bandwidth_Traced, Stats.BytesTraced);
-		SET_MEMORY_STAT(STAT_Trace_Bandwidth_Sent, Stats.BytesSent);
-		SET_MEMORY_STAT(STAT_Trace_MemoryError, (int64)Stats.MemoryUsed - (int64)TotalMemoryUsed);
-	}
-#endif
-
-#if UE_TRACEAUX_EMIT_TRACE_COUNTERS
-	if (UE_TRACE_CHANNELEXPR_IS_ENABLED(CountersChannel))
-	{
-		TRACE_CPUPROFILER_EVENT_SCOPE(FTraceAux_Stats_Counters);
-		TRACE_COUNTER_SET(Trace_MemoryUsed, Stats.MemoryUsed);
-		TRACE_COUNTER_SET(Trace_Memory_BlockPool, Stats.BlockPoolAllocated);
-		TRACE_COUNTER_SET(Trace_Memory_FixedBuffers, Stats.FixedBufferAllocated);
-		TRACE_COUNTER_SET(Trace_Memory_SharedBuffers, Stats.SharedBufferAllocated);
-		TRACE_COUNTER_SET(Trace_Memory_Cache, Stats.CacheAllocated);
-		TRACE_COUNTER_SET(Trace_Memory_Cache_Used, Stats.CacheUsed);
-		TRACE_COUNTER_SET(Trace_Memory_Cache_Unused, (int64)Stats.CacheAllocated - (int64)Stats.CacheUsed);
-		TRACE_COUNTER_SET(Trace_Memory_Cache_Wasted, Stats.CacheWaste);
-		TRACE_COUNTER_SET(Trace_Bandwidth_Emitted, Stats.BytesEmitted);
-		TRACE_COUNTER_SET(Trace_Bandwidth_EmittedNotTraced, (int64)Stats.BytesEmitted - (int64)Stats.BytesTraced);
-		TRACE_COUNTER_SET(Trace_Bandwidth_Traced, Stats.BytesTraced);
-		TRACE_COUNTER_SET(Trace_Bandwidth_Sent, Stats.BytesSent);
-		TRACE_COUNTER_SET(Trace_MemoryError, (int64)Stats.MemoryUsed - (int64)TotalMemoryUsed);
-	}
-#endif
-
-#if UE_TRACEAUX_EMIT_LLM_STATS && ENABLE_LOW_LEVEL_MEM_TRACKER
-	if (FLowLevelMemTracker::Get().IsEnabled())
-	{
-		TRACE_CPUPROFILER_EVENT_SCOPE(FTraceAux_Stats_LLM);
-		FLowLevelMemTracker::Get().SetTagAmountForTracker(ELLMTracker::Default, LLM_TAG_NAME(TraceLog), ELLMTagSet::None, Stats.MemoryUsed, true);
-		FLowLevelMemTracker::Get().SetTagAmountForTracker(ELLMTracker::Default, LLM_TAG_NAME(TraceLog_BlockPool), ELLMTagSet::None, Stats.BlockPoolAllocated, true);
-		FLowLevelMemTracker::Get().SetTagAmountForTracker(ELLMTracker::Default, LLM_TAG_NAME(TraceLog_FixedBuffers), ELLMTagSet::None, Stats.FixedBufferAllocated, true);
-		FLowLevelMemTracker::Get().SetTagAmountForTracker(ELLMTracker::Default, LLM_TAG_NAME(TraceLog_SharedBuffers), ELLMTagSet::None, Stats.SharedBufferAllocated, true);
-		FLowLevelMemTracker::Get().SetTagAmountForTracker(ELLMTracker::Default, LLM_TAG_NAME(TraceLog_Cache), ELLMTagSet::None, Stats.CacheAllocated, true);
-	}
-#endif
-}
-
 ////////////////////////////////////////////////////////////////////////////////
 void TraceAuxiliaryOnTraceUpdateCallback()
 {
@@ -1126,11 +1154,8 @@ void TraceAuxiliaryOnTraceUpdateCallback()
 	}
 #endif
 
-#if UE_TRACEAUX_EMIT_STAT_COUNTERS || UE_TRACEAUX_EMIT_TRACE_COUNTERS || UE_TRACEAUX_EMIT_LLM_STATS
-	UE::Trace::FStatistics Stats;
-	UE::Trace::GetStatistics(Stats);
-	TraceAuxiliaryEmitStats(Stats);
-#endif
+	GTraceAuxiliary.UpdateStats();
+
 }
 
 ////////////////////////////////////////////////////////////////////////////////
@@ -1142,32 +1167,7 @@ void TraceAuxiliaryOnEndFrameCallback()
 	// Reset the trace update callback. From now on, we only trace stats once per frame.
 	UE::Trace::SetUpdateCallback(nullptr);
 
-#if UE_TRACEAUX_EMIT_STAT_COUNTERS || UE_TRACEAUX_EMIT_TRACE_COUNTERS || UE_TRACEAUX_EMIT_LLM_STATS
-	UE::Trace::FStatistics Stats;
-	UE::Trace::GetStatistics(Stats);
-	TraceAuxiliaryEmitStats(Stats);
-#endif
-
-#if UE_TRACEAUX_EMIT_CSV_STATS && TRACE_PRIVATE_STATISTICS
-	// Only publish CSV stats if we have ever run tracing in order to reduce overhead in most runs.
-	static bool bDoCsvStats = false;
-	if (UE::Trace::IsTracing() || bDoCsvStats)
-	{
-		bDoCsvStats = true;
-
-#if !(UE_TRACEAUX_EMIT_STAT_COUNTERS || UE_TRACEAUX_EMIT_TRACE_COUNTERS || UE_TRACEAUX_EMIT_LLM_STATS)
-		UE::Trace::FStatistics Stats;
-		UE::Trace::GetStatistics(Stats);
-#endif
-
-		TRACE_CPUPROFILER_EVENT_SCOPE(FTraceAux_Stats_CSV);
-		CSV_CUSTOM_STAT(Trace, MemoryUsedMb,    double(Stats.MemoryUsed)            / 1024.0 / 1024.0, ECsvCustomStatOp::Set);
-		CSV_CUSTOM_STAT(Trace, BlockPoolMb,     double(Stats.BlockPoolAllocated)    / 1024.0 / 1024.0, ECsvCustomStatOp::Set);
-		CSV_CUSTOM_STAT(Trace, FixedBuffersMb,  double(Stats.FixedBufferAllocated)  / 1024.0 / 1024.0, ECsvCustomStatOp::Set);
-		CSV_CUSTOM_STAT(Trace, SharedBuffersMb, double(Stats.SharedBufferAllocated) / 1024.0 / 1024.0, ECsvCustomStatOp::Set);
-		CSV_CUSTOM_STAT(Trace, CacheMb,         double(Stats.CacheAllocated)        / 1024.0 / 1024.0, ECsvCustomStatOp::Set);
-	}
-#endif
+	GTraceAuxiliary.UpdateStats();
 }
 
 ////////////////////////////////////////////////////////////////////////////////
@@ -2096,7 +2096,7 @@ void FTraceAuxiliary::Initialize(const TCHAR* CommandLine)
 	{
 		TRACE_CPUPROFILER_EVENT_SCOPE(FTraceAux_InitialStats);
 		UE::Trace::FStatistics Stats;
-		TraceAuxiliaryEmitStats(Stats);
+		GTraceAuxiliary.EmitStats(Stats);
 	}
 
 	// Initialize Trace. The settings are stored in a static for posterity.
diff --git a/Engine/Source/Runtime/TraceLog/Private/Trace/Writer.cpp b/Engine/Source/Runtime/TraceLog/Private/Trace/Writer.cpp
--- a/Engine/Source/Runtime/TraceLog/Private/Trace/Writer.cpp
+++ b/Engine/Source/Runtime/TraceLog/Private/Trace/Writer.cpp
@@ -82,6 +82,22 @@ UE_TRACE_MINIMAL_EVENT_BEGIN($Trace, NewTrace, Important | NoSync)
 	UE_TRACE_MINIMAL_EVENT_FIELD(double, StartDateTime)
 UE_TRACE_MINIMAL_EVENT_END()
 
+UE_TRACE_MINIMAL_EVENT_BEGIN($Trace, BandwidthStats)
+	UE_TRACE_MINIMAL_EVENT_FIELD(uint64, Cycle)
+	UE_TRACE_MINIMAL_EVENT_FIELD(uint32, DeltaBytesSent)
+	UE_TRACE_MINIMAL_EVENT_FIELD(uint32, DeltaBytesEmitted)
+UE_TRACE_MINIMAL_EVENT_END()
+
+UE_TRACE_MINIMAL_EVENT_BEGIN($Trace, MemoryStats)
+	UE_TRACE_MINIMAL_EVENT_FIELD(uint64, Cycle)
+	UE_TRACE_MINIMAL_EVENT_FIELD(float,	 BlockPoolUsage)
+	UE_TRACE_MINIMAL_EVENT_FIELD(uint64, MemoryUsed)
+	UE_TRACE_MINIMAL_EVENT_FIELD(uint64, BlockPoolAllocated)
+	UE_TRACE_MINIMAL_EVENT_FIELD(uint32, FixedBufferAllocated)
+	UE_TRACE_MINIMAL_EVENT_FIELD(uint32, CacheAllocated)
+	UE_TRACE_MINIMAL_EVENT_FIELD(uint32, SharedAllocated)
+UE_TRACE_MINIMAL_EVENT_END()
+
 ////////////////////////////////////////////////////////////////////////////////
 static volatile bool			GInitialized;		// = false;
 FStatistics						GTraceStatistics;	// = {};
@@ -629,7 +645,61 @@ static OnConnectFunc*	GOnConnection = nullptr;
 static OnUpdateFunc*	GOnUpdate = nullptr;
 static FTraceGuid		GSessionGuid; // = {0, 0, 0, 0};
 static FTraceGuid		GTraceGuid; // = {0, 0, 0, 0};
+static const float		GMemoryStatsFrequencySeconds = 1.0f; // Every second
+
+////////////////////////////////////////////////////////////////////////////////
+static void Writer_TraceStats(const FStatistics& InStats)
+{
+#if UE_TRACE_EMIT_STATISTICS
+	static uint64 LastMemoryStatsCycle = 0;
+	static uint64 LastSent = InStats.BytesSent;
+	static uint64 LastTraced = InStats.BytesTraced;
+	static uint64 LastEmitted = InStats.BytesEmitted;
+
+	// Only send stats when there is an actual connection, so we
+	// don't use the tail buffer unnecessarily.
+	if (GDataHandle)
+	{
+		const uint64 CurrentTimeCycle = TimeGetRelativeTimestamp();
+		const float BlockPoolFreeFraction = InStats.BlockPoolAllocatedBlocks ? 
+			float(InStats.BlockPoolFreeBlocks) / float(InStats.BlockPoolAllocatedBlocks) : 1.0f;
+		const float BlockPoolUsage = 1.0f - BlockPoolFreeFraction;
+		const uint32 DeltaEmitted = uint32(InStats.BytesEmitted - LastEmitted);
+		uint32 DeltaSent = uint32(InStats.BytesSent - LastSent);
+		if (InStats.BytesSent < LastSent)
+		{
+			// BytesSent is reset on every new connection
+			DeltaSent = uint32(InStats.BytesSent);
+		}
+
+		// Send bandwidth stats every update
+		UE_TRACE_MINIMAL_LOG($Trace, BandwidthStats, TraceLogChannel)
+			<< BandwidthStats.Cycle(CurrentTimeCycle)
+			<< BandwidthStats.DeltaBytesSent(DeltaSent)
+			<< BandwidthStats.DeltaBytesEmitted(DeltaEmitted);
 
+		// Only send memory stats on configured frequency
+		const float MemoryStatDeltaSeconds = float(CurrentTimeCycle - LastMemoryStatsCycle) / float(TimeGetFrequency());
+		if (MemoryStatDeltaSeconds > GMemoryStatsFrequencySeconds)
+		{
+			UE_TRACE_MINIMAL_LOG($Trace, MemoryStats, TraceLogChannel)
+				<< MemoryStats.Cycle(CurrentTimeCycle)
+				<< MemoryStats.MemoryUsed(InStats.MemoryUsed)
+				<< MemoryStats.BlockPoolAllocated(InStats.BlockPoolAllocated)
+				<< MemoryStats.FixedBufferAllocated(InStats.FixedBufferAllocated)
+				<< MemoryStats.CacheAllocated(InStats.CacheAllocated)
+				<< MemoryStats.SharedAllocated(InStats.SharedBufferAllocated)
+				<< MemoryStats.BlockPoolUsage(BlockPoolUsage);
+
+			LastMemoryStatsCycle = CurrentTimeCycle;
+		}
+	}
+
+	LastSent = InStats.BytesSent;
+	LastTraced = InStats.BytesTraced;
+	LastEmitted = InStats.BytesEmitted;
+#endif
+}
 
 ////////////////////////////////////////////////////////////////////////////////
 static void Writer_SendSync()
@@ -852,6 +922,7 @@ static void Writer_WorkerUpdateInternal()
 	Writer_DescribeAnnounce();
 	Writer_UpdateSharedBuffers();
 	Writer_DrainBuffers();
+	Writer_TraceStats(GTraceStatistics);
 	Writer_SendSync();
 
 #if TRACE_PRIVATE_BUFFER_SEND
diff --git a/Engine/Source/Runtime/TraceLog/Public/Trace/Config.h b/Engine/Source/Runtime/TraceLog/Public/Trace/Config.h
--- a/Engine/Source/Runtime/TraceLog/Public/Trace/Config.h
+++ b/Engine/Source/Runtime/TraceLog/Public/Trace/Config.h
@@ -142,3 +142,18 @@ static_assert(UE_TRACE_MINIMAL_ENABLED == 0 || UE_TRACE_ENABLED_SHIPPING_EXPERIM
 #ifndef UE_TRACE_PACKET_VERIFICATION
 #	define UE_TRACE_PACKET_VERIFICATION 0
 #endif
+
+/**
+ * Enable tracking of memory and bandwidth statistics.
+ */
+#ifndef UE_TRACE_STATISTICS_ENABLED
+#	define UE_TRACE_STATISTICS_ENABLED (!(UE_BUILD_SHIPPING || USING_THREAD_SANITISER)) 
+#endif
+
+/**
+ * Enable built in events to send statistics. Enabled by default
+ * when statistics tracking is enabled.
+ */
+#ifndef UE_TRACE_EMIT_STATISTICS
+#	define UE_TRACE_EMIT_STATISTICS UE_TRACE_STATISTICS_ENABLED
+#endif
diff --git a/Engine/Source/Runtime/TraceLog/Public/Trace/Detail/Trace.h b/Engine/Source/Runtime/TraceLog/Public/Trace/Detail/Trace.h
--- a/Engine/Source/Runtime/TraceLog/Public/Trace/Detail/Trace.h
+++ b/Engine/Source/Runtime/TraceLog/Public/Trace/Detail/Trace.h
@@ -26,7 +26,7 @@ class FChannel;
 //-- Private functional variants
 //------------------------------------------------------------------------------------------------
 
-#define TRACE_PRIVATE_STATISTICS (!(UE_BUILD_SHIPPING || USING_THREAD_SANITISER))
+#define TRACE_PRIVATE_STATISTICS UE_TRACE_STATISTICS_ENABLED
 
 #define TRACE_PRIVATE_CHANNEL_DEFAULT_ARGS false, "None"
 
