Hololens removal: the final chapter
-Removing a couple warning disables and compiler options that were necessary for hololens.

diff --git a/Engine/Source/Programs/UnrealBuildTool/Platform/Windows/VCToolChain.cs b/Engine/Source/Programs/UnrealBuildTool/Platform/Windows/VCToolChain.cs
--- a/Engine/Source/Programs/UnrealBuildTool/Platform/Windows/VCToolChain.cs
+++ b/Engine/Source/Programs/UnrealBuildTool/Platform/Windows/VCToolChain.cs
@@ -734,15 +734,6 @@ namespace UnrealBuildTool
 				Arguments.Add("/Zc:checkGwOdr");
 			}
 
-			// @todo HoloLens: UE is non-compliant when it comes to use of %s and %S
-			// Previously %s meant "the current character set" and %S meant "the other one".
-			// Now %s means multibyte and %S means wide. %Ts means "natural width".
-			// Reverting this behaviour until the UE source catches up.
-			AddDefinition(Arguments, "_CRT_STDIO_LEGACY_WIDE_SPECIFIERS=1");
-
-			// @todo HoloLens: Silence the hash_map deprecation errors for now. This should be replaced with unordered_map for the real fix.
-			AddDefinition(Arguments, "_SILENCE_STDEXT_HASH_DEPRECATION_WARNINGS=1");
-
 			// Ignore secure CRT warnings on Clang
 			if (Target.WindowsPlatform.Compiler.IsClang())
 			{
diff --git a/Engine/Source/Runtime/Core/Public/Windows/WindowsPlatformCompilerSetup.h b/Engine/Source/Runtime/Core/Public/Windows/WindowsPlatformCompilerSetup.h
--- a/Engine/Source/Runtime/Core/Public/Windows/WindowsPlatformCompilerSetup.h
+++ b/Engine/Source/Runtime/Core/Public/Windows/WindowsPlatformCompilerSetup.h
@@ -128,9 +128,6 @@ static_assert(_MSC_VER >= 1938, "Visual Studio 2022 v17.8 (MSVC 14.38) or greate
 
 #pragma warning(disable: 4063) // case 'val' is not a valid value for switch of enum 'E'
 
-// @todo HoloLens: Disabled because DbgHelp.h has some anonymous typedefs in it (not allowed in Visual Studio 2015).  We should probably just wrap that header.
-#pragma warning(disable: 4091) // 'keyword' : ignored on left of 'type' when no variable is declared								https://docs.microsoft.com/en-us/cpp/error-messages/compiler-warnings/compiler-warning-level-1-c4091
-
 // Unwanted VC++ level 4 warnings to disable.
 #pragma warning(disable: 4100) // 'identifier' : unreferenced formal parameter														https://docs.microsoft.com/en-us/cpp/error-messages/compiler-warnings/compiler-warning-level-4-c4100
 #pragma warning(disable: 4121) // 'symbol' : alignment of a member was sensitive to packing											https://docs.microsoft.com/en-us/cpp/error-messages/compiler-warnings/compiler-warning-level-4-c4121
diff --git a/Engine/Source/Runtime/CoreUObject/Public/UObject/Class.h b/Engine/Source/Runtime/CoreUObject/Public/UObject/Class.h
--- a/Engine/Source/Runtime/CoreUObject/Public/UObject/Class.h
+++ b/Engine/Source/Runtime/CoreUObject/Public/UObject/Class.h
@@ -470,6 +470,12 @@ ENUM_CLASS_FLAGS(EStructPropertyLinkFlags)
 /**
  * Base class for all UObject types that contain fields.
  */
+#ifdef _MSC_VER
+//error C4091: '__declspec(empty_bases)': ignored on left of 'UStruct' when no variable is declared
+// In some configurations FStructBaseChain is empty, but in some it is not which will result in warning 4091.  empty_bases does nothing if there are no empty bases, so we can just suppress the warning.
+#pragma warning(push) // 4091
+#pragma warning(disable:4091)
+#endif //#ifdef _MSC_VER
 PLATFORM_EMPTY_BASES // FStructBaseChain may be empty
 UCLASS(MinimalAPI, Config = Engine)
 class UStruct : public UField
@@ -1003,6 +1009,9 @@ private:
 
 	void SerializeVersionedTaggedProperties(FStructuredArchive::FSlot Slot, uint8* Data, const UStruct* DefaultsStruct, const uint8* Defaults, const UObject* BreakRecursionIfFullyLoad) const;
 };
+#ifdef _MSC_VER
+#pragma warning(pop) // 4091
+#endif
 
 #if UE_WITH_CONSTINIT_UOBJECT
 // Private access to protected base class for UHT generated code
