Added TArray::PadToNum and TArray::PadToNumDefaulted.

diff --git a/Engine/Source/Runtime/Core/Private/Containers/ContainersTest.cpp b/Engine/Source/Runtime/Core/Private/Containers/ContainersTest.cpp
--- a/Engine/Source/Runtime/Core/Private/Containers/ContainersTest.cpp
+++ b/Engine/Source/Runtime/Core/Private/Containers/ContainersTest.cpp
@@ -296,6 +296,90 @@ namespace
 		}
 	};
 
+	template <typename ArrayType>
+	void RunPadToNumTests(const typename ArrayType::ElementType& PaddingValue1, const typename ArrayType::ElementType& PaddingValue2)
+	{
+		using ElementType = typename ArrayType::ElementType;
+
+		ArrayType Array;
+
+		// Check padding an empty array
+		{
+			auto Result = Array.PadToNumDefaulted(0);
+			check(Result == 0);
+			check(Array.IsEmpty());
+		}
+		{
+			auto Result = Array.PadToNum(0, PaddingValue1);
+			check(Result == 0);
+			check(Array.IsEmpty());
+		}
+
+		// Check padding with a value
+		{
+			auto Result = Array.PadToNumDefaulted(5);
+			check(Result == 5);
+			check(Array.Num() == 5);
+			check(Array == (ArrayType{ ElementType{}, ElementType{}, ElementType{}, ElementType{}, ElementType{} }));
+		}
+
+		// Check underpadding
+		{
+			auto Result = Array.PadToNumDefaulted(4);
+			check(Result == 0);
+			check(Array.Num() == 5);
+			check(Array == (ArrayType{ ElementType{}, ElementType{}, ElementType{}, ElementType{}, ElementType{} }));
+		}
+		{
+			auto Result = Array.PadToNum(4, PaddingValue1);
+			check(Result == 0);
+			check(Array.Num() == 5);
+			check(Array == (ArrayType{ ElementType{}, ElementType{}, ElementType{}, ElementType{}, ElementType{} }));
+		}
+
+		// Check adding exact padding
+		{
+			auto Result = Array.PadToNumDefaulted(5);
+			check(Result == 0);
+			check(Array.Num() == 5);
+			check(Array == (ArrayType{ ElementType{}, ElementType{}, ElementType{}, ElementType{}, ElementType{} }));
+		}
+		{
+			auto Result = Array.PadToNum(5, PaddingValue1);
+			check(Result == 0);
+			check(Array.Num() == 5);
+			check(Array == (ArrayType{ ElementType{}, ElementType{}, ElementType{}, ElementType{}, ElementType{} }));
+		}
+
+		// Check adding 1 more padding value to non-empty arrays
+		{
+			auto Result = Array.PadToNumDefaulted(6);
+			check(Result == 1);
+			check(Array.Num() == 6);
+			check(Array == (ArrayType{ ElementType{}, ElementType{}, ElementType{}, ElementType{}, ElementType{}, ElementType{} }));
+		}
+		{
+			auto Result = Array.PadToNum(7, PaddingValue1);
+			check(Result == 1);
+			check(Array.Num() == 7);
+			check(Array == (ArrayType{ ElementType{}, ElementType{}, ElementType{}, ElementType{}, ElementType{}, ElementType{}, PaddingValue1 }));
+		}
+
+		// Check adding multiple padding to non-empty arrays
+		{
+			auto Result = Array.PadToNumDefaulted(10);
+			check(Result == 3);
+			check(Array.Num() == 10);
+			check(Array == (ArrayType{ ElementType{}, ElementType{}, ElementType{}, ElementType{}, ElementType{}, ElementType{}, PaddingValue1, ElementType{}, ElementType{}, ElementType{} }));
+		}
+		{
+			auto Result = Array.PadToNum(14, PaddingValue2);
+			check(Result == 4);
+			check(Array.Num() == 14);
+			check(Array == (ArrayType{ ElementType{}, ElementType{}, ElementType{}, ElementType{}, ElementType{}, ElementType{}, PaddingValue1, ElementType{}, ElementType{}, ElementType{}, PaddingValue2, PaddingValue2, PaddingValue2, PaddingValue2 }));
+		}
+	}
+
 	template <typename MapType, typename KeyType>
 	void RunMapTests()
 	{
@@ -1080,6 +1164,9 @@ bool FContainersSmokeTest::RunTest( const FString& Parameters )
 IMPLEMENT_SIMPLE_AUTOMATION_TEST(FContainersFullTest, "System.Core.Containers.Full", EAutomationTestFlags_ApplicationContextMask | EAutomationTestFlags::EngineFilter)
 bool FContainersFullTest::RunTest(const FString& Parameters)
 {
+	RunPadToNumTests<TArray<int32>>(5, 7);
+	RunPadToNumTests<TArray<FString>>("Hello", "World");
+
 	RunMapTests<TMap<int32, FContainerTestValueType>, int32>();
 	RunMapTests<TMap<FName, FContainerTestValueType>, FName>();
 	RunMapTests<TMap<FString, FContainerTestValueType>, FString>();
diff --git a/Engine/Source/Runtime/Core/Public/Containers/Array.h b/Engine/Source/Runtime/Core/Public/Containers/Array.h
--- a/Engine/Source/Runtime/Core/Public/Containers/Array.h
+++ b/Engine/Source/Runtime/Core/Public/Containers/Array.h
@@ -2364,6 +2364,50 @@ public:
 		SlackTrackerNumChanged();
 	}
 
+	/**
+	 * Pad the array with default values to ensure that it is at least the specified length. Does nothing if Num() >= DesiredNum.
+	 *
+	 * @param DesiredNum The desired number of elements that should exist in the array.
+	 * @return The number of elements that were added to the array, or 0 if Num() >= DesiredNum.
+	 */
+	SizeType PadToNumDefaulted(SizeType DesiredNum)
+	{
+		SizeType Diff = DesiredNum - Num();
+		if (Diff <= 0)
+		{
+			return 0;
+		}
+
+		AddDefaulted(Diff);
+
+		return Diff;
+	}
+
+	/**
+	 * Pad the array with the specified value to ensure that it is at least the specified length. Does nothing if Num() >= DesiredNum.
+	 *
+	 * @param DesiredNum The desired number of elements that should exist in the array.
+	 * @param PadValue  The value to pad with
+	 * @return The number of elements that were added to the array, or 0 if Num() >= DesiredNum.
+	 */
+	SizeType PadToNum(SizeType DesiredNum, const ElementType& PadValue)
+	{
+		SizeType Diff = DesiredNum - Num();
+		if (Diff <= 0)
+		{
+			return 0;
+		}
+
+		SizeType Pos = AddUninitialized(Diff);
+		ElementType* Data = GetData() + Pos;
+		for (SizeType Index = 0; Index != Diff; ++Index, ++Data)
+		{
+			(void)new ((void*)Data) ElementType(PadValue);
+		}
+
+		return Diff;
+	}
+
 	/**
 	 * Appends the specified array to this array.
 	 *
