[TaskGraph]
- Allow named thread tasks retraction when possible

diff --git a/Engine/Source/Runtime/Core/Private/Tasks/TaskPrivate.cpp b/Engine/Source/Runtime/Core/Private/Tasks/TaskPrivate.cpp
--- a/Engine/Source/Runtime/Core/Private/Tasks/TaskPrivate.cpp
+++ b/Engine/Source/Runtime/Core/Private/Tasks/TaskPrivate.cpp
@@ -81,10 +81,30 @@ namespace UE::Tasks
 				return false;
 			}
 
-			// task retraction is not supported for named thread tasks
 			if (IsNamedThreadTask())
 			{
-				return false;
+				if (ExtendedPriority >= EExtendedTaskPriority::GameThreadNormalPri &&
+					ExtendedPriority <= EExtendedTaskPriority::GameThreadHiPriLocalQueue &&
+					IsInGameThread())
+				{
+					// Allow task retraction of game thread tasks on the game thread
+				}
+				else if (ExtendedPriority >= EExtendedTaskPriority::RenderThreadNormalPri &&
+					ExtendedPriority <= EExtendedTaskPriority::RenderThreadHiPriLocalQueue &&
+					IsInRenderingThread())
+				{
+					// Allow task retraction of rendering tasks on the rendering thread
+				}
+				else if (ExtendedPriority >= EExtendedTaskPriority::RHIThreadNormalPri &&
+					ExtendedPriority <= EExtendedTaskPriority::RHIThreadHiPriLocalQueue &&
+					IsInRHIThread())
+				{
+					// Allow task retraction of RHI tasks on the RHI thread
+				}
+				else
+				{
+					return false;
+				}
 			}
 
 			// avoid stack overflow. is not expected in a real-life cases but happens in stress tests
diff --git a/Engine/Source/Runtime/Core/Public/Async/TaskGraphInterfaces.h b/Engine/Source/Runtime/Core/Public/Async/TaskGraphInterfaces.h
--- a/Engine/Source/Runtime/Core/Public/Async/TaskGraphInterfaces.h
+++ b/Engine/Source/Runtime/Core/Public/Async/TaskGraphInterfaces.h
@@ -492,7 +492,7 @@ public:
 		checkSlow(NewTasks.Num() == 0);
 		checkSlow(bDeleteOnCompletion);
 		checkSlow(IsNamedThreadTask());
-		verify(TryExecuteTask());
+		TryExecuteTask(); // It's possible retraction already executed this task
 		ReleaseInternalReference(); // named tasks are executed by named threads, outside of the scheduler
 	}
 
diff --git a/Engine/Source/Runtime/Core/Tests/Tasks/TasksTest.cpp b/Engine/Source/Runtime/Core/Tests/Tasks/TasksTest.cpp
--- a/Engine/Source/Runtime/Core/Tests/Tasks/TasksTest.cpp
+++ b/Engine/Source/Runtime/Core/Tests/Tasks/TasksTest.cpp
@@ -376,6 +376,22 @@ namespace UE { namespace TasksTests
 			GTTask.Wait();
 		}
 
+		{	// a basic test for multiple named thread tasks
+			// This will exercise retraction on named thread
+			TArray<FTask> GTTasks;
+			for (int32 Index = 0; Index < 4; ++Index)
+			{
+				GTTasks.Add(Launch
+				(
+					UE_SOURCE_LOCATION,
+					[] { check(IsInGameThread()); },
+					ETaskPriority::Default,
+					EExtendedTaskPriority::GameThreadNormalPri
+				));
+			}
+			UE::Tasks::Wait(GTTasks);
+		}
+
 		//{	// Cancelled task is completed only after nested
 		//	FTaskEvent ParentBlocker{ UE_SOURCE_LOCATION };
 		//	FTaskEvent NestedBlocker{ UE_SOURCE_LOCATION };
