Increase default quality for view distance and textures in perf mode on gaming handhelds

Also change FWindowsPlatformMisc::IsDeviceHandheld() to FPlatformMisc::IsDeviceGamingHandheld() to include other desktop platforms and imply a level of gaming performance.

diff --git a/Engine/Source/Runtime/Core/Private/GenericPlatform/GenericPlatformMisc.cpp b/Engine/Source/Runtime/Core/Private/GenericPlatform/GenericPlatformMisc.cpp
--- a/Engine/Source/Runtime/Core/Private/GenericPlatform/GenericPlatformMisc.cpp
+++ b/Engine/Source/Runtime/Core/Private/GenericPlatform/GenericPlatformMisc.cpp
@@ -1851,6 +1851,11 @@ FText FGenericPlatformMisc::GetFileManagerName()
 	return NSLOCTEXT("GenericPlatform", "FileManagerName", "File Manager");
 }
 
+bool FGenericPlatformMisc::IsDeviceGamingHandheld()
+{
+	return false;
+}
+
 bool FGenericPlatformMisc::IsRunningOnBattery()
 {
 	return false;
diff --git a/Engine/Source/Runtime/Core/Private/Windows/WindowsPlatformMisc.cpp b/Engine/Source/Runtime/Core/Private/Windows/WindowsPlatformMisc.cpp
--- a/Engine/Source/Runtime/Core/Private/Windows/WindowsPlatformMisc.cpp
+++ b/Engine/Source/Runtime/Core/Private/Windows/WindowsPlatformMisc.cpp
@@ -950,7 +950,7 @@ void FWindowsPlatformMisc::PlatformInit()
 	UE_LOG(LogInit, Log, TEXT("High frequency timer resolution =%f MHz"), 0.000001 / FPlatformTime::GetSecondsPerCycle() );
 
 	// If the Device has been identified as a Handheld device
-	UE_LOG(LogInit, Log, TEXT("Handheld Device: %d"), FWindowsPlatformMisc::IsDeviceHandheld() ? 1 : 0);
+	UE_LOG(LogInit, Log, TEXT("Gaming Handheld Device: %d"), FPlatformMisc::IsDeviceGamingHandheld() ? 1 : 0);
 
 	// Register on the game thread.
 	FWindowsPlatformStackWalk::RegisterOnModulesChanged();
@@ -4251,7 +4251,7 @@ FText FWindowsPlatformMisc::GetFileManagerName()
 	return NSLOCTEXT("WindowsPlatform", "FileManagerName", "Explorer");
 }
 
-bool FWindowsPlatformMisc::IsDeviceHandheld()
+bool FWindowsPlatformMisc::IsDeviceGamingHandheld()
 {
 #ifndef DEVICEFAMILYDEVICEFORM_GAMING_HANDHELD
 	#define DEVICEFAMILYDEVICEFORM_GAMING_HANDHELD    0x0000002E 
diff --git a/Engine/Source/Runtime/Core/Public/GenericPlatform/GenericPlatformMisc.h b/Engine/Source/Runtime/Core/Public/GenericPlatform/GenericPlatformMisc.h
--- a/Engine/Source/Runtime/Core/Public/GenericPlatform/GenericPlatformMisc.h
+++ b/Engine/Source/Runtime/Core/Public/GenericPlatform/GenericPlatformMisc.h
@@ -1642,6 +1642,11 @@ public:
 		return false;
 	}
 
+	/**
+	 * Returns whether the device is considered a gaming handheld device. It's not a dynamic property, the status is cached at startup.
+	 */
+	static CORE_API bool IsDeviceGamingHandheld();
+
 	/**
 	 * Returns whether the platform is running on battery power or not.
 	 */
diff --git a/Engine/Source/Runtime/Core/Public/Windows/WindowsPlatformMisc.h b/Engine/Source/Runtime/Core/Public/Windows/WindowsPlatformMisc.h
--- a/Engine/Source/Runtime/Core/Public/Windows/WindowsPlatformMisc.h
+++ b/Engine/Source/Runtime/Core/Public/Windows/WindowsPlatformMisc.h
@@ -332,10 +332,9 @@ struct FWindowsPlatformMisc
 	}
 
 	/**
-	 * Returns whether the device is considered an Handheld device. It's not a dynamic property, the status is cached at startup.
+	 * Returns whether the device is considered a gaming handheld device. It's not a dynamic property, the status is cached at startup.
 	 */
-
-	static CORE_API bool IsDeviceHandheld();
+	static CORE_API bool IsDeviceGamingHandheld();
 	
 	/**
 	 * Returns whether the platform is running on battery power or not. For Windows laptops this is true regardless of charging state.
