- Move to the engine default stack size numbers since ini files being read too late TaskThreadStackSize and PoolThreadStackSize were effectively dead entries.
- Add TaskGraph.ThreadStackSize to allow configuring stack size for task graph worker threads

diff --git a/Engine/Config/BaseEngine.ini b/Engine/Config/BaseEngine.ini
--- a/Engine/Config/BaseEngine.ini
+++ b/Engine/Config/BaseEngine.ini
@@ -1603,9 +1603,6 @@ IgnoreInlineBulkDataReloadEnsures=False
 HangDuration=0.0
 AssetLogShowsDiskPath=True
 AssetLogShowsAbsolutePath=True
-TaskThreadStackSize=393216
-PoolThreadStackSize=131072
-DefaultThreadStackSize=1048576
 
 [Core.System.Experimental]
 EnablePackageSidecarSaving=False
diff --git a/Engine/Source/Runtime/Core/Private/Async/Fundamental/Scheduler.cpp b/Engine/Source/Runtime/Core/Private/Async/Fundamental/Scheduler.cpp
--- a/Engine/Source/Runtime/Core/Private/Async/Fundamental/Scheduler.cpp
+++ b/Engine/Source/Runtime/Core/Private/Async/Fundamental/Scheduler.cpp
@@ -175,7 +175,7 @@ namespace LowLevelTasks
 
 	FScheduler FScheduler::Singleton;
 
-	TUniquePtr<FThread> FScheduler::CreateWorker(uint32 WorkerId, const TCHAR* Name, bool bPermitBackgroundWork, FThread::EForkable IsForkable, Private::FWaitEvent* ExternalWorkerEvent, FSchedulerTls::FLocalQueueType* ExternalWorkerLocalQueue, EThreadPriority Priority, uint64 InAffinity)
+	TUniquePtr<FThread> FScheduler::CreateWorker(uint32 WorkerId, const TCHAR* Name, bool bPermitBackgroundWork, FThread::EForkable IsForkable, Private::FWaitEvent* ExternalWorkerEvent, FSchedulerTls::FLocalQueueType* ExternalWorkerLocalQueue, EThreadPriority Priority, uint64 InAffinity, uint32 StackSize)
 	{
 		const uint32 WaitTimes[8] = { 719, 991, 1361, 1237, 1597, 953, 587, 1439 };
 		uint32 WaitTime = WaitTimes[WorkerId % 8];
@@ -218,11 +218,11 @@ namespace LowLevelTasks
 			[this, ExternalWorkerEvent, ExternalWorkerLocalQueue, WaitTime, bPermitBackgroundWork]
 			{ 
 				WorkerMain(ExternalWorkerEvent, ExternalWorkerLocalQueue, WaitTime, bPermitBackgroundWork);
-			}, 0, Priority, FThreadAffinity{ ThreadAffinityMask & ProcessorGroups.ThreadAffinities[CpuGroup], CpuGroup }, IsForkable
+			}, StackSize, Priority, FThreadAffinity{ ThreadAffinityMask & ProcessorGroups.ThreadAffinities[CpuGroup], CpuGroup }, IsForkable
 		);
 	}
 
-	void FScheduler::StartWorkers(uint32 NumForegroundWorkers, uint32 NumBackgroundWorkers, FThread::EForkable IsForkable, EThreadPriority InWorkerPriority,  EThreadPriority InBackgroundPriority, uint64 InWorkerAffinity, uint64 InBackgroundAffinity)
+	void FScheduler::StartWorkers(uint32 NumForegroundWorkers, uint32 NumBackgroundWorkers, FThread::EForkable IsForkable, EThreadPriority InWorkerPriority,  EThreadPriority InBackgroundPriority, uint64 InWorkerAffinity, uint64 InBackgroundAffinity, uint32 InWorkersStackSize)
 	{
 		// It always been a given that only the game thread should start workers so just add validation.
 		check(IsInGameThread());
@@ -255,6 +255,8 @@ namespace LowLevelTasks
 			BackgroundAffinity = InBackgroundAffinity;
 		}
 
+		WorkerStackSize = InWorkersStackSize;
+
 		const bool bSupportsMultithreading = FPlatformProcess::SupportsMultithreading() || FForkProcessHelper::IsForkedMultithreadInstance();
 
 		uint32 OldActiveWorkers = ActiveWorkers.load(std::memory_order_relaxed);
@@ -285,7 +287,7 @@ namespace LowLevelTasks
 			WorkerThreads.Reset(new std::atomic<FThread*>[MaxWorkers]);
 
 			auto CreateThread =
-				[this, IsForkable](Private::ELocalQueueType LocalQueueType, const TCHAR* ThreadGroup, const TCHAR* Prefix, std::atomic<int32>& CreationIndex, int32 NumWorkers, int32 NumMaxWorkers, EThreadPriority Priority, uint64 Affinity)
+				[this, IsForkable](Private::ELocalQueueType LocalQueueType, const TCHAR* ThreadGroup, const TCHAR* Prefix, std::atomic<int32>& CreationIndex, int32 NumWorkers, int32 NumMaxWorkers, EThreadPriority Priority, uint64 Affinity, uint32 StackSize)
 				{
 					TRACE_CPUPROFILER_EVENT_SCOPE(CreateWorkerThread);
 					LLM_SCOPE_BYNAME(TEXT("EngineMisc/WorkerThreads"));
@@ -319,7 +321,8 @@ namespace LowLevelTasks
 							&WorkerEvents[WorkerId],
 							&WorkerLocalQueues[WorkerId],
 							Priority,
-							Affinity).Release();
+							Affinity,
+							StackSize).Release();
 
 					UE::Trace::ThreadGroupEnd();
 				};
@@ -327,13 +330,13 @@ namespace LowLevelTasks
 			TFunction<void()> ForegroundCreateThread =
 				[this, CreateThread, NumWorkers = NumForegroundWorkers, MaxWorkers = MaxForegroundWorkers]()
 				{
-					CreateThread(Private::ELocalQueueType::EForeground, TEXT("Foreground Workers"), TEXT("Foreground"), ForegroundCreationIndex, NumWorkers, MaxWorkers, WorkerPriority, WorkerAffinity);
+					CreateThread(Private::ELocalQueueType::EForeground, TEXT("Foreground Workers"), TEXT("Foreground"), ForegroundCreationIndex, NumWorkers, MaxWorkers, WorkerPriority, WorkerAffinity, WorkerStackSize);
 				};
 
 			TFunction<void()> BackgroundCreateThread =
 				[this, CreateThread, NumWorkers = NumBackgroundWorkers, MaxWorkers = MaxBackgroundWorkers, ActualBackgroundPriority]()
 				{
-					CreateThread(Private::ELocalQueueType::EBackground, TEXT("Background Workers"), TEXT("Background"), BackgroundCreationIndex, NumWorkers, MaxWorkers, ActualBackgroundPriority, BackgroundAffinity);
+					CreateThread(Private::ELocalQueueType::EBackground, TEXT("Background Workers"), TEXT("Background"), BackgroundCreationIndex, NumWorkers, MaxWorkers, ActualBackgroundPriority, BackgroundAffinity, WorkerStackSize);
 				};
 
 			WaitingQueue[0].Init(NumForegroundWorkers, MaxForegroundWorkers, ForegroundCreateThread, GTaskGraphUseDynamicThreadCreation ? 0 : MaxForegroundWorkers /* ActiveThreadCount */);
@@ -508,7 +511,7 @@ namespace LowLevelTasks
 		return false;
 	}
 
-	void FScheduler::RestartWorkers(uint32 NumForegroundWorkers, uint32 NumBackgroundWorkers, FThread::EForkable IsForkable, EThreadPriority InWorkerPriority, EThreadPriority InBackgroundPriority, uint64 InWorkerAffinity, uint64 InBackgroundAffinity)
+	void FScheduler::RestartWorkers(uint32 NumForegroundWorkers, uint32 NumBackgroundWorkers, FThread::EForkable IsForkable, EThreadPriority InWorkerPriority, EThreadPriority InBackgroundPriority, uint64 InWorkerAffinity, uint64 InBackgroundAffinity, uint32 InWorkersStackSize)
 	{
 		UE::TUniqueLock Lock(WorkerThreadsCS);
 		TemporaryShutdown.store(true, std::memory_order_release);
@@ -518,7 +521,7 @@ namespace LowLevelTasks
 		}
 		const bool bDrainGlobalQueue = false;
 		StopWorkers(bDrainGlobalQueue);
-		StartWorkers(NumForegroundWorkers, NumBackgroundWorkers, IsForkable, InWorkerPriority, InBackgroundPriority, InWorkerAffinity, InBackgroundAffinity);
+		StartWorkers(NumForegroundWorkers, NumBackgroundWorkers, IsForkable, InWorkerPriority, InBackgroundPriority, InWorkerAffinity, InBackgroundAffinity, InWorkersStackSize);
 		TemporaryShutdown.store(false, std::memory_order_release);
 	}
 
diff --git a/Engine/Source/Runtime/Core/Private/Async/TaskGraph.cpp b/Engine/Source/Runtime/Core/Private/Async/TaskGraph.cpp
--- a/Engine/Source/Runtime/Core/Private/Async/TaskGraph.cpp
+++ b/Engine/Source/Runtime/Core/Private/Async/TaskGraph.cpp
@@ -135,6 +135,14 @@ static FAutoConsoleVariableRef CVarNumForegroundWorkers(
 	TEXT("Configures the number of foreground worker threads. Requires the scheduler to be restarted to have an affect")
 );
 
+CORE_API int32 GTaskThreadStackSize = 0;
+static FAutoConsoleVariableRef CVarTaskThreadStackSize(
+	TEXT("TaskGraph.ThreadStackSize"),
+	GTaskThreadStackSize,
+	TEXT("Thread Stack Size, 0 means platform default")
+);
+
+
 #if CREATE_HIPRI_TASK_THREADS || CREATE_BACKGROUND_TASK_THREADS
 	static void ThreadSwitchForABTest(const TArray<FString>& Args)
 	{
diff --git a/Engine/Source/Runtime/Core/Public/Async/Fundamental/Scheduler.h b/Engine/Source/Runtime/Core/Public/Async/Fundamental/Scheduler.h
--- a/Engine/Source/Runtime/Core/Public/Async/Fundamental/Scheduler.h
+++ b/Engine/Source/Runtime/Core/Public/Async/Fundamental/Scheduler.h
@@ -113,9 +113,9 @@ namespace LowLevelTasks
 		inline static FScheduler& Get();
 
 		//start number of workers where 0 is the system default
-		CORE_API void StartWorkers(uint32 NumForegroundWorkers = 0, uint32 NumBackgroundWorkers = 0, FThread::EForkable IsForkable = FThread::NonForkable, EThreadPriority InWorkerPriority = EThreadPriority::TPri_Normal, EThreadPriority InBackgroundPriority = EThreadPriority::TPri_BelowNormal, uint64 InWorkerAffinity = 0, uint64 InBackgroundAffinity = 0);
+		CORE_API void StartWorkers(uint32 NumForegroundWorkers = 0, uint32 NumBackgroundWorkers = 0, FThread::EForkable IsForkable = FThread::NonForkable, EThreadPriority InWorkerPriority = EThreadPriority::TPri_Normal, EThreadPriority InBackgroundPriority = EThreadPriority::TPri_BelowNormal, uint64 InWorkerAffinity = 0, uint64 InBackgroundAffinity = 0, uint32 InWorkersStackSize = 0);
 		CORE_API void StopWorkers(bool DrainGlobalQueue = true);
-		CORE_API void RestartWorkers(uint32 NumForegroundWorkers = 0, uint32 NumBackgroundWorkers = 0, FThread::EForkable IsForkable = FThread::NonForkable, EThreadPriority WorkerPriority = EThreadPriority::TPri_Normal, EThreadPriority BackgroundPriority = EThreadPriority::TPri_BelowNormal, uint64 InWorkerAffinity = 0, uint64 InBackgroundAffinity = 0);
+		CORE_API void RestartWorkers(uint32 NumForegroundWorkers = 0, uint32 NumBackgroundWorkers = 0, FThread::EForkable IsForkable = FThread::NonForkable, EThreadPriority WorkerPriority = EThreadPriority::TPri_Normal, EThreadPriority BackgroundPriority = EThreadPriority::TPri_BelowNormal, uint64 InWorkerAffinity = 0, uint64 InBackgroundAffinity = 0, uint32 InWorkersStackSize = 0);
 
 		//try to launch the task, the return value will specify if the task was in the ready state and has been launched
 		inline bool TryLaunch(FTask& Task, EQueuePreference QueuePreference = EQueuePreference::DefaultPreference, bool bWakeUpWorker = true);
@@ -146,7 +146,7 @@ namespace LowLevelTasks
 
 	private: 
 		[[nodiscard]] FTask* ExecuteTask(FTask* InTask);
-		TUniquePtr<FThread> CreateWorker(uint32 WorkerId, const TCHAR* Name, bool bPermitBackgroundWork = false, FThread::EForkable IsForkable = FThread::NonForkable, Private::FWaitEvent* ExternalWorkerEvent = nullptr, FSchedulerTls::FLocalQueueType* ExternalWorkerLocalQueue = nullptr, EThreadPriority Priority = EThreadPriority::TPri_Normal, uint64 InAffinity = 0);
+		TUniquePtr<FThread> CreateWorker(uint32 WorkerId, const TCHAR* Name, bool bPermitBackgroundWork = false, FThread::EForkable IsForkable = FThread::NonForkable, Private::FWaitEvent* ExternalWorkerEvent = nullptr, FSchedulerTls::FLocalQueueType* ExternalWorkerLocalQueue = nullptr, EThreadPriority Priority = EThreadPriority::TPri_Normal, uint64 InAffinity = 0, uint32 StackSize = 0);
 		void WorkerMain(Private::FWaitEvent* WorkerEvent, FSchedulerTls::FLocalQueueType* ExternalWorkerLocalQueue, uint32 WaitCycles, bool bPermitBackgroundWork);
 		void StandbyLoop(Private::FWaitEvent* WorkerEvent, FSchedulerTls::FLocalQueueType* ExternalWorkerLocalQueue, uint32 WaitCycles, bool bPermitBackgroundWork);
 		void WorkerLoop(Private::FWaitEvent* WorkerEvent, FSchedulerTls::FLocalQueueType* ExternalWorkerLocalQueue, uint32 WaitCycles, bool bPermitBackgroundWork);
@@ -172,6 +172,7 @@ namespace LowLevelTasks
 		std::atomic<int32>                             BackgroundCreationIndex{ 0 };
 		uint64                                         WorkerAffinity = 0;
 		uint64                                         BackgroundAffinity = 0;
+		uint32										   WorkerStackSize = 0;
 		EThreadPriority                                WorkerPriority = EThreadPriority::TPri_Normal;
 		EThreadPriority                                BackgroundPriority = EThreadPriority::TPri_BelowNormal;
 		std::atomic_bool                               TemporaryShutdown{ false };
diff --git a/Engine/Source/Runtime/Core/Public/GenericPlatform/GenericPlatformProcess.h b/Engine/Source/Runtime/Core/Public/GenericPlatform/GenericPlatformProcess.h
--- a/Engine/Source/Runtime/Core/Public/GenericPlatform/GenericPlatformProcess.h
+++ b/Engine/Source/Runtime/Core/Public/GenericPlatform/GenericPlatformProcess.h
@@ -377,6 +377,18 @@ struct FGenericPlatformProcess
 	/** Get the active stack size for the currently running thread **/
 	static CORE_API uint32 GetStackSize();
 
+	/** Get Default Thread Stack Size in bytes, 0 means to use the platform OS default size **/
+	static uint32 GetDefaultThreadStackSize()
+	{
+		return 0;
+	}
+
+	/** Get Default Pool Thread Stack Size in bytes **/
+	static uint32 GetDefaultPoolThreadStackSize()
+	{
+		return 128 * 1024;
+	}
+
 	/** Output information about the currently active thread **/
 	static void DumpThreadInfo( const TCHAR* MarkerName ) { }
 
diff --git a/Engine/Source/Runtime/Engine/Private/UnrealEngine.cpp b/Engine/Source/Runtime/Engine/Private/UnrealEngine.cpp
--- a/Engine/Source/Runtime/Engine/Private/UnrealEngine.cpp
+++ b/Engine/Source/Runtime/Engine/Private/UnrealEngine.cpp
@@ -19805,7 +19805,18 @@ static void SetupThreadConfig(const TArray<FString>& Args)
 	int32 NumWorkerThreads = FPlatformMisc::NumberOfWorkerThreadsToSpawn();
 	int32 NumBackgroundWorkers = FMath::Max(1, NumWorkerThreads - FMath::Min<int>(2, NumWorkerThreads));
 	int32 NumForegroundWorkers = FMath::Max(1, NumWorkerThreads - NumBackgroundWorkers);
-	LowLevelTasks::FScheduler::Get().RestartWorkers(NumForegroundWorkers, NumBackgroundWorkers, FForkProcessHelper::IsForkedMultithreadInstance() ? FThread::Forkable : FThread::NonForkable, TaskThreadConfig.Priority, TaskBPThreadConfig.Priority, TaskThreadConfig.Affinity, TaskBPThreadConfig.Affinity);
+
+	// Setting the stack size won't have an effect for the threads created before InitializeConfigSystem() is called
+	int32 TaskGraphStackSize = 0;
+
+	const IConsoleVariable* TaskGraphStackSizeCVar = IConsoleManager::Get().FindConsoleVariable(TEXT("TaskGraph.ThreadStackSize"));
+	
+	if (TaskGraphStackSizeCVar)
+	{
+		TaskGraphStackSize = TaskGraphStackSizeCVar->GetInt();
+	}
+
+	LowLevelTasks::FScheduler::Get().RestartWorkers(NumForegroundWorkers, NumBackgroundWorkers, FForkProcessHelper::IsForkedMultithreadInstance() ? FThread::Forkable : FThread::NonForkable, TaskThreadConfig.Priority, TaskBPThreadConfig.Priority, TaskThreadConfig.Affinity, TaskBPThreadConfig.Affinity, TaskGraphStackSize);
 
 	WaitEvent.Notify();
 	ExitEvent.Wait();
@@ -19861,7 +19872,8 @@ void InitThreadConfig()
 		bool bRequiresRestart =
 			IsCVarModified(TEXT("TaskGraph.UseDynamicPrioritization")) ||
 			IsCVarModified(TEXT("TaskGraph.UseDynamicThreadCreation")) ||
-			IsCVarModified(TEXT("TaskGraph.OversubscriptionRatio"));
+			IsCVarModified(TEXT("TaskGraph.OversubscriptionRatio")) ||
+			IsCVarModified(TEXT("TaskGraph.ThreadStackSize"));
 
 		if (bRequiresRestart || !ThreadConfigEntries.IsEmpty())
 		{
diff --git a/Engine/Source/Runtime/Launch/Private/LaunchEngineLoop.cpp b/Engine/Source/Runtime/Launch/Private/LaunchEngineLoop.cpp
--- a/Engine/Source/Runtime/Launch/Private/LaunchEngineLoop.cpp
+++ b/Engine/Source/Runtime/Launch/Private/LaunchEngineLoop.cpp
@@ -2592,9 +2592,7 @@ int32 FEngineLoop::PreInitPreStartupScreen(const TCHAR* CmdLine)
 	{
 		SCOPED_BOOT_TIMING("Init FQueuedThreadPool's");
 
-		int32 StackSize = 128 * 1024;
-		// GConfig is not initialized yet, the only solution for now is to hardcode desired values
-		// GConfig->GetInt(TEXT("Core.System"), TEXT("PoolThreadStackSize"), StackSize, GEngineIni);
+		int32 StackSize = FPlatformProcess::GetDefaultPoolThreadStackSize();
 
 		bool bForceEditorStackSize = false;
 #if WITH_EDITOR
