- Fix potentially leaking tasks
- Fix skipping a write if GRemoteConfig.Write returns true

#rnx

diff --git a/Engine/Source/Runtime/Core/Private/Misc/RemoteConfigIni.cpp b/Engine/Source/Runtime/Core/Private/Misc/RemoteConfigIni.cpp
--- a/Engine/Source/Runtime/Core/Private/Misc/RemoteConfigIni.cpp
+++ b/Engine/Source/Runtime/Core/Private/Misc/RemoteConfigIni.cpp
@@ -113,11 +113,11 @@ void FRemoteConfigAsyncTaskManager::Tick()
 {
 	FScopeLock ScopeLock(&SynchronizationObject);
 	
-	for (int32 Idx = 0; Idx < CachedWriteTasks.Num(); ++Idx)
+	for ( auto It = CachedWriteTasks.CreateIterator(); It; ++It)
 	{
-		if (GRemoteConfig.Write(*CachedWriteTasks[Idx].Filename, CachedWriteTasks[Idx].Contents))
+		if (GRemoteConfig.Write(*It->Filename, It->Contents))
 		{
-			CachedWriteTasks.RemoveAt(Idx);
+			It.RemoveCurrent();
 		}
 	}
 }
@@ -150,6 +150,7 @@ bool FRemoteConfigAsyncTaskManager::StartTask(const TCHAR* InFilename, const TCH
 			if (AsyncTask->IsDone())
 			{
 				// Clear out old write tasks that have completed
+				delete AsyncTask;
 				PendingTasks.Remove(FString(InFilename));
 			}
 			else
@@ -200,13 +201,14 @@ bool FRemoteConfigAsyncTaskManager::AreAllTasksFinished(bool bDoRemoval)
 		}
 
 		FAsyncTask<FRemoteConfigAsyncWorker>* AsyncTask = NULL;
-		AsyncTask = PendingTasks.FindRef(FString(It.Key()));
+		AsyncTask = PendingTasks.FindRef(It.Key());
 			
 		if (AsyncTask)
 		{
 			if (AsyncTask->IsDone() && bDoRemoval)
 			{
 				// Remove completed tasks from the queue if so desired
+				delete AsyncTask;
 				It.RemoveCurrent();
 			}
 			else
