Csv Profiler: Add support for _NOLOG events to reduce logspam from events and reduce unnecessary string ops to speed up boot times
Also move some debug code behind CSV_PROFILER_ALLOW_DEBUG_FEATURES
#tests tested on console

diff --git a/Engine/Source/Runtime/Core/Private/ProfilingDebugging/CsvProfiler.cpp b/Engine/Source/Runtime/Core/Private/ProfilingDebugging/CsvProfiler.cpp
--- a/Engine/Source/Runtime/Core/Private/ProfilingDebugging/CsvProfiler.cpp
+++ b/Engine/Source/Runtime/Core/Private/ProfilingDebugging/CsvProfiler.cpp
@@ -4473,7 +4473,7 @@ FString FCsvProfiler::GetDefaultDirectory()
 	return FPaths::ProfilingDir() + TEXT("CSV/");
 }
 
-void FCsvProfiler::RecordEventfInternal(int32 CategoryIndex, const TCHAR* Fmt, ...)
+void FCsvProfiler::RecordEventfInternal(int32 CategoryIndex, bool bLog, const TCHAR* Fmt, ...)
 {
 	bool bIsCsvRecording = GCsvProfilerIsCapturing && GCsvCategoriesEnabled[CategoryIndex];
 	if (bIsCsvRecording || GStartOnEvent)
@@ -4486,13 +4486,16 @@ void FCsvProfiler::RecordEventfInternal(int32 CategoryIndex, const TCHAR* Fmt, .
 
 		if (bIsCsvRecording)
 		{
-			RecordEvent(CategoryIndex, Str);
+			RecordEvent(CategoryIndex, Str, bLog);
 
+#if CSV_PROFILER_ALLOW_DEBUG_FEATURES
 			if (GStopOnEvent && GStopOnEvent->Equals(Str, ESearchCase::IgnoreCase))
 			{
 				FCsvProfiler::Get()->EndCapture();
 			}
+#endif
 		}
+#if CSV_PROFILER_ALLOW_DEBUG_FEATURES
 		else
 		{
 			if (GStartOnEvent && GStartOnEvent->Equals(Str, ESearchCase::IgnoreCase))
@@ -4500,10 +4503,11 @@ void FCsvProfiler::RecordEventfInternal(int32 CategoryIndex, const TCHAR* Fmt, .
 				FCsvProfiler::Get()->BeginCapture(FCsvProfiler::Get()->GetNumFrameToCaptureOnEvent());
 			}
 		}
+#endif
 	}
 }
 
-void FCsvProfiler::RecordEvent(int32 CategoryIndex, const FString& EventText)
+void FCsvProfiler::RecordEvent(int32 CategoryIndex, const FString& EventText, bool bLog)
 {
 	if (GCsvProfilerIsCapturing && GCsvCategoriesEnabled[CategoryIndex])
 	{
@@ -4517,11 +4521,17 @@ void FCsvProfiler::RecordEvent(int32 CategoryIndex, const FString& EventText)
 		if (CategoryIndex != CSV_CATEGORY_INDEX_GLOBAL)
 		{
 			CategoryName = FCsvCategoryData::Get()->GetCategoryNameByIndex(CategoryIndex);
-			FullEventText = FString::Printf(TEXT("%s/%s"), *CategoryName, *EventText);
+			if (bLog || GCsvEventExecCmds)
+			{
+				FullEventText = FString::Printf(TEXT("%s/%s"), *CategoryName, *EventText);
+			}			
+		}
+		if (bLog)
+		{
+			UE_LOG(LogCsvProfiler, Display, TEXT("CSVEvent \"%s\" [Frame %d]"), *FullEventText, FCsvProfiler::Get()->GetCaptureFrameNumber());
 		}
-		UE_LOG(LogCsvProfiler, Display, TEXT("CSVEvent \"%s\" [Frame %d]"), *FullEventText, FCsvProfiler::Get()->GetCaptureFrameNumber());
+#if CSV_PROFILER_ALLOW_DEBUG_FEATURES
 		FCsvProfiler::Get()->OnCSVProfileEvent().Broadcast(CategoryName, EventText);
-
 		if (GCsvEventExecCmds)
 		{
 			FStringView Match(FullEventText);
@@ -4533,12 +4543,14 @@ void FCsvProfiler::RecordEvent(int32 CategoryIndex, const FString& EventText)
 				}
 			}
 		}
+#endif 
+
 	}
 }
 
-void FCsvProfiler::RecordEventAtFrameStart(int32 CategoryIndex, const FString& EventText)
+void FCsvProfiler::RecordEventAtFrameStart(int32 CategoryIndex, const FString& EventText, bool bLog)
 {
-	RecordEventAtTimestamp(CategoryIndex, EventText, FCsvProfiler::Get()->LastEndFrameTimestamp);
+	RecordEventAtTimestamp(CategoryIndex, EventText, FCsvProfiler::Get()->LastEndFrameTimestamp, bLog);
 }
 
 void FCsvProfiler::SetMetadata(const TCHAR* Key, const TCHAR* Value)
@@ -4599,12 +4611,15 @@ void FCsvProfiler::SetMetadataInternal(const TCHAR* Key, const TCHAR* Value, boo
 	}
 }
 
-void FCsvProfiler::RecordEventAtTimestamp(int32 CategoryIndex, const FString& EventText, uint64 Cycles64)
+void FCsvProfiler::RecordEventAtTimestamp(int32 CategoryIndex, const FString& EventText, uint64 Cycles64, bool bLog)
 {
 	if (GCsvProfilerIsCapturing && GCsvCategoriesEnabled[CategoryIndex])
 	{
 		LLM_SCOPE(ELLMTag::CsvProfiler);
-		UE_LOG(LogCsvProfiler, Display, TEXT("CSVEvent [Frame %d] : \"%s\""), FCsvProfiler::Get()->GetCaptureFrameNumber(), *EventText);
+		if (bLog)
+		{
+			UE_LOG(LogCsvProfiler, Display, TEXT("CSVEvent [Frame %d] : \"%s\""), FCsvProfiler::Get()->GetCaptureFrameNumber(), *EventText);
+		}
 		FCsvProfilerThreadData::Get().AddEventWithTimestamp(EventText, CategoryIndex,Cycles64);
 
 		if (IsContinuousWriteEnabled(false))
diff --git a/Engine/Source/Runtime/Core/Public/ProfilingDebugging/CsvProfiler.h b/Engine/Source/Runtime/Core/Public/ProfilingDebugging/CsvProfiler.h
--- a/Engine/Source/Runtime/Core/Public/ProfilingDebugging/CsvProfiler.h
+++ b/Engine/Source/Runtime/Core/Public/ProfilingDebugging/CsvProfiler.h
@@ -62,6 +62,14 @@ struct FCsvDeclaredStat;
 		FCsvProfiler::RecordEventf( CSV_CATEGORY_INDEX_GLOBAL, Format, ##__VA_ARGS__ ); \
 		TRACE_BOOKMARK(Format, ##__VA_ARGS__)
 
+	#define CSV_EVENT_NOLOG(Category, Format, ...) \
+		FCsvProfiler::RecordEventfNoLog( CSV_CATEGORY_INDEX(Category), Format, ##__VA_ARGS__ ); \
+		TRACE_BOOKMARK(TEXT(PREPROCESSOR_TO_STRING(Category)) TEXT("/") Format, ##__VA_ARGS__)
+
+	#define CSV_EVENT_GLOBAL_NOLOG(Format, ...) \
+		FCsvProfiler::RecordEventfNoLog( CSV_CATEGORY_INDEX_GLOBAL, Format, ##__VA_ARGS__ ); \
+		TRACE_BOOKMARK(Format, ##__VA_ARGS__)
+
 	// Metadata
 	#define CSV_METADATA(Key,Value)									FCsvProfiler::SetMetadata( Key, Value )
 	#define CSV_NON_PERSISTENT_METADATA(Key,Value)					FCsvProfiler::SetNonPersistentMetadata( Key, Value )
@@ -80,6 +88,8 @@ struct FCsvDeclaredStat;
 	#define CSV_DECLARE_CATEGORY_MODULE_EXTERN(Module_API,CategoryName)			
 	#define CSV_EVENT(Category, Format, ...) 						
 	#define CSV_EVENT_GLOBAL(Format, ...)
+	#define CSV_EVENT_NOLOG(Category, Format, ...) 						
+	#define CSV_EVENT_GLOBAL_NOLOG(Format, ...)
 	#define CSV_METADATA(Key,Value)
 	#define CSV_NON_PERSISTENT_METADATA(Key,Value)
 	#define CSV_CUSTOM_STAT_MINIMAL(Category,StatName,Value,Op)
@@ -394,9 +404,9 @@ public:
 	CORE_API static void RecordCustomStatMinimal(const char* StatName, uint32 CategoryIndex, int32 Value, const ECsvCustomStatOp CustomStatOp);
 	CORE_API static void RecordCustomStatMinimal(const char* StatName, uint32 CategoryIndex, double Value, const ECsvCustomStatOp CustomStatOp);
 
-	CORE_API static void RecordEvent(int32 CategoryIndex, const FString& EventText);
-	CORE_API static void RecordEventAtFrameStart(int32 CategoryIndex, const FString& EventText);
-	CORE_API static void RecordEventAtTimestamp(int32 CategoryIndex, const FString& EventText, uint64 Cycles64);
+	CORE_API static void RecordEvent(int32 CategoryIndex, const FString& EventText, bool bLog=true);
+	CORE_API static void RecordEventAtFrameStart(int32 CategoryIndex, const FString& EventText, bool bLog = true);
+	CORE_API static void RecordEventAtTimestamp(int32 CategoryIndex, const FString& EventText, uint64 Cycles64, bool bLog = true);
 
 	/** Metadata values set with this function will persist between captures. */
 	CORE_API static void SetMetadata(const TCHAR* Key, const TCHAR* Value);
@@ -413,7 +423,18 @@ public:
 		{
 			static_assert(TIsArrayOrRefOfTypeByPredicate<FmtType, TIsCharEncodingCompatibleWithTCHAR>::Value, "Formatting string must be a TCHAR array.");
 			static_assert((TIsValidVariadicFunctionArg<Types>::Value && ...), "Invalid argument(s) passed to FCsvProfiler::RecordEventf");
-			RecordEventfInternal(CategoryIndex, (const TCHAR*)Fmt, Args...);
+			RecordEventfInternal(CategoryIndex, true, (const TCHAR*)Fmt, Args...);
+		};
+	}
+
+	template <typename FmtType, typename... Types>
+	inline static void RecordEventfNoLog(int32 CategoryIndex, const FmtType& Fmt, Types... Args)
+	{
+		UE_AUTORTFM_OPEN
+		{
+			static_assert(TIsArrayOrRefOfTypeByPredicate<FmtType, TIsCharEncodingCompatibleWithTCHAR>::Value, "Formatting string must be a TCHAR array.");
+			static_assert((TIsValidVariadicFunctionArg<Types>::Value && ...), "Invalid argument(s) passed to FCsvProfiler::RecordEventfNoLog");
+			RecordEventfInternal(CategoryIndex, false, (const TCHAR*)Fmt, Args...);
 		};
 	}
 
@@ -498,7 +519,8 @@ public:
 	DECLARE_MULTICAST_DELEGATE(FOnCSVProfileEndFrame);
 	FOnCSVProfileEndFrame& OnCSVProfileEndFrame() { return OnCSVProfileEndFrameDelegate; }
 
-	/** Called when an event is emitted */
+	/** Called when an event is emitted. 
+		NOTE: this will only fire if CSV_PROFILER_ALLOW_DEBUG_FEATURES is enabled */
 	DECLARE_MULTICAST_DELEGATE_TwoParams(FOnCSVEvent, const FString& /* CategoryName */, const FString& /* EventText */);
 	FOnCSVEvent& OnCSVProfileEvent() { return OnCSVProfileEventDelegate; }
 
@@ -518,7 +540,7 @@ public:
 	CORE_API TCsvPersistentCustomStat<float>* GetOrCreatePersistentCustomStatFloat(FName Name, int32 CategoryIndex = CSV_CATEGORY_INDEX_GLOBAL, bool bResetEachFrame = false);
 
 private:
-	CORE_API static void VARARGS RecordEventfInternal(int32 CategoryIndex, const TCHAR* Fmt, ...);
+	CORE_API static void VARARGS RecordEventfInternal(int32 CategoryIndex, bool bLog, const TCHAR* Fmt, ...);
 
 	enum class EMetadataPersistenceType : int8
 	{
diff --git a/Engine/Source/Runtime/CoreUObject/Private/Serialization/AsyncPackageLoader.cpp b/Engine/Source/Runtime/CoreUObject/Private/Serialization/AsyncPackageLoader.cpp
--- a/Engine/Source/Runtime/CoreUObject/Private/Serialization/AsyncPackageLoader.cpp
+++ b/Engine/Source/Runtime/CoreUObject/Private/Serialization/AsyncPackageLoader.cpp
@@ -326,7 +326,7 @@ void FlushAsyncLoading(TConstArrayView<int32> RequestIds)
 {
 	TRACE_CPUPROFILER_EVENT_SCOPE(FlushAsyncLoading);
 	CSV_SCOPED_TIMING_STAT_EXCLUSIVE(FlushAsyncLoading);
-	CSV_EVENT_GLOBAL(TEXT("FlushAsyncLoading"));
+	CSV_EVENT_GLOBAL_NOLOG(TEXT("FlushAsyncLoading"));
 
 #if defined(WITH_CODE_GUARD_HANDLER) && WITH_CODE_GUARD_HANDLER
 	void CheckImageIntegrityAtRuntime();
