Fix merge issue resulting from 47321784

diff --git a/Engine/Source/Runtime/Core/Public/ProfilingDebugging/CsvProfiler.h b/Engine/Source/Runtime/Core/Public/ProfilingDebugging/CsvProfiler.h
--- a/Engine/Source/Runtime/Core/Public/ProfilingDebugging/CsvProfiler.h
+++ b/Engine/Source/Runtime/Core/Public/ProfilingDebugging/CsvProfiler.h
@@ -64,7 +64,7 @@ struct FCsvDeclaredStat;
 
 	#define CSV_EVENT_NOLOG(Category, Format, ...) \
 		FCsvProfiler::RecordEventfNoLog( CSV_CATEGORY_INDEX(Category), Format, ##__VA_ARGS__ ); \
-		TRACE_BOOKMARK(TEXT(PREPROCESSOR_TO_STRING(Category)) TEXT("/") Format, ##__VA_ARGS__)
+		TRACE_BOOKMARK(TEXT(UE_STRINGIZE(Category)) TEXT("/") Format, ##__VA_ARGS__)
 
 	#define CSV_EVENT_GLOBAL_NOLOG(Format, ...) \
 		FCsvProfiler::RecordEventfNoLog( CSV_CATEGORY_INDEX_GLOBAL, Format, ##__VA_ARGS__ ); \
