Propagate SetFlags through the console variable Unset function

#tests PIE, Standalone
#rnx

diff --git a/Engine/Source/Runtime/Core/Private/HAL/ConsoleManager.cpp b/Engine/Source/Runtime/Core/Private/HAL/ConsoleManager.cpp
--- a/Engine/Source/Runtime/Core/Private/HAL/ConsoleManager.cpp
+++ b/Engine/Source/Runtime/Core/Private/HAL/ConsoleManager.cpp
@@ -1183,17 +1183,18 @@ protected:
 	 */
 	virtual void Unset(EConsoleVariableFlags SetBy, FName Tag) override
 	{
-		
 		if (PriorityHistory == nullptr)
 		{
 			return;
 		}
 		
+		uint32 UnsetFlags = (uint32)SetBy & ECVF_SetFlagMask;
+
 		// if we don't know the SetBy, unset from all
 		// this isn't ideal because it could call SetInternal multiple times
-		if (SetBy == ECVF_SetByMask)
+		if ((SetBy & ECVF_SetByMask) == ECVF_SetByMask)
 		{
-			#define RECURSE(x) Unset(ECVF_SetBy##x, Tag);
+			#define RECURSE(x) Unset(EConsoleVariableFlags(ECVF_SetBy##x | UnsetFlags), Tag);
 			ENUMERATE_SET_BY(RECURSE)
 			#undef RECURSE
 			
@@ -1215,13 +1216,15 @@ protected:
 			// now figure out the new value
 			EConsoleVariableFlags NewSetBy;
 			auto MaxValue = PriorityHistory->GetMaxValue(NewSetBy);
-			
+
 			// when we preview SGs, we set their value, but dont run the callbacks, so here we are doing the same operation
 			if (GetFlags() & ECVF_ScalabilityGroup)
 			{
-				NewSetBy = (EConsoleVariableFlags)(NewSetBy | ECVF_Set_SetOnly_Unsafe);
+				UnsetFlags |= ECVF_Set_SetOnly_Unsafe;
 			}
-			
+
+			NewSetBy = (EConsoleVariableFlags)(NewSetBy | UnsetFlags);
+
 			UE_LOG(LogConsoleManager, Verbose, TEXT(" |-> Unsetting %s, now %s"), *IConsoleManager::Get().FindConsoleObjectName(this),
 				*TTypeToString<T>::ToString(MaxValue.GetValueOnGameThread()));
 
