Fixing Linux compilation failures by adding missing header includes

diff --git a/Engine/Source/Runtime/Core/Private/Unix/UnixForkPageProtector.cpp b/Engine/Source/Runtime/Core/Private/Unix/UnixForkPageProtector.cpp
--- a/Engine/Source/Runtime/Core/Private/Unix/UnixForkPageProtector.cpp
+++ b/Engine/Source/Runtime/Core/Private/Unix/UnixForkPageProtector.cpp
@@ -3,8 +3,11 @@
 #include "Unix/UnixForkPageProtector.h"
 
 #if COMPILE_FORK_PAGE_PROTECTOR
+#include "CoreGlobals.h"
 #include "HAL/PlatformStackWalk.h"
+#include "Hash/CityHash.h"
 #include "Misc/Fork.h"
+#include "Misc/Paths.h"
 #include "Misc/ScopeExit.h"
 #include "Misc/ScopeLock.h"
 #include "Unix/UnixPlatformFile.h"
diff --git a/Engine/Source/Runtime/Core/Private/Unix/UnixForkPageProtector.h b/Engine/Source/Runtime/Core/Private/Unix/UnixForkPageProtector.h
--- a/Engine/Source/Runtime/Core/Private/Unix/UnixForkPageProtector.h
+++ b/Engine/Source/Runtime/Core/Private/Unix/UnixForkPageProtector.h
@@ -3,6 +3,7 @@
 #pragma once
 
 #include "CoreTypes.h"
+#include "HAL/CriticalSection.h"
 #include "HAL/PlatformMemory.h"
 
 #if COMPILE_FORK_PAGE_PROTECTOR
