Downgrade trace panic to a warning.

#rnx

diff --git a/Engine/Source/Runtime/Core/Private/ProfilingDebugging/TraceAuxiliary.cpp b/Engine/Source/Runtime/Core/Private/ProfilingDebugging/TraceAuxiliary.cpp
--- a/Engine/Source/Runtime/Core/Private/ProfilingDebugging/TraceAuxiliary.cpp
+++ b/Engine/Source/Runtime/Core/Private/ProfilingDebugging/TraceAuxiliary.cpp
@@ -1206,6 +1206,9 @@ void TraceAuxiliaryOnMessageCallback(const UE::Trace::FMessageEvent& Message)
 	case UE::Trace::EMessageType::Display:
 		LOG_WITH_VERBOSITY(Display, Msg);
 		break;
+	case UE::Trace::EMessageType::PanicError:
+		LOG_WITH_VERBOSITY(Warning, Msg);
+		break;
 	default:
 		{
 			if (Message.Type > UE::Trace::EMessageType::FatalStart)
