Made Algo::Sort and related functions accept ranges with arbitrary iterator types.

diff --git a/Engine/Source/Runtime/Core/Public/Algo/Impl/BinaryHeap.h b/Engine/Source/Runtime/Core/Public/Algo/Impl/BinaryHeap.h
--- a/Engine/Source/Runtime/Core/Public/Algo/Impl/BinaryHeap.h
+++ b/Engine/Source/Runtime/Core/Public/Algo/Impl/BinaryHeap.h
@@ -49,14 +49,14 @@ namespace AlgoImpl
 	/**
 	 * Fixes a possible violation of order property between node at Index and a child.
 	 *
-	 * @param	Heap		Pointer to the first element of a binary heap.
+	 * @param	Heap		Random access iterator to the first element of a binary heap.
 	 * @param	Index		Node index.
 	 * @param	Count		Size of the heap.
 	 * @param	InProj		The projection to apply to the elements.
 	 * @param	Predicate	A binary predicate object used to specify if one element should precede another.
 	 */
-	template <typename RangeValueType, typename IndexType, typename ProjectionType, typename PredicateType>
-	inline void HeapSiftDown(RangeValueType* Heap, IndexType Index, const IndexType Count, const ProjectionType& InProj, const PredicateType& Predicate)
+	template <typename IteratorType, typename IndexType, typename ProjectionType, typename PredicateType>
+	inline void HeapSiftDown(IteratorType Heap, IndexType Index, const IndexType Count, const ProjectionType& InProj, const PredicateType& Predicate)
 	{
 // Workaround for a codegen bug that was discovered related to member function pointer projections that are virtual
 #if defined(_MSC_FULL_VER) && _MSC_FULL_VER < 194134123
@@ -97,7 +97,7 @@ namespace AlgoImpl
 	/**
 	 * Fixes a possible violation of order property between node at NodeIndex and a parent.
 	 *
-	 * @param	Heap		Pointer to the first element of a binary heap.
+	 * @param	Heap		Random access iterator to the first element of a binary heap.
 	 * @param	RootIndex	How far to go up?
 	 * @param	NodeIndex	Node index.
 	 * @param	InProj		The projection to apply to the elements.
@@ -105,8 +105,8 @@ namespace AlgoImpl
 	 *
 	 * @return	The new index of the node that was at NodeIndex
 	 */
-	template <class RangeValueType, typename IndexType, typename ProjectionType, class PredicateType>
-	inline IndexType HeapSiftUp(RangeValueType* Heap, IndexType RootIndex, IndexType NodeIndex, const ProjectionType& InProj, const PredicateType& Predicate)
+	template <class IteratorType, typename IndexType, typename ProjectionType, class PredicateType>
+	inline IndexType HeapSiftUp(IteratorType Heap, IndexType RootIndex, IndexType NodeIndex, const ProjectionType& InProj, const PredicateType& Predicate)
 	{
 // Workaround for a codegen bug that was discovered related to member function pointer projections that are virtual
 #if defined(_MSC_FULL_VER) && _MSC_FULL_VER < 194134123
@@ -138,20 +138,20 @@ namespace AlgoImpl
 	 * Builds an implicit min-heap from a range of elements.
 	 * This is the internal function used by Heapify overrides.
 	 *
-	 * @param	First		pointer to the first element to heapify
-	 * @param	Num			the number of items to heapify
+	 * @param	First		Random access iterator to the first element to heapify.
+	 * @param	Num			The number of items to heapify.
 	 * @param	Proj		The projection to apply to the elements.
 	 * @param	Predicate	A binary predicate object used to specify if one element should precede another.
 	 */
-	template <typename RangeValueType, typename IndexType, typename ProjectionType, typename PredicateType>
-	inline void HeapifyInternal(RangeValueType* First, IndexType Num, ProjectionType Proj, PredicateType Predicate)
+	template <typename IteratorType, typename IndexType, typename ProjectionType, typename PredicateType>
+	inline void HeapifyInternal(IteratorType First, IndexType Num, ProjectionType Proj, PredicateType Predicate)
 	{
 		if constexpr (std::is_signed_v<IndexType>)
 		{
 			checkf(Num >= 0, TEXT("Algo::HeapifyInternal called with negative count"));
 		}
 
-		if (Num == 0)
+		if (Num < 2)
 		{
 			return;
 		}
@@ -172,20 +172,20 @@ namespace AlgoImpl
 	 * Performs heap sort on the elements.
 	 * This is the internal sorting function used by HeapSort overrides.
 	 *
-	 * @param	First		pointer to the first element to sort
-	 * @param	Num			the number of elements to sort
+	 * @param	First		Random access iterator to the first element to sort.
+	 * @param	Num			The number of elements to sort.
 	 * @param	Proj		The projection to apply to the elements.
 	 * @param	Predicate	A binary predicate object used to specify if one element should precede another.
 	 */
-	template <typename RangeValueType, typename IndexType, typename ProjectionType, class PredicateType>
-	void HeapSortInternal(RangeValueType* First, IndexType Num, ProjectionType Proj, PredicateType Predicate)
+	template <typename IteratorType, typename IndexType, typename ProjectionType, class PredicateType>
+	void HeapSortInternal(IteratorType First, IndexType Num, ProjectionType Proj, PredicateType Predicate)
 	{
 		if constexpr (std::is_signed_v<IndexType>)
 		{
 			checkf(Num >= 0, TEXT("Algo::HeapSortInternal called with negative count"));
 		}
 
-		if (Num == 0)
+		if (Num < 2)
 		{
 			return;
 		}
diff --git a/Engine/Source/Runtime/Core/Public/Algo/IntroSort.h b/Engine/Source/Runtime/Core/Public/Algo/IntroSort.h
--- a/Engine/Source/Runtime/Core/Public/Algo/IntroSort.h
+++ b/Engine/Source/Runtime/Core/Public/Algo/IntroSort.h
@@ -10,6 +10,8 @@
 #include "Templates/Less.h"
 #include "Templates/UnrealTemplate.h" // For GetData, GetNum
 
+#include <iterator>
+
 
 namespace AlgoImpl
 {
@@ -18,18 +20,18 @@ namespace AlgoImpl
 	 * The sort is unstable, meaning that the ordering of equal items is not necessarily preserved.
 	 * This is the internal sorting function used by IntroSort overrides.
 	 *
-	 * @param First			pointer to the first element to sort
-	 * @param Num			the number of items to sort
+	 * @param First			Random access iterator to the first element to sort.
+	 * @param Num			The number of items to sort.
 	 * @param Proj			The projection to sort by when applied to the element.
 	 * @param Predicate		predicate class
 	 */
-	template <typename T, typename IndexType, typename ProjectionType, typename PredicateType> 
-	void IntroSortInternal(T* First, IndexType Num, ProjectionType Proj, PredicateType Predicate)
+	template <typename IteratorType, typename IndexType, typename ProjectionType, typename PredicateType>
+	void IntroSortInternal(IteratorType First, IndexType Num, ProjectionType Proj, PredicateType Predicate)
 	{
 		struct FStack
 		{
-			T* Min;
-			T* Max;
+			IteratorType Min;
+			IteratorType Max;
 			uint32 MaxDepth;
 		};
 
@@ -38,7 +40,9 @@ namespace AlgoImpl
 			return;
 		}
 
-		FStack RecursionStack[32]={{First, First+Num-1, (uint32)(FMath::Loge((float)Num) * 2.f)}}, Current, Inner;
+		FStack RecursionStack[32] = { {First, First + (Num - 1), (uint32)(FMath::Loge((float)Num) * 2.f)} };
+		FStack Current;
+		FStack Inner;
 		for( FStack* StackTop=RecursionStack; StackTop>=RecursionStack; --StackTop ) //-V625
 		{
 			Current = *StackTop;
@@ -55,11 +59,12 @@ namespace AlgoImpl
 
 			if( Count <= 8 )
 			{
-				// Use simple bubble-sort.
+				// Use simple selection sort.
 				while( Current.Max > Current.Min )
 				{
-					T *Max, *Item;
-					for( Max=Current.Min, Item=Current.Min+1; Item<=Current.Max; Item++ )
+					IteratorType Max;
+					IteratorType Item;
+					for( Max=Current.Min, Item=Current.Min+1; Item<=Current.Max; ++Item )
 					{
 // Workaround for a codegen bug that was discovered related to member function pointer projections that are virtual
 #if defined(_MSC_FULL_VER) && _MSC_FULL_VER < 194134123
@@ -148,7 +153,7 @@ namespace Algo
 	template <typename RangeType>
 	UE_REWRITE void IntroSort(RangeType&& Range)
 	{
-		AlgoImpl::IntroSortInternal(GetData(Range), GetNum(Range), FIdentityFunctor(), TLess<>());
+		AlgoImpl::IntroSortInternal(std::begin(Range), std::end(Range) - std::begin(Range), FIdentityFunctor(), TLess<>());
 	}
 
 	/**
@@ -160,7 +165,7 @@ namespace Algo
 	template <typename RangeType, typename PredicateType>
 	UE_REWRITE void IntroSort(RangeType&& Range, PredicateType Predicate)
 	{
-		AlgoImpl::IntroSortInternal(GetData(Range), GetNum(Range), FIdentityFunctor(), MoveTemp(Predicate));
+		AlgoImpl::IntroSortInternal(std::begin(Range), std::end(Range) - std::begin(Range), FIdentityFunctor(), MoveTemp(Predicate));
 	}
 
 	/**
@@ -174,9 +179,9 @@ namespace Algo
 	{
 // Workaround for a codegen bug that was discovered related to member function pointer projections that are virtual
 #if defined(_MSC_FULL_VER) && _MSC_FULL_VER < 194134123
-		AlgoImpl::IntroSortInternal(GetData(Range), GetNum(Range), MoveTemp(Proj), TLess<>());
+		AlgoImpl::IntroSortInternal(std::begin(Range), std::end(Range) - std::begin(Range), MoveTemp(Proj), TLess<>());
 #else
-		AlgoImpl::IntroSortInternal(GetData(Range), GetNum(Range), Projection(MoveTemp(Proj)), TLess<>());
+		AlgoImpl::IntroSortInternal(std::begin(Range), std::end(Range) - std::begin(Range), Projection(MoveTemp(Proj)), TLess<>());
 #endif
 	}
 
@@ -192,9 +197,9 @@ namespace Algo
 	{
 // Workaround for a codegen bug that was discovered related to member function pointer projections that are virtual
 #if defined(_MSC_FULL_VER) && _MSC_FULL_VER < 194134123
-		AlgoImpl::IntroSortInternal(GetData(Range), GetNum(Range), MoveTemp(Proj), MoveTemp(Predicate));
+		AlgoImpl::IntroSortInternal(std::begin(Range), std::end(Range) - std::begin(Range), MoveTemp(Proj), MoveTemp(Predicate));
 #else
-		AlgoImpl::IntroSortInternal(GetData(Range), GetNum(Range), Projection(MoveTemp(Proj)), Projection(MoveTemp(Predicate)));
+		AlgoImpl::IntroSortInternal(std::begin(Range), std::end(Range) - std::begin(Range), Projection(MoveTemp(Proj)), Projection(MoveTemp(Predicate)));
 #endif
 	}
 }
diff --git a/Engine/Source/Runtime/Core/Public/Containers/Array.h b/Engine/Source/Runtime/Core/Public/Containers/Array.h
--- a/Engine/Source/Runtime/Core/Public/Containers/Array.h
+++ b/Engine/Source/Runtime/Core/Public/Containers/Array.h
@@ -6,6 +6,7 @@
 #include "Misc/AssertionMacros.h"
 #include "Misc/IntrusiveUnsetOptionalState.h"
 #include "Misc/ReverseIterate.h"
+#include "Misc/UEOps.h"
 #include "HAL/UnrealMemory.h"
 #include "Templates/UnrealTypeTraits.h"
 #include "Templates/UnrealTemplate.h"
@@ -206,12 +207,12 @@ UE_NODEBUG [[nodiscard]] UE_FORCEINLINE_HINT TIndexedContainerIterator<Container
 	{
 		// This iterator type only supports the minimal functionality needed to support
 		// C++ ranged-for syntax.  For example, it does not provide post-increment ++ or any operation you wouldn't expect from a raw pointer.
-		//
-		// We do add an operator-- to help FString implementation
+
+		TCheckedPointerIterator() = default;
 
 		UE_NODEBUG [[nodiscard]] explicit TCheckedPointerIterator(const SizeType& InNum, ElementType* InPtr)
 			: Ptr       (InPtr)
-			, CurrentNum(InNum)
+			, CurrentNum(&InNum)
 			, InitialNum(InNum)
 		{
 		}
@@ -253,6 +254,13 @@ UE_NODEBUG [[nodiscard]] UE_FORCEINLINE_HINT TIndexedContainerIterator<Container
 			return *this;
 		}
 
+		UE_NODEBUG UE_FORCEINLINE_HINT TCheckedPointerIterator operator++(int)
+		{
+			TCheckedPointerIterator Result = *this;
+			++*this;
+			return Result;
+		}
+
 		UE_NODEBUG UE_FORCEINLINE_HINT TCheckedPointerIterator& operator--()
 		{
 			if constexpr (bReverse)
@@ -266,27 +274,117 @@ UE_NODEBUG [[nodiscard]] UE_FORCEINLINE_HINT TIndexedContainerIterator<Container
 			return *this;
 		}
 
-		UE_NODEBUG [[nodiscard]] UE_FORCEINLINE_HINT bool operator!=(const TCheckedPointerIterator& Rhs) const
+		UE_NODEBUG UE_FORCEINLINE_HINT TCheckedPointerIterator operator--(int)
 		{
-			// We only need to do the check in this operator, because no other operator will be
-			// called until after this one returns.
-			//
-			// Also, we should only need to check one side of this comparison - if the other iterator isn't
-			// even from the same array then the compiler has generated bad code.
-			ensureMsgf(CurrentNum == InitialNum, TEXT("Array has changed during ranged-for iteration!"));
-			return Ptr != Rhs.Ptr;
+			TCheckedPointerIterator Result = *this;
+			--*this;
+			return Result;
+		}
+
+		UE_NODEBUG [[nodiscard]] UE_FORCEINLINE_HINT TCheckedPointerIterator operator+(ptrdiff_t N) const
+		{
+			TCheckedPointerIterator Result = *this;
+			if constexpr (bReverse)
+			{
+				Result.Ptr -= N;
+			}
+			else
+			{
+				Result.Ptr += N;
+			}
+			return Result;
+		}
+
+		UE_NODEBUG UE_FORCEINLINE_HINT TCheckedPointerIterator& operator+=(ptrdiff_t N)
+		{
+			if constexpr (bReverse)
+			{
+				this->Ptr -= N;
+			}
+			else
+			{
+				this->Ptr += N;
+			}
+			return *this;
+		}
+
+		UE_NODEBUG [[nodiscard]] UE_FORCEINLINE_HINT TCheckedPointerIterator operator-(ptrdiff_t N) const
+		{
+			TCheckedPointerIterator Result = *this;
+			if constexpr (bReverse)
+			{
+				Result.Ptr += N;
+			}
+			else
+			{
+				Result.Ptr -= N;
+			}
+			return Result;
 		}
 
-		UE_NODEBUG [[nodiscard]] FORCEINLINE bool operator==(const TCheckedPointerIterator& Rhs) const
+		UE_NODEBUG UE_FORCEINLINE_HINT TCheckedPointerIterator operator-=(ptrdiff_t N)
 		{
-			return !(*this != Rhs);
+			if constexpr (bReverse)
+			{
+				this->Ptr += N;
+			}
+			else
+			{
+				this->Ptr -= N;
+			}
+			return *this;
+		}
+
+		UE_NODEBUG [[nodiscard]] UE_FORCEINLINE_HINT ptrdiff_t operator-(const TCheckedPointerIterator& Rhs) const
+		{
+			if constexpr (bReverse)
+			{
+				return Rhs.Ptr - Ptr;
+			}
+			else
+			{
+				return Ptr - Rhs.Ptr;
+			}
 		}
 
+		UE_NODEBUG [[nodiscard]] UE_FORCEINLINE_HINT ElementType& operator[](ptrdiff_t N) const
+		{
+			if constexpr (bReverse)
+			{
+				return *(Ptr - (N + 1));
+			}
+			else
+			{
+				return Ptr[N];
+			}
+		}
+
+		UE_NODEBUG [[nodiscard]] UE_FORCEINLINE_HINT bool UEOpEquals(const TCheckedPointerIterator& Rhs) const
+		{
+			// We only need to do the check in comparison operators, because no other operator will be
+			// called until after they return.
+			//
+			// Also, we should only need to check one side of this comparison - comparing iterators from
+			// different containers is undefined.
+			ensureMsgf(!CurrentNum || *CurrentNum == InitialNum, TEXT("Array has changed during ranged-for iteration!"));
+			return Ptr == Rhs.Ptr;
+		}
+
+		UE_NODEBUG  [[nodiscard]] UE_FORCEINLINE_HINT bool UEOpLessThan(const TCheckedPointerIterator& Rhs) const
+		{
+			// We only need to do the check in comparison operators, because no other operator will be
+			// called until after they return.
+			//
+			// Also, we should only need to check one side of this comparison - comparing iterators from
+			// different containers is undefined.
+			ensureMsgf(!CurrentNum || *CurrentNum == InitialNum, TEXT("Array has changed during ranged-for iteration!"));
+			return Ptr < Rhs.Ptr;
+		}
 
 	private:
-		ElementType*    Ptr;
-		const SizeType& CurrentNum;
-		SizeType        InitialNum;
+		ElementType*    Ptr        = nullptr;
+		const SizeType* CurrentNum = nullptr;
+		SizeType        InitialNum = 0;
 	};
 #endif
 
diff --git a/Engine/Source/Runtime/Core/Public/Templates/Sorting.h b/Engine/Source/Runtime/Core/Public/Templates/Sorting.h
--- a/Engine/Source/Runtime/Core/Public/Templates/Sorting.h
+++ b/Engine/Source/Runtime/Core/Public/Templates/Sorting.h
@@ -52,8 +52,25 @@ struct TArrayRange
 	{
 	}
 
-	T* GetData() const { return Begin; }
-	int32 Num() const { return Size; }
+	T* GetData() const
+	{
+		return Begin;
+	}
+
+	int32 Num() const
+	{
+		return Size;
+	}
+
+	T* begin() const
+	{
+		return this->Begin;
+	}
+
+	T* end() const
+	{
+		return this->Begin + this->Size;
+	}
 
 private:
 	T* Begin;
