Workaround for compiler bug during static analysis.

diff --git a/Engine/Source/Runtime/Core/Public/Containers/Array.h b/Engine/Source/Runtime/Core/Public/Containers/Array.h
--- a/Engine/Source/Runtime/Core/Public/Containers/Array.h
+++ b/Engine/Source/Runtime/Core/Public/Containers/Array.h
@@ -351,7 +351,7 @@ UE_NODEBUG [[nodiscard]] UE_FORCEINLINE_HINT TIndexedContainerIterator<Container
 		{
 			if constexpr (bReverse)
 			{
-				return *(Ptr - (N + 1));
+				return *(Ptr - N);
 			}
 			else
 			{
@@ -366,7 +366,10 @@ UE_NODEBUG [[nodiscard]] UE_FORCEINLINE_HINT TIndexedContainerIterator<Container
 			//
 			// Also, we should only need to check one side of this comparison - comparing iterators from
 			// different containers is undefined.
-			ensureMsgf(!CurrentNum || *CurrentNum == InitialNum, TEXT("Array has changed during ranged-for iteration!"));
+
+			// This should use *CurrentNum rather than CurrentNum[0], but it works around the following compiler issue:
+			// https://developercommunity.visualstudio.com/t/False-C6269-during-static-analysis/10989722
+			ensureMsgf(!CurrentNum || CurrentNum[0] == InitialNum, TEXT("Array has changed during ranged-for iteration!"));
 			return Ptr == Rhs.Ptr;
 		}
 
@@ -377,7 +380,10 @@ UE_NODEBUG [[nodiscard]] UE_FORCEINLINE_HINT TIndexedContainerIterator<Container
 			//
 			// Also, we should only need to check one side of this comparison - comparing iterators from
 			// different containers is undefined.
-			ensureMsgf(!CurrentNum || *CurrentNum == InitialNum, TEXT("Array has changed during ranged-for iteration!"));
+
+			// This should use *CurrentNum rather than CurrentNum[0], but it works around the following compiler issue:
+			// https://developercommunity.visualstudio.com/t/False-C6269-during-static-analysis/10989722
+			ensureMsgf(!CurrentNum || CurrentNum[0] == InitialNum, TEXT("Array has changed during ranged-for iteration!"));
 			return Ptr < Rhs.Ptr;
 		}
 
