Reverted accidental change to TCheckedPointerIterator::operator[].

diff --git a/Engine/Source/Runtime/Core/Public/Containers/Array.h b/Engine/Source/Runtime/Core/Public/Containers/Array.h
--- a/Engine/Source/Runtime/Core/Public/Containers/Array.h
+++ b/Engine/Source/Runtime/Core/Public/Containers/Array.h
@@ -351,7 +351,7 @@ UE_NODEBUG [[nodiscard]] UE_FORCEINLINE_HINT TIndexedContainerIterator<Container
 		{
 			if constexpr (bReverse)
 			{
-				return *(Ptr - N);
+				return *(Ptr - (N + 1));
 			}
 			else
 			{
