- Minor BugHawk fix

diff --git a/Engine/Source/Runtime/Core/Private/Misc/ConfigCacheIni.cpp b/Engine/Source/Runtime/Core/Private/Misc/ConfigCacheIni.cpp
--- a/Engine/Source/Runtime/Core/Private/Misc/ConfigCacheIni.cpp
+++ b/Engine/Source/Runtime/Core/Private/Misc/ConfigCacheIni.cpp
@@ -2610,6 +2610,7 @@ bool FConfigFile::WriteInternal(const FString& Filename, bool bDoRemoteWrite, TM
 	// don't write anything out if it didn't actually change
 	if (OriginalContents != nullptr && Text == *OriginalContents)
 	{
+		Dirty = false;
 		return true;
 	}
 
