[LLM]
* Added support for -LLMTraceTagFilter=tag1,tag2,... command line parameter. When used, only the list of tags specified would be traced (for -trace=memtag or -trace=memory).
* Fixed size of ETagReferenceSource enum to be uint8. This improves the layout of FTagData class.

diff --git a/Engine/Source/Runtime/Core/Private/HAL/LowLevelMemTracker.cpp b/Engine/Source/Runtime/Core/Private/HAL/LowLevelMemTracker.cpp
--- a/Engine/Source/Runtime/Core/Private/HAL/LowLevelMemTracker.cpp
+++ b/Engine/Source/Runtime/Core/Private/HAL/LowLevelMemTracker.cpp
@@ -931,8 +931,8 @@ uint8 LlmGetActiveTag()
 UE_TRACE_CHANNEL(MemTagChannel, "Memory overview")
 
 UE_TRACE_EVENT_BEGIN(LLM, TagsSpec, NoSync|Important)
-	UE_TRACE_EVENT_FIELD(const void*, TagId)
-	UE_TRACE_EVENT_FIELD(const void*, ParentId)
+	UE_TRACE_EVENT_FIELD(uint64, TagId)
+	UE_TRACE_EVENT_FIELD(uint64, ParentId)
 	UE_TRACE_EVENT_FIELD(uint8, TagSetId)
 	UE_TRACE_EVENT_FIELD(UE::Trace::AnsiString, Name)
 UE_TRACE_EVENT_END()
@@ -950,7 +950,7 @@ UE_TRACE_EVENT_END()
 UE_TRACE_EVENT_BEGIN(LLM, TagValue)
 	UE_TRACE_EVENT_FIELD(uint8, TrackerId)
 	UE_TRACE_EVENT_FIELD(uint64, Cycle)
-	UE_TRACE_EVENT_FIELD(const void*[], Tags)
+	UE_TRACE_EVENT_FIELD(uint64[], Tags)
 	UE_TRACE_EVENT_FIELD(int64[], Values)
 UE_TRACE_EVENT_END()
 
@@ -1835,6 +1835,7 @@ FLowLevelMemTracker::FLowLevelMemTracker()
 	: TagDatas(nullptr)
 	, TagDataNameMap(nullptr)
 	, TagDataEnumMap(nullptr)
+	, TraceTagFilter(nullptr)
 	, ProgramSize(0)
 	, MemoryUsageCurrentOverhead(0)
 	, MemoryUsagePlatformTotalUntracked(0)
@@ -1953,6 +1954,12 @@ void FLowLevelMemTracker::Clear()
 		Trackers[TrackerIndex] = nullptr;
 	}
 
+	if (TraceTagFilter)
+	{
+		Allocator.Delete(TraceTagFilter);
+		TraceTagFilter = nullptr;
+	}
+
 	ClearTagDatas();
 	Allocator.Clear();
 	bFullyInitialised = false;
@@ -2422,6 +2429,35 @@ void FLowLevelMemTracker::ProcessCommandLine(const TCHAR* CmdLine)
 		}
 	}
 
+	if (FParse::Value(CmdLine, TEXT("LLMTraceTagFilter="), SetList, false /* bShouldStopOnSeparator */))
+	{
+		LLMCheck(TraceTagFilter == nullptr); // ProcessCommandLine() is expected to be called only once
+		TArray<FString> TagNames;
+		SetList.ParseIntoArray(TagNames, TEXT(","), true);
+		if (!TagNames.IsEmpty())
+		{
+			if (!TraceTagFilter)
+			{
+				TraceTagFilter = Allocator.New<UE::LLMPrivate::FTagFilter>();
+			}
+			for (FString& TagName : TagNames)
+			{
+				TraceTagFilter->Add(FName(TagName));
+			}
+			// Update bIsTraceReportable flag for already created tags.
+			// TagData is supposed to be immutable once constructed, but we use bIsTraceReportable flag
+			// as a mutable member to quickly filter tags for the trace purposes.
+			UE::TSharedLock Lock(TagDataMutex);
+			for (UE::LLMPrivate::FTagData* TagData : (*TagDatas))
+			{
+				if (!TraceTagFilter->Contains(TagData->GetName()))
+				{
+					TagData->SetIsTraceReportable(false);
+				}
+			}
+		}
+	}
+
 	// Commandline overrides for console variables
 	int TrackPeaks = 0;
 	if (FParse::Value(CmdLine, TEXT("LLMTrackPeaks="), TrackPeaks))
@@ -3149,6 +3185,11 @@ void FLowLevelMemTracker::FinishConstruct(UE::LLMPrivate::FTagData* TagData,
 		FinishConstruct(ParentData, ReferenceSource);
 	}
 
+	if (TraceTagFilter && !TraceTagFilter->Contains(TagData->GetName()))
+	{
+		TagData->SetIsTraceReportable(false);
+	}
+
 	TagData->SetFinishConstructed();
 
 	// Broadcast the tag creation, except for generic tags which are constructed before any subscriber could
@@ -4511,7 +4552,7 @@ FTagData::FTagData(FName InName, ELLMTagSet InTagSet, FName InDisplayName, FName
 	: Name(InName), DisplayName(InDisplayName), ParentName(InParentName), StatName(InStatName)
 	, SummaryStatName(InSummaryStatName), EnumTag(InEnumTag), ReferenceSource(InReferenceSource)
 	, TagSet(InTagSet), bIsFinishConstructed(false), bParentIsName(true), bHasEnumTag(bInHasEnumTag)
-	, bIsReportable(true)
+	, bIsReportable(true), bIsTraceReportable(true)
 {
 }
 
@@ -4685,6 +4726,11 @@ bool FTagData::IsStatsReportable() const
 	return bIsReportable;
 }
 
+bool FTagData::IsTraceReportable() const
+{
+	return bIsTraceReportable && bIsReportable;
+}
+
 void FTagData::SetParent(const FTagData* InParent)
 {
 	if (bParentIsName)
@@ -4736,6 +4782,11 @@ void FTagData::SetIsReportable(bool bInIsReportable)
 	bIsReportable = bInIsReportable;
 }
 
+void FTagData::SetIsTraceReportable(bool bInIsTraceReportable)
+{
+	bIsTraceReportable = bInIsTraceReportable;
+}
+
 bool FTagData::IsUsedAsDisplayParent() const
 {
 	// All Tags but one are UsedAsDisplayParent - their name is prepended during GetDisplayPath.
@@ -6425,6 +6476,8 @@ void FLLMTraceWriter::Publish(FLowLevelMemTracker& LLMRef, const FTrackerTagSize
 	{
 		bTrackerSpecSent = true;
 		FAnsiStringView TrackerNames[] = { "Platform", "Default" };
+		static_assert((int)ELLMTracker::Platform == 0, "");
+		static_assert((int)ELLMTracker::Default == 1, "");
 		static_assert(UE_ARRAY_COUNT(TrackerNames) == (int)ELLMTracker::Max, "");
 		FAnsiStringView TrackerName = TrackerNames[(uint8)Tracker];
 		UE_TRACE_LOG(LLM, TrackerSpec, MemTagChannel, TrackerName.Len() * sizeof(ANSICHAR))
@@ -6457,37 +6510,54 @@ void FLLMTraceWriter::Publish(FLowLevelMemTracker& LLMRef, const FTrackerTagSize
 		}
 	}
 
-	SendTagDeclaration(OverrideTrackedTotalTagData);
-	SendTagDeclaration(OverrideUntaggedTagData);
+	if (OverrideTrackedTotalTagData && !OverrideTrackedTotalTagData->IsTraceReportable())
+	{
+		OverrideTrackedTotalTagData = nullptr;
+	}
+	if (OverrideUntaggedTagData && !OverrideUntaggedTagData->IsTraceReportable())
+	{
+		OverrideUntaggedTagData = nullptr;
+	}
+
+	if (OverrideTrackedTotalTagData)
+	{
+		SendTagDeclaration(OverrideTrackedTotalTagData);
+	}
+	if (OverrideUntaggedTagData)
+	{
+		SendTagDeclaration(OverrideUntaggedTagData);
+	}
 	for (const TPair<const FTagData*, FTrackerTagSizeData>& It : TagSizes)
 	{
 		const FTagData* TagData = It.Key;
-		if (!TagData->IsStatsReportable() || !LLMRef.IsTagSetActive(TagData->GetTagSet()))
+		if (!TagData->IsTraceReportable() || !LLMRef.IsTagSetActive(TagData->GetTagSet()))
 		{
 			continue;
 		}
-		if (OverrideUntaggedTagData != nullptr && TagData->GetTagSet() == ELLMTagSet::None && TagData->GetName() == TagName_Untagged)
+		if (OverrideUntaggedTagData &&
+			TagData->GetTagSet() == ELLMTagSet::None &&
+			TagData->GetName() == TagName_Untagged)
 		{
 			continue; // Handled by OverrideUntaggedTagData.
 		}
 		SendTagDeclaration(TagData);
 	}
 
-	TArray<const void*, FDefaultLLMAllocator> TagIds;
+	TArray<uint64, FDefaultLLMAllocator> TagIds;
 	TArray<int64, FDefaultLLMAllocator> TagValues;
 	TagIds.Reserve(TagSizes.Num() + 2);
 	TagValues.Reserve(TagSizes.Num() + 2);
 	auto AddValue = [&TagIds, &TagValues](const FTagData* TagData, int64 Value)
 	{
-		if (!TagData)
-		{
-			return;
-		}
-		TagIds.Add(GetTagId(TagData));
+		const uint64 TagId = uint64(GetTagId(TagData));
+		TagIds.Add(TagId);
 		TagValues.Add(Value);
 	};
 
-	AddValue(OverrideTrackedTotalTagData, TrackedTotal);
+	if (OverrideTrackedTotalTagData)
+	{
+		AddValue(OverrideTrackedTotalTagData, TrackedTotal);
+	}
 	if (OverrideUntaggedTagData)
 	{
 		const FTagData* TagData = LLMRef.FindTagData(TagName_Untagged, ELLMTagSet::None);
@@ -6505,33 +6575,32 @@ void FLLMTraceWriter::Publish(FLowLevelMemTracker& LLMRef, const FTrackerTagSize
 	for (const TPair<const FTagData*, FTrackerTagSizeData>& It : TagSizes)
 	{
 		const FTagData* TagData = It.Key;
-		if (!TagData->IsStatsReportable() || !LLMRef.IsTagSetActive(TagData->GetTagSet()))
+		if (!TagData->IsTraceReportable() || !LLMRef.IsTagSetActive(TagData->GetTagSet()))
 		{
 			continue;
 		}
-		if (OverrideUntaggedTagData && TagData->GetTagSet() == ELLMTagSet::None && TagData->GetName() == TagName_Untagged)
+		if (OverrideUntaggedTagData &&
+			TagData->GetTagSet() == ELLMTagSet::None &&
+			TagData->GetName() == TagName_Untagged)
 		{
 			continue; // Handled by OverrideUntaggedTagData.
 		}
 		AddValue(TagData, It.Value.GetSize(SizeParams));
 	}
 
-	int TagCount = TagIds.Num();
-	LLMCheck(TagCount == TagValues.Num());
+	LLMCheck(TagIds.Num() == TagValues.Num());
 	const uint64 Cycle = FPlatformTime::Cycles64();
 	UE_TRACE_LOG(LLM, TagValue, MemTagChannel)
 		<< TagValue.TrackerId((uint8)Tracker)
 		<< TagValue.Cycle(Cycle)
-		<< TagValue.Tags(TagIds.GetData(), TagCount)
-		<< TagValue.Values(TagValues.GetData(), TagCount);
+		<< TagValue.Tags(TagIds.GetData(), TagIds.Num())
+		<< TagValue.Values(TagValues.GetData(), TagValues.Num());
 }
 
 void FLLMTraceWriter::SendTagDeclaration(const FTagData* TagData)
 {
-	if (!TagData)
-	{
-		return;
-	}
+	LLMCheck(TagData);
+
 	bool bAlreadyInSet;
 	DeclaredTags.Add(TagData, &bAlreadyInSet);
 	if (bAlreadyInSet)
@@ -6540,13 +6609,16 @@ void FLLMTraceWriter::SendTagDeclaration(const FTagData* TagData)
 	}
 
 	const FTagData* Parent = TagData->GetParent();
-	SendTagDeclaration(Parent);
+	if (Parent)
+	{
+		SendTagDeclaration(Parent);
+	}
 
 	TStringBuilder<1024> NameBuffer;
 	TagData->AppendDisplayPath(NameBuffer);
 	UE_TRACE_LOG(LLM, TagsSpec, MemTagChannel, NameBuffer.Len() * sizeof(ANSICHAR))
-		<< TagsSpec.TagId(GetTagId(TagData))
-		<< TagsSpec.ParentId(GetTagId(Parent))
+		<< TagsSpec.TagId(uint64(GetTagId(TagData)))
+		<< TagsSpec.ParentId(uint64(GetTagId(Parent)))
 		<< TagsSpec.TagSetId(static_cast<uint8>(TagData->GetTagSet()))
 		<< TagsSpec.Name(*NameBuffer, NameBuffer.Len());
 };
diff --git a/Engine/Source/Runtime/Core/Private/HAL/LowLevelMemTrackerPrivate.h b/Engine/Source/Runtime/Core/Private/HAL/LowLevelMemTrackerPrivate.h
--- a/Engine/Source/Runtime/Core/Private/HAL/LowLevelMemTrackerPrivate.h
+++ b/Engine/Source/Runtime/Core/Private/HAL/LowLevelMemTrackerPrivate.h
@@ -111,11 +111,12 @@ namespace LLMPrivate
 		int32 GetIndex() const;
 		bool IsReportable() const;
 		bool IsStatsReportable() const;
+		bool IsTraceReportable() const;
 
 		void SetParent(const FTagData* InParent);
 		void SetIndex(int32 InIndex);
 		void SetIsReportable(bool bInReportable);
-		void SetIsStatsReportable(bool bInReportable);
+		void SetIsTraceReportable(bool bInTraceReportable);
 		void SetFinishConstructed();
 
 		// These functions are normally invalid - these properties should be immutable - but are called for EnumTags during bootstrapping
@@ -146,6 +147,7 @@ namespace LLMPrivate
 		bool bParentIsName;
 		bool bHasEnumTag;
 		bool bIsReportable;
+		bool bIsTraceReportable;
 	};
 
 	// TagData container types that use FLLMAllocator
@@ -161,6 +163,10 @@ namespace LLMPrivate
 	{
 		using TArray<FTagData*, FDefaultLLMAllocator>::TArray;
 	};
+
+	class FTagFilter : public TLLMSet<FName>
+	{
+	};
 }
 }
 
diff --git a/Engine/Source/Runtime/Core/Public/HAL/LowLevelMemTracker.h b/Engine/Source/Runtime/Core/Public/HAL/LowLevelMemTracker.h
--- a/Engine/Source/Runtime/Core/Public/HAL/LowLevelMemTracker.h
+++ b/Engine/Source/Runtime/Core/Public/HAL/LowLevelMemTracker.h
@@ -481,6 +481,7 @@ class FLLMTracker;
 class FTagData;
 class FTagDataArray;
 class FTagDataNameMap;
+class FTagFilter;
 
 namespace AllocatorPrivate
 {
@@ -551,7 +552,7 @@ private:
 	friend struct UE::LLMPrivate::AllocatorPrivate::FBin;
 };
 
-enum class ETagReferenceSource
+enum class ETagReferenceSource : uint8
 {
 	Scope,
 	Declare,
@@ -876,6 +877,8 @@ private:
 	/** Array to Map from ELLMTag to the TagData for that tag, used to handle LLM_SCOPE with ELLMTag */
 	UE::LLMPrivate::FTagData** TagDataEnumMap;
 
+	UE::LLMPrivate::FTagFilter* TraceTagFilter;
+
 	UE::LLMPrivate::FLLMTracker* Trackers[static_cast<int32>(ELLMTracker::Max)];
 
 	mutable UE::FPlatformSharedMutex TagDataMutex;
