[AsyncFileHandle]
- Fix TSAN race by using proper atomics and barriers

diff --git a/Engine/Source/Runtime/Core/Public/Async/AsyncFileHandle.h b/Engine/Source/Runtime/Core/Public/Async/AsyncFileHandle.h
--- a/Engine/Source/Runtime/Core/Public/Async/AsyncFileHandle.h
+++ b/Engine/Source/Runtime/Core/Public/Async/AsyncFileHandle.h
@@ -36,20 +36,16 @@ protected:
 		uint8* Memory;
 	};
 	FAsyncFileCallBack Callback;
-	TSAN_ATOMIC(bool) bDataIsReady;
-	TSAN_ATOMIC(bool) bCompleteAndCallbackCalled;
-	TSAN_ATOMIC(bool) bCompleteSync;
-	TSAN_ATOMIC(bool) bCanceled;
+	std::atomic<bool> bDataIsReady = false;
+	std::atomic<bool> bCompleteAndCallbackCalled = false;
+	std::atomic<bool> bCompleteSync = false;
+	std::atomic<bool> bCanceled = false;
 	const bool bSizeRequest;
 	const bool bUserSuppliedMemory;
 public:
 
 	inline IAsyncReadRequest(FAsyncFileCallBack* InCallback, bool bInSizeRequest, uint8* UserSuppliedMemory)
-		: bDataIsReady(false)
-		, bCompleteAndCallbackCalled(false)
-		, bCompleteSync(false)
-		, bCanceled(false)
-		, bSizeRequest(bInSizeRequest)
+		: bSizeRequest(bInSizeRequest)
 		, bUserSuppliedMemory(!!UserSuppliedMemory)
 	{
 		if (InCallback)
@@ -73,7 +69,7 @@ public:
 	virtual ~IAsyncReadRequest()
 	{
 		// check(bCompleteAndCallbackCalled && (bSizeRequest || !Memory)); // must be complete, and if it was a read request, the memory should be gone
-		UE_CLOG(!(bCompleteAndCallbackCalled && (bSizeRequest || !Memory)),
+		UE_CLOG(!(bCompleteAndCallbackCalled.load(std::memory_order_relaxed) && (bSizeRequest || !Memory)),
 			LogHAL, Fatal, TEXT("IAsyncReadRequests must not be deleted until they are completed.")
 		);
 		DEC_DWORD_STAT(STAT_AsyncFileRequests);
@@ -85,7 +81,17 @@ public:
 	**/
 	UE_FORCEINLINE_HINT bool PollCompletion()
 	{
-		return bCompleteAndCallbackCalled;
+		// fast path to avoid a barrier until the value returns true
+		if (bCompleteAndCallbackCalled.load(std::memory_order_relaxed))
+		{
+			// Once we return true, we need a barrier to guarantee the memory
+			// writes that have happened before the release barrier can be seen.
+			// Because for all we know, PollCompletion() could be the only thing
+			// that is used to synchronize with the writing thread.
+			return bCompleteAndCallbackCalled.load(std::memory_order_acquire);
+		}
+
+		return false;
 	}
 
 	/**
@@ -119,11 +125,11 @@ public:
 	/** Cancel the request. This is a non-blocking async call and so does not ensure completion! **/
 	inline void Cancel()
 	{
-		if (!bCanceled)
+		if (!bCanceled.load(std::memory_order_relaxed))
 		{
-			bCanceled = true;
-			bDataIsReady = true;
-			FPlatformMisc::MemoryBarrier();
+			bCanceled.store(true, std::memory_order_release);
+			bDataIsReady.store(true, std::memory_order_release);
+
 			if (!PollCompletion())
 			{
 				return CancelImpl();
@@ -137,8 +143,8 @@ public:
 	**/
 	inline int64 GetSizeResults()
 	{
-		check(bDataIsReady && bSizeRequest);
-		return bCanceled ? -1 : Size;
+		check(bDataIsReady.load(std::memory_order_relaxed) && bSizeRequest);
+		return bCanceled.load(std::memory_order_relaxed) ? -1 : Size;
 	}
 
 	/**
@@ -147,9 +153,9 @@ public:
 	**/
 	inline uint8* GetReadResults()
 	{
-		check(bDataIsReady && !bSizeRequest);
+		check(bDataIsReady.load(std::memory_order_relaxed) && !bSizeRequest);
 		uint8* Result = Memory;
-		if (bCanceled && Result)
+		if (bCanceled.load(std::memory_order_relaxed) && Result)
 		{
 			Result = nullptr;
 		}
@@ -185,19 +191,16 @@ protected:
 
 	void SetDataComplete()
 	{
-		bDataIsReady = true;
-		FPlatformMisc::MemoryBarrier();
+		bDataIsReady.store(true, std::memory_order_release);
 		if (Callback)
 		{
-			Callback(bCanceled, this);
+			Callback(bCanceled.load(std::memory_order_relaxed), this);
 		}
-		FPlatformMisc::MemoryBarrier();
 	}
 
 	void SetAllComplete()
 	{
-		bCompleteAndCallbackCalled = true;
-		FPlatformMisc::MemoryBarrier();
+		bCompleteAndCallbackCalled.store(true, std::memory_order_release);
 	}
 
 	void SetComplete()
