Added UE::Zip, a range adapter which takes multiple ranges and iterates over them in parallel.  This also allows parallel sorting with Algo::Sort.
Added GetElementByType and GetElementByIndex callable objects that can be used as projections in Algo::SortBy to pick a particular range.
Added Swap overloads for TTuple to do element-wise swapping rather than whole-object swapping, so that tuples of references will swap the referred objects rather than try and fail to rebind the references.

diff --git a/Engine/Source/Runtime/Core/Private/Misc/ZipTests.cpp b/Engine/Source/Runtime/Core/Private/Misc/ZipTests.cpp
new file mode 100644
--- /dev/null
+++ b/Engine/Source/Runtime/Core/Private/Misc/ZipTests.cpp
@@ -0,0 +1,150 @@
+// Copyright Epic Games, Inc. All Rights Reserved.
+
+#include "CoreTypes.h"
+
+#include "Algo/Sort.h"
+#include "Containers/AnsiString.h"
+#include "Containers/Array.h"
+#include "Containers/UnrealString.h"
+#include "Misc/AutomationTest.h"
+#include "Misc/Zip.h"
+
+#include <type_traits>
+
+#if WITH_DEV_AUTOMATION_TESTS
+
+IMPLEMENT_SIMPLE_AUTOMATION_TEST(FUEZipTests, "System.Core.UEZip", EAutomationTestFlags_ApplicationContextMask | EAutomationTestFlags::SmokeFilter)
+bool FUEZipTests::RunTest(const FString& Parameters)
+{
+	const TArray<int32>   Arr1 = { 1, 2, 3, 4, 5 };
+	const TArray<FString> Arr2 = { "Orange", "Banana", "Apple", "Pear", "Cherry" };
+	const float           Arr3[] = { 3.1415f, 1.4142f, 2.7183f, 1.6180f, 0.1234f };
+
+	// Test basic zipped iteration
+	{
+		// Single
+		{
+			TArray<int32> Copy;
+			for (const auto& Single : UE::Zip(Arr1))
+			{
+				static_assert(std::is_same_v<decltype(Single), const TTuple<const int32&>&>);
+				Copy.Add(Single.Get<0>());
+			}
+			check(Copy == TArray<int32>({ 1, 2, 3, 4, 5 }));
+		}
+
+		// Double
+		{
+			TArray<int32>   Copy1;
+			TArray<FString> Copy2;
+			for (const auto& Pair : UE::Zip(Arr1, Arr2))
+			{
+				static_assert(std::is_same_v<decltype(Pair), const TTuple<const int32&, const FString&>&>);
+				Copy1.Add(Pair.Get<0>());
+				Copy2.Add(Pair.Get<1>());
+			}
+			check(Copy1 == TArray<int32>  ({ 1, 2, 3, 4, 5 }));
+			check(Copy2 == TArray<FString>({ "Orange", "Banana", "Apple", "Pear", "Cherry" }));
+		}
+
+		// Triple, with a different container type
+		{
+			TArray<int32>   Copy1;
+			TArray<FString> Copy2;
+			TArray<float>   Copy3;
+			for (const auto& Triple : UE::Zip(Arr1, Arr2, Arr3))
+			{
+				static_assert(std::is_same_v<decltype(Triple), const TTuple<const int32&, const FString&, const float&>&>);
+				Copy1.Add(Triple.Get<0>());
+				Copy2.Add(Triple.Get<1>());
+				Copy3.Add(Triple.Get<2>());
+			}
+			check(Copy1 == TArray<int32>  ({ 1, 2, 3, 4, 5 }));
+			check(Copy2 == TArray<FString>({ "Orange", "Banana", "Apple", "Pear", "Cherry" }));
+			check(Copy3 == TArray<float>  ({ 3.1415f, 1.4142f, 2.7183f, 1.6180f, 0.1234f }));
+		}
+
+		// Mutate container elements
+		{
+			TArray<int32>   Arr1Copy = Arr1;
+			TArray<FString> Arr2Copy = Arr2;
+			for (const auto& Pair : UE::Zip(Arr1Copy, Arr2Copy))
+			{
+				static_assert(std::is_same_v<decltype(Pair), const TTuple<int32&, FString&>&>);
+
+				// Scale int by 10
+				Pair.Get<0>() *= 10;
+
+				// Append a copy of the string separated by |
+				Pair.Get<1>() += FString::Printf(TEXT("|%s"), *Pair.Get<1>());
+			}
+			check(Arr1Copy == TArray<int32>  ({ 10, 20, 30, 40, 50 }));
+			check(Arr2Copy == TArray<FString>({ "Orange|Orange", "Banana|Banana", "Apple|Apple", "Pear|Pear", "Cherry|Cherry" }));
+		}
+
+		// Mix constness
+		{
+			TArray<int32>   Arr1Copy = Arr1;
+			TArray<FString> Arr2Copy = Arr2;
+			for (const auto& Pair : UE::Zip(Arr1Copy, AsConst(Arr2Copy)))
+			{
+				static_assert(std::is_same_v<decltype(Pair), const TTuple<int32&, const FString&>&>);
+			}
+		}
+
+		// Mix lengths, and test structured binding
+		{
+			TArray<int32> Short1 = { 2, 3, 5 };
+			TArray<float> Short3 = { 4.0f, 3.0f, 2.0f, 1.0f };
+
+			using TupleType = TTuple<int32, FString, float>;
+			using CopyType  = TArray<TupleType>;
+			CopyType Copy;
+
+			// Should only copy as many arguments as Short1 has
+			for (auto&& [Int, String, Float] : UE::Zip(Short1, Arr2, Short3))
+			{
+				Copy.Add({ Int, String, Float });
+			}
+			check(Copy == CopyType({ { 2, "Orange", 4.0f }, { 3, "Banana", 3.0f }, { 5, "Apple", 2.0f } }));
+		}
+
+		// Mix lvalues and rvalues
+		{
+			auto MakeArray = []()
+			{
+				return TArray<FAnsiString>{ "Abcde", "Fghi", "Klm", "No", "P" };
+			};
+
+			TArray<int32>       Copy1;
+			TArray<FAnsiString> Copy2;
+			for (const auto& Pair : UE::Zip(Arr1, MakeArray()))
+			{
+				static_assert(std::is_same_v<decltype(Pair), const TTuple<const int32&, const FAnsiString&>&>);
+				Copy1.Add(Pair.Key);
+				Copy2.Add(Pair.Value);
+			}
+			check(Copy1 == TArray<int32>({ 1, 2, 3, 4, 5 }));
+			check(Copy2 == TArray<FAnsiString>({ "Abcde", "Fghi", "Klm", "No", "P" }));
+		}
+	}
+
+	// Test zipped sorting
+	{
+		// Sort by the strings
+		TArray<int32>   Arr1Copy = Arr1;
+		TArray<FString> Arr2Copy = Arr2;
+		Algo::SortBy(UE::Zip(Arr1Copy, Arr2Copy), GetElementByIndex<1>);
+		check(Arr1Copy == TArray<int32>  ({ 3, 2, 5, 1, 4 }));
+		check(Arr2Copy == TArray<FString>({ "Apple", "Banana", "Cherry", "Orange", "Pear" }));
+
+		// Sort by the ints
+		Algo::SortBy(UE::Zip(Arr1Copy, Arr2Copy), GetElementByIndex<0>);
+		check(Arr1Copy == TArray<int32>({ 1, 2, 3, 4, 5 }));
+		check(Arr2Copy == TArray<FString>({ "Orange", "Banana", "Apple", "Pear", "Cherry" }));
+	}
+
+	return !HasAnyErrors();
+}
+
+#endif // WITH_DEV_AUTOMATION_TESTS
diff --git a/Engine/Source/Runtime/Core/Public/Misc/Zip.h b/Engine/Source/Runtime/Core/Public/Misc/Zip.h
new file mode 100644
--- /dev/null
+++ b/Engine/Source/Runtime/Core/Public/Misc/Zip.h
@@ -0,0 +1,198 @@
+// Copyright Epic Games, Inc. All Rights Reserved.
+
+#pragma once
+
+#include "Concepts/Integral.h"
+#include "Misc/UEOps.h"
+#include "Templates/Tuple.h"
+#include <iterator>
+
+namespace UE::Core::Private
+{
+	template <typename IteratorTupleType>
+	struct TIteratorTuple;
+
+	template <typename... IteratorTypes>
+	struct TIteratorTuple<TTuple<IteratorTypes...>>
+	{
+		using value_type = TTuple<std::decay_t<decltype(*std::declval<IteratorTypes>())>...>;
+
+		TTuple<IteratorTypes...> Iterators;
+
+		decltype(auto) operator*() const
+		{
+			return this->Iterators.ApplyAfter(
+				[](const auto&... Iters)
+				{
+					return TTuple<decltype(*Iters)...>(*Iters...);
+				}
+			);
+		}
+
+		TIteratorTuple& operator++()
+		{
+			VisitTupleElements([](auto& Iters){ ++Iters; }, this->Iterators);
+			return *this;
+		}
+
+		TIteratorTuple operator++(int)
+		{
+			TIteratorTuple Result = *this;
+			VisitTupleElements([](auto& Iters) { ++Iters; }, this->Iterators);
+			return Result;
+		}
+
+		TIteratorTuple& operator--()
+		{
+			VisitTupleElements([](auto& Iters) { --Iters; }, this->Iterators);
+			return *this;
+		}
+
+		TIteratorTuple operator--(int)
+		{
+			TIteratorTuple Result = *this;
+			VisitTupleElements([](auto& Iters) { --Iters; }, this->Iterators);
+			return Result;
+		}
+
+		template <UE::CIntegral T>
+		[[nodiscard]] TIteratorTuple operator+(T Num) const
+		{
+			TIteratorTuple Result = *this;
+			VisitTupleElements([Num](auto& Iterator) { Iterator += Num; }, Result.Iterators);
+			return Result;
+		}
+
+		template <UE::CIntegral T>
+		[[nodiscard]] TIteratorTuple operator-(T Num) const
+		{
+			TIteratorTuple Result = *this;
+			VisitTupleElements([Num](auto& Iterator) { Iterator -= Num; }, Result.Iterators);
+			return Result;
+		}
+
+		template <UE::CIntegral T>
+		TIteratorTuple& operator+=(T Num)
+		{
+			VisitTupleElements([Num](auto& Iterator) { Iterator += Num; }, this->Iterators);
+			return *this;
+		}
+
+		template <UE::CIntegral T>
+		TIteratorTuple& operator-=(T Num)
+		{
+			VisitTupleElements([Num](auto& Iterator) { Iterator -= Num; }, this->Iterators);
+			return *this;
+		}
+
+		[[nodiscard]] auto operator-(const TIteratorTuple& Rhs) const
+		{
+			return this->Iterators.template Get<0>() - Rhs.Iterators.template Get<0>();
+		}
+
+		template <UE::CIntegral T>
+		[[nodiscard]] decltype(auto) operator[](T Index) const
+		{
+			return *(*this + Index);
+		}
+
+		[[nodiscard]] bool UEOpEquals(const TIteratorTuple& Rhs) const
+		{
+			bool bResult = this->Iterators.template Get<0>() == Rhs.Iterators.template Get<0>();
+
+			// If the first elements compare unequal, all iterators should also compare unequal, or
+			// it means the first container is longer than one of the others.
+			checkfSlow(bResult || All(*this, Rhs, [](const auto& A, const auto& B) { return !(A == B); }), TEXT("UE::Zip - initial range is longer than other ranges"));
+
+			return bResult;
+		}
+
+		[[nodiscard]] bool UEOpLessThan(const TIteratorTuple& Rhs) const
+		{
+			bool bResult = this->Iterators.template Get<0>() < Rhs.Iterators.template Get<0>();
+
+			// If the first elements compare less-than, all iterators should also compare less-than, or
+			// it means the first container is longer than one of the others.
+			checkfSlow(!bResult || All(*this, Rhs, [](const auto& A, const auto& B) { return A < B; }), TEXT("UE::Zip - initial range is longer than other ranges"));
+
+			return bResult;
+		}
+
+	private:
+		template <typename Predicate>
+		static bool All(const TIteratorTuple& Lhs, const TIteratorTuple& Rhs, Predicate Pred)
+		{
+			bool bResult = true;
+
+			VisitTupleElements(
+				[&](const auto& LhsIterator, const auto& RhsIterator)
+				{
+					bResult = bResult && Pred(LhsIterator, RhsIterator);
+				},
+				Lhs.Iterators,
+				Rhs.Iterators
+			);
+
+			return bResult;
+		}
+	};
+
+	template <typename TupleType>
+	TIteratorTuple<std::decay_t<TupleType>> MakeIteratorTuple(TupleType&& Tuple)
+	{
+		return { (TupleType&&)Tuple };
+	}
+
+	template <typename TupleType>
+	struct TRangeTuple
+	{
+		TupleType Ranges;
+
+		auto begin() const
+		{
+			return MakeIteratorTuple(TransformTuple(this->Ranges, [](auto& Range){ return std::begin(Range); }));
+		}
+
+		auto end() const
+		{
+			return MakeIteratorTuple(TransformTuple(this->Ranges, [](auto& Range){ return std::end(Range); }));
+		}
+	};
+}
+
+namespace UE
+{
+	// Returns a range which allows iteration over multiple ranges in parallel, that is, the first element from each range
+	// is returned, followed by the second elmeent from each range, followed by the third etc.  This is similar to std::views::zip.
+	//
+	// The 'element type' of the zipped range is a tuple of references into each container.
+	//
+	// The constness of the references into each container will match the constness of the references
+	// returned by the original container, and the constness of the tuple is irrelevant.  AsConst can be used
+	// with a container (not view) to return const iterators.  UE::Zip supports mixing const and non-const ranges,
+	// and will 'lifetime extend' rvalue containers by taking ownership of them.
+	//
+	// The zipped range is as long as the first range argument, and it is undefined behaviour to have
+	// a first range which is longer than any of the others being zipped.
+	//
+	// Examples:
+	//
+	// // Iteration:
+	// //   { ContainerOfInts[0], AsConst(ContainerOfStrings)[0] }
+	// //   { ContainerOfInts[1], AsConst(ContainerOfStrings)[1] }
+	// //   { ContainerOfInts[2], AsConst(ContainerOfStrings)[2] }
+	// //   ...
+	// for (const auto& Pair : UE::Zip(ContainerOfInts, AsConst(ContainerOfStrings)))
+	// {
+	//     int32&         Int = Pair.Get<0>();
+	//     const FString& Str = Pair.Get<1>();
+	// }
+	//
+	// Sort Array1, Array2 and Array3 in parallel by the elements in Array2
+	// Algo::SortBy(UE::Zip(Container1, Container2), GetElementByIndex<1>);
+	template <typename Range0Type, typename... OtherRangeTypes>
+	auto Zip(Range0Type&& Range0, OtherRangeTypes&&... OtherRanges)
+	{
+		return UE::Core::Private::TRangeTuple{ TTuple<Range0Type, OtherRangeTypes...>((Range0Type&&)Range0, (OtherRangeTypes&&)OtherRanges...) };
+	}
+}
diff --git a/Engine/Source/Runtime/Core/Public/Templates/Tuple.h b/Engine/Source/Runtime/Core/Public/Templates/Tuple.h
--- a/Engine/Source/Runtime/Core/Public/Templates/Tuple.h
+++ b/Engine/Source/Runtime/Core/Public/Templates/Tuple.h
@@ -14,6 +14,7 @@
 #include "Templates/TypeHash.h"
 #include <tuple>
 #include <type_traits>
+#include <utility>
 
 // MSVC generates bad codegen when getting a reference element from an rvalue tuple.  So, if the element type is
 // an rvalue reference, get it via an lvalue - it will be the same reference value in either case.
@@ -30,6 +31,19 @@ class FArchive;
 template <typename... Types>
 struct TTuple;
 
+template <typename T> constexpr bool TIsTuple_V = false;
+
+template <typename... Types> constexpr bool TIsTuple_V<               TTuple<Types...>> = true;
+template <typename... Types> constexpr bool TIsTuple_V<const          TTuple<Types...>> = true;
+template <typename... Types> constexpr bool TIsTuple_V<      volatile TTuple<Types...>> = true;
+template <typename... Types> constexpr bool TIsTuple_V<const volatile TTuple<Types...>> = true;
+
+template <typename T>
+struct TIsTuple
+{
+	enum { Value = TIsTuple_V<T> };
+};
+
 template <typename KeyType, typename ValueType>
 using TPair = TTuple<KeyType, ValueType>;
 
@@ -443,6 +457,28 @@ namespace UE::Core::Private::Tuple
 			return Hash;
 		}
 	};
+
+	template <uint32 Index>
+	struct TGetElementByIndex
+	{
+		template <typename TupleType>
+			requires (TIsTuple_V<std::decay_t<TupleType>>)
+		constexpr auto operator()(TupleType&& Tuple) -> decltype(((TupleType&&)Tuple).template Get<Index>())
+		{
+			return ((TupleType&&)Tuple).template Get<Index>();
+		}
+	};
+
+	template <typename Type>
+	struct TGetElementByType
+	{
+		template <typename TupleType>
+			requires (TIsTuple_V<std::decay_t<TupleType>>)
+		constexpr auto operator()(TupleType&& Tuple) -> decltype(((TupleType&&)Tuple).template Get<Type>())
+		{
+			return ((TupleType&&)Tuple).template Get<Type>();
+		}
+	};
 }
 
 template <typename... Types>
@@ -699,19 +735,6 @@ UE_FORCEINLINE_HINT TTuple<Types&...> Tie(Types&... Args)
 	return TTuple<Types&...>(Args...);
 }
 
-template <typename T> constexpr bool TIsTuple_V = false;
-
-template <typename... Types> constexpr bool TIsTuple_V<               TTuple<Types...>> = true;
-template <typename... Types> constexpr bool TIsTuple_V<const          TTuple<Types...>> = true;
-template <typename... Types> constexpr bool TIsTuple_V<      volatile TTuple<Types...>> = true;
-template <typename... Types> constexpr bool TIsTuple_V<const volatile TTuple<Types...>> = true;
-
-template <typename T>
-struct TIsTuple
-{
-	enum { Value = TIsTuple_V<T> };
-};
-
 ///////////////////////////////////////////////////////////////////////////////
 // TTuple support for structured bindings - not intended to be used directly //
 ///////////////////////////////////////////////////////////////////////////////
@@ -743,6 +766,71 @@ decltype(auto) get(TupleType&& val)
 // End of structured binding support //
 ///////////////////////////////////////
 
+// Overload for std::swap - to support swapping of references
+template <typename... Types>
+void swap(const TTuple<Types&...>& Lhs, const TTuple<Types&...>& Rhs)
+{
+	VisitTupleElements(
+		[](auto& LhsElement, auto& RhsElement)
+		{
+			using std::swap;
+			swap(LhsElement, RhsElement);
+		},
+		Lhs,
+		Rhs
+	);
+}
+template <typename... Types>
+UE_REWRITE void swap(const TTuple<Types&...>& Lhs, TTuple<Types&...>& Rhs)
+{
+	std::swap(Lhs, AsConst(Rhs));
+}
+template <typename... Types>
+UE_REWRITE void swap(TTuple<Types&...>& Lhs, const TTuple<Types&...>& Rhs)
+{
+	std::swap(AsConst(Lhs), Rhs);
+}
+template <typename... Types>
+UE_REWRITE void swap(TTuple<Types&...>& Lhs, TTuple<Types&...>& Rhs)
+{
+	std::swap(AsConst(Lhs), AsConst(Rhs));
+}
+
+// Overload for UE Swap - to support swapping of references
+template <typename... Types>
+void Swap(const TTuple<Types&...>& Lhs, const TTuple<Types&...>& Rhs)
+{
+	VisitTupleElements(
+		[](auto& LhsElement, auto& RhsElement)
+		{
+			Swap(LhsElement, RhsElement);
+		},
+		Lhs,
+		Rhs
+	);
+}
+template <typename... Types>
+UE_REWRITE void Swap(const TTuple<Types&...>& Lhs, TTuple<Types&...>& Rhs)
+{
+	Swap(Lhs, AsConst(Rhs));
+}
+template <typename... Types>
+UE_REWRITE void Swap(TTuple<Types&...>& Lhs, const TTuple<Types&...>& Rhs)
+{
+	Swap(AsConst(Lhs), Rhs);
+}
+template <typename... Types>
+UE_REWRITE void Swap(TTuple<Types&...>& Lhs, TTuple<Types&...>& Rhs)
+{
+	Swap(AsConst(Lhs), AsConst(Rhs));
+}
+
+template <uint32 Index>
+inline constexpr UE::Core::Private::Tuple::TGetElementByIndex<Index> GetElementByIndex;
+
+template <typename Type>
+inline constexpr UE::Core::Private::Tuple::TGetElementByType<Type> GetElementByType;
+
 template <typename... Types>
 inline FArchive& operator<<(FArchive& Ar, TTuple<Types...>& Tuple)
 {
