Removing a unneeded CheckSlow in TCircularBuffer(uint32 Capacity) constructor that was checking for a capacity of 0. FMath::RoundUpToPowerOfTwo returns 1 for both a value of 0 and 1 so the buffer is established with a valid capacity regardless of what's passed in so the check isn't needed.

diff --git a/Engine/Source/Runtime/Core/Public/Containers/CircularBuffer.h b/Engine/Source/Runtime/Core/Public/Containers/CircularBuffer.h
--- a/Engine/Source/Runtime/Core/Public/Containers/CircularBuffer.h
+++ b/Engine/Source/Runtime/Core/Public/Containers/CircularBuffer.h
@@ -26,7 +26,6 @@ public:
 	 */
 	[[nodiscard]] explicit TCircularBuffer(uint32 Capacity)
 	{
-		checkSlow(Capacity > 0);
 		checkSlow(Capacity <= 0xffffffffU);
 
 		Elements.AddZeroed(FMath::RoundUpToPowerOfTwo(Capacity));
