CsvProfiler: Add timestamps for events. These are denoted via a ##<time in seconds> suffix. The latest Csv tools will process these and strip the suffix.
This opens up the ability to perform detailed performance analysis using events, since previously these only had per-frame timing accuracy.
This functionality can be disabled via Engine.ini, via [CsvProfiler]bEventTimestamps

diff --git a/Engine/Source/Runtime/Core/Private/ProfilingDebugging/CsvProfiler.cpp b/Engine/Source/Runtime/Core/Private/ProfilingDebugging/CsvProfiler.cpp
--- a/Engine/Source/Runtime/Core/Private/ProfilingDebugging/CsvProfiler.cpp
+++ b/Engine/Source/Runtime/Core/Private/ProfilingDebugging/CsvProfiler.cpp
@@ -37,9 +37,10 @@
 #include "Misc/WildcardString.h"
 #include "Modules/ModuleManager.h"
 #include "HAL/PlatformMemoryHelpers.h"
-
 #include "HAL/PlatformMisc.h"
 
+#include <limits>
+
 #if CSV_PROFILER
 
 #define CSV_PROFILER_INLINE FORCEINLINE
@@ -274,11 +275,14 @@ static uint32 GCsvProfilerFrameNumber = 0;
 static bool GCsvTrackWaitsOnAllThreads = false;
 static bool GCsvTrackWaitsOnGameThread = true;
 static bool GCsvTrackWaitsOnRenderThread = true;
-
 static FAutoConsoleVariableRef CVarTrackWaitsAllThreads(TEXT("csv.trackWaitsAllThreads"), GCsvTrackWaitsOnAllThreads, TEXT("Determines whether to track waits on all threads. Note that this incurs a lot of overhead"), ECVF_Default);
 static FAutoConsoleVariableRef CVarTrackWaitsGT(TEXT("csv.trackWaitsGT"), GCsvTrackWaitsOnGameThread, TEXT("Determines whether to track game thread waits. Note that this incurs overhead"), ECVF_Default);
 static FAutoConsoleVariableRef CVarTrackWaitsRT(TEXT("csv.trackWaitsRT"), GCsvTrackWaitsOnRenderThread, TEXT("Determines whether to track render thread waits. Note that this incurs overhead"), ECVF_Default);
 
+// Set via config [CsvProfiler]:bEventTimestamps. Defaults to true
+// If enabled, CSV events will include a timestamp, encoded with the following suffix : ##<time in seconds>
+static bool GCsvEventTimestamps = true;
+
 //
 // Categories
 //
@@ -294,6 +298,8 @@ static bool GCsvExitOnCompletion = false;
 
 static thread_local bool GCsvThreadLocalWaitsEnabled = false;
 
+static uint64 GFirstFrameTimestamp = 0;
+
 //
 // Forward declarations
 //
@@ -1038,6 +1044,8 @@ static void CsvProfilerEndFrameRT()
 
 static void CsvProfilerReadConfig()
 {
+	check(GConfig);
+	GConfig->GetBool(TEXT("CsvProfiler"), TEXT("bEventTimestamps"), GCsvEventTimestamps, GEngineIni);
 	FCsvCategoryData::Get()->UpdateCategoriesFromConfig();
 }
 
@@ -1407,7 +1415,7 @@ public:
 		return ThreadTimestamps.Num() - 1;
 	}
 
-	void AddBeginFrameTimestamp(ECsvTimeline::Type Timeline, const bool bDoThreadCheck = true)
+	uint64 AddBeginFrameTimestamp(ECsvTimeline::Type Timeline, const bool bDoThreadCheck = true)
 	{
 #if DO_CHECK
 		if (bDoThreadCheck)
@@ -1439,13 +1447,15 @@ public:
 			{
 				// Skip this frame boundary.
 				EOPCounter++;
-				return;
+				return 0;
 			}
 		}
 
+		uint64 Timestamp = FPlatformTime::Cycles64();
 		uint64* Element = FrameBoundaryTimestampsWriteBuffer[Timeline].ReserveElement();
-		*Element = FPlatformTime::Cycles64();
+		*Element = Timestamp;
 		FrameBoundaryTimestampsWriteBuffer[Timeline].CommitElement();
+		return Timestamp;
 	}
 
 private:
@@ -2014,15 +2024,26 @@ struct FCsvProcessedEvent
 {
 	inline uint64 GetAllocatedSize() const { return EventText.GetAllocatedSize(); }
 
-	FString GetFullName() const
+	FString GetFullText() const
 	{
+		// If the event has a timestamp, output it after the text with a ## separator
+		if (TimeOffsetSeconds != std::numeric_limits<double>::max())
+		{
+			if (CategoryIndex == 0)
+			{
+				return FString::Printf(TEXT("%s##%.6f"), *EventText, TimeOffsetSeconds);
+			}
+			return FString::Printf(TEXT("%s/%s##%.6f"), *FCsvCategoryData::Get()->GetCategoryNameByIndex(CategoryIndex), *EventText, TimeOffsetSeconds);
+		}
 		if (CategoryIndex == 0)
 		{
 			return EventText;
 		}
-		return FCsvCategoryData::Get()->GetCategoryNameByIndex(CategoryIndex) + TEXT("/") + EventText;
+		return FString::Printf(TEXT("%s/%s"), *FCsvCategoryData::Get()->GetCategoryNameByIndex(CategoryIndex), *EventText);
 	}
+
 	FString EventText;
+	double TimeOffsetSeconds;
 	uint32 FrameNumber;
 	uint32 CategoryIndex;
 };
@@ -3010,7 +3031,7 @@ void FCsvStreamWriter::FinalizeNextRow()
 			EventStrings.Reserve(Row->Events.Num());
 			for (FCsvProcessedEvent& Event : Row->Events)
 			{
-				EventStrings.Add(Event.GetFullName());
+				EventStrings.Add(Event.GetFullText());
 			}
 
 			Stream.WriteSemicolonSeparatedStringList(EventStrings);
@@ -3080,6 +3101,7 @@ void FCsvStreamWriter::Finalize(const TMap<FString, FString>& Metadata)
 
 	// Insert some metadata to indicate the file has a summary header row
 	Stream.WriteMetadataEntry((TEXT("HasHeaderRowAtEnd")), TEXT("1"));
+	Stream.WriteMetadataEntry((TEXT("EventTimestamps")), GCsvEventTimestamps ? TEXT("1") : TEXT("0") );
 
 	// Add metadata at the end of the file, making sure commandline is last (this is required for parsing)
 	const TPair<FString, FString>* CommandlineEntry = NULL;
@@ -3540,6 +3562,7 @@ void FCsvProfilerThreadDataProcessor::Process(FCsvProcessThreadDataStats& OutSta
 				ProcessedEvent.EventText = Event.EventText;
 				ProcessedEvent.FrameNumber = FrameNumber;
 				ProcessedEvent.CategoryIndex = Event.CategoryIndex;
+				ProcessedEvent.TimeOffsetSeconds = GCsvEventTimestamps ? FPlatformTime::ToSeconds64(Event.Timestamp-GFirstFrameTimestamp) : std::numeric_limits<double>::max();
 				Writer->PushEvent(ProcessedEvent);
 			}
 		}
@@ -3705,11 +3728,12 @@ void FCsvProfiler::BeginFrame()
 
 		if (GCsvProfilerIsCapturing)
 		{
-			GFrameBoundaries.AddBeginFrameTimestamp(ECsvTimeline::Gamethread);
+			uint64 BeginFrameTimestamp = GFrameBoundaries.AddBeginFrameTimestamp(ECsvTimeline::Gamethread);
 
 			if (CaptureFrameNumber == 0)
 			{
 				OnCSVProfileFirstFrameDelegate.Broadcast();
+				GFirstFrameTimestamp = BeginFrameTimestamp;
 			}
 
 			if (!bNamedEventsWasEnabled && (GCycleStatsShouldEmitNamedEvents > 0))
