Remove potential ambiguity if user code has overlapping names within the current namespace.

The Verse team is attempting to eliminate the uLang container types, and uLang currently contains clones of the Core types with these names. We hope to eliminate the uLang clone types shortly, but during the transition phase, both container systems will need to safely coexist. Adding :: qualifiers will solve our type ambiguity errors and should otherwise be harmless.

diff --git a/Engine/Source/Runtime/Core/Public/Templates/AndOrNot.h b/Engine/Source/Runtime/Core/Public/Templates/AndOrNot.h
--- a/Engine/Source/Runtime/Core/Public/Templates/AndOrNot.h
+++ b/Engine/Source/Runtime/Core/Public/Templates/AndOrNot.h
@@ -13,8 +13,8 @@ struct TAnd;
 template <bool LHSValue, typename... RHS>
 struct TAndValue
 {
-	static constexpr bool Value = TAnd<RHS...>::value;
-	static constexpr bool value = TAnd<RHS...>::value;
+	static constexpr bool Value = ::TAnd<RHS...>::value;
+	static constexpr bool value = ::TAnd<RHS...>::value;
 };
 
 template <typename... RHS>
@@ -25,7 +25,7 @@ struct TAndValue<false, RHS...>
 };
 
 template <typename LHS, typename... RHS>
-struct TAnd<LHS, RHS...> : TAndValue<LHS::Value, RHS...>
+struct TAnd<LHS, RHS...> : ::TAndValue<LHS::Value, RHS...>
 {
 };
 
@@ -45,8 +45,8 @@ struct TOr;
 template <bool LHSValue, typename... RHS>
 struct TOrValue
 {
-	static constexpr bool Value = TOr<RHS...>::value;
-	static constexpr bool value = TOr<RHS...>::value;
+	static constexpr bool Value = ::TOr<RHS...>::value;
+	static constexpr bool value = ::TOr<RHS...>::value;
 };
 
 template <typename... RHS>
@@ -57,7 +57,7 @@ struct TOrValue<true, RHS...>
 };
 
 template <typename LHS, typename... RHS>
-struct TOr<LHS, RHS...> : TOrValue<LHS::Value, RHS...>
+struct TOr<LHS, RHS...> : ::TOrValue<LHS::Value, RHS...>
 {
 };
 
diff --git a/Engine/Source/Runtime/Core/Public/Templates/IsEnumClass.h b/Engine/Source/Runtime/Core/Public/Templates/IsEnumClass.h
--- a/Engine/Source/Runtime/Core/Public/Templates/IsEnumClass.h
+++ b/Engine/Source/Runtime/Core/Public/Templates/IsEnumClass.h
@@ -23,5 +23,5 @@ namespace UE::Core::Private::IsEnumClass
 template <typename T>
 struct TIsEnumClass
 { 
-	enum { Value = TAndValue<__is_enum(T), TNot<UE::Core::Private::IsEnumClass::TIsEnumConvertibleToInt<T>>>::Value };
+	enum { Value = ::TAndValue<__is_enum(T), ::TNot<::UE::Core::Private::IsEnumClass::TIsEnumConvertibleToInt<T>>>::Value };
 };
