Fix misleading doc comment.

It appears that this doc-comment originated above TIsArithmetic and was copy-pasted here verbatim.

diff --git a/Engine/Source/Runtime/Core/Public/Templates/IsEnumClass.h b/Engine/Source/Runtime/Core/Public/Templates/IsEnumClass.h
--- a/Engine/Source/Runtime/Core/Public/Templates/IsEnumClass.h
+++ b/Engine/Source/Runtime/Core/Public/Templates/IsEnumClass.h
@@ -18,7 +18,7 @@ namespace UE::Core::Private::IsEnumClass
 }
 
 /**
- * Traits class which tests if a type is arithmetic.
+ * Traits class which tests if a type is an enum class.
  */
 template <typename T>
 struct TIsEnumClass
