Remove unnecessary check that sometimes triggerers due timing.

#rnx

diff --git a/Engine/Source/Runtime/Core/Private/ProfilingDebugging/TraceAuxiliary.cpp b/Engine/Source/Runtime/Core/Private/ProfilingDebugging/TraceAuxiliary.cpp
--- a/Engine/Source/Runtime/Core/Private/ProfilingDebugging/TraceAuxiliary.cpp
+++ b/Engine/Source/Runtime/Core/Private/ProfilingDebugging/TraceAuxiliary.cpp
@@ -928,7 +928,6 @@ void FTraceAuxiliaryImpl::EmitDiagnostics(const UE::Trace::FStatistics& Stats)
 	constexpr uint32 Stickyness = 50; // Number of frames the message will "stick" after
 	constexpr float BlockPoolUtilizationWarning = 0.8f; // Warn above 80% utilization
 	constexpr float BlockPoolGrowthWarning = 0.8f; // Warn above 80% growth
-	check(Stats.BlockPoolAllocatedBlocks >= Stats.BlockPoolFreeBlocks);
 	const float BlockPoolFreeFraction = Stats.BlockPoolAllocatedBlocks ? float(Stats.BlockPoolFreeBlocks)/float(Stats.BlockPoolAllocatedBlocks) : 1.0f;
 	const float BlockPoolUtilization = 1.0f - BlockPoolFreeFraction;
 	const float BlockPoolGrowth = GInitializeDesc.BlockPoolMaxSize ? float(Stats.BlockPoolAllocated) / float(GInitializeDesc.BlockPoolMaxSize) : 0.0f;
