Add telemetry to identify the build flavor

diff --git a/Engine/Source/Runtime/Core/Private/GenericPlatform/GenericPlatformCrashContext.cpp b/Engine/Source/Runtime/Core/Private/GenericPlatform/GenericPlatformCrashContext.cpp
--- a/Engine/Source/Runtime/Core/Private/GenericPlatform/GenericPlatformCrashContext.cpp
+++ b/Engine/Source/Runtime/Core/Private/GenericPlatform/GenericPlatformCrashContext.cpp
@@ -365,7 +365,10 @@ void FGenericCrashContext::Initialize()
 #endif
 	Symbols = FString::Printf(TEXT("%s-%s-%s"), *Symbols, FPlatformMisc::GetUBTPlatform(), NCached::Session.BuildConfigurationName).Replace(TEXT("+"), TEXT("*"));
 #ifdef UE_BUILD_FLAVOR
-	Symbols = FString::Printf(TEXT("%s-%s"), *Symbols, *FString(UE_BUILD_FLAVOR));
+	const FString BuildFlavor(UE_BUILD_FLAVOR);
+	Symbols = FString::Printf(TEXT("%s-%s"), *Symbols, *BuildFlavor);
+#else
+	const FString BuildFlavor;
 #endif
 	NCached::Set(NCached::Session.SymbolsLabel, *Symbols);
 	if (Symbols.Len() >= UE_ARRAY_COUNT(NCached::Session.SymbolsLabel))
@@ -513,6 +516,8 @@ void FGenericCrashContext::Initialize()
 		}
 	});
 
+	SetEngineData(TEXT("BuildFlavor"), BuildFlavor);
+
 	// Store some additional info in the generic maps if it is available
 	if (const TCHAR* BuildURL = FApp::GetBuildURL())
 	{
