[LLM] * Removed unneeded lock from FLLMTracker::PauseAndTrackMemory(). Added allocations lock in FLLMThreadState::Clear().
* Added comments re allocations lock. Also renamed TagSection to AllocationsLock in FLLMThreadState.
* Fixed few typos in comments and log messages.

#rnx

diff --git a/Engine/Source/Runtime/Core/Private/HAL/LowLevelMemTracker.cpp b/Engine/Source/Runtime/Core/Private/HAL/LowLevelMemTracker.cpp
--- a/Engine/Source/Runtime/Core/Private/HAL/LowLevelMemTracker.cpp
+++ b/Engine/Source/Runtime/Core/Private/HAL/LowLevelMemTracker.cpp
@@ -986,9 +986,9 @@ TAutoConsoleVariable<int32> CVarLLMHeaderMaxSize(
 );
 
 FAutoConsoleCommand LLMSnapshot(
-	TEXT("LLMSnapshot"), 
-	TEXT("Takes a single LLM Snapshot of one frame. This command requires the commandline -llmdisableautopublish"), 
-	FConsoleCommandDelegate::CreateLambda([]() 
+	TEXT("LLMSnapshot"),
+	TEXT("Takes a single LLM Snapshot of one frame. This command requires the commandline -llmdisableautopublish"),
+	FConsoleCommandDelegate::CreateLambda([]()
 	{
 		FLowLevelMemTracker::Get().PublishDataSingleFrame();
 	}));
@@ -1364,7 +1364,7 @@ public:
 		UE::LLM::ESizeParams SizeParams);
 protected:
 	void RecordTagToCsv(int32 CsvCategoryIndex, const FTagData* TagData, int64 TagSize);
-#if LLM_CSV_PROFILER_WRITER_ENABLED 
+#if LLM_CSV_PROFILER_WRITER_ENABLED
 	ELLMTracker Tracker;
 	TLLMMap<const FTagData*, FName> TagDataToCsvStatName;
 #endif
@@ -1381,9 +1381,9 @@ public:
 	void PushTag(const FTagData* TagData, ELLMTagSet TagSet);
 	void PopTag(ELLMTagSet TagSet);
 	const FTagData* GetTopTag(ELLMTagSet TagSet);
-	void TrackAllocation(const void* Ptr, int64 Size, ELLMTracker Tracker, ELLMAllocType AllocType, 
+	void TrackAllocation(const void* Ptr, int64 Size, ELLMTracker Tracker, ELLMAllocType AllocType,
 		const FTagData* TagData, const FTagData* AssetTagData, const FTagData* AssetClassTagData, const FTagData* UObjectClassesTagData, bool bTrackInMemPro);
-	void TrackFree(const void* Ptr, int64 Size, ELLMTracker Tracker, ELLMAllocType AllocType, 
+	void TrackFree(const void* Ptr, int64 Size, ELLMTracker Tracker, ELLMAllocType AllocType,
 		const FTagData* TagData, const FTagData* AssetTagData, const FTagData* AssetClassTagData, const FTagData* UObjectClassesTagData, bool bTrackInMemPro);
 	void TrackMemory(int64 Amount, ELLMTracker Tracker, ELLMAllocType AllocType,
 		const FTagData* TagData);
@@ -1397,9 +1397,14 @@ public:
 
 	void ClearAllocTypeAmounts();
 
+	// AllocationsLock must be held around read/write of Allocations and AllocTypeAmounts, because these are read/write
+	// from the Update thread. All other data on FLLMThreadState are accessed only from the thread that owns the
+	// FLLMThreadState and do not require a lock.
+	// AllocationLock is also used to guarantee our contract provided to FLLMTracker::LockAllThreadTags: We must not
+	// dereference any FTagData* held by the FLLMThreadState unless we have entered AllocationLock.
 	FConstTagDataArray TagStack[static_cast<int32>(ELLMTagSet::Max)];
 	FThreadTagSizeMap Allocations;
-	FCriticalSection TagSection;
+	FCriticalSection AllocationsLock;
 
 	int8 PausedCounter[(int32)ELLMAllocType::Count];
 	int64 AllocTypeAmounts[(int32)ELLMAllocType::Count];
@@ -1440,10 +1445,10 @@ public:
 
 	void Clear();
 
-	// Dump the allocation count and size for each tag, along with the porportion that is private/shared/unreferenced
+	// Dump the allocation count and size for each tag, along with the proportion that is private/shared/unreferenced
 	// on linux machines. This will dump to a CSV in the project Saved/LLM directory with a separate file for each
 	// tracker and forked child. Integers will be appended to the filename to prevent overwrites. Returns false
-	// if the memory inforation can't be retrieved (likely because it's not on a supported platform)
+	// if the memory information can't be retrieved (likely because it's not on a supported platform)
 	bool DumpForkedAllocationInfo();
 
 	void PublishStats(UE::LLM::ESizeParams SizeParam);
@@ -1474,7 +1479,7 @@ public:
 				case ELLMTagSet::None: return 0;
 				case ELLMTagSet::Assets: return 1;
 				case ELLMTagSet::AssetClasses: return 2;
-				default: 
+				default:
 					LLMCheck(false);
 					return 0;
 			}
@@ -1482,8 +1487,8 @@ public:
 			switch (InTagSet)
 			{
 				case ELLMTagSet::None: return 0;
-				case ELLMTagSet::UObjectClasses: return 1; 
-				default: 
+				case ELLMTagSet::UObjectClasses: return 1;
+				default:
 					LLMCheck(false);
 					return 0;
 			}
@@ -1564,7 +1569,7 @@ public:
 #else
 		int32 Tag[LLM_TAGSET_MAX];
 
-		static_assert(LLM_TAGSET_MAX <= 3, 
+		static_assert(LLM_TAGSET_MAX <= 3,
 			"Incrementing the maximum amount of TagSet enums that can be defined at the same time adds 4 bytes of LLM overhead memory FOR EVERY ACTIVE ALLOCATION (e.g. 10 million * 4 == 40 million bytes). "
 			"Do not add them without measuring the cost of the overhead in the cook and runtime of large projects.");
 #endif // LLM_ALLOW_ASSETS_TAGS
@@ -1574,7 +1579,7 @@ public:
 		static const int32 InvalidCompressedTagValue = -1;
 	};
 
-	typedef LLMMap<PointerKey, uint32, FLowLevelAllocInfo, LLMNumAllocsType> FLLMAllocMap;	// pointer, size, info, Capacity SizeType
+	typedef LLMMap<PointerKey, uint32, FLowLevelAllocInfo, LLMNumAllocsType> FLLMAllocMap;  // pointer, size, info, Capacity SizeType
 
 	void SetTotalTags(const FTagData* InOverrideUntaggedTagData, const FTagData* InOverrideTrackedTotalTagData);
 	void Update();
@@ -1631,9 +1636,9 @@ protected:
 	TMap<uint32, FLLMThreadState*, FDefaultSetLLMAllocator> ThreadIdToThreadState;
 
 	/** Sum of memory from all tracked tags. Duplicated in separate storage to make it instantly available at any time of frame without waiting for accumulation from threads during update.
-	    GCC_ALIGN is required because it is modified in FPlatformAtomics::InterlockedAdd, which requires aligned values. */
+		GCC_ALIGN is required because it is modified in FPlatformAtomics::InterlockedAdd, which requires aligned values. */
 	int64 TrackedTotal GCC_ALIGN(8);
-	
+
 	// The total tracked memory when the snapshot was taken
 	int64 TrackedTotalInSnapshot;
 
@@ -2112,7 +2117,7 @@ void FLowLevelMemTracker::TickInternal()
 	PlatformTracker.SetTagAmountInUpdate(FindOrAddTagData(ELLMTag::PlatformOSAvailable), PlatformStats.AvailablePhysical, false);
 
 	// remove the MemoryUsageCurrentOverhead from the "Total" for the default LLM as it's not something anyone needs
-	// to investigate when finding what to reduce the platform LLM will have the info 
+	// to investigate when finding what to reduce the platform LLM will have the info
 	const int64 DefaultProcessMemory = PlatformProcessMemory - (!FPlatformMemory::TracksLLMAllocations() ? MemoryUsageCurrentOverhead : 0);
 	const int64 ImmediatelyFreeableCachedMemory = UE::Private::GMalloc ? UE::Private::GMalloc->GetImmediatelyFreeableCachedMemorySize() : 0;
 	// we need to subtract diff between total cached free mem and immediately freeable cached memory as it is not accounted anywhere in the LLM and would artificially inflate Untracked.
@@ -2333,7 +2338,7 @@ void FLowLevelMemTracker::ProcessCommandLine(const TCHAR* CmdLine)
 		}
 		bShouldDisable = false;
 	}
-	
+
 	bAutoPublish = FParse::Param(CmdLine, TEXT("LLMDISABLEAUTOPUBLISH")) == false;
 
 	if (!bCanEnable)
@@ -2934,13 +2939,13 @@ void FLowLevelMemTracker::ClearTagDatas()
 void FLowLevelMemTracker::RegisterTagDeclaration(FLLMTagDeclaration& TagDeclaration)
 {
 	TagDeclaration.ConstructUniqueName();
-	RegisterTagData(TagDeclaration.UniqueName, TagDeclaration.DisplayName, TagDeclaration.ParentTagName, 
-		TagDeclaration.StatName, TagDeclaration.SummaryStatName, false, ELLMTag::CustomName, false, 
+	RegisterTagData(TagDeclaration.UniqueName, TagDeclaration.DisplayName, TagDeclaration.ParentTagName,
+		TagDeclaration.StatName, TagDeclaration.SummaryStatName, false, ELLMTag::CustomName, false,
 		UE::LLMPrivate::ETagReferenceSource::Declare, TagDeclaration.TagSet);
 }
 
-UE::LLMPrivate::FTagData& FLowLevelMemTracker::RegisterTagData(FName Name, FName DisplayName, FName ParentName, 
-	FName StatName, FName SummaryStatName, bool bHasEnumTag, ELLMTag EnumTag, bool bIsStatTag, 
+UE::LLMPrivate::FTagData& FLowLevelMemTracker::RegisterTagData(FName Name, FName DisplayName, FName ParentName,
+	FName StatName, FName SummaryStatName, bool bHasEnumTag, ELLMTag EnumTag, bool bIsStatTag,
 	UE::LLMPrivate::ETagReferenceSource ReferenceSource, ELLMTagSet TagSet/* = ELLMTagSet::None*/)
 {
 	LLMCheckf(!bIsBootstrapping,
@@ -3221,7 +3226,7 @@ TArray<const UE::LLMPrivate::FTagData*> FLowLevelMemTracker::GetTrackedTags(ELLM
 	BootstrapInitialise();
 
 	UE::TSharedLock Lock(TagDataMutex);
-	
+
 	TArray<const FTagData*> FoundResults(TagDatas->FilterByPredicate([TagSet](const FTagData* InTagData)
 	{
 		return InTagData != nullptr && InTagData->GetTagSet() == TagSet;
@@ -3238,7 +3243,7 @@ TArray<const UE::LLMPrivate::FTagData*> FLowLevelMemTracker::GetTrackedTags(ELLM
 	{
 		return TArray<const FTagData*>();
 	}
-	
+
 	BootstrapInitialise();
 
 	UE::TUniqueLock UpdateLock(UpdateMutex); // uses of TagSizes are guarded by the UpdateLock
@@ -3543,7 +3548,7 @@ void FLowLevelMemTracker::PublishDataSingleFrame()
 	else
 	{
 		bPublishSingleFrame = true;
-	}	
+	}
 }
 
 const UE::LLMPrivate::FTagData* FLowLevelMemTracker::FindOrAddTagData(ELLMTag EnumTag,
@@ -3568,7 +3573,7 @@ const UE::LLMPrivate::FTagData* FLowLevelMemTracker::FindOrAddTagData(ELLMTag En
 	if (!bFullyInitialised)
 	{
 		FinishInitialise();
-		// Reeneter this function so that we retry the find above; we avoid infinite recursion because
+		// Reenter this function so that we retry the find above; we avoid infinite recursion because
 		// bFullyInitialised is now true.
 		return FindOrAddTagData(EnumTag, ReferenceSource);
 	}
@@ -3611,7 +3616,7 @@ const UE::LLMPrivate::FTagData* FLowLevelMemTracker::FindOrAddTagData(FName TagN
 	if (!bFullyInitialised)
 	{
 		FinishInitialise();
-		// Reeneter this function so that we retry the find above; note we avoid infinite recursion because
+		// Reenter this function so that we retry the find above; note we avoid infinite recursion because
 		// bFullyInitialised is now true.
 		return FindOrAddTagData(TagName, TagSet, StatName, ReferenceSource);
 	}
@@ -3729,7 +3734,7 @@ FLLMScope::FLLMScope(FName TagName, bool bInIsStatTag, ELLMTagSet InTagSet, ELLM
 					bEnabled = true;
 					Tracker = InTracker;
 					TagSet = InTagSet;
-	
+
 					LLMRef.GetTracker(Tracker)->PushTag(TagName, bInIsStatTag, InTagSet);
 				}
 			}
@@ -3940,7 +3945,7 @@ void FLLMPauseScope::Init(FName TagName, ELLMTag EnumTag, bool bIsEnumTag, bool
 				}
 				else
 				{
-					LLMRef.GetTracker(Tracker)->PauseAndTrackMemory(TagName, ELLMTagSet::None, bIsStatTag, 
+					LLMRef.GetTracker(Tracker)->PauseAndTrackMemory(TagName, ELLMTagSet::None, bIsStatTag,
 						static_cast<int64>(Amount), InAllocType);
 				}
 			}
@@ -4695,8 +4700,8 @@ const FTagData* FTagData::GetContainingEnumTagData() const
 		}
 		TagData = TagData->GetParent();
 	} while (TagData);
-	LLMCheckf(false, TEXT("TagData is not a descendent of an ELLMTag TagData. ")
-		TEXT("All TagDatas must be descendents of ELLMTag::CustomName if they are not descendets of any other ELLMTag"));
+	LLMCheckf(false, TEXT("TagData is not a descendant of an ELLMTag TagData. ")
+		TEXT("All TagDatas must be descendants of ELLMTag::CustomName if they are not descendants of any other ELLMTag"));
 		return this;
 }
 
@@ -4821,7 +4826,7 @@ FLLMTracker::~FLLMTracker()
 
 void FLLMTracker::Initialise(
 	ELLMTracker InTracker,
-	FLLMAllocator* InAllocator) 
+	FLLMAllocator* InAllocator)
 {
 	check(InTracker != ELLMTracker::Max);
 	Tracker = InTracker;
@@ -4937,7 +4942,7 @@ void FLLMTracker::TrackAllocation(const void* Ptr, int64 Size, const FTagData* A
 {
 	if (IsPaused(AllocType))
 	{
-		// When Paused, we do not track any new allocations and we we do not update the counters for the memory
+		// When Paused, we do not track any new allocations and we do not update the counters for the memory
 		// they use; the code that triggered the pause is responsible for updating those counters. Since we do not
 		// track the allocations, TrackFree will likewise not update the counters when those allocations are freed.
 		return;
@@ -4945,7 +4950,7 @@ void FLLMTracker::TrackAllocation(const void* Ptr, int64 Size, const FTagData* A
 
 	// track the total quickly
 	FPlatformAtomics::InterlockedAdd(&TrackedTotal, Size);
-	
+
 #if !LLM_ENABLED_FULL_TAGS
 	// When full tags are disabled, we instead store the top-level enumtag parent of the allocation's tag.
 	ActiveTagData = ActiveTagData->GetContainingEnumTagData();
@@ -4954,7 +4959,7 @@ void FLLMTracker::TrackAllocation(const void* Ptr, int64 Size, const FTagData* A
 #if LLM_ALLOW_ASSETS_TAGS
 	const FTagData* AssetTagData = State->GetTopTag(ELLMTagSet::Assets);
 	const FTagData* AssetClassTagData = State->GetTopTag(ELLMTagSet::AssetClasses);
-	
+
 	if (AssetTagData == nullptr && !LLMRef.IsBootstrapping() && LLMRef.IsTagSetActive(ELLMTagSet::Assets))
 	{
 		AssetTagData = LLMRef.FindOrAddTagData(TagName_UntaggedAsset, ELLMTagSet::Assets, false, UE::LLMPrivate::ETagReferenceSource::FunctionAPI);
@@ -5036,12 +5041,12 @@ void FLLMTracker::TrackFree(const void* Ptr, ELLMAllocType AllocType, bool bTrac
 #if LLM_ALLOW_ASSETS_TAGS
 	const FTagData* AssetTagData = AllocInfo.GetTag(LLMRef, ELLMTagSet::Assets);
 	const FTagData* AssetClassTagData = AllocInfo.GetTag(LLMRef, ELLMTagSet::AssetClasses);
-#else 
+#else
 	const FTagData* AssetTagData = nullptr;
 	const FTagData* AssetClassTagData = nullptr;
 #endif // LLM_ALLOW_ASSETS_TAGS
 
-#if LLM_ALLOW_UOBJECTCLASSES_TAGS	
+#if LLM_ALLOW_UOBJECTCLASSES_TAGS
 	const FTagData* UObjectClassesTagData = AllocInfo.GetTag(LLMRef, ELLMTagSet::UObjectClasses);
 #else
 	const FTagData* UObjectClassesTagData = nullptr;
@@ -5172,21 +5177,22 @@ void FLLMTracker::PauseAndTrackMemory(const FTagData* TagData, int64 Amount, ELL
 	FLLMThreadState* State = GetOrCreateState();
 	FPlatformAtomics::InterlockedAdd(&TrackedTotal, Amount);
 	State->TrackMemory(Amount, Tracker, AllocType, TagData);
-	FScopeLock Lock(&State->TagSection);
+	LLMCheck(State->PausedCounter[static_cast<int32>(AllocType)] < 127);
 	State->PausedCounter[static_cast<int32>(AllocType)]++;
 }
 
 void FLLMTracker::Pause(ELLMAllocType AllocType)
 {
 	FLLMThreadState* State = GetOrCreateState();
+	LLMCheck(State->PausedCounter[static_cast<int32>(AllocType)] < 127);
 	State->PausedCounter[static_cast<int32>(AllocType)]++;
 }
 
 void FLLMTracker::Unpause(ELLMAllocType AllocType)
 {
 	FLLMThreadState* State = GetOrCreateState();
+	LLMCheck(State->PausedCounter[static_cast<int32>(AllocType)] > 0);
 	State->PausedCounter[static_cast<int32>(AllocType)]--;
-	LLMCheck( State->PausedCounter[static_cast<int32>(AllocType)] >= 0 );
 }
 
 bool FLLMTracker::IsPaused(ELLMAllocType AllocType)
@@ -5424,7 +5430,7 @@ bool FLLMTracker::DumpForkedAllocationInfo()
 		uint64 ByteCount[Class_COUNT] = {};
 	};
 
-	
+
 	// We cannot call GetTag to find the Tag from the Allocation when we are within AllocationMap, because GetTag
 	// locks the TagDataMutex and TagDataMutex has to be entered before entering AllocationMap.LockAll.
 	// So for the TagIdentifier use the CompressedTag rather than calling GetTag.
@@ -5476,7 +5482,7 @@ bool FLLMTracker::DumpForkedAllocationInfo()
 		uint32 SharedCount = 0;
 		uint32 PrivateCount = 0;
 		uint32 UnreferencedCount = 0;
-		
+
 		// Walk the page allocations until we get them all if we aren't all in the same one.
 		uint64 Remaining = Size;
 		while (Remaining)
@@ -5632,7 +5638,7 @@ bool FLLMTracker::DumpForkedAllocationInfo()
 		UniqueFilename = SavedDir / *FString::Printf(TEXT("LLM_PrivateShared_Tracker%s_%s_%d.csv"), TrackerString, *ForkString, Counter);
 		Counter++;
 	}
-	
+
 	if (FFileHelper::SaveStringArrayToFile(Lines, *UniqueFilename) == false)
 	{
 		FPlatformMisc::LowLevelOutputDebugStringf(TEXT("Failed to write llm private/shared csv file: %s\n"), *UniqueFilename);
@@ -5669,7 +5675,7 @@ void FLLMTracker::OnTagsResorted(FTagDataArray& OldTagDatas)
 		{
 			for (int32 TagSetAsInteger = 0; TagSetAsInteger < static_cast<int32>(ELLMTagSet::Max); TagSetAsInteger++)
 			{
-				if (!LLMRef.IsTagSetActive(static_cast<ELLMTagSet>(TagSetAsInteger))) 
+				if (!LLMRef.IsTagSetActive(static_cast<ELLMTagSet>(TagSetAsInteger)))
 				{
 					continue;
 				}
@@ -5801,7 +5807,7 @@ void FLLMTracker::GetTagsNamesWithAmountFiltered(TMap<FName, uint64>& OutTagsNam
 				break;
 			}
 		}
-			
+
 		if (bIncludeAllocation)
 		{
 #if LLM_ALLOW_ASSETS_TAGS || LLM_ALLOW_UOBJECTCLASSES_TAGS
@@ -5818,7 +5824,7 @@ void FLLMTracker::GetTagsNamesWithAmountFiltered(TMap<FName, uint64>& OutTagsNam
 	}
 	AllocationMap.UnlockAll();
 #endif // !UE_ONLY_USE_PLATFORM_TRACKER
-	
+
 	for (const FLocalTagData& LocalTagData : TagDataCached)
 	{
 		if (LocalTagData.Size != 0)
@@ -5848,8 +5854,8 @@ bool FLLMTracker::FindTagsForPtr(void* InPtr, TArray<const FTagData *, TInlineAl
 		if (!LLMRef.IsTagSetActive(static_cast<ELLMTagSet>(TagSetAsInteger)))
 		{
 			continue;
-		} 
-		
+		}
+
 		OutTags[TagSetAsInteger] = AllocInfoPtr.GetTag(LLMRef, static_cast<ELLMTagSet>(TagSetAsInteger));
 	}
 #else
@@ -5908,11 +5914,15 @@ FLLMThreadState::FLLMThreadState()
 		PausedCounter[Index] = 0;
 	}
 
+	// By contract AllocationsLock must be entered around ClearAllocTypeAmounts, but multithreading contracts of that type
+	// are not required in constructors, which complete before the object can be used from multiple threads.
 	ClearAllocTypeAmounts();
 }
 
 void FLLMThreadState::Clear()
 {
+	FScopeLock Lock(&AllocationsLock);
+
 	for (int32 TagSetAsInteger = 0; TagSetAsInteger < static_cast<int32>(ELLMTagSet::Max); TagSetAsInteger++)
 	{
 		TagStack[TagSetAsInteger].Empty();
@@ -5923,7 +5933,7 @@ void FLLMThreadState::Clear()
 
 void FLLMThreadState::PushTag(const FTagData* TagData, ELLMTagSet TagSet)
 {
-	FScopeLock Lock(&TagSection);
+	FScopeLock Lock(&AllocationsLock);
 
 	// Push a tag.
 	TagStack[static_cast<int32>(TagSet)].Add(TagData);
@@ -5931,7 +5941,7 @@ void FLLMThreadState::PushTag(const FTagData* TagData, ELLMTagSet TagSet)
 
 void FLLMThreadState::PopTag(ELLMTagSet TagSet)
 {
-	FScopeLock Lock(&TagSection);
+	FScopeLock Lock(&AllocationsLock);
 
 	LLMCheckf(TagStack[static_cast<int32>(TagSet)].Num() > 0,
 		TEXT("Called FLLMThreadState::PopTag without a matching Push (stack was empty on pop)"));
@@ -5950,16 +5960,16 @@ const FTagData* FLLMThreadState::GetTopTag(ELLMTagSet TagSet)
 
 void FLLMThreadState::IncrTag(const FTagData* TagData, int64 Amount)
 {
-	// Caller is responsible for holding a lock on TagSection.
+	// Caller is responsible for holding a lock on AllocationsLock.
 	FThreadTagSizeData& AllocationSize = Allocations.FindOrAdd(TagData->GetIndex());
 	AllocationSize.TagData = TagData;
 	AllocationSize.Size += Amount;
 }
 
-void FLLMThreadState::TrackAllocation(const void* Ptr, int64 Size, ELLMTracker Tracker, ELLMAllocType AllocType, 
+void FLLMThreadState::TrackAllocation(const void* Ptr, int64 Size, ELLMTracker Tracker, ELLMAllocType AllocType,
 	const FTagData* TagData, const FTagData* AssetTagData, const FTagData* AssetClassTagData, const FTagData* UObjectClassesTagData, bool bTrackInMemPro)
 {
-	FScopeLock Lock(&TagSection);
+	FScopeLock Lock(&AllocationsLock);
 
 	AllocTypeAmounts[static_cast<int32>(AllocType)] += Size;
 
@@ -5999,10 +6009,10 @@ void FLLMThreadState::TrackAllocation(const void* Ptr, int64 Size, ELLMTracker T
 #endif
 }
 
-void FLLMThreadState::TrackFree(const void* Ptr, int64 Size, ELLMTracker Tracker, ELLMAllocType AllocType, 
+void FLLMThreadState::TrackFree(const void* Ptr, int64 Size, ELLMTracker Tracker, ELLMAllocType AllocType,
 	const FTagData* TagData, const FTagData* AssetTagData, const FTagData* AssetClassTagData, const FTagData* UObjectClassesTagData, bool bTrackInMemPro)
 {
-	FScopeLock Lock(&TagSection);
+	FScopeLock Lock(&AllocationsLock);
 
 	AllocTypeAmounts[static_cast<int32>(AllocType)] -= Size;
 	if (TagData)
@@ -6043,7 +6053,7 @@ void FLLMThreadState::TrackFree(const void* Ptr, int64 Size, ELLMTracker Tracker
 
 void FLLMThreadState::TrackMemory(int64 Amount, ELLMTracker Tracker, ELLMAllocType AllocType, const FTagData* TagData)
 {
-	FScopeLock Lock(&TagSection);
+	FScopeLock Lock(&AllocationsLock);
 	AllocTypeAmounts[static_cast<int32>(AllocType)] += Amount;
 	IncrTag(TagData, Amount);
 
@@ -6072,7 +6082,7 @@ void FLLMThreadState::TrackMoved(const void* Dest, const void* Source, int64 Siz
 
 void FLLMThreadState::PropagateChildSizesToParents()
 {
-	FScopeLock Lock(&TagSection);
+	FScopeLock Lock(&AllocationsLock);
 
 	// Make sure all parents of any TagDatas in the Allocations are also present.
 	FConstTagDataArray ParentsToAdd;
@@ -6107,7 +6117,7 @@ void FLLMThreadState::PropagateChildSizesToParents()
 
 void FLLMThreadState::OnTagsResorted(FTagDataArray& OldTagDatas)
 {
-	FScopeLock Lock(&TagSection);
+	FScopeLock Lock(&AllocationsLock);
 	TArray<FThreadTagSizeData, FDefaultLLMAllocator> AllocationDatas;
 	AllocationDatas.Reserve(Allocations.Num());
 	for (const TPair<int32, FThreadTagSizeData>& It : Allocations)
@@ -6125,18 +6135,18 @@ void FLLMThreadState::LockTags(bool bLock)
 {
 	if (bLock)
 	{
-		TagSection.Lock();
+		AllocationsLock.Lock();
 	}
 	else
 	{
-		TagSection.Unlock();
+		AllocationsLock.Unlock();
 	}
 }
 
-void FLLMThreadState::FetchAndClearTagSizes(FTrackerTagSizeMap& TagSizes, int64* InAllocTypeAmounts, 
+void FLLMThreadState::FetchAndClearTagSizes(FTrackerTagSizeMap& TagSizes, int64* InAllocTypeAmounts,
 	bool bTrimAllocations)
 {
-	FScopeLock Lock(&TagSection);
+	FScopeLock Lock(&AllocationsLock);
 	for (TPair<int32, FThreadTagSizeData>& Item : Allocations)
 	{
 		FThreadTagSizeData& ThreadInfo = Item.Value;
@@ -6170,6 +6180,8 @@ void SetMemoryStatByFName(FName Name, int64 Amount)
 
 void FLLMThreadState::ClearAllocTypeAmounts()
 {
+	// Caller is responsible for holding a lock on AllocationsLock.
+
 	for (int32 Index = 0; Index < static_cast<int32>(ELLMAllocType::Count); ++Index)
 	{
 		AllocTypeAmounts[Index] = 0;
@@ -6230,7 +6242,7 @@ void FLLMCsvWriter::Publish(FLowLevelMemTracker& LLMRef, const FTrackerTagSizeMa
 {
 	const double Now = FPlatformTime::Seconds();
 
-	if ((FLowLevelMemTracker::Get().bPublishSingleFrame == false) && 
+	if ((FLowLevelMemTracker::Get().bPublishSingleFrame == false) &&
 		(Now - LastWriteTime < (double)CVarLLMWriteInterval.GetValueOnAnyThread()))
 	{
 		return;
@@ -6261,7 +6273,7 @@ const TCHAR* FLLMCsvWriter::GetTrackerCsvName(ELLMTracker InTracker)
 	}
 }
 
-// Archive is a binary stream, so we can't just serialise an FString using <<.
+// Archive is a binary stream, so we can't just serialize an FString using <<.
 void FLLMCsvWriter::Write(FStringView Text)
 {
 	Archive->Serialize(const_cast<ANSICHAR*>(
@@ -6429,7 +6441,7 @@ void FLLMCsvWriter::AddRow(FLowLevelMemTracker& LLMRef, const FTrackerTagSizeMap
 	else if (!bRegisteredFlushDelegate)
 	{
 		// If we're not flushing every row, lazily register flush delegates to ensure we flush on a crash
-		// Note: we intentionally leak this since we can't clean it up safely 
+		// Note: we intentionally leak this since we can't clean it up safely
 		FCoreDelegates::OnHandleSystemError.AddRaw(this, &FLLMCsvWriter::Flush, true);
 		FCoreDelegates::GetOutOfMemoryDelegate().AddRaw(this, &FLLMCsvWriter::Flush, true);
 		FCoreDelegates::ApplicationWillEnterBackgroundDelegate.AddRaw(this, &FLLMCsvWriter::Flush, false); // When you close the app on some platforms it is suspended rather than terminated
@@ -6445,7 +6457,7 @@ FLLMTraceWriter::FLLMTraceWriter()
 }
 
 inline void FLLMTraceWriter::SetTracker(ELLMTracker InTracker)
-{ 
+{
 	Tracker = InTracker;
 }
 
@@ -6488,11 +6500,11 @@ void FLLMTraceWriter::Publish(FLowLevelMemTracker& LLMRef, const FTrackerTagSize
 	if (!bTagSetSpecSent && Tracker == ELLMTracker::Default)
 	{
 		bTagSetSpecSent = true;
-		FAnsiStringView TagSetNames[] = { 
-			"System", 
-			"Assets", 
-			"AssetClasses", 
-			"UObjectClasses" 
+		FAnsiStringView TagSetNames[] = {
+			"System",
+			"Assets",
+			"AssetClasses",
+			"UObjectClasses"
 		};
 		static_assert((int)ELLMTagSet::Assets == 1, "");
 		static_assert((int)ELLMTagSet::AssetClasses == 2, "");
