HazardPointer: Use FPlatformRecursiveMutex to avoid a dependency on ParkingLot

HazardPointer is used by oversubscription, which is depended on by ParkingLot.

diff --git a/Engine/Source/Runtime/Core/Private/Experimental/Containers/HazardPointer.cpp b/Engine/Source/Runtime/Core/Private/Experimental/Containers/HazardPointer.cpp
--- a/Engine/Source/Runtime/Core/Private/Experimental/Containers/HazardPointer.cpp
+++ b/Engine/Source/Runtime/Core/Private/Experimental/Containers/HazardPointer.cpp
@@ -3,8 +3,8 @@
 #include "Experimental/Containers/HazardPointer.h"
 
 #include "Algo/Sort.h"
+#include "Async/UniqueLock.h"
 #include "Misc/App.h"
-#include "Misc/ScopeLock.h"
 
 FHazardPointerCollection::~FHazardPointerCollection()
 {
@@ -25,7 +25,7 @@ FHazardPointerCollection::FHazardRecord* FHazardPointerCollection::Grow()
 {
 	uint32 TotalNumBeforeLock = TotalNumHazardRecords;
 	//No empty Entry found: create new
-	FScopeLock Lock(&HazardRecordBlocksCS);
+	UE::TUniqueLock Lock(HazardRecordBlocksCS);
 	if (TotalNumHazardRecords != TotalNumBeforeLock)
 	{
 		return Acquire<Cached>();
@@ -54,7 +54,7 @@ void FHazardPointerCollection::Delete(const HazardPointer_Impl::FHazardDeleter&
 	FTlsData* Collectables = (FTlsData*)FPlatformTLS::GetTlsValue(CollectablesTlsSlot);
 	if (Collectables == nullptr)
 	{
-		FScopeLock Lock(&AllTlsVariablesCS);
+		UE::TUniqueLock Lock(AllTlsVariablesCS);
 		Collectables = new FTlsData();
 		AllTlsVariables.Add(Collectables);
 		FPlatformTLS::SetTlsValue(CollectablesTlsSlot, Collectables);
diff --git a/Engine/Source/Runtime/Core/Public/Experimental/Containers/HazardPointer.h b/Engine/Source/Runtime/Core/Public/Experimental/Containers/HazardPointer.h
--- a/Engine/Source/Runtime/Core/Public/Experimental/Containers/HazardPointer.h
+++ b/Engine/Source/Runtime/Core/Public/Experimental/Containers/HazardPointer.h
@@ -3,8 +3,8 @@
 #pragma once
 #include "Containers/Array.h"
 #include "CoreMinimal.h"
-#include "HAL/CriticalSection.h"
 #include "HAL/PlatformMemory.h"
+#include "HAL/PlatformMutex.h"
 #include "HAL/PlatformTLS.h"
 #include "HAL/UnrealMemory.h"
 #include "Misc/AssertionMacros.h"
@@ -154,8 +154,8 @@ class FHazardPointerCollection
 
 	FHazardRecordChunk Head;
 
-	FCriticalSection AllTlsVariablesCS;
-	FCriticalSection HazardRecordBlocksCS;
+	UE::FPlatformRecursiveMutex AllTlsVariablesCS;
+	UE::FPlatformRecursiveMutex HazardRecordBlocksCS;
 	TArray<FTlsData*> AllTlsVariables;
 	TArray<FHazardRecordChunk*> HazardRecordBlocks;
 
