Fixed FVirtualAllocator::GetStats to include peak virtual memory consumption

diff --git a/Engine/Source/Runtime/Core/Public/HAL/VirtualAllocator.h b/Engine/Source/Runtime/Core/Public/HAL/VirtualAllocator.h
--- a/Engine/Source/Runtime/Core/Public/HAL/VirtualAllocator.h
+++ b/Engine/Source/Runtime/Core/Public/HAL/VirtualAllocator.h
@@ -81,12 +81,18 @@ class FVirtualAllocator
 	}
 protected:
 	size_t SpaceConsumed;
+	size_t SpaceConsumedPeak;
+
 	virtual uint8* AllocNewVM(size_t AlignedSize)
 	{
 		uint8* Result = NextAlloc;
 		check(IsAligned(Result, MaximumAlignment) && IsAligned(AlignedSize, MaximumAlignment));
 		NextAlloc = Result + AlignedSize;
 		SpaceConsumed = NextAlloc - LowAddress;
+		if (SpaceConsumed > SpaceConsumedPeak)
+		{
+			SpaceConsumedPeak = SpaceConsumed;
+		}
 		return Result;
 	}
 
@@ -102,6 +108,7 @@ public:
 		, LinkSize(0)
 		, bBacksMalloc(bInBacksMalloc)
 		, SpaceConsumed(0)
+		, SpaceConsumedPeak(0)
 	{
 		TotalSize = HighAddress - LowAddress;
 		check(LowAddress && HighAddress && LowAddress < HighAddress && IsAligned(LowAddress, MaximumAlignment));
@@ -232,6 +239,7 @@ public:
 		OutStats.MaximumAlignment = MaximumAlignment;
 		OutStats.VMSpaceTotal = HighAddress - LowAddress;
 		OutStats.VMSpaceConsumed = SpaceConsumed;
+		OutStats.VMSpaceConsumedPeak = SpaceConsumedPeak;
 		OutStats.FreeListLinks = LinkSize;
 
 		for (int32 Index = 0; Index < 64; Index++)
