[LLM] * Added support for "LLMEM DUMPTAGS" console command. It would generate a csv file with the list of available LLM tags (and their info: index, id, name, display name, stat names, flags, etc.).

diff --git a/Engine/Source/Runtime/Core/Private/HAL/LowLevelMemTracker.cpp b/Engine/Source/Runtime/Core/Private/HAL/LowLevelMemTracker.cpp
--- a/Engine/Source/Runtime/Core/Private/HAL/LowLevelMemTracker.cpp
+++ b/Engine/Source/Runtime/Core/Private/HAL/LowLevelMemTracker.cpp
@@ -1063,7 +1063,7 @@ FRWLock& GetEnableStateLock()
 	return Lock;
 }
 
-}
+} // namespace UE::LLMPrivate
 
 bool FLowLevelMemTracker::TryEnterEnabled(FLowLevelMemTracker::FEnableStateScopeLock& ScopeLock)
 {
@@ -1084,6 +1084,16 @@ bool FLowLevelMemTracker::TryEnterEnabled(FLowLevelMemTracker::FEnableStateScope
 }
 #endif // !PLATFORM_HAS_MULTITHREADED_PREMAIN
 
+namespace UE::LLM
+{
+
+FString GetLLMProfilingDir()
+{
+	return FPaths::ProfilingDir() + TEXT("LLM/");
+}
+
+} // namespace UE::LLM
+
 void FLowLevelMemTracker::DumpToLog(EDumpFormat DumpFormat, FOutputDevice* OutputDevice, UE::LLM::ESizeParams SizeParams, ELLMTagSet TagSet)
 {
 	if (!IsEnabled())
@@ -1400,8 +1410,8 @@ public:
 	// AllocationsLock must be held around read/write of Allocations and AllocTypeAmounts, because these are read/write
 	// from the Update thread. All other data on FLLMThreadState are accessed only from the thread that owns the
 	// FLLMThreadState and do not require a lock.
-	// AllocationLock is also used to guarantee our contract provided to FLLMTracker::LockAllThreadTags: We must not
-	// dereference any FTagData* held by the FLLMThreadState unless we have entered AllocationLock.
+	// AllocationsLock is also used to guarantee our contract provided to FLLMTracker::LockAllThreadTags: We must not
+	// dereference any FTagData* held by the FLLMThreadState unless we have entered AllocationsLock.
 	FConstTagDataArray TagStack[static_cast<int32>(ELLMTagSet::Max)];
 	FThreadTagSizeMap Allocations;
 	FCriticalSection AllocationsLock;
@@ -1451,6 +1461,10 @@ public:
 	// if the memory information can't be retrieved (likely because it's not on a supported platform)
 	bool DumpForkedAllocationInfo();
 
+	// Dump the list of all registered LLM tags and their properties to a CSV file, in the LLM's profiling directory
+	// (Saved/Profiling/LLM/); see UE::LLM::GetLLMProfilingDir().
+	bool DumpTags();
+
 	void PublishStats(UE::LLM::ESizeParams SizeParam);
 	void PublishCsv(UE::LLM::ESizeParams SizeParam);
 	void PublishTrace(UE::LLM::ESizeParams SizeParam);
@@ -2596,6 +2610,11 @@ bool FLowLevelMemTracker::Exec(const TCHAR* Cmd, FOutputDevice& Ar)
 {
 	if (!IsEnabled())
 	{
+		if (FParse::Command(&Cmd, TEXT("LLMEM")))
+		{
+			UE_LOG(LogHAL, Warning, TEXT("Cannot execute LLMEM command! LLM is not enabled."));
+			return true;
+		}
 		return false;
 	}
 
@@ -2670,6 +2689,14 @@ bool FLowLevelMemTracker::Exec(const TCHAR* Cmd, FOutputDevice& Ar)
 			}
 			bCapturedSizeSnapshot = false;
 		}
+		else if (FParse::Command(&Cmd, TEXT("DUMPTAGS")))
+		{
+			if (GetTracker(ELLMTracker::Default)->DumpTags() == false)
+			{
+				FPlatformMisc::LowLevelOutputDebugString(TEXT("Failed to dump tags info\n"));
+				return true;
+			}
+		}
 		return true;
 	}
 
@@ -5618,11 +5645,12 @@ bool FLLMTracker::DumpForkedAllocationInfo()
 			);
 	}
 
+	// This should be UE::LLM::GetLLMProfilingDir(), i.e. "Saved/Profiling/LLM", but kept to "Saved/LLM" for legacy reasons.
 	FString SavedDir = FPaths::ProjectSavedDir() / TEXT("LLM");
 
 	uint16 ForkId = FForkProcessHelper::GetForkedChildProcessIndex();
-
 	FString ForkString = ForkId == 0 ? FString("Parent") : FString::Printf(TEXT("Child%d"), ForkId);
+
 	const TCHAR* TrackerString = TEXT("Default");
 	if (Tracker == ELLMTracker::Platform)
 	{
@@ -5630,19 +5658,81 @@ bool FLLMTracker::DumpForkedAllocationInfo()
 	}
 	static_assert((int)ELLMTracker::Max == 2, "Add other tracker type strings here");
 
-	FString UniqueFilename = SavedDir / FString::Printf(TEXT("LLM_PrivateShared_Tracker%s_%s.csv"), TrackerString, *ForkString);
+	FString FileName = FString::Printf(TEXT("LLM_PrivateShared_Tracker%s_%s"), TrackerString, *ForkString);
 
+	FString UniqueFilename = SavedDir / FString::Printf(TEXT("%s.csv"), *FileName);
 	uint32 Counter = 1;
 	while (IFileManager::Get().FileSize(*UniqueFilename) >= 0)
 	{
-		UniqueFilename = SavedDir / *FString::Printf(TEXT("LLM_PrivateShared_Tracker%s_%s_%d.csv"), TrackerString, *ForkString, Counter);
+		UniqueFilename = SavedDir / FString::Printf(TEXT("%s_%d.csv"), *FileName, Counter);
 		Counter++;
 	}
 
 	if (FFileHelper::SaveStringArrayToFile(Lines, *UniqueFilename) == false)
 	{
-		FPlatformMisc::LowLevelOutputDebugStringf(TEXT("Failed to write llm private/shared csv file: %s\n"), *UniqueFilename);
+		UE_LOG(LogHAL, Error, TEXT("Failed to write LLM Private/Shared CSV file: %s\n"), *UniqueFilename);
+		return false;
 	}
+
+	UE_LOG(LogHAL, Log, TEXT("Wrote LLM Private/Shared CSV file: %s"), *UniqueFilename);
+	return true;
+}
+
+bool FLLMTracker::DumpTags()
+{
+	TArray<FString> Lines;
+
+	{
+		UE::TSharedLock Lock(LLMRef.TagDataMutex);
+
+		int32 NumTags = LLMRef.TagDatas->Num();
+		Lines.Reserve(NumTags + 1);
+
+		Lines.Add(TEXT("Index,Id,Name,DisplayName,Parent,StatName,SummaryStatName,ReferenceSource,EnumTag,bHasEnumTag,TagSet,bIsFinishConstructed,bIsReportable,bIsTraceReportable"));
+
+		for (FTagData* TagData : (*LLMRef.TagDatas))
+		{
+			Lines.Add(
+				FString::Printf(TEXT("%d,0x%X,%s,%s,%s,%s,%s,%s,%u,%s,%u,%s,%s,%s"),
+					TagData->GetIndex(),
+					uint64(TagData),
+					*TagData->GetName().ToString(),
+					*TagData->GetDisplayName().ToString(),
+					TagData->IsParentConstructed() ?
+						(TagData->GetParent() ? *FString::Printf(TEXT("0x%X"), uint64(TagData->GetParent())) : TEXT("")) :
+						*TagData->GetParentName().ToString(),
+					*TagData->GetStatName().ToString(),
+					*TagData->GetSummaryStatName().ToString(),
+					ToString(TagData->GetReferenceSource()),
+					uint32(TagData->GetEnumTag()),
+					TagData->HasEnumTag() ? TEXT("true") : TEXT("false"),
+					uint32(TagData->GetTagSet()),
+					TagData->IsFinishConstructed() ? TEXT("true") : TEXT("false"),
+					TagData->IsStatsReportable() ? TEXT("true") : TEXT("false"),
+					TagData->IsTraceReportable() ? TEXT("true") : TEXT("false"))
+				);
+		}
+	}
+
+	FString SavedDir = UE::LLM::GetLLMProfilingDir();
+
+	FString FileName("LLM_Tags");
+
+	FString UniqueFilename = SavedDir / FString::Printf(TEXT("%s.csv"), *FileName);
+	uint32 Counter = 1;
+	while (IFileManager::Get().FileSize(*UniqueFilename) >= 0)
+	{
+		UniqueFilename = SavedDir / FString::Printf(TEXT("%s_%d.csv"), *FileName, Counter);
+		Counter++;
+	}
+
+	if (FFileHelper::SaveStringArrayToFile(Lines, *UniqueFilename) == false)
+	{
+		UE_LOG(LogHAL, Error, TEXT("Failed to write LLM Tags CSV file: %s\n"), *UniqueFilename);
+		return false;
+	}
+
+	UE_LOG(LogHAL, Log, TEXT("Wrote LLM Tags CSV file: %s"), *UniqueFilename);
 	return true;
 }
 
@@ -6288,7 +6378,7 @@ bool FLLMCsvWriter::CreateArchive()
 	}
 
 	// Create the csv file.
-	FString Directory = FPaths::ProfilingDir() + TEXT("LLM/");
+	FString Directory = UE::LLM::GetLLMProfilingDir();
 	IFileManager::Get().MakeDirectory(*Directory, true);
 
 	const TCHAR* TrackerName = GetTrackerCsvName(Tracker);
diff --git a/Engine/Source/Runtime/Core/Public/HAL/LowLevelMemTracker.h b/Engine/Source/Runtime/Core/Public/HAL/LowLevelMemTracker.h
--- a/Engine/Source/Runtime/Core/Public/HAL/LowLevelMemTracker.h
+++ b/Engine/Source/Runtime/Core/Public/HAL/LowLevelMemTracker.h
@@ -137,12 +137,12 @@ enum class ELLMTagSet : uint8
 	Max, // note: see FLowLevelMemTracker::ShouldReduceThreads and IsAssetTagForAssets if you add any asset-style tagsets
 };
 
-/*
- * Size parameter flags used when requesting the size of the tracked tag data.
- */
 namespace UE::LLM
 {
 
+/*
+ * Size parameter flags used when requesting the size of the tracked tag data.
+ */
 enum class ESizeParams : uint8
 {
 	Default = 0,
@@ -153,6 +153,11 @@ enum class ESizeParams : uint8
 
 ENUM_CLASS_FLAGS(ESizeParams);
 
+/*
+ * Returns the preferred directory path for LLM profiling output files (i.e.: Saved/Profiling/LLM/).
+ */
+FString GetLLMProfilingDir();
+
 } // UE::LLM
 
 // Do not add to these macros. Please use the LLM_DECLARE_TAG family of macros below to create new tags.
