Core: Fixed WarnFilterScope for structured logging and added CHECK_LOG_SCOPE for common cases

diff --git a/Engine/Source/Runtime/Core/Private/Tests/WarnFilterScope.cpp b/Engine/Source/Runtime/Core/Private/Tests/WarnFilterScope.cpp
--- a/Engine/Source/Runtime/Core/Private/Tests/WarnFilterScope.cpp
+++ b/Engine/Source/Runtime/Core/Private/Tests/WarnFilterScope.cpp
@@ -1,16 +1,17 @@
 // Copyright Epic Games, Inc. All Rights Reserved.
 
 #include "Tests/WarnFilterScope.h"
+#include "Logging/StructuredLog.h"
 #include "Misc/FeedbackContext.h"
+#include "Misc/StringBuilder.h"
 
 struct FFilterFeedback : public FFeedbackContext
 {
 
-	FFilterFeedback(FFeedbackContext* OldFeedback, TFunction<bool(const TCHAR* Message, ELogVerbosity::Type Verbosity, const FName& Category)> Handler)
+	FFilterFeedback(FFeedbackContext* OldFeedback, TUniqueFunction<bool(const TCHAR* Message, ELogVerbosity::Type Verbosity, const FName& Category)> LogHandler)
 		: OldFeedback(OldFeedback)
-		, Handler(Handler)
+		, Handler(MoveTemp(LogHandler))
 	{
-
 	}
 
 	virtual void Serialize(const TCHAR* V, ELogVerbosity::Type Verbosity, const FName& Category) override
@@ -20,6 +21,7 @@ struct FFilterFeedback : public FFeedbackContext
 			OldFeedback->Serialize(V, Verbosity, Category);
 		}
 	}
+
 	virtual void Serialize(const TCHAR* V, ELogVerbosity::Type Verbosity, const FName& Category, double Time) override
 	{
 		if (!Handler(V, Verbosity, Category))
@@ -28,13 +30,23 @@ struct FFilterFeedback : public FFeedbackContext
 		}
 	}
 
+	virtual void SerializeRecord(const UE::FLogRecord& Record) override
+	{
+		TStringBuilder<256> Message;
+		Record.FormatMessageTo(Message);
+		if (!Handler(*Message, Record.GetVerbosity(), Record.GetCategory()))
+		{
+			OldFeedback->SerializeRecord(Record);
+		}
+	}
+
 	FFeedbackContext* OldFeedback;
-	TFunction<bool(const TCHAR* Message, ELogVerbosity::Type Verbosity, const FName& Category)> Handler;
+	TUniqueFunction<bool(const TCHAR* Message, ELogVerbosity::Type Verbosity, const FName& Category)> Handler;
 };
 
-FWarnFilterScope::FWarnFilterScope(TFunction<bool(const TCHAR* Message, ELogVerbosity::Type Verbosity, const FName& Category)> LogHandler)
+FWarnFilterScope::FWarnFilterScope(TUniqueFunction<bool(const TCHAR* Message, ELogVerbosity::Type Verbosity, const FName& Category)> LogHandler)
 	: OldWarn(GWarn)
-	, Feedback(new FFilterFeedback(OldWarn, LogHandler))
+	, Feedback(new FFilterFeedback(OldWarn, MoveTemp(LogHandler)))
 {
 	GWarn = Feedback;
 }
@@ -43,4 +55,4 @@ FWarnFilterScope::~FWarnFilterScope()
 {
 	GWarn = OldWarn;
 	delete Feedback;
-}
\ No newline at end of file
+}
diff --git a/Engine/Source/Runtime/Core/Public/Tests/WarnFilterScope.h b/Engine/Source/Runtime/Core/Public/Tests/WarnFilterScope.h
--- a/Engine/Source/Runtime/Core/Public/Tests/WarnFilterScope.h
+++ b/Engine/Source/Runtime/Core/Public/Tests/WarnFilterScope.h
@@ -1,24 +1,63 @@
 // Copyright Epic Games, Inc. All Rights Reserved.
 
 #pragma once
+
+#include "Containers/StringView.h"
 #include "HAL/Platform.h"
+#include "Misc/ScopeExit.h"
 #include "Templates/Function.h"
+#include "UObject/NameTypes.h"
 
 class FFeedbackContext;
-
 struct FFilterFeedback;
 
-///@brief Scope that captures warning log messages
-///necessary to filter out warning message from appearing the log file as Horde will flag the warnings
+/**
+ * Scope that captures and optionally filters display/warning/error log messages.
+ *
+ * Filtering expected warnings and errors during tests is necessary to avoid Horde flagging them as failures.
+ */
 struct FWarnFilterScope
 {
-	/// @param LogHandler function to filter log messages. returning true filters message from log file
-	FWarnFilterScope(TFunction<bool(const TCHAR* Message, ELogVerbosity::Type Verbosity, const FName& Category)> LogHandler);
+	/**
+	 * @param LogHandler Function to handle log messages. Returning true filters the message from the log output.
+	 */
+	explicit FWarnFilterScope(TUniqueFunction<bool(const TCHAR* Message, ELogVerbosity::Type Verbosity, const FName& Category)> LogHandler);
+
 	FWarnFilterScope(const FWarnFilterScope&) = delete;
 	FWarnFilterScope& operator=(const FWarnFilterScope&) = delete;
 
 	~FWarnFilterScope();
+
 private:
 	FFeedbackContext* OldWarn;
 	FFilterFeedback* Feedback;
-};
\ No newline at end of file
+};
+
+/**
+ * Asserts that this exact log occurs within the scope of this statement.
+ *
+ * Supported by the LowLevelTests style of test code.
+ */
+#define CHECK_LOG_SCOPE(ExpectedCategory, ExpectedVerbosity, ExpectedMessage) \
+	static_assert( \
+		::ELogVerbosity::ExpectedVerbosity == ::ELogVerbosity::Error || \
+		::ELogVerbosity::ExpectedVerbosity == ::ELogVerbosity::Warning || \
+		::ELogVerbosity::ExpectedVerbosity == ::ELogVerbosity::Display, \
+		"Only Error/Warning/Display verbosity is supported by CHECK_LOG_SCOPE."); \
+	bool UE_JOIN(bLogScopeLogged, __LINE__) = false; \
+	::FWarnFilterScope UE_JOIN(LogScope, __LINE__)([&UE_JOIN(bLogScopeLogged, __LINE__)](const TCHAR* Message, ::ELogVerbosity::Type Verbosity, const FName& Category) \
+	{ \
+		if (Category == #ExpectedCategory && Verbosity == ::ELogVerbosity::ExpectedVerbosity && Message == TEXTVIEW(ExpectedMessage)) \
+		{ \
+			UE_JOIN(bLogScopeLogged, __LINE__) = true; \
+			return true; \
+		} \
+		return false; \
+	}); \
+	ON_SCOPE_EXIT \
+	{ \
+		if (!UE_JOIN(bLogScopeLogged, __LINE__)) \
+		{ \
+			FAIL_CHECK("Missing expected log message. Inspect the log for similar messages. Expected:\n" #ExpectedCategory ": " #ExpectedVerbosity ": " ExpectedMessage); \
+		} \
+	}
