Add UE_CLANG_SUPPRESS macro to be able to disregard clang static analyzer false positives
Apply it to resolve analyzer misanalysis reporting Address of stack memory associated with local variable 'ExportedPathToObjectCreated' is still referred to by the caller variable 'Params' upon returning to the caller.  This will be a dangling reference

diff --git a/Engine/Source/Editor/UnrealEd/Private/EditorObject.cpp b/Engine/Source/Editor/UnrealEd/Private/EditorObject.cpp
--- a/Engine/Source/Editor/UnrealEd/Private/EditorObject.cpp
+++ b/Engine/Source/Editor/UnrealEd/Private/EditorObject.cpp
@@ -993,6 +993,8 @@ const TCHAR* ImportObjectProperties( FImportObjectParams& InParams )
 		InParams.ObjectRemapper = nullptr;
 	}
 
+	// Suppress clang analyzer false positive about InParams.ObjectRemapper retaining reference to ExportedPathToObjectCreated
+	UE_CLANG_SUPPRESS
 	return NewSourceText;
 }
 
diff --git a/Engine/Source/Runtime/Core/Public/Clang/ClangPlatformCompilerPreSetup.h b/Engine/Source/Runtime/Core/Public/Clang/ClangPlatformCompilerPreSetup.h
--- a/Engine/Source/Runtime/Core/Public/Clang/ClangPlatformCompilerPreSetup.h
+++ b/Engine/Source/Runtime/Core/Public/Clang/ClangPlatformCompilerPreSetup.h
@@ -6,6 +6,10 @@
 
 #ifdef __clang__
 
+#ifndef UE_CLANG_SUPPRESS
+	#define UE_CLANG_SUPPRESS [[clang::suppress]]
+#endif
+
 #pragma clang diagnostic warning "-Wdeprecated-declarations"
 
 #ifndef PRAGMA_DISABLE_DEPRECATION_WARNINGS
diff --git a/Engine/Source/Runtime/Core/Public/GenericPlatform/GenericPlatformCompilerPreSetup.h b/Engine/Source/Runtime/Core/Public/GenericPlatform/GenericPlatformCompilerPreSetup.h
--- a/Engine/Source/Runtime/Core/Public/GenericPlatform/GenericPlatformCompilerPreSetup.h
+++ b/Engine/Source/Runtime/Core/Public/GenericPlatform/GenericPlatformCompilerPreSetup.h
@@ -4,6 +4,10 @@
 
 // HEADER_UNIT_SKIP - Cause cirular dependencies for header units
 
+#ifndef UE_CLANG_SUPPRESS
+	#define UE_CLANG_SUPPRESS
+#endif
+
 #ifndef PRAGMA_DISABLE_DEPRECATION_WARNINGS
 	#define PRAGMA_DISABLE_DEPRECATION_WARNINGS
 #endif
