Fix FAutoConsoleVariableSink referencing released memory in most cases

diff --git a/Engine/Source/Runtime/Core/Public/HAL/IConsoleManager.h b/Engine/Source/Runtime/Core/Public/HAL/IConsoleManager.h
--- a/Engine/Source/Runtime/Core/Public/HAL/IConsoleManager.h
+++ b/Engine/Source/Runtime/Core/Public/HAL/IConsoleManager.h
@@ -1401,6 +1401,13 @@ struct TGuardConsoleVariable
 class FAutoConsoleVariableSink
 {
 public:
+	/** Constructor, saves the argument for future removal from the console variable system **/
+	FAutoConsoleVariableSink(FConsoleCommandDelegate&& InCommand)
+		: Command(MoveTemp(InCommand))
+	{
+		Handle = IConsoleManager::Get().RegisterConsoleVariableSink_Handle(Command);
+	}
+
 	/** Constructor, saves the argument for future removal from the console variable system **/
 	FAutoConsoleVariableSink(const FConsoleCommandDelegate& InCommand)
 		: Command(InCommand)
@@ -1413,7 +1420,7 @@ public:
 //disabled for now, destruction order makes this not always working		IConsoleManager::Get().UnregisterConsoleVariableSink_Handle(Handle);
 	}
 
-	const FConsoleCommandDelegate& Command;
+	FConsoleCommandDelegate Command;
 	FConsoleVariableSinkHandle Handle;
 };
 
