Fixed potentially calling GetSetKey on an invalid element in TCompactSet and TSparseSet's Add/FindOrAdd.

#tests PIE with and without compact set enabled

diff --git a/Engine/Source/Runtime/Core/Public/Containers/CompactSet.h.inl b/Engine/Source/Runtime/Core/Public/Containers/CompactSet.h.inl
--- a/Engine/Source/Runtime/Core/Public/Containers/CompactSet.h.inl
+++ b/Engine/Source/Runtime/Core/Public/Containers/CompactSet.h.inl
@@ -651,11 +651,11 @@ public:
 	 */
 	UE_FORCEINLINE_HINT FSetElementId Add(const ElementType&  InElement, bool* bIsAlreadyInSetPtr = nullptr)
 	{
-		return EmplaceByHash(KeyFuncs::GetKeyHash(KeyFuncs::GetSetKey(InElement)), InElement, bIsAlreadyInSetPtr);
+		return Emplace(InElement, bIsAlreadyInSetPtr);
 	}
 	UE_FORCEINLINE_HINT FSetElementId Add(ElementType && InElement, bool* bIsAlreadyInSetPtr = nullptr)
 	{
-		return EmplaceByHash(KeyFuncs::GetKeyHash(KeyFuncs::GetSetKey(InElement)), MoveTempIfPossible(InElement), bIsAlreadyInSetPtr);
+		return Emplace(MoveTempIfPossible(InElement), bIsAlreadyInSetPtr);
 	}
 
 	/**
@@ -808,7 +808,8 @@ public:
 	}
 	UE_FORCEINLINE_HINT ElementType& FindOrAdd(InElementType && InElement, bool* bIsAlreadyInSetPtr = nullptr)
 	{
-		return FindOrAddByHash(KeyFuncs::GetKeyHash(KeyFuncs::GetSetKey(InElement)), MoveTempIfPossible(InElement), bIsAlreadyInSetPtr);
+		const uint32 KeyHash = KeyFuncs::GetKeyHash(KeyFuncs::GetSetKey(InElement));
+		return FindOrAddByHash(KeyHash, MoveTempIfPossible(InElement), bIsAlreadyInSetPtr);
 	}
 
 	/**
diff --git a/Engine/Source/Runtime/Core/Public/Containers/SparseSet.h.inl b/Engine/Source/Runtime/Core/Public/Containers/SparseSet.h.inl
--- a/Engine/Source/Runtime/Core/Public/Containers/SparseSet.h.inl
+++ b/Engine/Source/Runtime/Core/Public/Containers/SparseSet.h.inl
@@ -428,7 +428,8 @@ public:
 	}
 	UE_FORCEINLINE_HINT ElementType& FindOrAdd(InElementType&& InElement, bool* bIsAlreadyInSetPtr = nullptr)
 	{
-		return FindOrAddByHash(KeyFuncs::GetKeyHash(KeyFuncs::GetSetKey(InElement)), MoveTempIfPossible(InElement), bIsAlreadyInSetPtr);
+		const uint32 KeyHash = KeyFuncs::GetKeyHash(KeyFuncs::GetSetKey(InElement));
+		return FindOrAddByHash(KeyHash, MoveTempIfPossible(InElement), bIsAlreadyInSetPtr);
 	}
 
 	/**
