Added a macro for the virtual projection workaround on MSVC.

diff --git a/Engine/Source/Runtime/Core/Public/Algo/Impl/BinaryHeap.h b/Engine/Source/Runtime/Core/Public/Algo/Impl/BinaryHeap.h
--- a/Engine/Source/Runtime/Core/Public/Algo/Impl/BinaryHeap.h
+++ b/Engine/Source/Runtime/Core/Public/Algo/Impl/BinaryHeap.h
@@ -8,6 +8,13 @@
 
 #include <type_traits>
 
+// Workaround for a codegen bug that was discovered related to member function pointer projections that are virtual
+#if defined(_MSC_FULL_VER) && _MSC_FULL_VER < 194134123
+	#define UE_VIRUTAL_PROJECTION_WORKAROUND 1
+#else
+	#define UE_VIRUTAL_PROJECTION_WORKAROUND 0
+#endif
+
 namespace AlgoImpl
 {
 	/**
@@ -58,8 +65,7 @@ namespace AlgoImpl
 	template <typename IteratorType, typename IndexType, typename ProjectionType, typename PredicateType>
 	inline void HeapSiftDown(IteratorType Heap, IndexType Index, const IndexType Count, const ProjectionType& InProj, const PredicateType& Predicate)
 	{
-// Workaround for a codegen bug that was discovered related to member function pointer projections that are virtual
-#if defined(_MSC_FULL_VER) && _MSC_FULL_VER < 194134123
+#if UE_VIRUTAL_PROJECTION_WORKAROUND
 		const ProjectionType& Proj = InProj;
 #else
 		auto&& Proj = Projection(InProj);
@@ -72,15 +78,14 @@ namespace AlgoImpl
 			IndexType MinChildIndex = LeftChildIndex;
 			if (RightChildIndex < Count)
 			{
-// Workaround for a codegen bug that was discovered related to member function pointer projections that are virtual
-#if defined(_MSC_FULL_VER) && _MSC_FULL_VER < 194134123
+#if UE_VIRUTAL_PROJECTION_WORKAROUND
 				MinChildIndex = Predicate( Invoke(Proj, Heap[LeftChildIndex]), Invoke(Proj, Heap[RightChildIndex]) ) ? LeftChildIndex : RightChildIndex;
 #else
 				MinChildIndex = Predicate( Proj(Heap[LeftChildIndex]), Proj(Heap[RightChildIndex]) ) ? LeftChildIndex : RightChildIndex;
 #endif
 			}
 
-#if defined(_MSC_FULL_VER) && _MSC_FULL_VER < 194134123
+#if UE_VIRUTAL_PROJECTION_WORKAROUND
 			if (!Predicate( Invoke(Proj, Heap[MinChildIndex]), Invoke(Proj, Heap[Index]) ))
 #else
 			if (!Predicate( Proj(Heap[MinChildIndex]), Proj(Heap[Index]) ))
@@ -108,8 +113,7 @@ namespace AlgoImpl
 	template <class IteratorType, typename IndexType, typename ProjectionType, class PredicateType>
 	inline IndexType HeapSiftUp(IteratorType Heap, IndexType RootIndex, IndexType NodeIndex, const ProjectionType& InProj, const PredicateType& Predicate)
 	{
-// Workaround for a codegen bug that was discovered related to member function pointer projections that are virtual
-#if defined(_MSC_FULL_VER) && _MSC_FULL_VER < 194134123
+#if UE_VIRUTAL_PROJECTION_WORKAROUND
 		const ProjectionType& Proj = InProj;
 #else
 		auto&& Proj = Projection(InProj);
@@ -117,8 +121,7 @@ namespace AlgoImpl
 		while (NodeIndex > RootIndex)
 		{
 			IndexType ParentIndex = HeapGetParentIndex(NodeIndex);
-// Workaround for a codegen bug that was discovered related to member function pointer projections that are virtual
-#if defined(_MSC_FULL_VER) && _MSC_FULL_VER < 194134123
+#if UE_VIRUTAL_PROJECTION_WORKAROUND
 			if (!Predicate( Invoke(Proj, Heap[NodeIndex]), Invoke(Proj, Heap[ParentIndex]) ))
 #else
 			if (!Predicate( Proj(Heap[NodeIndex]), Proj(Heap[ParentIndex]) ))
diff --git a/Engine/Source/Runtime/Core/Public/Algo/IntroSort.h b/Engine/Source/Runtime/Core/Public/Algo/IntroSort.h
--- a/Engine/Source/Runtime/Core/Public/Algo/IntroSort.h
+++ b/Engine/Source/Runtime/Core/Public/Algo/IntroSort.h
@@ -12,6 +12,9 @@
 
 #include <iterator>
 
+#ifndef UE_VIRUTAL_PROJECTION_WORKAROUND
+	#error "Expected UE_VIRUTAL_PROJECTION_WORKAROUND from BinaryHeap.h to be defined."
+#endif
 
 namespace AlgoImpl
 {
@@ -66,8 +69,7 @@ namespace AlgoImpl
 					IteratorType Item;
 					for( Max=Current.Min, Item=Current.Min+1; Item<=Current.Max; ++Item )
 					{
-// Workaround for a codegen bug that was discovered related to member function pointer projections that are virtual
-#if defined(_MSC_FULL_VER) && _MSC_FULL_VER < 194134123
+#if UE_VIRUTAL_PROJECTION_WORKAROUND
 						if( Invoke( Predicate, Invoke( Proj, *Max ), Invoke( Proj, *Item ) ) )
 #else
 						if( Predicate( Proj( *Max ), Proj( *Item ) ) )
@@ -89,8 +91,7 @@ namespace AlgoImpl
 				Inner.Max = Current.Max+1;
 				for( ; ; )
 				{
-// Workaround for a codegen bug that was discovered related to member function pointer projections that are virtual
-#if defined(_MSC_FULL_VER) && _MSC_FULL_VER < 194134123
+#if UE_VIRUTAL_PROJECTION_WORKAROUND
 					while( ++Inner.Min<=Current.Max && !Invoke( Predicate, Invoke( Proj, *Current.Min ), Invoke( Proj, *Inner.Min ) ) );
 					while( --Inner.Max> Current.Min && !Invoke( Predicate, Invoke( Proj, *Inner.Max ), Invoke( Proj, *Current.Min ) ) );
 #else
@@ -177,8 +178,7 @@ namespace Algo
 	template <typename RangeType, typename ProjectionType>
 	UE_REWRITE void IntroSortBy(RangeType&& Range, ProjectionType Proj)
 	{
-// Workaround for a codegen bug that was discovered related to member function pointer projections that are virtual
-#if defined(_MSC_FULL_VER) && _MSC_FULL_VER < 194134123
+#if UE_VIRUTAL_PROJECTION_WORKAROUND
 		AlgoImpl::IntroSortInternal(std::begin(Range), std::end(Range) - std::begin(Range), MoveTemp(Proj), TLess<>());
 #else
 		AlgoImpl::IntroSortInternal(std::begin(Range), std::end(Range) - std::begin(Range), Projection(MoveTemp(Proj)), TLess<>());
@@ -195,8 +195,7 @@ namespace Algo
 	template <typename RangeType, typename ProjectionType, typename PredicateType>
 	UE_REWRITE void IntroSortBy(RangeType&& Range, ProjectionType Proj, PredicateType Predicate)
 	{
-// Workaround for a codegen bug that was discovered related to member function pointer projections that are virtual
-#if defined(_MSC_FULL_VER) && _MSC_FULL_VER < 194134123
+#if UE_VIRUTAL_PROJECTION_WORKAROUND
 		AlgoImpl::IntroSortInternal(std::begin(Range), std::end(Range) - std::begin(Range), MoveTemp(Proj), MoveTemp(Predicate));
 #else
 		AlgoImpl::IntroSortInternal(std::begin(Range), std::end(Range) - std::begin(Range), Projection(MoveTemp(Proj)), Projection(MoveTemp(Predicate)));
