[TaskGraphTests]
- Fix test to account for retraction

diff --git a/Engine/Source/Runtime/Core/Tests/Async/TaskGraphTest.cpp b/Engine/Source/Runtime/Core/Tests/Async/TaskGraphTest.cpp
--- a/Engine/Source/Runtime/Core/Tests/Async/TaskGraphTest.cpp
+++ b/Engine/Source/Runtime/Core/Tests/Async/TaskGraphTest.cpp
@@ -952,8 +952,18 @@ PRAGMA_ENABLE_DEPRECATION_WARNINGS
 		}
 
 		{	// check ref count for named thread tasks
-			FGraphEventRef LocalQueueTask = FFunctionGraphTask::CreateAndDispatchWhenReady([] {}, TStatId{}, nullptr, ENamedThreads::GameThread_Local);
+			FGraphEventRef LocalQueueTask = FFunctionGraphTask::CreateAndDispatchWhenReady(
+				[] 
+				{
+					check(IsInGameThread());  
+				}, TStatId{}, nullptr, ENamedThreads::GameThread_Local);
+
 			LocalQueueTask->Wait(ENamedThreads::GameThread_Local);
+			check(LocalQueueTask->IsComplete());
+			check(LocalQueueTask.GetRefCount() >= 1);
+			// Now that the wait above can perform retraction, it will not necessarily pump the game thread queue and will just leave an empty task with a refcount on it
+			// so pump it ourself to verify that once the game thread queue is empty, the refcount should have been cleaned up as well.
+			FTaskGraphInterface::Get().ProcessThreadUntilIdle(ENamedThreads::GameThread_Local);
 			check(LocalQueueTask.GetRefCount() == 1);
 		}
 
@@ -961,20 +971,40 @@ PRAGMA_ENABLE_DEPRECATION_WARNINGS
 		{	// a particular real-life case that doesn't work in the old TaskGraph if run in single-threaded mode.
 			// the culprit is that when a task is waited for, in single-threaded mode the queue it was pushed to is executed. 
 			// Here the (local queue) task  depends on a task that is not in the same queue and so it doesn't get executed
+			std::atomic<bool> bWorkerDone = false;
+			bool bHasRun = false;
+			FGraphEventRef AnyTask = FFunctionGraphTask::CreateAndDispatchWhenReady([&bWorkerDone] { bWorkerDone = true; });
+			FGraphEventRef LocalQueueTask = FFunctionGraphTask::CreateAndDispatchWhenReady(
+				[&bHasRun]
+				{
+					check(bHasRun == false);
+					bHasRun = true;
+					check(IsInGameThread());
+				}, TStatId{}, AnyTask, ENamedThreads::GameThread_Local);
 
-			FGraphEventRef AnyTask = FFunctionGraphTask::CreateAndDispatchWhenReady([] {});
-			FGraphEventRef LocalQueueTask = FFunctionGraphTask::CreateAndDispatchWhenReady([] {}, TStatId{}, AnyTask, ENamedThreads::GameThread_Local);
+			check(bHasRun == false);
 			LocalQueueTask->Wait(ENamedThreads::GameThread_Local);
+			check(bWorkerDone);
+			check(bHasRun);
+			check(LocalQueueTask->IsComplete());
+			check(LocalQueueTask.GetRefCount() >= 1);
+			// Now that the wait above can perform retraction, it will not necessarily pump the game thread queue and will just leave an empty task with a refcount on it
+			// so pump it ourself to verify that once the game thread queue is empty, the refcount should have been cleaned up as well.
+			// The reason we do it in a loop, is that the task on the game thread could be queued after the wait returns from a successful retraction.
+			while (LocalQueueTask.GetRefCount() > 1)
+			{
+				FTaskGraphInterface::Get().ProcessThreadUntilIdle(ENamedThreads::GameThread_Local);
+			}
 			check(LocalQueueTask.GetRefCount() == 1);
 		}
 
 		{	// launch a GT task, then an any-thread task that depends on it. wait for the any-thread task. this was a deadlock on the new frontend
-			FGraphEventRef GTTask = FFunctionGraphTask::CreateAndDispatchWhenReady([] { IsInGameThread(); }, TStatId{}, nullptr, ENamedThreads::GameThread); //-V530
+			FGraphEventRef GTTask = FFunctionGraphTask::CreateAndDispatchWhenReady([] { check(IsInGameThread()); }, TStatId{}, nullptr, ENamedThreads::GameThread); //-V530
 			FGraphEventRef AnyThreadTask = FFunctionGraphTask::CreateAndDispatchWhenReady([] {}, TStatId{}, GTTask);
 			AnyThreadTask->Wait();
 		}
 	}
-
+	
 	// it's fast because tasks are too lightweight and so are executed almost as fast
 	template<int NumTasks>
 	void TestPerfBasic()
