Logging

diff --git a/Engine/Source/Runtime/Core/Private/Modules/ModuleManager.cpp b/Engine/Source/Runtime/Core/Private/Modules/ModuleManager.cpp
--- a/Engine/Source/Runtime/Core/Private/Modules/ModuleManager.cpp
+++ b/Engine/Source/Runtime/Core/Private/Modules/ModuleManager.cpp
@@ -2100,6 +2100,7 @@ void FModuleManager::RegisterStaticallyLinkedModule(const FLazyName InModuleName
 
 void FModuleManager::StartProcessingNewlyLoadedObjects()
 {
+	UE_LOGFMT(LogModuleManager, Log, "StartProcessingNewlyLoadedObjects");
 	// Only supposed to be called once
 	ensure( bCanProcessNewlyLoadedObjects == false );	
 	bCanProcessNewlyLoadedObjects = true;
