Logging: Moved log category macros to LogCategory.h

diff --git a/Engine/Source/Runtime/Core/Public/Logging/LogCategory.h b/Engine/Source/Runtime/Core/Public/Logging/LogCategory.h
--- a/Engine/Source/Runtime/Core/Public/Logging/LogCategory.h
+++ b/Engine/Source/Runtime/Core/Public/Logging/LogCategory.h
@@ -4,6 +4,7 @@
 
 #include "CoreTypes.h"
 #include "Logging/LogVerbosity.h"
+#include "Misc/Build.h"
 #include "UObject/NameTypes.h"
 
 #define UE_API CORE_API
@@ -102,4 +103,70 @@ struct FLogCategory : public FLogCategoryBase
 	}
 };
 
+#if NO_LOGGING
+
+	struct FNoLoggingCategory {};
+
+	#define DECLARE_LOG_CATEGORY_EXTERN(CategoryName, DefaultVerbosity, CompileTimeVerbosity) extern FNoLoggingCategory CategoryName;
+	#define DEFINE_LOG_CATEGORY(CategoryName, ...) FNoLoggingCategory CategoryName;
+	#define DEFINE_LOG_CATEGORY_STATIC(...)
+	#define DECLARE_LOG_CATEGORY_CLASS(...)
+	#define DEFINE_LOG_CATEGORY_CLASS(...)
+
+#else
+
+	#define UE_PRIVATE_DEFINE_LOG_CATEGORY(CategoryName, DefaultVerbosity, CompileTimeVerbosity) \
+		struct FLogCategory##CategoryName : public ::FLogCategory<::ELogVerbosity::DefaultVerbosity, ::ELogVerbosity::CompileTimeVerbosity> \
+		{ \
+			static_assert(::ELogVerbosity::DefaultVerbosity <= ::ELogVerbosity::CompileTimeVerbosity, \
+				"Default verbosity of log category " #CategoryName " cannot be higher than its compile time verbosity."); \
+			UE_FORCEINLINE_HINT FLogCategory##CategoryName() : FLogCategory(TEXT(#CategoryName)) {} \
+		} CategoryName;
+
+	/**
+	 * A macro to declare a logging category as a C++ "extern", usually declared in the header and paired with DEFINE_LOG_CATEGORY in the source. Accessible by all files that include the header.
+	 *
+	 * @param CategoryName   Name of the category, typically prefixed with Log.
+	 * @param DefaultVerbosity   Default runtime verbosity of the category.
+	 * @param CompileTimeVerbosity   Maximum verbosity to compile into the code.
+	 */
+	#define DECLARE_LOG_CATEGORY_EXTERN(CategoryName, DefaultVerbosity, CompileTimeVerbosity) \
+		extern UE_PRIVATE_DEFINE_LOG_CATEGORY(CategoryName, DefaultVerbosity, CompileTimeVerbosity);
+
+	/**
+	 * A macro to define a logging category, usually paired with DECLARE_LOG_CATEGORY_EXTERN from the header.
+	 * @param CategoryName category to define
+	 */
+	#define DEFINE_LOG_CATEGORY(CategoryName) FLogCategory##CategoryName CategoryName;
+
+	/**
+	 * A macro to define a logging category as a C++ "static". This should ONLY be declared in a source file. Only accessible in that single file.
+	 *
+	 * @param CategoryName   Name of the category, typically prefixed with Log.
+	 * @param DefaultVerbosity   Default runtime verbosity of the category.
+	 * @param CompileTimeVerbosity   Maximum verbosity to compile into the code.
+	 */
+	//-V:DEFINE_LOG_CATEGORY_STATIC:501
+	#define DEFINE_LOG_CATEGORY_STATIC(CategoryName, DefaultVerbosity, CompileTimeVerbosity) \
+		static UE_PRIVATE_DEFINE_LOG_CATEGORY(CategoryName, DefaultVerbosity, CompileTimeVerbosity);
+
+	/**
+	 * A macro to declare a logging category as a C++ "class static"
+	 *
+	 * @param CategoryName   Name of the category, typically prefixed with Log.
+	 * @param DefaultVerbosity   Default runtime verbosity of the category.
+	 * @param CompileTimeVerbosity   Maximum verbosity to compile into the code.
+	 */
+	//-V:DECLARE_LOG_CATEGORY_CLASS:501
+	#define DECLARE_LOG_CATEGORY_CLASS(CategoryName, DefaultVerbosity, CompileTimeVerbosity) \
+		static UE_PRIVATE_DEFINE_LOG_CATEGORY(CategoryName, DefaultVerbosity, CompileTimeVerbosity)
+
+	/**
+	 * A macro to define a logging category, usually paired with DECLARE_LOG_CATEGORY_CLASS from the header.
+	 * @param CategoryName category to define
+	 */
+	#define DEFINE_LOG_CATEGORY_CLASS(Class, CategoryName) Class::FLogCategory##CategoryName Class::CategoryName;
+
+#endif // NO_LOGGING
+
 #undef UE_API
diff --git a/Engine/Source/Runtime/Core/Public/Logging/LogMacros.h b/Engine/Source/Runtime/Core/Public/Logging/LogMacros.h
--- a/Engine/Source/Runtime/Core/Public/Logging/LogMacros.h
+++ b/Engine/Source/Runtime/Core/Public/Logging/LogMacros.h
@@ -176,8 +176,6 @@ CORE_API void BasicFatalLogWithProgramCounter(FStaticBasicLogRecordParam Log, co
 
 #if NO_LOGGING
 
-	struct FNoLoggingCategory {};
-	
 	// This will only log Fatal errors
 	#define UE_LOG(CategoryName, Verbosity, Format, ...) \
 	{ \
@@ -222,11 +220,6 @@ CORE_API void BasicFatalLogWithProgramCounter(FStaticBasicLogRecordParam Log, co
 	#define UE_SUPPRESS(...) {}
 	#define UE_GET_LOG_VERBOSITY(...)		(ELogVerbosity::NoLogging)
 	#define UE_SET_LOG_VERBOSITY(...)
-	#define DECLARE_LOG_CATEGORY_EXTERN(CategoryName, DefaultVerbosity, CompileTimeVerbosity) extern FNoLoggingCategory CategoryName;
-	#define DEFINE_LOG_CATEGORY(CategoryName, ...) FNoLoggingCategory CategoryName;
-	#define DEFINE_LOG_CATEGORY_STATIC(...)
-	#define DECLARE_LOG_CATEGORY_CLASS(...)
-	#define DEFINE_LOG_CATEGORY_CLASS(...)
 
 #else
 
@@ -352,54 +345,6 @@ CORE_API void BasicFatalLogWithProgramCounter(FStaticBasicLogRecordParam Log, co
 		} \
 	}
 
-	/** 
-	 * A macro to declare a logging category as a C++ "extern", usually declared in the header and paired with DEFINE_LOG_CATEGORY in the source. Accessible by all files that include the header.
-	 * @param CategoryName, category to declare
-	 * @param DefaultVerbosity, default run time verbosity
-	 * @param CompileTimeVerbosity, maximum verbosity to compile into the code
-	 **/
-	#define DECLARE_LOG_CATEGORY_EXTERN(CategoryName, DefaultVerbosity, CompileTimeVerbosity) \
-		extern struct FLogCategory##CategoryName : public FLogCategory<ELogVerbosity::DefaultVerbosity, ELogVerbosity::CompileTimeVerbosity> \
-		{ \
-			UE_FORCEINLINE_HINT FLogCategory##CategoryName() : FLogCategory(TEXT(#CategoryName)) {} \
-		} CategoryName;
-
-	/** 
-	 * A macro to define a logging category, usually paired with DECLARE_LOG_CATEGORY_EXTERN from the header.
-	 * @param CategoryName, category to define
-	**/
-	#define DEFINE_LOG_CATEGORY(CategoryName) FLogCategory##CategoryName CategoryName;
-
-	/** 
-	 * A macro to define a logging category as a C++ "static". This should ONLY be declared in a source file. Only accessible in that single file.
-	 * @param CategoryName, category to declare
-	 * @param DefaultVerbosity, default run time verbosity
-	 * @param CompileTimeVerbosity, maximum verbosity to compile into the code
-	**/
-	//-V:DEFINE_LOG_CATEGORY_STATIC:501
-	#define DEFINE_LOG_CATEGORY_STATIC(CategoryName, DefaultVerbosity, CompileTimeVerbosity) \
-		static struct FLogCategory##CategoryName : public FLogCategory<ELogVerbosity::DefaultVerbosity, ELogVerbosity::CompileTimeVerbosity> \
-		{ \
-			static_assert(ELogVerbosity::DefaultVerbosity <= ELogVerbosity::CompileTimeVerbosity, "Default verbosity of log category cannot be higher than its compile time verbosity"); \
-			UE_FORCEINLINE_HINT FLogCategory##CategoryName() : FLogCategory(TEXT(#CategoryName)) {} \
-		} CategoryName;
-
-	/** 
-	 * A macro to declare a logging category as a C++ "class static" 
-	 * @param CategoryName, category to declare
-	 * @param DefaultVerbosity, default run time verbosity
-	 * @param CompileTimeVerbosity, maximum verbosity to compile into the code
-	**/
-	//-V:DECLARE_LOG_CATEGORY_CLASS:501
-	#define DECLARE_LOG_CATEGORY_CLASS(CategoryName, DefaultVerbosity, CompileTimeVerbosity) \
-		DEFINE_LOG_CATEGORY_STATIC(CategoryName, DefaultVerbosity, CompileTimeVerbosity)
-
-	/** 
-	 * A macro to define a logging category, usually paired with DECLARE_LOG_CATEGORY_CLASS from the header.
-	 * @param CategoryName, category to define
-	**/
-	#define DEFINE_LOG_CATEGORY_CLASS(Class, CategoryName) Class::FLogCategory##CategoryName Class::CategoryName;
-
 #endif // NO_LOGGING
 
 #if UE_BUILD_SHIPPING
