Make CeilLogTwo constexpr as it is in GenericPlatformMath

diff --git a/Engine/Source/Runtime/Core/Public/Microsoft/MicrosoftPlatformMath.h b/Engine/Source/Runtime/Core/Public/Microsoft/MicrosoftPlatformMath.h
--- a/Engine/Source/Runtime/Core/Public/Microsoft/MicrosoftPlatformMath.h
+++ b/Engine/Source/Runtime/Core/Public/Microsoft/MicrosoftPlatformMath.h
@@ -65,7 +65,7 @@ struct FMicrosoftPlatformMathBase : public TUnrealPlatformMathSSE4Base<FGenericP
 		}
 	}
 
-	static inline uint32 CeilLogTwo( uint32 Arg )
+	static constexpr inline uint32 CeilLogTwo( uint32 Arg )
 	{
 		// if Arg is 0, change it to 1 so that we return 0
 		Arg = Arg ? Arg : 1;
