Silence warning V501: There are identical sub-expressions to the left and to the right of the '<=' operator. for DECLARE_LOG_CATEGORY_EXTERN
#rnx

diff --git a/Engine/Source/Runtime/Core/Public/Analysis/PvsStudioCodeAnalysis.h b/Engine/Source/Runtime/Core/Public/Analysis/PvsStudioCodeAnalysis.h
--- a/Engine/Source/Runtime/Core/Public/Analysis/PvsStudioCodeAnalysis.h
+++ b/Engine/Source/Runtime/Core/Public/Analysis/PvsStudioCodeAnalysis.h
@@ -25,6 +25,7 @@
 	//-V:UE_BUILD_SHIPPING:501
 	//-V:WITH_EDITOR:501
 	//-V:TestTrueExpr:501
+	//-V:DECLARE_LOG_CATEGORY_EXTERN:501
 	//-V:PLATFORM_:517,547
 	//-V:ensureMsgf:562
 	//-V:MAX_VERTS_PER_POLY:512
