ParkingLot: Extended the documentation on ParkingLot::Wait

#rnx

diff --git a/Engine/Source/Runtime/Core/Public/Async/ParkingLot.h b/Engine/Source/Runtime/Core/Public/Async/ParkingLot.h
--- a/Engine/Source/Runtime/Core/Public/Async/ParkingLot.h
+++ b/Engine/Source/Runtime/Core/Public/Async/ParkingLot.h
@@ -45,6 +45,14 @@ UE_API void WakeOne(const void* Address, uint64 (*OnWakeState)(void*, FWakeState
 /**
  * Queue the calling thread to wait if CanWait returns true. BeforeWait is only called if CanWait returns true.
  *
+ * CanWait is called while the queue for the address is locked. It is invalid to call into any ParkingLot function from
+ * within a CanWait function. No other thread can wake or start waiting on this address while within CanWait. Returning
+ * false from CanWait will return from Wait without waiting.
+ *
+ * BeforeWait is called only when CanWait returns true. It is called after the queue for the address has been unlocked.
+ * It is invalid to call into any ParkingLot wait function from within BeforeWait, but it is safe to call into any wake
+ * function and may be necessary as part of the implementation of some synchronization primitives.
+ *
  * @param Address      Address to use as the key for the queue. The same address is used to wake the thread.
  * @param CanWait      Function called while the queue is locked. A return of false cancels the wait.
  * @param BeforeWait   Function called after the queue is unlocked and before the thread waits.
