Fixed the comment on FRecursiveWordMutex

#rnx

diff --git a/Engine/Source/Runtime/Core/Public/Async/RecursiveWordMutex.h b/Engine/Source/Runtime/Core/Public/Async/RecursiveWordMutex.h
--- a/Engine/Source/Runtime/Core/Public/Async/RecursiveWordMutex.h
+++ b/Engine/Source/Runtime/Core/Public/Async/RecursiveWordMutex.h
@@ -11,7 +11,7 @@ namespace UE
 {
 
 /**
- * A recursive mutex that is the size of a pointer and does not depend on ParkingLot.
+ * A recursive mutex that is the size of two pointers and does not depend on ParkingLot.
  *
  * Prefer FRecursiveMutex to FRecursiveWordMutex whenever possible.
  * This mutex is not fair and supports recursive locking.
