Traits: Added TIsCharEncodingCompatibleWithUTF8CHAR

diff --git a/Engine/Source/Runtime/Core/Public/Traits/IsCharEncodingCompatibleWith.h b/Engine/Source/Runtime/Core/Public/Traits/IsCharEncodingCompatibleWith.h
--- a/Engine/Source/Runtime/Core/Public/Traits/IsCharEncodingCompatibleWith.h
+++ b/Engine/Source/Runtime/Core/Public/Traits/IsCharEncodingCompatibleWith.h
@@ -71,3 +71,6 @@ constexpr inline bool TIsCharEncodingCompatibleWith_V = TIsCharEncodingCompatibl
 
 template <typename SrcEncoding>
 using TIsCharEncodingCompatibleWithTCHAR = TIsCharEncodingCompatibleWith<SrcEncoding, TCHAR>;
+
+template <typename SrcEncoding>
+using TIsCharEncodingCompatibleWithUTF8CHAR = TIsCharEncodingCompatibleWith<SrcEncoding, UTF8CHAR>;
