Support querying TDR information and log on startup

diff --git a/Engine/Source/Runtime/Core/Private/Windows/WindowsPlatformMisc.cpp b/Engine/Source/Runtime/Core/Private/Windows/WindowsPlatformMisc.cpp
--- a/Engine/Source/Runtime/Core/Private/Windows/WindowsPlatformMisc.cpp
+++ b/Engine/Source/Runtime/Core/Private/Windows/WindowsPlatformMisc.cpp
@@ -715,6 +715,30 @@ const TCHAR* LexToString(EStorageDeviceType StorageType)
 	}
 }
 
+const TCHAR* LexToString(ETDRLevel TDRLevel)
+{
+	switch (TDRLevel)
+	{
+	case ETDRLevel::Disabled:
+		return TEXT("Disabled");
+		break;
+	case ETDRLevel::BugCheck:
+		return TEXT("BugCheck");
+		break;
+	case ETDRLevel::RecoverVGA:
+		return TEXT("RecoverVGA");
+		break;
+	case ETDRLevel::Recover:
+		return TEXT("Recover");
+		break;
+	case ETDRLevel::Invalid:
+		[[fallthrough]];
+	default:
+		return TEXT("Invalid");
+		break;
+	}
+}
+
 #include "Windows/HideWindowsPlatformTypes.h"
 
 /** 
@@ -4031,6 +4055,37 @@ FGPUDriverInfo FWindowsPlatformMisc::GetGPUDriverInfo(const FString& DeviceDescr
 	return Ret;
 }
 
+FTDRInfo FWindowsPlatformMisc::GetTDRInfo()
+{
+	FTDRInfo TDRInfo;
+	TDRInfo.Level = ETDRLevel::Recover;
+	TDRInfo.Delay = FTDRInfo::DefaultDelay;
+	TDRInfo.DdiDelay = FTDRInfo::DefaultDdiDelay;
+
+	const TCHAR* GraphicsDriverKey = TEXT("System\\CurrentControlSet\\Control\\GraphicsDrivers");
+
+	FString TDRLevel;
+	if (FWindowsPlatformMisc::QueryRegKey(HKEY_LOCAL_MACHINE, GraphicsDriverKey, TEXT("TdrLevel"), TDRLevel))
+	{
+		int32 Value = FCString::Atoi(*TDRLevel);
+		TDRInfo.Level = (Value >= 0 && Value < int32(ETDRLevel::Invalid)) ? ETDRLevel(Value) : ETDRLevel::Invalid;
+	}
+
+	FString TDRDelay;
+	if (FWindowsPlatformMisc::QueryRegKey(HKEY_LOCAL_MACHINE, GraphicsDriverKey, TEXT("TdrDelay"), TDRDelay))
+	{
+		LexFromString(TDRInfo.Delay, *TDRDelay);
+	}
+
+	FString TDRDdiDelay;
+	if (FWindowsPlatformMisc::QueryRegKey(HKEY_LOCAL_MACHINE, GraphicsDriverKey, TEXT("TdrDdiDelay"), TDRDdiDelay))
+	{
+		LexFromString(TDRInfo.DdiDelay, *TDRDdiDelay);
+	}
+
+	return TDRInfo;
+}
+
 #include "Windows/HideWindowsPlatformTypes.h"
 
 void FWindowsPlatformMisc::GetOSVersions( FString& OutOSVersionLabel, FString& OutOSSubVersionLabel )
diff --git a/Engine/Source/Runtime/Core/Public/Windows/WindowsPlatformMisc.h b/Engine/Source/Runtime/Core/Public/Windows/WindowsPlatformMisc.h
--- a/Engine/Source/Runtime/Core/Public/Windows/WindowsPlatformMisc.h
+++ b/Engine/Source/Runtime/Core/Public/Windows/WindowsPlatformMisc.h
@@ -90,6 +90,39 @@ struct FPlatformDriveStats
 	EStorageDeviceType DriveType;
 };
 
+enum class ETDRLevel
+{
+	/** TDR was disabled explicitly */
+	Disabled = 0,
+
+	/** Don't attempt to recover from a GPU or driver timeout and display a bugcheck screen instead */
+	BugCheck,
+
+	/** Not implemented */
+	RecoverVGA,
+
+	/** Recover adapter/driver on timeout (Default) */
+	Recover,
+
+	Invalid
+};
+
+CORE_API const TCHAR* LexToString(ETDRLevel TDRLevel);
+
+
+/**
+ *	Information about Windows' Timeout Detection and Recovery (TDR) mechanism for detecting display adapter hangs
+ */
+struct FTDRInfo
+{
+	static constexpr uint32 DefaultDelay = 2;
+	static constexpr uint32 DefaultDdiDelay = 5;
+
+	ETDRLevel Level;
+	uint32 Delay;		// GPU timeout threshold in seconds after which TDR will kick in
+	uint32 DdiDelay;	// Graphics driver timeout threshold in seconds after which TDR will kick in
+};
+
 /**
 * Windows implementation of the misc OS functions
 **/
@@ -258,6 +291,7 @@ struct FWindowsPlatformMisc
 
 	static CORE_API FString GetPrimaryGPUBrand();
 	static CORE_API struct FGPUDriverInfo GetGPUDriverInfo(const FString& DeviceDescription, bool bVerbose = true);
+	static CORE_API struct FTDRInfo GetTDRInfo();
 	static CORE_API void GetOSVersions( FString& out_OSVersionLabel, FString& out_OSSubVersionLabel );
 	static CORE_API FString GetOSVersion();
 	static CORE_API bool GetDiskTotalAndFreeSpace( const FString& InPath, uint64& TotalNumberOfBytes, uint64& NumberOfFreeBytes );
diff --git a/Engine/Source/Runtime/RHI/Private/Windows/WindowsDynamicRHI.cpp b/Engine/Source/Runtime/RHI/Private/Windows/WindowsDynamicRHI.cpp
--- a/Engine/Source/Runtime/RHI/Private/Windows/WindowsDynamicRHI.cpp
+++ b/Engine/Source/Runtime/RHI/Private/Windows/WindowsDynamicRHI.cpp
@@ -1144,6 +1144,25 @@ static bool HandleUnsupportedRHI(EWindowsRHI& WindowsRHI, ERHIFeatureLevel::Type
 	return false;
 }
 
+static void LogTDRSummary()
+{
+	FTDRInfo TDRInfo = FWindowsPlatformMisc::GetTDRInfo();
+	FString TDRSummary;
+	if (TDRInfo.Level == ETDRLevel::Recover &&
+		TDRInfo.Delay == FTDRInfo::DefaultDelay &&
+		TDRInfo.DdiDelay == FTDRInfo::DefaultDdiDelay)
+	{
+		TDRSummary += FString::Printf(TEXT("\r\n\t- TDR settings OK"));
+	}
+	else
+	{
+		TDRSummary += FString::Printf(TEXT("\r\n\t- TDR using non-default settings"));
+	}
+
+	TDRSummary += FString::Printf(TEXT(" - Level: %s, Delay: %d, DdiDelay: %d"), LexToString(TDRInfo.Level), TDRInfo.Delay, TDRInfo.DdiDelay);
+	UE_LOG(LogRHI, Display, TEXT("Timeout Detection and Recovery (TDR) summary: %s"), *TDRSummary)
+}
+
 static IDynamicRHIModule* LoadDynamicRHIModule(ERHIFeatureLevel::Type& DesiredFeatureLevel, const TCHAR*& LoadedRHIModuleName)
 {
 	// Make sure the DDSPI is initialized before we try and read from it
@@ -1255,6 +1274,8 @@ FDynamicRHI* PlatformCreateDynamicRHI()
 		GLoadedRHIModuleName = LoadedRHIModuleName;
 	}
 
+	LogTDRSummary();
+
 	return DynamicRHI;
 }
 
