CsvProfiler: Add support for timing regions for simple concurrency profiling. These are a subclass of events with a "Region:" prefix which can be processed by PerfReportTool for visualization and generation of metrics. Use macros CSV_REGION_BEGIN/CSV_REGION_END or CSV_REGION_BEGIN_GLOBAL/CSV_REGION_END_GLOBAL to instrument
- Improve boot capture support to ensure that capture timings start from the GStartTime rather than when the CSV profiler is initialized, and ensure frame times and event timestamps take this into account.

diff --git a/Engine/Source/Runtime/Core/Private/ProfilingDebugging/CsvProfiler.cpp b/Engine/Source/Runtime/Core/Private/ProfilingDebugging/CsvProfiler.cpp
--- a/Engine/Source/Runtime/Core/Private/ProfilingDebugging/CsvProfiler.cpp
+++ b/Engine/Source/Runtime/Core/Private/ProfilingDebugging/CsvProfiler.cpp
@@ -298,7 +298,7 @@ static bool GCsvExitOnCompletion = false;
 
 static thread_local bool GCsvThreadLocalWaitsEnabled = false;
 
-static uint64 GFirstFrameTimestamp = 0;
+static uint64 GCaptureStartTimestamp = 0;
 
 //
 // Forward declarations
@@ -3562,7 +3562,7 @@ void FCsvProfilerThreadDataProcessor::Process(FCsvProcessThreadDataStats& OutSta
 				ProcessedEvent.EventText = Event.EventText;
 				ProcessedEvent.FrameNumber = FrameNumber;
 				ProcessedEvent.CategoryIndex = Event.CategoryIndex;
-				ProcessedEvent.TimeOffsetSeconds = GCsvEventTimestamps ? FPlatformTime::ToSeconds64(Event.Timestamp-GFirstFrameTimestamp) : std::numeric_limits<double>::max();
+				ProcessedEvent.TimeOffsetSeconds = GCsvEventTimestamps ? FPlatformTime::ToSeconds64(Event.Timestamp-GCaptureStartTimestamp) : std::numeric_limits<double>::max();
 				Writer->PushEvent(ProcessedEvent);
 			}
 		}
@@ -3728,12 +3728,11 @@ void FCsvProfiler::BeginFrame()
 
 		if (GCsvProfilerIsCapturing)
 		{
-			uint64 BeginFrameTimestamp = GFrameBoundaries.AddBeginFrameTimestamp(ECsvTimeline::Gamethread);
+			GFrameBoundaries.AddBeginFrameTimestamp(ECsvTimeline::Gamethread);
 
 			if (CaptureFrameNumber == 0)
 			{
 				OnCSVProfileFirstFrameDelegate.Broadcast();
-				GFirstFrameTimestamp = BeginFrameTimestamp;
 			}
 
 			if (!bNamedEventsWasEnabled && (GCycleStatsShouldEmitNamedEvents > 0))
@@ -3822,7 +3821,28 @@ void FCsvProfiler::BeginCaptureInternal(const FCsvCaptureCommand& CurrentCommand
 	GCsvRepeatFrameCount = NumFramesToCapture;
 	CaptureFrameNumber = 0;
 	CaptureFrameNumberRT = 0;
-	LastEndFrameTimestamp = FPlatformTime::Cycles64();
+
+	if (EnumHasAnyFlags(CurrentCommand.Flags, ECsvProfilerFlags::BootCapture))
+	{
+		// If we're capturing from boot, set the start time to GStartTime 
+		// This ensures the first time and boot timings are relative to GStartTime
+		check(GFrameNumber==1);
+		UE_LOG(LogCsvProfiler, Log, TEXT("Doing a boot time capture. Start time offset %.2f"), FPlatformTime::Seconds()-GStartTime);
+		CaptureStartTime = GStartTime;
+		GCaptureStartTimestamp = FPlatformTime::SecondsToCycles64(GStartTime);
+
+		// signal bInsertEndFrameAtFrameStart to insert an EndFrame() at the start of the first _real_ frame
+		// We also add a FrameBeginTimestampsRT timestamp here, to create a dummy renderthread frame, to ensure the rows match up in the CSV
+		GFrameBoundaries.AddBeginFrameTimestamp(ECsvTimeline::Renderthread, false);
+		bInsertEndFrameAtFrameStart = true;
+	}
+	else
+	{
+		GCaptureStartTimestamp = FPlatformTime::Cycles64();
+		CaptureStartTime = FPlatformTime::ToSeconds64(GCaptureStartTimestamp);
+	}
+	LastEndFrameTimestamp = GCaptureStartTimestamp;
+
 	CurrentFlags = CurrentCommand.Flags;
 
 	if (GCsvUseProcessingThread && ProcessingThread == nullptr)
@@ -3898,7 +3918,6 @@ void FCsvProfiler::BeginCaptureInternal(const FCsvCaptureCommand& CurrentCommand
 	check(FCsvProfilerThreadData::IsTlsSlotInitialized());
 	TRACE_CSV_PROFILER_BEGIN_CAPTURE(*Filename, RenderThreadId, RHIThreadId, GDefaultWaitStatName.StatName, GCsvStatCounts);
 	GCsvProfilerIsCapturing = true;
-	CaptureStartTime = FPlatformTime::Seconds();
 
 	// Run the Csv benchmark if enabled
 	if (CVarCsvBenchmark.GetValueOnGameThread() || FParse::Param(FCommandLine::Get(), TEXT("csvBench")))
@@ -4568,7 +4587,6 @@ void FCsvProfiler::RecordEvent(int32 CategoryIndex, const FString& EventText, bo
 			}
 		}
 #endif 
-
 	}
 }
 
@@ -4577,6 +4595,42 @@ void FCsvProfiler::RecordEventAtFrameStart(int32 CategoryIndex, const FString& E
 	RecordEventAtTimestamp(CategoryIndex, EventText, FCsvProfiler::Get()->LastEndFrameTimestamp, bLog);
 }
 
+void FCsvProfiler::BeginRegion(int32 CategoryIndex, const TCHAR* RegionName)
+{
+	if (GCsvProfilerIsCapturing && GCsvCategoriesEnabled[CategoryIndex])
+	{
+		FString EventText;
+		if (CategoryIndex != CSV_CATEGORY_INDEX_GLOBAL)
+		{
+			EventText = FString::Printf(TEXT("Region: %s/%s - Begin"), *FCsvCategoryData::Get()->GetCategoryNameByIndex(CategoryIndex), RegionName);
+		}
+		else
+		{
+			EventText = FString::Printf(TEXT("Region: %s - Begin"), RegionName);
+		}
+		RecordEvent(CSV_CATEGORY_INDEX_GLOBAL, EventText, true );
+		// TODO: Track timings for each region and output metadata when CSV ends. Or we just use PRT for that...
+	}
+}
+
+void FCsvProfiler::EndRegion(int32 CategoryIndex, const TCHAR* RegionName)
+{
+	if (GCsvProfilerIsCapturing && GCsvCategoriesEnabled[CategoryIndex])
+	{
+		FString EventText;
+		if (CategoryIndex != CSV_CATEGORY_INDEX_GLOBAL)
+		{
+			EventText = FString::Printf(TEXT("Region: %s/%s - End"), *FCsvCategoryData::Get()->GetCategoryNameByIndex(CategoryIndex), RegionName);
+		}
+		else
+		{
+			EventText = FString::Printf(TEXT("Region: %s - End"), RegionName);
+		}
+		RecordEvent(CSV_CATEGORY_INDEX_GLOBAL, EventText, true);
+	}
+}
+
+
 void FCsvProfiler::SetMetadata(const TCHAR* Key, const TCHAR* Value)
 {
 	FCsvProfiler::Get()->SetMetadataInternal(Key, Value, true, EMetadataPersistenceType::Persistent);
@@ -4866,14 +4920,9 @@ void FCsvProfiler::Init()
 	if (FParse::Value(FCommandLine::Get(), TEXT("csvCaptureFrames="), NumCsvFrames))
 	{
 		check(IsInGameThread());
-		BeginCapture(NumCsvFrames);
-
+		BeginCapture(NumCsvFrames, FString(), FString(), ECsvProfilerFlags::BootCapture);
 		// Call BeginFrame() to start capturing a dummy first "frame"
-		// signal bInsertEndFrameAtFrameStart to insert an EndFrame() at the start of the first _real_ frame
-		// We also add a FrameBeginTimestampsRT timestamp here, to create a dummy renderthread frame, to ensure the rows match up in the CSV
 		BeginFrame();
-		GFrameBoundaries.AddBeginFrameTimestamp(ECsvTimeline::Renderthread, false);
-		bInsertEndFrameAtFrameStart = true;
 	}
 	FParse::Value(FCommandLine::Get(), TEXT("csvRepeat="), GCsvRepeatCount);
 
diff --git a/Engine/Source/Runtime/Core/Public/ProfilingDebugging/CsvProfiler.h b/Engine/Source/Runtime/Core/Public/ProfilingDebugging/CsvProfiler.h
--- a/Engine/Source/Runtime/Core/Public/ProfilingDebugging/CsvProfiler.h
+++ b/Engine/Source/Runtime/Core/Public/ProfilingDebugging/CsvProfiler.h
@@ -70,6 +70,19 @@ struct FCsvDeclaredStat;
 		FCsvProfiler::RecordEventfNoLog( CSV_CATEGORY_INDEX_GLOBAL, Format, ##__VA_ARGS__ ); \
 		TRACE_BOOKMARK(Format, ##__VA_ARGS__)
 
+	#define CSV_REGION_BEGIN(Category, RegionName) \
+		FCsvProfiler::BeginRegion(CSV_CATEGORY_INDEX(Category), RegionName);
+
+	#define CSV_REGION_END(Category, RegionName) \
+		FCsvProfiler::EndRegion(CSV_CATEGORY_INDEX(Category), RegionName);
+
+	#define CSV_REGION_BEGIN_GLOBAL(RegionName) \
+		FCsvProfiler::BeginRegion(CSV_CATEGORY_INDEX_GLOBAL, RegionName);
+
+	#define CSV_REGION_END_GLOBAL(RegionName) \
+		FCsvProfiler::EndRegion(CSV_CATEGORY_INDEX_GLOBAL, RegionName);
+
+
 	// Metadata
 	#define CSV_METADATA(Key,Value)									FCsvProfiler::SetMetadata( Key, Value )
 	#define CSV_NON_PERSISTENT_METADATA(Key,Value)					FCsvProfiler::SetNonPersistentMetadata( Key, Value )
@@ -90,6 +103,10 @@ struct FCsvDeclaredStat;
 	#define CSV_EVENT_GLOBAL(Format, ...)
 	#define CSV_EVENT_NOLOG(Category, Format, ...) 						
 	#define CSV_EVENT_GLOBAL_NOLOG(Format, ...)
+	#define CSV_REGION_BEGIN(Category, RegionName)
+	#define CSV_REGION_END(Category, RegionName)
+	#define CSV_REGION_BEGIN_GLOBAL(RegionName)
+	#define CSV_REGION_END_GLOBAL(RegionName)
 	#define CSV_METADATA(Key,Value)
 	#define CSV_NON_PERSISTENT_METADATA(Key,Value)
 	#define CSV_CUSTOM_STAT_MINIMAL(Category,StatName,Value,Op)
@@ -216,9 +233,10 @@ struct FCsvDeclaredStat
 
 enum class ECsvProfilerFlags
 {
-	None = 0,
-	WriteCompletionFile = 1,
-	CompressOutput = 2
+	None = 0x0,
+	WriteCompletionFile = 0x1,
+	CompressOutput = 0x2,
+	BootCapture = 0x4
 };
 ENUM_CLASS_FLAGS(ECsvProfilerFlags);
 
@@ -444,6 +462,9 @@ public:
 	CORE_API static void BeginWait();
 	CORE_API static void EndWait();
 
+	CORE_API static void BeginRegion(int32 CategoryIndex, const TCHAR* RegionName);
+	CORE_API static void EndRegion(int32 CategoryIndex, const TCHAR* RegionName);
+
 	CORE_API static FString GetDefaultDirectory();
 
 	/** Singleton interface */
