CSV profiler. Add CSV_TRACE_REGION macros so we can emit global region to MiscTrace and CSV profiler with the same macro
Add -csvContinuousWritesDisabled to disable continuous writes via the commandline. This is useful to minimize I/O overhead for boot profiling.
- Disable logging for region events

#tests tested on Android and console

diff --git a/Engine/Source/Runtime/Core/Private/ProfilingDebugging/CsvProfiler.cpp b/Engine/Source/Runtime/Core/Private/ProfilingDebugging/CsvProfiler.cpp
--- a/Engine/Source/Runtime/Core/Private/ProfilingDebugging/CsvProfiler.cpp
+++ b/Engine/Source/Runtime/Core/Private/ProfilingDebugging/CsvProfiler.cpp
@@ -480,6 +480,13 @@ static TArray<FEventExecCmds>* GCsvEventExecCmds = NULL;
 
 bool IsContinuousWriteEnabled(bool bGameThread)
 {
+#if CSV_PROFILER_ALLOW_DEBUG_FEATURES
+	static bool bContinuousWritesDisabled = FParse::Param(FCommandLine::Get(), TEXT("csvContinuousWritesDisabled"));
+	if (bContinuousWritesDisabled)
+	{
+		return false;
+	}
+#endif
 	int32 CVarValue = -1;
 #if UE_BUILD_SHIPPING
 	CVarValue = bGameThread ? CVarCsvShippingContinuousWrites.GetValueOnGameThread() : CVarCsvShippingContinuousWrites.GetValueOnAnyThread();
@@ -4608,7 +4615,7 @@ void FCsvProfiler::BeginRegion(int32 CategoryIndex, const TCHAR* RegionName)
 		{
 			EventText = FString::Printf(TEXT("Region: %s - Begin"), RegionName);
 		}
-		RecordEvent(CSV_CATEGORY_INDEX_GLOBAL, EventText, true );
+		RecordEvent(CSV_CATEGORY_INDEX_GLOBAL, EventText, false);
 		// TODO: Track timings for each region and output metadata when CSV ends. Or we just use PRT for that...
 	}
 }
@@ -4626,7 +4633,7 @@ void FCsvProfiler::EndRegion(int32 CategoryIndex, const TCHAR* RegionName)
 		{
 			EventText = FString::Printf(TEXT("Region: %s - End"), RegionName);
 		}
-		RecordEvent(CSV_CATEGORY_INDEX_GLOBAL, EventText, true);
+		RecordEvent(CSV_CATEGORY_INDEX_GLOBAL, EventText, false);
 	}
 }
 
diff --git a/Engine/Source/Runtime/Core/Public/ProfilingDebugging/CsvProfiler.h b/Engine/Source/Runtime/Core/Public/ProfilingDebugging/CsvProfiler.h
--- a/Engine/Source/Runtime/Core/Public/ProfilingDebugging/CsvProfiler.h
+++ b/Engine/Source/Runtime/Core/Public/ProfilingDebugging/CsvProfiler.h
@@ -41,7 +41,6 @@ class FScopedCsvStat;
 class FScopedCsvStatExclusive;
 struct FCsvDeclaredStat;
 
-
 #if CSV_PROFILER
 	// Helpers
 	#define CSV_CATEGORY_INDEX(CategoryName)									(_GCsvCategory_##CategoryName.Index)
@@ -192,6 +191,14 @@ struct FCsvDeclaredStat;
 	#define CSV_CUSTOM_STAT_DEFINED(StatName,Value,Op)				
 #endif
 
+// Shared macros. These emit trace and CSV instrumentation
+#define CSV_TRACE_REGION_BEGIN_GLOBAL(RegionName) \
+	TRACE_BEGIN_REGION(RegionName) \
+	CSV_REGION_BEGIN_GLOBAL(RegionName)
+
+#define CSV_TRACE_REGION_END_GLOBAL(RegionName) \
+	CSV_REGION_END_GLOBAL(RegionName) \
+	TRACE_END_REGION(RegionName)
 
 #if CSV_PROFILER
 
