[Async] AUTORTFM_INFER TFutureBase::Then() and TFutureBase::Next()

Prevents compile-time errors when passing an AutoRTFM-disabled function as the Continuation argument, as these functions are currently always AutoRTFM-enabled.

Longer term, these functions (along with the vast majority of Async) should be annotated with AUTORTFM_DISABLE, but this provides a stepping stone towards disabling more of the codebase.

See Engine\Source\Runtime\AutoRTFM\Documentation\ModesAndAttributes.md for more information about AutoRTFM modes.

diff --git a/Engine/Source/Runtime/Core/Public/Async/Future.h b/Engine/Source/Runtime/Core/Public/Async/Future.h
--- a/Engine/Source/Runtime/Core/Public/Async/Future.h
+++ b/Engine/Source/Runtime/Core/Public/Async/Future.h
@@ -2,6 +2,7 @@
 
 #pragma once
 
+#include "AutoRTFM/Defines.h"
 #include "CoreTypes.h"
 #include "Misc/AssertionMacros.h"
 #include "Templates/UnrealTemplate.h"
@@ -351,7 +352,7 @@ protected:
 	 * @return nothing at the moment but could return another future to allow future chaining
 	 */
 	template<typename Func>
-	auto Then(Func Continuation);
+	AUTORTFM_INFER auto Then(Func Continuation);
 
 	/**
 	 * Convenience wrapper for Then that
@@ -361,7 +362,7 @@ protected:
 	 * @return nothing at the moment but could return another future to allow future chaining
 	 */
 	template<typename Func>
-	auto Next(Func Continuation);
+	AUTORTFM_INFER auto Next(Func Continuation);
 
 	/**
 	 * Reset the future.
@@ -733,4 +734,4 @@ TPromise<ResultType> MakeFulfilledPromise(ArgTypes&&... Args)
 
 #if UE_ENABLE_INCLUDE_ORDER_DEPRECATED_IN_5_7
 #include "Templates/Requires.h"
-#endif
\ No newline at end of file
+#endif
