Explicitly qualify a handful of symbols that exist in both uLang and the top-level namespace.

As we remove the uLang Core reimplementation, we will briefly need the uLang libraries to coexist with the equivalent Core libraries.

diff --git a/Engine/Source/Runtime/Core/Public/Containers/Array.h b/Engine/Source/Runtime/Core/Public/Containers/Array.h
--- a/Engine/Source/Runtime/Core/Public/Containers/Array.h
+++ b/Engine/Source/Runtime/Core/Public/Containers/Array.h
@@ -449,8 +449,8 @@ namespace UE::Core::Private
 			return
 				!TLosesQualifiersFromTo_V<FromElementType, ToElementType> &&
 				(
-					std::is_same_v         <const ToElementType, const FromElementType> ||               // The element type of the container must be the same, or...
-					TIsBitwiseConstructible<UnqualifiedToElementType, UnqualifiedFromElementType>::Value // ... the element type of the source container must be bitwise constructible from the element type in the destination container
+					std::is_same_v           <const ToElementType, const FromElementType> ||               // The element type of the container must be the same, or...
+					::TIsBitwiseConstructible<UnqualifiedToElementType, UnqualifiedFromElementType>::Value // ... the element type of the source container must be bitwise constructible from the element type in the destination container
 				);
 		}
 		else
@@ -495,11 +495,11 @@ namespace UE::Core::Private
 	UE_NODEBUG [[nodiscard]] constexpr uint32 GetAllocatorFlags()
 	{
 		uint32 Result = 0;
-		if constexpr (TAllocatorTraits<AllocatorType>::SupportsElementAlignment)
+		if constexpr (::TAllocatorTraits<AllocatorType>::SupportsElementAlignment)
 		{
 			Result |= 1;
 		}
-		if constexpr (TAllocatorTraits<AllocatorType>::SupportsSlackTracking)
+		if constexpr (::TAllocatorTraits<AllocatorType>::SupportsSlackTracking)
 		{
 			Result |= 2;
 		}
@@ -936,7 +936,7 @@ private:
 	UE_FORCEINLINE_HINT void SlackTrackerNumChanged()
 	{
 #if UE_ENABLE_ARRAY_SLACK_TRACKING
-		if constexpr (TAllocatorTraits<InAllocatorType>::SupportsSlackTracking)
+		if constexpr (::TAllocatorTraits<InAllocatorType>::SupportsSlackTracking)
 		{
 			AllocatorInstance.SlackTrackerLogNum(ArrayNum);
 		}
@@ -2574,7 +2574,7 @@ public:
 	template <
 		typename RangeType
 		UE_REQUIRES(
-			TIsContiguousContainer<RangeType>::Value &&
+			::TIsContiguousContainer<RangeType>::Value &&
 			!UE::Core::Private::TIsTArrayOrDerivedFromTArray_V<std::remove_reference_t<RangeType>> &&
 			UE::Core::Private::TArrayElementsAreCompatible_V<ElementType, TElementType_T<RangeType>>
 		)
@@ -3674,7 +3674,7 @@ protected:
 public:
 	void WriteMemoryImage(FMemoryImageWriter& Writer) const
 	{
-		if constexpr (TAllocatorTraits<AllocatorType>::SupportsFreezeMemoryImage && THasTypeLayout<ElementType>::Value)
+		if constexpr (::TAllocatorTraits<AllocatorType>::SupportsFreezeMemoryImage && THasTypeLayout<ElementType>::Value)
 		{
 			this->AllocatorInstance.WriteMemoryImage(Writer, StaticGetTypeLayoutDesc<ElementType>(), this->ArrayNum);
 			Writer.WriteBytes(this->ArrayNum);
@@ -3691,7 +3691,7 @@ public:
 
 	void CopyUnfrozen(const FMemoryUnfreezeContent& Context, void* Dst) const
 	{
-		if constexpr (TAllocatorTraits<AllocatorType>::SupportsFreezeMemoryImage && THasTypeLayout<ElementType>::Value)
+		if constexpr (::TAllocatorTraits<AllocatorType>::SupportsFreezeMemoryImage && THasTypeLayout<ElementType>::Value)
 		{
 			TArray* DstArray = ::new(Dst) TArray();
 			DstArray->SetNumZeroed(this->ArrayNum);
@@ -3705,7 +3705,7 @@ public:
 
 	static void AppendHash(const FPlatformTypeLayoutParameters& LayoutParams, FSHA1& Hasher)
 	{
-		if constexpr (TAllocatorTraits<AllocatorType>::SupportsFreezeMemoryImage && THasTypeLayout<ElementType>::Value)
+		if constexpr (::TAllocatorTraits<AllocatorType>::SupportsFreezeMemoryImage && THasTypeLayout<ElementType>::Value)
 		{
 			Freeze::AppendHash(StaticGetTypeLayoutDesc<ElementType>(), LayoutParams, Hasher);
 		}
@@ -3713,7 +3713,7 @@ public:
 
 	void ToString(const FPlatformTypeLayoutParameters& LayoutParams, FMemoryToStringContext& OutContext) const
 	{
-		if constexpr (TAllocatorTraits<AllocatorType>::SupportsFreezeMemoryImage && THasTypeLayout<ElementType>::Value)
+		if constexpr (::TAllocatorTraits<AllocatorType>::SupportsFreezeMemoryImage && THasTypeLayout<ElementType>::Value)
 		{
 			this->AllocatorInstance.ToString(StaticGetTypeLayoutDesc<ElementType>(), this->ArrayNum, this->ArrayMax, LayoutParams, OutContext);
 		}
@@ -4054,7 +4054,7 @@ DECLARE_TEMPLATE_INTRINSIC_TYPE_LAYOUT((template <typename T, typename Allocator
 template <typename InElementType, typename AllocatorType>
 struct TIsZeroConstructType<TArray<InElementType, AllocatorType>>
 {
-	enum { Value = TAllocatorTraits<AllocatorType>::IsZeroConstruct };
+	enum { Value = ::TAllocatorTraits<AllocatorType>::IsZeroConstruct };
 };
 
 template <typename T, typename AllocatorType>
diff --git a/Engine/Source/Runtime/Core/Public/HAL/UnrealMemory.h b/Engine/Source/Runtime/Core/Public/HAL/UnrealMemory.h
--- a/Engine/Source/Runtime/Core/Public/HAL/UnrealMemory.h
+++ b/Engine/Source/Runtime/Core/Public/HAL/UnrealMemory.h
@@ -7,7 +7,8 @@
 #include "HAL/MemoryBase.h"
 #include "HAL/PlatformMemory.h"
 #include "ProfilingDebugging/MemoryTrace.h"
-#include "Templates/IsPointer.h"
+
+#include <type_traits>
 
 #ifndef UE_USE_VERYLARGEPAGEALLOCATOR
 #define UE_USE_VERYLARGEPAGEALLOCATOR 0
@@ -122,7 +123,7 @@ struct FMemory
 	template< class T > 
 	static inline void Memset( T& Src, uint8 ValueToSet )
 	{
-		static_assert( !TIsPointer<T>::Value, "For pointers use the three parameters function");
+		static_assert( !std::is_pointer_v<T>, "For pointers use the three parameters function");
 		Memset( &Src, ValueToSet, sizeof( T ) );
 	}
 
@@ -151,7 +152,7 @@ struct FMemory
 	template< class T > 
 	static inline void Memzero( T& Src )
 	{
-		static_assert( !TIsPointer<T>::Value, "For pointers use the two parameters function");
+		static_assert( !std::is_pointer_v<T>, "For pointers use the two parameters function");
 		Memzero( &Src, sizeof( T ) );
 	}
 
@@ -163,7 +164,7 @@ struct FMemory
 	template< class T > 
 	static inline void Memcpy( T& Dest, const T& Src )
 	{
-		static_assert( !TIsPointer<T>::Value, "For pointers use the three parameters function");
+		static_assert( !std::is_pointer_v<T>, "For pointers use the three parameters function");
 		Memcpy( &Dest, &Src, sizeof( T ) );
 	}
 
diff --git a/Engine/Source/Runtime/Core/Public/Math/Interval.h b/Engine/Source/Runtime/Core/Public/Math/Interval.h
--- a/Engine/Source/Runtime/Core/Public/Math/Interval.h
+++ b/Engine/Source/Runtime/Core/Public/Math/Interval.h
@@ -13,7 +13,7 @@
  */
 template <typename ElementType> struct TIntervalTraits
 {
-	static_assert(TIsArithmetic<ElementType>::Value, "Incompatible TInterval element type.");
+	static_assert(::TIsArithmetic<ElementType>::Value, "Incompatible TInterval element type.");
 
 	static ElementType Max()
 	{
diff --git a/Engine/Source/Runtime/Core/Public/Templates/MemoryOps.h b/Engine/Source/Runtime/Core/Public/Templates/MemoryOps.h
--- a/Engine/Source/Runtime/Core/Public/Templates/MemoryOps.h
+++ b/Engine/Source/Runtime/Core/Public/Templates/MemoryOps.h
@@ -16,7 +16,7 @@ namespace UE::Core::Private::MemoryOps
 	template <typename DestinationElementType, typename SourceElementType>
 	constexpr inline bool TCanBitwiseRelocate_V =
 		std::is_same_v<DestinationElementType, SourceElementType> || (
-			TIsBitwiseConstructible<DestinationElementType, SourceElementType>::Value &&
+			::TIsBitwiseConstructible<DestinationElementType, SourceElementType>::Value &&
 			std::is_trivially_destructible_v<SourceElementType>
 		);
 }
@@ -105,7 +105,7 @@ AUTORTFM_INFER FORCENOINLINE constexpr void DestructItems(ElementType* Element,
  * @param	Count		The number of elements to copy.
  */
 template <typename DestinationElementType, typename SourceElementType, typename SizeType>
-	requires (sizeof(DestinationElementType) > 0 && sizeof(SourceElementType) > 0 && !!TIsBitwiseConstructible<DestinationElementType, SourceElementType>::Value) // the sizeof here should improve the error messages we get when we try to call this function with incomplete types
+	requires (sizeof(DestinationElementType) > 0 && sizeof(SourceElementType) > 0 && !!::TIsBitwiseConstructible<DestinationElementType, SourceElementType>::Value) // the sizeof here should improve the error messages we get when we try to call this function with incomplete types
 FORCEINLINE void ConstructItems(void* Dest, const SourceElementType* Source, SizeType Count)
 {
 	if (Count)
@@ -114,7 +114,7 @@ FORCEINLINE void ConstructItems(void* Dest, const SourceElementType* Source, Siz
 	}
 }
 template <typename DestinationElementType, typename SourceElementType, typename SizeType>
-	requires (sizeof(DestinationElementType) > 0 && sizeof(SourceElementType) > 0 && !TIsBitwiseConstructible<DestinationElementType, SourceElementType>::Value) // the sizeof here should improve the error messages we get when we try to call this function with incomplete types
+	requires (sizeof(DestinationElementType) > 0 && sizeof(SourceElementType) > 0 && !::TIsBitwiseConstructible<DestinationElementType, SourceElementType>::Value) // the sizeof here should improve the error messages we get when we try to call this function with incomplete types
 AUTORTFM_INFER FORCENOINLINE void ConstructItems(void* Dest, const SourceElementType* Source, SizeType Count)
 {
 	while (Count)
