Suppress warning: Storage type is aligned to 1 bytes but allocated type is aligned to 8 bytes [cplusplus.PlacementNew]

#rnx

diff --git a/Engine/Source/Runtime/Core/Private/MemPro/MemPro.cpp b/Engine/Source/Runtime/Core/Private/MemPro/MemPro.cpp
--- a/Engine/Source/Runtime/Core/Private/MemPro/MemPro.cpp
+++ b/Engine/Source/Runtime/Core/Private/MemPro/MemPro.cpp
@@ -2554,6 +2554,7 @@ void MemPro::InitialiseInternal()
 	if(!gp_MemPro && !g_ShuttingDown)
 	{
 		CMemPro* p_mempro = (CMemPro*)g_MemProMem;
+		UE_CLANG_SUPPRESS //warning: Storage type is aligned to 1 bytes but allocated type is aligned to 8 bytes [cplusplus.PlacementNew]
 		new (p_mempro)CMemPro();
 		gp_MemPro = p_mempro;
 		p_mempro->Initialise();
