Remove namespaced FIoStatusBuilder since it is not used anymore.

#rnx

diff --git a/Engine/Source/Runtime/Core/Private/IO/IoStatus.cpp b/Engine/Source/Runtime/Core/Private/IO/IoStatus.cpp
--- a/Engine/Source/Runtime/Core/Private/IO/IoStatus.cpp
+++ b/Engine/Source/Runtime/Core/Private/IO/IoStatus.cpp
@@ -138,16 +138,6 @@ FIoStatusBuilder::FIoStatusBuilder(EIoErrorCode InStatusCode, uint32 InSystemErr
 {
 }
 
-FIoStatusBuilder::FIoStatusBuilder(
-	const TCHAR* NamespaceStr, EIoErrorCode StatusCode,
-	const uint_least32_t Line, const uint_least32_t Column)
-	: FIoStatusBuilder(StatusCode)
-{
-	Message
-		<< TEXT('<') << NamespaceStr << TEXT('-')
-		<< Line << TEXT(',') << Column << TEXT(">");
-}
-
 FIoStatusBuilder::~FIoStatusBuilder()
 {
 }
diff --git a/Engine/Source/Runtime/Core/Public/IO/IoStatus.h b/Engine/Source/Runtime/Core/Public/IO/IoStatus.h
--- a/Engine/Source/Runtime/Core/Public/IO/IoStatus.h
+++ b/Engine/Source/Runtime/Core/Public/IO/IoStatus.h
@@ -154,11 +154,6 @@ protected:
 public:
 	CORE_API explicit	FIoStatusBuilder(EIoErrorCode StatusCode);
 	CORE_API explicit	FIoStatusBuilder(EIoErrorCode StatusCode, uint32 SystemErrorCode);
-
-	CORE_API			FIoStatusBuilder(const TCHAR* NamespaceStr, EIoErrorCode StatusCode, 
-							const uint_least32_t Line = __builtin_LINE(),
-							const uint_least32_t Column = __builtin_COLUMN());
-
 	CORE_API			~FIoStatusBuilder();
 
 	CORE_API			operator FIoStatus();
@@ -172,27 +167,6 @@ public:
 	CORE_API FIoStatusBuilder& operator<<(FStringView String);
 };
 
-/**
- * Creates an FIoStatusBuilder that generates unique messages by prepending a namespace and source location.
- * Yes, this is in fact a builder builder - deal with it!
- */
-class FNamespacedIoStatusBuilderBuilder
-{
-private:
-	const TCHAR* NamespaceStr = nullptr;
-public:
-	FNamespacedIoStatusBuilderBuilder(const TCHAR* InNamespaceStr) : NamespaceStr(InNamespaceStr)
-	{
-	}
-
-	FIoStatusBuilder operator()(EIoErrorCode StatusCode, 
-		const uint_least32_t Line = __builtin_LINE(),
-		const uint_least32_t Column = __builtin_COLUMN()) const
-	{
-		return FIoStatusBuilder(NamespaceStr, StatusCode, Line, Column);
-	}
-};
-
 CORE_API void StatusOrCrash(const FIoStatus& Status);
 
 template<typename T>
