[RaceDetector]
- Fixed a missing function tag causing recursion inside GetObject
- Fixed a missing argument to Array.Pop causing it to have a high cost during race reporting

diff --git a/Engine/Source/Runtime/Core/Internal/Sanitizer/RaceDetectorTypes.h b/Engine/Source/Runtime/Core/Internal/Sanitizer/RaceDetectorTypes.h
--- a/Engine/Source/Runtime/Core/Internal/Sanitizer/RaceDetectorTypes.h
+++ b/Engine/Source/Runtime/Core/Internal/Sanitizer/RaceDetectorTypes.h
@@ -600,7 +600,7 @@ namespace UE::Sanitizer::RaceDetector {
 							}
 							break;
 							case EHistoryEntryType::FunctionExit:
-								Stack.Pop();
+								Stack.Pop(EAllowShrinking::No);
 								Offset += sizeof(FHistoryEntryBase);
 							break;
 							default:
@@ -983,7 +983,7 @@ namespace UE::Sanitizer::RaceDetector {
 		FSyncObjectRef& operator=(const FSyncObjectRef&) = delete;
 		FSyncObjectRef& operator=(FSyncObjectRef&&) = delete;
 
-		FSyncObjectRef(FSyncObjectRef&& Other)
+		INSTRUMENTATION_FUNCTION_ATTRIBUTES FSyncObjectRef(FSyncObjectRef&& Other)
 		{
 			Bank = Other.Bank;
 			Object = Other.Object;
