Removed deprecated support for the int continuation parameter in TFuture<void>.

diff --git a/Engine/Source/Runtime/Core/Public/Async/Future.h b/Engine/Source/Runtime/Core/Public/Async/Future.h
--- a/Engine/Source/Runtime/Core/Public/Async/Future.h
+++ b/Engine/Source/Runtime/Core/Public/Async/Future.h
@@ -13,25 +13,12 @@
 #include "HAL/Event.h"
 #include "HAL/PooledSyncEvent.h"
 #include "Misc/ScopeLock.h"
-#include "Templates/Models.h"
 #include <type_traits>
 
 template<typename ResultType> class TFuture;
 template<typename ResultType> class TSharedFuture;
 template<typename ResultType> class TPromise;
 
-namespace UE::Core::Private
-{
-	// A concept check for a type being callable
-	struct CIntCallable
-	{
-		template <typename Type>
-		auto Requires(Type& Callable, int32 Val) -> decltype(
-			Callable(Val)
-		);
-	};
-}
-
 /**
  * Base class for the internal state of asynchronous return values (futures).
  */
@@ -706,15 +693,7 @@ auto TFutureBase<ResultType>::Next(Func Continuation) //-> TFuture<decltype(Cont
 		if constexpr (std::is_void_v<ResultType>)
 		{
 			Self.Consume();
-			if constexpr (TModels_V<UE::Core::Private::CIntCallable, Func>)
-			{
-				UE_STATIC_DEPRECATE(5.6, true, "Passing continuations to TFuture<void>::Next which take int parameters has been deprecated - please remove the parameter.");
-				return Continuation(1);
-			}
-			else
-			{
-				return Continuation();
-			}
+			return Continuation();
 		}
 		else
 		{
@@ -735,3 +714,7 @@ TPromise<ResultType> MakeFulfilledPromise(ArgTypes&&... Args)
 #if UE_ENABLE_INCLUDE_ORDER_DEPRECATED_IN_5_7
 #include "Templates/Requires.h"
 #endif
+
+#if UE_ENABLE_INCLUDE_ORDER_DEPRECATED_IN_5_8
+#include "Templates/Models.h"
+#endif
