- Added the OptStreamGotTo optional param to most of the FParse::Value() functions (it was only on string versions before). Didn't add to FText version because it uses a different parsing code that doesn't return the end of the value

diff --git a/Engine/Source/Runtime/Core/Private/Misc/Parse.cpp b/Engine/Source/Runtime/Core/Private/Misc/Parse.cpp
--- a/Engine/Source/Runtime/Core/Private/Misc/Parse.cpp
+++ b/Engine/Source/Runtime/Core/Private/Misc/Parse.cpp
@@ -353,7 +353,7 @@ bool FParse::Param( const TCHAR* Stream, const TCHAR* Param )
 // 
 // Parse a string.
 //
-bool FParse::Value( const TCHAR* Stream, const TCHAR* Match, FString& Value, bool bShouldStopOnSeparator, const TCHAR** OptStreamGotTo)
+bool FParse::Value(const TCHAR* Stream, const TCHAR* Match, FString& Value, bool bShouldStopOnSeparator, const TCHAR** OptStreamGotTo)
 {
 	if (!Stream)
 	{
@@ -533,7 +533,7 @@ bool FParse::Text( const TCHAR* Buffer, FText& Value, const TCHAR* Namespace )
 // Parse an Text.
 // This is expected to in the form NSLOCTEXT("Namespace","Key","SourceString") or LOCTEXT("Key","SourceString")
 //
-bool FParse::Value( const TCHAR* Stream, const TCHAR* Match, FText& Value, const TCHAR* Namespace )
+bool FParse::Value(const TCHAR* Stream, const TCHAR* Match, FText& Value, const TCHAR* Namespace)
 {
 	// The FText 
 	Stream = FCString::Strifind( Stream, Match );
@@ -549,19 +549,19 @@ bool FParse::Value( const TCHAR* Stream, const TCHAR* Match, FText& Value, const
 //
 // Parse a quadword.
 //
-bool FParse::Value( const TCHAR* Stream, const TCHAR* Match, uint64& Value )
+bool FParse::Value(const TCHAR* Stream, const TCHAR* Match, uint64& Value, const TCHAR** OptStreamGotTo)
 {
-	return FParse::Value( Stream, Match, *(int64*)&Value );
+	return FParse::Value( Stream, Match, *(int64*)&Value, OptStreamGotTo);
 }
 
 //
 // Parse a signed quadword.
 //
-bool FParse::Value( const TCHAR* Stream, const TCHAR* Match, int64& Value )
+bool FParse::Value(const TCHAR* Stream, const TCHAR* Match, int64& Value, const TCHAR** OptStreamGotTo)
 {
 	TCHAR Temp[4096] = {};
 	TCHAR* Ptr = Temp;
-	if( FParse::Value( Stream, Match, Temp, UE_ARRAY_COUNT(Temp) ) )
+	if (FParse::Value( Stream, Match, Temp, UE_ARRAY_COUNT(Temp), true, OptStreamGotTo))
 	{
 		Value = 0;
 		bool Negative = (*Ptr=='-');
@@ -581,11 +581,11 @@ bool FParse::Value( const TCHAR* Stream, const TCHAR* Match, int64& Value )
 //
 // Get a name.
 //
-bool FParse::Value(	const TCHAR* Stream, const TCHAR* Match, FName& Name )
+bool FParse::Value(const TCHAR* Stream, const TCHAR* Match, FName& Name, const TCHAR** OptStreamGotTo)
 {
 	TCHAR TempStr[NAME_SIZE];
 
-	if( !FParse::Value(Stream,Match,TempStr,NAME_SIZE) )
+	if (!FParse::Value(Stream, Match, TempStr, NAME_SIZE, true, OptStreamGotTo))
 	{
 		return false;
 	}
@@ -598,10 +598,10 @@ bool FParse::Value(	const TCHAR* Stream, const TCHAR* Match, FName& Name )
 //
 // Get a uint32.
 //
-bool FParse::Value( const TCHAR* Stream, const TCHAR* Match, uint32& Value )
+bool FParse::Value(const TCHAR* Stream, const TCHAR* Match, uint32& Value, const TCHAR** OptStreamGotTo)
 {
 	TCHAR Temp[256];
-	if (!FParse::Value(Stream, Match, Temp, UE_ARRAY_COUNT(Temp)))
+	if (!FParse::Value(Stream, Match, Temp, UE_ARRAY_COUNT(Temp), true, OptStreamGotTo))
 	{
 		return false;
 	}
@@ -615,10 +615,10 @@ bool FParse::Value( const TCHAR* Stream, const TCHAR* Match, uint32& Value )
 //
 // Get a byte.
 //
-bool FParse::Value( const TCHAR* Stream, const TCHAR* Match, uint8& Value )
+bool FParse::Value(const TCHAR* Stream, const TCHAR* Match, uint8& Value, const TCHAR** OptStreamGotTo)
 {
 	TCHAR Temp[256];
-	if (!FParse::Value(Stream, Match, Temp, UE_ARRAY_COUNT(Temp)))
+	if (!FParse::Value(Stream, Match, Temp, UE_ARRAY_COUNT(Temp), true, OptStreamGotTo))
 	{
 		return false;
 	}
@@ -630,10 +630,10 @@ bool FParse::Value( const TCHAR* Stream, const TCHAR* Match, uint8& Value )
 //
 // Get a signed byte.
 //
-bool FParse::Value( const TCHAR* Stream, const TCHAR* Match, int8& Value )
+bool FParse::Value(const TCHAR* Stream, const TCHAR* Match, int8& Value, const TCHAR** OptStreamGotTo)
 {
 	TCHAR Temp[256];
-	if (!FParse::Value(Stream, Match, Temp, UE_ARRAY_COUNT(Temp)))
+	if (!FParse::Value(Stream, Match, Temp, UE_ARRAY_COUNT(Temp), true, OptStreamGotTo))
 	{
 		return false;
 	}
@@ -645,10 +645,10 @@ bool FParse::Value( const TCHAR* Stream, const TCHAR* Match, int8& Value )
 //
 // Get a word.
 //
-bool FParse::Value( const TCHAR* Stream, const TCHAR* Match, uint16& Value )
+bool FParse::Value(const TCHAR* Stream, const TCHAR* Match, uint16& Value, const TCHAR** OptStreamGotTo)
 {
 	TCHAR Temp[256];
-	if (!FParse::Value(Stream, Match, Temp, UE_ARRAY_COUNT(Temp)))
+	if (!FParse::Value(Stream, Match, Temp, UE_ARRAY_COUNT(Temp), true, OptStreamGotTo))
 	{
 		return false;
 	}
@@ -660,10 +660,10 @@ bool FParse::Value( const TCHAR* Stream, const TCHAR* Match, uint16& Value )
 //
 // Get a signed word.
 //
-bool FParse::Value( const TCHAR* Stream, const TCHAR* Match, int16& Value )
+bool FParse::Value(const TCHAR* Stream, const TCHAR* Match, int16& Value, const TCHAR** OptStreamGotTo)
 {
 	TCHAR Temp[256];
-	if (!FParse::Value(Stream, Match, Temp, UE_ARRAY_COUNT(Temp)))
+	if (!FParse::Value(Stream, Match, Temp, UE_ARRAY_COUNT(Temp), true, OptStreamGotTo))
 	{
 		return false;
 	}
@@ -675,10 +675,10 @@ bool FParse::Value( const TCHAR* Stream, const TCHAR* Match, int16& Value )
 //
 // Get a floating-point number.
 //
-bool FParse::Value( const TCHAR* Stream, const TCHAR* Match, float& Value )
+bool FParse::Value(const TCHAR* Stream, const TCHAR* Match, float& Value, const TCHAR** OptStreamGotTo)
 {
 	TCHAR Temp[256];
-	if (!FParse::Value(Stream, Match, Temp, UE_ARRAY_COUNT(Temp)))
+	if (!FParse::Value(Stream, Match, Temp, UE_ARRAY_COUNT(Temp), true, OptStreamGotTo))
 	{
 		return false;
 	}
@@ -690,10 +690,10 @@ bool FParse::Value( const TCHAR* Stream, const TCHAR* Match, float& Value )
 //
 // Get a double precision floating-point number.
 //
-bool FParse::Value(const TCHAR* Stream, const TCHAR* Match, double& Value)
+bool FParse::Value(const TCHAR* Stream, const TCHAR* Match, double& Value, const TCHAR** OptStreamGotTo)
 {
 	TCHAR Temp[256];
-	if (!FParse::Value(Stream, Match, Temp, UE_ARRAY_COUNT(Temp)))
+	if (!FParse::Value(Stream, Match, Temp, UE_ARRAY_COUNT(Temp), true, OptStreamGotTo))
 	{
 		return false;
 	}
@@ -706,10 +706,10 @@ bool FParse::Value(const TCHAR* Stream, const TCHAR* Match, double& Value)
 //
 // Get a signed double word.
 //
-bool FParse::Value( const TCHAR* Stream, const TCHAR* Match, int32& Value )
+bool FParse::Value(const TCHAR* Stream, const TCHAR* Match, int32& Value, const TCHAR** OptStreamGotTo)
 {
 	TCHAR Temp[256];
-	if (!FParse::Value(Stream, Match, Temp, UE_ARRAY_COUNT(Temp)))
+	if (!FParse::Value(Stream, Match, Temp, UE_ARRAY_COUNT(Temp), true, OptStreamGotTo))
 	{
 		return false;
 	}
@@ -721,10 +721,10 @@ bool FParse::Value( const TCHAR* Stream, const TCHAR* Match, int32& Value )
 //
 // Get a boolean value.
 //
-bool FParse::Bool(const TCHAR* Stream, const TCHAR* Match, bool& OnOff)
+bool FParse::Bool(const TCHAR* Stream, const TCHAR* Match, bool& OnOff, const TCHAR** OptStreamGotTo)
 {
 	TCHAR TempStr[16];
-	if (FParse::Value(Stream, Match, TempStr, UE_ARRAY_COUNT(TempStr)))
+	if (FParse::Value(Stream, Match, TempStr, UE_ARRAY_COUNT(TempStr), true, OptStreamGotTo))
 	{
 		OnOff = FCString::ToBool(TempStr);
 		return true;
@@ -738,7 +738,7 @@ bool FParse::Bool(const TCHAR* Stream, const TCHAR* Match, bool& OnOff)
 //
 // Get a globally unique identifier.
 //
-bool FParse::Value( const TCHAR* Stream, const TCHAR* Match, struct FGuid& Guid )
+bool FParse::Value(const TCHAR* Stream, const TCHAR* Match, struct FGuid& Guid, const TCHAR** OptStreamGotTo)
 {
 	TCHAR Temp[256];
 	if (!FParse::Value(Stream, Match, Temp, UE_ARRAY_COUNT(Temp)))
diff --git a/Engine/Source/Runtime/Core/Public/Misc/Parse.h b/Engine/Source/Runtime/Core/Public/Misc/Parse.h
--- a/Engine/Source/Runtime/Core/Public/Misc/Parse.h
+++ b/Engine/Source/Runtime/Core/Public/Misc/Parse.h
@@ -27,11 +27,11 @@ struct FParse
 	 */
 	static CORE_API bool Command( const TCHAR** Stream, const TCHAR* Match, bool bParseMightTriggerExecution = true );
 	/** Parses a name. */
-	static CORE_API bool Value( const TCHAR* Stream, const TCHAR* Match, FName& Name );
+	static CORE_API bool Value(const TCHAR* Stream, const TCHAR* Match, FName& Name, const TCHAR** OptStreamGotTo = nullptr);
 	/** Parses a uint32. */
-	static CORE_API bool Value( const TCHAR* Stream, const TCHAR* Match, uint32& Value );
+	static CORE_API bool Value(const TCHAR* Stream, const TCHAR* Match, uint32& Value, const TCHAR** OptStreamGotTo = nullptr);
 	/** Parses a globally unique identifier. */
-	static CORE_API bool Value( const TCHAR* Stream, const TCHAR* Match, struct FGuid& Guid );
+	static CORE_API bool Value(const TCHAR* Stream, const TCHAR* Match, struct FGuid& Guid, const TCHAR** OptStreamGotTo = nullptr);
 	/** Parses a string from a text string. 
 	 * @param Stream, the string you want to extract the value from.
 	 * @param Match, the identifier for the value in the stream.
@@ -43,21 +43,21 @@ struct FParse
 				of the end of the value within Stream. This permits consuming of stream in a loop where Match may
 				occur multiple times.
 	*/
-	static CORE_API bool Value( const TCHAR* Stream, const TCHAR* Match, TCHAR* Value, int32 MaxLen, bool bShouldStopOnSeparator=true, const TCHAR** OptStreamGotTo = nullptr);
+	static CORE_API bool Value(const TCHAR* Stream, const TCHAR* Match, TCHAR* Value, int32 MaxLen, bool bShouldStopOnSeparator=true, const TCHAR** OptStreamGotTo = nullptr);
 	/** Parses a byte. */
-	static CORE_API bool Value( const TCHAR* Stream, const TCHAR* Match, uint8& Value );
+	static CORE_API bool Value(const TCHAR* Stream, const TCHAR* Match, uint8& Value, const TCHAR** OptStreamGotTo = nullptr);
 	/** Parses a signed byte. */
-	static CORE_API bool Value( const TCHAR* Stream, const TCHAR* Match, int8& Value );
+	static CORE_API bool Value(const TCHAR* Stream, const TCHAR* Match, int8& Value, const TCHAR** OptStreamGotTo = nullptr);
 	/** Parses a uint16. */
-	static CORE_API bool Value( const TCHAR* Stream, const TCHAR* Match, uint16& Value );
+	static CORE_API bool Value(const TCHAR* Stream, const TCHAR* Match, uint16& Value, const TCHAR** OptStreamGotTo = nullptr);
 	/** Parses a signed word. */
-	static CORE_API bool Value( const TCHAR* Stream, const TCHAR* Match, int16& Value );
+	static CORE_API bool Value(const TCHAR* Stream, const TCHAR* Match, int16& Value, const TCHAR** OptStreamGotTo = nullptr);
 	/** Parses a floating-point value. */
-	static CORE_API bool Value( const TCHAR* Stream, const TCHAR* Match, float& Value );
+	static CORE_API bool Value(const TCHAR* Stream, const TCHAR* Match, float& Value, const TCHAR** OptStreamGotTo = nullptr);
 	/** Parses a double precision floating-point value. */
-	static CORE_API bool Value(const TCHAR* Stream, const TCHAR* Match, double& Value);
+	static CORE_API bool Value(const TCHAR* Stream, const TCHAR* Match, double& Value, const TCHAR** OptStreamGotTo = nullptr);
 	/** Parses a signed double word. */
-	static CORE_API bool Value( const TCHAR* Stream, const TCHAR* Match, int32& Value );
+	static CORE_API bool Value(const TCHAR* Stream, const TCHAR* Match, int32& Value, const TCHAR** OptStreamGotTo = nullptr);
 	/** Parses a string. 
 	 * @param Stream, the string you want to extract the value from.
 	 * @param Match, the identifier for the value in the stream.
@@ -68,15 +68,21 @@ struct FParse
 				of the end of the value within Stream. This permits consuming of stream in a loop where Match may
 				occur multiple times.
 	*/
-	static CORE_API bool Value( const TCHAR* Stream, const TCHAR* Match, FString& Value, bool bShouldStopOnSeparator =true, const TCHAR** OptStreamGotTo = nullptr);
-	/** Parses an FText. */
-	static CORE_API bool Value( const TCHAR* Stream, const TCHAR* Match, FText& Value, const TCHAR* Namespace = NULL );
+	static CORE_API bool Value(const TCHAR* Stream, const TCHAR* Match, FString& Value, bool bShouldStopOnSeparator =true, const TCHAR** OptStreamGotTo = nullptr);
+	/** Parses an FText. Doesn't support OptStreamGotTo like the other Value versions due to FText parsing
+	 * not returning the end of the text
+	 * @param Stream, the string you want to extract the value from.
+	 * @param Match, the identifier for the value in the stream.
+	 * @param Value, the destination FText value to be extracted to.
+	 * @param Namespace, an optional namespace to use when parsing texts that use LOCTEXT (default is an empty namespace).
+	*/
+	static CORE_API bool Value(const TCHAR* Stream, const TCHAR* Match, FText& Value, const TCHAR* Namespace = NULL);
 	/** Parses a quadword. */
-	static CORE_API bool Value( const TCHAR* Stream, const TCHAR* Match, uint64& Value );
+	static CORE_API bool Value(const TCHAR* Stream, const TCHAR* Match, uint64& Value, const TCHAR** OptStreamGotTo = nullptr);
 	/** Parses a signed quadword. */
-	static CORE_API bool Value( const TCHAR* Stream, const TCHAR* Match, int64& Value );
+	static CORE_API bool Value(const TCHAR* Stream, const TCHAR* Match, int64& Value, const TCHAR** OptStreamGotTo = nullptr);
 	/** Parses a boolean value. */
-	static CORE_API bool Bool( const TCHAR* Stream, const TCHAR* Match, bool& OnOff );
+	static CORE_API bool Bool(const TCHAR* Stream, const TCHAR* Match, bool& OnOff, const TCHAR** OptStreamGotTo = nullptr);
 	/** Get a line of Stream (everything up to, but not including, CR/LF. Returns 0 if ok, nonzero if at end of stream and returned 0-length string. */
 	static CORE_API bool Line( const TCHAR** Stream, TCHAR* Result, int32 MaxLen, bool Exact= false );
 	/** Get a line of Stream (everything up to, but not including, CR/LF. Returns 0 if ok, nonzero if at end of stream and returned 0-length string. */
