[RaceDetector]
- Improve memory mapping speed by making sure alignment function aren't instrumented

diff --git a/Engine/Source/Runtime/Core/Private/Sanitizer/RaceDetectorWindows.cpp b/Engine/Source/Runtime/Core/Private/Sanitizer/RaceDetectorWindows.cpp
--- a/Engine/Source/Runtime/Core/Private/Sanitizer/RaceDetectorWindows.cpp
+++ b/Engine/Source/Runtime/Core/Private/Sanitizer/RaceDetectorWindows.cpp
@@ -1072,10 +1072,26 @@ namespace UE::Sanitizer::RaceDetector::Platform {
 		return PageSize;
 	}
 
+	template <typename T>
+	FORCEINLINE INSTRUMENTATION_FUNCTION_ATTRIBUTES constexpr T AlignDown(T Val, uint64 Alignment)
+	{
+		static_assert(TIsIntegral<T>::Value || TIsPointer<T>::Value, "AlignDown expects an integer or pointer type");
+
+		return (T)(((uint64)Val) & ~(Alignment - 1));
+	}
+
+	template <typename T>
+	FORCEINLINE INSTRUMENTATION_FUNCTION_ATTRIBUTES constexpr T Align(T Val, uint64 Alignment)
+	{
+		static_assert(TIsIntegral<T>::Value || TIsPointer<T>::Value, "Align expects an integer or pointer type");
+
+		return (T)(((uint64)Val + Alignment - 1) & ~(Alignment - 1));
+	}
+
 	INSTRUMENTATION_FUNCTION_ATTRIBUTES bool IsShadowMemoryMapped(UPTRINT Base, UPTRINT Size)
 	{
 		UPTRINT AlignedStart = AlignDown(Base, PageSize);
-		UPTRINT AlignedEnd = Align(Base + Size, PageSize);
+		UPTRINT AlignedEnd   = Align(Base + Size, PageSize);
 
 		UPTRINT ShadowPageStart = ((AlignedStart - ShadowBase) >> PageSizeBitShift);
 		UPTRINT ShadowPageEnd = ((AlignedEnd - ShadowBase) >> PageSizeBitShift);
