Add missing overload for xDefine(const TArray<FString>...) in automation test framework

diff --git a/Engine/Source/Runtime/Core/Public/Misc/AutomationTest.h b/Engine/Source/Runtime/Core/Public/Misc/AutomationTest.h
--- a/Engine/Source/Runtime/Core/Public/Misc/AutomationTest.h
+++ b/Engine/Source/Runtime/Core/Public/Misc/AutomationTest.h
@@ -3300,6 +3300,11 @@ public:
 		//disabled
 	}
 
+	void xDescribe(const TArray<FString> InDescriptions, TFunction<void(int32 /*Iteration*/)> DoWork)
+	{
+		//disabled
+	}
+
 	void Describe(const TArray<FString> InDescriptions, TFunction<void(int32 /*Iteration*/)> DoWork)
 	{
 		LLM_SCOPE_BYNAME(TEXT("AutomationTest/Framework"));
