[Core] Fix bug with remove when duplicates values are located at the end of the container

#tests Unit tests added to cover various scenarios

diff --git a/Engine/Source/Runtime/Core/Public/Containers/CompactHashTable.h b/Engine/Source/Runtime/Core/Public/Containers/CompactHashTable.h
--- a/Engine/Source/Runtime/Core/Public/Containers/CompactHashTable.h
+++ b/Engine/Source/Runtime/Core/Public/Containers/CompactHashTable.h
@@ -32,7 +32,7 @@ namespace UE::Core::CompactHashTable::Private
 {
 	// Remove an element from the list and patch up any next index references, after this the spot is empty and needs to be filled
 	template<typename IndexType>
-	inline void RemoveInternal(const uint32 Index, const uint32 Key, IndexType* HashData, const uint32 HashCount, IndexType* NextIndexData, const uint32 NextIndexCount)
+	inline IndexType RemoveInternal(const uint32 Index, const uint32 Key, IndexType* HashData, const uint32 HashCount, IndexType* NextIndexData, const uint32 NextIndexCount)
 	{
 		const uint32 HashIndex = Key & (HashCount - 1);
 
@@ -42,9 +42,12 @@ namespace UE::Core::CompactHashTable::Private
 			{
 				// Next = Next->Next
 				*Lookup = NextIndexData[Index];
-				return;
+				return *Lookup;
 			}
 		}
+
+		checkNoEntry(); // At this point we must be pointing to something valid in the list so this lookup should never fail
+		return (IndexType)INDEX_NONE;
 	}
 }
 
@@ -71,6 +74,21 @@ namespace UE::Core::CompactHashTable
 		return 4; // only support uint32 for now, pick highest alignment so we don't change alignment between allocations
 	}
 
+	template<typename IndexType>
+	[[nodiscard]] UE_FORCEINLINE_HINT constexpr uint32 IndexTypeToSizeType(const IndexType Index)
+	{
+		if constexpr (sizeof(IndexType) == 4)
+		{
+			return Index;
+		}
+		else
+		{
+			// Convert smaller type INDEX_NONE to uint32 INDEX_NONE if data is invalid
+			constexpr IndexType InvalidIndex = (IndexType)INDEX_NONE;
+			return Index == InvalidIndex ? (uint32)INDEX_NONE : (uint32)Index;
+		}
+	}
+
 	/* Calculate the size of the hash table from the number of elements in the set */
 	[[nodiscard]] inline constexpr size_t GetHashCount(uint32 NumElements)
 	{
@@ -92,23 +110,15 @@ namespace UE::Core::CompactHashTable
 	[[nodiscard]] inline uint32 GetFirst(uint32 Key, const IndexType* HashData, const uint32 HashCount)
 	{
 		const uint32 HashIndex = Key & (HashCount - 1);
-		const IndexType FirstIndex = HashData[HashIndex];
-
-		// Convert smaller type INDEX_NONE to uint32 INDEX_NONE if data is invalid
-		constexpr IndexType InvalidIndex = (IndexType)INDEX_NONE;
-		return FirstIndex == InvalidIndex ? (uint32)INDEX_NONE : (uint32)FirstIndex;
+		return IndexTypeToSizeType(HashData[HashIndex]);
 	}
 
 	/* Return the first index of a key from the hash table portion */
 	template<typename IndexType>
 	[[nodiscard]] inline uint32 GetFirstByIndex(uint32 HashIndex, const IndexType* HashData, const uint32 HashCount)
 	{
-		checkSlow(HashIndex <= HashCount);
-		const IndexType FirstIndex = HashData[HashIndex];
-
-		// Convert smaller type INDEX_NONE to uint32 INDEX_NONE if data is invalid
-		constexpr IndexType InvalidIndex = (IndexType)INDEX_NONE;
-		return FirstIndex == InvalidIndex ? (uint32)INDEX_NONE : (uint32)FirstIndex;
+		checkSlow(HashIndex < HashCount);
+		return IndexTypeToSizeType(HashData[HashIndex]);
 	}
 
 	/* Given an existing index, return the next index in case there was a collision in the has table */
@@ -116,11 +126,7 @@ namespace UE::Core::CompactHashTable
 	[[nodiscard]] inline uint32 GetNext(uint32 Index, const IndexType* NextIndexData, const uint32 NextIndexCount)
 	{
 		checkSlow(Index < NextIndexCount);
-		const IndexType NextIndex = NextIndexData[Index];
-		
-		// Convert smaller type INDEX_NONE to uint32 INDEX_NONE if data is invalid
-		constexpr IndexType InvalidIndex = (IndexType)INDEX_NONE;
-		return NextIndex == InvalidIndex ? (uint32)INDEX_NONE : (uint32)NextIndex;
+		return IndexTypeToSizeType(NextIndexData[Index]);
 	}
 
 	/* Do a full search for an existing element in the table given a function to compare if a found element is what you're looking for */
@@ -154,55 +160,73 @@ namespace UE::Core::CompactHashTable
 	 If the item to remove is the last element then the last element's key will be ignored (you can skip calculating it if it's expensive)
 	*/
 	template<typename IndexType>
-	inline void Remove(const uint32 Index, const uint32 Key, const uint32 LastIndex, uint32 OptLastKey, IndexType* HashData, const uint32 HashCount, IndexType* NextIndexData, const uint32 NextIndexCount)
+	inline uint32 Remove(const uint32 Index, const uint32 Key, const uint32 LastIndex, uint32 OptLastKey, IndexType* HashData, const uint32 HashCount, IndexType* NextIndexData, const uint32 NextIndexCount)
 	{
 		checkSlow(LastIndex < NextIndexCount && Index <= LastIndex);
 
-		UE::Core::CompactHashTable::Private::RemoveInternal<IndexType>(Index, Key, HashData, HashCount, NextIndexData, NextIndexCount);
+		IndexType NextIndex = UE::Core::CompactHashTable::Private::RemoveInternal<IndexType>(Index, Key, HashData, HashCount, NextIndexData, NextIndexCount);
 
 		if (Index != LastIndex)
 		{
-			// Remove the last element and add it into the empty spot
-			UE::Core::CompactHashTable::Private::RemoveInternal<IndexType>(LastIndex, OptLastKey, HashData, HashCount, NextIndexData, NextIndexCount);
-
+			// Remap the last element into the now empty spot
 			const uint32 HashIndex = OptLastKey & (HashCount - 1);
-			NextIndexData[Index] = HashData[HashIndex];
-			HashData[HashIndex] = (IndexType)Index;
+
+			for (IndexType* Lookup = &HashData[HashIndex]; *Lookup < NextIndexCount; Lookup = &NextIndexData[*Lookup])
+			{
+				if (*Lookup == LastIndex)
+				{
+					// Fixup to look at now empty slot
+					*Lookup = (IndexType)Index;
+					// Now empty slot looks where the last item looked
+					NextIndexData[Index] = NextIndexData[LastIndex];
+					break;
+				}
+			}
 		}
+
+		// If the next index points to the last element, then it will get relocated into the current element so return the current instead of the next
+		return NextIndex == LastIndex ? Index : IndexTypeToSizeType(NextIndex);
 	}
 
 	/* Remove an element from the list, shift all indexes down to preserve the order of elements in a list
 	 This is a very expensive operation so it should only be used if absolutely necessary (i.e. generally only for user facing data)
 	*/
 	template<typename IndexType>
-	inline void RemoveStable(const uint32 Index, const uint32 Key, IndexType* HashData, const uint32 HashCount, IndexType* NextIndexData, const uint32 NextIndexCount)
+	inline uint32 RemoveStable(const uint32 Index, const uint32 Key, IndexType* HashData, const uint32 HashCount, IndexType* NextIndexData, const uint32 NextIndexCount)
 	{
 		checkSlow(Index < NextIndexCount);
 
-		UE::Core::CompactHashTable::Private::RemoveInternal<IndexType>(Index, Key, HashData, HashCount, NextIndexData, NextIndexCount);
+		constexpr IndexType InvalidIndex = (IndexType)INDEX_NONE;
 
-		// For the hash indexes, just decrement any that are bigger than the removed element
-		for (uint32 HashIndex = 0; HashIndex < HashCount; ++HashIndex)
+		// If TestIndex is valid and above Index just remove 1 so it will reference the new location that elements should get relocated to
+		auto FixIndex = [Index](IndexType& TestIndex)
 		{
-			if (HashData[HashIndex] > Index && HashData[HashIndex] != (IndexType)INDEX_NONE)
+			if (TestIndex != InvalidIndex && TestIndex > Index)
 			{
-				--HashData[HashIndex];
+				--TestIndex;
 			}
+		};
+
+		IndexType NextIndex = UE::Core::CompactHashTable::Private::RemoveInternal<IndexType>(Index, Key, HashData, HashCount, NextIndexData, NextIndexCount);
+		FixIndex(NextIndex);
+
+		// For the hash indexes, just decrement any that are bigger than the removed element
+		for (uint32 HashIndex = 0; HashIndex < HashCount; ++HashIndex)
+		{
+			FixIndex(HashData[HashIndex]);
 		}
 
 		// Decrement values for all next index elements that are before the removed element
 		for (uint32 NextIndexIndex = 0; NextIndexIndex < Index; ++NextIndexIndex)
 		{
-			if (NextIndexData[NextIndexIndex] > Index && NextIndexData[NextIndexIndex] != (IndexType)INDEX_NONE)
-			{
-				--NextIndexData[NextIndexIndex];
-			}
+			FixIndex(NextIndexData[NextIndexIndex]);
 		}
 
 		// Move AND Decrement values for all next index elements that are after the removed element
 		for (uint32 NextIndexIndex = Index + 1; NextIndexIndex < NextIndexCount; ++NextIndexIndex)
 		{
-			if (NextIndexData[NextIndexIndex] > Index && NextIndexData[NextIndexIndex] != (IndexType)INDEX_NONE)
+			// Not using FixIndex so we can use the else
+			if (NextIndexData[NextIndexIndex] != InvalidIndex && NextIndexData[NextIndexIndex] > Index)
 			{
 				NextIndexData[NextIndexIndex - 1] = NextIndexData[NextIndexIndex] - 1;
 			}
@@ -211,6 +235,8 @@ namespace UE::Core::CompactHashTable
 				NextIndexData[NextIndexIndex - 1] = NextIndexData[NextIndexIndex];
 			}
 		}
+
+		return IndexTypeToSizeType(NextIndex);
 	}
 }
 
@@ -370,17 +396,19 @@ public:
 		#undef UE_COMPACTHASHTABLE_EXECUTEBYTYPE
 	}
 
-	UE_FORCEINLINE_HINT void Remove(uint32 Index, uint32 Key, uint32 LastIndex, uint32 OptLastKey) const
+	UE_FORCEINLINE_HINT uint32 Remove(uint32 Index, uint32 Key, uint32 LastIndex, uint32 OptLastKey) const
 	{
 		#define UE_COMPACTHASHTABLE_EXECUTEBYTYPE(Type) return UE::Core::CompactHashTable::Remove(Index, Key, LastIndex, OptLastKey, (Type*)(HashData), HashCount, (Type*)(NextIndexData), NextIndexCount)
 		UE_COMPACTHASHTABLE_CALLBYTYPE(NextIndexCount);
 		#undef UE_COMPACTHASHTABLE_EXECUTEBYTYPE
+		return INDEX_NONE;
 	}
 
-	UE_FORCEINLINE_HINT void RemoveStable(uint32 Index, uint32 Key) const
+	UE_FORCEINLINE_HINT uint32 RemoveStable(uint32 Index, uint32 Key) const
 	{
 		#define UE_COMPACTHASHTABLE_EXECUTEBYTYPE(Type) return UE::Core::CompactHashTable::RemoveStable(Index, Key, (Type*)(HashData), HashCount, (Type*)(NextIndexData), NextIndexCount)
 		UE_COMPACTHASHTABLE_CALLBYTYPE(NextIndexCount);
 		#undef UE_COMPACTHASHTABLE_EXECUTEBYTYPE
+		return INDEX_NONE;
 	}
 };
diff --git a/Engine/Source/Runtime/Core/Public/Containers/CompactSet.h.inl b/Engine/Source/Runtime/Core/Public/Containers/CompactSet.h.inl
--- a/Engine/Source/Runtime/Core/Public/Containers/CompactSet.h.inl
+++ b/Engine/Source/Runtime/Core/Public/Containers/CompactSet.h.inl
@@ -1496,38 +1496,37 @@ private:
 	}
 
 	template <bool IsStable = false>
-	void RemoveByIndexAndHash(const SizeType ElementIndex, const uint32 KeyHash)
+	SizeType RemoveByIndexAndHash(const SizeType ElementIndex, const uint32 KeyHash)
 	{
 		checkf(ElementIndex >= 0 && ElementIndex < this->NumElements, TEXT("Invalid ElementIndex passed to UE_TCOMPACT_SET::RemoveByIndex"));
 
 		ElementType* ElementsData = GetData();
 		FCompactHashTableView HashTable = GetHashTableView();
 
-		const SizeType LastElementIndex = this->NumElements - 1;
+		--this->NumElements;
+
+		const SizeType LastElementIndex = this->NumElements;
 		if (ElementIndex == LastElementIndex)
 		{
-			HashTable.Remove(ElementIndex, KeyHash, ElementIndex, KeyHash);
 			ElementsData[LastElementIndex].~ElementType();
+			return HashTable.Remove(ElementIndex, KeyHash, ElementIndex, KeyHash);
+		}
+
+		if constexpr (IsStable)
+		{
+			ElementsData[ElementIndex].~ElementType();
+			RelocateConstructItems<ElementType>(ElementsData + ElementIndex, ElementsData + ElementIndex + 1, LastElementIndex - ElementIndex);
+
+			return HashTable.RemoveStable(ElementIndex, KeyHash);
 		}
 		else
 		{
-			if constexpr (IsStable)
-			{
-				HashTable.RemoveStable(ElementIndex, KeyHash);
+			const uint32 LastElementHash = KeyFuncs::GetKeyHash(KeyFuncs::GetSetKey(ElementsData[LastElementIndex]));
 
-				ElementsData[ElementIndex].~ElementType();
-				RelocateConstructItems<ElementType>(ElementsData + ElementIndex, ElementsData + ElementIndex + 1, LastElementIndex - ElementIndex);
-			}
-			else
-			{
-				const uint32 LastElementHash = KeyFuncs::GetKeyHash(KeyFuncs::GetSetKey(ElementsData[LastElementIndex]));
-				HashTable.Remove(ElementIndex, KeyHash, LastElementIndex, LastElementHash);
+			MoveByRelocate(ElementsData[ElementIndex], ElementsData[LastElementIndex]);
 
-				MoveByRelocate(ElementsData[ElementIndex], ElementsData[LastElementIndex]);
-			}
+			return HashTable.Remove(ElementIndex, KeyHash, LastElementIndex, LastElementHash);
 		}
-
-		--this->NumElements;
 	}
 
 	template<typename ComparableKey, bool IsStable = false>
@@ -1539,42 +1538,27 @@ private:
 		const ElementType* ElementsData = GetData();
 		FCompactHashTableView HashTable = GetHashTableView();
 
-		SizeType LastElementIndex = INDEX_NONE;
 		SizeType ElementIndex = HashTable.GetFirst(KeyHash);
 
 		while (ElementIndex != INDEX_NONE)
 		{
 			if (KeyFuncs::Matches(KeyFuncs::GetSetKey(ElementsData[ElementIndex]), Key))
 			{
-				RemoveByIndexAndHash<IsStable>(ElementIndex, KeyHash);
-				NumRemovedElements++;
-
 				if constexpr (!KeyFuncs::bAllowDuplicateKeys)
 				{
+					RemoveByIndexAndHash<IsStable>(ElementIndex, KeyHash);
 					// If the hash disallows duplicate keys, we're done removing after the first matched key.
-					break;
+					return 1;
 				}
 				else
 				{
-					if (LastElementIndex == INDEX_NONE)
-					{
-						ElementIndex = HashTable.GetFirst(KeyHash);
-					}
-					else
-					{
-						if (LastElementIndex == this->NumElements) // Would have been remapped to ElementIndex
-						{
-							LastElementIndex = ElementIndex;
-						}
-
-						ElementIndex = HashTable.GetNext(LastElementIndex, this->NumElements);
-					}
+					ElementIndex = RemoveByIndexAndHash<IsStable>(ElementIndex, KeyHash);
+					NumRemovedElements++;
 				}
 			}
 			else
 			{
-				LastElementIndex = ElementIndex;
-				ElementIndex = HashTable.GetNext(LastElementIndex, this->NumElements);
+				ElementIndex = HashTable.GetNext(ElementIndex, this->NumElements);
 			}
 		}
 
diff --git a/Engine/Source/Runtime/Core/Public/Containers/Map.h b/Engine/Source/Runtime/Core/Public/Containers/Map.h
--- a/Engine/Source/Runtime/Core/Public/Containers/Map.h
+++ b/Engine/Source/Runtime/Core/Public/Containers/Map.h
@@ -7,6 +7,8 @@
 #include "Algo/Reverse.h"
 #include "Containers/ContainerElementTypeCompatibility.h"
 #include "Containers/Set.h"
+#include "Containers/SparseSet.h"
+#include "Containers/CompactSet.h"
 #include "Containers/UnrealString.h"
 #include "Misc/AssertionMacros.h"
 #include "Misc/StructBuilder.h"
diff --git a/Engine/Source/Runtime/Core/Tests/Containers/MapTest.cpp b/Engine/Source/Runtime/Core/Tests/Containers/MapTest.cpp
--- a/Engine/Source/Runtime/Core/Tests/Containers/MapTest.cpp
+++ b/Engine/Source/Runtime/Core/Tests/Containers/MapTest.cpp
@@ -1,39 +1,485 @@
 ﻿// Copyright Epic Games, Inc. All Rights Reserved.
 
-#if WITH_LOW_LEVEL_TESTS
+#if WITH_TESTS
+
+#include "Containers/SparseSet.h"
+#include "Containers/CompactSet.h"
 
 #include "Containers/Map.h"
 
 #include "Tests/TestHarnessAdapter.h"
 
-TEST_CASE_NAMED(FMapIteratorTest, "System::Core::Containers::TMap::TIterator::RemoveCurrent", "[Core][Containers][TMap]")
+struct FSparseMapTestResolver
 {
-	TMap<int, int> MapUnderTest;
+	template<typename InKeyType, typename InValueType, typename Allocator = FDefaultSparseSetAllocator, typename KeyFuncs = TDefaultMapHashableKeyFuncs<InKeyType, InValueType, false>>
+	using MapType = TSparseMap<InKeyType, InValueType, Allocator, KeyFuncs>;
+
+	template<typename InKeyType, typename InValueType, typename Allocator = FDefaultSparseSetAllocator, typename KeyFuncs = TDefaultMapHashableKeyFuncs<InKeyType, InValueType, true>>
+	using MultiMapType = TSparseMultiMap<InKeyType, InValueType, Allocator, KeyFuncs>;
+};
 
-	MapUnderTest.Add(5, 55);
-	MapUnderTest.Add(1, 11);
-	MapUnderTest.Add(3, 33);
-	MapUnderTest.Add(4, 44);
+struct FCompactMapTestResolver
+{
+	template<typename InKeyType, typename InValueType, typename Allocator = FDefaultCompactSetAllocator, typename KeyFuncs = TDefaultMapHashableKeyFuncs<InKeyType, InValueType, false>>
+	using MapType = TCompactMap<InKeyType, InValueType, Allocator, KeyFuncs>;
 
-	int VisitedItemsCount = 0;
-	for (auto It = MapUnderTest.CreateIterator(); It; ++It)
+	template<typename InKeyType, typename InValueType, typename Allocator = FDefaultCompactSetAllocator, typename KeyFuncs = TDefaultMapHashableKeyFuncs<InKeyType, InValueType, true>>
+	using MultiMapType = TCompactMultiMap<InKeyType, InValueType, Allocator, KeyFuncs>;
+};
+
+template<typename MapResolver = FSparseMapTestResolver>
+struct TMapTestHelper
+{
+	template<typename Key, typename Value>
+	using MapType = typename MapResolver::template MapType<Key, Value>;
+
+	template<typename Key, typename Value>
+	using MultiMapType = typename MapResolver::template MultiMapType<Key, Value>;
+
+	TMapTestHelper()
 	{
-		++VisitedItemsCount;
-		if (It->Key == 1)
+		SECTION("Default Constructor")
 		{
-			It.RemoveCurrent();
+			MapType<int, int> Map;
+			CHECK(Map.Num() == 0);
+
+			MultiMapType<int, int> MultiMap;
+			CHECK(MultiMap.Num() == 0);
 		}
-	}
 
-	CHECK(MapUnderTest.Num() == 3);
+		SECTION("Iterator Test")
+		{
+			MapType<int, int> MapUnderTest;
+
+			MapUnderTest.Add(5, 55);
+			MapUnderTest.Add(1, 11);
+			MapUnderTest.Add(3, 33);
+			MapUnderTest.Add(4, 44);
+
+			int VisitedItemsCount = 0;
+			for (auto It = MapUnderTest.CreateIterator(); It; ++It)
+			{
+				++VisitedItemsCount;
+				if (It->Key == 1)
+				{
+					It.RemoveCurrent();
+				}
+			}
+
+			CHECK(MapUnderTest.Num() == 3);
 	
-	CHECK(VisitedItemsCount == 4);
+			CHECK(VisitedItemsCount == 4);
 	
-	CHECK(MapUnderTest.Find(5) != nullptr);
-	CHECK(MapUnderTest.Find(1) == nullptr);
-	CHECK(MapUnderTest.Find(3) != nullptr);
-	CHECK(MapUnderTest.Find(4) != nullptr);
+			CHECK(MapUnderTest.Find(5) != nullptr);
+			CHECK(MapUnderTest.Find(1) == nullptr);
+			CHECK(MapUnderTest.Find(3) != nullptr);
+			CHECK(MapUnderTest.Find(4) != nullptr);
+		}
+
+		SECTION("Multi Map Test")
+		{
+			const int32 Keys = 100;
+			const int32 Values = 10;
+
+			MultiMapType<int, int> MultiMap;
+
+			for (int32 Value = 0; Value < Values; ++Value)
+			{
+				for (int32 Key = 0; Key < Keys; ++Key)
+				{
+					MultiMap.Add(Key, (Key << 16) + Value);
+				}
+			}
+
+			CHECK(MultiMap.Num() == Keys * Values);
+
+			{
+				int32 ValueCount = 0;
+				for (auto KeyIter = MultiMap.CreateKeyIterator(3); KeyIter; ++KeyIter)
+				{
+					++ValueCount;
+				}
+
+				CHECK(ValueCount == Values);
+			}
+
+			for (int32 Key = 0; Key < Keys; ++Key)
+			{
+				CHECK(MultiMap.Remove(Key) == Values);
+				CHECK(MultiMap.Num() == Values * (Keys - (Key + 1)));
+			}
+		}
+
+		SECTION("Very Specific broken sequence for multi map remove")
+		{
+			MultiMapType<int, int> Map;
+			{
+				Map.Add(70, 1);
+				Map.Add(8, 3);
+				Map.Add(88, 4);
+				Map.Add(30, 6);
+				Map.Add(103, 9);
+				Map.Add(6, 10);
+				Map.Add(0, 12);
+				Map.Add(62, 15);
+				Map.Add(78, 24);
+				Map.Add(70, 30);
+				Map.Add(53, 40);
+				Map.Add(53, 41);
+				Map.Add(72, 44);
+				Map.Add(72, 47);
+				Map.Add(8, 48);
+				Map.Add(8, 49);
+				Map.Add(70, 54);
+				Map.Add(70, 55);
+				Map.Add(53, 56);
+				Map.Add(53, 57);
+				Map.Add(78, 59);
+				Map.Add(93, 64);
+				Map.Add(93, 65);
+				Map.Add(28, 67);
+				Map.Add(28, 68);
+				Map.Add(28, 72);
+				Map.Add(28, 73);
+				Map.Add(30, 74);
+				Map.Add(30, 75);
+				Map.Add(93, 76);
+				Map.Add(93, 77);
+				Map.Add(88, 78);
+				Map.Add(88, 79);
+				Map.Add(93, 80);
+				Map.Add(93, 81);
+				Map.Add(90, 82);
+				Map.Add(90, 83);
+				Map.Add(66, 84);
+				Map.Add(2, 86);
+				Map.Add(64, 91);
+				Map.Add(2, 94);
+				Map.Add(2, 106);
+				Map.Add(2, 107);
+				Map.Add(66, 108);
+				Map.Add(64, 111);
+				Map.Add(103, 117);
+				Map.Add(104, 118);
+				Map.Add(60, 120);
+				Map.Add(60, 121);
+				Map.Add(62, 123);
+				Map.Add(93, 132);
+				Map.Add(90, 134);
+				Map.Add(52, 138);
+				Map.Add(52, 139);
+				Map.Add(52, 140);
+				Map.Add(52, 141);
+				Map.Add(52, 142);
+				Map.Add(52, 143);
+				Map.Add(52, 144);
+				Map.Add(52, 145);
+				Map.Add(52, 146);
+				Map.Add(52, 147);
+				Map.Add(52, 148);
+				Map.Add(52, 149);
+				Map.Add(52, 150);
+				Map.Add(52, 151);
+				Map.Add(52, 152);
+				Map.Add(52, 153);
+				Map.Add(52, 156);
+				Map.Add(102, 158);
+				Map.Add(8, 160);
+				Map.Add(8, 162);
+				Map.Add(68, 165);
+				Map.Add(53, 172);
+				Map.Add(54, 174);
+				Map.Add(54, 175);
+				Map.Add(54, 177);
+				Map.Add(102, 178);
+				Map.Add(102, 179);
+				Map.Add(53, 181);
+				Map.Add(53, 182);
+				Map.Add(53, 183);
+				Map.Add(52, 184);
+				Map.Add(52, 185);
+				Map.Add(68, 186);
+				Map.Add(68, 187);
+				Map.Add(54, 188);
+				Map.Add(54, 189);
+				Map.Add(54, 190);
+				Map.Add(54, 191);
+				Map.Add(53, 192);
+				Map.Add(53, 193);
+				Map.Add(53, 194);
+				Map.Add(53, 195);
+				Map.Add(8, 196);
+				Map.Add(68, 198);
+				Map.Add(30, 201);
+				Map.Add(90, 203);
+				Map.Add(6, 204);
+				Map.Add(104, 207);
+				Map.Add(60, 209);
+				Map.Add(2, 210);
+				Map.Add(52, 212);
+				Map.Add(52, 213);
+				Map.Add(52, 214);
+				Map.Add(52, 215);
+				Map.Add(52, 216);
+				Map.Add(52, 217);
+				Map.Add(52, 218);
+				Map.Add(52, 219);
+				Map.Add(52, 221);
+				Map.Add(10, 222);
+				Map.Add(10, 223);
+				Map.Add(10, 224);
+				Map.Add(10, 225);
+				Map.Add(10, 227);
+				Map.Add(54, 228);
+				Map.Add(54, 229);
+				Map.Add(53, 230);
+				Map.Add(53, 231);
+				Map.Add(53, 232);
+				Map.Add(53, 233);
+				Map.Add(54, 234);
+				Map.Add(54, 235);
+				Map.Add(10, 238);
+				Map.Add(10, 239);
+				Map.Add(53, 241);
+				Map.Add(54, 242);
+				Map.Add(54, 246);
+				Map.Add(54, 247);
+				Map.Add(54, 248);
+				Map.Add(54, 249);
+				Map.Add(53, 254);
+				Map.Add(53, 256);
+				Map.Add(53, 257);
+				Map.Add(28, 258);
+				Map.Add(28, 259);
+				Map.Add(93, 262);
+				Map.Add(88, 265);
+				Map.Add(28, 266);
+				Map.Add(28, 267);
+				Map.Add(6, 281);
+				Map.Add(60, 283);
+				Map.Add(2, 284);
+				Map.Add(2, 285);
+				Map.Add(104, 288);
+				Map.Add(104, 289);
+				Map.Add(103, 291);
+				Map.Add(5, 292);
+				Map.Add(5, 293);
+				Map.Add(6, 296);
+				Map.Add(6, 297);
+				Map.Add(0, 298);
+				Map.Add(0, 299);
+				Map.Add(5, 305);
+				Map.Add(5, 306);
+				Map.Add(103, 308);
+				Map.Add(103, 309);
+				Map.Add(104, 310);
+				Map.Add(104, 311);
+				Map.Add(3, 312);
+				Map.Add(3, 315);
+				Map.Add(2, 323);
+				Map.Add(0, 324);
+				Map.Add(0, 326);
+				Map.Add(0, 327);
+				Map.Add(0, 329);
+				Map.Add(0, 331);
+				Map.Add(2, 350);
+				Map.Add(2, 351);
+				Map.Add(60, 353);
+				Map.Add(10, 355);
+				Map.Add(102, 357);
+				Map.Add(103, 359);
+				Map.Add(6, 360);
+				Map.Add(2, 362);
+				Map.Add(2, 363);
+				Map.Add(2, 364);
+				Map.Add(2, 365);
+				Map.Add(2, 366);
+				Map.Add(2, 367);
+				Map.Add(64, 369);
+				Map.Add(64, 372);
+				Map.Add(60, 376);
+			}
+
+			Map.Remove(70);
+			Map.Add(120, 54);
+			Map.Add(120, 30);
+			Map.Remove(2);
+			Map.Add(123, 366);
+			Map.Add(123, 365);
+			Map.Add(123, 364);
+			Map.Add(123, 363);
+			Map.Add(123, 362);
+			Map.Add(123, 351);
+			Map.Add(123, 350);
+			Map.Add(123, 323);
+			Map.Add(2, 285);
+			Map.Add(123, 210);
+			Map.Add(123, 107);
+			Map.Add(123, 106);
+			Map.Add(2, 94);
+			Map.Add(123, 86);
+			Map.Remove(120);
+			Map.Remove(2);
+			Map.Remove(123);
+			Map.Add(129, 323);
+			Map.Add(129, 350);
+			Map.Add(129, 351);
+			Map.Add(129, 362);
+			Map.Add(129, 363);
+			Map.Add(129, 364);
+			Map.Add(129, 365);
+			Map.Add(129, 366);
+			Map.Remove(0);
+			Map.Add(131, 326);
+			Map.Add(131, 329);
+			Map.Add(131, 331);
+			Map.Add(131, 299);
+			Map.Add(131, 298);
+			Map.Add(131, 12);
+			Map.Remove(129);
+			Map.Add(132, 323);
+			Map.Add(132, 350);
+			Map.Add(132, 351);
+			Map.Add(132, 364);
+			Map.Remove(131);
+			Map.Add(135, 326);
+			Map.Add(135, 329);
+			Map.Remove(132);
+			Map.Add(132, 323);
+			Map.Add(132, 350);
+			Map.Remove(104);
+			Map.Add(139, 311);
+			Map.Add(138, 288);
+			Map.Add(139, 207);
+			Map.Add(138, 118);
+			Map.Remove(132);
+			Map.Remove(139);
+			Map.Remove(138);
+			Map.Remove(135);
+			Map.Remove(6);
+			Map.Add(149, 204);
+			Map.Add(149, 10);
+			Map.Remove(5);
+			Map.Remove(149);
+			Map.Remove(28);
+			Map.Add(155, 266);
+			Map.Add(28, 259);
+			Map.Add(155, 258);
+			Map.Add(28, 67);
+			Map.Remove(88);
+			Map.Add(88, 78);
+			Map.Add(88, 4);
+			Map.Remove(28);
+			Map.Remove(88);
+			Map.Remove(54);
+			Map.Add(163, 248);
+			Map.Add(162, 247);
+			Map.Add(163, 246);
+			Map.Add(162, 242);
+			Map.Add(163, 235);
+			Map.Add(163, 234);
+			Map.Add(162, 229);
+			Map.Add(162, 228);
+			Map.Add(162, 191);
+			Map.Add(162, 190);
+			Map.Add(162, 177);
+			Map.Add(162, 174);
+			Map.Remove(10);
+			Map.Add(10, 225);
+			Map.Add(164, 355);
+			Map.Add(164, 224);
+			Map.Add(10, 223);
+			Map.Remove(162);
+			Map.Add(166, 229);
+			Map.Add(166, 228);
+			Map.Add(166, 242);
+			Map.Add(166, 247);
+			Map.Remove(164);
+			Map.Remove(166);
+			Map.Remove(163);
+			Map.Remove(52);
+			Map.Add(174, 219);
+			Map.Add(174, 217);
+			Map.Add(174, 216);
+			Map.Add(174, 213);
+			Map.Add(174, 212);
+			Map.Add(174, 185);
+			Map.Add(174, 184);
+			Map.Add(174, 156);
+			Map.Add(174, 153);
+			Map.Add(174, 152);
+			Map.Add(174, 151);
+			Map.Add(174, 150);
+			Map.Add(174, 149);
+			Map.Add(174, 148);
+			Map.Add(174, 147);
+			Map.Add(174, 146);
+			Map.Add(174, 145);
+			Map.Add(174, 144);
+			Map.Add(174, 143);
+			Map.Add(174, 142);
+			Map.Add(174, 141);
+			Map.Add(174, 140);
+			Map.Add(174, 139);
+			Map.Add(174, 138);
+			Map.Remove(60);
+			Map.Add(177, 209);
+			Map.Add(177, 353);
+			Map.Add(177, 376);
+			Map.Add(177, 121);
+			Map.Remove(174);
+			Map.Add(179, 143);
+			Map.Add(174, 142);
+			Map.Add(174, 141);
+			Map.Add(174, 140);
+			Map.Add(174, 139);
+			Map.Add(179, 144);
+			Map.Add(179, 145);
+			Map.Add(179, 146);
+			Map.Add(179, 147);
+			Map.Add(174, 148);
+			Map.Add(174, 149);
+			Map.Add(174, 152);
+			Map.Add(174, 153);
+			Map.Add(179, 156);
+			Map.Add(179, 184);
+			Map.Add(179, 185);
+			Map.Add(178, 212);
+			Map.Add(178, 213);
+			Map.Add(178, 216);
+			Map.Add(178, 219);
+			Map.Remove(102);
+			Map.Add(181, 178);
+			Map.Add(181, 158);
+			Map.Remove(178);
+			Map.Remove(179);
+			Map.Add(184, 147);
+			Map.Add(184, 146);
+			Map.Add(184, 144);
+			Map.Add(184, 143);
+			Map.Remove(174);
+			Map.Add(187, 149);
+			Map.Add(187, 148);
+			Map.Add(187, 141);
+			Map.Add(187, 142);
+			Map.Remove(184);
+			Map.Remove(187);
+
+			CHECK(Map.Contains(53));
+			Map.Remove(53);
+			CHECK(!Map.Contains(53));
+		}
+	}
+};
+
+TEST_CASE_NAMED(FSparseMapTest, "System::Core::Containers::SparseMap", "[Core][Containers][SparseMap]")
+{
+	TMapTestHelper<FSparseMapTestResolver>();
 }
 
+TEST_CASE_NAMED(FCompactMapTest, "System::Core::Containers::CompactMap", "[Core][Containers][CompactMap]")
+{
+	TMapTestHelper<FCompactMapTestResolver>();
+}
 
-#endif //WITH_LOW_LEVEL_TESTS
+#endif // WITH_TESTS
