[Core]
* Changed FGraphEventRef EventToSignal = nullptr to FGraphEventRef EventToSignal = {} to use cheaper ctor

diff --git a/Engine/Source/Runtime/Core/Public/ProfilingDebugging/CsvProfiler.h b/Engine/Source/Runtime/Core/Public/ProfilingDebugging/CsvProfiler.h
--- a/Engine/Source/Runtime/Core/Public/ProfilingDebugging/CsvProfiler.h
+++ b/Engine/Source/Runtime/Core/Public/ProfilingDebugging/CsvProfiler.h
@@ -505,7 +505,7 @@ public:
 	 * EventToSignal is optional. If provided, the CSV profiler will signal the event when the async file write is complete.
 	 * The returned TFuture can be waited on to retrieve the filename of the csv file that was written to disk.
 	 */
-	CORE_API TSharedFuture<FString> EndCapture(FGraphEventRef EventToSignal = nullptr);
+	CORE_API TSharedFuture<FString> EndCapture(FGraphEventRef EventToSignal = {});
 
 	/** Called at the end of the first frame after forking */
 	CORE_API void OnEndFramePostFork();
