Add explicit qualifications to avoid ambiguous symbol errors.

This fixes an error on Horde in the "Verse PresubmitTest.py (Non-Linux)" build.

diff --git a/Engine/Source/Runtime/Core/Public/Async/Fundamental/TaskDelegate.h b/Engine/Source/Runtime/Core/Public/Async/Fundamental/TaskDelegate.h
--- a/Engine/Source/Runtime/Core/Public/Async/Fundamental/TaskDelegate.h
+++ b/Engine/Source/Runtime/Core/Public/Async/Fundamental/TaskDelegate.h
@@ -136,7 +136,7 @@ namespace LowLevelTasks
 			template<typename CallableT>
 			inline TTaskDelegateImpl(CallableT&& Callable, void* InlineData)
 			{
-				static_assert(TIsInvocable<TCallableType, ParamTypes...>::Value, "TCallableType is not invocable");
+				static_assert(::TIsInvocable<TCallableType, ParamTypes...>::Value, "TCallableType is not invocable");
 				static_assert(std::is_same_v<ReturnType, decltype(Callable(UE::Core::Private::IsInvocable::DeclVal<ParamTypes>()...))>, "TCallableType return type does not match");
 				static_assert(sizeof(TTaskDelegateImpl<TCallableType, false>) == sizeof(TTaskDelegateBase), "Size must match the Baseclass");
 				new (InlineData) TCallableType(Forward<CallableT>(Callable));
@@ -206,7 +206,7 @@ namespace LowLevelTasks
 			template<typename CallableT>
 			inline TTaskDelegateImpl(CallableT&& Callable, void* InlineData)
 			{
-				static_assert(TIsInvocable<TCallableType, ParamTypes...>::Value, "TCallableType is not invocable");
+				static_assert(::TIsInvocable<TCallableType, ParamTypes...>::Value, "TCallableType is not invocable");
 				static_assert(std::is_same_v<ReturnType, decltype(Callable(UE::Core::Private::IsInvocable::DeclVal<ParamTypes>()...))>, "TCallableType return type does not match");
 				static_assert(sizeof(TTaskDelegateImpl<TCallableType, true>) == sizeof(TTaskDelegateBase), "Size must match the Baseclass");
 				static_assert(alignof(TCallableType) <= FLowLevelTasksBlockAllocationTag::Allocator::MaxAlignment);
diff --git a/Engine/Source/Runtime/Core/Public/Templates/Function.h b/Engine/Source/Runtime/Core/Public/Templates/Function.h
--- a/Engine/Source/Runtime/Core/Public/Templates/Function.h
+++ b/Engine/Source/Runtime/Core/Public/Templates/Function.h
@@ -111,7 +111,7 @@ namespace UE::Core::Private::Function
 	template <typename T>
 	FORCEINLINE bool IsBound(const T& Func)
 	{
-		if constexpr (std::is_pointer_v<T> || std::is_member_pointer_v<T> || TIsTFunction<T>::Value)
+		if constexpr (std::is_pointer_v<T> || std::is_member_pointer_v<T> || ::TIsTFunction<T>::Value)
 		{
 			// Function pointers, data member pointers, member function pointers and TFunctions
 			// can all be null/unbound, so test them using their boolean state.
@@ -389,7 +389,7 @@ namespace UE::Core::Private::Function
 				}
 			}
 
-			using DecayedFunctorType = typename TRemovePointer<decltype(Binding)>::Type;
+			using DecayedFunctorType = typename ::TRemovePointer<decltype(Binding)>::Type;
 
 			Callable = &TFunctionRefCaller<DecayedFunctorType, Ret, ParamTypes...>::Call;
 		}
@@ -562,7 +562,7 @@ public:
 	 * Constructor which binds a TFunctionRef to a callable object.
 	 */
 	template <typename FunctorType>
-		requires (!TIsTFunctionRef<std::decay_t<FunctorType>>::Value && std::is_invocable_r_v<Ret, std::decay_t<FunctorType>, ParamTypes...>)
+		requires (!::TIsTFunctionRef<std::decay_t<FunctorType>>::Value && std::is_invocable_r_v<Ret, std::decay_t<FunctorType>, ParamTypes...>)
 	TFunctionRef(FunctorType&& InFunc UE_LIFETIMEBOUND)
 		: Super(Forward<FunctorType>(InFunc))
 	{
@@ -652,7 +652,7 @@ public:
 	 * Constructor which binds a TFunction to any function object.
 	 */
 	template <typename FunctorType>
-		requires (!TIsTFunction<std::decay_t<FunctorType>>::Value && std::is_invocable_r_v<Ret, std::decay_t<FunctorType>, ParamTypes...>)
+		requires (!::TIsTFunction<std::decay_t<FunctorType>>::Value && std::is_invocable_r_v<Ret, std::decay_t<FunctorType>, ParamTypes...>)
 	TFunction(FunctorType&& InFunc)
 		: Super(static_cast<UE::Core::Private::Function::TFuncPtrTypeIfPossible_T<FunctorType>>(InFunc))
 	{
@@ -666,7 +666,7 @@ public:
 		// reference):
 		//
 		// TFunction<int32(float)> MyFunction = [MyFunctionRef](float F) { return MyFunctionRef(F); };
-		static_assert(!TIsTFunctionRef<std::decay_t<FunctorType>>::Value, "Cannot construct a TFunction from a TFunctionRef");
+		static_assert(!::TIsTFunctionRef<std::decay_t<FunctorType>>::Value, "Cannot construct a TFunction from a TFunctionRef");
 	}
 
 	TFunction(TFunction&&) = default;
@@ -767,7 +767,7 @@ public:
 	 * Constructor which binds a TFunction to any function object.
 	 */
 	template <typename FunctorType>
-		requires(!TIsTUniqueFunction<std::decay_t<FunctorType>>::Value && !TIsTFunction      <std::decay_t<FunctorType>>::Value && std::is_invocable_r_v<Ret, std::decay_t<FunctorType>, ParamTypes...>)
+		requires(!::TIsTUniqueFunction<std::decay_t<FunctorType>>::Value && !::TIsTFunction<std::decay_t<FunctorType>>::Value && std::is_invocable_r_v<Ret, std::decay_t<FunctorType>, ParamTypes...>)
 	TUniqueFunction(FunctorType&& InFunc)
 		: Super(static_cast<UE::Core::Private::Function::TFuncPtrTypeIfPossible_T<FunctorType>>(InFunc))
 	{
@@ -781,7 +781,7 @@ public:
 		// reference):
 		//
 		// TUniqueFunction<int32(float)> MyFunction = [MyFunctionRef](float F) { return MyFunctionRef(F); };
-		static_assert(!TIsTFunctionRef<std::decay_t<FunctorType>>::Value, "Cannot construct a TUniqueFunction from a TFunctionRef");
+		static_assert(!::TIsTFunctionRef<std::decay_t<FunctorType>>::Value, "Cannot construct a TUniqueFunction from a TFunctionRef");
 	}
 
 	/**
