Add support for 2-param CSV_TRACE_REGION macros, using the same API as TRACE_BEGIN_REGION. Note: unlike TRACE_END_REGION, we need to pass a category name to CSV_TRACE_REGION_END - this is necessary to avoid collisions.
Add region instrumentation for AssetRegistry and Map loading screen (using shared macros where possible)
#tests tested on android and console

diff --git a/Engine/Source/Runtime/AssetRegistry/Private/AssetRegistry.cpp b/Engine/Source/Runtime/AssetRegistry/Private/AssetRegistry.cpp
--- a/Engine/Source/Runtime/AssetRegistry/Private/AssetRegistry.cpp
+++ b/Engine/Source/Runtime/AssetRegistry/Private/AssetRegistry.cpp
@@ -46,6 +46,7 @@
 #include "AssetRegistry/PackageReader.h"
 #include "Modules/ModuleManager.h"
 #include "ProfilingDebugging/CpuProfilerTrace.h"
+#include "ProfilingDebugging/CsvProfiler.h"
 #include "ProfilingDebugging/MiscTrace.h"
 #include "Serialization/ArrayReader.h"
 #include "Serialization/CompactBinarySerialization.h"
@@ -2343,7 +2344,7 @@ void FAssetRegistryImpl::SearchAllAssets(Impl::FEventContext& EventContext,
 
 	if (!bInitialSearchStarted)
 	{
-		TRACE_BEGIN_REGION(TEXT("Asset Registry Scan"));
+		CSV_TRACE_REGION_BEGIN(TEXT("Asset Registry Scan"));
 		InitialSearchStartTime = FPlatformTime::Seconds();
 		bInitialSearchStarted = true;
 		bInitialSearchCompleted.store(false, std::memory_order_relaxed);
@@ -6027,7 +6028,7 @@ void FAssetRegistryImpl::OnInitialSearchCompleted(Impl::FEventContext& EventCont
 	SetPerformanceMode(Impl::EPerformanceMode::MostlyStatic);
 
 	LogSearchDiagnostics(InitialSearchStartTime, /* bInitialScan */ true);
-	TRACE_END_REGION(TEXT("Asset Registry Scan"));
+	CSV_TRACE_REGION_END(TEXT("Asset Registry Scan"));
 
 	GlobalGatherer->UpdateCacheForSaving();
 	GlobalGatherer->OnInitialSearchCompleted();
@@ -6050,7 +6051,7 @@ void FAssetRegistryImpl::OnAdditionalMountSearchCompleted(Impl::FEventContext& E
 	SetPerformanceMode(Impl::EPerformanceMode::MostlyStatic);
 
 	LogSearchDiagnostics(AdditionalMountSearchStartTime, /* bInitialScan */ false);
-	TRACE_END_REGION(TEXT("Asset Registry - Additional Mount Search")); // Matching TRACE_BEGIN_REGION in OnContentPathMounted
+	CSV_TRACE_REGION_END(TEXT("Asset Registry - Additional Mount Search")); // Matching CSV_TRACE_REGION_BEGIN in OnContentPathMounted
 
 	GlobalGatherer->UpdateCacheForSaving();
 	GlobalGatherer->OnAdditionalMountSearchCompleted();
@@ -9264,7 +9265,7 @@ void FAssetRegistryImpl::OnContentPathMounted(UE::AssetRegistry::FInterfaceWrite
 		{
 			if (!IsGathering())
 			{
-				TRACE_BEGIN_REGION(TEXT("Asset Registry - Additional Mount Search")); // Matching TRACE_END_REGION in OnAdditionalMountSearchCompleted
+				CSV_TRACE_REGION_BEGIN(TEXT("Asset Registry - Additional Mount Search")); // Matching CSV_TRACE_REGION_END in OnAdditionalMountSearchCompleted
 				AdditionalMountSearchStartTime = FPlatformTime::Seconds();
 				bAdditionalMountSearchInProgress.store(true, std::memory_order_relaxed);
 				GlobalGatherer->SetIsAdditionalMountSearchInProgress(true);
diff --git a/Engine/Source/Runtime/Core/Private/ProfilingDebugging/CsvProfiler.cpp b/Engine/Source/Runtime/Core/Private/ProfilingDebugging/CsvProfiler.cpp
--- a/Engine/Source/Runtime/Core/Private/ProfilingDebugging/CsvProfiler.cpp
+++ b/Engine/Source/Runtime/Core/Private/ProfilingDebugging/CsvProfiler.cpp
@@ -4602,39 +4602,39 @@ void FCsvProfiler::RecordEventAtFrameStart(int32 CategoryIndex, const FString& E
 	RecordEventAtTimestamp(CategoryIndex, EventText, FCsvProfiler::Get()->LastEndFrameTimestamp, bLog);
 }
 
-void FCsvProfiler::BeginRegion(int32 CategoryIndex, const TCHAR* RegionName)
+void FCsvProfiler::BeginRegion(int32 CategoryIndex, const TCHAR* RegionName, const TCHAR* SubCategoryName)
 {
 	if (GCsvProfilerIsCapturing && GCsvCategoriesEnabled[CategoryIndex])
 	{
-		FString EventText;
-		if (CategoryIndex != CSV_CATEGORY_INDEX_GLOBAL)
-		{
-			EventText = FString::Printf(TEXT("Region: %s/%s - Begin"), *FCsvCategoryData::Get()->GetCategoryNameByIndex(CategoryIndex), RegionName);
-		}
-		else
-		{
-			EventText = FString::Printf(TEXT("Region: %s - Begin"), RegionName);
-		}
-		RecordEvent(CSV_CATEGORY_INDEX_GLOBAL, EventText, false);
-		// TODO: Track timings for each region and output metadata when CSV ends. Or we just use PRT for that...
+		RecordRegionEventInternal(CategoryIndex == CSV_CATEGORY_INDEX_GLOBAL ? nullptr : *FCsvCategoryData::Get()->GetCategoryNameByIndex(CategoryIndex), RegionName, SubCategoryName, true);
 	}
 }
 
-void FCsvProfiler::EndRegion(int32 CategoryIndex, const TCHAR* RegionName)
+void FCsvProfiler::EndRegion(int32 CategoryIndex, const TCHAR* RegionName, const TCHAR* SubCategoryName)
 {
 	if (GCsvProfilerIsCapturing && GCsvCategoriesEnabled[CategoryIndex])
 	{
-		FString EventText;
-		if (CategoryIndex != CSV_CATEGORY_INDEX_GLOBAL)
-		{
-			EventText = FString::Printf(TEXT("Region: %s/%s - End"), *FCsvCategoryData::Get()->GetCategoryNameByIndex(CategoryIndex), RegionName);
-		}
-		else
-		{
-			EventText = FString::Printf(TEXT("Region: %s - End"), RegionName);
-		}
-		RecordEvent(CSV_CATEGORY_INDEX_GLOBAL, EventText, false);
+		RecordRegionEventInternal(CategoryIndex == CSV_CATEGORY_INDEX_GLOBAL ? nullptr : *FCsvCategoryData::Get()->GetCategoryNameByIndex(CategoryIndex), RegionName, SubCategoryName, false);
+	}
+}
+
+void FCsvProfiler::RecordRegionEventInternal(const TCHAR* CategoryName, const TCHAR* RegionName, const TCHAR* SubCategoryName, bool bBegin)
+{
+	const TCHAR* Suffix = bBegin ? TEXT("Begin") : TEXT("End");
+	FString EventText;
+	if (CategoryName)
+	{
+		EventText = SubCategoryName ?
+			FString::Printf(TEXT("Region: %s/%s/%s - %s"), CategoryName, SubCategoryName, RegionName, Suffix) :
+			FString::Printf(TEXT("Region: %s/%s - %s"), CategoryName, RegionName, Suffix);
+	}
+	else
+	{
+		EventText = SubCategoryName ? 
+			FString::Printf(TEXT("Region: %s/%s - %s"), SubCategoryName, RegionName, Suffix) : 
+			FString::Printf(TEXT("Region: %s - %s"), RegionName, Suffix);
 	}
+	RecordEvent(CSV_CATEGORY_INDEX_GLOBAL, EventText, false);
 }
 
 
diff --git a/Engine/Source/Runtime/Core/Public/ProfilingDebugging/CsvProfiler.h b/Engine/Source/Runtime/Core/Public/ProfilingDebugging/CsvProfiler.h
--- a/Engine/Source/Runtime/Core/Public/ProfilingDebugging/CsvProfiler.h
+++ b/Engine/Source/Runtime/Core/Public/ProfilingDebugging/CsvProfiler.h
@@ -81,7 +81,6 @@ struct FCsvDeclaredStat;
 	#define CSV_REGION_END_GLOBAL(RegionName) \
 		FCsvProfiler::EndRegion(CSV_CATEGORY_INDEX_GLOBAL, RegionName);
 
-
 	// Metadata
 	#define CSV_METADATA(Key,Value)									FCsvProfiler::SetMetadata( Key, Value )
 	#define CSV_NON_PERSISTENT_METADATA(Key,Value)					FCsvProfiler::SetNonPersistentMetadata( Key, Value )
@@ -191,14 +190,23 @@ struct FCsvDeclaredStat;
 	#define CSV_CUSTOM_STAT_DEFINED(StatName,Value,Op)				
 #endif
 
-// Shared macros. These emit trace and CSV instrumentation
-#define CSV_TRACE_REGION_BEGIN_GLOBAL(RegionName) \
-	TRACE_BEGIN_REGION(RegionName) \
-	CSV_REGION_BEGIN_GLOBAL(RegionName)
+// Shared macros. These emit Trace and CSV instrumentation. Can be used in place of TRACE_BEGIN_REGION and TRACE_END_REGION
+// NOTE: CSV_TRACE_REGION_END requires a category name, unlike TRACE_END_REGION. This is in order to avoid potential naming clashes
+#if CSV_PROFILER
+	#define CSV_TRACE_REGION_BEGIN(RegionName, ...) \
+		TRACE_BEGIN_REGION(RegionName __VA_OPT__(,) __VA_ARGS__) \
+		FCsvProfiler::BeginRegion(CSV_CATEGORY_INDEX_GLOBAL, RegionName __VA_OPT__(,) __VA_ARGS__);
+
+	#define CSV_TRACE_REGION_END(RegionName, ...) \
+		FCsvProfiler::EndRegion(CSV_CATEGORY_INDEX_GLOBAL, RegionName __VA_OPT__(,) __VA_ARGS__); \
+		TRACE_END_REGION(RegionName)
+#else // CSV_PROFILER
+	#define CSV_TRACE_REGION_BEGIN(RegionName, ...) \
+		TRACE_BEGIN_REGION(RegionName __VA_OPT__(,) __VA_ARGS__)
 
-#define CSV_TRACE_REGION_END_GLOBAL(RegionName) \
-	CSV_REGION_END_GLOBAL(RegionName) \
-	TRACE_END_REGION(RegionName)
+	#define CSV_TRACE_REGION_END(RegionName, ...) \
+		TRACE_END_REGION(RegionName)
+#endif // !CSV_PROFILER
 
 #if CSV_PROFILER
 
@@ -469,8 +477,8 @@ public:
 	CORE_API static void BeginWait();
 	CORE_API static void EndWait();
 
-	CORE_API static void BeginRegion(int32 CategoryIndex, const TCHAR* RegionName);
-	CORE_API static void EndRegion(int32 CategoryIndex, const TCHAR* RegionName);
+	CORE_API static void BeginRegion(int32 CategoryIndex, const TCHAR* RegionName, const TCHAR* SubCategoryName = nullptr);
+	CORE_API static void EndRegion(int32 CategoryIndex, const TCHAR* RegionName, const TCHAR* SubCategoryName = nullptr);
 
 	CORE_API static FString GetDefaultDirectory();
 
@@ -569,6 +577,8 @@ public:
 
 private:
 	CORE_API static void VARARGS RecordEventfInternal(int32 CategoryIndex, bool bLog, const TCHAR* Fmt, ...);
+	
+	static void RecordRegionEventInternal(const TCHAR* CategoryName, const TCHAR* RegionName, const TCHAR* SubCategoryName, bool bBegin);
 
 	enum class EMetadataPersistenceType : int8
 	{
