[RaceDetector]
- Fix an issue where the context of the racing thread was out of date by the time a race needed to be reported
- This could result in races we weren't interested in to get reported, or some we would be interested in to not be reported
- This issue was only present when filtering threads with ToggleFilterOtherThreads()

diff --git a/Engine/Source/Runtime/Core/Internal/Sanitizer/RaceDetectorTypes.h b/Engine/Source/Runtime/Core/Internal/Sanitizer/RaceDetectorTypes.h
--- a/Engine/Source/Runtime/Core/Internal/Sanitizer/RaceDetectorTypes.h
+++ b/Engine/Source/Runtime/Core/Internal/Sanitizer/RaceDetectorTypes.h
@@ -246,6 +246,7 @@ namespace UE::Sanitizer::RaceDetector {
 		ACCESS_TYPE_WRITE = 0b10,
 		ACCESS_TYPE_ATOMIC = 0b100,
 		ACCESS_TYPE_VPTR = 0b1000,
+		ACCESS_TYPE_ALWAYSREPORT = 0b10000,
 
 		ACCESS_TYPE_ATOMIC_READ = ACCESS_TYPE_ATOMIC | ACCESS_TYPE_READ,
 		ACCESS_TYPE_ATOMIC_WRITE = ACCESS_TYPE_ATOMIC | ACCESS_TYPE_WRITE,
@@ -296,13 +297,17 @@ namespace UE::Sanitizer::RaceDetector {
 		{
 		}
 
-		INSTRUMENTATION_FUNCTION_ATTRIBUTES FORCEINLINE FMemoryAccess(FContextId InContextId, FClock InClock, uint8 InOffset, uint8 InSize, EMemoryAccessType InAccessType)
+		INSTRUMENTATION_FUNCTION_ATTRIBUTES FORCEINLINE FMemoryAccess(FContextId InContextId, FClock InClock, uint8 InOffset, uint8 InSize, EMemoryAccessType InAccessType, bool bInAlwaysReport)
 		{
 			// Apparently this is much faster than using the bitfields. 30% faster in warm TSAN benchmark!!!
 			// The constructor was accessing the same value in memory/store buffer multiple time doing its bit tweaking and it caused tons of Core::X86::Pmc::Core::LsBadStatus2.
 			// That ended up stalling on S[0].RawValue == CurrentAccess.RawValue in InstrumentMemoryAccess when trying to extract the full uint64.
 			// https://blog.stuffedcow.net/2014/01/x86-memory-disambiguation/
 			const uint8 AccessValue = (uint8)(((1ull << InSize) - 1ull) << InOffset);
+			if (bInAlwaysReport)
+			{
+				InAccessType |= EMemoryAccessType::ACCESS_TYPE_ALWAYSREPORT;
+			}
 			RawValue = (uint64)InAccessType << 48 | (uint64)AccessValue << 40 | (uint64)InContextId << 32 | (uint64)InClock;
 		}
 		
@@ -324,12 +329,13 @@ namespace UE::Sanitizer::RaceDetector {
 				{
 					struct
 					{
-						uint8 bIsRead   : 1;
-						uint8 bIsWrite  : 1;
-						uint8 bIsAtomic : 1;
-						uint8 bIsVPtr   : 1;
+						uint8 bIsRead       : 1;
+						uint8 bIsWrite      : 1;
+						uint8 bIsAtomic     : 1;
+						uint8 bIsVPtr       : 1;
+						uint8 bAlwaysReport : 1;
 					};
-					EMemoryAccessType AccessType : 4; // 6  plenty of bits left here
+					EMemoryAccessType AccessType : 5; // 6  plenty of bits left here
 				};
 				uint8 Reserved;                           // 7  plenty of bits left here
 			};
diff --git a/Engine/Source/Runtime/Core/Private/Sanitizer/RaceDetector.cpp b/Engine/Source/Runtime/Core/Private/Sanitizer/RaceDetector.cpp
--- a/Engine/Source/Runtime/Core/Private/Sanitizer/RaceDetector.cpp
+++ b/Engine/Source/Runtime/Core/Private/Sanitizer/RaceDetector.cpp
@@ -610,7 +610,7 @@ namespace UE::Sanitizer::RaceDetector {
 
 			if (GRaceDetectorFilterOtherThreads &&
 				Context.bAlwaysReport != true &&
-				RacingContext->bAlwaysReport != true)
+				RacingAccess.bAlwaysReport != true)
 			{
 				return;
 			}
@@ -760,7 +760,9 @@ namespace UE::Sanitizer::RaceDetector {
 
 		const uint32 Epoch = GlobalEpoch;
 		const uint8 Offset = GetOffsetInBytes(Ptr);
-		const FMemoryAccess CurrentAccess(Context.ContextId, Context.CurrentClock(), Offset, Size, AccessType);
+		// Store the bAlwaysReport of this context directly in the access bitfield because at the time of the race, the racingcontext might have a different bAlwaysReport state than when the access was made.
+		// So to get consistent results, it's important to store the state of the context as it was when the access was made.
+		const FMemoryAccess CurrentAccess(Context.ContextId, Context.CurrentClock(), Offset, Size, AccessType, Context.bAlwaysReport);
 
 		// This is very practical to scope small section of code that we want to understand
 		// exactly what's going on inside.
