[RaceDetector]
- Fix potential infinite loop during shutdown because we didn't put back NOPs at the JMP site and it kept jumping back to the function prefix which was all NOPs
- Fix potential crash at shutdown if race detector was still active or the shutdown raced with toggling race detector

diff --git a/Engine/Source/Runtime/Core/Private/Sanitizer/RaceDetector.cpp b/Engine/Source/Runtime/Core/Private/Sanitizer/RaceDetector.cpp
--- a/Engine/Source/Runtime/Core/Private/Sanitizer/RaceDetector.cpp
+++ b/Engine/Source/Runtime/Core/Private/Sanitizer/RaceDetector.cpp
@@ -1530,6 +1530,9 @@ namespace UE::Sanitizer::RaceDetector {
 			return true;
 		}
 		
+		// Needed to avoid race with ToggleRaceDetection
+		UE::TUniqueLock ScopeLock(Mutex);
+
 		UnhookInstrumentation();
 
 		ResetShadow();
@@ -1546,8 +1549,9 @@ namespace UE::Sanitizer::RaceDetector {
 		if (ContextTLSIndex != UINT32_MAX)
 		{
 			ReleaseCurrentContext();
-			FPlatformTLS::FreeTlsSlot(ContextTLSIndex);
-			ContextTLSIndex = UINT32_MAX;
+
+			// Because we unhooked the detoured function doesn't mean there isn't a call still inside one.
+			// We have no choice but to leak the context tls to avoid crashes if that ever happens.
 		}
 
 		bRuntimeInitialized = false;
@@ -1584,6 +1588,11 @@ namespace UE::Sanitizer::RaceDetector {
 			// Need to take the mutex during reset shadow since we're unhooking and rehooking instrumentation, which is not thread-safe.
 			UE::TUniqueLock ScopeLock(Mutex);
 
+			if (!bRuntimeInitialized)
+			{
+				return;
+			}
+
 			FContext& Context = EnsureCurrentContext();
 
 			FInstrumentationScope InstrumentationScope;
diff --git a/Engine/Source/Runtime/Core/Private/Sanitizer/RaceDetectorWindows.cpp b/Engine/Source/Runtime/Core/Private/Sanitizer/RaceDetectorWindows.cpp
--- a/Engine/Source/Runtime/Core/Private/Sanitizer/RaceDetectorWindows.cpp
+++ b/Engine/Source/Runtime/Core/Private/Sanitizer/RaceDetectorWindows.cpp
@@ -827,6 +827,10 @@ namespace UE::Sanitizer::RaceDetector::Platform {
 			UE_LOG(LogRaceDetector, Fatal, TEXT("Unable to change page protection for hotpatching at %p (error %d)"), ThunkFunctionAddress, GetLastError());
 		}
 
+		// Make sure there is no jmp left, equivalent to deactivation.
+		// This is actually atomic and won't cause any crashes during the transition.
+		*(uint16*)ThunkFunctionStart = 0x9066; /* Replace by 2-byte NOP */
+
 		// Restore all the NOPs that were overwritten
 		memset(PrefixStart, 0x90, INSTRUMENTATION_HOTPATCH_PREFIX_NOPS);
 
@@ -876,7 +880,7 @@ namespace UE::Sanitizer::RaceDetector::Platform {
 		}
 		else
 		{
-			// This is actually atomic a wont cause any crashes during the transition.
+			// This is actually atomic and won't cause any crashes during the transition.
 			*(uint16*)ThunkFunctionStart = 0x9066; /* Replace by 2-byte NOP */
 		}
 
