Add ApplicationHeartbeat core delegate and broadcast from FFeedbackContext when progress is made

#rnx

diff --git a/Engine/Source/Runtime/Core/Private/Misc/CoreDelegates.cpp b/Engine/Source/Runtime/Core/Private/Misc/CoreDelegates.cpp
--- a/Engine/Source/Runtime/Core/Private/Misc/CoreDelegates.cpp
+++ b/Engine/Source/Runtime/Core/Private/Misc/CoreDelegates.cpp
@@ -257,6 +257,8 @@ TMulticastDelegate<void(const TCHAR*)> FCoreDelegates::StatDisabled;
 TMulticastDelegate<void(const bool)> FCoreDelegates::StatDisableAll;
 
 TMulticastDelegate<void()> FCoreDelegates::ApplicationLicenseChange;
+TTSMulticastDelegate<void()> FCoreDelegates::ApplicationHeartbeat;
+
 TMulticastDelegate<void(EConvertibleLaptopMode)> FCoreDelegates::PlatformChangedLaptopMode;
 
 TMulticastDelegate<void()> FCoreDelegates::VRHeadsetTrackingInitializingAndNeedsHMDToBeTrackedDelegate;
diff --git a/Engine/Source/Runtime/Core/Public/Misc/CoreDelegates.h b/Engine/Source/Runtime/Core/Public/Misc/CoreDelegates.h
--- a/Engine/Source/Runtime/Core/Public/Misc/CoreDelegates.h
+++ b/Engine/Source/Runtime/Core/Public/Misc/CoreDelegates.h
@@ -443,6 +443,11 @@ public:
 	// Called when an application is notified that the application license info has been updated.
 	// The new license data should be polled and steps taken based on the results (i.e. halt application if license is no longer valid).
 	static CORE_API TMulticastDelegate<void()> ApplicationLicenseChange;
+	
+	// Called periodically from within slow operations in order to indicate the application is still alive.
+	// Progress must be made for this to be called, or it could prevent the application from detecting soft-locks and unresponsiveness.
+	// Subscribers should be lightweight and thread-safe.
+	static CORE_API TTSMulticastDelegate<void()> ApplicationHeartbeat;
 
 	/** Sent when the platform changed its laptop mode (for convertible laptops).*/
 	static CORE_API TMulticastDelegate<void(EConvertibleLaptopMode)> PlatformChangedLaptopMode;
diff --git a/Engine/Source/Runtime/Core/Public/Misc/FeedbackContext.h b/Engine/Source/Runtime/Core/Public/Misc/FeedbackContext.h
--- a/Engine/Source/Runtime/Core/Public/Misc/FeedbackContext.h
+++ b/Engine/Source/Runtime/Core/Public/Misc/FeedbackContext.h
@@ -5,6 +5,7 @@
 #include "Containers/Array.h"
 #include "Containers/StringFwd.h"
 #include "Containers/UnrealString.h"
+#include "CoreDelegates.h"
 #include "CoreGlobals.h"
 #include "CoreTypes.h"
 #include "HAL/CriticalSection.h"
@@ -100,7 +101,13 @@ protected:
 	 * @param	TotalProgressInterp		[0..1] Value indicating the total progress of the slow task
 	 * @param	DisplayMessage			The message to display on the slow task
 	 */
-	virtual void ProgressReported( const float TotalProgressInterp, FText DisplayMessage ) {}
+	virtual void ProgressReported( const float TotalProgressInterp, FText DisplayMessage )
+	{
+		if (TotalProgressInterp > 0)
+		{
+			FCoreDelegates::ApplicationHeartbeat.Broadcast();
+		}
+	}
 
 	/** Called to check whether we are playing in editor when starting a slow task */
 	CORE_API virtual bool IsPlayingInEditor() const;
