[RaceDetector]
- Replace costly instrumented TArray with simpler structure to significantly improve performance when resolving callstacks for races

diff --git a/Engine/Source/Runtime/Core/Internal/Sanitizer/RaceDetectorTypes.h b/Engine/Source/Runtime/Core/Internal/Sanitizer/RaceDetectorTypes.h
--- a/Engine/Source/Runtime/Core/Internal/Sanitizer/RaceDetectorTypes.h
+++ b/Engine/Source/Runtime/Core/Internal/Sanitizer/RaceDetectorTypes.h
@@ -564,65 +564,7 @@ namespace UE::Sanitizer::RaceDetector {
 			EmplaceEntry<FHistoryEntryFunctionExit>();
 		}
 
-		INSTRUMENTATION_FUNCTION_ATTRIBUTES bool ResolveAccess(const void* InAlignedPointer, FMemoryAccess InAccess, FCallstackLocation& OutLocation, FClockRange& OutClockRange)
-		{
-			TReadScopeLock Scope(Lock);
-
-			// Make sure we sync with the light fence.
-			Platform::AsymmetricThreadFenceHeavy();
-
-			OutClockRange.First = Head ? Head->StartClock : 0;
-
-			for (FHistoryChunk* Chunk = Head; Chunk; Chunk = Chunk->Next)
-			{
-				OutClockRange.Last = Chunk->EndClock;
-
-				// Do not bother searching a chunk that is outside the range we're looking for.
-				// Do the range comparison in a way that handles clock wrapping
-				if ((InAccess.Clock - Chunk->StartClock) <= (Chunk->EndClock - Chunk->StartClock))
-				{
-					TArray<void*, TInlineAllocator<1024>> Stack;
-
-					int32 Offset = 0;
-					while (Offset < Chunk->Offset)
-					{
-						FHistoryEntryBase* Entry = (FHistoryEntryBase*)(Chunk->Buffer + Offset);
-						switch (Entry->Type)
-						{
-							case EHistoryEntryType::FunctionEntry:
-								Stack.Add(((FHistoryEntryFunctionEntry*)Entry)->ReturnAddress);
-								Offset += sizeof(FHistoryEntryFunctionEntry);
-							break;
-							case EHistoryEntryType::MemoryAccess:
-							{
-								FHistoryEntryAccess* MemoryAccessEntry = (FHistoryEntryAccess*)Entry;
-								if (MemoryAccessEntry->Pointer == InAlignedPointer &&
-									MemoryAccessEntry->Access.RawValue == InAccess.RawValue)
-								{
-									OutLocation = FCallstackLocation(Stack.GetData(), Stack.Num());
-									return true;
-								}
-								Offset += sizeof(FHistoryEntryAccess);
-							}
-							break;
-							case EHistoryEntryType::FunctionExit:
-								Stack.Pop(EAllowShrinking::No);
-								Offset += sizeof(FHistoryEntryBase);
-							break;
-							default:
-								check(false);
-								// This should never happen, but if it does it is most likely a race condition
-								// so just restart the tracing from the beginning as a last resort.
-								Stack.Reset();
-								Offset = 0;
-							break;
-						}
-					}
-				}
-			}
-
-			return false;
-		}
+		inline INSTRUMENTATION_FUNCTION_ATTRIBUTES bool ResolveAccess(struct FContext& Context, const void* InAlignedPointer, FMemoryAccess InAccess, FCallstackLocation& OutLocation, FClockRange& OutClockRange);
 
 		INSTRUMENTATION_FUNCTION_ATTRIBUTES int32 GetOffset() const
 		{
@@ -653,6 +595,11 @@ namespace UE::Sanitizer::RaceDetector {
 				delete AccessHistory;
 				AccessHistory = nullptr;
 			}
+			if (ResolvingStack)
+			{
+				FInstrumentationSafeWinAllocator::FreeWithGuards(ResolvingStack);
+				ResolvingStack = nullptr;
+			}
 		}
 
 		uint32 GlobalEpoch = 0;
@@ -684,7 +631,8 @@ namespace UE::Sanitizer::RaceDetector {
 		FAccessHistory* AccessHistory = nullptr;
 		// Wether or not to always report race for this thread.
 		bool bAlwaysReport = false;
-
+		// Used when we need to resolve a stack
+		void** ResolvingStack = nullptr;
 		// This can be bumped again if we ever face a need for deeper callstacks since
 		// this is a virtual allocation anyway so it's not going to take physical memory
 		// until it is used. 
@@ -725,6 +673,73 @@ namespace UE::Sanitizer::RaceDetector {
 		}
 	};
 
+	INSTRUMENTATION_FUNCTION_ATTRIBUTES bool FAccessHistory::ResolveAccess(FContext& Context, const void* InAlignedPointer, FMemoryAccess InAccess, FCallstackLocation& OutLocation, FClockRange& OutClockRange)
+	{
+		TReadScopeLock Scope(Lock);
+
+		// Make sure we sync with the light fence.
+		Platform::AsymmetricThreadFenceHeavy();
+
+		OutClockRange.First = Head ? Head->StartClock : 0;
+		OutClockRange.Last = OutClockRange.First;
+
+		for (FHistoryChunk* Chunk = Head; Chunk; Chunk = Chunk->Next)
+		{
+			OutClockRange.Last = Chunk->EndClock;
+
+			// Do not bother searching a chunk that is outside the range we're looking for.
+			// Do the range comparison in a way that handles clock wrapping
+			if ((InAccess.Clock - Chunk->StartClock) <= (Chunk->EndClock - Chunk->StartClock))
+			{
+				int32 StackNum = 0;
+				if (Context.ResolvingStack == nullptr)
+				{
+					// Each thread has its own resolving stack with page guards to automatically protect
+					// against underflow and overflow without additional logic.
+					Context.ResolvingStack = (void**)FInstrumentationSafeWinAllocator::AllocWithGuards(256 * 1024);
+				}
+
+				int32 Offset = 0;
+				while (Offset < Chunk->Offset)
+				{
+					FHistoryEntryBase* Entry = (FHistoryEntryBase*)(Chunk->Buffer + Offset);
+					switch (Entry->Type)
+					{
+					case EHistoryEntryType::FunctionEntry:
+						Context.ResolvingStack[StackNum++] = ((FHistoryEntryFunctionEntry*)Entry)->ReturnAddress;
+						Offset += sizeof(FHistoryEntryFunctionEntry);
+					break;
+					case EHistoryEntryType::MemoryAccess:
+					{
+						FHistoryEntryAccess* MemoryAccessEntry = (FHistoryEntryAccess*)Entry;
+						if (MemoryAccessEntry->Pointer == InAlignedPointer &&
+							MemoryAccessEntry->Access.RawValue == InAccess.RawValue)
+						{
+							OutLocation = FCallstackLocation(Context.ResolvingStack, StackNum);
+							return true;
+						}
+						Offset += sizeof(FHistoryEntryAccess);
+					}
+					break;
+					case EHistoryEntryType::FunctionExit:
+						StackNum--;
+						Offset += sizeof(FHistoryEntryBase);
+					break;
+					default:
+						check(false);
+						// This should never happen, but if it does it is most likely a race condition
+						// so just restart the tracing from the beginning as a last resort.
+						StackNum = 0;
+						Offset = 0;
+					break;
+					}
+				}
+			}
+		}
+
+		return false;
+	}
+
 	CORE_API extern FContext* GetThreadContext();
 
 	struct FInstrumentationScope
@@ -897,7 +912,7 @@ namespace UE::Sanitizer::RaceDetector {
 
 		INSTRUMENTATION_FUNCTION_ATTRIBUTES ~FSyncObjectBank()
 		{
-			for (int32 Index = 0; Index < 8; ++Index)
+			for (int32 Index = 0; Index < 7; ++Index)
 			{
 				if (SyncObjects[Index])
 				{
diff --git a/Engine/Source/Runtime/Core/Private/Sanitizer/RaceDetector.cpp b/Engine/Source/Runtime/Core/Private/Sanitizer/RaceDetector.cpp
--- a/Engine/Source/Runtime/Core/Private/Sanitizer/RaceDetector.cpp
+++ b/Engine/Source/Runtime/Core/Private/Sanitizer/RaceDetector.cpp
@@ -571,7 +571,7 @@ namespace UE::Sanitizer::RaceDetector {
 		const void* AlignedPointer = GetAlignedPointer(Pointer);
 		FAccessHistory& SecondAccessHistory = *RacingContext->AccessHistory;
 		FClockRange HistoryRange;
-		bool bFoundLocation = SecondAccessHistory.ResolveAccess(AlignedPointer, RacingAccess, SecondLocation, HistoryRange);
+		bool bFoundLocation = SecondAccessHistory.ResolveAccess(Context, AlignedPointer, RacingAccess, SecondLocation, HistoryRange);
 
 		// Resolving symbols can cause us to enter a wait that could try to start a new thread
 		// and wait for it to be started. We can't allow that since we could deadlock
