Expose the platform time offset applied inside FPlatformTime::Seconds so we can correct for it in calling code if necessary (e.g when converting between cycles and seconds)

diff --git a/Engine/Source/Runtime/Core/Private/ProfilingDebugging/CsvProfiler.cpp b/Engine/Source/Runtime/Core/Private/ProfilingDebugging/CsvProfiler.cpp
--- a/Engine/Source/Runtime/Core/Private/ProfilingDebugging/CsvProfiler.cpp
+++ b/Engine/Source/Runtime/Core/Private/ProfilingDebugging/CsvProfiler.cpp
@@ -3836,8 +3836,7 @@ void FCsvProfiler::BeginCaptureInternal(const FCsvCaptureCommand& CurrentCommand
 		check(GFrameNumber==1);
 		UE_LOG(LogCsvProfiler, Log, TEXT("Doing a boot time capture. Start time offset %.2f"), FPlatformTime::Seconds()-GStartTime);
 		CaptureStartTime = GStartTime;
-		GCaptureStartTimestamp = FPlatformTime::SecondsToCycles64(GStartTime);
-
+		
 		// signal bInsertEndFrameAtFrameStart to insert an EndFrame() at the start of the first _real_ frame
 		// We also add a FrameBeginTimestampsRT timestamp here, to create a dummy renderthread frame, to ensure the rows match up in the CSV
 		GFrameBoundaries.AddBeginFrameTimestamp(ECsvTimeline::Renderthread, false);
@@ -3845,9 +3844,11 @@ void FCsvProfiler::BeginCaptureInternal(const FCsvCaptureCommand& CurrentCommand
 	}
 	else
 	{
-		GCaptureStartTimestamp = FPlatformTime::Cycles64();
-		CaptureStartTime = FPlatformTime::ToSeconds64(GCaptureStartTimestamp);
+		CaptureStartTime = FPlatformTime::Seconds();
 	}
+	// Note: we have to remove the platform time offset when converting between absolute seconds and clock cycles
+	GCaptureStartTimestamp = FPlatformTime::SecondsToCycles64(CaptureStartTime - FPlatformTime::GetSecondsTimeOffset());
+
 	LastEndFrameTimestamp = GCaptureStartTimestamp;
 
 	CurrentFlags = CurrentCommand.Flags;
diff --git a/Engine/Source/Runtime/Core/Public/Apple/ApplePlatformTime.h b/Engine/Source/Runtime/Core/Public/Apple/ApplePlatformTime.h
--- a/Engine/Source/Runtime/Core/Public/Apple/ApplePlatformTime.h
+++ b/Engine/Source/Runtime/Core/Public/Apple/ApplePlatformTime.h
@@ -33,7 +33,7 @@ struct CORE_API FApplePlatformTime : public FGenericPlatformTime
 	{
 		uint64 Cycles = mach_absolute_time();
 		// Add big number to make bugs apparent where return value is being passed to float
-		return Cycles * GetSecondsPerCycle() + 16777216.0;
+		return Cycles * GetSecondsPerCycle() + GetSecondsTimeOffset();
 	}
 
 	static inline uint32 Cycles()
@@ -49,6 +49,11 @@ struct CORE_API FApplePlatformTime : public FGenericPlatformTime
 	}
 
 	static FCPUTime GetCPUTime();
+
+	static inline double GetSecondsTimeOffset()
+	{
+		return 16777216.0;
+	}
 };
 
 typedef FApplePlatformTime FPlatformTime;
diff --git a/Engine/Source/Runtime/Core/Public/GenericPlatform/GenericPlatformTime.h b/Engine/Source/Runtime/Core/Public/GenericPlatform/GenericPlatformTime.h
--- a/Engine/Source/Runtime/Core/Public/GenericPlatform/GenericPlatformTime.h
+++ b/Engine/Source/Runtime/Core/Public/GenericPlatform/GenericPlatformTime.h
@@ -216,6 +216,12 @@ struct FGenericPlatformTime
 		return static_cast<uint64>(Seconds / GetSecondsPerCycle64());
 	}
 
+	/** Gets the platform time offset, if there is one. This is added to the value returned from Seconds() */
+	static double GetSecondsTimeOffset()
+	{
+		return 0.0;
+	}
+
 protected:
 
 	static CORE_API double SecondsPerCycle;
diff --git a/Engine/Source/Runtime/Core/Public/Windows/WindowsPlatformTime.h b/Engine/Source/Runtime/Core/Public/Windows/WindowsPlatformTime.h
--- a/Engine/Source/Runtime/Core/Public/Windows/WindowsPlatformTime.h
+++ b/Engine/Source/Runtime/Core/Public/Windows/WindowsPlatformTime.h
@@ -26,7 +26,7 @@ struct FWindowsPlatformTime
 		Windows::QueryPerformanceCounter(&Cycles);
 
 		// add big number to make bugs apparent where return value is being passed to float
-		return (double)Cycles.QuadPart * GetSecondsPerCycle() + 16777216.0;
+		return (double)Cycles.QuadPart * GetSecondsPerCycle() + GetSecondsTimeOffset();
 	}
 
 	static inline uint32 Cycles()
@@ -43,6 +43,10 @@ struct FWindowsPlatformTime
 		return Cycles.QuadPart;
 	}
 
+	static inline double GetSecondsTimeOffset()
+	{
+		return 16777216.0;
+	}
 
 	static CORE_API void SystemTime( int32& Year, int32& Month, int32& DayOfWeek, int32& Day, int32& Hour, int32& Min, int32& Sec, int32& MSec );
 	static CORE_API void UtcTime( int32& Year, int32& Month, int32& DayOfWeek, int32& Day, int32& Hour, int32& Min, int32& Sec, int32& MSec );
