[StripedMap]
- Make it possible to emplace when using FindOrProduceAndApplyForWrite
- Add GetAllocatedSize and Reserve

diff --git a/Engine/Source/Runtime/Core/Public/Containers/StripedMap.h b/Engine/Source/Runtime/Core/Public/Containers/StripedMap.h
--- a/Engine/Source/Runtime/Core/Public/Containers/StripedMap.h
+++ b/Engine/Source/Runtime/Core/Public/Containers/StripedMap.h
@@ -508,7 +508,7 @@ public:
 			}
 			else
 			{
-				InApplyFunction(Map.AddByHash(Hash, InKey, InProduceFunction()));
+				InApplyFunction(Map.EmplaceByHash(Hash, InKey, InProduceFunction()));
 			}
 		});
 	}
@@ -702,6 +702,43 @@ public:
 		return Count;
 	}
 
+	/**
+	 * Helper function to return the amount of memory allocated by this
+	 * container.
+	 * Only returns the size of allocations made directly by the container, not the elements themselves.
+	 *
+	 * @returns Number of bytes allocated by this container.
+	 */
+	[[nodiscard]] SIZE_T GetAllocatedSize() const
+	{
+		SIZE_T Total = 0;
+		ForEachMap(
+			[&Total](const MapType& Map)
+			{
+				Total += Map.GetAllocatedSize();
+			}
+		);
+
+		return Total;
+	}
+
+	/**
+	 * Reserves memory such that the map can contain at least Number elements. The reserved memory is distributed across stripes.
+	 *
+	 * @param Number The number of elements that the map should be able to contain after allocation.
+	 * @see Shrink
+	 */
+	void Reserve(SIZE_T Number)
+	{
+		SIZE_T NumberPerMap = Number / BucketCount;
+		ForEachMap(
+			[NumberPerMap](MapType& Map)
+			{
+				Map.Reserve(NumberPerMap);
+			}
+		);
+	}
+
 	/**
 	 * Calls a function on all elements of the map with exclusive access (elements can be modified).
 	 *
