[Core] Fix missing destruct items on compact set move
* Added unit test to ensure destruction happens as expected

#tests Local cook on data causing crash and auto test cooke preflights
#pf2 690bc1de7428591f0e17cbbc

diff --git a/Engine/Source/Runtime/Core/Public/Containers/CompactSet.h.inl b/Engine/Source/Runtime/Core/Public/Containers/CompactSet.h.inl
--- a/Engine/Source/Runtime/Core/Public/Containers/CompactSet.h.inl
+++ b/Engine/Source/Runtime/Core/Public/Containers/CompactSet.h.inl
@@ -399,6 +399,7 @@ public:
 	{
 		if (this != &Other)
 		{
+			DestructItems(GetData(), this->NumElements);
 			this->Elements.MoveToEmpty(Other.Elements);
 			this->NumElements = Other.NumElements;
 			this->MaxElements = Other.MaxElements;
diff --git a/Engine/Source/Runtime/Core/Tests/Containers/SetTests.cpp b/Engine/Source/Runtime/Core/Tests/Containers/SetTests.cpp
--- a/Engine/Source/Runtime/Core/Tests/Containers/SetTests.cpp
+++ b/Engine/Source/Runtime/Core/Tests/Containers/SetTests.cpp
@@ -46,11 +46,48 @@ void SetTestExists(const Type& Key, const SetType& Set)
 	}
 }
 
-struct FSetTestNonCopyableTestData
+struct FSetTestDestructableData
+{
+	FSetTestDestructableData()
+	{
+		++Count;
+	}
+	FSetTestDestructableData(uint32 InValue)
+	: Value(InValue)
+	{
+		++Count;
+	}
+
+	FSetTestDestructableData(const FSetTestDestructableData& Other) : FSetTestDestructableData(Other.Value)
+	{
+	}
+
+	FSetTestDestructableData(FSetTestDestructableData&& Other) : FSetTestDestructableData(Other.Value)
+	{
+	}
+
+	~FSetTestDestructableData()
+	{
+		--Count;
+	}
+
+	FSetTestDestructableData& operator=(const FSetTestDestructableData&) = default;
+	FSetTestDestructableData& operator=(FSetTestDestructableData&&) = default;
+	
+	bool operator==(const FSetTestDestructableData& Other) const
+	{
+		return Value == Other.Value;
+	}
+
+	int32 Value = 0;
+
+	static inline int32 Count;
+};
+
+struct FSetTestNonCopyableTestData : public FSetTestDestructableData
 {
 	FSetTestNonCopyableTestData() = default;
-	FSetTestNonCopyableTestData(uint32 V)
-	: Value(V)
+	FSetTestNonCopyableTestData(uint32 InValue) : FSetTestDestructableData(InValue)
 	{
 	}
 
@@ -61,21 +98,13 @@ struct FSetTestNonCopyableTestData
 	// But allow for proper use of move operations
 	FSetTestNonCopyableTestData(FSetTestNonCopyableTestData&&) = default;
 	FSetTestNonCopyableTestData& operator=(FSetTestNonCopyableTestData&&) = default;
-	
-	bool operator==(const FSetTestNonCopyableTestData& Other) const
-	{
-		return Value == Other.Value;
-	}
-
-	uint32 Value = 0;
 };
 
-uint32 GetTypeHash(const FSetTestNonCopyableTestData& A)
+uint32 GetTypeHash(const FSetTestDestructableData& A)
 {
 	return A.Value;
 }
 
-
 template<typename SetResolver = FCompactSetTestResolver>
 struct TSetTestHelper
 {
@@ -127,6 +156,46 @@ struct TSetTestHelper
 			}
 		}
 
+		SECTION("Destruction Tests")
+		{
+			constexpr int32 TestCount = 10;
+
+			SetType<FSetTestDestructableData> Set;
+			for (uint32 Index = 0; Index < TestCount; ++Index)
+			{
+				Set.Emplace(Index);
+			}
+
+			CHECK(FSetTestDestructableData::Count == TestCount);
+
+			SetType<FSetTestDestructableData> SetCopy = Set;
+
+			CHECK(FSetTestDestructableData::Count == TestCount * 2);
+
+			SetType<FSetTestDestructableData> SetMove = MoveTemp(Set);
+
+			CHECK(FSetTestDestructableData::Count == TestCount * 2);
+
+			for (uint32 Index = 0; Index < TestCount; ++Index)
+			{
+				Set.Emplace(Index);
+			}
+
+			CHECK(FSetTestDestructableData::Count == TestCount * 3);
+
+			SetCopy = Set;
+
+			CHECK(FSetTestDestructableData::Count == TestCount * 3);
+
+			SetMove = MoveTemp(Set);
+
+			CHECK(FSetTestDestructableData::Count == TestCount * 2);
+
+			SetCopy.Reset();
+
+			CHECK(FSetTestDestructableData::Count == TestCount);
+		}
+
 		SECTION("Advanced Tests")
 		{
 			int32 IntTestItems[256];
@@ -466,6 +535,11 @@ struct TSetTestHelper
 		}
 #endif
 	}
+
+	~TSetTestHelper()
+	{
+		CHECK(FSetTestDestructableData::Count == 0);
+	}
 };
 
 TEST_CASE_NAMED(FCompactSetTest, "System::Core::Containers::CompactSet", "[Core][Containers][CompactSet]")
