TaskGraph - Fix TGraphTask execute scopes not being emitted to Insights when using lightweight stats.

[RN] Fix TGraphTask execute scopes not being emitted to Insights when using lightweight stats
#tests insights replay run

diff --git a/Engine/Source/Runtime/Core/Public/Async/TaskGraphInterfaces.h b/Engine/Source/Runtime/Core/Public/Async/TaskGraphInterfaces.h
--- a/Engine/Source/Runtime/Core/Public/Async/TaskGraphInterfaces.h
+++ b/Engine/Source/Runtime/Core/Public/Async/TaskGraphInterfaces.h
@@ -706,7 +706,7 @@ private:
 		// Only use this scope when GetStatId() is actually implemented on TaskObject for backward compat, sometimes it's not.
 		if constexpr (TModels_V<CGetStatIdProvider, TTask>)
 		{
-			FScopeCycleCounter Scope(TaskObject->GetStatId(), true);
+			SCOPE_CYCLE_COUNTER_STATID(TaskObject->GetStatId());
 			TaskObject->DoTask(TaskObject->GetDesiredThread(), GraphEventRef);
 		}
 		else
