Added TTuplePopFront and TTuplePopBack for removing element types from the start or end of a tuple type.

diff --git a/Engine/Source/Runtime/Core/Private/Templates/TuplesTest.cpp b/Engine/Source/Runtime/Core/Private/Templates/TuplesTest.cpp
--- a/Engine/Source/Runtime/Core/Private/Templates/TuplesTest.cpp
+++ b/Engine/Source/Runtime/Core/Private/Templates/TuplesTest.cpp
@@ -248,6 +248,34 @@ static_assert(TIsTuple<const          TTuple<double, float&, char&&>>::Value, ""
 static_assert(TIsTuple<      volatile TTuple<double, float&, char&&>>::Value, "");
 static_assert(TIsTuple<const volatile TTuple<double, float&, char&&>>::Value, "");
 
+using TestPopTupleType = TTuple<float, char, int, void, bool>;
+
+// Check that TTuplePopBack works
+static_assert(std::is_same_v<TTuplePopBack<TestPopTupleType, 0>::Type, TTuple<float, char, int, void, bool>>);
+static_assert(std::is_same_v<TTuplePopBack<TestPopTupleType, 1>::Type, TTuple<float, char, int, void>>);
+static_assert(std::is_same_v<TTuplePopBack<TestPopTupleType, 2>::Type, TTuple<float, char, int>>);
+static_assert(std::is_same_v<TTuplePopBack<TestPopTupleType, 3>::Type, TTuple<float, char>>);
+static_assert(std::is_same_v<TTuplePopBack<TestPopTupleType, 4>::Type, TTuple<float>>);
+static_assert(std::is_same_v<TTuplePopBack<TestPopTupleType, 5>::Type, TTuple<>>);
+
+// Check that TTuplePopFront works
+static_assert(std::is_same_v<TTuplePopFront<TestPopTupleType, 0>::Type, TTuple<float, char, int, void, bool>>);
+static_assert(std::is_same_v<TTuplePopFront<TestPopTupleType, 1>::Type, TTuple<char, int, void, bool>>);
+static_assert(std::is_same_v<TTuplePopFront<TestPopTupleType, 2>::Type, TTuple<int, void, bool>>);
+static_assert(std::is_same_v<TTuplePopFront<TestPopTupleType, 3>::Type, TTuple<void, bool>>);
+static_assert(std::is_same_v<TTuplePopFront<TestPopTupleType, 4>::Type, TTuple<bool>>);
+static_assert(std::is_same_v<TTuplePopFront<TestPopTupleType, 5>::Type, TTuple<>>);
+
+// Check that TTuplePopBack preserves qualifiers
+static_assert(std::is_same_v<TTuplePopBack<const          TestPopTupleType, 2>::Type, const          TTuple<float, char, int>>);
+static_assert(std::is_same_v<TTuplePopBack<      volatile TestPopTupleType, 2>::Type,       volatile TTuple<float, char, int>>);
+static_assert(std::is_same_v<TTuplePopBack<const volatile TestPopTupleType, 2>::Type, const volatile TTuple<float, char, int>>);
+
+// Check that TTuplePopFront preserves qualifiers
+static_assert(std::is_same_v<TTuplePopFront<const          TestPopTupleType, 2>::Type, const          TTuple<int, void, bool>>);
+static_assert(std::is_same_v<TTuplePopFront<      volatile TestPopTupleType, 2>::Type,       volatile TTuple<int, void, bool>>);
+static_assert(std::is_same_v<TTuplePopFront<const volatile TestPopTupleType, 2>::Type, const volatile TTuple<int, void, bool>>);
+
 // These shouldn't compile - ideally giving a meaningful error message
 #if 0
 	// TTupleElement passed a non-tuple
diff --git a/Engine/Source/Runtime/Core/Public/Templates/Tuple.h b/Engine/Source/Runtime/Core/Public/Templates/Tuple.h
--- a/Engine/Source/Runtime/Core/Public/Templates/Tuple.h
+++ b/Engine/Source/Runtime/Core/Public/Templates/Tuple.h
@@ -391,6 +391,58 @@ namespace UE::Core::Private::Tuple
 		enum { Value = sizeof...(Types) };
 	};
 
+	template <typename TupleType, uint32 Num>
+	struct TTuplePopFront
+	{
+	};
+
+	template <typename Type0, typename... Types, uint32 Num>
+	struct TTuplePopFront<TTuple<Type0, Types...>, Num> : TTuplePopFront<TTuple<Types...>, Num - 1>
+	{
+	};
+
+	template <typename Type0, typename... Types>
+	struct TTuplePopFront<TTuple<Type0, Types...>, 0>
+	{
+		using Type = TTuple<Type0, Types...>;
+	};
+
+	template <>
+	struct TTuplePopFront<TTuple<>, 0>
+	{
+		using Type = TTuple<>;
+	};
+
+	template <typename TupleType, uint32 Num> struct TTuplePopFront<const          TupleType, Num> { using Type = const          typename TTuplePopFront<TupleType, Num>::Type; };
+	template <typename TupleType, uint32 Num> struct TTuplePopFront<      volatile TupleType, Num> { using Type =       volatile typename TTuplePopFront<TupleType, Num>::Type; };
+	template <typename TupleType, uint32 Num> struct TTuplePopFront<const volatile TupleType, Num> { using Type = const volatile typename TTuplePopFront<TupleType, Num>::Type; };
+
+	template <typename TupleType, uint32 Num, typename ResultTuple = TTuple<>>
+	struct TTuplePopBack
+	{
+	};
+
+	template <typename Type0, typename... Types, uint32 Num, typename... ResultTypes>
+	struct TTuplePopBack<TTuple<Type0, Types...>, Num, TTuple<ResultTypes...>> : TTuplePopBack<TTuple<Types...>, Num - 1, TTuple<ResultTypes..., Type0>>
+	{
+	};
+
+	template <typename Type0, typename... Types, typename... ResultTypes>
+	struct TTuplePopBack<TTuple<Type0, Types...>, 0, TTuple<ResultTypes...>>
+	{
+		using Type = TTuple<ResultTypes...>;
+	};
+
+	template <typename... ResultTypes>
+	struct TTuplePopBack<TTuple<>, 0, TTuple<ResultTypes...>>
+	{
+		using Type = TTuple<ResultTypes...>;
+	};
+
+	template <typename TupleType, uint32 Num> struct TTuplePopBack<const          TupleType, Num> { using Type = const          typename TTuplePopBack<TupleType, Num>::Type; };
+	template <typename TupleType, uint32 Num> struct TTuplePopBack<      volatile TupleType, Num> { using Type =       volatile typename TTuplePopBack<TupleType, Num>::Type; };
+	template <typename TupleType, uint32 Num> struct TTuplePopBack<const volatile TupleType, Num> { using Type = const volatile typename TTuplePopBack<TupleType, Num>::Type; };
+
 	template <typename Type, typename TupleType>
 	struct TCVTupleIndex
 	{
@@ -587,6 +639,34 @@ struct TTupleArity : UE::Core::Private::Tuple::TCVTupleArity<const volatile Tupl
 };
 
 
+/**
+ * Traits class which pops a number of element types from the start of a tuple type.
+ *
+ * Example: TTuplePopFront<TTuple<A, B, C, D, E>, 2>::Type == TTuple<C, D, E>
+ */
+template <typename TupleType, uint32 Num = 1>
+struct TTuplePopFront : UE::Core::Private::Tuple::TTuplePopFront<TupleType, Num>
+{
+};
+
+template <typename TupleType, uint32 Num = 1>
+using TTuplePopFront_T = typename TTuplePopFront<TupleType, Num>::Type;
+
+
+/**
+ * Traits class which pops a number of element types from the end of a tuple type.
+ *
+ * Example: TTuplePopBack<TTuple<A, B, C, D, E>, 2>::Type == TTuple<A, B, C>
+ */
+template <typename TupleType, uint32 Num = 1>
+struct TTuplePopBack : UE::Core::Private::Tuple::TTuplePopBack<TupleType, TTupleArity<TupleType>::Value - Num>
+{
+};
+
+template <typename TupleType, uint32 Num = 1>
+using TTuplePopBack_T = typename TTuplePopBack<TupleType, Num>::Type;
+
+
 /**
  * Traits class which gets the tuple index of a given type from a given TTuple.
  * If the type doesn't appear, or appears more than once, a compile error is generated.
