[Core]
* Removed experimental include from CoreDelegates.h

diff --git a/Engine/Plugins/Experimental/DefaultInstallBundleManager/Source/Private/DefaultInstallBundleManager.cpp b/Engine/Plugins/Experimental/DefaultInstallBundleManager/Source/Private/DefaultInstallBundleManager.cpp
--- a/Engine/Plugins/Experimental/DefaultInstallBundleManager/Source/Private/DefaultInstallBundleManager.cpp
+++ b/Engine/Plugins/Experimental/DefaultInstallBundleManager/Source/Private/DefaultInstallBundleManager.cpp
@@ -12,6 +12,7 @@
 #include "Misc/CommandLine.h"
 #include "IO/IoStoreOnDemand.h"
 #include "Misc/EnumClassFlags.h"
+#include "MountPaksExArgs.h"
 #include "PatchCheck.h"
 #include "Templates/Greater.h"
 #include "ShaderPipelineCache.h"
diff --git a/Engine/Source/Editor/UnrealEd/Private/Cooker/DiffPackageWriter.cpp b/Engine/Source/Editor/UnrealEd/Private/Cooker/DiffPackageWriter.cpp
--- a/Engine/Source/Editor/UnrealEd/Private/Cooker/DiffPackageWriter.cpp
+++ b/Engine/Source/Editor/UnrealEd/Private/Cooker/DiffPackageWriter.cpp
@@ -19,6 +19,7 @@
 #include "Misc/OutputDeviceHelper.h"
 #include "Misc/App.h"
 #include "ProfilingDebugging/CookStats.h"
+#include "Serialization/CompactBinaryWriter.h"
 #include "Serialization/LargeMemoryWriter.h"
 #include "Serialization/JsonWriter.h"
 #include "Serialization/CompactBinarySerialization.h"
diff --git a/Engine/Source/Runtime/Core/Private/Windows/WindowsPlatformProcess.cpp b/Engine/Source/Runtime/Core/Private/Windows/WindowsPlatformProcess.cpp
--- a/Engine/Source/Runtime/Core/Private/Windows/WindowsPlatformProcess.cpp
+++ b/Engine/Source/Runtime/Core/Private/Windows/WindowsPlatformProcess.cpp
@@ -22,6 +22,7 @@
 #include "Misc/Fork.h"
 #include "Misc/Parse.h"
 #include "Misc/Paths.h"
+#include "Misc/ScopeExit.h"
 #include "Misc/SingleThreadEvent.h"
 #include "Misc/PathViews.h"
 #include "Misc/TrackedActivity.h"
diff --git a/Engine/Source/Runtime/Core/Public/Misc/CoreDelegates.h b/Engine/Source/Runtime/Core/Public/Misc/CoreDelegates.h
--- a/Engine/Source/Runtime/Core/Public/Misc/CoreDelegates.h
+++ b/Engine/Source/Runtime/Core/Public/Misc/CoreDelegates.h
@@ -8,7 +8,6 @@
 #include "Containers/Map.h"
 #include "Containers/Set.h"
 #include "Containers/UnrealString.h"
-#include "Experimental/UnifiedError/UnifiedError.h"
 #include "CoreTypes.h"
 #include "Delegates/Delegate.h"
 #include "GenericPlatform/GenericPlatformFile.h"
@@ -82,19 +81,8 @@ struct FMultiprocessDetachedContext
 	bool bAbnormalDetach;
 };
 
-struct FMountPaksExArgs
-{
-	struct FMountResult
-	{
-		// Any pak file that was mounted, may be null on success if an encryption key is pending
-		IPakFile* PakFile = nullptr;
-	};
+struct FMountPaksExArgs;
 
-	const TCHAR* PakFilePath = nullptr;
-	int32 Order = INDEX_NONE;
-	FPakMountOptions MountOptions;
-	TValueOrError<FMountResult, UE::UnifiedError::FError> Result = MakeValue();
-};
 }
 
 class FCoreDelegates
@@ -641,3 +629,7 @@ private:
 	// This class is only for namespace use
 	FCoreDelegates() = default;
 };
+
+#if UE_ENABLE_INCLUDE_ORDER_DEPRECATED_IN_5_8
+#include "Experimental/UnifiedError/UnifiedError.h"
+#endif
\ No newline at end of file
diff --git a/Engine/Source/Runtime/PakFile/Private/IPlatformFilePak.cpp b/Engine/Source/Runtime/PakFile/Private/IPlatformFilePak.cpp
--- a/Engine/Source/Runtime/PakFile/Private/IPlatformFilePak.cpp
+++ b/Engine/Source/Runtime/PakFile/Private/IPlatformFilePak.cpp
@@ -16,6 +16,7 @@
 #include "Misc/ConfigCacheIni.h"
 #include "Misc/SecureHash.h"
 #include "Misc/ScopeLock.h"
+#include "MountPaksExArgs.h"
 #include "HAL/FileManagerGeneric.h"
 #include "SignedArchiveReader.h"
 #include "Misc/AES.h"
diff --git a/Engine/Source/Runtime/PakFile/Public/MountPaksExArgs.h b/Engine/Source/Runtime/PakFile/Public/MountPaksExArgs.h
new file mode 100644
--- /dev/null
+++ b/Engine/Source/Runtime/PakFile/Public/MountPaksExArgs.h
@@ -0,0 +1,24 @@
+// Copyright Epic Games, Inc. All Rights Reserved.
+
+#pragma once
+
+#include "Experimental/UnifiedError/UnifiedError.h"
+#include "GenericPlatform/GenericPlatformFile.h"
+#include "Templates/ValueOrError.h"
+
+namespace UE
+{
+	struct FMountPaksExArgs
+	{
+		struct FMountResult
+		{
+			// Any pak file that was mounted, may be null on success if an encryption key is pending
+			IPakFile* PakFile = nullptr;
+		};
+
+		const TCHAR* PakFilePath = nullptr;
+		int32 Order = INDEX_NONE;
+		FPakMountOptions MountOptions;
+		TValueOrError<FMountResult, UE::UnifiedError::FError> Result = MakeValue();
+	};
+}
