Fixed missing include in MacPlatformProcess

diff --git a/Engine/Source/Runtime/Core/Private/Mac/MacPlatformProcess.cpp b/Engine/Source/Runtime/Core/Private/Mac/MacPlatformProcess.cpp
--- a/Engine/Source/Runtime/Core/Private/Mac/MacPlatformProcess.cpp
+++ b/Engine/Source/Runtime/Core/Private/Mac/MacPlatformProcess.cpp
@@ -11,6 +11,7 @@
 #include "Misc/App.h"
 #include "Misc/CoreDelegates.h"
 #include "Misc/Paths.h"
+#include "Misc/ScopeExit.h"
 #include "Misc/StringBuilder.h"
 #include "HAL/FileManager.h"
 #include "Apple/PreAppleSystemHeaders.h"
