[StripedMap]
- Fix compilation error on type conversion

diff --git a/Engine/Source/Runtime/Core/Public/Containers/StripedMap.h b/Engine/Source/Runtime/Core/Public/Containers/StripedMap.h
--- a/Engine/Source/Runtime/Core/Public/Containers/StripedMap.h
+++ b/Engine/Source/Runtime/Core/Public/Containers/StripedMap.h
@@ -728,9 +728,9 @@ public:
 	 * @param Number The number of elements that the map should be able to contain after allocation.
 	 * @see Shrink
 	 */
-	void Reserve(SIZE_T Number)
+	void Reserve(int32 Number)
 	{
-		SIZE_T NumberPerMap = Number / BucketCount;
+		int32 NumberPerMap = Number / BucketCount;
 		ForEachMap(
 			[NumberPerMap](MapType& Map)
 			{
