I/O store container metadata sidecar file(s).
- extended I/O store container writer with support for writing an additional sidecar file <container>.umeta which contains mappings from FIoChunkId to filename and container name.
- disabled by default and enabled by passing -ContainerMeta to UnrealPak.exe
- creates a global.umeta file when creating base game container files and a per container .umeta file when creating DLC containers
- for encrypted containers the filename is replaced with <Encrypted>
- the metadata reader can read into memory as well as from a memory mapped file region

#rnx

diff --git a/Engine/Source/Developer/IoStoreUtilities/Private/IoStoreUtilities.cpp b/Engine/Source/Developer/IoStoreUtilities/Private/IoStoreUtilities.cpp
--- a/Engine/Source/Developer/IoStoreUtilities/Private/IoStoreUtilities.cpp
+++ b/Engine/Source/Developer/IoStoreUtilities/Private/IoStoreUtilities.cpp
@@ -20,6 +20,7 @@
 #include "Hash/xxhash.h"
 #include "Interfaces/ITargetPlatformManagerModule.h"
 #include "Interfaces/ITargetPlatform.h"
+#include "IO/IoContainerMeta.h"
 #include "IO/IoDispatcher.h"
 #include "Misc/App.h"
 #include "Misc/CommandLine.h"
@@ -5813,6 +5814,12 @@ int32 CreateTarget(const FIoStoreArguments& Arguments, const FIoStoreWriterSetti
 
 			UE_LOG(LogIoStore, Display, TEXT("%s"), *ProgressStringBuilder);
 		}
+
+		// When staging base game, all container metadata is aggregated into a single global.umeta file
+		const FString GlobalContainerMetaPath = FPaths::GetExtension(Arguments.GlobalContainerPath).IsEmpty()
+			? Arguments.GlobalContainerPath + FIoContainerMetaHeader::FileExtension
+			: FPaths::ChangeExtension(Arguments.GlobalContainerPath, FIoContainerMetaHeader::FileExtension);
+		IoStoreWriterContext->SaveContainerMeta(GlobalContainerMetaPath);
 	}
 
 	if (GeneralIoWriterSettings.bCompressionEnableDDC)
@@ -6897,6 +6904,84 @@ bool ListIoStoreContainerBulkData(const TCHAR* CmdLine)
 	return ListContainerBulkData(KeyChain, ContainerPathOrWildcard, OutFile) == 0;
 }
 
+bool ListIoStoreContainerMetadataFile(const FString& ContainerPathOrWildcard, const FString& OutFile)
+{
+	TArray<FString> ContainerFilePaths = GetContainersFromPathOrWildcard(ContainerPathOrWildcard);
+
+	if (ContainerFilePaths.Num() == 0)
+	{
+		UE_LOG(LogIoStore, Error, TEXT("Container '%s' doesn't exist and no container matches wildcard."), *ContainerPathOrWildcard);
+		return false;
+	}
+
+	TUniquePtr<FArchive> CsvAr;
+	if (OutFile.IsEmpty() == false)
+	{
+		CsvAr.Reset(IFileManager::Get().CreateFileWriter(*OutFile));
+		if (CsvAr.IsValid())
+		{
+			CsvAr->Logf(TEXT("ChunkId,ContainerName,Filename"));
+		}
+		else
+		{
+			UE_LOG(LogIoStore, Error, TEXT("Failed to create CSV '%s'"), *OutFile);
+			return false;
+		}
+	}
+
+	for (const FString& ContainerFilePath : ContainerFilePaths)
+	{
+		FString MetaPath = FPaths::ChangeExtension(ContainerFilePath, FIoContainerMetaHeader::FileExtension);
+		TIoStatusOr<FIoContainerMetaReader> MaybeReader = FIoContainerMetaReader::Load(MetaPath);
+		if (MaybeReader.IsOk() == false)
+		{
+			UE_LOG(LogIoStore, Error, TEXT("Failed load metadata file '%s', reason: %s"), *MetaPath, *MaybeReader.Status().ToString());
+			continue;
+		}
+
+		FIoContainerMetaReader Reader = MaybeReader.ConsumeValueOrDie();
+		Reader.Iterate([&CsvAr](const FIoChunkId& ChunkId, FUtf8StringView ContainerName, FUtf8StringView Filename)
+		{
+			if (CsvAr.IsValid())
+			{
+				CsvAr->Logf(TEXT("%s, %s, %s"), *LexToString(ChunkId), *FString(ContainerName), *FString(Filename));
+			}
+			else
+			{
+				UE_LOG(LogIoStore, Display, TEXT("%s, %s, %s"), *LexToString(ChunkId), *FString(ContainerName), *FString(Filename));
+			}
+
+			return true;
+		});
+	}
+
+	if (CsvAr.IsValid())
+	{
+		UE_LOG(LogIoStore, Display, TEXT("Saving '%s'"), *OutFile);
+		CsvAr.Reset();
+	}
+
+	return true;
+}
+
+bool ListIoStoreContainerMetadataFile(const TCHAR* CmdLine)
+{
+	FString ContainerPathOrWildcard;
+	if (!FParse::Value(FCommandLine::Get(), TEXT("ListContainerMetadataFile="), ContainerPathOrWildcard))
+	{
+		UE_LOG(LogIoStore, Error, TEXT("Missing argument -ListContainerMetadataFile=<ContainerFileOrWildCard>"));
+		return false;
+	}
+
+	FString OutFile;
+	if (!FParse::Value(FCommandLine::Get(), TEXT("Out="), OutFile))
+	{
+		OutFile.Empty();
+	}
+
+	return ListIoStoreContainerMetadataFile(ContainerPathOrWildcard, OutFile);
+}
+
 bool LegacyListIoStoreContainer(
 	const TCHAR* InContainerFilename,
 	int64 InSizeFilter,
@@ -9893,6 +9978,9 @@ bool ParseContainerGenerationArguments(FIoStoreArguments& Arguments, FIoStoreWri
 		WriterSettings.CompressionMinSizeToConsiderDDC = CompressionMinSizeToConsiderDDC;
 	}
 
+	WriterSettings.bContainerMetaEnabled = FParse::Param(FCommandLine::Get(), TEXT("ContainerMeta"));
+	WriterSettings.bPerContainerMeta = Arguments.IsDLC();
+
 	UE_LOG(LogIoStore, Display, TEXT("Using memory mapping alignment '%ld'"), WriterSettings.MemoryMappingAlignment);
 	UE_LOG(LogIoStore, Display, TEXT("Using compression block size '%ld'"), WriterSettings.CompressionBlockSize);
 	UE_LOG(LogIoStore, Display, TEXT("Using compression block alignment '%ld'"), WriterSettings.CompressionBlockAlignment);
@@ -9913,6 +10001,22 @@ bool ParseContainerGenerationArguments(FIoStoreArguments& Arguments, FIoStoreWri
 		UE_LOG(LogIoStore, Display, TEXT("Not using DDC for compression"));
 	}
 
+	if (WriterSettings.bContainerMetaEnabled)
+	{
+		if (WriterSettings.bPerContainerMeta)
+		{
+			UE_LOG(LogIoStore, Display, TEXT("Generating per container metadata file(s)"));
+		}
+		else
+		{
+			UE_LOG(LogIoStore, Display, TEXT("Generating global container metadata file"));
+		}
+	}
+	else
+	{
+		UE_LOG(LogIoStore, Display, TEXT("Container metadata disabled"));
+	}
+
 	FString CommandListFile;
 	if (FParse::Value(FCommandLine::Get(), TEXT("Commands="), CommandListFile))
 	{
diff --git a/Engine/Source/Developer/IoStoreUtilities/Private/IoStoreWriter.cpp b/Engine/Source/Developer/IoStoreUtilities/Private/IoStoreWriter.cpp
--- a/Engine/Source/Developer/IoStoreUtilities/Private/IoStoreWriter.cpp
+++ b/Engine/Source/Developer/IoStoreUtilities/Private/IoStoreWriter.cpp
@@ -16,6 +16,7 @@
 #include "DerivedDataRequestOwner.h"
 #include "HAL/FileManager.h"
 #include "HAL/PlatformFileManager.h"
+#include "IO/IoContainerMeta.h"
 #include "IO/IoDirectoryIndex.h"
 #include "Memory/SharedBuffer.h"
 #include "Misc/Compression.h"
@@ -561,6 +562,28 @@ public:
 		return ErrorCount.Load(EMemoryOrder::Relaxed);
 	}
 
+	void AddContainerMeta(const FString& ContainerName, const FIoChunkId& ChunkId, const FString& Filename)
+	{
+		GlobalContainerMetaWriter.AddFile(ContainerName, ChunkId, Filename);
+	}
+
+	void SaveContainerMeta(const FString& FilePath)
+	{
+		if (!WriterSettings.bContainerMetaEnabled || WriterSettings.bPerContainerMeta)
+		{
+			return;
+		}
+
+		if (const int64 FileSize = GlobalContainerMetaWriter.Save(FilePath); FileSize > 0)
+		{
+			UE_LOG(LogIoStore, Display, TEXT("Saved global metadata '%s' (%.2lf KiB)"), *FilePath, double(FileSize) / 1024.0);
+		}
+		else
+		{
+			UE_LOG(LogIoStore, Warning, TEXT("Failed to save global metadata '%s'"), *FilePath);
+		}
+	}
+
 private:
 	UE::DerivedData::FCacheKey MakeDDCKey(FIoStoreWriteQueueEntry* Entry) const;
 
@@ -603,6 +626,7 @@ private:
 	uint64 MaxCompressionBufferMemory = 0;
 	uint64 CompressionBufferSize = 0;
 	TArray<TSharedPtr<FIoStoreWriter>> IoStoreWriters;
+	FIoContainerMetaWriter GlobalContainerMetaWriter;
 
 	friend class FIoStoreWriter;
 };
@@ -643,6 +667,11 @@ uint32 FIoStoreWriterContext::GetErrors() const
 	return Impl->GetErrors();
 }
 
+void FIoStoreWriterContext::SaveContainerMeta(const FString& FilePath)
+{
+	Impl->SaveContainerMeta(FilePath);
+}
+
 class FIoStoreTocBuilder
 {
 public:
@@ -1325,6 +1354,44 @@ public:
 				ContainerSettings.IsEncrypted() ? ContainerSettings.EncryptionKey : FAES::FAESKey());
 		}
 
+		if (WriterContext->GetSettings().bContainerMetaEnabled)
+		{
+			const FString EncryptedFilename = TEXT("<Encrypted>");
+			FString ContainerName = FPaths::GetBaseFilename(TocFilePath);
+
+			if (WriterContext->GetSettings().bPerContainerMeta)
+			{
+				FIoContainerMetaWriter ContainerMetaWriter;
+				for (const FIoChunkId& ChunkId : TocResource.ChunkIds)
+				{
+					if (const FString* Filename = TocBuilder.GetFileName(ChunkId))
+					{
+						ContainerMetaWriter.AddFile(ContainerName, ChunkId, ContainerSettings.IsEncrypted() ? EncryptedFilename : *Filename);
+					}
+				}
+
+				FString ContainerMetaPath = FPaths::ChangeExtension(TocFilePath, FIoContainerMetaHeader::FileExtension);
+				if (const int64 FileSize = ContainerMetaWriter.Save(ContainerMetaPath); FileSize > 0)
+				{
+					UE_LOG(LogIoStore, Display, TEXT("Saved container metadata '%s' (%.2lf KiB)"), *ContainerMetaPath, double(FileSize) / 1024.0);
+				}
+				else
+				{
+					UE_LOG(LogIoStore, Warning, TEXT("Failed to save container metadata '%s'"), *ContainerMetaPath);
+				}
+			}
+			else
+			{
+				for (const FIoChunkId& ChunkId : TocResource.ChunkIds)
+				{
+					if (const FString* Filename = TocBuilder.GetFileName(ChunkId))
+					{
+						WriterContext->AddContainerMeta(ContainerName, ChunkId, ContainerSettings.IsEncrypted() ? EncryptedFilename : *Filename);
+					}
+				}
+			}
+		}
+
 		TIoStatusOr<uint64> TocSize = FIoStoreTocResource::Write(*TocFilePath, TocResource, static_cast<uint32>(WriterSettings.CompressionBlockSize), MaxPartitionSize, ContainerSettings);
 		checkf(TocSize.IsOk(), TEXT("FIoStoreTocResource::Write failed with error %s"), *TocSize.Status().ToString());
 		
diff --git a/Engine/Source/Developer/IoStoreUtilities/Private/IoStoreWriter.h b/Engine/Source/Developer/IoStoreUtilities/Private/IoStoreWriter.h
--- a/Engine/Source/Developer/IoStoreUtilities/Private/IoStoreWriter.h
+++ b/Engine/Source/Developer/IoStoreUtilities/Private/IoStoreWriter.h
@@ -24,6 +24,8 @@ struct FIoStoreWriterSettings
 	uint64 MaxPartitionSize = 0;
 	TMap<FString, uint64> MaxPartitionSizeOverride;
 	bool bEnableFileRegions = false;
+	bool bContainerMetaEnabled = false;
+	bool bPerContainerMeta = false;
 	bool bCompressionEnableDDC = false;
 	bool bValidateChunkHashes = false;
 };
@@ -97,6 +99,7 @@ public:
 	void Flush();
 	FProgress GetProgress() const;
 	uint32 GetErrors() const;
+	void SaveContainerMeta(const FString& FilePath);
 
 private:
 	FIoStoreWriterContextImpl* Impl;
diff --git a/Engine/Source/Developer/IoStoreUtilities/Public/IoStoreUtilities.h b/Engine/Source/Developer/IoStoreUtilities/Public/IoStoreUtilities.h
--- a/Engine/Source/Developer/IoStoreUtilities/Public/IoStoreUtilities.h
+++ b/Engine/Source/Developer/IoStoreUtilities/Public/IoStoreUtilities.h
@@ -31,6 +31,7 @@ UE_API bool LegacyListIoStoreContainer(
 
 UE_API bool ListIoStoreContainer(const TCHAR* CmdLine);
 UE_API bool ListIoStoreContainerBulkData(const TCHAR* CmdLine);
+UE_API bool ListIoStoreContainerMetadataFile(const TCHAR* CmdLine);
 
 UE_API bool DiffIoStoreContainer(const TCHAR* CmdLine);
 UE_API bool LegacyDiffIoStoreContainers(
diff --git a/Engine/Source/Developer/PakFileUtilities/Private/PakFileUtilities.cpp b/Engine/Source/Developer/PakFileUtilities/Private/PakFileUtilities.cpp
--- a/Engine/Source/Developer/PakFileUtilities/Private/PakFileUtilities.cpp
+++ b/Engine/Source/Developer/PakFileUtilities/Private/PakFileUtilities.cpp
@@ -5636,6 +5636,11 @@ bool ExecuteUnrealPak(const TCHAR* CmdLine)
 			return ListIoStoreContainerBulkData(CmdLine);
 		}
 
+		if (FParse::Value(CmdLine, TEXT("-ListContainerMetadataFile="), IoStoreArg))
+		{
+			return ListIoStoreContainerMetadataFile(CmdLine);
+		}
+
 		if (FParse::Value(CmdLine, TEXT("-DiffContainer="), IoStoreArg))
 		{
 			return DiffIoStoreContainer(CmdLine);
diff --git a/Engine/Source/Runtime/Core/Internal/IO/IoContainerMeta.h b/Engine/Source/Runtime/Core/Internal/IO/IoContainerMeta.h
new file mode 100644
--- /dev/null
+++ b/Engine/Source/Runtime/Core/Internal/IO/IoContainerMeta.h
@@ -0,0 +1,176 @@
+// Copyright Epic Games, Inc. All Rights Reserved.
+
+#pragma once
+
+#include "Async/MappedFileHandle.h"
+#include "Containers/Array.h"
+#include "Containers/ArrayView.h"
+#include "Containers/ChunkedArray.h"
+#include "Containers/Map.h"
+#include "Containers/StringView.h"
+#include "Containers/Utf8String.h"
+#include "IO/IoChunkId.h"
+#include "IO/IoStatus.h"
+#include "Memory/MemoryView.h"
+#include "Misc/StringBuilder.h"
+#include "Misc/TVariant.h"
+
+////////////////////////////////////////////////////////////////////////////////
+struct FIoMetaDirectoryIndexEntry
+{
+	uint32 Name				= ~uint32(0);
+	uint32 ParentEntry		= ~uint32(0);
+	uint32 FirstChildEntry	= ~uint32(0);
+	uint32 NextSiblingEntry	= ~uint32(0);
+	uint32 FirstFileEntry	= ~uint32(0);
+
+	CORE_API friend FArchive& operator<<(FArchive& Ar, FIoMetaDirectoryIndexEntry& Entry);
+};
+
+////////////////////////////////////////////////////////////////////////////////
+struct FIoMetaFileIndexEntry
+{
+	uint32 Name				= ~uint32(0);
+	uint32 ContainerName	= ~uint32(0);
+	uint32 DirectoryEntry	= ~uint32(0);
+	uint32 NextFileEntry	= ~uint32(0);
+
+	CORE_API friend FArchive& operator<<(FArchive& Ar, FIoMetaFileIndexEntry& Entry);
+};
+
+////////////////////////////////////////////////////////////////////////////////
+struct FIoMetaStringTableEntry
+{
+	const uint32 Offset = 0;
+	const uint32 Len	= 0;
+};
+
+////////////////////////////////////////////////////////////////////////////////
+struct FIoContainerMetaHeader
+{
+	CORE_API static const TCHAR*	FileExtension;
+	static const inline uint8		MagicSequence[16] = {'C', 'O', 'N', 'T', 'A', 'I', 'N', 'E', 'R', 'M', 'E', 'T', 'A', 'H', 'D', 'R'};
+
+	enum class EVersion : uint32
+	{
+		Invalid	= 0,
+		Initial,
+
+		LatestPlusOne,
+		Latest = LatestPlusOne - 1
+	};
+
+	CORE_API bool IsValid() const;
+
+	uint8	Magic[16]		= {0};
+	uint32	Version			= 0;
+	uint32	FileCount		= 0;
+	uint32	DirectoryCount	= 0;
+	uint32	StringCount		= 0;
+	uint8	Pad[32]			= {0};
+};
+static_assert(sizeof(FIoContainerMetaHeader) == 64);
+
+////////////////////////////////////////////////////////////////////////////////
+struct FIoContainerMetaResourceView
+{
+	bool											IsValid() const { return Header != nullptr && Header->IsValid(); }
+	inline FUtf8StringView							GetString(uint32 StringTableEntryIndex) const;
+	inline FUtf8StringView							GetString(FIoMetaStringTableEntry Entry) const;
+
+	CORE_API static FIoContainerMetaResourceView	MakeView(FMemoryView MemoryView);
+
+	const FIoContainerMetaHeader*					Header = nullptr;
+	TConstArrayView<FIoChunkId>						ChunkIds;
+	TConstArrayView<uint32>							FileEntryIndices;
+	TConstArrayView<FIoMetaFileIndexEntry>			FileEntries;
+	TConstArrayView<FIoMetaDirectoryIndexEntry>		DirectoryEntries;
+	TConstArrayView<FIoMetaStringTableEntry>		StringTableEntries;
+	FMemoryView										StringTable;
+};
+
+////////////////////////////////////////////////////////////////////////////////
+FUtf8StringView FIoContainerMetaResourceView::GetString(FIoMetaStringTableEntry Entry) const
+{
+	return FUtf8StringView(reinterpret_cast<const UTF8CHAR*>(StringTable.RightChop(Entry.Offset).GetData()), Entry.Len);
+}
+
+FUtf8StringView FIoContainerMetaResourceView::GetString(uint32 StringTableEntryIndex) const
+{
+	check(StringTableEntryIndex != ~uint32(0));
+	return GetString(StringTableEntries[int32(StringTableEntryIndex)]);
+}
+
+////////////////////////////////////////////////////////////////////////////////
+using FIoContainerMetaVisistor = TFunction<bool(const FIoChunkId& ChunkId, FUtf8StringView ContainerName, FUtf8StringView Filename)>;
+
+////////////////////////////////////////////////////////////////////////////////
+class FIoContainerMetaReader
+{
+	struct FMappedFileStorage
+	{
+		TUniquePtr<IMappedFileHandle> MappedFile;
+		TUniquePtr<IMappedFileRegion> MappedRegion;
+	};
+
+	struct FMemoryStorage
+	{
+		TArray<uint8> Buffer;
+	};
+
+	using FStorage = TVariant<FEmptyVariantState, FMappedFileStorage, FMemoryStorage>;
+
+public:
+	CORE_API					~FIoContainerMetaReader();
+	CORE_API FUtf8StringView	GetFilename(const FIoChunkId& ChunkId, FUtf8StringBuilderBase& OutFilename, FUtf8StringView& OutContainerName);
+	CORE_API void				Iterate(FIoContainerMetaVisistor&& Visitor);
+	CORE_API static TIoStatusOr<FIoContainerMetaReader> Load(const FString& FilePath, bool bEnableMemoryMapping = true);
+
+	FIoContainerMetaReader& operator=(const FIoContainerMetaReader&) = delete;
+	FIoContainerMetaReader&	operator=(FIoContainerMetaReader&&) = default;
+	FIoContainerMetaReader(FIoContainerMetaReader&&) = default;
+	FIoContainerMetaReader(const FIoContainerMetaReader&) = delete;
+	
+private:
+	FUtf8StringView	GetFilename(uint32 File, FUtf8StringBuilderBase& OutFilename, FUtf8StringView& OutContainerName);
+
+	FIoContainerMetaReader(FMemoryStorage&& InStorage, FIoContainerMetaResourceView InView)
+		: View(InView)
+	{
+		Storage.Emplace<FMemoryStorage>(MoveTemp(InStorage));
+	}
+
+	FIoContainerMetaReader(FMappedFileStorage&& InStorage, FIoContainerMetaResourceView InView)
+		: View(InView)
+	{
+		Storage.Emplace<FMappedFileStorage>(MoveTemp(InStorage));
+	}
+
+	FIoContainerMetaResourceView	View;
+	FStorage						Storage;
+};
+
+////////////////////////////////////////////////////////////////////////////////
+class FIoContainerMetaWriter
+{
+public:
+	CORE_API		FIoContainerMetaWriter();
+	CORE_API void	AddFile(FStringView ContainerName, const FIoChunkId& ChunkId, FStringView Filename);
+	CORE_API void	Save(FArchive& Ar);
+	CORE_API int64	Save(const FString& FilePath);
+
+	FIoContainerMetaWriter(const FIoContainerMetaWriter&) = delete;
+	FIoContainerMetaWriter& operator=(const FIoContainerMetaWriter&) = delete;
+private:
+	uint32			GetDirectory(uint32 DirectoryName, uint32 Parent);
+	uint32 			CreateDirectory(FStringView DirectoryName, uint32 Parent);
+	uint32 			CreateDirectoryTree(FStringView Path);
+	uint32 			AddFile(FStringView FileName, FStringView ContainerName, uint32 Directory);
+	uint32 			GetNameIndex(const FStringView& String);
+
+	TMap<FUtf8StringView, uint32>		StringToIndex;
+	TMap<FIoChunkId, uint32>			ChunkIdToFileEntry;
+	TChunkedArray<FUtf8String>			Strings;
+	TArray<FIoMetaDirectoryIndexEntry>	DirectoryEntries;
+	TArray<FIoMetaFileIndexEntry>		FileEntries;
+};
diff --git a/Engine/Source/Runtime/Core/Private/IO/IoContainerMeta.cpp b/Engine/Source/Runtime/Core/Private/IO/IoContainerMeta.cpp
new file mode 100644
--- /dev/null
+++ b/Engine/Source/Runtime/Core/Private/IO/IoContainerMeta.cpp
@@ -0,0 +1,412 @@
+// Copyright Epic Games, Inc. All Rights Reserved.
+
+#include "IO/IoContainerMeta.h"
+
+#include "Algo/BinarySearch.h"
+#include "HAL/FileManager.h"
+#include "HAL/PlatformFileManager.h"
+#include "Misc/PathViews.h"
+#include "Serialization/Archive.h"
+#include "Serialization/LargeMemoryWriter.h"
+
+////////////////////////////////////////////////////////////////////////////////
+namespace IoContainerMeta
+{
+
+static bool IsValid(uint32 Index)
+{
+	return Index != ~uint32(0);
+}
+
+static FStringView GetNextDirectoryName(const FStringView& Path)
+{
+	int32 SeparatorIndex = INDEX_NONE;
+	if (Path.FindChar(TEXT('/'), SeparatorIndex))
+	{
+		return Path.Left(SeparatorIndex);
+	}
+
+	return FStringView();
+}
+
+} // namespace IoContainerMeta
+
+////////////////////////////////////////////////////////////////////////////////
+FArchive& operator<<(FArchive& Ar, FIoMetaDirectoryIndexEntry& Entry)
+{
+	Ar << Entry.Name;
+	Ar << Entry.ParentEntry;
+	Ar << Entry.FirstChildEntry;
+	Ar << Entry.NextSiblingEntry;
+	Ar << Entry.FirstFileEntry;
+
+	return Ar;
+}
+
+////////////////////////////////////////////////////////////////////////////////
+FArchive& operator<<(FArchive& Ar, FIoMetaFileIndexEntry& Entry)
+{
+	Ar << Entry.Name;
+	Ar << Entry.ContainerName;
+	Ar << Entry.DirectoryEntry;
+	Ar << Entry.NextFileEntry;
+
+	return Ar;
+}
+
+////////////////////////////////////////////////////////////////////////////////
+bool FIoContainerMetaHeader::IsValid() const
+{
+	if (FMemory::Memcmp(Magic, MagicSequence, sizeof(MagicSequence)) != 0)
+	{
+		return false;
+	}
+
+	return static_cast<EVersion>(Version) != EVersion::Invalid;
+}
+
+const TCHAR* FIoContainerMetaHeader::FileExtension = TEXT(".umeta");
+
+////////////////////////////////////////////////////////////////////////////////
+FIoContainerMetaReader::~FIoContainerMetaReader()
+{
+}
+
+FUtf8StringView	FIoContainerMetaReader::GetFilename(const FIoChunkId& ChunkId, FUtf8StringBuilderBase& OutFilename, FUtf8StringView& OutContainerName)
+{
+	if (int32 Index = Algo::LowerBound(View.ChunkIds, ChunkId); Index < View.ChunkIds.Num())
+	{
+		if (View.ChunkIds[Index] == ChunkId)
+		{
+			const uint32 File = View.FileEntryIndices[Index];
+			return GetFilename(File, OutFilename, OutContainerName);
+		}
+	}
+
+	OutFilename.Reset();
+	return OutContainerName = FUtf8StringView();
+}
+
+void FIoContainerMetaReader::Iterate(FIoContainerMetaVisistor&& Visitor)
+{
+	TUtf8StringBuilder<2048> Sb;
+	for (int32 Idx = 0; const FIoChunkId& ChunkId : View.ChunkIds)
+	{
+		const uint32	File = View.FileEntryIndices[Idx++];
+		FUtf8StringView ContainerName;
+		FUtf8StringView Filename = GetFilename(File, Sb, ContainerName);
+
+		if (Visitor(ChunkId, ContainerName, Filename) == false)
+		{
+			break;
+		}
+	}
+}
+
+TIoStatusOr<FIoContainerMetaReader> FIoContainerMetaReader::Load(const FString& FilePath, bool bEnableMemoryMapping)
+{
+	IPlatformFile& Ipf = FPlatformFileManager::Get().GetPlatformFile();
+
+	const int64 FileSize = Ipf.FileSize(*FilePath);
+	if (FileSize < int64(sizeof(FIoContainerMetaHeader)) || FileSize >= MAX_int32)
+	{
+		return FIoStatus(EIoErrorCode::FileNotOpen);
+	}
+
+	if (FPlatformProperties::SupportsMemoryMappedFiles() && bEnableMemoryMapping)
+	{
+		FOpenMappedResult OpenResult = Ipf.OpenMappedEx(*FilePath);
+		if (OpenResult.HasError() == false)
+		{
+			FMappedFileStorage FileStorage;
+			FileStorage.MappedFile		= OpenResult.StealValue();
+			FileStorage.MappedRegion	= TUniquePtr<IMappedFileRegion>(FileStorage.MappedFile->MapRegion());
+
+			FIoContainerMetaResourceView View = FIoContainerMetaResourceView::MakeView(
+				FMemoryView(FileStorage.MappedRegion->GetMappedPtr(), FileStorage.MappedRegion->GetMappedSize()));
+
+			if (View.IsValid())
+			{
+				return FIoContainerMetaReader(MoveTemp(FileStorage), View);
+			}
+		}
+	}
+
+	TUniquePtr<IFileHandle> FileHandle(Ipf.OpenRead(*FilePath));
+	if (FileHandle.IsValid() == false)
+	{
+		return FIoStatus(EIoErrorCode::FileOpenFailed);
+	}
+
+	FMemoryStorage MemoryStorage;
+	MemoryStorage.Buffer.SetNumUninitialized(int32(FileSize));
+
+	if (FileHandle->Read(MemoryStorage.Buffer.GetData(), FileSize) == false)
+	{
+		return FIoStatus(EIoErrorCode::ReadError);
+	}
+
+	FIoContainerMetaResourceView View = FIoContainerMetaResourceView::MakeView(
+		FMemoryView(MemoryStorage.Buffer.GetData(), MemoryStorage.Buffer.Num()));
+
+	if (View.IsValid() == false)
+	{
+		return FIoStatus(EIoErrorCode::FileCorrupt);
+	}
+
+	return FIoContainerMetaReader(MoveTemp(MemoryStorage), View);
+}
+
+FUtf8StringView	FIoContainerMetaReader::GetFilename(uint32 File, FUtf8StringBuilderBase& OutFilename, FUtf8StringView& OutContainerName)
+{
+	check(File != ~uint32(0));
+
+	OutFilename.Reset();
+
+	const FIoMetaFileIndexEntry&			FileEntry = View.FileEntries[File];
+	uint32									Directory = FileEntry.DirectoryEntry;
+	TArray<uint32, TInlineAllocator<32>>	Segments;
+
+	check(FileEntry.Name != ~uint32(0));
+	check(Directory != ~uint32(0));
+	Segments.Add(FileEntry.Name);
+
+	while (Directory != ~uint32(0) && Directory != 0)
+	{
+		const FIoMetaDirectoryIndexEntry& DirectoryEntry = View.DirectoryEntries[Directory];
+		check(DirectoryEntry.Name != ~uint32(0));
+		Segments.Add(DirectoryEntry.Name);
+		Directory = DirectoryEntry.ParentEntry;
+	}
+
+	for (int32 Seg = Segments.Num() - 1; Seg >= 0; --Seg)
+	{
+		FUtf8StringView String = View.GetString(Segments[Seg]);
+
+		OutFilename << FString(String);
+		if (Seg > 0)
+		{
+			OutFilename << TEXT("/");
+		}
+	}
+
+	OutContainerName = View.GetString(FileEntry.ContainerName);
+	return OutFilename.ToView();
+}
+
+////////////////////////////////////////////////////////////////////////////////
+FIoContainerMetaWriter::FIoContainerMetaWriter()
+{
+	// Root directory
+	DirectoryEntries.AddDefaulted();
+}
+
+void FIoContainerMetaWriter::AddFile(FStringView ContainerName, const FIoChunkId& ChunkId, FStringView Filename)
+{
+	// Only store filenames for packages for now
+	if (ChunkIdToFileEntry.Contains(ChunkId) || ChunkId.GetChunkType() != EIoChunkType::ExportBundleData)
+	{
+		return;
+	}
+
+	FStringView		CleanFilename = FPathViews::GetCleanFilename(Filename);
+	const uint32	Directory = CreateDirectoryTree(Filename);
+	const uint32	FileEntry = AddFile(CleanFilename, ContainerName, Directory);
+
+	ChunkIdToFileEntry.Add(ChunkId, FileEntry);
+}
+
+void FIoContainerMetaWriter::Save(FArchive& Ar)
+{
+	check(ChunkIdToFileEntry.Num() == FileEntries.Num());
+
+	TArray<FIoChunkId>	ChunkIds;
+	TArray<uint32>		FileEntryIndices;
+
+	ChunkIdToFileEntry.GetKeys(ChunkIds);
+	ChunkIds.Sort();
+	FileEntryIndices.Reserve(ChunkIds.Num());
+
+	for (const FIoChunkId& ChunkId : ChunkIds)
+	{
+		FileEntryIndices.Add(ChunkIdToFileEntry.FindRef(ChunkId));
+	}
+
+	check(ChunkIds.Num() == FileEntryIndices.Num());
+	check(FileEntryIndices.Num() == FileEntries.Num());
+
+	FIoContainerMetaHeader Header;
+	Header.Version			= static_cast<uint32>(FIoContainerMetaHeader::EVersion::Latest);
+	Header.FileCount		= ChunkIds.Num();
+	Header.DirectoryCount	= DirectoryEntries.Num();
+	Header.StringCount		= Strings.Num();
+	FMemory::Memcpy(Header.Magic, FIoContainerMetaHeader::MagicSequence, sizeof(FIoContainerMetaHeader::MagicSequence));
+
+	Ar.Serialize(&Header, sizeof(FIoContainerMetaHeader));
+	Ar.Serialize(ChunkIds.GetData(), ChunkIds.GetTypeSize() * ChunkIds.Num());
+	Ar.Serialize(FileEntryIndices.GetData(), FileEntryIndices.GetTypeSize() * FileEntryIndices.Num());
+	Ar.Serialize(FileEntries.GetData(), FileEntries.GetTypeSize() * FileEntries.Num());
+	Ar.Serialize(DirectoryEntries.GetData(), DirectoryEntries.GetTypeSize() * DirectoryEntries.Num());
+	
+	{
+		FLargeMemoryWriter	StringTableEntryAr;
+		FLargeMemoryWriter	StringsAr;
+
+		for (const FUtf8String& String : Strings)
+		{
+			uint32 Offset	= static_cast<uint32>(StringsAr.Tell());
+			uint32 Len		= String.Len();
+
+			StringTableEntryAr << Offset;
+			StringTableEntryAr << Len;
+
+			StringsAr.Serialize((void*)*String, sizeof(UTF8CHAR) * String.Len());
+		}
+
+		Ar.Serialize(StringTableEntryAr.GetData(), StringTableEntryAr.TotalSize());
+		Ar.Serialize(StringsAr.GetData(), StringsAr.TotalSize());
+	}
+
+	StringToIndex.Empty();
+	ChunkIdToFileEntry.Empty();
+	Strings.Empty();
+	DirectoryEntries = TArray<FIoMetaDirectoryIndexEntry>();
+	DirectoryEntries.AddDefaulted();
+	FileEntries = TArray<FIoMetaFileIndexEntry>();
+}
+
+int64 FIoContainerMetaWriter::Save(const FString& FilePath)
+{
+	if (TUniquePtr<FArchive> Ar(IFileManager::Get().CreateFileWriter(*FilePath)); Ar.IsValid())
+	{
+		Save(*Ar);
+		if (int64 FileSize = Ar->TotalSize(); Ar->Close())
+		{
+			return FileSize;
+		}
+	}
+
+	return -1;
+}
+
+uint32 FIoContainerMetaWriter::GetDirectory(uint32 DirectoryName, uint32 Parent)
+{
+	check(DirectoryEntries.IsValidIndex(int32(Parent)));
+	uint32 Directory = DirectoryEntries[Parent].FirstChildEntry;
+	while (IoContainerMeta::IsValid(Directory))
+	{
+		const FIoMetaDirectoryIndexEntry& Entry = DirectoryEntries[Directory];
+		if (Entry.Name == DirectoryName)
+		{
+			return Directory;
+		}
+		Directory = Entry.NextSiblingEntry;
+	}
+
+	return ~uint32(0);
+}
+
+uint32 FIoContainerMetaWriter::CreateDirectory(FStringView DirectoryName, uint32 Parent)
+{
+	check(DirectoryEntries.IsValidIndex(int32(Parent)));
+
+	uint32 Name = GetNameIndex(DirectoryName);
+	uint32 Directory = GetDirectory(Name, Parent);
+
+	if (IoContainerMeta::IsValid(Directory))
+	{
+		return Directory;
+	}
+
+	Directory = DirectoryEntries.Num();
+	FIoMetaDirectoryIndexEntry& NewEntry = DirectoryEntries.AddDefaulted_GetRef();
+	NewEntry.Name = Name;
+	NewEntry.ParentEntry = Parent;
+	NewEntry.NextSiblingEntry = DirectoryEntries[Parent].FirstChildEntry;
+	DirectoryEntries[Parent].FirstChildEntry = Directory;
+
+	return Directory;
+}
+
+uint32 FIoContainerMetaWriter::CreateDirectoryTree(FStringView Path)
+{
+	using namespace IoContainerMeta;
+
+	uint32		Directory = 0; // Root
+	FStringView DirectoryName = GetNextDirectoryName(Path);
+
+	while (DirectoryName.IsEmpty() == false)
+	{
+		Directory = CreateDirectory(DirectoryName, Directory);
+		Path.RightChopInline(DirectoryName.Len() + 1);
+		DirectoryName = GetNextDirectoryName(Path);
+	}
+
+	return Directory;
+}
+
+uint32 FIoContainerMetaWriter::AddFile(FStringView FileName, FStringView ContainerName, uint32 Directory)
+{
+	check(Directory != ~uint32(0));
+
+	const uint32			File		= FileEntries.Num();
+	FIoMetaFileIndexEntry&	FileEntry	= FileEntries.AddDefaulted_GetRef();
+	FileEntry.Name						= GetNameIndex(FileName);
+	FileEntry.ContainerName				= GetNameIndex(ContainerName);
+	FileEntry.DirectoryEntry			= Directory;
+
+	FIoMetaDirectoryIndexEntry& DirectoryEntry	= DirectoryEntries[Directory];
+	FileEntry.NextFileEntry						= DirectoryEntry.FirstFileEntry;
+	DirectoryEntry.FirstFileEntry				= File;
+
+	return File;
+}
+
+uint32 FIoContainerMetaWriter::GetNameIndex(const FStringView& String)
+{
+	FUtf8String StringEntry(String);
+	if (const uint32* Index = StringToIndex.Find(StringEntry))
+	{
+		return *Index;
+	}
+	else
+	{
+		const uint32 NewIndex = Strings.AddElement(MoveTemp(StringEntry));
+		StringToIndex.Add(Strings[NewIndex], NewIndex);
+		return NewIndex;
+	}
+}
+
+FIoContainerMetaResourceView FIoContainerMetaResourceView::MakeView(FMemoryView Memory)
+{
+	if (Memory.GetSize() < sizeof(FIoContainerMetaHeader))
+	{
+		return FIoContainerMetaResourceView();
+	}
+	
+	FIoContainerMetaResourceView View;
+	View.Header = reinterpret_cast<const FIoContainerMetaHeader*>(Memory.GetData());
+
+	if (View.Header->IsValid() == false)
+	{
+		return View;
+	}
+	Memory.RightChopInline(sizeof(FIoContainerMetaHeader));
+
+	View.ChunkIds = MakeConstArrayView(reinterpret_cast<const FIoChunkId*>(Memory.GetData()), View.Header->FileCount);
+	Memory.RightChopInline(View.ChunkIds.GetTypeSize() * View.ChunkIds.Num());
+
+	View.FileEntryIndices = MakeConstArrayView(reinterpret_cast<const uint32*>(Memory.GetData()), View.Header->FileCount);
+	Memory.RightChopInline(View.FileEntryIndices.GetTypeSize() * View.FileEntryIndices.Num());
+
+	View.FileEntries = MakeConstArrayView(reinterpret_cast<const FIoMetaFileIndexEntry*>(Memory.GetData()), View.Header->FileCount);
+	Memory.RightChopInline(View.FileEntries.GetTypeSize() * View.FileEntries.Num());
+
+	View.DirectoryEntries = MakeConstArrayView(reinterpret_cast<const FIoMetaDirectoryIndexEntry*>(Memory.GetData()), View.Header->DirectoryCount);
+	Memory.RightChopInline(View.DirectoryEntries.GetTypeSize() * View.DirectoryEntries.Num());
+
+	View.StringTableEntries	= MakeConstArrayView(reinterpret_cast<const FIoMetaStringTableEntry*>(Memory.GetData()), View.Header->StringCount);
+	View.StringTable = Memory.RightChop(View.StringTableEntries.GetTypeSize() * View.StringTableEntries.Num());
+
+	return View;
+}
