Regex: Add Replace methods.

diff --git a/Engine/Source/Runtime/Core/Private/Internationalization/ICURegex.cpp b/Engine/Source/Runtime/Core/Private/Internationalization/ICURegex.cpp
--- a/Engine/Source/Runtime/Core/Private/Internationalization/ICURegex.cpp
+++ b/Engine/Source/Runtime/Core/Private/Internationalization/ICURegex.cpp
@@ -188,6 +188,34 @@ bool FRegexMatcher::FindNext()
 	return ICURegexMatcher.IsValid() && ICURegexMatcher->find() != 0;
 }
 
+FString FRegexMatcher::ReplaceFirst(FStringView Replacement)
+{
+	TSharedPtr<icu::RegexMatcher> ICURegexMatcher = Implementation->GetInternalRegexMatcher();
+	if (ICURegexMatcher.IsValid())
+	{
+		UErrorCode ICUStatus = U_ZERO_ERROR;
+		icu::UnicodeString ReplacementString = ICUUtilities::ConvertString(Replacement);
+		icu::UnicodeString Result = ICURegexMatcher->replaceFirst(ReplacementString, ICUStatus);
+		return ICUUtilities::ConvertString(Result);
+	}
+
+	return FString(Replacement);
+}
+
+FString FRegexMatcher::ReplaceAll(FStringView Replacement)
+{
+	TSharedPtr<icu::RegexMatcher> ICURegexMatcher = Implementation->GetInternalRegexMatcher();
+	if (ICURegexMatcher.IsValid())
+	{
+		UErrorCode ICUStatus = U_ZERO_ERROR;
+		icu::UnicodeString ReplacementString = ICUUtilities::ConvertString(Replacement);
+		icu::UnicodeString Result = ICURegexMatcher->replaceAll(ReplacementString, ICUStatus);
+		return ICUUtilities::ConvertString(Result);
+	}
+
+	return FString(Replacement);
+}
+
 int32 FRegexMatcher::GetMatchBeginning()
 {
 	UErrorCode ICUStatus = U_ZERO_ERROR;
diff --git a/Engine/Source/Runtime/Core/Private/Internationalization/LegacyRegex.cpp b/Engine/Source/Runtime/Core/Private/Internationalization/LegacyRegex.cpp
--- a/Engine/Source/Runtime/Core/Private/Internationalization/LegacyRegex.cpp
+++ b/Engine/Source/Runtime/Core/Private/Internationalization/LegacyRegex.cpp
@@ -38,6 +38,16 @@ bool FRegexMatcher::FindNext()
 	return false;
 }
 
+FString FRegexMatcher::ReplaceFirst(FStringView Replacement)
+{
+	return FString();
+}
+
+FString FRegexMatcher::ReplaceAll(FStringView Replacement)
+{
+	return FString();
+}
+
 int32 FRegexMatcher::GetMatchBeginning()
 {
 	return INDEX_NONE;
diff --git a/Engine/Source/Runtime/Core/Public/Internationalization/Regex.h b/Engine/Source/Runtime/Core/Public/Internationalization/Regex.h
--- a/Engine/Source/Runtime/Core/Public/Internationalization/Regex.h
+++ b/Engine/Source/Runtime/Core/Public/Internationalization/Regex.h
@@ -54,6 +54,9 @@ public:
 
 	CORE_API bool FindNext();
 
+	CORE_API FString ReplaceFirst(FStringView Replacement);
+	CORE_API FString ReplaceAll(FStringView Replacement);
+
 	CORE_API int32 GetMatchBeginning();
 	CORE_API int32 GetMatchEnding();
 
