Text debugging fixes

* Log DisplayString values as escaped, so that we can more easily view whitespace (like newlines) in the value.
* FTextLocalizationManager::AddOrUpdateDisplayStringInLiveTable now updates the SourceStringHash when a source string is provided.

#rnx

diff --git a/Engine/Source/Runtime/Core/Private/Internationalization/TextLocalizationManager.cpp b/Engine/Source/Runtime/Core/Private/Internationalization/TextLocalizationManager.cpp
--- a/Engine/Source/Runtime/Core/Private/Internationalization/TextLocalizationManager.cpp
+++ b/Engine/Source/Runtime/Core/Private/Internationalization/TextLocalizationManager.cpp
@@ -187,6 +187,13 @@ FString KeyifyTextId(const FTextId& TextId)
 	// If the key/namespace pair is too long, the Slate.LogPaintedText cvar can be used to see the entire thing.
 	return FString::Printf(TEXT("%s, %s"), *TextId.GetKey().ToString(), *TextId.GetNamespace().ToString());
 }
+
+FString ConditionDisplayStringForLog(const FString& DisplayString)
+{
+	static const TArray<TCHAR> CharsToReplace = { TEXT('\n'), TEXT('\r'), TEXT('\t'), TEXT('"'), TEXT('\\') };
+	return DisplayString.ReplaceCharWithEscapedChar(&CharsToReplace);
+}
+
 }
 
 enum class ERequestedCultureOverrideLevel : uint8
@@ -793,7 +800,7 @@ void FTextLocalizationManager::DumpLiveTable(const FString* NamespaceFilter, con
 
 	DumpLiveTableImpl(NamespaceFilter, KeyFilter, DisplayStringFilter, [&Category](const FTextId& Id, const FTextConstDisplayStringRef& DisplayString)
 	{
-		UE_LOG_REF(Category, Display, TEXT("LiveTableEntry: Namespace: '%s', Key: '%s', DisplayString: '%s'"), *Id.GetNamespace().ToString(), *Id.GetKey().ToString(), **DisplayString); //-V510
+		UE_LOG_REF(Category, Display, TEXT("LiveTableEntry: Namespace: '%s', Key: '%s', DisplayString: '%s'"), *Id.GetNamespace().ToString(), *Id.GetKey().ToString(), *TextLocalizationManager::ConditionDisplayStringForLog(*DisplayString)); //-V510
 	});
 
 	UE_LOG_REF(Category, Display, TEXT("----------------------------------------------------------------------"));
@@ -806,7 +813,7 @@ void FTextLocalizationManager::DumpLiveTable(const FString& OutputFilename, cons
 
 	DumpLiveTableImpl(NamespaceFilter, KeyFilter, DisplayStringFilter, [&DumpString](const FTextId& Id, const FTextConstDisplayStringRef& DisplayString)
 	{
-		DumpString += FString::Printf(TEXT("LiveTableEntry: Namespace: '%s', Key: '%s', DisplayString: '%s'"), *Id.GetNamespace().ToString(), *Id.GetKey().ToString(), **DisplayString);
+		DumpString += FString::Printf(TEXT("LiveTableEntry: Namespace: '%s', Key: '%s', DisplayString: '%s'"), *Id.GetNamespace().ToString(), *Id.GetKey().ToString(), *TextLocalizationManager::ConditionDisplayStringForLog(*DisplayString));
 		DumpString += LINE_TERMINATOR;
 	});
 
@@ -822,10 +829,14 @@ void FTextLocalizationManager::AddOrUpdateDisplayStringInLiveTable(const FString
 
 	if (FDisplayStringEntry* LiveEntry = DisplayStringLookupTable.Find(TextId))
 	{
+		if (SourceStringPtr)
+		{
+			LiveEntry->SourceStringHash = FTextLocalizationResource::HashString(*SourceStringPtr);
+		}
 		LiveEntry->DisplayString = MakeTextDisplayString(CopyTemp(DisplayString));
 		DirtyLocalRevisionForTextId(TextId);
 
-		UE_LOG(LogConsoleResponse, Display, TEXT("Updated string for Namespace='%s', Key='%s' to DisplayString='%s'"), *Namespace, *Key, *DisplayString);
+		UE_LOG(LogConsoleResponse, Display, TEXT("Updated string for Namespace='%s', Key='%s' to DisplayString='%s'"), *Namespace, *Key, *TextLocalizationManager::ConditionDisplayStringForLog(DisplayString));
 	}
 	else if (SourceStringPtr)
 	{
@@ -841,7 +852,7 @@ void FTextLocalizationManager::AddOrUpdateDisplayStringInLiveTable(const FString
 			DisplayStringLookupTable.Add(TextId, MoveTemp(NewLiveEntry));
 		}
 
-		UE_LOG(LogConsoleResponse, Display, TEXT("Added string for Namespace='%s', Key='%s' to DisplayString='%s'"), *Namespace, *Key, *DisplayString);
+		UE_LOG(LogConsoleResponse, Display, TEXT("Added string for Namespace='%s', Key='%s' to DisplayString='%s'"), *Namespace, *Key, *TextLocalizationManager::ConditionDisplayStringForLog(DisplayString));
 	}
 	else
 	{
