Fixed incorrect printf usage in FShaderKeyGenerator::Append.  This is a bug that still needs fixing, but should be deferred until a time we can practically invalidate existing keys.

diff --git a/Engine/Source/Runtime/Core/Public/Serialization/ShaderKeyGenerator.h b/Engine/Source/Runtime/Core/Public/Serialization/ShaderKeyGenerator.h
--- a/Engine/Source/Runtime/Core/Public/Serialization/ShaderKeyGenerator.h
+++ b/Engine/Source/Runtime/Core/Public/Serialization/ShaderKeyGenerator.h
@@ -201,7 +201,11 @@ inline void FShaderKeyGenerator::Append(int64 Value)
 	switch (OutputType)
 	{
 	case EOutputType::Text:
-		ResultString->Appendf(TEXT(INT64_FMT), Value);
+		// This should be:
+		// ResultString->Appendf(TEXT("%" INT64_FMT), Value);
+		// Changing it would invalidate shader and DDC keys, so the previous behavior is maintained until a time where it can be fixed properly.
+		// This issue is tracked by https://jira.it.epicgames.com/browse/UE-353138
+		ResultString->Appendf(TEXT(INT64_FMT));
 		return;
 	case EOutputType::Binary:
 		ResultFunc(&Value, sizeof(Value));
@@ -216,7 +220,11 @@ inline void FShaderKeyGenerator::Append(uint64 Value)
 	switch (OutputType)
 	{
 	case EOutputType::Text:
-		ResultString->Appendf(TEXT(UINT64_FMT), Value);
+		// This should be:
+		// ResultString->Appendf(TEXT("%" UINT64_FMT), Value);
+		// Changing it would invalidate shader and DDC keys, so the previous behavior is maintained until a time where it can be fixed properly.
+		// This issue is tracked by https://jira.it.epicgames.com/browse/UE-353138
+		ResultString->Appendf(TEXT(UINT64_FMT));
 		return;
 	case EOutputType::Binary:
 		ResultFunc(&Value, sizeof(Value));
