Logging: Updated UE_LOG_CONTEXT to be safe for use in transactions

diff --git a/Engine/Source/Programs/AutoRTFMTests/Private/UECoreTests.cpp b/Engine/Source/Programs/AutoRTFMTests/Private/UECoreTests.cpp
--- a/Engine/Source/Programs/AutoRTFMTests/Private/UECoreTests.cpp
+++ b/Engine/Source/Programs/AutoRTFMTests/Private/UECoreTests.cpp
@@ -1650,14 +1650,14 @@ TEST_CASE("UECore.UE_LOGFMT")
 {
 	SECTION("Commit")
 	{
-		AutoRTFM::Testing::Commit([&]
+		AutoRTFM::Testing::Commit([]
 		{
 			UE_LOGFMT(LogAutoRTFM_UECoreTests, Log, "{Animal} says {Sound}", TEXT("Cat"), TEXT("meow!"));
 		});
 	}
 	SECTION("Abort")
 	{
-		AutoRTFM::Testing::Abort([&]
+		AutoRTFM::Testing::Abort([]
 		{
 			UE_LOGFMT(LogAutoRTFM_UECoreTests, Log, "{Animal} says {Sound}", TEXT("Cat"), TEXT("meow!"));
 			AutoRTFM::AbortTransaction();
@@ -1665,6 +1665,25 @@ TEST_CASE("UECore.UE_LOGFMT")
 	}
 }
 
+TEST_CASE("UECore.UE_LOG_CONTEXT")
+{
+	SECTION("Commit")
+	{
+		AutoRTFM::Testing::Commit([]
+		{
+			UE_LOG_CONTEXT("Test", "Commit");
+		});
+	}
+	SECTION("Abort")
+	{
+		AutoRTFM::Testing::Abort([]
+		{
+			UE_LOG_CONTEXT("Test", "Abort");
+			AutoRTFM::AbortTransaction();
+		});
+	}
+}
+
 TEST_CASE("UECore.FOutputDeviceRedirector")
 {
 	FOutputDeviceRedirector Redirector;
diff --git a/Engine/Source/Runtime/Core/Private/Logging/StructuredLog.cpp b/Engine/Source/Runtime/Core/Private/Logging/StructuredLog.cpp
--- a/Engine/Source/Runtime/Core/Private/Logging/StructuredLog.cpp
+++ b/Engine/Source/Runtime/Core/Private/Logging/StructuredLog.cpp
@@ -1110,30 +1110,33 @@ FLogCategory<ELogVerbosity::Fatal, ELogVerbosity::Fatal> LogFatal(TEXT("Fatal"))
 thread_local FLogContext* LogContextHead = nullptr;
 thread_local FLogContext* LogContextTail = nullptr;
 
-FLogContext::FLogContext(const FLogField& InField)
-: Prev(LogContextTail)
+FLogContext::FLogContext(const ANSICHAR* InName)
+	: FLogContext(InName, nullptr, nullptr)
+{
+}
+
+FLogContext::FLogContext(const ANSICHAR* InName, const void* InValue, FLogField::FWriteFn* InWriteValue)
+: Name(InName)
+, Prev(LogContextTail)
 {
 	(Prev ? Prev->Next : LogContextHead) = this;
 	LogContextTail = this;
 
-	const FUtf8StringView Name = InField.Name;
-
 	TCbWriter<256> Writer;
-	Writer.SetName(Name);
-	if (InField.WriteValue)
+	if (InWriteValue)
 	{
-		InField.WriteValue(Writer, InField.Value);
+		InWriteValue(Writer, InValue);
 	}
 	else
 	{
 		Writer.AddNull();
 	}
-	Field = Writer.Save();
-	Field.MakeOwned();
+	Value.SetNumUninitialized((int32)Writer.GetSaveSize());
+	Writer.Save(MakeMemoryView(Value));
 
 	for (FLogContext* Node = Prev; Node; Node = Node->Prev)
 	{
-		if (Node->Field.GetName().Equals(Name))
+		if (FCStringAnsi::Strcmp(Name, Node->Name) == 0)
 		{
 			Node->bDisabledByNewerContext = true;
 			bDisabledOlderContext = true;
@@ -1148,10 +1151,9 @@ FLogContext::~FLogContext()
 
 	if (bDisabledOlderContext)
 	{
-		const FUtf8StringView Name = Field.GetName();
 		for (FLogContext* Node = Prev; Node; Node = Node->Prev)
 		{
-			if (Node->Field.GetName().Equals(Name))
+			if (FCStringAnsi::Strcmp(Name, Node->Name) == 0)
 			{
 				Node->bDisabledByNewerContext = false;
 			}
@@ -1375,7 +1377,7 @@ FORCENOINLINE static FCbObject SerializeLogFields(
 		{
 			if (!Node->bDisabledByNewerContext)
 			{
-				const FUtf8StringView NodeName = Node->Field.GetName();
+				const FAnsiStringView NodeName = Node->Name;
 				ActiveContext.Add(!FieldNames.FindByPredicate([NodeName](FAnsiStringView Name) { return NodeName.Equals(Name); }));
 				++ContextIndex;
 			}
@@ -1386,7 +1388,7 @@ FORCENOINLINE static FCbObject SerializeLogFields(
 		{
 			if (!Node->bDisabledByNewerContext && ActiveContext[--ContextIndex])
 			{
-				Writer.AddField(Node->Field.GetName(), Node->Field);
+				Writer.AddField(Node->Name, FCbFieldView(Node->Value.GetData()));
 			}
 		}
 
@@ -1396,7 +1398,7 @@ FORCENOINLINE static FCbObject SerializeLogFields(
 		{
 			if (!Node->bDisabledByNewerContext)
 			{
-				Writer.AddString(Node->Field.GetName());
+				Writer.AddString(Node->Name);
 			}
 		}
 		Writer.EndArray();
@@ -1606,7 +1608,7 @@ FORCENOINLINE static FCbObject SerializeBasicLogMessage(const FStaticBasicLogRec
 	TCbWriter<512> Writer;
 	Writer.BeginObject();
 
-	const FUtf8StringView MessageName = UTF8TEXTVIEW("Message");
+	const FAnsiStringView MessageName = ANSITEXTVIEW("Message");
 	Writer.AddString(MessageName, Message);
 
 	if (LogContextHead)
@@ -1614,10 +1616,10 @@ FORCENOINLINE static FCbObject SerializeBasicLogMessage(const FStaticBasicLogRec
 		// Traverse contexts forward and copy any that Message did not override.
 		for (FLogContext* Node = LogContextHead; Node; Node = Node->Next)
 		{
-			const FUtf8StringView Name = Node->Field.GetName();
+			const FAnsiStringView Name = Node->Name;
 			if (!Node->bDisabledByNewerContext && !Name.Equals(MessageName))
 			{
-				Writer.AddField(Name, Node->Field);
+				Writer.AddField(Name, FCbFieldView(Node->Value.GetData()));
 			}
 		}
 
@@ -1627,7 +1629,7 @@ FORCENOINLINE static FCbObject SerializeBasicLogMessage(const FStaticBasicLogRec
 		{
 			if (!Node->bDisabledByNewerContext)
 			{
-				Writer.AddString(Node->Field.GetName());
+				Writer.AddString(Node->Name);
 			}
 		}
 		Writer.EndArray();
@@ -1749,7 +1751,9 @@ void VisitLogContext(TFunctionRef<void (const FCbField&)> Visitor)
 	{
 		if (!Node->bDisabledByNewerContext)
 		{
-			Visitor(Node->Field);
+			TCbWriter<256> Writer;
+			Writer.AddField(Node->Name, FCbFieldView(Node->Value.GetData()));
+			Visitor(Writer.Save());
 		}
 	}
 }
diff --git a/Engine/Source/Runtime/Core/Public/Logging/StructuredLog.h b/Engine/Source/Runtime/Core/Public/Logging/StructuredLog.h
--- a/Engine/Source/Runtime/Core/Public/Logging/StructuredLog.h
+++ b/Engine/Source/Runtime/Core/Public/Logging/StructuredLog.h
@@ -140,7 +140,7 @@ template <typename FuncType> class TFunctionRef;
  * UE_LOG_CONTEXT("Count", 123.0);
  * UE_LOG_CONTEXT("Asset", FAssetLog(AssetPath));
  */
-#define UE_LOG_CONTEXT(Name, ...) ::UE::Logging::Private::FLogContext ANONYMOUS_VARIABLE(LogContext_)(Name, ##__VA_ARGS__)
+#define UE_LOG_CONTEXT(NameLiteral, ...) ::UE::Logging::Private::FLogContext ANONYMOUS_VARIABLE(LogContext_)(::UE::Logging::Private::CheckFieldName(NameLiteral).Name, ##__VA_ARGS__)
 
 #else // #if !NO_LOGGING
 
@@ -508,24 +508,21 @@ UE_API extern FLogCategory<ELogVerbosity::Fatal, ELogVerbosity::Fatal> LogFatal;
 
 struct FLogContext
 {
-	inline explicit FLogContext(const ANSICHAR* Name)
-		: FLogContext({Name, nullptr, nullptr})
-	{
-	}
-
 	template <typename ValueType>
-	inline explicit FLogContext(const ANSICHAR* Name, const ValueType& Value)
-		: FLogContext({Name, &Value, FLogField::Write<ValueType>})
+	inline explicit FLogContext(const ANSICHAR* InName, const ValueType& InValue)
+		: FLogContext(InName, &InValue, FLogField::Write<ValueType>)
 	{
 	}
 
-	UE_API explicit FLogContext(const FLogField& Field);
+	UE_API explicit FLogContext(const ANSICHAR* Name);
+	UE_API explicit FLogContext(const ANSICHAR* Name, const void* Value, FLogField::FWriteFn* WriteValue);
 	UE_API ~FLogContext();
 
 	FLogContext(const FLogContext&) = delete;
 	FLogContext& operator=(const FLogContext&) = delete;
 
-	FCbField Field;
+	const ANSICHAR* Name = nullptr;
+	TArray<uint8> Value;
 	FLogContext* Next = nullptr;
 	FLogContext* Prev = nullptr;
 	bool bDisabledByNewerContext = false;
