[StripedMap]
- Add ParallelForEach and ParallelRemoveIf where we parallelize each stripe

diff --git a/Engine/Source/Runtime/Core/Public/Containers/StripedMap.h b/Engine/Source/Runtime/Core/Public/Containers/StripedMap.h
--- a/Engine/Source/Runtime/Core/Public/Containers/StripedMap.h
+++ b/Engine/Source/Runtime/Core/Public/Containers/StripedMap.h
@@ -222,6 +222,20 @@ protected:
 		}
 	}
 
+	template <typename FunctionType>
+	void ParallelForEachMap(FunctionType&& InFunction)
+	{
+		ParallelFor(
+			BucketCount,
+			[this, &InFunction](int32 BucketIndex)
+			{
+				FBucket& Bucket = Buckets[BucketIndex];
+				ExclusiveLockType ScopeLock(Bucket.Lock);
+				InFunction(Bucket.Map);
+			}
+		);
+	}
+
 	template <typename FunctionType>
 	void ForEachMap(FunctionType&& InFunction) const
 	{
@@ -232,6 +246,20 @@ protected:
 		}
 	}
 
+	template <typename FunctionType>
+	void ParallelForEachMap(FunctionType&& InFunction) const
+	{
+		ParallelFor(
+			BucketCount,
+			[this, &InFunction](int32 BucketIndex)
+			{
+				const FBucket& Bucket = Buckets[BucketIndex];
+				SharedLockType ScopeLock(Bucket.Lock);
+				InFunction(Bucket.Map);
+			}
+		);
+	}
+
 public:
 	/**
 	 * Gets you a copy of the value.
@@ -611,6 +639,34 @@ public:
 		return RemovedCount;
 	}
 
+	/**
+	 * Removes all elements where the predicate returns true.
+	 *
+	 * @param InPredicate  A thread-safe predicate to call to determine if a pair should be removed.
+	 * @return The number of pairs that were removed.
+	 */
+	template <typename PredicateType>
+	int32 ParallelRemoveIf(PredicateType&& InPredicate)
+	{
+		std::atomic<int32> RemovedCount = 0;
+		ParallelForEachMap(
+			[&RemovedCount, &InPredicate](MapType& Map)
+			{
+				int32 LocalRemoveCount = 0;
+				for (auto It = Map.CreateIterator(); It; ++It)
+				{
+					if (InPredicate(*It))
+					{
+						It.RemoveCurrent();
+						LocalRemoveCount++;
+					}
+				}
+				RemovedCount += LocalRemoveCount;
+			}
+		);
+		return RemovedCount.load();
+	}
+
 	/**
 	 * Remove the pair with the specified key and copies the value
 	 * that was removed to the ref parameter
@@ -758,6 +814,25 @@ public:
 		);
 	}
 
+	/**
+	 * Calls a function on all elements of the map with exclusive access (elements can be modified).
+	 *
+	 * @param InFunction A thread-safe callback to run on each Pair of the map.
+	 */
+	template <typename FunctionType>
+	void ParallelForEach(FunctionType&& InFunction)
+	{
+		ParallelForEachMap(
+			[&InFunction](MapType& Map)
+			{
+				for (auto& Item : Map)
+				{
+					InFunction(Item);
+				}
+			}
+		);
+	}
+
 	/**
 	 * Calls a function on all elements of the map with shared access (elements can only be read).
 	 *
@@ -777,6 +852,25 @@ public:
 		);
 	}
 
+	/**
+	 * Calls a function on all elements of the map with shared access (elements can only be read).
+	 *
+	 * @param InFunction A thread-safe callback to run on each Pair of the map.
+	 */
+	template <typename FunctionType>
+	void ParallelForEach(FunctionType&& InFunction) const
+	{
+		ParallelForEachMap(
+			[&InFunction](const MapType& Map)
+			{
+				for (const auto& Item : Map)
+				{
+					InFunction(Item);
+				}
+			}
+		);
+	}
+
 	/**
 	 * Get the unique keys contained within this map.
 	 *
