Updated Windows telemetry for host architecture and form factor to be in missing events (CrashReportClient.ReportCrash)

#tests In progress...

diff --git a/Engine/Source/Runtime/Core/Private/Windows/WindowsPlatformCrashContext.cpp b/Engine/Source/Runtime/Core/Private/Windows/WindowsPlatformCrashContext.cpp
--- a/Engine/Source/Runtime/Core/Private/Windows/WindowsPlatformCrashContext.cpp
+++ b/Engine/Source/Runtime/Core/Private/Windows/WindowsPlatformCrashContext.cpp
@@ -327,6 +327,9 @@ void FWindowsPlatformCrashContext::AddPlatformSpecificProperties() const
 	AddCrashProperty(TEXT("PlatformSupportsWindows11"), FWindowsPlatformMisc::SupportsWindows11());
 	AddCrashProperty(TEXT("IsRunningOnBattery"), FPlatformMisc::IsRunningOnBattery());
 	AddCrashProperty(TEXT("IsUsingBatterySaver"), FWindowsPlatformMisc::IsUsingBatterySaver());
+	AddCrashProperty(TEXT("HostArchitecture"), FWindowsPlatformMisc::GetHostArchitecture());
+	AddCrashProperty(TEXT("IsDeviceHandheld"), FPlatformMisc::IsDeviceGamingHandheld());
+
 	WIDECHAR DriveName = 0;
 	const TCHAR* BaseDir = FWindowsPlatformProcess::BaseDir();
 	if (BaseDir && *BaseDir)
