[Core] Update check message to add info about what a user could look into to resolve the situation

#tests Compiled
#rnx

diff --git a/Engine/Source/Runtime/Core/Public/Containers/CompactHashTable.h b/Engine/Source/Runtime/Core/Public/Containers/CompactHashTable.h
--- a/Engine/Source/Runtime/Core/Public/Containers/CompactHashTable.h
+++ b/Engine/Source/Runtime/Core/Public/Containers/CompactHashTable.h
@@ -46,7 +46,7 @@ namespace UE::Core::CompactHashTable::Private
 			}
 		}
 
-		checkNoEntry(); // At this point we must be pointing to something valid in the list so this lookup should never fail
+		checkf(false, TEXT("Failed to find object to remove from set.")); // This usually indicates that something changed in the entry that affected the hash so it no longer matches up.
 		return (IndexType)INDEX_NONE;
 	}
 }
