CsvProfiler: Don't fire a delegate for region events
#tests tested on console

diff --git a/Engine/Source/Runtime/Core/Private/ProfilingDebugging/CsvProfiler.cpp b/Engine/Source/Runtime/Core/Private/ProfilingDebugging/CsvProfiler.cpp
--- a/Engine/Source/Runtime/Core/Private/ProfilingDebugging/CsvProfiler.cpp
+++ b/Engine/Source/Runtime/Core/Private/ProfilingDebugging/CsvProfiler.cpp
@@ -4558,7 +4558,7 @@ void FCsvProfiler::RecordEventfInternal(int32 CategoryIndex, bool bLog, const TC
 	}
 }
 
-void FCsvProfiler::RecordEvent(int32 CategoryIndex, const FString& EventText, bool bLog)
+void FCsvProfiler::RecordEvent(int32 CategoryIndex, const FString& EventText, bool bLog, bool bAllowDelegate)
 {
 	if (GCsvProfilerIsCapturing && GCsvCategoriesEnabled[CategoryIndex])
 	{
@@ -4582,7 +4582,10 @@ void FCsvProfiler::RecordEvent(int32 CategoryIndex, const FString& EventText, bo
 			UE_LOG(LogCsvProfiler, Display, TEXT("CSVEvent \"%s\" [Frame %d]"), *FullEventText, FCsvProfiler::Get()->GetCaptureFrameNumber());
 		}
 #if CSV_PROFILER_ALLOW_DEBUG_FEATURES
-		FCsvProfiler::Get()->OnCSVProfileEvent().Broadcast(CategoryName, EventText);
+		if (bAllowDelegate)
+		{
+			FCsvProfiler::Get()->OnCSVProfileEvent().Broadcast(CategoryName, EventText);
+		}
 		if (GCsvEventExecCmds)
 		{
 			FStringView Match(FullEventText);
@@ -4635,7 +4638,7 @@ void FCsvProfiler::RecordRegionEventInternal(const TCHAR* CategoryName, const TC
 			FString::Printf(TEXT("Region: %s/%s - %s"), SubCategoryName, RegionName, Suffix) : 
 			FString::Printf(TEXT("Region: %s - %s"), RegionName, Suffix);
 	}
-	RecordEvent(CSV_CATEGORY_INDEX_GLOBAL, EventText, false);
+	RecordEvent(CSV_CATEGORY_INDEX_GLOBAL, EventText, false, false);
 }
 
 
diff --git a/Engine/Source/Runtime/Core/Public/ProfilingDebugging/CsvProfiler.h b/Engine/Source/Runtime/Core/Public/ProfilingDebugging/CsvProfiler.h
--- a/Engine/Source/Runtime/Core/Public/ProfilingDebugging/CsvProfiler.h
+++ b/Engine/Source/Runtime/Core/Public/ProfilingDebugging/CsvProfiler.h
@@ -437,7 +437,7 @@ public:
 	CORE_API static void RecordCustomStatMinimal(const char* StatName, uint32 CategoryIndex, int32 Value, const ECsvCustomStatOp CustomStatOp);
 	CORE_API static void RecordCustomStatMinimal(const char* StatName, uint32 CategoryIndex, double Value, const ECsvCustomStatOp CustomStatOp);
 
-	CORE_API static void RecordEvent(int32 CategoryIndex, const FString& EventText, bool bLog=true);
+	CORE_API static void RecordEvent(int32 CategoryIndex, const FString& EventText, bool bLog=true, bool bAllowDelegate=true);
 	CORE_API static void RecordEventAtFrameStart(int32 CategoryIndex, const FString& EventText, bool bLog = true);
 	CORE_API static void RecordEventAtTimestamp(int32 CategoryIndex, const FString& EventText, uint64 Cycles64, bool bLog = true);
 
