Add support for CSV regions from the boot profiler. This will work even if BOOT_PROFILING is disabled, provided the Csv profiler is capturing on startup. This is only active in non-shipping builds.

diff --git a/Engine/Source/Runtime/Core/Private/Misc/CoreGlobals.cpp b/Engine/Source/Runtime/Core/Private/Misc/CoreGlobals.cpp
--- a/Engine/Source/Runtime/Core/Private/Misc/CoreGlobals.cpp
+++ b/Engine/Source/Runtime/Core/Private/Misc/CoreGlobals.cpp
@@ -14,6 +14,7 @@
 #include "Misc/LazySingleton.h"
 #include "Misc/PlayInEditorLoadingScope.h"
 #include "Misc/CommandLine.h"
+#include "ProfilingDebugging/CsvProfiler.h"
 #include "ProfilingDebugging/MiscTrace.h"
 #include "GenericPlatform/GenericPlatformCrashContext.h"
 #include "CoreGlobalsInternal.h"
@@ -32,6 +33,7 @@
 
 #define LOCTEXT_NAMESPACE "Core"
 
+CSV_DEFINE_CATEGORY(EngineInit, true);
 
 class FCoreModule : public FDefaultModuleImpl
 {
@@ -524,21 +526,51 @@ FEngineTrackedActivityScope::~FEngineTrackedActivityScope()
 	FTrackedActivity::GetEngineActivity().Pop();
 }
 
+static bool GBootTimingCompleted = false;
+
 #define USE_BOOT_PROFILING 0
+#define CSV_BOOT_PROFILING (CSV_PROFILER_STATS && !UE_BUILD_SHIPPING)
 
 #if !USE_BOOT_PROFILING
 FScopedBootTiming::FScopedBootTiming(const ANSICHAR *InMessage)
 {
 	UE_LOGFMT(LogInit, Verbose, "{Message}", InMessage);
+#if CSV_BOOT_PROFILING
+	if (!GBootTimingCompleted && IsInGameThread())
+	{
+		Message = InMessage;
+		if (Message.Len())
+		{
+			CSV_REGION_BEGIN(EngineInit, *Message);
+		}
+	}
+#endif
 }
 
 FScopedBootTiming::FScopedBootTiming(const ANSICHAR *InMessage, FName Suffix)
 {
 	UE_LOGFMT(LogInit, Verbose, "{Message}: {Suffix}", InMessage, Suffix);
+#if CSV_BOOT_PROFILING
+	if (!GBootTimingCompleted && IsInGameThread())
+	{
+		Message = FString(InMessage) + Suffix.ToString();
+		if (Message.Len())
+		{
+			CSV_REGION_BEGIN(EngineInit, *Message);
+		}
+	}
+#endif
 }
 
 FScopedBootTiming::~FScopedBootTiming()
 {
+#if CSV_BOOT_PROFILING
+	if (!Message.IsEmpty())
+	{
+		CSV_REGION_END(EngineInit, *Message);
+	}
+#endif
+
 }
 void BootTimingPoint(const ANSICHAR *Message)
 {
@@ -546,10 +578,10 @@ void BootTimingPoint(const ANSICHAR *Message)
 }
 void DumpBootTiming()
 {
+	GBootTimingCompleted = true;
 }
 #else
 static TArray<FString> GAllBootTiming;
-static bool GBootTimingCompleted = false;
 static int32 GBootScopeDepth = 0;
 
 static void DumpBootTimingString(const TCHAR* Message)
@@ -688,6 +720,9 @@ FScopedBootTiming::FScopedBootTiming(const ANSICHAR *InMessage)
 	Message = InMessage;
 	StartTime = FPlatformTime::Seconds();
 	BootTimingPoint(*Message, TEXT("{"), GBootScopeDepth);
+#if CSV_BOOT_PROFILING
+	CSV_REGION_BEGIN(EngineInit, *Message);
+#endif
 	GBootScopeDepth++;
 }
 
@@ -700,6 +735,9 @@ FScopedBootTiming::FScopedBootTiming(const ANSICHAR *InMessage, FName Suffix)
 	Message = FString(InMessage) + Suffix.ToString();
 	StartTime = FPlatformTime::Seconds();
 	BootTimingPoint(*Message, TEXT("{"), GBootScopeDepth);
+#if CSV_BOOT_PROFILING
+	CSV_REGION_BEGIN(EngineInit, *Message);
+#endif
 	GBootScopeDepth++;
 }
 
@@ -709,6 +747,9 @@ FScopedBootTiming::~FScopedBootTiming()
 	{
 		GBootScopeDepth--;
 		BootTimingPoint(*Message, TEXT("}"), GBootScopeDepth, FPlatformTime::Seconds() - StartTime);
+#if CSV_BOOT_PROFILING
+		CSV_REGION_END(EngineInit, *Message);
+#endif
 	}
 }
 
