Enable mounting of container .umeta file(s) in non-shipping builds.

#rnx

diff --git a/Engine/Source/Runtime/Core/Internal/IO/IoDispatcherInternal.h b/Engine/Source/Runtime/Core/Internal/IO/IoDispatcherInternal.h
--- a/Engine/Source/Runtime/Core/Internal/IO/IoDispatcherInternal.h
+++ b/Engine/Source/Runtime/Core/Internal/IO/IoDispatcherInternal.h
@@ -3,6 +3,7 @@
 #pragma once
 #include "IO/IoDispatcher.h"
 #include "IO/IoChunkId.h"
+#include "Containers/StringFwd.h"
 
 // This header is used for IoDispatcher functionality we don't want to expose outside of the engine
 
@@ -17,4 +18,11 @@ public:
 		static const bool HasScriptObjectsChunk = Dispatcher.DoesChunkExist(CreateIoChunkId(0, 0, EIoChunkType::ScriptObjects));
 		return HasScriptObjectsChunk;
 	}
-};
\ No newline at end of file
+	
+	static void MountContainerMeta(const FString& FilePath)
+	{
+		FIoDispatcher::Get().MountContainerMeta(FilePath);
+	}
+
+	FUtf8StringView GetFilename(const FIoChunkId& ChunkId, FUtf8StringBuilderBase& OutFilename, FUtf8StringView& OutContainerName);
+};
diff --git a/Engine/Source/Runtime/Core/Private/IO/IoDispatcher.cpp b/Engine/Source/Runtime/Core/Private/IO/IoDispatcher.cpp
--- a/Engine/Source/Runtime/Core/Private/IO/IoDispatcher.cpp
+++ b/Engine/Source/Runtime/Core/Private/IO/IoDispatcher.cpp
@@ -24,6 +24,16 @@
 #include "ProfilingDebugging/TraceAuxiliary.h"
 #include "Templates/Greater.h"
 
+#define UE_WITH_CONTAINER_META !UE_BUILD_SHIPPING
+
+#if UE_WITH_CONTAINER_META
+#include "HAL/PlatformFileManager.h"
+#include "IO/IoContainerMeta.h"
+#include "Misc/Paths.h"
+#else
+#include "Containers/Utf8String.h"
+#endif // UE_WITH_CONTAINER_META
+
 DEFINE_LOG_CATEGORY(LogIoDispatcher);
 
 const FIoChunkId FIoChunkId::InvalidChunkId = FIoChunkId::CreateEmptyId();
@@ -381,6 +391,80 @@ private:
 	TBlockAllocator<FIoRequestImpl, 4096> BlockAllocator;
 };
 
+#if UE_WITH_CONTAINER_META
+class FIoMountedContainerMeta
+{
+	struct FContainerMeta
+	{
+		FContainerMeta(FIoContainerMetaReader&& InReader, FString&& InFilePath)
+			: Reader(MoveTemp(InReader))
+			, FilePath(MoveTemp(InFilePath))
+		{ }
+
+		FIoContainerMetaReader	Reader;
+		FString					FilePath;
+	};
+
+public:
+	void			Mount(const FString& FilePath);
+	void			Unmount(const FString& FilePath);
+	FUtf8StringView	GetFilename(const FIoChunkId& ChunkId, FUtf8StringBuilderBase& OutFilename, FUtf8StringView& OutContainerName) const;
+
+private:
+	mutable UE::FMutex					Mutex;
+	TArray<TUniquePtr<FContainerMeta>>	Meta;
+};
+
+void FIoMountedContainerMeta::Mount(const FString& FilePath)
+{
+	LLM_SCOPE_BYNAME(TEXT("FileSystem/Debug"));
+
+	IPlatformFile& Ipf = FPlatformFileManager::Get().GetPlatformFile();
+
+	FString MetaPath = FPaths::ChangeExtension(FilePath, FIoContainerMetaHeader::FileExtension);
+	if (Ipf.FileExists(*MetaPath) == false)
+	{
+		return;
+	}
+
+	TIoStatusOr<FIoContainerMetaReader> MaybeReader = FIoContainerMetaReader::Load(*MetaPath);
+	if (MaybeReader.IsOk() == false)
+	{
+		UE_LOG(LogIoDispatcher, Warning, TEXT("Failed to load container metadata '%s', reason: %s"), *MetaPath, *MaybeReader.Status().ToString());
+		return;
+	}
+
+	UE_LOG(LogIoDispatcher, Log, TEXT("Mounted container metadata '%s'"), *MetaPath);
+	{
+		UE::TUniqueLock Lock(Mutex);
+		Meta.Add(MakeUnique<FContainerMeta>(MaybeReader.ConsumeValueOrDie(), MoveTemp(MetaPath)));
+	}
+}
+
+FUtf8StringView FIoMountedContainerMeta::GetFilename(const FIoChunkId& ChunkId, FUtf8StringBuilderBase& OutFilename, FUtf8StringView& OutContainerName) const
+{
+	FUtf8StringView Result;
+	UE::TUniqueLock Lock(Mutex);
+	for (const TUniquePtr<FContainerMeta>& ContainerMeta : Meta)
+	{
+		if (Result = ContainerMeta->Reader.GetFilename(ChunkId, OutFilename, OutContainerName); !Result.IsEmpty())
+		{
+			return Result;
+		}
+	}
+
+	return Result;
+}
+#else
+class FIoMountedContainerMeta
+{
+public:
+	void			Mount(const FString& FilePath) { }
+	void			Unmount(const FString& FilePath) { }
+	FUtf8StringView	GetFilename(const FIoChunkId& ChunkId, FUtf8StringBuilderBase& OutFilename, FUtf8StringView& OutContainerName) const { return FUtf8StringView(); }
+};
+#endif // UE_WITH_CONTAINER_META
+
 class FIoDispatcherImpl
 	: public FRunnable
 {
@@ -626,6 +710,16 @@ public:
 		}	
 	}
 
+	void MountContainerMeta(const FString& FilePath)
+	{
+		ContainerMeta.Mount(FilePath);
+	}
+
+	FUtf8StringView	GetFilename(const FIoChunkId& ChunkId, FUtf8StringBuilderBase& OutFilename, FUtf8StringView& OutContainerName) const
+	{
+		return ContainerMeta.GetFilename(ChunkId, OutFilename, OutContainerName);
+	}
+
 	FIoSignatureErrorDelegate& OnSignatureError()
 	{
 		return BackendContext->SignatureErrorDelegate;
@@ -988,6 +1082,7 @@ private:
 	using FBatchAllocator = TBlockAllocator<FIoBatchImpl, 4096>;
 	using FBackendAndPriority = TTuple<int32, TSharedRef<IIoDispatcherBackend>>;
 
+	FIoMountedContainerMeta ContainerMeta;
 	TSharedRef<FIoDispatcherBackendContext> BackendContext;
 	FDelegateHandle MemoryTrimDelegateHandle;
 	mutable FRWLock BackendsLock;
@@ -1065,6 +1160,16 @@ FIoDispatcher::GetSizeForChunk(const FIoChunkId& ChunkId, const FIoOffsetAndLeng
 	return Impl->GetSizeForChunk(ChunkId, ChunkRange, OutAvailable);
 }
 
+void FIoDispatcher::MountContainerMeta(const FString& FilePath)
+{
+	Impl->MountContainerMeta(FilePath);
+}
+
+FUtf8StringView	FIoDispatcher::GetFilename(const FIoChunkId& ChunkId, FUtf8StringBuilderBase& OutFilename, FUtf8StringView& OutContainerName) const
+{
+	return Impl->GetFilename(ChunkId, OutFilename, OutContainerName);
+}
+
 int64
 FIoDispatcher::GetTotalLoaded() const
 {
diff --git a/Engine/Source/Runtime/Core/Public/IO/IoDispatcher.h b/Engine/Source/Runtime/Core/Public/IO/IoDispatcher.h
--- a/Engine/Source/Runtime/Core/Public/IO/IoDispatcher.h
+++ b/Engine/Source/Runtime/Core/Public/IO/IoDispatcher.h
@@ -382,6 +382,8 @@ public:
 private:
 	CORE_API bool					DoesChunkExist(const FIoChunkId& ChunkId, const FIoOffsetAndLength& ChunkRange) const;
 	CORE_API TIoStatusOr<uint64>	GetSizeForChunk(const FIoChunkId& ChunkId, const FIoOffsetAndLength& ChunkRange, uint64& OutAvailable) const;
+	CORE_API void					MountContainerMeta(const FString& FilePath);
+	CORE_API FUtf8StringView		GetFilename(const FIoChunkId& ChunkId, FUtf8StringBuilderBase& OutFilename, FUtf8StringView& OutContainerName) const;
 
 	FIoDispatcherImpl* Impl = nullptr;
 
@@ -389,6 +391,7 @@ private:
 	friend class FIoBatch;
 	friend class FIoQueue;
 	friend class FBulkData;
+	friend class FIoDispatcherInternal;
 };
 
 //////////////////////////////////////////////////////////////////////////
diff --git a/Engine/Source/Runtime/PakFile/Private/IPlatformFilePak.cpp b/Engine/Source/Runtime/PakFile/Private/IPlatformFilePak.cpp
--- a/Engine/Source/Runtime/PakFile/Private/IPlatformFilePak.cpp
+++ b/Engine/Source/Runtime/PakFile/Private/IPlatformFilePak.cpp
@@ -2,6 +2,7 @@
 
 #include "IPlatformFilePak.h"
 #include "HAL/FileManager.h"
+#include "IO/IoDispatcherInternal.h"
 #include "IO/PlatformIoDispatcher.h"
 #include "Math/GuardedInt.h"
 #include "Misc/CoreMisc.h"
@@ -5692,6 +5693,7 @@ bool FPakPlatformFile::Initialize(IPlatformFile* Inner, const TCHAR* CmdLine)
 #endif
 						FPakPlatformFile::BroadcastPakChunkSignatureCheckFailure(FailedData);
 					});
+					FIoDispatcherInternal::MountContainerMeta(GlobalUTocPath);
 				}
 				else
 				{
@@ -6035,6 +6037,8 @@ bool FPakPlatformFile::Mount(const FPakMountArgs& MountArgs, FIoStatus* OutIoMou
 				const FString UtocPath = FPaths::ChangeExtension(InPakFilename, TEXT(".utoc"));
 				if (FPlatformFileManager::Get().GetPlatformFile().FileExists(*UtocPath))
 				{
+					FIoDispatcherInternal::MountContainerMeta(UtocPath);
+
 					if (ShouldCheckPak())
 					{
 						ensure(CheckIoStoreContainerBlockSignatures(*UtocPath));
