Clean up the `FMemory` transactional code so that it all uses the same macros. Disable validation because in theory (and hopefully in practice!) we shouldn't be able to hit open <-> closed hazards in allocators.

diff --git a/Engine/Source/Runtime/Core/Public/HAL/FMemory.inl b/Engine/Source/Runtime/Core/Public/HAL/FMemory.inl
--- a/Engine/Source/Runtime/Core/Public/HAL/FMemory.inl
+++ b/Engine/Source/Runtime/Core/Public/HAL/FMemory.inl
@@ -13,7 +13,6 @@ struct FScopedMallocTimer;
 void FMemory_FreeInline(void* Original);
 SIZE_T FMemory_GetAllocSizeInline(void* Original);
 
-
 FORCEINLINE void* FMemory_MallocInline(SIZE_T Count, uint32 Alignment)
 {
 	// AutoRTFM: For non-transactional code, all of these calls optimize away and the
@@ -22,33 +21,33 @@ FORCEINLINE void* FMemory_MallocInline(SIZE_T Count, uint32 Alignment)
 	// we don't end up keeping track of the writes to the allocator's internal data structures.
 	// This is because allocators are already transactional - malloc can be rolled back by
 	// calling free.
-	void* Ptr = AutoRTFM::Open([Count, Alignment]
+	void* Ptr = nullptr;
+
+	UE_AUTORTFM_OPEN_NO_VALIDATION
 	{
-		void* Alloc = nullptr;
 		if (UNLIKELY(!FMEMORY_INLINE_GMalloc))
 		{
-			Alloc = FMemory::MallocExternal(Count, Alignment);
+			Ptr = FMemory::MallocExternal(Count, Alignment);
 		}
 		else
 		{
 			DoGamethreadHook(0);
 			FScopedMallocTimer Timer(0);
-			Alloc = FMEMORY_INLINE_GMalloc->Malloc(Count, Alignment);
+			Ptr = FMEMORY_INLINE_GMalloc->Malloc(Count, Alignment);
 		}
 		// optional tracking of every allocation
-		LLM_IF_ENABLED(FLowLevelMemTracker::Get().OnLowLevelAlloc(ELLMTracker::Default, Alloc, Count, ELLMTag::Untagged, ELLMAllocType::FMalloc));
-		return Alloc;
-	});
+		LLM_IF_ENABLED(FLowLevelMemTracker::Get().OnLowLevelAlloc(ELLMTracker::Default, Ptr, Count, ELLMTag::Untagged, ELLMAllocType::FMalloc));
+	};
 
 	// AutoRTFM: This is a no-op for non-transactional code.
 	// For transactional code, this defers a call to Free if the transaction aborts,
 	// so that rolling back this allocation will end up freeing the memory.
-	AutoRTFM::OnAbort([Ptr]
+	UE_AUTORTFM_ONABORT(Ptr)
 	{
 		// Disable the code analysis warning that complains that Free is being passed
 		// a pointer that may be null. Free explicitly handles this case already.
 		FMemory::Free(Ptr); //-V575
-	});
+	};
 
 	return AutoRTFM::DidAllocate(Ptr, Count);
 }
@@ -167,8 +166,11 @@ FORCEINLINE void FMemory_FreeInline(void* Original)
 
 FORCEINLINE SIZE_T FMemory_GetAllocSizeInline(void* Original)
 {
+	const bool bWasClosed = AutoRTFM::IsClosed();
+
 	SIZE_T Result;
-	UE_AUTORTFM_OPEN
+
+	UE_AUTORTFM_OPEN_NO_VALIDATION
 	{
 		if (UNLIKELY(!FMEMORY_INLINE_GMalloc))
 		{
@@ -183,7 +185,7 @@ FORCEINLINE SIZE_T FMemory_GetAllocSizeInline(void* Original)
 			// This folds away at compile time so that the check is only ever performed inside transactional
 			// code paths. The check is to ensure that the allocator used will return the correct allocation
 			// size, which is a cornerstone requirement for AutoRTFM to function.
-			if (AutoRTFM::IsClosed())
+			if (bWasClosed)
 			{
 				checkf(bGotSize, TEXT("For AutoRTFM to function it must be able to get the size of an allocation"));
 			}
@@ -203,7 +205,7 @@ FORCEINLINE void* FMemory_MallocZeroedInline(SIZE_T Count, uint32 Alignment)
 	// we don't end up keeping track of the writes to the allocator's internal data structures.
 	// This is because allocators are already transactional - malloc can be rolled back by
 	// calling free.
-	UE_AUTORTFM_OPEN
+	UE_AUTORTFM_OPEN_NO_VALIDATION
 	{
 		if (UNLIKELY(!FMEMORY_INLINE_GMalloc))
 		{
@@ -222,12 +224,12 @@ FORCEINLINE void* FMemory_MallocZeroedInline(SIZE_T Count, uint32 Alignment)
 	// AutoRTFM: This is a no-op for non-transactional code.
 	// For transactional code, this defers a call to Free if the transaction aborts,
 	// so that rolling back this allocation will end up freeing the memory.
-	AutoRTFM::OnAbort([Ptr]
+	UE_AUTORTFM_ONABORT(Ptr)
 	{
 		// Disable the code analysis warning that complains that Free is being passed
 		// a pointer that may be null. Free explicitly handles this case already.
 		FMemory::Free(Ptr); //-V575
-	});
+	};
 
 	return AutoRTFM::DidAllocate(Ptr, Count);
 }
@@ -235,7 +237,8 @@ FORCEINLINE void* FMemory_MallocZeroedInline(SIZE_T Count, uint32 Alignment)
 FORCEINLINE SIZE_T FMemory_QuantizeSizeInline(SIZE_T Count, uint32 Alignment)
 {
 	SIZE_T Result;
-	UE_AUTORTFM_OPEN
+
+	UE_AUTORTFM_OPEN_NO_VALIDATION
 	{
 		if (UNLIKELY(!FMEMORY_INLINE_GMalloc))
 		{
