Provide a hint if FCbWriter::Save() is called from the closed.

This has bitten us twice so far.

#Synced-CL 48085702

diff --git a/Engine/Source/Runtime/Core/Private/Serialization/CompactBinaryWriter.cpp b/Engine/Source/Runtime/Core/Private/Serialization/CompactBinaryWriter.cpp
--- a/Engine/Source/Runtime/Core/Private/Serialization/CompactBinaryWriter.cpp
+++ b/Engine/Source/Runtime/Core/Private/Serialization/CompactBinaryWriter.cpp
@@ -97,6 +97,8 @@ void FCbWriter::Reset()
 
 FCbFieldIterator FCbWriter::Save() const
 {
+	AutoRTFM::UnreachableIfClosed("FUniqueBuffer is not yet transactionally safe. Call FCbWriter::Save(const FMutableMemoryView Buffer) instead.");
+
 	const uint64 Size = GetSaveSize();
 	FUniqueBuffer Buffer = FUniqueBuffer::Alloc(Size);
 	const FCbFieldViewIterator Output = Save(Buffer);
