Added WaitFor and WaitUntil to FConditionVariable

#Synced-CL 47599469

diff --git a/Engine/Source/Runtime/Core/Public/Experimental/Async/ConditionVariable.h b/Engine/Source/Runtime/Core/Public/Experimental/Async/ConditionVariable.h
--- a/Engine/Source/Runtime/Core/Public/Experimental/Async/ConditionVariable.h
+++ b/Engine/Source/Runtime/Core/Public/Experimental/Async/ConditionVariable.h
@@ -20,14 +20,14 @@ public:
 
     void NotifyOne()
     {
-        if (bHasWaiters)
+        if (bMayHaveWaiters)
         {
             ParkingLot::WakeOne(
-                &bHasWaiters,
+                &bMayHaveWaiters,
                 [this] (ParkingLot::FWakeState WakeState) -> uint64
                 {
                     if (!WakeState.bHasWaitingThreads)
-                        bHasWaiters = false;
+                        bMayHaveWaiters = false;
                     return 0;
                 });
         }
@@ -35,10 +35,10 @@ public:
     
     void NotifyAll()
     {
-        if (bHasWaiters)
+        if (bMayHaveWaiters)
         {
-            bHasWaiters = false;
-            ParkingLot::WakeAll(&bHasWaiters);
+            bMayHaveWaiters = false;
+            ParkingLot::WakeAll(&bMayHaveWaiters);
         }
     }
 
@@ -46,10 +46,10 @@ public:
     void Wait(TLock& Lock)
     {
         ParkingLot::Wait(
-            &bHasWaiters,
+            &bMayHaveWaiters,
             [this] () -> bool
             {
-                bHasWaiters = true;
+                bMayHaveWaiters = true;
                 return true;
             },
             [&Lock] ()
@@ -59,8 +59,46 @@ public:
         Lock.Lock();
     }
     
+	template<typename TLock>
+	bool WaitFor(TLock& Lock, FMonotonicTimeSpan WaitTime)
+    {
+    	ParkingLot::FWaitState WaitState = ParkingLot::WaitFor(
+			&bMayHaveWaiters,
+			[this] () -> bool
+			{
+				bMayHaveWaiters = true;
+				return true;
+			},
+			[&Lock] ()
+			{
+				Lock.Unlock();
+			},
+			WaitTime);
+    	Lock.Lock();
+     	return WaitState.bDidWake;
+	}
+
+	template<typename TLock>
+	bool WaitUntil(TLock& Lock, FMonotonicTimePoint WaitTime)
+    {
+    	ParkingLot::FWaitState WaitState = ParkingLot::WaitUntil(
+			&bMayHaveWaiters,
+			[this] () -> bool
+			{
+				bMayHaveWaiters = true;
+				return true;
+			},
+			[&Lock] ()
+			{
+				Lock.Unlock();
+			},
+			WaitTime);
+    	Lock.Lock();
+    	return WaitState.bDidWake;
+    }
+
 private:
-    std::atomic<bool> bHasWaiters = false;
+    std::atomic<bool> bMayHaveWaiters = false;
 };
 
 } // namespace UE
