[AutoRTFM] Add an AutoRTFM hazard message for UE::FSharedMutex.

We will now prompt users to switch to an FTransactionallySafeSharedMutex instead.

#Synced-CL 48085702

diff --git a/Engine/Source/Programs/AutoRTFMTests/Private/Abort.cpp b/Engine/Source/Programs/AutoRTFMTests/Private/Abort.cpp
--- a/Engine/Source/Programs/AutoRTFMTests/Private/Abort.cpp
+++ b/Engine/Source/Programs/AutoRTFMTests/Private/Abort.cpp
@@ -5,6 +5,7 @@
 #include "AutoRTFMTestUtils.h"
 #include "AutoRTFM/AutoRTFMUE.h"
 #include "Async/Mutex.h"
+#include "Async/SharedMutex.h"
 #include "Catch2Includes.h"
 #include "Misc/StringOutputDevice.h"
 #include "Misc/OutputDeviceRedirector.h"
@@ -1480,6 +1481,42 @@ TEST_CASE("Abort.Language.ErrorHints.MutexIsLocked")
 	REQUIRE(WarningContext.HasWarningSubstring("A mutex was accessed while inside an AutoRTFM transaction."));
 }
 
+TEST_CASE("Abort.Language.ErrorHints.SharedMutexLockShared")
+{
+	FScopedErrorHintTest ScopedTest;
+    AutoRTFMTestUtils::FCaptureWarningContext WarningContext;
+
+	UE::FSharedMutex Mutex;
+
+	const AutoRTFM::ETransactionResult Result = AutoRTFM::Transact([&]
+	{
+		Mutex.LockShared();
+	});
+
+	REQUIRE(AutoRTFM::ETransactionResult::AbortedByLanguage == Result);
+
+	// Did we print our message?
+	REQUIRE(WarningContext.HasWarningSubstring("A shared mutex was accessed while inside an AutoRTFM transaction."));
+}
+
+TEST_CASE("Abort.Language.ErrorHints.SharedMutexIsLockShared")
+{
+	FScopedErrorHintTest ScopedTest;
+    AutoRTFMTestUtils::FCaptureWarningContext WarningContext;
+
+	UE::FSharedMutex Mutex;
+
+	const AutoRTFM::ETransactionResult Result = AutoRTFM::Transact([&]
+	{
+		REQUIRE(!Mutex.IsLockShared());
+	});
+
+	REQUIRE(AutoRTFM::ETransactionResult::AbortedByLanguage == Result);
+
+	// Did we print our message?
+	REQUIRE(WarningContext.HasWarningSubstring("A shared mutex was accessed while inside an AutoRTFM transaction."));
+}
+
 #if 0
 // Testing an open free is flaky, since it leads to a double free--one inside the open block, and one during the
 // transactional abort caused by an AutoRTFM hazard. This can cause the unit test harness to fail on a subsequent test,
diff --git a/Engine/Source/Runtime/AutoRTFM/Private/ReportHazard.cpp b/Engine/Source/Runtime/AutoRTFM/Private/ReportHazard.cpp
--- a/Engine/Source/Runtime/AutoRTFM/Private/ReportHazard.cpp
+++ b/Engine/Source/Runtime/AutoRTFM/Private/ReportHazard.cpp
@@ -49,6 +49,16 @@ it does support transactionally-safe mutexes that can be used instead.
 
 Use the call stack below to find the FMutex accessed immediately before this error occurred, and
 replace it with FTransactionallySafeMutex.
+)";
+			break;
+
+		case EHazardType::SharedMutex:
+			Message = R"(
+A shared mutex was accessed while inside an AutoRTFM transaction. AutoRTFM doesn't support native
+shared mutexes, but it does support transactionally-safe shared mutexes that can be used instead.
+
+Use the call stack below to find the FSharedMutex accessed immediately before this error occurred,
+and replace it with FTransactionallySafeSharedMutex.
 )";
 			break;
 
diff --git a/Engine/Source/Runtime/AutoRTFM/Public/AutoRTFM/ReportHazard.h b/Engine/Source/Runtime/AutoRTFM/Public/AutoRTFM/ReportHazard.h
--- a/Engine/Source/Runtime/AutoRTFM/Public/AutoRTFM/ReportHazard.h
+++ b/Engine/Source/Runtime/AutoRTFM/Public/AutoRTFM/ReportHazard.h
@@ -15,6 +15,7 @@ enum class EHazardType
 	CriticalSection,
 	RWLock,
 	Mutex,
+	SharedMutex,
 	OpenFree,
 };
 
diff --git a/Engine/Source/Runtime/Core/Public/Async/SharedMutex.h b/Engine/Source/Runtime/Core/Public/Async/SharedMutex.h
--- a/Engine/Source/Runtime/Core/Public/Async/SharedMutex.h
+++ b/Engine/Source/Runtime/Core/Public/Async/SharedMutex.h
@@ -3,6 +3,7 @@
 #pragma once
 
 #include "CoreTypes.h"
+#include "AutoRTFM/ReportHazard.h"
 #include "Misc/AssertionMacros.h"
 #include <atomic>
 
@@ -28,11 +29,13 @@ public:
 
 	[[nodiscard]] inline bool IsLocked() const
 	{
+		UE_AUTORTFM_REPORT_HAZARD_IF_CLOSED(SharedMutex);
 		return !!(State.load(std::memory_order_relaxed) & IsLockedFlag);
 	}
 
 	[[nodiscard]] inline bool TryLock()
 	{
+		UE_AUTORTFM_REPORT_HAZARD_IF_CLOSED(SharedMutex);
 		uint32 Expected = State.load(std::memory_order_relaxed);
 		return !(Expected & (IsLockedFlag | SharedLockCountMask)) &&
 			State.compare_exchange_strong(Expected, Expected | IsLockedFlag,
@@ -41,6 +44,7 @@ public:
 
 	inline void Lock()
 	{
+		UE_AUTORTFM_REPORT_HAZARD_IF_CLOSED(SharedMutex);
 		uint32 Expected = 0;
 		if (LIKELY(State.compare_exchange_weak(Expected, IsLockedFlag, std::memory_order_acquire, std::memory_order_relaxed)))
 		{
@@ -52,6 +56,7 @@ public:
 	inline void Unlock()
 	{
 		// Unlock immediately to allow other threads to acquire the lock while this thread looks for a thread to wake.
+		UE_AUTORTFM_REPORT_HAZARD_IF_CLOSED(SharedMutex);
 		uint32 LastState = State.fetch_sub(IsLockedFlag, std::memory_order_release);
 		checkSlow(LastState & IsLockedFlag);
 		if (LIKELY(!(LastState & (MayHaveWaitingLockFlag | MayHaveWaitingSharedLockFlag))))
@@ -63,11 +68,13 @@ public:
 
 	[[nodiscard]] inline bool IsLockShared() const
 	{
+		UE_AUTORTFM_REPORT_HAZARD_IF_CLOSED(SharedMutex);
 		return !!(State.load(std::memory_order_relaxed) & SharedLockCountMask);
 	}
 
 	[[nodiscard]] inline bool TryLockShared()
 	{
+		UE_AUTORTFM_REPORT_HAZARD_IF_CLOSED(SharedMutex);
 		uint32 Expected = State.load(std::memory_order_relaxed);
 		while (LIKELY(!(Expected & (IsLockedFlag | MayHaveWaitingLockFlag))))
 		{
@@ -82,6 +89,7 @@ public:
 
 	inline void LockShared()
 	{
+		UE_AUTORTFM_REPORT_HAZARD_IF_CLOSED(SharedMutex);
 		uint32 Expected = State.load(std::memory_order_relaxed);
 		if (LIKELY(!(Expected & (IsLockedFlag | MayHaveWaitingLockFlag)) &&
 			State.compare_exchange_weak(Expected, Expected + (1 << SharedLockCountShift),
@@ -95,6 +103,7 @@ public:
 	inline void UnlockShared()
 	{
 		// Unlock immediately to allow other threads to acquire the lock while this thread looks for a thread to wake.
+		UE_AUTORTFM_REPORT_HAZARD_IF_CLOSED(SharedMutex);
 		const uint32 LastState = State.fetch_sub(1 << SharedLockCountShift, std::memory_order_release);
 		checkSlow(LastState & SharedLockCountMask);
 		constexpr uint32 WakeState = MayHaveWaitingLockFlag | (1 << SharedLockCountShift);
