[Core/Stats] Fixed SCOPE_CYCLE_COUNTER_STATID to allow multiple declarations in the same scope.

#Synced-CL 48145873

diff --git a/Engine/Source/Runtime/Core/Public/Stats/Stats.h b/Engine/Source/Runtime/Core/Public/Stats/Stats.h
--- a/Engine/Source/Runtime/Core/Public/Stats/Stats.h
+++ b/Engine/Source/Runtime/Core/Public/Stats/Stats.h
@@ -230,7 +230,7 @@ struct FStat_##StatName\
 	FScopeCycleCounter CycleCount_##Stat(GET_STATID(Stat), GET_STATFLAGS(Stat));
 
 #define SCOPE_CYCLE_COUNTER_STATID(StatId) \
-	FScopeCycleCounter CycleCount_FromStatId(StatId);
+	FScopeCycleCounter UE_JOIN(CycleCount_FromStatId_, __LINE__)(StatId);
 
 #define CONDITIONAL_SCOPE_CYCLE_COUNTER(Stat,bCondition) \
 	FScopeCycleCounter CycleCount_##Stat(bCondition ? GET_STATID(Stat) : TStatId(), GET_STATFLAGS(Stat));
@@ -582,7 +582,7 @@ struct FStat_##StatName\
 
 // NOTE: This macro bypasses the standard enable checks as we don't know the stat's type/group.
 #define SCOPE_CYCLE_COUNTER_STATID(StatId) \
-	FScopeCycleCounter StatNamedEventsScope_STATID(StatId); \
+	FScopeCycleCounter UE_JOIN(CycleCount_FromStatId_, __LINE__)(StatId); \
 	TRACE_CPUPROFILER_EVENT_SCOPE_TEXT_CONDITIONAL(StatId.StatString, StatId.IsValidStat() && GCycleStatsShouldEmitNamedEvents > 0);
 
 #define CONDITIONAL_SCOPE_CYCLE_COUNTER(Stat,bCondition) \
@@ -618,7 +618,7 @@ struct FStat_##StatName\
 	UE::Stats::FHitchTrackingStatScope HitchTrackingStatScope_##Stat(ANSI_TO_PROFILING(#Stat));
 
 #define SCOPE_CYCLE_COUNTER_STATID(StatId) \
-	UE::Stats::FHitchTrackingStatScope HitchTrackingStatScope_##Stat(ANSI_TO_PROFILING("Lightweight StatId Scope"));
+	UE::Stats::FHitchTrackingStatScope UE_JOIN(HitchTrackingStatScope_FromStatId_, __LINE__)(ANSI_TO_PROFILING("Lightweight StatId Scope"));
 
 #define CONDITIONAL_SCOPE_CYCLE_COUNTER(Stat,bCondition) \
 	UE::Stats::FHitchTrackingStatScope HitchTrackingStatScope_##Stat(bCondition ? ANSI_TO_PROFILING(#Stat) : nullptr);
