Disable AutoRTFM validation on `FThreadSingletonInitializer::Get` as the scope of it is limited, and it is heavily hit running BPVM.

#Synced-CL 48143733

diff --git a/Engine/Source/Runtime/Core/Private/HAL/ThreadingBase.cpp b/Engine/Source/Runtime/Core/Private/HAL/ThreadingBase.cpp
--- a/Engine/Source/Runtime/Core/Private/HAL/ThreadingBase.cpp
+++ b/Engine/Source/Runtime/Core/Private/HAL/ThreadingBase.cpp
@@ -1480,7 +1480,11 @@ FTlsAutoCleanup* FThreadSingletonInitializer::Get( TFunctionRef<FTlsAutoCleanup*
 {
 	uint32 TlsSlot;
 	FTlsAutoCleanup* ThreadSingleton = nullptr;
-	UE_AUTORTFM_OPEN
+
+	// This open block is a) hit frequently and b) constrained in the scope of open code that is touched.
+	// When running with memory validation on this becomes a performance bottleneck even though there is
+	// no additional safety achieved. Thus we run this open with validation disabled.
+	UE_AUTORTFM_OPEN_NO_VALIDATION
 	{
 		TlsSlot = (uint32)FPlatformAtomics::AtomicRead_Relaxed((int32*)&InOutTlsSlot);
 		if (TlsSlot == FPlatformTLS::InvalidTlsSlot)
