Add RandomShuffle version using a FRandomStream

#Synced-CL 48202918

diff --git a/Engine/Source/Runtime/Core/Public/Algo/RandomShuffle.h b/Engine/Source/Runtime/Core/Public/Algo/RandomShuffle.h
--- a/Engine/Source/Runtime/Core/Public/Algo/RandomShuffle.h
+++ b/Engine/Source/Runtime/Core/Public/Algo/RandomShuffle.h
@@ -2,18 +2,14 @@
 
 #pragma once
 
+#include "Math/RandomStream.h"
 #include "Math/UnrealMathUtility.h"
 #include "Templates/UnrealTemplate.h"
 
-namespace Algo
+namespace AlgoImpl
 {
-	/**
-	 * Randomly shuffle a range of elements.
-	 *
-	 * @param  Range  Any contiguous container.
-	 */
-	template <typename RangeType>
-	void RandomShuffle(RangeType&& Range)
+	template <typename RangeType, typename RandomFunc>
+	void RandomShuffle(RangeType&& Range, RandomFunc Random)
 	{
 		auto Data = GetData(Range);
 
@@ -23,7 +19,7 @@ namespace Algo
 		for (SizeType Index = 0; Index < Num - 1; ++Index)
 		{
 			// Get a random integer in [Index, Num)
-			const SizeType RandomIndex = Index + (SizeType)FMath::RandHelper64(Num - Index);
+			const SizeType RandomIndex = Index + (SizeType)::Invoke(Random, Num - Index);
 			if (RandomIndex != Index)
 			{
 				Swap(Data[Index], Data[RandomIndex]);
@@ -31,3 +27,31 @@ namespace Algo
 		}
 	}
 }
+
+namespace Algo
+{
+	/**
+	 * Randomly shuffle a range of elements.
+	 *
+	 * @param  Range  Any contiguous container.
+	 */
+	template <typename RangeType>
+	void RandomShuffle(RangeType&& Range)
+	{
+		using SizeType = decltype(GetNum(Range));
+		AlgoImpl::RandomShuffle(Range, [] (const SizeType MaxIndex) { return FMath::RandHelper64(MaxIndex); });
+	}
+
+	/**
+	 * Randomly shuffle a range of elements.
+	 *
+	 * @param  Range  Any contiguous container.
+	 * @param  RandomStream - Random stream from which to retrieve random indices.
+	 */
+	template <typename RangeType>
+	void RandomShuffle(RangeType&& Range, FRandomStream& RandomStream)
+	{
+		using SizeType = decltype(GetNum(Range));
+		AlgoImpl::RandomShuffle(Range, [&RandomStream] (const SizeType MaxIndex) { return RandomStream.RandHelper64(MaxIndex); });
+	}
+}
diff --git a/Engine/Source/Runtime/Core/Public/Math/RandomStream.h b/Engine/Source/Runtime/Core/Public/Math/RandomStream.h
--- a/Engine/Source/Runtime/Core/Public/Math/RandomStream.h
+++ b/Engine/Source/Runtime/Core/Public/Math/RandomStream.h
@@ -189,6 +189,18 @@ public:
 		return ((A > 0) ? FMath::TruncToInt(GetFraction() * float(A)) : 0);
 	}
 
+	/**
+	 * Helper function for rand implementations.
+	 *
+	 * @return A random number in [0..A)
+	 */
+	UE_FORCEINLINE_HINT int64 RandHelper64(int64 A) const
+	{
+		// Note that on some platforms RAND_MAX is a large number so we cannot do ((rand()/(RAND_MAX+1)) * A)
+		// or else we may include the upper bound results, which should be excluded.
+		return A > 0 ? FMath::Min<int64>(FMath::TruncToInt(GetFraction() * float(A)), A - 1) : 0;
+	}
+
 	/**
 	 * Helper function for rand implementations.
 	 *
