[UE] TCompactSet: Ensure that elements are rehashed after being serialized during reference replacement.

#Synced-CL 48038701

diff --git a/Engine/Source/Runtime/Core/Public/Containers/CompactSet.h.inl b/Engine/Source/Runtime/Core/Public/Containers/CompactSet.h.inl
--- a/Engine/Source/Runtime/Core/Public/Containers/CompactSet.h.inl
+++ b/Engine/Source/Runtime/Core/Public/Containers/CompactSet.h.inl
@@ -1634,6 +1634,15 @@ struct TSETPRIVATEFRIEND
 			{
 				Ar << Element;
 			}
+
+			// Modifying refs is similar to load, in that we need to rehash the values after serializing.
+			// Ex: See FArchiveReplaceObjectRef, which is often used to patch references via serialization.
+			// This is reflecting the same logic being used in TScriptSparseSet for now to ensure correctness.
+			// @todo - Maybe this could be optimized later to only rehash when the elements are pointer types.
+			if (Ar.IsModifyingWeakAndStrongReferences() && !Ar.IsSaving())
+			{
+				Set.Rehash();
+			}
 		}
 		return Ar;
 	}
@@ -1645,7 +1654,8 @@ struct TSETPRIVATEFRIEND
 		int32 NumElements = Set.Num();
 		FStructuredArchive::FArray Array = Slot.EnterArray(NumElements);
 
-		if (Slot.GetUnderlyingArchive().IsLoading())
+		FArchive& Ar = Slot.GetUnderlyingArchive();
+		if (Ar.IsLoading())
 		{
 			// We can skip the reset on Empty and do it once at the end with Rehash
 			DestructItems(Set.GetData(), Set.Num());
@@ -1669,6 +1679,15 @@ struct TSETPRIVATEFRIEND
 				FStructuredArchive::FSlot ElementSlot = Array.EnterElement();
 				ElementSlot << Element;
 			}
+
+			// Modifying refs is similar to load, in that we need to rehash the values after serializing.
+			// Ex: See FArchiveReplaceObjectRef, which is often used to patch references via serialization.
+			// This is reflecting the same logic being used in TScriptSparseSet for now to ensure correctness.
+			// @todo - Maybe this could be optimized later to only rehash when the elements are pointer types.
+			if (Ar.IsModifyingWeakAndStrongReferences() && !Ar.IsSaving())
+			{
+				Set.Rehash();
+			}
 		}
 	}
 
