Fixed ReadAt bugs in IOSPlatformFile, ApplePlatformFile and UnixPlatformFile

#rnx

#Synced-CL Unknown

diff --git a/Engine/Source/Runtime/Core/Private/Apple/ApplePlatformFile.cpp b/Engine/Source/Runtime/Core/Private/Apple/ApplePlatformFile.cpp
--- a/Engine/Source/Runtime/Core/Private/Apple/ApplePlatformFile.cpp
+++ b/Engine/Source/Runtime/Core/Private/Apple/ApplePlatformFile.cpp
@@ -215,18 +215,22 @@ public:
 		do
 		{
 			size_t BytesToRead32 = static_cast<size_t>(FMath::Min<int64>(READWRITE_SIZE, BytesToRead));
-			ssize_t BytesRead = pread(FileHandle, Destination, BytesToRead, Offset);
+			ssize_t BytesRead = pread(FileHandle, Destination, BytesToRead32, Offset);
 
-			TotalBytesRead += BytesRead;
-
-			if (BytesRead != BytesToRead32)
+			if (BytesRead == -1 && errno == EINTR)
+			{
+				continue;
+			}
+			if (BytesRead <= 0)
 			{
 				TRACE_PLATFORMFILE_END_READ(this, TotalBytesRead);
 				return false;
 			}
 
+			TotalBytesRead += BytesRead;
 			Offset += BytesRead;
-			BytesToRead -= BytesToRead32;
+			BytesToRead -= BytesRead;
+			Destination += BytesRead;
 
 		} while (BytesToRead > 0);
 
diff --git a/Engine/Source/Runtime/Core/Private/IOS/IOSPlatformFile.cpp b/Engine/Source/Runtime/Core/Private/IOS/IOSPlatformFile.cpp
--- a/Engine/Source/Runtime/Core/Private/IOS/IOSPlatformFile.cpp
+++ b/Engine/Source/Runtime/Core/Private/IOS/IOSPlatformFile.cpp
@@ -224,15 +224,20 @@ public:
 		do
 		{
 			size_t BytesToRead32 = static_cast<size_t>(FMath::Min<int64>(READWRITE_SIZE, BytesToRead));
-			ssize_t BytesRead = pread(FileHandle, Destination, BytesToRead, Offset);
+			ssize_t BytesRead = pread(FileHandle, Destination, BytesToRead32, Offset);
 
-			if (BytesRead != BytesToRead32)
+			if (BytesRead == -1 && errno == EINTR)
+			{
+				continue;
+			}
+			if (BytesRead <= 0)
 			{
 				return false;
 			}
 
 			Offset += BytesRead;
-			BytesToRead -= BytesToRead32;
+			BytesToRead -= BytesRead;
+			Destination += BytesRead;
 
 		} while (BytesToRead > 0);
 
diff --git a/Engine/Source/Runtime/Core/Private/Unix/UnixPlatformFile.cpp b/Engine/Source/Runtime/Core/Private/Unix/UnixPlatformFile.cpp
--- a/Engine/Source/Runtime/Core/Private/Unix/UnixPlatformFile.cpp
+++ b/Engine/Source/Runtime/Core/Private/Unix/UnixPlatformFile.cpp
@@ -250,16 +250,16 @@ public:
 		do
 		{
 			size_t BytesToRead32 = static_cast<size_t>(FMath::Min<int64>(READWRITE_SIZE, BytesToRead));
-			ssize_t BytesRead = pread(FileHandle, Destination, BytesToRead, Offset);
+			ssize_t BytesRead = pread(FileHandle, Destination, BytesToRead32, Offset);
 
 			if (BytesRead == -1 && errno == EFAULT)
 			{
 				// Hacky workaround @see ReadInternal
-				void* TempDest = FMemory::Malloc(BytesToRead);
+				void* TempDest = FMemory::Malloc(BytesToRead32);
 				if (TempDest)
 				{
-					BytesRead = pread(FileHandle, TempDest, BytesToRead, Offset);
-					if (BytesRead > 0 && BytesRead <= BytesToRead)
+					BytesRead = pread(FileHandle, TempDest, BytesToRead32, Offset);
+					if (BytesRead > 0 && BytesRead <= BytesToRead32)
 					{
 						FMemory::Memcpy(Destination, TempDest, BytesRead);
 					}
@@ -267,13 +267,18 @@ public:
 				}
 			}
 
-			if (BytesRead != BytesToRead32)
+			if (BytesRead == -1 && errno == EINTR)
+			{
+				continue;
+			}
+			if (BytesRead <= 0)
 			{
 				return false;
 			}
 
 			Offset += BytesRead;
-			BytesToRead -= BytesToRead32;
+			BytesToRead -= BytesRead;
+			Destination += BytesRead;
 
 		} while (BytesToRead > 0);
 
