[Apple] os_log

Switching from NSLog to os_log as it is faster & allows to control privacy of our logs.
IOS 26 has logs private by default, so even in non-shipping builds old code logging is not visible ouside of Xcode.

Also marked Apple LocalPrint logging as thread-safe (os_log is thread-safe), so it now skips the buffered output path.

#rnx

diff --git a/Engine/Source/Runtime/Core/Private/Apple/ApplePlatformMisc.cpp b/Engine/Source/Runtime/Core/Private/Apple/ApplePlatformMisc.cpp
--- a/Engine/Source/Runtime/Core/Private/Apple/ApplePlatformMisc.cpp
+++ b/Engine/Source/Runtime/Core/Private/Apple/ApplePlatformMisc.cpp
@@ -15,6 +15,8 @@
 #include "FramePro/FrameProProfiler.h"
 #include "CoreGlobals.h"
 
+#include <os/log.h>
+
 DEFINE_LOG_CATEGORY_STATIC(LogApplePlatformMisc, Log, All);
 
 #if !UE_BUILD_SHIPPING
@@ -75,8 +77,12 @@ FString FApplePlatformMisc::GetEnvironmentVariable(const TCHAR* VariableName)
 
 void FApplePlatformMisc::LocalPrint(const TCHAR* Message)
 {
-	//NsLog will out to all iOS output consoles, instead of just the Xcode console.
-	NSLog(@"[UE] %s", TCHAR_TO_UTF8(Message));
+	// This goes to all output consoles, not only Xcode one.
+#if !UE_BUILD_SHIPPING
+	os_log(OS_LOG_DEFAULT, "[UE] %{public}s", TCHAR_TO_UTF8(Message));
+#else
+	os_log(OS_LOG_DEFAULT, "[UE] %s", TCHAR_TO_UTF8(Message)); // Private by default, would be visible in Xcode.
+#endif
 }
 
 const TCHAR* FApplePlatformMisc::GetSystemErrorMessage(TCHAR* OutBuffer, int32 BufferCount, int32 Error)
diff --git a/Engine/Source/Runtime/Core/Public/Apple/ApplePlatformMisc.h b/Engine/Source/Runtime/Core/Public/Apple/ApplePlatformMisc.h
--- a/Engine/Source/Runtime/Core/Public/Apple/ApplePlatformMisc.h
+++ b/Engine/Source/Runtime/Core/Public/Apple/ApplePlatformMisc.h
@@ -71,6 +71,7 @@ struct CORE_API FApplePlatformMisc : public FGenericPlatformMisc
 	}
 
 	static void LocalPrint(const TCHAR* Message);
+	static bool IsLocalPrintThreadSafe() { return true; }
 	static const TCHAR* GetSystemErrorMessage(TCHAR* OutBuffer, int32 BufferCount, int32 Error);
 	static uint32 GetLastError();
 	static int32 NumberOfCores();
