Fold `FRefCountBase`, `FThreadSafeRefCountedObject`, and `FRefCountedObject` together and use `FRefCountedObject` from now on.

We'll eventually deprecate and remove all but `FRefCountedObject`.

#Synced-CL 48280287

diff --git a/Engine/Source/Runtime/Core/Public/Templates/RefCounting.h b/Engine/Source/Runtime/Core/Public/Templates/RefCounting.h
--- a/Engine/Source/Runtime/Core/Public/Templates/RefCounting.h
+++ b/Engine/Source/Runtime/Core/Public/Templates/RefCounting.h
@@ -209,89 +209,15 @@ public:
 /**
  * Base class implementing thread-safe reference counting.
  */
-class FRefCountBase : UE::Private::TTransactionalAtomicRefCount<uint32>
+class FRefCountedObject : UE::Private::TTransactionalAtomicRefCount<uint32>
 {
 public:
-	FRefCountBase() = default;
-	virtual ~FRefCountBase() = default;
+	FRefCountedObject() = default;
+	virtual ~FRefCountedObject() = default;
 
-	FRefCountBase(const FRefCountBase& Rhs) = delete;
-	FRefCountBase& operator=(const FRefCountBase& Rhs) = delete;
-
-	FReturnedRefCountValue AddRef() const
-	{
-		return FReturnedRefCountValue{Super::AddRef<DeleteThis>()};
-	}
-
-	uint32 Release() const
-	{
-		return Super::Release<DeleteThis>();
-	}
-
-	uint32 GetRefCount() const
-	{
-		return Super::GetRefCount();
-	}
-
-private:
-	using Super = UE::Private::TTransactionalAtomicRefCount<uint32>;
-
-	static void DeleteThis(const Super* This)
-	{
-		delete static_cast<const FRefCountBase*>(This);
-	}
-};
-
-/**
- * The base class of reference counted objects.
- *
- * This class should not be used for new code as it does not use atomic operations to update 
- * the reference count.
- */
-class FRefCountedObject
-{
-public:
-	FRefCountedObject(): NumRefs(0) {}
-	virtual ~FRefCountedObject() { check(!NumRefs); }
 	FRefCountedObject(const FRefCountedObject& Rhs) = delete;
 	FRefCountedObject& operator=(const FRefCountedObject& Rhs) = delete;
-	FReturnedRefCountValue AddRef() const
-	{
-		return FReturnedRefCountValue{uint32(++NumRefs)};
-	}
-	uint32 Release() const
-	{
-		uint32 Refs = uint32(--NumRefs);
-		if(Refs == 0)
-		{
-			delete this;
-		}
-		return Refs;
-	}
-	uint32 GetRefCount() const
-	{
-		return uint32(NumRefs);
-	}
-private:
-	mutable int32 NumRefs;
-};
-
-/**
- * Like FRefCountedObject, but the reference count is thread-safe.
- */
-class FThreadSafeRefCountedObject : UE::Private::TTransactionalAtomicRefCount<uint32>
-{
-public:
-	FThreadSafeRefCountedObject() = default;
-
-	FThreadSafeRefCountedObject(const FThreadSafeRefCountedObject& Rhs) = delete;
-	FThreadSafeRefCountedObject& operator=(const FThreadSafeRefCountedObject& Rhs) = delete;
 
-	virtual ~FThreadSafeRefCountedObject()
-	{ 
-		check(Super::GetRefCount() == 0); 
-	}
-	
 	FReturnedRefCountValue AddRef() const
 	{
 		return FReturnedRefCountValue{Super::AddRef<DeleteThis>()};
@@ -312,10 +238,18 @@ private:
 
 	static void DeleteThis(const Super* This)
 	{
-		delete static_cast<const FThreadSafeRefCountedObject*>(This);
+		delete static_cast<const FRefCountedObject*>(This);
 	}
 };
 
+// TODO (SOL-7350): deprecate and remove this entirely. This is a class because there are
+// some uses of this as a friend class declaration.
+class FRefCountBase : public FRefCountedObject {};
+
+// TODO (SOL-7350): deprecate and remove this entirely. This is a class because there are
+// some uses of this as a friend class declaration.
+class FThreadSafeRefCountedObject : public FRefCountedObject {};
+
 /**
  * ERefCountingMode is used select between either 'fast' or 'thread safe' ref-counting types.
  * This is only used at compile time to select between template specializations.
