Fix removal during iteration at merged module unmount

#rnx

#Synced-CL 48260816

diff --git a/Engine/Source/Runtime/Core/Private/Modules/ModuleManager.cpp b/Engine/Source/Runtime/Core/Private/Modules/ModuleManager.cpp
--- a/Engine/Source/Runtime/Core/Private/Modules/ModuleManager.cpp
+++ b/Engine/Source/Runtime/Core/Private/Modules/ModuleManager.cpp
@@ -692,12 +692,15 @@ void FModuleManager::InternalFreeLibrary(FName ModuleName, void* Handle)
 				UE_LOG(LogModuleManager, Verbose, TEXT("InternalFreeLibrary: cleaning up for module '%s'"),
 					*OtherModuleNameAndInfo.Key.ToString());
 
-				StaticallyLinkedModuleInitializers.Remove(OtherModuleNameAndInfo.Key);
-
 				UnloadedModules.Add(OtherModuleNameAndInfo.Key);
 			}
 		}
 
+		for (FName UnloadedModule : UnloadedModules)
+		{
+			StaticallyLinkedModuleInitializers.Remove(UnloadedModule);
+		}
+
 		// Mark for delay unload to run after GC is done.
 		DelayUnloadLibraries.Add(LibraryToUnload, UnloadedModules);
 
