Fix pipelining of FApp game/delta times
- Previous FApp::GetDeltaTime() etc functions read from relaxed std::atomics, meaning any thread reading these values would race the game thread. Tasks deep in the frame pipeline might receive delta time values for the wrong frame.
- Moved all time related fields and functions out from FApp into a new FAppTime class. The game thread sets a global copy of this class. Tasks receive a TSharedPtr const reference of the current global value at the point they are spawned.
- Child tasks receive the copy from their parent, meaning all tasks in the pipeline receive the original FAppTime value from the game thread task that spawned them.
- A future change will deprecate the old FApp::Get... functions, to be replaced with FAppTime::Get().

#Synced-CL 48210367

diff --git a/Engine/Source/Runtime/Core/Private/Async/InheritedContext.cpp b/Engine/Source/Runtime/Core/Private/Async/InheritedContext.cpp
--- a/Engine/Source/Runtime/Core/Private/Async/InheritedContext.cpp
+++ b/Engine/Source/Runtime/Core/Private/Async/InheritedContext.cpp
@@ -7,22 +7,18 @@ namespace UE
 	// this method must be defined in a .cpp file to avoid exporting dependency on trace module
 	[[nodiscard]] FInheritedContextScope FInheritedContextBase::RestoreInheritedContext()
 	{
-		return FInheritedContextScope{
-#if ENABLE_LOW_LEVEL_MEM_TRACKER
-				InheritedLLMTags
-	#if (UE_MEMORY_TAGS_TRACE_ENABLED||UE_TRACE_METADATA_ENABLED) && UE_TRACE_ENABLED
-				,
-	#endif
-#endif
-#if UE_MEMORY_TAGS_TRACE_ENABLED && UE_TRACE_ENABLED
-				InheritedMemTag
-	#if UE_TRACE_METADATA_ENABLED && UE_TRACE_ENABLED
-				,
-	#endif
-#endif
-#if UE_TRACE_METADATA_ENABLED && UE_TRACE_ENABLED
-				InheritedMetadataId
-#endif
-		};
+		return FInheritedContextScope
+		(
+			AppTime
+		#if ENABLE_LOW_LEVEL_MEM_TRACKER
+			, InheritedLLMTags
+		#endif
+		#if (UE_MEMORY_TAGS_TRACE_ENABLED && UE_TRACE_ENABLED)
+			, InheritedMemTag
+		#endif
+		#if (UE_TRACE_METADATA_ENABLED && UE_TRACE_ENABLED) 
+			, InheritedMetadataId
+		#endif
+		);
 	}
 }
\ No newline at end of file
diff --git a/Engine/Source/Runtime/Core/Private/Experimental/Containers/HazardPointer.cpp b/Engine/Source/Runtime/Core/Private/Experimental/Containers/HazardPointer.cpp
--- a/Engine/Source/Runtime/Core/Private/Experimental/Containers/HazardPointer.cpp
+++ b/Engine/Source/Runtime/Core/Private/Experimental/Containers/HazardPointer.cpp
@@ -4,7 +4,7 @@
 
 #include "Algo/Sort.h"
 #include "Async/UniqueLock.h"
-#include "Misc/App.h"
+#include "HAL/PlatformTime.h"
 
 FHazardPointerCollection::~FHazardPointerCollection()
 {
@@ -65,7 +65,7 @@ void FHazardPointerCollection::Delete(const HazardPointer_Impl::FHazardDeleter&
 	Collectables->ReclamationList.Add(Ptr);
 
 	//maybe scan the list
-	const double CurrentTime = FApp::GetGameTime();
+	const double CurrentTime = FPlatformTime::Seconds();
 	const bool TimeLimit = (CurrentTime - Collectables->TimeOfLastCollection) > 1.0;
 	const int32 DeleteMetric = CollectLimit <= 0 ? ((5 * TotalNumHazardRecords) / 4) : CollectLimit;
 	if (TimeLimit || Collectables->ReclamationList.Num() >= DeleteMetric)
diff --git a/Engine/Source/Runtime/Core/Private/Misc/App.cpp b/Engine/Source/Runtime/Core/Private/Misc/App.cpp
--- a/Engine/Source/Runtime/Core/Private/Misc/App.cpp
+++ b/Engine/Source/Runtime/Core/Private/Misc/App.cpp
@@ -36,15 +36,6 @@ TArray<FString> FApp::SessionUsers = TArray<FString>();
 bool FApp::Standalone = true;
 bool FApp::bIsBenchmarking = false;
 bool FApp::bUseFixedSeed = false;
-bool FApp::bUseFixedTimeStep = false;
-std::atomic<double> FApp::FixedDeltaTime = 1 / 30.0;
-std::atomic<double> FApp::CurrentTime = 0.0;
-std::atomic<double> FApp::LastTime = 0.0;
-std::atomic<double> FApp::DeltaTime = 1 / 30.0;
-std::atomic<double> FApp::IdleTime = 0.0;
-std::atomic<double> FApp::IdleTimeOvershoot = 0.0;
-std::atomic<double> FApp::GameTime = 0.0;
-TOptional<FQualifiedFrameTime> FApp::CurrentFrameTime;
 float FApp::VolumeMultiplier = 1.0f;
 float FApp::UnfocusedVolumeMultiplier = 0.0f;
 bool FApp::bUseVRFocus = false;
@@ -317,24 +308,6 @@ bool FApp::IsMultithreadServer()
 	return bIsMultithreadServer;
 }
 
-FTimecode FApp::GetTimecode()
-{
-	if (CurrentFrameTime.IsSet())
-	{
-		return CurrentFrameTime->ToTimecode();
-	}
-	return FTimecode();
-}
-
-FFrameRate FApp::GetTimecodeFrameRate()
-{
-	if (CurrentFrameTime.IsSet())
-	{
-		return CurrentFrameTime->Rate;
-	}
-	return FFrameRate();
-}
-
 static bool GUnfocusedVolumeMultiplierInitialised = false;
 float FApp::GetUnfocusedVolumeMultiplier()
 {
diff --git a/Engine/Source/Runtime/Core/Private/Misc/AppTime.cpp b/Engine/Source/Runtime/Core/Private/Misc/AppTime.cpp
new file mode 100644
--- /dev/null
+++ b/Engine/Source/Runtime/Core/Private/Misc/AppTime.cpp
@@ -0,0 +1,54 @@
+// Copyright Epic Games, Inc. All Rights Reserved.
+
+#include "Misc/AppTime.h"
+#include "CoreGlobals.h"
+
+FAppTime FAppTime::GameThreadInstance;
+thread_local TSharedPtr<const FAppTime> FAppTime::TLSInstance;
+
+[[nodiscard]] CORE_API TSharedPtr<const FAppTime> FAppTime::Fork()
+{
+	if (IsInSlateThread())
+	{
+		// @todo delta time refactor - ignore races with the slate thread for now
+		return MakeShared<FAppTime>(GameThreadInstance);
+	}
+
+	return TLSInstance;
+}
+
+CORE_API void FAppTime::Restore(TSharedPtr<const FAppTime>&& Ptr)
+{
+	TLSInstance = MoveTemp(Ptr);
+}
+
+CORE_API const FAppTime& FAppTime::Get()
+{
+	if (TLSInstance.IsValid())
+	{
+		return *TLSInstance;
+	}
+	else if (IsInSlateThread())
+	{
+		// @todo delta time refactor - ignore races with the slate thread for now
+		return GameThreadInstance;
+	}
+	else
+	{
+		ensureMsgf(IsInGameThread(),
+			TEXT("Attempted to retrieve FAppTime on a thread where there is no inherited time context. "        )
+			TEXT("Copies of the game thread's FAppTime are passed automatically through the TaskGraph via "     )
+			TEXT("FInheritedContextBase / FInheritedContextScope, to ensure the values are pipelined correctly ")
+			TEXT("to tasks that may be running concurrently. The current thread or task has not been directly " )
+			TEXT("or indirectly spawned from the game thread, so it did not receive a time context."            ));
+
+		return GameThreadInstance;
+	}
+}
+
+void FAppTime::Set(const FAppTime& NewValue)
+{
+	checkf(IsInGameThread(), TEXT("Only the game thread may set the current FAppTime value, to ensure parallel tasks receive the correct value via task context inheritance."));
+	GameThreadInstance = NewValue;
+	TLSInstance = MakeShared<FAppTime>(GameThreadInstance);
+}
diff --git a/Engine/Source/Runtime/Core/Public/Async/InheritedContext.h b/Engine/Source/Runtime/Core/Public/Async/InheritedContext.h
--- a/Engine/Source/Runtime/Core/Public/Async/InheritedContext.h
+++ b/Engine/Source/Runtime/Core/Public/Async/InheritedContext.h
@@ -5,6 +5,7 @@
 #include "HAL/LowLevelMemTracker.h"
 #include "ProfilingDebugging/MetadataTrace.h"
 #include "ProfilingDebugging/TagTrace.h"
+#include "Misc/AppTime.h"
 
 namespace UE
 {
@@ -52,64 +53,55 @@ namespace UE
 	// An instance must be obtained by calling `FInheritedContextBase::RestoreInheritedContext()`
 	class FInheritedContextScope
 	{
+	public:
+		~FInheritedContextScope()
+		{
+			FAppTime::Restore(MoveTemp(PrevAppTime));
+		}
+
 	private:
 		UE_NONCOPYABLE(FInheritedContextScope);
 
 		friend class FInheritedContextBase; // allow construction only by `FInheritedContextBase`
 		
-		FInheritedContextScope(
-#if ENABLE_LOW_LEVEL_MEM_TRACKER
-			const FLLMActiveTagsCapture& InInheritedLLMTag
-	#if UE_TRACE_ENABLED && (UE_MEMORY_TAGS_TRACE_ENABLED || UE_TRACE_METADATA_ENABLED)
-			,
-	#endif
-#endif
-#if (UE_MEMORY_TAGS_TRACE_ENABLED && UE_TRACE_ENABLED)
-			int32 InInheritedMemTag
-	#if (UE_TRACE_METADATA_ENABLED && UE_TRACE_ENABLED)
-			,
-	#endif
-#endif
-#if (UE_TRACE_METADATA_ENABLED && UE_TRACE_ENABLED)
-			uint32 InInheritedMetadataId
-#endif
-		)
-#if ENABLE_LOW_LEVEL_MEM_TRACKER
-			: LLMScopes(InInheritedLLMTag)
-	#if UE_TRACE_ENABLED && (UE_MEMORY_TAGS_TRACE_ENABLED || UE_TRACE_METADATA_ENABLED)
-			,
-	#endif
-#endif
-#if (UE_MEMORY_TAGS_TRACE_ENABLED && UE_TRACE_ENABLED)
-	#if !ENABLE_LOW_LEVEL_MEM_TRACKER
-			:
-	#endif
-			MemScope(InInheritedMemTag)
-	#if (UE_TRACE_METADATA_ENABLED && UE_TRACE_ENABLED)
-			,
-	#endif
-#endif
-#if (UE_TRACE_METADATA_ENABLED && UE_TRACE_ENABLED)
-	#if !ENABLE_LOW_LEVEL_MEM_TRACKER && !(UE_MEMORY_TAGS_TRACE_ENABLED && UE_TRACE_ENABLED)
-			:
-	#endif
-			MetaScope(InInheritedMetadataId)
-#endif
+		FInheritedContextScope(TSharedPtr<const FAppTime> AppTime
+		#if ENABLE_LOW_LEVEL_MEM_TRACKER
+			, const FLLMActiveTagsCapture& InInheritedLLMTag
+		#endif
+		#if (UE_MEMORY_TAGS_TRACE_ENABLED && UE_TRACE_ENABLED)
+			, int32 InInheritedMemTag
+		#endif
+		#if (UE_TRACE_METADATA_ENABLED && UE_TRACE_ENABLED) 
+			, uint32 InInheritedMetadataId
+		#endif
+			)
+			: PrevAppTime(FAppTime::Fork())
+		#if ENABLE_LOW_LEVEL_MEM_TRACKER
+			, LLMScopes(InInheritedLLMTag)
+		#endif
+		#if (UE_MEMORY_TAGS_TRACE_ENABLED && UE_TRACE_ENABLED)
+			, MemScope(InInheritedMemTag)
+		#endif
+		#if (UE_TRACE_METADATA_ENABLED && UE_TRACE_ENABLED) 
+			, MetaScope(InInheritedMetadataId)
+		#endif
 		{
+			FAppTime::Restore(MoveTemp(AppTime));
 		}
 
-	private:
-#if ENABLE_LOW_LEVEL_MEM_TRACKER
+		TSharedPtr<const FAppTime> PrevAppTime;
+
+	#if ENABLE_LOW_LEVEL_MEM_TRACKER
 		FLLMActiveTagsScope LLMScopes;
-#endif
+	#endif
 
-#if (UE_MEMORY_TAGS_TRACE_ENABLED && UE_TRACE_ENABLED)
+	#if (UE_MEMORY_TAGS_TRACE_ENABLED && UE_TRACE_ENABLED)
 		FMemScope MemScope;
-#endif
+	#endif
 
-#if (UE_TRACE_METADATA_ENABLED && UE_TRACE_ENABLED) 
+	#if (UE_TRACE_METADATA_ENABLED && UE_TRACE_ENABLED) 
 		FMetadataRestoreScope MetaScope;
-#endif
+	#endif
 	};
 
 	// this class extends the inherited context (see private members for what the inherited context is) to cover async execution. 
@@ -120,6 +112,8 @@ namespace UE
 		// must be called in the inherited context, e.g. on launching an async task
 		void CaptureInheritedContext()
 		{
+			AppTime = FAppTime::Fork();
+
 #if ENABLE_LOW_LEVEL_MEM_TRACKER
 			InheritedLLMTags.CaptureActiveTagData();
 #endif
@@ -137,6 +131,8 @@ namespace UE
 		[[nodiscard]] CORE_API FInheritedContextScope RestoreInheritedContext();
 
 	private:
+		TSharedPtr<const FAppTime> AppTime;
+
 #if ENABLE_LOW_LEVEL_MEM_TRACKER
 		FLLMActiveTagsCapture InheritedLLMTags;
 #endif
diff --git a/Engine/Source/Runtime/Core/Public/Misc/App.h b/Engine/Source/Runtime/Core/Public/Misc/App.h
--- a/Engine/Source/Runtime/Core/Public/Misc/App.h
+++ b/Engine/Source/Runtime/Core/Public/Misc/App.h
@@ -10,6 +10,7 @@
 #include "HAL/PlatformCrt.h"
 #include "HAL/PlatformMisc.h"
 #include "HAL/PlatformProcess.h"
+#include "Misc/AppTime.h"
 #include "Misc/Build.h"
 #include "Misc/CString.h"
 #include "Misc/CommandLine.h"
@@ -555,21 +556,11 @@
 	 *
 	 * @return Time step in seconds for fixed delta time.
 	 */
-	FORCEINLINE static double GetFixedDeltaTime()
+	static double GetFixedDeltaTime()
 	{
-		return FixedDeltaTime.load(std::memory_order_relaxed);
+		return FAppTime::Get().FixedDeltaTime;
 	}
-
+	
 	/**
-	 * Sets time step in seconds if a fixed delta time is wanted.
-	 *
-	 * @param seconds Time step in seconds for fixed delta time.
-	 */
-	static void SetFixedDeltaTime(double Seconds)
-	{
-		FixedDeltaTime.store(Seconds, std::memory_order_relaxed);
-	}
-
-	/**
 	 * Gets whether we want to use a fixed time step or not.
 	 *
@@ -577,58 +568,27 @@
 	 */
 	static bool UseFixedTimeStep()
 	{
-#if WITH_FIXED_TIME_STEP_SUPPORT
-		return bUseFixedTimeStep;
-#else
-		return false;
-#endif
+		return FAppTime::Get().bUseFixedTimeStep;
 	}
 
 	/**
-	 * Enables or disabled usage of fixed time step.
-	 *
-	 * @param whether to use fixed time step or not
-	 */
-	static void SetUseFixedTimeStep(bool bVal)
-	{
-		bUseFixedTimeStep = bVal;
-	}
-
-	/**
 	 * Gets current time in seconds.
 	 *
 	 * @return Current time in seconds.
 	 */
-	FORCEINLINE static double GetCurrentTime()
+	static double GetCurrentTime()
 	{
-		return CurrentTime.load(std::memory_order_relaxed);
+		return FAppTime::Get().CurrentTime;
 	}
 
 	/**
-	 * Sets current time in seconds.
-	 *
-	 * @param seconds - Time in seconds.
-	 */
-	static void SetCurrentTime(double Seconds)
-	{
-		CurrentTime.store(Seconds, std::memory_order_relaxed);
-	}
-
-	/**
 	 * Gets previous value of CurrentTime.
 	 *
 	 * @return Previous value of CurrentTime.
 	 */
-	FORCEINLINE static double GetLastTime()
+	static double GetLastTime()
 	{
-		return LastTime.load(std::memory_order_relaxed);
-	}
-
-	/** Updates Last time to CurrentTime. */
-	static void UpdateLastTime()
-	{
-		// Not an atomic operation, but preferred to a compare and swap
-		LastTime.store(CurrentTime.load(std::memory_order_relaxed), std::memory_order_relaxed);
+		return FAppTime::Get().LastTime;
 	}
 
 	/**
@@ -636,88 +596,51 @@
 	 *
 	 * @return Time delta in seconds.
 	 */
-	FORCEINLINE static double GetDeltaTime()
+	static double GetDeltaTime()
 	{
-		return DeltaTime.load(std::memory_order_relaxed);
+		return FAppTime::Get().DeltaTime;
 	}
 
 	/**
-	 * Sets time delta in seconds.
-	 *
-	 * @param seconds Time in seconds.
-	 */
-	static void SetDeltaTime(double Seconds)
-	{
-		DeltaTime.store(Seconds, std::memory_order_relaxed);
-	}
-
-	/**
 	 * Gets idle time in seconds.
 	 *
 	 * @return Idle time in seconds.
 	 */
-	FORCEINLINE static double GetIdleTime()
+	static double GetIdleTime()
 	{
-		return IdleTime.load(std::memory_order_relaxed);
+		return FAppTime::Get().IdleTime;
 	}
 
 	/**
-	 * Sets idle time in seconds.
-	 *
-	 * @param seconds - Idle time in seconds.
-	 */
-	static void SetIdleTime(double Seconds)
-	{
-		IdleTime.store(Seconds, std::memory_order_relaxed);
-	}
-
-	/**
 	 * Gets overall game time in seconds.
 	 *
 	 * @return Overall game time in seconds.
 	 */
-	FORCEINLINE static double GetGameTime()
+	static double GetGameTime()
 	{
-		return GameTime.load(std::memory_order_relaxed);
+		return FAppTime::Get().GameTime;
 	}
 
 	/**
-	 * Sets overall game time in seconds.
-	 *
-	* @param seconds - Overall game time in seconds.
-	 */
-	static void SetGameTime(double Seconds)
-	{
-		GameTime.store(Seconds, std::memory_order_relaxed);
-	}
-
-	/**
 	 * Gets idle time overshoot in seconds (the time beyond the wait time we requested for the frame). Only valid when IdleTime is > 0.
 	 *
 	 * @return Idle time in seconds.
 	 */
-	FORCEINLINE static double GetIdleTimeOvershoot()
+	static double GetIdleTimeOvershoot()
 	{
-		return IdleTimeOvershoot.load(std::memory_order_relaxed);
+		return FAppTime::Get().IdleTimeOvershoot;
 	}
 
 	/**
-	 * Sets idle time overshoot in seconds (the time beyond the wait time we requested for the frame). Only valid when IdleTime is > 0.
-	 *
-	 * @param seconds - Idle time in seconds.
-	 */
-	static void SetIdleTimeOvershoot(double Seconds)
-	{
-		IdleTimeOvershoot.store(Seconds, std::memory_order_relaxed);
-	}
-
-	/**
 	 * Convert the current frame time into a readable timecode.
 	 * If the current frame time is not set, the timecode value will be defaulted.
 	 *
 	 * @return the current timecode.
 	 */
-	static CORE_API FTimecode GetTimecode();
+	static FTimecode GetTimecode()
+	{
+		return FAppTime::Get().GetTimecode();
+	}
 
 	/**
 	 * Get the frame rate of the current frame time.
@@ -725,7 +648,10 @@
 	 *
 	 * @return the current timecode frame rate.
 	 */
-	static CORE_API FFrameRate GetTimecodeFrameRate();
+	static FFrameRate GetTimecodeFrameRate()
+	{
+		return FAppTime::Get().GetTimecodeFrameRate();
+	}
 
 	/**
 	 * Gets a frame number generated by the engine's timecode provider.
@@ -734,21 +660,104 @@
 	 *
 	 * @return the current frame time.
 	 */
-	FORCEINLINE static TOptional<FQualifiedFrameTime> GetCurrentFrameTime()
+	static TOptional<FQualifiedFrameTime> GetCurrentFrameTime()
 	{
-		return CurrentFrameTime;
+		return FAppTime::Get().CurrentFrameTime;
 	}
 
 	/**
-	 * Sets the current timecode, and the frame rate to which it's relative.
+	 * Sets time step in seconds if a fixed delta time is wanted.
 	 *
-	 * @param InTimecode - current timecode.
-	 * @param InTimecodeFrameRate - current timecode framerate.
+	 * @param seconds Time step in seconds for fixed delta time.
 	 */
-	UE_DEPRECATED(4.25, "Please use SetQualifiedFrameTime")
-	static void SetTimecodeAndFrameRate(FTimecode InTimecode, FFrameRate InTimecodeFrameRate)
+	// UE_DEPRECATED(5.8, "@todo delta time refactor - deprecate this function")
+	static void SetFixedDeltaTime(double Seconds)
 	{
-		CurrentFrameTime = FQualifiedFrameTime(InTimecode, InTimecodeFrameRate);
+		FAppTime Time = FAppTime::Get();
+		Time.FixedDeltaTime = Seconds;
+		FAppTime::Set(Time);
 	}
 
 	/**
+	 * Enables or disabled usage of fixed time step.
+	 *
+	 * @param whether to use fixed time step or not
+	 */
+	// UE_DEPRECATED(5.8, "@todo delta time refactor - deprecate this function")
+	static void SetUseFixedTimeStep(bool bVal)
+	{
+		FAppTime Time = FAppTime::Get();
+		Time.bUseFixedTimeStep = bVal;
+		FAppTime::Set(Time);
+	}
+	/**
+	 * Sets current time in seconds.
+	 *
+	 * @param seconds - Time in seconds.
+	 */
+	// UE_DEPRECATED(5.8, "@todo delta time refactor - deprecate this function")
+	static void SetCurrentTime(double Seconds)
+	{
+		FAppTime Time = FAppTime::Get();
+		Time.CurrentTime = Seconds;
+		FAppTime::Set(Time);
+	}
+
+	/** Updates Last time to CurrentTime. */
+	// UE_DEPRECATED(5.8, "@todo delta time refactor - deprecate this function")
+	static void UpdateLastTime()
+	{
+		FAppTime Time = FAppTime::Get();
+		Time.LastTime = Time.CurrentTime;
+		FAppTime::Set(Time);
+	}
+	/**
+	 * Sets time delta in seconds.
+	 *
+	 * @param seconds Time in seconds.
+	 */
+	// UE_DEPRECATED(5.8, "@todo delta time refactor - deprecate this function")
+	static void SetDeltaTime(double Seconds)
+	{
+		FAppTime Time = FAppTime::Get();
+		Time.DeltaTime = Seconds;
+		FAppTime::Set(Time);
+	}
+	/**
+	 * Sets idle time in seconds.
+	 *
+	 * @param seconds - Idle time in seconds.
+	 */
+	// UE_DEPRECATED(5.8, "@todo delta time refactor - deprecate this function")
+	static void SetIdleTime(double Seconds)
+	{
+		FAppTime Time = FAppTime::Get();
+		Time.IdleTime = Seconds;
+		FAppTime::Set(Time);
+	}
+	/**
+	 * Sets overall game time in seconds.
+	 *
+	 * @param seconds - Overall game time in seconds.
+	 */
+	// UE_DEPRECATED(5.8, "@todo delta time refactor - deprecate this function")
+	static void SetGameTime(double Seconds)
+	{
+		FAppTime Time = FAppTime::Get();
+		Time.GameTime = Seconds;
+		FAppTime::Set(Time);
+	}
+	/**
+	 * Sets idle time overshoot in seconds (the time beyond the wait time we requested for the frame). Only valid when IdleTime is > 0.
+	 *
+	 * @param seconds - Idle time in seconds.
+	 */
+	// UE_DEPRECATED(5.8, "@todo delta time refactor - deprecate this function")
+	static void SetIdleTimeOvershoot(double Seconds)
+	{
+		FAppTime Time = FAppTime::Get();
+		Time.IdleTimeOvershoot = Seconds;
+		FAppTime::Set(Time);
+	}
+
+	/**
@@ -756,15 +765,21 @@
 	 *
 	 * @param InFrameTime - current frame time and frame rate.
 	 */
+	// UE_DEPRECATED(5.8, "@todo delta time refactor - deprecate this function")
 	static void SetCurrentFrameTime(FQualifiedFrameTime InFrameTime)
 	{
-		CurrentFrameTime = InFrameTime;
+		FAppTime Time = FAppTime::Get();
+		Time.CurrentFrameTime = InFrameTime;
+		FAppTime::Set(Time);
 	}
 
 	/** Invalidate the current frame time. It will reset the TOptional. */
+	// UE_DEPRECATED(5.8, "@todo delta time refactor - deprecate this function")
 	static void InvalidateCurrentFrameTime()
 	{
-		CurrentFrameTime.Reset();
+		FAppTime Time = FAppTime::Get();
+		Time.CurrentFrameTime.Reset();
+		FAppTime::Set(Time);
 	}
 
 	/**
@@ -877,33 +892,6 @@
 	/** Holds a flag Whether we are in benchmark mode or not. */
 	static CORE_API bool bIsBenchmarking;
 
-	/** Holds a flag whether we want to use a fixed time step or not. */
-	static CORE_API bool bUseFixedTimeStep;
-
-	/** Holds time step if a fixed delta time is wanted. */
-	static CORE_API std::atomic<double> FixedDeltaTime;
-
-	/** Holds current time. */
-	static CORE_API std::atomic<double> CurrentTime;
-
-	/** Holds previous value of CurrentTime. */
-	static CORE_API std::atomic<double> LastTime;
-
-	/** Holds current delta time in seconds. */
-	static CORE_API std::atomic<double> DeltaTime;
-
-	/** Holds time we spent sleeping in UpdateTimeAndHandleMaxTickRate() if our frame time was smaller than one allowed by target FPS. */
-	static CORE_API std::atomic<double> IdleTime;
-
-	/** Holds the amount of IdleTime that was LONGER than we tried to sleep. The OS can't sleep the exact amount of time, so this measures that overshoot. */
-	static CORE_API std::atomic<double> IdleTimeOvershoot;
-
-	/** Holds overall game time. */
-	static CORE_API std::atomic<double> GameTime;
-
-	/** Holds the current frame time and framerate. */
-	static CORE_API TOptional<FQualifiedFrameTime> CurrentFrameTime;
-
 	/** Use to affect the app volume when it loses focus */
 	static CORE_API float VolumeMultiplier;
 
diff --git a/Engine/Source/Runtime/Core/Public/Misc/AppTime.h b/Engine/Source/Runtime/Core/Public/Misc/AppTime.h
new file mode 100644
--- /dev/null
+++ b/Engine/Source/Runtime/Core/Public/Misc/AppTime.h
@@ -0,0 +1,90 @@
+// Copyright Epic Games, Inc. All Rights Reserved.
+
+#pragma once
+
+#include "Misc/FrameRate.h"
+#include "Misc/Optional.h"
+#include "Misc/QualifiedFrameTime.h"
+#include "Misc/Timecode.h"
+#include "Templates/SharedPointer.h"
+
+namespace UE
+{
+	class FInheritedContextScope;
+	class FInheritedContextBase;
+}
+
+/**
+ * Provides global game timing information (delta time, current time etc).
+ * The current time is set each frame by the game thread. It may then be retrieved by the game thread and other task threads via FAppTime::Get().
+ * The value from the game thread is pipelined through the TaskGraph such that child tasks receive the correct value for the frame the tasks were spawned in.
+ */
+class FAppTime
+{
+	friend UE::FInheritedContextScope;
+	friend UE::FInheritedContextBase;
+
+public:
+#if WITH_FIXED_TIME_STEP_SUPPORT
+	/** Holds a flag whether we want to use a fixed time step or not. */
+	bool bUseFixedTimeStep = false;
+#else
+	static constexpr bool bUseFixedTimeStep = false;
+#endif
+
+	/** Holds time step if a fixed delta time is wanted. */
+	double FixedDeltaTime = 1 / 30.0;
+
+	/** Holds current time. */
+	double CurrentTime = 0.0;
+
+	/** Holds previous value of CurrentTime. */
+	double LastTime = 0.0;
+
+	/** Holds current delta time in seconds. */
+	double DeltaTime = 1 / 30.0;
+
+	/** Holds time we spent sleeping in UpdateTimeAndHandleMaxTickRate() if our frame time was smaller than one allowed by target FPS. */
+	double IdleTime = 0.0;
+
+	/** Holds the amount of IdleTime that was LONGER than we tried to sleep. The OS can't sleep the exact amount of time, so this measures that overshoot. */
+	double IdleTimeOvershoot = 0.0;
+
+	/** Holds overall game time. */
+	double GameTime = 0.0;
+
+	/** Holds the current frame time and framerate. */
+	TOptional<FQualifiedFrameTime> CurrentFrameTime;
+
+	FTimecode GetTimecode() const
+	{
+		if (CurrentFrameTime.IsSet())
+		{
+			return CurrentFrameTime->ToTimecode();
+		}
+
+		return FTimecode();
+	}
+
+	FFrameRate GetTimecodeFrameRate() const
+	{
+		if (CurrentFrameTime.IsSet())
+		{
+			return CurrentFrameTime->Rate;
+		}
+
+		return FFrameRate();
+	}
+
+	static CORE_API const FAppTime& Get();
+
+	// May only be called from the game thread
+	static CORE_API void Set(const FAppTime& NewValue);
+
+private:
+	static FAppTime GameThreadInstance;
+	static thread_local TSharedPtr<const FAppTime> TLSInstance;
+
+	[[nodiscard]] static CORE_API TSharedPtr<const FAppTime> Fork();
+	static CORE_API void Restore(TSharedPtr<const FAppTime>&& Ptr);
+};
