ParkingLot: Assert on re-entrant calls

#Synced-CL 32500503

diff --git a/Engine/Source/Runtime/Core/Private/Async/ParkingLot.cpp b/Engine/Source/Runtime/Core/Private/Async/ParkingLot.cpp
--- a/Engine/Source/Runtime/Core/Private/Async/ParkingLot.cpp
+++ b/Engine/Source/Runtime/Core/Private/Async/ParkingLot.cpp
@@ -15,6 +15,33 @@
 namespace UE::ParkingLot::Private
 {
 
+/**
+ * Identifies which callback function is currently executing.
+ *
+ * NOTE: ORDER MATTERS FOR COMPARISONS!
+ */
+enum class EExecutingFunction : uint8
+{
+	None = 0,
+	BeforeWait = 1,
+	CanWait = 2,
+	OnWakeState = 3,
+};
+
+const TCHAR* LexToString(EExecutingFunction Function)
+{
+	switch (Function)
+	{
+	default:
+	case EExecutingFunction::None: return TEXT("None");
+	case EExecutingFunction::BeforeWait: return TEXT("BeforeWait");
+	case EExecutingFunction::CanWait: return TEXT("CanWait");
+	case EExecutingFunction::OnWakeState: return TEXT("OnWakeState");
+	}
+}
+
+static thread_local EExecutingFunction ParkingLotExecutingFunction = EExecutingFunction::None;
+
 ///////////////////////////////////////////////////////////////////////////////////////////////////////////////////////
 
 /**
@@ -51,10 +78,6 @@ struct alignas(PLATFORM_CACHE_LINE_SIZE) FThread final
 class FThreadLocalData
 {
 public:
-	// TODO: This version reserves the table as threads spawn. The function-local thread_local
-	//       reserves the table as threads wait for the first time. Not sure which is better.
-	//static thread_local FThreadLocalData ThreadLocalData;
-
 	static TRefCountPtr<FThread> Get()
 	{
 		static thread_local FThreadLocalData ThreadLocalData;
@@ -84,8 +107,6 @@ private:
 	FThreadLocalData& operator=(const FThreadLocalData&) = delete;
 };
 
-//thread_local FThreadLocalData FThreadLocalData::ThreadLocalData;
-
 ///////////////////////////////////////////////////////////////////////////////////////////////////////////////////////
 
 enum class EQueueAction
@@ -505,6 +526,11 @@ FWaitState Wait(const void* Address, bool (*CanWait)(void*), void* CanWaitContex
 {
 	using namespace Private;
 
+	// A failure of this assertion is either an error in a callback, because callbacks may not call into ParkingLot,
+	// or a dependency on ParkingLot from within the memory allocator used by ParkingLot.
+	checkf(ParkingLotExecutingFunction == EExecutingFunction::None,
+		TEXT("Detected a re-entrant call into ParkingLot from within a %s callback."), LexToString(ParkingLotExecutingFunction));
+
 	TRefCountPtr<FThread> Self = FThreadLocalData::Get();
 
 	checkfSlow(!Self->WaitAddress.load(std::memory_order_relaxed), TEXT("WaitAddress must be null. This can happen if Wait is called by BeforeWait."));
@@ -516,11 +542,15 @@ FWaitState Wait(const void* Address, bool (*CanWait)(void*), void* CanWaitContex
 	{
 		TDynamicUniqueLock<FWordMutex> BucketLock;
 		FBucket& Bucket = FTable::FindOrCreateBucket(Address, BucketLock);
-		State.bDidWait = !CanWait || CanWait(CanWaitContext);
-		if (!State.bDidWait)
+		if (CanWait)
 		{
-			return State;
+			TGuardValue GuardActive(ParkingLotExecutingFunction, EExecutingFunction::CanWait);
+			if (!CanWait(CanWaitContext))
+			{
+				return State;
+			}
 		}
+		State.bDidWait = true;
 		Self->WaitAddress.store(Address, std::memory_order_relaxed);
 		Self->Event.Reset();
 		Bucket.Enqueue(&*Self);
@@ -529,6 +559,7 @@ FWaitState Wait(const void* Address, bool (*CanWait)(void*), void* CanWaitContex
 	// BeforeWait must be invoked after the bucket is unlocked.
 	if (BeforeWait)
 	{
+		TGuardValue GuardActive(ParkingLotExecutingFunction, EExecutingFunction::BeforeWait);
 		BeforeWait(BeforeWaitContext);
 	}
 
@@ -536,7 +567,7 @@ FWaitState Wait(const void* Address, bool (*CanWait)(void*), void* CanWaitContex
 	Self->Event.Wait();
 
 	// WaitAddress is reset when the thread is dequeued.
-	verifySlow(!Self->WaitAddress.load(std::memory_order_relaxed));
+	checkSlow(!Self->WaitAddress.load(std::memory_order_relaxed));
 	State.bDidWake = true;
 	State.WakeToken = Self->WakeToken;
 	Self->WakeToken = 0;
@@ -554,6 +585,11 @@ static FWaitState TimedWait(
 {
 	using namespace Private;
 
+	// A failure of this assertion is either an error in a callback, because callbacks may not call into ParkingLot,
+	// or a dependency on ParkingLot from within the memory allocator used by ParkingLot.
+	checkf(ParkingLotExecutingFunction == EExecutingFunction::None,
+		TEXT("Detected a re-entrant call into ParkingLot from within a %s callback."), LexToString(ParkingLotExecutingFunction));
+
 	TRefCountPtr<FThread> Self = FThreadLocalData::Get();
 
 	checkfSlow(!Self->WaitAddress.load(std::memory_order_relaxed), TEXT("WaitAddress must be null. This can happen if Wait is called by BeforeWait."));
@@ -565,11 +601,15 @@ static FWaitState TimedWait(
 	{
 		TDynamicUniqueLock<FWordMutex> BucketLock;
 		FBucket& Bucket = FTable::FindOrCreateBucket(Address, BucketLock);
-		State.bDidWait = !CanWait || CanWait(CanWaitContext);
-		if (!State.bDidWait)
+		if (CanWait)
 		{
-			return State;
+			TGuardValue GuardActive(ParkingLotExecutingFunction, EExecutingFunction::CanWait);
+			if (!CanWait(CanWaitContext))
+			{
+				return State;
+			}
 		}
+		State.bDidWait = true;
 		Self->WaitAddress.store(Address, std::memory_order_relaxed);
 		Self->Event.Reset();
 		Bucket.Enqueue(&*Self);
@@ -578,6 +618,7 @@ static FWaitState TimedWait(
 	// BeforeWait must be invoked after the bucket is unlocked.
 	if (BeforeWait)
 	{
+		TGuardValue GuardActive(ParkingLotExecutingFunction, EExecutingFunction::BeforeWait);
 		BeforeWait(BeforeWaitContext);
 	}
 
@@ -652,6 +693,11 @@ void WakeOne(const void* Address, uint64 (*OnWakeState)(void*, FWakeState), void
 {
 	using namespace Private;
 
+	// A failure of this assertion is either an error in a callback, because callbacks may not call into ParkingLot,
+	// or a dependency on ParkingLot from within the memory allocator used by ParkingLot.
+	checkf(ParkingLotExecutingFunction <= EExecutingFunction::BeforeWait,
+		TEXT("Detected a re-entrant call into ParkingLot from within a %s callback."), LexToString(ParkingLotExecutingFunction));
+
 	TRefCountPtr<FThread> WakeThread;
 	uint64 WakeToken = 0;
 
@@ -667,6 +713,7 @@ void WakeOne(const void* Address, uint64 (*OnWakeState)(void*, FWakeState), void
 			}
 			return EQueueAction::Continue;
 		});
+		TGuardValue GuardActive(ParkingLotExecutingFunction, EExecutingFunction::OnWakeState);
 		FWakeState WakeState;
 		WakeState.bDidWake = !!WakeThread;
 		WakeState.bHasWaitingThreads = !Bucket.IsEmpty();
@@ -704,6 +751,11 @@ uint32 WakeMultiple(const void* const Address, const uint32 WakeCount)
 {
 	using namespace Private;
 
+	// A failure of this assertion is either an error in a callback, because callbacks may not call into ParkingLot,
+	// or a dependency on ParkingLot from within the memory allocator used by ParkingLot.
+	checkf(ParkingLotExecutingFunction <= EExecutingFunction::BeforeWait,
+		TEXT("Detected a re-entrant call into ParkingLot from within a %s callback."), LexToString(ParkingLotExecutingFunction));
+
 	TArray<TRefCountPtr<FThread>, TInlineAllocator<128>> WakeThreads;
 
 	if (TDynamicUniqueLock<FWordMutex> BucketLock; FBucket* Bucket = FTable::FindBucket(Address, BucketLock))
