Synchronize access to FIoContainerMetaWriter

#rnx
#Synced-CL 48026973

diff --git a/Engine/Source/Runtime/Core/Internal/IO/IoContainerMeta.h b/Engine/Source/Runtime/Core/Internal/IO/IoContainerMeta.h
--- a/Engine/Source/Runtime/Core/Internal/IO/IoContainerMeta.h
+++ b/Engine/Source/Runtime/Core/Internal/IO/IoContainerMeta.h
@@ -3,6 +3,7 @@
 #pragma once
 
 #include "Async/MappedFileHandle.h"
+#include "Async/Mutex.h"
 #include "Containers/Array.h"
 #include "Containers/ArrayView.h"
 #include "Containers/ChunkedArray.h"
@@ -173,4 +174,5 @@ private:
 	TChunkedArray<FUtf8String>			Strings;
 	TArray<FIoMetaDirectoryIndexEntry>	DirectoryEntries;
 	TArray<FIoMetaFileIndexEntry>		FileEntries;
+	UE::FMutex							Mutex;
 };
diff --git a/Engine/Source/Runtime/Core/Private/IO/IoContainerMeta.cpp b/Engine/Source/Runtime/Core/Private/IO/IoContainerMeta.cpp
--- a/Engine/Source/Runtime/Core/Private/IO/IoContainerMeta.cpp
+++ b/Engine/Source/Runtime/Core/Private/IO/IoContainerMeta.cpp
@@ -3,6 +3,7 @@
 #include "IO/IoContainerMeta.h"
 
 #include "Algo/BinarySearch.h"
+#include "Async/UniqueLock.h"
 #include "HAL/FileManager.h"
 #include "HAL/PlatformFileManager.h"
 #include "Misc/PathViews.h"
@@ -203,6 +204,8 @@ FIoContainerMetaWriter::FIoContainerMetaWriter()
 
 void FIoContainerMetaWriter::AddFile(FStringView ContainerName, const FIoChunkId& ChunkId, FStringView Filename)
 {
+	UE::TUniqueLock Lock(Mutex);
+
 	// Only store filenames for packages for now
 	if (ChunkIdToFileEntry.Contains(ChunkId) || ChunkId.GetChunkType() != EIoChunkType::ExportBundleData)
 	{
@@ -218,6 +221,8 @@ void FIoContainerMetaWriter::AddFile(FStringView ContainerName, const FIoChunkId
 
 void FIoContainerMetaWriter::Save(FArchive& Ar)
 {
+	UE::TUniqueLock Lock(Mutex);
+
 	check(ChunkIdToFileEntry.Num() == FileEntries.Num());
 
 	TArray<FIoChunkId>	ChunkIds;
