Added UE::CVoid concept.

diff --git a/Engine/Source/Runtime/Core/Public/Concepts/Void.h b/Engine/Source/Runtime/Core/Public/Concepts/Void.h
new file mode 100644
--- /dev/null
+++ b/Engine/Source/Runtime/Core/Public/Concepts/Void.h
@@ -0,0 +1,14 @@
+// Copyright Epic Games, Inc. All Rights Reserved.
+
+#pragma once
+
+#include <type_traits>
+
+namespace UE
+{
+	/**
+	 * Concept which describes a void type.
+	 */
+	template <typename T>
+	concept CVoid = std::is_void_v<T>;
+}
diff --git a/Engine/Source/Runtime/Core/Public/Memory/MemoryFwd.h b/Engine/Source/Runtime/Core/Public/Memory/MemoryFwd.h
--- a/Engine/Source/Runtime/Core/Public/Memory/MemoryFwd.h
+++ b/Engine/Source/Runtime/Core/Public/Memory/MemoryFwd.h
@@ -2,9 +2,11 @@
 
 #pragma once
 
+#include "Concepts/Void.h"
+
 // MemoryView
 
-template <typename DataType>
+template <UE::CVoid DataType>
 class TMemoryView;
 
 /** A non-owning view of a contiguous region of memory. */
diff --git a/Engine/Source/Runtime/Core/Public/Memory/MemoryView.h b/Engine/Source/Runtime/Core/Public/Memory/MemoryView.h
--- a/Engine/Source/Runtime/Core/Public/Memory/MemoryView.h
+++ b/Engine/Source/Runtime/Core/Public/Memory/MemoryView.h
@@ -24,11 +24,9 @@
  *
  * Functions that modify a view clamp sizes and offsets to always return a sub-view of the input.
  */
-template <typename DataType>
+template <UE::CVoid DataType>
 class TMemoryView
 {
-	static_assert(std::is_void_v<DataType>, "DataType must be cv-qualified void");
-
 	using ByteType = std::conditional_t<std::is_const_v<DataType>, const uint8, uint8>;
 
 public:
@@ -231,7 +229,7 @@ private:
 		return reinterpret_cast<ByteType*>(Data) + InOffset;
 	}
 
-	template <typename OtherDataType>
+	template <UE::CVoid OtherDataType>
 	friend class TMemoryView;
 
 private:
