Improve documentation for IModuleInterface and FModuleManager

#Synced-CL 48050961

diff --git a/Engine/Source/Runtime/Core/Public/Modules/ModuleInterface.h b/Engine/Source/Runtime/Core/Public/Modules/ModuleInterface.h
--- a/Engine/Source/Runtime/Core/Public/Modules/ModuleInterface.h
+++ b/Engine/Source/Runtime/Core/Public/Modules/ModuleInterface.h
@@ -1,27 +1,47 @@
 // Copyright Epic Games, Inc. All Rights Reserved.
 
-
 #pragma once
 
 #include "CoreTypes.h"
 
-
 /**
- * Interface class that all module implementations should derive from.  This is used to initialize
- * a module after it's been loaded, and also to clean it up before the module is unloaded.
+ * Interface class that all module implementations should derive from. This is used to initialize a module after it's
+ * been loaded, and also to clean it up before the module is unloaded. Callbacks are invoked at the following times:
+ *
+ * Operation                                         | Startup | PostLoad | PreUnload | Shutdown |
+ * --------------------------------------------------|---------|----------|-----------|----------|
+ * Engine Startup                                    | Yes     | No       | ----      | ----     |
+ * Engine Shutdown                                   | ----    | ----     | Yes       | Yes      |
+ * Hot Reload (Legacy)                               | Yes     | Yes      | Yes       | Yes      |
+ * Live Coding                                       | No      | No       | No        | No       |
+ * CCmd Load                                         | Yes     | Yes      | ----      | ----     |
+ * CCmd Unload                                       | ----    | ----     | Yes       | Yes      |
+ * CCmd Reload                                       | Yes     | Yes      | Yes       | Yes      |
+ * FModuleManager::LoadModule                        | Yes     | No       | ----      | ----     |
+ * FModuleManager::LoadModuleChecked                 | Yes     | No       | ----      | ----     |
+ * FModuleManager::LoadModuleWithCallback            | Yes     | Yes      | ----      | ----     |
+ * FModuleManager::LoadModulePtr                     | Yes     | No       | ----      | ----     |
+ * FModuleManager::LoadModuleWithFailureReason       | Yes     | No       | ----      | ----     |
+ * FModuleManager::UnloadModule                      | ----    | ----     | No        | Yes      |
+ * FModuleManager::UnloadOrAbandonModuleWithCallback | ----    | ----     | Yes       | Yes      |
+ * FModuleManager::AbandonModule                     | ----    | ----     | No        | Yes      |
+ * FModuleManager::AbandonModuleWithCallback         | ----    | ----     | Yes       | Yes      |
  */
 class IModuleInterface
 {
-
 public:
 
 	/**
-	 * Note: Even though this is an interface class we need a virtual destructor here because modules are deleted via a pointer to this interface                   
+	 * Note: Even though this is an interface class we need a virtual destructor here because modules are deleted via a
+	 * pointer to this interface.
 	 */
 	virtual ~IModuleInterface() = default;
 
 	/**
-	 * Called right after the module DLL has been loaded and the module object has been created
+	 * Called after the module is loaded. Occurs in all loading situations, including engine startup, loading or
+	 * reloading via console commands, hot reloading, and explicit requests through @ref FModuleManager public API.
+	 * @ref FModuleManager::CurrentModule is guaranteed to be set when this is called.
+	 *
 	 * Load dependent modules here, and they will be guaranteed to be available during ShutdownModule. ie:
 	 *
 	 * FModuleManager::Get().LoadModuleChecked(TEXT("HTTP"));
@@ -31,24 +51,30 @@ public:
 	}
 
 	/**
-	 * Called before the module has been unloaded
+	 * Called before the module is unloaded. Occurs during engine shutdown, hot reloading, and unloading or reloading
+	 * via console commands. During engine shutdown this is called for all modules before ShutdownModule is called for
+	 * any module. During engine shutdown, this is called in reverse order that modules finish StartupModule. Is not
+	 * called for explicit @ref FModuleManager::UnloadModule and @ref FModuleManager::AbandonModule requests.
 	 */
 	virtual void PreUnloadCallback()
 	{
 	}
 
 	/**
-	 * Called after the module has been reloaded
+	 * Called after the module has been reloaded. Occurs during hot reloading and loading or reloading via console
+	 * commands. Not called during engine startup or explicit requests through @ref FModuleManager public API, other
+	 * than @ref FModuleManager::LoadModuleWithCallback.
 	 */
 	virtual void PostLoadCallback()
 	{
 	}
 
 	/**
-	 * Called before the module is unloaded, right before the module object is destroyed.
-	 * During normal shutdown, this is called in reverse order that modules finish StartupModule().
-	 * This means that, as long as a module references dependent modules in it's StartupModule(), it
-	 * can safely reference those dependencies in ShutdownModule() as well.
+	 * Called before the module is unloaded. Occurs in all unloading situations, including hot reloading, unloading or
+	 * reloading via console commands, explicit requests through @ref FModuleManager public API, and engine shutdown.
+	 * During engine shutdown, this is called in reverse order that modules finish StartupModule. This means that, as
+	 * long as a module references dependent modules in it's StartupModule, it can safely reference those dependencies
+	 * in ShutdownModule as well.
 	 */
 	virtual void ShutdownModule()
 	{
@@ -57,7 +83,7 @@ public:
 	/**
 	 * Override this to set whether your module is allowed to be unloaded on the fly
 	 *
-	 * @return	Whether the module supports shutdown separate from the rest of the engine.
+	 * @return Whether the module supports shutdown separate from the rest of the engine.
 	 */
 	virtual bool SupportsDynamicReloading()
 	{
@@ -67,7 +93,7 @@ public:
 	/**
 	 * Override this to set whether your module would like cleanup on application shutdown
 	 *
-	 * @return	Whether the module supports shutdown on application exit
+	 * @return Whether the module supports shutdown on application exit
 	 */
 	virtual bool SupportsAutomaticShutdown()
 	{
@@ -84,5 +110,3 @@ public:
 		return false;
 	}
 };
-
-
diff --git a/Engine/Source/Runtime/Core/Public/Modules/ModuleManager.h b/Engine/Source/Runtime/Core/Public/Modules/ModuleManager.h
--- a/Engine/Source/Runtime/Core/Public/Modules/ModuleManager.h
+++ b/Engine/Source/Runtime/Core/Public/Modules/ModuleManager.h
@@ -163,8 +163,6 @@ namespace UE::Core::Private
 }
 
 /**
- * Implements the module manager.
- *
  * The module manager is used to load and unload modules, as well as to keep track of all of the
  * modules that are currently loaded. You can access this singleton using FModuleManager::Get().
  */
@@ -192,6 +190,7 @@ public:
 
 	/**
 	 * Abandons a loaded module, leaving it loaded in memory but no longer tracking it in the module manager.
+	 * Does not call @ref IModuleInterface::PreUnloadCallback. Calls @ref IModuleInterface::ShutdownModule.
 	 *
 	 * @param InModuleName The name of the module to abandon.  Should not include path, extension or platform/configuration info.  This is just the "module name" part of the module file name.
 	 * @see IsModuleLoaded, LoadModule, LoadModuleWithFailureReason, UnloadModule
@@ -239,6 +238,7 @@ public:
 
 	/**
 	 * Loads the specified module.
+	 * Calls @ref IModuleInterface::StartupModule. Does not call @ref IModuleInterface::PostLoadCallback.
 	 *
 	 * @param InModuleName The base name of the module file.  Should not include path, extension or platform/configuration info.  This is just the "module name" part of the module file name.  Names should be globally unique.
 	 * @param InLoadModuleFlags Optional flags for module load operation.
@@ -251,13 +251,14 @@ public:
 	 * Loads the specified module, checking to ensure it exists.
 	 *
 	 * @param InModuleName The base name of the module file.  Should not include path, extension or platform/configuration info.  This is just the "module name" part of the module file name.  Names should be globally unique.
-	 * @return The loaded module, or nullptr if the load operation failed.
+	 * @return The loaded module.
 	 * @see AbandonModule, IsModuleLoaded, LoadModuleChecked, LoadModulePtr, LoadModuleWithFailureReason, UnloadModule
 	 */
 	CORE_API IModuleInterface& LoadModuleChecked( const FName InModuleName );
 
 	/**
 	 * Loads a module in memory then calls PostLoad.
+	 * Calls @ref IModuleInterface::StartupModule and @ref IModuleInterface::PostLoadCallback.
 	 *
 	 * @param InModuleName The name of the module to load.
 	 * @param Ar The archive to receive error messages, if any.
@@ -268,6 +269,7 @@ public:
 
 	/**
 	 * Loads the specified module and returns a result.
+	 * Calls @ref IModuleInterface::StartupModule. Does not call @ref IModuleInterface::PostLoadCallback.
 	 *
 	 * @param InModuleName The base name of the module file.  Should not include path, extension or platform/configuration info.  This is just the "module name" part of the module file name.  Names should be globally unique.
 	 * @param OutFailureReason Will contain the result.
@@ -297,6 +299,7 @@ public:
 
 	/**
 	 * Unloads a specific module
+	 * Does not call @ref IModuleInterface::PreUnloadCallback. Calls @ref IModuleInterface::ShutdownModule.
 	 * NOTE: You can manually unload a module before the normal shutdown occurs with this, but be careful as you may be unloading another
 	 * module's dependency too early!
 	 *
@@ -310,6 +313,7 @@ public:
 
 	/**
 	 * Calls PreUnload then either unloads or abandons a module in memory, depending on whether the module supports unloading.
+	 * Calls @ref IModuleInterface::PreUnloadCallback and @ref IModuleInterface::ShutdownModule.
 	 *
 	 * @param InModuleName The name of the module to unload.
 	 * @param Ar The archive to receive error messages, if any.
@@ -319,6 +323,7 @@ public:
 
 	/**
 	 * Calls PreUnload then abandons a module in memory.
+	 * Calls @ref IModuleInterface::PreUnloadCallback and @ref IModuleInterface::ShutdownModule.
 	 *
 	 * @param InModuleName The name of the module to unload.
 	 * @see LoadModuleWithCallback
@@ -369,15 +374,16 @@ public:
 	}
 
 	/**
-	  * Loads a module by name, checking to ensure it exists.
-	  *
-	  * This method checks whether the module actually exists. If the module does not exist, an assertion will be triggered.
-	  * If the module was already loaded previously, the existing instance will be returned.
-	  *
-	  * @param ModuleName The module to find and load
-	  * @return	Returns the module interface, casted to the specified typename
-	  * @see GetModulePtr, LoadModulePtr, LoadModuleChecked
-	  */
+	 * Loads a module by name, checking to ensure it exists.
+	 * Calls @ref IModuleInterface::StartupModule. Does not call @ref IModuleInterface::PostLoadCallback.
+	 *
+	 * This method checks whether the module actually exists. If the module does not exist, an assertion will be triggered.
+	 * If the module was already loaded previously, the existing instance will be returned.
+	 *
+	 * @param ModuleName The module to find and load
+	 * @return Returns the module interface, casted to the specified typename
+	 * @see GetModulePtr, LoadModulePtr, LoadModuleChecked
+	 */
 	template<typename TModuleInterface>
 	static TModuleInterface& LoadModuleChecked( const FName InModuleName)
 	{
@@ -386,12 +392,13 @@ public:
 	}
 
 	/**
-	  * Loads a module by name.
-	  *
-	  * @param ModuleName The module to find and load.
-	  * @return The interface to the module, or nullptr if the module was not found.
-	  * @see GetModulePtr, GetModuleChecked, LoadModuleChecked
-	  */
+	 * Loads a module by name.
+	 * Calls @ref IModuleInterface::StartupModule. Does not call @ref IModuleInterface::PostLoadCallback.
+	 *
+	 * @param ModuleName The module to find and load.
+	 * @return The interface to the module, or nullptr if the module was not found.
+	 * @see GetModulePtr, GetModuleChecked, LoadModuleChecked
+	 */
 	template<typename TModuleInterface>
 	static TModuleInterface* LoadModulePtr( const FName InModuleName)
 	{
@@ -431,6 +438,7 @@ public:
 	 * that B will actually get Unloaded after A during shutdown. This allows A's ShutdownModule() call to still reference module B.
 	 * You can manually unload a module yourself which will change this ordering, but be careful as you may be unloading another 
 	 * module's dependency!
+	 * Calls @ref IModuleInterface::PreUnloadCallback and @ref IModuleInterface::ShutdownModule.
 	 *
 	 * This method is Usually called at various points while exiting an application.
 	 */
@@ -569,7 +577,7 @@ public:
 	/**
 	 * Gets a delegate that is executed when a module containing UObjects has been loaded.
 	 *
-	 * The first parameter is the name of the loaded module.
+	 * The first parameter is the name of script package associated with the module.
 	 *
 	 * @return The event delegate.
 	 */
@@ -725,8 +733,7 @@ private:
 
 #endif // UE_MERGED_MODULES
 
-
-#endif
+#endif // !IS_MONOLITHIC
 
 	/** Adds pending module initializer registrations to the StaticallyLinkedModuleInitializers map. */
 	CORE_API void ProcessPendingStaticallyLinkedModuleInitializers() const;
