Boot profiler: avoid all string processing if we're not capturing.
Csv profiler: Make IsCapturing static to avoid singleton overhead
#Synced-CL Unknown

diff --git a/Engine/Source/Runtime/Core/Private/Misc/CoreGlobals.cpp b/Engine/Source/Runtime/Core/Private/Misc/CoreGlobals.cpp
--- a/Engine/Source/Runtime/Core/Private/Misc/CoreGlobals.cpp
+++ b/Engine/Source/Runtime/Core/Private/Misc/CoreGlobals.cpp
@@ -536,7 +536,7 @@ FScopedBootTiming::FScopedBootTiming(const ANSICHAR *InMessage)
 {
 	UE_LOGFMT(LogInit, Verbose, "{Message}", InMessage);
 #if CSV_BOOT_PROFILING
-	if (!GBootTimingCompleted && IsInGameThread())
+	if (!GBootTimingCompleted && FCsvProfiler::IsCapturing() && IsInGameThread())
 	{
 		Message = InMessage;
 		if (Message.Len())
@@ -551,7 +551,7 @@ FScopedBootTiming::FScopedBootTiming(const ANSICHAR *InMessage, FName Suffix)
 {
 	UE_LOGFMT(LogInit, Verbose, "{Message}: {Suffix}", InMessage, Suffix);
 #if CSV_BOOT_PROFILING
-	if (!GBootTimingCompleted && IsInGameThread())
+	if (!GBootTimingCompleted && FCsvProfiler::IsCapturing() && IsInGameThread())
 	{
 		Message = FString(InMessage) + Suffix.ToString();
 		if (Message.Len())
diff --git a/Engine/Source/Runtime/Core/Private/ProfilingDebugging/CsvProfiler.cpp b/Engine/Source/Runtime/Core/Private/ProfilingDebugging/CsvProfiler.cpp
--- a/Engine/Source/Runtime/Core/Private/ProfilingDebugging/CsvProfiler.cpp
+++ b/Engine/Source/Runtime/Core/Private/ProfilingDebugging/CsvProfiler.cpp
@@ -5055,7 +5055,7 @@ void FCsvProfiler::Init()
 	}
 }
 
-bool FCsvProfiler::IsCapturing() const
+bool FCsvProfiler::IsCapturing()
 {
 	check(IsInGameThread());
 	return GCsvProfilerIsCapturing;
@@ -5137,7 +5137,7 @@ bool FCsvProfiler::IsCategoryEnabled(uint32 CategoryIndex) const
 	return GCsvCategoriesEnabled[CategoryIndex];
 }
 
-bool FCsvProfiler::IsCapturing_Renderthread() const
+bool FCsvProfiler::IsCapturing_Renderthread()
 {
 	check(IsInParallelRenderingThread());
 	return GCsvProfilerIsCapturingRT;
diff --git a/Engine/Source/Runtime/Core/Public/ProfilingDebugging/CsvProfiler.h b/Engine/Source/Runtime/Core/Public/ProfilingDebugging/CsvProfiler.h
--- a/Engine/Source/Runtime/Core/Public/ProfilingDebugging/CsvProfiler.h
+++ b/Engine/Source/Runtime/Core/Public/ProfilingDebugging/CsvProfiler.h
@@ -482,9 +482,10 @@ public:
 
 	CORE_API static FString GetDefaultDirectory();
 
+	CORE_API static bool IsCapturing();
+	CORE_API static bool IsCapturing_Renderthread();
+
 	/** Singleton interface */
-	CORE_API bool IsCapturing() const;
-	CORE_API bool IsCapturing_Renderthread() const;
 	CORE_API bool IsWritingFile() const;
 	CORE_API bool IsEndCapturePending() const;
 
