Fix IsInGameThread assert
#Synced-CL Unknown

diff --git a/Engine/Source/Runtime/Core/Private/Misc/CoreGlobals.cpp b/Engine/Source/Runtime/Core/Private/Misc/CoreGlobals.cpp
--- a/Engine/Source/Runtime/Core/Private/Misc/CoreGlobals.cpp
+++ b/Engine/Source/Runtime/Core/Private/Misc/CoreGlobals.cpp
@@ -536,7 +536,7 @@ FScopedBootTiming::FScopedBootTiming(const ANSICHAR *InMessage)
 {
 	UE_LOGFMT(LogInit, Verbose, "{Message}", InMessage);
 #if CSV_BOOT_PROFILING
-	if (!GBootTimingCompleted && FCsvProfiler::IsCapturing() && IsInGameThread())
+	if (!GBootTimingCompleted && IsInGameThread() && FCsvProfiler::IsCapturing())
 	{
 		Message = InMessage;
 		if (Message.Len())
@@ -551,7 +551,7 @@ FScopedBootTiming::FScopedBootTiming(const ANSICHAR *InMessage, FName Suffix)
 {
 	UE_LOGFMT(LogInit, Verbose, "{Message}: {Suffix}", InMessage, Suffix);
 #if CSV_BOOT_PROFILING
-	if (!GBootTimingCompleted && FCsvProfiler::IsCapturing() && IsInGameThread())
+	if (!GBootTimingCompleted && IsInGameThread() && FCsvProfiler::IsCapturing())
 	{
 		Message = FString(InMessage) + Suffix.ToString();
 		if (Message.Len())
