Add logging and tracing for Linux filename case-insensitivity operations

#rnx

#Synced-CL 48354687

diff --git a/Engine/Source/Runtime/Core/Private/Unix/UnixPlatformFile.cpp b/Engine/Source/Runtime/Core/Private/Unix/UnixPlatformFile.cpp
--- a/Engine/Source/Runtime/Core/Private/Unix/UnixPlatformFile.cpp
+++ b/Engine/Source/Runtime/Core/Private/Unix/UnixPlatformFile.cpp
@@ -3,10 +3,12 @@
 #include "Unix/UnixPlatformFile.h"
 #include "GenericPlatform/GenericPlatformProcess.h"
 #include "HAL/PlatformTime.h"
+#include "HAL/PlatformOutputDevices.h"
 #include "Containers/StringConv.h"
 #include "Containers/LruCache.h"
 #include "Logging/LogMacros.h"
 #include "Misc/Paths.h"
+#include "Misc/OutputDeviceRedirector.h"
 #include "Async/MappedFileHandle.h"
 #include "AutoRTFM.h"
 #include <sys/file.h>
@@ -14,6 +16,8 @@
 
 #include "HAL/PlatformFileCommon.h"
 #include "HAL/PlatformFileManager.h"
+#include "ProfilingDebugging/CpuProfilerTrace.h"
+#include "ProfilingDebugging/CountersTrace.h"
 
 DEFINE_LOG_CATEGORY_STATIC(LogUnixPlatformFile, Log, All);
 
@@ -502,6 +506,10 @@ namespace
  * A class to handle case insensitive file opening. This is a band-aid, non-performant approach,
  * without any caching.
  */
+
+ TRACE_DECLARE_INT_COUNTER(FUnixFileMapper_UnnecessaryFileRemapping, TEXT("FUnixFileMapper/Unnecessary File Remappings"));
+ TRACE_DECLARE_INT_COUNTER(FUnixFileMapper_NecessaryFileRemapping, TEXT("FUnixFileMapper/Necessary File Remappings"));
+
 class FUnixFileMapper
 {
 
@@ -513,6 +521,8 @@ public:
 
 	FString GetPathComponent(const FString & Filename, int NumPathComponent)
 	{
+		TRACE_CPUPROFILER_EVENT_SCOPE(FUnixFileMapper::GetPathComponent);
+
 		// skip over empty part
 		int StartPosition = (Filename[0] == TEXT('/')) ? 1 : 0;
 		
@@ -581,6 +591,8 @@ public:
 	 */
 	bool MapFileRecursively(const FString & Filename, int PathComponentToLookFor, int MaxPathComponents, FString & ConstructedPath)
 	{
+		TRACE_CPUPROFILER_EVENT_SCOPE(FUnixFileMapper::MapFileRecursively);
+
 		// get the directory without the last path component
 		FString BaseDir = ConstructedPath;
 
@@ -657,7 +669,9 @@ public:
 	 */
 	bool MapCaseInsensitiveFile(const FString & PossiblyWrongFilename, FString & ExistingFilename)
 	{
-		// Cannot log anything here, as this may result in infinite recursion when this function is called on log file itself
+		TRACE_CPUPROFILER_EVENT_SCOPE(FUnixFileMapper::MapCaseInsensitiveFile);
+
+		// Have to be careful logging in this function since this is called for setting up the log file
 
 		// We can get some "absolute" filenames like "D:/Blah/" here (e.g. non-Linux paths to source files embedded in assets).
 		// In that case, fail silently.
@@ -720,6 +734,23 @@ public:
 			}
 		});
 
+		if (bFound && !PossiblyWrongFilename.Equals(ExistingFilename, ESearchCase::CaseSensitive) && !ExistingFilename.IsEmpty())
+		{
+			// Can't just log directly here.  Use LowLevelOutputDebugString if logging isn't hooked up.
+			if (GLog && FPlatformOutputDevices::GetLog() && GLog->IsRedirectingTo(FPlatformOutputDevices::GetLog()))
+			{
+				UE_LOG(LogUnixPlatformFile, Log, TEXT("FUnixFileMapper::MapCaseInsensitiveFile: Lookup:%s Found:%s\n"), *PossiblyWrongFilename, *ExistingFilename);
+			}
+			else
+			{
+				FPlatformMisc::LowLevelOutputDebugStringf(TEXT("FUnixFileMapper::MapCaseInsensitiveFile: Lookup:%s Found:%s\n"), *PossiblyWrongFilename, *ExistingFilename);
+			}
+			TRACE_COUNTER_ADD(FUnixFileMapper_NecessaryFileRemapping, 1);
+		}
+		else
+		{
+			TRACE_COUNTER_ADD(FUnixFileMapper_UnnecessaryFileRemapping, 1);
+		}
 		return bFound;
 	}
 
@@ -762,6 +793,7 @@ public:
 					return -1;
 				}
 #endif
+				TRACE_CPUPROFILER_EVENT_SCOPE(FUnixFileMapper::OpenCaseInsensitiveRead);
 				// perform a case-insensitive search
 				// make sure we get the absolute filename
 				checkf(Filename[0] == TEXT('/'), TEXT("Filename '%s' given to OpenCaseInsensitiveRead is not absolute!"), *Filename);
@@ -776,12 +808,17 @@ public:
 						if (Handle != -1)
 						{
 							MappedToFilename = FoundFilename;
-							if (Filename != MappedToFilename)
+							if (!Filename.Equals(MappedToFilename, ESearchCase::CaseSensitive))
 							{
-								UE_LOG(LogUnixPlatformFile, Log, TEXT("Mapped '%s' to '%s'"), *Filename, *MappedToFilename);
+								UE_LOG(LogUnixPlatformFile, Log, TEXT("FUnixFileMapper::OpenCaseInsensitiveRead: '%s' to '%s'"), *Filename, *MappedToFilename);
+								TRACE_COUNTER_ADD(FUnixFileMapper_NecessaryFileRemapping, 1);
 							}
 						}
 					}
+					else
+					{
+						TRACE_COUNTER_ADD(FUnixFileMapper_UnnecessaryFileRemapping, 1);
+					}
 				}
 			}
 		}
