[Input] Add a "preferred input API" setting to the engine.

This will allow games to slowly roll out a change from an older input APIs like XInput to a more modern one like GameInput. The intent is that you could have a setting in your game, similar to a rendering setting for choosing which API to use, where players can enable/disable different input APIs to achieve the best compatibiltiy for thier setup. This has the added benefit of being able to "upgrade" a live project to use a newer input API with a fallback option via the config system.

To use this feature, set the bEnablePreferredInputAPIPreferences flag in the input project settings to true and provide a list of default preferred input APIs, such as "XInput,WinDualShock"

#Synced-CL 48353899

diff --git a/Engine/Plugins/GameInputWindows/Source/GameInputWindows/Private/GameInputWindowsModule.cpp b/Engine/Plugins/GameInputWindows/Source/GameInputWindows/Private/GameInputWindowsModule.cpp
--- a/Engine/Plugins/GameInputWindows/Source/GameInputWindows/Private/GameInputWindowsModule.cpp
+++ b/Engine/Plugins/GameInputWindows/Source/GameInputWindows/Private/GameInputWindowsModule.cpp
@@ -56,4 +56,10 @@ TSharedPtr<IInputDevice> FGameInputWindowsModule::CreateInputDevice(const TShare
 #endif	// GAME_INPUT_SUPPORT
 }
 
-IMPLEMENT_MODULE(FGameInputWindowsModule, GameInputWindows)
\ No newline at end of file
+const TCHAR* FGameInputWindowsModule::GetPreferredDeviceAPIString() const
+{
+	static const TCHAR* APIName = TEXT("GameInput");
+	return APIName;
+}
+
+IMPLEMENT_MODULE(FGameInputWindowsModule, GameInputWindows)
diff --git a/Engine/Plugins/GameInputWindows/Source/GameInputWindows/Private/GameInputWindowsModule.h b/Engine/Plugins/GameInputWindows/Source/GameInputWindows/Private/GameInputWindowsModule.h
--- a/Engine/Plugins/GameInputWindows/Source/GameInputWindows/Private/GameInputWindowsModule.h
+++ b/Engine/Plugins/GameInputWindows/Source/GameInputWindows/Private/GameInputWindowsModule.h
@@ -16,6 +16,7 @@ public:
 	static bool IsAvailable();
 
 	//~ Begin IInputDeviceModule interface
-	virtual TSharedPtr<IInputDevice> CreateInputDevice(const TSharedRef<FGenericApplicationMessageHandler>& InMessageHandler) override;		
+	virtual TSharedPtr<IInputDevice> CreateInputDevice(const TSharedRef<FGenericApplicationMessageHandler>& InMessageHandler) override;
+	virtual const TCHAR* GetPreferredDeviceAPIString() const override;
 	//~ End IInputDeviceModule interface
 };
\ No newline at end of file
diff --git a/Engine/Plugins/Runtime/Windows/WinDualShock/Source/WinDualShock/Private/WinDualShock.cpp b/Engine/Plugins/Runtime/Windows/WinDualShock/Source/WinDualShock/Private/WinDualShock.cpp
--- a/Engine/Plugins/Runtime/Windows/WinDualShock/Source/WinDualShock/Private/WinDualShock.cpp
+++ b/Engine/Plugins/Runtime/Windows/WinDualShock/Source/WinDualShock/Private/WinDualShock.cpp
@@ -824,6 +824,12 @@ public:
 		return TSharedPtr< class IInputDevice >(nullptr);
 #endif
 	}
+
+	virtual const TCHAR* GetPreferredDeviceAPIString() const override
+	{
+		static const TCHAR* APIName = TEXT("WinDualShock"); 
+		return APIName;
+	}
 };
 
 IMPLEMENT_MODULE(FWinDualShockPlugin, WinDualShock)
diff --git a/Engine/Plugins/Runtime/Windows/XInputDevice/Source/XInputDevice/Private/XInputDeviceModule.cpp b/Engine/Plugins/Runtime/Windows/XInputDevice/Source/XInputDevice/Private/XInputDeviceModule.cpp
--- a/Engine/Plugins/Runtime/Windows/XInputDevice/Source/XInputDevice/Private/XInputDeviceModule.cpp
+++ b/Engine/Plugins/Runtime/Windows/XInputDevice/Source/XInputDevice/Private/XInputDeviceModule.cpp
@@ -15,4 +15,10 @@ TSharedPtr<IInputDevice> FXInputDeviceModule::CreateInputDevice(const TSharedRef
 	return XInputInterface::Create(InMessageHandler, InParameters.bInitAsPrimaryDevice);
 }
 
+const TCHAR* FXInputDeviceModule::GetPreferredDeviceAPIString() const
+{
+	static const TCHAR* APIName = TEXT("XInput"); 
+	return APIName;
+}
+
 IMPLEMENT_MODULE(FXInputDeviceModule, XInputDevice)
diff --git a/Engine/Plugins/Runtime/Windows/XInputDevice/Source/XInputDevice/Public/XInputDeviceModule.h b/Engine/Plugins/Runtime/Windows/XInputDevice/Source/XInputDevice/Public/XInputDeviceModule.h
--- a/Engine/Plugins/Runtime/Windows/XInputDevice/Source/XInputDevice/Public/XInputDeviceModule.h
+++ b/Engine/Plugins/Runtime/Windows/XInputDevice/Source/XInputDevice/Public/XInputDeviceModule.h
@@ -13,5 +13,6 @@ public:
 	//~ Begin IInputDeviceModule interface
 	virtual TSharedPtr<IInputDevice> CreateInputDevice(const TSharedRef<FGenericApplicationMessageHandler>& InMessageHandler) override;
 	virtual TSharedPtr<IInputDevice> CreateInputDevice(const TSharedRef<FGenericApplicationMessageHandler>& InMessageHandler, FInputDeviceCreationParameters InParameters) override;
+	virtual const TCHAR* GetPreferredDeviceAPIString() const override;
 	//~ End IInputDeviceModule interface
 };
diff --git a/Engine/Source/Runtime/ApplicationCore/Private/Windows/WindowsApplication.cpp b/Engine/Source/Runtime/ApplicationCore/Private/Windows/WindowsApplication.cpp
--- a/Engine/Source/Runtime/ApplicationCore/Private/Windows/WindowsApplication.cpp
+++ b/Engine/Source/Runtime/ApplicationCore/Private/Windows/WindowsApplication.cpp
@@ -4103,7 +4103,7 @@ void FWindowsApplication::OnInputDeviceModuleRegistered(const FName& Type, IModu
 	{
 		IInputDeviceModule* Module = static_cast<IInputDeviceModule*>(ModularFeature);
 
-		TSharedPtr<IInputDevice> Device = Module->CreateInputDevice(MessageHandler);
+		TSharedPtr<IInputDevice> Device = Module->CreateInputDeviceIfPreferred(MessageHandler);
 		AddExternalInputDevice(MoveTemp(Device));
 	}
 }
@@ -4180,7 +4180,7 @@ void FWindowsApplication::PollGameDeviceState( const float TimeDelta )
 		TArray<IInputDeviceModule*> PluginImplementations = IModularFeatures::Get().GetModularFeatureImplementations<IInputDeviceModule>( IInputDeviceModule::GetModularFeatureName() );
 		for( auto InputPluginIt = PluginImplementations.CreateIterator(); InputPluginIt; ++InputPluginIt )
 		{
-			TSharedPtr<IInputDevice> Device = (*InputPluginIt)->CreateInputDevice(MessageHandler);
+			TSharedPtr<IInputDevice> Device = (*InputPluginIt)->CreateInputDeviceIfPreferred(MessageHandler);
 			AddExternalInputDevice(MoveTemp(Device));			
 		}
 
diff --git a/Engine/Source/Runtime/Core/Private/GenericPlatform/GenericPlatformMisc.cpp b/Engine/Source/Runtime/Core/Private/GenericPlatform/GenericPlatformMisc.cpp
--- a/Engine/Source/Runtime/Core/Private/GenericPlatform/GenericPlatformMisc.cpp
+++ b/Engine/Source/Runtime/Core/Private/GenericPlatform/GenericPlatformMisc.cpp
@@ -1303,6 +1303,85 @@ IPlatformHostCommunication& FGenericPlatformMisc::GetPlatformHostCommunication()
 	return Singleton;
 }
 
+bool FGenericPlatformMisc::IsPreferredInputDevice(const TCHAR* DeviceName)
+{
+	// If given a null treat it as always being enabled.
+	// This ensures that users who have written their own IInputDeviceModule code will
+	// continue working and they can add support for this feature if they want to.
+	if (!DeviceName)
+	{
+		return true;
+	}
+	
+	bool bIsPreferredAPICheckEnabled = false;
+	
+	GConfig->GetBool(TEXT("/Script/Engine.InputSettings"), TEXT("bEnablePreferredInputAPIPreferences"), OUT bIsPreferredAPICheckEnabled, GInputIni);
+	
+	// If the project setting is not enabled, always return true because we want every API to be enabled
+	if (!bIsPreferredAPICheckEnabled)
+	{
+		return true;
+	}
+
+	const FString CurrentPreferredInputAPIs = FGenericPlatformMisc::GetCurrentPreferredInputDevices();
+
+	UE_CLOG(CurrentPreferredInputAPIs.IsEmpty(), LogGenericPlatformMisc, Warning, TEXT("bEnablePreferredInputAPIPreferences is true, but no default preferred APIs were provided. Please fill out DefaultPreferredInputAPIList in the input settings for this to have any affect"));
+	
+	// If the given device is contained in our comma-seperated string value, then it is preferred
+	return CurrentPreferredInputAPIs.Contains(DeviceName);
+}
+
+void FGenericPlatformMisc::SetPreferredInputDevices(const TCHAR* DeviceNames)
+{
+	GConfig->SetString(TEXT("GamepadAPIPreference"), TEXT("PreferredAPI"), DeviceNames, GGameUserSettingsIni);
+
+	UE_LOG(LogGenericPlatformMisc, Log, TEXT("Changed preferred gamepad API to '%s'"), DeviceNames);
+}
+
+void FGenericPlatformMisc::ResetPreferredInputDevicesToDefault()
+{
+	// Clear the user set value
+	GConfig->RemoveKey(TEXT("GamepadAPIPreference"), TEXT("PreferredAPI"), GGameUserSettingsIni);
+
+	// Set it to the default specified in the project if there is one
+	FString DefaultPreferredAPIs = TEXT("");
+	const bool bHasDefaultPreferredList = GConfig->GetString(TEXT("/Script/Engine.InputSettings"), TEXT("DefaultPreferredInputAPIList"), DefaultPreferredAPIs, GInputIni);
+	if (bHasDefaultPreferredList && !DefaultPreferredAPIs.IsEmpty())
+	{
+		SetPreferredInputDevices(*DefaultPreferredAPIs);
+	}
+
+	UE_LOG(LogGenericPlatformMisc, Log, TEXT("Changed preferred gamepad API to '%s'"), *DefaultPreferredAPIs);
+}
+
+FString FGenericPlatformMisc::GetCurrentPreferredInputDevices()
+{
+	// If the API check is enabled, we want to look at the game's user preference INI and see what the use has selected for their API 
+	FString CurrentPreferredAPIs = TEXT("");
+
+#if WITH_EDITOR
+	// First, check the editor preference for overriden API preferences here allowing for easier testing
+	if (GConfig->GetString(TEXT("/Script/UnrealEd.LevelEditorPlaySettings"), TEXT("PreferredInputAPIOverride"), OUT CurrentPreferredAPIs, GEditorPerProjectIni) && !CurrentPreferredAPIs.IsEmpty())
+	{
+		return CurrentPreferredAPIs;
+	}
+#endif
+
+	// Next, check if the user has set a preference yet. If they have, use that
+	if (GConfig->GetString(TEXT("GamepadAPIPreference"), TEXT("PreferredAPI"), OUT CurrentPreferredAPIs, GGameUserSettingsIni))
+	{
+		return CurrentPreferredAPIs;
+	}
+		
+	// Lastly, use the project's default preferred list if there is one.
+	if (GConfig->GetString(TEXT("/Script/Engine.InputSettings"), TEXT("DefaultPreferredInputAPIList"), OUT CurrentPreferredAPIs, GInputIni))
+	{
+		return CurrentPreferredAPIs;
+	}
+		
+	return CurrentPreferredAPIs;
+}
+
 void GenericPlatformMisc_GetProjectFilePathProjectDir(FString& OutGameDir)
 {
 	// Here we derive the game path from the project file location.
diff --git a/Engine/Source/Runtime/Core/Public/GenericPlatform/GenericPlatformMisc.h b/Engine/Source/Runtime/Core/Public/GenericPlatform/GenericPlatformMisc.h
--- a/Engine/Source/Runtime/Core/Public/GenericPlatform/GenericPlatformMisc.h
+++ b/Engine/Source/Runtime/Core/Public/GenericPlatform/GenericPlatformMisc.h
@@ -1316,6 +1316,39 @@ public:
 	 */
 	static CORE_API IPlatformHostCommunication& GetPlatformHostCommunication();
 
+	/**
+	 * If bEnablePreferredInputAPIPreferences is enabled, then check if the given input device is
+	 * "Preferred" by the user. If a null device name is given, then this function will always treat 
+	 * that as a preferred API and return true.
+	 * 
+	 * @param DeviceName	The name of the input device API to check if it is preferred. This should come from
+	 *						an IInputDeviceModule::GetPreferredDeviceAPIString call.
+	 * @return True if the given input device API is preferred.
+	 */
+	static CORE_API bool IsPreferredInputDevice(const TCHAR* DeviceName);
+
+	/**
+	 * Sets the [GamepadAPIPreference]PreferredAPI config value in the GGameUserSettingsIni ini to the
+	 * given value. This will then be used to determine if an input API is preferred on the next boot.
+	 *
+	 * NOTE: This will not take affect immediately. An engine reboot is required to determine which input APIs should be created.
+	 * 
+	 * @param DeviceName The full comma-seperated list of preferred input APIs this application should use.
+	 */
+	static CORE_API void SetPreferredInputDevices(const TCHAR* DeviceName);
+
+	/**
+	 * Resets the users [GamepadAPIPreference]PreferredAPI config value
+	 * and calls SetPreferredInputDevices with the project's default preferred
+	 * input API list if available.
+	 */
+	static CORE_API void ResetPreferredInputDevicesToDefault();
+
+	/**
+	 * @return A comma-seperated string value which represents a list of the preferred input APIs.
+	 */
+	static CORE_API FString GetCurrentPreferredInputDevices();
+
 	/**
 	 * Has the OS execute a command and path pair (such as launch a browser)
 	 *
diff --git a/Engine/Source/Runtime/Engine/Classes/GameFramework/InputSettings.h b/Engine/Source/Runtime/Engine/Classes/GameFramework/InputSettings.h
--- a/Engine/Source/Runtime/Engine/Classes/GameFramework/InputSettings.h
+++ b/Engine/Source/Runtime/Engine/Classes/GameFramework/InputSettings.h
@@ -119,6 +119,41 @@ class UInputSettings
 	UPROPERTY(config, EditAnywhere, Category="Bindings")
 	uint8 bEnabledLegacyMappingDeprecationWarnings : 1;
 
+	/**
+	 * If true, then before creating an IInputDeviceModule, check for a list of "preferred"
+	 * input API names set by the user in their GGameUserSettingsIni.
+	 *
+	 * This will check the following value:
+	 * 
+	 *		; In GGameUserSettingsIni
+	 *		[GamepadAPIPreference]
+	 *		PreferredAPI=Example,API,List
+	 *
+	 *	The "PreferredAPI" config value is a comma seperated string value which can be modified in your gameplay code.
+	 *	If the value is not set, or empty, then the default value of "DefaultPreferredInputAPIList" will be checked.
+	 *
+	 * If the module is not contained in the comma-seperated list of preferred APIs,
+	 * then the input API will not be used.
+	 * 
+	 * This may be a useful setting for your project if you would like to provide users with multiple options
+	 * for their Input API to choose from to enjoy the best compatibility for their device. Because this setting
+	 * is config driven, it will allow you to package your game with all the input plugins you would like to support
+	 * enabled, but determine which to use at runtime, similar to how many projects allow you to select between different
+	 * rendering APIs. 
+	 *
+	 * This feature is only for Desktop platforms, and is experimental in 5.8
+	 */
+	UE_EXPERIMENTAL(5.8, "This setting is experimental")
+	UPROPERTY(config, EditAnywhere, Category = "Input", meta=(InlineEditConditionToggle))
+	uint8 bEnablePreferredInputAPIPreferences : 1;
+
+	/**
+	 * This feature is only for Desktop platforms, and is experimental in 5.8
+	 */
+	UE_EXPERIMENTAL(5.8, "This setting is experimental")
+	UPROPERTY(config, EditAnywhere, Category = "Input", AdvancedDisplay, meta=(EditCondition="bEnablePreferredInputAPIPreferences"))
+	FString DefaultPreferredInputAPIList;
+
 	/** 
 	 * Disables autocorrect for these operating systems, even if autocorrect is enabled. Use the format "[platform] [osversion]"
 	 * (e.g., "iOS 11.2" or "Android 6"). More specific versions will disable autocorrect for fewer devices ("iOS 11" will disable
diff --git a/Engine/Source/Runtime/Engine/Private/UserInterface/InputSettings.cpp b/Engine/Source/Runtime/Engine/Private/UserInterface/InputSettings.cpp
--- a/Engine/Source/Runtime/Engine/Private/UserInterface/InputSettings.cpp
+++ b/Engine/Source/Runtime/Engine/Private/UserInterface/InputSettings.cpp
@@ -46,6 +46,7 @@ PRAGMA_DISABLE_DEPRECATION_WARNINGS
 	, bEnableDynamicComponentInputBinding(true)
 PRAGMA_ENABLE_DEPRECATION_WARNINGS
 	, bEnabledLegacyMappingDeprecationWarnings(true)
+	, bEnablePreferredInputAPIPreferences(false)
 	, DefaultViewportMouseCaptureMode(EMouseCaptureMode::CapturePermanently_IncludingInitialMouseDown)
 	, DefaultViewportMouseLockMode(EMouseLockMode::LockOnCapture)
 	, DefaultPlayerInputClass(UPlayerInput::StaticClass())
diff --git a/Engine/Source/Runtime/InputDevice/Public/IInputDeviceModule.h b/Engine/Source/Runtime/InputDevice/Public/IInputDeviceModule.h
--- a/Engine/Source/Runtime/InputDevice/Public/IInputDeviceModule.h
+++ b/Engine/Source/Runtime/InputDevice/Public/IInputDeviceModule.h
@@ -4,6 +4,7 @@
 
 #include "CoreMinimal.h"
 #include "GenericPlatform/GenericApplicationMessageHandler.h"
+#include "GenericPlatform/GenericPlatformMisc.h"
 #include "Features/IModularFeatures.h"
 #include "Features/IModularFeature.h"
 #include "Modules/ModuleInterface.h"
@@ -54,13 +55,28 @@ public:
 		return FModuleManager::Get().IsModuleLoaded( "InputDevice" );
 	}
 
+	/**
+	 * Calls the normal "CreateInputDevice" logic only if this input device module is preferred.
+	 *
+	 * Returns null if this device is not preferred
+	 */
+	TSharedPtr<IInputDevice> CreateInputDeviceIfPreferred(const TSharedRef<FGenericApplicationMessageHandler>& InMessageHandler)
+	{
+		if (!FGenericPlatformMisc::IsPreferredInputDevice(GetPreferredDeviceAPIString()))
+		{
+			return nullptr;
+		}
+		
+		return CreateInputDevice(InMessageHandler);
+	}
+
 	/**
 	 * Attempts to create a new input device with the given message handler.  This version will eventually be deprecated and users should
 	 * migrate to the new version with an additional parameters struct.
 	 *
 	 * @return	Interface to the new input device, if we were able to successfully create one
 	 */
-	virtual TSharedPtr<class IInputDevice> CreateInputDevice(const TSharedRef<FGenericApplicationMessageHandler>& InMessageHandler) = 0;
+	virtual TSharedPtr<IInputDevice> CreateInputDevice(const TSharedRef<FGenericApplicationMessageHandler>& InMessageHandler) = 0;
 
 	/**
 	 * Attempts to create a new input device interface with the specified input
@@ -70,7 +86,7 @@ public:
 	 *
 	 * @return	Interface to the new input device, if we were able to successfully create one
 	 */
-	virtual TSharedPtr<class IInputDevice> CreateInputDevice(const TSharedRef<FGenericApplicationMessageHandler>& InMessageHandler, FInputDeviceCreationParameters InParams)
+	virtual TSharedPtr<IInputDevice> CreateInputDevice(const TSharedRef<FGenericApplicationMessageHandler>& InMessageHandler, FInputDeviceCreationParameters InParams)
 	{
 		// If this is a primary device and it should be constructed for the input system then call the default implementation.
 		if (InParams.bInitAsPrimaryDevice)
@@ -80,4 +96,19 @@ public:
 
 		return nullptr;
 	}
+
+protected:
+	
+	/**
+	 * Override this function if you would like your IInputDeviceModule to support
+	 * "Preferred Device API" settings.
+	 * 
+	 * @return Name of the input API as it would appear in the list of preferred input API names in the 
+	 * user settings. If null, then this input device module will never be checked for if the user prefers
+	 * it and will always be created.
+	 */
+	virtual const TCHAR* GetPreferredDeviceAPIString() const
+	{
+		return nullptr;
+	}
 };
