Fix comment in light of 48365800

#rnx
#Synced-CL 48299257

diff --git a/Engine/Source/Runtime/Core/Public/Misc/Build.h b/Engine/Source/Runtime/Core/Public/Misc/Build.h
--- a/Engine/Source/Runtime/Core/Public/Misc/Build.h
+++ b/Engine/Source/Runtime/Core/Public/Misc/Build.h
@@ -554,9 +554,10 @@
 //	5.	Relatedly avoids hashing of placeholder types that are not hashable - see PropertySet.cpp,
 //		PropertyMap.cpp
 //  6.	Assumes a read operation, avoiding expensive mutators - see PropertyBagRepository.cpp
-//	7.	Handles some unexpected seeking in complicated native asset loads - see Class.cpp
-//	8.	Allows serialization routines to expect unreachable objects (in the long run all loaded 
+//	7.	Allows serialization routines to expect unreachable objects (in the long run all loaded 
 //		data should be made reachable) - see JsonStringifyImpl.cpp
+//  8.	Allows loading of assets with a newer custom version because we are not running user native
+//		serialize routines
 //
 // All of these usages will be accomodated more elegantly by future work. Do not use this define.
 #ifndef UE_INTERNAL_UNIVERSAL_PLACEHOLDERS
