Added utils to return the language that we last loaded localization data for

This is cached each time the language changs or localization data is refreshed, to avoid hitting the disk for every call

#Synced-CL 48128650

diff --git a/Engine/Source/Runtime/Core/Private/Internationalization/TextLocalizationManager.cpp b/Engine/Source/Runtime/Core/Private/Internationalization/TextLocalizationManager.cpp
--- a/Engine/Source/Runtime/Core/Private/Internationalization/TextLocalizationManager.cpp
+++ b/Engine/Source/Runtime/Core/Private/Internationalization/TextLocalizationManager.cpp
@@ -558,6 +558,8 @@ void InitEngineTextLocalization()
 	TextLocalizationResourceUtil::ClearNativeEditorCultureName();
 #endif
 
+	FTextLocalizationManager::Get().CacheLocalizedLanguageName(LocLoadFlags);
+
 	FTextLocalizationManager::Get().QueueAsyncTask([LocLoadFlags, AvailableTextSources = FTextLocalizationManager::Get().LocalizedTextSources]()
 	{
 		FTextLocalizationManager::Get().LoadLocalizationResourcesForCulture_Sync(AvailableTextSources, FInternationalization::Get().GetCurrentLanguage()->GetName(), LocLoadFlags);
@@ -606,6 +608,8 @@ void InitGameTextLocalization()
 		LocLoadFlags |= ELocalizationLoadFlags::Additional;
 	}
 
+	FTextLocalizationManager::Get().CacheLocalizedLanguageName(LocLoadFlags);
+
 	FTextLocalizationManager::Get().QueueAsyncTask([LocLoadFlags, AvailableTextSources = FTextLocalizationManager::Get().LocalizedTextSources]()
 	{
 		FTextLocalizationManager::Get().LoadLocalizationResourcesForCulture_Sync(AvailableTextSources, FInternationalization::Get().GetCurrentLanguage()->GetName(), LocLoadFlags);
@@ -908,6 +912,11 @@ TArray<FString> FTextLocalizationManager::GetLocalizedCultureNames(const ELocali
 	return LocalizedCultureNames;
 }
 
+FString FTextLocalizationManager::GetLocalizedLanguageName() const
+{
+	return LocalizedLanguageName;
+}
+
 int32 FTextLocalizationManager::GetLocalizationTargetPathId(FStringView InLocalizationTargetPath)
 {
 	FWriteScopeLock ScopeLock(DisplayStringTableRW);
@@ -1240,6 +1249,8 @@ void FTextLocalizationManager::RefreshResources(TUniqueFunction<void()>&& Notifi
 	LocLoadFlags |= ELocalizationLoadFlags::Native;
 	LocLoadFlags |= ELocalizationLoadFlags::Additional;
 
+	CacheLocalizedLanguageName(LocLoadFlags);
+
 	LoadLocalizationResourcesForCulture_Async(FInternationalization::Get().GetCurrentLanguage()->GetName(), LocLoadFlags, MoveTemp(Notification));
 }
 
@@ -1802,6 +1813,36 @@ void FTextLocalizationManager::DirtyTextRevision()
 	}
 }
 
+void FTextLocalizationManager::CacheLocalizedLanguageName(const ELocalizationLoadFlags InLoadFlags)
+{
+	TRACE_CPUPROFILER_EVENT_SCOPE(FTextLocalizationManager::CacheLocalizedLanguageName);
+
+	{
+		const TArray<FString> AvailableCultures = GetLocalizedCultureNames(InLoadFlags);
+		const TArray<FString> PrioritizedCulturesToMatch = FInternationalization::Get().GetPrioritizedCultureNames(FInternationalization::Get().GetCurrentLanguage()->GetName());
+		for (const FString& PrioritizedCulture : PrioritizedCulturesToMatch)
+		{
+			if (AvailableCultures.Contains(PrioritizedCulture))
+			{
+				LocalizedLanguageName = PrioritizedCulture;
+				return;
+			}
+		}
+	}
+
+	if (EnumHasAnyFlags(InLoadFlags, ELocalizationLoadFlags::Game))
+	{
+		// If this is a game, use the native culture of the game as the fallback
+		LocalizedLanguageName = GetNativeCultureName(ELocalizedTextSourceCategory::Game);
+		if (!LocalizedLanguageName.IsEmpty())
+		{
+			return;
+		}
+	}
+	
+	LocalizedLanguageName = TEXT("en");
+}
+
 #if WITH_EDITOR
 void FTextLocalizationManager::EnableGameLocalizationPreview()
 {
diff --git a/Engine/Source/Runtime/Core/Public/Internationalization/TextLocalizationManager.h b/Engine/Source/Runtime/Core/Public/Internationalization/TextLocalizationManager.h
--- a/Engine/Source/Runtime/Core/Public/Internationalization/TextLocalizationManager.h
+++ b/Engine/Source/Runtime/Core/Public/Internationalization/TextLocalizationManager.h
@@ -177,6 +177,12 @@ public:
 	 */
 	CORE_API TArray<FString> GetLocalizedCultureNames(const ELocalizationLoadFlags InLoadFlags) const;
 
+	/**
+	 * Get the language that we last loaded localization data for.
+	 * eg) If the available languages were [en, fr, de] and the current language was "en-US", this function would return "en".
+	 */
+	CORE_API FString GetLocalizedLanguageName() const;
+
 	/**
 	 * Given a localization target path, get the ID associated with it.
 	 * @note This ID is unstable and should only be used for quick in-process comparison.
@@ -394,6 +400,9 @@ private:
 	/** Dirties the text revision counter by incrementing it, causing a revision mismatch for any information cached before this happens.  */
 	void DirtyTextRevision();
 
+	/** Cache the language that we last loaded localization data for. */
+	void CacheLocalizedLanguageName(const ELocalizationLoadFlags InLoadFlags);
+
 	/** Array of registered localized text sources, sorted by priority (@see RegisterTextSource) */
 	TArray<TSharedPtr<ILocalizedTextSource>> LocalizedTextSources;
 
@@ -405,4 +414,7 @@ private:
 
 	/** The latest async localization task (if any). Additional requests are chained to this so that they run in sequence */
 	FGraphEventRef AsyncLocalizationTask;
+
+	/** The language that we last loaded localization data for. */
+	FString LocalizedLanguageName;
 };
diff --git a/Engine/Source/Runtime/Engine/Classes/Kismet/KismetInternationalizationLibrary.h b/Engine/Source/Runtime/Engine/Classes/Kismet/KismetInternationalizationLibrary.h
--- a/Engine/Source/Runtime/Engine/Classes/Kismet/KismetInternationalizationLibrary.h
+++ b/Engine/Source/Runtime/Engine/Classes/Kismet/KismetInternationalizationLibrary.h
@@ -140,6 +140,13 @@ public:
 	 */
 	static ENGINE_API TArray<FString> GetLocalizedCultures(const ELocalizationLoadFlags LoadFlags);
 
+	/**
+	 * Get the language that we last loaded localization data for.
+	 * eg) If the available languages were [en, fr, de] and the current language was "en-US", this function would return "en".
+	 */
+	UFUNCTION(BlueprintPure, Category="Utilities|Internationalization")
+	static ENGINE_API FString GetLocalizedLanguage();
+
 	/**
 	 * Given a list of available cultures, try and find the most suitable culture from the list based on culture prioritization.
 	 *   eg) If your list was [en, fr, de] and the given culture was "en-US", this function would return "en".
diff --git a/Engine/Source/Runtime/Engine/Private/KismetInternationalizationLibrary.cpp b/Engine/Source/Runtime/Engine/Private/KismetInternationalizationLibrary.cpp
--- a/Engine/Source/Runtime/Engine/Private/KismetInternationalizationLibrary.cpp
+++ b/Engine/Source/Runtime/Engine/Private/KismetInternationalizationLibrary.cpp
@@ -139,6 +139,11 @@ TArray<FString> UKismetInternationalizationLibrary::GetLocalizedCultures(const E
 	return FTextLocalizationManager::Get().GetLocalizedCultureNames(LoadFlags);
 }
 
+FString UKismetInternationalizationLibrary::GetLocalizedLanguage()
+{
+	return FTextLocalizationManager::Get().GetLocalizedLanguageName();
+}
+
 FString UKismetInternationalizationLibrary::GetSuitableCulture(const TArray<FString>& AvailableCultures, const FString& CultureToMatch, const FString& FallbackCulture)
 {
 	const TArray<FString> PrioritizedCulturesToMatch = FInternationalization::Get().GetPrioritizedCultureNames(CultureToMatch);
