- Don't load plugin ini files from subdirs (like Tags)

#Synced-CL 48049239

diff --git a/Engine/Source/Runtime/Core/Private/Misc/ConfigCacheIni.cpp b/Engine/Source/Runtime/Core/Private/Misc/ConfigCacheIni.cpp
--- a/Engine/Source/Runtime/Core/Private/Misc/ConfigCacheIni.cpp
+++ b/Engine/Source/Runtime/Core/Private/Misc/ConfigCacheIni.cpp
@@ -6949,13 +6949,6 @@ void FConfigCacheIni::AddPluginsToBranches(const TArray<FName>& PluginNames, FCo
 	
 			SlowPluginConfigs = TSet<FString>(LocalConfigs);
 		}
-		
-#if !UE_BUILD_SHIPPING
-		for (const FString& F : *PluginConfigs)
-		{
-			UE_LOG(LogConfig, Verbose, TEXT("Found config file %s in plugin dir %s"), *F, *PluginInfo->PluginDir);
-		}
-#endif // !UE_BUILD_SHIPPING
 
 		// make a single context that can be used for all the branches modified by this plugin
 		FConfigContext Context = FConfigContext::ReadIntoConfigSystem(this, PlatformNameStr);
@@ -6971,6 +6964,15 @@ void FConfigCacheIni::AddPluginsToBranches(const TArray<FName>& PluginNames, FCo
 		TSet<FName> LoadedBranches;
 		for (const FString& ConfigFilename : *PluginConfigs)
 		{
+			// make sure this file isn't in a subdir (like Tags), this will make the StagedPluginConfigCache case match the SlowPluginConfigs case
+			FString DirName = FPaths::GetCleanFilename(FPaths::GetPath(ConfigFilename));
+			if (!DirName.IsEmpty() && DirName != TEXT("Config") && DirName != PlatformNameStr)
+			{
+				UE_LOG(LogConfig, Verbose, TEXT("Skipping config file %s in plugin dir %s, since it's in a subdirectory."), *ConfigFilename, *PluginInfo->PluginDir);
+				continue;
+			}
+			UE_LOG(LogConfig, Verbose, TEXT("Found config file %s in plugin dir %s"), *ConfigFilename, *PluginInfo->PluginDir);
+
 			FName BranchName = *FPaths::GetBaseFilename(ConfigFilename).Replace(*StrippedPart, TEXT("")).Replace(*PlatformNameStr, TEXT(""));
 			if (LoadedBranches.Contains(BranchName))
 			{
