Core: Moved GetDataHelper to UnrealTemplate.h

#Synced-CL 0

diff --git a/Engine/Source/Runtime/Core/Public/Containers/Array.h b/Engine/Source/Runtime/Core/Public/Containers/Array.h
--- a/Engine/Source/Runtime/Core/Public/Containers/Array.h
+++ b/Engine/Source/Runtime/Core/Public/Containers/Array.h
@@ -426,14 +426,6 @@ private:
 
 namespace UE::Core::Private
 {
-	// Simply forwards to an unqualified GetData(), but can be called from within a container or view
-	// where GetData() is already a member and so hides any others.
-	template <typename T>
-	[[nodiscard]] UE_FORCEINLINE_HINT decltype(auto) GetDataHelper(T&& Arg)
-	{
-		return GetData(Forward<T>(Arg));
-	}
-
 	template <typename FromArrayType, typename ToArrayType>
 	UE_NODEBUG [[nodiscard]] constexpr bool CanMoveTArrayPointersBetweenArrayTypes()
 	{
@@ -507,7 +499,6 @@ namespace UE::Core::Private
 		return Result;
 	}
 
-	
 	// Called only when we KNOW we are going to do a realloc increasing by 1.
 	// In this case, we know that max == num and can simplify things in a very 
 	// hot location in the code.
@@ -756,7 +747,7 @@ namespace UE::Core::Private
 		}
 		ArrayMax = NewMax;
 	}
-}
+} // namespace UE::Core::Private
 
 
 /**
@@ -2603,7 +2594,7 @@ public:
 
 		// Allocate memory for the new elements.
 		SizeType Pos = AddUninitialized(SourceCount);
-		ConstructItems<ElementType>((void*)(GetData() + Pos), UE::Core::Private::GetDataHelper(Source), SourceCount);
+		ConstructItems<ElementType>((void*)(GetData() + Pos), GetDataHelper(Source), SourceCount);
 	}
 
 	/**
diff --git a/Engine/Source/Runtime/Core/Public/Containers/ArrayView.h b/Engine/Source/Runtime/Core/Public/Containers/ArrayView.h
--- a/Engine/Source/Runtime/Core/Public/Containers/ArrayView.h
+++ b/Engine/Source/Runtime/Core/Public/Containers/ArrayView.h
@@ -26,14 +26,6 @@ namespace UE::Core::ArrayView::Private
 	template <typename T, typename ElementType>
 	constexpr bool TIsCompatibleElementType_V =std::is_convertible_v<T**, ElementType* const*>;
 
-	// Simply forwards to an unqualified GetData(), but can be called from within TArrayView
-	// where GetData() is already a member and so hides any others.
-	template <typename T>
-	UE_FORCEINLINE_HINT constexpr auto GetDataHelper(T&& Arg) -> decltype(GetData(Forward<T>(Arg)))
-	{
-		return GetData(Forward<T>(Arg));
-	}
-
 	// Gets the data from the passed argument and proceeds to reinterpret the resulting elements
 	template <typename T>
 	inline decltype(auto) GetReinterpretedDataHelper(T&& Arg)
@@ -59,7 +51,7 @@ namespace UE::Core::ArrayView::Private
 		template <typename T>
 		static constexpr decltype(auto) GetData(T&& Arg)
 		{
-			return UE::Core::ArrayView::Private::GetDataHelper(Forward<T>(Arg));
+			return GetDataHelper(Forward<T>(Arg));
 		}
 	};
 
@@ -260,7 +252,7 @@ public:
 	 * The caller is responsible for ensuring that the view does not outlive the initializer list.
 	 */
 	[[nodiscard]] constexpr TArrayView(std::initializer_list<ElementType> List UE_LIFETIMEBOUND)
-		: DataPtr(UE::Core::ArrayView::Private::GetDataHelper(List))
+		: DataPtr(GetDataHelper(List))
 		, ArrayNum(GetNum(List))
 	{
 		static_assert(std::is_const_v<ElementType>, "Only views of const elements can bind to initializer lists");
@@ -836,11 +828,11 @@ public:
 	 * @returns True if this array view's contents and the other ranged type match. False otherwise.
 	 */
 	template <typename RangeType>
-		requires (std::is_convertible_v<decltype(UE::Core::ArrayView::Private::GetDataHelper(std::declval<RangeType&>())), const ElementType*>)
+		requires (std::is_convertible_v<decltype(GetDataHelper(std::declval<RangeType&>())), const ElementType*>)
 	[[nodiscard]] bool UEOpEquals(RangeType&& Rhs) const
 	{
 		auto Count = GetNum(Rhs);
-		return Count == ArrayNum && CompareItems(DataPtr, UE::Core::ArrayView::Private::GetDataHelper(Rhs), Count);
+		return Count == ArrayNum && CompareItems(DataPtr, GetDataHelper(Rhs), Count);
 	}
 
 private:
diff --git a/Engine/Source/Runtime/Core/Public/Containers/StringView.h b/Engine/Source/Runtime/Core/Public/Containers/StringView.h
--- a/Engine/Source/Runtime/Core/Public/Containers/StringView.h
+++ b/Engine/Source/Runtime/Core/Public/Containers/StringView.h
@@ -27,14 +27,7 @@
 
 namespace UE::Core::Private
 {
-	/** Allow GetData to called unqualified from a scope with its own overload of GetData. */
-	template <typename... ArgTypes>
-	constexpr inline auto StringViewGetData(ArgTypes&&... Args) -> decltype(GetData(Forward<ArgTypes>(Args)...))
-	{
-		return GetData(Forward<ArgTypes>(Args)...);
-	}
-
-	template<UE::CCharType CharType>
+	template <UE::CCharType CharType>
 	constexpr int32 StringLength(const CharType* InString)
 	{
 		int32 Count = 0;
@@ -148,7 +141,7 @@ public:
 		// !std::is_pointer is used here to force the pointer constructor to be used when passed an array or pointer.
 		requires (!std::is_pointer_v<std::decay_t<CharRangeType>> && UE::CCompatibleStringViewable<CharRangeType, CharType>)
 	[[nodiscard]] constexpr inline TStringView(const CharRangeType& InRange UE_LIFETIMEBOUND)
-		: DataPtr((const CharType*)UE::Core::Private::StringViewGetData(InRange))
+		: DataPtr((const CharType*)GetDataHelper(InRange))
 		, Size(IntCastChecked<int32>(GetNum(InRange)))
 	{
 	}
diff --git a/Engine/Source/Runtime/Core/Public/Templates/UnrealTemplate.h b/Engine/Source/Runtime/Core/Public/Templates/UnrealTemplate.h
--- a/Engine/Source/Runtime/Core/Public/Templates/UnrealTemplate.h
+++ b/Engine/Source/Runtime/Core/Public/Templates/UnrealTemplate.h
@@ -92,6 +92,16 @@ constexpr const T* GetData(std::initializer_list<T> List)
 	return List.begin();
 }
 
+/**
+ * Forwards to an unqualified GetData(), but can be called from within a scope
+ * where GetData() is a member and hides any other
+ */
+template <typename T>
+[[nodiscard]] UE_REWRITE constexpr auto GetDataHelper(T&& Arg) -> decltype(GetData((T&&)Arg))
+{
+	return GetData((T&&)Arg);
+}
+
 /**
 * Generically gets the number of items in a contiguous container
 */
