Core: Added UE::CContiguousRangeOrRef

#Synced-CL 0

diff --git a/Engine/Source/Runtime/Core/Public/Concepts/ContiguousRange.h b/Engine/Source/Runtime/Core/Public/Concepts/ContiguousRange.h
--- a/Engine/Source/Runtime/Core/Public/Concepts/ContiguousRange.h
+++ b/Engine/Source/Runtime/Core/Public/Concepts/ContiguousRange.h
@@ -14,4 +14,10 @@ namespace UE
 	concept CContiguousRange =
 		!std::is_reference_v<T> && // TIsContiguousContainer wrongly identifies references to containers as containers
 		TIsContiguousContainer_V<T>;
+
+	/**
+	 * Concept which describes a contiguous range of elements or a reference to such a range.
+	 */
+	template <typename T>
+	concept CContiguousRangeOrRef = CContiguousRange<std::remove_reference_t<T>>;
 }
