AndroidStats.h int32 error, needed include.

diff --git a/Engine/Source/Runtime/Core/Public/Android/AndroidStats.h b/Engine/Source/Runtime/Core/Public/Android/AndroidStats.h
--- a/Engine/Source/Runtime/Core/Public/Android/AndroidStats.h
+++ b/Engine/Source/Runtime/Core/Public/Android/AndroidStats.h
@@ -1,6 +1,8 @@
 // Copyright Epic Games, Inc. All Rights Reserved.
 #pragma once
 
+#include "CoreMinimal.h"
+
 enum class EDeviceThermalState : int32;
 
 class FAndroidStats
