Deprecate `FRefCountBase` which is now just a mirage for `FRefCountedObject`.

#Synced-CL 48387602

diff --git a/Engine/Source/Editor/UnrealEd/Classes/CookOnTheSide/CookOnTheFlyServer.h b/Engine/Source/Editor/UnrealEd/Classes/CookOnTheSide/CookOnTheFlyServer.h
--- a/Engine/Source/Editor/UnrealEd/Classes/CookOnTheSide/CookOnTheFlyServer.h
+++ b/Engine/Source/Editor/UnrealEd/Classes/CookOnTheSide/CookOnTheFlyServer.h
@@ -274,7 +274,7 @@ private:
 	 * Wrapper around a function for cooker tasks that need to be ticked on a schedule.
 	 * Includes data to support calling it on a schedule or triggering it manually.
 	 */
-	struct FPollable : public FRefCountBase
+	struct FPollable : public FRefCountedObject
 	{
 	public:
 		enum EManualTrigger {};
diff --git a/Engine/Source/Editor/UnrealEd/Private/Cooker/DiffWriterArchive.h b/Engine/Source/Editor/UnrealEd/Private/Cooker/DiffWriterArchive.h
--- a/Engine/Source/Editor/UnrealEd/Private/Cooker/DiffWriterArchive.h
+++ b/Engine/Source/Editor/UnrealEd/Private/Cooker/DiffWriterArchive.h
@@ -370,7 +370,7 @@ public:
  * on disk. In the second pass, the diff map is read during each call to Serialize to decide whether we need to collect
  * the stack trace for that call.
  */
-class FAccumulator : public FRefCountBase
+class FAccumulator : public FRefCountedObject
 {
 public:
 	FAccumulator(FAccumulatorGlobals& InGlobals, UObject* InAsset, FName InPackageName, int32 InMaxDiffsToLog,
diff --git a/Engine/Source/Programs/AutoRTFMTests/Private/RefCountPtrTests.cpp b/Engine/Source/Programs/AutoRTFMTests/Private/RefCountPtrTests.cpp
--- a/Engine/Source/Programs/AutoRTFMTests/Private/RefCountPtrTests.cpp
+++ b/Engine/Source/Programs/AutoRTFMTests/Private/RefCountPtrTests.cpp
@@ -7,7 +7,7 @@
 
 namespace
 {
-	struct RefCountedInt : public FRefCountBase
+	struct RefCountedInt : public FRefCountedObject
 	{
 		RefCountedInt(int InValue) : Value(InValue) {}
 		int Value;
diff --git a/Engine/Source/Programs/BenchmarkTool/Private/BenchmarkTool.cpp b/Engine/Source/Programs/BenchmarkTool/Private/BenchmarkTool.cpp
--- a/Engine/Source/Programs/BenchmarkTool/Private/BenchmarkTool.cpp
+++ b/Engine/Source/Programs/BenchmarkTool/Private/BenchmarkTool.cpp
@@ -586,7 +586,7 @@ void BM_TSharedPtrAssign_NoTS(BenchmarkState& State)
 	}
 }
 
-struct DummyRefCount : public FRefCountBase
+struct DummyRefCount : public FRefCountedObject
 {
 	int _ = 0;
 };
diff --git a/Engine/Source/Runtime/AssetRegistry/Private/AssetDataGathererPrivate.h b/Engine/Source/Runtime/AssetRegistry/Private/AssetDataGathererPrivate.h
--- a/Engine/Source/Runtime/AssetRegistry/Private/AssetDataGathererPrivate.h
+++ b/Engine/Source/Runtime/AssetRegistry/Private/AssetDataGathererPrivate.h
@@ -365,7 +365,7 @@ private:
  *
  * This class is not ThreadSafe; The FAssetDataDiscovery reads/writes its data only while holding TreeLock.
  */
-class FScanDir : public FRefCountBase
+class FScanDir : public FRefCountedObject
 {
 public:
 	/** Holder for fields on FScanDir that are inherited from parents */
diff --git a/Engine/Source/Runtime/Core/Private/Experimental/IO/IoStatusError.cpp b/Engine/Source/Runtime/Core/Private/Experimental/IO/IoStatusError.cpp
--- a/Engine/Source/Runtime/Core/Private/Experimental/IO/IoStatusError.cpp
+++ b/Engine/Source/Runtime/Core/Private/Experimental/IO/IoStatusError.cpp
@@ -17,7 +17,7 @@ namespace UE::UnifiedError::IoStore
 		ErrorCode.ReplaceCharInline(UTF8TEXT(' '), UTF8TEXT('_'));
 	}
 
-	class FIoStoreErrorDetails : public UE::UnifiedError::FManditoryErrorDetails, public FRefCountBase
+	class FIoStoreErrorDetails : public UE::UnifiedError::FManditoryErrorDetails, public FRefCountedObject
 	{
 	private:
 		FString ErrorMessage;
@@ -30,9 +30,9 @@ namespace UE::UnifiedError::IoStore
 			ErrorMessage = InErrorMessage;
 		}
 
-		virtual FReturnedRefCountValue AddRef() const final { return FRefCountBase::AddRef(); }
-		virtual uint32 Release() const final { return FRefCountBase::Release(); }
-		virtual uint32 GetRefCount() const final { return FRefCountBase::GetRefCount(); }
+		virtual FReturnedRefCountValue AddRef() const final override { return FRefCountedObject::AddRef(); }
+		virtual uint32 Release() const final override { return FRefCountedObject::Release(); }
+		virtual uint32 GetRefCount() const final override { return FRefCountedObject::GetRefCount(); }
 
 		const FUtf8String& GetErrorName(const FError& Error) const
 		{
diff --git a/Engine/Source/Runtime/Core/Public/Experimental/UnifiedError/UnifiedError.h b/Engine/Source/Runtime/Core/Public/Experimental/UnifiedError/UnifiedError.h
--- a/Engine/Source/Runtime/Core/Public/Experimental/UnifiedError/UnifiedError.h
+++ b/Engine/Source/Runtime/Core/Public/Experimental/UnifiedError/UnifiedError.h
@@ -302,16 +302,16 @@ namespace UE::UnifiedError
 	/// <summary>
 	/// FRefCountedErrorDetails; base implementation of refcounting for IErrorDetails, this is used for heap allocated IErrorDetails implementations
 	/// </summary>
-	class FRefCountedErrorDetails : public IErrorDetails, public FRefCountBase
+	class FRefCountedErrorDetails : public IErrorDetails, public FRefCountedObject
 	{
 		UE_DECLARE_FERROR_DETAILS_ABSTRACT(UnifiedError, FRefCountedErrorDetails);
 
 
 		CORE_API virtual ~FRefCountedErrorDetails();
 
-		virtual FReturnedRefCountValue AddRef() const final { return FRefCountBase::AddRef(); }
-		virtual uint32 Release() const final { return FRefCountBase::Release(); }
-		virtual uint32 GetRefCount() const final { return FRefCountBase::GetRefCount(); }
+		virtual FReturnedRefCountValue AddRef() const final override { return FRefCountedObject::AddRef(); }
+		virtual uint32 Release() const final override { return FRefCountedObject::Release(); }
+		virtual uint32 GetRefCount() const final override { return FRefCountedObject::GetRefCount(); }
 	};
 
 
diff --git a/Engine/Source/Runtime/Core/Public/Misc/ConfigAccessTracking.h b/Engine/Source/Runtime/Core/Public/Misc/ConfigAccessTracking.h
--- a/Engine/Source/Runtime/Core/Public/Misc/ConfigAccessTracking.h
+++ b/Engine/Source/Runtime/Core/Public/Misc/ConfigAccessTracking.h
@@ -60,7 +60,7 @@ namespace UE::ConfigAccessTracking
  * data about their FConfigFile when their FConfigValues are accessed. Due to FConfigSection being movable out of
  * FConfigFiles, this struct might outlive its FConfigFile (but its backpointer will be set to null).
  */
-struct FFile : public FRefCountBase
+struct FFile : public FRefCountedObject
 {
 	CORE_API FFile(const FConfigFile* InConfigFile);
 	CORE_API FName GetFilenameToLoad() const;
@@ -88,7 +88,7 @@ private:
  * of FConfigFiles, this struct might outlive the existence of its FConfigSection in the FConfigFile, or even outlive
  * the FConfigFile.
  */
-struct FSection : public FRefCountBase
+struct FSection : public FRefCountedObject
 {
 	CORE_API FSection(FFile& InFileAccess, FStringView InSectionName);
 
diff --git a/Engine/Source/Runtime/Core/Public/Tasks/TaskConcurrencyLimiter.h b/Engine/Source/Runtime/Core/Public/Tasks/TaskConcurrencyLimiter.h
--- a/Engine/Source/Runtime/Core/Public/Tasks/TaskConcurrencyLimiter.h
+++ b/Engine/Source/Runtime/Core/Public/Tasks/TaskConcurrencyLimiter.h
@@ -54,7 +54,7 @@ class FPimpl : public TSharedFromThis<FPimpl>
 {
 	// Reference counted wrapper for FTask to manage the task lifetime between worker threads and 
 	// any thread that might retract the task once it is submitted to the scheduler
-	struct FLimiterTask : FRefCountBase
+	struct FLimiterTask : FRefCountedObject
 	{
 		LowLevelTasks::FTask Task;
 	};
diff --git a/Engine/Source/Runtime/Core/Public/Templates/RefCounting.h b/Engine/Source/Runtime/Core/Public/Templates/RefCounting.h
--- a/Engine/Source/Runtime/Core/Public/Templates/RefCounting.h
+++ b/Engine/Source/Runtime/Core/Public/Templates/RefCounting.h
@@ -242,9 +242,7 @@ private:
 	}
 };
 
-// TODO (SOL-7350): deprecate and remove this entirely. This is a class because there are
-// some uses of this as a friend class declaration.
-class FRefCountBase : public FRefCountedObject {};
+class UE_DEPRECATED(5.8, "`FRefCountedObject` is now thread-safe so use that instead.") FRefCountBase : public FRefCountedObject {};
 
 class UE_DEPRECATED(5.8, "`FRefCountedObject` is now thread-safe so use that instead.") FThreadSafeRefCountedObject : public FRefCountedObject {};
 
diff --git a/Engine/Source/Runtime/CoreUObject/Private/Serialization/AsyncLoading2.cpp b/Engine/Source/Runtime/CoreUObject/Private/Serialization/AsyncLoading2.cpp
--- a/Engine/Source/Runtime/CoreUObject/Private/Serialization/AsyncLoading2.cpp
+++ b/Engine/Source/Runtime/CoreUObject/Private/Serialization/AsyncLoading2.cpp
@@ -3947,7 +3947,7 @@ private:
 * Structure containing intermediate data required for async loading of all exports of a package.
 */
 
-struct FAsyncPackage2 : public FRefCountBase
+struct FAsyncPackage2 : public FRefCountedObject
 {
 	friend struct FAsyncPackageScope2;
 	friend class FAsyncLoadingThread2;
diff --git a/Engine/Source/Runtime/CoreUObject/Public/Cooker/CookArtifact.h b/Engine/Source/Runtime/CoreUObject/Public/Cooker/CookArtifact.h
--- a/Engine/Source/Runtime/CoreUObject/Public/Cooker/CookArtifact.h
+++ b/Engine/Source/Runtime/CoreUObject/Public/Cooker/CookArtifact.h
@@ -65,7 +65,7 @@ namespace UE::Cook
  * Provides hooks to save global settings that invalidate the artifact when changed.
  * Provides hooks to clean the artifact when invalidated.
  */
-class ICookArtifact : public FRefCountBase
+class ICookArtifact : public FRefCountedObject
 {
 public:
 	virtual ~ICookArtifact() = default;
diff --git a/Engine/Source/Runtime/CoreUObject/Public/Cooker/CookDeterminismHelper.h b/Engine/Source/Runtime/CoreUObject/Public/Cooker/CookDeterminismHelper.h
--- a/Engine/Source/Runtime/CoreUObject/Public/Cooker/CookDeterminismHelper.h
+++ b/Engine/Source/Runtime/CoreUObject/Public/Cooker/CookDeterminismHelper.h
@@ -116,7 +116,7 @@ struct IDeterminismModifiedPackageContext
  * DeterminismHelper objects are registered via FObjectPreSaveContext.RegisterDeterminismHelper
  * in the PreSave of a UObject.
  */
-class IDeterminismHelper : public FRefCountBase
+class IDeterminismHelper : public FRefCountedObject
 {
 public:
 	/**
diff --git a/Engine/Source/Runtime/CoreUObject/Public/Cooker/MPCollector.h b/Engine/Source/Runtime/CoreUObject/Public/Cooker/MPCollector.h
--- a/Engine/Source/Runtime/CoreUObject/Public/Cooker/MPCollector.h
+++ b/Engine/Source/Runtime/CoreUObject/Public/Cooker/MPCollector.h
@@ -215,7 +215,7 @@ private:
  * Interface used during cooking to send data collected from save/load on a remote CookWorker
  * to the Director for aggregation into files saves at the end of the cook. 
  */
-class IMPCollector : public FRefCountBase
+class IMPCollector : public FRefCountedObject
 {
 public:
 	virtual ~IMPCollector() = default;
diff --git a/Engine/Source/Runtime/D3D12RHI/Private/D3D12PipelineState.h b/Engine/Source/Runtime/D3D12RHI/Private/D3D12PipelineState.h
--- a/Engine/Source/Runtime/D3D12RHI/Private/D3D12PipelineState.h
+++ b/Engine/Source/Runtime/D3D12RHI/Private/D3D12PipelineState.h
@@ -296,7 +296,7 @@ struct FD3D12PipelineStateWorker : public FD3D12AdapterChild, public FNonAbandon
 	TRefCountPtr<ID3D12PipelineState> PSO;
 };
 
-struct FD3D12PipelineState : public FD3D12AdapterChild, public FD3D12MultiNodeGPUObject, public FNoncopyable, public FRefCountBase
+struct FD3D12PipelineState : public FD3D12AdapterChild, public FD3D12MultiNodeGPUObject, public FNoncopyable, public FRefCountedObject
 {
 public:
 	explicit FD3D12PipelineState(FD3D12Adapter* Parent);
diff --git a/Engine/Source/Runtime/D3D12RHI/Private/D3D12TransientResourceAllocator.h b/Engine/Source/Runtime/D3D12RHI/Private/D3D12TransientResourceAllocator.h
--- a/Engine/Source/Runtime/D3D12RHI/Private/D3D12TransientResourceAllocator.h
+++ b/Engine/Source/Runtime/D3D12RHI/Private/D3D12TransientResourceAllocator.h
@@ -7,7 +7,7 @@
 
 class FD3D12TransientHeap final
 	: public FRHITransientHeap
-	, public FRefCountBase
+	, public FRefCountedObject
 	, public FD3D12LinkedAdapterObject<FD3D12TransientHeap>
 {
 public:
diff --git a/Engine/Source/Runtime/Engine/Private/ShaderCompiler/ShaderCompilerJobCache.cpp b/Engine/Source/Runtime/Engine/Private/ShaderCompiler/ShaderCompilerJobCache.cpp
--- a/Engine/Source/Runtime/Engine/Private/ShaderCompiler/ShaderCompilerJobCache.cpp
+++ b/Engine/Source/Runtime/Engine/Private/ShaderCompiler/ShaderCompilerJobCache.cpp
@@ -411,7 +411,7 @@ public:
 	/** Path to where the cached debug info is stored. */
 	FString CachedDebugInfoPath;
 
-	/** Similar to FRefCountBase AddRef, but not atomic */
+	/** Similar to FRefCountedObject AddRef, but not atomic */
 	int32 AddRef()
 	{
 		++NumReferences;
@@ -424,7 +424,7 @@ public:
 		return NumReferences;
 	}
 
-	/** Similar to FRefCountBase Release, but not atomic */
+	/** Similar to FRefCountedObject Release, but not atomic */
 	int32 Release()
 	{
 		checkf(NumReferences >= 0, TEXT("Attempting to release shader job cache output that was already released"));
diff --git a/Engine/Source/Runtime/Engine/Public/Rendering/SkeletalMeshLODRenderData.h b/Engine/Source/Runtime/Engine/Public/Rendering/SkeletalMeshLODRenderData.h
--- a/Engine/Source/Runtime/Engine/Public/Rendering/SkeletalMeshLODRenderData.h
+++ b/Engine/Source/Runtime/Engine/Public/Rendering/SkeletalMeshLODRenderData.h
@@ -142,7 +142,7 @@ private:
 	TArray<FBoneIndexType> Indices;
 };
 
-class FSkeletalMeshLODRenderData : public FRefCountBase
+class FSkeletalMeshLODRenderData : public FRefCountedObject
 {
 public:
 
diff --git a/Engine/Source/Runtime/Engine/Public/StaticMeshResources.h b/Engine/Source/Runtime/Engine/Public/StaticMeshResources.h
--- a/Engine/Source/Runtime/Engine/Public/StaticMeshResources.h
+++ b/Engine/Source/Runtime/Engine/Public/StaticMeshResources.h
@@ -422,7 +422,7 @@ struct TIsContiguousContainer<FStaticMeshSectionArray> : TIsContiguousContainer<
  * This structure is ref counted to allow the LOD streamer to evaluate the number of readers to it (readers that could access the CPU data).
  * Because the stream out clears the CPU readable data, CPU code that samples it must ensure to only reference LODs above CurrentFirstLODIdx.
  */
-struct FStaticMeshLODResources : public FRefCountBase
+struct FStaticMeshLODResources : public FRefCountedObject
 {
 public:
 
diff --git a/Engine/Source/Runtime/PakFile/Public/IPlatformFilePak.h b/Engine/Source/Runtime/PakFile/Public/IPlatformFilePak.h
--- a/Engine/Source/Runtime/PakFile/Public/IPlatformFilePak.h
+++ b/Engine/Source/Runtime/PakFile/Public/IPlatformFilePak.h
@@ -748,7 +748,7 @@ struct FPakEntryPair
 /**
  * Pak file.
  */
-class FPakFile : FNoncopyable, public FRefCountBase, public IPakFile
+class FPakFile : FNoncopyable, public FRefCountedObject, public IPakFile
 {
 public:
 	/** Index data that provides a map from the hash of a Filename to an FPakEntryLocation */
@@ -976,7 +976,7 @@ public:
 private:
 	/** Private destructor, use AddRef/Release instead */
 	PAKFILE_API virtual ~FPakFile();
-	friend class FRefCountBase;
+	friend class FRefCountedObject;
 
 public:
 
diff --git a/Engine/Source/Runtime/RenderCore/Public/RenderGraphResources.h b/Engine/Source/Runtime/RenderCore/Public/RenderGraphResources.h
--- a/Engine/Source/Runtime/RenderCore/Public/RenderGraphResources.h
+++ b/Engine/Source/Runtime/RenderCore/Public/RenderGraphResources.h
@@ -533,7 +533,7 @@ inline FRDGTextureDesc Translate(const FPooledRenderTargetDesc& InDesc);
 inline FPooledRenderTargetDesc Translate(const FRHITextureDesc& InDesc);
 
 class FRDGPooledTexture final
-	: public FRefCountBase
+	: public FRefCountedObject
 {
 public:
 	FRDGPooledTexture(FRHITexture* InTexture)
@@ -1192,7 +1192,7 @@ struct FRDGBufferUAVDesc final
 inline FRHIBufferCreateInfo Translate(const FRDGBufferDesc& InDesc);
 
 class FRDGPooledBuffer final
-	: public FRefCountBase
+	: public FRefCountedObject
 {
 public:
 	FRDGPooledBuffer(FRHICommandListBase& RHICmdList, TRefCountPtr<FRHIBuffer> InBuffer, const FRDGBufferDesc& InDesc, uint32 InNumAllocatedElements, const TCHAR* InName)
diff --git a/Engine/Source/Runtime/RenderCore/Public/ShaderCompilerJobTypes.h b/Engine/Source/Runtime/RenderCore/Public/ShaderCompilerJobTypes.h
--- a/Engine/Source/Runtime/RenderCore/Public/ShaderCompilerJobTypes.h
+++ b/Engine/Source/Runtime/RenderCore/Public/ShaderCompilerJobTypes.h
@@ -40,7 +40,7 @@ struct FShaderMapCompileResults
 
 struct FPendingShaderMapCompileResults
 	: public FShaderMapCompileResults
-	, public FRefCountBase
+	, public FRefCountedObject
 {};
 using FPendingShaderMapCompileResultsPtr = TRefCountPtr<FPendingShaderMapCompileResults>;
 
diff --git a/Engine/Source/Runtime/RenderCore/Public/ShaderCore.h b/Engine/Source/Runtime/RenderCore/Public/ShaderCore.h
--- a/Engine/Source/Runtime/RenderCore/Public/ShaderCore.h
+++ b/Engine/Source/Runtime/RenderCore/Public/ShaderCore.h
@@ -727,7 +727,7 @@ private:
 	} while(0)
 
 
-struct FSharedShaderCompilerEnvironment final : public FShaderCompilerEnvironment, public FRefCountBase
+struct FSharedShaderCompilerEnvironment final : public FShaderCompilerEnvironment, public FRefCountedObject
 {
 	virtual ~FSharedShaderCompilerEnvironment() = default;
 };
diff --git a/Engine/Source/Runtime/Renderer/Internal/MegaFiltersInterface.h b/Engine/Source/Runtime/Renderer/Internal/MegaFiltersInterface.h
--- a/Engine/Source/Runtime/Renderer/Internal/MegaFiltersInterface.h
+++ b/Engine/Source/Runtime/Renderer/Internal/MegaFiltersInterface.h
@@ -20,7 +20,7 @@ class IMegaFilters : public ISceneViewFamilyExtention
 public:
 
 	/** Ref counted history to be saved in the history*/
-	class IHistory : public FRefCountBase
+	class IHistory : public FRefCountedObject
 	{
 	public:
 		virtual ~IHistory() {}
diff --git a/Engine/Source/Runtime/Renderer/Private/SceneRendering.cpp b/Engine/Source/Runtime/Renderer/Private/SceneRendering.cpp
--- a/Engine/Source/Runtime/Renderer/Private/SceneRendering.cpp
+++ b/Engine/Source/Runtime/Renderer/Private/SceneRendering.cpp
@@ -3601,7 +3601,7 @@ struct FCrossGPUTarget
 	TArray<FCrossGPUTransfer> Transfers;
 };
 
-class FCrossGPUTransfersDeferred : public FRefCountBase
+class FCrossGPUTransfersDeferred : public FRefCountedObject
 {
 public:
 	TArray<FCrossGPUTarget> Targets;
diff --git a/Engine/Source/Runtime/Renderer/Private/Skinning/SkinningSceneExtension.h b/Engine/Source/Runtime/Renderer/Private/Skinning/SkinningSceneExtension.h
--- a/Engine/Source/Runtime/Renderer/Private/Skinning/SkinningSceneExtension.h
+++ b/Engine/Source/Runtime/Renderer/Private/Skinning/SkinningSceneExtension.h
@@ -215,7 +215,7 @@ private:
 
 	bool Tick(float DeltaTime);
 
-	struct FTickState : public FRefCountBase
+	struct FTickState : public FRefCountedObject
 	{
 		FVector CameraLocation = FVector::ZeroVector;
 	};
diff --git a/Engine/Source/Runtime/Renderer/Public/HairStrandsInterface.h b/Engine/Source/Runtime/Renderer/Public/HairStrandsInterface.h
--- a/Engine/Source/Runtime/Renderer/Public/HairStrandsInterface.h
+++ b/Engine/Source/Runtime/Renderer/Public/HairStrandsInterface.h
@@ -212,7 +212,7 @@ enum EHairInterpolationType
 ////////////////////////////////////////////////////////////////////////////////////////////////////
 // Public group data 
 
-struct FHairStrandsInstance : public FRefCountBase
+struct FHairStrandsInstance : public FRefCountedObject
 {
 	virtual ~FHairStrandsInstance() = default;
 	int32 RegisteredIndex = -1;
diff --git a/Engine/Source/Runtime/Renderer/Public/PathTracingDenoiser.h b/Engine/Source/Runtime/Renderer/Public/PathTracingDenoiser.h
--- a/Engine/Source/Runtime/Renderer/Public/PathTracingDenoiser.h
+++ b/Engine/Source/Runtime/Renderer/Public/PathTracingDenoiser.h
@@ -44,7 +44,7 @@ class IPathTracingSpatialTemporalDenoiser
 public:
 
 	/** Ref counted history to be saved in the history. */
-	class IHistory : public FRefCountBase //IRefCountedObject
+	class IHistory : public FRefCountedObject
 	{
 	public:
 		virtual ~IHistory() = default;
diff --git a/Engine/Source/Runtime/VulkanRHI/Private/VulkanTransientResourceAllocator.h b/Engine/Source/Runtime/VulkanRHI/Private/VulkanTransientResourceAllocator.h
--- a/Engine/Source/Runtime/VulkanRHI/Private/VulkanTransientResourceAllocator.h
+++ b/Engine/Source/Runtime/VulkanRHI/Private/VulkanTransientResourceAllocator.h
@@ -8,7 +8,7 @@
 
 class FVulkanTransientHeap final
 	: public FRHITransientHeap
-	, public FRefCountBase
+	, public FRefCountedObject
 {
 public:
 	FVulkanTransientHeap(const FInitializer& Initializer, FVulkanDevice& InDevice);
