Add EAllowShrinking arg to TIndexedContainerIterator::RemoveCurrent and TIndexedContainerIterator::RemoveCurrentSwap.

#Synced-CL Unknown

diff --git a/Engine/Source/Runtime/Core/Public/Containers/Array.h b/Engine/Source/Runtime/Core/Public/Containers/Array.h
--- a/Engine/Source/Runtime/Core/Public/Containers/Array.h
+++ b/Engine/Source/Runtime/Core/Public/Containers/Array.h
@@ -64,6 +64,8 @@ template <typename T, typename AllocatorType> inline void* operator new(size_t S
 template <typename ContainerType, typename ElementType, typename SizeType>
 class TIndexedContainerIterator
 {
+	using ContainerAllocatorType = typename ContainerType::AllocatorType;
+
 public:
 	UE_NODEBUG [[nodiscard]] TIndexedContainerIterator(ContainerType& InContainer, SizeType StartIndex = 0)
 		: Container(InContainer)
@@ -156,9 +158,9 @@ public:
 	}
 
 	/** Removes current element in array. This invalidates the current iterator value and it must be incremented */
-	UE_NODEBUG void RemoveCurrent()
+	UE_NODEBUG void RemoveCurrent(EAllowShrinking AllowShrinking = UE::Core::Private::AllowShrinkingByDefault<ContainerAllocatorType>())
 	{
-		Container.RemoveAt(Index);
+		Container.RemoveAt(Index, AllowShrinking);
 		Index--;
 	}
 
@@ -167,9 +169,9 @@ public:
 	 * This invalidates the current iterator value and it must be incremented.
 	 * Note this modifies the order of the remaining elements in the array.
 	 */
-	UE_NODEBUG void RemoveCurrentSwap()
+	UE_NODEBUG void RemoveCurrentSwap(EAllowShrinking AllowShrinking = UE::Core::Private::AllowShrinkingByDefault<ContainerAllocatorType>())
 	{
-		Container.RemoveAtSwap(Index);
+		Container.RemoveAtSwap(Index, AllowShrinking);
 		Index--;
 	}
 
diff --git a/Engine/Source/Runtime/Core/Public/Containers/IndirectArray.h b/Engine/Source/Runtime/Core/Public/Containers/IndirectArray.h
--- a/Engine/Source/Runtime/Core/Public/Containers/IndirectArray.h
+++ b/Engine/Source/Runtime/Core/Public/Containers/IndirectArray.h
@@ -21,6 +21,7 @@ class TIndirectArray
 public:
 	typedef T                        ElementType;
 	typedef TArray<void*, Allocator> InternalArrayType;
+	typedef Allocator 				 AllocatorType;
 
 	/** Default constructors. */
 	[[nodiscard]] TIndirectArray() = default;
