Add PID to default names for server utraces.

#Synced-CL 48364980

diff --git a/Engine/Source/Runtime/Core/Private/ProfilingDebugging/TraceAuxiliary.cpp b/Engine/Source/Runtime/Core/Private/ProfilingDebugging/TraceAuxiliary.cpp
--- a/Engine/Source/Runtime/Core/Private/ProfilingDebugging/TraceAuxiliary.cpp
+++ b/Engine/Source/Runtime/Core/Private/ProfilingDebugging/TraceAuxiliary.cpp
@@ -790,7 +790,11 @@ bool FTraceAuxiliaryImpl::FinalizeFilePath(const TCHAR* InPath, FString& OutPath
 	auto GetDefaultName = []
 		{
 			const FDateTime Now = FDateTime::Now();
+#if PLATFORM_LINUX
+			return Now.ToString(TEXT("%Y%m%d_%H%M%S_")).Appendf(TEXT("%06" INT64_X_FMT "_%06d.utrace"), (Now.GetTicks() % 10000000LL), FPlatformProcess::GetCurrentProcessId());
+#else
 			return Now.ToString(TEXT("%Y%m%d_%H%M%S_")).Appendf(TEXT("%06" INT64_X_FMT ".utrace"), (Now.GetTicks() % 10000000LL));
+#endif
 		};
 
 	if (Path.IsEmpty())
