[Core] Fix for FHashTable buffer overrun on resize. Added unit/automation test to verify correct behaviour

#Synced-CL 48390192

diff --git a/Engine/Source/Runtime/Core/Private/Containers/ContainersTest.cpp b/Engine/Source/Runtime/Core/Private/Containers/ContainersTest.cpp
--- a/Engine/Source/Runtime/Core/Private/Containers/ContainersTest.cpp
+++ b/Engine/Source/Runtime/Core/Private/Containers/ContainersTest.cpp
@@ -16,6 +16,7 @@
 #include "Stats/StatsMisc.h"
 #include "Templates/Tuple.h"
 #include "Templates/UniquePtr.h"
+#include "Containers/HashTable.h"
 
 #if WITH_DEV_AUTOMATION_TESTS
 
@@ -1577,6 +1578,109 @@ bool FContainersTSortedMapTest::RunTest(const FString& Parameters)
 	return !HasAnyErrors();
 }
 
+namespace UE::Private
+{
+	struct FHashTableHelper
+	{
+		// Given a hash table, a key and a list of values - test that the linked list of values for this key matches the passed in values
+		// Also tests that the list is correctly terminated to ~0u
+		// e.g.: 
+		//		FHashTable Table;
+		//		Table.Add(1, 1);
+		//		Table.Add(1, 2);
+		//		UE::Private::FHashTableHelper::TestHashList(TEXT("Description"), TestPtr, Table, 1, 1, 2);
+		template<typename... Values>
+		static void TestHashList(const TCHAR* TestString, FAutomationTestBase* Test, const FHashTable& InTable, uint32 First, uint32 Second, Values&&... Rest)
+		{
+			Test->TestEqual(TestString, InTable.First(First), Second);
+
+			if constexpr(sizeof...(Rest) > 0)
+			{
+				TestHashListInner(TestString, Test, InTable, Second, Rest...);
+			}
+			else
+			{
+				// Test end of list is correctly terminated
+				Test->TestEqual(TestString, InTable.Next(Second), ~0u);
+			}
+		}
+
+		// Secondary test as the first value requires a call to FHashTable::First and all subsequents require FHashTable::Next
+		template<typename... Values>
+		static void TestHashListInner(const TCHAR* TestString, FAutomationTestBase* Test, const FHashTable& InTable, uint32 First, uint32 Second, Values&&... Rest)
+		{
+			Test->TestEqual(TestString, InTable.Next(First), Second);
+
+			if constexpr(sizeof...(Rest) > 0)
+			{
+				TestHashListInner(TestString, Test, InTable, Second, Rest...);
+			}
+			else
+			{
+				// Test end of list is correctly terminated
+				Test->TestEqual(TestString, InTable.Next(Second), ~0u);
+			}
+		}
+	};
+}
+
+IMPLEMENT_SIMPLE_AUTOMATION_TEST(FContainersFHashTableTest, "System.Core.Containers.FHashTable", EAutomationTestFlags_ApplicationContextMask | EAutomationTestFlags::SmokeFilter)
+bool FContainersFHashTableTest::RunTest(const FString& Parameters)
+{
+	FHashTable HashTable(1024, 16);
+
+	// Resizing < size
+	{
+		// Key list for key == 1 -> [3, 10, 13, 2, 1]
+		HashTable.Add(1, 1);
+		HashTable.Add(1, 2);
+		HashTable.Add(1, 13);
+		HashTable.Add(1, 10);
+		HashTable.Add(1, 3);
+
+		UE::Private::FHashTableHelper::TestHashList(TEXT("Middle gap initial list is correct"), this, HashTable, 1, 3, 10, 13, 2, 1);
+		
+		// Key list for key == 1 -> [3, 2, 1]
+		HashTable.Resize(8);
+
+		UE::Private::FHashTableHelper::TestHashList(TEXT("Middle gap resized list is correct"), this, HashTable, 1, 3, 2, 1);
+
+		HashTable.Clear(1024, 16);
+
+		// Key list for key == 1 -> [1, 13, 12, 11, 10]
+		HashTable.Add(1, 10);
+		HashTable.Add(1, 11);
+		HashTable.Add(1, 12);
+		HashTable.Add(1, 13);
+		HashTable.Add(1, 1);
+
+		UE::Private::FHashTableHelper::TestHashList(TEXT("Open end list is correct"), this, HashTable, 1, 1, 13, 12, 11, 10);
+
+		// Key list for key == 1 -> [1]
+		HashTable.Resize(8);
+
+		UE::Private::FHashTableHelper::TestHashList(TEXT("Open end list is correct after resize"), this, HashTable, 1, 1);
+
+		HashTable.Clear(1024, 16);
+
+		// Key list for key == 1 -> [13, 12, 11, 10, 1]
+		HashTable.Add(1, 1);
+		HashTable.Add(1, 10);
+		HashTable.Add(1, 11);
+		HashTable.Add(1, 12);
+		HashTable.Add(1, 13);
+
+		UE::Private::FHashTableHelper::TestHashList(TEXT("Open begin list is correct"), this, HashTable, 1, 13, 12, 11, 10, 1);
+
+		// Key list for key == 1 -> [1]
+		HashTable.Resize(8);
+
+		UE::Private::FHashTableHelper::TestHashList(TEXT("Open begin list is correct after resize"), this, HashTable, 1, 1);
+	}
+
+	return !HasAnyErrors();
+}
+
 IMPLEMENT_SIMPLE_AUTOMATION_TEST(FContainersTArrayTest, "System.Core.Containers.TArray", EAutomationTestFlags_ApplicationContextMask | EAutomationTestFlags::SmokeFilter)
 bool FContainersTArrayTest::RunTest(const FString& Parameters)
 {
diff --git a/Engine/Source/Runtime/Core/Private/Containers/HashTable.cpp b/Engine/Source/Runtime/Core/Private/Containers/HashTable.cpp
--- a/Engine/Source/Runtime/Core/Private/Containers/HashTable.cpp
+++ b/Engine/Source/Runtime/Core/Private/Containers/HashTable.cpp
@@ -28,7 +28,31 @@ CORE_API void FHashTable::Resize( uint32 NewIndexSize )
 
 	if( NextIndex )
 	{
-		FMemory::Memcpy( NewNextIndex, NextIndex, IndexSize * 4 );
+		if(NewIndexSize < IndexSize)
+		{
+			// If the new size is less than the old size then the linked list of indices could have been invalidated
+			// Collapse gaps in the resulting lists for each hash
+			for(uint32 ListIndex = 0; ListIndex < HashSize; ++ListIndex)
+			{
+				uint32* Node = &Hash[ListIndex];
+
+				while(IsValid(*Node))
+				{
+					// For any list node >= NewIndex size: Next = Next->Next
+					if(*Node >= NewIndexSize)
+					{
+						*Node = NextIndex[*Node];
+
+						continue;
+					}
+
+					Node = &NextIndex[*Node];
+				}
+			}
+		}
+
+		FMemory::Memcpy( NewNextIndex, NextIndex, FMath::Min(IndexSize, NewIndexSize) * 4 );
+
 		delete[] NextIndex;
 	}
 	
