Add FAppTime::Get_UnsafeNoCheck() function
- Temporary function to avoid triggering the ensure in FAppTime::Get() for existing engine systems that read the FAppTime value in a non-threadsafe manner.

#Synced-CL 48409310

diff --git a/Engine/Plugins/Animation/LiveLink/Source/LiveLink/Private/LiveLinkClient.cpp b/Engine/Plugins/Animation/LiveLink/Source/LiveLink/Private/LiveLinkClient.cpp
--- a/Engine/Plugins/Animation/LiveLink/Source/LiveLink/Private/LiveLinkClient.cpp
+++ b/Engine/Plugins/Animation/LiveLink/Source/LiveLink/Private/LiveLinkClient.cpp
@@ -360,8 +360,13 @@ void FLiveLinkClient::RebroadcastSubjects(const TSet<FLiveLinkSubjectKey>& Updat
 void FLiveLinkClient::CacheValues()
 {
 #if WITH_EDITOR
-	CachedEngineTime = FApp::GetCurrentTime();
-	CachedEngineFrameTime = FApp::GetCurrentFrameTime();
+
+PRAGMA_DISABLE_DEPRECATION_WARNINGS
+	const FAppTime& Time = FAppTime::Get_UnsafeNoCheck();
+PRAGMA_ENABLE_DEPRECATION_WARNINGS
+
+	CachedEngineTime      = Time.CurrentTime;
+	CachedEngineFrameTime = Time.CurrentFrameTime;
 #endif
 }
 
diff --git a/Engine/Plugins/VirtualProduction/StageMonitoring/Source/StageDataProvider/Private/FramePerformanceProvider.cpp b/Engine/Plugins/VirtualProduction/StageMonitoring/Source/StageDataProvider/Private/FramePerformanceProvider.cpp
--- a/Engine/Plugins/VirtualProduction/StageMonitoring/Source/StageDataProvider/Private/FramePerformanceProvider.cpp
+++ b/Engine/Plugins/VirtualProduction/StageMonitoring/Source/StageDataProvider/Private/FramePerformanceProvider.cpp
@@ -61,10 +61,9 @@ FFramePerformanceProviderMessage GetLatestPerformanceData()
 	const float GameThreadWaitTime = FPlatformTime::ToMilliseconds(GGameThreadWaitTime);
 	const float RenderThreadTime = FPlatformTime::ToMilliseconds(GRenderThreadTime);
 	const float RenderThreadWaitTime = FPlatformTime::ToMilliseconds(GRenderThreadWaitTime);
-
-	// @todo delta time refactor - The FApp::GetIdleTime function requires proper pipelining of the delta time values from the game thread. Disabling this value for now.
-	// The above stats are also not threadsafe.
-	const float IdleTimeMilli = 0; // (FApp::GetIdleTime() * 1000.0);
+PRAGMA_DISABLE_DEPRECATION_WARNINGS
+	const float IdleTimeMilli = (FAppTime::Get_UnsafeNoCheck().IdleTime * 1000.0);
+PRAGMA_ENABLE_DEPRECATION_WARNINGS
 
 #if STATS
 	FPlatformMemoryStats Stats = FPlatformMemory::GetStats();
diff --git a/Engine/Source/Runtime/Core/Private/Misc/AppTime.cpp b/Engine/Source/Runtime/Core/Private/Misc/AppTime.cpp
--- a/Engine/Source/Runtime/Core/Private/Misc/AppTime.cpp
+++ b/Engine/Source/Runtime/Core/Private/Misc/AppTime.cpp
@@ -46,6 +46,19 @@ CORE_API const FAppTime& FAppTime::Get()
 	}
 }
 
+CORE_API const FAppTime& FAppTime::Get_UnsafeNoCheck()
+{
+	// Prefer the TLS value if we have one.
+	if (TLSInstance.IsValid())
+	{
+		return *TLSInstance;
+	}
+
+	// Unsynchronized read from the game thread's copy of FAppTime.
+	// This is a data race if read from any thread except the game thread.
+	return GameThreadInstance;
+}
+
 void FAppTime::Set(const FAppTime& NewValue)
 {
 	checkf(IsInGameThread(), TEXT("Only the game thread may set the current FAppTime value, to ensure parallel tasks receive the correct value via task context inheritance."));
diff --git a/Engine/Source/Runtime/Core/Public/Misc/AppTime.h b/Engine/Source/Runtime/Core/Public/Misc/AppTime.h
--- a/Engine/Source/Runtime/Core/Public/Misc/AppTime.h
+++ b/Engine/Source/Runtime/Core/Public/Misc/AppTime.h
@@ -76,11 +76,27 @@ public:
 		return FFrameRate();
 	}
 
+	//
+	// Returns the FAppTime value for the current thread's time context. May be called by the game thread, or decendent TaskGraph tasks of the game thread.
+	// Calling this function from any other context will result in an ensure failure. This is because the game thread is the authority on the FAppTime value,
+	// and this value is pipelined through the TaskGraph so that decendent tasks receive the correct time value for the frame they belong to.
+	// Reading the FAppTime value outside of these contexts is a data race with the game thread.
+	//
 	static CORE_API const FAppTime& Get();
 
+	//
+	// Sets the new FAppTime value.
 	// May only be called from the game thread
+	//
 	static CORE_API void Set(const FAppTime& NewValue);
 
+	UE_DEPRECATED(5.8,
+		"Do not call FAppTime::Get_UnsafeNoCheck() in new engine code. This function is provided "
+		"as a temporary measure to avoid triggering the time context ensure in FAppTime::Get() "
+		"for engine systems that read the FAppTime value in a non-threadsafe manner, until those "
+		"systems can be refactored.")
+	static CORE_API const FAppTime& Get_UnsafeNoCheck();
+
 private:
 	static FAppTime GameThreadInstance;
 	static thread_local TSharedPtr<const FAppTime> TLSInstance;
diff --git a/Engine/Source/Runtime/VirtualProduction/StageDataCore/Public/IStageDataProvider.h b/Engine/Source/Runtime/VirtualProduction/StageDataCore/Public/IStageDataProvider.h
--- a/Engine/Source/Runtime/VirtualProduction/StageDataCore/Public/IStageDataProvider.h
+++ b/Engine/Source/Runtime/VirtualProduction/StageDataCore/Public/IStageDataProvider.h
@@ -79,8 +79,11 @@ private:
 
 			static_assert(TIsDerivedFrom<MessageType, FStageProviderEventMessage>::IsDerived || TIsDerivedFrom<MessageType, FStageProviderPeriodicMessage>::IsDerived, "MessageType must be a FStageProviderEventMessage or FStageProviderPeriodicMessage derived UStruct.");
 
+		PRAGMA_DISABLE_DEPRECATION_WARNINGS
 			// Timestamp the outgoing message.
-			TOptional<FQualifiedFrameTime> CurrentFrameTime = FApp::GetCurrentFrameTime();
+			TOptional<FQualifiedFrameTime> CurrentFrameTime = FAppTime::Get_UnsafeNoCheck().CurrentFrameTime;
+		PRAGMA_ENABLE_DEPRECATION_WARNINGS
+
 			Message.DateTime = FDateTime::Now();
 
 			if (CurrentFrameTime.IsSet())
