Move CompactTSet metadata into CSVProfiler

#Synced-CL 48454375

diff --git a/Engine/Source/Runtime/Core/Private/ProfilingDebugging/CsvProfiler.cpp b/Engine/Source/Runtime/Core/Private/ProfilingDebugging/CsvProfiler.cpp
--- a/Engine/Source/Runtime/Core/Private/ProfilingDebugging/CsvProfiler.cpp
+++ b/Engine/Source/Runtime/Core/Private/ProfilingDebugging/CsvProfiler.cpp
@@ -3649,6 +3649,7 @@ FCsvProfiler::FCsvProfiler()
 	SetMetadataInternal(TEXT("LTOEnabled"), PLATFORM_COMPILER_OPTIMIZATION_LTCG ? TEXT("1") : TEXT("0"));//True if Link Time Optimisation is enabled 
 	SetMetadataInternal(TEXT("ASan"), USING_ADDRESS_SANITISER ? TEXT("1") : TEXT("0"));
 	SetMetadataInternal(TEXT("MergedModules"), UE_MERGED_MODULES ? TEXT("1") : TEXT("0"));
+	SetMetadataInternal(TEXT("CompactTSet"), UE_USE_COMPACT_SET_AS_DEFAULT ? TEXT("1") : TEXT("0"));
 
 	FCoreDelegates::OnSystemResolutionChanged.AddLambda([](uint32 ResX, uint32 ResY)
 		{
