Change error warning of corruption to a structured log so it can grab context.

#Synced-CL Unknown

diff --git a/Engine/Source/Runtime/Core/Private/Serialization/Archive.cpp b/Engine/Source/Runtime/Core/Private/Serialization/Archive.cpp
--- a/Engine/Source/Runtime/Core/Private/Serialization/Archive.cpp
+++ b/Engine/Source/Runtime/Core/Private/Serialization/Archive.cpp
@@ -8,6 +8,7 @@
 #include "Containers/UnrealString.h"
 #include "UObject/NameTypes.h"
 #include "Logging/LogMacros.h"
+#include "Logging/StructuredLog.h"
 #include "Misc/Parse.h"
 #include "UObject/ObjectVersion.h"
 #include "Serialization/NameAsStringProxyArchive.h"
@@ -565,7 +566,7 @@ void FArchive::SerializeBool( bool& D )
 
 	if (OldUBoolValue > 1)
 	{
-		UE_LOG(LogSerialization, Error, TEXT("Invalid boolean encountered while reading archive %s - stream is most likely corrupted."), *GetArchiveName());
+		UE_LOGFMT(LogSerialization, Error, "Invalid boolean encountered while reading archive {ArchiveName} - stream is most likely corrupted.", GetArchiveName());
 
 		this->SetError();
 	}
