Another comment fix - we don't currently disregard custom versions in this config

#rnx
#Synced-CL 48375597

diff --git a/Engine/Source/Runtime/Core/Public/Misc/Build.h b/Engine/Source/Runtime/Core/Public/Misc/Build.h
--- a/Engine/Source/Runtime/Core/Public/Misc/Build.h
+++ b/Engine/Source/Runtime/Core/Public/Misc/Build.h
@@ -543,7 +543,7 @@
 // 
 //	1.	Assume no imported packages need to actually be loaded, and can therefore be universally mocked, 
 //		hence 'universal placeholders' - see FLinkerLoad::CreateUniversalPlaceholders
-//	2.	Accomodate missing types with native stream serializers by adding extra bounds checking to serialization
+//	2.	Accommodate missing types with native stream serializers by adding extra bounds checking to serialization
 //		This extra bounds checking would waste critical main thread time in other configurations, but in the 
 //		the long run the workloads can be moved to other threads - see Archive.h, PropertyTag.cpp
 //	3.	Provide safe behavior to placeholders which would need their type to 'emerge' from usage. 
@@ -556,8 +556,6 @@
 //  6.	Assumes a read operation, avoiding expensive mutators - see PropertyBagRepository.cpp
 //	7.	Allows serialization routines to expect unreachable objects (in the long run all loaded 
 //		data should be made reachable) - see JsonStringifyImpl.cpp
-//  8.	Allows loading of assets with a newer custom version because we are not running user native
-//		serialize routines
 //
 // All of these usages will be accomodated more elegantly by future work. Do not use this define.
 #ifndef UE_INTERNAL_UNIVERSAL_PLACEHOLDERS
