Added FPlatformString::FormatV and TCString::FormatV to provide GetVarArgs with a usable return value

Updated TStringBuilder::Appendf to use the return value to resize distinct from the error case. Previously any error would cause a resize loop until crashing due to OOM.

#Synced-CL 0

diff --git a/Engine/Source/Runtime/Core/Private/GenericPlatform/GenericPlatformString.cpp b/Engine/Source/Runtime/Core/Private/GenericPlatform/GenericPlatformString.cpp
--- a/Engine/Source/Runtime/Core/Private/GenericPlatform/GenericPlatformString.cpp
+++ b/Engine/Source/Runtime/Core/Private/GenericPlatform/GenericPlatformString.cpp
@@ -43,28 +43,30 @@ namespace
 		TSafeDestIterator(const TSafeDestIterator&) = delete;
 		TSafeDestIterator& operator=(const TSafeDestIterator&) = delete;
 
-		explicit TSafeDestIterator(CharType* InPtr, SIZE_T InEnd)
+		explicit TSafeDestIterator(CharType* InPtr, SIZE_T InSize)
 			: NumCharsWritten(0)
-			, Ptr(InPtr)
-			, EndMinusOne(InPtr + InEnd - 1)
+			, Ptr(InSize ? InPtr : nullptr)
+			, EndMinusOne(InSize ? InPtr + InSize - 1 : nullptr)
 		{
-			check(InPtr);
 		}
 
 		FORCENOINLINE ~TSafeDestIterator()
 		{
 			// Auto-terminate buffer
-			*Ptr = 0;
+			if (Ptr)
+			{
+				*Ptr = 0;
+			}
 		}
 
 		// Writes Count instances of Ch to the destination buffer
 		FORCENOINLINE TSafeDestIterator& Write(CharType Ch, int32 Count = 1)
 		{
+			NumCharsWritten += Count;
+
 			// Ensure we're not in the error state
 			if (EndMinusOne)
 			{
-				NumCharsWritten += Count;
-
 				int32 NumToWrite = FPlatformMath::Min(Count, (int32)(EndMinusOne - Ptr));
 				for (int32 Val = NumToWrite; Val; --Val)
 				{
@@ -73,7 +75,7 @@ namespace
 
 				if (NumToWrite != Count)
 				{
-					EndMinusOne = 0;
+					EndMinusOne = nullptr;
 				}
 			}
 
@@ -84,35 +86,39 @@ namespace
 		template <typename OtherCharType>
 		FORCENOINLINE TSafeDestIterator& Write(const OtherCharType* Src, int32 Count)
 		{
-			// Ensure we're not in the error state
-			if (EndMinusOne)
+			using DstType = std::conditional_t<std::is_same_v<ANSICHAR, CharType>, UTF8CHAR, CharType>;
+			using SrcType = std::conditional_t<std::is_same_v<ANSICHAR, OtherCharType>, UTF8CHAR, OtherCharType>;
+			if constexpr (!std::is_same_v<DstType, SrcType>)
 			{
-				using DstType = std::conditional_t<std::is_same_v<ANSICHAR, CharType>, UTF8CHAR, CharType>;
-				using SrcType = std::conditional_t<std::is_same_v<ANSICHAR, OtherCharType>, UTF8CHAR, OtherCharType>;
-				if constexpr (!std::is_same_v<DstType, SrcType>)
-				{
-					int32 ConvertedLength = FPlatformString::ConvertedLength<DstType, SrcType>((const SrcType*)Src, Count);
-					int32 NumToWrite = FPlatformMath::Min(ConvertedLength, (int32)(EndMinusOne - Ptr));
-					CharType* Result = (CharType*)FPlatformString::Convert((DstType*)Ptr, NumToWrite, (const SrcType*)Src, Count);
+				int32 ConvertedLength = FPlatformString::ConvertedLength<DstType, SrcType>((const SrcType*)Src, Count);
+				NumCharsWritten += ConvertedLength;
 
-					NumCharsWritten += ConvertedLength;
-
-					if (NumToWrite != ConvertedLength)
+				// Ensure we're not in the error state
+				if (EndMinusOne)
+				{
+					// Copy the source only if the destination will fit it completely because
+					// the conversion has undefined output when the destination is too small.
+					if (ConvertedLength <= EndMinusOne - Ptr)
 					{
-						Ptr = Ptr + NumToWrite;
-						EndMinusOne = 0;
+						Ptr = (CharType*)FPlatformString::Convert((DstType*)Ptr, ConvertedLength, (const SrcType*)Src, Count);
+						check(Ptr);
 					}
 					else
 					{
-						Ptr = Result;
+						EndMinusOne = nullptr;
 					}
 				}
-				else
-				{
-					NumCharsWritten += Count;
+			}
+			else
+			{
+				NumCharsWritten += Count;
 
+				// Ensure we're not in the error state
+				if (EndMinusOne)
+				{
 					int32 NumToWrite = FPlatformMath::Min(Count, (int32)(EndMinusOne - Ptr));
 
+					// Copy as much of the source as will fit in the destination
 					for (int32 Val = NumToWrite; Val; --Val)
 					{
 						*Ptr++ = *Src++;
@@ -120,7 +126,7 @@ namespace
 
 					if (NumToWrite != Count)
 					{
-						EndMinusOne = 0;
+						EndMinusOne = nullptr;
 					}
 				}
 			}
@@ -128,11 +134,6 @@ namespace
 			return *this;
 		}
 
-		explicit operator bool() const
-		{
-			return !!EndMinusOne;
-		}
-
 		int32 GetNumCharsWritten() const
 		{
 			return NumCharsWritten;
@@ -213,10 +214,7 @@ namespace
 		{
 			if (*Src != '%')
 			{
-				if (!DestIter.Write(*Src++))
-				{
-					return -1;
-				}
+				DestIter.Write(*Src++);
 				continue;
 			}
 
@@ -295,10 +293,7 @@ namespace
 			case '%':
 			{
 				Src++;
-				if (!DestIter.Write('%'))
-				{
-					return -1;
-				}
+				DestIter.Write('%');
 				break;
 			}
 
@@ -306,19 +301,14 @@ namespace
 			{
 				CharType Val = (CharType)va_arg(ArgPtr, int);
 				Src++;
-				bool bSuccess = true;
 				if (FieldLen > 1)
 				{
-					bSuccess = !!DestIter.Write(CharType(' '), FieldLen - 1);
+					DestIter.Write(CharType(' '), FieldLen - 1);
 				}
-				bSuccess = bSuccess && !!DestIter.Write(Val);
+				DestIter.Write(Val);
 				if (FieldLen < -1)
 				{
-					bSuccess = bSuccess && !!DestIter.Write(CharType(' '), FPlatformMath::Abs(FieldLen + 1));
-				}
-				if (!bSuccess)
-				{
-					return -1;
+					DestIter.Write(CharType(' '), FPlatformMath::Abs(FieldLen + 1));
 				}
 				break;
 			}
@@ -359,10 +349,11 @@ namespace
 				FmtBuf[CpyIdx] = 0;
 
 				int RetCnt = snprintf(AnsiNum, sizeof(AnsiNum), FmtBuf, Val);
-				if (RetCnt < 0 || !DestIter.Write(AnsiNum, RetCnt))
+				if (RetCnt < 0)
 				{
 					return -1;
 				}
+				DestIter.Write(AnsiNum, RetCnt);
 				break;
 			}
 
@@ -408,10 +399,11 @@ namespace
 				FmtBuf[CpyIdx] = 0;
 
 				int RetCnt = snprintf(AnsiNum, sizeof(AnsiNum), FmtBuf, Val);
-				if (RetCnt < 0 || !DestIter.Write(AnsiNum, RetCnt))
+				if (RetCnt < 0)
 				{
 					return -1;
 				}
+				DestIter.Write(AnsiNum, RetCnt);
 				break;
 			}
 
@@ -433,10 +425,11 @@ namespace
 				FmtBuf[CpyIdx] = 0;
 
 				int RetCnt = snprintf(AnsiNum, sizeof(AnsiNum), FmtBuf, Val);
-				if (RetCnt < 0 || !DestIter.Write(AnsiNum, RetCnt))
+				if (RetCnt < 0)
 				{
 					return -1;
 				}
+				DestIter.Write(AnsiNum, RetCnt);
 				break;
 			}
 
@@ -472,10 +465,11 @@ namespace
 					FmtBuf[CpyIdx] = 0;
 
 					int RetCnt = snprintf(AnsiNum, sizeof(AnsiNum), FmtBuf, Val);
-					if (RetCnt < 0 || !DestIter.Write(AnsiNum, RetCnt))
+					if (RetCnt < 0)
 					{
 						return -1;
 					}
+					DestIter.Write(AnsiNum, RetCnt);
 					break;
 				}
 				// Treat %lf as a %f
@@ -497,10 +491,11 @@ namespace
 					FmtBuf[CpyIdx] = 0;
 
 					int RetCnt = snprintf(AnsiNum, sizeof(AnsiNum), FmtBuf, Val);
-					if (RetCnt < 0 || !DestIter.Write(AnsiNum, RetCnt))
+					if (RetCnt < 0)
 					{
 						return -1;
 					}
+					DestIter.Write(AnsiNum, RetCnt);
 					break;
 				}
 
@@ -537,10 +532,11 @@ namespace
 					FmtBuf[CpyIdx] = 0;
 
 					int RetCnt = snprintf(AnsiNum, sizeof(AnsiNum), FmtBuf, Val);
-					if (RetCnt < 0 || !DestIter.Write(AnsiNum, RetCnt))
+					if (RetCnt < 0)
 					{
 						return -1;
 					}
+					DestIter.Write(AnsiNum, RetCnt);
 					break;
 				}
 				else if (Src[0] == 'l')
@@ -605,10 +601,11 @@ namespace
 				FmtBuf[CpyIdx] = 0;
 
 				int RetCnt = snprintf(AnsiNum, sizeof(AnsiNum), FmtBuf, Val);
-				if (RetCnt < 0 || !DestIter.Write(AnsiNum, RetCnt))
+				if (RetCnt < 0)
 				{
 					return -1;
 				}
+				DestIter.Write(AnsiNum, RetCnt);
 				break;
 			}
 
@@ -646,10 +643,11 @@ namespace
 					AnsiNum[UE_ARRAY_COUNT(AnsiNum) - 1] = '\0';
 					checkf(0, TEXT("Attempting to read past the size our buffer. Buffer Size: %" SIZE_T_FMT " Size to read: %d. Current contents: '%s'\n"), UE_ARRAY_COUNT(AnsiNum), RetCnt, UTF8_TO_TCHAR(AnsiNum));
 				}
-				if (RetCnt < 0 || !DestIter.Write(AnsiNum, RetCnt))
+				if (RetCnt < 0)
 				{
 					return -1;
 				}
+				DestIter.Write(AnsiNum, RetCnt);
 				break;
 			}
 
@@ -663,10 +661,6 @@ namespace
 				{
 					ProcessStringArg<TCHAR>(DestIter, Src, FieldLen, PrecisionLen, ArgPtr);
 				}
-				if (!DestIter)
-				{
-					return -1;
-				}
 				break;
 			}
 
@@ -674,12 +668,7 @@ namespace
 			{
 				// The %S format represents a string which is the opposite of %s - wide if CharType is narrow, or narrow if CharType is wide
 				using OtherCharType = std::conditional_t<std::is_same_v<CharType, ANSICHAR>, WIDECHAR, ANSICHAR>;
-
 				ProcessStringArg<OtherCharType>(DestIter, Src, FieldLen, PrecisionLen, ArgPtr);
-				if (!DestIter)
-				{
-					return -1;
-				}
 				break;
 			}
 
@@ -1095,6 +1084,8 @@ namespace UE::Core::Private
 	template <typename DestType, typename DestBufferType, typename SourceEndType>
 	static int32 ConvertFromUTF8(DestBufferType& ConvertedBuffer, int32 DestLen, const UTF8CHAR* Source, SourceEndType SourceEnd)
 	{
+		check(DestLen >= 0);
+
 		DestBufferType DestStartingPosition = ConvertedBuffer;
 
 		const uint64 ExtendedCharMask = 0x8080808080808080;
@@ -1389,12 +1380,12 @@ int32 StrncmpImpl(const CharType1* String1, const CharType2* String2, SIZE_T Cou
 
 }
 
-int32 FGenericPlatformString::GetVarArgs(ANSICHAR* Dest, SIZE_T DestSize, const ANSICHAR*& Fmt, va_list ArgPtr)
+int32 FGenericPlatformString::FormatV(ANSICHAR* Dest, SIZE_T DestSize, const ANSICHAR* Fmt, va_list ArgPtr)
 {
 	return GetVarArgsImpl(Dest, DestSize, Fmt, ArgPtr);
 }
 
-int32 FGenericPlatformString::GetVarArgs(WIDECHAR* Dest, SIZE_T DestSize, const WIDECHAR*& Fmt, va_list ArgPtr)
+int32 FGenericPlatformString::FormatV(WIDECHAR* Dest, SIZE_T DestSize, const WIDECHAR* Fmt, va_list ArgPtr)
 {
 	return GetVarArgsImpl(Dest, DestSize, Fmt, ArgPtr);
 }
@@ -1417,32 +1408,72 @@ int32 FGenericPlatformString::Strncmp(const UTF8CHAR* Str1, const UTF8CHAR* Str2
 #include <catch2/generators/catch_generators.hpp>
 
 template <int32 OutputSize, typename ExpectedCharType, typename FormatCharType>
-static void TestGetVarArgs(const ExpectedCharType* Expected, const FormatCharType* Format, ...)
+static void TestFormatV(const ExpectedCharType* Expected, const FormatCharType* Format, ...)
 {
 	FormatCharType OutputString[OutputSize];
 
-	va_list ArgPtr;
-	va_start(ArgPtr, Format);
-	int32 Result = FGenericPlatformString::GetVarArgs(OutputString, OutputSize, Format, ArgPtr);
-	va_end(ArgPtr);
+	int32 Result;
+	{
+		va_list ArgPtr;
+		va_start(ArgPtr, Format);
+		Result = FGenericPlatformString::FormatV(OutputString, OutputSize, Format, ArgPtr);
+		va_end(ArgPtr);
+	}
 
 	TString<ExpectedCharType> ExpectedString(Expected);
 	CAPTURE(ExpectedString);
-	CHECK(ExpectedString == TString<ExpectedCharType>(OutputString));
 	CHECK(ExpectedString.Len() == Result);
+	CHECK(ExpectedString.Left(OutputSize - 1) == TString<ExpectedCharType>(OutputString));
+	CHECK(OutputString[FPlatformMath::Min(Result, OutputSize - 1)] == 0);
+
+	if (!ExpectedString.IsEmpty())
+	{
+		static_assert(OutputSize > 1);
+		OutputString[0] = '.';
+		OutputString[1] = '.';
+
+		va_list ArgPtr;
+		va_start(ArgPtr, Format);
+		Result = FGenericPlatformString::FormatV(OutputString, 1, Format, ArgPtr);
+		va_end(ArgPtr);
+
+		CHECK(ExpectedString.Len() == Result);
+		CHECK(OutputString[0] == 0);
+		CHECK(OutputString[1] == '.');
+	}
 }
 
-TEST_CASE("Core::GenericPlatformString::GetVarArgs", "[Core][String][Smoke]")
+template <typename CharType>
+static int32 TestFormat(CharType* Dest, SIZE_T DestSize, const CharType* Format, ...)
+{
+	va_list Args;
+	va_start(Args, Format);
+	int32 Result = FGenericPlatformString::FormatV(Dest, DestSize, Format, Args);
+	va_end(Args);
+	return Result;
+}
+
+template <typename CharType>
+static int32 TestGetVarArgs(CharType* Dest, SIZE_T DestSize, const CharType* Format, ...)
+{
+	va_list Args;
+	va_start(Args, Format);
+	int32 Result = FGenericPlatformString::GetVarArgs(Dest, DestSize, Format, Args);
+	va_end(Args);
+	return Result;
+}
+
+TEST_CASE("Core::GenericPlatformString::FormatV", "[Core][String][Smoke]")
 {
 	constexpr int32 OutputSize = 256;
 
-	TestGetVarArgs<OutputSize>(TEXT("Test A|LEFT                |               RIGHT|     33.33|66.67     |"), TEXT("Test A|%-20s|%20s|%10.2f|%-10.2f|"), TEXT("LEFT"), TEXT("RIGHT"), 33.333333, 66.666666);
-	TestGetVarArgs<OutputSize>(TEXT("Test B|Percents:%%%3|"), TEXT("Test B|Percents:%%%%%%%d|"), 3);
-	TestGetVarArgs<OutputSize>(TEXT("Test C|12345|54321|123ABC|f|99|"), TEXT("Test C|%d|%i|%X|%x|%u|"), 12345, 54321, 0x123AbC, 15, 99);
+	TestFormatV<OutputSize>(TEXT("Test A|LEFT                |               RIGHT|     33.33|66.67     |"), TEXT("Test A|%-20s|%20s|%10.2f|%-10.2f|"), TEXT("LEFT"), TEXT("RIGHT"), 33.333333, 66.666666);
+	TestFormatV<OutputSize>(TEXT("Test B|Percents:%%%3|"), TEXT("Test B|Percents:%%%%%%%d|"), 3);
+	TestFormatV<OutputSize>(TEXT("Test C|12345|54321|123ABC|f|99|"), TEXT("Test C|%d|%i|%X|%x|%u|"), 12345, 54321, 0x123AbC, 15, 99);
 
 #ifdef _MSC_VER
 	// MSVC's standard library formats pointers differently
-	TestGetVarArgs<OutputSize>(TEXT("Test D|0000000000012345|"), TEXT("Test D|%p|"), 0x12345);
+	TestFormatV<OutputSize>(TEXT("Test D|0000000000012345|"), TEXT("Test D|%p|"), 0x12345);
 #else
 	// Pointer format (which we get from snprintf) varies from platform to platform.
 	// Call it to make sure it puts the output in the proper place, but don't test for the exact format
@@ -1455,28 +1486,28 @@ TEST_CASE("Core::GenericPlatformString::GetVarArgs", "[Core][String][Smoke]")
 	}
 #endif
 
-	TestGetVarArgs<OutputSize>(TEXT("Test E|12345678912345|"), TEXT("Test E|%" INT64_FMT "|"), int64(12345678912345LL));
-	TestGetVarArgs<OutputSize>(TEXT("Test F|123.456000|1.234560e+02|123.456|"), TEXT("Test F|%f|%e|%g|"), 123.456, 123.456, 123.456);
-	TestGetVarArgs<OutputSize>(TEXT("Test G|C0DE|"), TEXT("Test G|%" UPTRINT_X_FMT"|"), UPTRINT(49374));
-	TestGetVarArgs<OutputSize>(TEXT("Test H|c0de|"), TEXT("Test H|%" UPTRINT_x_FMT "|"), UPTRINT(49374));
-	TestGetVarArgs<OutputSize>(TEXT("Test I|18446744073709551615|"), TEXT("Test I|%" UINT64_FMT "|"), MAX_uint64);
-	TestGetVarArgs<OutputSize>(TEXT("Test J|         J|"), TEXT("Test J|%*c|"), 10, 'J');
-	TestGetVarArgs<OutputSize>(TEXT("Test K|K    |"), TEXT("Test K|%-5c|"), 'K');
-	TestGetVarArgs<OutputSize>(TEXT("Test L|Hello World|"), TEXT("Test L|%hs|"), UTF8TEXT("Hello World"));
-	TestGetVarArgs<OutputSize>(TEXT("Test M|\U0001F606\U0001F41C|"), TEXT("Test M|%hs|"), UTF8TEXT("\U0001F606\U0001F41C"));
-	TestGetVarArgs<OutputSize>(TEXT("Test N|\U0001F606\U0001F41C|"), TEXT("Test N|%ls|"), TEXT("\U0001F606\U0001F41C"));
-	TestGetVarArgs<OutputSize>(TEXT("Test O|\U0001F606\U0001F41C|"), TEXT("Test O|%s|"), TEXT("\U0001F606\U0001F41C"));
-	TestGetVarArgs<OutputSize>(TEXT("Test P|\U0001F606\U0001F41C|"), TEXT("Test P|%S|"), UTF8TEXT("\U0001F606\U0001F41C"));
-
-	TestGetVarArgs<9>(TEXT("12345678"), TEXT("%08x"), 0x12345678);
-
-	TestGetVarArgs<OutputSize>("Test A|LEFT                |               RIGHT|     33.33|66.67     |", "Test A|%-20s|%20s|%10.2f|%-10.2f|", "LEFT", "RIGHT", 33.333333, 66.666666);
-	TestGetVarArgs<OutputSize>("Test B|Percents:%%%3|", "Test B|Percents:%%%%%%%d|", 3);
-	TestGetVarArgs<OutputSize>("Test C|12345|54321|123ABC|f|99|", "Test C|%d|%i|%X|%x|%u|", 12345, 54321, 0x123AbC, 15, 99);
+	TestFormatV<OutputSize>(TEXT("Test E|12345678912345|"), TEXT("Test E|%" INT64_FMT "|"), int64(12345678912345LL));
+	TestFormatV<OutputSize>(TEXT("Test F|123.456000|1.234560e+02|123.456|"), TEXT("Test F|%f|%e|%g|"), 123.456, 123.456, 123.456);
+	TestFormatV<OutputSize>(TEXT("Test G|C0DE|"), TEXT("Test G|%" UPTRINT_X_FMT"|"), UPTRINT(49374));
+	TestFormatV<OutputSize>(TEXT("Test H|c0de|"), TEXT("Test H|%" UPTRINT_x_FMT "|"), UPTRINT(49374));
+	TestFormatV<OutputSize>(TEXT("Test I|18446744073709551615|"), TEXT("Test I|%" UINT64_FMT "|"), MAX_uint64);
+	TestFormatV<OutputSize>(TEXT("Test J|         J|"), TEXT("Test J|%*c|"), 10, 'J');
+	TestFormatV<OutputSize>(TEXT("Test K|K    |"), TEXT("Test K|%-5c|"), 'K');
+	TestFormatV<OutputSize>(TEXT("Test L|Hello World|"), TEXT("Test L|%hs|"), UTF8TEXT("Hello World"));
+	TestFormatV<OutputSize>(TEXT("Test M|\U0001F606\U0001F41C|"), TEXT("Test M|%hs|"), UTF8TEXT("\U0001F606\U0001F41C"));
+	TestFormatV<OutputSize>(TEXT("Test N|\U0001F606\U0001F41C|"), TEXT("Test N|%ls|"), TEXT("\U0001F606\U0001F41C"));
+	TestFormatV<OutputSize>(TEXT("Test O|\U0001F606\U0001F41C|"), TEXT("Test O|%s|"), TEXT("\U0001F606\U0001F41C"));
+	TestFormatV<OutputSize>(TEXT("Test P|\U0001F606\U0001F41C|"), TEXT("Test P|%S|"), UTF8TEXT("\U0001F606\U0001F41C"));
+
+	TestFormatV<9>(TEXT("12345678"), TEXT("%08x"), 0x12345678);
+
+	TestFormatV<OutputSize>("Test A|LEFT                |               RIGHT|     33.33|66.67     |", "Test A|%-20s|%20s|%10.2f|%-10.2f|", "LEFT", "RIGHT", 33.333333, 66.666666);
+	TestFormatV<OutputSize>("Test B|Percents:%%%3|", "Test B|Percents:%%%%%%%d|", 3);
+	TestFormatV<OutputSize>("Test C|12345|54321|123ABC|f|99|", "Test C|%d|%i|%X|%x|%u|", 12345, 54321, 0x123AbC, 15, 99);
 
 #ifdef _MSC_VER
 	// MSVC's standard library formats pointers differently
-	TestGetVarArgs<OutputSize>("Test D|0000000000012345|", "Test D|%p|", 0x12345);
+	TestFormatV<OutputSize>("Test D|0000000000012345|", "Test D|%p|", 0x12345);
 #else
 	// Pointer format (which we get from snprintf) varies from platform to platform.
 	// Call it to make sure it puts the output in the proper place, but don't test for the exact format
@@ -1489,19 +1520,66 @@ TEST_CASE("Core::GenericPlatformString::GetVarArgs", "[Core][String][Smoke]")
 	}
 #endif
 
-	TestGetVarArgs<OutputSize>("Test E|12345678912345|", "Test E|%" INT64_FMT "|", int64(12345678912345LL));
-	TestGetVarArgs<OutputSize>("Test F|123.456000|1.234560e+02|123.456|", "Test F|%f|%e|%g|", 123.456, 123.456, 123.456);
-	TestGetVarArgs<OutputSize>("Test G|C0DE|", "Test G|%" UPTRINT_X_FMT"|", UPTRINT(49374));
-	TestGetVarArgs<OutputSize>("Test H|c0de|", "Test H|%" UPTRINT_x_FMT "|", UPTRINT(49374));
-	TestGetVarArgs<OutputSize>("Test I|18446744073709551615|", "Test I|%" UINT64_FMT "|", MAX_uint64);
-	TestGetVarArgs<OutputSize>("Test J|         J|", "Test J|%*c|", 10, 'J');
-	TestGetVarArgs<OutputSize>("Test K|K    |", "Test K|%-5c|", 'K');
-	TestGetVarArgs<OutputSize>("Test L|Hello World|", "Test L|%hs|", UTF8TEXT("Hello World"));
-	TestGetVarArgs<OutputSize>(UTF8TEXT("Test M|\U0001F606\U0001F41C|"), "Test M|%ls|", TEXT("\U0001F606\U0001F41C"));
-	TestGetVarArgs<OutputSize>(UTF8TEXT("Test N|\U0001F606\U0001F41C|"), "Test N|%hs|", UTF8TEXT("\U0001F606\U0001F41C"));
-	TestGetVarArgs<OutputSize>(UTF8TEXT("Test O|\U0001F606\U0001F41C|"), "Test O|%s|", UTF8TEXT("\U0001F606\U0001F41C"));
-	TestGetVarArgs<OutputSize>(UTF8TEXT("Test P|\U0001F606\U0001F41C|"), "Test P|%S|", TEXT("\U0001F606\U0001F41C"));
-
-	TestGetVarArgs<9>("12345678", "%08x", 0x12345678);
+	TestFormatV<OutputSize>("Test E|12345678912345|", "Test E|%" INT64_FMT "|", int64(12345678912345LL));
+	TestFormatV<OutputSize>("Test F|123.456000|1.234560e+02|123.456|", "Test F|%f|%e|%g|", 123.456, 123.456, 123.456);
+	TestFormatV<OutputSize>("Test G|C0DE|", "Test G|%" UPTRINT_X_FMT"|", UPTRINT(49374));
+	TestFormatV<OutputSize>("Test H|c0de|", "Test H|%" UPTRINT_x_FMT "|", UPTRINT(49374));
+	TestFormatV<OutputSize>("Test I|18446744073709551615|", "Test I|%" UINT64_FMT "|", MAX_uint64);
+	TestFormatV<OutputSize>("Test J|         J|", "Test J|%*c|", 10, 'J');
+	TestFormatV<OutputSize>("Test K|K    |", "Test K|%-5c|", 'K');
+	TestFormatV<OutputSize>("Test L|Hello World|", "Test L|%hs|", UTF8TEXT("Hello World"));
+	TestFormatV<OutputSize>(UTF8TEXT("Test M|\U0001F606\U0001F41C|"), "Test M|%ls|", TEXT("\U0001F606\U0001F41C"));
+	TestFormatV<OutputSize>(UTF8TEXT("Test N|\U0001F606\U0001F41C|"), "Test N|%hs|", UTF8TEXT("\U0001F606\U0001F41C"));
+	TestFormatV<OutputSize>(UTF8TEXT("Test O|\U0001F606\U0001F41C|"), "Test O|%s|", UTF8TEXT("\U0001F606\U0001F41C"));
+	TestFormatV<OutputSize>(UTF8TEXT("Test P|\U0001F606\U0001F41C|"), "Test P|%S|", TEXT("\U0001F606\U0001F41C"));
+
+	TestFormatV<9>("12345678", "%08x", 0x12345678);
+
+	constexpr uint8 Sentinel = 0xcc;
+	const auto CheckAllSentinel = [Sentinel]<typename T, int32 N>(T (&Array)[N]) -> bool
+	{
+		bool bOk = true;
+		for (uint8 Value : MakeArrayView((const uint8*)Array, N * sizeof(T)))
+		{
+			bOk &= Value == Sentinel;
+		}
+		return bOk;
+	};
+
+	{
+		WIDECHAR Output[OutputSize];
+		FMemory::Memset(Output, Sentinel, OutputSize * sizeof(WIDECHAR));
+
+		CHECK(TestFormat((WIDECHAR*)nullptr, 0, WIDETEXT("")) == 0);
+		CHECK(TestFormat(Output + 8, 0, WIDETEXT("")) == 0);
+		CHECK(CheckAllSentinel(Output));
+		CHECK(TestFormat((WIDECHAR*)nullptr, 0, WIDETEXT("%d"), 1) == 1);
+		CHECK(TestFormat(Output + 8, 0, WIDETEXT("%d"), 1) == 1);
+		CHECK(CheckAllSentinel(Output));
+	}
+
+	{
+		ANSICHAR Output[OutputSize];
+		FMemory::Memset(Output, Sentinel, OutputSize * sizeof(ANSICHAR));
+
+		CHECK(TestFormat((ANSICHAR*)nullptr, 0, "") == 0);
+		CHECK(TestFormat(Output + 8, 0, "") == 0);
+		CHECK(CheckAllSentinel(Output));
+		CHECK(TestFormat((ANSICHAR*)nullptr, 0, "%d", 1) == 1);
+		CHECK(TestFormat(Output + 8, 0, "%d", 1) == 1);
+		CHECK(CheckAllSentinel(Output));
+	}
+}
+
+TEST_CASE("Core::GenericPlatformString::GetVarArgs", "[Core][String][Smoke]")
+{
+	// GetVarArgs is a wrapper for FormatV. Test only the expected difference in return value.
+	TCHAR Output[16];
+	CHECK(TestGetVarArgs(Output, 0, TEXT("")) == -1);
+	CHECK(TestGetVarArgs(Output, 1, TEXT("")) == 0);
+	CHECK(TestGetVarArgs(Output, 0, TEXT("%d"), 1) == -1);
+	CHECK(TestGetVarArgs(Output, 1, TEXT("%d"), 1) == -1);
+	CHECK(TestGetVarArgs(Output, 2, TEXT("%d"), 1) == 1);
+	CHECK(TestGetVarArgs(Output, 3, TEXT("%d"), 1) == 1);
 }
 #endif // WITH_LOW_LEVEL_TESTS
\ No newline at end of file
diff --git a/Engine/Source/Runtime/Core/Private/Misc/StringBuilder.cpp b/Engine/Source/Runtime/Core/Private/Misc/StringBuilder.cpp
--- a/Engine/Source/Runtime/Core/Private/Misc/StringBuilder.cpp
+++ b/Engine/Source/Runtime/Core/Private/Misc/StringBuilder.cpp
@@ -74,22 +74,26 @@ TStringBuilderBase<C>& TStringBuilderBase<C>::AppendV(const C* Fmt, va_list Args
 {
 	for (;;)
 	{
-		if (const int32 RemainingSize = (int32)(End - CurPos))
+		const int32 RemainingSize = (int32)(End - CurPos);
+
+		va_list ArgPack;
+		va_copy(ArgPack, Args);
+		const int32 RequiredSize = TCString<C>::FormatV(CurPos, RemainingSize, Fmt, ArgPack);
+		va_end(ArgPack);
+
+		if (!ensure(RequiredSize >= 0))
+		{
+			return *this;
+		}
+
+		if (RequiredSize < RemainingSize)
 		{
-			va_list ArgPack;
-			va_copy(ArgPack, Args);
-			const int32 Result = TCString<C>::GetVarArgs(CurPos, RemainingSize, Fmt, ArgPack);
-			va_end(ArgPack);
-
-			if (Result >= 0 && Result < RemainingSize)
-			{
-				CurPos += Result;
-				return *this;
-			}
+			CurPos += RequiredSize;
+			return *this;
 		}
 
-		// Total size will be rounded up to the next power of two. Start with at least 64.
-		Extend(64);
+		// Total size will be rounded up to the next power of two.
+		Extend(RequiredSize - RemainingSize + 1);
 	}
 }
 
diff --git a/Engine/Source/Runtime/Core/Public/GenericPlatform/GenericPlatformString.h b/Engine/Source/Runtime/Core/Public/GenericPlatform/GenericPlatformString.h
--- a/Engine/Source/Runtime/Core/Public/GenericPlatform/GenericPlatformString.h
+++ b/Engine/Source/Runtime/Core/Public/GenericPlatform/GenericPlatformString.h
@@ -340,11 +340,53 @@ struct FGenericPlatformString : public FGenericPlatformStricmp
 	CORE_API static int32 Strncmp(const WIDECHAR* String1, const UTF8CHAR* String2, SIZE_T Count);
 	CORE_API static int32 Strncmp(const UTF8CHAR* String1, const UTF8CHAR* String2, SIZE_T Count);
 
-	static CORE_API int32 GetVarArgs(ANSICHAR* Dest, SIZE_T DestSize, const ANSICHAR*& Fmt, va_list ArgPtr);
-	static CORE_API int32 GetVarArgs(WIDECHAR* Dest, SIZE_T DestSize, const WIDECHAR*& Fmt, va_list ArgPtr);
+	/**
+	 * Format to a destination buffer. Similar to vsnprintf.
+	 *
+	 * Writes the formatted output to Dest, up to a maximum of DestSize-1 characters, excluding the null terminator.
+	 * Content of Dest is undefined when the return value is outside of the range [0, DestSize).
+	 * A return value is calculated even when DestSize is 0, in which case Dest may be null.
+	 * A null terminator is always written unless DestSize is 0.
+	 * Success is indicated by ReturnValue >= 0 && ReturnValue < DestSize.
+	 *
+	 * %hs is always treated as UTF-8.
+	 * %ls is always treated as UTF-16.
+	 * %s is treated as UTF-8 for narrow format strings and UTF-16 for wide format strings.
+	 *
+	 * @param Dest   Pointer to the destination buffer to write to.
+	 * @param DestSize   Maximum number of characters to write to Dest, including the null terminator.
+	 * @param Format   Format string with which to format the output.
+	 * @param Args   Arguments containing the data referenced by the format string.
+	 * @return Number of characters, excluding the null terminator, that would be
+	 *         written to Dest if DestSize was ignored, or -1 on error.
+	 */
+	[[nodiscard]] CORE_API static int32 FormatV(ANSICHAR* Dest, SIZE_T DestSize, const ANSICHAR* Format, va_list Args);
+	[[nodiscard]] CORE_API static int32 FormatV(WIDECHAR* Dest, SIZE_T DestSize, const WIDECHAR* Format, va_list Args);
+	[[nodiscard]] UE_REWRITE static int32 FormatV(UTF8CHAR* Dest, SIZE_T DestSize, const ANSICHAR* Format, va_list Args)
+	{
+		return FormatV((ANSICHAR*)Dest, DestSize, Format, Args);
+	}
+	[[nodiscard]] UE_REWRITE static int32 FormatV(UTF8CHAR* Dest, SIZE_T DestSize, const UTF8CHAR* Format, va_list Args)
+	{
+		return FormatV((ANSICHAR*)Dest, DestSize, (const ANSICHAR*)Format, Args);
+	}
+
+	UE_REWRITE static int32 GetVarArgs(ANSICHAR* Dest, SIZE_T DestSize, const ANSICHAR*& Fmt, va_list ArgPtr)
+	{
+		int32 Result = FormatV(Dest, DestSize, Fmt, ArgPtr);
+		return (Result >= 0 && Result < (int32)DestSize) ? Result : -1;
+	}
+
+	UE_REWRITE static int32 GetVarArgs(WIDECHAR* Dest, SIZE_T DestSize, const WIDECHAR*& Fmt, va_list ArgPtr)
+	{
+		int32 Result = FormatV(Dest, DestSize, Fmt, ArgPtr);
+		return (Result >= 0 && Result < (int32)DestSize) ? Result : -1;
+	}
+
 	UE_REWRITE static int32 GetVarArgs(UTF8CHAR* Dest, SIZE_T DestSize, const UTF8CHAR*& Fmt, va_list ArgPtr)
 	{
-		return GetVarArgs((ANSICHAR*)Dest, DestSize, (const ANSICHAR*&)(Fmt), ArgPtr);
+		int32 Result = FormatV((ANSICHAR*)Dest, DestSize, (const ANSICHAR*)Fmt, ArgPtr);
+		return (Result >= 0 && Result < (int32)DestSize) ? Result : -1;
 	}
 
 private:
diff --git a/Engine/Source/Runtime/Core/Public/Misc/CString.h b/Engine/Source/Runtime/Core/Public/Misc/CString.h
--- a/Engine/Source/Runtime/Core/Public/Misc/CString.h
+++ b/Engine/Source/Runtime/Core/Public/Misc/CString.h
@@ -586,6 +586,31 @@ public:
 		return SnprintfImpl(Dest, DestSize, (const CharType*)Fmt, Args...);
 	}
 
+	/**
+	 * Format to a destination buffer. Similar to vsnprintf.
+	 *
+	 * Writes the formatted output to Dest, up to a maximum of DestSize-1 characters, excluding the null terminator.
+	 * Content of Dest is undefined when the return value is outside of the range [0, DestSize).
+	 * A return value is calculated even when DestSize is 0, in which case Dest may be null.
+	 * A null terminator is always written unless DestSize is 0.
+	 * Success is indicated by ReturnValue >= 0 && ReturnValue < DestSize.
+	 *
+	 * %hs is always treated as UTF-8.
+	 * %ls is always treated as UTF-16.
+	 * %s is treated as UTF-8 for narrow format strings and UTF-16 for wide format strings.
+	 *
+	 * @param Dest   Pointer to the destination buffer to write to.
+	 * @param DestSize   Maximum number of characters to write to Dest, including the null terminator.
+	 * @param Format   Format string with which to format the output.
+	 * @param Args   Arguments containing the data referenced by the format string.
+	 * @return Number of characters, excluding the null terminator, that would be
+	 *         written to Dest if DestSize was ignored, or -1 on error.
+	 */
+	UE_REWRITE static int32 FormatV(CharType* Dest, SIZE_T DestSize, const CharType* Format, va_list Args)
+	{
+		return FPlatformString::FormatV(Dest, DestSize, Format, Args);
+	}
+
 	/**
 	 * Helper function to write formatted output using an argument list
 	 *
