Logging: Added the UE_LOGF family of macros for unstructured logging with UTF-8 format strings

#Synced-CL 32500503

diff --git a/Engine/Source/Runtime/Core/Private/AutoRTFM/AutoRTFMUE.cpp b/Engine/Source/Runtime/Core/Private/AutoRTFM/AutoRTFMUE.cpp
--- a/Engine/Source/Runtime/Core/Private/AutoRTFM/AutoRTFMUE.cpp
+++ b/Engine/Source/Runtime/Core/Private/AutoRTFM/AutoRTFMUE.cpp
@@ -285,7 +285,7 @@ void AutoRTFMLog(const char* File, int Line, void* ProgramCounter, autortfm_log_
 	}
 
 	static ::UE::Logging::Private::FStaticBasicLogDynamicData DynamicData;
-	UE::Logging::Private::FStaticBasicLogRecord Record(
+	UE::Logging::Private::TStaticBasicLogRecord<TCHAR> Record(
 		/* Format */ TEXT("%s"),
 		/* File */ File,
 		/* Line */ Line,
diff --git a/Engine/Source/Runtime/Core/Private/Logging/StructuredLog.cpp b/Engine/Source/Runtime/Core/Private/Logging/StructuredLog.cpp
--- a/Engine/Source/Runtime/Core/Private/Logging/StructuredLog.cpp
+++ b/Engine/Source/Runtime/Core/Private/Logging/StructuredLog.cpp
@@ -37,6 +37,7 @@
 
 void StaticFailDebug (const TCHAR* Error, const ANSICHAR* File, int32 Line, void* ProgramCounter, const TCHAR* Message);
 void StaticFailDebugV(const TCHAR* Error, const ANSICHAR* File, int32 Line, void* ProgramCounter, const TCHAR* DescriptionFormat, va_list DescriptionArgs);
+void StaticFailDebugV(const TCHAR* Error, const ANSICHAR* File, int32 Line, void* ProgramCounter, const UTF8CHAR* DescriptionFormat, va_list DescriptionArgs);
 
 namespace UE::Serialization::Private
 {
@@ -1600,9 +1601,10 @@ static FLogTemplate* GetStaticBasicLogTemplate()
 
 // Serializing the log to compact binary happens in its own function because that allows stack
 // space for the writer to be returned before calling into the output devices.
-FORCENOINLINE static FCbObject SerializeBasicLogMessage(const FStaticBasicLogRecord& Log, va_list Args)
+template <UE::CCharType FormatCharType>
+FORCENOINLINE static FCbObject SerializeBasicLogMessage(TStaticBasicLogRecordParam<FormatCharType> Log, va_list Args)
 {
-	TStringBuilder<512> Message;
+	TStringBuilderWithBuffer<FormatCharType, 512> Message;
 	Message.AppendV(Log.Format, Args);
 
 	TCbWriter<512> Writer;
@@ -1640,7 +1642,7 @@ FORCENOINLINE static FCbObject SerializeBasicLogMessage(const FStaticBasicLogRec
 }
 
 UE_AUTORTFM_ALWAYS_OPEN
-static void BasicLogV(FStaticBasicLogRecordParam Log, const FLogCategoryBase& Category, ELogVerbosity::Type Verbosity, va_list Args)
+static void BasicLogV(TStaticBasicLogRecordParam<TCHAR> Log, const FLogCategoryBase& Category, ELogVerbosity::Type Verbosity, va_list Args)
 {
 #if !NO_LOGGING
 #if LOGTRACE_ENABLED
@@ -1682,8 +1684,37 @@ static void BasicLogV(FStaticBasicLogRecordParam Log, const FLogCategoryBase& Ca
 #endif
 }
 
-template <ELogVerbosity::Type Verbosity>
-FORCENOINLINE void BasicLog(FStaticBasicLogRecordParam Log, const FLogCategoryBase* Category, ...)
+template <UE::CCharType FormatCharType>
+	requires TIsCharEncodingCompatibleWith_V<FormatCharType, UTF8CHAR>
+UE_AUTORTFM_ALWAYS_OPEN
+static void BasicLogV(TStaticBasicLogRecordParam<FormatCharType> Log, const FLogCategoryBase& Category, ELogVerbosity::Type Verbosity, va_list Args)
+{
+#if !NO_LOGGING
+#if LOGTRACE_ENABLED
+	if (UE_TRACE_CHANNELEXPR_IS_ENABLED(LogChannel))
+	{
+		EnsureLogMessageSpec(Category, Log, Verbosity);
+	}
+#endif
+
+	FCbObject Fields = SerializeBasicLogMessage(Log, Args);
+
+	FLogRecord Record;
+	Record.SetFormat(GStaticBasicLogFormat);
+	Record.SetTemplate(GetStaticBasicLogTemplate());
+	Record.SetFields(MoveTemp(Fields));
+	Record.SetFile(Log.File);
+	Record.SetLine(Log.Line);
+	Record.SetCategory(Category.GetCategoryName());
+	Record.SetVerbosity(Verbosity);
+	Record.SetTime(FLogTime::Now());
+
+	DispatchStaticLogRecord(Log, Record);
+#endif
+}
+
+template <ELogVerbosity::Type Verbosity, UE::CCharType FormatCharType>
+FORCENOINLINE void BasicLog(TStaticBasicLogRecordParam<FormatCharType> Log, const FLogCategoryBase* Category, ...)
 {
 #if !NO_LOGGING
 	va_list Args;
@@ -1693,26 +1724,43 @@ FORCENOINLINE void BasicLog(FStaticBasicLogRecordParam Log, const FLogCategoryBa
 #endif
 }
 
-template CORE_API void BasicLog<ELogVerbosity::Error>(FStaticBasicLogRecordParam Log, const FLogCategoryBase* Category, ...);
-template CORE_API void BasicLog<ELogVerbosity::Warning>(FStaticBasicLogRecordParam Log, const FLogCategoryBase* Category, ...);
-template CORE_API void BasicLog<ELogVerbosity::Display>(FStaticBasicLogRecordParam Log, const FLogCategoryBase* Category, ...);
-template CORE_API void BasicLog<ELogVerbosity::Log>(FStaticBasicLogRecordParam Log, const FLogCategoryBase* Category, ...);
-template CORE_API void BasicLog<ELogVerbosity::Verbose>(FStaticBasicLogRecordParam Log, const FLogCategoryBase* Category, ...);
-template CORE_API void BasicLog<ELogVerbosity::VeryVerbose>(FStaticBasicLogRecordParam Log, const FLogCategoryBase* Category, ...);
-template CORE_API void BasicLog<ELogVerbosity::SetColor>(FStaticBasicLogRecordParam Log, const FLogCategoryBase* Category, ...);
-
+template CORE_API void BasicLog<ELogVerbosity::Error>(TStaticBasicLogRecordParam<TCHAR> Log, const FLogCategoryBase* Category, ...);
+template CORE_API void BasicLog<ELogVerbosity::Warning>(TStaticBasicLogRecordParam<TCHAR> Log, const FLogCategoryBase* Category, ...);
+template CORE_API void BasicLog<ELogVerbosity::Display>(TStaticBasicLogRecordParam<TCHAR> Log, const FLogCategoryBase* Category, ...);
+template CORE_API void BasicLog<ELogVerbosity::Log>(TStaticBasicLogRecordParam<TCHAR> Log, const FLogCategoryBase* Category, ...);
+template CORE_API void BasicLog<ELogVerbosity::Verbose>(TStaticBasicLogRecordParam<TCHAR> Log, const FLogCategoryBase* Category, ...);
+template CORE_API void BasicLog<ELogVerbosity::VeryVerbose>(TStaticBasicLogRecordParam<TCHAR> Log, const FLogCategoryBase* Category, ...);
+template CORE_API void BasicLog<ELogVerbosity::SetColor>(TStaticBasicLogRecordParam<TCHAR> Log, const FLogCategoryBase* Category, ...);
+
+template CORE_API void BasicLog<ELogVerbosity::Error>(TStaticBasicLogRecordParam<ANSICHAR> Log, const FLogCategoryBase* Category, ...);
+template CORE_API void BasicLog<ELogVerbosity::Warning>(TStaticBasicLogRecordParam<ANSICHAR> Log, const FLogCategoryBase* Category, ...);
+template CORE_API void BasicLog<ELogVerbosity::Display>(TStaticBasicLogRecordParam<ANSICHAR> Log, const FLogCategoryBase* Category, ...);
+template CORE_API void BasicLog<ELogVerbosity::Log>(TStaticBasicLogRecordParam<ANSICHAR> Log, const FLogCategoryBase* Category, ...);
+template CORE_API void BasicLog<ELogVerbosity::Verbose>(TStaticBasicLogRecordParam<ANSICHAR> Log, const FLogCategoryBase* Category, ...);
+template CORE_API void BasicLog<ELogVerbosity::VeryVerbose>(TStaticBasicLogRecordParam<ANSICHAR> Log, const FLogCategoryBase* Category, ...);
+template CORE_API void BasicLog<ELogVerbosity::SetColor>(TStaticBasicLogRecordParam<ANSICHAR> Log, const FLogCategoryBase* Category, ...);
+
+template CORE_API void BasicLog<ELogVerbosity::Error>(TStaticBasicLogRecordParam<UTF8CHAR> Log, const FLogCategoryBase* Category, ...);
+template CORE_API void BasicLog<ELogVerbosity::Warning>(TStaticBasicLogRecordParam<UTF8CHAR> Log, const FLogCategoryBase* Category, ...);
+template CORE_API void BasicLog<ELogVerbosity::Display>(TStaticBasicLogRecordParam<UTF8CHAR> Log, const FLogCategoryBase* Category, ...);
+template CORE_API void BasicLog<ELogVerbosity::Log>(TStaticBasicLogRecordParam<UTF8CHAR> Log, const FLogCategoryBase* Category, ...);
+template CORE_API void BasicLog<ELogVerbosity::Verbose>(TStaticBasicLogRecordParam<UTF8CHAR> Log, const FLogCategoryBase* Category, ...);
+template CORE_API void BasicLog<ELogVerbosity::VeryVerbose>(TStaticBasicLogRecordParam<UTF8CHAR> Log, const FLogCategoryBase* Category, ...);
+template CORE_API void BasicLog<ELogVerbosity::SetColor>(TStaticBasicLogRecordParam<UTF8CHAR> Log, const FLogCategoryBase* Category, ...);
+
+template <UE::CCharType FormatCharType>
 UE_AUTORTFM_ALWAYS_OPEN
-static void BasicFatalLogV(FStaticBasicLogRecordParam Log, void* ProgramCounter, va_list Args)
+static void BasicFatalLogV(TStaticBasicLogRecordParam<FormatCharType> Log, void* ProgramCounter, va_list Args)
 {
 #if !NO_LOGGING
 	StaticFailDebugV(TEXT("Fatal error:"), Log.File, Log.Line, ProgramCounter, Log.Format, Args);
-
 	UE_DEBUG_BREAK_AND_PROMPT_FOR_REMOTE();
 	FDebug::ProcessFatalError(ProgramCounter);
 #endif
 }
 
-FORCENOINLINE void BasicFatalLog(FStaticBasicLogRecordParam Log, const FLogCategoryBase* Category, ...)
+template <UE::CCharType FormatCharType>
+FORCENOINLINE void BasicFatalLog(TStaticBasicLogRecordParam<FormatCharType> Log, const FLogCategoryBase* Category, ...)
 {
 #if !NO_LOGGING
 	va_list Args;
@@ -1722,7 +1770,11 @@ FORCENOINLINE void BasicFatalLog(FStaticBasicLogRecordParam Log, const FLogCateg
 #endif
 }
 
-FORCENOINLINE void BasicFatalLogWithProgramCounter(FStaticBasicLogRecordParam Log, const FLogCategoryBase* Category, void* ProgramCounter, ...)
+template CORE_API void BasicFatalLog(TStaticBasicLogRecordParam<TCHAR> Log, const FLogCategoryBase* Category, ...);
+template CORE_API void BasicFatalLog(TStaticBasicLogRecordParam<UTF8CHAR> Log, const FLogCategoryBase* Category, ...);
+
+template <UE::CCharType FormatCharType>
+FORCENOINLINE void BasicFatalLogWithProgramCounter(TStaticBasicLogRecordParam<FormatCharType> Log, const FLogCategoryBase* Category, void* ProgramCounter, ...)
 {
 #if !NO_LOGGING
 	va_list Args;
@@ -1732,6 +1784,9 @@ FORCENOINLINE void BasicFatalLogWithProgramCounter(FStaticBasicLogRecordParam Lo
 #endif
 }
 
+template CORE_API void BasicFatalLogWithProgramCounter(TStaticBasicLogRecordParam<TCHAR> Log, const FLogCategoryBase* Category, void* ProgramCounter, ...);
+template CORE_API void BasicFatalLogWithProgramCounter(TStaticBasicLogRecordParam<UTF8CHAR> Log, const FLogCategoryBase* Category, void* ProgramCounter, ...);
+
 } // UE::Logging::Private
 
 ///////////////////////////////////////////////////////////////////////////////////////////////////////////////////////
diff --git a/Engine/Source/Runtime/Core/Private/Misc/AssertionMacros.cpp b/Engine/Source/Runtime/Core/Private/Misc/AssertionMacros.cpp
--- a/Engine/Source/Runtime/Core/Private/Misc/AssertionMacros.cpp
+++ b/Engine/Source/Runtime/Core/Private/Misc/AssertionMacros.cpp
@@ -240,6 +240,14 @@ public:
 		Index = FMath::Clamp(Index, Index + Len, Capacity - 1);
 	}
 
+	void AppendV(const UTF8CHAR* Format, va_list Args)
+	{
+		if (const int32 Len = TCString<UTF8CHAR>::GetVarArgs(Utf8Scratch, UE_ARRAY_COUNT(Utf8Scratch), Format, Args); Len > 0)
+		{
+			Append(Utf8Scratch);
+		}
+	}
+
 	void Appendf(const TCHAR* Format, ...)
 	{
 		va_list Args;
@@ -252,11 +260,17 @@ private:
 	int32 Index = 0;
 
 	static constexpr int32 Capacity = UE_ARRAY_COUNT(GErrorHist);
+
+	// Scratch space for formatting as UTF-8.
+	// This is static to avoid stack overflows.
+	// This is safe to be static because its use is guarded by FailDebugMutex.
+	inline static UTF8CHAR Utf8Scratch[Capacity] = {};
 };
 
 /**
  * Prints error to the debug output and copies the error into the global error message.
  */
+template <UE::CCharType FormatCharType>
 FORCENOINLINE void StaticFailDebugV(
 	const TCHAR* Error,
 	const ANSICHAR* Expression,
@@ -265,7 +279,7 @@ FORCENOINLINE void StaticFailDebugV(
 	bool bIsEnsure,
 	void* ProgramCounter,
 	const TCHAR* Message,
-	const TCHAR* DescriptionFormat,
+	const FormatCharType* DescriptionFormat,
 	va_list DescriptionArgs)
 {
 	PrintScriptCallstack();
@@ -316,7 +330,7 @@ FORCENOINLINE void StaticFailDebugV(
 FORCENOINLINE void VARARGS StaticFailDebug(const TCHAR* Error, const ANSICHAR* File, int32 Line, void* ProgramCounter, const TCHAR* Message)
 {
 	va_list DescriptionArgs{};
-	StaticFailDebugV(Error, /*Expression*/ "", File, Line, /*bIsEnsure*/ false, ProgramCounter, Message, /*DescriptionFormat*/ nullptr, DescriptionArgs);
+	StaticFailDebugV(Error, /*Expression*/ "", File, Line, /*bIsEnsure*/ false, ProgramCounter, Message, /*DescriptionFormat*/ (const TCHAR*)nullptr, DescriptionArgs);
 }
 
 FORCENOINLINE void VARARGS StaticFailDebugV(const TCHAR* Error, const ANSICHAR* File, int32 Line, void* ProgramCounter, const TCHAR* DescriptionFormat, va_list DescriptionArgs)
@@ -324,6 +338,11 @@ FORCENOINLINE void VARARGS StaticFailDebugV(const TCHAR* Error, const ANSICHAR*
 	StaticFailDebugV(Error, /*Expression*/ "", File, Line, /*bIsEnsure*/ false, ProgramCounter, /*Message*/ nullptr, DescriptionFormat, DescriptionArgs);
 }
 
+FORCENOINLINE void VARARGS StaticFailDebugV(const TCHAR* Error, const ANSICHAR* File, int32 Line, void* ProgramCounter, const UTF8CHAR* DescriptionFormat, va_list DescriptionArgs)
+{
+	StaticFailDebugV(Error, /*Expression*/ "", File, Line, /*bIsEnsure*/ false, ProgramCounter, /*Message*/ nullptr, DescriptionFormat, DescriptionArgs);
+}
+
 /// track thread asserts
 bool FDebug::HasAsserted()
 {
@@ -490,7 +509,7 @@ FORCENOINLINE void FDebug::EnsureFailed(const ANSICHAR* Expr, const ANSICHAR* Fi
 	FPlatformAtomics::InterlockedIncrement(&ActiveEnsureCount);
 
 	va_list DescriptionArgs{};
-	StaticFailDebugV(TEXT("Ensure condition failed: "), Expr, File, Line, /*bIsEnsure*/ true, ProgramCounter, Msg, nullptr, DescriptionArgs);
+	StaticFailDebugV(TEXT("Ensure condition failed: "), Expr, File, Line, /*bIsEnsure*/ true, ProgramCounter, Msg, (const TCHAR*)nullptr, DescriptionArgs);
 
 	// Is there a debugger attached?  If not we'll submit an error report.
 	if (FPlatformMisc::IsDebuggerPresent() && !GAlwaysReportCrash)
@@ -827,6 +846,28 @@ void UE_COLD UE_DEBUG_SECTION VARARGS UE::Assert::Private::ProcessLowLevelFatalE
 	FDebug::ProcessFatalError(PLATFORM_RETURN_ADDRESS());
 }
 
+void UE_COLD UE_DEBUG_SECTION VARARGS UE::Assert::Private::ProcessLowLevelFatalError(const ANSICHAR* File, int32 Line, const ANSICHAR* Format, ...)
+{
+	va_list Args;
+	va_start(Args, Format);
+	StaticFailDebugV(TEXT("LowLevelFatalError"), /*Expression*/ "", File, Line, /*bIsEnsure*/ false, PLATFORM_RETURN_ADDRESS(), /*Message*/ nullptr, (const UTF8CHAR*)Format, Args);
+	va_end(Args);
+
+	UE_DEBUG_BREAK_AND_PROMPT_FOR_REMOTE();
+	FDebug::ProcessFatalError(PLATFORM_RETURN_ADDRESS());
+}
+
+void UE_COLD UE_DEBUG_SECTION VARARGS UE::Assert::Private::ProcessLowLevelFatalError(const ANSICHAR* File, int32 Line, const UTF8CHAR* Format, ...)
+{
+	va_list Args;
+	va_start(Args, Format);
+	StaticFailDebugV(TEXT("LowLevelFatalError"), /*Expression*/ "", File, Line, /*bIsEnsure*/ false, PLATFORM_RETURN_ADDRESS(), /*Message*/ nullptr, Format, Args);
+	va_end(Args);
+
+	UE_DEBUG_BREAK_AND_PROMPT_FOR_REMOTE();
+	FDebug::ProcessFatalError(PLATFORM_RETURN_ADDRESS());
+}
+
 UE_AUTORTFM_ALWAYS_OPEN
 void FDebug::DumpStackTraceToLog(const ELogVerbosity::Type LogVerbosity)
 {
diff --git a/Engine/Source/Runtime/Core/Public/Logging/LogMacros.h b/Engine/Source/Runtime/Core/Public/Logging/LogMacros.h
--- a/Engine/Source/Runtime/Core/Public/Logging/LogMacros.h
+++ b/Engine/Source/Runtime/Core/Public/Logging/LogMacros.h
@@ -2,6 +2,7 @@
 
 #pragma once
 
+#include "Concepts/CharType.h"
 #include "Containers/UnrealString.h"
 #include "CoreTypes.h"
 #include "HAL/PreprocessorHelpers.h"
@@ -16,6 +17,7 @@
 #include "Templates/EnableIf.h"
 #include "Templates/IsArrayOrRefOfTypeByPredicate.h"
 #include "Templates/IsValidVariadicFunctionArg.h"
+#include "Traits/ElementType.h"
 #include "Traits/IsCharEncodingCompatibleWith.h"
 
 #include <type_traits>
@@ -126,9 +128,10 @@ struct FStaticBasicLogDynamicData
 };
 
 /** Data about a static basic log that is constant for every occurrence. */
-struct FStaticBasicLogRecord
+template <UE::CCharType FormatCharType>
+struct TStaticBasicLogRecord
 {
-	const TCHAR* Format = nullptr;
+	const FormatCharType* Format = nullptr;
 	UE_IF(UE_LOG_INCLUDE_SOURCE_LOCATION,, inline static constexpr) const ANSICHAR* File = nullptr;
 	UE_IF(UE_LOG_INCLUDE_SOURCE_LOCATION,, inline static constexpr) int32 Line = 0;
 #if LOGTRACE_ENABLED
@@ -136,8 +139,8 @@ struct FStaticBasicLogRecord
 #endif
 
 	// Workaround for https://developercommunity.visualstudio.com/t/Incorrect-warning-C4700-with-unrelated-s/10285950
-	constexpr FStaticBasicLogRecord(
-		const TCHAR* InFormat,
+	constexpr TStaticBasicLogRecord(
+		const FormatCharType* InFormat,
 		const ANSICHAR* InFile,
 		int32 InLine,
 		FStaticBasicLogDynamicData* InDynamicData)
@@ -155,22 +158,26 @@ struct FStaticBasicLogRecord
 #if !UE_PRIVATE_LOG_USE_STATIC_RECORD
 	// Delete the address-of operator because it is unsafe to pass pointers to records
 	// when they do not have static storage duration.
-	FStaticBasicLogRecord* operator&() const = delete;
+	TStaticBasicLogRecord* operator&() const = delete;
 #endif
 };
 
-// Pass FStaticBasicLogRecord by value when the record is not static, because in that
+// Pass TStaticBasicLogRecord by value when the record is not static, because in that
 // configuration it is simply a wrapper for the format string pointer.
 #if UE_PRIVATE_LOG_USE_STATIC_RECORD
-using FStaticBasicLogRecordParam = const FStaticBasicLogRecord&;
+template <UE::CCharType FormatCharType>
+using TStaticBasicLogRecordParam = const TStaticBasicLogRecord<FormatCharType>&;
 #else
-using FStaticBasicLogRecordParam = FStaticBasicLogRecord;
+template <UE::CCharType FormatCharType>
+using TStaticBasicLogRecordParam = TStaticBasicLogRecord<FormatCharType>;
 #endif
 
-template <ELogVerbosity::Type Verbosity>
-CORE_API void BasicLog(FStaticBasicLogRecordParam Log, const FLogCategoryBase* Category, ...);
-CORE_API void BasicFatalLog(FStaticBasicLogRecordParam Log, const FLogCategoryBase* Category, ...);
-CORE_API void BasicFatalLogWithProgramCounter(FStaticBasicLogRecordParam Log, const FLogCategoryBase* Category, void* ProgramCounter, ...);
+template <ELogVerbosity::Type Verbosity, UE::CCharType FormatCharType>
+CORE_API void BasicLog(TStaticBasicLogRecordParam<FormatCharType> Log, const FLogCategoryBase* Category, ...);
+template <UE::CCharType FormatCharType>
+CORE_API void BasicFatalLog(TStaticBasicLogRecordParam<FormatCharType> Log, const FLogCategoryBase* Category, ...);
+template <UE::CCharType FormatCharType>
+CORE_API void BasicFatalLogWithProgramCounter(TStaticBasicLogRecordParam<FormatCharType> Log, const FLogCategoryBase* Category, void* ProgramCounter, ...);
 
 } // UE::Logging::Private
 
@@ -179,31 +186,38 @@ CORE_API void BasicFatalLogWithProgramCounter(FStaticBasicLogRecordParam Log, co
 	// This will only log Fatal errors
 	#define UE_LOG(CategoryName, Verbosity, Format, ...) \
 	{ \
-		if constexpr ((ELogVerbosity::Verbosity & ELogVerbosity::VerbosityMask) == ELogVerbosity::Fatal) \
+		if constexpr ((::ELogVerbosity::Verbosity & ::ELogVerbosity::VerbosityMask) == ::ELogVerbosity::Fatal) \
 		{ \
 			LowLevelFatalError(Format, ##__VA_ARGS__); \
 			CA_ASSUME(false); \
 		} \
 	}
-	#define UE_LOG_REF(CategoryRef, Verbosity, Format, ...) \
+	// DO NOT USE. Use UE_LOG because this will be deprecated in a future release.
+	#define UE_LOG_REF UE_LOG
+	// DO NOT USE. Use UE_LOG because this will be deprecated in a future release.
+	#define UE_LOG_CLINKAGE UE_LOG
+	#define UE_LOGF(CategoryName, Verbosity, Format, ...) \
 	{ \
-		if constexpr ((ELogVerbosity::Verbosity & ELogVerbosity::VerbosityMask) == ELogVerbosity::Fatal) \
+		if constexpr ((::ELogVerbosity::Verbosity & ::ELogVerbosity::VerbosityMask) == ::ELogVerbosity::Fatal) \
 		{ \
 			LowLevelFatalError(Format, ##__VA_ARGS__); \
 			CA_ASSUME(false); \
 		} \
 	}
-	#define UE_LOG_CLINKAGE(CategoryName, Verbosity, Format, ...) \
+
+	// Conditional logging (fatal errors only).
+	#define UE_CLOG(Condition, CategoryName, Verbosity, Format, ...) \
 	{ \
-		if constexpr ((ELogVerbosity::Verbosity & ELogVerbosity::VerbosityMask) == ELogVerbosity::Fatal) \
+		if constexpr ((::ELogVerbosity::Verbosity & ::ELogVerbosity::VerbosityMask) == ::ELogVerbosity::Fatal) \
 		{ \
-			LowLevelFatalError(Format, ##__VA_ARGS__); \
-			CA_ASSUME(false); \
+			if (Condition) \
+			{ \
+				LowLevelFatalError(Format, ##__VA_ARGS__); \
+				CA_ASSUME(false); \
+			} \
 		} \
 	}
-
-	// Conditional logging (fatal errors only).
-	#define UE_CLOG(Condition, CategoryName, Verbosity, Format, ...) \
+	#define UE_CLOGF(Condition, CategoryName, Verbosity, Format, ...) \
 	{ \
 		if constexpr ((::ELogVerbosity::Verbosity & ::ELogVerbosity::VerbosityMask) == ::ELogVerbosity::Fatal) \
 		{ \
@@ -218,7 +232,7 @@ CORE_API void BasicFatalLogWithProgramCounter(FStaticBasicLogRecordParam Log, co
 	#define UE_LOG_ACTIVE(...)				(0)
 	#define UE_LOG_ANY_ACTIVE(...)			(0)
 	#define UE_SUPPRESS(...) {}
-	#define UE_GET_LOG_VERBOSITY(...)		(ELogVerbosity::NoLogging)
+	#define UE_GET_LOG_VERBOSITY(...)		(::ELogVerbosity::NoLogging)
 	#define UE_SET_LOG_VERBOSITY(...)
 
 #else
@@ -245,23 +259,40 @@ CORE_API void BasicFatalLogWithProgramCounter(FStaticBasicLogRecordParam Log, co
 	 * @param CategoryName name of the logging category
 	 * @param Verbosity, verbosity level to test against
 	**/
-	#define UE_LOG_ACTIVE(CategoryName, Verbosity) (::UEAsserts_Private::IsLogActive<(int32)ELogVerbosity::Verbosity>(CategoryName))
+	#define UE_LOG_ACTIVE(CategoryName, Verbosity) (::UEAsserts_Private::IsLogActive<(int32)::ELogVerbosity::Verbosity>(CategoryName))
 
 	#define UE_GET_LOG_VERBOSITY(CategoryName) \
 		CategoryName.GetVerbosity()
 
 	#define UE_SET_LOG_VERBOSITY(CategoryName, Verbosity) \
-		CategoryName.SetVerbosity(ELogVerbosity::Verbosity);
+		CategoryName.SetVerbosity(::ELogVerbosity::Verbosity);
 
 	/**
 	 * A macro that logs a formatted message if the log category is active at the requested verbosity level.
+	 * Prefer UE_LOGF because its ASCII or UTF-8 format strings require less memory. 
 	 *
 	 * @param CategoryName   Name of the log category as provided to DEFINE_LOG_CATEGORY.
 	 * @param Verbosity      Verbosity level of this message. See ELogVerbosity.
-	 * @param Format         Format string literal in the style of printf.
+	 * @param Format         Format string literal in the style of printf. Specify as TEXT("...")
 	 */
 	#define UE_LOG(CategoryName, Verbosity, Format, ...) \
-		UE_PRIVATE_LOG(UE_EMPTY, constexpr, CategoryName, Verbosity, Format, ##__VA_ARGS__)
+	{ \
+		static_assert(::TIsArrayOrRefOfTypeByPredicate<decltype(Format), ::TIsCharEncodingCompatibleWithTCHAR>::Value, "Formatting string must be a TCHAR array."); \
+		UE_PRIVATE_LOG(UE_EMPTY, constexpr, CategoryName, Verbosity, Format, ##__VA_ARGS__) \
+	}
+
+	/**
+	 * A macro that logs a formatted message if the log category is active at the requested verbosity level.
+	 *
+	 * @param CategoryName   Name of the log category as provided to DEFINE_LOG_CATEGORY.
+	 * @param Verbosity      Verbosity level of this message. See ELogVerbosity.
+	 * @param Format         Format string literal in the style of printf. Specify as "..." or UTF8TEXT("...")
+	 */
+	#define UE_LOGF(CategoryName, Verbosity, Format, ...) \
+	{ \
+		static_assert(::TIsArrayOrRefOfTypeByPredicate<decltype(Format), ::TIsCharEncodingCompatibleWithUTF8CHAR>::Value, "Formatting string must be a UTF8CHAR or ANSICHAR array."); \
+		UE_PRIVATE_LOG(UE_EMPTY, constexpr, CategoryName, Verbosity, Format, ##__VA_ARGS__) \
+	}
 
 	/**
 	 * DO NOT USE. A macro that logs a formatted message if the log category is active at the requested verbosity level.
@@ -273,7 +304,10 @@ CORE_API void BasicFatalLogWithProgramCounter(FStaticBasicLogRecordParam Log, co
 	 * @param Format        Format string literal in the style of printf.
 	 */
 	#define UE_LOG_REF(CategoryRef, Verbosity, Format, ...) \
-		UE_PRIVATE_LOG(UE_EMPTY, UE_EMPTY, CategoryRef, Verbosity, Format, ##__VA_ARGS__)
+	{ \
+		static_assert(::TIsArrayOrRefOfTypeByPredicate<decltype(Format), ::TIsCharEncodingCompatibleWithTCHAR>::Value, "Formatting string must be a TCHAR array."); \
+		UE_PRIVATE_LOG(UE_EMPTY, UE_EMPTY, CategoryRef, Verbosity, Format, ##__VA_ARGS__) \
+	}
 
 	// DO NOT USE. Use UE_LOG because this will be deprecated in a future release.
 	#define UE_LOG_CLINKAGE UE_LOG
@@ -289,16 +323,22 @@ CORE_API void BasicFatalLogWithProgramCounter(FStaticBasicLogRecordParam Log, co
 	 * @param Format         Format string literal in the style of printf.
 	 */
 	#define UE_CLOG(Condition, CategoryName, Verbosity, Format, ...) \
-		UE_PRIVATE_LOG(if (Condition), constexpr, CategoryName, Verbosity, Format, ##__VA_ARGS__)
+	{ \
+		static_assert(::TIsArrayOrRefOfTypeByPredicate<decltype(Format), ::TIsCharEncodingCompatibleWithTCHAR>::Value, "Formatting string must be a TCHAR array."); \
+		UE_PRIVATE_LOG(if (Condition), constexpr, CategoryName, Verbosity, Format, ##__VA_ARGS__) \
+	}
+	#define UE_CLOGF(Condition, CategoryName, Verbosity, Format, ...) \
+	{ \
+		static_assert(::TIsArrayOrRefOfTypeByPredicate<decltype(Format), ::TIsCharEncodingCompatibleWithUTF8CHAR>::Value, "Formatting string must be a UTF8CHAR or ANSICHAR array."); \
+		UE_PRIVATE_LOG(if (Condition), constexpr, CategoryName, Verbosity, UTF8TEXT(Format), ##__VA_ARGS__) \
+	}
 
 	/** Private macro used to implement the public log macros. DO NOT CALL DIRECTLY! */
 	#define UE_PRIVATE_LOG(Condition, CategoryConst, Category, Verbosity, Format, ...) \
 	{ \
-		static_assert(std::is_const_v<std::remove_reference_t<decltype(Format)>>, "Formatting string must be a const TCHAR array."); \
-		static_assert(TIsArrayOrRefOfTypeByPredicate<decltype(Format), TIsCharEncodingCompatibleWithTCHAR>::Value, "Formatting string must be a TCHAR array."); \
 		UE_VALIDATE_FORMAT_STRING(Format, ##__VA_ARGS__); \
 		UE_IF(LOGTRACE_ENABLED, static ::UE::Logging::Private::FStaticBasicLogDynamicData LOG_Dynamic, ); \
-		UE_IF(UE_PRIVATE_LOG_USE_STATIC_RECORD, static,) constexpr ::UE::Logging::Private::FStaticBasicLogRecord LOG_Static(Format, __builtin_FILE(), __builtin_LINE(), UE_IF(LOGTRACE_ENABLED, &LOG_Dynamic, nullptr)); \
+		UE_IF(UE_PRIVATE_LOG_USE_STATIC_RECORD, static,) constexpr ::UE::Logging::Private::TStaticBasicLogRecord<TElementType_T<decltype(Format)>> LOG_Static(Format, __builtin_FILE(), __builtin_LINE(), UE_IF(LOGTRACE_ENABLED, &LOG_Dynamic, nullptr)); \
 		static_assert((::ELogVerbosity::Verbosity & ::ELogVerbosity::VerbosityMask) < ::ELogVerbosity::NumVerbosity && ::ELogVerbosity::Verbosity > 0, "Verbosity must be constant and in range."); \
 		if constexpr ((::ELogVerbosity::Verbosity & ::ELogVerbosity::VerbosityMask) == ::ELogVerbosity::Fatal) \
 		{ \
@@ -333,14 +373,14 @@ CORE_API void BasicFatalLogWithProgramCounter(FStaticBasicLogRecordParam Log, co
 	 ***/
 	#define UE_SUPPRESS(CategoryName, Verbosity, ExecuteIfUnsuppressed) \
 	{ \
-		static_assert((ELogVerbosity::Verbosity & ELogVerbosity::VerbosityMask) < ELogVerbosity::NumVerbosity && ELogVerbosity::Verbosity > 0, "Verbosity must be constant and in range."); \
-		CA_CONSTANT_IF((ELogVerbosity::Verbosity & ELogVerbosity::VerbosityMask) <= ELogVerbosity::COMPILED_IN_MINIMUM_VERBOSITY && (ELogVerbosity::Verbosity & ELogVerbosity::VerbosityMask) <= FLogCategory##CategoryName::CompileTimeVerbosity) \
+		static_assert((::ELogVerbosity::Verbosity & ::ELogVerbosity::VerbosityMask) < ::ELogVerbosity::NumVerbosity && ::ELogVerbosity::Verbosity > 0, "Verbosity must be constant and in range."); \
+		CA_CONSTANT_IF((::ELogVerbosity::Verbosity & ::ELogVerbosity::VerbosityMask) <= ::ELogVerbosity::COMPILED_IN_MINIMUM_VERBOSITY && (::ELogVerbosity::Verbosity & ::ELogVerbosity::VerbosityMask) <= FLogCategory##CategoryName::CompileTimeVerbosity) \
 		{ \
-			if (!CategoryName.IsSuppressed(ELogVerbosity::Verbosity)) \
+			if (!CategoryName.IsSuppressed(::ELogVerbosity::Verbosity)) \
 			{ \
-				FScopedCategoryAndVerbosityOverride TEMP__##CategoryName(CategoryName.GetCategoryName(), ELogVerbosity::Type(ELogVerbosity::Verbosity & ELogVerbosity::VerbosityMask)); \
+				::FScopedCategoryAndVerbosityOverride TEMP__##CategoryName(CategoryName.GetCategoryName(), ::ELogVerbosity::Type(::ELogVerbosity::Verbosity & ::ELogVerbosity::VerbosityMask)); \
 				ExecuteIfUnsuppressed; \
-				CategoryName.PostTrigger(ELogVerbosity::Verbosity); \
+				CategoryName.PostTrigger(::ELogVerbosity::Verbosity); \
 			} \
 		} \
 	}
diff --git a/Engine/Source/Runtime/Core/Public/Misc/AssertionMacros.h b/Engine/Source/Runtime/Core/Public/Misc/AssertionMacros.h
--- a/Engine/Source/Runtime/Core/Public/Misc/AssertionMacros.h
+++ b/Engine/Source/Runtime/Core/Public/Misc/AssertionMacros.h
@@ -572,6 +572,8 @@ namespace UEAsserts_Private
 namespace UE::Assert::Private
 {
 	CORE_API void UE_COLD UE_DEBUG_SECTION VARARGS ProcessLowLevelFatalError(const ANSICHAR* File, int32 Line, const TCHAR* Format, ...);
+	CORE_API void UE_COLD UE_DEBUG_SECTION VARARGS ProcessLowLevelFatalError(const ANSICHAR* File, int32 Line, const ANSICHAR* Format, ...);
+	CORE_API void UE_COLD UE_DEBUG_SECTION VARARGS ProcessLowLevelFatalError(const ANSICHAR* File, int32 Line, const UTF8CHAR* Format, ...);
 } // UE::Assert::Private
 
 template <typename... ArgTypes>
