Deprecated PLATFORM_TCHAR_IS_4_BYTES

#Synced-CL 0

diff --git a/Engine/Source/Developer/ScriptDisassembler/Private/ScriptDisassembler.cpp b/Engine/Source/Developer/ScriptDisassembler/Private/ScriptDisassembler.cpp
--- a/Engine/Source/Developer/ScriptDisassembler/Private/ScriptDisassembler.cpp
+++ b/Engine/Source/Developer/ScriptDisassembler/Private/ScriptDisassembler.cpp
@@ -238,9 +238,6 @@ FString FKismetBytecodeDisassembler::ReadString16(int32& ScriptIndex)
 	}
 	while ((Script[ScriptIndex-1] != 0) || (Script[ScriptIndex-2] != 0));
 
-	// Inline combine any surrogate pairs in the data when loading into a UTF-32 string
-	StringConv::InlineCombineSurrogates(Result);
-
 	return Result;
 }
 
diff --git a/Engine/Source/Programs/UnrealObjectPtrTool/Private/UnrealObjectPtrTool.cpp b/Engine/Source/Programs/UnrealObjectPtrTool/Private/UnrealObjectPtrTool.cpp
--- a/Engine/Source/Programs/UnrealObjectPtrTool/Private/UnrealObjectPtrTool.cpp
+++ b/Engine/Source/Programs/UnrealObjectPtrTool/Private/UnrealObjectPtrTool.cpp
@@ -444,12 +444,6 @@ bool FPointerUpgrader::LoadFileToString(FString& OutResult, EByteOrderMarkerType
 	{
 		// Else ensure null terminator is present
 		ResultArray.Last() = 0;
-
-		if (bIsUnicode)
-		{
-			// Inline combine any surrogate pairs in the data when loading into a UTF-32 string
-			StringConv::InlineCombineSurrogates(OutResult);
-		}
 	}
 
 	return true;
diff --git a/Engine/Source/Runtime/Analytics/Analytics/Private/AnalyticsPropertyStore.cpp b/Engine/Source/Runtime/Analytics/Analytics/Private/AnalyticsPropertyStore.cpp
--- a/Engine/Source/Runtime/Analytics/Analytics/Private/AnalyticsPropertyStore.cpp
+++ b/Engine/Source/Runtime/Analytics/Analytics/Private/AnalyticsPropertyStore.cpp
@@ -798,9 +798,6 @@ IAnalyticsPropertyStore::EStatusCode FAnalyticsPropertyStore::GetStringValueInte
 			Passthru.Get()[CharCountWithNull - 1] = '\0';
 			Passthru.Apply();
 
-			// Inline combine any surrogate pairs in the data when loading into a UTF-32 string
-			StringConv::InlineCombineSurrogates(OutValue);
-
 			return EStatusCode::Success;
 		}
 		return EStatusCode::BadType;
diff --git a/Engine/Source/Runtime/Core/Private/Containers/String.cpp b/Engine/Source/Runtime/Core/Private/Containers/String.cpp
--- a/Engine/Source/Runtime/Core/Private/Containers/String.cpp
+++ b/Engine/Source/Runtime/Core/Private/Containers/String.cpp
@@ -46,11 +46,6 @@ void FTextRange::CalculateLineRangesFromString(const FString& Input, TArray<FTex
 	}
 }
 
-void StringConv::InlineCombineSurrogates(FString& Str)
-{
-	InlineCombineSurrogates_Array(Str.GetCharArray());
-}
-
 namespace UE::Core::Private
 {
 	UE_DISABLE_OPTIMIZATION_SHIP
diff --git a/Engine/Source/Runtime/Core/Private/Containers/String.cpp.inl b/Engine/Source/Runtime/Core/Private/Containers/String.cpp.inl
--- a/Engine/Source/Runtime/Core/Private/Containers/String.cpp.inl
+++ b/Engine/Source/Runtime/Core/Private/Containers/String.cpp.inl
@@ -356,15 +356,7 @@ NSString* UE_STRING_CLASS::GetNSString() const
 /** Convert string to bridgeable CFString */
 CFStringRef UE_STRING_CLASS::GetCFString() const
 {
-    CFStringRef OutString;
-    
-#if PLATFORM_TCHAR_IS_4_BYTES
-    OutString = CFStringCreateWithBytes(kCFAllocatorDefault, (UInt8*)Data.GetData(), Len() * sizeof(ElementType), kCFStringEncodingUTF32LE, false);
-#else
-    OutString = CFStringCreateWithBytes(kCFAllocatorDefault, (UInt8*)Data.GetData(), Len() * sizeof(ElementType), kCFStringEncodingUTF16LE, false);
-#endif
-    
-    return OutString;
+	return CFStringCreateWithBytes(kCFAllocatorDefault, (UInt8*)Data.GetData(), Len() * sizeof(ElementType), kCFStringEncodingUTF16LE, false);
 }
 #endif
 
@@ -1849,9 +1841,6 @@ FArchive& operator<<(FArchive& Ar, UE_STRING_CLASS& Str)
 					Passthru.Get()[SaveNum - 1] = '\0';
 					Passthru.Apply();
 
-					// Inline combine any surrogate pairs in the data when loading into a UTF-32 string
-					StringConv::InlineCombineSurrogates(Str);
-
 					// Since Microsoft's vsnwprintf implementation raises an invalid parameter warning
 					// with a character of 0xffff, scan for it and terminate the string there.
 					// 0xffff isn't an actual Unicode character anyway.
diff --git a/Engine/Source/Runtime/Core/Private/GenericPlatform/GenericWidePlatformString.cpp b/Engine/Source/Runtime/Core/Private/GenericPlatform/GenericWidePlatformString.cpp
--- a/Engine/Source/Runtime/Core/Private/GenericPlatform/GenericWidePlatformString.cpp
+++ b/Engine/Source/Runtime/Core/Private/GenericPlatform/GenericWidePlatformString.cpp
@@ -100,7 +100,7 @@ WIDECHAR* FGenericWidePlatformString::Strncat(WIDECHAR* Dest, const WIDECHAR* Sr
 
 int32 FGenericWidePlatformString::Strtoi( const WIDECHAR* Start, WIDECHAR** End, int32 Base )
 {
-#if PLATFORM_TCHAR_IS_4_BYTES || PLATFORM_TCHAR_IS_UTF8CHAR
+#if PLATFORM_TCHAR_IS_UTF8CHAR
 	unimplemented();
 #else
 	static_assert(sizeof(TCHAR) == 2, "TCHAR is expected to be 16-bit");
@@ -129,7 +129,7 @@ int32 FGenericWidePlatformString::Strtoi( const WIDECHAR* Start, WIDECHAR** End,
 
 int64 FGenericWidePlatformString::Strtoi64( const WIDECHAR* Start, WIDECHAR** End, int32 Base )
 {
-#if PLATFORM_TCHAR_IS_4_BYTES || PLATFORM_TCHAR_IS_UTF8CHAR
+#if PLATFORM_TCHAR_IS_UTF8CHAR
 	unimplemented();
 #else
 	static_assert(sizeof(TCHAR) == 2, "TCHAR is expected to be 16-bit");
diff --git a/Engine/Source/Runtime/Core/Private/Internationalization/Text.cpp b/Engine/Source/Runtime/Core/Private/Internationalization/Text.cpp
--- a/Engine/Source/Runtime/Core/Private/Internationalization/Text.cpp
+++ b/Engine/Source/Runtime/Core/Private/Internationalization/Text.cpp
@@ -2132,29 +2132,27 @@ UTF32CHAR FTextChar::GetCodepoint(const TCHAR* Buffer, int32* OutNumTCHARInCodep
 		return 0;
 	}
 
-    UTF32CHAR Codepoint = 0;
+	UTF32CHAR Codepoint = 0;
 
-#if !PLATFORM_TCHAR_IS_4_BYTES
 	const TCHAR NextChar = *(Buffer + 1);
-    if (NextChar && StringConv::IsHighSurrogate(*Buffer) && StringConv::IsLowSurrogate(NextChar))
-    {
-        Codepoint = StringConv::EncodeSurrogate(*Buffer, NextChar);
-        if (OutNumTCHARInCodepoint)
-        {
-            *OutNumTCHARInCodepoint = 2;
-        }
-    }
-    else
-#endif  // !PLATFORM_TCHAR_IS_4_BYTES
-    {
-        Codepoint = *Buffer;
-        if (OutNumTCHARInCodepoint)
-        {
-            *OutNumTCHARInCodepoint = 1;
-        }
-    }
-
-    return Codepoint;
+	if (NextChar && StringConv::IsHighSurrogate(*Buffer) && StringConv::IsLowSurrogate(NextChar))
+	{
+		Codepoint = StringConv::EncodeSurrogate(*Buffer, NextChar);
+		if (OutNumTCHARInCodepoint)
+		{
+			*OutNumTCHARInCodepoint = 2;
+		}
+	}
+	else
+	{
+		Codepoint = *Buffer;
+		if (OutNumTCHARInCodepoint)
+		{
+			*OutNumTCHARInCodepoint = 1;
+		}
+	}
+
+	return Codepoint;
 }
 
 bool FTextChar::ConvertCodepointToString(const UTF32CHAR Codepoint, FString& OutString)
@@ -2177,7 +2175,6 @@ bool FTextChar::AppendCodepointToString(const UTF32CHAR Codepoint, FString& OutS
 		return false;
 	}
 
-#if !PLATFORM_TCHAR_IS_4_BYTES
 	if (StringConv::IsEncodedSurrogate(Codepoint))
 	{
 		uint16 HighSurrogate = 0;
@@ -2193,7 +2190,6 @@ bool FTextChar::AppendCodepointToString(const UTF32CHAR Codepoint, FString& OutS
 		return false;
 	}
 	else
-#endif	// !PLATFORM_TCHAR_IS_4_BYTES
 	{
 		OutString.AppendChar((TCHAR)Codepoint);
 	}
diff --git a/Engine/Source/Runtime/Core/Private/Internationalization/TextKey.cpp b/Engine/Source/Runtime/Core/Private/Internationalization/TextKey.cpp
--- a/Engine/Source/Runtime/Core/Private/Internationalization/TextKey.cpp
+++ b/Engine/Source/Runtime/Core/Private/Internationalization/TextKey.cpp
@@ -742,9 +742,6 @@ bool LoadKeyString(FArchive& Ar, FInlineStringBuffer& OutStrBuffer)
 			Ar.Serialize(Passthru.Get(), SaveNum * sizeof(UCS2CHAR));
 			Passthru.Get()[SaveNum - 1] = 0; // Ensure the string has a null terminator
 			Passthru.Apply();
-
-			// Inline combine any surrogate pairs in the data when loading into a UTF-32 string
-			StringConv::InlineCombineSurrogates_Array(OutStrBuffer);
 		}
 		else
 		{
diff --git a/Engine/Source/Runtime/Core/Private/Misc/FileHelper.cpp b/Engine/Source/Runtime/Core/Private/Misc/FileHelper.cpp
--- a/Engine/Source/Runtime/Core/Private/Misc/FileHelper.cpp
+++ b/Engine/Source/Runtime/Core/Private/Misc/FileHelper.cpp
@@ -145,7 +145,6 @@ void FFileHelper::BufferToString( FString& Result, const uint8* Buffer, int32 Si
 	TArray<TCHAR, FString::AllocatorType>& ResultArray = Result.GetCharArray();
 	ResultArray.Empty();
 
-	bool bIsUnicode = false;
 	if( Size >= 2 && !( Size & 1 ) && Buffer[0] == 0xff && Buffer[1] == 0xfe )
 	{
 		// Unicode Intel byte order. Less 1 for the FFFE header, additional 1 for null terminator.
@@ -154,7 +153,6 @@ void FFileHelper::BufferToString( FString& Result, const uint8* Buffer, int32 Si
 		{
 			ResultArray[ i ] = CharCast<TCHAR>( (UCS2CHAR)(( uint16 )Buffer[i * 2 + 2] + ( uint16 )Buffer[i * 2 + 3] * 256) );
 		}
-		bIsUnicode = true;
 	}
 	else if( Size >= 2 && !( Size & 1 ) && Buffer[0] == 0xfe && Buffer[1] == 0xff )
 	{
@@ -164,7 +162,6 @@ void FFileHelper::BufferToString( FString& Result, const uint8* Buffer, int32 Si
 		{
 			ResultArray[ i ] = CharCast<TCHAR>( (UCS2CHAR)(( uint16 )Buffer[i * 2 + 3] + ( uint16 )Buffer[i * 2 + 2] * 256) );
 		}
-		bIsUnicode = true;
 	}
 	else
 	{
@@ -190,12 +187,6 @@ void FFileHelper::BufferToString( FString& Result, const uint8* Buffer, int32 Si
 	{
 		// Else ensure null terminator is present
 		ResultArray.Last() = TCHAR('\0');
-
-		if (bIsUnicode)
-		{
-			// Inline combine any surrogate pairs in the data when loading into a UTF-32 string
-			StringConv::InlineCombineSurrogates(Result);
-		}
 	}
 }
 
diff --git a/Engine/Source/Runtime/Core/Private/UObject/UnrealNames.cpp b/Engine/Source/Runtime/Core/Private/UObject/UnrealNames.cpp
--- a/Engine/Source/Runtime/Core/Private/UObject/UnrealNames.cpp
+++ b/Engine/Source/Runtime/Core/Private/UObject/UnrealNames.cpp
@@ -4413,11 +4413,6 @@ FArchive& operator<<(FArchive& Ar, FNameEntrySerialized& E)
 			auto Sink = StringMemoryPassthru<UCS2CHAR>(WideName, StringLen, StringLen);
 			Ar.Serialize(Sink.Get(), StringLen * sizeof(UCS2CHAR));
 			Sink.Apply();
-
-#if PLATFORM_TCHAR_IS_4_BYTES
-			// Inline combine any surrogate pairs in the data when loading into a UTF-32 string
-			StringLen = StringConv::InlineCombineSurrogates_Buffer(WideName, StringLen);
-#endif	// PLATFORM_TCHAR_IS_4_BYTES
 		}
 		else
 		{
@@ -4876,11 +4871,6 @@ static FDisplayNameEntryId BatchLoadNameWithoutHash(const UTF16CHAR* Str, uint32
 	{
 		Temp[Idx] = INTEL_ORDER16(Str[Idx]);
 	}
-	
-#if PLATFORM_TCHAR_IS_4_BYTES
-	// Inline combine any surrogate pairs in the data when loading into a UTF-32 string
-	Len = StringConv::InlineCombineSurrogates_Buffer(Temp, Len);
-#endif
 
 	FNameStringView Name(Temp, Len);
 	FNameHash Hash = HashName<ENameCase::IgnoreCase>(Name);
diff --git a/Engine/Source/Runtime/Core/Public/Containers/StringConv.h b/Engine/Source/Runtime/Core/Public/Containers/StringConv.h
--- a/Engine/Source/Runtime/Core/Public/Containers/StringConv.h
+++ b/Engine/Source/Runtime/Core/Public/Containers/StringConv.h
@@ -151,12 +151,9 @@ namespace StringConv
 
 	/** Inline combine any UTF-16 surrogate pairs in the given null-terminated TCHAR array */
 	template<typename AllocatorType>
+	UE_DEPRECATED(5.8, "Remove calls to this function because it does nothing.")
 	inline void InlineCombineSurrogates_Array(TArray<TCHAR, AllocatorType>& StrBuffer)
 	{
-#if PLATFORM_TCHAR_IS_4_BYTES
-		const int32 NewStrLen = InlineCombineSurrogates_Buffer(StrBuffer.GetData(), StrBuffer.Num() - 1);
-		StrBuffer.SetNum(NewStrLen + 1, EAllowShrinking::No);
-#endif	// PLATFORM_TCHAR_IS_4_BYTES
 	}
 
 	struct FUnused;
@@ -1029,21 +1026,7 @@ using FUTF8ToTCHAR /*UE_DEPRECATED(5.xx, "FUTF8ToTCHAR(PtrToUTF8) is deprecated,
 // UTF32CHAR.                                                     //
 ////////////////////////////////////////////////////////////////////
 
-// special handling for platforms still using a 32-bit TCHAR
-#if PLATFORM_TCHAR_IS_4_BYTES
-
-typedef TStringConversion<TUTF32ToUTF16_Convert<TCHAR, UTF16CHAR>> FTCHARToUTF16;
-typedef TStringConversion<TUTF16ToUTF32_Convert<UTF16CHAR, TCHAR>> FUTF16ToTCHAR;
-#define TCHAR_TO_UTF16(str) (UTF16CHAR*)FTCHARToUTF16((const TCHAR*)str).Get()
-#define UTF16_TO_TCHAR(str) (TCHAR*)FUTF16ToTCHAR((const UTF16CHAR*)str).Get()
-
-static_assert(sizeof(TCHAR) == sizeof(UTF32CHAR), "TCHAR and UTF32CHAR are expected to be the same size for inline conversion! PLATFORM_TCHAR_IS_4_BYTES is not configured correctly for this platform.");
-typedef TStringPointer<TCHAR, UTF32CHAR> FTCHARToUTF32;
-typedef TStringPointer<UTF32CHAR, TCHAR> FUTF32ToTCHAR;
-#define TCHAR_TO_UTF32(str) (UTF32CHAR*)(str)
-#define UTF32_TO_TCHAR(str) (TCHAR*)(str)
-
-#elif PLATFORM_TCHAR_IS_UTF8CHAR
+#if PLATFORM_TCHAR_IS_UTF8CHAR
 
 typedef TStringConversion<TStringConvert<TCHAR, UTF16CHAR>> FTCHARToUTF16;
 typedef TStringConversion<TStringConvert<UTF16CHAR, TCHAR>> FUTF16ToTCHAR;
@@ -1057,7 +1040,7 @@ typedef TStringConversion<TStringConvert<UTF32CHAR, TCHAR>> FUTF32ToTCHAR;
 
 #else
 
-static_assert(sizeof(TCHAR) == sizeof(UTF16CHAR), "TCHAR and UTF16CHAR are expected to be the same size for inline conversion! PLATFORM_TCHAR_IS_4_BYTES is not configured correctly for this platform.");
+static_assert(sizeof(TCHAR) == sizeof(UTF16CHAR), "TCHAR and UTF16CHAR are expected to be the same size for inline conversion!");
 typedef TStringPointer<TCHAR, UTF16CHAR> FTCHARToUTF16;
 typedef TStringPointer<UTF16CHAR, TCHAR> FUTF16ToTCHAR;
 #define TCHAR_TO_UTF16(str) (UTF16CHAR*)(str)
diff --git a/Engine/Source/Runtime/Core/Public/Containers/UnrealString.h b/Engine/Source/Runtime/Core/Public/Containers/UnrealString.h
--- a/Engine/Source/Runtime/Core/Public/Containers/UnrealString.h
+++ b/Engine/Source/Runtime/Core/Public/Containers/UnrealString.h
@@ -227,7 +227,10 @@ struct TTypeFromString
 namespace StringConv
 {
 	/** Inline combine any UTF-16 surrogate pairs in the given string */
-	CORE_API void InlineCombineSurrogates(FString& Str);
+	UE_DEPRECATED(5.8, "Remove calls to this function because it does nothing.")
+	inline void InlineCombineSurrogates(FString& Str)
+	{
+	}
 }
 
 struct FTextRange
diff --git a/Engine/Source/Runtime/Core/Public/Containers/UnrealString.h.inl b/Engine/Source/Runtime/Core/Public/Containers/UnrealString.h.inl
--- a/Engine/Source/Runtime/Core/Public/Containers/UnrealString.h.inl
+++ b/Engine/Source/Runtime/Core/Public/Containers/UnrealString.h.inl
@@ -190,12 +190,7 @@ public:
 		{
 			// Convert the NSString data into the native TCHAR format for UE4
 			// This returns a buffer of bytes, but they can be safely cast to a buffer of TCHARs
-#if PLATFORM_TCHAR_IS_4_BYTES
-			const CFStringEncoding Encoding = kCFStringEncodingUTF32LE;
-#else
 			const CFStringEncoding Encoding = kCFStringEncodingUTF16LE;
-#endif
-
 			CFRange Range = CFRangeMake(0, StringLength);
 			CFIndex BytesNeeded;
 			if (CFStringGetBytes(In, Range, Encoding, '?', false, NULL, 0, &BytesNeeded) > 0)
diff --git a/Engine/Source/Runtime/Core/Public/GenericPlatform/GenericPlatformString.h b/Engine/Source/Runtime/Core/Public/GenericPlatform/GenericPlatformString.h
--- a/Engine/Source/Runtime/Core/Public/GenericPlatform/GenericPlatformString.h
+++ b/Engine/Source/Runtime/Core/Public/GenericPlatform/GenericPlatformString.h
@@ -113,11 +113,7 @@ struct FGenericPlatformString : public FGenericPlatformStricmp
 
 	static const ANSICHAR* GetEncodingName()
 	{
-#if PLATFORM_TCHAR_IS_4_BYTES
-		return "UTF-32LE";
-#else
 		return "UTF-16LE";
-#endif
 	}
 
 	/**
diff --git a/Engine/Source/Runtime/Core/Public/HAL/Platform.h b/Engine/Source/Runtime/Core/Public/HAL/Platform.h
--- a/Engine/Source/Runtime/Core/Public/HAL/Platform.h
+++ b/Engine/Source/Runtime/Core/Public/HAL/Platform.h
@@ -273,7 +273,7 @@
 	#define PLATFORM_COMPILER_HAS_GENERATED_COMPARISON_OPERATORS (__cplusplus >= 202002L)
 #endif
 // Feature test macro for constexpr __builtin_FILE() and __builtin_LINE()
-#define PLATFORM_COMPILER_SUPPORTS_CONSTEXPR_BUILTIN_FILE_AND_LINE 1 UE_DEPRECATED_MACRO(5.7, "PLATFORM_COMPILER_SUPPORTS_CONSTEXPR_BUILTIN_FILE_AND_LINE has been deprecated and should be replaced with 1.")
+#define PLATFORM_COMPILER_SUPPORTS_CONSTEXPR_BUILTIN_FILE_AND_LINE UE_DEPRECATED_MACRO(5.7, "PLATFORM_COMPILER_SUPPORTS_CONSTEXPR_BUILTIN_FILE_AND_LINE has been deprecated and should be replaced with 1.") 1
 #ifndef PLATFORM_COMPILER_SUPPORTS_BUILTIN_BITCAST
 	#error PLATFORM_COMPILER_SUPPORTS_BUILTIN_BITCAST should have been defined by now
 #endif
@@ -287,9 +287,7 @@
 		#define PLATFORM_COMPILER_SUPPORTS_CONSTRAINED_DESTRUCTORS 0
 	#endif
 #endif
-#ifndef PLATFORM_TCHAR_IS_4_BYTES
-	#define PLATFORM_TCHAR_IS_4_BYTES			0
-#endif
+#define PLATFORM_TCHAR_IS_4_BYTES UE_DEPRECATED_MACRO(5.8, "PLATFORM_TCHAR_IS_4_BYTES has been deprecated and should be replaced with 0.") 0
 #ifndef PLATFORM_WCHAR_IS_4_BYTES
 	#define PLATFORM_WCHAR_IS_4_BYTES			0
 #endif
@@ -1167,9 +1165,7 @@ namespace TypeTests
 	template <typename T>
 	inline constexpr bool TAreTypesEqual_V<T, T> = true;
 
-#if PLATFORM_TCHAR_IS_4_BYTES
-	static_assert(sizeof(TCHAR) == 4, "TCHAR size must be 4 bytes.");
-#elif PLATFORM_TCHAR_IS_UTF8CHAR
+#if PLATFORM_TCHAR_IS_UTF8CHAR
 	static_assert(sizeof(TCHAR) == 1, "TCHAR size must be 1 byte.");
 #else
 	static_assert(sizeof(TCHAR) == 2, "TCHAR size must be 2 bytes.");
diff --git a/Engine/Source/Runtime/Core/Public/UObject/NameBatchSerialization.h b/Engine/Source/Runtime/Core/Public/UObject/NameBatchSerialization.h
--- a/Engine/Source/Runtime/Core/Public/UObject/NameBatchSerialization.h
+++ b/Engine/Source/Runtime/Core/Public/UObject/NameBatchSerialization.h
@@ -6,7 +6,7 @@
 #include "UObject/NameTypes.h"
 #include "Templates/Function.h"
 
-#define ALLOW_NAME_BATCH_SAVING PLATFORM_LITTLE_ENDIAN && !PLATFORM_TCHAR_IS_4_BYTES
+#define ALLOW_NAME_BATCH_SAVING PLATFORM_LITTLE_ENDIAN
 
 //////////////////////////////////////////////////////////////////////////
 
diff --git a/Engine/Source/Runtime/Core/Public/Windows/WindowsPlatform.h b/Engine/Source/Runtime/Core/Public/Windows/WindowsPlatform.h
--- a/Engine/Source/Runtime/Core/Public/Windows/WindowsPlatform.h
+++ b/Engine/Source/Runtime/Core/Public/Windows/WindowsPlatform.h
@@ -79,7 +79,6 @@ typedef FWindowsPlatformTypes FPlatformTypes;
 
 #define PLATFORM_USE_LS_SPEC_FOR_WIDECHAR					0
 //#define PLATFORM_USE_SYSTEM_VSWPRINTF						1
-//#define PLATFORM_TCHAR_IS_4_BYTES							0
 #define PLATFORM_HAS_BSD_TIME								0
 #define PLATFORM_USE_PTHREADS								0
 #define PLATFORM_USES_UNFAIR_LOCKS							1
diff --git a/Engine/Source/Runtime/CoreUObject/Private/Serialization/PropertyLocalizationDataGathering.cpp b/Engine/Source/Runtime/CoreUObject/Private/Serialization/PropertyLocalizationDataGathering.cpp
--- a/Engine/Source/Runtime/CoreUObject/Private/Serialization/PropertyLocalizationDataGathering.cpp
+++ b/Engine/Source/Runtime/CoreUObject/Private/Serialization/PropertyLocalizationDataGathering.cpp
@@ -712,9 +712,6 @@ private:
 				iCode += sizeof(uint16);
 			}
 			while (Script[iCode-1] || Script[iCode-2]);
-
-			// Inline combine any surrogate pairs in the data when loading into a UTF-32 string
-			StringConv::InlineCombineSurrogates(LastParsedString);
 		}
 		else
 		{
diff --git a/Engine/Source/Runtime/CoreUObject/Private/UObject/ScriptCore.cpp b/Engine/Source/Runtime/CoreUObject/Private/UObject/ScriptCore.cpp
--- a/Engine/Source/Runtime/CoreUObject/Private/UObject/ScriptCore.cpp
+++ b/Engine/Source/Runtime/CoreUObject/Private/UObject/ScriptCore.cpp
@@ -3454,9 +3454,6 @@ DEFINE_FUNCTION(UObject::execUnicodeStringConst)
 	FString& ResultStr = *(FString*)RESULT_PARAM;
 	ResultStr = FString((UCS2CHAR*)Stack.Code);
 
-	// Inline combine any surrogate pairs in the data when loading into a UTF-32 string
-	StringConv::InlineCombineSurrogates(ResultStr);
-
 	while( *(uint16*)Stack.Code )
 	{
 		Stack.Code+=sizeof(uint16);
diff --git a/Engine/Source/Runtime/CoreUObject/Private/UObject/SoftObjectPath.cpp b/Engine/Source/Runtime/CoreUObject/Private/UObject/SoftObjectPath.cpp
--- a/Engine/Source/Runtime/CoreUObject/Private/UObject/SoftObjectPath.cpp
+++ b/Engine/Source/Runtime/CoreUObject/Private/UObject/SoftObjectPath.cpp
@@ -510,9 +510,6 @@ void SoftObjectPathLoadSubPathWorkaround(FArchive& Ar, FUtf8String& OutStr)
 	{
 		FString Temp = (UCS2CHAR*)SavedBytes.GetData();
 
-		// Inline combine any surrogate pairs in the data when loading into a UTF-32 string
-		StringConv::InlineCombineSurrogates(Temp);
-
 		// Since Microsoft's vsnwprintf implementation raises an invalid parameter warning
 		// with a character of 0xffff, scan for it and terminate the string there.
 		// 0xffff isn't an actual Unicode character anyway.
diff --git a/Engine/Source/Runtime/CoreUObject/Public/UObject/ScriptSerialization.inl b/Engine/Source/Runtime/CoreUObject/Public/UObject/ScriptSerialization.inl
--- a/Engine/Source/Runtime/CoreUObject/Public/UObject/ScriptSerialization.inl
+++ b/Engine/Source/Runtime/CoreUObject/Public/UObject/ScriptSerialization.inl
@@ -421,9 +421,6 @@
 				else
 				{
 					LongPath = FString((UCS2CHAR*)&Script[iCode]);
-
-					// Inline combine any surrogate pairs in the data when loading into a UTF-32 string
-					StringConv::InlineCombineSurrogates(LongPath);
 					XFERUNICODESTRING();
 				}
 				FSoftObjectPath Path(LongPath);
diff --git a/Engine/Source/Runtime/Json/Public/Serialization/JsonReader.h b/Engine/Source/Runtime/Json/Public/Serialization/JsonReader.h
--- a/Engine/Source/Runtime/Json/Public/Serialization/JsonReader.h
+++ b/Engine/Source/Runtime/Json/Public/Serialization/JsonReader.h
@@ -731,9 +731,6 @@ private:
 
 		StringValue = StoredStringType(StringBuffer);
 
-		// Inline combine any surrogate pairs in the data when loading into a UTF-32 string
-		InlineCombineSurrogates(StringValue);
-
 		return true;
 	}
 
@@ -944,15 +941,12 @@ protected:
 	}
 
 	template <typename Type>
+	UE_DEPRECATED(5.8, "Remove calls to this function because it does nothing.")
 	void InlineCombineSurrogates(Type& String)
 	{
-		if constexpr (std::is_same_v<Type, FString>)
-		{
-			StringConv::InlineCombineSurrogates(String);
-		}
 	}
-protected:
 
+protected:
 	TArray<EJson> ParseState;
 	EJsonToken CurrentToken;
 
diff --git a/Engine/Source/Runtime/Launch/Private/LaunchEngineLoop.cpp b/Engine/Source/Runtime/Launch/Private/LaunchEngineLoop.cpp
--- a/Engine/Source/Runtime/Launch/Private/LaunchEngineLoop.cpp
+++ b/Engine/Source/Runtime/Launch/Private/LaunchEngineLoop.cpp
@@ -6875,25 +6875,19 @@ bool FEngineLoop::AppInit( )
 #endif
 
 	// Checks.
-	check(sizeof(uint8) == 1);
-	check(sizeof(int8) == 1);
-	check(sizeof(uint16) == 2);
-	check(sizeof(uint32) == 4);
-	check(sizeof(uint64) == 8);
-	check(sizeof(ANSICHAR) == 1);
-
-#if PLATFORM_TCHAR_IS_4_BYTES
-	check(sizeof(TCHAR) == 4);
-#else
-	check(sizeof(TCHAR) == 2);
-#endif
-
-	check(sizeof(int16) == 2);
-	check(sizeof(int32) == 4);
-	check(sizeof(int64) == 8);
-	check(sizeof(bool) == 1);
-	check(sizeof(float) == 4);
-	check(sizeof(double) == 8);
+	static_assert(sizeof(uint8) == 1);
+	static_assert(sizeof(int8) == 1);
+	static_assert(sizeof(uint16) == 2);
+	static_assert(sizeof(uint32) == 4);
+	static_assert(sizeof(uint64) == 8);
+	static_assert(sizeof(ANSICHAR) == 1);
+	static_assert(sizeof(TCHAR) == 2);
+	static_assert(sizeof(int16) == 2);
+	static_assert(sizeof(int32) == 4);
+	static_assert(sizeof(int64) == 8);
+	static_assert(sizeof(bool) == 1);
+	static_assert(sizeof(float) == 4);
+	static_assert(sizeof(double) == 8);
 
 	// Init list of common colors.
 	GColorList.CreateColorMap();
diff --git a/Engine/Source/Runtime/SlateCore/Private/Fonts/SlateTextShaper.cpp b/Engine/Source/Runtime/SlateCore/Private/Fonts/SlateTextShaper.cpp
--- a/Engine/Source/Runtime/SlateCore/Private/Fonts/SlateTextShaper.cpp
+++ b/Engine/Source/Runtime/SlateCore/Private/Fonts/SlateTextShaper.cpp
@@ -31,12 +31,10 @@ bool RenderCodepointAsWhitespace(const TCHAR InCodepoint)
 		|| InCodepoint == TEXT('\u202F');	// Narrow No-Break Space
 }
 
-#if !PLATFORM_TCHAR_IS_4_BYTES
 bool RenderCodepointAsWhitespace(const UTF32CHAR InCodepoint)
 {
 	return RenderCodepointAsWhitespace((TCHAR)InCodepoint);
 }
-#endif
 
 struct FKerningOnlyTextSequenceEntry
 {
@@ -256,7 +254,6 @@ void FSlateTextShaper::PerformTextShaping(const TCHAR* InText, const int32 InTex
 						return true;
 					}
 
-#if !PLATFORM_TCHAR_IS_4_BYTES
 					{
 						const int32 NextTextIndex = RunningTextIndex + 1;
 						if (NextTextIndex < TextEndIndex && StringConv::IsHighSurrogate(Char) && StringConv::IsLowSurrogate(InText[NextTextIndex]))
@@ -264,7 +261,6 @@ void FSlateTextShaper::PerformTextShaping(const TCHAR* InText, const int32 InTex
 							return true;
 						}
 					}
-#endif	// !PLATFORM_TCHAR_IS_4_BYTES
 				}
 			}
 
@@ -572,7 +568,6 @@ void FSlateTextShaper::PerformHarfBuzzTextShaping(const TCHAR* InText, const int
 			
 			// Combine any surrogate pairs into a complete codepoint
 			UTF32CHAR CurrentCodepoint = InText[RunningTextIndex];
-#if !PLATFORM_TCHAR_IS_4_BYTES
 			if (ClusterSize > 1)
 			{
 				const int32 NextTextIndex = RunningTextIndex + 1;
@@ -585,7 +580,6 @@ void FSlateTextShaper::PerformHarfBuzzTextShaping(const TCHAR* InText, const int
 					bBypassLetterSpacing = true;
 				}
 			}
-#endif	// !PLATFORM_TCHAR_IS_4_BYTES
 
 			const bool bShouldRenderAsWhitespace = RenderCodepointAsWhitespace(CurrentCodepoint);
 
