BLAKE3: Removed Archive.h and UnrealString.h header dependencies

#rnx
#Synced-CL 0

diff --git a/Engine/Source/Runtime/Core/Private/Hash/Blake3.cpp b/Engine/Source/Runtime/Core/Private/Hash/Blake3.cpp
--- a/Engine/Source/Runtime/Core/Private/Hash/Blake3.cpp
+++ b/Engine/Source/Runtime/Core/Private/Hash/Blake3.cpp
@@ -5,10 +5,10 @@
 #include "AutoRTFM.h"
 #include "Containers/Array.h"
 #include "Containers/ArrayView.h"
-#include "Containers/ContainersFwd.h"
 #include "Containers/UnrealString.h"
 #include "Memory/CompositeBuffer.h"
 #include "Memory/SharedBuffer.h"
+#include "Serialization/Archive.h"
 #include "blake3.h"
 
 ///////////////////////////////////////////////////////////////////////////////////////////////////////////////////////
@@ -98,6 +98,12 @@ FBlake3Hash FBlake3::HashBuffer(const FCompositeBuffer& Buffer)
 
 ///////////////////////////////////////////////////////////////////////////////////////////////////////////////////////
 
+FArchive& operator<<(FArchive& Ar, FBlake3Hash& Value)
+{
+	Ar.Serialize(Value.GetBytes(), sizeof(decltype(Value.GetBytes())));
+	return Ar;
+}
+
 FString LexToString(const FBlake3Hash& Hash)
 {
 	FString Output;
diff --git a/Engine/Source/Runtime/Core/Private/IO/IoHash.cpp b/Engine/Source/Runtime/Core/Private/IO/IoHash.cpp
--- a/Engine/Source/Runtime/Core/Private/IO/IoHash.cpp
+++ b/Engine/Source/Runtime/Core/Private/IO/IoHash.cpp
@@ -4,6 +4,19 @@
 
 #include "Containers/Array.h"
 #include "Containers/UnrealString.h"
+#include "Memory/CompositeBuffer.h"
+#include "Serialization/Archive.h"
+
+FIoHash FIoHash::HashBuffer(const FCompositeBuffer& Buffer)
+{
+	return FBlake3::HashBuffer(Buffer);
+}
+
+FArchive& operator<<(FArchive& Ar, FIoHash& InHash)
+{
+	Ar.Serialize(InHash.GetBytes(), sizeof(decltype(InHash.GetBytes())));
+	return Ar;
+}
 
 FString LexToString(const FIoHash& Hash)
 {
@@ -14,3 +27,8 @@ FString LexToString(const FIoHash& Hash)
 	CharArray.Last() = TEXT('\0');
 	return Output;
 }
+
+FIoHash FIoHashBuilder::HashBuffer(const FCompositeBuffer& Buffer)
+{
+	return FBlake3::HashBuffer(Buffer);
+}
diff --git a/Engine/Source/Runtime/Core/Public/Hash/Blake3.h b/Engine/Source/Runtime/Core/Public/Hash/Blake3.h
--- a/Engine/Source/Runtime/Core/Public/Hash/Blake3.h
+++ b/Engine/Source/Runtime/Core/Public/Hash/Blake3.h
@@ -4,21 +4,15 @@
 
 #include "Containers/StringFwd.h"
 #include "Containers/StringView.h"
-#include "Containers/UnrealString.h"
-#include "HAL/Platform.h"
-#include "HAL/PlatformString.h"
 #include "HAL/UnrealMemory.h"
 #include "Memory/MemoryFwd.h"
 #include "Memory/MemoryView.h"
 #include "Misc/AssertionMacros.h"
-#include "Serialization/Archive.h"
 #include "String/BytesToHex.h"
 #include "String/HexToBytes.h"
 #include "Templates/TypeCompatibleBytes.h"
 
 class FArchive;
-class FCompositeBuffer;
-template <typename CharType> class TStringBuilderBase;
 
 ///////////////////////////////////////////////////////////////////////////////////////////////////////////////////////
 
@@ -70,11 +64,7 @@ public:
 		return FMemory::Memcmp(GetBytes(), B.GetBytes(), sizeof(decltype(GetBytes()))) < 0;
 	}
 
-	friend inline FArchive& operator<<(FArchive& Ar, FBlake3Hash& Value)
-	{
-		Ar.Serialize(Value.GetBytes(), sizeof(decltype(Value.GetBytes())));
-		return Ar;
-	}
+	friend CORE_API FArchive& operator<<(FArchive& Ar, FBlake3Hash& Value);
 
 	friend inline uint32 GetTypeHash(const FBlake3Hash& Value)
 	{
@@ -185,3 +175,8 @@ inline void LexFromString(FBlake3Hash& OutHash, const TCHAR* Buffer)
 
 /** Convert a hash to a 64-character hex string. */
 [[nodiscard]] CORE_API FString LexToString(const FBlake3Hash& Hash);
+
+#if UE_ENABLE_INCLUDE_ORDER_DEPRECATED_IN_5_8
+#include "Containers/UnrealString.h"
+#include "Serialization/Archive.h"
+#endif
diff --git a/Engine/Source/Runtime/Core/Public/IO/IoHash.h b/Engine/Source/Runtime/Core/Public/IO/IoHash.h
--- a/Engine/Source/Runtime/Core/Public/IO/IoHash.h
+++ b/Engine/Source/Runtime/Core/Public/IO/IoHash.h
@@ -4,20 +4,15 @@
 
 #include "Containers/StringFwd.h"
 #include "Containers/StringView.h"
-#include "Containers/UnrealString.h"
-#include "HAL/Platform.h"
-#include "HAL/PlatformString.h"
 #include "HAL/UnrealMemory.h"
 #include "Hash/Blake3.h"
 #include "Memory/MemoryFwd.h"
 #include "Memory/MemoryView.h"
 #include "Misc/AssertionMacros.h"
-#include "Serialization/Archive.h"
 #include "String/BytesToHex.h"
 #include "String/HexToBytes.h"
 
-class FCompositeBuffer;
-template <typename CharType> class TStringBuilderBase;
+class FArchive;
 
 /**
  * Stores a BLAKE3-160 hash, taken from the first 20 bytes of a BLAKE3-256 hash.
@@ -62,9 +57,9 @@ public:
 	inline const ByteArray& GetBytes() const { return Hash; }
 
 	/** Calculate the hash of the buffer. */
-	[[nodiscard]] static inline FIoHash HashBuffer(FMemoryView View);
-	[[nodiscard]] static inline FIoHash HashBuffer(const void* Data, uint64 Size);
-	[[nodiscard]] static inline FIoHash HashBuffer(const FCompositeBuffer& Buffer);
+	[[nodiscard]] inline static FIoHash HashBuffer(FMemoryView View);
+	[[nodiscard]] inline static FIoHash HashBuffer(const void* Data, uint64 Size);
+	[[nodiscard]] CORE_API static FIoHash HashBuffer(const FCompositeBuffer& Buffer);
 
 	/** A zero hash. */
 	static const FIoHash Zero;
@@ -92,11 +87,7 @@ private:
 		return *reinterpret_cast<const uint32*>(Value.GetBytes());
 	}
 
-	friend inline FArchive& operator<<(FArchive& Ar, FIoHash& InHash)
-	{
-		Ar.Serialize(InHash.GetBytes(), sizeof(decltype(InHash.GetBytes())));
-		return Ar;
-	}
+	friend CORE_API FArchive& operator<<(FArchive& Ar, FIoHash& InHash);
 };
 
 inline const FIoHash FIoHash::Zero;
@@ -163,11 +154,6 @@ inline FIoHash FIoHash::HashBuffer(const void* Data, uint64 Size)
 	return FBlake3::HashBuffer(Data, Size);
 }
 
-inline FIoHash FIoHash::HashBuffer(const FCompositeBuffer& Buffer)
-{
-	return FBlake3::HashBuffer(Buffer);
-}
-
 template <typename CharType>
 inline TStringBuilderBase<CharType>& operator<<(TStringBuilderBase<CharType>& Builder, const FIoHash& Hash)
 {
@@ -200,5 +186,10 @@ public:
 	/** Calculate the hash of the buffer. */
 	[[nodiscard]] inline static FIoHash HashBuffer(FMemoryView View) { return FBlake3::HashBuffer(View); }
 	[[nodiscard]] inline static FIoHash HashBuffer(const void* Data, uint64 Size) { return FBlake3::HashBuffer(Data, Size); }
-	[[nodiscard]] inline static FIoHash HashBuffer(const FCompositeBuffer& Buffer) { return FBlake3::HashBuffer(Buffer); }
+	[[nodiscard]] CORE_API static FIoHash HashBuffer(const FCompositeBuffer& Buffer);
 };
+
+#if UE_ENABLE_INCLUDE_ORDER_DEPRECATED_IN_5_8
+#include "Containers/UnrealString.h"
+#include "Serialization/Archive.h"
+#endif
