CsvProfiler - Remove CsvWaitStat_EventWait/* named event scopes

#rnx
#Synced-CL Unknown

diff --git a/Engine/Source/Runtime/Core/Private/ProfilingDebugging/CsvProfiler.cpp b/Engine/Source/Runtime/Core/Private/ProfilingDebugging/CsvProfiler.cpp
--- a/Engine/Source/Runtime/Core/Private/ProfilingDebugging/CsvProfiler.cpp
+++ b/Engine/Source/Runtime/Core/Private/ProfilingDebugging/CsvProfiler.cpp
@@ -4423,13 +4423,6 @@ void FCsvProfiler::BeginSetWaitStat(const char* StatName, const char* FormattedS
 #if !CSV_PROFILER_MINIMAL
 	if (GCsvProfilerIsCapturing && GCsvCategoriesEnabled[CSV_CATEGORY_INDEX(Exclusive)])
 	{
-#if CSV_PROFILER_SUPPORT_NAMED_EVENTS
-		if (UNLIKELY(GCsvProfilerNamedEventsExclusive))
-		{
-			FPlatformMisc::BeginNamedEvent(FColor::Yellow, *FString::Printf(TEXT("CsvWaitStat_%s"), StringCast<TCHAR>(StatName).Get()));
-		}
-#endif
-		
 		if (StatName == nullptr)
 		{
 			FCsvProfilerThreadData::Get().PushWaitStatName(GIgnoreWaitStatName);
@@ -4447,12 +4440,6 @@ void FCsvProfiler::EndSetWaitStat()
 #if !CSV_PROFILER_MINIMAL
 	if (GCsvProfilerIsCapturing && GCsvCategoriesEnabled[CSV_CATEGORY_INDEX(Exclusive)])
 	{
-#if CSV_PROFILER_SUPPORT_NAMED_EVENTS
-		if (UNLIKELY(GCsvProfilerNamedEventsExclusive))
-		{
-			FPlatformMisc::EndNamedEvent();
-		}
-#endif
 		FCsvProfilerThreadData::Get().PopWaitStatName();
 	}
 #endif
