Core: Moved FName::AutoTest() into its own test function

#Synced-CL 0

diff --git a/Engine/Source/Runtime/Core/Private/Tests/HAL/PlatformTest.cpp b/Engine/Source/Runtime/Core/Private/Tests/HAL/PlatformTest.cpp
--- a/Engine/Source/Runtime/Core/Private/Tests/HAL/PlatformTest.cpp
+++ b/Engine/Source/Runtime/Core/Private/Tests/HAL/PlatformTest.cpp
@@ -81,8 +81,6 @@ bool FPlatformVerificationTest::RunTest (const FString& Parameters)
 
 	static_assert(alignof(FTestAlign) == 16, "Align of FTestAlign is not 16.");
 
-	FName::AutoTest();
-
 	return true;
 }
 
diff --git a/Engine/Source/Runtime/Core/Private/UObject/UnrealNames.cpp b/Engine/Source/Runtime/Core/Private/UObject/UnrealNames.cpp
--- a/Engine/Source/Runtime/Core/Private/UObject/UnrealNames.cpp
+++ b/Engine/Source/Runtime/Core/Private/UObject/UnrealNames.cpp
@@ -43,6 +43,10 @@
 #include "Misc/AsciiSet.h"
 #include "AutoRTFM.h"
 
+#if WITH_TESTS
+#include "Tests/TestHarnessAdapter.h"
+#endif
+
 #ifndef USE_FNAME_MMAP
 #define USE_FNAME_MMAP 0
 #endif
@@ -3995,17 +3999,19 @@ FUtf8StringBuilderBase& operator<<(FUtf8StringBuilderBase& Builder, FNameEntryId
 	return Builder;
 }
 
+#if WITH_TESTS
+
 template <typename CharType, int N>
 void CheckLazyName(const CharType(&Literal)[N])
 {
-	check(FName(Literal) == FLazyName(Literal));
-	check(FLazyName(Literal) == FName(Literal));
-	check(FLazyName(Literal) == FLazyName(Literal));
-	check(FName(Literal) == FLazyName(Literal).Resolve());
+	CHECK(FName(Literal) == FLazyName(Literal));
+	CHECK(FLazyName(Literal) == FName(Literal));
+	CHECK(FLazyName(Literal) == FLazyName(Literal));
+	CHECK(FName(Literal) == FLazyName(Literal).Resolve());
 
 	CharType Literal2[N];
 	FMemory::Memcpy(Literal2, Literal);
-	check(FLazyName(Literal) == FLazyName(Literal2));
+	CHECK(FLazyName(Literal) == FLazyName(Literal2));
 
 	constexpr ESearchCase::Type ComparisonMode =
 #if WITH_CASE_PRESERVING_NAME
@@ -4014,12 +4020,11 @@ void CheckLazyName(const CharType(&Literal)[N])
 		ESearchCase::IgnoreCase;
 #endif
 
-	check(WriteToString<64>(FLazyName(Literal).Resolve()).ToView().Equals(WriteToString<64>(Literal).ToView(), ComparisonMode));
+	CHECK(WriteToString<64>(FLazyName(Literal).Resolve()).ToView().Equals(WriteToString<64>(Literal).ToView(), ComparisonMode));
 }
 
 static void TestNameBatch();
 
-#if DO_CHECK
 struct FNameAutoTestTruncateHelper
 {
 	static constexpr const TCHAR TestStr[] = TEXT("ABC_2147483646");
@@ -4040,31 +4045,31 @@ struct FNameAutoTestTruncateHelper
 		if constexpr (RequiredLen >= TestStrLen)
 		{
 			// Whole string should be returned
-			check(Len == TestStrLen);
+			CHECK(Len == TestStrLen);
 		}
 		else
 		{
 			// String should be truncated
-			check(Len == RequiredLen);
+			CHECK(Len == RequiredLen);
 		}
 
 		for (uint32 Index = 0; Index <= RequiredLen; ++Index)
 		{
 			if (Index < Len)
 			{
-				check(Buffer.Data[Index] == TestStr[Index]);
+				CHECK(Buffer.Data[Index] == TestStr[Index]);
 			}
 			else if (Index == Len)
 			{
-				check(Buffer.Data[Index] == 0); // null terminator
+				CHECK(Buffer.Data[Index] == 0); // null terminator
 			}
 			else
 			{
-				check(Buffer.Data[Index] == CanaryValue); // should not have been written
+				CHECK(Buffer.Data[Index] == CanaryValue); // should not have been written
 			}
 		}
 
-		check(Buffer.Canary == CanaryValue);
+		CHECK(Buffer.Canary == CanaryValue);
 	}
 
 	template <std::size_t... Ints>
@@ -4078,19 +4083,17 @@ struct FNameAutoTestTruncateHelper
 		TestOuter(std::make_index_sequence<TestStrLen + 4>());
 	}
 };
-#endif
 
-void FName::AutoTest()
+TEST_CASE_NAMED(FNameTest, "System::Core::Name", "[ApplicationContextMask][SmokeFilter]")
 {
-#if DO_CHECK
-	check(FNameHash::IsAnsiNone("None", 4) == 1);
-	check(FNameHash::IsAnsiNone("none", 4) == 1);
-	check(FNameHash::IsAnsiNone("NONE", 4) == 1);
-	check(FNameHash::IsAnsiNone("nOnE", 4) == 1);
-	check(FNameHash::IsAnsiNone("None", 5) == 0);
-	check(FNameHash::IsAnsiNone(TEXT("None"), 4) == 0);
-	check(FNameHash::IsAnsiNone("nono", 4) == 0);
-	check(FNameHash::IsAnsiNone("enon", 4) == 0);
+	CHECK(FNameHash::IsAnsiNone("None", 4) == 1);
+	CHECK(FNameHash::IsAnsiNone("none", 4) == 1);
+	CHECK(FNameHash::IsAnsiNone("NONE", 4) == 1);
+	CHECK(FNameHash::IsAnsiNone("nOnE", 4) == 1);
+	CHECK(FNameHash::IsAnsiNone("None", 5) == 0);
+	CHECK(FNameHash::IsAnsiNone(TEXT("None"), 4) == 0);
+	CHECK(FNameHash::IsAnsiNone("nono", 4) == 0);
+	CHECK(FNameHash::IsAnsiNone("enon", 4) == 0);
 
 	const FName AutoTest_1("AutoTest_1");
 	const FName autoTest_1("autoTest_1");
@@ -4098,135 +4101,141 @@ void FName::AutoTest()
 	const FName AutoTest_2(TEXT("AutoTest_2"));
 	const FName AutoTestB_2(TEXT("AutoTestB_2"));
 
-	check(AutoTest_1.GetNumber() == NAME_EXTERNAL_TO_INTERNAL(1));
-	check(autoTest_1.GetNumber() == NAME_EXTERNAL_TO_INTERNAL(1));
-	check(autoTeSt_1.GetNumber() == NAME_EXTERNAL_TO_INTERNAL(1));
-	check(AutoTest_2.GetNumber() == NAME_EXTERNAL_TO_INTERNAL(2));
-	check(AutoTestB_2.GetNumber() == NAME_EXTERNAL_TO_INTERNAL(2));
+	CHECK(AutoTest_1.GetNumber() == NAME_EXTERNAL_TO_INTERNAL(1));
+	CHECK(autoTest_1.GetNumber() == NAME_EXTERNAL_TO_INTERNAL(1));
+	CHECK(autoTeSt_1.GetNumber() == NAME_EXTERNAL_TO_INTERNAL(1));
+	CHECK(AutoTest_2.GetNumber() == NAME_EXTERNAL_TO_INTERNAL(2));
+	CHECK(AutoTestB_2.GetNumber() == NAME_EXTERNAL_TO_INTERNAL(2));
 
-	check(AutoTest_1 != AutoTest_2);
-	check(AutoTest_1 == autoTest_1);
-	check(AutoTest_1 == autoTeSt_1);
+	CHECK(AutoTest_1 != AutoTest_2);
+	CHECK(AutoTest_1 == autoTest_1);
+	CHECK(AutoTest_1 == autoTeSt_1);
 
-	check(AutoTest_1.ToUnstableInt() == autoTest_1.ToUnstableInt());
-	check(AutoTest_1.ToUnstableInt() != AutoTest_2.ToUnstableInt());
+	CHECK(AutoTest_1.ToUnstableInt() == autoTest_1.ToUnstableInt());
+	CHECK(AutoTest_1.ToUnstableInt() != AutoTest_2.ToUnstableInt());
 
 	TCHAR Buffer[FName::StringBufferSize];
 
 #if WITH_CASE_PRESERVING_NAME
-	check(!FCString::Strcmp(*AutoTest_1.ToString(), TEXT("AutoTest_1")));
-	check(!FCString::Strcmp(*autoTest_1.ToString(), TEXT("autoTest_1")));
-	check(!FCString::Strcmp(*autoTeSt_1.ToString(), TEXT("autoTeSt_1")));
-	check(!FCString::Strcmp(*AutoTestB_2.ToString(), TEXT("AutoTestB_2")));
+	CHECK(!FCString::Strcmp(*AutoTest_1.ToString(), TEXT("AutoTest_1")));
+	CHECK(!FCString::Strcmp(*autoTest_1.ToString(), TEXT("autoTest_1")));
+	CHECK(!FCString::Strcmp(*autoTeSt_1.ToString(), TEXT("autoTeSt_1")));
+	CHECK(!FCString::Strcmp(*AutoTestB_2.ToString(), TEXT("AutoTestB_2")));
 	
-	check(FName("ABC").ToString(Buffer) == 3 &&			!FCString::Strcmp(Buffer, TEXT("ABC")));
-	check(FName("abc").ToString(Buffer) == 3 &&			!FCString::Strcmp(Buffer, TEXT("abc")));
-	check(FName(TEXT("abc")).ToString(Buffer) == 3 &&	!FCString::Strcmp(Buffer, TEXT("abc")));
-	check(FName("ABC_0").ToString(Buffer) == 5 &&		!FCString::Strcmp(Buffer, TEXT("ABC_0")));
-	check(FName("ABC_10").ToString(Buffer) == 6 &&		!FCString::Strcmp(Buffer, TEXT("ABC_10")));	
+	CHECK(FName("ABC").ToString(Buffer) == 3);
+	CHECK(!FCString::Strcmp(Buffer, TEXT("ABC")));
+	CHECK(FName("abc").ToString(Buffer) == 3);
+	CHECK(!FCString::Strcmp(Buffer, TEXT("abc")));
+	CHECK(FName(TEXT("abc")).ToString(Buffer) == 3);
+	CHECK(!FCString::Strcmp(Buffer, TEXT("abc")));
+	CHECK(FName("ABC_0").ToString(Buffer) == 5);
+	CHECK(!FCString::Strcmp(Buffer, TEXT("ABC_0")));
+	CHECK(FName("ABC_10").ToString(Buffer) == 6);
+	CHECK(!FCString::Strcmp(Buffer, TEXT("ABC_10")));	
 #endif
 
-	check(autoTest_1.GetComparisonIndex() == AutoTest_2.GetComparisonIndex());
-	check(autoTest_1.GetPlainNameString() == AutoTest_1.GetPlainNameString());
-	check(autoTest_1.GetPlainNameString() == AutoTest_2.GetPlainNameString());
-	check(*AutoTestB_2.GetPlainNameString() != *AutoTest_2.GetPlainNameString());
-	check(AutoTestB_2.GetNumber() == AutoTest_2.GetNumber());
-	check(autoTest_1.GetNumber() != AutoTest_2.GetNumber());
-	check(AutoTest_1.GetNumber() == autoTest_1.GetNumber());
-
-	check(FCStringAnsi::Strlen("None") == FName().GetStringLength());
-	check(FCStringAnsi::Strlen("ABC") == FName("ABC").GetStringLength());
-	check(FCStringAnsi::Strlen("ABC_0") == FName("ABC_0").GetStringLength());
-	check(FCStringAnsi::Strlen("ABC_9") == FName("ABC_9").GetStringLength());
-	check(FCStringAnsi::Strlen("ABC_10") == FName("ABC_10").GetStringLength());
-	check(FCStringAnsi::Strlen("ABC_2000000000") == FName("ABC_2000000000").GetStringLength());
-	check(FCStringAnsi::Strlen("ABC_4000000000") == FName("ABC_4000000000").GetStringLength());
+	CHECK(autoTest_1.GetComparisonIndex() == AutoTest_2.GetComparisonIndex());
+	CHECK(autoTest_1.GetPlainNameString() == AutoTest_1.GetPlainNameString());
+	CHECK(autoTest_1.GetPlainNameString() == AutoTest_2.GetPlainNameString());
+	CHECK(*AutoTestB_2.GetPlainNameString() != *AutoTest_2.GetPlainNameString());
+	CHECK(AutoTestB_2.GetNumber() == AutoTest_2.GetNumber());
+	CHECK(autoTest_1.GetNumber() != AutoTest_2.GetNumber());
+	CHECK(AutoTest_1.GetNumber() == autoTest_1.GetNumber());
+
+	CHECK(FCStringAnsi::Strlen("None") == FName().GetStringLength());
+	CHECK(FCStringAnsi::Strlen("ABC") == FName("ABC").GetStringLength());
+	CHECK(FCStringAnsi::Strlen("ABC_0") == FName("ABC_0").GetStringLength());
+	CHECK(FCStringAnsi::Strlen("ABC_9") == FName("ABC_9").GetStringLength());
+	CHECK(FCStringAnsi::Strlen("ABC_10") == FName("ABC_10").GetStringLength());
+	CHECK(FCStringAnsi::Strlen("ABC_2000000000") == FName("ABC_2000000000").GetStringLength());
+	CHECK(FCStringAnsi::Strlen("ABC_4000000000") == FName("ABC_4000000000").GetStringLength());
 
 	const FName NullName(static_cast<ANSICHAR*>(nullptr));
-	check(NullName.IsNone());
-	check(NullName == FName(static_cast<WIDECHAR*>(nullptr)));
-	check(NullName == FName(NAME_None));
-	check(NullName == FName());
-	check(NullName == FName(""));
-	check(NullName == FName(TEXT("")));
-	check(NullName == FName("None"));
-	check(NullName == FName("none"));
-	check(NullName == FName("NONE"));
-	check(NullName == FName(TEXT("None")));
-	check(FName().ToEName());
-	check(*FName().ToEName() == NAME_None);
-	check(NullName.GetComparisonIndex().ToUnstableInt() == 0);
+	CHECK(NullName.IsNone());
+	CHECK(NullName == FName(static_cast<WIDECHAR*>(nullptr)));
+	CHECK(NullName == FName(NAME_None));
+	CHECK(NullName == FName());
+	CHECK(NullName == FName(""));
+	CHECK(NullName == FName(TEXT("")));
+	CHECK(NullName == FName("None"));
+	CHECK(NullName == FName("none"));
+	CHECK(NullName == FName("NONE"));
+	CHECK(NullName == FName(TEXT("None")));
+	CHECK(FName().ToEName());
+	CHECK(*FName().ToEName() == NAME_None);
+	CHECK(NullName.GetComparisonIndex().ToUnstableInt() == 0);
 
 	// Numbered nones
 	FName None_7 = FName("None_7");
-	check(None_7.ToString() == "None_7");
-	check(None_7.GetComparisonIndex() == FNameEntryId::FromEName(EName::None));
-	check(None_7.GetNumber() == NAME_EXTERNAL_TO_INTERNAL(7));
-	check(FName(EName::None, NAME_EXTERNAL_TO_INTERNAL(7)) == None_7);
-	check(FName(FName(), NAME_EXTERNAL_TO_INTERNAL(7)) == None_7);
-	check(FName("None", NAME_EXTERNAL_TO_INTERNAL(7)) == None_7);
-	check(FName(FNameEntryId::FromEName(EName::None), FNameEntryId::FromEName(EName::None), NAME_EXTERNAL_TO_INTERNAL(7)) == None_7);
-	check(!None_7.IsNone());
+	CHECK(None_7.ToString() == "None_7");
+	CHECK(None_7.GetComparisonIndex() == FNameEntryId::FromEName(EName::None));
+	CHECK(None_7.GetNumber() == NAME_EXTERNAL_TO_INTERNAL(7));
+	CHECK(FName(EName::None, NAME_EXTERNAL_TO_INTERNAL(7)) == None_7);
+	CHECK(FName(FName(), NAME_EXTERNAL_TO_INTERNAL(7)) == None_7);
+	CHECK(FName("None", NAME_EXTERNAL_TO_INTERNAL(7)) == None_7);
+	CHECK(FName(FNameEntryId::FromEName(EName::None), FNameEntryId::FromEName(EName::None), NAME_EXTERNAL_TO_INTERNAL(7)) == None_7);
+	CHECK(!None_7.IsNone());
 
 	// "_[num]" is considered None to keep semantics of old implementation, unlike "None_[num]"
-	check(FName("_7").IsNone()); 
-	check(FName("_0").IsNone());
-	check(FName("_2147483646").IsNone());
-	check(!FName("_2147483647").IsNone()); // invalid number
+	CHECK(FName("_7").IsNone()); 
+	CHECK(FName("_0").IsNone());
+	CHECK(FName("_2147483646").IsNone());
+	CHECK(!FName("_2147483647").IsNone()); // invalid number
 
 	const FName Cylinder(NAME_Cylinder);
-	check(Cylinder == FName("Cylinder"));
-	check(Cylinder.ToEName());
-	check(*Cylinder.ToEName() == NAME_Cylinder);
-	check(Cylinder.GetPlainNameString() == TEXT("Cylinder"));
+	CHECK(Cylinder == FName("Cylinder"));
+	CHECK(Cylinder.ToEName());
+	CHECK(*Cylinder.ToEName() == NAME_Cylinder);
+	CHECK(Cylinder.GetPlainNameString() == TEXT("Cylinder"));
 
 	// Test numbers
-	check(FName("Text_0") == FName("Text", NAME_EXTERNAL_TO_INTERNAL(0)));
-	check(FName("Text_1") == FName("Text", NAME_EXTERNAL_TO_INTERNAL(1)));
-	check(FName("Text_1_0") == FName("Text_1", NAME_EXTERNAL_TO_INTERNAL(0)));
-	check(FName("Text_0_1") == FName("Text_0", NAME_EXTERNAL_TO_INTERNAL(1)));
-	check(FName("Text_00") == FName("Text_00", NAME_NO_NUMBER_INTERNAL));
-	check(FName("Text_01") == FName("Text_01", NAME_NO_NUMBER_INTERNAL));
+	CHECK(FName("Text_0") == FName("Text", NAME_EXTERNAL_TO_INTERNAL(0)));
+	CHECK(FName("Text_1") == FName("Text", NAME_EXTERNAL_TO_INTERNAL(1)));
+	CHECK(FName("Text_1_0") == FName("Text_1", NAME_EXTERNAL_TO_INTERNAL(0)));
+	CHECK(FName("Text_0_1") == FName("Text_0", NAME_EXTERNAL_TO_INTERNAL(1)));
+	CHECK(FName("Text_00") == FName("Text_00", NAME_NO_NUMBER_INTERNAL));
+	CHECK(FName("Text_01") == FName("Text_01", NAME_NO_NUMBER_INTERNAL));
 
 	// Test unterminated strings
-	check(FName("") == FName(0, "Unused"));
-	check(FName("Used") == FName(4, "UsedUnused"));
-	check(FName("Used") == FName(4, "Used"));
-	check(FName("Used_0") == FName(6, "Used_01"));
-	check(FName("Used_01") == FName(7, "Used_012"));
-	check(FName("Used_123") == FName(8, "Used_123456"));
-	check(FName("Used_123") == FName(8, "Used_123_456"));
-	check(FName("Used_123") == FName(8, TEXT("Used_123456")));
-	check(FName("Used_123") == FName(8, TEXT("Used_123_456")));
-	check(FName("Used_2147483646") == FName(15, TEXT("Used_2147483646123")));
-	check(FName("Used_2147483647") == FName(15, TEXT("Used_2147483647123")));
-	check(FName("Used_2147483648") == FName(15, TEXT("Used_2147483648123")));
+	CHECK(FName("") == FName(0, "Unused"));
+	CHECK(FName("Used") == FName(4, "UsedUnused"));
+	CHECK(FName("Used") == FName(4, "Used"));
+	CHECK(FName("Used_0") == FName(6, "Used_01"));
+	CHECK(FName("Used_01") == FName(7, "Used_012"));
+	CHECK(FName("Used_123") == FName(8, "Used_123456"));
+	CHECK(FName("Used_123") == FName(8, "Used_123_456"));
+	CHECK(FName("Used_123") == FName(8, TEXT("Used_123456")));
+	CHECK(FName("Used_123") == FName(8, TEXT("Used_123_456")));
+	CHECK(FName("Used_2147483646") == FName(15, TEXT("Used_2147483646123")));
+	CHECK(FName("Used_2147483647") == FName(15, TEXT("Used_2147483647123")));
+	CHECK(FName("Used_2147483648") == FName(15, TEXT("Used_2147483648123")));
 
 	// Test wide strings
 	FString Wide("Wide ");
 	Wide[4] = 60000;
 	FName WideName(*Wide);
-	check(WideName.GetPlainNameString() == Wide);
-	check(FName(*Wide).GetPlainNameString() == Wide);
-	check(FName(*Wide).ToString(Buffer) == 5 && !FCString::Strcmp(Buffer, *Wide));
-	check(Wide.Len() == WideName.GetStringLength());
+	CHECK(WideName.GetPlainNameString() == Wide);
+	CHECK(FName(*Wide).GetPlainNameString() == Wide);
+	CHECK(FName(*Wide).ToString(Buffer) == 5);
+	CHECK(!FCString::Strcmp(Buffer, *Wide));
+	CHECK(Wide.Len() == WideName.GetStringLength());
 	FString WideLong = FString::ChrN(1000, 60000);
-	check(FName(*WideLong).GetPlainNameString() == WideLong);
+	CHECK(FName(*WideLong).GetPlainNameString() == WideLong);
 
 
 	// Check that FNAME_Find doesn't add entries
 	static bool Once = true;
 	if (Once)
 	{
-		check(FName("UniqueUnicorn!!", FNAME_Find) == FName());
+		CHECK(FName("UniqueUnicorn!!", FNAME_Find) == FName());
 
 		// Check that FNAME_Find can find entries
 		const FName UniqueName("UniqueUnicorn!!", FNAME_Add);
-		check(FName("UniqueUnicorn!!", FNAME_Find) == UniqueName);
-		check(FName(TEXT("UniqueUnicorn!!"), FNAME_Find) == UniqueName);
-		check(FName("UNIQUEUNICORN!!", FNAME_Find) == UniqueName);
-		check(FName(TEXT("UNIQUEUNICORN!!"), FNAME_Find) == UniqueName);
-		check(FName("uniqueunicorn!!", FNAME_Find) == UniqueName);
+		CHECK(FName("UniqueUnicorn!!", FNAME_Find) == UniqueName);
+		CHECK(FName(TEXT("UniqueUnicorn!!"), FNAME_Find) == UniqueName);
+		CHECK(FName("UNIQUEUNICORN!!", FNAME_Find) == UniqueName);
+		CHECK(FName(TEXT("UNIQUEUNICORN!!"), FNAME_Find) == UniqueName);
+		CHECK(FName("uniqueunicorn!!", FNAME_Find) == UniqueName);
 
 		Once = false;
 	}
@@ -4242,43 +4251,43 @@ void FName::AutoTest()
 	Names.Add(FooWide);
 	Algo::Sort(Names, FNameLexicalLess()); 
 
-	check(Names[0] == "FooA");
-	check(Names[1] == "FooAB");
-	check(Names[2] == "FooABC");
-	check(Names[3] == "FooABCD");
-	check(Names[4] == "FooB");
-	check(Names[5] == "FooC");
-	check(Names[6] == FooWide);
+	CHECK(Names[0] == "FooA");
+	CHECK(Names[1] == "FooAB");
+	CHECK(Names[2] == "FooABC");
+	CHECK(Names[3] == "FooABCD");
+	CHECK(Names[4] == "FooB");
+	CHECK(Names[5] == "FooC");
+	CHECK(Names[6] == FooWide);
 		
-	check(FName() == TEXTVIEW("None"));
-	check(FName() == TEXTVIEW("none"));
-	check(FName() == TEXTVIEW(""));
-	check(FName() == FStringView());
-	check(FName("a") == TEXTVIEW("a"));
-	check(FName("a") == TEXTVIEW("A"));
-	check(FName("A") == TEXTVIEW("a"));
-	check(FName("aa") != TEXTVIEW("a"));
-	check(FName("a") != TEXTVIEW("aa"));
-	check(FName("a") == TEXTVIEW("aa").Left(1)); // Unterminated
-	check(FName("aa") != TEXTVIEW("ab"));
-	check(FName("a_0") == TEXTVIEW("a_0"));
-	check(FName("a_1") != TEXTVIEW("a_0"));
-	check(FName("a_10") == TEXTVIEW("a_10"));
-	check(FName("a")    == TEXTVIEW("a_10").Left(1)); // Unterminated number
-	check(FName("a_")   == TEXTVIEW("a_10").Left(2)); // Unterminated number
-	check(FName("a_1")  == TEXTVIEW("a_10").Left(3)); // Unterminated number
-	check(FName("a_2147483646") == TEXTVIEW("a_2147483646")); // Largest valid number
-	check(FName("a_2147483647") == TEXTVIEW("a_2147483647")); // Smallest invalid number
-	check(FName("a_") == TEXTVIEW("a_"));
-	check(FName("_") == TEXTVIEW("_"));
-	check(FName("_3") == TEXTVIEW("_3")); // Special none
-	check(FName() == TEXTVIEW("_123")); // Special none
-	check(FName("_3") == TEXTVIEW("None")); // Special none
-	check(FName("none_1") == TEXTVIEW("none_1")); // Numbered none
-	check(WideName == FStringView(Wide));
-	check(FName("a") == ANSITEXTVIEW("a"));
-	check(FName("a") == ANSITEXTVIEW("A"));
-	check(FName("A") != ANSITEXTVIEW("b"));
+	CHECK(FName() == TEXTVIEW("None"));
+	CHECK(FName() == TEXTVIEW("none"));
+	CHECK(FName() == TEXTVIEW(""));
+	CHECK(FName() == FStringView());
+	CHECK(FName("a") == TEXTVIEW("a"));
+	CHECK(FName("a") == TEXTVIEW("A"));
+	CHECK(FName("A") == TEXTVIEW("a"));
+	CHECK(FName("aa") != TEXTVIEW("a"));
+	CHECK(FName("a") != TEXTVIEW("aa"));
+	CHECK(FName("a") == TEXTVIEW("aa").Left(1)); // Unterminated
+	CHECK(FName("aa") != TEXTVIEW("ab"));
+	CHECK(FName("a_0") == TEXTVIEW("a_0"));
+	CHECK(FName("a_1") != TEXTVIEW("a_0"));
+	CHECK(FName("a_10") == TEXTVIEW("a_10"));
+	CHECK(FName("a")    == TEXTVIEW("a_10").Left(1)); // Unterminated number
+	CHECK(FName("a_")   == TEXTVIEW("a_10").Left(2)); // Unterminated number
+	CHECK(FName("a_1")  == TEXTVIEW("a_10").Left(3)); // Unterminated number
+	CHECK(FName("a_2147483646") == TEXTVIEW("a_2147483646")); // Largest valid number
+	CHECK(FName("a_2147483647") == TEXTVIEW("a_2147483647")); // Smallest invalid number
+	CHECK(FName("a_") == TEXTVIEW("a_"));
+	CHECK(FName("_") == TEXTVIEW("_"));
+	CHECK(FName("_3") == TEXTVIEW("_3")); // Special none
+	CHECK(FName() == TEXTVIEW("_123")); // Special none
+	CHECK(FName("_3") == TEXTVIEW("None")); // Special none
+	CHECK(FName("none_1") == TEXTVIEW("none_1")); // Numbered none
+	CHECK(WideName == FStringView(Wide));
+	CHECK(FName("a") == ANSITEXTVIEW("a"));
+	CHECK(FName("a") == ANSITEXTVIEW("A"));
+	CHECK(FName("A") != ANSITEXTVIEW("b"));
 
 	// Check truncating ToString functions
 	FNameAutoTestTruncateHelper::Run();
@@ -4302,12 +4311,12 @@ void FName::AutoTest()
 	CheckLazyName("None_0");
 	CheckLazyName("None_1");
 
-	check(FLazyName("Hej") == FLazyName(UTF8TEXT("Hej")));
-	check(FLazyName("Hej") == FLazyName(WIDETEXT("Hej")));
-	check(FLazyName(UTF8TEXT("Hej")) == FLazyName("Hej"));
-	check(FLazyName(UTF8TEXT("Hej")) == FLazyName(WIDETEXT("Hej")));
-	check(FLazyName(WIDETEXT("Hej")) == FLazyName("Hej"));
-	check(FLazyName(WIDETEXT("Hej")) == FLazyName(UTF8TEXT("Hej")));
+	CHECK(FLazyName("Hej") == FLazyName(UTF8TEXT("Hej")));
+	CHECK(FLazyName("Hej") == FLazyName(WIDETEXT("Hej")));
+	CHECK(FLazyName(UTF8TEXT("Hej")) == FLazyName("Hej"));
+	CHECK(FLazyName(UTF8TEXT("Hej")) == FLazyName(WIDETEXT("Hej")));
+	CHECK(FLazyName(WIDETEXT("Hej")) == FLazyName("Hej"));
+	CHECK(FLazyName(WIDETEXT("Hej")) == FLazyName(UTF8TEXT("Hej")));
 
 	TestNameBatch();
 
@@ -4330,7 +4339,7 @@ void FName::AutoTest()
 		{
 			Id = FName(*FString::Printf(TEXT("%d"), I)).GetComparisonIndex();
 		}
-		check(Id == OldId);
+		CHECK(Id == OldId);
 	}
 #endif
 
@@ -4340,30 +4349,31 @@ void FName::AutoTest()
 		for (uint32 Index : SetBits)
 		{
 			const bool bWasSet = TestSet.SetBitAtomic(Index);
-			check(bWasSet == false);
+			CHECK(bWasSet == false);
 		}
 
 		const bool b32WasSet = TestSet.SetBitAtomic(34);
-		check(b32WasSet);
+		CHECK(b32WasSet);
 
 		uint32 Expected = 0;
 		TestSet.ForEachSetBit([&](uint32 Index)
 		{
-			check(Index == SetBits[Expected]);
+			CHECK(Index == SetBits[Expected]);
 			++Expected;
 		});
-		check(Expected == UE_ARRAY_COUNT(SetBits));
+		CHECK(Expected == UE_ARRAY_COUNT(SetBits));
 
 		const bool b1023WasSet = TestSet.ClearBitAtomic(1023);
-		check(b1023WasSet == true);
+		CHECK(b1023WasSet == true);
 		TestSet.ForEachSetBit([&](uint32 Index)
 		{
-			check(Index != 1023);
+			CHECK(Index != 1023);
 		});
 	}
-#endif // DO_CHECK
 }
 
+#endif // WITH_TESTS
+
 
 /*-----------------------------------------------------------------------------
 	FNameEntry implementation.
@@ -5491,6 +5501,7 @@ int LoadNameBatchTestFiles(bool bArchivePath)
 
 #endif
 
+#if WITH_TESTS
 static void TestNameBatch()
 {
 #if ALLOW_NAME_BATCH_SAVING
@@ -5612,6 +5623,7 @@ static void TestNameBatch()
 
 #endif // ALLOW_NAME_BATCH_SAVING
 }
+#endif // WITH_TESTS
 
 #if !UE_BUILD_SHIPPING && !UE_BUILD_TEST
 
diff --git a/Engine/Source/Runtime/Core/Public/UObject/NameTypes.h b/Engine/Source/Runtime/Core/Public/UObject/NameTypes.h
--- a/Engine/Source/Runtime/Core/Public/UObject/NameTypes.h
+++ b/Engine/Source/Runtime/Core/Public/UObject/NameTypes.h
@@ -1179,8 +1179,10 @@ public:
 
 	//@}
 
-	/** Run autotest on FNames. */
-	CORE_API static void AutoTest();
+	UE_DEPRECATED(5.8, "Tests for FName are registered with both test frameworks now.")
+	inline static void AutoTest()
+	{
+	}
 	
 	/**
 	 * Takes a string and breaks it down into a human readable string.
