CSV profiler: Workaround for unix time issues, since Cycles64 and Seconds methods can use different clocks. This was causing very large or negative first frame time. Capturing from boot is still broken but it's rare that we need to do this.

#tests tested on console

diff --git a/Engine/Source/Runtime/Core/Private/ProfilingDebugging/CsvProfiler.cpp b/Engine/Source/Runtime/Core/Private/ProfilingDebugging/CsvProfiler.cpp
--- a/Engine/Source/Runtime/Core/Private/ProfilingDebugging/CsvProfiler.cpp
+++ b/Engine/Source/Runtime/Core/Private/ProfilingDebugging/CsvProfiler.cpp
@@ -3837,6 +3837,9 @@ void FCsvProfiler::BeginCaptureInternal(const FCsvCaptureCommand& CurrentCommand
 		check(GFrameNumber==1);
 		UE_LOG(LogCsvProfiler, Log, TEXT("Doing a boot time capture. Start time offset %.2f"), FPlatformTime::Seconds()-GStartTime);
 		CaptureStartTime = GStartTime;
+		// Note: we have to remove the platform time offset when converting between absolute seconds and clock cycles
+        // Warning: this won't work on unix (see comment in the else case below). First frame timing will be inaccurate there when capturing from boot
+		GCaptureStartTimestamp = FPlatformTime::SecondsToCycles64(CaptureStartTime - FPlatformTime::GetSecondsTimeOffset());
 		
 		// signal bInsertEndFrameAtFrameStart to insert an EndFrame() at the start of the first _real_ frame
 		// We also add a FrameBeginTimestampsRT timestamp here, to create a dummy renderthread frame, to ensure the rows match up in the CSV
@@ -3845,10 +3848,12 @@ void FCsvProfiler::BeginCaptureInternal(const FCsvCaptureCommand& CurrentCommand
 	}
 	else
 	{
+		// Workaround for unix time inconsistencies. We can't convert between Cycles64 and Seconds because they can come from different clocks (see UnixPlatformTime.h)
+		// Implementing FPlatformTime::GetSecondsTimeOffset() might help but we'd still get precision differences
+		// Instead we'll record StartTime and CaptureStartTimestamp independently
 		CaptureStartTime = FPlatformTime::Seconds();
+		GCaptureStartTimestamp = FPlatformTime::Cycles64();
 	}
-	// Note: we have to remove the platform time offset when converting between absolute seconds and clock cycles
-	GCaptureStartTimestamp = FPlatformTime::SecondsToCycles64(CaptureStartTime - FPlatformTime::GetSecondsTimeOffset());
 
 	LastEndFrameTimestamp = GCaptureStartTimestamp;
 
