Fix Linux ReadInternal file error handling

#rnx
#Synced-CL 48510256

diff --git a/Engine/Source/Runtime/Core/Private/Unix/UnixPlatformFile.cpp b/Engine/Source/Runtime/Core/Private/Unix/UnixPlatformFile.cpp
--- a/Engine/Source/Runtime/Core/Private/Unix/UnixPlatformFile.cpp
+++ b/Engine/Source/Runtime/Core/Private/Unix/UnixPlatformFile.cpp
@@ -360,6 +360,7 @@ private:
 			int64 ThisSize = FMath::Min<int64>(READWRITE_SIZE, BytesToRead);
 			check(Destination);
 			int64 ThisRead = read(FileHandle, Destination, ThisSize);
+			
 			if (ThisRead == -1 && errno == EFAULT)
 			{
 				// hacky workaround for UE-69018
@@ -374,6 +375,17 @@ private:
 					FMemory::Free(TempDest);
 				}
 			}
+			
+			if (ThisRead == -1 && errno == EINTR)   // need to retry
+			{
+				continue;
+			}			
+			
+			if (ThisRead == -1)  // some other error just bail
+			{
+				return BytesRead;
+			}
+
 			BytesRead += ThisRead;
 			if (ThisRead != ThisSize)
 			{
