Add a helper to UE::Logging::Private to allow ensures to emit their logging via structured logging so the messages inherit UE_LOG_CONTEXT

#Synced-CL Unknown

diff --git a/Engine/Source/Runtime/Core/Private/Logging/StructuredLog.cpp b/Engine/Source/Runtime/Core/Private/Logging/StructuredLog.cpp
--- a/Engine/Source/Runtime/Core/Private/Logging/StructuredLog.cpp
+++ b/Engine/Source/Runtime/Core/Private/Logging/StructuredLog.cpp
@@ -1599,27 +1599,15 @@ static FLogTemplate* GetStaticBasicLogTemplate()
 	return Template.Get();
 }
 
-// Serializing the log to compact binary happens in its own function because that allows stack
-// space for the writer to be returned before calling into the output devices.
-template <UE::CCharType FormatCharType>
-FORCENOINLINE static FCbObject SerializeBasicLogMessage(TStaticBasicLogRecordParam<FormatCharType> Log, va_list Args)
+static void AddContextsToBasicLog(TCbWriter<512>& Writer, FAnsiStringView Ignore1, FAnsiStringView Ignore2 = {})
 {
-	TStringBuilderWithBuffer<FormatCharType, 512> Message;
-	Message.AppendV(Log.Format, Args);
-
-	TCbWriter<512> Writer;
-	Writer.BeginObject();
-
-	const FAnsiStringView MessageName = ANSITEXTVIEW("Message");
-	Writer.AddString(MessageName, Message);
-
 	if (LogContextHead)
 	{
 		// Traverse contexts forward and copy any that Message did not override.
 		for (FLogContext* Node = LogContextHead; Node; Node = Node->Next)
 		{
 			const FAnsiStringView Name = Node->Name;
-			if (!Node->bDisabledByNewerContext && !Name.Equals(MessageName))
+			if (!Node->bDisabledByNewerContext && !Name.Equals(Ignore1) && !Name.Equals(Ignore2))
 			{
 				Writer.AddField(Name, FCbFieldView(Node->Value.GetData()));
 			}
@@ -1636,6 +1624,23 @@ FORCENOINLINE static FCbObject SerializeBasicLogMessage(TStaticBasicLogRecordPar
 		}
 		Writer.EndArray();
 	}
+}
+
+// Serializing the log to compact binary happens in its own function because that allows stack
+// space for the writer to be returned before calling into the output devices.
+template <UE::CCharType FormatCharType>
+FORCENOINLINE static FCbObject SerializeBasicLogMessage(TStaticBasicLogRecordParam<FormatCharType> Log, va_list Args)
+{
+	TStringBuilderWithBuffer<FormatCharType, 512> Message;
+	Message.AppendV(Log.Format, Args);
+
+	TCbWriter<512> Writer;
+	Writer.BeginObject();
+
+	const FAnsiStringView MessageName = ANSITEXTVIEW("Message");
+	Writer.AddString(MessageName, Message);
+
+	AddContextsToBasicLog(Writer, MessageName);
 
 	Writer.EndObject();
 	return Writer.Save().AsObject();
@@ -1787,6 +1792,51 @@ FORCENOINLINE void BasicFatalLogWithProgramCounter(TStaticBasicLogRecordParam<Fo
 template CORE_API void BasicFatalLogWithProgramCounter(TStaticBasicLogRecordParam<TCHAR> Log, const FLogCategoryBase* Category, void* ProgramCounter, ...);
 template CORE_API void BasicFatalLogWithProgramCounter(TStaticBasicLogRecordParam<UTF8CHAR> Log, const FLogCategoryBase* Category, void* ProgramCounter, ...);
 
+#if !NO_LOGGING
+static constexpr const TCHAR* const GDebugMessageLogFormat = TEXT("{Prefix}{Message}");
+
+static FLogTemplate* GetDebugMessageLogTemplate()
+{
+	static FInlineLogTemplate Template(GDebugMessageLogFormat);
+	return Template.Get();
+}
+
+// Serializing the log to compact binary happens in its own function because that allows stack
+// space for the writer to be returned before calling into the output devices.
+FORCENOINLINE static FCbObject SerializeDebugLogMessage(const TCHAR* Prefix, const TCHAR* Message)
+{
+	TCbWriter<512> Writer;
+	Writer.BeginObject();
+
+	const FAnsiStringView PrefixName = ANSITEXTVIEW("Prefix");
+	Writer.AddString(PrefixName, Prefix);
+	const FAnsiStringView MessageName = ANSITEXTVIEW("Message");
+	Writer.AddString(MessageName, Message);
+
+	AddContextsToBasicLog(Writer, MessageName, PrefixName);
+
+	Writer.EndObject();
+	return Writer.Save().AsObject();
+}
+
+UE_AUTORTFM_ALWAYS_OPEN
+void LogDebugMessageLine(const ANSICHAR* File, int32 Line, const FName& InLogName, ELogVerbosity::Type Verbosity, const TCHAR* Message, const TCHAR* Prefix)
+{
+	FCbObject Fields = SerializeDebugLogMessage(Prefix, Message);
+
+	FLogRecord Record;
+	Record.SetFormat(GDebugMessageLogFormat);
+	Record.SetTemplate(GetDebugMessageLogTemplate());
+	Record.SetFields(MoveTemp(Fields));
+	Record.SetFile(File);
+	Record.SetLine(Line);
+	Record.SetCategory(InLogName);
+	Record.SetVerbosity(Verbosity);
+	Record.SetTime(FLogTime::Now());
+
+	DispatchLogRecord(Record);
+}
+#endif // !NO_LOGGING
 } // UE::Logging::Private
 
 ///////////////////////////////////////////////////////////////////////////////////////////////////////////////////////
diff --git a/Engine/Source/Runtime/Core/Private/Misc/AssertionMacros.cpp b/Engine/Source/Runtime/Core/Private/Misc/AssertionMacros.cpp
--- a/Engine/Source/Runtime/Core/Private/Misc/AssertionMacros.cpp
+++ b/Engine/Source/Runtime/Core/Private/Misc/AssertionMacros.cpp
@@ -15,6 +15,7 @@
 #include "GenericPlatform/GenericPlatformStackWalk.h"
 #include "HAL/PlatformStackWalk.h"
 #include "Logging/LogMacros.h"
+#include "Logging/StructuredLog.h"
 #include "CoreGlobals.h"
 #include "Misc/Parse.h"
 #include "Misc/ScopeLock.h"
@@ -368,7 +369,7 @@ void FDebug::LogFormattedMessageWithCallstack(const FName& InLogName, const ANSI
 
 	if (bWriteUATMarkers)
 	{
-		FMsg::Logf(File, Line, LogName, Verbosity, TEXT("begin: stack for UAT"));
+		UE::Logging::Private::LogDebugMessageLine(File, Line, LogName, Verbosity, TEXT("begin: stack for UAT"));
 	}
 
 	if (bLowLevel)
@@ -377,8 +378,8 @@ void FDebug::LogFormattedMessageWithCallstack(const FName& InLogName, const ANSI
 	}
 	else
 	{
-		FMsg::Logf(File, Line, LogName, Verbosity, TEXT("%s"), Heading);
-		FMsg::Logf(File, Line, LogName, Verbosity, TEXT(""));
+		UE::Logging::Private::LogDebugMessageLine(File, Line, LogName, Verbosity, Heading);
+		UE::Logging::Private::LogDebugMessageLine(File, Line, LogName, Verbosity, TEXT(""));
 	}
 
 	for (const TCHAR* LineStart = Message;; )
@@ -409,7 +410,7 @@ void FDebug::LogFormattedMessageWithCallstack(const FName& InLogName, const ANSI
 		}
 		else
 		{
-			FMsg::Logf(File, Line, LogName, Verbosity, TEXT("%s%s"), Prefix, SingleLine);
+			UE::Logging::Private::LogDebugMessageLine(File, Line, LogName, Verbosity, SingleLine, Prefix);
 		}
 		
 		// Quit if this was the last line
@@ -424,7 +425,7 @@ void FDebug::LogFormattedMessageWithCallstack(const FName& InLogName, const ANSI
 
 	if (bWriteUATMarkers)
 	{
-		FMsg::Logf(File, Line, LogName, Verbosity, TEXT("end: stack for UAT"));
+		UE::Logging::Private::LogDebugMessageLine(File, Line, LogName, Verbosity, TEXT("end: stack for UAT"));
 	}
 }
 
@@ -520,11 +521,11 @@ FORCENOINLINE void FDebug::EnsureFailed(const ANSICHAR* Expr, const ANSICHAR* Fi
 #if !NO_LOGGING
 		if (GEnsuresAreErrors)
 		{
-			UE_LOG(LogOutputDevice, Error, TEXT("%s"), *CombinedErrorMsg);
+			UE_LOGFMT(LogOutputDevice, Error, "{CombinedErrorMsg}", CombinedErrorMsg);
 		}
 		else
 		{
-			UE_LOG(LogOutputDevice, Warning, TEXT("%s"), *CombinedErrorMsg);
+			UE_LOGFMT(LogOutputDevice, Warning, "{CombinedErrorMsg}", CombinedErrorMsg);
 		}
 #endif
 	}
@@ -657,11 +658,11 @@ FORCENOINLINE void FDebug::EnsureFailed(const ANSICHAR* Expr, const ANSICHAR* Fi
 #if !NO_LOGGING
 			if (GEnsuresAreErrors)
 			{
-				UE_LOG(LogOutputDevice, Error, TEXT("Ensure condition failed: %hs " FILE_LINE_DESC_ANSI), Expr, File, Line);
+				UE_LOGFMT(LogOutputDevice, Error, "Ensure condition failed: {Expr} [File:{File}] [Line: {Line}]", Expr, File, Line);
 			}
 			else
 			{
-				UE_LOG(LogOutputDevice, Warning, TEXT("Ensure condition failed: %hs " FILE_LINE_DESC_ANSI), Expr, File, Line);
+				UE_LOGFMT(LogOutputDevice, Warning, "Ensure condition failed: {Expr} [File:{File}] [Line: {Line}]", Expr, File, Line);
 			}
 #endif
 		}
diff --git a/Engine/Source/Runtime/Core/Public/Logging/StructuredLog.h b/Engine/Source/Runtime/Core/Public/Logging/StructuredLog.h
--- a/Engine/Source/Runtime/Core/Public/Logging/StructuredLog.h
+++ b/Engine/Source/Runtime/Core/Public/Logging/StructuredLog.h
@@ -529,6 +529,11 @@ struct FLogContext
 	bool bDisabledOlderContext = false;
 };
 
+#if !NO_LOGGING
+// Compatibility function for FDebug::LogFormattedMessageWithCallstack to log a single line without checking log channel verbosity, accepting channel by name
+// This exists to be able to emit those lines as structured logs inheriting the log context fields
+void LogDebugMessageLine(const ANSICHAR* File, int32 Line, const FName& InLogName, ELogVerbosity::Type Verbosity, const TCHAR* Message, const TCHAR* Prefix = TEXT(""));
+#endif
 } // UE::Logging::Private
 
 #if !NO_LOGGING
