Small performance boosts in WriteLineToStdOut and FWindowsPlatformProcess::ReadPipe

#Synced-CL 48491776

diff --git a/Engine/Source/Runtime/Core/Private/Misc/OutputDeviceStdOut.cpp b/Engine/Source/Runtime/Core/Private/Misc/OutputDeviceStdOut.cpp
--- a/Engine/Source/Runtime/Core/Private/Misc/OutputDeviceStdOut.cpp
+++ b/Engine/Source/Runtime/Core/Private/Misc/OutputDeviceStdOut.cpp
@@ -32,7 +32,7 @@ using StdOutCharType = WIDECHAR;
 
 static void WriteLineToStdOut(FUtf8StringBuilderBase& Line)
 {
-	printf("%s", (const char*)*Line);
+	fwrite(Line.GetData(), 1, Line.Len(), stdout);
 	fflush(stdout);
 }
 
diff --git a/Engine/Source/Runtime/Core/Private/Windows/WindowsPlatformProcess.cpp b/Engine/Source/Runtime/Core/Private/Windows/WindowsPlatformProcess.cpp
--- a/Engine/Source/Runtime/Core/Private/Windows/WindowsPlatformProcess.cpp
+++ b/Engine/Source/Runtime/Core/Private/Windows/WindowsPlatformProcess.cpp
@@ -1978,28 +1978,22 @@ bool FWindowsPlatformProcess::CreatePipe( void*& ReadPipe, void*& WritePipe, boo
 	return true;
 }
 
-FString FWindowsPlatformProcess::ReadPipe( void* ReadPipe )
+FString FWindowsPlatformProcess::ReadPipe(void* ReadPipe)
 {
-	FString Output;
-
 	// Note: String becomes corrupted when more than one byte per character and all bytes are not available
-	uint32 BytesAvailable = 0;
-	if (::PeekNamedPipe(ReadPipe, NULL, 0, NULL, (::DWORD*)&BytesAvailable, NULL) && (BytesAvailable > 0))
+	::DWORD BytesAvailable;
+	if (::PeekNamedPipe(ReadPipe, nullptr, 0, nullptr, &BytesAvailable, nullptr) && BytesAvailable)
 	{
-		UTF8CHAR* Buffer = new UTF8CHAR[BytesAvailable + 1];
-		uint32 BytesRead = 0;
-		if (::ReadFile(ReadPipe, Buffer, BytesAvailable, (::DWORD*)&BytesRead, NULL))
+		TUtf8StringBuilder<256> Buffer;
+		Buffer.AddUninitialized(int32(BytesAvailable));
+		::DWORD BytesRead;
+		if (::ReadFile(ReadPipe, Buffer.GetData(), BytesAvailable, &BytesRead, nullptr))
 		{
-			if (BytesRead > 0)
-			{
-				Buffer[BytesRead] = (UTF8CHAR)'\0';
-				Output += FUTF8ToTCHAR((const ANSICHAR*)Buffer).Get();
-			}
+			return FString{Buffer.ToView().Left(int32(BytesRead))};
 		}
-		delete [] Buffer;
 	}
 
-	return Output;
+	return {};
 }
 
 bool FWindowsPlatformProcess::ReadPipeToArray(void* ReadPipe, TArray<uint8> & Output)
