Fix new ensure output compilation in NO_LOGGING case

#Synced-CL Unknown

diff --git a/Engine/Source/Runtime/Core/Private/Logging/StructuredLog.cpp b/Engine/Source/Runtime/Core/Private/Logging/StructuredLog.cpp
--- a/Engine/Source/Runtime/Core/Private/Logging/StructuredLog.cpp
+++ b/Engine/Source/Runtime/Core/Private/Logging/StructuredLog.cpp
@@ -1818,10 +1818,12 @@ FORCENOINLINE static FCbObject SerializeDebugLogMessage(const TCHAR* Prefix, con
 	Writer.EndObject();
 	return Writer.Save().AsObject();
 }
+#endif // !NO_LOGGING
 
 UE_AUTORTFM_ALWAYS_OPEN
 void LogDebugMessageLine(const ANSICHAR* File, int32 Line, const FName& InLogName, ELogVerbosity::Type Verbosity, const TCHAR* Message, const TCHAR* Prefix)
 {
+#if !NO_LOGGING
 	FCbObject Fields = SerializeDebugLogMessage(Prefix, Message);
 
 	FLogRecord Record;
@@ -1835,8 +1837,8 @@ void LogDebugMessageLine(const ANSICHAR* File, int32 Line, const FName& InLogNam
 	Record.SetTime(FLogTime::Now());
 
 	DispatchLogRecord(Record);
-}
 #endif // !NO_LOGGING
+}
 } // UE::Logging::Private
 
 ///////////////////////////////////////////////////////////////////////////////////////////////////////////////////////
diff --git a/Engine/Source/Runtime/Core/Public/Logging/StructuredLog.h b/Engine/Source/Runtime/Core/Public/Logging/StructuredLog.h
--- a/Engine/Source/Runtime/Core/Public/Logging/StructuredLog.h
+++ b/Engine/Source/Runtime/Core/Public/Logging/StructuredLog.h
@@ -529,11 +529,9 @@ struct FLogContext
 	bool bDisabledOlderContext = false;
 };
 
-#if !NO_LOGGING
 // Compatibility function for FDebug::LogFormattedMessageWithCallstack to log a single line without checking log channel verbosity, accepting channel by name
 // This exists to be able to emit those lines as structured logs inheriting the log context fields
 void LogDebugMessageLine(const ANSICHAR* File, int32 Line, const FName& InLogName, ELogVerbosity::Type Verbosity, const TCHAR* Message, const TCHAR* Prefix = TEXT(""));
-#endif
 } // UE::Logging::Private
 
 #if !NO_LOGGING
