Revert changes to assertion macros which broke handling of expected errors in some tests

#Synced-CL Unknown

diff --git a/Engine/Source/Runtime/Core/Private/Misc/AssertionMacros.cpp b/Engine/Source/Runtime/Core/Private/Misc/AssertionMacros.cpp
--- a/Engine/Source/Runtime/Core/Private/Misc/AssertionMacros.cpp
+++ b/Engine/Source/Runtime/Core/Private/Misc/AssertionMacros.cpp
@@ -15,7 +15,6 @@
 #include "GenericPlatform/GenericPlatformStackWalk.h"
 #include "HAL/PlatformStackWalk.h"
 #include "Logging/LogMacros.h"
-#include "Logging/StructuredLog.h"
 #include "CoreGlobals.h"
 #include "Misc/Parse.h"
 #include "Misc/ScopeLock.h"
@@ -369,7 +368,7 @@ void FDebug::LogFormattedMessageWithCallstack(const FName& InLogName, const ANSI
 
 	if (bWriteUATMarkers)
 	{
-		UE::Logging::Private::LogDebugMessageLine(File, Line, LogName, Verbosity, TEXT("begin: stack for UAT"));
+		FMsg::Logf(File, Line, LogName, Verbosity, TEXT("begin: stack for UAT"));
 	}
 
 	if (bLowLevel)
@@ -378,8 +377,8 @@ void FDebug::LogFormattedMessageWithCallstack(const FName& InLogName, const ANSI
 	}
 	else
 	{
-		UE::Logging::Private::LogDebugMessageLine(File, Line, LogName, Verbosity, Heading);
-		UE::Logging::Private::LogDebugMessageLine(File, Line, LogName, Verbosity, TEXT(""));
+		FMsg::Logf(File, Line, LogName, Verbosity, TEXT("%s"), Heading);
+		FMsg::Logf(File, Line, LogName, Verbosity, TEXT(""));
 	}
 
 	for (const TCHAR* LineStart = Message;; )
@@ -410,7 +409,7 @@ void FDebug::LogFormattedMessageWithCallstack(const FName& InLogName, const ANSI
 		}
 		else
 		{
-			UE::Logging::Private::LogDebugMessageLine(File, Line, LogName, Verbosity, SingleLine, Prefix);
+			FMsg::Logf(File, Line, LogName, Verbosity, TEXT("%s%s"), Prefix, SingleLine);
 		}
 		
 		// Quit if this was the last line
@@ -425,7 +424,7 @@ void FDebug::LogFormattedMessageWithCallstack(const FName& InLogName, const ANSI
 
 	if (bWriteUATMarkers)
 	{
-		UE::Logging::Private::LogDebugMessageLine(File, Line, LogName, Verbosity, TEXT("end: stack for UAT"));
+		FMsg::Logf(File, Line, LogName, Verbosity, TEXT("end: stack for UAT"));
 	}
 }
 
@@ -521,11 +520,11 @@ FORCENOINLINE void FDebug::EnsureFailed(const ANSICHAR* Expr, const ANSICHAR* Fi
 #if !NO_LOGGING
 		if (GEnsuresAreErrors)
 		{
-			UE_LOGFMT(LogOutputDevice, Error, "{CombinedErrorMsg}", CombinedErrorMsg);
+			UE_LOG(LogOutputDevice, Error, TEXT("%s"), *CombinedErrorMsg);
 		}
 		else
 		{
-			UE_LOGFMT(LogOutputDevice, Warning, "{CombinedErrorMsg}", CombinedErrorMsg);
+			UE_LOG(LogOutputDevice, Warning, TEXT("%s"), *CombinedErrorMsg);
 		}
 #endif
 	}
@@ -658,11 +657,11 @@ FORCENOINLINE void FDebug::EnsureFailed(const ANSICHAR* Expr, const ANSICHAR* Fi
 #if !NO_LOGGING
 			if (GEnsuresAreErrors)
 			{
-				UE_LOGFMT(LogOutputDevice, Error, "Ensure condition failed: {Expr} [File:{File}] [Line: {Line}]", Expr, File, Line);
+				UE_LOG(LogOutputDevice, Error, TEXT("Ensure condition failed: %hs " FILE_LINE_DESC_ANSI), Expr, File, Line);
 			}
 			else
 			{
-				UE_LOGFMT(LogOutputDevice, Warning, "Ensure condition failed: {Expr} [File:{File}] [Line: {Line}]", Expr, File, Line);
+				UE_LOG(LogOutputDevice, Warning, TEXT("Ensure condition failed: %hs " FILE_LINE_DESC_ANSI), Expr, File, Line);
 			}
 #endif
 		}
