Math: Part of 2 DiagnosticCheckNaN changeover

Log NaNs but not modify the data - mutation here is problematic.

#Synced-CL 44544657

diff --git a/Engine/Source/Runtime/Core/Public/Math/BoxSphereBounds.h b/Engine/Source/Runtime/Core/Public/Math/BoxSphereBounds.h
--- a/Engine/Source/Runtime/Core/Public/Math/BoxSphereBounds.h
+++ b/Engine/Source/Runtime/Core/Public/Math/BoxSphereBounds.h
@@ -304,17 +304,14 @@ public:
 		if (Origin.ContainsNaN())
 		{
 			logOrEnsureNanError(TEXT("Origin contains NaN: %s"), *Origin.ToString());
-			const_cast<TBoxSphereBounds*>(this)->Origin = TVector<T>::ZeroVector;
 		}
 		if (BoxExtent.ContainsNaN())
 		{
 			logOrEnsureNanError(TEXT("BoxExtent contains NaN: %s"), *BoxExtent.ToString());
-			const_cast<TBoxSphereBounds*>(this)->BoxExtent = TVector<TExtent>::ZeroVector;
 		}
 		if (FMath::IsNaN(SphereRadius) || !FMath::IsFinite(SphereRadius))
 		{
 			logOrEnsureNanError(TEXT("SphereRadius contains NaN: %f"), SphereRadius);
-			const_cast<TBoxSphereBounds*>(this)->SphereRadius = 0.f;
 		}
 	}
 #else
diff --git a/Engine/Source/Runtime/Core/Public/Math/Plane.h b/Engine/Source/Runtime/Core/Public/Math/Plane.h
--- a/Engine/Source/Runtime/Core/Public/Math/Plane.h
+++ b/Engine/Source/Runtime/Core/Public/Math/Plane.h
@@ -50,7 +50,6 @@ public:
 		if (TVector<T>::ContainsNaN() || !FMath::IsFinite(W))
 		{
 			logOrEnsureNanError(TEXT("FPlane contains NaN: %s W=%3.3f"), *TVector<T>::ToString(), W);
-			*const_cast<TPlane<T>*>(static_cast<const TPlane<T>*>(this)) = TPlane<T>(ForceInitToZero);
 		}
 	}
 #else
diff --git a/Engine/Source/Runtime/Core/Public/Math/Quat.h b/Engine/Source/Runtime/Core/Public/Math/Quat.h
--- a/Engine/Source/Runtime/Core/Public/Math/Quat.h
+++ b/Engine/Source/Runtime/Core/Public/Math/Quat.h
@@ -591,7 +591,6 @@ public:
 		if (ContainsNaN())
 		{
 			logOrEnsureNanError(TEXT("Quat contains NaN: %s"), *ToString());
-			*const_cast<TQuat*>(this) = TQuat<T>(0.f, 0.f, 0.f, 1.f);
 		}
 	}
 
@@ -600,7 +599,6 @@ public:
 		if (ContainsNaN())
 		{
 			logOrEnsureNanError(TEXT("%s: Quat contains NaN: %s"), Message, *ToString());
-			*const_cast<TQuat*>(this) = TQuat<T>(0.f, 0.f, 0.f, 1.f);
 		}
 	}
 #else
diff --git a/Engine/Source/Runtime/Core/Public/Math/Rotator.h b/Engine/Source/Runtime/Core/Public/Math/Rotator.h
--- a/Engine/Source/Runtime/Core/Public/Math/Rotator.h
+++ b/Engine/Source/Runtime/Core/Public/Math/Rotator.h
@@ -64,7 +64,6 @@ public:
 		if (ContainsNaN())
 		{
 			logOrEnsureNanError(TEXT("TRotator contains NaN: %s"), *ToString());
-			*const_cast<TRotator<T>*>(this) = ZeroRotator;
 		}
 	}
 
@@ -73,7 +72,6 @@ public:
 		if (ContainsNaN())
 		{
 			logOrEnsureNanError(TEXT("%s: TRotator contains NaN: %s"), Message, *ToString());
-			*const_cast<TRotator<T>*>(this) = ZeroRotator;
 		}
 	}
 #else
diff --git a/Engine/Source/Runtime/Core/Public/Math/TransformCalculus2D.h b/Engine/Source/Runtime/Core/Public/Math/TransformCalculus2D.h
--- a/Engine/Source/Runtime/Core/Public/Math/TransformCalculus2D.h
+++ b/Engine/Source/Runtime/Core/Public/Math/TransformCalculus2D.h
@@ -612,7 +612,6 @@ public:
 		if (ContainsNaN())
 		{
 			logOrEnsureNanError(TEXT("Matrix2x2 contains NaN"));
-			*const_cast<TMatrix2x2<T>*>(static_cast<const TMatrix2x2<T>*>(this)) = TMatrix2x2<T>(ForceInitToZero);
 		}
 	}
 #else
diff --git a/Engine/Source/Runtime/Core/Public/Math/Vector2D.h b/Engine/Source/Runtime/Core/Public/Math/Vector2D.h
--- a/Engine/Source/Runtime/Core/Public/Math/Vector2D.h
+++ b/Engine/Source/Runtime/Core/Public/Math/Vector2D.h
@@ -673,7 +673,6 @@ public:
 		if (ContainsNaN())
 		{
 			logOrEnsureNanError(TEXT("FVector2 contains NaN: %s"), *ToString());
-			*this = TVector2<T>::ZeroVector;
 		}
 	}
 #else
diff --git a/Engine/Source/Runtime/Core/Public/Math/Vector4.h b/Engine/Source/Runtime/Core/Public/Math/Vector4.h
--- a/Engine/Source/Runtime/Core/Public/Math/Vector4.h
+++ b/Engine/Source/Runtime/Core/Public/Math/Vector4.h
@@ -547,8 +547,6 @@ public:
 		if (ContainsNaN())
 		{
 			logOrEnsureNanError(TEXT("FVector4 contains NaN: %s"), *ToString());
-			*this = TVector4();
-
 		}
 	}
 #else
