[Backout] - CL48529305
#rnx
Original CL Desc
-----------------------------------------------------------------
Revert changes to assertion macros which broke handling of expected errors in some tests

#Synced-CL Unknown

#Synced-CL Unknown

diff --git a/Engine/Source/Runtime/Core/Private/Misc/AssertionMacros.cpp b/Engine/Source/Runtime/Core/Private/Misc/AssertionMacros.cpp
--- a/Engine/Source/Runtime/Core/Private/Misc/AssertionMacros.cpp
+++ b/Engine/Source/Runtime/Core/Private/Misc/AssertionMacros.cpp
@@ -15,6 +15,7 @@
 #include "GenericPlatform/GenericPlatformStackWalk.h"
 #include "HAL/PlatformStackWalk.h"
 #include "Logging/LogMacros.h"
+#include "Logging/StructuredLog.h"
 #include "CoreGlobals.h"
 #include "Misc/Parse.h"
 #include "Misc/ScopeLock.h"
@@ -368,7 +369,7 @@ void FDebug::LogFormattedMessageWithCallstack(const FName& InLogName, const ANSI
 
 	if (bWriteUATMarkers)
 	{
-		FMsg::Logf(File, Line, LogName, Verbosity, TEXT("begin: stack for UAT"));
+		UE::Logging::Private::LogDebugMessageLine(File, Line, LogName, Verbosity, TEXT("begin: stack for UAT"));
 	}
 
 	if (bLowLevel)
@@ -377,8 +378,8 @@ void FDebug::LogFormattedMessageWithCallstack(const FName& InLogName, const ANSI
 	}
 	else
 	{
-		FMsg::Logf(File, Line, LogName, Verbosity, TEXT("%s"), Heading);
-		FMsg::Logf(File, Line, LogName, Verbosity, TEXT(""));
+		UE::Logging::Private::LogDebugMessageLine(File, Line, LogName, Verbosity, Heading);
+		UE::Logging::Private::LogDebugMessageLine(File, Line, LogName, Verbosity, TEXT(""));
 	}
 
 	for (const TCHAR* LineStart = Message;; )
@@ -409,7 +410,7 @@ void FDebug::LogFormattedMessageWithCallstack(const FName& InLogName, const ANSI
 		}
 		else
 		{
-			FMsg::Logf(File, Line, LogName, Verbosity, TEXT("%s%s"), Prefix, SingleLine);
+			UE::Logging::Private::LogDebugMessageLine(File, Line, LogName, Verbosity, SingleLine, Prefix);
 		}
 		
 		// Quit if this was the last line
@@ -424,7 +425,7 @@ void FDebug::LogFormattedMessageWithCallstack(const FName& InLogName, const ANSI
 
 	if (bWriteUATMarkers)
 	{
-		FMsg::Logf(File, Line, LogName, Verbosity, TEXT("end: stack for UAT"));
+		UE::Logging::Private::LogDebugMessageLine(File, Line, LogName, Verbosity, TEXT("end: stack for UAT"));
 	}
 }
 
@@ -520,11 +521,11 @@ FORCENOINLINE void FDebug::EnsureFailed(const ANSICHAR* Expr, const ANSICHAR* Fi
 #if !NO_LOGGING
 		if (GEnsuresAreErrors)
 		{
-			UE_LOG(LogOutputDevice, Error, TEXT("%s"), *CombinedErrorMsg);
+			UE_LOGFMT(LogOutputDevice, Error, "{CombinedErrorMsg}", CombinedErrorMsg);
 		}
 		else
 		{
-			UE_LOG(LogOutputDevice, Warning, TEXT("%s"), *CombinedErrorMsg);
+			UE_LOGFMT(LogOutputDevice, Warning, "{CombinedErrorMsg}", CombinedErrorMsg);
 		}
 #endif
 	}
@@ -657,11 +658,11 @@ FORCENOINLINE void FDebug::EnsureFailed(const ANSICHAR* Expr, const ANSICHAR* Fi
 #if !NO_LOGGING
 			if (GEnsuresAreErrors)
 			{
-				UE_LOG(LogOutputDevice, Error, TEXT("Ensure condition failed: %hs " FILE_LINE_DESC_ANSI), Expr, File, Line);
+				UE_LOGFMT(LogOutputDevice, Error, "Ensure condition failed: {Expr} [File:{File}] [Line: {Line}]", Expr, File, Line);
 			}
 			else
 			{
-				UE_LOG(LogOutputDevice, Warning, TEXT("Ensure condition failed: %hs " FILE_LINE_DESC_ANSI), Expr, File, Line);
+				UE_LOGFMT(LogOutputDevice, Warning, "Ensure condition failed: {Expr} [File:{File}] [Line: {Line}]", Expr, File, Line);
 			}
 #endif
 		}
