Remove erroneous CVar performance warning from Check Asset References validation step

#Synced-CL 48528694

diff --git a/Engine/Source/Runtime/Core/Private/HAL/ConsoleManager.cpp b/Engine/Source/Runtime/Core/Private/HAL/ConsoleManager.cpp
--- a/Engine/Source/Runtime/Core/Private/HAL/ConsoleManager.cpp
+++ b/Engine/Source/Runtime/Core/Private/HAL/ConsoleManager.cpp
@@ -3552,7 +3552,7 @@ void FConsoleManager::LoadAllPlatformCVars(FName PlatformName, const FString& De
 		int32 SetBy = SetByAndPreview & ECVF_SetByMask;
 			
 		// make sure the named cvar exists
-		IConsoleVariable* CVar = IConsoleManager::Get().FindConsoleVariable(*CVarName);
+		IConsoleVariable* CVar = IConsoleManager::Get().FindConsoleVariable(*CVarName, false /* bTrackFrequentCalls */);
 		if (CVar == nullptr)
 		{
 			// store this unknown value in a map, so when the cvar gets created later, we can grab this platform's values into the platform variable
