[AppTime]
- Fix asserts coming from the async loading thread to unblock us and because we don't care about having the contextual time during asset loading

#Synced-CL Unknown

diff --git a/Engine/Source/Runtime/Core/Private/Misc/AppTime.cpp b/Engine/Source/Runtime/Core/Private/Misc/AppTime.cpp
--- a/Engine/Source/Runtime/Core/Private/Misc/AppTime.cpp
+++ b/Engine/Source/Runtime/Core/Private/Misc/AppTime.cpp
@@ -28,9 +28,9 @@ CORE_API const FAppTime& FAppTime::Get()
 	{
 		return *TLSInstance;
 	}
-	else if (IsInSlateThread())
+	else if (IsInSlateThread() || IsInAsyncLoadingThread() || IsInParallelLoadingThread())
 	{
-		// @todo delta time refactor - ignore races with the slate thread for now
+		// @todo delta time refactor - ignore races with the slate thread and loading threads for now
 		return GameThreadInstance;
 	}
 	else
