Marked unused types as deprecated and added some documentation to the I/O dispatcher API.

#rnx
#Synced-CL 48546254

diff --git a/Engine/Source/Runtime/Core/Private/IO/IoDispatcher.cpp b/Engine/Source/Runtime/Core/Private/IO/IoDispatcher.cpp
--- a/Engine/Source/Runtime/Core/Private/IO/IoDispatcher.cpp
+++ b/Engine/Source/Runtime/Core/Private/IO/IoDispatcher.cpp
@@ -1095,7 +1095,6 @@ private:
 	TArray<FIoRequestImpl*> RequestsToCancel;
 	TArray<FIoRequestImpl*> RequestsToReprioritize;
 	TAtomic<bool> bStopRequested { false };
-	FIoDispatcher::FIoContainerUnmountedEvent ContainerUnmountedEvent;
 	uint64 PendingIoRequestsCount = 0;
 	int64 TotalLoaded = 0;
 	FIoRequestStats RequestStats;
diff --git a/Engine/Source/Runtime/Core/Public/IO/IoDispatcher.h b/Engine/Source/Runtime/Core/Public/IO/IoDispatcher.h
--- a/Engine/Source/Runtime/Core/Public/IO/IoDispatcher.h
+++ b/Engine/Source/Runtime/Core/Public/IO/IoDispatcher.h
@@ -69,11 +69,10 @@ template <typename OptionalType> struct TOptional;
 
 CORE_API DECLARE_LOG_CATEGORY_EXTERN(LogIoDispatcher, Log, All);
 
-//////////////////////////////////////////////////////////////////////////
-
-/** Helper used to manage creation of I/O store file handles etc
-  */
-class FIoStoreEnvironment
+/** Helper used to manage creation of I/O store file handles etc. */
+class
+UE_DEPRECATED(5.8, "FIoStoreEnvironment is not used and is deprecated.")
+FIoStoreEnvironment
 {
 public:
 	CORE_API FIoStoreEnvironment();
@@ -89,6 +88,7 @@ private:
 	int32			Order = 0;
 };
 
+/** Represents the hash of an I/O store chunk. */
 class
 UE_DEPRECATED(5.5, "FIoChunkHash is deprecated. Use FIoHash instead.")
 FIoChunkHash
@@ -149,8 +149,7 @@ private:
 	uint8	Hash[32];
 };
 
-//////////////////////////////////////////////////////////////////////////
-
+/** Additional I/O dispatcher read option flags. */
 enum class EIoReadOptionsFlags : uint32
 {
 	None = 0,
@@ -161,6 +160,7 @@ enum class EIoReadOptionsFlags : uint32
 };
 ENUM_CLASS_FLAGS(EIoReadOptionsFlags);
 
+/** I/O dispatcher read options. */
 class FIoReadOptions
 {
 public:
@@ -186,37 +186,44 @@ public:
 
 	~FIoReadOptions() = default;
 
+	/** Set the requested range within the I/O store chunk. */
 	void SetRange(uint64 Offset, uint64 Size)
 	{
 		RequestedOffset = Offset;
 		RequestedSize	= Size;
 	}
 
+	/** Set the target address where to write the decoded I/O store chunk data. */
 	void SetTargetVa(void* InTargetVa)
 	{
 		TargetVa = InTargetVa;
 	}
 
+	/** Set additional read option flags. */
 	void SetFlags(EIoReadOptionsFlags InValue)
 	{
 		Flags = InValue;
 	}
 
+	/** Returns the offset within the I/O store chunk. */
 	uint64 GetOffset() const
 	{
 		return RequestedOffset;
 	}
 
+	/** Returns the range size to read. */
 	uint64 GetSize() const
 	{
 		return RequestedSize;
 	}
 
+	/** Returns the target address. */
 	void* GetTargetVa() const
 	{
 		return TargetVa;
 	}
 
+	/** Returns the read option flags. */
 	EIoReadOptionsFlags GetFlags() const
 	{
 		return Flags;
@@ -229,10 +236,7 @@ private:
 	EIoReadOptionsFlags Flags = EIoReadOptionsFlags::None;
 };
 
-//////////////////////////////////////////////////////////////////////////
-
-/**
-  */
+/** Handle to an I/O dispatcher read request. */
 class FIoRequest final
 {
 public:
@@ -262,6 +266,7 @@ private:
 
 using FIoReadCallback = TFunction<void(TIoStatusOr<FIoBuffer>)>;
 
+/** Converts filesystem priority flags to I/O dispatcher priority. */
 inline int32 ConvertToIoDispatcherPriority(EAsyncIOPriorityAndFlags AIOP)
 {
 	int32 AIOPriorityToIoDispatcherPriorityMap[] = {
@@ -276,11 +281,7 @@ inline int32 ConvertToIoDispatcherPriority(EAsyncIOPriorityAndFlags AIOP)
 	return AIOPriorityToIoDispatcherPriorityMap[AIOP & AIOP_PRIORITY_MASK];
 }
 
-/** I/O batch
-
-	This is a primitive used to group I/O requests for synchronization
-	purposes
-  */
+/** Enables dispatching multiple I/O dispatcher request(s) as a single request. */
 class FIoBatch final
 {
 	friend class FIoDispatcher;
@@ -309,21 +310,25 @@ private:
 	FIoRequestImpl*		TailRequest = nullptr;
 };
 
-/**
- * Mapped region.
- */
+/** File and region handles for a memory mapped I/O store chunk. */
 struct FIoMappedRegion
 {
 	IMappedFileHandle* MappedFileHandle = nullptr;
 	IMappedFileRegion* MappedFileRegion = nullptr;
 };
 
-struct FIoDispatcherMountedContainer
+/** Holds information about a mounted container. */
+struct
+UE_DEPRECATED(5.8, "FIoDispatcherMountedContainer is deprecated.")
+FIoDispatcherMountedContainer
 {
+	PRAGMA_DISABLE_DEPRECATION_WARNINGS
 	FIoStoreEnvironment Environment;
+	PRAGMA_ENABLE_DEPRECATION_WARNINGS
 	FIoContainerId ContainerId;
 };
 
+/** Holds information to where a signature error occurred. */
 struct FIoSignatureError
 {
 	FString ContainerName;
@@ -332,28 +337,36 @@ struct FIoSignatureError
 	FSHAHash ActualHash;
 };
 
+/** Signature error delegate. */
 DECLARE_MULTICAST_DELEGATE_OneParam(FIoSignatureErrorDelegate, const FIoSignatureError&);
 
-struct FIoSignatureErrorEvent
+struct
+UE_DEPRECATED(5.8, "FIoSignatureErrorEvent is deprecated. Use FIoSignatureError.")
+FIoSignatureErrorEvent
 {
 	FCriticalSection CriticalSection;
 	FIoSignatureErrorDelegate SignatureErrorDelegate;
 };
 
-DECLARE_MULTICAST_DELEGATE_OneParam(FIoSignatureErrorDelegate, const FIoSignatureError&);
-
-DECLARE_MULTICAST_DELEGATE_OneParam(FIoContainerMountedDelegate, const FIoContainerId&);
-
-/** I/O dispatcher
-  */
+/**
+ * The I/O dispatcher enables reading data addressed with chunk identifiers (@see FIoChunkId).
+ * Multiple backends can be mounted enabling reading data from multiple storage mediums such
+ * as the fileystem or via HTTP.
+ */
 class FIoDispatcher final
 {
 public:
+	PRAGMA_DISABLE_DEPRECATION_WARNINGS
+	UE_DEPRECATED(5.8, "FIoContainerMountedEvent is deprecated")
 	DECLARE_EVENT_OneParam(FIoDispatcher, FIoContainerMountedEvent, const FIoDispatcherMountedContainer&);
+	UE_DEPRECATED(5.8, "FIoContainerUnmountedEvent is deprecated")
 	DECLARE_EVENT_OneParam(FIoDispatcher, FIoContainerUnmountedEvent, const FIoDispatcherMountedContainer&);
+	PRAGMA_ENABLE_DEPRECATION_WARNINGS
 
 	CORE_API						FIoDispatcher();
 	CORE_API						~FIoDispatcher();
+									FIoDispatcher(const FIoDispatcher&) = delete;
+									FIoDispatcher& operator=(const FIoDispatcher&) = delete;
 
 	CORE_API void					Mount(TSharedRef<IIoDispatcherBackend> Backend, int32 Priority = 0);
 
@@ -366,18 +379,14 @@ public:
 	CORE_API TIoStatusOr<uint64>	GetSizeForChunk(const FIoChunkId& ChunkId) const;
 	CORE_API int64					GetTotalLoaded() const;
 
-
 	// Events
 	CORE_API FIoSignatureErrorDelegate& OnSignatureError();
 
-	FIoDispatcher(const FIoDispatcher&) = default;
-	FIoDispatcher& operator=(const FIoDispatcher&) = delete;
-
-	static CORE_API bool IsInitialized();
-	static CORE_API FIoStatus Initialize();
-	static CORE_API void InitializePostSettings();
-	static CORE_API void Shutdown();
-	static CORE_API FIoDispatcher& Get();
+	static CORE_API bool			IsInitialized();
+	static CORE_API FIoStatus		Initialize();
+	static CORE_API void			InitializePostSettings();
+	static CORE_API void			Shutdown();
+	static CORE_API FIoDispatcher&	Get();
 
 private:
 	CORE_API FUtf8StringView		GetFilename(const FIoChunkId& ChunkId, FUtf8StringBuilderBase& OutFilename, FUtf8StringView& OutContainerName) const;
@@ -395,8 +404,7 @@ private:
 	friend class FIoDispatcherInternal;
 };
 
-//////////////////////////////////////////////////////////////////////////
-
+/** Represents a handle into the I/O store directory index. */
 class FIoDirectoryIndexHandle
 {
 	static constexpr uint32 InvalidHandle = ~uint32(0);
@@ -454,6 +462,11 @@ private:
 
 using FDirectoryIndexVisitorFunction = TFunctionRef<bool(FStringView, const uint32)>;
 
+/**
+ * The directory index reader enables reading filenames mapped from chunk identifiers.
+ * Note that the directory index is still considered metadata and is not available in
+ * the game runtime.
+ */
 class FIoDirectoryIndexReader
 {
 public:
@@ -478,8 +491,7 @@ private:
 	FIoDirectoryIndexReaderImpl* Impl;
 };
 
-//////////////////////////////////////////////////////////////////////////
-
+/** I/O store container flags. */
 enum class EIoContainerFlags : uint8
 {
 	None,
@@ -495,6 +507,7 @@ ENUM_CLASS_FLAGS(EIoContainerFlags);
 CORE_API FStringBuilderBase& operator<<(FStringBuilderBase& Sb, EIoContainerFlags Flags);
 CORE_API FString LexToString(EIoContainerFlags Flags);
 
+/** Settings used when creating I/O store container files (.ucas). */
 struct FIoContainerSettings
 {
 	FIoContainerId ContainerId;
@@ -530,6 +543,7 @@ struct FIoContainerSettings
 	}
 };
 
+/** Holds information about an I/O store chunk found in the container TOC. */
 struct FIoStoreTocChunkInfo
 {
 	FIoChunkId Id;
@@ -560,6 +574,7 @@ struct FIoStoreTocChunkInfo
 	PRAGMA_ENABLE_DEPRECATION_WARNINGS
 };
 
+/** Holds information about an I/O store chunk. */
 struct FIoStoreTocCompressedBlockInfo
 {
 	uint64 Offset;
@@ -568,6 +583,7 @@ struct FIoStoreTocCompressedBlockInfo
 	uint8 CompressionMethodIndex;
 };
 
+/** Holds information about an I/O store compression block. */
 struct FIoStoreCompressedBlockInfo
 {
 	/** Name of the method used to compress the block. */
@@ -582,6 +598,7 @@ struct FIoStoreCompressedBlockInfo
 	uint64 OffsetInBuffer;
 };
 
+/** Holds information about a decoded I/O store chunk. */
 struct FIoStoreCompressedReadResult
 {
 	/** The buffer containing the chunk. */
@@ -598,6 +615,7 @@ struct FIoStoreCompressedReadResult
 	uint64 TotalCompressedSize = 0;
 };
 
+/** The I/O store reader provides non-runtime convenience functionality for reading I/O store container files. */
 class FIoStoreReader
 {
 public:
@@ -645,4 +663,3 @@ public:
 private:
 	FIoStoreReaderImpl* Impl;
 };
-//////////////////////////////////////////////////////////////////////////
