Fixup deprecation macro in event declaration.

#rnx
#Synced-CL 48546254

diff --git a/Engine/Source/Runtime/Core/Public/IO/IoDispatcher.h b/Engine/Source/Runtime/Core/Public/IO/IoDispatcher.h
--- a/Engine/Source/Runtime/Core/Public/IO/IoDispatcher.h
+++ b/Engine/Source/Runtime/Core/Public/IO/IoDispatcher.h
@@ -357,10 +357,16 @@ class FIoDispatcher final
 {
 public:
 	PRAGMA_DISABLE_DEPRECATION_WARNINGS
-	UE_DEPRECATED(5.8, "FIoContainerMountedEvent is deprecated")
-	DECLARE_EVENT_OneParam(FIoDispatcher, FIoContainerMountedEvent, const FIoDispatcherMountedContainer&);
-	UE_DEPRECATED(5.8, "FIoContainerUnmountedEvent is deprecated")
-	DECLARE_EVENT_OneParam(FIoDispatcher, FIoContainerUnmountedEvent, const FIoDispatcherMountedContainer&);
+	DECLARE_EVENT_OneParam(
+		FIoDispatcher,
+		UE_DEPRECATED(5.8, "FIoContainerMountedEvent is deprecated")
+		FIoContainerMountedEvent,
+		const FIoDispatcherMountedContainer&);
+	DECLARE_EVENT_OneParam(
+		FIoDispatcher,
+		UE_DEPRECATED(5.8, "FIoContainerUnmountedEvent is deprecated")
+		FIoContainerUnmountedEvent,
+		const FIoDispatcherMountedContainer&);
 	PRAGMA_ENABLE_DEPRECATION_WARNINGS
 
 	CORE_API						FIoDispatcher();
