Replaced FutureDetail::SetPromiseValue with a if constexpr.

diff --git a/Engine/Source/Runtime/Core/Public/Async/Future.h b/Engine/Source/Runtime/Core/Public/Async/Future.h
--- a/Engine/Source/Runtime/Core/Public/Async/Future.h
+++ b/Engine/Source/Runtime/Core/Public/Async/Future.h
@@ -644,24 +644,6 @@ private:
 /* TFuture::Then
 *****************************************************************************/
 
-namespace FutureDetail
-{
-	/**
-	* Template for setting a promise value from a continuation.
-	*/
-	template<typename Func, typename ParamType, typename ResultType>
-	AUTORTFM_INFER inline void SetPromiseValue(TPromise<ResultType>& Promise, Func& Function, TFuture<ParamType>&& Param)
-	{
-		Promise.SetValue(Function(MoveTemp(Param)));
-	}
-	template<typename Func, typename ParamType>
-	AUTORTFM_INFER inline void SetPromiseValue(TPromise<void>& Promise, Func& Function, TFuture<ParamType>&& Param)
-	{
-		Function(MoveTemp(Param));
-		Promise.SetValue();
-	}
-}
-
 // Then implementation
 template<typename ResultType>
 template<typename Func>
@@ -674,7 +656,15 @@ auto TFutureBase<ResultType>::Then(Func Continuation) //-> TFuture<decltype(Cont
 	TFuture<ReturnValue> FutureResult = Promise.GetFuture();
 	TUniqueFunction<void()> Callback = [PromiseCapture = MoveTemp(Promise), ContinuationCapture = MoveTemp(Continuation), StateCapture = this->State]() mutable
 	{
-		FutureDetail::SetPromiseValue(PromiseCapture, ContinuationCapture, TFuture<ResultType>(MoveTemp(StateCapture)));
+		if constexpr (std::is_void_v<ReturnValue>)
+		{
+			ContinuationCapture(TFuture<ResultType>(MoveTemp(StateCapture)));
+			PromiseCapture.SetValue();
+		}
+		else
+		{
+			PromiseCapture.SetValue(ContinuationCapture(TFuture<ResultType>(MoveTemp(StateCapture))));
+		}
 	};
 
 	// This invalidate this future.
