Logging: Use UTF-8 for structured logging format strings

#Synced-CL 32500503

diff --git a/Engine/Source/Editor/UnrealEd/Private/Cooker/CookImportsChecker.cpp b/Engine/Source/Editor/UnrealEd/Private/Cooker/CookImportsChecker.cpp
--- a/Engine/Source/Editor/UnrealEd/Private/Cooker/CookImportsChecker.cpp
+++ b/Engine/Source/Editor/UnrealEd/Private/Cooker/CookImportsChecker.cpp
@@ -737,10 +737,10 @@ void FEDLCookChecker::Verify(const TFunction<void(UE::FLogRecord&& Record)>& Mes
 #endif
 					Record.SetVerbosity(MissingContentSeverity);
 					Record.SetTime(UE::FLogTime::Now());
-					Record.SetFormat(TEXT("Content is missing from cook. Source package referenced an object in target package but {Reason}.")
-						TEXT("\n\tSource package: {Source}")
-						TEXT("\n\tTarget package: {Target}")
-						TEXT("\n\tReferenced object: {ReferencedObject}"));
+					Record.SetUtf8Format(UTF8TEXT("Content is missing from cook. Source package referenced an object in target package but {Reason}.")
+						UTF8TEXT("\n\tSource package: {Source}")
+						UTF8TEXT("\n\tTarget package: {Target}")
+						UTF8TEXT("\n\tReferenced object: {ReferencedObject}"));
 					{
 						FCbWriter Writer;
 						Writer.BeginObject();
diff --git a/Engine/Source/Editor/UnrealEd/Private/Cooker/CookLog.cpp b/Engine/Source/Editor/UnrealEd/Private/Cooker/CookLog.cpp
--- a/Engine/Source/Editor/UnrealEd/Private/Cooker/CookLog.cpp
+++ b/Engine/Source/Editor/UnrealEd/Private/Cooker/CookLog.cpp
@@ -4,7 +4,7 @@
 
 #include "Cooker/CookLogPrivate.h"
 
-#include "Containers/AnsiString.h"
+#include "Containers/Utf8String.h"
 #include "Cooker/CookPackageData.h"
 #include "Cooker/CookPlatformManager.h"
 #include "Cooker/CookWorkerClient.h"
@@ -142,7 +142,8 @@ private:
 		TFunctionRef<bool(FName, const FString&)> MessagePassesFilter,
 		TFunctionRef<bool(const FString&, FString&)> TryTransformMessage);
 
-	void ReportActiveLog(FReplicatedLogData&& LogData, FStringView FormatMessage, ELogVerbosity::Type Verbosity);
+	template <UE::CCharType CharType>
+	void ReportActiveLog(FReplicatedLogData&& LogData, TStringView<CharType> FormatMessage, ELogVerbosity::Type Verbosity);
 	void RecordLogForIncrementalCook(FReplicatedLogData&& LogData, ELogVerbosity::Type Verbosity);
 	void RecordLogForIncrementalCookGameThreadPortion(FName PackageName, FReplicatedLogData&& LogData);
 
@@ -157,7 +158,7 @@ private:
 
 	FCriticalSection TableLock;
 	TArray<FString> StringTable;
-	TArray<FAnsiString> AnsiStringTable;
+	TArray<FUtf8String> Utf8StringTable;
 	TArray<FUniqueLogTemplate> TemplateTable;
 };
 
@@ -248,7 +249,7 @@ void FLogHandler::ConditionalPruneReplay()
 	// Flush if the tables in the serialization context have exceeded 100 entries
 	const int32 TableSizeToFlushAt = 100;
 	if ((StringTable.Num() > TableSizeToFlushAt)
-		|| (AnsiStringTable.Num() > TableSizeToFlushAt)
+		|| (Utf8StringTable.Num() > TableSizeToFlushAt)
 		|| (TemplateTable.Num() > TableSizeToFlushAt))
 	{
 		PruneReplay();
@@ -259,7 +260,7 @@ void FLogHandler::PruneReplay()
 {
 	// We are going to drop data from our tables that might be pointed to from logs still pending in GLog. So Flush
 	// logs before we prune the tables.
-	if (!StringTable.IsEmpty() || !AnsiStringTable.IsEmpty() || !TemplateTable.IsEmpty())
+	if (!StringTable.IsEmpty() || !Utf8StringTable.IsEmpty() || !TemplateTable.IsEmpty())
 	{
 		// NOTE: We only call FlushThreadedLogs on GLog even though we might serialize structured logs via GLog or GWarn.
 		// GWarn is an output device, but GLog is a an output redirector, and only the redirector has/needs FlushThreadedLogs.
@@ -268,7 +269,7 @@ void FLogHandler::PruneReplay()
 	}
 
 	StringTable.Empty();
-	AnsiStringTable.Empty();
+	Utf8StringTable.Empty();
 	TemplateTable.Empty();
 }
 
@@ -291,7 +292,16 @@ void FLogHandler::SerializeRecord(const UE::FLogRecord& Record)
 {
 	FReplicatedLogData SerializedData;
 	Marshal(SerializedData, Record);
-	ReportActiveLog(MoveTemp(SerializedData), Record.GetFormat(), Record.GetVerbosity());
+	if (const UTF8CHAR* Utf8Format = Record.GetUtf8Format())
+	{
+		ReportActiveLog<UTF8CHAR>(MoveTemp(SerializedData), Utf8Format, Record.GetVerbosity());
+	}
+	else
+	{
+		PRAGMA_DISABLE_DEPRECATION_WARNINGS;
+		ReportActiveLog<TCHAR>(MoveTemp(SerializedData), Record.GetFormat(), Record.GetVerbosity());
+		PRAGMA_ENABLE_DEPRECATION_WARNINGS;
+	}
 }
 
 void FLogHandler::Flush()
@@ -329,7 +339,16 @@ void FLogHandler::Marshal(FReplicatedLogData& OutData, const FLogRecord& LogReco
 	Writer << LogRecord.GetCategory();
 	Writer << static_cast<uint8>(LogRecord.GetVerbosity());
 	Writer << LogRecord.GetTime().GetUtcTime();
-	Writer << LogRecord.GetFormat();
+	if (const UTF8CHAR* Utf8Format = LogRecord.GetUtf8Format())
+	{
+		Writer << Utf8Format;
+	}
+	else
+	{
+		PRAGMA_DISABLE_DEPRECATION_WARNINGS;
+		Writer << LogRecord.GetFormat();
+		PRAGMA_ENABLE_DEPRECATION_WARNINGS;
+	}
 	Writer << LogRecord.GetFields();
 	Writer << LogRecord.GetFile();
 	Writer << LogRecord.GetLine();
@@ -434,9 +453,8 @@ bool FLogHandler::UnMarshal(FCbFieldView Field, FLogRecord& OutLogRecord,
 		{
 			SerializedString = MoveTemp(TransformedString);
 		}
-		FString& FormatString = StringTable.AddDefaulted_GetRef();
-		FormatString = MoveTemp(SerializedString);
-		OutLogRecord.SetFormat(*FormatString);
+		FUtf8String& FormatString = Utf8StringTable.Emplace_GetRef(SerializedString);
+		OutLogRecord.SetUtf8Format(*FormatString);
 	}
 	else
 	{
@@ -450,9 +468,9 @@ bool FLogHandler::UnMarshal(FCbFieldView Field, FLogRecord& OutLogRecord,
 
 	if (TUtf8StringBuilder<64> FileStringBuilder; LoadFromCompactBinary(*It++, FileStringBuilder))
 	{
-		FAnsiString& FileString = AnsiStringTable.AddDefaulted_GetRef();
+		FUtf8String& FileString = Utf8StringTable.AddDefaulted_GetRef();
 		FileString = FileStringBuilder.ToString();
-		OutLogRecord.SetFile(*FileString);
+		OutLogRecord.SetFile((const ANSICHAR*)*FileString);
 	}
 	else
 	{
@@ -501,19 +519,20 @@ bool FLogHandler::UnMarshal(FCbFieldView Field, FLogRecord& OutLogRecord,
 
 	if (bHasTextKey)
 	{
-		FLogTemplate* LogTemplate = TemplateTable.Emplace_GetRef(OutLogRecord.GetTextNamespace(), OutLogRecord.GetTextKey(), OutLogRecord.GetFormat()).Get();
+		FLogTemplate* LogTemplate = TemplateTable.Emplace_GetRef(OutLogRecord.GetTextNamespace(), OutLogRecord.GetTextKey(), OutLogRecord.GetUtf8Format()).Get();
 		OutLogRecord.SetTemplate(LogTemplate);
 	}
 	else
 	{
-		FLogTemplate* LogTemplate = TemplateTable.Emplace_GetRef(OutLogRecord.GetFormat()).Get();
+		FLogTemplate* LogTemplate = TemplateTable.Emplace_GetRef(OutLogRecord.GetUtf8Format()).Get();
 		OutLogRecord.SetTemplate(LogTemplate);
 	}
 
 	return bOk;
 }
 
-void FLogHandler::ReportActiveLog(FReplicatedLogData&& LogData, FStringView FormatMessage,
+template <UE::CCharType CharType>
+void FLogHandler::ReportActiveLog(FReplicatedLogData&& LogData, TStringView<CharType> FormatMessage,
 	ELogVerbosity::Type Verbosity)
 {
 	if (COTFS.CookWorkerClient)
@@ -522,7 +541,7 @@ void FLogHandler::ReportActiveLog(FReplicatedLogData&& LogData, FStringView Form
 	}
 	else if (COTFS.CookDirector)
 	{
-		if (FormatMessage.StartsWith(TEXT("[CookWorker")))
+		if (FormatMessage.StartsWith(LITERAL(CharType, "[CookWorker")))
 		{
 			// Do not store logs from CookWorkers; only the CookWorker saving the package needs to store those logs.
 			return;
diff --git a/Engine/Source/Runtime/Core/Private/Logging/StructuredLog.cpp b/Engine/Source/Runtime/Core/Private/Logging/StructuredLog.cpp
--- a/Engine/Source/Runtime/Core/Private/Logging/StructuredLog.cpp
+++ b/Engine/Source/Runtime/Core/Private/Logging/StructuredLog.cpp
@@ -996,8 +996,11 @@ FDateTime FLogTime::GetUtcTime() const
 template <typename CharType>
 FORCENOINLINE static void FormatDynamicRecordMessageTo(TStringBuilderBase<CharType>& Out, const FLogRecord& Record)
 {
+	const UTF8CHAR* Utf8Format = Record.GetUtf8Format();
+	PRAGMA_DISABLE_DEPRECATION_WARNINGS;
 	const TCHAR* Format = Record.GetFormat();
-	if (UNLIKELY(!Format))
+	PRAGMA_ENABLE_DEPRECATION_WARNINGS;
+	if (UNLIKELY(!Utf8Format && !Format))
 	{
 		return;
 	}
@@ -1005,10 +1008,18 @@ FORCENOINLINE static void FormatDynamicRecordMessageTo(TStringBuilderBase<CharTy
 	const TCHAR* TextNamespace = Record.GetTextNamespace();
 	const TCHAR* TextKey = Record.GetTextKey();
 	checkf(!TextNamespace == !TextKey,
-		TEXT("Log record must have both or neither of the text namespace and text key. [[%s]]"), Format);
+		TEXT("Log record must have both or neither of the text namespace and text key. [[%s]]"),
+		Utf8Format ? StringCast<TCHAR>(Utf8Format).Get() : Format);
 
 	TOptional<FInlineLogTemplate> LocalTemplate;
-	TextKey ? LocalTemplate.Emplace(TextNamespace, TextKey, Format) : LocalTemplate.Emplace(Format);
+	if (LIKELY(Utf8Format))
+	{
+		TextKey ? LocalTemplate.Emplace(TextNamespace, TextKey, Utf8Format) : LocalTemplate.Emplace(Utf8Format);
+	}
+	else
+	{
+		TextKey ? LocalTemplate.Emplace(TextNamespace, TextKey, Format) : LocalTemplate.Emplace(Format);
+	}
 	LocalTemplate->FormatTo(Out, Record.GetFields().CreateViewIterator());
 }
 
@@ -1091,7 +1102,11 @@ void FLogRecord::ConvertToCommonLog(FUtf8StringBuilderBase& OutFormat, FCbWriter
 	}
 
 	// TODO: Process localized format strings to remove argument modifiers and convert escaped braces.
-	if (LIKELY(Format))
+	if (LIKELY(Utf8Format))
+	{
+		OutFormat.Append(Utf8Format);
+	}
+	else if (Format)
 	{
 		OutFormat.Append(Format);
 	}
@@ -1169,19 +1184,19 @@ class FLogTemplateFieldIterator
 public:
 	inline explicit FLogTemplateFieldIterator(const FLogTemplate& Template)
 		: NextOp(Template.GetOpData())
-		, NextFormat(Template.GetFormat())
+		, NextFormat(Template.GetUtf8Format())
 	{
 		++*this;
 	}
 
 	FLogTemplateFieldIterator& operator++();
 	inline explicit operator bool() const { return !!NextOp; }
-	inline const FStringView& GetName() const { return Name; }
+	inline const FUtf8StringView& GetName() const { return Name; }
 
 private:
-	FStringView Name;
+	FUtf8StringView Name;
 	const uint8* NextOp = nullptr;
-	const TCHAR* NextFormat = nullptr;
+	const UTF8CHAR* NextFormat = nullptr;
 };
 
 FLogTemplateFieldIterator& FLogTemplateFieldIterator::operator++()
@@ -1193,7 +1208,7 @@ FLogTemplateFieldIterator& FLogTemplateFieldIterator::operator++()
 		FLogTemplateOp Op = FLogTemplateOp::Load(NextOp);
 		if (Op.Code == FLogTemplateOp::OpName)
 		{
-			Name = FStringView(NextFormat + 1, Op.Value - 2);
+			Name = FUtf8StringView(NextFormat + 1, Op.Value - 2);
 			NextFormat += Op.GetSkipSize();
 			return *this;
 		}
@@ -1336,7 +1351,7 @@ FORCENOINLINE static FCbObject SerializeLogFields(
 	TCbWriter<1024> Writer;
 	Writer.BeginObject();
 
-	TArray<FAnsiString, TInlineAllocator<16>> FieldNames;
+	TArray<FUtf8StringView, TInlineAllocator<16>> FieldNames;
 
 	// Anonymous. Extract names from Template.
 	if (!Fields->Name)
@@ -1345,12 +1360,10 @@ FORCENOINLINE static FCbObject SerializeLogFields(
 		for (const FLogField* FieldsEnd = Fields + FieldCount; Fields != FieldsEnd; ++Fields, ++It)
 		{
 			check(It);
-			const auto Name = StringCast<ANSICHAR>(It.GetName().GetData(), It.GetName().Len());
-			const FAnsiStringView NameView(Name.Get(), Name.Length());
-			Fields->WriteValue(Writer.SetName(NameView), Fields->Value);
+			Fields->WriteValue(Writer.SetName(It.GetName()), Fields->Value);
 			if (bHasContext)
 			{
-				FieldNames.Emplace(NameView);
+				FieldNames.Emplace(It.GetName());
 			}
 		}
 		check(!It);
@@ -1379,7 +1392,7 @@ FORCENOINLINE static FCbObject SerializeLogFields(
 			if (!Node->bDisabledByNewerContext)
 			{
 				const FAnsiStringView NodeName = Node->Name;
-				ActiveContext.Add(!FieldNames.FindByPredicate([NodeName](FAnsiStringView Name) { return NodeName.Equals(Name); }));
+				ActiveContext.Add(!FieldNames.FindByPredicate([NodeName](FUtf8StringView Name) { return NodeName.Equals(Name); }));
 				++ContextIndex;
 			}
 		}
@@ -1414,7 +1427,7 @@ FORCENOINLINE static FLogTemplate& CreateLogTemplate(const FLogCategoryBase& Cat
 {
 	for (FLogTemplate* Template = Log.DynamicData.Template.load(std::memory_order_acquire);;)
 	{
-		if (Template && Template->GetFormat() == Log.Format)
+		if (Template && Template->GetUtf8Format() == Log.Format)
 		{
 			return *Template;
 		}
@@ -1446,7 +1459,7 @@ template <typename StaticLogRecordType>
 inline static FLogTemplate& EnsureLogTemplate(const FLogCategoryBase& Category, const StaticLogRecordType& Log, const FLogField* Fields, const int32 FieldCount)
 {
 	// Format can change on a static log record due to Live Coding.
-	if (FLogTemplate* Template = Log.DynamicData.Template.load(std::memory_order_acquire); LIKELY(Template && Template->GetFormat() == Log.Format))
+	if (FLogTemplate* Template = Log.DynamicData.Template.load(std::memory_order_acquire); LIKELY(Template && Template->GetUtf8Format() == Log.Format))
 	{
 		return *Template;
 	}
@@ -1466,7 +1479,7 @@ inline static FLogRecord CreateLogRecord(const FLogCategoryBase& Category, const
 	FLogTemplate& Template = EnsureLogTemplate(Category, Log, Fields, FieldCount);
 
 	FLogRecord Record;
-	Record.SetFormat(Log.Format);
+	Record.SetUtf8Format(Log.Format);
 	Record.SetTemplate(&Template);
 	Record.SetFields(SerializeLogFields(Template, Fields, FieldCount));
 	Record.SetFile(Log.File);
@@ -1591,7 +1604,7 @@ void FatalLogWithNoFields(const FLogCategoryBase& Category, const FStaticLocaliz
 namespace UE::Logging::Private
 {
 
-static constexpr const TCHAR* const GStaticBasicLogFormat = TEXT("{Message}");
+static constexpr const UTF8CHAR* const GStaticBasicLogFormat = UTF8TEXT("{Message}");
 
 static FLogTemplate* GetStaticBasicLogTemplate()
 {
@@ -1662,7 +1675,7 @@ static void BasicLogV(TStaticBasicLogRecordParam<TCHAR> Log, const FLogCategoryB
 		FCbObject Fields = SerializeBasicLogMessage(Log, Args);
 
 		FLogRecord Record;
-		Record.SetFormat(GStaticBasicLogFormat);
+		Record.SetUtf8Format(GStaticBasicLogFormat);
 		Record.SetTemplate(GetStaticBasicLogTemplate());
 		Record.SetFields(MoveTemp(Fields));
 		Record.SetFile(Log.File);
@@ -1705,7 +1718,7 @@ static void BasicLogV(TStaticBasicLogRecordParam<FormatCharType> Log, const FLog
 	FCbObject Fields = SerializeBasicLogMessage(Log, Args);
 
 	FLogRecord Record;
-	Record.SetFormat(GStaticBasicLogFormat);
+	Record.SetUtf8Format(GStaticBasicLogFormat);
 	Record.SetTemplate(GetStaticBasicLogTemplate());
 	Record.SetFields(MoveTemp(Fields));
 	Record.SetFile(Log.File);
@@ -1793,7 +1806,7 @@ template CORE_API void BasicFatalLogWithProgramCounter(TStaticBasicLogRecordPara
 template CORE_API void BasicFatalLogWithProgramCounter(TStaticBasicLogRecordParam<UTF8CHAR> Log, const FLogCategoryBase* Category, void* ProgramCounter, ...);
 
 #if !NO_LOGGING
-static constexpr const TCHAR* const GDebugMessageLogFormat = TEXT("{Prefix}{Message}");
+static constexpr const UTF8CHAR* const GDebugMessageLogFormat = UTF8TEXT("{Prefix}{Message}");
 
 static FLogTemplate* GetDebugMessageLogTemplate()
 {
@@ -1827,7 +1840,7 @@ void LogDebugMessageLine(const ANSICHAR* File, int32 Line, const FName& InLogNam
 	FCbObject Fields = SerializeDebugLogMessage(Prefix, Message);
 
 	FLogRecord Record;
-	Record.SetFormat(GDebugMessageLogFormat);
+	Record.SetUtf8Format(GDebugMessageLogFormat);
 	Record.SetTemplate(GetDebugMessageLogTemplate());
 	Record.SetFields(MoveTemp(Fields));
 	Record.SetFile(File);
diff --git a/Engine/Source/Runtime/Core/Private/Misc/OutputDeviceDebug.cpp b/Engine/Source/Runtime/Core/Private/Misc/OutputDeviceDebug.cpp
--- a/Engine/Source/Runtime/Core/Private/Misc/OutputDeviceDebug.cpp
+++ b/Engine/Source/Runtime/Core/Private/Misc/OutputDeviceDebug.cpp
@@ -84,10 +84,16 @@ void FOutputDeviceDebug::SerializeRecord(const UE::FLogRecord& Record)
 		Writer.AddString(ANSITEXTVIEW("Category"), WriteToUtf8String<64>(Record.GetCategory()));
 		Writer.AddString(ANSITEXTVIEW("Verbosity"), ToString(Record.GetVerbosity()));
 		Writer.AddString(ANSITEXTVIEW("Message"), Message);
-		if (const TCHAR* Format = Record.GetFormat())
+		PRAGMA_DISABLE_DEPRECATION_WARNINGS;
+		if (const UTF8CHAR* Utf8Format = Record.GetUtf8Format())
+		{
+			Writer.AddString(ANSITEXTVIEW("Format"), Utf8Format);
+		}
+		else if (const TCHAR* Format = Record.GetFormat())
 		{
 			Writer.AddString(ANSITEXTVIEW("Format"), Format);
 		}
+		PRAGMA_ENABLE_DEPRECATION_WARNINGS;
 		if (const ANSICHAR* File = Record.GetFile())
 		{
 			Writer.AddString(ANSITEXTVIEW("File"), File);
diff --git a/Engine/Source/Runtime/Core/Public/Logging/StructuredLog.h b/Engine/Source/Runtime/Core/Public/Logging/StructuredLog.h
--- a/Engine/Source/Runtime/Core/Public/Logging/StructuredLog.h
+++ b/Engine/Source/Runtime/Core/Public/Logging/StructuredLog.h
@@ -99,7 +99,7 @@ template <typename FuncType> class TFunctionRef;
  * Records a localized structured log event if this category is active at this level of verbosity.
  *
  * Example:
- * UE_LOGFMT_LOC_EX(LogCore, Warning, "LoadingFailed", "Loading '{PackageName}' failed with error {Error}",
+ * UE_LOGFMT_LOC_EX(LogCore, Warning, "LoadingFailed", "Loading '{Name}' failed with error {Error}",
  *     UE_LOGFMT_FIELD("Name", Package->GetName()), UE_LOGFMT_FIELD("Error", ErrorCode), UE_LOGFMT_FIELD("Flags", LoadFlags));
  *
  * Same as UE_LOGFMT_LOC and works for any number of fields.
@@ -194,9 +194,15 @@ public:
 	void SetTime(const FLogTime& InTime) { Time = InTime; }
 
 	/** The format string that serves as the message for the log record. Example: TEXT("FieldName is {FieldName}") */
+	UE_DEPRECATED(5.8, "Use GetUtf8Format.")
 	const TCHAR* GetFormat() const { return Format; }
+	UE_DEPRECATED(5.8, "Use SetUtf8Format.")
 	void SetFormat(const TCHAR* InFormat) { Format = InFormat; }
 
+	/** The format string that serves as the message for the log record. Example: UTF8TEXT("FieldName is {FieldName}") */
+	const UTF8CHAR* GetUtf8Format() const { return Utf8Format; }
+	void SetUtf8Format(const UTF8CHAR* InFormat) { Utf8Format = InFormat; }
+
 	/** The optional template for the format string. */
 	const FLogTemplate* GetTemplate() const { return Template; }
 	void SetTemplate(const FLogTemplate* InTemplate) { Template = InTemplate; }
@@ -242,6 +248,7 @@ public:
 
 private:
 	const TCHAR* Format = nullptr;
+	const UTF8CHAR* Utf8Format = nullptr;
 	const ANSICHAR* File = nullptr;
 	int32 Line = 0;
 	FName Category;
@@ -336,7 +343,7 @@ struct FStaticLogDynamicData
 /** Data about a static log that is constant for every occurrence. */
 struct FStaticLogRecord
 {
-	const TCHAR* Format = nullptr;
+	const UTF8CHAR* Format = nullptr;
 	UE_IF(UE_LOG_INCLUDE_SOURCE_LOCATION,, inline static constexpr) const ANSICHAR* File = nullptr;
 	UE_IF(UE_LOG_INCLUDE_SOURCE_LOCATION,, inline static constexpr) int32 Line = 0;
 	ELogVerbosity::Type Verbosity = ELogVerbosity::Log;
@@ -344,7 +351,7 @@ struct FStaticLogRecord
 
 	// Workaround for https://developercommunity.visualstudio.com/t/Incorrect-warning-C4700-with-unrelated-s/10285950
 	constexpr FStaticLogRecord(
-		const TCHAR* InFormat,
+		const UTF8CHAR* InFormat,
 		const ANSICHAR* InFile,
 		int32 InLine,
 		ELogVerbosity::Type InVerbosity,
@@ -365,7 +372,7 @@ struct FStaticLocalizedLogRecord
 {
 	const TCHAR* TextNamespace = nullptr;
 	const TCHAR* TextKey = nullptr;
-	const TCHAR* Format = nullptr;
+	const UTF8CHAR* Format = nullptr;
 	UE_IF(UE_LOG_INCLUDE_SOURCE_LOCATION,, inline static constexpr) const ANSICHAR* File = nullptr;
 	UE_IF(UE_LOG_INCLUDE_SOURCE_LOCATION,, inline static constexpr) int32 Line = 0;
 	ELogVerbosity::Type Verbosity = ELogVerbosity::Log;
@@ -375,7 +382,7 @@ struct FStaticLocalizedLogRecord
 	constexpr FStaticLocalizedLogRecord(
 		const TCHAR* InTextNamespace,
 		const TCHAR* InTextKey,
-		const TCHAR* InFormat,
+		const UTF8CHAR* InFormat,
 		const ANSICHAR* InFile,
 		int32 InLine,
 		ELogVerbosity::Type InVerbosity,
@@ -548,7 +555,7 @@ void LogDebugMessageLine(const ANSICHAR* File, int32 Line, const FName& InLogNam
 			 (::ELogVerbosity::Verbosity & ::ELogVerbosity::VerbosityMask) <= UE_PRIVATE_LOG_CATEGORY(CategoryName).CompileTimeVerbosity)) \
 		{ \
 			static ::UE::Logging::Private::FStaticLogDynamicData LOG_Dynamic; \
-			static constexpr ::UE::Logging::Private::FStaticLogRecord LOG_Static UE_PRIVATE_LOGFMT_AGGREGATE(TEXT(Format), __builtin_FILE(), __builtin_LINE(), ::ELogVerbosity::Verbosity, LOG_Dynamic); \
+			static constexpr ::UE::Logging::Private::FStaticLogRecord LOG_Static UE_PRIVATE_LOGFMT_AGGREGATE(UTF8TEXT(Format), __builtin_FILE(), __builtin_LINE(), ::ELogVerbosity::Verbosity, LOG_Dynamic); \
 			UE_PRIVATE_LOGFMT_LOG_IF_ACTIVE(Condition, CategoryName, Verbosity, LOG_Static, ##__VA_ARGS__); \
 		} \
 	} \
@@ -562,7 +569,7 @@ void LogDebugMessageLine(const ANSICHAR* File, int32 Line, const FName& InLogNam
 			 (::ELogVerbosity::Verbosity & ::ELogVerbosity::VerbosityMask) <= UE_PRIVATE_LOG_CATEGORY(CategoryName).CompileTimeVerbosity)) \
 		{ \
 			static ::UE::Logging::Private::FStaticLogDynamicData LOG_Dynamic; \
-			static constexpr ::UE::Logging::Private::FStaticLocalizedLogRecord LOG_Static UE_PRIVATE_LOGFMT_AGGREGATE(TEXT(Namespace), TEXT(Key), TEXT(Format), __builtin_FILE(), __builtin_LINE(), ::ELogVerbosity::Verbosity, LOG_Dynamic); \
+			static constexpr ::UE::Logging::Private::FStaticLocalizedLogRecord LOG_Static UE_PRIVATE_LOGFMT_AGGREGATE(TEXT(Namespace), TEXT(Key), UTF8TEXT(Format), __builtin_FILE(), __builtin_LINE(), ::ELogVerbosity::Verbosity, LOG_Dynamic); \
 			UE_PRIVATE_LOGFMT_LOG_IF_ACTIVE(Condition, CategoryName, Verbosity, LOG_Static, ##__VA_ARGS__); \
 		} \
 	} \
diff --git a/Engine/Source/Runtime/CoreUObject/Private/UObject/UObjectGlobals.cpp b/Engine/Source/Runtime/CoreUObject/Private/UObject/UObjectGlobals.cpp
--- a/Engine/Source/Runtime/CoreUObject/Private/UObject/UObjectGlobals.cpp
+++ b/Engine/Source/Runtime/CoreUObject/Private/UObject/UObjectGlobals.cpp
@@ -6042,7 +6042,7 @@ void UE::Core::Private::RecordAssetLog(
 	Writer.EndObject();
 
 	FLogRecord Record;
-	Record.SetFormat(TEXT("[AssetLog] {Asset}: {Message}"));
+	Record.SetUtf8Format(UTF8TEXT("[AssetLog] {Asset}: {Message}"));
 	Record.SetFields(Writer.Save().AsObject());
 	Record.SetFile(File);
 	Record.SetLine(Line);
