[Core] Remove unused feature to auto relax map after iteration if elements were removed
* Cleaned up unsupported features in TCompactSet and allow for conditional compiling in TMap

#tests Compiled, preflights
#Synced-CL 48466649

diff --git a/Engine/Plugins/Enterprise/DataprepEditor/Source/DataprepEditor/Private/PreviewSystem/DataprepPreviewSystem.cpp b/Engine/Plugins/Enterprise/DataprepEditor/Source/DataprepEditor/Private/PreviewSystem/DataprepPreviewSystem.cpp
--- a/Engine/Plugins/Enterprise/DataprepEditor/Source/DataprepEditor/Private/PreviewSystem/DataprepPreviewSystem.cpp
+++ b/Engine/Plugins/Enterprise/DataprepEditor/Source/DataprepEditor/Private/PreviewSystem/DataprepPreviewSystem.cpp
@@ -310,7 +310,7 @@ void FDataprepPreviewSystem::IncrementalProcess()
 		{
 			CurrentProgress.CurrentFilterIndex++;
 			CurrentProgress.CurrentObjectProcessed = 0;
-			CurrentProgress.Iterator = MakeUnique<FResultIterator>( PreviewResult, false );
+			CurrentProgress.Iterator = MakeUnique<FResultIterator>( PreviewResult );
 		}
 	}
 }
@@ -324,11 +324,11 @@ void FDataprepPreviewSystem::RestartProcessing()
 
 	if ( PreviewResult.Num() > 0 )
 	{
-		CurrentProgress.Iterator = MakeUnique<FResultIterator>( PreviewResult, false );
+		CurrentProgress.Iterator = MakeUnique<FResultIterator>( PreviewResult );
 
 		if ( bResetProcessStatus )
 		{
-			FResultIterator ResultIterator( PreviewResult, false );
+			FResultIterator ResultIterator( PreviewResult );
 			while ( ResultIterator )
 			{
 				FDataprepPreviewProcessingResult& Result = ResultIterator.Value().Get();
diff --git a/Engine/Source/Runtime/Core/Private/Containers/ContainersTest.cpp b/Engine/Source/Runtime/Core/Private/Containers/ContainersTest.cpp
--- a/Engine/Source/Runtime/Core/Private/Containers/ContainersTest.cpp
+++ b/Engine/Source/Runtime/Core/Private/Containers/ContainersTest.cpp
@@ -814,21 +814,24 @@ namespace
 			check(Map.Remove(Key1) == 1);
 			check(Map.Remove(Key2) == 1);
 
-			// Re-add the 2 elements in the same order. Without the call to SortFreeList()
-			// the elements would end up in a different locations/order compared to the
-			// original insertions. SortFreeList() should ensure that re-adding the elements
-			// gives use the same container layout as long as we perform the same operations
-			// in the same order as before.
-			Map.SortFreeList();
-			Map.Add(Key1, Value);
-			Map.Add(Key2, Value);
-			const FContainerTestValueType* NewValuePtr0 = Map.Find(Key0);
-			const FContainerTestValueType* NewValuePtr1 = Map.Find(Key1);
-			const FContainerTestValueType* NewValuePtr2 = Map.Find(Key2);
-
-			check(ValuePtr0 == NewValuePtr0);
-			check(ValuePtr1 == NewValuePtr1);
-			check(ValuePtr2 == NewValuePtr2);
+			if constexpr (requires { Map.SortFreeList(); })
+			{
+				// Re-add the 2 elements in the same order. Without the call to SortFreeList()
+				// the elements would end up in a different locations/order compared to the
+				// original insertions. SortFreeList() should ensure that re-adding the elements
+				// gives use the same container layout as long as we perform the same operations
+				// in the same order as before.
+				Map.SortFreeList();
+				Map.Add(Key1, Value);
+				Map.Add(Key2, Value);
+				const FContainerTestValueType* NewValuePtr0 = Map.Find(Key0);
+				const FContainerTestValueType* NewValuePtr1 = Map.Find(Key1);
+				const FContainerTestValueType* NewValuePtr2 = Map.Find(Key2);
+
+				check(ValuePtr0 == NewValuePtr0);
+				check(ValuePtr1 == NewValuePtr1);
+				check(ValuePtr2 == NewValuePtr2);
+			}
 		}
 	}
 
diff --git a/Engine/Source/Runtime/Core/Public/Containers/CompactSet.h.inl b/Engine/Source/Runtime/Core/Public/Containers/CompactSet.h.inl
--- a/Engine/Source/Runtime/Core/Public/Containers/CompactSet.h.inl
+++ b/Engine/Source/Runtime/Core/Public/Containers/CompactSet.h.inl
@@ -536,25 +536,12 @@ public:
 		}
 	}
 
-	/** Deprecated - default behavior now, keeping this here so UE_TCOMPACT_SET can be swapped with TSet without changing code */
+	/** Deprecated */
 	void Compact()
 	{
-	}
-
-	/** Deprecated - use sparse array if this behavior is required, keeping this here so UE_TCOMPACT_SET can be swapped with TSet without changing code */
-	void CompactStable()
-	{
-		ensureMsgf(false, TEXT("Compact sets are always compact so CompactStable will not do anything. If you hit this then you likely need to use a different pattern to maintain order, see RemoveStable"));
-	}
-
-	/** Deprecated - unnecessary, keeping this here so UE_TCOMPACT_SET can be swapped with TSet without changing code */
-	void SortFreeList()
-	{
-	}
-
-	/** Deprecated - default behavior now, keeping this here so UE_TCOMPACT_SET can be swapped with TSet without changing code */
-	void Relax()
-	{
+		// Data is always compacted, but this is useful to be able to more easily swap between sparse and compact sets
+		// We can at least perform a shrink here to maintain memory parity with sparse set
+		Shrink();
 	}
 
 	/** 
diff --git a/Engine/Source/Runtime/Core/Public/Containers/Map.h.inl b/Engine/Source/Runtime/Core/Public/Containers/Map.h.inl
--- a/Engine/Source/Runtime/Core/Public/Containers/Map.h.inl
+++ b/Engine/Source/Runtime/Core/Public/Containers/Map.h.inl
@@ -37,6 +37,8 @@ public:
 	typedef TPair<KeyType, ValueType> ElementType;
 
 protected:
+	using ElementSetType = TSET<ElementType, KeyFuncs, SetAllocator>;
+
 	[[nodiscard]] constexpr TMAPBASE() = default;
 	[[nodiscard]] explicit consteval TMAPBASE(EConstEval)
 	: Pairs(ConstEval)
@@ -168,7 +170,7 @@ public:
 	}
 
 	/** Compacts the pair set to remove holes. Does not change the iteration order of the elements. */
-	UE_FORCEINLINE_HINT void CompactStable()
+	UE_FORCEINLINE_HINT void CompactStable() requires (&ElementSetType::CompactStable != nullptr)
 	{
 		Pairs.CompactStable();
 	}
@@ -761,8 +763,6 @@ public:
 	}
 
 protected:
-	using ElementSetType = TSET<ElementType, KeyFuncs, SetAllocator>;
-
 	/** The base of TMAPBASE iterators. */
 	template<bool bConst>
 	class TBaseIterator
@@ -921,7 +921,16 @@ public:
 	public:
 
 		/** Initialization constructor. */
-		[[nodiscard]] inline TIterator(TMAPBASE& InMap, bool bInRequiresRehashOnRemoval = false)
+		[[nodiscard]] inline TIterator(TMAPBASE& InMap)
+		: TBaseIterator<false>(InMap.Pairs.CreateIterator())
+		, Map(InMap)
+		, bElementsHaveBeenRemoved(false)
+		, bRequiresRehashOnRemoval(false)
+		{
+		}
+
+		UE_DEPRECATED(5.8, "TSet::Relax has been deprecated. Memory updates post iteration should be handled manually as necessary.")
+		[[nodiscard]] inline TIterator(TMAPBASE& InMap, bool bInRequiresRehashOnRemoval /* = false */)
 		: TBaseIterator<false>(InMap.Pairs.CreateIterator())
 		, Map(InMap)
 		, bElementsHaveBeenRemoved(false)
@@ -932,10 +941,15 @@ public:
 		/** Destructor. */
 		inline ~TIterator()
 		{
-			if (bElementsHaveBeenRemoved && bRequiresRehashOnRemoval)
-			{
-				Map.Pairs.Relax();
-			}
+		    if constexpr (requires { Map.Pairs.Relax(); })
+		    {
+			    if (bElementsHaveBeenRemoved && bRequiresRehashOnRemoval)
+			    {
+					PRAGMA_DISABLE_DEPRECATION_WARNINGS
+					Map.Pairs.Relax();
+					PRAGMA_ENABLE_DEPRECATION_WARNINGS
+			    }
+		    }
 		}
 
 		/** Removes the current pair from the map without losing the iteration position.
@@ -1158,7 +1172,7 @@ public:
 	 * TSet index resulting in tighter packing without moving any existing items. Also useful for
 	 * some types of determinism. @see TSparseArray::SortFreeList() for more info.
 	 */
-	void SortFreeList()
+	void SortFreeList() requires (&Super::ElementSetType::SortFreeList != nullptr)
 	{
 		Super::Pairs.SortFreeList();
 	}
diff --git a/Engine/Source/Runtime/Core/Public/Containers/SparseSet.h.inl b/Engine/Source/Runtime/Core/Public/Containers/SparseSet.h.inl
--- a/Engine/Source/Runtime/Core/Public/Containers/SparseSet.h.inl
+++ b/Engine/Source/Runtime/Core/Public/Containers/SparseSet.h.inl
@@ -257,7 +257,7 @@ public:
 	inline void Shrink()
 	{
 		Elements.Shrink();
-		Relax();
+		ConditionalRehash(Elements.Num(), EAllowShrinking::Yes);
 	}
 
 	/** Compacts the allocated elements into a contiguous range. */
@@ -310,6 +310,7 @@ public:
 	}
 
 	/** Relaxes the set's hash to a size strictly bounded by the number of elements in the set. */
+	UE_DEPRECATED(5.8, "Removing support for relax, use Shrink instead for better memory savings.")
 	UE_FORCEINLINE_HINT void Relax()
 	{
 		ConditionalRehash(Elements.Num(), EAllowShrinking::Yes);
