[IOS] NSLog in non-shipping builds

48282166 switched Apple logging to use os_log, but unfortunately that doesn't go to stdout, so that is invisible to "xcrun devicectl process ... --console" which is required for automation to access the "live" log.

This brings back the issue of logs being private by default in IOS26, as we can't control the visibility of NSLog.
Kept os_log for Shipping as that is faster and we do not want to go to stdout there, anyway.

#rnx

diff --git a/Engine/Source/Runtime/Core/Private/Apple/ApplePlatformMisc.cpp b/Engine/Source/Runtime/Core/Private/Apple/ApplePlatformMisc.cpp
--- a/Engine/Source/Runtime/Core/Private/Apple/ApplePlatformMisc.cpp
+++ b/Engine/Source/Runtime/Core/Private/Apple/ApplePlatformMisc.cpp
@@ -77,11 +77,13 @@ FString FApplePlatformMisc::GetEnvironmentVariable(const TCHAR* VariableName)
 
 void FApplePlatformMisc::LocalPrint(const TCHAR* Message)
 {
-	// This goes to all output consoles, not only Xcode one.
 #if !UE_BUILD_SHIPPING
-	os_log(OS_LOG_DEFAULT, "[UE] %{public}s", TCHAR_TO_UTF8(Message));
+	// This goes to all output consoles, not only system one, stdout included.
+	NSLog(@"[UE] %s", TCHAR_TO_UTF8(Message));
 #else
-	os_log(OS_LOG_DEFAULT, "[UE] %s", TCHAR_TO_UTF8(Message)); // Private by default, would be visible in Xcode.
+	// This only goes to the system console.
+	// Private by default, would be visible in Xcode.
+	os_log(OS_LOG_DEFAULT, "[UE] %s", TCHAR_TO_UTF8(Message));
 #endif
 }
 
