Exposed utils to import/export a string table from CSV to BP/Python

#Synced-CL 48448085

diff --git a/Engine/Source/Editor/StringTableEditor/Private/StringTableEditor.cpp b/Engine/Source/Editor/StringTableEditor/Private/StringTableEditor.cpp
--- a/Engine/Source/Editor/StringTableEditor/Private/StringTableEditor.cpp
+++ b/Engine/Source/Editor/StringTableEditor/Private/StringTableEditor.cpp
@@ -672,7 +672,7 @@ FReply FStringTableEditor::OnImportFromCSVClicked()
 			{
 				const FScopedTransaction Transaction(LOCTEXT("ImportStringTableEntries", "Import String Table Entries"));
 				StringTable->Modify();
-				StringTable->GetMutableStringTable()->ImportStrings(OutFiles[0]);
+				StringTable->GetMutableStringTable()->ImportStringsFromCSVFile(OutFiles[0]);
 				HandlePostChange();
 			}
 		}
@@ -708,7 +708,7 @@ FReply FStringTableEditor::OnExportToCSVClicked()
 			UStringTable* StringTable = Cast<UStringTable>(GetEditingObject());
 			if (StringTable)
 			{
-				StringTable->GetStringTable()->ExportStrings(OutFiles[0]);
+				StringTable->GetStringTable()->ExportStringsToCSVFile(OutFiles[0]);
 			}
 		}
 	}
diff --git a/Engine/Source/Editor/UnrealEd/Private/Commandlets/GatherTextFromSourceCommandlet.cpp b/Engine/Source/Editor/UnrealEd/Private/Commandlets/GatherTextFromSourceCommandlet.cpp
--- a/Engine/Source/Editor/UnrealEd/Private/Commandlets/GatherTextFromSourceCommandlet.cpp
+++ b/Engine/Source/Editor/UnrealEd/Private/Commandlets/GatherTextFromSourceCommandlet.cpp
@@ -1935,7 +1935,7 @@ void UGatherTextFromSourceCommandlet::FSourceFileParseContext::AddStringTableFro
 		const FString FullImportPath = InRootPath / InTableFilename;
 
 		FStringTableRef TmpStringTable = FStringTable::NewStringTable();
-		if (TmpStringTable->ImportStrings(FullImportPath))
+		if (TmpStringTable->ImportStringsFromCSVFile(FullImportPath))
 		{
 			const FSourceLocation SourceLocation = FSourceLocation(InTableFilename, INDEX_NONE);
 			const FName TablePlatformName = GetSplitPlatformNameFromPath_Static(InTableFilename, SplitPlatforms);
diff --git a/Engine/Source/Runtime/Core/Private/Internationalization/StringTableCore.cpp b/Engine/Source/Runtime/Core/Private/Internationalization/StringTableCore.cpp
--- a/Engine/Source/Runtime/Core/Private/Internationalization/StringTableCore.cpp
+++ b/Engine/Source/Runtime/Core/Private/Internationalization/StringTableCore.cpp
@@ -409,6 +409,11 @@ void FStringTable::Serialize(FArchive& Ar)
 }
 
 bool FStringTable::ExportStrings(const FString& InFilename) const
+{
+	return ExportStringsToCSVFile(InFilename);
+}
+
+FString FStringTable::ExportStringsToCSVString() const
 {
 	FString ExportedStrings;
 
@@ -471,25 +476,29 @@ bool FStringTable::ExportStrings(const FString& InFilename) const
 		}
 	}
 
+	return ExportedStrings;
+}
+
+bool FStringTable::ExportStringsToCSVFile(const FString& InFilename) const
+{
+	FString ExportedStrings = ExportStringsToCSVString();
 	return FFileHelper::SaveStringToFile(ExportedStrings, *InFilename);
 }
 
 bool FStringTable::ImportStrings(const FString& InFilename)
 {
-	FString ImportedStrings;
-	if (!FFileHelper::LoadFileToString(ImportedStrings, *InFilename))
-	{
-		UE_LOG(LogStringTable, Warning, TEXT("Failed to import string table from '%s'. Could not open file."), *InFilename);
-		return false;
-	}
+	return ImportStringsFromCSVFile(InFilename);
+}
 
-	const FCsvParser ImportedStringsParser(ImportedStrings);
+bool FStringTable::ImportStringsFromCSVString(const FString& InCSVString, const TCHAR* InErrorContext)
+{
+	const FCsvParser ImportedStringsParser(InCSVString);
 	const FCsvParser::FRows& Rows = ImportedStringsParser.GetRows();
 
 	// Must have at least 2 rows (header and content)
 	if (Rows.Num() <= 1)
 	{
-		UE_LOG(LogStringTable, Warning, TEXT("Failed to import string table from '%s'. Incorrect number of rows (must be at least 2)."), *InFilename);
+		UE_LOG(LogStringTable, Warning, TEXT("Failed to import string table from '%s'. Incorrect number of rows (must be at least 2)."), InErrorContext);
 		return false;
 	}
 
@@ -526,12 +535,12 @@ bool FStringTable::ImportStrings(const FString& InFilename)
 		if (KeyColumn == INDEX_NONE)
 		{
 			bValidHeader = false;
-			UE_LOG(LogStringTable, Warning, TEXT("Failed to import string table from '%s'. Failed to find required column 'Key'."), *InFilename);
+			UE_LOG(LogStringTable, Warning, TEXT("Failed to import string table from '%s'. Failed to find required column 'Key'."), InErrorContext);
 		}
 		if (SourceStringColumn == INDEX_NONE)
 		{
 			bValidHeader = false;
-			UE_LOG(LogStringTable, Warning, TEXT("Failed to import string table from '%s'. Failed to find required column 'SourceString'."), *InFilename);
+			UE_LOG(LogStringTable, Warning, TEXT("Failed to import string table from '%s'. Failed to find required column 'SourceString'."), InErrorContext);
 		}
 		if (!bValidHeader)
 		{
@@ -582,6 +591,17 @@ bool FStringTable::ImportStrings(const FString& InFilename)
 	return true;
 }
 
+bool FStringTable::ImportStringsFromCSVFile(const FString& InFilename)
+{
+	FString ImportedStrings;
+	if (!FFileHelper::LoadFileToString(ImportedStrings, *InFilename))
+	{
+		UE_LOG(LogStringTable, Warning, TEXT("Failed to import string table from '%s'. Could not open file."), *InFilename);
+		return false;
+	}
+	return ImportStringsFromCSVString(ImportedStrings, *InFilename);
+}
+
 
 void FStringTableRedirects::InitStringTableRedirects()
 {
diff --git a/Engine/Source/Runtime/Core/Private/Internationalization/StringTableRegistry.cpp b/Engine/Source/Runtime/Core/Private/Internationalization/StringTableRegistry.cpp
--- a/Engine/Source/Runtime/Core/Private/Internationalization/StringTableRegistry.cpp
+++ b/Engine/Source/Runtime/Core/Private/Internationalization/StringTableRegistry.cpp
@@ -176,7 +176,7 @@ void FStringTableRegistry::OnDirectoryChanged(const TArray<FFileChangeData>& InF
 				FStringTablePtr StringTable = FindMutableStringTable(TableId);
 				if (StringTable.IsValid())
 				{
-					if (!StringTable->ImportStrings(FileChange.Filename))
+					if (!StringTable->ImportStringsFromCSVFile(FileChange.Filename))
 					{
 						UE_LOG(LogStringTable, Warning, TEXT("Failed to import strings from '%s'"), *FileChange.Filename);
 					}
@@ -201,7 +201,7 @@ void FStringTableRegistry::Internal_LocTableFromFile(const FName InTableId, cons
 	StringTable->SetNamespace(InNamespace);
 
 	const FString FullPath = InRootPath / InFilePath;
-	if (!StringTable->ImportStrings(FullPath))
+	if (!StringTable->ImportStringsFromCSVFile(FullPath))
 	{
 		UE_LOG(LogStringTable, Warning, TEXT("Failed to import strings from '%s'"), *FullPath);
 	}
diff --git a/Engine/Source/Runtime/Core/Public/Internationalization/StringTableCore.h b/Engine/Source/Runtime/Core/Public/Internationalization/StringTableCore.h
--- a/Engine/Source/Runtime/Core/Public/Internationalization/StringTableCore.h
+++ b/Engine/Source/Runtime/Core/Public/Internationalization/StringTableCore.h
@@ -161,11 +161,25 @@ public:
 	CORE_API void Serialize(FArchive& Ar);
 
 	/** Export the key, string, and meta-data information in this string table to a CSV file (does not export the namespace) */
+	UE_DEPRECATED(5.8, "ExportStrings is deprecated in favor of ExportStringsToCSVFile")
 	CORE_API bool ExportStrings(const FString& InFilename) const;
 
+	/** Export the key, string, and meta-data information in this string table to a CSV string (does not export the namespace) */
+	CORE_API FString ExportStringsToCSVString() const;
+
+	/** Export the key, string, and meta-data information in this string table to a CSV file (does not export the namespace) */
+	CORE_API bool ExportStringsToCSVFile(const FString& InFilename) const;
+
 	/** Import key, string, and meta-data information from a CSV file to this string table (does not import the namespace) */
+	UE_DEPRECATED(5.8, "ImportStrings is deprecated in favor of ImportStringsFromCSVFile")
 	CORE_API bool ImportStrings(const FString& InFilename);
 
+	/** Import key, string, and meta-data information from a CSV string to this string table (does not import the namespace) */
+	CORE_API bool ImportStringsFromCSVString(const FString& InCSVString, const TCHAR* InErrorContext = TEXT("<unknown>"));
+
+	/** Import key, string, and meta-data information from a CSV file to this string table (does not import the namespace) */
+	CORE_API bool ImportStringsFromCSVFile(const FString& InFilename);
+
 private:
 	/** Pointer back to the asset that owns this table */
 	UStringTable* OwnerAsset;
diff --git a/Engine/Source/Runtime/Engine/Classes/Kismet/KismetStringTableLibrary.h b/Engine/Source/Runtime/Engine/Classes/Kismet/KismetStringTableLibrary.h
--- a/Engine/Source/Runtime/Engine/Classes/Kismet/KismetStringTableLibrary.h
+++ b/Engine/Source/Runtime/Engine/Classes/Kismet/KismetStringTableLibrary.h
@@ -8,10 +8,15 @@
 #include "Kismet/BlueprintFunctionLibrary.h"
 #include "KismetStringTableLibrary.generated.h"
 
+class UStringTable;
+
 UCLASS(meta=(BlueprintThreadSafe, ScriptName="StringTableLibrary"), MinimalAPI)
 class UKismetStringTableLibrary : public UBlueprintFunctionLibrary
 {
-	GENERATED_UCLASS_BODY()
+	GENERATED_BODY()
+
+	UFUNCTION(BlueprintPure, Category="Utilities|String Table", meta=(DisplayName="Get String Table ID"))
+	static ENGINE_API FName GetTableId(const UStringTable* StringTable);
 
 	/** Returns true if the given table ID corresponds to a registered string table. */
 	UFUNCTION(BlueprintPure, Category="Utilities|String Table", meta=(DisplayName="Is String Table Registered"))
@@ -44,4 +49,20 @@ class UKismetStringTableLibrary : public UBlueprintFunctionLibrary
 	/** Returns an array of all meta-data IDs within the given string table entry */
 	UFUNCTION(BlueprintPure, Category="Utilities|String Table", meta=(DisplayName="Get Meta-Data IDs from String Table Entry"))
 	static ENGINE_API TArray<FName> GetMetaDataIdsFromStringTableEntry(const FName TableId, const FString& Key);
+
+	/** Export the key, string, and meta-data information in the given string table to a CSV string (does not export the namespace) */
+	UFUNCTION(BlueprintCallable, Category="Utilities|String Table", meta=(DisplayName="Export String Table to CSV String"))
+	static ENGINE_API bool ExportTableToCSVString(const FName TableId, FString& OutCSVString);
+
+	/** Export the key, string, and meta-data information in the given string table to a CSV file (does not export the namespace) */
+	UFUNCTION(BlueprintCallable, Category="Utilities|String Table", meta=(DisplayName="Export String Table to CSV File"))
+	static ENGINE_API bool ExportTableToCSVFile(const FName TableId, const FString& Filename);
+
+	/** Import key, string, and meta-data information from a CSV string to the given string table (does not import the namespace) */
+	UFUNCTION(BlueprintCallable, Category="Utilities|String Table", meta=(DisplayName="Import String Table from CSV String"))
+	static ENGINE_API bool ImportTableFromCSVString(const FName TableId, const FString& CSVString);
+
+	/** Import key, string, and meta-data information from a CSV file to the given string table (does not import the namespace) */
+	UFUNCTION(BlueprintCallable, Category="Utilities|String Table", meta=(DisplayName="Import String Table from CSV File"))
+	static ENGINE_API bool ImportTableFromCSVFile(const FName TableId, const FString& Filename);
 };
diff --git a/Engine/Source/Runtime/Engine/Private/KismetStringTableLibrary.cpp b/Engine/Source/Runtime/Engine/Private/KismetStringTableLibrary.cpp
--- a/Engine/Source/Runtime/Engine/Private/KismetStringTableLibrary.cpp
+++ b/Engine/Source/Runtime/Engine/Private/KismetStringTableLibrary.cpp
@@ -1,6 +1,7 @@
 // Copyright Epic Games, Inc. All Rights Reserved.
 
 #include "Kismet/KismetStringTableLibrary.h"
+#include "Internationalization/StringTable.h"
 #include "Internationalization/StringTableCore.h"
 #include "Internationalization/StringTableRegistry.h"
 
@@ -8,9 +9,9 @@
 
 #define LOCTEXT_NAMESPACE "Kismet"
 
-UKismetStringTableLibrary::UKismetStringTableLibrary(const FObjectInitializer& ObjectInitializer)
-	: Super(ObjectInitializer)
+FName UKismetStringTableLibrary::GetTableId(const UStringTable* StringTable)
 {
+	return StringTable ? StringTable->GetStringTableId() : FName();
 }
 
 
@@ -115,5 +116,50 @@ TArray<FName> UKismetStringTableLibrary::GetMetaDataIdsFromStringTableEntry(cons
 	return MetaDataIdsFromStringTable;
 }
 
+
+bool UKismetStringTableLibrary::ExportTableToCSVString(const FName TableId, FString& OutCSVString)
+{
+	if (FStringTableConstPtr StringTable = FStringTableRegistry::Get().FindStringTable(TableId))
+	{
+		OutCSVString = StringTable->ExportStringsToCSVString();
+		return true;
+	}
+
+	return false;
+}
+
+
+bool UKismetStringTableLibrary::ExportTableToCSVFile(const FName TableId, const FString& Filename)
+{
+	if (FStringTableConstPtr StringTable = FStringTableRegistry::Get().FindStringTable(TableId))
+	{
+		return StringTable->ExportStringsToCSVFile(Filename);
+	}
+
+	return false;
+}
+
+
+bool UKismetStringTableLibrary::ImportTableFromCSVString(const FName TableId, const FString& CSVString)
+{
+	if (FStringTablePtr StringTable = FStringTableRegistry::Get().FindMutableStringTable(TableId))
+	{
+		return StringTable->ImportStringsFromCSVString(CSVString, *TableId.ToString());
+	}
+
+	return false;
+}
+
+
+bool UKismetStringTableLibrary::ImportTableFromCSVFile(const FName TableId, const FString& Filename)
+{
+	if (FStringTablePtr StringTable = FStringTableRegistry::Get().FindMutableStringTable(TableId))
+	{
+		return StringTable->ImportStringsFromCSVFile(Filename);
+	}
+
+	return false;
+}
+
 #undef LOCTEXT_NAMESPACE
 
