Correct handling of chained unwind entries.

Function may be split into parts and scattered about a code segment. This means that for a given program counter, the prologue that establishes the stack frame may be elsewhere. Unwind info handles this by chaining info together. I had previously misunderstood this and the handling of chains was incorrect; they do not reference and contribute to the of the prior entry. Instead the relevant stack frame for this region of code segment is found elsewhere.

#rnx

#Synced-CL Unknown

diff --git a/Engine/Source/Runtime/Core/Private/ProfilingDebugging/Microsoft/WindowsCallstackTrace.cpp b/Engine/Source/Runtime/Core/Private/ProfilingDebugging/Microsoft/WindowsCallstackTrace.cpp
--- a/Engine/Source/Runtime/Core/Private/ProfilingDebugging/Microsoft/WindowsCallstackTrace.cpp
+++ b/Engine/Source/Runtime/Core/Private/ProfilingDebugging/Microsoft/WindowsCallstackTrace.cpp
@@ -342,6 +342,17 @@ void FBacktracer::AddModule(UPTRINT ModuleBase, const TCHAR* Name)
 		UPTRINT UnwindInfoAddr = ModuleBase + FunctionTable->UnwindInfoAddress;
 		const auto* UnwindInfo = (FUnwindInfo*)UnwindInfoAddr;
 
+		// Functions can end up split into parts and spread about code segments. In
+		// such cases there is a chain of unwind info which must be followed to find
+		// the primary info that describes the prologue.
+		while (UnwindInfo->Flags & UNW_FLAG_CHAININFO)
+		{
+			uint32 CodeStep = (UnwindInfo->NumUnwindCodes + 1) & ~1;
+			const void* OuterAddr = (FUnwindCode*)(UnwindInfo + 1) + CodeStep;
+			const RUNTIME_FUNCTION* Outer = (RUNTIME_FUNCTION*)OuterAddr;
+			UnwindInfo = (FUnwindInfo*)(ModuleBase + Outer->UnwindInfoAddress);
+		}
+
 		if (UnwindInfo->Version != 1)
 		{
 			/* some v2s have been seen in msvc. Always seem to be assembly
@@ -440,34 +451,20 @@ void FBacktracer::AddModule(UPTRINT ModuleBase, const TCHAR* Name)
 			}
 		}
 
-		// "Chained" simply means that multiple RUNTIME_FUNCTIONs pertains to a
-		// single actual function in the .text segment.
-		bool bIsChained = (UnwindInfo->Flags & UNW_FLAG_CHAININFO);
-
 		RspBias /= sizeof(void*);	// stack push/popds in units of one machine word
-		RspBias += !bIsChained;		// and one extra push for the ret address
+		RspBias += 1;				// and one extra push for the ret address
 		RspBias |= FpInfo;			// pack in details about possible frame pointer
 
-		if (bIsChained)
-		{
-			OutTableCursor[-1].RspBias += RspBias;
+		*OutTableCursor = {
+			.Id = FunctionTable->BeginAddress,
+			.RspBias = RspBias,
 #if BACKTRACE_DBGLVL >= 2
-			OutTableCursor[-1].Size += (FunctionTable->EndAddress - FunctionTable->BeginAddress);
+			.Size = FunctionTable->EndAddress - FunctionTable->BeginAddress,
+			.UnwindInfo = UnwindInfo,
 #endif
-		}
-		else
-		{
-			*OutTableCursor = {
-				FunctionTable->BeginAddress,
-				RspBias,
-#if BACKTRACE_DBGLVL >= 2
-				FunctionTable->EndAddress - FunctionTable->BeginAddress,
-				UnwindInfo,
-#endif
-			};
+		};
 
-			++OutTableCursor;
-		}
+		++OutTableCursor;
 	}
 
 	UPTRINT ModuleSize = NtHeader->OptionalHeader.SizeOfImage;
