Allow Implicit Conversion of Contiguous Ranges to Memory Views

#Synced-CL 48488326

diff --git a/Engine/Source/Runtime/Core/Public/Memory/MemoryView.h b/Engine/Source/Runtime/Core/Public/Memory/MemoryView.h
--- a/Engine/Source/Runtime/Core/Public/Memory/MemoryView.h
+++ b/Engine/Source/Runtime/Core/Public/Memory/MemoryView.h
@@ -2,6 +2,7 @@
 
 #pragma once
 
+#include "Concepts/ContiguousRange.h"
 #include "CoreTypes.h"
 #include "HAL/PlatformString.h"
 #include "Math/NumericLimits.h"
@@ -62,6 +63,14 @@ public:
 	{
 	}
 
+	/** Construct a view from a contiguous range of any type with compatible const/volatile qualifiers. */
+	template <UE::CContiguousRange RangeType>
+		requires std::is_convertible_v<decltype(GetDataHelper(DeclVal<RangeType>())), DataType*>
+	constexpr inline TMemoryView(RangeType& Range)
+		: TMemoryView(GetDataHelper(Range), (uint64)GetNum(Range) * sizeof(*GetDataHelper(Range)))
+	{
+	}
+
 	/** Returns a pointer to the start of the view. */
 	[[nodiscard]] constexpr inline DataType* GetData() const { return Data; }
 
diff --git a/Engine/Source/Runtime/Core/Public/Serialization/MemoryReader.h b/Engine/Source/Runtime/Core/Public/Serialization/MemoryReader.h
--- a/Engine/Source/Runtime/Core/Public/Serialization/MemoryReader.h
+++ b/Engine/Source/Runtime/Core/Public/Serialization/MemoryReader.h
@@ -107,8 +107,9 @@ public:
 		}
 	}
 
-	explicit FMemoryReaderView(TArrayView<const uint8> InBytes, bool bIsPersistent = false)
-		: FMemoryReaderView(MakeMemoryView(InBytes), bIsPersistent)
+	template <typename InElementType, typename InSizeType>
+	explicit FMemoryReaderView(TArrayView<InElementType, InSizeType>&& InBytes, bool bIsPersistent = false)
+		: FMemoryReaderView(InBytes, bIsPersistent)
 	{
 	}
 
diff --git a/Engine/Source/Runtime/Core/Tests/Memory/MemoryViewTest.cpp b/Engine/Source/Runtime/Core/Tests/Memory/MemoryViewTest.cpp
--- a/Engine/Source/Runtime/Core/Tests/Memory/MemoryViewTest.cpp
+++ b/Engine/Source/Runtime/Core/Tests/Memory/MemoryViewTest.cpp
@@ -5,7 +5,10 @@
 #include "Memory/MemoryView.h"
 
 #include "Containers/ArrayView.h"
+#include "Containers/StaticArray.h"
+#include "Containers/UnrealString.h"
 #include "Tests/TestHarnessAdapter.h"
+#include "UObject/NameTypes.h"
 
 #include <type_traits>
 
@@ -59,6 +62,180 @@ TEST_CASE_NAMED(FMemoryMemoryViewTest, "System::Core::Memory::MemoryView", "[Cor
 
 		STATIC_CHECK(std::is_assignable<FMemoryView, const FMutableMemoryView&>::value);
 		STATIC_CHECK(!std::is_assignable<FMutableMemoryView, const FMemoryView&>::value);
+
+		// FMemoryView is constructible from TArray
+		STATIC_CHECK( std::is_constructible<FMemoryView,       TArray<      uint8>&>::value);
+		STATIC_CHECK( std::is_constructible<FMemoryView,       TArray<const uint8>&>::value);
+		STATIC_CHECK( std::is_constructible<FMemoryView, const TArray<      uint8>&>::value);
+		STATIC_CHECK( std::is_constructible<FMemoryView, const TArray<const uint8>&>::value);
+
+		// FMemoryView is not constructible from a temporary TArray
+		STATIC_CHECK(!std::is_constructible<FMemoryView,       TArray<      uint8>>::value);
+		STATIC_CHECK(!std::is_constructible<FMemoryView,       TArray<const uint8>>::value);
+
+		// FMemoryView is constructible from TArrayView
+		STATIC_CHECK( std::is_constructible<FMemoryView,       TArrayView<      uint8>&>::value);
+		STATIC_CHECK( std::is_constructible<FMemoryView,       TArrayView<const uint8>&>::value);
+		STATIC_CHECK( std::is_constructible<FMemoryView, const TArrayView<      uint8>&>::value);
+		STATIC_CHECK( std::is_constructible<FMemoryView, const TArrayView<const uint8>&>::value);
+
+		// FMemoryView is not constructible from a temporary TArrayView
+		STATIC_CHECK(!std::is_constructible<FMemoryView,       TArrayView<      uint8>>::value);
+		STATIC_CHECK(!std::is_constructible<FMemoryView,       TArrayView<const uint8>>::value);
+
+		// FMemoryView is constructible from a raw array
+		STATIC_CHECK( std::is_constructible<FMemoryView,              uint8(&)[256]>::value);
+		STATIC_CHECK( std::is_constructible<FMemoryView,        const uint8(&)[256]>::value);
+
+		// FMemoryView is not constructible from a temporary raw array
+		STATIC_CHECK(!std::is_constructible<FMemoryView,              uint8   [256]>::value);
+
+		// FMutableMemoryView is constructible from TArray when appropriate. Respects constness.
+		STATIC_CHECK( std::is_constructible<FMutableMemoryView,       TArray<      uint8>&>::value);
+		STATIC_CHECK(!std::is_constructible<FMutableMemoryView,       TArray<const uint8>&>::value);
+		STATIC_CHECK(!std::is_constructible<FMutableMemoryView, const TArray<      uint8>&>::value);
+		STATIC_CHECK(!std::is_constructible<FMutableMemoryView, const TArray<const uint8>&>::value);
+
+		// FMutableMemoryView is not constructible from a temporary TArray
+		STATIC_CHECK(!std::is_constructible<FMutableMemoryView,       TArray<      uint8>>::value);
+
+		// FMutableMemoryView is constructible from TArrayView when appropriate. Respects constness.
+		STATIC_CHECK( std::is_constructible<FMutableMemoryView,       TArrayView<      uint8>&>::value);
+		STATIC_CHECK(!std::is_constructible<FMutableMemoryView,       TArrayView<const uint8>&>::value);
+		STATIC_CHECK( std::is_constructible<FMutableMemoryView, const TArrayView<      uint8>&>::value);
+		STATIC_CHECK(!std::is_constructible<FMutableMemoryView, const TArrayView<const uint8>&>::value);
+
+		// FMutableMemoryView is not constructible from a temporary TArrayView
+		STATIC_CHECK(!std::is_constructible<FMutableMemoryView,       TArrayView<      uint8>>::value);
+
+		// FMutableMemoryView is constructible from a raw array when appropriate. Respects constness.
+		STATIC_CHECK( std::is_constructible<FMutableMemoryView,       uint8(&)[256]>::value);
+		STATIC_CHECK(!std::is_constructible<FMutableMemoryView, const uint8(&)[256]>::value);
+
+		// FMutableMemoryView is not constructible from a temporary raw array
+		STATIC_CHECK(!std::is_constructible<FMutableMemoryView,       uint8   [256]>::value);
+
+		// Implicitly converts TArray to a memory view when appropriate. FMutableMemoryView respects constness.
+		STATIC_CHECK( std::is_convertible<      TArray<      uint8>&, FMemoryView>::value);
+		STATIC_CHECK( std::is_convertible<      TArray<      uint8>&, FMutableMemoryView>::value);
+		STATIC_CHECK( std::is_convertible<      TArray<const uint8>&, FMemoryView>::value);
+		STATIC_CHECK(!std::is_convertible<      TArray<const uint8>&, FMutableMemoryView>::value);
+		STATIC_CHECK( std::is_convertible<const TArray<      uint8>&, FMemoryView>::value);
+		STATIC_CHECK(!std::is_convertible<const TArray<      uint8>&, FMutableMemoryView>::value);
+		STATIC_CHECK( std::is_convertible<const TArray<const uint8>&, FMemoryView>::value);
+		STATIC_CHECK(!std::is_convertible<const TArray<const uint8>&, FMutableMemoryView>::value);
+
+		// Can't implicitly convert a temporary TArray to a non-owning memory view
+		STATIC_CHECK(!std::is_convertible<      TArray<      uint8>, FMemoryView>::value);
+		STATIC_CHECK(!std::is_convertible<      TArray<      uint8>, FMutableMemoryView>::value);
+		STATIC_CHECK(!std::is_convertible<      TArray<const uint8>, FMemoryView>::value);
+
+		// Implicitly converts TStaticArray to a memory view when appropriate. FMutableMemoryView respects constness.
+		STATIC_CHECK( std::is_convertible<      TStaticArray<      uint8, 32>&, FMemoryView>::value);
+		STATIC_CHECK( std::is_convertible<      TStaticArray<      uint8, 32>&, FMutableMemoryView>::value);
+		STATIC_CHECK( std::is_convertible<      TStaticArray<const uint8, 32>&, FMemoryView>::value);
+		STATIC_CHECK(!std::is_convertible<      TStaticArray<const uint8, 32>&, FMutableMemoryView>::value);
+		STATIC_CHECK( std::is_convertible<const TStaticArray<      uint8, 32>&, FMemoryView>::value);
+		STATIC_CHECK(!std::is_convertible<const TStaticArray<      uint8, 32>&, FMutableMemoryView>::value);
+		STATIC_CHECK( std::is_convertible<const TStaticArray<const uint8, 32>&, FMemoryView>::value);
+		STATIC_CHECK(!std::is_convertible<const TStaticArray<const uint8, 32>&, FMutableMemoryView>::value);
+
+		// Can't implicitly convert a temporary TStaticArray to a non-owning memory view
+		STATIC_CHECK(!std::is_convertible<      TStaticArray<      uint8, 32>, FMemoryView>::value);
+		STATIC_CHECK(!std::is_convertible<      TStaticArray<      uint8, 32>, FMutableMemoryView>::value);
+		STATIC_CHECK(!std::is_convertible<      TStaticArray<const uint8, 32>, FMemoryView>::value);
+
+		// Implicitly converts TArrayView to a memory view when appropriate. FMutableMemoryView respects constness.
+		STATIC_CHECK( std::is_convertible<      TArrayView<      uint8>&, FMemoryView>::value);
+		STATIC_CHECK( std::is_convertible<      TArrayView<      uint8>&, FMutableMemoryView>::value);
+		STATIC_CHECK( std::is_convertible<      TArrayView<const uint8>&, FMemoryView>::value);
+		STATIC_CHECK(!std::is_convertible<      TArrayView<const uint8>&, FMutableMemoryView>::value);
+		STATIC_CHECK( std::is_convertible<const TArrayView<      uint8>&, FMemoryView>::value);
+		STATIC_CHECK( std::is_convertible<const TArrayView<      uint8>&, FMutableMemoryView>::value);
+		STATIC_CHECK( std::is_convertible<const TArrayView<const uint8>&, FMemoryView>::value);
+		STATIC_CHECK(!std::is_convertible<const TArrayView<const uint8>&, FMutableMemoryView>::value);
+
+		// Can't implicitly convert a temporary TArrayView to a non-owning memory view
+		// Note that ideally this would be allowed (this is safe for any non-owning range) but we, as of now,
+		// don't have a way to differentiate owning from non-owning ranges so choose to take the safe route.
+		STATIC_CHECK(!std::is_convertible<      TArrayView<      uint8>, FMemoryView>::value);
+		STATIC_CHECK(!std::is_convertible<      TArrayView<      uint8>, FMutableMemoryView>::value);
+		STATIC_CHECK(!std::is_convertible<      TArrayView<const uint8>, FMemoryView>::value);
+		STATIC_CHECK(!std::is_convertible<      TArrayView<const uint8>, FMutableMemoryView>::value);
+
+		// Implicitly converts raw arrays to a memory view when appropriate. FMutableMemoryView respects constness.
+		STATIC_CHECK( std::is_convertible<      uint8(&)[256], FMemoryView>::value);
+		STATIC_CHECK( std::is_convertible<      uint8(&)[256], FMutableMemoryView>::value);
+		STATIC_CHECK( std::is_convertible<const uint8(&)[256], FMemoryView>::value);
+		STATIC_CHECK(!std::is_convertible<const uint8(&)[256], FMutableMemoryView>::value);
+		STATIC_CHECK( std::is_convertible<const uint8(&)[  0], FMemoryView>::value);
+		STATIC_CHECK(!std::is_convertible<const uint8(&)[  0], FMutableMemoryView>::value);
+
+		// Can't implicitly convert a raw array temporary
+		STATIC_CHECK(!std::is_convertible<      uint8   [256], FMemoryView>::value);
+		STATIC_CHECK(!std::is_convertible<      uint8   [256], FMutableMemoryView>::value);
+
+		// Implicitly converts TEXT() when appropriate
+		STATIC_CHECK( std::is_convertible<decltype(TEXT("HELLO")), FMemoryView>::value);
+		STATIC_CHECK(!std::is_convertible<decltype(TEXT("HELLO")), FMutableMemoryView>::value);
+		
+		// Allows assignment of an ArrayView to a MemoryView when appropriate. FMutableMemoryView respects constness.
+		STATIC_CHECK( std::is_assignable<      FMemoryView,        TArrayView<      uint8>&>::value);
+		STATIC_CHECK( std::is_assignable<      FMutableMemoryView, TArrayView<      uint8>&>::value);
+		STATIC_CHECK( std::is_assignable<      FMemoryView,        TArrayView<const uint8>&>::value);
+		STATIC_CHECK(!std::is_assignable<      FMutableMemoryView, TArrayView<const uint8>&>::value);
+		STATIC_CHECK(!std::is_assignable<const FMemoryView,        TArrayView<      uint8>&>::value);
+		STATIC_CHECK(!std::is_assignable<const FMutableMemoryView, TArrayView<      uint8>&>::value);
+		STATIC_CHECK(!std::is_assignable<const FMemoryView,        TArrayView<const uint8>&>::value);
+		STATIC_CHECK(!std::is_assignable<const FMutableMemoryView, TArrayView<const uint8>&>::value);
+
+		// Can't assign a temporary ArrayView to a MemoryView
+		// Note that ideally this would be allowed (this is safe for any non-owning range) but we, as of now,
+		// don't have a way to differentiate owning from non-owning ranges so choose to take the safe route.
+		STATIC_CHECK(!std::is_assignable<      FMemoryView,        TArrayView<      uint8>>::value);
+		STATIC_CHECK(!std::is_assignable<      FMutableMemoryView, TArrayView<      uint8>>::value);
+		STATIC_CHECK(!std::is_assignable<      FMemoryView,        TArrayView<const uint8>>::value);
+		STATIC_CHECK(!std::is_assignable<const FMemoryView,        TArrayView<      uint8>>::value);
+		STATIC_CHECK(!std::is_assignable<const FMutableMemoryView, TArrayView<      uint8>>::value);
+		STATIC_CHECK(!std::is_assignable<const FMemoryView,        TArrayView<const uint8>>::value);
+		STATIC_CHECK(!std::is_assignable<const FMutableMemoryView, TArrayView<const uint8>>::value);
+
+		// Some implementations of std::is_assignable<> incorrectly evaluate assignment from raw arrays. If a 
+		// basic assignment from an array to a pointer of the same type doesn't work, then we expect all of the
+		// following positive tests to fail too. Still checking for failure makes sure we're not incorrectly
+		// expecting failure on a platform where is_assignable<> would actually work for these tests.
+		constexpr bool bIsAssignableWorksForRawArrays = std::is_assignable<uint8*&, uint8(&)[256]>::value;
+		STATIC_CHECK( std::is_assignable<FMemoryView,              uint8(&)[256]>::value == bIsAssignableWorksForRawArrays);
+		STATIC_CHECK( std::is_assignable<FMutableMemoryView,       uint8(&)[256]>::value == bIsAssignableWorksForRawArrays);
+		STATIC_CHECK( std::is_assignable<FMemoryView,        const uint8(&)[256]>::value == bIsAssignableWorksForRawArrays);
+		STATIC_CHECK( std::is_assignable<FMutableMemoryView, const uint8(&)[256]>::value == false);
+
+		// Can't assign a temporary raw array to a MemoryView
+		STATIC_CHECK( std::is_assignable<FMemoryView,              uint8   [256]>::value == false);
+		STATIC_CHECK( std::is_assignable<FMutableMemoryView,       uint8   [256]>::value == false);
+
+		// Can't implicitly cast underlying shapeless memory via assignment of a MemoryView to an ArrayView
+		STATIC_CHECK(!std::is_assignable<      TArrayView<      uint8>, FMemoryView&>::value);
+		STATIC_CHECK(!std::is_assignable<      TArrayView<      uint8>, FMutableMemoryView&>::value);
+		STATIC_CHECK(!std::is_assignable<      TArrayView<const uint8>, FMemoryView&>::value);
+		STATIC_CHECK(!std::is_assignable<      TArrayView<const uint8>, FMutableMemoryView&>::value);
+		STATIC_CHECK(!std::is_assignable<const TArrayView<      uint8>, FMemoryView&>::value);
+		STATIC_CHECK(!std::is_assignable<const TArrayView<      uint8>, FMutableMemoryView&>::value);
+		STATIC_CHECK(!std::is_assignable<const TArrayView<const uint8>, FMemoryView&>::value);
+		STATIC_CHECK(!std::is_assignable<const TArrayView<const uint8>, FMutableMemoryView&>::value);
+
+		// Doesn't allows FMutableMemoryView construction from FString
+		STATIC_CHECK( std::is_constructible<FMemoryView, FString&>::value);
+		STATIC_CHECK(!std::is_constructible<FMutableMemoryView, FString&>::value);
+
+		// Can't construct a MemoryView from a temporary FString
+		STATIC_CHECK(!std::is_constructible<FMemoryView, FString>::value);
+
+		// Doesn't implicitly convert FName
+		STATIC_CHECK(!std::is_convertible<FName&, FMemoryView>::value);
+		STATIC_CHECK(!std::is_convertible<FName&, FMutableMemoryView>::value);
+		STATIC_CHECK(!std::is_convertible<FName , FMemoryView>::value);
+		STATIC_CHECK(!std::is_convertible<FName , FMutableMemoryView>::value);
 	}
 
 	SECTION("Empty Views")
@@ -87,6 +264,60 @@ TEST_CASE_NAMED(FMemoryMemoryViewTest, "System::Core::Memory::MemoryView", "[Cor
 		TestMemoryView(MakeMemoryView(InitializerList), GetData(InitializerList), GetNum(InitializerList) * sizeof(uint8));
 	}
 
+	SECTION("Construction from TArray")
+	{
+		TArray64<uint32> Array(GetData(IntArray), GetNum(IntArray));
+		const TArray64<uint32> ConstArray(GetData(IntArray), GetNum(IntArray));
+
+		FMemoryView MemoryView(Array);
+		TestMemoryView(FMemoryView(MemoryView), Array.GetData(), Array.NumBytes());
+
+		FMemoryView MemoryViewOfConst(ConstArray);
+		TestMemoryView(FMemoryView(MemoryViewOfConst), ConstArray.GetData(), ConstArray.NumBytes());
+
+		FMutableMemoryView MutableMemoryView(Array);
+		TestMutableMemoryView(MutableMemoryView, Array.GetData(), ConstArray.NumBytes());
+	}
+
+	SECTION("Construction from TArrayView")
+	{
+		const TArrayView<uint32> ArrayView = MakeArrayView(IntArray);
+		const TArrayView<const uint32> ConstArrayView = MakeConstArrayView(IntArray);
+
+		FMemoryView MemoryView(ArrayView);
+		TestMemoryView(FMemoryView(MemoryView), IntArray, sizeof(IntArray));
+
+		FMemoryView MemoryViewOfConst(ConstArrayView);
+		TestMemoryView(FMemoryView(MemoryViewOfConst), GetData(IntArray), sizeof(IntArray));
+
+		FMutableMemoryView MutableMemoryView(ArrayView);
+		TestMutableMemoryView(MutableMemoryView, GetData(IntArray), sizeof(IntArray));
+	}
+
+	SECTION("Construction from TStaticArray")
+	{
+		TStaticArray<uint32, 16> StaticArray = {};
+		const TStaticArray<uint32, 16> ConstStaticArray = {};
+
+		const uint64 StaticArrayBytes = StaticArray.Num() * sizeof(decltype(StaticArray)::ElementType);
+		const uint64 ConstStaticArrayBytes = ConstStaticArray.Num() * sizeof(decltype(ConstStaticArray)::ElementType);
+
+		FMemoryView MemoryView(StaticArray);
+		TestMemoryView(FMemoryView(MemoryView), GetData(StaticArray), StaticArrayBytes);
+
+		FMemoryView MemoryViewOfConst(ConstStaticArray);
+		TestMemoryView(FMemoryView(MemoryViewOfConst), GetData(ConstStaticArray), ConstStaticArrayBytes);
+
+		FMutableMemoryView MutableMemoryView(StaticArray);
+		TestMutableMemoryView(MutableMemoryView, GetData(StaticArray), StaticArrayBytes);
+	}
+
+	SECTION("Construction from Raw Array")
+	{
+		TestMemoryView(FMemoryView(IntArray), IntArray, sizeof(IntArray)); 
+		TestMutableMemoryView(FMutableMemoryView(IntArray), IntArray, sizeof(IntArray));
+	}
+
 	SECTION("Reset")
 	{
 		FMutableMemoryView View = MakeMemoryView(IntArray);
@@ -268,6 +499,60 @@ TEST_CASE_NAMED(FMemoryMemoryViewTest, "System::Core::Memory::MemoryView", "[Cor
 		CHECK((MakeMemoryView(ByteArray) + 32) == MakeMemoryView(ByteArray + 16, 0));
 		CHECK((32 + MakeMemoryView(ByteArray)) == MakeMemoryView(ByteArray + 16, 0));
 	}
+
+	SECTION("Assignment from TArrayView")
+	{
+		const TArrayView<uint32> ArrayView = MakeArrayView(IntArray);
+		const TArrayView<const uint32> ConstArrayView = MakeConstArrayView(IntArray);
+
+		FMemoryView MemoryView = ArrayView;
+		TestMemoryView(MemoryView, IntArray, sizeof(IntArray));
+		
+		FMemoryView MemoryViewOfConst = ConstArrayView;
+		TestMemoryView(MemoryViewOfConst, IntArray, sizeof(IntArray));
+
+		FMutableMemoryView MutableMemoryView = ArrayView;
+		TestMutableMemoryView(MutableMemoryView, IntArray, sizeof(IntArray));
+	}
+
+	SECTION("Assignment from TArray")
+	{
+		TArray64<uint32> Array(GetData(IntArray), GetNum(IntArray));
+		const TArray64<uint32> ConstArray(GetData(IntArray), GetNum(IntArray));
+
+		FMemoryView MemoryView = Array;
+		TestMemoryView(MemoryView, Array.GetData(), Array.NumBytes());
+		
+		FMemoryView MemoryViewOfConst = ConstArray;
+		TestMemoryView(MemoryViewOfConst, ConstArray.GetData(), ConstArray.NumBytes());
+
+		FMutableMemoryView MutableMemoryView = Array;
+		TestMutableMemoryView(MutableMemoryView, Array.GetData(), Array.NumBytes());
+	}
+
+	SECTION("Assignment from TStaticArray")
+	{
+		TStaticArray<uint64, 21> Array = {};
+		const TStaticArray<uint16, 44> ConstArray = {};
+
+		FMemoryView MemoryView = Array;
+		TestMemoryView(MemoryView, Array.GetData(), Array.Num() * sizeof(decltype(Array)::ElementType));
+		
+		FMemoryView MemoryViewOfConst = ConstArray;
+		TestMemoryView(MemoryViewOfConst, ConstArray.GetData(), ConstArray.Num() * sizeof(decltype(ConstArray)::ElementType));
+
+		FMutableMemoryView MutableMemoryView = Array;
+		TestMutableMemoryView(MutableMemoryView, Array.GetData(), Array.Num() * sizeof(decltype(Array)::ElementType));
+	}
+
+	SECTION("Assignment from Raw Array")
+	{
+		FMemoryView MemoryView = IntArray;
+		TestMemoryView(MemoryView, IntArray, sizeof(IntArray));
+
+		FMutableMemoryView MutableMemoryView = IntArray;
+		TestMutableMemoryView(MutableMemoryView, IntArray, sizeof(IntArray));
+	}
 }
 
 #endif // WITH_TESTS
