ShaderKeys: Fix missing values in shader keys due to incorrect printf format specifier: the marshalling to string of uint64 and int64 values was hardcoding the string "lld" instead of converting the int64 and uint64 value to string.
#rnx
#Synced-CL 46581523

diff --git a/Engine/Source/Runtime/Core/Public/Serialization/ShaderKeyGenerator.h b/Engine/Source/Runtime/Core/Public/Serialization/ShaderKeyGenerator.h
--- a/Engine/Source/Runtime/Core/Public/Serialization/ShaderKeyGenerator.h
+++ b/Engine/Source/Runtime/Core/Public/Serialization/ShaderKeyGenerator.h
@@ -201,11 +201,7 @@ inline void FShaderKeyGenerator::Append(int64 Value)
 	switch (OutputType)
 	{
 	case EOutputType::Text:
-		// This should be:
-		// ResultString->Appendf(TEXT("%" INT64_FMT), Value);
-		// Changing it would invalidate shader and DDC keys, so the previous behavior is maintained until a time where it can be fixed properly.
-		// This issue is tracked by https://jira.it.epicgames.com/browse/UE-353138
-		ResultString->Appendf(TEXT(INT64_FMT));
+		ResultString->Appendf(TEXT("%" INT64_FMT), Value);
 		return;
 	case EOutputType::Binary:
 		ResultFunc(&Value, sizeof(Value));
@@ -220,11 +216,7 @@ inline void FShaderKeyGenerator::Append(uint64 Value)
 	switch (OutputType)
 	{
 	case EOutputType::Text:
-		// This should be:
-		// ResultString->Appendf(TEXT("%" UINT64_FMT), Value);
-		// Changing it would invalidate shader and DDC keys, so the previous behavior is maintained until a time where it can be fixed properly.
-		// This issue is tracked by https://jira.it.epicgames.com/browse/UE-353138
-		ResultString->Appendf(TEXT(UINT64_FMT));
+		ResultString->Appendf(TEXT("%" UINT64_FMT), Value);
 		return;
 	case EOutputType::Binary:
 		ResultFunc(&Value, sizeof(Value));
