Follow redirecting runtime functions.

#rnx

#Synced-CL Unknown

diff --git a/Engine/Source/Runtime/Core/Private/ProfilingDebugging/Microsoft/WindowsCallstackTrace.cpp b/Engine/Source/Runtime/Core/Private/ProfilingDebugging/Microsoft/WindowsCallstackTrace.cpp
--- a/Engine/Source/Runtime/Core/Private/ProfilingDebugging/Microsoft/WindowsCallstackTrace.cpp
+++ b/Engine/Source/Runtime/Core/Private/ProfilingDebugging/Microsoft/WindowsCallstackTrace.cpp
@@ -342,6 +342,15 @@ void FBacktracer::AddModule(UPTRINT ModuleBase, const TCHAR* Name)
 		UPTRINT UnwindInfoAddr = ModuleBase + FunctionTable->UnwindInfoAddress;
 		const auto* UnwindInfo = (FUnwindInfo*)UnwindInfoAddr;
 
+		// Unwind infos are two-byte aligned. Unaligned info addresses have been
+		// encoutered where the aligned RVA holds a redirecting runtime function
+		// and not unwind op codes.
+		if (UnwindInfoAddr & 1)
+		{
+			const RUNTIME_FUNCTION* Redirect = (RUNTIME_FUNCTION*)(UnwindInfoAddr & ~1);
+			UnwindInfo = (FUnwindInfo*)(ModuleBase + Redirect->UnwindInfoAddress);
+		}
+
 		// Functions can end up split into parts and spread about code segments. In
 		// such cases there is a chain of unwind info which must be followed to find
 		// the primary info that describes the prologue.
