Correctly handle functions with no stack frame that do not have unwind table entries

Unwind tables only have entries for functions with a prologue and epilogue that establish a stack frame. Leaf functions - functions that do not establish a stack frame - do not have any unwind codes. This change adds an additional check against the end of a function to correctly handle leafs with a zero-sized frame. Fortunately leafs are almost never encountered in practice.

#rnx

#Synced-CL Unknown

diff --git a/Engine/Source/Runtime/Core/Private/ProfilingDebugging/Microsoft/WindowsCallstackTrace.cpp b/Engine/Source/Runtime/Core/Private/ProfilingDebugging/Microsoft/WindowsCallstackTrace.cpp
--- a/Engine/Source/Runtime/Core/Private/ProfilingDebugging/Microsoft/WindowsCallstackTrace.cpp
+++ b/Engine/Source/Runtime/Core/Private/ProfilingDebugging/Microsoft/WindowsCallstackTrace.cpp
@@ -155,8 +155,8 @@ private:
 	{
 		uint32				Id;
 		int32				RspBias;
+		uint32				End;
 #if BACKTRACE_DBGLVL >= 2
-		uint32				Size;
 		const FUnwindInfo*	UnwindInfo;
 #endif
 	};
@@ -467,8 +467,8 @@ void FBacktracer::AddModule(UPTRINT ModuleBase, const TCHAR* Name)
 		*OutTableCursor = {
 			.Id = FunctionTable->BeginAddress,
 			.RspBias = RspBias,
+			.End = FunctionTable->EndAddress,
 #if BACKTRACE_DBGLVL >= 2
-			.Size = FunctionTable->EndAddress - FunctionTable->BeginAddress,
 			.UnwindInfo = UnwindInfo,
 #endif
 		};
@@ -605,6 +605,17 @@ const FBacktracer::FFunction* FBacktracer::LookupFunction(UPTRINT Address, FLook
 		return nullptr;
 	}
 #endif
+
+	if (FuncId >= Function->End)
+	{
+		// The lookup fell outside of functions with stack frames so we will return
+		// a generic result that gives a stack frame size of 0.
+		static FFunction NoFrameFunction = {
+			.RspBias = 0,
+		};
+		return &NoFrameFunction;
+	}
+
 	return Function;
 }
 
