ConfigBinary for MobileSession Scalability
ExtraBinaryConfigWithScalabilityOverride Array Option added to /Script/UnrealEd.ProjectPackagingSettings section to instruct UnrealPak to create additional ConfigBinaries if needed for ScalabilityOverrides at Staging time.
BinaryConfigScalabilityOverride command line argument added in order to create GConfig from a different ConfigBinary than the default one at Runtime

#Synced-CL 48599132

diff --git a/Engine/Source/Developer/PakFileUtilities/Private/PakFileUtilities.cpp b/Engine/Source/Developer/PakFileUtilities/Private/PakFileUtilities.cpp
--- a/Engine/Source/Developer/PakFileUtilities/Private/PakFileUtilities.cpp
+++ b/Engine/Source/Developer/PakFileUtilities/Private/PakFileUtilities.cpp
@@ -5539,6 +5539,19 @@ bool MakeBinaryConfig(const TCHAR* CmdLine)
 	// building the binary config, but aren't staged, and shouldn't be in the binary)
 	LoadIniCache(Config, PluginAndConfigsList, true);
 
+	FString ScalabilityPlatformOverride;
+	FParse::Value(FCommandLine::Get(), TEXT("ScalabilityIniPlatformOverride="), ScalabilityPlatformOverride);
+	if (!ScalabilityPlatformOverride.IsEmpty())
+	{
+		FString EngineConfigPlatformOverride = FString::Format(TEXT("{0}{1}/{1}Scalability.ini"), { *Context.EngineConfigDir, *ScalabilityPlatformOverride });
+		FString ProjectConfigPlatformOverride = FString::Format(TEXT("{0}{1}/{1}Scalability.ini"), { *Context.ProjectConfigDir, *ScalabilityPlatformOverride });
+		if (Config.StagedGlobalConfigCache == nullptr)
+		{
+			Config.StagedGlobalConfigCache = new TSet<FString>();
+		}
+		Config.StagedGlobalConfigCache->Add(EngineConfigPlatformOverride);
+		Config.StagedGlobalConfigCache->Add(ProjectConfigPlatformOverride);
+	}
 	Config.InitializeKnownConfigFiles(Context);
 
 	// then we need to save out with staged paths so that runtime can know what files exist in the cache
diff --git a/Engine/Source/Programs/AutomationTool/AutomationUtils/Utils.cs b/Engine/Source/Programs/AutomationTool/AutomationUtils/Utils.cs
--- a/Engine/Source/Programs/AutomationTool/AutomationUtils/Utils.cs
+++ b/Engine/Source/Programs/AutomationTool/AutomationUtils/Utils.cs
@@ -477,7 +477,9 @@ namespace AutomationTool
 					bool bSkipSizeCheck = false;
 					// BinaryConfig.ini is a special case with binary data, but with same extension to handle all ini 
 					// file chunking/packaging/etc rules
-					if ((IniKeyDenyList != null || IniSectionDenyList != null || IniSectionAllowList != null) && Path.GetExtension(SourceName) == ".ini" && Path.GetFileName(SourceName) != "BinaryConfig.ini")
+					string FileName = Path.GetFileName(SourceName);
+					bool bIsNotBinaryConfigIni = string.Equals(Path.GetExtension(SourceName), ".ini", StringComparison.OrdinalIgnoreCase) && Path.GetFileName(SourceName).IndexOf("BinaryConfig", StringComparison.OrdinalIgnoreCase) == -1;
+					if ((IniKeyDenyList != null || IniSectionDenyList != null || IniSectionAllowList != null) && bIsNotBinaryConfigIni)
 					{
 						FilterIniFile(SourceName, TargetName, IniKeyDenyList, IniSectionDenyList, IniSectionAllowList);
 						// ini files may change size, don't check
diff --git a/Engine/Source/Programs/AutomationTool/CookedEditor/CookedEditor.Automation.cs b/Engine/Source/Programs/AutomationTool/CookedEditor/CookedEditor.Automation.cs
--- a/Engine/Source/Programs/AutomationTool/CookedEditor/CookedEditor.Automation.cs
+++ b/Engine/Source/Programs/AutomationTool/CookedEditor/CookedEditor.Automation.cs
@@ -355,7 +355,7 @@ public class ModifyStageContext
 			AddUFSFilesToList(NonUFSFilesToStage, ".uplugin", SC);
 			AddUFSFilesToList(NonUFSFilesToStage, ".ini", SC);
 			AddUFSFilesToList(NonUFSFilesToStage, "SDK.json", SC);
-			NonUFSFilesToStage = NonUFSFilesToStage.Where(x => x.GetFileName() != "BinaryConfig.ini").ToList();
+			NonUFSFilesToStage = NonUFSFilesToStage.Where(x => x.GetFileName().IndexOf("BinaryConfig", StringComparison.OrdinalIgnoreCase) == -1).ToList();
 		}
 	}
 
diff --git a/Engine/Source/Programs/AutomationTool/Scripts/CopyBuildToStagingDirectory.Automation.cs b/Engine/Source/Programs/AutomationTool/Scripts/CopyBuildToStagingDirectory.Automation.cs
--- a/Engine/Source/Programs/AutomationTool/Scripts/CopyBuildToStagingDirectory.Automation.cs
+++ b/Engine/Source/Programs/AutomationTool/Scripts/CopyBuildToStagingDirectory.Automation.cs
@@ -2138,6 +2138,26 @@ namespace AutomationScripts
 						RunUnrealPak(Params.RawProjectPath, UnrealPakParams, "MakeBinaryConfig");
 
 						SC.StageFile(StagedFileType.UFS, OutputFile, StagedFileReference.Combine(SC.RelativeProjectRootForStage, "Config", OutputFile.GetFileName()));
+
+						if (PlatformGameConfig.TryGetValues("/Script/UnrealEd.ProjectPackagingSettings", "ExtraBinaryConfigWithScalabilityOverride", out IReadOnlyList<string> ScalabilityOverrideValues))
+						{
+							foreach (string ScalabilityOverride in ScalabilityOverrideValues)
+							{
+								if(!string.IsNullOrEmpty(ScalabilityOverride))
+								{
+									FileReference OutputExtraBinary = FileReference.Combine(SC.ProjectRoot, "Intermediate", "Config", TargetPlatformName, String.Format("BinaryConfig{0}.ini", ScalabilityOverride));
+									String UnrealPakParamsExtraBinary = String.Format("MakeBinaryConfig -Project=\"{0}\" -Platform={1} -OutputFile=\"{2}\" -StagedPluginsFile=\"{3}\" -ScalabilityIniPlatformOverride={4}",
+										SC.RawProjectPath, ConfigHierarchy.GetIniPlatformName(ThisPlatform.IniPlatformType), OutputExtraBinary.FullName, ConfigCacheInputFile.FullName, ScalabilityOverride);
+									if (!string.IsNullOrEmpty(SC.CustomConfig))
+									{
+										UnrealPakParamsExtraBinary += $" -CustomConfig={SC.CustomConfig}";
+									}
+
+									RunUnrealPak(Params.RawProjectPath, UnrealPakParamsExtraBinary, "MakeBinaryConfig");
+									SC.StageFile(StagedFileType.UFS, OutputExtraBinary, StagedFileReference.Combine(SC.RelativeProjectRootForStage, "Config", OutputExtraBinary.GetFileName()));
+								}
+							}
+						}
 					}
 
 					bCreatePluginManifest = true;
diff --git a/Engine/Source/Runtime/Core/Private/Misc/ConfigCacheIni.cpp b/Engine/Source/Runtime/Core/Private/Misc/ConfigCacheIni.cpp
--- a/Engine/Source/Runtime/Core/Private/Misc/ConfigCacheIni.cpp
+++ b/Engine/Source/Runtime/Core/Private/Misc/ConfigCacheIni.cpp
@@ -6136,8 +6136,7 @@ bool FConfigCacheIni::InitializeKnownConfigFiles(FConfigContext& Context)
 {
 	// check for scalability platform override.
 	FConfigContext* ScalabilityPlatformOverrideContext = nullptr;
-#if !UE_BUILD_SHIPPING && WITH_EDITOR
-	if (Context.ConfigSystem == GConfig)
+#if WITH_EDITOR || IS_PROGRAM
 	{
 		FString ScalabilityPlatformOverrideCommandLine;
 		FParse::Value(FCommandLine::Get(), TEXT("ScalabilityIniPlatformOverride="), ScalabilityPlatformOverrideCommandLine);
@@ -6285,6 +6284,34 @@ bool FConfigCacheIni::CreateGConfigFromSaved(const TCHAR* Filename)
 	// get the already loaded file
 	int64 Size;
 	void* PreloadedData = GPreLoadConfigBin.TakeOwnershipOfLoadedData(&Size);
+
+	FString BinaryConfigScalabilityOverride;
+	FParse::Value(FCommandLine::Get(), TEXT("BinaryConfigScalabilityOverride="), BinaryConfigScalabilityOverride);
+	bool bIsConfigScalabilityOverride = !BinaryConfigScalabilityOverride.IsEmpty();
+	TArray<uint8> FileContent;
+	if (bIsConfigScalabilityOverride)
+	{
+		FString BinaryConfig = FString::Printf(TEXT("%sBinaryConfig%s.ini"), *FPaths::ProjectConfigDir(), *BinaryConfigScalabilityOverride);
+		if (FFileHelper::LoadFileToArray(FileContent, *BinaryConfig, FILEREAD_Silent))
+		{
+			// Free immediately the default Binary Config preloaded Data
+			if (PreloadedData != nullptr)
+			{
+				FMemory::Free(PreloadedData);
+			}
+
+			// Set the Preload Data to the Override Scalability Binary Config file
+			PreloadedData = FileContent.GetData();
+			Size = FileContent.Num();
+		}
+		else
+		{
+			UE_LOG(LogConfig, Log, TEXT("BinaryConfigScalabilityOverride was set to %s but BinaryConfig %s could not be loaded. Reverting to default BinaryConfig.ini file."),
+				*BinaryConfigScalabilityOverride, *BinaryConfig);
+			bIsConfigScalabilityOverride = false;
+		}
+	}
+
 	if (PreloadedData == nullptr)
 	{
 		return false;
@@ -6325,8 +6352,10 @@ bool FConfigCacheIni::CreateGConfigFromSaved(const TCHAR* Filename)
 	// It must be done after GConfig has been allocated and ready to use to ensure IFileHandle log wont be a 
 	// FManagedStorageFileWriteHandle (which might cause dead lock on some platforms when reporting a crash)
 	UE_LOG(LogInit, Display, TEXT("Loaded binary GConfig from %lld bytes of data..."), Size);
-
-	FMemory::Free(PreloadedData);
+	if (!bIsConfigScalabilityOverride)
+	{
+		FMemory::Free(PreloadedData);
+	}
 	return true;
 }
 
