Compile fix when callstack tracing debug/verification path is compiled in.

#rnx

#Synced-CL Unknown

diff --git a/Engine/Source/Runtime/Core/Private/ProfilingDebugging/Microsoft/WindowsCallstackTrace.cpp b/Engine/Source/Runtime/Core/Private/ProfilingDebugging/Microsoft/WindowsCallstackTrace.cpp
--- a/Engine/Source/Runtime/Core/Private/ProfilingDebugging/Microsoft/WindowsCallstackTrace.cpp
+++ b/Engine/Source/Runtime/Core/Private/ProfilingDebugging/Microsoft/WindowsCallstackTrace.cpp
@@ -598,24 +598,24 @@ const FBacktracer::FFunction* FBacktracer::LookupFunction(UPTRINT Address, FLook
 	}
 
 	const FFunction* Function = Module->Functions + (Index - 1);
-#if BACKTRACE_DBGLVL >= 2
-	if ((FuncId - Function->Id) >= Function->Size)
-	{
-		PLATFORM_BREAK();
-		return nullptr;
-	}
-#endif
-
 	if (FuncId >= Function->End)
 	{
 		// The lookup fell outside of functions with stack frames so we will return
-		// a generic result that gives a stack frame size of 0.
+		// a generic result that gives a stack frame size of 0 (+1 for return addr).
 		static FFunction NoFrameFunction = {
-			.RspBias = 0,
+			.RspBias = 1,
 		};
 		return &NoFrameFunction;
 	}
 
+#if BACKTRACE_DBGLVL >= 2
+	if ((FuncId - Function->Id) >= Function->End - Function->Id)
+	{
+		PLATFORM_BREAK();
+		return nullptr;
+	}
+#endif
+
 	return Function;
 }
 
