Correct version check as we are already handling version 2 codes.

#rnx

#Synced-CL Unknown

diff --git a/Engine/Source/Runtime/Core/Private/ProfilingDebugging/Microsoft/WindowsCallstackTrace.cpp b/Engine/Source/Runtime/Core/Private/ProfilingDebugging/Microsoft/WindowsCallstackTrace.cpp
--- a/Engine/Source/Runtime/Core/Private/ProfilingDebugging/Microsoft/WindowsCallstackTrace.cpp
+++ b/Engine/Source/Runtime/Core/Private/ProfilingDebugging/Microsoft/WindowsCallstackTrace.cpp
@@ -362,10 +362,9 @@ void FBacktracer::AddModule(UPTRINT ModuleBase, const TCHAR* Name)
 			UnwindInfo = (FUnwindInfo*)(ModuleBase + Outer->UnwindInfoAddress);
 		}
 
-		if (UnwindInfo->Version != 1)
+		// Expect version 1 or 2 (2 introduced the Epilog code).
+		if (uint32 Version = UnwindInfo->Version; Version != 1 && Version != 2)
 		{
-			/* some v2s have been seen in msvc. Always seem to be assembly
-			 * routines (memset, memcpy, etc) */
 			continue;
 		}
 
@@ -382,11 +381,15 @@ void FBacktracer::AddModule(UPTRINT ModuleBase, const TCHAR* Name)
 		while (Code < EndCode && !bEncounteredUnknownOpCode)
 		{
 #if BACKTRACE_DBGLVL >= 2
-			if (Code->PrologOffset > PrologVerify)
+			if (Code->OpCode != UWOP_EPILOG)
 			{
-				PLATFORM_BREAK();
+				// PrologOffset has a different purpose in v2 epilog ops
+				if (Code->PrologOffset > PrologVerify)
+				{
+					PLATFORM_BREAK();
+				}
+				PrologVerify = Code->PrologOffset;
 			}
-			PrologVerify = Code->PrologOffset;
 #endif
 
 			switch (Code->OpCode)
@@ -444,11 +447,17 @@ void FBacktracer::AddModule(UPTRINT ModuleBase, const TCHAR* Name)
 				Code += 1;
 				break;
 
+			case UWOP_EPILOG:
+				// The epilog op replaced an obsolete save-xmm op. Save-xmm was two codes
+				// long. Epilog retained this and is always an even code count. +2 will work
+				// even if epilog is three ops long as the opcode value is always present.
+				Code += 2;
+				break;
+
 			case UWOP_SAVE_NONVOL:		Code += 2; break; /* saves are movs instead of pushes */
 			case UWOP_SAVE_NONVOL_FAR:	Code += 3; break;
 			case UWOP_SAVE_XMM128:		Code += 2; break;
 			case UWOP_SAVE_XMM128_FAR:	Code += 3; break;
-			case UWOP_EPILOG:			Code += 2; break;
 			case UWOP_SPARE_CODE:		Code += 3; break;
 
 			default:
