Added a delegate for property permission lists and path permission lists to be enabled or disabled dynamically when they are checked, instead of globally.

#Synced-CL 48600603

diff --git a/Engine/Source/Editor/PropertyEditor/Private/PropertyPermissionList.cpp b/Engine/Source/Editor/PropertyEditor/Private/PropertyPermissionList.cpp
--- a/Engine/Source/Editor/PropertyEditor/Private/PropertyPermissionList.cpp
+++ b/Engine/Source/Editor/PropertyEditor/Private/PropertyPermissionList.cpp
@@ -56,6 +56,7 @@ FPropertyPermissionList::FPropertyPermissionList()
 FPropertyPermissionList::~FPropertyPermissionList()
 {
 	FTSTicker::RemoveTicker(OnTickHandle);
+	FCoreDelegates::OnPostEngineInit.RemoveAll(this);
 	if (GEditor)
 	{
 		GEditor->OnBlueprintCompiled().RemoveAll(this);
@@ -233,13 +234,24 @@ void FPropertyPermissionList::UnregisterCustomDetailsTreeNodePermissionDelegate(
 
 void FPropertyPermissionList::SetEnabled(bool bEnable, bool bAnnounce)
 {
-	if (bEnablePermissionList != bEnable)
-	{
-		bEnablePermissionList = bEnable;
-		if (bAnnounce)
+	EnablePermissionList = FDelegateEnablePathFilter::CreateLambda([bEnable]()
 		{
-			PermissionListEnabledDelegate.Broadcast();
-		}
+			return bEnable;
+		});
+
+	if (bAnnounce)
+	{
+		PermissionListEnabledDelegate.Broadcast();
+	}
+}
+
+void FPropertyPermissionList::SetEnabled(const FDelegateEnablePathFilter& bEnable, bool bAnnounce)
+{
+	EnablePermissionList = bEnable;
+
+	if (bAnnounce)
+	{
+		PermissionListEnabledDelegate.Broadcast();
 	}
 }
 
@@ -255,7 +267,7 @@ bool FPropertyPermissionList::HasFiltering(const UStruct* ObjectStruct) const
 
 bool FPropertyPermissionList::DoesPropertyPassFilter(const UStruct* ObjectStruct, FName PropertyName) const
 {
-	if (bEnablePermissionList && ObjectStruct)
+	if (ObjectStruct && IsEnabled())
 	{
 		return GetCachedPermissionListForStruct(ObjectStruct).PassesFilter(PropertyName);
 	}
@@ -264,7 +276,7 @@ bool FPropertyPermissionList::DoesPropertyPassFilter(const UStruct* ObjectStruct
 
 bool FPropertyPermissionList::DoesDetailTreeNodePassFilter(const UStruct* ObjectStruct, TSharedRef<FDetailTreeNode> DetailTreeNode)
 {
-	if (bEnablePermissionList && ObjectStruct)
+	if (ObjectStruct && IsEnabled())
 	{
 		if(CustomDetailTreeNodePermissionDelegates.Contains(ObjectStruct) && CustomDetailTreeNodePermissionDelegates[ObjectStruct].IsBound())
 		{
diff --git a/Engine/Source/Editor/PropertyEditor/Public/PropertyPermissionList.h b/Engine/Source/Editor/PropertyEditor/Public/PropertyPermissionList.h
--- a/Engine/Source/Editor/PropertyEditor/Public/PropertyPermissionList.h
+++ b/Engine/Source/Editor/PropertyEditor/Public/PropertyPermissionList.h
@@ -26,7 +26,7 @@ DECLARE_DELEGATE_RetVal_OneParam(TOptional<bool>, FCustomDetailTreeNodePermissio
  * 
  * For example, if you have:
  * class A - (UseExistingPermissionList "MyProp")						PermissionList = "MyProp"
- * class B : public class A - (AllowListListAllProperties)				PermissionList = "MyProp","PropA1","PropA2"
+ * class B : public class A - (AllowListAllProperties)					PermissionList = "MyProp","PropA1","PropA2"
  * class C : public class B - (UseExistingPermissionList "AnotherProp")	PermissionList = "MyProp","PropA1","PropA2","AnotherProp"
  * class D : public class B - (UseExistingPermissionList)				PermissionList = "MyProp","PropA1","PropA2","PropD1","PropD2"
  * Note that because class C manually defines a PermissionList, it does not inherit the AllowListAllProperties rule from class B, while
@@ -88,13 +88,17 @@ public:
 	FSimpleMulticastDelegate PermissionListEnabledDelegate;
 
 	/** Controls whether DoesPropertyPassFilter always returns true or performs property-based filtering. */
-	bool IsEnabled() const { return bEnablePermissionList; }
+	bool IsEnabled() const { return EnablePermissionList.IsBound() && EnablePermissionList.Execute(); }
 	/** 
 	 * Turn on or off the property PermissionList. DoesPropertyPassFilter will always return true if disabled.
 	 * Optionally the change can be announced through a broadcast, which is typically recommended, but may not be needed
 	 * if the permission list is only temporary enabled, e.g. during saving.
+	 * 
+	 * Under the hood, it will be a delegate that will always return bEnabled. If a delegate is passed in instead, that
+	 * delegate is called to return dynamically whether the filter is enabled or not.
 	 */
 	UE_API void SetEnabled(bool bEnable, bool bAnnounce = true);
+	UE_API void SetEnabled(const FDelegateEnablePathFilter& bEnable, bool bAnnounce = true);
 
 	/** Checks if the provided struct has any entries in the allow, deny and deny all lists and returns true in that case. */
 	UE_API bool HasFiltering(const UStruct* ObjectStruct) const;
@@ -125,6 +129,8 @@ public:
 	/** If true, PermissionListUpdatedDelegate will not broadcast when modifying the permission lists */
 	bool bSuppressUpdateDelegate = false;
 
+	const FDelegateEnablePathFilter& DelegateEnablePathFilter() { return EnablePermissionList; }
+
 protected:
 	UE_API FPropertyPermissionList();
 	UE_API ~FPropertyPermissionList();
@@ -134,8 +140,8 @@ private:
 	UE_API void RegisterOnBlueprintCompiled();
 	UE_API void ClearCacheAndBroadcast(TSoftObjectPtr<const UStruct> ObjectStruct = nullptr, FName OwnerName = NAME_None);
 
-	/** Whether DoesPropertyPassFilter should perform its PermissionList check or always return true */
-	bool bEnablePermissionList = false;
+	/** Delegate to return whether DoesPropertyPassFilter should perform its PermissionList check or always return true */
+	FDelegateEnablePathFilter EnablePermissionList;
 
 	/** Stores assigned PermissionLists from AddPermissionList(), which are later flattened and stored in CachedPropertyPermissionList. */
 	TMap<TSoftObjectPtr<const UStruct>, FPropertyPermissionListEntry> RawPropertyPermissionList;
diff --git a/Engine/Source/Runtime/Core/Private/Misc/PathPermissionList.cpp b/Engine/Source/Runtime/Core/Private/Misc/PathPermissionList.cpp
--- a/Engine/Source/Runtime/Core/Private/Misc/PathPermissionList.cpp
+++ b/Engine/Source/Runtime/Core/Private/Misc/PathPermissionList.cpp
@@ -43,12 +43,54 @@ bool FPathPermissionList::PassesFilter(const FStringView Item) const
 	{
 		if (!AllowTree.IsEmpty() && !AllowTree.Contains(Item))
 		{
-			return false;
+			// If the allow tree doesn't include the item, and
+			// no delegates exist that could disable any
+			// filters, then the item doesn't pass the filter.
+			if (EnablePathFilterDelegates.IsEmpty())
+			{
+				return false;
+			}
+
+			// Check to make sure at least one filter is still enabled
+			for (const TTuple<FName, FDelegateEnablePathFilter>& Entry : EnablePathFilterDelegates)
+			{
+				// If there is a filter that enables this permission list, and is keeping
+				// it enabled, then the item doesn't pass the filter. But if a filter is
+				// disabled, we can check the others.
+				if (Entry.Value.IsBound() && Entry.Value.Execute())
+				{
+					return false;
+				}
+			}
 		}
 
-		if (DenyTree.Contains(Item))
+		// It's possible that the allow tree is empty or the filters are all disabled, so
+		// see if this is an item we are explicitly denying.
+		if (const FPermissionListOwners* Owners = DenyTree.Find(Item))
 		{
-			return false;
+			// If the deny tree includes the item, and  no delegates exist that would disable
+			// this list, then the item doesn't pass the filter.
+			if (EnablePathFilterDelegates.IsEmpty())
+			{
+				return false;
+			}
+
+			// At this point there are some delegates that could disable this list, so make sure
+			// at least one wants to keep it on.
+			for (FName Owner : *Owners)
+			{
+				if (const FDelegateEnablePathFilter* Filter = EnablePathFilterDelegates.Find(Owner))
+				{
+					// If the filter is enabled that excludes this item, then it doesn't pass the filter.
+					if (Filter->IsBound() && Filter->Execute())
+					{
+						return false;
+					}
+				}
+			}
+
+			// Here, we found no enabled filters that would deny the item even though it's in the deny list.
+			// At that point, we should allow the item despite it being in the deny list.
 		}
 	}
 
@@ -245,7 +287,7 @@ bool FPathPermissionList::AddAllowListItem(const FName OwnerName, const FStringV
 {
 	VerifyItemMatchesListType(Item);
 	bool bExisted = false;
-	AllowTree.FindOrAdd(Item).AddUnique(OwnerName);
+	AllowTree.FindOrAdd(Item, &bExisted).AddUnique(OwnerName);
 	bool bFilterChanged = !bExisted;
 	if (bFilterChanged && !bSuppressOnFilterChanged)
 	{
@@ -321,9 +363,12 @@ TArray<FName> FPathPermissionList::GetOwnerNames() const
 	DenyTree.TryGetChildren({}, Entries, EDirectoryTreeGetFlags::Recursive | EDirectoryTreeGetFlags::ImpliedParent);
 	for (const FString& DenyEntry : Entries)
 	{
-		for (FName OwnerName : *DenyTree.Find(DenyEntry))
+		if (const FPermissionListOwners* Owners = DenyTree.Find(DenyEntry))
 		{
-			OwnerNames.AddUnique(OwnerName);
+			for (FName OwnerName : *Owners)
+			{
+				OwnerNames.AddUnique(OwnerName);
+			}
 		}
 	}
 
@@ -353,12 +398,14 @@ bool FPathPermissionList::UnregisterOwner(const FName OwnerName)
 	DenyTree.TryGetChildren({}, Entries, EDirectoryTreeGetFlags::Recursive | EDirectoryTreeGetFlags::ImpliedParent);
 	for (const FString& DenyEntry : Entries)
 	{
-		FPermissionListOwners* Owners = DenyTree.Find(DenyEntry);
-		Owners->Remove(OwnerName);
-		if (Owners->Num() == 0)
+		if (FPermissionListOwners* Owners = DenyTree.Find(DenyEntry))
 		{
-			DenyTree.Remove(DenyEntry);
-			bFilterChanged = true;
+			Owners->Remove(OwnerName);
+			if (Owners->Num() == 0)
+			{
+				DenyTree.Remove(DenyEntry);
+				bFilterChanged = true;
+			}
 		}
 	}
 
@@ -366,17 +413,24 @@ bool FPathPermissionList::UnregisterOwner(const FName OwnerName)
 	AllowTree.TryGetChildren({}, Entries, EDirectoryTreeGetFlags::Recursive | EDirectoryTreeGetFlags::ImpliedParent);
 	for (const FString& AllowEntry : Entries)
 	{
-		FPermissionListOwners* Owners = AllowTree.Find(AllowEntry);
-		Owners->Remove(OwnerName);
-		if (Owners->Num() == 0)
+		if (FPermissionListOwners* Owners = AllowTree.Find(AllowEntry))
 		{
-			AllowTree.Remove(AllowEntry);
-			bFilterChanged = true;
+			Owners->Remove(OwnerName);
+			if (Owners->Num() == 0)
+			{
+				AllowTree.Remove(AllowEntry);
+				bFilterChanged = true;
+			}
 		}
 	}
 
 	bFilterChanged |= (DenyListAll.Remove(OwnerName) > 0);
 
+	{
+		TGuardValue<bool> Guard(bSuppressOnFilterChanged, true);
+		bFilterChanged |= RemoveEnablePathFilterDelegate(OwnerName);
+	}
+
 	if (bFilterChanged && !bSuppressOnFilterChanged)
 	{
 		OnFilterChanged().Broadcast();
@@ -597,7 +651,11 @@ FString FPathPermissionList::ToString() const
 			StringBuilder.Append(Path);
 			StringBuilder.AppendChar(TCHAR('"'));
 			StringBuilder.AppendChar(TCHAR(' '));
-			SortAndAppendOwners(*Tree.Find(Path));
+
+			if (const FPermissionListOwners* List = Tree.Find(Path))
+			{
+				SortAndAppendOwners(*List);
+			}
 			StringBuilder.AppendChar(TCHAR('\n'));
 		}
 	};
@@ -615,4 +673,26 @@ FString FPathPermissionList::ToString() const
 	}
 
 	return FString(StringBuilder);
+}
+
+void FPathPermissionList::AddEnablePathFilterDelegate(FName Owner, FDelegateEnablePathFilter Delegate)
+{
+	EnablePathFilterDelegates.Add(Owner, Delegate);
+
+	if (!bSuppressOnFilterChanged)
+	{
+		OnFilterChanged().Broadcast();
+	}
+}
+
+bool FPathPermissionList::RemoveEnablePathFilterDelegate(FName Owner)
+{
+	bool bRemovedDelegate = EnablePathFilterDelegates.Remove(Owner) > 0;
+
+	if (bRemovedDelegate && !bSuppressOnFilterChanged)
+	{
+		OnFilterChanged().Broadcast();
+	}
+
+	return bRemovedDelegate;
 }
\ No newline at end of file
diff --git a/Engine/Source/Runtime/Core/Public/Misc/NamePermissionList.h b/Engine/Source/Runtime/Core/Public/Misc/NamePermissionList.h
--- a/Engine/Source/Runtime/Core/Public/Misc/NamePermissionList.h
+++ b/Engine/Source/Runtime/Core/Public/Misc/NamePermissionList.h
@@ -145,6 +145,8 @@ enum class EPathPermissionPrefixResult
 	PassRecursive,
 };
 
+DECLARE_DELEGATE_RetVal(bool, FDelegateEnablePathFilter);
+
 /**
  * Set of paths that are allowd and/or denied for certain use cases.
  * A permission list may contain
@@ -379,6 +381,9 @@ public:
 	/** Triggered when filter changes */
 	FSimpleMulticastDelegate& OnFilterChanged() { return OnFilterChangedDelegate; }
 
+	CORE_API void AddEnablePathFilterDelegate(FName Owner, FDelegateEnablePathFilter Delegate);
+	CORE_API bool RemoveEnablePathFilterDelegate(FName Owner);
+
 	/** Dumps the path permission list details into a multi-line string */
 	CORE_API FString ToString() const;
 
@@ -402,9 +407,13 @@ protected:
 	/** Triggered when filter changes */
 	FSimpleMulticastDelegate OnFilterChangedDelegate;
 
+	FDelegateEnablePathFilter OnEnablePathFilterDelegate;
+
 	/** Temporarily prevent delegate from being triggered */
 	bool bSuppressOnFilterChanged = false;
 
 	/** Type of paths this list represent */
 	EPathPermissionListType ListType;
+
+	TMap<FName, FDelegateEnablePathFilter> EnablePathFilterDelegates;
 };
