[RaceDetector]
- Fix potential crash since during a race we could end up releasing a 4KB page containing the Next pointer for a linked list. Put it back in the free queue instead.
- Release the non-aligned sync object during shadow memory reset

#Synced-CL Unknown

diff --git a/Engine/Source/Runtime/Core/Internal/Sanitizer/RaceDetectorTypes.h b/Engine/Source/Runtime/Core/Internal/Sanitizer/RaceDetectorTypes.h
--- a/Engine/Source/Runtime/Core/Internal/Sanitizer/RaceDetectorTypes.h
+++ b/Engine/Source/Runtime/Core/Internal/Sanitizer/RaceDetectorTypes.h
@@ -912,14 +912,7 @@ namespace UE::Sanitizer::RaceDetector {
 
 		INSTRUMENTATION_FUNCTION_ATTRIBUTES ~FSyncObjectBank()
 		{
-			for (int32 Index = 0; Index < 7; ++Index)
-			{
-				if (SyncObjects[Index])
-				{
-					delete SyncObjects[Index];
-				}
-			}
-
+			Reset();
 			FPlatformAtomics::InterlockedDecrement(&ObjectCount);
 		}
 
@@ -971,6 +964,18 @@ namespace UE::Sanitizer::RaceDetector {
 			return RefCount;
 		}
 
+		INSTRUMENTATION_FUNCTION_ATTRIBUTES void Reset()
+		{
+			for (int32 Index = 0; Index < 7; ++Index)
+			{
+				if (SyncObjects[Index])
+				{
+					delete SyncObjects[Index];
+					SyncObjects[Index] = nullptr;
+				}
+			}
+		}
+
 		static volatile int64 ObjectCount;
 	private:
 		volatile int32 RefCount = 1;
diff --git a/Engine/Source/Runtime/Core/Private/Sanitizer/RaceDetector.cpp b/Engine/Source/Runtime/Core/Private/Sanitizer/RaceDetector.cpp
--- a/Engine/Source/Runtime/Core/Private/Sanitizer/RaceDetector.cpp
+++ b/Engine/Source/Runtime/Core/Private/Sanitizer/RaceDetector.cpp
@@ -393,6 +393,12 @@ namespace UE::Sanitizer::RaceDetector {
 
 	INSTRUMENTATION_FUNCTION_ATTRIBUTES void PushFreeObjectBank(FSyncObjectBank* SyncObjectBank)
 	{
+		RACEDETECTOR_CHECK(SyncObjectBank->GetRefCount() == 1);
+
+		// Cleanup all non-aligned objects, those are unlikely to be needed when the
+		// sync bank is reused for a different atomic address.
+		SyncObjectBank->Reset();
+
 		while (true)
 		{
 			FSyncObjectBank* LocalSyncBankHead = (FSyncObjectBank*)FreeObjectBankHead;
@@ -467,7 +473,9 @@ namespace UE::Sanitizer::RaceDetector {
 				LocalBank = OldObjectBank;
 
 				// Should not happen very often, and at most once per thread taking part of the race
-				int32 Result = NewObjectBank->Release();
+				// Put it back in the list, freeing it now could cause a crash since that is fresh out a linked-list
+				// that another thread might still try to dereference.
+				PushFreeObjectBank(NewObjectBank);
 				RACEDETECTOR_CHECK(Result == 0);
 			}
 
