Fixed IDO still referencing old re-instantiated class after re-instantiation.

#rnx
#Synced-CL 48408763

diff --git a/Engine/Source/Runtime/Core/Public/Serialization/ArchiveSerializedPropertyChain.h b/Engine/Source/Runtime/Core/Public/Serialization/ArchiveSerializedPropertyChain.h
--- a/Engine/Source/Runtime/Core/Public/Serialization/ArchiveSerializedPropertyChain.h
+++ b/Engine/Source/Runtime/Core/Public/Serialization/ArchiveSerializedPropertyChain.h
@@ -3,6 +3,8 @@
 #pragma once
 
 #include "Containers/Array.h"
+#include "Misc/NotNull.h"
+#include "Templates/Function.h"
 
 /** Structure containing the stack of properties that are currently being serialized by an archive */
 struct FArchiveSerializedPropertyChain
@@ -13,6 +15,7 @@ public:
 		: SerializedPropertyChainUpdateCount(0)
 #if WITH_EDITORONLY_DATA
 		, EditorOnlyPropertyStack(0)
+		, UnknownPropertyStack(0)
 #endif
 	{
 	}
@@ -35,7 +38,13 @@ public:
 		{
 			++EditorOnlyPropertyStack;
 		}
-#endif
+		PRAGMA_DISABLE_INTERNAL_WARNINGS
+		if (IsUnknownPropertyFunc.IsSet() && IsUnknownPropertyFunc(InProperty))
+		{
+			++UnknownPropertyStack;
+		}
+		PRAGMA_ENABLE_INTERNAL_WARNINGS
+#endif // WITH_EDITORONLY_DATA
 	}
 
 	/**
@@ -58,7 +67,13 @@ public:
 			--EditorOnlyPropertyStack;
 			check(EditorOnlyPropertyStack >= 0);
 		}
-#endif
+		PRAGMA_DISABLE_INTERNAL_WARNINGS
+		if (IsUnknownPropertyFunc.IsSet() && IsUnknownPropertyFunc(InProperty))
+		{
+			--UnknownPropertyStack;
+		}
+		PRAGMA_ENABLE_INTERNAL_WARNINGS
+#endif // WITH_EDITORONLY_DATA
 	}
 
 	/**
@@ -107,6 +122,21 @@ public:
 #endif
 	}
 
+	/**
+	 * Check to see whether there are any properties on the stack is unknown
+	 */
+	bool HasUnknownProperty() const
+	{
+#if WITH_EDITORONLY_DATA
+		PRAGMA_DISABLE_INTERNAL_WARNINGS
+		checkf(IsUnknownPropertyFunc.IsSet(), TEXT("There was no function pointer provided to retrieve this information, so the result will not be correctly returned."))
+		PRAGMA_ENABLE_INTERNAL_WARNINGS
+		return UnknownPropertyStack > 0;
+#else // WITH_EDITORONLY_DATA
+		return false;
+#endif// WITH_EDITORONLY_DATA
+	}
+
 	UE_FORCEINLINE_HINT TArray<class FProperty*, TInlineAllocator<8>>::TConstIterator GetRootIterator() const
 	{
 		return SerializedPropertyChain.CreateConstIterator();
@@ -126,5 +156,13 @@ private:
 #if WITH_EDITORONLY_DATA
 	/** Non-zero if on the current stack of serialized properties there's an editor-only property */
 	int32 EditorOnlyPropertyStack;
-#endif
+
+	/** Non-zero if on the current stack of serialized properties is a unknown property */
+	int32 UnknownPropertyStack;
+
+public:
+	PRAGMA_DISABLE_INTERNAL_WARNINGS
+	UE_INTERNAL static inline TFunction<bool(TNotNull<FProperty*>)> IsUnknownPropertyFunc;
+	PRAGMA_ENABLE_INTERNAL_WARNINGS
+#endif // WITH_EDITORONLY_DATA
 };
\ No newline at end of file
diff --git a/Engine/Source/Runtime/CoreUObject/Internal/UObject/InstanceDataObjectUtils.h b/Engine/Source/Runtime/CoreUObject/Internal/UObject/InstanceDataObjectUtils.h
--- a/Engine/Source/Runtime/CoreUObject/Internal/UObject/InstanceDataObjectUtils.h
+++ b/Engine/Source/Runtime/CoreUObject/Internal/UObject/InstanceDataObjectUtils.h
@@ -3,6 +3,7 @@
 #pragma once
 
 #include "CoreTypes.h"
+#include "Misc/NotNull.h"
 
 #if WITH_EDITORONLY_DATA
 
@@ -43,7 +44,10 @@ UE_API bool IsInstanceDataObjectOuterChainEnabled();
 UE_API bool StructContainsLooseProperties(const UStruct* Struct);
 UE_API bool StructIsInstanceDataObjectStruct(const UStruct* Struct);
 
-UE_API bool IsPropertyLoose(const FProperty* Property);
+/** @return Whether the property's has the "LooseProperty" MetaData */
+UE_API bool IsPropertyLoose(TNotNull<const FProperty*> Property);
+/** @return Whether the provided property's or any properties in the archive serialization chain has the "LooseProperty" MetaData */
+UE_API bool IsSerializingUnknownData(const FArchive& Archive, const FProperty* InProperty = nullptr);
 
 /** Query if placeholder support is enabled for a specific import class type. */
 bool CanCreatePropertyBagPlaceholderTypeForImportClass(const UClass* ClassKind, const UE::Serialization::Private::FImportTypeHierarchy* ImportTypeHierarchy);
diff --git a/Engine/Source/Runtime/CoreUObject/Private/UObject/Class.cpp b/Engine/Source/Runtime/CoreUObject/Private/UObject/Class.cpp
--- a/Engine/Source/Runtime/CoreUObject/Private/UObject/Class.cpp
+++ b/Engine/Source/Runtime/CoreUObject/Private/UObject/Class.cpp
@@ -1,4 +1,4 @@
-// Copyright Epic Games, Inc. All Rights Reserved.
+﻿// Copyright Epic Games, Inc. All Rights Reserved.
 
 #include "UObject/Class.h"
 #include "HAL/ThreadSafeBool.h"
@@ -1846,8 +1846,11 @@ void UStruct::SerializeVersionedTaggedProperties(FStructuredArchive::FSlot Slot,
 					}
 #if WITH_EDITORONLY_DATA
 					// Handle the sub object shadowing serialization,
-					// do not skip this property if loading into a loose property or into a placeholder
-					else if (!SerializeContext->bImpersonateProperties && Tag.OverrideOperation == EOverriddenPropertyOperation::SubObjectsShadowing)
+					// do not skip this property if loading into a unknown property or into a placeholder
+					// When loading into the IDO, we should only consider the unknown properties.
+					// (IDO are serializing properties using bImpersonateProperties(true) but they are not necessary unknown)
+					else if (Tag.OverrideOperation == EOverriddenPropertyOperation::SubObjectsShadowing && 
+						(!SerializeContext->bImpersonateProperties || !UE::IsSerializingUnknownData(UnderlyingArchive, Property)))
 					{
 						UnderlyingArchive.Seek(StartOfProperty + Tag.Size);
 						bAdvanceProperty = true;
diff --git a/Engine/Source/Runtime/CoreUObject/Private/UObject/InstanceDataObjectUtils.cpp b/Engine/Source/Runtime/CoreUObject/Private/UObject/InstanceDataObjectUtils.cpp
--- a/Engine/Source/Runtime/CoreUObject/Private/UObject/InstanceDataObjectUtils.cpp
+++ b/Engine/Source/Runtime/CoreUObject/Private/UObject/InstanceDataObjectUtils.cpp
@@ -15,6 +15,7 @@
 #include "Serialization/ObjectWriter.h"
 #include "Serialization/StructuredArchive.h"
 #include "Serialization/ArchiveCountMem.h"
+#include "Serialization/ArchiveSerializedPropertyChain.h"
 #include "Serialization/ImportTypeHierarchy.h"
 #include "UObject/Class.h"
 #include "UObject/CommonFNames.h"
@@ -1203,11 +1204,44 @@ namespace UE
 		// struct sub-properties don't need to be handled. If the sub-properties are loose, this method will get called on them regardless when they're created.
 	}
 
-	bool IsPropertyLoose(const FProperty* Property)
+	bool IsPropertyLoose(TNotNull<const FProperty*> Property)
 	{
-		return Property->GetBoolMetaData(NAME_IsLooseMetadata);
+		return IsInstanceDataObjectSupportEnabled() && Property->GetBoolMetaData(NAME_IsLooseMetadata);
 	}
 
+	struct FSetupIsUnknownPropertyFnPtr
+	{
+		FSetupIsUnknownPropertyFnPtr()
+		{
+			FArchiveSerializedPropertyChain::IsUnknownPropertyFunc = &IsPropertyLoose;
+		}
+		~FSetupIsUnknownPropertyFnPtr()
+		{
+			FArchiveSerializedPropertyChain::IsUnknownPropertyFunc.Reset();
+		}
+	} SetupGetIsUnknownPropertyFnPtr;
+
+	bool IsSerializingUnknownData(const FArchive& Archive, const FProperty* InProperty /* = nullptr*/)
+	{
+		if (!IsInstanceDataObjectSupportEnabled())
+		{
+			return false;
+		}
+
+		bool bHasUnknownProperty = false;
+		if (const FArchiveSerializedPropertyChain * PropertyChain = Archive.GetSerializedPropertyChain())
+		{
+			bHasUnknownProperty = PropertyChain->HasUnknownProperty();
+		}
+		if (!bHasUnknownProperty && InProperty)
+		{
+			bHasUnknownProperty = InProperty->HasMetaData(NAME_IsLooseMetadata);
+		}
+
+		return bHasUnknownProperty;
+	}
+
+
 	// creates a struct that doesn't contain unknown (loose) properties except those in sub-structs
 	static UStruct* CreateInstanceDataObjectSuperHierarchy(const UClass* StructClass, UStruct* TemplateStruct, UObject* Outer, const FPropertyPathNameTree* PropertyTree, const FUnknownEnumNames* EnumNames, TSet<FPropertyPathName>& OutSuperPropertyPathsFromTree)
 	{
