Core: Updated FName tests to assert unless compiling for low level test targets

#rnx
#Synced-CL 32500503

diff --git a/Engine/Source/Runtime/Core/Private/UObject/UnrealNames.cpp b/Engine/Source/Runtime/Core/Private/UObject/UnrealNames.cpp
--- a/Engine/Source/Runtime/Core/Private/UObject/UnrealNames.cpp
+++ b/Engine/Source/Runtime/Core/Private/UObject/UnrealNames.cpp
@@ -4001,17 +4001,23 @@ FUtf8StringBuilderBase& operator<<(FUtf8StringBuilderBase& Builder, FNameEntryId
 
 #if WITH_TESTS
 
+#if WITH_LOW_LEVEL_TESTS
+#define NAME_CHECK CHECK
+#else
+#define NAME_CHECK check
+#endif
+
 template <typename CharType, int N>
 void CheckLazyName(const CharType(&Literal)[N])
 {
-	CHECK(FName(Literal) == FLazyName(Literal));
-	CHECK(FLazyName(Literal) == FName(Literal));
-	CHECK(FLazyName(Literal) == FLazyName(Literal));
-	CHECK(FName(Literal) == FLazyName(Literal).Resolve());
+	NAME_CHECK(FName(Literal) == FLazyName(Literal));
+	NAME_CHECK(FLazyName(Literal) == FName(Literal));
+	NAME_CHECK(FLazyName(Literal) == FLazyName(Literal));
+	NAME_CHECK(FName(Literal) == FLazyName(Literal).Resolve());
 
 	CharType Literal2[N];
 	FMemory::Memcpy(Literal2, Literal);
-	CHECK(FLazyName(Literal) == FLazyName(Literal2));
+	NAME_CHECK(FLazyName(Literal) == FLazyName(Literal2));
 
 	constexpr ESearchCase::Type ComparisonMode =
 #if WITH_CASE_PRESERVING_NAME
@@ -4020,7 +4026,7 @@ void CheckLazyName(const CharType(&Literal)[N])
 		ESearchCase::IgnoreCase;
 #endif
 
-	CHECK(WriteToString<64>(FLazyName(Literal).Resolve()).ToView().Equals(WriteToString<64>(Literal).ToView(), ComparisonMode));
+	NAME_CHECK(WriteToString<64>(FLazyName(Literal).Resolve()).ToView().Equals(WriteToString<64>(Literal).ToView(), ComparisonMode));
 }
 
 static void TestNameBatch();
@@ -4045,31 +4051,31 @@ struct FNameAutoTestTruncateHelper
 		if constexpr (RequiredLen >= TestStrLen)
 		{
 			// Whole string should be returned
-			CHECK(Len == TestStrLen);
+			NAME_CHECK(Len == TestStrLen);
 		}
 		else
 		{
 			// String should be truncated
-			CHECK(Len == RequiredLen);
+			NAME_CHECK(Len == RequiredLen);
 		}
 
 		for (uint32 Index = 0; Index <= RequiredLen; ++Index)
 		{
 			if (Index < Len)
 			{
-				CHECK(Buffer.Data[Index] == TestStr[Index]);
+				NAME_CHECK(Buffer.Data[Index] == TestStr[Index]);
 			}
 			else if (Index == Len)
 			{
-				CHECK(Buffer.Data[Index] == 0); // null terminator
+				NAME_CHECK(Buffer.Data[Index] == 0); // null terminator
 			}
 			else
 			{
-				CHECK(Buffer.Data[Index] == CanaryValue); // should not have been written
+				NAME_CHECK(Buffer.Data[Index] == CanaryValue); // should not have been written
 			}
 		}
 
-		CHECK(Buffer.Canary == CanaryValue);
+		NAME_CHECK(Buffer.Canary == CanaryValue);
 	}
 
 	template <std::size_t... Ints>
@@ -4086,14 +4092,14 @@ struct FNameAutoTestTruncateHelper
 
 TEST_CASE_NAMED(FNameTest, "System::Core::Name", "[ApplicationContextMask][SmokeFilter]")
 {
-	CHECK(FNameHash::IsAnsiNone("None", 4) == 1);
-	CHECK(FNameHash::IsAnsiNone("none", 4) == 1);
-	CHECK(FNameHash::IsAnsiNone("NONE", 4) == 1);
-	CHECK(FNameHash::IsAnsiNone("nOnE", 4) == 1);
-	CHECK(FNameHash::IsAnsiNone("None", 5) == 0);
-	CHECK(FNameHash::IsAnsiNone(TEXT("None"), 4) == 0);
-	CHECK(FNameHash::IsAnsiNone("nono", 4) == 0);
-	CHECK(FNameHash::IsAnsiNone("enon", 4) == 0);
+	NAME_CHECK(FNameHash::IsAnsiNone("None", 4) == 1);
+	NAME_CHECK(FNameHash::IsAnsiNone("none", 4) == 1);
+	NAME_CHECK(FNameHash::IsAnsiNone("NONE", 4) == 1);
+	NAME_CHECK(FNameHash::IsAnsiNone("nOnE", 4) == 1);
+	NAME_CHECK(FNameHash::IsAnsiNone("None", 5) == 0);
+	NAME_CHECK(FNameHash::IsAnsiNone(TEXT("None"), 4) == 0);
+	NAME_CHECK(FNameHash::IsAnsiNone("nono", 4) == 0);
+	NAME_CHECK(FNameHash::IsAnsiNone("enon", 4) == 0);
 
 	const FName AutoTest_1("AutoTest_1");
 	const FName autoTest_1("autoTest_1");
@@ -4101,141 +4107,141 @@ TEST_CASE_NAMED(FNameTest, "System::Core::Name", "[ApplicationContextMask][Smoke
 	const FName AutoTest_2(TEXT("AutoTest_2"));
 	const FName AutoTestB_2(TEXT("AutoTestB_2"));
 
-	CHECK(AutoTest_1.GetNumber() == NAME_EXTERNAL_TO_INTERNAL(1));
-	CHECK(autoTest_1.GetNumber() == NAME_EXTERNAL_TO_INTERNAL(1));
-	CHECK(autoTeSt_1.GetNumber() == NAME_EXTERNAL_TO_INTERNAL(1));
-	CHECK(AutoTest_2.GetNumber() == NAME_EXTERNAL_TO_INTERNAL(2));
-	CHECK(AutoTestB_2.GetNumber() == NAME_EXTERNAL_TO_INTERNAL(2));
+	NAME_CHECK(AutoTest_1.GetNumber() == NAME_EXTERNAL_TO_INTERNAL(1));
+	NAME_CHECK(autoTest_1.GetNumber() == NAME_EXTERNAL_TO_INTERNAL(1));
+	NAME_CHECK(autoTeSt_1.GetNumber() == NAME_EXTERNAL_TO_INTERNAL(1));
+	NAME_CHECK(AutoTest_2.GetNumber() == NAME_EXTERNAL_TO_INTERNAL(2));
+	NAME_CHECK(AutoTestB_2.GetNumber() == NAME_EXTERNAL_TO_INTERNAL(2));
 
-	CHECK(AutoTest_1 != AutoTest_2);
-	CHECK(AutoTest_1 == autoTest_1);
-	CHECK(AutoTest_1 == autoTeSt_1);
+	NAME_CHECK(AutoTest_1 != AutoTest_2);
+	NAME_CHECK(AutoTest_1 == autoTest_1);
+	NAME_CHECK(AutoTest_1 == autoTeSt_1);
 
-	CHECK(AutoTest_1.ToUnstableInt() == autoTest_1.ToUnstableInt());
-	CHECK(AutoTest_1.ToUnstableInt() != AutoTest_2.ToUnstableInt());
+	NAME_CHECK(AutoTest_1.ToUnstableInt() == autoTest_1.ToUnstableInt());
+	NAME_CHECK(AutoTest_1.ToUnstableInt() != AutoTest_2.ToUnstableInt());
 
 	TCHAR Buffer[FName::StringBufferSize];
 
 #if WITH_CASE_PRESERVING_NAME
-	CHECK(!FCString::Strcmp(*AutoTest_1.ToString(), TEXT("AutoTest_1")));
-	CHECK(!FCString::Strcmp(*autoTest_1.ToString(), TEXT("autoTest_1")));
-	CHECK(!FCString::Strcmp(*autoTeSt_1.ToString(), TEXT("autoTeSt_1")));
-	CHECK(!FCString::Strcmp(*AutoTestB_2.ToString(), TEXT("AutoTestB_2")));
+	NAME_CHECK(!FCString::Strcmp(*AutoTest_1.ToString(), TEXT("AutoTest_1")));
+	NAME_CHECK(!FCString::Strcmp(*autoTest_1.ToString(), TEXT("autoTest_1")));
+	NAME_CHECK(!FCString::Strcmp(*autoTeSt_1.ToString(), TEXT("autoTeSt_1")));
+	NAME_CHECK(!FCString::Strcmp(*AutoTestB_2.ToString(), TEXT("AutoTestB_2")));
 	
-	CHECK(FName("ABC").ToString(Buffer) == 3);
-	CHECK(!FCString::Strcmp(Buffer, TEXT("ABC")));
-	CHECK(FName("abc").ToString(Buffer) == 3);
-	CHECK(!FCString::Strcmp(Buffer, TEXT("abc")));
-	CHECK(FName(TEXT("abc")).ToString(Buffer) == 3);
-	CHECK(!FCString::Strcmp(Buffer, TEXT("abc")));
-	CHECK(FName("ABC_0").ToString(Buffer) == 5);
-	CHECK(!FCString::Strcmp(Buffer, TEXT("ABC_0")));
-	CHECK(FName("ABC_10").ToString(Buffer) == 6);
-	CHECK(!FCString::Strcmp(Buffer, TEXT("ABC_10")));	
+	NAME_CHECK(FName("ABC").ToString(Buffer) == 3);
+	NAME_CHECK(!FCString::Strcmp(Buffer, TEXT("ABC")));
+	NAME_CHECK(FName("abc").ToString(Buffer) == 3);
+	NAME_CHECK(!FCString::Strcmp(Buffer, TEXT("abc")));
+	NAME_CHECK(FName(TEXT("abc")).ToString(Buffer) == 3);
+	NAME_CHECK(!FCString::Strcmp(Buffer, TEXT("abc")));
+	NAME_CHECK(FName("ABC_0").ToString(Buffer) == 5);
+	NAME_CHECK(!FCString::Strcmp(Buffer, TEXT("ABC_0")));
+	NAME_CHECK(FName("ABC_10").ToString(Buffer) == 6);
+	NAME_CHECK(!FCString::Strcmp(Buffer, TEXT("ABC_10")));	
 #endif
 
-	CHECK(autoTest_1.GetComparisonIndex() == AutoTest_2.GetComparisonIndex());
-	CHECK(autoTest_1.GetPlainNameString() == AutoTest_1.GetPlainNameString());
-	CHECK(autoTest_1.GetPlainNameString() == AutoTest_2.GetPlainNameString());
-	CHECK(*AutoTestB_2.GetPlainNameString() != *AutoTest_2.GetPlainNameString());
-	CHECK(AutoTestB_2.GetNumber() == AutoTest_2.GetNumber());
-	CHECK(autoTest_1.GetNumber() != AutoTest_2.GetNumber());
-	CHECK(AutoTest_1.GetNumber() == autoTest_1.GetNumber());
-
-	CHECK(FCStringAnsi::Strlen("None") == FName().GetStringLength());
-	CHECK(FCStringAnsi::Strlen("ABC") == FName("ABC").GetStringLength());
-	CHECK(FCStringAnsi::Strlen("ABC_0") == FName("ABC_0").GetStringLength());
-	CHECK(FCStringAnsi::Strlen("ABC_9") == FName("ABC_9").GetStringLength());
-	CHECK(FCStringAnsi::Strlen("ABC_10") == FName("ABC_10").GetStringLength());
-	CHECK(FCStringAnsi::Strlen("ABC_2000000000") == FName("ABC_2000000000").GetStringLength());
-	CHECK(FCStringAnsi::Strlen("ABC_4000000000") == FName("ABC_4000000000").GetStringLength());
+	NAME_CHECK(autoTest_1.GetComparisonIndex() == AutoTest_2.GetComparisonIndex());
+	NAME_CHECK(autoTest_1.GetPlainNameString() == AutoTest_1.GetPlainNameString());
+	NAME_CHECK(autoTest_1.GetPlainNameString() == AutoTest_2.GetPlainNameString());
+	NAME_CHECK(*AutoTestB_2.GetPlainNameString() != *AutoTest_2.GetPlainNameString());
+	NAME_CHECK(AutoTestB_2.GetNumber() == AutoTest_2.GetNumber());
+	NAME_CHECK(autoTest_1.GetNumber() != AutoTest_2.GetNumber());
+	NAME_CHECK(AutoTest_1.GetNumber() == autoTest_1.GetNumber());
+
+	NAME_CHECK(FCStringAnsi::Strlen("None") == FName().GetStringLength());
+	NAME_CHECK(FCStringAnsi::Strlen("ABC") == FName("ABC").GetStringLength());
+	NAME_CHECK(FCStringAnsi::Strlen("ABC_0") == FName("ABC_0").GetStringLength());
+	NAME_CHECK(FCStringAnsi::Strlen("ABC_9") == FName("ABC_9").GetStringLength());
+	NAME_CHECK(FCStringAnsi::Strlen("ABC_10") == FName("ABC_10").GetStringLength());
+	NAME_CHECK(FCStringAnsi::Strlen("ABC_2000000000") == FName("ABC_2000000000").GetStringLength());
+	NAME_CHECK(FCStringAnsi::Strlen("ABC_4000000000") == FName("ABC_4000000000").GetStringLength());
 
 	const FName NullName(static_cast<ANSICHAR*>(nullptr));
-	CHECK(NullName.IsNone());
-	CHECK(NullName == FName(static_cast<WIDECHAR*>(nullptr)));
-	CHECK(NullName == FName(NAME_None));
-	CHECK(NullName == FName());
-	CHECK(NullName == FName(""));
-	CHECK(NullName == FName(TEXT("")));
-	CHECK(NullName == FName("None"));
-	CHECK(NullName == FName("none"));
-	CHECK(NullName == FName("NONE"));
-	CHECK(NullName == FName(TEXT("None")));
-	CHECK(FName().ToEName());
-	CHECK(*FName().ToEName() == NAME_None);
-	CHECK(NullName.GetComparisonIndex().ToUnstableInt() == 0);
+	NAME_CHECK(NullName.IsNone());
+	NAME_CHECK(NullName == FName(static_cast<WIDECHAR*>(nullptr)));
+	NAME_CHECK(NullName == FName(NAME_None));
+	NAME_CHECK(NullName == FName());
+	NAME_CHECK(NullName == FName(""));
+	NAME_CHECK(NullName == FName(TEXT("")));
+	NAME_CHECK(NullName == FName("None"));
+	NAME_CHECK(NullName == FName("none"));
+	NAME_CHECK(NullName == FName("NONE"));
+	NAME_CHECK(NullName == FName(TEXT("None")));
+	NAME_CHECK(FName().ToEName());
+	NAME_CHECK(*FName().ToEName() == NAME_None);
+	NAME_CHECK(NullName.GetComparisonIndex().ToUnstableInt() == 0);
 
 	// Numbered nones
 	FName None_7 = FName("None_7");
-	CHECK(None_7.ToString() == "None_7");
-	CHECK(None_7.GetComparisonIndex() == FNameEntryId::FromEName(EName::None));
-	CHECK(None_7.GetNumber() == NAME_EXTERNAL_TO_INTERNAL(7));
-	CHECK(FName(EName::None, NAME_EXTERNAL_TO_INTERNAL(7)) == None_7);
-	CHECK(FName(FName(), NAME_EXTERNAL_TO_INTERNAL(7)) == None_7);
-	CHECK(FName("None", NAME_EXTERNAL_TO_INTERNAL(7)) == None_7);
-	CHECK(FName(FNameEntryId::FromEName(EName::None), FNameEntryId::FromEName(EName::None), NAME_EXTERNAL_TO_INTERNAL(7)) == None_7);
-	CHECK(!None_7.IsNone());
+	NAME_CHECK(None_7.ToString() == "None_7");
+	NAME_CHECK(None_7.GetComparisonIndex() == FNameEntryId::FromEName(EName::None));
+	NAME_CHECK(None_7.GetNumber() == NAME_EXTERNAL_TO_INTERNAL(7));
+	NAME_CHECK(FName(EName::None, NAME_EXTERNAL_TO_INTERNAL(7)) == None_7);
+	NAME_CHECK(FName(FName(), NAME_EXTERNAL_TO_INTERNAL(7)) == None_7);
+	NAME_CHECK(FName("None", NAME_EXTERNAL_TO_INTERNAL(7)) == None_7);
+	NAME_CHECK(FName(FNameEntryId::FromEName(EName::None), FNameEntryId::FromEName(EName::None), NAME_EXTERNAL_TO_INTERNAL(7)) == None_7);
+	NAME_CHECK(!None_7.IsNone());
 
 	// "_[num]" is considered None to keep semantics of old implementation, unlike "None_[num]"
-	CHECK(FName("_7").IsNone()); 
-	CHECK(FName("_0").IsNone());
-	CHECK(FName("_2147483646").IsNone());
-	CHECK(!FName("_2147483647").IsNone()); // invalid number
+	NAME_CHECK(FName("_7").IsNone()); 
+	NAME_CHECK(FName("_0").IsNone());
+	NAME_CHECK(FName("_2147483646").IsNone());
+	NAME_CHECK(!FName("_2147483647").IsNone()); // invalid number
 
 	const FName Cylinder(NAME_Cylinder);
-	CHECK(Cylinder == FName("Cylinder"));
-	CHECK(Cylinder.ToEName());
-	CHECK(*Cylinder.ToEName() == NAME_Cylinder);
-	CHECK(Cylinder.GetPlainNameString() == TEXT("Cylinder"));
+	NAME_CHECK(Cylinder == FName("Cylinder"));
+	NAME_CHECK(Cylinder.ToEName());
+	NAME_CHECK(*Cylinder.ToEName() == NAME_Cylinder);
+	NAME_CHECK(Cylinder.GetPlainNameString() == TEXT("Cylinder"));
 
 	// Test numbers
-	CHECK(FName("Text_0") == FName("Text", NAME_EXTERNAL_TO_INTERNAL(0)));
-	CHECK(FName("Text_1") == FName("Text", NAME_EXTERNAL_TO_INTERNAL(1)));
-	CHECK(FName("Text_1_0") == FName("Text_1", NAME_EXTERNAL_TO_INTERNAL(0)));
-	CHECK(FName("Text_0_1") == FName("Text_0", NAME_EXTERNAL_TO_INTERNAL(1)));
-	CHECK(FName("Text_00") == FName("Text_00", NAME_NO_NUMBER_INTERNAL));
-	CHECK(FName("Text_01") == FName("Text_01", NAME_NO_NUMBER_INTERNAL));
+	NAME_CHECK(FName("Text_0") == FName("Text", NAME_EXTERNAL_TO_INTERNAL(0)));
+	NAME_CHECK(FName("Text_1") == FName("Text", NAME_EXTERNAL_TO_INTERNAL(1)));
+	NAME_CHECK(FName("Text_1_0") == FName("Text_1", NAME_EXTERNAL_TO_INTERNAL(0)));
+	NAME_CHECK(FName("Text_0_1") == FName("Text_0", NAME_EXTERNAL_TO_INTERNAL(1)));
+	NAME_CHECK(FName("Text_00") == FName("Text_00", NAME_NO_NUMBER_INTERNAL));
+	NAME_CHECK(FName("Text_01") == FName("Text_01", NAME_NO_NUMBER_INTERNAL));
 
 	// Test unterminated strings
-	CHECK(FName("") == FName(0, "Unused"));
-	CHECK(FName("Used") == FName(4, "UsedUnused"));
-	CHECK(FName("Used") == FName(4, "Used"));
-	CHECK(FName("Used_0") == FName(6, "Used_01"));
-	CHECK(FName("Used_01") == FName(7, "Used_012"));
-	CHECK(FName("Used_123") == FName(8, "Used_123456"));
-	CHECK(FName("Used_123") == FName(8, "Used_123_456"));
-	CHECK(FName("Used_123") == FName(8, TEXT("Used_123456")));
-	CHECK(FName("Used_123") == FName(8, TEXT("Used_123_456")));
-	CHECK(FName("Used_2147483646") == FName(15, TEXT("Used_2147483646123")));
-	CHECK(FName("Used_2147483647") == FName(15, TEXT("Used_2147483647123")));
-	CHECK(FName("Used_2147483648") == FName(15, TEXT("Used_2147483648123")));
+	NAME_CHECK(FName("") == FName(0, "Unused"));
+	NAME_CHECK(FName("Used") == FName(4, "UsedUnused"));
+	NAME_CHECK(FName("Used") == FName(4, "Used"));
+	NAME_CHECK(FName("Used_0") == FName(6, "Used_01"));
+	NAME_CHECK(FName("Used_01") == FName(7, "Used_012"));
+	NAME_CHECK(FName("Used_123") == FName(8, "Used_123456"));
+	NAME_CHECK(FName("Used_123") == FName(8, "Used_123_456"));
+	NAME_CHECK(FName("Used_123") == FName(8, TEXT("Used_123456")));
+	NAME_CHECK(FName("Used_123") == FName(8, TEXT("Used_123_456")));
+	NAME_CHECK(FName("Used_2147483646") == FName(15, TEXT("Used_2147483646123")));
+	NAME_CHECK(FName("Used_2147483647") == FName(15, TEXT("Used_2147483647123")));
+	NAME_CHECK(FName("Used_2147483648") == FName(15, TEXT("Used_2147483648123")));
 
 	// Test wide strings
 	FString Wide("Wide ");
 	Wide[4] = 60000;
 	FName WideName(*Wide);
-	CHECK(WideName.GetPlainNameString() == Wide);
-	CHECK(FName(*Wide).GetPlainNameString() == Wide);
-	CHECK(FName(*Wide).ToString(Buffer) == 5);
-	CHECK(!FCString::Strcmp(Buffer, *Wide));
-	CHECK(Wide.Len() == WideName.GetStringLength());
+	NAME_CHECK(WideName.GetPlainNameString() == Wide);
+	NAME_CHECK(FName(*Wide).GetPlainNameString() == Wide);
+	NAME_CHECK(FName(*Wide).ToString(Buffer) == 5);
+	NAME_CHECK(!FCString::Strcmp(Buffer, *Wide));
+	NAME_CHECK(Wide.Len() == WideName.GetStringLength());
 	FString WideLong = FString::ChrN(1000, 60000);
-	CHECK(FName(*WideLong).GetPlainNameString() == WideLong);
+	NAME_CHECK(FName(*WideLong).GetPlainNameString() == WideLong);
 
 
 	// Check that FNAME_Find doesn't add entries
 	static bool Once = true;
 	if (Once)
 	{
-		CHECK(FName("UniqueUnicorn!!", FNAME_Find) == FName());
+		NAME_CHECK(FName("UniqueUnicorn!!", FNAME_Find) == FName());
 
 		// Check that FNAME_Find can find entries
 		const FName UniqueName("UniqueUnicorn!!", FNAME_Add);
-		CHECK(FName("UniqueUnicorn!!", FNAME_Find) == UniqueName);
-		CHECK(FName(TEXT("UniqueUnicorn!!"), FNAME_Find) == UniqueName);
-		CHECK(FName("UNIQUEUNICORN!!", FNAME_Find) == UniqueName);
-		CHECK(FName(TEXT("UNIQUEUNICORN!!"), FNAME_Find) == UniqueName);
-		CHECK(FName("uniqueunicorn!!", FNAME_Find) == UniqueName);
+		NAME_CHECK(FName("UniqueUnicorn!!", FNAME_Find) == UniqueName);
+		NAME_CHECK(FName(TEXT("UniqueUnicorn!!"), FNAME_Find) == UniqueName);
+		NAME_CHECK(FName("UNIQUEUNICORN!!", FNAME_Find) == UniqueName);
+		NAME_CHECK(FName(TEXT("UNIQUEUNICORN!!"), FNAME_Find) == UniqueName);
+		NAME_CHECK(FName("uniqueunicorn!!", FNAME_Find) == UniqueName);
 
 		Once = false;
 	}
@@ -4251,43 +4257,43 @@ TEST_CASE_NAMED(FNameTest, "System::Core::Name", "[ApplicationContextMask][Smoke
 	Names.Add(FooWide);
 	Algo::Sort(Names, FNameLexicalLess()); 
 
-	CHECK(Names[0] == "FooA");
-	CHECK(Names[1] == "FooAB");
-	CHECK(Names[2] == "FooABC");
-	CHECK(Names[3] == "FooABCD");
-	CHECK(Names[4] == "FooB");
-	CHECK(Names[5] == "FooC");
-	CHECK(Names[6] == FooWide);
+	NAME_CHECK(Names[0] == "FooA");
+	NAME_CHECK(Names[1] == "FooAB");
+	NAME_CHECK(Names[2] == "FooABC");
+	NAME_CHECK(Names[3] == "FooABCD");
+	NAME_CHECK(Names[4] == "FooB");
+	NAME_CHECK(Names[5] == "FooC");
+	NAME_CHECK(Names[6] == FooWide);
 		
-	CHECK(FName() == TEXTVIEW("None"));
-	CHECK(FName() == TEXTVIEW("none"));
-	CHECK(FName() == TEXTVIEW(""));
-	CHECK(FName() == FStringView());
-	CHECK(FName("a") == TEXTVIEW("a"));
-	CHECK(FName("a") == TEXTVIEW("A"));
-	CHECK(FName("A") == TEXTVIEW("a"));
-	CHECK(FName("aa") != TEXTVIEW("a"));
-	CHECK(FName("a") != TEXTVIEW("aa"));
-	CHECK(FName("a") == TEXTVIEW("aa").Left(1)); // Unterminated
-	CHECK(FName("aa") != TEXTVIEW("ab"));
-	CHECK(FName("a_0") == TEXTVIEW("a_0"));
-	CHECK(FName("a_1") != TEXTVIEW("a_0"));
-	CHECK(FName("a_10") == TEXTVIEW("a_10"));
-	CHECK(FName("a")    == TEXTVIEW("a_10").Left(1)); // Unterminated number
-	CHECK(FName("a_")   == TEXTVIEW("a_10").Left(2)); // Unterminated number
-	CHECK(FName("a_1")  == TEXTVIEW("a_10").Left(3)); // Unterminated number
-	CHECK(FName("a_2147483646") == TEXTVIEW("a_2147483646")); // Largest valid number
-	CHECK(FName("a_2147483647") == TEXTVIEW("a_2147483647")); // Smallest invalid number
-	CHECK(FName("a_") == TEXTVIEW("a_"));
-	CHECK(FName("_") == TEXTVIEW("_"));
-	CHECK(FName("_3") == TEXTVIEW("_3")); // Special none
-	CHECK(FName() == TEXTVIEW("_123")); // Special none
-	CHECK(FName("_3") == TEXTVIEW("None")); // Special none
-	CHECK(FName("none_1") == TEXTVIEW("none_1")); // Numbered none
-	CHECK(WideName == FStringView(Wide));
-	CHECK(FName("a") == ANSITEXTVIEW("a"));
-	CHECK(FName("a") == ANSITEXTVIEW("A"));
-	CHECK(FName("A") != ANSITEXTVIEW("b"));
+	NAME_CHECK(FName() == TEXTVIEW("None"));
+	NAME_CHECK(FName() == TEXTVIEW("none"));
+	NAME_CHECK(FName() == TEXTVIEW(""));
+	NAME_CHECK(FName() == FStringView());
+	NAME_CHECK(FName("a") == TEXTVIEW("a"));
+	NAME_CHECK(FName("a") == TEXTVIEW("A"));
+	NAME_CHECK(FName("A") == TEXTVIEW("a"));
+	NAME_CHECK(FName("aa") != TEXTVIEW("a"));
+	NAME_CHECK(FName("a") != TEXTVIEW("aa"));
+	NAME_CHECK(FName("a") == TEXTVIEW("aa").Left(1)); // Unterminated
+	NAME_CHECK(FName("aa") != TEXTVIEW("ab"));
+	NAME_CHECK(FName("a_0") == TEXTVIEW("a_0"));
+	NAME_CHECK(FName("a_1") != TEXTVIEW("a_0"));
+	NAME_CHECK(FName("a_10") == TEXTVIEW("a_10"));
+	NAME_CHECK(FName("a")    == TEXTVIEW("a_10").Left(1)); // Unterminated number
+	NAME_CHECK(FName("a_")   == TEXTVIEW("a_10").Left(2)); // Unterminated number
+	NAME_CHECK(FName("a_1")  == TEXTVIEW("a_10").Left(3)); // Unterminated number
+	NAME_CHECK(FName("a_2147483646") == TEXTVIEW("a_2147483646")); // Largest valid number
+	NAME_CHECK(FName("a_2147483647") == TEXTVIEW("a_2147483647")); // Smallest invalid number
+	NAME_CHECK(FName("a_") == TEXTVIEW("a_"));
+	NAME_CHECK(FName("_") == TEXTVIEW("_"));
+	NAME_CHECK(FName("_3") == TEXTVIEW("_3")); // Special none
+	NAME_CHECK(FName() == TEXTVIEW("_123")); // Special none
+	NAME_CHECK(FName("_3") == TEXTVIEW("None")); // Special none
+	NAME_CHECK(FName("none_1") == TEXTVIEW("none_1")); // Numbered none
+	NAME_CHECK(WideName == FStringView(Wide));
+	NAME_CHECK(FName("a") == ANSITEXTVIEW("a"));
+	NAME_CHECK(FName("a") == ANSITEXTVIEW("A"));
+	NAME_CHECK(FName("A") != ANSITEXTVIEW("b"));
 
 	// Check truncating ToString functions
 	FNameAutoTestTruncateHelper::Run();
@@ -4311,12 +4317,12 @@ TEST_CASE_NAMED(FNameTest, "System::Core::Name", "[ApplicationContextMask][Smoke
 	CheckLazyName("None_0");
 	CheckLazyName("None_1");
 
-	CHECK(FLazyName("Hej") == FLazyName(UTF8TEXT("Hej")));
-	CHECK(FLazyName("Hej") == FLazyName(WIDETEXT("Hej")));
-	CHECK(FLazyName(UTF8TEXT("Hej")) == FLazyName("Hej"));
-	CHECK(FLazyName(UTF8TEXT("Hej")) == FLazyName(WIDETEXT("Hej")));
-	CHECK(FLazyName(WIDETEXT("Hej")) == FLazyName("Hej"));
-	CHECK(FLazyName(WIDETEXT("Hej")) == FLazyName(UTF8TEXT("Hej")));
+	NAME_CHECK(FLazyName("Hej") == FLazyName(UTF8TEXT("Hej")));
+	NAME_CHECK(FLazyName("Hej") == FLazyName(WIDETEXT("Hej")));
+	NAME_CHECK(FLazyName(UTF8TEXT("Hej")) == FLazyName("Hej"));
+	NAME_CHECK(FLazyName(UTF8TEXT("Hej")) == FLazyName(WIDETEXT("Hej")));
+	NAME_CHECK(FLazyName(WIDETEXT("Hej")) == FLazyName("Hej"));
+	NAME_CHECK(FLazyName(WIDETEXT("Hej")) == FLazyName(UTF8TEXT("Hej")));
 
 	TestNameBatch();
 
@@ -4339,7 +4345,7 @@ TEST_CASE_NAMED(FNameTest, "System::Core::Name", "[ApplicationContextMask][Smoke
 		{
 			Id = FName(*FString::Printf(TEXT("%d"), I)).GetComparisonIndex();
 		}
-		CHECK(Id == OldId);
+		NAME_CHECK(Id == OldId);
 	}
 #endif
 
@@ -4349,29 +4355,31 @@ TEST_CASE_NAMED(FNameTest, "System::Core::Name", "[ApplicationContextMask][Smoke
 		for (uint32 Index : SetBits)
 		{
 			const bool bWasSet = TestSet.SetBitAtomic(Index);
-			CHECK(bWasSet == false);
+			NAME_CHECK(bWasSet == false);
 		}
 
 		const bool b32WasSet = TestSet.SetBitAtomic(34);
-		CHECK(b32WasSet);
+		NAME_CHECK(b32WasSet);
 
 		uint32 Expected = 0;
 		TestSet.ForEachSetBit([&](uint32 Index)
 		{
-			CHECK(Index == SetBits[Expected]);
+			NAME_CHECK(Index == SetBits[Expected]);
 			++Expected;
 		});
-		CHECK(Expected == UE_ARRAY_COUNT(SetBits));
+		NAME_CHECK(Expected == UE_ARRAY_COUNT(SetBits));
 
 		const bool b1023WasSet = TestSet.ClearBitAtomic(1023);
-		CHECK(b1023WasSet == true);
+		NAME_CHECK(b1023WasSet == true);
 		TestSet.ForEachSetBit([&](uint32 Index)
 		{
-			CHECK(Index != 1023);
+			NAME_CHECK(Index != 1023);
 		});
 	}
 }
 
+#undef NAME_CHECK
+
 #endif // WITH_TESTS
 
 
