Keep module list matching the address space and avoid duplicates.

DLLs use the same address space system wide. When they are unloaded and later loaded again they return to the region of address space there were at before. Because RemoveModule() is skipped, the reload would result in a duplicate Modules entry. This change detects those duplicate.

Should an unloaded DLL not be in use by any other process it vacates its address space and another DLL may be loaded in its place. With stale Modules entries, backtracing could use the wrong Functions array. Note that this highlights a flaw in the lookup set cache used to avoid the lock, as that has a slim chance of holding items for DLLs that have been unloaded.

#rnx

#Synced-CL Unknown

diff --git a/Engine/Source/Runtime/Core/Private/ProfilingDebugging/Microsoft/WindowsCallstackTrace.cpp b/Engine/Source/Runtime/Core/Private/ProfilingDebugging/Microsoft/WindowsCallstackTrace.cpp
--- a/Engine/Source/Runtime/Core/Private/ProfilingDebugging/Microsoft/WindowsCallstackTrace.cpp
+++ b/Engine/Source/Runtime/Core/Private/ProfilingDebugging/Microsoft/WindowsCallstackTrace.cpp
@@ -20,6 +20,7 @@
 #	include <winternl.h>
 #include "Windows/HideWindowsPlatformTypes.h"
 #include "Misc/CoreDelegates.h"
+#include "Misc/ScopeExit.h"
 #include "Trace/Trace.inl"
 #include "HAL/Runnable.h"
 #include "HAL/PlatformProcess.h"
@@ -211,6 +212,7 @@ private:
 	};
 	typedef TGrowOnlyLockFreeHash<FFunctionLookupSetEntry, uint64, int32> FFunctionLookupSet;
 
+	void					ReplaceModule(TArrayView<FModule> ModulesView, FModule& Module);
 	const FFunction*		LookupFunction(UPTRINT Address, FLookupState& State) const;
 	static FBacktracer*		Instance;
 	mutable FCriticalSection Lock;
@@ -266,6 +268,39 @@ FBacktracer* FBacktracer::Get()
 	return Instance;
 }
 
+////////////////////////////////////////////////////////////////////////////////
+void FBacktracer::ReplaceModule(TArrayView<FModule> ModulesView, FModule& Module)
+{
+	FModule& Existing = ModulesView[0];
+
+	// Is the module being added a dulplicate? It has been unloaded and reloaded
+	// again into the same region
+	if (Module.Id == Existing.Id
+		&& Module.IdSize == Existing.IdSize
+		&& Module.NumFunctions == Existing.NumFunctions)
+	{
+		Malloc->Free(Module.Functions);
+		return;
+	}
+
+	// Otherwise the module is new and located in a region previously occupied
+	// by one or more modules that have since been unloaded.
+	Malloc->Free(Existing.Functions);
+	Existing = Module;
+
+	uint32 IdEnd = Module.Id + Module.IdSize;
+	for (FModule& Iter : ModulesView.Mid(1))
+	{
+		if (Iter.Id >= IdEnd)
+		{
+			break;
+		}
+
+		Iter.Id = 0;
+		Iter.IdSize = 0;
+	}
+}
+
 ////////////////////////////////////////////////////////////////////////////////
 void FBacktracer::AddModule(UPTRINT ModuleBase, const TCHAR* Name)
 {
@@ -501,8 +536,42 @@ void FBacktracer::AddModule(UPTRINT ModuleBase, const TCHAR* Name)
 	{
 		FScopeLock _(&Lock);
 
+		ON_SCOPE_EXIT
+		{
+			// For ease of early outs and logical paragraphs, we'll sort out the
+			// way out of the locked scope.
+			auto Predicate = [] (const FModule& A, const FModule& B)
+			{
+				return A.Id < B.Id;
+			};
+			Algo::Sort(TArrayView<FModule>(Modules, ModulesNum), Predicate);
+		};
+
 		TArrayView<FModule> ModulesView(Modules, ModulesNum);
-		int32 Index = Algo::UpperBound(ModulesView, Module.Id, FIdPredicate());
+		int32 Index = Algo::UpperBound(ModulesView, Module.Id, FIdPredicate()) - 1;
+		if (Index >= 0)
+		{
+			// Move the view so it starts at the first module of interest.
+			ModulesView = ModulesView.Mid(Index);
+
+			// Is the new module's starting address (Id) contained within an
+			// existing unloaded module?
+			uint32 IdEnd = ModulesView[0].Id + ModulesView[0].IdSize;
+			if (Module.Id < IdEnd)
+			{
+				ReplaceModule(ModulesView, Module);
+				return;
+			}
+
+			// Or perhaps the new module ends in the address space of the next
+			// module along (which has since been unloaded)?
+			ModulesView = ModulesView.Mid(1);
+			if (!ModulesView.IsEmpty() && (Module.Id + Module.IdSize) > ModulesView[0].Id)
+			{
+				ReplaceModule(ModulesView, Module);
+				return;
+			}
+		}
 
 		if (ModulesNum + 1 > ModulesCapacity)
 		{
@@ -510,7 +579,6 @@ void FBacktracer::AddModule(UPTRINT ModuleBase, const TCHAR* Name)
 			Modules = (FModule*) Malloc->Realloc(Modules, sizeof(FModule)* ModulesCapacity);
 		}
 		Modules[ModulesNum++] = Module;
-		Algo::Sort(TArrayView<FModule>(Modules, ModulesNum), [](const FModule& A, const FModule& B) { return A.Id < B.Id; });
 	}
 
 #if BACKTRACE_DBGLVL >= 1
