Add a new `FQueryableRefCountedObject` that allows the reference count adds and releases to be returned and queried.

This will let us deprecate the return from `Release` of `FRefCountedObject` which is unsafe under transactions, but give us a new class to let us get these accurately (at the expense of more memory used for the queryable variant of ref counted objects).

#Synced-CL 48759341

diff --git a/Engine/Plugins/Experimental/UAF/UAF/Source/UAF/Private/UAF/Attributes/AttributeBindingData.cpp b/Engine/Plugins/Experimental/UAF/UAF/Source/UAF/Private/UAF/Attributes/AttributeBindingData.cpp
--- a/Engine/Plugins/Experimental/UAF/UAF/Source/UAF/Private/UAF/Attributes/AttributeBindingData.cpp
+++ b/Engine/Plugins/Experimental/UAF/UAF/Source/UAF/Private/UAF/Attributes/AttributeBindingData.cpp
@@ -436,7 +436,7 @@ namespace UE::UAF
 
 	uint32 FAttributeBindingData::Release() const
 	{
-		const uint32 RefsBeforeRelease = FRefCountedObject::Release();
+		const uint32 RefsBeforeRelease = UE::Private::FQueryableRefCountedObject::Release();
 		if (RefsBeforeRelease == NumInternalReferences + 1)
 		{
 			// Our last external reference has been released, release our internal references to free the memory
@@ -449,12 +449,12 @@ namespace UE::UAF
 	void FAttributeBindingData::ResetInternalReferences()
 	{
 		// Add an extra reference to keep this container alive during destruction
-		FRefCountedObject::AddRef();
+		UE::Private::FQueryableRefCountedObject::AddRef();
 
 		DefaultAttributeValuesByName.Empty();
 
 		// We are done with destruction, destroy ourselves now
-		FRefCountedObject::Release();
+		UE::Private::FQueryableRefCountedObject::Release();
 	}
 
 	FAttributeBindingDataPtr MakeAttributeBindingData(const UAbstractSkeletonSetBinding& Binding)
diff --git a/Engine/Plugins/Experimental/UAF/UAF/Source/UAF/Public/UAF/Attributes/AttributeBindingData.h b/Engine/Plugins/Experimental/UAF/UAF/Source/UAF/Public/UAF/Attributes/AttributeBindingData.h
--- a/Engine/Plugins/Experimental/UAF/UAF/Source/UAF/Public/UAF/Attributes/AttributeBindingData.h
+++ b/Engine/Plugins/Experimental/UAF/UAF/Source/UAF/Public/UAF/Attributes/AttributeBindingData.h
@@ -18,7 +18,7 @@ namespace UE::UAF
 
 	// Our baked abstract skeleton binding data
 	// This class is ref counted and all named/typed sets increment/decrement its ref count
-	class FAttributeBindingData : public FRefCountedObject
+	class FAttributeBindingData : public UE::Private::FQueryableRefCountedObject
 	{
 	public:
 		using NamedSetMapType = TSortedMap<FName, FAttributeNamedSet*, FDefaultAllocator, FNameFastLess>;
diff --git a/Engine/Source/Programs/AutoRTFMTests/Private/RefCountPtrTests.cpp b/Engine/Source/Programs/AutoRTFMTests/Private/RefCountPtrTests.cpp
--- a/Engine/Source/Programs/AutoRTFMTests/Private/RefCountPtrTests.cpp
+++ b/Engine/Source/Programs/AutoRTFMTests/Private/RefCountPtrTests.cpp
@@ -7,26 +7,32 @@
 
 namespace
 {
-	struct RefCountedInt : public FRefCountedObject
+	struct FRefCountedInt : public FRefCountedObject
 	{
-		RefCountedInt(int InValue) : Value(InValue) {}
+		FRefCountedInt(int InValue) : Value(InValue) {}
 		int Value;
 	};
 
-	struct RefCountedMixinInt : public TRefCountingMixin<RefCountedMixinInt>
+	struct FRefCountedMixinInt : public TRefCountingMixin<FRefCountedMixinInt>
 	{
-		RefCountedMixinInt(int InValue) : Value(InValue) {}
+		FRefCountedMixinInt(int InValue) : Value(InValue) {}
 		int Value;
 	};
 
-	struct ThreadSafeRefCountedInt : public FRefCountedObject
+	struct FThreadSafeRefCountedInt : public FRefCountedObject
 	{
-		ThreadSafeRefCountedInt(int InValue) : Value(InValue) {}
+		FThreadSafeRefCountedInt(int InValue) : Value(InValue) {}
+		int Value;
+	};
+
+	struct FQueryableRefCountedInt : public UE::Private::FQueryableRefCountedObject
+	{
+		FQueryableRefCountedInt(int InValue) : Value(InValue) {}
 		int Value;
 	};
 }
 
-TEMPLATE_TEST_CASE("CheckRefCounts", "[RefCountPtr]", RefCountedInt, RefCountedMixinInt, ThreadSafeRefCountedInt)
+TEMPLATE_TEST_CASE("CheckRefCounts", "[RefCountPtr]", FQueryableRefCountedInt)
 {
 	// Refcounts should be correct outside of a transaction.
 	SECTION("Non-transactional refcounts should be accurate")
@@ -34,13 +40,13 @@ TEMPLATE_TEST_CASE("CheckRefCounts", "[RefCountPtr]", RefCountedInt, RefCountedM
 		TestType* Ptr = new TestType(42);
 
 		REQUIRE(Ptr->GetRefCount() == 0);
-		       (Ptr->AddRef().CheckAtLeast(1));
+		REQUIRE(Ptr->AddRef() == 1);
 		REQUIRE(Ptr->GetRefCount() == 1);
-		       (Ptr->AddRef().CheckAtLeast(2));
+		REQUIRE(Ptr->AddRef() == 2);
 		REQUIRE(Ptr->GetRefCount() == 2);
 		REQUIRE(Ptr->Release() == 1);
 		REQUIRE(Ptr->GetRefCount() == 1);
-		       (Ptr->AddRef().CheckAtLeast(2));
+		REQUIRE(Ptr->AddRef() == 2);
 		REQUIRE(Ptr->GetRefCount() == 2);
 		REQUIRE(Ptr->Release() == 1);
 		REQUIRE(Ptr->GetRefCount() == 1);
@@ -56,25 +62,73 @@ TEMPLATE_TEST_CASE("CheckRefCounts", "[RefCountPtr]", RefCountedInt, RefCountedM
 		{
 			TestType* Ptr = new TestType(42);
 
-			REQUIRE(Ptr->GetRefCount() >= 0);
-				   (Ptr->AddRef().CheckAtLeast(1));
-			REQUIRE(Ptr->GetRefCount() >= 1);
-				   (Ptr->AddRef().CheckAtLeast(2));
-			REQUIRE(Ptr->GetRefCount() >= 2);
-			REQUIRE(Ptr->Release() >= 1);
-			REQUIRE(Ptr->GetRefCount() >= 1);
-				   (Ptr->AddRef().CheckAtLeast(2));
-			REQUIRE(Ptr->GetRefCount() >= 2);
-			REQUIRE(Ptr->Release() >= 1);
-			REQUIRE(Ptr->GetRefCount() >= 1);
-			REQUIRE(Ptr->Release() >= 0);
+			REQUIRE(Ptr->GetRefCount() == 0);
+			REQUIRE(Ptr->AddRef() == 1);
+			REQUIRE(Ptr->GetRefCount() == 1);
+			REQUIRE(Ptr->AddRef() == 2);
+			REQUIRE(Ptr->GetRefCount() == 2);
+			REQUIRE(Ptr->Release() == 1);
+			REQUIRE(Ptr->GetRefCount() == 1);
+			REQUIRE(Ptr->AddRef() == 2);
+			REQUIRE(Ptr->GetRefCount() == 2);
+			REQUIRE(Ptr->Release() == 1);
+			REQUIRE(Ptr->GetRefCount() == 1);
+			REQUIRE(Ptr->Release() == 0);
 
 			// Ptr is dead
 		});
 	}
 }
 
-TEMPLATE_TEST_CASE("PreviouslyAllocated", "[RefCountPtr]", RefCountedInt, RefCountedMixinInt, ThreadSafeRefCountedInt)
+TEMPLATE_TEST_CASE("CheckRefCountsToDeprecate", "[RefCountPtr]", FRefCountedInt, FRefCountedMixinInt, FThreadSafeRefCountedInt)
+{
+	// Refcounts should be correct outside of a transaction.
+	SECTION("Non-transactional refcounts should be accurate")
+	{
+		TestType* Ptr = new TestType(42);
+
+		REQUIRE(Ptr->GetRefCount() == 0);
+		(Ptr->AddRef().CheckAtLeast(1));
+		REQUIRE(Ptr->GetRefCount() == 1);
+		(Ptr->AddRef().CheckAtLeast(2));
+		REQUIRE(Ptr->GetRefCount() == 2);
+		REQUIRE(Ptr->Release() == 1);
+		REQUIRE(Ptr->GetRefCount() == 1);
+		(Ptr->AddRef().CheckAtLeast(2));
+		REQUIRE(Ptr->GetRefCount() == 2);
+		REQUIRE(Ptr->Release() == 1);
+		REQUIRE(Ptr->GetRefCount() == 1);
+		REQUIRE(Ptr->Release() == 0);
+
+		// Ptr is dead
+	}
+
+	// Refcounts can be inflated inside a transaction, since releases are deferred.
+	SECTION("Transactional refcounts should represent a lower bound")
+	{
+		AutoRTFM::Testing::Commit([]
+			{
+				TestType* Ptr = new TestType(42);
+
+				REQUIRE(Ptr->GetRefCount() >= 0);
+				(Ptr->AddRef().CheckAtLeast(1));
+				REQUIRE(Ptr->GetRefCount() >= 1);
+				(Ptr->AddRef().CheckAtLeast(2));
+				REQUIRE(Ptr->GetRefCount() >= 2);
+				REQUIRE(Ptr->Release() >= 1);
+				REQUIRE(Ptr->GetRefCount() >= 1);
+				(Ptr->AddRef().CheckAtLeast(2));
+				REQUIRE(Ptr->GetRefCount() >= 2);
+				REQUIRE(Ptr->Release() >= 1);
+				REQUIRE(Ptr->GetRefCount() >= 1);
+				REQUIRE(Ptr->Release() >= 0);
+
+				// Ptr is dead
+			});
+	}
+}
+
+TEMPLATE_TEST_CASE("PreviouslyAllocated", "[RefCountPtr]", FRefCountedInt, FRefCountedMixinInt, FThreadSafeRefCountedInt, FQueryableRefCountedInt)
 {
 	SECTION("Using new T()")
 	{
@@ -107,7 +161,7 @@ TEMPLATE_TEST_CASE("PreviouslyAllocated", "[RefCountPtr]", RefCountedInt, RefCou
 	}
 }
 
-TEMPLATE_TEST_CASE("AbortWithPreviouslyAllocated", "[RefCountPtr]", RefCountedInt, RefCountedMixinInt, ThreadSafeRefCountedInt)
+TEMPLATE_TEST_CASE("AbortWithPreviouslyAllocated", "[RefCountPtr]", FRefCountedInt, FRefCountedMixinInt, FThreadSafeRefCountedInt, FQueryableRefCountedInt)
 {
 	SECTION("Using new T()")
 	{
@@ -144,7 +198,7 @@ TEMPLATE_TEST_CASE("AbortWithPreviouslyAllocated", "[RefCountPtr]", RefCountedIn
 	}
 }
 
-TEMPLATE_TEST_CASE("AbortWithNewlyAllocated", "[RefCountPtr]", RefCountedInt, RefCountedMixinInt, ThreadSafeRefCountedInt)
+TEMPLATE_TEST_CASE("AbortWithNewlyAllocated", "[RefCountPtr]", FRefCountedInt, FRefCountedMixinInt, FThreadSafeRefCountedInt, FQueryableRefCountedInt)
 {
 	int Result = 42;
 
@@ -193,7 +247,7 @@ TEMPLATE_TEST_CASE("AbortWithNewlyAllocated", "[RefCountPtr]", RefCountedInt, Re
 	REQUIRE(Result == 42);
 }
 
-TEMPLATE_TEST_CASE("OnCommitCapturingRefCountPtr", "[RefCountPtr]", RefCountedInt, RefCountedMixinInt, ThreadSafeRefCountedInt)
+TEMPLATE_TEST_CASE("OnCommitCapturingRefCountPtr", "[RefCountPtr]", FRefCountedInt, FRefCountedMixinInt, FThreadSafeRefCountedInt, FQueryableRefCountedInt)
 {
 	TRefCountPtr<TestType> Foo = MakeRefCount<TestType>(42);
 
@@ -225,7 +279,7 @@ TEMPLATE_TEST_CASE("OnCommitCapturingRefCountPtr", "[RefCountPtr]", RefCountedIn
 	}
 }
 
-TEMPLATE_TEST_CASE("OnAbortCapturingRefCountPtr", "[RefCountPtr]", RefCountedInt, RefCountedMixinInt, ThreadSafeRefCountedInt)
+TEMPLATE_TEST_CASE("OnAbortCapturingRefCountPtr", "[RefCountPtr]", FRefCountedInt, FRefCountedMixinInt, FThreadSafeRefCountedInt, FQueryableRefCountedInt)
 {
 	TRefCountPtr<TestType> Foo = MakeRefCount<TestType>(42);
 
diff --git a/Engine/Source/Runtime/Core/Public/Templates/RefCounting.h b/Engine/Source/Runtime/Core/Public/Templates/RefCounting.h
--- a/Engine/Source/Runtime/Core/Public/Templates/RefCounting.h
+++ b/Engine/Source/Runtime/Core/Public/Templates/RefCounting.h
@@ -65,21 +65,66 @@ namespace UE::Private
 
 CORE_API void CheckRefCountIsNonZero();
 
-template <typename AtomicType>
+enum class ERefCountIsQueryable : bool
+{
+	Yes = true,
+	No = false,
+};
+
+template <typename IntType, ERefCountIsQueryable Queryable> struct TQueryableTransactionalType;
+
+template <typename IntType> struct TQueryableTransactionalType<IntType, ERefCountIsQueryable::Yes>
+{
+	TQueryableTransactionalType(IntType Init) : Payload(Init) {}
+
+	IntType operator++(int)
+	{
+		return Payload++;
+	}
+
+	IntType operator--(int)
+	{
+		return Payload--;
+	}
+
+	operator IntType() const
+	{
+		return Payload;
+	}
+private:
+	IntType Payload = 0;
+};
+
+template <typename IntType> struct TQueryableTransactionalType<IntType, ERefCountIsQueryable::No>
+{
+	TQueryableTransactionalType(IntType Init) {}
+
+	operator IntType() const
+	{
+		return 0;
+	}
+};
+
+template <typename IntType, ERefCountIsQueryable Queryable = ERefCountIsQueryable::Yes>
 class TTransactionalAtomicRefCount
 {
 public:
 	template <auto DeleteFn>
-	AtomicType AddRef() const
+	IntType AddRef() const
 	{
-		AtomicType Refs;
+		if (!AutoRTFM::IsClosed())
+		{
+			return ++RefCount;
+		}
+
+		IntType Refs;
 
 		UE_AUTORTFM_OPEN
 		{
 			Refs = RefCount++;
 		};
 
-		// If we are inside a transaction, a rollback must undo our refcount bump.
+		// As we are inside a transaction, the rollback must undo our refcount bump.
 		// In general, this is best handled by Release(). However, there is one case 
 		// that needs to be handled with special care. A brand-new object has a
 		// refcount of zero, and a rollback must return it to this zero-refcount state.
@@ -104,11 +149,11 @@ public:
 			};
 		}
 
-		return Refs + 1;
+		return Refs + 1 - TransactionallyDeferredReleaseCount;
 	}
 
 	template <auto DeleteFn>
-	AtomicType Release() const
+	IntType Release() const
 	{
 #if DO_GUARD_SLOW
 		if (RefCount == 0)
@@ -117,59 +162,70 @@ public:
 		}
 #endif
 
-		AtomicType RefsToReturn;
-
 		if (AutoRTFM::IsClosed())
 		{
-			// We return the active number of references minus one to maintain the existing
-			// Release() behavior as closely as possible while inside a transaction, even 
-			// though we are deferring reference count changes until commit time.
-			// Be advised that GetRefCount() would reveal our trickery, since it
-			// always returns the true refcount.
+			IntType Refs;
+
 			UE_AUTORTFM_OPEN
 			{
-				RefsToReturn = AtomicType(RefCount.load(std::memory_order_relaxed)) - 1;
+				TransactionallyDeferredReleaseCount++;
+				Refs = RefCount.load(std::memory_order_relaxed) - TransactionallyDeferredReleaseCount;
+			};
+
+			UE_AUTORTFM_ONABORT(this)
+			{
+				// TODO: degrade this to a `checkSlow` once time has passed.
+				check(TransactionallyDeferredReleaseCount > 0);
+				TransactionallyDeferredReleaseCount--;
 			};
 
 			// Refcount changes and frees are deferred until the transaction is concluded.
 			UE_AUTORTFM_ONCOMMIT(this)
 			{
+				// TODO (SOL-7350): degrade this to a `checkSlow` once time has passed.
+				check(TransactionallyDeferredReleaseCount > 0);
+				TransactionallyDeferredReleaseCount--;
 				ImmediatelyRelease<DeleteFn>();
 			};
+
+			return Refs;
 		}
 		else
 		{
-			RefsToReturn = ImmediatelyRelease<DeleteFn>() - 1;
+			return ImmediatelyRelease<DeleteFn>() - 1;
 		}
-
-		return RefsToReturn;
 	}
 
-	AtomicType GetRefCount() const
+	// This is equivalent to https://en.cppreference.com/w/cpp/memory/shared_ptr/use_count
+	// This loads with a relaxed memory ordering because a 'live' reference count is
+	// unstable by nature and so there's no benefit to try and enforce memory ordering
+	// around the reading of it.
+	IntType GetRefCount() const
 	{
-		// This is equivalent to https://en.cppreference.com/w/cpp/memory/shared_ptr/use_count
-		// 
-		// Inside of an AutoRTFM transaction, the returned refcount value may be higher than you'd
-		// expect, because all Release calls are deferred until the transaction commit time.
-		AtomicType Refs;
-		UE_AUTORTFM_OPEN
+		if (AutoRTFM::IsClosed())
 		{
-			// A 'live' reference count is unstable by nature and so there's no benefit
-			// to try and enforce memory ordering around the reading of it.
-			Refs = RefCount.load(std::memory_order_relaxed);
-		};
-		return Refs;
+			IntType Refs;
+			UE_AUTORTFM_OPEN
+			{
+				Refs = RefCount.load(std::memory_order_relaxed) - TransactionallyDeferredReleaseCount;
+			};
+			return Refs;
+		}
+		else
+		{
+			return RefCount.load(std::memory_order_relaxed);
+		}
 	}
 
 private:
 	template <auto DeleteFn>
-	AtomicType ImmediatelyRelease() const
+	IntType ImmediatelyRelease() const
 	{
 		// fetch_sub returns the refcount _before_ it was decremented. std::memory_order_acq_rel is
 		// used so that, if we do end up executing the destructor, it's not possible for side effects 
 		// from executing the destructor to end up being visible before we've determined that the 
 		// reference count is actually zero.
-		AtomicType RefsBeforeRelease = RefCount.fetch_sub(1, std::memory_order_acq_rel);
+		IntType RefsBeforeRelease = RefCount.fetch_sub(1, std::memory_order_acq_rel);
 
 #if DO_GUARD_SLOW
 		// A check-failure is issued if an object is over-released.
@@ -178,17 +234,58 @@ private:
 			CheckRefCountIsNonZero();
 		}
 #endif
+
 		// We immediately free the object if its refcount has become zero.
 		if (RefsBeforeRelease == 1)
 		{
 			DeleteFn(this);
 		}
+
 		return RefsBeforeRelease;
 	}
 
-	mutable std::atomic<AtomicType> RefCount = 0;
+	mutable std::atomic<IntType> RefCount = 0;
+	mutable TQueryableTransactionalType<IntType, Queryable> TransactionallyDeferredReleaseCount = 0;
 };
 
+/**
+ * Please do not use this, long term we want it to go away. Being able to query the reference count can
+ * easily lead into really gnarly bugs and we generally don't want this sort of paradigm to be used more
+ * widely. But there are a few places where it is currently being legitimately used so we'll provide a
+ * queryable variant to let us deprecate the return of the reference count from `FRefCountedObject` asap.
+ */
+class FQueryableRefCountedObject : UE::Private::TTransactionalAtomicRefCount<uint32, UE::Private::ERefCountIsQueryable::Yes>
+{
+public:
+	FQueryableRefCountedObject() = default;
+	virtual ~FQueryableRefCountedObject() = default;
+
+	FQueryableRefCountedObject(const FQueryableRefCountedObject& Rhs) = delete;
+	FQueryableRefCountedObject& operator=(const FQueryableRefCountedObject& Rhs) = delete;
+
+	uint32 AddRef() const
+	{
+		return Super::AddRef<DeleteThis>();
+	}
+
+	uint32 Release() const
+	{
+		return Super::Release<DeleteThis>();
+	}
+
+	uint32 GetRefCount() const
+	{
+		return Super::GetRefCount();
+	}
+
+private:
+	using Super = UE::Private::TTransactionalAtomicRefCount<uint32, UE::Private::ERefCountIsQueryable::Yes>;
+
+	static void DeleteThis(const Super* This)
+	{
+		delete static_cast<const FQueryableRefCountedObject*>(This);
+	}
+};
 }
 
 /** A virtual interface for ref counted objects to implement. */
@@ -246,6 +343,8 @@ class UE_DEPRECATED(5.8, "`FRefCountedObject` is now thread-safe so use that ins
 
 class UE_DEPRECATED(5.8, "`FRefCountedObject` is now thread-safe so use that instead.") FThreadSafeRefCountedObject : public FRefCountedObject {};
 
+
+
 /**
  * ERefCountingMode is used select between either 'fast' or 'thread safe' ref-counting types.
  * This is only used at compile time to select between template specializations.
