Refactor FAppTime to add GetApproximateCurrentTime() function
- Some systems want a periodically incrementing timer that is cheaper to read than the full precision FPlatformTime::Seconds(). For these systems, it doesn't matter if the value returned races the game thread, so therefore can be implemented with a relaxed atomic read of the game thread's instance.
- To achieve this, CurrentTime has been wrapped in an atomic helper to always perform relaxed reads/writes.
- Added getter/setter functions for FAppTime, as this was necessary to hide the detail of CurrentTime being an atomic, and also allows us to potentially deprecate access in future.
- Get_UnsafeNoCheck is still provided as a temporary workaround for systems that need refactor. E.g. the CurrentFrameTime field is currently accessed by multiple threads, and cannot be made atomic.

#Synced-CL 48676711

diff --git a/Engine/Source/Runtime/Core/Private/Experimental/Containers/HazardPointer.cpp b/Engine/Source/Runtime/Core/Private/Experimental/Containers/HazardPointer.cpp
--- a/Engine/Source/Runtime/Core/Private/Experimental/Containers/HazardPointer.cpp
+++ b/Engine/Source/Runtime/Core/Private/Experimental/Containers/HazardPointer.cpp
@@ -4,7 +4,7 @@
 
 #include "Algo/Sort.h"
 #include "Async/UniqueLock.h"
-#include "HAL/PlatformTime.h"
+#include "Misc/AppTime.h"
 
 FHazardPointerCollection::~FHazardPointerCollection()
 {
@@ -65,7 +65,7 @@ void FHazardPointerCollection::Delete(const HazardPointer_Impl::FHazardDeleter&
 	Collectables->ReclamationList.Add(Ptr);
 
 	//maybe scan the list
-	const double CurrentTime = FPlatformTime::Seconds();
+	const double CurrentTime = FAppTime::GetApproximateCurrentTime();
 	const bool TimeLimit = (CurrentTime - Collectables->TimeOfLastCollection) > 1.0;
 	const int32 DeleteMetric = CollectLimit <= 0 ? ((5 * TotalNumHazardRecords) / 4) : CollectLimit;
 	if (TimeLimit || Collectables->ReclamationList.Num() >= DeleteMetric)
diff --git a/Engine/Source/Runtime/Core/Public/Misc/App.h b/Engine/Source/Runtime/Core/Public/Misc/App.h
--- a/Engine/Source/Runtime/Core/Public/Misc/App.h
+++ b/Engine/Source/Runtime/Core/Public/Misc/App.h
@@ -558,7 +558,7 @@ public:
 	 */
 	static double GetFixedDeltaTime()
 	{
-		return FAppTime::Get().FixedDeltaTime;
+		return FAppTime::GetFixedDeltaTime();
 	}
 	
 	/**
@@ -568,7 +568,7 @@ public:
 	 */
 	static bool UseFixedTimeStep()
 	{
-		return FAppTime::Get().bUseFixedTimeStep;
+		return FAppTime::UseFixedTimeStep();
 	}
 
 	/**
@@ -578,7 +578,7 @@ public:
 	 */
 	static double GetCurrentTime()
 	{
-		return FAppTime::Get().CurrentTime;
+		return FAppTime::GetCurrentTime();
 	}
 
 	/**
@@ -588,7 +588,7 @@ public:
 	 */
 	static double GetLastTime()
 	{
-		return FAppTime::Get().LastTime;
+		return FAppTime::GetLastTime();
 	}
 
 	/**
@@ -598,7 +598,7 @@ public:
 	 */
 	static double GetDeltaTime()
 	{
-		return FAppTime::Get().DeltaTime;
+		return FAppTime::GetDeltaTime();
 	}
 
 	/**
@@ -608,7 +608,7 @@ public:
 	 */
 	static double GetIdleTime()
 	{
-		return FAppTime::Get().IdleTime;
+		return FAppTime::GetIdleTime();
 	}
 
 	/**
@@ -618,7 +618,7 @@ public:
 	 */
 	static double GetGameTime()
 	{
-		return FAppTime::Get().GameTime;
+		return FAppTime::GetGameTime();
 	}
 
 	/**
@@ -628,7 +628,7 @@ public:
 	 */
 	static double GetIdleTimeOvershoot()
 	{
-		return FAppTime::Get().IdleTimeOvershoot;
+		return FAppTime::GetIdleTimeOvershoot();
 	}
 
 	/**
@@ -639,7 +639,7 @@ public:
 	 */
 	static FTimecode GetTimecode()
 	{
-		return FAppTime::Get().GetTimecode();
+		return FAppTime::GetTimecode();
 	}
 
 	/**
@@ -650,7 +650,7 @@ public:
 	 */
 	static FFrameRate GetTimecodeFrameRate()
 	{
-		return FAppTime::Get().GetTimecodeFrameRate();
+		return FAppTime::GetTimecodeFrameRate();
 	}
 
 	/**
@@ -662,7 +662,7 @@ public:
 	 */
 	static TOptional<FQualifiedFrameTime> GetCurrentFrameTime()
 	{
-		return FAppTime::Get().CurrentFrameTime;
+		return FAppTime::GetCurrentFrameTime();
 	}
 
 	/**
@@ -673,9 +673,7 @@ public:
 	// UE_DEPRECATED(5.8, "@todo delta time refactor - deprecate this function")
 	static void SetFixedDeltaTime(double Seconds)
 	{
-		FAppTime Time = FAppTime::Get();
-		Time.FixedDeltaTime = Seconds;
-		FAppTime::Set(Time);
+		FAppTime::SetFixedDeltaTime(Seconds);
 	}
 
 	/**
@@ -686,10 +684,9 @@ public:
 	// UE_DEPRECATED(5.8, "@todo delta time refactor - deprecate this function")
 	static void SetUseFixedTimeStep(bool bVal)
 	{
-		FAppTime Time = FAppTime::Get();
-		Time.bUseFixedTimeStep = bVal;
-		FAppTime::Set(Time);
+		FAppTime::SetUseFixedTimeStep(bVal);
 	}
+
 	/**
 	 * Sets current time in seconds.
 	 *
@@ -698,19 +695,16 @@ public:
 	// UE_DEPRECATED(5.8, "@todo delta time refactor - deprecate this function")
 	static void SetCurrentTime(double Seconds)
 	{
-		FAppTime Time = FAppTime::Get();
-		Time.CurrentTime = Seconds;
-		FAppTime::Set(Time);
+		FAppTime::SetCurrentTime(Seconds);
 	}
 
 	/** Updates Last time to CurrentTime. */
 	// UE_DEPRECATED(5.8, "@todo delta time refactor - deprecate this function")
 	static void UpdateLastTime()
 	{
-		FAppTime Time = FAppTime::Get();
-		Time.LastTime = Time.CurrentTime;
-		FAppTime::Set(Time);
+		FAppTime::UpdateLastTime();
 	}
+
 	/**
 	 * Sets time delta in seconds.
 	 *
@@ -719,10 +713,9 @@ public:
 	// UE_DEPRECATED(5.8, "@todo delta time refactor - deprecate this function")
 	static void SetDeltaTime(double Seconds)
 	{
-		FAppTime Time = FAppTime::Get();
-		Time.DeltaTime = Seconds;
-		FAppTime::Set(Time);
+		FAppTime::SetDeltaTime(Seconds);
 	}
+
 	/**
 	 * Sets idle time in seconds.
 	 *
@@ -731,10 +724,9 @@ public:
 	// UE_DEPRECATED(5.8, "@todo delta time refactor - deprecate this function")
 	static void SetIdleTime(double Seconds)
 	{
-		FAppTime Time = FAppTime::Get();
-		Time.IdleTime = Seconds;
-		FAppTime::Set(Time);
+		FAppTime::SetIdleTime(Seconds);
 	}
+
 	/**
 	 * Sets overall game time in seconds.
 	 *
@@ -743,10 +735,9 @@ public:
 	// UE_DEPRECATED(5.8, "@todo delta time refactor - deprecate this function")
 	static void SetGameTime(double Seconds)
 	{
-		FAppTime Time = FAppTime::Get();
-		Time.GameTime = Seconds;
-		FAppTime::Set(Time);
+		FAppTime::SetGameTime(Seconds);
 	}
+
 	/**
 	 * Sets idle time overshoot in seconds (the time beyond the wait time we requested for the frame). Only valid when IdleTime is > 0.
 	 *
@@ -755,9 +746,7 @@ public:
 	// UE_DEPRECATED(5.8, "@todo delta time refactor - deprecate this function")
 	static void SetIdleTimeOvershoot(double Seconds)
 	{
-		FAppTime Time = FAppTime::Get();
-		Time.IdleTimeOvershoot = Seconds;
-		FAppTime::Set(Time);
+		FAppTime::SetIdleTimeOvershoot(Seconds);
 	}
 
 	/**
@@ -768,18 +757,14 @@ public:
 	// UE_DEPRECATED(5.8, "@todo delta time refactor - deprecate this function")
 	static void SetCurrentFrameTime(FQualifiedFrameTime InFrameTime)
 	{
-		FAppTime Time = FAppTime::Get();
-		Time.CurrentFrameTime = InFrameTime;
-		FAppTime::Set(Time);
+		FAppTime::SetCurrentFrameTime(InFrameTime);
 	}
 
 	/** Invalidate the current frame time. It will reset the TOptional. */
 	// UE_DEPRECATED(5.8, "@todo delta time refactor - deprecate this function")
 	static void InvalidateCurrentFrameTime()
 	{
-		FAppTime Time = FAppTime::Get();
-		Time.CurrentFrameTime.Reset();
-		FAppTime::Set(Time);
+		FAppTime::InvalidateCurrentFrameTime();
 	}
 
 	/**
diff --git a/Engine/Source/Runtime/Core/Public/Misc/AppTime.h b/Engine/Source/Runtime/Core/Public/Misc/AppTime.h
--- a/Engine/Source/Runtime/Core/Public/Misc/AppTime.h
+++ b/Engine/Source/Runtime/Core/Public/Misc/AppTime.h
@@ -21,43 +21,230 @@ namespace UE
  */
 class FAppTime
 {
+public:
+	// Helper type to perform automatic relaxed-order loads/stores of a std::atomic.
+	template <typename TValue>
+	struct TRelaxedAtomic
+	{
+		std::atomic<TValue> Value;
+
+		TRelaxedAtomic() = default;
+
+		TRelaxedAtomic(TRelaxedAtomic const& Other)
+		{
+			(*this) = Other.Value;
+		}
+
+		TRelaxedAtomic(TValue InValue)
+		{
+			(*this) = InValue;
+		}
+
+		TRelaxedAtomic& operator = (TValue InValue)
+		{
+			Value.store(InValue, std::memory_order_relaxed);
+			return *this;
+		}
+
+		TRelaxedAtomic& operator = (TRelaxedAtomic const& RHS)
+		{
+			TValue LoadedValue = RHS;
+			(*this) = LoadedValue;
+			return *this;
+		}
+
+		operator TValue() const
+		{
+			return Value.load(std::memory_order_relaxed);
+		}
+	};
+
 	friend UE::FInheritedContextScope;
 	friend UE::FInheritedContextBase;
+	friend class FApp;
 
-public:
 #if WITH_FIXED_TIME_STEP_SUPPORT
-	/** Holds a flag whether we want to use a fixed time step or not. */
+	// Holds a flag whether we want to use a fixed time step or not.
 	bool bUseFixedTimeStep = false;
 #else
 	static constexpr bool bUseFixedTimeStep = false;
 #endif
 
-	/** Holds time step if a fixed delta time is wanted. */
+	// Holds time step if a fixed delta time is wanted.
 	double FixedDeltaTime = 1 / 30.0;
 
-	/** Holds current time. */
-	double CurrentTime = 0.0;
+	// Holds current time.
+	TRelaxedAtomic<double> CurrentTime = 0.0;
 
-	/** Holds previous value of CurrentTime. */
+	// Holds previous value of CurrentTime.
 	double LastTime = 0.0;
 
-	/** Holds current delta time in seconds. */
+	// Holds current delta time in seconds.
 	double DeltaTime = 1 / 30.0;
 
-	/** Holds time we spent sleeping in UpdateTimeAndHandleMaxTickRate() if our frame time was smaller than one allowed by target FPS. */
+	// Holds time we spent sleeping in UpdateTimeAndHandleMaxTickRate() if our frame time was smaller than one allowed by target FPS.
 	double IdleTime = 0.0;
 
-	/** Holds the amount of IdleTime that was LONGER than we tried to sleep. The OS can't sleep the exact amount of time, so this measures that overshoot. */
+	// Holds the amount of IdleTime that was LONGER than we tried to sleep. The OS can't sleep the exact amount of time, so this measures that overshoot.
 	double IdleTimeOvershoot = 0.0;
 
-	/** Holds overall game time. */
+	// Holds overall game time.
 	double GameTime = 0.0;
 
 	/** Holds the current frame time and framerate. */
 	TOptional<FQualifiedFrameTime> CurrentFrameTime;
 
-	FTimecode GetTimecode() const
+private:
+	//
+	// Returns the FAppTime value for the current thread's time context. May be called by the game thread, or decendent TaskGraph tasks of the game thread.
+	// Calling this function from any other context will result in an ensure failure. This is because the game thread is the authority on the FAppTime value,
+	// and this value is pipelined through the TaskGraph so that decendent tasks receive the correct time value for the frame they belong to.
+	// Reading the FAppTime value outside of these contexts is a data race with the game thread.
+	//
+	static CORE_API const FAppTime& Get();
+
+public:
+	UE_DEPRECATED(5.8,
+		"Do not call FAppTime::Get_UnsafeNoCheck() in new engine code. This function is provided "
+		"as a temporary measure to avoid triggering the time context ensure in FAppTime::Get() "
+		"for engine systems that read the FAppTime value in a non-threadsafe manner, until those "
+		"systems can be refactored.")
+	static CORE_API const FAppTime& Get_UnsafeNoCheck();
+
+	// Sets the new FAppTime value. May only be called from the game thread.
+	static CORE_API void Set(const FAppTime& NewValue);
+
+	//
+	// Setter functions matching those from FApp
+	//
+
+	static void SetFixedDeltaTime(double Seconds)
+	{
+		FAppTime Time = Get();
+		Time.FixedDeltaTime = Seconds;
+		Set(Time);
+	}
+
+	static void SetUseFixedTimeStep(bool bVal)
+	{
+		FAppTime Time = Get();
+		Time.bUseFixedTimeStep = bVal;
+		Set(Time);
+	}
+
+	static void SetCurrentTime(double Seconds)
+	{
+		FAppTime Time = Get();
+		Time.CurrentTime = Seconds;
+		Set(Time);
+	}
+
+	static void UpdateLastTime()
+	{
+		FAppTime Time = Get();
+		Time.LastTime = Time.CurrentTime;
+		Set(Time);
+	}
+
+	static void SetDeltaTime(double Seconds)
+	{
+		FAppTime Time = Get();
+		Time.DeltaTime = Seconds;
+		Set(Time);
+	}
+
+	static void SetIdleTime(double Seconds)
 	{
+		FAppTime Time = Get();
+		Time.IdleTime = Seconds;
+		Set(Time);
+	}
+
+	static void SetGameTime(double Seconds)
+	{
+		FAppTime Time = Get();
+		Time.GameTime = Seconds;
+		Set(Time);
+	}
+
+	static void SetIdleTimeOvershoot(double Seconds)
+	{
+		FAppTime Time = Get();
+		Time.IdleTimeOvershoot = Seconds;
+		Set(Time);
+	}
+
+	static void SetCurrentFrameTime(FQualifiedFrameTime InFrameTime)
+	{
+		FAppTime Time = Get();
+		Time.CurrentFrameTime = InFrameTime;
+		Set(Time);
+	}
+
+	static void InvalidateCurrentFrameTime()
+	{
+		FAppTime Time = Get();
+		Time.CurrentFrameTime.Reset();
+		Set(Time);
+	}
+
+	//
+	// Equivalent to GetCurrentTime(), but can be called from any thread, rather than only from threads with valid time contexts.
+	// The returned value is loaded via a relaxed atomic read, and therefore can race updates from the game thread.
+	// 
+	// Use only for systems that require an approximate incrementing timer, where getting a value from a frame that is off-by-one is not important.
+	// Do not use in systems where the current time directly affects the game simulation or rendered pixels. Those systems should use GetCurrentTime()
+	// to ensure they retrieve the appropriate time value via the inherited time context.
+	//
+	static double GetApproximateCurrentTime() { return GameThreadInstance.CurrentTime; }
+
+	//
+	// Pipelined Getter functions.
+	// Must only be called from the game thread, or threads with a valid time context.
+	//
+
+	// Gets time step in seconds if a fixed delta time is wanted.
+	static double GetFixedDeltaTime() { return Get().FixedDeltaTime; }
+
+	// Gets whether we want to use a fixed time step or not.
+	static bool UseFixedTimeStep() { return Get().bUseFixedTimeStep; }
+
+	// Gets current frame's time in seconds.
+	static double GetCurrentTime() { return Get().CurrentTime; }
+
+	// Gets the previous frame's value of CurrentTime
+	static double GetLastTime() { return Get().LastTime; }
+
+	// Gets the current frame's time delta in seconds.
+	static double GetDeltaTime() { return Get().DeltaTime; }
+
+	// Gets idle time in seconds.
+	static double GetIdleTime() { return Get().IdleTime; }
+
+	// Gets overall game time in seconds.
+	static double GetGameTime() { return Get().GameTime; }
+
+	// Gets idle time overshoot in seconds (the time beyond the wait time we requested for the frame). Only valid when IdleTime is > 0.
+	static double GetIdleTimeOvershoot() { return Get().IdleTimeOvershoot; }
+
+	//
+	// Gets a frame number generated by the engine's timecode provider.
+	// The current frame time is used to generate a timecode value.
+	// The optional will be false if no timecode provider was set or is not in a synchronized state.
+	//
+	static TOptional<FQualifiedFrameTime> GetCurrentFrameTime()
+	{
+		return Get().CurrentFrameTime;
+	}
+
+	//
+	// Convert the current frame time into a readable timecode.
+	// If the current frame time is not set, the timecode value will be defaulted.
+	//
+	static FTimecode GetTimecode()
+	{
+		TOptional<FQualifiedFrameTime> CurrentFrameTime = GetCurrentFrameTime();
+
 		if (CurrentFrameTime.IsSet())
 		{
 			return CurrentFrameTime->ToTimecode();
@@ -66,8 +253,11 @@ public:
 		return FTimecode();
 	}
 
-	FFrameRate GetTimecodeFrameRate() const
+	// Get the frame rate of the current frame time. If the current frame time is not set, the frame rate value will be defaulted.
+	static FFrameRate GetTimecodeFrameRate()
 	{
+		TOptional<FQualifiedFrameTime> CurrentFrameTime = GetCurrentFrameTime();
+
 		if (CurrentFrameTime.IsSet())
 		{
 			return CurrentFrameTime->Rate;
@@ -76,27 +266,6 @@ public:
 		return FFrameRate();
 	}
 
-	//
-	// Returns the FAppTime value for the current thread's time context. May be called by the game thread, or decendent TaskGraph tasks of the game thread.
-	// Calling this function from any other context will result in an ensure failure. This is because the game thread is the authority on the FAppTime value,
-	// and this value is pipelined through the TaskGraph so that decendent tasks receive the correct time value for the frame they belong to.
-	// Reading the FAppTime value outside of these contexts is a data race with the game thread.
-	//
-	static CORE_API const FAppTime& Get();
-
-	//
-	// Sets the new FAppTime value.
-	// May only be called from the game thread
-	//
-	static CORE_API void Set(const FAppTime& NewValue);
-
-	UE_DEPRECATED(5.8,
-		"Do not call FAppTime::Get_UnsafeNoCheck() in new engine code. This function is provided "
-		"as a temporary measure to avoid triggering the time context ensure in FAppTime::Get() "
-		"for engine systems that read the FAppTime value in a non-threadsafe manner, until those "
-		"systems can be refactored.")
-	static CORE_API const FAppTime& Get_UnsafeNoCheck();
-
 private:
 	static FAppTime GameThreadInstance;
 	static thread_local TSharedPtr<const FAppTime> TLSInstance;
