Add missing CORE_API on FAppTime::GameThreadInstance static.
- Required as FAppTime::GetApproximateCurrentTime() is a static inline function that returns a member of GameThreadInstance.

#Synced-CL 48676711

diff --git a/Engine/Source/Runtime/Core/Private/Misc/AppTime.cpp b/Engine/Source/Runtime/Core/Private/Misc/AppTime.cpp
--- a/Engine/Source/Runtime/Core/Private/Misc/AppTime.cpp
+++ b/Engine/Source/Runtime/Core/Private/Misc/AppTime.cpp
@@ -3,7 +3,7 @@
 #include "Misc/AppTime.h"
 #include "CoreGlobals.h"
 
-FAppTime FAppTime::GameThreadInstance;
+CORE_API FAppTime FAppTime::GameThreadInstance;
 thread_local TSharedPtr<const FAppTime> FAppTime::TLSInstance;
 
 [[nodiscard]] CORE_API TSharedPtr<const FAppTime> FAppTime::Fork()
diff --git a/Engine/Source/Runtime/Core/Public/Misc/AppTime.h b/Engine/Source/Runtime/Core/Public/Misc/AppTime.h
--- a/Engine/Source/Runtime/Core/Public/Misc/AppTime.h
+++ b/Engine/Source/Runtime/Core/Public/Misc/AppTime.h
@@ -267,7 +267,7 @@ public:
 	}
 
 private:
-	static FAppTime GameThreadInstance;
+	static CORE_API FAppTime GameThreadInstance;
 	static thread_local TSharedPtr<const FAppTime> TLSInstance;
 
 	[[nodiscard]] static CORE_API TSharedPtr<const FAppTime> Fork();
