ParkingLot: Fix a rare false positive in the tests

#rnx
#Synced-CL 32500503

diff --git a/Engine/Source/Runtime/Core/Tests/Async/ParkingLotTest.cpp b/Engine/Source/Runtime/Core/Tests/Async/ParkingLotTest.cpp
--- a/Engine/Source/Runtime/Core/Tests/Async/ParkingLotTest.cpp
+++ b/Engine/Source/Runtime/Core/Tests/Async/ParkingLotTest.cpp
@@ -47,15 +47,19 @@ TEST_CASE("Core::Async::ParkingLot", "[Core][Async]")
 		for (int32 Index = 0; Index <= TaskCount; ++Index)
 		{
 			int32 WakeCount = 0;
-			ParkingLot::WakeOne(&WaitCount, [&WakeCount, &Sequence, Index](ParkingLot::FWakeState WakeState) -> uint64
+			bool bDidWake = false;
+			bool bHasWaitingThreads = false;
+			ParkingLot::WakeOne(&WaitCount, [&](ParkingLot::FWakeState WakeState) -> uint64
 			{
 				++WakeCount;
-				CHECK(WakeState.bDidWake == (Index < TaskCount));
-				CHECK(WakeState.bHasWaitingThreads == (Index + 1 < TaskCount));
+				bDidWake = WakeState.bDidWake;
+				bHasWaitingThreads = WakeState.bHasWaitingThreads;
 				return ++Sequence;
 			});
 			// The callback must be invoked exactly once.
 			CHECK(WakeCount == 1);
+			CHECK(bDidWake == (Index < TaskCount));
+			CHECK(bHasWaitingThreads >= (Index + 1 < TaskCount));
 		}
 
 		// Spin until the tasks are complete.
