Added missing RoundDownToPowerOfTwo64 to GenericPlatformMath.h

#Synced-CL 46666047

diff --git a/Engine/Source/Runtime/Core/Public/GenericPlatform/GenericPlatformMath.h b/Engine/Source/Runtime/Core/Public/GenericPlatform/GenericPlatformMath.h
--- a/Engine/Source/Runtime/Core/Public/GenericPlatform/GenericPlatformMath.h
+++ b/Engine/Source/Runtime/Core/Public/GenericPlatform/GenericPlatformMath.h
@@ -840,6 +840,11 @@ struct FGenericPlatformMath
 		return uint64(1) << CeilLogTwo64(V);
 	}
 
+	[[nodiscard]] static constexpr UE_FORCEINLINE_HINT uint64 RoundDownToPowerOfTwo64(uint64 V)
+	{
+		return (V == 0) ? 0u : (uint64(1) << FloorLog2_64(V));
+	}
+
 	/** Spreads bits to every other. */
 	[[nodiscard]] static constexpr inline uint32 MortonCode2( uint32 x )
 	{
