Core: Fixed TFunctionWithContext to support assignment from a const lambda

#Synced-CL 32500503

diff --git a/Engine/Source/Runtime/Core/Public/Templates/FunctionWithContext.h b/Engine/Source/Runtime/Core/Public/Templates/FunctionWithContext.h
--- a/Engine/Source/Runtime/Core/Public/Templates/FunctionWithContext.h
+++ b/Engine/Source/Runtime/Core/Public/Templates/FunctionWithContext.h
@@ -48,7 +48,7 @@ public:
 		requires (!std::is_same_v<std::decay_t<LambdaType>, TFunctionWithContext>) &&
 			std::is_invocable_r_v<ReturnType, std::decay_t<LambdaType>, ArgTypes...>
 		: Function(&Call<std::decay_t<LambdaType>>)
-		, Context(&Lambda)
+		, Context((void*)&Lambda)
 	{
 	}
 
@@ -59,7 +59,7 @@ public:
 			std::is_invocable_r_v<ReturnType, std::decay_t<LambdaType>, ArgTypes...>
 	{
 		Function = &Call<std::decay_t<LambdaType>>;
-		Context = &Lambda;
+		Context = (void*)&Lambda;
 		return *this;
 	}
 
