Fix access violations triggered by the Stall Detector Thread
- Only CONTEXT_CONTROL was requested when reading the remote thread's context. Walking some stack frames requires register values which are not captured by this.
- Replaced with CONTEXT_ALL to get all possible register state.

Also add a missing call to CloseHandle().

#Synced-CL 48676711

diff --git a/Engine/Source/Runtime/Core/Private/Windows/WindowsPlatformStackWalk.cpp b/Engine/Source/Runtime/Core/Private/Windows/WindowsPlatformStackWalk.cpp
--- a/Engine/Source/Runtime/Core/Private/Windows/WindowsPlatformStackWalk.cpp
+++ b/Engine/Source/Runtime/Core/Private/Windows/WindowsPlatformStackWalk.cpp
@@ -179,16 +179,22 @@ UE_AUTORTFM_ALWAYS_OPEN void FWindowsPlatformStackWalk::ThreadStackWalkAndDump(A
 
 		// Suspend the thread before grabbing its context (possible fix for incomplete callstacks)
 		SuspendThread(ThreadHandle);
+
 		// Give task scheduler some time to actually suspend the thread
 		FPlatformProcess::Sleep(0.01f);
+
 		FWindowsThreadContextWrapper ContextWapper;
-		ContextWapper.Context.ContextFlags = CONTEXT_CONTROL;
+		ContextWapper.Context.ContextFlags = CONTEXT_ALL;
 		ContextWapper.ThreadHandle = ThreadHandle;
+
 		if (GetThreadContext(ThreadHandle, &ContextWapper.Context))
 		{
 			FGenericPlatformStackWalk::StackWalkAndDump(HumanReadableString, HumanReadableStringSize, IgnoreCount, &ContextWapper);
 		}
+
 		ResumeThread(ThreadHandle);
+
+		CloseHandle(ThreadHandle);
 	}
 }
 
@@ -340,7 +346,7 @@ UE_AUTORTFM_ALWAYS_OPEN uint32 FWindowsPlatformStackWalk::CaptureThreadStackBack
 	{
 		CONTEXT ThreadContext;
 		FMemory::Memzero(ThreadContext);
-		ThreadContext.ContextFlags = CONTEXT_CONTROL;
+		ThreadContext.ContextFlags = CONTEXT_ALL;
 
 		if (::GetThreadContext(ThreadHandle, &ThreadContext))
 		{
