Added missing AlignmentTemplates include

#Synced-CL 46666047

diff --git a/Engine/Source/Runtime/Core/Public/HAL/MipFieldPageAllocator.h b/Engine/Source/Runtime/Core/Public/HAL/MipFieldPageAllocator.h
--- a/Engine/Source/Runtime/Core/Public/HAL/MipFieldPageAllocator.h
+++ b/Engine/Source/Runtime/Core/Public/HAL/MipFieldPageAllocator.h
@@ -3,8 +3,9 @@
 #pragma once
 
 #include "HAL/Platform.h"
-#include "Math/UnrealMathUtility.h"
 #include "HAL/UnrealMemory.h"
+#include "Math/UnrealMathUtility.h"
+#include "Templates/AlignmentTemplates.h"
 
 
 #ifndef MIP_FIELD_ALLOCATOR_DEBUG
