Added  parameters to TGuardConsoleVariable to be able to specify priority when setting cvars from testcode.

#Synced-CL 48774492

diff --git a/Engine/Source/Runtime/Core/Public/HAL/IConsoleManager.h b/Engine/Source/Runtime/Core/Public/HAL/IConsoleManager.h
--- a/Engine/Source/Runtime/Core/Public/HAL/IConsoleManager.h
+++ b/Engine/Source/Runtime/Core/Public/HAL/IConsoleManager.h
@@ -1356,20 +1356,22 @@ struct TGuardConsoleVariable
 {
 	T PreviousValue;
 	IConsoleVariable* ConsoleVariable;
-
-	TGuardConsoleVariable(const TCHAR* ConsoleVariableName, const T& NewValue)
-		: TGuardConsoleVariable(IConsoleManager::Get().FindConsoleVariable(ConsoleVariableName), NewValue)
+	IConsoleVariable::FSetContext SetContext;
+	
+	TGuardConsoleVariable(const TCHAR* ConsoleVariableName, const T& NewValue, EConsoleVariableFlags Flags = ECVF_SetByCode, FName Tag = NAME_None)
+		: TGuardConsoleVariable(IConsoleManager::Get().FindConsoleVariable(ConsoleVariableName), NewValue, Flags, Tag)
 	{
 		ensureMsgf(ConsoleVariable, TEXT("TGuardConsoleVariable: Trying to set console variable \"%s\" but it doesn't exist"), ConsoleVariableName);
 	}
 
-	TGuardConsoleVariable(IConsoleVariable* ConsoleVariable, const T& NewValue)
+	TGuardConsoleVariable(IConsoleVariable* ConsoleVariable, const T& NewValue, EConsoleVariableFlags Flags = ECVF_SetByCode, FName Tag = NAME_None)
 		: ConsoleVariable(ConsoleVariable)
+		, SetContext(Flags, Tag)
 	{
 		if (ConsoleVariable)
 		{
 			ConsoleVariable->GetValue(PreviousValue);
-			ConsoleVariable->Set(NewValue);
+			ConsoleVariable->Set(NewValue, SetContext);
 		}
 	}
 
@@ -1390,7 +1392,7 @@ struct TGuardConsoleVariable
 	{
 		if (ConsoleVariable)
 		{
-			ConsoleVariable->Set(PreviousValue);
+			ConsoleVariable->Set(PreviousValue, SetContext);
 		}
 	}
 };
