Defer frees until outside of the lock.

Calling malloc free needs to take a lock. A lock that could be taken by another thread that is currently waiting on the callstack-tracing lock. This has almost never happened, only induced by adding sleeps. Note that the realloc is likely still susceptible, but that realloc has been there for quite some time by this point.

#rnx

#Synced-CL Unknown

diff --git a/Engine/Source/Runtime/Core/Private/ProfilingDebugging/Microsoft/WindowsCallstackTrace.cpp b/Engine/Source/Runtime/Core/Private/ProfilingDebugging/Microsoft/WindowsCallstackTrace.cpp
--- a/Engine/Source/Runtime/Core/Private/ProfilingDebugging/Microsoft/WindowsCallstackTrace.cpp
+++ b/Engine/Source/Runtime/Core/Private/ProfilingDebugging/Microsoft/WindowsCallstackTrace.cpp
@@ -212,7 +212,7 @@ private:
 	};
 	typedef TGrowOnlyLockFreeHash<FFunctionLookupSetEntry, uint64, int32> FFunctionLookupSet;
 
-	void					ReplaceModule(TArrayView<FModule> ModulesView, FModule& Module);
+	void					ReplaceModule(TArrayView<FModule> ModulesView, FModule& Module, void (*FreeFunc)(void*, void*), void* FreeParam);
 	const FFunction*		LookupFunction(UPTRINT Address, FLookupState& State) const;
 	static FBacktracer*		Instance;
 	mutable FCriticalSection Lock;
@@ -269,7 +269,9 @@ FBacktracer* FBacktracer::Get()
 }
 
 ////////////////////////////////////////////////////////////////////////////////
-void FBacktracer::ReplaceModule(TArrayView<FModule> ModulesView, FModule& Module)
+void FBacktracer::ReplaceModule(
+	TArrayView<FModule> ModulesView, FModule& Module,
+	void (*FreeFunc)(void*, void*), void* FreeParam)
 {
 	FModule& Existing = ModulesView[0];
 
@@ -279,13 +281,13 @@ void FBacktracer::ReplaceModule(TArrayView<FModule> ModulesView, FModule& Module
 		&& Module.IdSize == Existing.IdSize
 		&& Module.NumFunctions == Existing.NumFunctions)
 	{
-		Malloc->Free(Module.Functions);
+		FreeFunc(Module.Functions, FreeParam);
 		return;
 	}
 
 	// Otherwise the module is new and located in a region previously occupied
 	// by one or more modules that have since been unloaded.
-	Malloc->Free(Existing.Functions);
+	FreeFunc(Existing.Functions, FreeParam);
 	Existing = Module;
 
 	uint32 IdEnd = Module.Id + Module.IdSize;
@@ -296,8 +298,11 @@ void FBacktracer::ReplaceModule(TArrayView<FModule> ModulesView, FModule& Module
 			break;
 		}
 
+		FreeFunc(Iter.Functions, FreeParam);
+
 		Iter.Id = 0;
 		Iter.IdSize = 0;
+		Iter.Functions = nullptr;
 	}
 }
 
@@ -533,6 +538,23 @@ void FBacktracer::AddModule(UPTRINT ModuleBase, const TCHAR* Name)
 		OutTable,
 	};
 
+	struct FDeferredFrees
+	{
+		enum : uint32 { Max = 8 };
+		void*	Addresses[Max];
+		uint32	Num = 0;
+		void	Defer(void* Addr) { if (Num < Max) Addresses[Num++] = Addr; }
+	};
+
+	FDeferredFrees DeferredFrees;
+	ON_SCOPE_EXIT
+	{
+		for (uint32 i = 0; i < DeferredFrees.Num; ++i)
+		{
+			Malloc->Free(DeferredFrees.Addresses[i]);
+		}
+	};
+
 	{
 		FScopeLock _(&Lock);
 
@@ -551,6 +573,11 @@ void FBacktracer::AddModule(UPTRINT ModuleBase, const TCHAR* Name)
 		int32 Index = Algo::UpperBound(ModulesView, Module.Id, FIdPredicate()) - 1;
 		if (Index >= 0)
 		{
+			auto FreeThunk = [] (void* AddrToFree, void* Param)
+			{
+				((FDeferredFrees*)Param)->Defer(AddrToFree);
+			};
+
 			// Move the view so it starts at the first module of interest.
 			ModulesView = ModulesView.Mid(Index);
 
@@ -559,7 +586,7 @@ void FBacktracer::AddModule(UPTRINT ModuleBase, const TCHAR* Name)
 			uint32 IdEnd = ModulesView[0].Id + ModulesView[0].IdSize;
 			if (Module.Id < IdEnd)
 			{
-				ReplaceModule(ModulesView, Module);
+				ReplaceModule(ModulesView, Module, FreeThunk, &DeferredFrees);
 				return;
 			}
 
@@ -568,15 +595,15 @@ void FBacktracer::AddModule(UPTRINT ModuleBase, const TCHAR* Name)
 			ModulesView = ModulesView.Mid(1);
 			if (!ModulesView.IsEmpty() && (Module.Id + Module.IdSize) > ModulesView[0].Id)
 			{
-				ReplaceModule(ModulesView, Module);
+				ReplaceModule(ModulesView, Module, FreeThunk, &DeferredFrees);
 				return;
 			}
 		}
 
 		if (ModulesNum + 1 > ModulesCapacity)
 		{
-			ModulesCapacity += 8;
-			Modules = (FModule*) Malloc->Realloc(Modules, sizeof(FModule)* ModulesCapacity);
+			ModulesCapacity += 32;
+			Modules = (FModule*)Malloc->Realloc(Modules, sizeof(FModule) * ModulesCapacity);
 		}
 		Modules[ModulesNum++] = Module;
 	}
