[Scheduler]
- Fix priority order where low prio game thread tasks were picked up before high prio worker thread tasks

#Synced-CL Unknown

diff --git a/Engine/Source/Runtime/Core/Private/Async/Fundamental/Scheduler.cpp b/Engine/Source/Runtime/Core/Private/Async/Fundamental/Scheduler.cpp
--- a/Engine/Source/Runtime/Core/Private/Async/Fundamental/Scheduler.cpp
+++ b/Engine/Source/Runtime/Core/Private/Async/Fundamental/Scheduler.cpp
@@ -706,15 +706,52 @@ namespace LowLevelTasks
 		return AnyExecuted;
 	}
 
+	template <typename LocalQueueType>
+	class TCombinedQueue
+	{
+		LocalQueueType* GameThreadLocalQueue;
+		LocalQueueType* WorkerLocalQueue;
+	public:
+		TCombinedQueue(LocalQueueType* InGameThreadLocalQueue, LocalQueueType* InWorkerLocalQueue)
+			: GameThreadLocalQueue(InGameThreadLocalQueue)
+			, WorkerLocalQueue(InWorkerLocalQueue)
+		{
+		}
+
+		FTask* Dequeue(bool bPermitBackgroundWork)
+		{
+			const int32 MaxPriority = bPermitBackgroundWork ? int32(ETaskPriority::Count) : int32(ETaskPriority::ForegroundCount);
+
+			for (int32 PriorityIndex = 0; PriorityIndex < MaxPriority; ++PriorityIndex)
+			{
+				FTask* Item = GameThreadLocalQueue->StealLocal(PriorityIndex);
+				if (Item)
+				{
+					return Item;
+				}
+
+				Item = WorkerLocalQueue->Dequeue(PriorityIndex);
+				if (Item)
+				{
+					return Item;
+				}
+			}
+			return nullptr;
+		}
+	};
+
 	void FScheduler::StandbyLoop(Private::FWaitEvent* WorkerEvent, FSchedulerTls::FLocalQueueType* WorkerLocalQueue, uint32 WaitCycles, bool bPermitBackgroundWork)
 	{
 		bool bPreparingStandby = false;
 		Private::FOutOfWork OutOfWork;
+
+		using FCombinedQueue = TCombinedQueue<FSchedulerTls::FLocalQueueType>;
+		FCombinedQueue CombinedQueue(GameThreadLocalQueue.Get(), WorkerLocalQueue);
+
 		while (true)
 		{
 			bool bExecutedSomething = false;
-			while (TryExecuteTaskFrom<FSchedulerTls::FLocalQueueType, &FSchedulerTls::FLocalQueueType::StealLocal, true>(WorkerEvent, GameThreadLocalQueue.Get(), OutOfWork, bPermitBackgroundWork)
-				|| TryExecuteTaskFrom<FSchedulerTls::FLocalQueueType, &FSchedulerTls::FLocalQueueType::Dequeue, true>(WorkerEvent, WorkerLocalQueue, OutOfWork, bPermitBackgroundWork)
+			while (TryExecuteTaskFrom<FCombinedQueue, &FCombinedQueue::Dequeue, true>(WorkerEvent, &CombinedQueue, OutOfWork, bPermitBackgroundWork)
 				|| TryExecuteTaskFrom<FSchedulerTls::FLocalQueueType, &FSchedulerTls::FLocalQueueType::DequeueSteal, true>(WorkerEvent, WorkerLocalQueue, OutOfWork, bPermitBackgroundWork))
 			{
 				bPreparingStandby = false;
@@ -750,13 +787,15 @@ namespace LowLevelTasks
 
 	void FScheduler::WorkerLoop(Private::FWaitEvent* WorkerEvent, FSchedulerTls::FLocalQueueType* WorkerLocalQueue, uint32 WaitCycles, bool bPermitBackgroundWork)
 	{
+		using FCombinedQueue = TCombinedQueue<FSchedulerTls::FLocalQueueType>;
+		FCombinedQueue CombinedQueue(GameThreadLocalQueue.Get(), WorkerLocalQueue);
+
 		bool bPreparingWait = false;
 		Private::FOutOfWork OutOfWork;
 		while (true)
 		{
 			bool bExecutedSomething = false;
-			while (TryExecuteTaskFrom<FSchedulerTls::FLocalQueueType, &FSchedulerTls::FLocalQueueType::StealLocal, false>(WorkerEvent, GameThreadLocalQueue.Get(), OutOfWork, bPermitBackgroundWork)
-				|| TryExecuteTaskFrom<FSchedulerTls::FLocalQueueType, &FSchedulerTls::FLocalQueueType::Dequeue, false>(WorkerEvent, WorkerLocalQueue, OutOfWork, bPermitBackgroundWork)
+			while (TryExecuteTaskFrom<FCombinedQueue, &FCombinedQueue::Dequeue, false>(WorkerEvent, &CombinedQueue, OutOfWork, bPermitBackgroundWork)
 				|| TryExecuteTaskFrom<FSchedulerTls::FLocalQueueType, &FSchedulerTls::FLocalQueueType::DequeueSteal, false>(WorkerEvent, WorkerLocalQueue, OutOfWork, bPermitBackgroundWork))
 			{
 				bPreparingWait = false;
diff --git a/Engine/Source/Runtime/Core/Public/Async/Fundamental/LocalQueue.h b/Engine/Source/Runtime/Core/Public/Async/Fundamental/LocalQueue.h
--- a/Engine/Source/Runtime/Core/Public/Async/Fundamental/LocalQueue.h
+++ b/Engine/Source/Runtime/Core/Public/Async/Fundamental/LocalQueue.h
@@ -238,14 +238,24 @@ public:
 			}
 		}
 
-		inline FTask* StealLocal(bool GetBackGroundTasks)
+		inline FTask* StealLocal(int32 PriorityIndex)
 		{
-			const int32 MaxPriority = GetBackGroundTasks ? int32(ETaskPriority::Count) : int32(ETaskPriority::ForegroundCount);
+			FTask* Item;
+			if (LocalQueues[PriorityIndex].Steal(Item))
+			{
+				return Item;
+			}
+
+			return nullptr;
+		}
+
+		inline FTask* StealLocal(bool bPermitBackgroundTasks)
+		{
+			const int32 MaxPriority = bPermitBackgroundTasks ? int32(ETaskPriority::Count) : int32(ETaskPriority::ForegroundCount);
 
 			for (int32 PriorityIndex = 0; PriorityIndex < MaxPriority; ++PriorityIndex)
 			{
-				FTask* Item;
-				if (LocalQueues[PriorityIndex].Steal(Item))
+				if (FTask* Item = StealLocal(PriorityIndex))
 				{
 					return Item;
 				}
@@ -253,36 +263,48 @@ public:
 			return nullptr;
 		}
 
+		// Check both the local and global queue for a specific priority
+		inline FTask* Dequeue(int32 PriorityIndex)
+		{
+			FTask* Item;
+			if (LocalQueues[PriorityIndex].Get(Item))
+			{
+				return Item;
+			}
+
+			Item = Registry->OverflowQueues[PriorityIndex].dequeue(DequeueHazards[PriorityIndex]);
+			if (Item)
+			{
+				return Item;
+			}
+			
+			return nullptr;
+		}
+
 		// Check both the local and global queue in priority order
-		inline FTask* Dequeue(bool GetBackGroundTasks)
+		inline FTask* Dequeue(bool bPermitBackgroundTasks)
 		{
-			const int32 MaxPriority = GetBackGroundTasks ? int32(ETaskPriority::Count)   : int32(ETaskPriority::ForegroundCount);
+			const int32 MaxPriority = bPermitBackgroundTasks ? int32(ETaskPriority::Count)   : int32(ETaskPriority::ForegroundCount);
 
 			for (int32 PriorityIndex = 0; PriorityIndex < MaxPriority; ++PriorityIndex)
 			{
-				FTask* Item;
-				if (LocalQueues[PriorityIndex].Get(Item))
-				{
-					return Item;
-				}
-
-				Item = Registry->OverflowQueues[PriorityIndex].dequeue(DequeueHazards[PriorityIndex]);
-				if (Item)
+				if (FTask* Item = Dequeue(PriorityIndex))
 				{
 					return Item;
 				}
 			}
+
 			return nullptr;
 		}
 
-		inline FTask* DequeueSteal(bool GetBackGroundTasks)
+		inline FTask* DequeueSteal(bool bPermitBackgroundTasks)
 		{
 			if (CachedRandomIndex == InvalidIndex)
 			{
 				CachedRandomIndex = Rand();
 			}
 
-			FTask* Result = Registry->StealItem(CachedRandomIndex, CachedPriorityIndex, GetBackGroundTasks);
+			FTask* Result = Registry->StealItem(CachedRandomIndex, CachedPriorityIndex, bPermitBackgroundTasks);
 			if (Result)
 			{
 				return Result;
