Fix compile error when WITH_FIXED_TIME_STEP_SUPPORT is defined to 1.
- bUseFixedTimeStep is a compile-time constant when the define is 1, meaning SetUseFixedTimeStep can't change its value.
- Wrapping the setter in #if WITH_FIXED_TIME_STEP_SUPPORT matches how the function was in FApp before 48290569 was submitted.

#Synced-CL 48803475

diff --git a/Engine/Source/Runtime/Core/Public/Misc/AppTime.h b/Engine/Source/Runtime/Core/Public/Misc/AppTime.h
--- a/Engine/Source/Runtime/Core/Public/Misc/AppTime.h
+++ b/Engine/Source/Runtime/Core/Public/Misc/AppTime.h
@@ -127,9 +127,11 @@ public:
 
 	static void SetUseFixedTimeStep(bool bVal)
 	{
+#if WITH_FIXED_TIME_STEP_SUPPORT
 		FAppTime Time = Get();
 		Time.bUseFixedTimeStep = bVal;
 		Set(Time);
+#endif
 	}
 
 	static void SetCurrentTime(double Seconds)
