Cloud Gaming - Fix version mismatch during early startup patch flow by adding a launch argument as a fallback mechanism

#tests Tested the fix on GFN using the debug seats
#rnx
#Synced-CL 48686588

diff --git a/Engine/Source/Runtime/Core/Private/GenericPlatform/GenericPlatformMisc.cpp b/Engine/Source/Runtime/Core/Private/GenericPlatform/GenericPlatformMisc.cpp
--- a/Engine/Source/Runtime/Core/Private/GenericPlatform/GenericPlatformMisc.cpp
+++ b/Engine/Source/Runtime/Core/Private/GenericPlatform/GenericPlatformMisc.cpp
@@ -1197,7 +1197,12 @@ const FString FGenericPlatformMisc::VersionCheckPlatformName()
 	FString OverrideVersionCheckPlatformName;
 	if (!GConfig->GetString(TEXT("PatchConfig"), TEXT("PlatformOverrideName"), OverrideVersionCheckPlatformName, GEngineIni))
 	{
-		return ANSI_TO_TCHAR(FPlatformProperties::IniPlatformName());
+		OverrideVersionCheckPlatformName = ANSI_TO_TCHAR(FPlatformProperties::IniPlatformName());
+		FString PlatformNameSuffix;
+		if (FParse::Value(FCommandLine::Get(), TEXT("PlatformSuffix="), PlatformNameSuffix))
+		{
+			OverrideVersionCheckPlatformName += PlatformNameSuffix;
+		}
 	}
 	return OverrideVersionCheckPlatformName;
 }
