Added some additional detection of usage post move temp UnifiedError.
#Synced-CL 48601431

diff --git a/Engine/Source/Runtime/Core/Private/Experimental/UnifiedError/UnifiedError.cpp b/Engine/Source/Runtime/Core/Private/Experimental/UnifiedError/UnifiedError.cpp
--- a/Engine/Source/Runtime/Core/Private/Experimental/UnifiedError/UnifiedError.cpp
+++ b/Engine/Source/Runtime/Core/Private/Experimental/UnifiedError/UnifiedError.cpp
@@ -16,6 +16,7 @@ namespace UE::UnifiedError
 
 void FError::PushErrorDetails(TRefCountPtr<IErrorDetails> InErrorDetails)
 {
+	ensureMsgf(ErrorDetails.IsValid(), TEXT("All FErrors should have been initialized with an FMandatoryErrorDetails"));
 	if (!InErrorDetails.IsValid())
 	{
 		return;
@@ -26,6 +27,10 @@ void FError::PushErrorDetails(TRefCountPtr<IErrorDetails> InErrorDetails)
 
 TRefCountPtr<const IErrorDetails> FError::GetInnerMostErrorDetails() const
 {
+	if (!ErrorDetails)
+	{
+		return ErrorDetails;
+	}
 	const IErrorDetails* Result = ErrorDetails;
 	while (true)
 	{
@@ -42,13 +47,23 @@ TRefCountPtr<const IErrorDetails> FError::GetInnerMostErrorDetails() const
 
 FText FError::GetFormatErrorText() const
 {
+	if (!ErrorDetails)
+	{
+		return FText::FromString(TEXT("UninitializedError1"));
+	}
+	
 	return ErrorDetails->GetErrorFormatString(*this);
 }
 
-const FManditoryErrorDetails* FError::GetManditoryErrorDetails() const
+void FError::Invalidate()
 {
+	// inverse the error code to help track usage after MoveTemp
+	ErrorCode = -ErrorCode;
+	ErrorDetails = nullptr;
+}
 
-
+const FManditoryErrorDetails* FError::GetManditoryErrorDetails() const
+{
 	TRefCountPtr<const IErrorDetails> CurrentIt = ErrorDetails;
 	while (CurrentIt != nullptr)
 	{
@@ -60,7 +75,7 @@ const FManditoryErrorDetails* FError::GetManditoryErrorDetails() const
 		CurrentIt = Next;
 	}
 
-	if (!ensure(CurrentIt.GetReference()))
+	if (!ensureMsgf(CurrentIt.GetReference() != nullptr, TEXT("Attempting to read error manditory error details from an invalid ferror, ModuleId:%d ErrorCode:%d"), ModuleId, ErrorCode))
 	{
 		return nullptr;
 	}
@@ -77,14 +92,20 @@ FAnsiString FError::GetModuleIdAndErrorCodeString() const
 FAnsiString FError::GetErrorCodeString() const
 {
 	const FManditoryErrorDetails* ManditoryErrorDetails = GetManditoryErrorDetails();
-	check(ManditoryErrorDetails);
+	if (!ManditoryErrorDetails)
+	{
+		return FAnsiString("Invalid"); // if you are seeing this it's possible this FError was moved and you are using it after the move, or it was constructed in an invalid way (memzero). 
+	}
 	return ManditoryErrorDetails->GetErrorCodeString(*this);
 }
 
 FAnsiString FError::GetModuleIdString() const
 {
 	const FManditoryErrorDetails* ManditoryErrorDetails = GetManditoryErrorDetails();
-	check(ManditoryErrorDetails);
+	if (!ManditoryErrorDetails)
+	{
+		return FAnsiString("Invalid"); // if you are seeing this it's possible this FError was moved and you are using it after the move, or it was constructed in an invalid way (memzero). 
+	}
 	return ManditoryErrorDetails->GetModuleIdString(*this);
 }
 
@@ -96,11 +117,20 @@ FAnsiString FError::GetModuleIdString() const
 
 FText FError::GetErrorMessage(bool bIncludeContext) const
 {
+	if (!ErrorDetails)
+	{
+		return FText::FromString(TEXT("UninitializedError2"));
+	}
+	
 	return GetErrorMessage(FText::Join(FText::FromString(": "), NSLOCTEXT("UnifiedError", "DefaultLogPrepend", "{Root/ModuleIdString}.{Root/ErrorCodeString}"), GetFormatErrorText()), bIncludeContext);
 }
 
 FText FError::GetErrorMessage(const FText& InFormatText, bool bIncludeContext) const 
 {
+	if (!ErrorDetails)
+	{
+		return FText::FromString(TEXT("UninitializedError3"));
+	}
 
 	EDetailFilter LogDetailFilter = EDetailFilter::IncludeInLogMessage;
 	if (bIncludeContext)
diff --git a/Engine/Source/Runtime/Core/Public/Experimental/UnifiedError/UnifiedError.h b/Engine/Source/Runtime/Core/Public/Experimental/UnifiedError/UnifiedError.h
--- a/Engine/Source/Runtime/Core/Public/Experimental/UnifiedError/UnifiedError.h
+++ b/Engine/Source/Runtime/Core/Public/Experimental/UnifiedError/UnifiedError.h
@@ -440,14 +440,12 @@ namespace UE::UnifiedError
 		FError(FError&& InError)
 		{
 			ErrorCode = InError.ErrorCode;
-			InError.ErrorCode = 0;
 			ModuleId = InError.ModuleId;
-			InError.ModuleId = 0;
 			if (InError.ErrorDetails)
 			{
-				ErrorDetails = InError.ErrorDetails;
-				InError.ErrorDetails = nullptr;
+				ErrorDetails = MoveTemp(InError.ErrorDetails);
 			}
+			InError.Invalidate();
 		}
 
 		FError(const FError& InError)
@@ -507,15 +505,10 @@ namespace UE::UnifiedError
 
 		UE_FORCEINLINE_HINT bool IsValid() const
 		{
-			return (ErrorCode != 0) && (ModuleId != 0);
+			return ErrorDetails != nullptr;
 		}
 
-		inline void Invalidate()
-		{
-			ErrorCode = 0;
-			ModuleId = 0;
-			ErrorDetails = nullptr;
-		}
+		CORE_API void Invalidate();
 	private:
 
 		friend CORE_API void SerializeForLog(FCbWriter& Writer, const FError& Error);
