Fix for FHLSLMaterialTranslator::CustomOutput that always returns INDEX_NONE, preventing uniform expressions from being properly tagged.

#rnx

#Synced-CL 48767541

diff --git a/Engine/Source/Runtime/Core/Private/UObject/FortniteMainBranchObjectVersion.cpp b/Engine/Source/Runtime/Core/Private/UObject/FortniteMainBranchObjectVersion.cpp
--- a/Engine/Source/Runtime/Core/Private/UObject/FortniteMainBranchObjectVersion.cpp
+++ b/Engine/Source/Runtime/Core/Private/UObject/FortniteMainBranchObjectVersion.cpp
@@ -10,10 +10,10 @@ TMap<FGuid, FGuid> FFortniteMainBranchObjectVersion::GetSystemGuids()
 
 	SystemGuids.Add(DevGuids.GLOBALSHADERMAP_DERIVEDDATA_VER, FGuid("EF39291B950C42C99FFE8D7AA98064BA"));
 	SystemGuids.Add(DevGuids.LANDSCAPE_MOBILE_COOK_VERSION, FGuid("32D02EF867C74B71A0D4E0FA41392732"));
-	SystemGuids.Add(DevGuids.MATERIALSHADERMAP_DERIVEDDATA_VER, FGuid("AF4566B7CCBB41649A079D009AACA4B2"));
+	SystemGuids.Add(DevGuids.MATERIALSHADERMAP_DERIVEDDATA_VER, FGuid("7DC13FE42B9645D8BFEE5B5CB16C78ED"));
 	SystemGuids.Add(DevGuids.NIAGARASHADERMAP_DERIVEDDATA_VER, FGuid("CC229D919A214E59B79D01AC2E5BF981"));
 	SystemGuids.Add(DevGuids.Niagara_LatestScriptCompileVersion, FGuid("FB15C2870A4E4AE8B5F58566F0E3825E"));
-	SystemGuids.Add(DevGuids.MaterialTranslationDDCVersion, FGuid("7B2702334BEB44B787B8C7C252B5C026"));
+	SystemGuids.Add(DevGuids.MaterialTranslationDDCVersion, FGuid("BD6B188514464989B75EC2C964F4E7FD"));
 
 	// These GUIDs are in AppendRestrictedFortniteMainStreamObjectVersionGuids as they should rarely be changed.
 	//SystemGuids.Add(DevGuids.NANITE_DERIVEDDATA_VER, FGuid("ACF98184978742A592537DA223A1E6BE"));
diff --git a/Engine/Source/Runtime/Engine/Private/Materials/HLSLMaterialTranslator.cpp b/Engine/Source/Runtime/Engine/Private/Materials/HLSLMaterialTranslator.cpp
--- a/Engine/Source/Runtime/Engine/Private/Materials/HLSLMaterialTranslator.cpp
+++ b/Engine/Source/Runtime/Engine/Private/Materials/HLSLMaterialTranslator.cpp
@@ -1035,9 +1035,9 @@ void FHLSLMaterialTranslator::CompileCustomOutputs(TArray<UMaterialExpressionCus
 							bCompilingPreviousFrame = PreviousFrameEvaluation != 0;
 							TArray<FShaderCodeChunk> CustomExpressionChunks;
 							AssignTempScope(CustomExpressionChunks);
-							int32 CodeChunk = CustomOutput->Compile(this, Index);
+							CustomOutput->Compile(this, Index);
 
-							if (CodeChunk != INDEX_NONE && ShaderFrequency == SF_Vertex)
+							if (ShaderFrequency == SF_Vertex)
 							{
 								for (FShaderCodeChunk& ShaderCodeChunk : CustomExpressionChunks)
 								{
