Add bool return type to FConfigCacheIni::Flush to know if the information was actually written to disk

#Synced-CL 48798601

diff --git a/Engine/Source/Runtime/Core/Private/Misc/ConfigCacheIni.cpp b/Engine/Source/Runtime/Core/Private/Misc/ConfigCacheIni.cpp
--- a/Engine/Source/Runtime/Core/Private/Misc/ConfigCacheIni.cpp
+++ b/Engine/Source/Runtime/Core/Private/Misc/ConfigCacheIni.cpp
@@ -4389,12 +4389,12 @@ TArray<FString> FConfigCacheIni::GetFilenames()
 	return Result;
 }
 
-
-
-void FConfigCacheIni::Flush(bool bRemoveFromCache, const FString& Filename )
+bool FConfigCacheIni::Flush(bool bRemoveFromCache, const FString& Filename)
 {
 	TRACE_CPUPROFILER_EVENT_SCOPE(FConfigCacheIni::Flush);
 
+	bool bSuccess = true;
+
 	// never Flush temporary cache objects
 	if (Type != EConfigCacheType::Temporary)
 	{
@@ -4406,7 +4406,7 @@ void FConfigCacheIni::Flush(bool bRemoveFromCache, const FString& Filename )
 				// flush single file
 				if (FConfigBranch* Branch = FindBranch(*Filename, Filename))
 				{
-					SaveBranch(*Branch);
+					bSuccess = SaveBranch(*Branch) && bSuccess;
 				}
 			}
 			else
@@ -4415,11 +4415,11 @@ void FConfigCacheIni::Flush(bool bRemoveFromCache, const FString& Filename )
 				// flush all files
 				for (TPair<FString, FConfigBranch*>& Pair : OtherFiles)
 				{
-					SaveBranch(*Pair.Value);
+					bSuccess = SaveBranch(*Pair.Value) && bSuccess;
 				}
 				for (FConfigBranch& Branch : KnownFiles.Branches)
 				{
-					SaveBranch(Branch);
+					bSuccess = SaveBranch(Branch) && bSuccess;
 				}
 			}
 		}
@@ -4431,10 +4431,8 @@ void FConfigCacheIni::Flush(bool bRemoveFromCache, const FString& Filename )
 		if (bAreFileOperationsDisabled)
 		{
 			UE_LOG(LogConfig, Warning, TEXT("Tried to flush the config cache and read it back in, but File Operations are disabled!!"));
-			return;
 		}
-
-		if (Filename.Len() != 0)
+		else if (Filename.Len() != 0)
 		{
 			Remove(Filename);
 		}
@@ -4449,6 +4447,8 @@ void FConfigCacheIni::Flush(bool bRemoveFromCache, const FString& Filename )
 			OtherFileNames.Empty();
 		}
 	}
+
+	return bSuccess;
 }
 
 /**
diff --git a/Engine/Source/Runtime/Core/Public/Misc/ConfigCacheIni.h b/Engine/Source/Runtime/Core/Public/Misc/ConfigCacheIni.h
--- a/Engine/Source/Runtime/Core/Public/Misc/ConfigCacheIni.h
+++ b/Engine/Source/Runtime/Core/Public/Misc/ConfigCacheIni.h
@@ -1367,8 +1367,7 @@ public:
 	
 	CORE_API TArray<FString> GetFilenames();
 
-
-	CORE_API void Flush(bool bRemoveFromCache, const FString& Filename=TEXT(""));
+	CORE_API bool Flush(bool bRemoveFromCache, const FString& Filename=TEXT(""));
 
 	CORE_API void LoadFile( const FString& InFilename, const FConfigFile* Fallback = NULL, const TCHAR* PlatformString = NULL );
 	CORE_API void SetFile( const FString& InFilename, const FConfigFile* NewConfigFile );
