Fixed comment on UE::CNotCVRefTo and updated it to use a C++20 type trait.

diff --git a/Engine/Source/Runtime/Core/Public/Concepts/NotCVRefTo.h b/Engine/Source/Runtime/Core/Public/Concepts/NotCVRefTo.h
--- a/Engine/Source/Runtime/Core/Public/Concepts/NotCVRefTo.h
+++ b/Engine/Source/Runtime/Core/Public/Concepts/NotCVRefTo.h
@@ -8,7 +8,7 @@
 namespace UE
 {
 	/**
-	 * Concept which describes convertibility from one type to another.
+	 * Concept which describes a type which is not a possibly-cv-qualified or reference to a given type.
 	 *
 	 * This is useful for templated constructors to prevent them being instantiated for the class type itself
 	 * so that it will call a move or copy constructor instead:
@@ -16,11 +16,11 @@ namespace UE
 	 * struct FSomeClass
 	 * {
 	 *     template <CNotCVRefTo<FSomeClass> ArgType>
-	 *     FSomeClass(ArgType&& Arg);
+	 *     FSomeClass(ArgType&& Arg)
 	 *     {
 	 *     }
 	 * };
 	 */
 	template <typename T, typename SelfType>
-	concept CNotCVRefTo = !UE::CSameAs<std::remove_cv_t<std::remove_reference_t<T>>, SelfType>;
+	concept CNotCVRefTo = !UE::CSameAs<std::remove_cvref_t<T>, SelfType>;
 }
