Reverted C++20 trait.

diff --git a/Engine/Source/Runtime/Core/Public/Concepts/NotCVRefTo.h b/Engine/Source/Runtime/Core/Public/Concepts/NotCVRefTo.h
--- a/Engine/Source/Runtime/Core/Public/Concepts/NotCVRefTo.h
+++ b/Engine/Source/Runtime/Core/Public/Concepts/NotCVRefTo.h
@@ -22,5 +22,5 @@ namespace UE
 	 * };
 	 */
 	template <typename T, typename SelfType>
-	concept CNotCVRefTo = !UE::CSameAs<std::remove_cvref_t<T>, SelfType>;
+	concept CNotCVRefTo = !UE::CSameAs<std::remove_cv_t<std::remove_reference_t<T>>, SelfType>;
 }
