[Core] Break Core dependency on CoreUObject

CoreUObject depends on Core.
AutoRTFMUE.cpp cannot reference Verse::RollbackAutoRTFMWrite as this adds a Core dependency on CoreUObject, forming a cyclic dependency.

Pass the Verse::RollbackAutoRTFMWrite function pointer down to AutoRTFM::InitializeForUE() from Runtime, which depends on both Core and CoreUObject.

#Synced-CL 48856386

diff --git a/Engine/Source/Runtime/Core/Private/AutoRTFM/AutoRTFMUE.cpp b/Engine/Source/Runtime/Core/Private/AutoRTFM/AutoRTFMUE.cpp
--- a/Engine/Source/Runtime/Core/Private/AutoRTFM/AutoRTFMUE.cpp
+++ b/Engine/Source/Runtime/Core/Private/AutoRTFM/AutoRTFMUE.cpp
@@ -13,7 +13,6 @@
 #include "HAL/UnrealMemory.h"
 #include "Logging/LogVerbosity.h"
 #include "Misc/AssertionMacros.h"
-#include "VerseVM/VVMWriteBarrier.h"
 
 #include <cstdarg>
 #include <cstdio>
@@ -390,7 +389,7 @@ static void AutoRTFMFree(void* Pointer)
 	FMemory::Free(Pointer);
 }
 
-CORE_API void AutoRTFM::InitializeForUE()
+CORE_API void AutoRTFM::InitializeForUE(RollbackWriteFn RollbackWrite /* = nullptr */)
 {
 	AutoRTFM::ForTheRuntime::FExternAPI ExternAPI
 	{
@@ -408,9 +407,7 @@ CORE_API void AutoRTFM::InitializeForUE()
 		.OnMemoryValidationThrottlingChanged = OnAutoRTFMMemoryValidationThrottlingChanged,
 		.OnMemoryValidationWriteIdentifierChanged = OnAutoRTFMMemoryValidationWriteIdentifierChanged,
 		.OnMemoryValidationStatisticsChanged = OnAutoRTFMMemoryValidationStatisticsChanged,
-#if WITH_VERSE_VM
-		.RollbackWrite = Verse::RollbackAutoRTFMWrite,
-#endif
+		.RollbackWrite = RollbackWrite,
 	};
 	
 	AutoRTFM::ForTheRuntime::Initialize(ExternAPI);
diff --git a/Engine/Source/Runtime/Core/Public/AutoRTFM/AutoRTFMUE.h b/Engine/Source/Runtime/Core/Public/AutoRTFM/AutoRTFMUE.h
--- a/Engine/Source/Runtime/Core/Public/AutoRTFM/AutoRTFMUE.h
+++ b/Engine/Source/Runtime/Core/Public/AutoRTFM/AutoRTFMUE.h
@@ -13,14 +13,18 @@ UE_AUTORTFM_API DECLARE_LOG_CATEGORY_EXTERN(LogAutoRTFM, Display, All)
 namespace AutoRTFM
 {
 
+using RollbackWriteFn = void (*)(void* Address, const void* Value, size_t Size, autortfm_write_flags Flags);
+
 #if UE_AUTORTFM
 
 // Initializes AutoRTFM for use by the Unreal Engine.
-UE_AUTORTFM_API void InitializeForUE();
+// RollbackWrite is an optional function pointer used to rollback writes written
+// with the AutoRTFM::EWriteFlags::CustomRollback flag.
+UE_AUTORTFM_API void InitializeForUE(RollbackWriteFn RollbackWrite = nullptr);
 
 #else
 
-inline void InitializeForUE() {}
+inline void InitializeForUE(RollbackWriteFn) {}
 
 #endif
 
diff --git a/Engine/Source/Runtime/Launch/Private/LaunchEngineLoop.cpp b/Engine/Source/Runtime/Launch/Private/LaunchEngineLoop.cpp
--- a/Engine/Source/Runtime/Launch/Private/LaunchEngineLoop.cpp
+++ b/Engine/Source/Runtime/Launch/Private/LaunchEngineLoop.cpp
@@ -112,6 +112,7 @@
 
 	#if WITH_VERSE_VM || defined(__INTELLISENSE__)
 		#include "VerseVM/VVMVerse.h"
+		#include "VerseVM/VVMWriteBarrier.h"
 	#else
 		#include "VerseVM/VVMExecutionContext.h"
 	#endif
@@ -2872,7 +2873,11 @@ int32 FEngineLoop::PreInitPreStartupScreen(const TCHAR* CmdLine)
 
 	{
 		SCOPED_BOOT_TIMING("AutoRTFM");
-		AutoRTFM::InitializeForUE();
+#if WITH_VERSE_VM
+		AutoRTFM::InitializeForUE(Verse::RollbackAutoRTFMWrite);
+#else
+		AutoRTFM::InitializeForUE(nullptr);
+#endif
 	}
 
 #if WITH_ENGINE
