Correct error message in UTEST_GREATER for floats and doubles. Was incorrectly logging that the value should be 'less than' the expected value, when the logic was checking for greater.

#rnx
#Synced-CL 48866341

diff --git a/Engine/Source/Runtime/Core/Private/Misc/AutomationTest.cpp b/Engine/Source/Runtime/Core/Private/Misc/AutomationTest.cpp
--- a/Engine/Source/Runtime/Core/Private/Misc/AutomationTest.cpp
+++ b/Engine/Source/Runtime/Core/Private/Misc/AutomationTest.cpp
@@ -2572,14 +2572,14 @@ bool FAutomationTestBase::TestGreaterThan(const TCHAR* What, const float Actual,
 {
 	if (FMath::IsNearlyEqual(Actual, Expected, Tolerance))
 	{
-		AddError(FString::Printf(TEXT("Expected '%s' to be less than %f, but it was %f and within equality tolerance %f."), What, Expected, Actual, Tolerance), 1);
+		AddError(FString::Printf(TEXT("Expected '%s' to be greater than %f, but it was %f and within equality tolerance %f."), What, Expected, Actual, Tolerance), 1);
 		return false;
 	}
 	if (Actual > Expected)
 	{
 		return true;
 	}
-	AddError(FString::Printf(TEXT("Expected '%s' to be less than %f, but it was %f and outside equality tolerance %f."), What, Expected, Actual, Tolerance), 1);
+	AddError(FString::Printf(TEXT("Expected '%s' to be greater than %f, but it was %f and outside equality tolerance %f."), What, Expected, Actual, Tolerance), 1);
 	return false;
 }
 
@@ -2587,14 +2587,14 @@ bool FAutomationTestBase::TestGreaterThan(const TCHAR* What, const double Actual
 {
 	if (FMath::IsNearlyEqual(Actual, Expected, Tolerance))
 	{
-		AddError(FString::Printf(TEXT("Expected '%s' to be less than %f, but it was %f and within equality tolerance %f."), What, Expected, Actual, Tolerance), 1);
+		AddError(FString::Printf(TEXT("Expected '%s' to be greater than %f, but it was %f and within equality tolerance %f."), What, Expected, Actual, Tolerance), 1);
 		return false;
 	}
 	if (Actual > Expected)
 	{
 		return true;
 	}
-	AddError(FString::Printf(TEXT("Expected '%s' to be less than %f, but it was %f and outside equality tolerance %f."), What, Expected, Actual, Tolerance), 1);
+	AddError(FString::Printf(TEXT("Expected '%s' to be greater than %f, but it was %f and outside equality tolerance %f."), What, Expected, Actual, Tolerance), 1);
 	return false;
 }
 
