[Core] Add explicit to iterator constructors for map/set
* Rework map iterators to share common code between regular and key iterators

#tests Compiled with compact and sparse sets, ran PFs
#Synced-CL 48831686

diff --git a/Engine/Plugins/Experimental/SerializationUtils/Source/JsonSerialization/Private/Formatters/JsonArchiveInputFormatterEx.cpp b/Engine/Plugins/Experimental/SerializationUtils/Source/JsonSerialization/Private/Formatters/JsonArchiveInputFormatterEx.cpp
--- a/Engine/Plugins/Experimental/SerializationUtils/Source/JsonSerialization/Private/Formatters/JsonArchiveInputFormatterEx.cpp
+++ b/Engine/Plugins/Experimental/SerializationUtils/Source/JsonSerialization/Private/Formatters/JsonArchiveInputFormatterEx.cpp
@@ -200,7 +200,7 @@ void FJsonArchiveInputFormatterEx::EnterMap(int32& NumElements)
 {
 	TSharedPtr<FJsonValue>& Value = ValueStack.Top();
 	TSharedPtr<FJsonObject> Object = Value->AsObject();
-	MapIteratorStack.Add(Object->Values);
+	MapIteratorStack.Add(Object->Values.CreateIterator());
 	NumElements = Object->Values.Num();
 }
 
diff --git a/Engine/Source/Editor/Sequencer/Private/Sequencer.cpp b/Engine/Source/Editor/Sequencer/Private/Sequencer.cpp
--- a/Engine/Source/Editor/Sequencer/Private/Sequencer.cpp
+++ b/Engine/Source/Editor/Sequencer/Private/Sequencer.cpp
@@ -2452,9 +2452,8 @@ void FSequencer::TransformSelection(FFrameTime InDeltaTime, float InScale)
 			UMovieScene* FocusedMovieScene = GetFocusedMovieSceneSequence()->GetMovieScene();
 			FocusedMovieScene->Modify();
 			
-			for (TSet<int32>::TConstIterator It = SelectedMarkedFrames.GetSelected(); It; ++It)
+			for (const int32 MarkIndex : SelectedMarkedFrames)
 			{
-				const int32 MarkIndex = *It;
 				FFrameNumber FrameNumber = FocusedMovieScene->GetMarkedFrames()[MarkIndex].FrameNumber;
 				FrameNumber = (OriginTime + InDeltaTime + (FrameNumber - OriginTime) * InScale).FloorToFrame();
 				FocusedMovieScene->SetMarkedFrame(MarkIndex, FrameNumber);
diff --git a/Engine/Source/Runtime/Core/Public/Containers/CompactSet.h.inl b/Engine/Source/Runtime/Core/Public/Containers/CompactSet.h.inl
--- a/Engine/Source/Runtime/Core/Public/Containers/CompactSet.h.inl
+++ b/Engine/Source/Runtime/Core/Public/Containers/CompactSet.h.inl
@@ -86,12 +86,12 @@ private:
 	public:
 		using ElementItType = std::conditional_t<bConst, const ElementType, ElementType>;
 
-		[[nodiscard]] UE_FORCEINLINE_HINT TBaseIterator(SetType& InSet)
+		[[nodiscard]] explicit UE_FORCEINLINE_HINT TBaseIterator(SetType& InSet)
 		: TBaseIterator(InSet, 0)
 		{
 		}
 
-		[[nodiscard]] inline TBaseIterator(SetType& InSet, SizeType StartIndex)
+		[[nodiscard]] explicit inline TBaseIterator(SetType& InSet, SizeType StartIndex)
 		: Set(InSet)
 		, Index(StartIndex)
 #if DO_CHECK
@@ -171,7 +171,7 @@ private:
 		>;
 
 		/** Initialization constructor. */
-		[[nodiscard]] inline TBaseKeyIterator(SetType& InSet, KeyArgumentType InKey)
+		[[nodiscard]] explicit inline TBaseKeyIterator(SetType& InSet, KeyArgumentType InKey)
 		: Set  (InSet)
 		, Key  (InKey)
 		, HashTable(InSet.Num() ? Set.GetConstHashTableView() : FConstCompactHashTableView())
diff --git a/Engine/Source/Runtime/Core/Public/Containers/Map.h.inl b/Engine/Source/Runtime/Core/Public/Containers/Map.h.inl
--- a/Engine/Source/Runtime/Core/Public/Containers/Map.h.inl
+++ b/Engine/Source/Runtime/Core/Public/Containers/Map.h.inl
@@ -764,27 +764,31 @@ public:
 
 protected:
 	/** The base of TMAPBASE iterators. */
-	template<bool bConst>
-	class TBaseIterator
+	template<bool bConst, bool bKeyIterator>
+	class TCommonBaseIterator
 	{
-	public:
-		typedef std::conditional_t<
-			bConst,
-		typename ElementSetType::TConstIterator,
-		typename ElementSetType::TIterator
-			> PairItType;
 	private:
-		typedef std::conditional_t<bConst, const KeyType, KeyType> ItKeyType;
-		typedef std::conditional_t<bConst, const ValueType, ValueType> ItValueType;
-		typedef std::conditional_t<bConst, const typename ElementSetType::ElementType, typename ElementSetType::ElementType> PairType;
+		using BaseIteratorType = std::conditional_t<bConst, typename ElementSetType::TConstIterator, typename ElementSetType::TIterator>;
+		using KeyIteratorType = std::conditional_t<bConst, typename ElementSetType::TConstKeyIterator, typename ElementSetType::TKeyIterator>;
 
-	public:
-		[[nodiscard]] UE_FORCEINLINE_HINT TBaseIterator(const PairItType& InElementIt)
+	protected:
+		using MapType = std::conditional_t<bConst, const TMAPBASE, TMAPBASE>;
+		using PairType = std::conditional_t<bConst, const typename ElementSetType::ElementType, typename ElementSetType::ElementType>;
+
+		using PairItType = std::conditional_t<bKeyIterator, KeyIteratorType, BaseIteratorType>;
+
+		using ItKeyType = std::conditional_t<bConst, const KeyType, KeyType>;
+		using ItValueType = std::conditional_t<bConst, const ValueType, ValueType>;
+
+	protected:
+		[[nodiscard]] explicit UE_FORCEINLINE_HINT TCommonBaseIterator(const PairItType& InElementIt)
 		: PairIt(InElementIt)
 		{
 		}
 
-		UE_FORCEINLINE_HINT TBaseIterator& operator++()
+	public:
+
+		UE_FORCEINLINE_HINT TCommonBaseIterator& operator++()
 		{
 			++PairIt;
 			return *this;
@@ -801,11 +805,11 @@ protected:
 			return !(bool)*this;
 		}
 
-		[[nodiscard]] UE_FORCEINLINE_HINT bool operator==(const TBaseIterator& Rhs) const
+		[[nodiscard]] UE_FORCEINLINE_HINT bool operator==(const TCommonBaseIterator& Rhs) const
 		{
 			return PairIt == Rhs.PairIt;
 		}
-		[[nodiscard]] UE_FORCEINLINE_HINT bool operator!=(const TBaseIterator& Rhs) const
+		[[nodiscard]] UE_FORCEINLINE_HINT bool operator!=(const TCommonBaseIterator& Rhs) const
 		{
 			return PairIt != Rhs.PairIt;
 		}
@@ -836,64 +840,41 @@ protected:
 	protected:
 		PairItType PairIt;
 	};
-
+	
 	/** The base type of iterators that iterate over the values associated with a specified key. */
 	template<bool bConst>
-	class TBaseKeyIterator
+	class TBaseIterator : public TCommonBaseIterator<bConst, false>
 	{
 	private:
-		typedef std::conditional_t<bConst, typename ElementSetType::TConstKeyIterator, typename ElementSetType::TKeyIterator> SetItType;
-		typedef std::conditional_t<bConst, const KeyType, KeyType> ItKeyType;
-		typedef std::conditional_t<bConst, const ValueType, ValueType> ItValueType;
+		using Super = TCommonBaseIterator<bConst, false>;
+		using MapType = typename Super::MapType;
+		using PairItType = typename Super::PairItType;
 
 	public:
 		/** Initialization constructor. */
-		[[nodiscard]] UE_FORCEINLINE_HINT TBaseKeyIterator(const SetItType& InSetIt)
-		: SetIt(InSetIt)
+		[[nodiscard]] explicit UE_FORCEINLINE_HINT TBaseIterator(MapType& Map)
+		: Super(PairItType(Map.Pairs))
 		{
 		}
+	};
 
-		UE_FORCEINLINE_HINT TBaseKeyIterator& operator++()
-		{
-			++SetIt;
-			return *this;
-		}
-
-		/** conversion to "bool" returning true if the iterator is valid. */
-		[[nodiscard]] UE_FORCEINLINE_HINT explicit operator bool() const
-		{
-			return !!SetIt; 
-		}
-		/** inverse of the "bool" operator */
-		[[nodiscard]] UE_FORCEINLINE_HINT bool operator !() const 
-		{
-			return !(bool)*this;
-		}
+	/** The base type of iterators that iterate over the values associated with a specified key. */
+	template<bool bConst>
+	class TBaseKeyIterator : public TCommonBaseIterator<bConst, true>
+	{
+	private:
+		using Super = TCommonBaseIterator<bConst, true>;
+		using MapType = typename Super::MapType;
+		using PairItType = typename Super::PairItType;
 
-		[[nodiscard]] UE_FORCEINLINE_HINT FSetElementId GetId() const
-		{
-			return SetIt.GetId();
-		}
-		[[nodiscard]] UE_FORCEINLINE_HINT ItKeyType& Key() const
-		{
-			return SetIt->Key;
-		}
-		[[nodiscard]] UE_FORCEINLINE_HINT ItValueType& Value() const
-		{
-			return SetIt->Value;
-		}
+	public:
+		using KeyArgumentType = typename PairItType::KeyArgumentType;
 
-		[[nodiscard]] UE_FORCEINLINE_HINT decltype(auto) operator*() const
-		{
-			return SetIt.operator*();
-		}
-		[[nodiscard]] UE_FORCEINLINE_HINT decltype(auto) operator->() const
+		/** Initialization constructor. */
+		[[nodiscard]] explicit UE_FORCEINLINE_HINT TBaseKeyIterator(MapType& Map, KeyArgumentType InKey)
+		: Super(PairItType(Map.Pairs, InKey))
 		{
-			return SetIt.operator->();
 		}
-
-	protected:
-		SetItType SetIt;
 	};
 
 	/** A set of the key-value pairs in the map. */
@@ -918,11 +899,13 @@ public:
 	/** Map iterator. */
 	class TIterator : public TBaseIterator<false>
 	{
-	public:
+	private:
+		using Super = TBaseIterator<false>;
 
+	public:
 		/** Initialization constructor. */
-		[[nodiscard]] inline TIterator(TMAPBASE& InMap)
-		: TBaseIterator<false>(InMap.Pairs.CreateIterator())
+		[[nodiscard]] explicit inline TIterator(TMAPBASE& InMap)
+		: Super(InMap)
 		, Map(InMap)
 		, bElementsHaveBeenRemoved(false)
 		, bRequiresRehashOnRemoval(false)
@@ -930,8 +913,8 @@ public:
 		}
 
 		UE_DEPRECATED(5.8, "TSet::Relax has been deprecated. Memory updates post iteration should be handled manually as necessary.")
-		[[nodiscard]] inline TIterator(TMAPBASE& InMap, bool bInRequiresRehashOnRemoval /* = false */)
-		: TBaseIterator<false>(InMap.Pairs.CreateIterator())
+		[[nodiscard]] explicit inline TIterator(TMAPBASE& InMap, bool bInRequiresRehashOnRemoval /* = false */)
+		: Super(InMap.Pairs.CreateIterator())
 		, Map(InMap)
 		, bElementsHaveBeenRemoved(false)
 		, bRequiresRehashOnRemoval(bInRequiresRehashOnRemoval)
@@ -958,7 +941,7 @@ public:
 		 */
 		inline void RemoveCurrent()
 		{
-			TBaseIterator<false>::PairIt.RemoveCurrent();
+			this->PairIt.RemoveCurrent();
 			bElementsHaveBeenRemoved = true;
 		}
 
@@ -968,57 +951,29 @@ public:
 		bool      bRequiresRehashOnRemoval;
 	};
 
-	/** Const map iterator. */
-	class TConstIterator : public TBaseIterator<true>
-	{
-	public:
-		[[nodiscard]] UE_FORCEINLINE_HINT TConstIterator(const TMAPBASE& InMap)
-		: TBaseIterator<true>(InMap.Pairs.CreateConstIterator())
-		{
-		}
-	};
-
-	using TRangedForIterator = typename ElementSetType::TRangedForIterator;
-	using TRangedForConstIterator = typename ElementSetType::TRangedForConstIterator;
-
-	/** Iterates over values associated with a specified key in a const map. */
-	class TConstKeyIterator : public TBaseKeyIterator<true>
-	{
-	private:
-		using Super        = TBaseKeyIterator<true>;
-		using IteratorType = typename ElementSetType::TConstKeyIterator;
-
-	public:
-		using KeyArgumentType = typename IteratorType::KeyArgumentType;
-
-		[[nodiscard]] UE_FORCEINLINE_HINT TConstKeyIterator(const TMAPBASE& InMap, KeyArgumentType InKey)
-		: Super(IteratorType(InMap.Pairs, InKey))
-		{
-		}
-	};
-
 	/** Iterates over values associated with a specified key in a map. */
 	class TKeyIterator : public TBaseKeyIterator<false>
 	{
 	private:
-		using Super        = TBaseKeyIterator<false>;
-		using IteratorType = typename ElementSetType::TKeyIterator;
+		using Super = TBaseKeyIterator<false>;
 
 	public:
-		using KeyArgumentType = typename IteratorType::KeyArgumentType;
-
-		[[nodiscard]] UE_FORCEINLINE_HINT TKeyIterator(TMAPBASE& InMap, KeyArgumentType InKey)
-		: Super(IteratorType(InMap.Pairs, InKey))
-		{
-		}
+		using Super::Super;
 
 		/** Removes the current key-value pair from the map. */
 		UE_FORCEINLINE_HINT void RemoveCurrent()
 		{
-			TBaseKeyIterator<false>::SetIt.RemoveCurrent();
+			this->PairIt.RemoveCurrent();
 		}
 	};
 
+	/** Const map iterator. */
+	using TConstIterator = TBaseIterator<true>;
+	using TConstKeyIterator = TBaseKeyIterator<true>;
+
+	using TRangedForIterator = typename ElementSetType::TRangedForIterator;
+	using TRangedForConstIterator = typename ElementSetType::TRangedForConstIterator;
+
 	/** Creates an iterator over all the pairs in this map */
 	[[nodiscard]] UE_FORCEINLINE_HINT TIterator CreateIterator()
 	{
diff --git a/Engine/Source/Runtime/Core/Public/Containers/SparseSet.h.inl b/Engine/Source/Runtime/Core/Public/Containers/SparseSet.h.inl
--- a/Engine/Source/Runtime/Core/Public/Containers/SparseSet.h.inl
+++ b/Engine/Source/Runtime/Core/Public/Containers/SparseSet.h.inl
@@ -1441,7 +1441,7 @@ private:
 			std::conditional_t<bRangedFor, typename ElementArrayType::TRangedForIterator,      typename ElementArrayType::TIterator     >
 		> ElementItType;
 
-		[[nodiscard]] UE_FORCEINLINE_HINT TBaseIterator(const ElementItType& InElementIt)
+		[[nodiscard]] UE_FORCEINLINE_HINT explicit TBaseIterator(const ElementItType& InElementIt)
 			: ElementIt(InElementIt)
 		{
 		}
@@ -1502,7 +1502,7 @@ private:
 			>;
 
 		/** Initialization constructor. */
-		[[nodiscard]] inline TBaseKeyIterator(SetType& InSet, KeyArgumentType InKey)
+		[[nodiscard]] inline explicit TBaseKeyIterator(SetType& InSet, KeyArgumentType InKey)
 			: Set  (InSet)
 			, Key  (InKey) //-V1041
 			, Index(INDEX_NONE)
@@ -1578,7 +1578,7 @@ public:
 		friend class UE_TSPARSE_SET;
 
 	public:
-		[[nodiscard]] UE_FORCEINLINE_HINT TConstIterator(const UE_TSPARSE_SET& InSet)
+		[[nodiscard]] UE_FORCEINLINE_HINT explicit TConstIterator(const UE_TSPARSE_SET& InSet)
 			: TBaseIterator<true>(InSet.Elements.begin())
 		{
 		}
@@ -1590,7 +1590,7 @@ public:
 		friend class UE_TSPARSE_SET;
 
 	public:
-		[[nodiscard]] inline TIterator(UE_TSPARSE_SET& InSet)
+		[[nodiscard]] inline explicit TIterator(UE_TSPARSE_SET& InSet)
 			: TBaseIterator<false>(InSet.Elements.begin())
 			, Set                 (InSet)
 		{
@@ -1618,7 +1618,7 @@ public:
 	public:
 		using KeyArgumentType = typename Super::KeyArgumentType;
 
-		[[nodiscard]] UE_FORCEINLINE_HINT TConstKeyIterator(const UE_TSPARSE_SET& InSet, KeyArgumentType InKey)
+		[[nodiscard]] UE_FORCEINLINE_HINT explicit TConstKeyIterator(const UE_TSPARSE_SET& InSet, KeyArgumentType InKey)
 			: Super(InSet, InKey)
 		{
 		}
@@ -1633,7 +1633,7 @@ public:
 	public:
 		using KeyArgumentType = typename Super::KeyArgumentType;
 
-		[[nodiscard]] UE_FORCEINLINE_HINT TKeyIterator(UE_TSPARSE_SET& InSet, KeyArgumentType InKey)
+		[[nodiscard]] UE_FORCEINLINE_HINT explicit TKeyIterator(UE_TSPARSE_SET& InSet, KeyArgumentType InKey)
 			: Super(InSet, InKey)
 		{
 		}
diff --git a/Engine/Source/Runtime/CoreUObject/Private/Serialization/Formatters/JsonArchiveInputFormatter.cpp b/Engine/Source/Runtime/CoreUObject/Private/Serialization/Formatters/JsonArchiveInputFormatter.cpp
--- a/Engine/Source/Runtime/CoreUObject/Private/Serialization/Formatters/JsonArchiveInputFormatter.cpp
+++ b/Engine/Source/Runtime/CoreUObject/Private/Serialization/Formatters/JsonArchiveInputFormatter.cpp
@@ -165,7 +165,7 @@ void FJsonArchiveInputFormatter::EnterMap(int32& NumElements)
 {
 	TSharedPtr<FJsonValue>& Value = ValueStack.Top();
 	TSharedPtr<FJsonObject> Object = Value->AsObject();
-	MapIteratorStack.Add(Object->Values);
+	MapIteratorStack.Add(Object->Values.CreateIterator());
 	NumElements = Object->Values.Num();
 }
 
