[AutoRTFM] Allow shared pointers, shared references and weak pointers to be returned from AutoRTFM::Open.

#Synced-CL 48800164

diff --git a/Engine/Source/Programs/AutoRTFMTests/Private/SharedPointerTests.cpp b/Engine/Source/Programs/AutoRTFMTests/Private/SharedPointerTests.cpp
--- a/Engine/Source/Programs/AutoRTFMTests/Private/SharedPointerTests.cpp
+++ b/Engine/Source/Programs/AutoRTFMTests/Private/SharedPointerTests.cpp
@@ -230,3 +230,34 @@ TEST_CASE("SharedPointer.OnAbortCapturesWeakPtr")
 		AutoRTFM::AbortTransaction();
 	});
 }
+
+TEST_CASE("SharedPointer.ReturnFromOpen")
+{
+	TSharedPtr<int> Shared = AutoRTFM::Open([]
+	{
+		return TSharedPtr<int>{new int(42)};
+	});
+
+	REQUIRE(*Shared == 42);
+}
+
+TEST_CASE("SharedReference.ReturnFromOpen")
+{
+	TSharedRef<int> Shared = AutoRTFM::Open([]
+	{
+		return MakeShared<int>(42);
+	});
+
+	REQUIRE(*Shared == 42);
+}
+
+TEST_CASE("WeakPointer.ReturnFromOpen")
+{
+	TSharedPtr<int> Shared = TSharedPtr<int>{new int(42)};
+	TWeakPtr<int> Weak = AutoRTFM::Open([&]
+	{
+		return TWeakPtr<int>{Shared};
+	});
+
+	REQUIRE(*Weak.Pin() == 42);
+}
diff --git a/Engine/Source/Runtime/Core/Public/Templates/SharedPointer.h b/Engine/Source/Runtime/Core/Public/Templates/SharedPointer.h
--- a/Engine/Source/Runtime/Core/Public/Templates/SharedPointer.h
+++ b/Engine/Source/Runtime/Core/Public/Templates/SharedPointer.h
@@ -3,6 +3,7 @@
 #pragma once
 
 #include "CoreTypes.h"
+#include "AutoRTFM/ReturnFromOpenMode.h"
 #include "Misc/IntrusiveUnsetOptionalState.h"
 #include "Templates/PointerIsConvertibleFromTo.h"
 #include "Misc/AssertionMacros.h"
@@ -446,6 +447,9 @@ public:
 		return *this;
 	}
 
+	/** TSharedRef can be returned from AutoRTFM::Open() via a copy-construction in the closed. */
+	static constexpr AutoRTFM::EReturnFromOpenMode AutoRTFMReturnFromOpenMode = AutoRTFM::EReturnFromOpenMode::CopyConstructInClosed;
+
 	/**
 	 * Converts a shared reference to a shared pointer.
 	 *
@@ -1019,6 +1023,9 @@ public:
 		return *this;
 	}
 
+	/** TSharedPtr can be returned from AutoRTFM::Open() via a copy-construction in the closed. */
+	static constexpr AutoRTFM::EReturnFromOpenMode AutoRTFMReturnFromOpenMode = AutoRTFM::EReturnFromOpenMode::CopyConstructInClosed;
+
 	/**
 	 * Converts a shared pointer to a shared reference.  The pointer *must* be valid or an assertion will trigger.
 	 *
@@ -1501,6 +1508,9 @@ public:
 		return *this;
 	}
 
+	/** TWeakPtr can be returned from AutoRTFM::Open() via a copy-construction in the closed. */
+	static constexpr AutoRTFM::EReturnFromOpenMode AutoRTFMReturnFromOpenMode = AutoRTFM::EReturnFromOpenMode::CopyConstructInClosed;
+
 	/**
 	 * Converts this weak pointer to a shared pointer that you can use to access the object (if it
 	 * hasn't expired yet.)  Remember, if there are no more shared references to the object, the
