UnixPlatformTime: Prefer CLOCK_MONOTONIC

Consolidates time functions on a single clock source, preferably CLOCK_MONOTONIC, rather than CLOCK_MONOTONIC_RAW.

Reduces the duration of clock benchmarking at startup.

[RN] Linux: Timing functions now prefer the CLOCK_MONOTONIC reference clock over CLOCK_MONOTONIC_RAW.

diff --git a/Engine/Source/Runtime/Core/Private/Unix/UnixPlatformTime.cpp b/Engine/Source/Runtime/Core/Private/Unix/UnixPlatformTime.cpp
--- a/Engine/Source/Runtime/Core/Private/Unix/UnixPlatformTime.cpp
+++ b/Engine/Source/Runtime/Core/Private/Unix/UnixPlatformTime.cpp
@@ -10,10 +10,10 @@
 #include "GenericPlatform/GenericPlatformTime.h"
 #include "CoreGlobals.h"
 #include "Containers/Ticker.h"
+#include <inttypes.h>
 #include <sys/resource.h>
 
-int FUnixTime::ClockSource = -1;
-char FUnixTime::CalibrationLog[4096] = {0};
+clockid_t FUnixTime::ClockSource = -1;
 
 namespace FUnixTimeInternal
 {
@@ -76,6 +76,11 @@ namespace FUnixTimeInternal
 
 	/** Process lifetime Involuntary Context Switch count */
 	static uint64 InvoluntaryContextSwitchCount = 0;
+
+	static char CalibrationLog[4096] = "";
+
+	// Benchmark period is in terms of the clock being benchmarked, so results are not exactly comparable.
+	static constexpr uint64 kBenchmarkPeriodNanoSec = 1'000'000'000ULL / 100;	// 0.01s
 }
 
 double FUnixTime::InitTiming()
@@ -96,13 +101,8 @@ FCPUTime FUnixTime::GetCPUTime()
 	// 250ms minimum delay between checks to minimize overhead (and also match Windows version)
 	constexpr double MinDelayBetweenChecksNanoSec = 250 * 1e6;
 	
-	if (UNLIKELY(ClockSource < 0))
-	{
-		ClockSource = FUnixTime::CalibrateAndSelectClock();
-	}
-
 	struct timespec ts;
-	if (0 == clock_gettime(ClockSource, &ts))
+	if (0 == clock_gettime(GetClockSource(), &ts))
 	{
 		const double CurrentTimeNanoSec = static_cast<double>(FUnixTimeInternal::TimeSpecToNanoSec(ts));
 
@@ -233,85 +233,93 @@ double FUnixTime::GetLastIntervalThreadCPUTimeInSeconds()
 	return static_cast<double>(FUnixTimeInternal::CurrentThreadCPUStats.LastIntervalThreadTimeNS) / 1e9;
 }
 
-uint64 FUnixTime::CallsPerSecondBenchmark(clockid_t BenchClockId, const char * BenchClockIdName)
+// Benchmark a clock.
+// Returns the number of times this clock can be called per second, or zero if the clock is totally unusable.
+[[nodiscard]]
+static uint64 BenchmarkClock(const clockid_t Id, const char* Name)
 {
-	char Buffer[256];
-	const uint64 kBenchmarkPeriodMicroSec = 1000000000ULL / 10;	// 0.1s
+	using namespace FUnixTimeInternal;
+
+	char Buffer[256] = "";
 
-	// clock_getres() can fail when running on Windows Subsystem for Linux (but the clock can still be supported).
 	struct timespec ts;
-	if (clock_gettime(BenchClockId, &ts) == -1)
+	if (clock_gettime(Id, &ts) == -1)
 	{
-		FCStringAnsi::Snprintf(Buffer, sizeof(Buffer), "Clock_id %d (%s) is not supported on this system, clock_gettime() fails.\n", BenchClockId, BenchClockIdName);
+		FCStringAnsi::Snprintf(Buffer, sizeof(Buffer), "Clock_id %d (%s) is not supported on this system, clock_gettime() fails.\n", Id, Name);
 		FCStringAnsi::StrncatTruncateDest(CalibrationLog, sizeof(CalibrationLog), Buffer);
-		return 0;	// unsupported clock id either
+		return 0;
+	}
+
+	// clock_getres() can fail when running on Windows Subsystem for Linux version 1 (but the clock can still be supported).
+	struct timespec Resolution;
+	if (clock_getres(Id, &Resolution) == -1)
+	{
+		Resolution = {-1, -1};
 	}
 
 	// from now on we'll assume that clock_gettime cannot fail
-	uint64 StartTimestamp = FUnixTimeInternal::TimeSpecToNanoSec(ts);
-	uint64 EndTimeStamp = StartTimestamp;
+	uint64 StartTimeNS = TimeSpecToNanoSec(ts);
+	uint64 EndTimeNS = StartTimeNS;
 
 	uint64 NumCalls = 1;	// account for starting timestamp
 	uint64 NumZeroDeltas = 0;
 	const uint64 kHardLimitOnZeroDeltas = (1 << 26);	// arbitrary, but high enough so we don't hit it on fast coarse clocks
 	do
 	{
-		clock_gettime(BenchClockId, &ts);
+		clock_gettime(Id, &ts);
 
-		uint64 NewEndTimeStamp = FUnixTimeInternal::TimeSpecToNanoSec(ts);
+		uint64 NewEndTimeNS = TimeSpecToNanoSec(ts);
 		++NumCalls;
 
-		if (NewEndTimeStamp < EndTimeStamp)
+		if (NewEndTimeNS < EndTimeNS)
 		{
-			FCStringAnsi::Snprintf(Buffer, sizeof(Buffer), "Clock_id %d (%s) is unusable, can go backwards.\n", BenchClockId, BenchClockIdName);
+			FCStringAnsi::Snprintf(Buffer, sizeof(Buffer), "Clock_id %d (%s) is not monotonic.\n", Id, Name);
 			FCStringAnsi::StrncatTruncateDest(CalibrationLog, sizeof(CalibrationLog), Buffer);
-			return 0;
+			return NumCalls;
 		}
-		else if (NewEndTimeStamp == EndTimeStamp)
+
+		if (NewEndTimeNS == EndTimeNS)
 		{
 			++NumZeroDeltas;
 
 			// do not lock up if the clock is broken (e.g. stays in place)
 			if (NumZeroDeltas > kHardLimitOnZeroDeltas)
 			{
-				FCStringAnsi::Snprintf(Buffer, sizeof(Buffer), "Clock_id %d (%s) is unusable, too many (%llu) zero deltas.\n", BenchClockId, BenchClockIdName, NumZeroDeltas);
+				FCStringAnsi::Snprintf(Buffer, sizeof(Buffer), "Clock_id %d (%s) has many (%llu) zero deltas.\n", Id, Name, NumZeroDeltas);
 				FCStringAnsi::StrncatTruncateDest(CalibrationLog, sizeof(CalibrationLog), Buffer);
-				return 0;
+				return NumCalls;
 			}
 		}
 
-		EndTimeStamp = NewEndTimeStamp;
+		EndTimeNS = NewEndTimeNS;
 	}
-	while (EndTimeStamp - StartTimestamp < kBenchmarkPeriodMicroSec);
+	while (EndTimeNS - StartTimeNS < kBenchmarkPeriodNanoSec);
 
-	double TimesPerSecond = 1e9 / static_cast<double>(EndTimeStamp - StartTimestamp);
+	const double BenchmarksPerSecond = 1e9 / static_cast<double>(EndTimeNS - StartTimeNS);
+	uint64 RealNumCallsPerSecond = static_cast<uint64>(BenchmarksPerSecond * static_cast<double>(NumCalls));
 
-	uint64 RealNumCalls = static_cast<uint64> (TimesPerSecond * static_cast<double>(NumCalls));
-
-	char ZeroDeltasBuf[128];
+	char ZeroDeltasBuf[128] = "";
 	if (NumZeroDeltas)
 	{
 		FCStringAnsi::Snprintf(ZeroDeltasBuf, sizeof(ZeroDeltasBuf), "with %f%% zero deltas", 100.0 * static_cast<double>(NumZeroDeltas) / static_cast<double>(NumCalls));
 	}
 
-	FCStringAnsi::Snprintf(Buffer, sizeof(Buffer), (" - %s (id=%d) can sustain %llu (%lluK, %lluM) calls per second %s.\n"), BenchClockIdName, BenchClockId,
-		RealNumCalls, (RealNumCalls + 500) / 1000, (RealNumCalls + 500000) / 1000000,
+	FCStringAnsi::Snprintf(Buffer, sizeof(Buffer), (" - %s (id=%d), resolution %" PRId64 ".%.09" PRId32 "s, can sustain %" PRIu64 " (%" PRIu64 "K, % " PRIu64 "M) calls per second %s.\n"),
+		Name, static_cast<int>(Id),
+		static_cast<int64>(Resolution.tv_sec), static_cast<int32>(Resolution.tv_nsec),
+		RealNumCallsPerSecond, (RealNumCallsPerSecond + 500) / 1'000, (RealNumCallsPerSecond + 500'000) / 1'000'000,
 		NumZeroDeltas ? ZeroDeltasBuf : "without zero deltas");
 	FCStringAnsi::StrncatTruncateDest(CalibrationLog, sizeof(CalibrationLog), Buffer);
 
-	// if clock had zero deltas, we don't want it
-	if (NumZeroDeltas)
-	{
-		return 0;
-	}
-
-	return RealNumCalls;
+	return RealNumCallsPerSecond;
 }
 
-int FUnixTime::CalibrateAndSelectClock()
+clockid_t FUnixTime::CalibrateAndSelectClock()
 {
+	using namespace FUnixTimeInternal;
+
 	// do not calibrate in case of programs, so e.g. ShaderCompileWorker speed is not impacted
-	if (IS_PROGRAM)
+	if constexpr (IS_PROGRAM)
 	{
 		struct timespec ts;
 		if (clock_gettime(CLOCK_MONOTONIC, &ts) == -1)
@@ -320,65 +328,69 @@ int FUnixTime::CalibrateAndSelectClock()
 				"Skipped benchmarking clocks because the engine is running in a standalone program mode: CLOCK_MONOTONIC is unavailable, CLOCK_REALTIME will be used.\n");
 			return CLOCK_REALTIME;
 		}
-		else
-		{
-			FCStringAnsi::Snprintf(CalibrationLog, sizeof(CalibrationLog),
-				"Skipped benchmarking clocks because the engine is running in a standalone program mode - CLOCK_MONOTONIC will be used.\n");
-			return CLOCK_MONOTONIC;
-		}
+
+		FCStringAnsi::Snprintf(CalibrationLog, sizeof(CalibrationLog),
+			"Skipped benchmarking clocks because the engine is running in a standalone program mode - CLOCK_MONOTONIC will be used.\n");
+		return CLOCK_MONOTONIC;
 	}
-	else
-	{
-		char Buffer[256];
 
-		// init calibration log
-		FCStringAnsi::Snprintf(CalibrationLog, sizeof(CalibrationLog), "Benchmarking clocks:\n");
+	char Buffer[256];
 
-		struct ClockDesc
-		{
-			int Id;
-			const char *Desc;
-			uint64 Rate;
-		}
-		Clocks[] =
-		{
-			{ CLOCK_REALTIME, "CLOCK_REALTIME", 0 },
-			{ CLOCK_MONOTONIC, "CLOCK_MONOTONIC", 0 },
-			{ CLOCK_MONOTONIC_RAW, "CLOCK_MONOTONIC_RAW", 0 },
-			{ CLOCK_MONOTONIC_COARSE, "CLOCK_MONOTONIC_COARSE", 0 }
-		};
+	// init calibration log
+	FCStringAnsi::Snprintf(CalibrationLog, sizeof(CalibrationLog), "Benchmarking clocks:\n");
+
+	struct FClockDesc
+	{
+		clockid_t Id = -1;
+		const char* Name = "";
+		// Higher values are more preferred.
+		int Preference = 0;
+		uint64 CallsPerSecond = 0;
+	}
+	Clocks[] =
+	{
+		{ CLOCK_MONOTONIC, "CLOCK_MONOTONIC", 100 }, // Highest preference
+		{ CLOCK_MONOTONIC_COARSE, "CLOCK_MONOTONIC_COARSE", 20 },
+		{ CLOCK_MONOTONIC_RAW, "CLOCK_MONOTONIC_RAW", 10 },
+		{ CLOCK_REALTIME, "CLOCK_REALTIME", 1 }, // Last resort.
+	};
 
-		int ChosenClock = 0;	
-		// Skip measuring CLOCK_REALTIME, so it's never picked up if any other is available.
-		// CLOCK_REALTIME should be always supported - as a last resort.
-		for (int Idx = 1; Idx < UE_ARRAY_COUNT(Clocks); ++Idx)
+	for (FClockDesc& Clock : Clocks)
+	{
+		Clock.CallsPerSecond = BenchmarkClock(Clock.Id, Clock.Name);
+		if (Clock.CallsPerSecond == 0)
 		{
-			Clocks[Idx].Rate = CallsPerSecondBenchmark(Clocks[Idx].Id, Clocks[Idx].Desc);
-			if (Clocks[Idx].Rate > Clocks[ChosenClock].Rate)
-			{
-				ChosenClock = Idx;
-			}
+			Clock.Preference = 0;
 		}
+	}
 
-		FCStringAnsi::Snprintf(Buffer, sizeof(Buffer), "Selected clock_id %d (%s) since it is the fastest support clock without zero deltas.\n",
-			Clocks[ChosenClock].Id, Clocks[ChosenClock].Desc);
-		FCStringAnsi::StrncatTruncateDest(CalibrationLog, sizeof(CalibrationLog), Buffer);
-
-		// Warn if our current clock source cannot be called at least 1M times a second (<30k a frame) as this may affect tight loops
-		if (ChosenClock != 0 && Clocks[ChosenClock].Rate < 1000000)
+	// Select preferred clock
+	int ChosenClockIdx = 0;
+	for (int Idx = 0; Idx < UE_ARRAY_COUNT(Clocks); ++Idx)
+	{
+		if (Clocks[Idx].Preference > Clocks[ChosenClockIdx].Preference)
 		{
-			FCStringAnsi::Snprintf(Buffer, sizeof(Buffer), "The clock source is too slow on this machine, performance may be affected.\n");
-			FCStringAnsi::StrncatTruncateDest(CalibrationLog, sizeof(CalibrationLog), Buffer);
+			ChosenClockIdx = Idx;
 		}
+	}
+
+	FCStringAnsi::Snprintf(Buffer, sizeof(Buffer), "Selected %s (%d) as the best available clock.\n", Clocks[ChosenClockIdx].Name, Clocks[ChosenClockIdx].Id);
+	FCStringAnsi::StrncatTruncateDest(CalibrationLog, sizeof(CalibrationLog), Buffer);
 
-		return Clocks[ChosenClock].Id;
+	// Warn if our current clock source cannot be called at least 1M times a second (<30k a frame) as this may affect tight loops
+	if (Clocks[ChosenClockIdx].CallsPerSecond < 1'000'000)
+	{
+		FCStringAnsi::Snprintf(Buffer, sizeof(Buffer), "The clock source (%s) is too slow on this machine, performance may be affected.\n", Clocks[ChosenClockIdx].Name);
+		FCStringAnsi::StrncatTruncateDest(CalibrationLog, sizeof(CalibrationLog), Buffer);
 	}
+
+	return Clocks[ChosenClockIdx].Id;
 }
 
 void FUnixTime::PrintCalibrationLog()
 {
 	// clock selection happens too early to be printed to log, print it now
-	FString Buffer(ANSI_TO_TCHAR(CalibrationLog));
+	FString Buffer(ANSI_TO_TCHAR(FUnixTimeInternal::CalibrationLog));
 
 	TArray<FString> Lines;
 	Buffer.ParseIntoArrayLines(Lines);
diff --git a/Engine/Source/Runtime/Core/Public/Unix/UnixPlatformTime.h b/Engine/Source/Runtime/Core/Public/Unix/UnixPlatformTime.h
--- a/Engine/Source/Runtime/Core/Public/Unix/UnixPlatformTime.h
+++ b/Engine/Source/Runtime/Core/Public/Unix/UnixPlatformTime.h
@@ -20,32 +20,22 @@ struct FUnixTime : public FGenericPlatformTime
 
 	static inline double Seconds()
 	{
-		if (UNLIKELY(ClockSource < 0))
-		{
-			ClockSource = CalibrateAndSelectClock();
-		}
-
 		struct timespec ts;
-		clock_gettime(ClockSource, &ts);
+		clock_gettime(GetClockSource(), &ts);
 		return static_cast<double>(ts.tv_sec) + static_cast<double>(ts.tv_nsec) / 1e9;
 	}
 
 	static inline uint32 Cycles()
 	{
-		if (UNLIKELY(ClockSource < 0))
-		{
-			ClockSource = CalibrateAndSelectClock();
-		}
-
 		struct timespec ts;
-		clock_gettime(CLOCK_MONOTONIC, &ts);
+		clock_gettime(GetClockSource(), &ts);
 		return static_cast<uint32>(static_cast<uint64>(ts.tv_sec) * (uint64)1e6 + static_cast<uint64>(ts.tv_nsec) / 1000ULL);
 	}
 
 	static inline uint64 Cycles64()
 	{
 		struct timespec ts;
-		clock_gettime(CLOCK_MONOTONIC, &ts);
+		clock_gettime(GetClockSource(), &ts);
 		return static_cast<uint64>(static_cast<uint64>(ts.tv_sec) * (uint64)1e7 + static_cast<uint64>(ts.tv_nsec) / 100ULL);
 	}
 
@@ -62,29 +52,28 @@ struct FUnixTime : public FGenericPlatformTime
 	 */
 	static CORE_API void PrintCalibrationLog();
 
-	UE_INTERNAL static int GetClockSource()
+	/**
+	 * Return the clock source, initiating clock selection (calibration/benchmarking) first if necessary.
+	 */
+	static inline clockid_t GetClockSource()
 	{
+		if (UNLIKELY(ClockSource < 0))
+		{
+			ClockSource = CalibrateAndSelectClock();
+		}
+
 		return ClockSource;
 	}
 
 private:
 
 	/** Clock source to use */
-	static CORE_API int ClockSource;
-
-	/** Log information about calibrating the clock. */
-	static CORE_API char CalibrationLog[4096];
-
-	/**
-	 * Benchmarks clock_gettime(), possibly switches to something else is too slow.
-	 * Unix-specific.
-	 */
-	static CORE_API int CalibrateAndSelectClock();
+	static CORE_API clockid_t ClockSource;
 
 	/**
-	 * Returns number of time we can call the clock per second.
+	 * Benchmarks clock_gettime() and returns the best available clock.
 	 */
-	static CORE_API uint64 CallsPerSecondBenchmark(clockid_t BenchClockId, const char * BenchClockIdName);
+	static CORE_API clockid_t CalibrateAndSelectClock();
 };
 
 typedef FUnixTime FPlatformTime;
