Fix cis issue in main mac compile step

#Synced-CL 48831686

diff --git a/Engine/Source/Runtime/Core/Private/Mac/MacPlatformProcess.cpp b/Engine/Source/Runtime/Core/Private/Mac/MacPlatformProcess.cpp
--- a/Engine/Source/Runtime/Core/Private/Mac/MacPlatformProcess.cpp
+++ b/Engine/Source/Runtime/Core/Private/Mac/MacPlatformProcess.cpp
@@ -708,7 +708,7 @@ namespace UE::Mac
 
 					if (!ProcessesToReap.Remove(PId))
 					{
-						for (TSet<pid_t>::TIterator Iterator = ProcessesToReap; Iterator; ++Iterator)
+						for (TSet<pid_t>::TIterator Iterator = ProcessesToReap.CreateIterator(); Iterator; ++Iterator)
 						{
 							if (Reap(*Iterator))
 							{
