iostore: Fixed TSAN warning

#rnx

#Synced-CL Unknown

diff --git a/Engine/Source/Runtime/Core/Private/IO/IoDispatcher.cpp b/Engine/Source/Runtime/Core/Private/IO/IoDispatcher.cpp
--- a/Engine/Source/Runtime/Core/Private/IO/IoDispatcher.cpp
+++ b/Engine/Source/Runtime/Core/Private/IO/IoDispatcher.cpp
@@ -828,7 +828,7 @@ public:
 
 	int64 GetTotalLoaded() const
 	{
-		return TotalLoaded;
+		return TotalLoaded.load(std::memory_order_relaxed);
 	}
 
 	bool HasMountedBackend() const
@@ -865,7 +865,7 @@ private:
 				else
 				{
 					UE_CLOG(!CompletedRequestsHead->HasBuffer(), LogStreaming, Fatal, TEXT("Backend provided a completed request without an IoBuiffer. Requests that are not failed or cancelled must have an IoBuffer"));
-					FPlatformAtomics::InterlockedAdd(&TotalLoaded, CompletedRequestsHead->GetBuffer().DataSize());
+					TotalLoaded.fetch_add(CompletedRequestsHead->GetBuffer().DataSize(), std::memory_order_relaxed);
 					CompleteRequest(CompletedRequestsHead, EIoErrorCode::Ok);
 				}
 				CompletedRequestsHead->ReleaseRef();
@@ -1096,7 +1096,7 @@ private:
 	TArray<FIoRequestImpl*> RequestsToReprioritize;
 	TAtomic<bool> bStopRequested { false };
 	uint64 PendingIoRequestsCount = 0;
-	int64 TotalLoaded = 0;
+	std::atomic<int64> TotalLoaded = 0;
 	FIoRequestStats RequestStats;
 	UE::FIoDispatcherChunkBlockDecoder ChunkBlockDecoder;
 	bool bIsInitialized = false;
