[AutoRTFM] Default the RollbackWrite parameter of AutoRTFM::InitializeForUE() when !UE_AUTORTFM

#Synced-CL 48856386

diff --git a/Engine/Source/Runtime/Core/Public/AutoRTFM/AutoRTFMUE.h b/Engine/Source/Runtime/Core/Public/AutoRTFM/AutoRTFMUE.h
--- a/Engine/Source/Runtime/Core/Public/AutoRTFM/AutoRTFMUE.h
+++ b/Engine/Source/Runtime/Core/Public/AutoRTFM/AutoRTFMUE.h
@@ -24,7 +24,7 @@ UE_AUTORTFM_API void InitializeForUE(RollbackWriteFn RollbackWrite = nullptr);
 
 #else
 
-inline void InitializeForUE(RollbackWriteFn) {}
+inline void InitializeForUE([[maybe_unused]] RollbackWriteFn RollbackWrite = nullptr) {}
 
 #endif
 
