[RaceDetector]
- Improve performance by removing some costly debug code

#Synced-CL 48817740

diff --git a/Engine/Source/Runtime/Core/Private/HAL/ThreadingBase.cpp b/Engine/Source/Runtime/Core/Private/HAL/ThreadingBase.cpp
--- a/Engine/Source/Runtime/Core/Private/HAL/ThreadingBase.cpp
+++ b/Engine/Source/Runtime/Core/Private/HAL/ThreadingBase.cpp
@@ -188,7 +188,7 @@ CORE_API bool IsInGameThread()
 	{
 PRAGMA_DISABLE_DEPRECATION_WARNINGS
 		bool newValue = FTaskTagScope::IsCurrentTag(ETaskTag::EGameThread) || FTaskTagScope::IsRunningDuringStaticInit();
-#if !UE_BUILD_SHIPPING && !UE_BUILD_TEST
+#if !UE_BUILD_SHIPPING && !UE_BUILD_TEST && !USING_INSTRUMENTATION
 		if (!UE::Tasks::Private::IsThreadRetractingTask())
 		{
 			const uint32 CurrentThreadId = FPlatformTLS::GetCurrentThreadId();
@@ -224,7 +224,7 @@ CORE_API bool IsInSlateThread()
 PRAGMA_DISABLE_DEPRECATION_WARNINGS
 	// If this explicitly is a slate thread, not just the main thread running slate
 	bool newValue = FTaskTagScope::IsCurrentTag(ETaskTag::ESlateThread);
-#if !UE_BUILD_SHIPPING && !UE_BUILD_TEST
+#if !UE_BUILD_SHIPPING && !UE_BUILD_TEST && !USING_INSTRUMENTATION
 	if (!UE::Tasks::Private::IsThreadRetractingTask())
 	{
 		bool oldValue = GSlateLoadingThreadId != 0 && FPlatformTLS::GetCurrentThreadId() == GSlateLoadingThreadId;
@@ -275,7 +275,7 @@ PRAGMA_DISABLE_DEPRECATION_WARNINGS
 		? FTaskTagScope::IsCurrentTag(ETaskTag::EGameThread) || FTaskTagScope::IsCurrentTag(ETaskTag::ERenderingThread) || FTaskTagScope::IsRunningDuringStaticInit()
 		: FTaskTagScope::IsCurrentTag(ETaskTag::ERenderingThread);
 
-#if !UE_BUILD_SHIPPING && !UE_BUILD_TEST
+#if !UE_BUILD_SHIPPING && !UE_BUILD_TEST && !USING_INSTRUMENTATION
 	if (!UE::Tasks::Private::IsThreadRetractingTask())
 	{
 		const uint32 CurrentThreadId = FPlatformTLS::GetCurrentThreadId();
@@ -312,7 +312,7 @@ PRAGMA_DISABLE_DEPRECATION_WARNINGS
 			|| FTaskTagScope::IsCurrentTag(ETaskTag::ERhiThread); //TODO lots of RHI functions rely on our broken IsInParallelRenderingThread;
 	}
 
-#if !UE_BUILD_SHIPPING && !UE_BUILD_TEST
+#if !UE_BUILD_SHIPPING && !UE_BUILD_TEST && !USING_INSTRUMENTATION
 	if (!UE::Tasks::Private::IsThreadRetractingTask())
 	{
 		const uint32 CurrentThreadId = FPlatformTLS::GetCurrentThreadId();
