[CsvProfilerTrace]
- Silence TSAN warning

#Synced-CL 48817740

diff --git a/Engine/Source/Runtime/Core/Public/ProfilingDebugging/CsvProfilerTrace.h b/Engine/Source/Runtime/Core/Public/ProfilingDebugging/CsvProfilerTrace.h
--- a/Engine/Source/Runtime/Core/Public/ProfilingDebugging/CsvProfilerTrace.h
+++ b/Engine/Source/Runtime/Core/Public/ProfilingDebugging/CsvProfilerTrace.h
@@ -47,14 +47,14 @@ struct FCsvProfilerTrace
 	FCsvProfilerTrace::OutputRegisterCategory(Index, Name);
 
 #define TRACE_CSV_PROFILER_INLINE_STAT(StatName, CategoryIndex) \
-	static bool UE_JOIN(__CsvProfilerStat, __LINE__); \
+	static TSAN_ATOMIC(bool) UE_JOIN(__CsvProfilerStat, __LINE__); \
 	if (!UE_JOIN(__CsvProfilerStat, __LINE__)) { \
 		FCsvProfilerTrace::OutputInlineStat(StatName, CategoryIndex); \
 		UE_JOIN(__CsvProfilerStat, __LINE__) = true; \
 	}
 
 #define TRACE_CSV_PROFILER_INLINE_STAT_EXCLUSIVE(StatName) \
-	static bool UE_JOIN(__CsvProfilerStat, __LINE__); \
+	static TSAN_ATOMIC(bool) UE_JOIN(__CsvProfilerStat, __LINE__); \
 	if (!UE_JOIN(__CsvProfilerStat, __LINE__)) { \
 		FCsvProfilerTrace::OutputInlineStatExclusive(StatName); \
 		UE_JOIN(__CsvProfilerStat, __LINE__) = true; \
