GetDeviceThermalState() comment fix.
-The comment was simply duplicated from GetDeviceTemperature().

#rnx
#Synced-CL 48669812

diff --git a/Engine/Source/Runtime/Core/Public/GenericPlatform/GenericPlatformMisc.h b/Engine/Source/Runtime/Core/Public/GenericPlatform/GenericPlatformMisc.h
--- a/Engine/Source/Runtime/Core/Public/GenericPlatform/GenericPlatformMisc.h
+++ b/Engine/Source/Runtime/Core/Public/GenericPlatform/GenericPlatformMisc.h
@@ -1574,16 +1574,13 @@ public:
 	static CORE_API float GetDeviceTemperature();
 
 	/**
-	 * Returns the current device temperature in degrees Celcius.
-	 * Keep in mind that no specific temperature probe position can be assumed.  It will vary across platforms and devices.
-	 * Inspecting the platform implementation might yeild more information about this stat or more detailed platform specific stats.
+	 * Returns the current device thermal state, which is a cross platform qualitative representation of the device's heat status.
+	 * Briefly EDeviceThermalState::None (0) is the ideal state, EDeviceThermalState::Severe+ (3) is where serious throttling will start.
+	 * See EDeviceThermalState and device specific implemetations for details.
 	 * -1.0 is returned for a platform that does not support this measure at all.
-	 * A platform that does support it, but is not currently providing data may return -1 or 0.
-	 * Theoretically -1 and 0 could be real temperatures, but we will assume they are not.
 	 */
 	static CORE_API EDeviceThermalState GetDeviceThermalState();
 
-
 	/**
 	 * Allows a game/program/etc to control the game directory in a special place (for instance, monolithic programs that don't have .uprojects)
 	 */
