LLM: Add LowLevelMemTracking to AnsiMalloc,AnsiRealloc,AnsiFree.
#rnx
#Synced-CL 48825415

diff --git a/Engine/Source/Runtime/Core/Private/HAL/MallocAnsi.cpp b/Engine/Source/Runtime/Core/Private/HAL/MallocAnsi.cpp
--- a/Engine/Source/Runtime/Core/Private/HAL/MallocAnsi.cpp
+++ b/Engine/Source/Runtime/Core/Private/HAL/MallocAnsi.cpp
@@ -5,11 +5,13 @@
 =============================================================================*/
 
 #include "HAL/MallocAnsi.h"
-#include "Misc/AssertionMacros.h"
+
+#include "AutoRTFM.h"
+#include "HAL/LowLevelMemTracker.h"
 #include "Math/UnrealMathUtility.h"
+#include "Misc/AssertionMacros.h"
 #include "ProfilingDebugging/MemoryTrace.h"
 #include "Templates/AlignmentTemplates.h"
-#include "AutoRTFM.h"
 
 #if PLATFORM_USE_ANSI_POSIX_MALLOC
 	#include <malloc.h>
@@ -25,8 +27,9 @@
 
 #define MALLOC_ANSI_USES__ALIGNED_MALLOC !UE_AUTORTFM && PLATFORM_USES__ALIGNED_MALLOC
 
-void* AnsiMalloc(SIZE_T Size, uint32 Alignment)
+static void* AnsiMallocInternal(SIZE_T Size, uint32 Alignment)
 {
+	LLM_PLATFORM_SCOPE(ELLMTag::FMalloc);
 #if MALLOC_ANSI_USES__ALIGNED_MALLOC
 	void* Result = _aligned_malloc( Size, Alignment );
 #elif PLATFORM_USE_ANSI_POSIX_MALLOC
@@ -47,11 +50,24 @@ void* AnsiMalloc(SIZE_T Size, uint32 Alignment)
 		*((SIZE_T*)((uint8*)Result - sizeof(void*) - sizeof(SIZE_T))) = Size;
 	}
 #endif
+	LLM_IF_ENABLED(FLowLevelMemTracker::Get().OnLowLevelAlloc(ELLMTracker::Platform, Result, Size));
 	MemoryTrace_Alloc(uint64(Result), Size, Alignment, EMemoryTraceRootHeap::SystemMemory);
 
 	return Result;
 }
 
+void* AnsiMalloc(SIZE_T Size, uint32 Alignment)
+{
+	// AnsiMalloc always needs to report its allocs to LLM's platform tracker, but it does not need to report its
+	// allocs to LLM's default tracker when called from FMallocAnsi or from other internal functions, because they
+	// have already been reported by the tracking in FMemory.inl or the other external caller. So do LLM's default
+	// tracking only when called externally.
+	void* Result = AnsiMallocInternal(Size, Alignment);
+	LLM_IF_ENABLED(FLowLevelMemTracker::Get().OnLowLevelAlloc(ELLMTracker::Default, Result, Size,
+		ELLMTag::Untagged, ELLMAllocType::FMalloc));
+	return Result;
+}
+
 static SIZE_T AnsiGetAllocationSize(void* Original)
 {
 #if MALLOC_ANSI_USES__ALIGNED_MALLOC
@@ -63,11 +79,15 @@ static SIZE_T AnsiGetAllocationSize(void* Original)
 #endif
 }
 
-void* AnsiRealloc(void* Ptr, SIZE_T NewSize, uint32 Alignment)
+static void AnsiFreeInternal(void* Ptr);
+
+static void* AnsiReallocInternal(void* Ptr, SIZE_T NewSize, uint32 Alignment)
 {
 	void* Result;
 
+	LLM_PLATFORM_SCOPE(ELLMTag::FMalloc);
 #if MALLOC_ANSI_USES__ALIGNED_MALLOC
+	LLM_IF_ENABLED(FLowLevelMemTracker::Get().OnLowLevelFree(ELLMTracker::Platform, Ptr));
 	MemoryTrace_ReallocFree(uint64(Ptr), EMemoryTraceRootHeap::SystemMemory);
 	if (Ptr && NewSize)
 	{
@@ -82,8 +102,10 @@ void* AnsiRealloc(void* Ptr, SIZE_T NewSize, uint32 Alignment)
 		_aligned_free(Ptr);
 		Result = nullptr;
 	}
+	LLM_IF_ENABLED(FLowLevelMemTracker::Get().OnLowLevelAlloc(ELLMTracker::Platform, Result, NewSize));
 	MemoryTrace_ReallocAlloc(uint64(Result), NewSize, Alignment, EMemoryTraceRootHeap::SystemMemory);
 #elif PLATFORM_USE_ANSI_POSIX_MALLOC
+	LLM_IF_ENABLED(FLowLevelMemTracker::Get().OnLowLevelFree(ELLMTracker::Platform, Ptr));
 	MemoryTrace_ReallocFree(uint64(Ptr), EMemoryTraceRootHeap::SystemMemory);
 	if (Ptr && NewSize)
 	{
@@ -110,26 +132,30 @@ void* AnsiRealloc(void* Ptr, SIZE_T NewSize, uint32 Alignment)
 		free(Ptr);
 		Result = nullptr;
 	}
+	LLM_IF_ENABLED(FLowLevelMemTracker::Get().OnLowLevelAlloc(ELLMTracker::Platform, Result, NewSize));
 	MemoryTrace_ReallocAlloc(uint64(Result), NewSize, Alignment, EMemoryTraceRootHeap::SystemMemory);
 #elif PLATFORM_USE_ANSI_MEMALIGN
+	LLM_IF_ENABLED(FLowLevelMemTracker::Get().OnLowLevelFree(ELLMTracker::Platform, Ptr));
 	MemoryTrace_ReallocFree(uint64(Ptr), EMemoryTraceRootHeap::SystemMemory);
 	Result = reallocalign(Ptr, NewSize, Alignment);
+	LLM_IF_ENABLED(FLowLevelMemTracker::Get().OnLowLevelAlloc(ELLMTracker::Platform, Result, NewSize));
 	MemoryTrace_ReallocAlloc(uint64(Result), NewSize, Alignment, EMemoryTraceRootHeap::SystemMemory);
 #else
 	if (Ptr && NewSize)
 	{
 		// Can't use realloc as it might screw with alignment.
-		Result = AnsiMalloc(NewSize, Alignment);
+		Result = AnsiMallocInternal(NewSize, Alignment);
 		SIZE_T PtrSize = AnsiGetAllocationSize(Ptr);
 		FMemory::Memcpy(Result, Ptr, FMath::Min(NewSize, PtrSize));
-		AnsiFree(Ptr);
+		AnsiFreeInternal(Ptr);
 	}
 	else if (Ptr == nullptr)
 	{
-		Result = AnsiMalloc(NewSize, Alignment);
+		Result = AnsiMallocInternal(NewSize, Alignment);
 	}
 	else
 	{
+		LLM_IF_ENABLED(FLowLevelMemTracker::Get().OnLowLevelFree(ELLMTracker::Platform, Ptr));
 		MemoryTrace_Free(uint64(Ptr), EMemoryTraceRootHeap::SystemMemory);
 		free(*((void**)((uint8*)Ptr - sizeof(void*))));
 		Result = nullptr;
@@ -139,8 +165,23 @@ void* AnsiRealloc(void* Ptr, SIZE_T NewSize, uint32 Alignment)
 	return Result;
 }
 
-void AnsiFree(void* Ptr)
+void* AnsiRealloc(void* Ptr, SIZE_T NewSize, uint32 Alignment)
 {
+	// AnsiRealloc always needs to report its allocations and frees to LLM's platform tracker, but it does not need to
+	// report its allocations to LLM's default tracker when called from FMallocAnsi, because they have already been
+	// reported by the tracking in FMemory.inl. So do LLM's default tracking only when called externally.
+	LLM_REALLOC_SCOPE(Ptr);
+	LLM_IF_ENABLED(if (Ptr != nullptr) FLowLevelMemTracker::Get().OnLowLevelFree(ELLMTracker::Default,
+		Ptr, ELLMAllocType::FMalloc));
+	void* Result = AnsiReallocInternal(Ptr, NewSize, Alignment);
+	LLM_IF_ENABLED(if (Result != nullptr) FLowLevelMemTracker::Get().OnLowLevelAlloc(ELLMTracker::Default,
+		Result, NewSize, ELLMTag::Untagged, ELLMAllocType::FMalloc));
+	return Result;
+}
+
+void AnsiFreeInternal(void* Ptr)
+{
+	LLM_IF_ENABLED(FLowLevelMemTracker::Get().OnLowLevelFree(ELLMTracker::Platform, Ptr));
 	MemoryTrace_Free(uint64(Ptr), EMemoryTraceRootHeap::SystemMemory);
 #if MALLOC_ANSI_USES__ALIGNED_MALLOC
 	_aligned_free(Ptr);
@@ -154,6 +195,16 @@ void AnsiFree(void* Ptr)
 #endif
 }
 
+void AnsiFree(void* Ptr)
+{
+	// AnsiFree always needs to report its frees to LLM's platform tracker, but it does not need to report its frees
+	// to LLM's default tracker when called from FMallocAnsi or from other internal functions, because they have
+	// already been reported by the tracking in FMemory.inl or the other external caller. So do LLM's default tracking
+	// only when called externally.
+	LLM_IF_ENABLED(FLowLevelMemTracker::Get().OnLowLevelFree(ELLMTracker::Default, Ptr, ELLMAllocType::FMalloc));
+	AnsiFreeInternal(Ptr);
+}
+
 FMallocAnsi::FMallocAnsi()
 {
 #if PLATFORM_WINDOWS
@@ -176,7 +227,7 @@ void* FMallocAnsi::TryMalloc(SIZE_T Size, uint32 Alignment)
 
 	Alignment = FMath::Max(Size >= 16 ? (uint32)16 : (uint32)8, Alignment);
 
-	void* Result = AnsiMalloc(Size, Alignment);
+	void* Result = AnsiMallocInternal(Size, Alignment);
 
 	return Result;
 }
@@ -205,7 +256,7 @@ void* FMallocAnsi::TryRealloc(void* Ptr, SIZE_T NewSize, uint32 Alignment)
 
 	Alignment = FMath::Max(NewSize >= 16 ? (uint32)16 : (uint32)8, Alignment);
 
-	void* Result = AnsiRealloc(Ptr, NewSize, Alignment);
+	void* Result = AnsiReallocInternal(Ptr, NewSize, Alignment);
 
 	return Result;
 }
@@ -224,7 +275,7 @@ void* FMallocAnsi::Realloc( void* Ptr, SIZE_T NewSize, uint32 Alignment )
 
 void FMallocAnsi::Free( void* Ptr )
 {
-	AnsiFree(Ptr);
+	AnsiFreeInternal(Ptr);
 }
 
 bool FMallocAnsi::GetAllocationSize(void* Original, SIZE_T& SizeOut)
