Move x86 feature querying to Generic Platform and add support for FMA feature detection.

#rnx
#tests compiling with -stresstestunity & clang to prevent regression of CL 40787765

Moves GetFeatureBits_X86() and its helper functions from WindowsPlatfromMisc to GenericPlatformMisc to make them available on all platforms. In addition the following changes are made:
- Make GetFeatureBits_X86() allways callable by returning 0 on none x86 platforms.
- Use the getbv implementation from Engine/Source/Runtime/RadAudioCodec/SDK/Src/RadAudio/cpux86.cpp to avoid a "'__builtin_ia32_xgetbv' needs target feature xsave" error on Mac and Linux.
- Add support for detecting FMA capability.

#Synced-CL 48903505

diff --git a/Engine/Source/Runtime/Core/Private/GenericPlatform/GenericPlatformMisc.cpp b/Engine/Source/Runtime/Core/Private/GenericPlatform/GenericPlatformMisc.cpp
--- a/Engine/Source/Runtime/Core/Private/GenericPlatform/GenericPlatformMisc.cpp
+++ b/Engine/Source/Runtime/Core/Private/GenericPlatform/GenericPlatformMisc.cpp
@@ -641,6 +641,189 @@ EProcessDiagnosticFlags FGenericPlatformMisc::GetProcessDiagnostics()
 	return FoundDiagnostics;
 }
 
+
+#if PLATFORM_CPU_X86_FAMILY
+
+#ifdef _MSC_VER
+#include <intrin.h>
+#define CpuIdEx __cpuidex
+#define CpuId __cpuid
+static inline uint64 xgetbv(uint32 xcr)
+{
+#if _MSC_VER >= 1500 // VC++2008 or later
+	return _xgetbv(xcr);
+#else
+	return 0;
+#endif
+}
+#else
+// GCC/Clang
+
+// 64-bit: GCC/Clang won't let us use "=b" constraint on Mac64, and we need to preserve RBX
+// (PIC/PIE base)
+#define CpuIdEx(out, leaf_id, subleaf_id)\
+			asm("xchgq %%rbx,%q1\n" \
+				"cpuid\n" \
+				"xchgq %%rbx,%q1\n" \
+				: "=a" (out[0]), "=&r" (out[1]), "=c" (out[2]), "=d" (out[3]): "0" (leaf_id), "2"(subleaf_id));
+
+#define CpuId(out, leaf_id) CpuIdEx(out, leaf_id, 0)
+static inline uint64 xgetbv(uint32 xcr)
+{
+	uint32 lo, hi;
+	__asm__ __volatile__("xgetbv" : "=a"(lo), "=d"(hi) : "c"(xcr));
+	return ((uint64)hi << 32) | lo;
+}
+#endif // if not msc
+
+static std::atomic_uint32_t CachedX86FeatureBits = 0;
+uint32 FGenericPlatformMisc::GetFeatureBits_X86()
+{
+	//
+	// Note we are 64bit+ now so we know we have cpuid.
+	//
+
+	uint32 FeatureBits = CachedX86FeatureBits.load(std::memory_order_relaxed);
+	if (FeatureBits)
+	{
+		return FeatureBits;
+	}
+
+	int CpuInfo[4];
+	uint32 MaxLeaf;
+
+	// Basic CPUID information
+	CpuId(CpuInfo, 0);
+	MaxLeaf = CpuInfo[0];
+
+	// Basic feature flags
+	CpuId(CpuInfo, 1);
+
+	FeatureBits |= (CpuInfo[3] & (1u << 26)) ? ECPUFeatureBits_X86::SSE2 : 0;
+	FeatureBits |= (CpuInfo[2] & (1u << 9)) ? ECPUFeatureBits_X86::SSSE3 : 0;
+	FeatureBits |= (CpuInfo[2] & (1u << 20)) ? ECPUFeatureBits_X86::SSE42 : 0;
+
+	// Feature flags we don't report directly (because they need further checks), but use internally
+	bool bHasFMA = (CpuInfo[2] & (1u << 12)) != 0;
+	bool bHasPopcnt = (CpuInfo[2] & (1u << 23)) != 0;
+	bool bHasOsXsave = (CpuInfo[2] & (1u << 27)) != 0;
+	bool bHasAVX = (CpuInfo[2] & (1u << 28)) != 0;
+	bool bHasF16C = (CpuInfo[2] & (1u << 29)) != 0;
+
+	// Get XCR0 if available (context save bits) - required for AVX and other extensions
+	// to test this, on Admin command prompt: (WARNING: this can brick your machine
+	// depending on your drivers. You might need to boot from a recovery/install USB
+	// image to undo if that's the case!)
+	//     bcdedit /set xsavedisable 1
+	// (be sure to "/set xsavedisable 0" again once you're done testing)
+	uint64 Xcr0 = 0;
+	if (bHasOsXsave)
+	{
+		Xcr0 = xgetbv(0);
+	}
+
+	// Required context save bits to use AVX and AVX-512, respectively
+	bool bAVXContextSupport = (Xcr0 & 0x6) == 0x6;
+	bool bAVX512ContextSupport = (Xcr0 & 0xe6) == 0xe6;
+
+	// Combine these to determine effective AVX support
+	if (bAVXContextSupport)
+	{
+		if (bHasAVX)
+		{
+			FeatureBits |= ECPUFeatureBits_X86::AVX;
+		}
+
+		// Even though F16C does not require using AVX registers, the relevant instructions
+		// fault if XSAVE support is disabled
+		if (bHasF16C)
+		{
+			FeatureBits |= ECPUFeatureBits_X86::F16C;
+		}
+
+		if (bHasFMA)
+		{
+			FeatureBits |= ECPUFeatureBits_X86::FMA;
+		}
+	}
+
+	if (MaxLeaf >= 7)
+	{
+		// "Structured extended feature flags enumeration"
+		CpuIdEx(CpuInfo, 7, 0);
+
+		// Some (Celeron) Skylakes erroneously report BMI1/BMI2 even though they don't have it.
+		// These Celerons also don't have AVX.
+		//
+		// All CPUs that actually have BMI1/BMI2 (as of this writing, 2016-05-11) have AVX.
+		// (The ones we care about, anyway.) So only report BMI1/BMI2 if AVX is present.
+		// Also only report AVX or the BMIs if POPCNT is present; all processors I know of
+		// have either both or neither, and it's convenient for us to be able to assume
+		// that either BMI1/BMI2 or AVX2 implies POPCNT.
+		if (bAVXContextSupport && bHasAVX && bHasPopcnt)
+		{
+			// These don't actually require XSAVE support to use, but let's not increase
+			// the size of the testing matrix further.
+			FeatureBits |= (CpuInfo[1] & (1u << 3)) ? ECPUFeatureBits_X86::BMI1 : 0;
+			FeatureBits |= (CpuInfo[1] & (1u << 8)) ? ECPUFeatureBits_X86::BMI2 : 0;
+
+			// In addition to the above, only report AVX2 if BMI1 (and thus LZCNT/TZCNT)
+			// are also reported present; finally VC++ with /arch:AVX2 will emit BMI2
+			// instructions for things like variable shifts so we require BMI2 for AVX2
+			// as well.
+			//
+			// In practice this is not a limitation, AVX2 and BMI2 are a package deal on
+			// all uArchs I'm aware of.
+			const uint32 Avx2Bits = (1u << 3) /* BMI1 */ | (1u << 5) /* AVX2 */ | (1u << 8) /* BMI2 */;
+			if ((CpuInfo[1] & Avx2Bits) == Avx2Bits)
+			{
+				FeatureBits |= ECPUFeatureBits_X86::AVX2;
+			}
+
+			if (bAVX512ContextSupport)
+			{
+				// For us to report AVX512, we want the Skylake feature set
+				const uint32 Avx512Bits = (1u << 31) /* AVX512VL */ | (1u << 30) /* AVX512BW */ | (1u << 17) /* AVX512DQ */ | (1u << 16) /* AVX512F */;
+				if ((CpuInfo[1] & Avx512Bits) == Avx512Bits)
+				{
+					FeatureBits |= ECPUFeatureBits_X86::AVX512;
+				}
+
+				// Use the VBMI2 bit (set on ICL+) to set the NOCAVEATS flag. This is available
+				// on a generation of cores where AVX-512 has no major clock penalty anymore so
+				// whether to use AVX-512 or not is a much more straightforward calculation,
+				// and not so dependent on what else is running at the same time.
+				if (CpuInfo[2] & (1u << 6))
+				{
+					FeatureBits |= ECPUFeatureBits_X86::AVX512_NOCAVEATS;
+				}
+			}
+		}
+	}
+
+	// write detected features
+	// only write value once at end of the function!
+	FeatureBits |= 1; // initialized flag
+
+	CachedX86FeatureBits.store(FeatureBits, std::memory_order_release);
+	return FeatureBits;
+}
+#undef CpuIdEx
+#undef CpuId
+#else
+
+uint32 FGenericPlatformMisc::GetFeatureBits_X86()
+{
+	return 0;
+}
+
+#endif //PLATFORM_CPU_X86_FAMILY
+
+bool FGenericPlatformMisc::HasAVX2InstructionSupport()
+{
+	return CheckFeatureBit_X86(ECPUFeatureBits_X86::AVX2);
+}
+
 FString FGenericPlatformMisc::GetCPUVendor()
 {
 	// Not implemented cross-platform. Each platform may or may not choose to implement this.
diff --git a/Engine/Source/Runtime/Core/Private/Windows/WindowsPlatformMisc.cpp b/Engine/Source/Runtime/Core/Private/Windows/WindowsPlatformMisc.cpp
--- a/Engine/Source/Runtime/Core/Private/Windows/WindowsPlatformMisc.cpp
+++ b/Engine/Source/Runtime/Core/Private/Windows/WindowsPlatformMisc.cpp
@@ -3244,163 +3244,6 @@ uint32 FWindowsPlatformMisc::GetCPUMicrocodeRevision()
 	return MicrocodeVersion;
 }
 
-#if PLATFORM_CPU_X86_FAMILY
-
-#ifdef _MSC_VER
-	#define CpuIdEx __cpuidex
-	#define CpuId __cpuid
-#else
-	// GCC/Clang
-
-	// 64-bit: GCC/Clang won't let us use "=b" constraint on Mac64, and we need to preserve RBX
-	// (PIC/PIE base)
-	#define CpuIdEx(out, leaf_id, subleaf_id)\
-			asm("xchgq %%rbx,%q1\n" \
-				"cpuid\n" \
-				"xchgq %%rbx,%q1\n" \
-				: "=a" (out[0]), "=&r" (out[1]), "=c" (out[2]), "=d" (out[3]): "0" (leaf_id), "2"(subleaf_id));
-
-	#define CpuId(out, leaf_id) CpuIdEx(out, leaf_id, 0)
-
-#endif // if not msc
-
-static std::atomic_uint32_t CachedX86FeatureBits = 0;
-uint32 FWindowsPlatformMisc::GetFeatureBits_X86()
-{
-	//
-	// Note we are 64bit+ now so we know we have cpuid.
-	//
-
-	uint32 FeatureBits = CachedX86FeatureBits.load(std::memory_order_relaxed);
-	if (FeatureBits)
-	{
-		return FeatureBits;
-	}
-
-	int CpuInfo[4];
-	uint32 MaxLeaf;
-
-	// Basic CPUID information
-	UE_CPUID(CpuInfo, 0);
-	MaxLeaf = CpuInfo[0];
-
-	// Basic feature flags
-	UE_CPUID(CpuInfo, 1);
-
-	FeatureBits |= (CpuInfo[3] & (1u << 26)) ? ECPUFeatureBits_X86::SSE2 : 0;
-	FeatureBits |= (CpuInfo[2] & (1u << 9)) ? ECPUFeatureBits_X86::SSSE3 : 0;
-	FeatureBits |= (CpuInfo[2] & (1u << 20)) ? ECPUFeatureBits_X86::SSE42 : 0;
-
-	// Feature flags we don't report directly (because they need further checks), but use internally
-	bool bHasPopcnt = (CpuInfo[2] & (1u << 23)) != 0;
-	bool bHasOsXsave = (CpuInfo[2] & (1u << 27)) != 0;
-	bool bHasAVX = (CpuInfo[2] & (1u << 28)) != 0;
-	bool bHasF16C = (CpuInfo[2] & (1u << 29)) != 0;
-
-	// Get XCR0 if available (context save bits) - required for AVX and other extensions
-	// to test this, on Admin command prompt: (WARNING: this can brick your machine
-	// depending on your drivers. You might need to boot from a recovery/install USB
-	// image to undo if that's the case!)
-	//     bcdedit /set xsavedisable 1
-	// (be sure to "/set xsavedisable 0" again once you're done testing)
-	uint64 Xcr0 = 0;
-	if (bHasOsXsave)
-	{
-		Xcr0 = _xgetbv(0);
-	}
-
-	// Required context save bits to use AVX and AVX-512, respectively
-	bool bAVXContextSupport = (Xcr0 & 0x6) == 0x6;
-	bool bAVX512ContextSupport = (Xcr0 & 0xe6) == 0xe6;
-
-	// Combine these to determine effective AVX support
-	if (bAVXContextSupport)
-	{
-		if (bHasAVX)
-		{
-			FeatureBits |= ECPUFeatureBits_X86::AVX;
-		}
-
-		// Even though F16C does not require using AVX registers, the relevant instructions
-		// fault if XSAVE support is disabled
-		if (bHasF16C)
-		{
-			FeatureBits |= ECPUFeatureBits_X86::F16C;
-		}
-	}
-
-	if (MaxLeaf >= 7)
-	{
-		// "Structured extended feature flags enumeration"
-		CpuIdEx(CpuInfo, 7, 0);
-
-		// Some (Celeron) Skylakes erroneously report BMI1/BMI2 even though they don't have it.
-		// These Celerons also don't have AVX.
-		//
-		// All CPUs that actually have BMI1/BMI2 (as of this writing, 2016-05-11) have AVX.
-		// (The ones we care about, anyway.) So only report BMI1/BMI2 if AVX is present.
-		// Also only report AVX or the BMIs if POPCNT is present; all processors I know of
-		// have either both or neither, and it's convenient for us to be able to assume
-		// that either BMI1/BMI2 or AVX2 implies POPCNT.
-		if (bAVXContextSupport && bHasAVX && bHasPopcnt)
-		{
-			// These don't actually require XSAVE support to use, but let's not increase
-			// the size of the testing matrix further.
-			FeatureBits |= (CpuInfo[1] & (1u << 3)) ? ECPUFeatureBits_X86::BMI1 : 0;
-			FeatureBits |= (CpuInfo[1] & (1u << 8)) ? ECPUFeatureBits_X86::BMI2 : 0;
-
-			// In addition to the above, only report AVX2 if BMI1 (and thus LZCNT/TZCNT)
-			// are also reported present; finally VC++ with /arch:AVX2 will emit BMI2
-			// instructions for things like variable shifts so we require BMI2 for AVX2
-			// as well.
-			//
-			// In practice this is not a limitation, AVX2 and BMI2 are a package deal on
-			// all uArchs I'm aware of.
-			const uint32 Avx2Bits = (1u << 3) /* BMI1 */ | (1u << 5) /* AVX2 */ | (1u << 8) /* BMI2 */;
-			if ((CpuInfo[1] & Avx2Bits) == Avx2Bits)
-			{
-				FeatureBits |= ECPUFeatureBits_X86::AVX2;
-			}
-
-			if (bAVX512ContextSupport)
-			{
-				// For us to report AVX512, we want the Skylake feature set
-				const uint32 Avx512Bits = (1u << 31) /* AVX512VL */ | (1u << 30) /* AVX512BW */ | (1u << 17) /* AVX512DQ */ | (1u << 16) /* AVX512F */;
-				if ((CpuInfo[1] & Avx512Bits) == Avx512Bits)
-				{
-					FeatureBits |= ECPUFeatureBits_X86::AVX512;
-				}
-
-				// Use the VBMI2 bit (set on ICL+) to set the NOCAVEATS flag. This is available
-				// on a generation of cores where AVX-512 has no major clock penalty anymore so
-				// whether to use AVX-512 or not is a much more straightforward calculation,
-				// and not so dependent on what else is running at the same time.
-				if (CpuInfo[2] & (1u << 6))
-				{
-					FeatureBits |= ECPUFeatureBits_X86::AVX512_NOCAVEATS;
-				}
-			}
-		}
-	}
-
-	// write detected features
-	// only write value once at end of the function!
-	FeatureBits |= 1; // initialized flag
-
-	CachedX86FeatureBits.store(FeatureBits, std::memory_order_release);
-	return FeatureBits;
-}
-#endif
-
-bool FWindowsPlatformMisc::HasAVX2InstructionSupport()
-{
-#if PLATFORM_CPU_ARM_FAMILY
-	return false;
-#else
-	return CheckFeatureBit_X86(ECPUFeatureBits_X86::AVX2);
-#endif
-}
-
 #include "Windows/AllowWindowsPlatformTypes.h"
 #include <tbs.h>
 bool FWindowsPlatformMisc::HasTPM2Support()
diff --git a/Engine/Source/Runtime/Core/Public/GenericPlatform/GenericPlatformMisc.h b/Engine/Source/Runtime/Core/Public/GenericPlatform/GenericPlatformMisc.h
--- a/Engine/Source/Runtime/Core/Public/GenericPlatform/GenericPlatformMisc.h
+++ b/Engine/Source/Runtime/Core/Public/GenericPlatform/GenericPlatformMisc.h
@@ -574,6 +574,21 @@ enum class EDeviceThermalState : int32
 	Shutdown = 6,	// Device should shutdown.
 };
 
+namespace ECPUFeatureBits_X86
+{
+	inline constexpr uint32 SSE2 = 1U << 2;
+	inline constexpr uint32 SSSE3 = 1U << 3;
+	inline constexpr uint32 SSE42 = 1U << 4;
+	inline constexpr uint32 AVX = 1U << 5;
+	inline constexpr uint32 BMI1 = 1U << 6; // Bit Manipulation Instructions - 1
+	inline constexpr uint32 BMI2 = 1U << 7; // Bit Manipulation Instructions - 2
+	inline constexpr uint32 AVX2 = 1U << 8;
+	inline constexpr uint32 F16C = 1U << 9; // Float16 conversion instructions
+	inline constexpr uint32 AVX512 = 1U << 10; // Skylake feature set : AVXF512{ F,VL,BW,DQ}.
+	inline constexpr uint32 AVX512_NOCAVEATS = 1U << 11; // Set when we have AVX512 without caveats like throttling.
+	inline constexpr uint32 FMA = 1U << 12;
+}
+
 /**
  * Generic implementation for most platforms
  */
@@ -738,6 +753,20 @@ struct FGenericPlatformMisc
 		}
 	}
 
+
+	// Query the CPUID and parse out various feature bits. This is safe to call multiple times and caches the result internally for rapid access.
+	// Bits are all from the ECPUFeatureBits_X86 namespace.
+	static CORE_API uint32 GetFeatureBits_X86();
+	static bool CheckFeatureBit_X86(uint32 FeatureBit_X86) { return (GetFeatureBits_X86() & FeatureBit_X86) != 0; }
+	static bool CheckAllFeatureBits_X86(uint32 FeatureBits_X86) { return (GetFeatureBits_X86() & FeatureBits_X86) == FeatureBits_X86; }
+
+	/**
+	 * Determines if AVX2 instruction set is supported on this platform
+	 *
+	 * @return	Returns true if instruction-set is supported
+	 */
+	static CORE_API bool HasAVX2InstructionSupport();
+
 	/**
 	 * Uses cpuid instruction to get the vendor string
 	 *
diff --git a/Engine/Source/Runtime/Core/Public/Windows/WindowsPlatformMisc.h b/Engine/Source/Runtime/Core/Public/Windows/WindowsPlatformMisc.h
--- a/Engine/Source/Runtime/Core/Public/Windows/WindowsPlatformMisc.h
+++ b/Engine/Source/Runtime/Core/Public/Windows/WindowsPlatformMisc.h
@@ -12,22 +12,6 @@ class GenericApplication;
 struct FGuid;
 class IPlatformChunkInstall;
 
-#if PLATFORM_CPU_X86_FAMILY
-namespace ECPUFeatureBits_X86
-{
-	inline constexpr uint32 SSE2 = 1U << 2;
-	inline constexpr uint32 SSSE3 = 1U << 3;
-	inline constexpr uint32 SSE42 = 1U << 4;
-	inline constexpr uint32 AVX = 1U << 5;
-	inline constexpr uint32 BMI1 = 1U << 6; // Bit Manipulation Instructions - 1
-	inline constexpr uint32 BMI2 = 1U << 7; // Bit Manipulation Instructions - 2
-	inline constexpr uint32 AVX2 = 1U << 8;
-	inline constexpr uint32 F16C = 1U << 9; // Float16 conversion instructions
-	inline constexpr uint32 AVX512 = 1U << 10; // Skylake feature set : AVXF512{ F,VL,BW,DQ}.
-	inline constexpr uint32 AVX512_NOCAVEATS = 1U << 11; // Set when we have AVX512 without caveats like throttling.
-}
-#endif
-
 /** Helper struct used to get the string version of the Windows version. */
 struct FWindowsOSVersionHelper
 {
@@ -260,21 +244,6 @@ struct FWindowsPlatformMisc
 	 * @return	Returns true if cpuid is supported
 	 */
 	static CORE_API bool HasCPUIDInstruction();
-
-#if PLATFORM_CPU_X86_FAMILY
-	// Query the CPUID and parse out various feature bits. This is safe to call multiple times and caches the result internally for rapid access.
-	// Bits are all from the ECPUFeatureBits_X86 namespace.
-	static CORE_API uint32 GetFeatureBits_X86();
-	static bool CheckFeatureBit_X86(uint32 FeatureBit_X86) { return (GetFeatureBits_X86() & FeatureBit_X86) != 0; }
-	static bool CheckAllFeatureBits_X86(uint32 FeatureBits_X86) { return (GetFeatureBits_X86() & FeatureBits_X86) == FeatureBits_X86; }
-#endif
-
-	/**
-	 * Determines if AVX2 instruction set is supported on this platform
-	 *
-	 * @return	Returns true if instruction-set is supported
-	 */
-	static CORE_API bool HasAVX2InstructionSupport();
 	
 	static CORE_API bool HasTPM2Support();
 	static CORE_API bool HasSecureBootSupport();
