iostore: Fixed TSAN warning about NumRefs access

#rnx

#Synced-CL 48897982

diff --git a/Engine/Source/Runtime/Core/Private/IO/IoBuffer.cpp b/Engine/Source/Runtime/Core/Private/IO/IoBuffer.cpp
--- a/Engine/Source/Runtime/Core/Private/IO/IoBuffer.cpp
+++ b/Engine/Source/Runtime/Core/Private/IO/IoBuffer.cpp
@@ -49,13 +49,6 @@ FIoBuffer::BufCore::BufCore(ECloneTag, uint8* InData, uint64 InSize)
 	FMemory::Memcpy(Data(), InData, InSize);
 }
 
-void
-FIoBuffer::BufCore::CheckRefCount() const
-{
-	// Verify that Release() is not being called on an object which is already at a zero refcount
-	check(NumRefs != 0);
-}
-
 void
 FIoBuffer::BufCore::SetDataAndSize(const uint8* InData, uint64 InSize)
 {
diff --git a/Engine/Source/Runtime/Core/Public/IO/IoBuffer.h b/Engine/Source/Runtime/Core/Public/IO/IoBuffer.h
--- a/Engine/Source/Runtime/Core/Public/IO/IoBuffer.h
+++ b/Engine/Source/Runtime/Core/Public/IO/IoBuffer.h
@@ -111,17 +111,16 @@ private:
 
 		inline uint32 AddRef() const
 		{
-			return uint32(FPlatformAtomics::InterlockedIncrement(&NumRefs));
+			return uint32(NumRefs.fetch_add(1, std::memory_order_seq_cst) + 1);
 		}
 
 		inline uint32 Release() const
 		{
+			const int32 Refs = NumRefs.fetch_sub(1, std::memory_order_seq_cst);
 #if DO_CHECK
-			CheckRefCount();
+			check(Refs > 0);
 #endif
-
-			const int32 Refs = FPlatformAtomics::InterlockedDecrement(&NumRefs);
-			if (Refs == 0)
+			if (Refs == 1)
 			{
 				delete this;
 			}
@@ -131,18 +130,16 @@ private:
 
 		uint32 GetRefCount() const
 		{
-			return uint32(NumRefs);
+			return uint32(NumRefs.load(std::memory_order_relaxed));
 		}
 
 		bool IsMemoryOwned() const	{ return Flags & OwnsMemory; }
 
 	private:
-		CORE_API void				CheckRefCount() const;
-
 		uint8*						DataPtr = nullptr;
 
 		uint32						DataSizeLow = 0;
-		mutable int32				NumRefs = 0;
+		mutable std::atomic<int32>	NumRefs = 0;
 
 		// Reference-counted outer "core", used for views into other buffer
 		//
