Unwind table-based backtracing for Clang

Walking the stack involves locating return addresses pushed at the beginning of each stack frame. To get started we need a current stack frame and instruction address inside a procedure that create the frame. With MSVC we got both of these using an intrinsic and told us where the return address was. Clang doesn't have such a built-in but there is a built-in to get the current frame addresses. This is enough to get the caller's stack frame addresses.

Note that the first backtrace frame is no longer recorded as that is part of the bootstrap to walking. Tested on Win64 with both MSVC and Clang, and on all MS platforms.

#rnx

#Synced-CL 48897982

diff --git a/Engine/Source/Runtime/Core/Private/ProfilingDebugging/Microsoft/WindowsCallstackTrace.cpp b/Engine/Source/Runtime/Core/Private/ProfilingDebugging/Microsoft/WindowsCallstackTrace.cpp
--- a/Engine/Source/Runtime/Core/Private/ProfilingDebugging/Microsoft/WindowsCallstackTrace.cpp
+++ b/Engine/Source/Runtime/Core/Private/ProfilingDebugging/Microsoft/WindowsCallstackTrace.cpp
@@ -45,20 +45,7 @@ static bool GModulesAreInitialized = false;
 // stack walking. In some cases this is not suitable, and we then fall back
 // to the standard stack walking implementation.
 #if !defined(UE_CALLSTACK_TRACE_USE_UNWIND_TABLES)
-	#if defined(__clang__)
-		#define UE_CALLSTACK_TRACE_USE_UNWIND_TABLES 0
-	#else
-		#define UE_CALLSTACK_TRACE_USE_UNWIND_TABLES 1
-	#endif
-#endif
-
-// stacktrace tracking using clang intrinsic __builtin_frame_address(0) doesn't work correctly on all windows platforms
-#if !defined(PLATFORM_USE_CALLSTACK_ADDRESS_POINTER)
-	#if defined(__clang__)
-		#define PLATFORM_USE_CALLSTACK_ADDRESS_POINTER 0
-	#else
-		#define PLATFORM_USE_CALLSTACK_ADDRESS_POINTER 1
-	#endif
+	#define UE_CALLSTACK_TRACE_USE_UNWIND_TABLES 1
 #endif
 
 
@@ -149,7 +136,7 @@ public:
 	static FBacktracer*		Get();
 	void					AddModule(UPTRINT Base, const TCHAR* Name);
 	void					RemoveModule(UPTRINT Base);
-	uint32					GetBacktraceId(void* AddressOfReturnAddress);
+	uint32					GetBacktraceId(UPTRINT ProgCounter, const UPTRINT* StackPtr, uint32 FrameSkip=1);
 
 private:
 	struct FFunction
@@ -724,13 +711,11 @@ const FBacktracer::FFunction* FBacktracer::LookupFunction(UPTRINT Address, FLook
 }
 
 ////////////////////////////////////////////////////////////////////////////////
-uint32 FBacktracer::GetBacktraceId(void* AddressOfReturnAddress) 
+uint32 FBacktracer::GetBacktraceId(UPTRINT ProgCounter, const UPTRINT* StackPointer, uint32 FrameSkip) 
 {
 	FLookupState LookupState = {};
 	uint64 Frames[UE_CALLSTACK_TRACE_MAX_FRAMES];
 
-	UPTRINT* StackPointer = (UPTRINT*)AddressOfReturnAddress;
-
 #if BACKTRACE_DBGLVL >= 3
 	UPTRINT TruthBacktrace[1024];
 	uint32 NumTruth = RtlCaptureStackBackTrace(0, 1024, (void**)TruthBacktrace, nullptr);
@@ -746,21 +731,38 @@ uint32 FBacktracer::GetBacktraceId(void* AddressOfReturnAddress)
 	uint64 BacktraceHash = 0;
 	uint32 FrameIdx = 0;
 
+	auto AddFrame = [&] (int32 RspBias)
+	{
+		StackPointer += RspBias;
+		ProgCounter = StackPointer[-1];
+		if (ProgCounter == 0)
+		{
+			return false;
+		}
+
+		if (FrameSkip != 0)
+		{
+			--FrameSkip;
+			return true;
+		}
+
+		// This is a simple order-dependent LCG. Should be sufficient enough
+		BacktraceHash += ProgCounter;
+		BacktraceHash *= 0x30be8efa499c249dull;
+
+		Frames[FrameIdx++] = ProgCounter;
+		return FrameIdx < UE_ARRAY_COUNT(Frames);
+	};
+
 #if BACKTRACE_LOCK_FREE
 	// When running lock free, we defer the lock until a lock free function lookup fails
 	bool Locked = false;
 #else
 	FScopeLock _(&Lock);
 #endif
-	do
+	while (true)
 	{
-		UPTRINT RetAddr = *StackPointer;
-		
-		Frames[FrameIdx++] = RetAddr;
-
-		// This is a simple order-dependent LCG. Should be sufficient enough
-		BacktraceHash += RetAddr;
-		BacktraceHash *= 0x30be8efa499c249dull;
+		UPTRINT RetAddr = ProgCounter;
 
 #if BACKTRACE_LOCK_FREE
 		int32 RspBias;
@@ -772,11 +774,11 @@ uint32 FBacktracer::GetBacktraceId(void* AddressOfReturnAddress)
 			{
 				break;
 			}
-			else
+			else if (!AddFrame(RspBias))
 			{
-				StackPointer += RspBias;
-				continue;
+				break;
 			}
+			continue;
 		}
 		if (!Locked)
 		{
@@ -841,10 +843,11 @@ uint32 FBacktracer::GetBacktraceId(void* AddressOfReturnAddress)
 			break;
 		}
 
-		StackPointer += Function->RspBias;
+		if (!AddFrame(Function->RspBias))
+		{
+			break;
+		}
 	}
-	// Trunkate callstacks longer than MaxStackDepth
-	while (*StackPointer && FrameIdx < UE_ARRAY_COUNT(Frames));
 
 	// Build the backtrace entry for submission
 	FCallstackTracer::FBacktraceEntry BacktraceEntry;
@@ -1029,24 +1032,62 @@ bool CallstackTrace_FilterModule(FStringView ModuleName)
 #endif
 
 ////////////////////////////////////////////////////////////////////////////////
-uint32 CallstackTrace_GetCurrentId()
+FORCENOINLINE static const UPTRINT* CallstackTrace_GetStackFrame()
 {
-	if (!UE_TRACE_CHANNELEXPR_IS_ENABLED(CallstackChannel))
+	// Returns the stack frame's address for the caller.
+
+#if defined(__clang__)
+	// Using __builtin_frame_address(0) will cause this function to get a frame
+	// pointer, which takes one push on the stack. As the function will not be
+	// inlined, there is a return address too - another push. So the bias adjust
+	// the stack frame to match the caller.
+	enum { Bias = 2 };
+	return (UPTRINT*)__builtin_frame_address(0) + Bias;
+
+#elif defined(_MSC_VER)
+	// _AddressOfReturnAddress() is already pointing at the callers frame, but it
+	// is off by one push due to the return address itself. Bias adjusts for that.
+	enum { Bias = 1 };
+	return (UPTRINT*)_AddressOfReturnAddress() + Bias;
+
+#else
+	return nullptr;
+#endif
+}
+
+////////////////////////////////////////////////////////////////////////////////
+FORCENOINLINE static uint32 CallstackTrace_GetCurrentId_Private()
+{
+	FBacktracer* Instance = FBacktracer::Get();
+	if (Instance == nullptr)
 	{
 		return 0;
 	}
 
-	void* StackAddress = PLATFORM_RETURN_ADDRESS_FOR_CALLSTACKTRACING();
-	if (FBacktracer* Instance = FBacktracer::Get())
+	const UPTRINT* StackPtr = CallstackTrace_GetStackFrame();
+	if (StackPtr == nullptr)
 	{
-#if PLATFORM_USE_CALLSTACK_ADDRESS_POINTER
-		return Instance->GetBacktraceId(StackAddress);
-#else
-		return Instance->GetBacktraceId((uint64)StackAddress);
-#endif
+		return 0;
 	}
+	// It is expected that StackPtr == RSP at this point.
 
-	return 0;
+	// We don't need the parent function in the traced callstack which will be
+	// CallstackTrace_GetCurrentId() or possibly MemoryTrace_Alloc/Free (inlining)
+	const uint32 FrameSkip = 1;
+
+	UPTRINT ProgCounter = UPTRINT(&CallstackTrace_GetCurrentId_Private);
+	return Instance->GetBacktraceId(ProgCounter, StackPtr, FrameSkip);
+}
+
+////////////////////////////////////////////////////////////////////////////////
+uint32 CallstackTrace_GetCurrentId()
+{
+	if (!UE_TRACE_CHANNELEXPR_IS_ENABLED(CallstackChannel))
+	{
+		return 0;
+	}
+
+	return CallstackTrace_GetCurrentId_Private();
 }
 
 #endif // defined(PLATFORM_SUPPORTS_TRACE_WIN32_CALLSTACK) && UE_CALLSTACK_TRACE_ENABLED
