- Added support for an optional .ini file that is applied after engine/project plugins and before GFPs and Hotfixes. This allows a project to override plugin settings if needed (likely to be rare, but there is a use case for it). It is called PluginOverride<branch>.ini (ie PluginOverrideEngine.ini) and is stored in Project/Config. There is no per-platform version.

#Synced-CL 48934447

diff --git a/Engine/Source/Programs/AutomationTool/Scripts/CopyBuildToStagingDirectory.Automation.cs b/Engine/Source/Programs/AutomationTool/Scripts/CopyBuildToStagingDirectory.Automation.cs
--- a/Engine/Source/Programs/AutomationTool/Scripts/CopyBuildToStagingDirectory.Automation.cs
+++ b/Engine/Source/Programs/AutomationTool/Scripts/CopyBuildToStagingDirectory.Automation.cs
@@ -2539,6 +2539,13 @@ namespace AutomationScripts
 					return false;
 				}
 
+				// special dynamic layer file
+				const string PluginOverridePrefix = "pluginoverride";
+				if (NormalizedPath.StartsWith(PluginOverridePrefix))
+				{
+					return true;
+				}
+
 				if (PlatformExtensionName != null)
 				{
 					if (NormalizedPath.StartsWith(PlatformExtensionName, StringComparison.InvariantCultureIgnoreCase))
diff --git a/Engine/Source/Runtime/Core/Private/Misc/ConfigCacheIni.cpp b/Engine/Source/Runtime/Core/Private/Misc/ConfigCacheIni.cpp
--- a/Engine/Source/Runtime/Core/Private/Misc/ConfigCacheIni.cpp
+++ b/Engine/Source/Runtime/Core/Private/Misc/ConfigCacheIni.cpp
@@ -3412,7 +3412,7 @@ void FConfigBranch::RunOnEachCommandStream(TFunction<void(FConfigCommandStream&
 }
 
 
-bool FConfigBranch::AddDynamicLayerToHierarchy(const FString& Filename, FConfigModificationTracker* ModificationTracker, TSet<FString>* GlobalConfigFileCache, TSet<FString>* PluginConfigFileCache)
+bool FConfigBranch::AddDynamicLayerToHierarchy(const FString& Filename, FConfigModificationTracker* ModificationTracker, const TSet<FString>* GlobalConfigFileCache, const TSet<FString>* PluginConfigFileCache)
 {
 	if (!DoesConfigFileExistWrapper(*Filename, nullptr, GlobalConfigFileCache, PluginConfigFileCache))
 	{
@@ -3433,7 +3433,7 @@ bool FConfigBranch::AddDynamicLayersToHierarchy(const TArray<FString>& Filenames
 	return AddDynamicLayersToHierarchy(Layers, ModificationTracker);
 }
 
-bool FConfigBranch::AddDynamicLayersToHierarchy(const TArray<FDynamicLayerInfo>& Layers, FConfigModificationTracker* ModificationTracker, TSet<FString>* GlobalConfigFileCache, TSet<FString>* PluginConfigFileCache, bool bForceFullDynamicLayerUpdate)
+bool FConfigBranch::AddDynamicLayersToHierarchy(const TArray<FDynamicLayerInfo>& Layers, FConfigModificationTracker* ModificationTracker, const TSet<FString>* GlobalConfigFileCache, const TSet<FString>* PluginConfigFileCache, bool bForceFullDynamicLayerUpdate)
 {
 	LLM_SCOPE(ELLMTag::ConfigSystem);
 	if (bIsSafeUnloaded)
@@ -3466,6 +3466,7 @@ bool FConfigBranch::AddDynamicLayersToHierarchy(const TArray<FDynamicLayerInfo>&
 			continue;
 		}
 
+		bFoundAFile = true;
 		UE_CLOG(bDumpIniLoadInfo, LogConfig, Display, TEXT("   Found %s!"), *Layer.Filename);
 
 		if (AddedLayers.Num() == 0 && !bForceFullDynamicLayerUpdate)
@@ -3484,6 +3485,7 @@ bool FConfigBranch::AddDynamicLayersToHierarchy(const TArray<FDynamicLayerInfo>&
 		// remember in local array, then figure out how to remember it permanently
 		AddedLayers.Add(DynamicLayer);
 
+
 		// if we aren't caching dynamic layers, then we need a temp layer
 		if (ReplayMethod == EBranchReplayMethod::NoReplay)
 		{
diff --git a/Engine/Source/Runtime/Core/Private/Misc/ConfigContext.cpp b/Engine/Source/Runtime/Core/Private/Misc/ConfigContext.cpp
--- a/Engine/Source/Runtime/Core/Private/Misc/ConfigContext.cpp
+++ b/Engine/Source/Runtime/Core/Private/Misc/ConfigContext.cpp
@@ -1205,6 +1205,19 @@ bool FConfigContext::LoadIniFileHierarchy()
 	}
 
 	Branch->FinalCombinedLayers = Branch->InMemoryFile;
+
+	// look for some standard dynamic layers that we know to look for now
+	if (!bIsForPlugin && !ProjectConfigDir.IsEmpty())
+	{
+		FString ProjectPluginsOverride = FPaths::Combine(ProjectConfigDir, FString::Printf(TEXT("PluginOverride%s.ini"), *BaseIniName));
+
+		// use the DynamicLayers version so we can specify the priority
+		if (Branch->AddDynamicLayersToHierarchy({ { ProjectPluginsOverride, NAME_None, (uint16)DynamicLayerPriority::ProjectPluginOverrides } }, nullptr, StagedGlobalConfigCache, nullptr))
+		{
+			bReadAnyFile = true;
+		}
+	}
+
 	return bReadAnyFile;
 }
 
diff --git a/Engine/Source/Runtime/Core/Public/Misc/ConfigCacheIni.h b/Engine/Source/Runtime/Core/Public/Misc/ConfigCacheIni.h
--- a/Engine/Source/Runtime/Core/Public/Misc/ConfigCacheIni.h
+++ b/Engine/Source/Runtime/Core/Public/Misc/ConfigCacheIni.h
@@ -1088,12 +1088,22 @@ enum class EBranchReplayMethod : uint8
 };
 
 
-// NOTE: These are currently unused - here for future use
+// this is the ordering within dynamic layers - for any given priority, the order is indeterminate (controlled by discovery of the 
+// file on disk, or pulled down from hotfix server, etc)
 enum class DynamicLayerPriority : uint16
 {
+	// unknown priority, lowest priority
 	Unknown = 0,
-	Plugin = 20,
+	// plugins built into UE (a project can disable them, but these may want to override settings in a project plugin)
+	EnginePlugin = 10,
+	Plugin UE_DEPRECATED(5.8, "DynamicLayerPriority::Plugin is deprecated, use EnginePlugin or ProjectPlugin instead") = 20,
+	// plugins in Project/Plugins, or external locations, that the project has enabled
+	ProjectPlugin = 30,
+	// this special layer allows for a .ini file specified in the project to override any plugin settings, just in cast that is needed
+	ProjectPluginOverrides = 49,
+	// GameFeaturePlugins tend to come and go, and expect to override other plugin settings
 	GameFeature = 50,
+	// hotfixes are generally going to the highest priority to be able to override any setting
 	Hotfix = 80,
 };
 
@@ -1155,8 +1165,8 @@ public:
 	 * Additonally, returns the FConfigFIle object that contains just the loaded sections
 	 * Can return the modified sections if the callers wants to reloadconfig on classes
 	 */
-	CORE_API bool AddDynamicLayerToHierarchy(const FString& Filename, FConfigModificationTracker* ModificationTracker=nullptr, TSet<FString>* GlobalConfigFileCache=nullptr, TSet<FString>* PluginConfigFileCache=nullptr);
-	CORE_API bool AddDynamicLayersToHierarchy(const TArray<FDynamicLayerInfo>& Layers, FConfigModificationTracker* ModificationTracker=nullptr, TSet<FString>* GlobalConfigFileCache=nullptr, TSet<FString>* PluginConfigFileCache=nullptr, bool bForceFullDynamicLayerUpdate=false);
+	CORE_API bool AddDynamicLayerToHierarchy(const FString& Filename, FConfigModificationTracker* ModificationTracker=nullptr, const TSet<FString>* GlobalConfigFileCache=nullptr, const TSet<FString>* PluginConfigFileCache=nullptr);
+	CORE_API bool AddDynamicLayersToHierarchy(const TArray<FDynamicLayerInfo>& Layers, FConfigModificationTracker* ModificationTracker=nullptr, const TSet<FString>* GlobalConfigFileCache=nullptr, const TSet<FString>* PluginConfigFileCache=nullptr, bool bForceFullDynamicLayerUpdate=false);
 
 	UE_DEPRECATED(5.6, "Use AddDynamicLayersToHierarchy that takes a FDynamicLayerInfo list")
 	CORE_API bool AddDynamicLayersToHierarchy(const TArray<FString>& Filenames, FName Tag=NAME_None, DynamicLayerPriority Priority=DynamicLayerPriority::Unknown, FConfigModificationTracker* ModificationTracker =nullptr);
diff --git a/Engine/Source/Runtime/Projects/Private/PluginManager.cpp b/Engine/Source/Runtime/Projects/Private/PluginManager.cpp
--- a/Engine/Source/Runtime/Projects/Private/PluginManager.cpp
+++ b/Engine/Source/Runtime/Projects/Private/PluginManager.cpp
@@ -182,7 +182,7 @@ namespace DiscoveredPluginMapUtils
 	}
 
 	/** Removes/Forgets all other known versions of the specified plugin entry. Retains the one "offered" version (which was prioritized and positioned for the rest of UE). */
-	static void DiscardAllSupressedVersions(FDiscoveredPluginMap::ValueType& PluginMapValue)
+	static void DiscardAllSuppressedVersions(FDiscoveredPluginMap::ValueType& PluginMapValue)
 	{
 		// The "top" (last) plugin in the list is the one we "offer". Keep it, discard the rest.
 		PluginMapValue.RemoveAtSwap(0, PluginMapValue.Num()-1);
@@ -566,7 +566,7 @@ void FPluginManager::RefreshPluginsList()
 		else
 		{
 			// Forget all the other discovered versions (which we assume aren't enabled/mounted)
-			DiscoveredPluginMapUtils::DiscardAllSupressedVersions(Iter.Value());
+			DiscoveredPluginMapUtils::DiscardAllSuppressedVersions(Iter.Value());
 		}
 	}
 
@@ -1921,7 +1921,8 @@ void FPluginManager::ProcessEnabledPlugins(FPluginDiscoveryContext& Context)
 
 		// register this plugin, so it can be found right now, but also when loading other platforms, or when modifying it by other plugins
 		const bool bIncludePluginNameInBranchName = false;
-		FConfigCacheIni::RegisterPlugin(PluginName, Plugin.GetBaseDir(), Plugin.GetExtensionBaseDirs(), DynamicLayerPriority::Plugin, bIncludePluginNameInBranchName);
+		DynamicLayerPriority LayerPrio = Plugin.GetType() == EPluginType::Engine ? DynamicLayerPriority::EnginePlugin : DynamicLayerPriority::ProjectPlugin;
+		FConfigCacheIni::RegisterPlugin(PluginName, Plugin.GetBaseDir(), Plugin.GetExtensionBaseDirs(), LayerPrio, bIncludePluginNameInBranchName);
 
 		FConfigContext ConfigContext = FConfigContext::ReadIntoGConfig();
 		ConfigContext.ChangeTracker = &ChangeTracker;
