- Added TAutoConsoleVariableDeprecated_NoReplace CVar type that is used to deprecate a cvar without replacing it with another one. The value can still be queried, but it will ensure (or warn or assert depending on setting). Similarly,  FAutoConsoleCommandDeprecated_NoReplace was added for deprecating a command without pointing at another

#Synced-CL 48934447

diff --git a/Engine/Source/Runtime/Core/Private/HAL/ConsoleManager.cpp b/Engine/Source/Runtime/Core/Private/HAL/ConsoleManager.cpp
--- a/Engine/Source/Runtime/Core/Private/HAL/ConsoleManager.cpp
+++ b/Engine/Source/Runtime/Core/Private/HAL/ConsoleManager.cpp
@@ -2705,6 +2705,11 @@ void FConsoleManager::DumpObjects(const TCHAR* Params, FOutputDevice& InAr, bool
 			IConsoleVariable* CVar = Obj->AsVariable();
 			IConsoleCommand* CCmd = Obj->AsCommand();
 
+			if (Obj->TestFlags(ECVF_Unregistered))
+			{
+				continue;
+			}
+			
 			bool bDisplay = true;
 			bDisplay &= bShowDeprecated == Obj->IsDeprecated();
 			bDisplay &= bDisplayCommands ? CCmd != nullptr : CVar != nullptr;
@@ -2747,7 +2752,7 @@ void FConsoleManager::DumpObjects(const TCHAR* Params, FOutputDevice& InAr, bool
 				}
 
 				// Don't read the value when deprecated. It will trip the UsageBehavior, which may assert.
-				FString Value = Obj->IsDeprecated() ? TEXT("") : CVar->GetString();
+				FString Value = Obj->IsDeprecated() ? TEXT("<deprecated, can't show>") : CVar->GetString();
 
 				if (bWriteToCSV)
 				{
@@ -4740,13 +4745,21 @@ protected:
 				FText Message;
 				if (DeprecationMessage.IsEmpty())
 				{
-					if (DeprecatedVersion != TEXT(""))
+					if (ShadowName.IsEmpty() || !!(GetRealObject()->GetFlags() & ECVF_Unregistered))
 					{
-						Message = FText::Format(LOCTEXT("ShadowCVarUsage_Deprecated", "Using a deprecated (as of UE {DeprecatedVersion}) CVar: '{ThisName}'. It will be removed in the future. Change all uses to '{ShadowName}' instead."), Arguments);
+						// in-place deprecation requires a version number
+						Message = FText::Format(LOCTEXT("ShadowCVarUsage_DeprecatedInPlace", "Using a deprecated (as of UE {DeprecatedVersion}) CVar/Command: '{ThisName}'. It will be removed in the future, and may not function as expected. Please review your usage of {ThisName}."), Arguments);
 					}
 					else
 					{
-						Message = FText::Format(LOCTEXT("ShadowCVarUsage", "Using a shadowed CVar '{ThisName}'. It is recommended to change all uses to '{ShadowName}' instead."), Arguments);
+						if (DeprecatedVersion != TEXT(""))
+						{
+							Message = FText::Format(LOCTEXT("ShadowCVarUsage_Deprecated", "Using a deprecated (as of UE {DeprecatedVersion}) CVar/Command: '{ThisName}'. It will be removed in the future. Change all uses to '{ShadowName}' instead."), Arguments);
+						}
+						else
+						{
+							Message = FText::Format(LOCTEXT("ShadowCVarUsage", "Using a shadowed CVar '{ThisName}'. It is recommended to change all uses to '{ShadowName}' instead."), Arguments);
+						}
 					}
 				}
 				else
@@ -5241,6 +5254,187 @@ FAutoConsoleCommandDeprecated::FAutoConsoleCommandDeprecated(const TCHAR* Name,
 }
 
 
+// helper class for in-place deprecation, which has a backing cvar unregistered cvar created with a _dep nam, but since 
+// it's ECVF_Unregistered, it doesn't show up in auto-complete etc
+class FConsoleVariableDeprecatedNoReplace : public FConsoleVariableShadow
+{
+public:
+	FConsoleVariableDeprecatedNoReplace(const TCHAR* CVarToShadow, const TCHAR* DeprecatedAtVersion, EShadowCVarBehavior UsageBehavior, const TCHAR* DeprecationMessage)
+	: FConsoleVariableShadow(CVarToShadow, EShadowCVarBehavior::NoMessaging, UsageBehavior, true, DeprecatedAtVersion, DeprecationMessage)
+	{
+		
+	}
+	
+	virtual EConsoleVariableFlags GetFlags() const override
+	{
+		if (Bind())
+		{
+			// we need to remove the ECVF_Unregistered flag when returning the info about the "good" cvar
+			return (EConsoleVariableFlags)(GetRealObject()->GetFlags() & ~ECVF_Unregistered);
+		}
+		return (EConsoleVariableFlags)0;
+
+	}
+};
+
+class FConsoleCommandDeprecatedNoReplace : public FConsoleCommandShadow
+{
+public:
+	FConsoleCommandDeprecatedNoReplace(const TCHAR* CVarToShadow, const TCHAR* DeprecatedAtVersion, EShadowCVarBehavior UsageBehavior, const TCHAR* DeprecationMessage)
+		: FConsoleCommandShadow(CVarToShadow, EShadowCVarBehavior::NoMessaging, UsageBehavior, true, DeprecatedAtVersion, DeprecationMessage)
+	{
+
+	}
+
+	virtual EConsoleVariableFlags GetFlags() const override
+	{
+		if (Bind())
+		{
+			// we need to remove the ECVF_Unregistered flag when returning the info about the "good" cvar
+			return (EConsoleVariableFlags)(GetRealObject()->GetFlags() & ~ECVF_Unregistered);
+		}
+		return (EConsoleVariableFlags)0;
+
+	}
+};
+
+static FString MakeDeprecatedHelp(const TCHAR* Name, const TCHAR* DeprecatedAtVersion, const TCHAR* DeprecationMessage)
+{
+	FString Message = DeprecationMessage;
+	if (Message.IsEmpty())
+	{
+		Message = FString::Printf(TEXT("'%s' has been deprecated as of UE %s. It should no longer be used."), Name, DeprecatedAtVersion);
+	}
+
+	return Message;
+}
+
+template<typename T>
+void CreateInPlaceDeprecatedVar(const TCHAR* Name, const T& DefaultValue, const TCHAR* DeprecatedAtVersion, EShadowCVarBehavior UsageBehavior, const TCHAR* DeprecationMessage)
+{
+#if !NO_CVARS
+	// create a backing unregistered cvar with a dummy name, but this name won't show up since Unregistered
+	FString BackingName = FString(Name) + "_dep";
+	IConsoleVariable* BackingCVar = new FConsoleVariable<T>(DefaultValue, *MakeDeprecatedHelp(Name, DeprecatedAtVersion, DeprecationMessage), (EConsoleVariableFlags)(ECVF_Default | ECVF_Unregistered));
+	GetManager().AddConsoleObject(*BackingName, BackingCVar);
+	
+	// create the shadow cvar that shadows the backing cvar, so we can use all the existing logic
+	FConsoleVariableDeprecatedNoReplace* DeprecatedCVar = new FConsoleVariableDeprecatedNoReplace(*BackingName, DeprecatedAtVersion, UsageBehavior, DeprecationMessage);
+	// don't let the AddCosoleObject throw an error by calling GetFlags()
+	GetManager().AddShadowConsoleObject(Name, DeprecatedCVar);
+
+#endif
+}
+
+template<typename T, typename CmdType>
+void CreateInPlaceDeprecatedCommand(const TCHAR* Name, const CmdType& Command, const TCHAR* DeprecatedAtVersion, EShadowCVarBehavior UsageBehavior, const TCHAR* DeprecationMessage)
+{
+#if !NO_CVARS
+	check(Name != nullptr && FCString::Strlen(Name) > 0);
+	check(DeprecatedAtVersion != nullptr && FCString::Strlen(DeprecatedAtVersion) > 0);
+
+	// create a backing unregistered cvar with a dummy name, but this name won't show up since Unregistered
+	FString BackingName = FString(Name) + "_dep";
+	T* BackingCommand = new T(Command, *MakeDeprecatedHelp(Name, DeprecatedAtVersion, DeprecationMessage), (EConsoleVariableFlags)(ECVF_Default | ECVF_Unregistered));
+	GetManager().AddConsoleObject(*BackingName, BackingCommand);
+
+	// create the shadow cvar that shadows the backing cvar, so we can use all the existing logic
+	FConsoleCommandDeprecatedNoReplace* DeprecatedCVar = new FConsoleCommandDeprecatedNoReplace(*BackingName, DeprecatedAtVersion, UsageBehavior, DeprecationMessage);
+	// don't let the AddCosoleObject throw an error by calling GetFlags()
+	GetManager().AddShadowConsoleObject(Name, DeprecatedCVar);
+
+#endif
+}
+
+void CreateInPlaceDeprecatedCommandExec(const TCHAR* Name, const TCHAR* DeprecatedAtVersion, EShadowCVarBehavior UsageBehavior, const TCHAR* DeprecationMessage)
+{
+#if !NO_CVARS
+	check(Name != nullptr && FCString::Strlen(Name) > 0);
+	check(DeprecatedAtVersion != nullptr && FCString::Strlen(DeprecatedAtVersion) > 0);
+
+	// create a backing unregistered cvar with a dummy name, but this name won't show up since Unregistered
+	FString BackingName = FString(Name) + "_dep";
+	FConsoleCommandExec* BackingCommand = new FConsoleCommandExec(*MakeDeprecatedHelp(Name, DeprecatedAtVersion, DeprecationMessage), (EConsoleVariableFlags)(ECVF_Default | ECVF_Unregistered));
+	GetManager().AddConsoleObject(*BackingName, BackingCommand);
+
+	// create the shadow cvar that shadows the backing cvar, so we can use all the existing logic
+	FConsoleCommandDeprecatedNoReplace* DeprecatedCVar = new FConsoleCommandDeprecatedNoReplace(*BackingName, DeprecatedAtVersion, UsageBehavior, DeprecationMessage);
+	// don't let the AddCosoleObject throw an error by calling GetFlags()
+	GetManager().AddShadowConsoleObject(Name, DeprecatedCVar);
+
+#endif
+}
+
+// because we create a variable whose type is not in the header, we have to instantiate each possible type here
+template<>
+TAutoConsoleVariableDeprecated_NoReplace<int32>::TAutoConsoleVariableDeprecated_NoReplace(const TCHAR* Name, const int32& DefaultValue, const TCHAR* DeprecatedAtVersion, EShadowCVarBehavior UsageBehavior, const TCHAR* DeprecationMessage)
+{
+	CreateInPlaceDeprecatedVar(Name, DefaultValue, DeprecatedAtVersion, UsageBehavior, DeprecationMessage);
+}
+
+template<>
+TAutoConsoleVariableDeprecated_NoReplace<float>::TAutoConsoleVariableDeprecated_NoReplace(const TCHAR* Name, const float& DefaultValue, const TCHAR* DeprecatedAtVersion, EShadowCVarBehavior UsageBehavior, const TCHAR* DeprecationMessage)
+{
+	CreateInPlaceDeprecatedVar(Name, DefaultValue, DeprecatedAtVersion, UsageBehavior, DeprecationMessage);
+}
+
+template<>
+TAutoConsoleVariableDeprecated_NoReplace<bool>::TAutoConsoleVariableDeprecated_NoReplace(const TCHAR* Name, const bool& DefaultValue, const TCHAR* DeprecatedAtVersion, EShadowCVarBehavior UsageBehavior, const TCHAR* DeprecationMessage)
+{
+	CreateInPlaceDeprecatedVar(Name, DefaultValue, DeprecatedAtVersion, UsageBehavior, DeprecationMessage);
+}
+
+template<>
+TAutoConsoleVariableDeprecated_NoReplace<FString>::TAutoConsoleVariableDeprecated_NoReplace(const TCHAR* Name, const FString& DefaultValue, const TCHAR* DeprecatedAtVersion, EShadowCVarBehavior UsageBehavior, const TCHAR* DeprecationMessage)
+{
+	CreateInPlaceDeprecatedVar(Name, DefaultValue, DeprecatedAtVersion, UsageBehavior, DeprecationMessage);
+}
+
+
+FAutoConsoleCommandDeprecated_NoReplace::FAutoConsoleCommandDeprecated_NoReplace(const TCHAR* Name, const TCHAR* DeprecatedAtVersion, EShadowCVarBehavior UsageBehavior, const TCHAR* DeprecationMessage)
+{
+	CreateInPlaceDeprecatedCommandExec(Name, DeprecatedAtVersion, UsageBehavior, DeprecationMessage);
+}
+
+FAutoConsoleCommandDeprecated_NoReplace::FAutoConsoleCommandDeprecated_NoReplace(const TCHAR* Name, const FConsoleCommandDelegate& Command, const TCHAR* DeprecatedAtVersion, EShadowCVarBehavior UsageBehavior, const TCHAR* DeprecationMessage)
+{
+	CreateInPlaceDeprecatedCommand<FConsoleCommand>(Name, Command, DeprecatedAtVersion, UsageBehavior, DeprecationMessage);
+}
+
+FAutoConsoleCommandDeprecated_NoReplace::FAutoConsoleCommandDeprecated_NoReplace(const TCHAR* Name, const FConsoleCommandWithArgsDelegate& Command, const TCHAR* DeprecatedAtVersion, EShadowCVarBehavior UsageBehavior, const TCHAR* DeprecationMessage)
+{
+	CreateInPlaceDeprecatedCommand<FConsoleCommandWithArgs>(Name, Command, DeprecatedAtVersion, UsageBehavior, DeprecationMessage);
+}
+
+FAutoConsoleCommandDeprecated_NoReplace::FAutoConsoleCommandDeprecated_NoReplace(const TCHAR* Name, const FConsoleCommandWithOutputDeviceDelegate& Command, const TCHAR* DeprecatedAtVersion, EShadowCVarBehavior UsageBehavior, const TCHAR* DeprecationMessage)
+{
+	CreateInPlaceDeprecatedCommand<FConsoleCommandWithOutputDevice>(Name, Command, DeprecatedAtVersion, UsageBehavior, DeprecationMessage);
+}
+
+FAutoConsoleCommandDeprecated_NoReplace::FAutoConsoleCommandDeprecated_NoReplace(const TCHAR* Name, const FConsoleCommandWithArgsAndOutputDeviceDelegate& Command, const TCHAR* DeprecatedAtVersion, EShadowCVarBehavior UsageBehavior, const TCHAR* DeprecationMessage)
+{
+	CreateInPlaceDeprecatedCommand<FConsoleCommanWithArgsAndOutputDevice>(Name, Command, DeprecatedAtVersion, UsageBehavior, DeprecationMessage);
+}
+
+FAutoConsoleCommandDeprecated_NoReplace::FAutoConsoleCommandDeprecated_NoReplace(const TCHAR* Name, const FConsoleCommandWithWorldDelegate& Command, const TCHAR* DeprecatedAtVersion, EShadowCVarBehavior UsageBehavior, const TCHAR* DeprecationMessage)
+{
+	CreateInPlaceDeprecatedCommand<FConsoleCommandWithWorld>(Name, Command, DeprecatedAtVersion, UsageBehavior, DeprecationMessage);
+}
+
+FAutoConsoleCommandDeprecated_NoReplace::FAutoConsoleCommandDeprecated_NoReplace(const TCHAR* Name, const FConsoleCommandWithWorldArgsAndOutputDeviceDelegate& Command, const TCHAR* DeprecatedAtVersion, EShadowCVarBehavior UsageBehavior, const TCHAR* DeprecationMessage)
+{
+	CreateInPlaceDeprecatedCommand<FConsoleCommandWithWorldArgsAndOutputDevice>(Name, Command, DeprecatedAtVersion, UsageBehavior, DeprecationMessage);
+}
+
+FAutoConsoleCommandDeprecated_NoReplace::FAutoConsoleCommandDeprecated_NoReplace(const TCHAR* Name, const FConsoleCommandWithWorldAndArgsDelegate& Command, const TCHAR* DeprecatedAtVersion, EShadowCVarBehavior UsageBehavior, const TCHAR* DeprecationMessage)
+{
+	CreateInPlaceDeprecatedCommand<FConsoleCommandWithWorldAndArgs>(Name, Command, DeprecatedAtVersion, UsageBehavior, DeprecationMessage);
+}
+
+
+
+
+
 namespace ConsoleVariablePlatformMapping
 {
 	TMap<int, FName> ShaderPlatformToPlatformMap;
diff --git a/Engine/Source/Runtime/Core/Public/HAL/IConsoleManager.h b/Engine/Source/Runtime/Core/Public/HAL/IConsoleManager.h
--- a/Engine/Source/Runtime/Core/Public/HAL/IConsoleManager.h
+++ b/Engine/Source/Runtime/Core/Public/HAL/IConsoleManager.h
@@ -2318,6 +2318,46 @@ public:
 	CORE_API FAutoConsoleCommandDeprecated(const TCHAR* Name, const TCHAR* CVarToShadow, const TCHAR* DeprecatedAtVersion, EShadowCVarBehavior UsageBehavior = EShadowCVarBehavior::Ensure, EShadowCVarBehavior LookupFailureBehavior = EShadowCVarBehavior::Ensure, const TCHAR* DeprecationMessage = TEXT(""));
 };
 
+/**
+ * A class that can be used to deprecate a cvar, but without making a new one to replace it.
+ */
+template <typename T>
+class TAutoConsoleVariableDeprecated_NoReplace
+{
+public:
+	/**
+	 * @param Name The name of this cvar
+	 * @param DefaultValue The value this cvar has by default
+	 * @param DeprecatedAtVersion The first UE version when this was deprecated
+	 * @param UsageBehavior How to handle uses of the deprecated cvar  - defaults to ensure
+	 * @param DeprecationMessage Optional deprecation message - defaults to empty, auto-generated.
+	 */
+	CORE_API TAutoConsoleVariableDeprecated_NoReplace(const TCHAR* Name, const T& DefaultValue, const TCHAR* DeprecatedAtVersion, EShadowCVarBehavior UsageBehavior=EShadowCVarBehavior::Ensure, const TCHAR* DeprecationMessage = TEXT(""));
+};
+
+/**
+ * A class that can be used to deprecate a ConsoleCommand, but without making a new one to replace it.
+ */
+class FAutoConsoleCommandDeprecated_NoReplace
+{
+public:
+	/**
+	 * @param Name The name of this command
+	 * @param Command The user function to call when this command is executed
+	 * @param DeprecatedAtVersion The first UE version when this was deprecated
+	 * @param UsageBehavior How to handle uses of the deprecated cvar  - defaults to ensure
+	 * @param DeprecationMessage Optional deprecation message - defaults to empty, auto-generated.
+	 */
+	CORE_API FAutoConsoleCommandDeprecated_NoReplace(const TCHAR* Name, const TCHAR* DeprecatedAtVersion, EShadowCVarBehavior UsageBehavior = EShadowCVarBehavior::Ensure, const TCHAR* DeprecationMessage = TEXT(""));
+	CORE_API FAutoConsoleCommandDeprecated_NoReplace(const TCHAR* Name, const FConsoleCommandDelegate& Command, const TCHAR* DeprecatedAtVersion, EShadowCVarBehavior UsageBehavior = EShadowCVarBehavior::Ensure, const TCHAR* DeprecationMessage = TEXT(""));
+	CORE_API FAutoConsoleCommandDeprecated_NoReplace(const TCHAR* Name, const FConsoleCommandWithArgsDelegate& Command, const TCHAR* DeprecatedAtVersion, EShadowCVarBehavior UsageBehavior = EShadowCVarBehavior::Ensure, const TCHAR* DeprecationMessage = TEXT(""));
+	CORE_API FAutoConsoleCommandDeprecated_NoReplace(const TCHAR* Name, const FConsoleCommandWithOutputDeviceDelegate& Command, const TCHAR* DeprecatedAtVersion, EShadowCVarBehavior UsageBehavior = EShadowCVarBehavior::Ensure, const TCHAR* DeprecationMessage = TEXT(""));
+	CORE_API FAutoConsoleCommandDeprecated_NoReplace(const TCHAR* Name, const FConsoleCommandWithArgsAndOutputDeviceDelegate& Command, const TCHAR* DeprecatedAtVersion, EShadowCVarBehavior UsageBehavior = EShadowCVarBehavior::Ensure, const TCHAR* DeprecationMessage = TEXT(""));
+	CORE_API FAutoConsoleCommandDeprecated_NoReplace(const TCHAR* Name, const FConsoleCommandWithWorldDelegate& Command, const TCHAR* DeprecatedAtVersion, EShadowCVarBehavior UsageBehavior = EShadowCVarBehavior::Ensure, const TCHAR* DeprecationMessage = TEXT(""));
+	CORE_API FAutoConsoleCommandDeprecated_NoReplace(const TCHAR* Name, const FConsoleCommandWithWorldArgsAndOutputDeviceDelegate& Command, const TCHAR* DeprecatedAtVersion, EShadowCVarBehavior UsageBehavior = EShadowCVarBehavior::Ensure, const TCHAR* DeprecationMessage = TEXT(""));
+	CORE_API FAutoConsoleCommandDeprecated_NoReplace(const TCHAR* Name, const FConsoleCommandWithWorldAndArgsDelegate& Command, const TCHAR* DeprecatedAtVersion, EShadowCVarBehavior UsageBehavior = EShadowCVarBehavior::Ensure, const TCHAR* DeprecationMessage = TEXT(""));
+};
+
 
 /**
  * A class to create a ranged console variable. The min and max can be changed with <name>.min and <name.max>
