[Core] Add check for double removecurrent in compact set
* Added optional strict mode which adds modification and double remove checks to sparse sets as well.

#tests Ran standard builds in multiple branches to cover both sparse and compact sets
#Synced-CL 48915923

diff --git a/Engine/Source/Programs/UnrealBuildTool/Configuration/Rules/ReadOnlyTargetRules.cs b/Engine/Source/Programs/UnrealBuildTool/Configuration/Rules/ReadOnlyTargetRules.cs
--- a/Engine/Source/Programs/UnrealBuildTool/Configuration/Rules/ReadOnlyTargetRules.cs
+++ b/Engine/Source/Programs/UnrealBuildTool/Configuration/Rules/ReadOnlyTargetRules.cs
@@ -273,6 +273,8 @@ namespace UnrealBuildTool
 
 		public bool bDefineForceInlineHintToInline => Inner.bDefineForceInlineHintToInline;
 
+		public bool bUseSparseSetStrictMode => Inner.bUseSparseSetStrictMode;
+
 		public bool bUseCompactSetAsDefault => Inner.bUseCompactSetAsDefault;
 
 		public bool bUseAndroidJni => Inner.bUseAndroidJni;
diff --git a/Engine/Source/Programs/UnrealBuildTool/Configuration/Rules/TargetRules.cs b/Engine/Source/Programs/UnrealBuildTool/Configuration/Rules/TargetRules.cs
--- a/Engine/Source/Programs/UnrealBuildTool/Configuration/Rules/TargetRules.cs
+++ b/Engine/Source/Programs/UnrealBuildTool/Configuration/Rules/TargetRules.cs
@@ -1640,6 +1640,13 @@ namespace UnrealBuildTool
 		}
 		private bool? bDefineForceInlineHintToInlinePrivate;
 
+		/// <summary>
+		/// If enabled sparse set will have additional runtime checks against usage that conflicts with compact sets. This can be used to catch 
+		/// </summary>
+		[CommandLine("-UseSparseSetStrictMode")]
+		[XmlConfigFile(Category = "BuildConfiguration")]
+		public bool bUseSparseSetStrictMode { get; set; } = false;
+
 		/// <summary>
 		/// By default TSparseSet is used as the underlying TSet container. Use this flag to change that to use TCompactSet instead.
 		/// </summary>
diff --git a/Engine/Source/Programs/UnrealBuildTool/Configuration/UEBuildTarget.cs b/Engine/Source/Programs/UnrealBuildTool/Configuration/UEBuildTarget.cs
--- a/Engine/Source/Programs/UnrealBuildTool/Configuration/UEBuildTarget.cs
+++ b/Engine/Source/Programs/UnrealBuildTool/Configuration/UEBuildTarget.cs
@@ -6566,6 +6566,15 @@ namespace UnrealBuildTool
 				GlobalCompileEnvironment.Definitions.Add("UE_DEFINE_FORCEINLINE_HINT_TO_INLINE=0");
 			}
 
+			if (Rules.bUseSparseSetStrictMode)
+			{
+				GlobalCompileEnvironment.Definitions.Add("UE_USE_SPARSE_SET_STRICT_MODE=1");
+			}
+			else
+			{
+				GlobalCompileEnvironment.Definitions.Add("UE_USE_SPARSE_SET_STRICT_MODE=0");
+			}
+
 			if (Rules.bUseCompactSetAsDefault)
 			{
 				GlobalCompileEnvironment.Definitions.Add("UE_USE_COMPACT_SET_AS_DEFAULT=1");
diff --git a/Engine/Source/Runtime/Core/Public/Containers/CompactSet.h.inl b/Engine/Source/Runtime/Core/Public/Containers/CompactSet.h.inl
--- a/Engine/Source/Runtime/Core/Public/Containers/CompactSet.h.inl
+++ b/Engine/Source/Runtime/Core/Public/Containers/CompactSet.h.inl
@@ -105,11 +105,17 @@ private:
 
 		[[nodiscard]] UE_FORCEINLINE_HINT ElementItType & operator*() const
 		{
+#if DO_CHECK
+			checkf(LastRemoved == INDEX_NONE, TEXT("Accessing iterator after it was removed from the set and before it was incremented to the next valid elements."));
+#endif
 			return Set[FSetElementId::FromInteger(Index)];
 		}
 
 		[[nodiscard]] UE_FORCEINLINE_HINT ElementItType* operator->() const
 		{
+#if DO_CHECK
+			checkf(LastRemoved == INDEX_NONE, TEXT("Accessing iterator after it was removed from the set and before it was incremented to the next valid elements."));
+#endif
 			return &Set[FSetElementId::FromInteger(Index)];
 		}
 
@@ -118,6 +124,7 @@ private:
 #if DO_CHECK
 			const int32 SetNum = Set.Num();
 			checkf(SetNum >= InitialNum, TEXT("Sets/Maps should never have elements removed during iteration outside of Iterator.RemoveCurrent(). InitialNum %d CurrentNum %d"), InitialNum, SetNum);
+			LastRemoved = INDEX_NONE;
 #endif
 			++Index;
 			return *this;
@@ -151,6 +158,7 @@ private:
 		SizeType Index;
 #if DO_CHECK
 		SizeType InitialNum;
+		SizeType LastRemoved = INDEX_NONE;
 #endif
 	};
 
@@ -191,6 +199,7 @@ private:
 			// Note: Adding new elements is safe, however they will be guaranteed to be missed by the current iteration
 #if DO_CHECK
 			checkf(SetNum >= InitialNum, TEXT("Sets/Maps should never have elements removed during iteration outside of Iterator.RemoveCurrent(). InitialNum %d CurrentNum %d"), InitialNum, SetNum);
+			LastRemoved = INDEX_NONE;
 #endif
 
 			Index = NextIndex;
@@ -242,6 +251,7 @@ private:
 
 #if DO_CHECK
 		SizeType InitialNum;
+		SizeType LastRemoved = INDEX_NONE;
 #endif
 	};
 
@@ -262,11 +272,14 @@ public:
 		/** Removes the current element from the set. */
 		inline void RemoveCurrent()
 		{
-			this->Set.Remove(this->GetId());
-			--this->Index;
 #if DO_CHECK
+			checkf(this->LastRemoved == INDEX_NONE, TEXT("RemoveCurrent was already called for the current iteration index of %d"), this->LastRemoved);
+			this->LastRemoved = this->Index;
 			--this->InitialNum;
 #endif
+
+			this->Set.Remove(this->GetId());
+			--this->Index;
 		}
 	};
 
@@ -286,15 +299,18 @@ public:
 		/** Removes the current element from the set. */
 		inline void RemoveCurrent()
 		{
-			this->Set.Remove(this->GetId());
 #if DO_CHECK
+			checkf(this->LastRemoved == INDEX_NONE, TEXT("RemoveCurrent was already called for the current iteration index of %d"), this->LastRemoved);
+			this->LastRemoved = this->Index;
 			--this->InitialNum;
 #endif
 
+			this->Set.Remove(this->GetId());
+
 			// If the next element was the last in the set then it will get remapped to the current index, so fix that up
-			if (this->NextIndex == TBaseKeyIterator<false>::Set.Num())
+			if (this->NextIndex == this->Set.Num())
 			{
-				this->NextIndex = TBaseKeyIterator<false>::Index;
+				this->NextIndex = this->Index;
 			}
 
 			this->Index = INDEX_NONE;
diff --git a/Engine/Source/Runtime/Core/Public/Containers/SparseSet.h.inl b/Engine/Source/Runtime/Core/Public/Containers/SparseSet.h.inl
--- a/Engine/Source/Runtime/Core/Public/Containers/SparseSet.h.inl
+++ b/Engine/Source/Runtime/Core/Public/Containers/SparseSet.h.inl
@@ -29,6 +29,12 @@ class TScriptSparseSet;
 
 #define TSETPRIVATEFRIEND UE_JOIN(UE_TSPARSE_SET, PrivateFriend)
 
+#if !defined(UE_USE_SPARSE_SET_STRICT_MODE)
+#define UE_USE_SPARSE_SET_STRICT_MODE 0
+#endif
+
+#define DO_SPARSE_SET_ITERATOR_CHECKS DO_CHECK && UE_USE_SPARSE_SET_STRICT_MODE
+
 /**
  * A set with an optional KeyFuncs parameters for customizing how the elements are compared and searched.
  * E.g. You can specify a mapping from elements to keys if you want to find elements by specifying a subset of
@@ -257,7 +263,7 @@ public:
 	inline void Shrink()
 	{
 		Elements.Shrink();
-		ConditionalRehash(Elements.Num(), EAllowShrinking::Yes);
+		Relax();
 	}
 
 	/** Compacts the allocated elements into a contiguous range. */
@@ -310,7 +316,6 @@ public:
 	}
 
 	/** Relaxes the set's hash to a size strictly bounded by the number of elements in the set. */
-	UE_DEPRECATED(5.8, "Removing support for relax, use Shrink instead for better memory savings.")
 	UE_FORCEINLINE_HINT void Relax()
 	{
 		ConditionalRehash(Elements.Num(), EAllowShrinking::Yes);
@@ -1430,9 +1435,8 @@ private:
 	class TBaseIterator
 	{
 	private:
-		friend class UE_TSPARSE_SET;
-
-		typedef std::conditional_t<bConst,const ElementType,ElementType> ItElementType;
+		using SetType = std::conditional_t<bConst, const UE_TSPARSE_SET, UE_TSPARSE_SET>;
+		using ItElementType = std::conditional_t<bConst,const ElementType,ElementType>;
 
 	public:
 		typedef std::conditional_t<
@@ -1441,14 +1445,33 @@ private:
 			std::conditional_t<bRangedFor, typename ElementArrayType::TRangedForIterator,      typename ElementArrayType::TIterator     >
 		> ElementItType;
 
-		[[nodiscard]] UE_FORCEINLINE_HINT explicit TBaseIterator(const ElementItType& InElementIt)
+		[[nodiscard]] UE_FORCEINLINE_HINT explicit TBaseIterator(SetType& InSet)
+			: ElementIt(InSet.Elements.begin())
+			, Set(InSet)
+#if DO_SPARSE_SET_ITERATOR_CHECKS
+			, InitialNum(InSet.Num())
+#endif
+		{
+		}
+
+		[[nodiscard]] UE_FORCEINLINE_HINT explicit TBaseIterator(SetType& InSet, const ElementItType& InElementIt)
 			: ElementIt(InElementIt)
+			, Set(InSet)
+#if DO_SPARSE_SET_ITERATOR_CHECKS
+			, InitialNum(InSet.Num())
+#endif
 		{
 		}
 
 		/** Advances the iterator to the next element. */
 		UE_FORCEINLINE_HINT TBaseIterator& operator++()
 		{
+#if DO_SPARSE_SET_ITERATOR_CHECKS
+			const int32 SetNum = Set.Num();
+			ensureMsgf(SetNum >= InitialNum, TEXT("Sets/Maps should never have elements removed during iteration outside of Iterator.RemoveCurrent(). InitialNum %d CurrentNum %d"), InitialNum, SetNum);
+			LastRemoved = INDEX_NONE;
+#endif
+
 			++ElementIt;
 			return *this;
 		}
@@ -1471,17 +1494,30 @@ private:
 		}
 		[[nodiscard]] UE_FORCEINLINE_HINT ItElementType* operator->() const
 		{
+#if DO_SPARSE_SET_ITERATOR_CHECKS
+			checkf(LastRemoved == INDEX_NONE, TEXT("Accessing iterator after it was removed from the set and before it was incremented to the next valid elements."));
+#endif
 			return &ElementIt->Value;
 		}
 		[[nodiscard]] UE_FORCEINLINE_HINT ItElementType& operator*() const
 		{
+#if DO_SPARSE_SET_ITERATOR_CHECKS
+			checkf(LastRemoved == INDEX_NONE, TEXT("Accessing iterator after it was removed from the set and before it was incremented to the next valid elements."));
+#endif
 			return ElementIt->Value;
 		}
 
 		[[nodiscard]] UE_FORCEINLINE_HINT bool operator!=(const TBaseIterator& Rhs) const { return ElementIt != Rhs.ElementIt; }
 		[[nodiscard]] UE_FORCEINLINE_HINT bool operator==(const TBaseIterator& Rhs) const { return ElementIt == Rhs.ElementIt; }
 
+	protected:
 		ElementItType ElementIt;
+		SetType& Set;
+
+#if DO_SPARSE_SET_ITERATOR_CHECKS
+		SizeType InitialNum;
+		SizeType LastRemoved = INDEX_NONE;
+#endif
 	};
 
 	/** The base type of whole set iterators. */
@@ -1489,9 +1525,9 @@ private:
 	class TBaseKeyIterator
 	{
 	private:
-		typedef std::conditional_t<bConst, const UE_TSPARSE_SET, UE_TSPARSE_SET> SetType;
-		typedef std::conditional_t<bConst,const ElementType,ElementType> ItElementType;
-		typedef typename TTypeTraits<typename KeyFuncs::KeyType>::ConstPointerType ReferenceOrValueType;
+		using SetType = std::conditional_t<bConst, const UE_TSPARSE_SET, UE_TSPARSE_SET>;
+		using ItElementType = std::conditional_t<bConst,const ElementType,ElementType>;
+		using ReferenceOrValueType = typename TTypeTraits<typename KeyFuncs::KeyType>::ConstPointerType;
 
 	public:
 		using KeyArgumentType =
@@ -1506,6 +1542,9 @@ private:
 			: Set  (InSet)
 			, Key  (InKey) //-V1041
 			, Index(INDEX_NONE)
+#if DO_SPARSE_SET_ITERATOR_CHECKS
+			, InitialNum(InSet.Num())
+#endif
 		{
 			if (Set.HashSize)
 			{
@@ -1521,6 +1560,13 @@ private:
 		/** Advances the iterator to the next element. */
 		inline TBaseKeyIterator& operator++()
 		{
+			// Note: Adding new elements is safe, however they will be guaranteed to be missed by the current iteration
+#if DO_SPARSE_SET_ITERATOR_CHECKS
+			const int32 SetNum = Set.Num();
+			ensureMsgf(SetNum >= InitialNum, TEXT("Sets/Maps should never have elements removed during iteration outside of Iterator.RemoveCurrent(). InitialNum %d CurrentNum %d"), InitialNum, SetNum);
+			LastRemoved = INDEX_NONE;
+#endif
+
 			Index = NextIndex;
 
 			while (Index != INDEX_NONE)
@@ -1568,61 +1614,45 @@ private:
 		ReferenceOrValueType Key;
 		SizeType Index;
 		SizeType NextIndex;
+
+#if DO_SPARSE_SET_ITERATOR_CHECKS
+		SizeType InitialNum;
+		SizeType LastRemoved = INDEX_NONE;
+#endif
 	};
 
 public:
 
 	/** Used to iterate over the elements of a const UE_TSPARSE_SET. */
-	class TConstIterator : public TBaseIterator<true>
-	{
-		friend class UE_TSPARSE_SET;
-
-	public:
-		[[nodiscard]] UE_FORCEINLINE_HINT explicit TConstIterator(const UE_TSPARSE_SET& InSet)
-			: TBaseIterator<true>(InSet.Elements.begin())
-		{
-		}
-	};
+	using TConstIterator = TBaseIterator<true>;
 
 	/** Used to iterate over the elements of a UE_TSPARSE_SET. */
 	class TIterator : public TBaseIterator<false>
 	{
 		friend class UE_TSPARSE_SET;
+		using Super = TBaseIterator<false>;
 
 	public:
-		[[nodiscard]] inline explicit TIterator(UE_TSPARSE_SET& InSet)
-			: TBaseIterator<false>(InSet.Elements.begin())
-			, Set                 (InSet)
-		{
-		}
+		using Super::Super;
 
 		/** Removes the current element from the set. */
 		UE_FORCEINLINE_HINT void RemoveCurrent()
 		{
-			Set.RemoveByIndex(TBaseIterator<false>::ElementIt.GetIndex());
-		}
+#if DO_SPARSE_SET_ITERATOR_CHECKS
+			ensureMsgf(this->LastRemoved == INDEX_NONE, TEXT("RemoveCurrent was already called for the current iteration index of %d"), this->LastRemoved);
+			this->LastRemoved = this->ElementIt.GetIndex();
+			--this->InitialNum;
+#endif
 
-	private:
-		UE_TSPARSE_SET& Set;
+			this->Set.RemoveByIndex(this->ElementIt.GetIndex());
+		}
 	};
 
 	using TRangedForConstIterator = TBaseIterator<true, true>;
 	using TRangedForIterator      = TBaseIterator<false, true>;
 
 	/** Used to iterate over the elements of a const UE_TSPARSE_SET. */
-	class TConstKeyIterator : public TBaseKeyIterator<true>
-	{
-	private:
-		using Super = TBaseKeyIterator<true>;
-
-	public:
-		using KeyArgumentType = typename Super::KeyArgumentType;
-
-		[[nodiscard]] UE_FORCEINLINE_HINT explicit TConstKeyIterator(const UE_TSPARSE_SET& InSet, KeyArgumentType InKey)
-			: Super(InSet, InKey)
-		{
-		}
-	};
+	using TConstKeyIterator = TBaseKeyIterator<true>;
 
 	/** Used to iterate over the elements of a UE_TSPARSE_SET. */
 	class TKeyIterator : public TBaseKeyIterator<false>
@@ -1633,16 +1663,19 @@ public:
 	public:
 		using KeyArgumentType = typename Super::KeyArgumentType;
 
-		[[nodiscard]] UE_FORCEINLINE_HINT explicit TKeyIterator(UE_TSPARSE_SET& InSet, KeyArgumentType InKey)
-			: Super(InSet, InKey)
-		{
-		}
+		using Super::Super;
 
 		/** Removes the current element from the set. */
 		inline void RemoveCurrent()
 		{
-			this->Set.RemoveByIndex(TBaseKeyIterator<false>::Index);
-			TBaseKeyIterator<false>::Index = INDEX_NONE;
+#if DO_SPARSE_SET_ITERATOR_CHECKS
+			ensureMsgf(this->LastRemoved == INDEX_NONE, TEXT("RemoveCurrent was already called for the current iteration index of %d"), this->LastRemoved);
+			this->LastRemoved = this->Index;
+			--this->InitialNum;
+#endif
+
+			this->Set.RemoveByIndex(this->Index);
+			this->Index = INDEX_NONE;
 		}
 	};
 
@@ -1667,19 +1700,19 @@ public:
 	 */
 	[[nodiscard]] UE_FORCEINLINE_HINT TRangedForIterator begin()
 	{
-		return TRangedForIterator(Elements.begin());
+		return TRangedForIterator(*this);
 	}
 	[[nodiscard]] UE_FORCEINLINE_HINT TRangedForConstIterator begin() const
 	{
-		return TRangedForConstIterator(Elements.begin());
+		return TRangedForConstIterator(*this);
 	}
 	[[nodiscard]] UE_FORCEINLINE_HINT TRangedForIterator end()
 	{
-		return TRangedForIterator(Elements.end());
+		return TRangedForIterator(*this, Elements.end());
 	}
 	[[nodiscard]] UE_FORCEINLINE_HINT TRangedForConstIterator end() const
 	{
-		return TRangedForConstIterator(Elements.end());
+		return TRangedForConstIterator(*this, Elements.end());
 	}
 
 	// Maps are deliberately prevented from being hashed or compared, because this would hide potentially major performance problems behind default operations.
