Added && to Value and operator* in TOptional

#Synced-CL 48956957

diff --git a/Engine/Source/Runtime/Core/Public/Misc/Optional.h b/Engine/Source/Runtime/Core/Public/Misc/Optional.h
--- a/Engine/Source/Runtime/Core/Public/Misc/Optional.h
+++ b/Engine/Source/Runtime/Core/Public/Misc/Optional.h
@@ -449,15 +449,19 @@ public:
 	}
 
 	/** @return The optional value; undefined when IsSet() returns false. */
-	[[nodiscard]] constexpr OptionalType& GetValue()
+	[[nodiscard]] constexpr OptionalType& GetValue() &
 	{
 		checkf(IsSet(), TEXT("It is an error to call GetValue() on an unset TOptional. Please either check IsSet() or use Get(DefaultValue) instead."));
 		return this->TypedValue;
 	}
-	[[nodiscard]] UE_FORCEINLINE_HINT constexpr const OptionalType& GetValue() const
+	[[nodiscard]] UE_FORCEINLINE_HINT constexpr const OptionalType& GetValue() const&
 	{
 		return const_cast<TOptional*>(this)->GetValue();
 	}
+	[[nodiscard]] UE_FORCEINLINE_HINT constexpr OptionalType&& GetValue() &&
+	{
+		return MoveTempIfPossible(GetValue());
+	}
 
 	[[nodiscard]] constexpr OptionalType* operator->()
 	{
@@ -468,14 +472,18 @@ public:
 		return const_cast<TOptional*>(this)->operator->();
 	}
 
-	[[nodiscard]] constexpr OptionalType& operator*()
+	[[nodiscard]] constexpr OptionalType& operator*() &
 	{
 		return GetValue();
 	}
-	[[nodiscard]] UE_FORCEINLINE_HINT constexpr const OptionalType& operator*() const
+	[[nodiscard]] UE_FORCEINLINE_HINT constexpr const OptionalType& operator*() const&
 	{
 		return const_cast<TOptional*>(this)->operator*();
 	}
+	[[nodiscard]] UE_FORCEINLINE_HINT constexpr OptionalType&& operator*() &&
+	{
+		return MoveTempIfPossible(operator*());
+	}
 
 	/** @return The optional value when set; DefaultValue otherwise. */
 	[[nodiscard]] constexpr const OptionalType& Get(const OptionalType& DefaultValue UE_LIFETIMEBOUND) const UE_LIFETIMEBOUND
