[RaceDetector]
- Fix the timeout value that would never be applied when using r.RaceDetector.Activate [Timeout]

#Synced-CL 48817740

diff --git a/Engine/Source/Runtime/Core/Private/Sanitizer/RaceDetector.cpp b/Engine/Source/Runtime/Core/Private/Sanitizer/RaceDetector.cpp
--- a/Engine/Source/Runtime/Core/Private/Sanitizer/RaceDetector.cpp
+++ b/Engine/Source/Runtime/Core/Private/Sanitizer/RaceDetector.cpp
@@ -1571,12 +1571,6 @@ namespace UE::Sanitizer::RaceDetector {
 		GRaceDetectorFilterDuplicates = bEnable;
 	}
 
-	INSTRUMENTATION_FUNCTION_ATTRIBUTES void ToggleRaceDetectionUntil(UE::FTimeout Timeout)
-	{
-		RaceDetectorTimeout = Timeout;
-		ToggleRaceDetection(true);
-	}
-	
 	INSTRUMENTATION_FUNCTION_ATTRIBUTES void ToggleFilterOtherThreads(bool bEnable)
 	{
 		GRaceDetectorFilterOtherThreads = bEnable;
@@ -1587,7 +1581,7 @@ namespace UE::Sanitizer::RaceDetector {
 		return RaceDetectorActive != 0;
 	}
 
-	INSTRUMENTATION_FUNCTION_ATTRIBUTES void ToggleRaceDetection(bool bEnable, float GracePeriodSeconds)
+	INSTRUMENTATION_FUNCTION_ATTRIBUTES void ToggleRaceDetectionInternal(bool bEnable, float GracePeriodSeconds, UE::FTimeout Timeout = UE::FTimeout::Never())
 	{
 		int32 Result = 0;
 
@@ -1608,7 +1602,7 @@ namespace UE::Sanitizer::RaceDetector {
 			{
 				Context.bAlwaysReport = true;
 
-				RaceDetectorTimeout = UE::FTimeout::Never();
+				RaceDetectorTimeout = Timeout;
 				if (RaceDetectorActive++ == 0)
 				{
 					GracePeriodTimeout = UE::FTimeout::Never();
@@ -1656,10 +1650,20 @@ namespace UE::Sanitizer::RaceDetector {
 		// detected when the stack memory used as synchronization during logging starts being reused.
 		if (Result)
 		{
-			UE_LOG(LogRaceDetector, Log, TEXT("Race detector has been toggled %s"),	Result == 1 ? TEXT("on") : TEXT("off"))
+			UE_LOG(LogRaceDetector, Log, TEXT("Race detector has been toggled %s"), Result == 1 ? TEXT("on") : TEXT("off"))
 		}
 	}
 
+	INSTRUMENTATION_FUNCTION_ATTRIBUTES void ToggleRaceDetectionUntil(UE::FTimeout Timeout)
+	{
+		ToggleRaceDetectionInternal(true, .0f, Timeout);
+	}
+
+	INSTRUMENTATION_FUNCTION_ATTRIBUTES void ToggleRaceDetection(bool bEnable, float GracePeriodSeconds)
+	{
+		ToggleRaceDetectionInternal(bEnable, GracePeriodSeconds);
+	}
+
 	INSTRUMENTATION_FUNCTION_ATTRIBUTES void ToggleThreadDetailedLog(bool bEnabled)
 	{
 		FContext& Context = EnsureCurrentContext();
