Used MPSC queue instead of MPMC queue for log output

#Synced-CL 48956957

diff --git a/Engine/Source/Runtime/Core/Private/Misc/OutputDeviceRedirector.cpp b/Engine/Source/Runtime/Core/Private/Misc/OutputDeviceRedirector.cpp
--- a/Engine/Source/Runtime/Core/Private/Misc/OutputDeviceRedirector.cpp
+++ b/Engine/Source/Runtime/Core/Private/Misc/OutputDeviceRedirector.cpp
@@ -5,7 +5,7 @@
 #include "Async/EventCount.h"
 #include "AutoRTFM.h"
 #include "Containers/BitArray.h"
-#include "Containers/ConsumeAllMpmcQueue.h"
+#include "Containers/MpscQueue.h"
 #include "Experimental/ConcurrentLinearAllocator.h"
 #include "HAL/PlatformProcess.h"
 #include "HAL/PlatformTime.h"
@@ -140,7 +140,8 @@ struct FOutputDeviceRedirectorState
 	uint8 OutputDevicesLockPadding[CalculateRedirectorCacheLinePadding(sizeof(OutputDevicesLock) + sizeof(OutputDevicesLockState))]{};
 
 	/** A queue of items logged by non-primary threads. */
-	TConsumeAllMpmcQueue<FOutputDeviceItem, FOutputDeviceLinearAllocator> BufferedItems;
+	TMpscQueue<FOutputDeviceItem, FOutputDeviceLinearAllocator> BufferedItems;
+	std::atomic<uint32> BufferedItemCount = 0;
 
 	/** Array of output devices to redirect to from the primary thread. */
 	TArray<FOutputDevice*> BufferedOutputDevices;
@@ -500,7 +501,7 @@ void FOutputDeviceRedirectorState::ThreadLoop()
 		{
 			break;
 		}
-		while (!BufferedItems.IsEmpty() && IsPrimaryThread(ThreadId))
+		while (BufferedItemCount.load(std::memory_order_relaxed) && IsPrimaryThread(ThreadId))
 		{
 			if (FOutputDevicesPrimaryScopeLock Lock(*this); Lock.IsLocked())
 			{
@@ -536,7 +537,7 @@ void FOutputDeviceRedirectorState::ConditionalFlushBufferedItems(EOutputDeviceRe
 UE_AUTORTFM_NOAUTORTFM
 void FOutputDeviceRedirectorState::FlushBufferedItems()
 {
-	if (BufferedItems.IsEmpty())
+	if (!BufferedItemCount.load(std::memory_order_relaxed))
 	{
 		return;
 	}
@@ -567,60 +568,60 @@ template <typename VisitorType>
 UE_AUTORTFM_NOAUTORTFM
 void FOutputDeviceRedirectorState::FlushBufferedItemsTo(VisitorType&& Visitor)
 {
-	if (BufferedItems.IsEmpty())
+	if (!BufferedItemCount.load(std::memory_order_relaxed))
 	{
 		return;
 	}
 
 	TRACE_CPUPROFILER_EVENT_SCOPE(FOutputDeviceRedirector::FlushBufferedItems);
 
-	BufferedItems.ConsumeAllFifo([this, &Visitor](FOutputDeviceItem&& Item)
+	for (uint32 Count = BufferedItemCount.exchange(0, std::memory_order_relaxed); Count; --Count)
 	{
-		if (!Visitor.bIsPanicThread && HasPanicThread())
+		for (;;)
 		{
-			// Re-queue items to be consumed by the panic thread.
-			// ProduceItem within ConsumeAllFifo is okay because the items passed to the consumer
-			// are from a moved list that is no longer associated with the queue.
-			// ConsumeAllFifo will return even though new items are in the queue.
-			BufferedItems.ProduceItem(MoveTemp(Item));
-		}
-		else
-		{
-			Visit([this, &Visitor](auto&& Value)
+			if (!Visitor.bIsPanicThread && HasPanicThread())
 			{
-				using ValueType = std::decay_t<decltype(Value)>;
-				if constexpr (std::is_same_v<ValueType, FOutputDeviceLine>)
-				{
-					const FOutputDeviceLine& Line = Value;
-					Visitor(Line);
-				}
-				else if constexpr (std::is_same_v<ValueType, FLogRecord>)
-				{
-					const FLogRecord& Record = Value;
-					Visitor(Record);
-				}
-				else if constexpr (std::is_same_v<ValueType, uint64>)
+				// Re-queue items to be consumed by the panic thread.
+				BufferedItemCount.fetch_add(Count, std::memory_order_relaxed);
+				return;
+			}
+			
+			// Dequeue can return NullOpt even though it's not empty due to races of setting Prev->Next	in Enqueue		
+			if (TOptional<FOutputDeviceItem> Item = BufferedItems.Dequeue())
+			{
+				Visit([this, &Visitor]<class ValueType>(const ValueType& Value)
 				{
-					const uint64 FenceIndex = Value;
-					for (uint64 LastFenceIndex = FenceCompleteIndex.load(std::memory_order_relaxed);;)
+					if constexpr (!std::is_same_v<ValueType, uint64>)
 					{
-						if (FenceIndex <= LastFenceIndex || FenceCompleteIndex.compare_exchange_weak(LastFenceIndex, FenceIndex, std::memory_order_release))
+						Visitor(Value);
+					}
+					else
+					{
+						const uint64 FenceIndex = Value;
+						for (uint64 LastFenceIndex = FenceCompleteIndex.load(std::memory_order_relaxed);;)
 						{
-							break;
+							if (FenceIndex <= LastFenceIndex || FenceCompleteIndex.compare_exchange_weak(LastFenceIndex, FenceIndex, std::memory_order_release))
+							{
+								break;
+							}
 						}
+						FenceCompleteEvent.Notify();
 					}
-					FenceCompleteEvent.Notify();
-				}
-			}, Item.Value);
+				}, Item->Value);
+				break;
+			}
+			
+			FPlatformProcess::Yield();
 		}
-	});
+	}
 }
 
 UE_AUTORTFM_NOAUTORTFM
 uint64 FOutputDeviceRedirectorState::CreateFence()
 {
 	uint64 Index = FenceCreateIndex.fetch_add(1, std::memory_order_relaxed) + 1;
-	if (BufferedItems.ProduceItem(Index) == EConsumeAllMpmcQueueResult::WasEmpty)
+	BufferedItems.Enqueue(Index);
+	if (!BufferedItemCount.fetch_add(1, std::memory_order_relaxed))
 	{
 		ThreadWakeEvent.Notify();
 	}
@@ -670,8 +671,7 @@ void FOutputDeviceRedirectorState::EnableBacklog(bool bEnable)
 			struct FVisitor
 			{
 				FOutputDeviceRedirectorState& State;
-				const uint32 ThreadId;
-				const bool bIsPanicThread = State.IsPanicThread(ThreadId);
+				const bool bIsPanicThread;
 
 				inline void operator()(const FOutputDeviceLine& Line)
 				{
@@ -685,7 +685,7 @@ void FOutputDeviceRedirectorState::EnableBacklog(bool bEnable)
 					State.BacklogLines.Emplace(*Text, Record.GetCategory(), Record.GetVerbosity(), FPlatformTime::Seconds() - GStartTime);
 				}
 			};
-			FlushBufferedItemsTo(FVisitor{*this, ThreadId});
+			FlushBufferedItemsTo(FVisitor{*this, IsPanicThread(ThreadId)});
 		}
 	}
 }
@@ -870,7 +870,8 @@ void FOutputDeviceRedirector::SerializeRecord(const UE::FLogRecord& Record)
 	}
 
 	// Queue the record to serialize to buffered output devices from the primary thread.
-	if (State->BufferedItems.ProduceItem(Record) == UE::EConsumeAllMpmcQueueResult::WasEmpty)
+	State->BufferedItems.Enqueue(Record);
+	if (!State->BufferedItemCount.fetch_add(1, std::memory_order_relaxed))
 	{
 		State->ThreadWakeEvent.Notify();
 	}
@@ -932,7 +933,8 @@ void FOutputDeviceRedirector::Serialize(const TCHAR* const Data, const ELogVerbo
 	}
 
 	// Queue the line to serialize to buffered output devices from the primary thread.
-	if (State->BufferedItems.ProduceItem(Data, Category, Verbosity, RealTime) == UE::EConsumeAllMpmcQueueResult::WasEmpty)
+	State->BufferedItems.Enqueue(Data, Category, Verbosity, RealTime);
+	if (!State->BufferedItemCount.fetch_add(1, std::memory_order_relaxed))
 	{
 		State->ThreadWakeEvent.Notify();
 	}
