Fixed comment for TMemberFunctionPtrOuter_T.

diff --git a/Engine/Source/Runtime/Core/Public/Traits/MemberFunctionPtrOuter.h b/Engine/Source/Runtime/Core/Public/Traits/MemberFunctionPtrOuter.h
--- a/Engine/Source/Runtime/Core/Public/Traits/MemberFunctionPtrOuter.h
+++ b/Engine/Source/Runtime/Core/Public/Traits/MemberFunctionPtrOuter.h
@@ -5,7 +5,7 @@
 /**
  * Type trait which yields the type of the class given a pointer to a member function of that class, e.g.:
  *
- * TMemberFunctionPtrOuter_T<decltype(&FVector::DotProduct)>::Type == FVector
+ * TMemberFunctionPtrOuter_T<decltype(&FVector::Dot)> == FVector
  */
 template <typename T>
 struct TMemberFunctionPtrOuter;
