Added move constructor and asserts to IntrusiveDoubleLinkedList.

#Synced-CL 48968555

diff --git a/Engine/Source/Runtime/Core/Public/Containers/IntrusiveDoubleLinkedList.h b/Engine/Source/Runtime/Core/Public/Containers/IntrusiveDoubleLinkedList.h
--- a/Engine/Source/Runtime/Core/Public/Containers/IntrusiveDoubleLinkedList.h
+++ b/Engine/Source/Runtime/Core/Public/Containers/IntrusiveDoubleLinkedList.h
@@ -185,6 +185,11 @@ public:
 	UE_FORCEINLINE_HINT TIntrusiveDoubleLinkedList(const TIntrusiveDoubleLinkedList&) = delete;
 	UE_FORCEINLINE_HINT TIntrusiveDoubleLinkedList& operator=(const TIntrusiveDoubleLinkedList&) = delete;
 
+	TIntrusiveDoubleLinkedList(TIntrusiveDoubleLinkedList&& Other)
+	{
+		AddHead(Forward<TIntrusiveDoubleLinkedList&&>(Other));
+	}
+
 	/**
 	 * Fast empty that clears this list *without* changing the links in any elements.
 	 * 
@@ -210,6 +215,7 @@ public:
 
 	UE_FORCEINLINE_HINT void AddHead(ElementType* Element)
 	{
+		check(Element && !Element->IsInList());
 		static_cast<NodeType*>(Element)->InsertAfter(GetSentinel());
 	}
 
@@ -227,6 +233,7 @@ public:
 
 	UE_FORCEINLINE_HINT void AddTail(ElementType* Element)
 	{
+		check(Element && !Element->IsInList());
 		static_cast<NodeType*>(Element)->InsertBefore(GetSentinel());
 	}
 
