[RaceDetector]
- Precompute the TLS position

#Synced-CL 48817740

diff --git a/Engine/Source/Runtime/Core/Private/Sanitizer/RaceDetector.cpp b/Engine/Source/Runtime/Core/Private/Sanitizer/RaceDetector.cpp
--- a/Engine/Source/Runtime/Core/Private/Sanitizer/RaceDetector.cpp
+++ b/Engine/Source/Runtime/Core/Private/Sanitizer/RaceDetector.cpp
@@ -129,7 +129,8 @@ namespace UE::Sanitizer::RaceDetector {
 	UPTRINT ShadowMemorySize = 0;
 	UPTRINT ShadowClockBase = 0;
 	UPTRINT ShadowMemoryEnd = 0;
-	uint32  ContextTLSIndex  = UINT32_MAX;
+	uint32  ContextTLSIndex = UINT32_MAX;
+	uint64  ContextTLSValue = 0;
 	volatile uint32 GlobalEpoch = 0;
 	volatile int64 FSyncObject::ObjectCount = 0;
 	volatile int64 FSyncObjectBank::ObjectCount = 0;
@@ -175,14 +176,14 @@ namespace UE::Sanitizer::RaceDetector {
 		return !!RaceDetectorActive && !Context.WinInstrumentationDepth && !Context.InstrumentationDepth && !bIsResettingShadow;
 	}
 
-	FORCEINLINE INSTRUMENTATION_FUNCTION_ATTRIBUTES void* GetTlsValueFast(uint32 Index)
+	FORCEINLINE INSTRUMENTATION_FUNCTION_ATTRIBUTES void* GetTlsValueFast(uint32 Value)
 	{
-		return (void*)__readgsqword(0x1480lu + (unsigned long)Index * sizeof(void*));
+		return (void*)__readgsqword(Value);
 	}
 
 	FORCEINLINE INSTRUMENTATION_FUNCTION_ATTRIBUTES FContext* GetThreadContext()
 	{
-		return (FContext*)GetTlsValueFast(ContextTLSIndex);
+		return (FContext*)GetTlsValueFast(ContextTLSValue);
 	}
 
 	INSTRUMENTATION_FUNCTION_ATTRIBUTES uint64 GetSyncObjectMemoryUsage()
@@ -542,7 +543,7 @@ namespace UE::Sanitizer::RaceDetector {
 
 	INSTRUMENTATION_FUNCTION_ATTRIBUTES void FHistoryChunk::InitStack()
 	{
-		FContext* ContextPtr = (FContext*)GetTlsValueFast(ContextTLSIndex);
+		FContext* ContextPtr = (FContext*)GetTlsValueFast(ContextTLSValue);
 		if (FContext::IsValid(ContextPtr))
 		{
 			StartClock = ContextPtr->CurrentClock();
@@ -1492,7 +1493,7 @@ namespace UE::Sanitizer::RaceDetector {
 
 	INSTRUMENTATION_FUNCTION_ATTRIBUTES FContext& EnsureCurrentContext()
 	{
-		check(ContextTLSIndex != UINT32_MAX);
+		RACEDETECTOR_CHECK(ContextTLSValue != 0 && ContextTLSIndex != UINT32_MAX);
 		FContext* Context = GetThreadContext();
 		if (Context == nullptr)
 		{
@@ -1501,7 +1502,7 @@ namespace UE::Sanitizer::RaceDetector {
 			// Refcount owned by the thread itself.
 			Context->AddRef();
 			Platform::SetTlsValue(ContextTLSIndex, Context);
-			if (GetTlsValueFast(ContextTLSIndex) != Context)
+			if (GetTlsValueFast(ContextTLSValue) != Context)
 			{
 				UE_LOG(LogRaceDetector, Fatal, TEXT("GetTlsValueFast implementation is invalid"));
 			}
@@ -1696,7 +1697,7 @@ using namespace UE::Sanitizer::RaceDetector;
 // We don't need to test for bRaceDetectorActive or bIsResettingShadow here
 // since we unhook the instrumentation instead.
 #define BEGIN_HANDLE_INSTRUMENTATION(FuncDoIfNotInstrumenting) \
-	FContext* ContextPtr = (FContext*)GetTlsValueFast(ContextTLSIndex); \
+	FContext* ContextPtr = (FContext*)GetTlsValueFast(ContextTLSValue); \
 	if (!FContext::IsValid(ContextPtr) || (ContextPtr->WinInstrumentationDepth | ContextPtr->InstrumentationDepth) != 0) \
 	{ \
 		FuncDoIfNotInstrumenting; \
@@ -2302,6 +2303,7 @@ INSTRUMENTATION_FUNCTION_ATTRIBUTES void UnhookInstrumentation()
 INSTRUMENTATION_FUNCTION_ATTRIBUTES void InitializeBoot()
 {
 	ContextTLSIndex = Platform::AllocTlsSlot();
+	ContextTLSValue = 0x1480lu + (unsigned long)ContextTLSIndex * sizeof(void*);
 	EnsureCurrentContext();
 
 	// Avoid doing that while reporting a race as it may cause reentrancy issues because of tracing.
