[Core]
* Fixed TStringView::Find() and TStringView::Contains() to not crash (check) when testing with an empty string parameter. Now Find() returns valid index and Contains() returns true for empty strings.
* Fixed UE::String::FindFirst(), UE::String::FindLast() to not crash (check) when Search string is empty. Now for empty strings, FindFirst returns 0 and FindLast returns Len().
* Fixed UE::String::FindFirstOfAny() and UE::String::FindLastOfAny() when one of the strings in array is empty to not crash (check) anymore. Empty strings match at index 0 / Len() as FindFirst/FindLast.
* Added more tests in the String's FindTest (for FindFirst, FindLast, FindFirstOfAny, FindLastOfAny with empty strings).

#Synced-CL 48842043

diff --git a/Engine/Source/Runtime/Core/Private/String/Find.cpp b/Engine/Source/Runtime/Core/Private/String/Find.cpp
--- a/Engine/Source/Runtime/Core/Private/String/Find.cpp
+++ b/Engine/Source/Runtime/Core/Private/String/Find.cpp
@@ -22,7 +22,10 @@ namespace UE::String::Private
 template <typename CharType>
 static inline int32 FindFirst(TStringView<CharType> View, TStringView<CharType> Search, ESearchCase::Type SearchCase)
 {
-	check(!Search.IsEmpty());
+	if (Search.IsEmpty())
+	{
+		return 0;
+	}
 	const int32 SearchLen = Search.Len();
 	if (SearchLen == 1)
 	{
@@ -62,7 +65,10 @@ static inline int32 FindFirst(TStringView<CharType> View, TStringView<CharType>
 template <typename CharType>
 static inline int32 FindLast(TStringView<CharType> View, TStringView<CharType> Search, ESearchCase::Type SearchCase)
 {
-	check(!Search.IsEmpty());
+	if (Search.IsEmpty())
+	{
+		return View.Len();
+	}
 	const int32 SearchLen = Search.Len();
 	if (SearchLen == 1)
 	{
@@ -102,7 +108,6 @@ static inline int32 FindLast(TStringView<CharType> View, TStringView<CharType> S
 template <typename CharType>
 static inline int32 FindFirstOfAny(TStringView<CharType> View, TConstArrayView<TStringView<CharType>> Search, ESearchCase::Type SearchCase, int32* OutMatchIndex)
 {
-	check(Algo::NoneOf(Search, &TStringView<CharType>::IsEmpty));
 	switch (Search.Num())
 	{
 	case 0:
@@ -132,17 +137,33 @@ static inline int32 FindFirstOfAny(TStringView<CharType> View, TConstArrayView<T
 
 	const CharType* const ViewBegin = View.GetData();
 	const int32 ViewLen = View.Len();
-	for (int32 ViewIndex = 0; ViewIndex != ViewLen; ++ViewIndex)
+	if (ViewLen > 0)
 	{
-		const TStringView<CharType> RemainingView(ViewBegin + ViewIndex, ViewLen - ViewIndex);
-		auto MatchPattern = [&RemainingView, SearchCase](const TStringView<CharType>& Pattern) { return RemainingView.StartsWith(Pattern, SearchCase); };
-		if (const TStringView<CharType>* Match = Algo::FindByPredicate(Search, MatchPattern))
+		for (int32 ViewIndex = 0; ViewIndex != ViewLen; ++ViewIndex)
+		{
+			const TStringView<CharType> RemainingView(ViewBegin + ViewIndex, ViewLen - ViewIndex);
+			auto MatchPattern = [&RemainingView, SearchCase](const TStringView<CharType>& Pattern) { return RemainingView.StartsWith(Pattern, SearchCase); };
+			if (const TStringView<CharType>* Match = Algo::FindByPredicate(Search, MatchPattern))
+			{
+				if (OutMatchIndex)
+				{
+					*OutMatchIndex = UE_PTRDIFF_TO_INT32(Match - Search.GetData());
+				}
+				return ViewIndex;
+			}
+		}
+	}
+	else
+	{
+		// An empty view can match an empty pattern.
+		auto MatchEmptyPattern = [](const TStringView<CharType>& Pattern) { return Pattern.IsEmpty(); };
+		if (const TStringView<CharType>* Match = Algo::FindByPredicate(Search, MatchEmptyPattern))
 		{
 			if (OutMatchIndex)
 			{
-				*OutMatchIndex = UE_PTRDIFF_TO_INT32(Match - Search.GetData());
+				*OutMatchIndex = UE_PTRDIFF_TO_INT32(Match - Search.GetData());;
 			}
-			return ViewIndex;
+			return 0;
 		}
 	}
 	return INDEX_NONE;
@@ -151,7 +172,6 @@ static inline int32 FindFirstOfAny(TStringView<CharType> View, TConstArrayView<T
 template <typename CharType>
 static inline int32 FindLastOfAny(TStringView<CharType> View, TConstArrayView<TStringView<CharType>> Search, ESearchCase::Type SearchCase, int32* OutMatchIndex)
 {
-	check(Algo::NoneOf(Search, &TStringView<CharType>::IsEmpty));
 	switch (Search.Num())
 	{
 	case 0:
@@ -181,21 +201,33 @@ static inline int32 FindLastOfAny(TStringView<CharType> View, TConstArrayView<TS
 
 	const CharType* const ViewBegin = View.GetData();
 	const int32 ViewLen = View.Len();
-	if (ViewLen == 0)
+	if (ViewLen > 0)
 	{
-		return INDEX_NONE;
+		for (int32 ViewIndex = ViewLen; ViewIndex >= 0; --ViewIndex)
+		{
+			const TStringView<CharType> RemainingView(ViewBegin + ViewIndex, ViewLen - ViewIndex);
+			auto MatchPattern = [&RemainingView, SearchCase](const TStringView<CharType>& Pattern) { return RemainingView.StartsWith(Pattern, SearchCase); };
+			if (const TStringView<CharType>* Match = Algo::FindByPredicate(Search, MatchPattern))
+			{
+				if (OutMatchIndex)
+				{
+					*OutMatchIndex = UE_PTRDIFF_TO_INT32(Match - Search.GetData());
+				}
+				return ViewIndex;
+			}
+		}
 	}
-	for (int32 ViewIndex = ViewLen - 1; ViewIndex >= 0; --ViewIndex)
+	else
 	{
-		const TStringView<CharType> RemainingView(ViewBegin + ViewIndex, ViewLen - ViewIndex);
-		auto MatchPattern = [&RemainingView, SearchCase](const TStringView<CharType>& Pattern) { return RemainingView.StartsWith(Pattern, SearchCase); };
-		if (const TStringView<CharType>* Match = Algo::FindByPredicate(Search, MatchPattern))
+		// An empty view can match an empty pattern.
+		auto MatchEmptyPattern = [](const TStringView<CharType>& Pattern) { return Pattern.IsEmpty(); };
+		if (const TStringView<CharType>* Match = Algo::FindByPredicate(Search, MatchEmptyPattern))
 		{
 			if (OutMatchIndex)
 			{
-				*OutMatchIndex = UE_PTRDIFF_TO_INT32(Match - Search.GetData());
+				*OutMatchIndex = UE_PTRDIFF_TO_INT32(Match - Search.GetData());;
 			}
-			return ViewIndex;
+			return 0;
 		}
 	}
 	return INDEX_NONE;
diff --git a/Engine/Source/Runtime/Core/Tests/String/FindTest.cpp b/Engine/Source/Runtime/Core/Tests/String/FindTest.cpp
--- a/Engine/Source/Runtime/Core/Tests/String/FindTest.cpp
+++ b/Engine/Source/Runtime/Core/Tests/String/FindTest.cpp
@@ -33,6 +33,8 @@ TEST_CASE_NAMED(FStringFindTest, "System::Core::String::Find", "[Core][String][S
 		CHECK(FindFirst(TEXT("AbCABCAbCABD"), TEXT("Bd"), ESearchCase::IgnoreCase) == 10);
 
 		CHECK(FindFirst(TEXT(""), TEXT("A")) == INDEX_NONE);
+		CHECK(FindFirst(TEXT(""), TEXT("")) == 0);
+		CHECK(FindFirst(TEXT("A"), TEXT("")) == 0);
 		CHECK(FindFirst(TEXT("A"), TEXT("A")) == 0);
 		CHECK(FindFirst(TEXT("A"), TEXT("A"), ESearchCase::IgnoreCase) == 0);
 		CHECK(FindFirst(TEXT("ABC"), TEXT("ABC")) == 0);
@@ -43,6 +45,8 @@ TEST_CASE_NAMED(FStringFindTest, "System::Core::String::Find", "[Core][String][S
 
 		CHECK(FindFirst(FStringView(nullptr, 0), TEXT("SearchTerm")) == INDEX_NONE); //-V575
 		CHECK(FindFirst(FStringView(), TEXT("SearchTerm")) == INDEX_NONE);
+		CHECK(FindFirst(FStringView(), FStringView()) == 0);
+		CHECK(FindFirst(TEXT("A"), FStringView()) == 0);
 	}
 
 	SECTION("FindLast")
@@ -63,6 +67,8 @@ TEST_CASE_NAMED(FStringFindTest, "System::Core::String::Find", "[Core][String][S
 		CHECK(FindLast(TEXT("AbCABCAbCABC"), TEXT("Bc"), ESearchCase::IgnoreCase) == 10);
 
 		CHECK(FindLast(TEXT(""), TEXT("A")) == INDEX_NONE);
+		CHECK(FindLast(TEXT(""), TEXT("")) == 0);
+		CHECK(FindLast(TEXT("A"), TEXT("")) == 1);
 		CHECK(FindLast(TEXT("A"), TEXT("A")) == 0);
 		CHECK(FindLast(TEXT("A"), TEXT("A"), ESearchCase::IgnoreCase) == 0);
 		CHECK(FindLast(TEXT("ABC"), TEXT("ABC")) == 0);
@@ -73,6 +79,8 @@ TEST_CASE_NAMED(FStringFindTest, "System::Core::String::Find", "[Core][String][S
 
 		CHECK(FindLast(FStringView(nullptr, 0), TEXT("SearchTerm")) == INDEX_NONE); //-V575
 		CHECK(FindLast(FStringView(), TEXT("SearchTerm")) == INDEX_NONE);
+		CHECK(FindLast(FStringView(), FStringView()) == 0);
+		CHECK(FindLast(TEXT("A"), FStringView()) == 1);
 	}
 
 	SECTION("FindFirstOfAny")
@@ -95,6 +103,7 @@ TEST_CASE_NAMED(FStringFindTest, "System::Core::String::Find", "[Core][String][S
 		CHECK(FindFirstOfAny(TEXT("AbCABCAbCABD"), {TEXT("BD"), TEXT("CABB")}) == 10);
 		CHECK(FindFirstOfAny(TEXT("AbCABCAbCABD"), {TEXT("Bd"), TEXT("CABB")}, ESearchCase::IgnoreCase) == 10);
 
+		CHECK(FindFirstOfAny(TEXT(""), {TEXT("AB")}) == INDEX_NONE);
 		CHECK(FindFirstOfAny(TEXT(""), {TEXT("A"), TEXT("B")}) == INDEX_NONE);
 		CHECK(FindFirstOfAny(TEXT("A"), {TEXT("A"), TEXT("B")}) == 0);
 		CHECK(FindFirstOfAny(TEXT("A"), {TEXT("A"), TEXT("B")}, ESearchCase::IgnoreCase) == 0);
@@ -114,6 +123,30 @@ TEST_CASE_NAMED(FStringFindTest, "System::Core::String::Find", "[Core][String][S
 		CHECK(MatchIndex == 1);
 		FindFirstOfAny(TEXT("AbCABCAbCABC"), {TEXT("DEF"), TEXT("ABCD"), TEXT("Cab")}, ESearchCase::IgnoreCase, &MatchIndex);
 		CHECK(MatchIndex == 2);
+
+		CHECK(FindFirstOfAny(TEXT(""), {TEXT("")}, ESearchCase::IgnoreCase, &MatchIndex) == 0);
+		CHECK(MatchIndex == 0);
+		CHECK(FindFirstOfAny(TEXT(""), {TEXT("AB"), TEXT("")}, ESearchCase::IgnoreCase, &MatchIndex) == 0);
+		CHECK(MatchIndex == 1);
+
+		CHECK(FindFirstOfAny(FStringView(), {FStringView()}, ESearchCase::IgnoreCase, &MatchIndex) == 0);
+		CHECK(MatchIndex == 0);
+		CHECK(FindFirstOfAny(FStringView(), {TEXTVIEW("AB"), FStringView()}, ESearchCase::IgnoreCase, &MatchIndex) == 0);
+		CHECK(MatchIndex == 1);
+
+		CHECK(FindFirstOfAny(TEXT("ABC"), {TEXT("")}, ESearchCase::IgnoreCase, &MatchIndex) == 0);
+		CHECK(MatchIndex == 0);
+		CHECK(FindFirstOfAny(TEXT("--abCDABCD--"), {TEXT("XY"), TEXT("")}, ESearchCase::IgnoreCase, &MatchIndex) == 0);
+		CHECK(MatchIndex == 1);
+		CHECK(FindFirstOfAny(TEXT("--abCDABCD--"), {TEXT("AB"), TEXT("CD"), TEXT("")}, ESearchCase::IgnoreCase, &MatchIndex) == 0);
+		CHECK(MatchIndex == 2);
+
+		CHECK(FindFirstOfAny(TEXT("ABC"), {FStringView()}, ESearchCase::IgnoreCase, &MatchIndex) == 0);
+		CHECK(MatchIndex == 0);
+		CHECK(FindFirstOfAny(TEXTVIEW("--abCDABCD--"), {TEXTVIEW("XY"), FStringView()}, ESearchCase::IgnoreCase, &MatchIndex) == 0);
+		CHECK(MatchIndex == 1);
+		CHECK(FindFirstOfAny(TEXTVIEW("--abCDABCD--"), {TEXTVIEW("AB"), TEXTVIEW("CD"), FStringView()}, ESearchCase::IgnoreCase, &MatchIndex) == 0);
+		CHECK(MatchIndex == 2);
 	}
 
 	SECTION("FindLastOfAny")
@@ -136,6 +169,7 @@ TEST_CASE_NAMED(FStringFindTest, "System::Core::String::Find", "[Core][String][S
 		CHECK(FindLastOfAny(TEXT("AbCABCAbCABD"), {TEXT("BD"), TEXT("CABB")}) == 10);
 		CHECK(FindLastOfAny(TEXT("AbCABCAbCABD"), {TEXT("Bd"), TEXT("CABB")}, ESearchCase::IgnoreCase) == 10);
 
+		CHECK(FindLastOfAny(TEXT(""), {TEXT("AB")}) == INDEX_NONE);
 		CHECK(FindLastOfAny(TEXT(""), {TEXT("A"), TEXT("B")}) == INDEX_NONE);
 		CHECK(FindLastOfAny(TEXT("A"), {TEXT("A"), TEXT("B")}) == 0);
 		CHECK(FindLastOfAny(TEXT("A"), {TEXT("A"), TEXT("B")}, ESearchCase::IgnoreCase) == 0);
@@ -145,8 +179,8 @@ TEST_CASE_NAMED(FStringFindTest, "System::Core::String::Find", "[Core][String][S
 
 		CHECK(FindLastOfAny(ANSITEXTVIEW("AbCABCAbCABC"), {"CABc", "ABC"}) == 9);
 
-		CHECK(FindLastOfAny(FStringView(nullptr, 0), { TEXT("ABC"), TEXT("ABD") }) == INDEX_NONE); //-V575
-		CHECK(FindLastOfAny(FStringView(), { TEXT("ABC"), TEXT("ABD") }) == INDEX_NONE);
+		CHECK(FindLastOfAny(FStringView(nullptr, 0), {TEXT("ABC"), TEXT("ABD")}) == INDEX_NONE); //-V575
+		CHECK(FindLastOfAny(FStringView(), {TEXT("ABC"), TEXT("ABD")}) == INDEX_NONE);
 
 		int32 MatchIndex = INDEX_NONE;
 		FindLastOfAny(TEXT("AbCABCAbCABC"), {TEXT("AbC")}, ESearchCase::CaseSensitive, &MatchIndex);
@@ -155,6 +189,30 @@ TEST_CASE_NAMED(FStringFindTest, "System::Core::String::Find", "[Core][String][S
 		CHECK(MatchIndex == 1);
 		FindLastOfAny(TEXT("AbCABCAbCABC"), {TEXT("DEF"), TEXT("CAB"), TEXT("AbC")}, ESearchCase::IgnoreCase, &MatchIndex);
 		CHECK(MatchIndex == 2);
+
+		CHECK(FindLastOfAny(TEXT(""), {TEXT("")}, ESearchCase::IgnoreCase, &MatchIndex) == 0);
+		CHECK(MatchIndex == 0);
+		CHECK(FindLastOfAny(TEXT(""), {TEXT("AB"), TEXT("")}, ESearchCase::IgnoreCase, &MatchIndex) == 0);
+		CHECK(MatchIndex == 1);
+
+		CHECK(FindLastOfAny(FStringView(), {FStringView()}, ESearchCase::IgnoreCase, &MatchIndex) == 0);
+		CHECK(MatchIndex == 0);
+		CHECK(FindLastOfAny(FStringView(), {TEXTVIEW("AB"), FStringView()}, ESearchCase::IgnoreCase, &MatchIndex) == 0);
+		CHECK(MatchIndex == 1);
+
+		CHECK(FindLastOfAny(TEXT("ABC"), {TEXT("")}, ESearchCase::IgnoreCase, &MatchIndex) == 3);
+		CHECK(MatchIndex == 0);
+		CHECK(FindLastOfAny(TEXT("--ABCDABcd--"), {TEXT("XY"), TEXT("")}, ESearchCase::IgnoreCase, &MatchIndex) == 12);
+		CHECK(MatchIndex == 1);
+		CHECK(FindLastOfAny(TEXT("--ABCDABcd--"), {TEXT("AB"), TEXT("CD"), TEXT("")}, ESearchCase::IgnoreCase, &MatchIndex) == 12);
+		CHECK(MatchIndex == 2);
+
+		CHECK(FindLastOfAny(TEXT("ABC"), {FStringView()}, ESearchCase::IgnoreCase, &MatchIndex) == 3);
+		CHECK(MatchIndex == 0);
+		CHECK(FindLastOfAny(TEXT("--ABCDABcd--"), {TEXTVIEW("XY"), FStringView()}, ESearchCase::IgnoreCase, &MatchIndex) == 12);
+		CHECK(MatchIndex == 1);
+		CHECK(FindLastOfAny(TEXT("--ABCDABcd--"), {TEXTVIEW("AB"), TEXTVIEW("CD"), FStringView()}, ESearchCase::IgnoreCase, &MatchIndex) == 12);
+		CHECK(MatchIndex == 2);
 	}
 
 	SECTION("FindFirstChar")
@@ -213,8 +271,8 @@ TEST_CASE_NAMED(FStringFindTest, "System::Core::String::Find", "[Core][String][S
 
 		CHECK(FindFirstOfAnyChar("AbCABCAbcABC", {'c', 'B'}) == 4);
 
-		CHECK(FindFirstOfAnyChar(FStringView(nullptr, 0), { TEXT('A'), TEXT('B') }) == INDEX_NONE); //-V575
-		CHECK(FindFirstOfAnyChar(FStringView(), { TEXT('A'), TEXT('B') }) == INDEX_NONE);
+		CHECK(FindFirstOfAnyChar(FStringView(nullptr, 0), {TEXT('A'), TEXT('B')}) == INDEX_NONE); //-V575
+		CHECK(FindFirstOfAnyChar(FStringView(), {TEXT('A'), TEXT('B')}) == INDEX_NONE);
 
 		int32 MatchIndex = INDEX_NONE;
 		FindFirstOfAnyChar(TEXT("AbCABCAbCABC"), {TEXT('B')}, ESearchCase::CaseSensitive, &MatchIndex);
@@ -241,8 +299,8 @@ TEST_CASE_NAMED(FStringFindTest, "System::Core::String::Find", "[Core][String][S
 
 		CHECK(FindLastOfAnyChar("AbCABCAbcABC", {'c', 'B'}) == 10);
 
-		CHECK(FindLastOfAnyChar(FStringView(nullptr, 0), { TEXT('A'), TEXT('B') }) == INDEX_NONE); //-V575
-		CHECK(FindLastOfAnyChar(FStringView(), { TEXT('A'), TEXT('B') }) == INDEX_NONE);
+		CHECK(FindLastOfAnyChar(FStringView(nullptr, 0), {TEXT('A'), TEXT('B')}) == INDEX_NONE); //-V575
+		CHECK(FindLastOfAnyChar(FStringView(), {TEXT('A'), TEXT('B')}) == INDEX_NONE);
 
 		int32 MatchIndex = INDEX_NONE;
 		FindLastOfAnyChar(TEXT("AbCABCAbCABC"), {TEXT('A')}, ESearchCase::CaseSensitive, &MatchIndex);
