[Trace Services] Compress callstacks (XOR+Shift)

#rnx
#Synced-CL 49033424

diff --git a/Engine/Source/Developer/TraceServices/Private/Analyzers/CallstacksAnalysis.cpp b/Engine/Source/Developer/TraceServices/Private/Analyzers/CallstacksAnalysis.cpp
--- a/Engine/Source/Developer/TraceServices/Private/Analyzers/CallstacksAnalysis.cpp
+++ b/Engine/Source/Developer/TraceServices/Private/Analyzers/CallstacksAnalysis.cpp
@@ -2,14 +2,22 @@
 
 #include "CallstacksAnalysis.h"
 
+#include "Containers/StaticArray.h"
 #include "HAL/LowLevelMemTracker.h"
 #include "Model/CallstacksProvider.h"
+#include "ProfilingDebugging/CallstackTrace.h"
 #include "TraceServices/Model/AnalysisSession.h"
 #include "Common/Utils.h"
 
 namespace TraceServices
 {
 
+enum ERoutes
+{
+	RouteId_Callstack,
+	RouteId_CallstackXORAndRLE
+};
+	
 ////////////////////////////////////////////////////////////////////////////////
 FCallstacksAnalyzer::FCallstacksAnalyzer(IAnalysisSession& InSession, FCallstacksProvider* InProvider)
 	: Session(InSession)
@@ -23,6 +31,7 @@ void FCallstacksAnalyzer::OnAnalysisBegin(const FOnAnalysisContext& Context)
 {
 	FInterfaceBuilder& Builder = Context.InterfaceBuilder;
 	Builder.RouteEvent(RouteId_Callstack, "Memory", "CallstackSpec");
+	Builder.RouteEvent(RouteId_CallstackXORAndRLE, "Memory", "CallstackSpecXORAndRLE");
 }
 
 ////////////////////////////////////////////////////////////////////////////////
@@ -33,6 +42,7 @@ bool FCallstacksAnalyzer::OnEvent(uint16 RouteId, EStyle Style, const FOnEventCo
 	switch (RouteId)
 	{
 		case RouteId_Callstack:
+		{
 			const TArrayReader<uint64>& Frames = Context.EventData.GetArray<uint64>("Frames");
 			uint8 NumFrames = (uint8)FMath::Min(255u, Frames.Num());
 			if (const uint32 Id = Context.EventData.GetValue<uint32>("CallstackId"))
@@ -52,8 +62,25 @@ bool FCallstacksAnalyzer::OnEvent(uint16 RouteId, EStyle Style, const FOnEventCo
 				}
 				Provider->AddCallstackWithHash(Hash, Frames.GetData(), NumFrames);
 			}
-			break;
+		}
+		break;
+
+		case RouteId_CallstackXORAndRLE:
+		{
+			const TArrayReader<uint8>& CompressedFrames = Context.EventData.GetArray<uint8>("CompressedFrames");
+			TStaticArray<uint64, UE_CALLSTACK_TRACE_MAX_FRAMES> Frames;
+			const uint32 FrameCount = FCallstackXORAndRLE::Uncompress(
+				MakeArrayView(CompressedFrames.GetData(), CompressedFrames.Num()),
+				Frames
+			);
+			if (const uint32 Id = Context.EventData.GetValue<uint32>("CallstackId"))
+			{
+				Provider->AddCallstack(Id, Frames.GetData(), uint8(FrameCount));
+			}
+		}
+		break;
 	}
+
 	return true;
 }
 
diff --git a/Engine/Source/Developer/TraceServices/Private/Analyzers/CallstacksAnalysis.h b/Engine/Source/Developer/TraceServices/Private/Analyzers/CallstacksAnalysis.h
--- a/Engine/Source/Developer/TraceServices/Private/Analyzers/CallstacksAnalysis.h
+++ b/Engine/Source/Developer/TraceServices/Private/Analyzers/CallstacksAnalysis.h
@@ -21,10 +21,6 @@ public:
 	virtual bool OnEvent(uint16 RouteId, EStyle Style, const FOnEventContext& Context) override;
 
 private:
-	enum Routes
-	{
-		RouteId_Callstack,
-	};
 
 	IAnalysisSession& Session;
 	FCallstacksProvider* Provider;
diff --git a/Engine/Source/Runtime/Core/Private/ProfilingDebugging/CallstackTrace.cpp b/Engine/Source/Runtime/Core/Private/ProfilingDebugging/CallstackTrace.cpp
--- a/Engine/Source/Runtime/Core/Private/ProfilingDebugging/CallstackTrace.cpp
+++ b/Engine/Source/Runtime/Core/Private/ProfilingDebugging/CallstackTrace.cpp
@@ -2,6 +2,7 @@
 
 #include "ProfilingDebugging/CallstackTrace.h"
 #include "CallstackTracePrivate.h"
+#include "Containers/StaticArray.h"
 
 #if UE_CALLSTACK_TRACE_ENABLED
 
@@ -13,6 +14,7 @@ void	CallstackTrace_InitializeInternal();
 ////////////////////////////////////////////////////////////////////////////////
 UE_TRACE_CHANNEL_DEFINE(CallstackChannel)
 UE_TRACE_EVENT_DEFINE(Memory, CallstackSpec)
+UE_TRACE_EVENT_DEFINE(Memory, CallstackSpecXORAndRLE)
 
 uint32 GCallStackTracingTlsSlotIndex = FPlatformTLS::InvalidTlsSlot;
 
@@ -39,4 +41,104 @@ void CallstackTrace_Initialize()
 	}();
 }
 
+#define UE_CALLSTACK_TRACE_TESTS 0
+#if UE_CALLSTACK_TRACE_TESTS
+
+#include "Algo/Compare.h"
+
+void Test_Compression()
+{
+	auto Check = [](TConstArrayView<uint64> Input)
+	{
+		uint8 Compressed[UE_CALLSTACK_TRACE_MAX_FRAMES*sizeof(uint64)];
+		const uint32 CompressedSize = FCallstackXORAndRLE::Compress(Input, Compressed);
+		const auto CompressedView = MakeArrayView(Compressed, CompressedSize);
+
+		TStaticArray<uint64, UE_CALLSTACK_TRACE_MAX_FRAMES> Uncompressed;
+		const uint32 UncompressedCount = FCallstackXORAndRLE::Uncompress(CompressedView, Uncompressed);
+		auto UncompressedView = MakeArrayView(Uncompressed.GetData(), UncompressedCount);
+
+		check(Algo::Compare(UncompressedView, Input));
+	};
+
+	uint64 Frames[] = {
+		0x00007ff693ab7b09,
+		0x00007ff693ab7cd2,
+		0x00007ff693aa9fe1,
+		0x00007ff693aaa0c8,
+		0x00007ff693abb0d0,
+		0x00007ff693abb0d0,
+		0x00007ffd61917344,
+		0x00007ffd639226b1,
+	};
+	uint8 FramesExpected[] = { 0x0 };
+	Check(Frames);
+	
+	uint64 Zeroes[] = {
+		0x0, 0x0, 0x0, 0x0, 0x0
+	};
+	uint8 ZeroesExpected[] = {
+		64, 64, 64, 64, 64
+	};
+	Check(Zeroes);
+	
+	uint64 SingleOne[] = {
+		0x1,
+	};
+	uint8 SingleOneExpected [] = {
+		63, 0x1
+	};
+	Check(SingleOne);
+	
+	uint64 SingleNormal[] = {
+		0x2d4f763d,
+	};
+	uint8 SingleNormalExpected [] = {
+		34, 0x3d, 0x76, 0x4f, 0x2d
+	};
+	Check(SingleNormal);
+	
+	uint64 Mixed[] = {
+		0x000000002d4f829c,
+		0x00000001ed4f763d,
+		0xffffffffffffffff,
+		0x0000001800000000,
+		0x000000002d4f7680,
+		0x00000001ed4f7638,
+		0x0,
+		0x000000002d4f829c,
+	};
+	uint8 MixedExpected [] = {
+		34, 0x9c, 0x82, 0x4f, 0x2d, 
+		31, 0x3d, 0x76, 0x4f, 0xed, 0x01,
+		0,  0xff, 0xff, 0xff, 0xff, 0xff, 0xff, 0xff, 0xff,
+		27, 0x00, 0x00, 0x00, 0x00, 0x18,
+		34, 0x80, 0x76, 0x4f, 0x2d,
+		31, 0x38, 0x76, 0x4f, 0xed, 0x01,
+		64,
+		34, 0x9c, 0x82, 0x4f, 0x2d,
+	};
+	Check(Mixed);
+
+	uint64 Random[] = {
+		0x85EA1249376518D5,
+		0x006B579412698957,
+		0x7B505F48F7D4B719,		
+	};
+	uint8 RandomExpected[] = {
+		0, 0xd5, 0x18, 0x65, 0x37, 0x49, 0x12, 0xEA, 0x85,
+		9, 0x57, 0x89, 0x69, 0x12, 0x94, 0x57, 0x6B,
+		1, 0x19, 0xB7, 0xD4, 0xF7, 0x48, 0x5F, 0x50, 0x7B
+	};
+	Check(Random);
+}
+
+static bool bCompressionTestsWasRun = []()
+{
+	Test_Compression();
+	return true;
+}();
+
+#endif //UE_CALLSTACK_TRACE_TESTS
+
 #endif //UE_CALLSTACK_TRACE_ENABLED
\ No newline at end of file
diff --git a/Engine/Source/Runtime/Core/Private/ProfilingDebugging/CallstackTracePrivate.h b/Engine/Source/Runtime/Core/Private/ProfilingDebugging/CallstackTracePrivate.h
--- a/Engine/Source/Runtime/Core/Private/ProfilingDebugging/CallstackTracePrivate.h
+++ b/Engine/Source/Runtime/Core/Private/ProfilingDebugging/CallstackTracePrivate.h
@@ -25,6 +25,10 @@
 	#define UE_CALLSTACK_TRACE_RESERVE_GROWABLE 1
 #endif
 
+#if !defined(UE_CALLSTACK_TRACE_UNCOMPRESSED)
+	// If enabled, callstack addresses aggregate without any compression
+	#define UE_CALLSTACK_TRACE_UNCOMPRESSED 0
+#endif
 
 ////////////////////////////////////////////////////////////////////////////////
 UE_TRACE_CHANNEL_EXTERN(CallstackChannel)
@@ -34,6 +38,11 @@ UE_TRACE_EVENT_BEGIN_EXTERN(Memory, CallstackSpec, NoSync)
 	UE_TRACE_EVENT_FIELD(uint64[], Frames)
 UE_TRACE_EVENT_END()
 
+UE_TRACE_EVENT_BEGIN_EXTERN(Memory, CallstackSpecXORAndRLE, NoSync)
+	UE_TRACE_EVENT_FIELD(uint32, CallstackId)
+	UE_TRACE_EVENT_FIELD(uint8[], CompressedFrames)
+UE_TRACE_EVENT_END()
+
 ////////////////////////////////////////////////////////////////////////////////
 
 	class FCallstackTracer
@@ -51,7 +60,7 @@ UE_TRACE_EVENT_END()
 			, CallstackIdCounter(1) // 0 is reserved for "unknown callstack"
 		{
 		}
-		
+
 		uint32 AddCallstack(const FBacktraceEntry& Entry)
 		{
 			bool bAlreadyAdded = false;
@@ -76,11 +85,23 @@ UE_TRACE_EVENT_END()
 				{
 					return 0;
 				}
-#endif
+#endif //UE_CALLSTACK_TRACE_RESERVE_GROWABLE
 				KnownSet.Emplace(Hash, Id);
+#if UE_CALLSTACK_TRACE_UNCOMPRESSED
 				UE_TRACE_LOG(Memory, CallstackSpec, CallstackChannel)
 					<< CallstackSpec.CallstackId(Id)
 					<< CallstackSpec.Frames(Entry.Frames, Entry.FrameCount);
+#else
+				uint8 Compressed[UE_CALLSTACK_TRACE_MAX_FRAMES*sizeof(uint64)];
+				const uint32 CompressedSize = FCallstackXORAndRLE::Compress(
+					MakeArrayView(Entry.Frames, Entry.FrameCount),
+					Compressed
+				);
+				
+				UE_TRACE_LOG(Memory, CallstackSpecXORAndRLE, CallstackChannel)
+					<< CallstackSpecXORAndRLE.CallstackId(Id)
+					<< CallstackSpecXORAndRLE.CompressedFrames(Compressed, CompressedSize);
+#endif //UE_CALLSTACK_TRACE_UNCOMPRESSED
 			}
 
 			return Id;
@@ -114,4 +135,4 @@ UE_TRACE_EVENT_END()
 		std::atomic_uint32_t		CallstackIdCounter;
 	};
 
-#endif
+#endif //UE_CALLSTACK_TRACE_ENABLED
diff --git a/Engine/Source/Runtime/Core/Public/ProfilingDebugging/CallstackTrace.h b/Engine/Source/Runtime/Core/Public/ProfilingDebugging/CallstackTrace.h
--- a/Engine/Source/Runtime/Core/Public/ProfilingDebugging/CallstackTrace.h
+++ b/Engine/Source/Runtime/Core/Public/ProfilingDebugging/CallstackTrace.h
@@ -2,9 +2,7 @@
 #pragma once
 #include "HAL/Platform.h"
 #include "HAL/PlatformTLS.h"
-#include "Math/NumericLimits.h"
-#include "Misc/Build.h"
-#include "Trace/Config.h"
+#include "Containers/ArrayView.h"
 #include "Containers/StringFwd.h"
 
 ////////////////////////////////////////////////////////////////////////////////
@@ -118,3 +116,73 @@ inline bool  CallstackTrace_ResolveFullCallStack() { return true; }
 #define CALLSTACK_TRACE_LIMIT_CALLSTACKRESOLVE_SCOPE()
 
 #endif // UE_CALLSTACK_TRACE_ENABLED
+
+/**
+ * Simple run length compression of arrays of addresses (e.g stack frames).
+ * Needs to be available regardless if tracing is enabled to allow analysis to decompress.
+ */
+namespace FCallstackXORAndRLE
+{
+	inline uint32 Compress(TConstArrayView<uint64> Addresses, TArrayView<uint8> Compressed)
+	{
+		check(Addresses.Num() * 8 < Compressed.Num());
+		uint32 OutIdx = 0;
+		uint32 CompressedBytes = 0;
+		uint64 LastAddress = 0;
+		for(uint8 Index = 0; Index < Addresses.Num(); ++Index)
+		{
+			const uint64 Address = Addresses[Index];
+			uint64 XorAddress = LastAddress ^ Address;
+			LastAddress = Address;
+			const uint8 LeadingZeroes = uint8(FMath::CountLeadingZeros64(XorAddress));
+			Compressed[OutIdx++] = LeadingZeroes;
+			++CompressedBytes;
+			if (LeadingZeroes < 64)
+			{
+				int8 ValueBits = 64 - LeadingZeroes;
+				do
+				{
+					Compressed[OutIdx++] = uint8(0xff & XorAddress);
+					XorAddress >>= 8;
+					ValueBits -= 8;
+					++CompressedBytes;
+				} while(ValueBits > 0);  
+			}
+		}
+		return CompressedBytes;
+	}
+
+	inline uint32 Uncompress(TConstArrayView<uint8> Compressed, TArrayView<uint64> Addresses)
+	{
+		const int32 Bytes = Compressed.Num();
+		int32 ByteIndex = 0;
+		uint32 AddressCount = 0;
+		uint64 LastAddress = 0;
+		while (ByteIndex < Bytes)
+		{
+			uint64 OutAddress = 0;
+			const uint8 LeadingZeroes = Compressed[ByteIndex++];
+			if (LeadingZeroes < 64)
+			{
+				int8 ValueBits = 64 - LeadingZeroes;
+				uint8 BitsShift = 0;
+				check(ValueBits <= (Bytes - ByteIndex)*8);
+				while (ValueBits > 0)
+				{
+					OutAddress = (uint64(Compressed[ByteIndex++]) << BitsShift) | OutAddress;
+					ValueBits -= 8;
+					BitsShift += 8;
+				}
+			}
+			
+			if (AddressCount)
+			{
+				OutAddress = LastAddress ^ OutAddress;
+			}
+
+			LastAddress = OutAddress;
+			Addresses[AddressCount++] = OutAddress;
+		}
+		return AddressCount;
+	}
+}
