Put JNI class initializations into a seperate data section to make statics constexpr

#android
#Synced-CL 49001105

diff --git a/Engine/Source/Runtime/Core/Private/Android/AndroidJavaEnv.cpp b/Engine/Source/Runtime/Core/Private/Android/AndroidJavaEnv.cpp
--- a/Engine/Source/Runtime/Core/Private/Android/AndroidJavaEnv.cpp
+++ b/Engine/Source/Runtime/Core/Private/Android/AndroidJavaEnv.cpp
@@ -20,7 +20,7 @@ void AndroidJavaEnv::InitializeJavaEnv(JavaVM* VM, jint Version, jobject GlobalT
 
 void AndroidJavaEnv::InitializeJavaEnv()
 {
-	UE::Jni::Env.InitializeChain();
+	UE::Jni::Env.Initialize();
 }
 
 jobject AndroidJavaEnv::GetClassLoader()
@@ -324,6 +324,8 @@ namespace UE::Jni
 		return true;
 	}
 	
+    extern "C" void (FEnv::*const __start_UE_Jni_Initialize [[gnu::weak]])() const;
+    extern "C" void (FEnv::*const __stop_UE_Jni_Initialize [[gnu::weak]])() const;
 	void FEnv::Initialize() const
 	{
 		using namespace Java::Lang;
@@ -336,6 +338,11 @@ namespace UE::Jni
 		Initialize<FGameActivity, false>(Class<FGameActivity>);
 
 		ClassLoader = NewGlobalRef(*Get<FClassLoader* ()>(*FScopedJavaObject<TClass<FClass>*>{Value->FindClass(ClassName<FClass>.GetData())}, "getClassLoader")(Class<FGameActivity>));
+		
+		for (auto* Initialize = &__start_UE_Jni_Initialize; Initialize < &__stop_UE_Jni_Initialize; ++Initialize)
+		{
+			(this->**Initialize)();
+		}
 	}
 }
 #endif
\ No newline at end of file
diff --git a/Engine/Source/Runtime/Core/Public/Android/AndroidJavaEnv.h b/Engine/Source/Runtime/Core/Public/Android/AndroidJavaEnv.h
--- a/Engine/Source/Runtime/Core/Public/Android/AndroidJavaEnv.h
+++ b/Engine/Source/Runtime/Core/Public/Android/AndroidJavaEnv.h
@@ -279,12 +279,6 @@ namespace UE::Jni
 	template <typename T>
 	inline Java::Lang::TClass<T>* Class;
 
-	inline namespace Experimental
-	{
-		template <typename T>
-		struct TInitialize;
-	}
-
 	inline thread_local const struct FEnv
 	{
 		FEnv();
@@ -338,30 +332,39 @@ namespace UE::Jni
 		{
 			if constexpr (std::is_base_of_v<Java::Lang::FObject, T>)
 			{
-				checkf(!Class<T>, TEXT("FEnv::Initialize() called more than once for java class: %s"), StringCast<TCHAR>(T::ClassName.GetData()).Get());
-
-				FScopedJavaObject<Java::Lang::TClass<T>*> LocalClass = Find<T>();
-
-				if constexpr (bIsOptional)
+				if (!Class<T>)
 				{
-					if (!LocalClass)
+					FScopedJavaObject<Java::Lang::TClass<T>*> LocalClass = Find<T>();
+
+					if constexpr (bIsOptional)
 					{
-						return;
+						if (!LocalClass)
+						{
+							return;
+						}
 					}
+					else
+					{
+						checkf(LocalClass, TEXT("Failed to find java class: %s"), StringCast<TCHAR>(T::ClassName.GetData()).Get());
+					}
+				
+					Class<T> = NewGlobalRef(*LocalClass);
 				}
-				else
-				{
-					checkf(LocalClass, TEXT("Failed to find java class: %s"), StringCast<TCHAR>(T::ClassName.GetData()).Get());
-				}
-			
-				Class<T> = NewGlobalRef(*LocalClass);
+
 				Initialize<T, false>(Class<T>);
 			}
 			else
 			{
 				if (!Class<typename T::PartialClass>)
 				{
-					return;
+					FScopedJavaObject<Java::Lang::TClass<typename T::PartialClass>*> LocalClass = Find<typename T::PartialClass>();
+
+					if (!LocalClass)
+					{
+						return;
+					}
+
+					Class<typename T::PartialClass> = NewGlobalRef(*LocalClass);
 				}
 			
 				Initialize<T, bIsOptional>(Class<typename T::PartialClass>);
@@ -457,10 +460,6 @@ namespace UE::Jni
 		
 	private:
 		friend void AndroidJavaEnv::InitializeJavaEnv();
-		template <typename T>
-		friend struct TInitialize;
-
-		static void (*InitializeChain)();
 
 		JNIEnv* Value;
 		bool bAttached;
@@ -469,26 +468,14 @@ namespace UE::Jni
 	}
 	Env;
 
-	inline void (*FEnv::InitializeChain)() = []
+	template <typename T>
+	struct TInitialize
 	{
-		Env.Initialize();
+	private:
+		[[gnu::retain, gnu::used, gnu::section("UE_Jni_Initialize")]]
+		static constexpr void (FEnv::*Function)() const = &FEnv::Initialize<T>;
 	};
 
-	inline namespace Experimental
-	{
-		template <typename T>
-		struct TInitialize
-		{
-		private:
-			inline static void (*const Chain)() = std::exchange(FEnv::InitializeChain, []
-			{
-				Chain();
-
-				Env.Initialize<T>();
-			});
-		};
-	}
-
 	template <typename T, typename... TArgs>
 	struct TConstructor
 	{
