[Linux][Server] Make fork delay configurable and add a batch count for how many forks can be fulfilled per signal

#tests ran in debugger to validate logic, ran a client and server to ensure child process was functional

#Synced-CL 48838953

diff --git a/Engine/Source/Runtime/Core/Private/Unix/UnixPlatformProcess.cpp b/Engine/Source/Runtime/Core/Private/Unix/UnixPlatformProcess.cpp
--- a/Engine/Source/Runtime/Core/Private/Unix/UnixPlatformProcess.cpp
+++ b/Engine/Source/Runtime/Core/Private/Unix/UnixPlatformProcess.cpp
@@ -59,6 +59,20 @@ namespace UnixPlatformProcess
 		TEXT("When true the parent process will log memory stats before every fork."),
 		ECVF_Default);
 
+	static float GParentSpawnDelay = 0.125f;
+	static TAutoConsoleVariable<float> CVarForkParentSpawnDelay(
+		TEXT("server.WaitAndForkDelay"),
+		GParentSpawnDelay,
+		TEXT("Controls the minimum time between processing queued fork requests"),
+		ECVF_Default);
+
+	static int GParentForkBatchSize = 1;
+	static TAutoConsoleVariable<int> CVarForkBatchSize(
+		TEXT("server.ForkBatchSize"),
+		GParentForkBatchSize,
+		TEXT("Controls the number of queued forks that will be processed when a signal is received"),
+		ECVF_Default);
+
 
 	FString GetReadableTime(double InDeltaSeconds)
 	{
@@ -1454,7 +1468,6 @@ static FDelegateHandle OnEndFrameHandle;
 FGenericPlatformProcess::EWaitAndForkResult FUnixPlatformProcess::WaitAndFork()
 {
 #define WAIT_AND_FORK_QUEUE_LENGTH 4096
-#define WAIT_AND_FORK_CHILD_SPAWN_DELAY 0.125
 #ifndef WAIT_AND_FORK_PARENT_SHUTDOWN_EXIT_CODE
 	#define WAIT_AND_FORK_PARENT_SHUTDOWN_EXIT_CODE 0
 #endif
@@ -1566,173 +1579,199 @@ FGenericPlatformProcess::EWaitAndForkResult FUnixPlatformProcess::WaitAndFork()
 	};
 	TArray<FPidAndSignal> AllChildren;
 	AllChildren.Reserve(1024); // Sized to be big enough that it probably wont reallocate, but its not the end of the world if it does.
-	while (!IsEngineExitRequested())
+	bool ChildBreakOut = false;	// due to batch forking support, child processes can no longer use 'break' to exit the while loops
+	while ( !ChildBreakOut && !IsEngineExitRequested() )
 	{
 		BeginExitIfRequested();
 
+		float SpawnDelay = FMath::Max(0, UnixPlatformProcess::CVarForkParentSpawnDelay.GetValueOnAnyThread());
+		int BatchCount = FMath::Max(1,UnixPlatformProcess::CVarForkBatchSize.GetValueOnAnyThread());
 		FForkSignalData SignalData;
-		if (WaitAndForkSignalQueue.Dequeue(SignalData))
+		if ( !WaitAndForkSignalQueue.IsEmpty() )
 		{
-			uint16 Cookie = (SignalData.SignalValue >> 16) & 0xffff;
-			uint16 ChildIdx = SignalData.SignalValue & 0xffff;
-
-			UE_LOG(LogFork, Log, TEXT("[Parent] WaitAndForkSignalQueue has %d signals to execute. Sleeping for %f secs before processing request %04hx-%04hx "), WaitAndForkSignalQueue.Count(), WAIT_AND_FORK_CHILD_SPAWN_DELAY, Cookie, ChildIdx);
+			UE_LOG(LogFork, Log, TEXT("[Parent] WaitAndForkSignalQueue has %d signals to execute. Sleeping for %f secs before processing"), WaitAndForkSignalQueue.Count(), SpawnDelay);
 
 			// Sleep for a short while to avoid spamming new processes to the OS all at once
-			FPlatformProcess::Sleep(WAIT_AND_FORK_CHILD_SPAWN_DELAY);
-
-			FCoreDelegates::OnParentPreFork.Broadcast();
+			FPlatformProcess::Sleep(SpawnDelay);
 
-			UE_LOG(LogFork, Log, TEXT("[Parent] WaitAndFork processing child request %04hx-%04hx received at: %s"), Cookie, ChildIdx, *UnixPlatformProcess::GetReadableTime(FPlatformTime::Seconds() - SignalData.TimeSeconds));
-
-			if (UnixPlatformProcess::GLogMemoryStatsWhenForking)
-			{
-				FMemoryStatsHolder CurrentMasterMemStats(FPlatformMemory::GetStats());
-				UE_LOG(LogFork, Log, TEXT("[Parent] MemoryStats PreFork: AvailablePhysical: %.02fMiB (%+.02fMiB), PeakPhysical: %.02fMiB, PeakVirtual: %.02fMiB"),
-					CurrentMasterMemStats.AvailablePhysical, (CurrentMasterMemStats.AvailablePhysical - PreviousMasterMemStats->AvailablePhysical),
-					CurrentMasterMemStats.PeakUsedPhysical,
-					CurrentMasterMemStats.PeakUsedVirtual
-				);
-				PreviousMasterMemStats = CurrentMasterMemStats;
-			}
 			
-			// Make sure there are no pending messages in the log.
-			GLog->Flush();
-
-			// This should be the very last thing we do before forking for optimal interaction with GMalloc
-			FForkProcessHelper::LowLevelPreFork();
-			// ******** The fork happens here! ********
-			pid_t ChildPID = fork();
-			// ******** The fork happened! This is now either the parent process or the new child process ********
-
-			if (ChildPID == -1)
-			{
-				// Error handling
-				// We could return with an error code here, but instead it is somewhat better to log out an error and continue since this loop is supposed to be stable.
-				// Fork errors may include hitting process limits or other environmental factors so we will just report the issue since the environmental factor can be
-				// fixed while the process is still running.
-				int ErrNo = errno;
-				UE_LOG(LogFork, Error, TEXT("WaitAndFork failed to fork! fork() error:%d"), ErrNo);
-			}
-			else if (ChildPID == 0)
+			while ( !ChildBreakOut && !IsEngineExitRequested() && BatchCount > 0 )
 			{
-				// This should be the very first thing we do after forking for optimal interaction with GMalloc
-				FForkProcessHelper::LowLevelPostForkChild(ChildIdx);
+				// need to hold up the fork signal while we Dequeue since there's no way for the queue itself to guard against races due to a signal handler
+				sigset_t OldSigMask, NewSigMask;
+				sigfillset(&NewSigMask);	
+				int MaskResult = sigprocmask(SIG_BLOCK, &NewSigMask, &OldSigMask);
+				bool HasWaitingSignals = WaitAndForkSignalQueue.Dequeue(SignalData);
+				// unblock fork signal
+				if ( MaskResult == 0 )
+				{
+					// don't try to revert mask if it wasn't successfully changed above
+					sigprocmask(SIG_SETMASK, &OldSigMask, nullptr);
+				}
 
-				if (FPlatformMemory::HasForkPageProtectorEnabled())
+				if (!HasWaitingSignals)
 				{
-					UE::FForkPageProtector::OverrideGMalloc();
-					UE::FForkPageProtector::Get().ProtectMemoryRegions();
+					// all signals already processed, return to sleep and then wait for a new signal
+					break;
 				}
 
-				// Close the log state we inherited from our parent
-				GLog->TearDown();
+				uint16 Cookie = (SignalData.SignalValue >> 16) & 0xffff;
+				uint16 ChildIdx = SignalData.SignalValue & 0xffff;
+
+				FCoreDelegates::OnParentPreFork.Broadcast();
 
-				// Update GGameThreadId
-				FUnixTLS::ClearThreadIdTLS();
-				GGameThreadId = FUnixTLS::GetCurrentThreadId();
+				UE_LOG(LogFork, Log, TEXT("[Parent] WaitAndFork processing child request %04hx-%04hx received at: %s"), Cookie, ChildIdx, *UnixPlatformProcess::GetReadableTime(FPlatformTime::Seconds() - SignalData.TimeSeconds));
 
-				// Fix the command line, if a path to command line parameters was specified
-				if (!ChildParametersPath.IsEmpty() && ChildIdx > 0)
+				if (UnixPlatformProcess::GLogMemoryStatsWhenForking)
 				{
-					FString NewCmdLine;
-					const FString CmdLineFilename = ChildParametersPath / FString::Printf(TEXT("%d"), ChildIdx);
-					FFileHelper::LoadFileToString(NewCmdLine, *CmdLineFilename);
-					if (!NewCmdLine.IsEmpty())
-					{
-						FCommandLine::Set(*NewCmdLine);
-					}
-					else
-					{
-						UE_LOG(LogFork, Error, TEXT("[Child] WaitAndFork child %04hx-%04hx failed to set command line from: %s"), Cookie, ChildIdx, *CmdLineFilename);
-					}
+					FMemoryStatsHolder CurrentMasterMemStats(FPlatformMemory::GetStats());
+					UE_LOG(LogFork, Log, TEXT("[Parent] MemoryStats PreFork: AvailablePhysical: %.02fMiB (%+.02fMiB), PeakPhysical: %.02fMiB, PeakVirtual: %.02fMiB"),
+						CurrentMasterMemStats.AvailablePhysical, (CurrentMasterMemStats.AvailablePhysical - PreviousMasterMemStats->AvailablePhysical),
+						CurrentMasterMemStats.PeakUsedPhysical,
+						CurrentMasterMemStats.PeakUsedVirtual
+					);
+					PreviousMasterMemStats = CurrentMasterMemStats;
 				}
 
-				// Start up the log again
-				FPlatformOutputDevices::SetupOutputDevices();
-				GLog->SetCurrentThreadAsPrimaryThread();
-				
+				// Make sure there are no pending messages in the log.
+				GLog->Flush();
 
-				// Set the process name, if specified
-				if (ChildIdx > 0)
+				// This should be the very last thing we do before forking for optimal interaction with GMalloc
+				FForkProcessHelper::LowLevelPreFork();
+				// ******** The fork happens here! ********
+				pid_t ChildPID = fork();
+				// ******** The fork happened! This is now either the parent process or the new child process ********
+
+				if (ChildPID == -1)
 				{
-					if (prctl(PR_SET_NAME, TCHAR_TO_UTF8(*FString::Printf(TEXT("DS-%04hx-%04hx"), Cookie, ChildIdx))) != 0)
+					// Error handling
+					// We could return with an error code here, but instead it is somewhat better to log out an error and continue since this loop is supposed to be stable.
+					// Fork errors may include hitting process limits or other environmental factors so we will just report the issue since the environmental factor can be
+					// fixed while the process is still running.
+					int ErrNo = errno;
+					UE_LOG(LogFork, Error, TEXT("WaitAndFork failed to fork! fork() error:%d"), ErrNo);
+				}
+				else if (ChildPID == 0)
+				{
+					// This should be the very first thing we do after forking for optimal interaction with GMalloc
+					FForkProcessHelper::LowLevelPostForkChild(ChildIdx);
+
+					if (FPlatformMemory::HasForkPageProtectorEnabled())
 					{
-						int ErrNo = errno;
-						UE_LOG(LogFork, Fatal, TEXT("[Child] WaitAndFork failed to set process name with prctl! error:%d"), ErrNo);
+						UE::FForkPageProtector::OverrideGMalloc();
+						UE::FForkPageProtector::Get().ProtectMemoryRegions();
 					}
-				}
 
-				// Need to remove the child process from the responsibility of keeping track of a valid sibling process.
-				UnixCrashReporterTracker::RemoveValidCrashReportTickerForChildProcess();
+					// Close the log state we inherited from our parent
+					GLog->TearDown();
 
-				// If requested, now wait for a SIGRTMIN+2 signal before continuing execution.
-				if (bRequireResponseSignal && (ChildIdx == 0 || ChildIdx > NumForks))
-				{
-					static bool bResponseReceived = false;
-					struct sigaction Action;
-					FMemory::Memzero(Action);
-					sigfillset(&Action.sa_mask);
-					Action.sa_flags = SA_SIGINFO | SA_RESTART | SA_ONSTACK;
-					Action.sa_sigaction = [](int32 Signal, siginfo_t* Info, void* Context) {
-						if (Signal == WAIT_AND_FORK_RESPONSE_SIGNAL)
+					// Update GGameThreadId
+					FUnixTLS::ClearThreadIdTLS();
+					GGameThreadId = FUnixTLS::GetCurrentThreadId();
+
+					// Fix the command line, if a path to command line parameters was specified
+					if (!ChildParametersPath.IsEmpty() && ChildIdx > 0)
+					{
+						FString NewCmdLine;
+						const FString CmdLineFilename = ChildParametersPath / FString::Printf(TEXT("%d"), ChildIdx);
+						FFileHelper::LoadFileToString(NewCmdLine, *CmdLineFilename);
+						if (!NewCmdLine.IsEmpty())
+						{
+							FCommandLine::Set(*NewCmdLine);
+						}
+						else
 						{
-							bResponseReceived = true;
+							UE_LOG(LogFork, Error, TEXT("[Child] WaitAndFork child %04hx-%04hx failed to set command line from: %s"), Cookie, ChildIdx, *CmdLineFilename);
 						}
-					};
-					sigaction(WAIT_AND_FORK_RESPONSE_SIGNAL, &Action, nullptr);
+					}
 
-					const double StartChildWaitSeconds = FPlatformTime::Seconds();
+					// Start up the log again
+					FPlatformOutputDevices::SetupOutputDevices();
+					GLog->SetCurrentThreadAsPrimaryThread();
 
-					UE_LOG(LogFork, Log, TEXT("[Child] WaitAndFork child %04hx-%04hx waiting for signal %d to proceed."), Cookie, ChildIdx, WAIT_AND_FORK_RESPONSE_SIGNAL);
 
-					while (!IsEngineExitRequested() && !bResponseReceived)
+					// Set the process name, if specified
+					if (ChildIdx > 0)
 					{
-						UE_LOG(LogFork, Verbose, TEXT("[Child] WaitAndFork child %04hx-%04hx signal %d not received. Sleeping for 1sec."), Cookie, ChildIdx, WAIT_AND_FORK_RESPONSE_SIGNAL);
+						if (prctl(PR_SET_NAME, TCHAR_TO_UTF8(*FString::Printf(TEXT("DS-%04hx-%04hx"), Cookie, ChildIdx))) != 0)
+						{
+							int ErrNo = errno;
+							UE_LOG(LogFork, Fatal, TEXT("[Child] WaitAndFork failed to set process name with prctl! error:%d"), ErrNo);
+						}
+					}
 
-						FPlatformProcess::Sleep(1);
+					// Need to remove the child process from the responsibility of keeping track of a valid sibling process.
+					UnixCrashReporterTracker::RemoveValidCrashReportTickerForChildProcess();
 
-						// Check to see how long we've been waiting and if we should time out.
-						if ((WaitAndForkResponseTimeout > 0.0) && ((FPlatformTime::Seconds() - StartChildWaitSeconds) > WaitAndForkResponseTimeout))
+					// If requested, now wait for a SIGRTMIN+2 signal before continuing execution.
+					if (bRequireResponseSignal && (ChildIdx == 0 || ChildIdx > NumForks))
+					{
+						static bool bResponseReceived = false;
+						struct sigaction Action;
+						FMemory::Memzero(Action);
+						sigfillset(&Action.sa_mask);
+						Action.sa_flags = SA_SIGINFO | SA_RESTART | SA_ONSTACK;
+						Action.sa_sigaction = [](int32 Signal, siginfo_t* Info, void* Context) {
+							if (Signal == WAIT_AND_FORK_RESPONSE_SIGNAL)
+							{
+								bResponseReceived = true;
+							}
+							};
+						sigaction(WAIT_AND_FORK_RESPONSE_SIGNAL, &Action, nullptr);
+
+						const double StartChildWaitSeconds = FPlatformTime::Seconds();
+
+						UE_LOG(LogFork, Log, TEXT("[Child] WaitAndFork child %04hx-%04hx waiting for signal %d to proceed."), Cookie, ChildIdx, WAIT_AND_FORK_RESPONSE_SIGNAL);
+
+						while (!IsEngineExitRequested() && !bResponseReceived)
 						{
-							UE_LOG(LogFork, Error, TEXT("[Child] WaitAndFork child %04hx-%04hx has exceeded WAIT_AND_FORK_RESPONSE_SIGNAL timeout"), Cookie, ChildIdx);
-							FPlatformMisc::RequestExitWithStatus(true, WAIT_AND_FORK_RESPONSE_TIMEOUT_EXIT_CODE);
-							break;
+							UE_LOG(LogFork, Verbose, TEXT("[Child] WaitAndFork child %04hx-%04hx signal %d not received. Sleeping for 1sec."), Cookie, ChildIdx, WAIT_AND_FORK_RESPONSE_SIGNAL);
+
+							FPlatformProcess::Sleep(1);
+
+							// Check to see how long we've been waiting and if we should time out.
+							if ((WaitAndForkResponseTimeout > 0.0) && ((FPlatformTime::Seconds() - StartChildWaitSeconds) > WaitAndForkResponseTimeout))
+							{
+								UE_LOG(LogFork, Error, TEXT("[Child] WaitAndFork child %04hx-%04hx has exceeded WAIT_AND_FORK_RESPONSE_SIGNAL timeout"), Cookie, ChildIdx);
+								FPlatformMisc::RequestExitWithStatus(true, WAIT_AND_FORK_RESPONSE_TIMEOUT_EXIT_CODE);
+								break;
+							}
 						}
-					}
 
-					FMemory::Memzero(Action);
-					sigaction(WAIT_AND_FORK_RESPONSE_SIGNAL, &Action, nullptr);
-				}
+						FMemory::Memzero(Action);
+						sigaction(WAIT_AND_FORK_RESPONSE_SIGNAL, &Action, nullptr);
+					}
 
-				UE_LOG(LogFork, Log, TEXT("[Child] WaitAndFork child process %04hx-%04hx has started with pid %d."), Cookie, ChildIdx, GetCurrentProcessId());
-				FApp::PrintStartupLogMessages();
+					UE_LOG(LogFork, Log, TEXT("[Child] WaitAndFork child process %04hx-%04hx has started with pid %d."), Cookie, ChildIdx, GetCurrentProcessId());
+					FApp::PrintStartupLogMessages();
 
-				OnEndFrameHandle = FCoreDelegates::OnEndFrame.AddStatic(FUnixPlatformProcess::OnChildEndFramePostFork);
-				FCoreDelegates::OnPostFork.Broadcast(EForkProcessRole::Child);
+					OnEndFrameHandle = FCoreDelegates::OnEndFrame.AddStatic(FUnixPlatformProcess::OnChildEndFramePostFork);
+					FCoreDelegates::OnPostFork.Broadcast(EForkProcessRole::Child);
 
 #if UE_CHECK_LARGE_ALLOCATIONS
-				if (CVarEnableLargeAllocationChecksAfterFork.GetValueOnAnyThread())
-				{
-					UE::Memory::Private::CVarEnableLargeAllocationChecks->Set(true);
-				}
+					if (CVarEnableLargeAllocationChecksAfterFork.GetValueOnAnyThread())
+					{
+						UE::Memory::Private::CVarEnableLargeAllocationChecks->Set(true);
+					}
 #endif
 
-				// Children break out of the loop and return
-				RetVal = EWaitAndForkResult::Child;
-				break;
-			}
-			else
-			{
-				// This should be the very first thing we do after forking for optimal interaction with GMalloc
-				FForkProcessHelper::LowLevelPostForkParent();
+					// Children break out of the loop and return
+					RetVal = EWaitAndForkResult::Child;
+					ChildBreakOut = true;
+				}
+				else
+				{
+					// This should be the very first thing we do after forking for optimal interaction with GMalloc
+					FForkProcessHelper::LowLevelPostForkParent();
 
-				// Parent
-				AllChildren.Emplace(ChildPID, SignalData.SignalValue);
+					// Parent
+					AllChildren.Emplace(ChildPID, SignalData.SignalValue);
 
-				FCoreDelegates::OnPostFork.Broadcast(EForkProcessRole::Parent);
+					FCoreDelegates::OnPostFork.Broadcast(EForkProcessRole::Parent);
 
-				UE_LOG(LogFork, Log, TEXT("[Parent] WaitAndFork Successfully processed request %04hx-%04hx, made a child with pid %d! Total number of children: %d."), Cookie, ChildIdx, ChildPID, AllChildren.Num());
+					UE_LOG(LogFork, Log, TEXT("[Parent] WaitAndFork Successfully processed request %04hx-%04hx, made a child with pid %d! Total number of children: %d."), Cookie, ChildIdx, ChildPID, AllChildren.Num());
+				}
+				BatchCount--;
 			}
 		}
 		else
