Allow capturing non-error/warning/display log messages for automated tests (to match against ExpectedMessages using log verbosity threshold
#misc added automation test to validate UE_LOG message capturing (including > Display verbosity)

#Synced-CL 49030801

diff --git a/Engine/Plugins/Tests/TestFramework/Source/TestFramework/Private/AutomationExpectedMessage.spec.cpp b/Engine/Plugins/Tests/TestFramework/Source/TestFramework/Private/AutomationExpectedMessage.spec.cpp
--- a/Engine/Plugins/Tests/TestFramework/Source/TestFramework/Private/AutomationExpectedMessage.spec.cpp
+++ b/Engine/Plugins/Tests/TestFramework/Source/TestFramework/Private/AutomationExpectedMessage.spec.cpp
@@ -263,3 +263,25 @@ void FAutomationExpectedMessageTest::Define()
 	});
 
 }
+
+BEGIN_DEFINE_SPEC(FAutomationExpectedLogMessageTest, "TestFramework.ExpectedMessagesFromLog", EAutomationTestFlags::EngineFilter | EAutomationTestFlags_ApplicationContextMask)
+END_DEFINE_SPEC(FAutomationExpectedLogMessageTest)
+void FAutomationExpectedLogMessageTest::Define()
+{
+	Describe("A defined expected message (captured from log) with a specified verbosity ", [this]()
+	{
+		It("will match a UE_LOG entry that contains expected message", [this]()
+		{
+			AddExpectedMessage(TEXT("suppress this error message"), ELogVerbosity::Error, EAutomationExpectedMessageFlags::Contains, 1);
+
+			UE_LOG(LogAutomationTestFramework, Error, TEXT("suppress this error message please"))
+		});
+
+		It("will match a UE_LOG entry of higher than Display verbosity level", [this]()
+		{
+			AddExpectedMessage(TEXT("suppress this log message"), ELogVerbosity::Log, EAutomationExpectedMessageFlags::Contains, 1);
+			UE_LOG(LogAutomationTestFramework, Log, TEXT("suppress this log message"));
+		});
+	});
+}
+
diff --git a/Engine/Source/Runtime/Core/Private/Misc/AutomationTest.cpp b/Engine/Source/Runtime/Core/Private/Misc/AutomationTest.cpp
--- a/Engine/Source/Runtime/Core/Private/Misc/AutomationTest.cpp
+++ b/Engine/Source/Runtime/Core/Private/Misc/AutomationTest.cpp
@@ -225,7 +225,7 @@ void FAutomationTestFramework::FAutomationTestOutputDevice::Serialize( const TCH
 	if (LocalCurTest)
 	{
 		bool CaptureLog = !LocalCurTest->SuppressLogs()
-			&& (Verbosity == ELogVerbosity::Error || Verbosity == ELogVerbosity::Warning || Verbosity == ELogVerbosity::Display)
+			&& LocalCurTest->ShouldCaptureLogVerbosity(Verbosity)
 			&& LocalCurTest->ShouldCaptureLogCategory(Category);
 
 		if (CaptureLog)
@@ -266,10 +266,6 @@ void FAutomationTestFramework::FAutomationTestOutputDevice::Serialize( const TCH
 
 				LocalCurTest->AddAnalyticsItem(LogString);
 			}
-			//else
-			//{
-			//	LocalCurTest->AddInfo(LogString, STACK_OFFSET);
-			//}
 		}
 	}
 }
@@ -1373,6 +1369,7 @@ bool FAutomationTestFramework::InternalStopTest(FAutomationTestExecutionInfo& Ou
 	{
 		UE::TWriteScopeLock Lock(CurrentTest->ActionCS);
 		CurrentTest->ExpectedMessages.Empty();
+		CurrentTest->ExpectedMessageVerbosityThreshold = ELogVerbosity::Display;
 	}
 
 	// Set the success state of the test based on the above criteria
@@ -1403,7 +1400,7 @@ bool FAutomationTestFramework::InternalStopTest(FAutomationTestExecutionInfo& Ou
 	AutomationTestMessageFilter.SetCurrentAutomationTest(nullptr);
 
 	// Release pointers to now-invalid data
-	CurrentTest = NULL;
+	CurrentTest = nullptr;
 
 	return bTestSuccessful;
 }
@@ -1923,6 +1920,11 @@ void FAutomationTestBase::AddExpectedMessage(
 {
 	UE::TWriteScopeLock Lock(ActionCS);
 	ExpectedMessages.Add(FAutomationExpectedMessage(ExpectedPatternString, ExpectedVerbosity, CompareType, Occurrences, IsRegex));
+
+	if (ExpectedVerbosity > ExpectedMessageVerbosityThreshold)
+	{
+		ExpectedMessageVerbosityThreshold = ExpectedVerbosity; 
+	}
 }
 
 void FAutomationTestBase::AddExpectedMessage(
@@ -2707,3 +2709,8 @@ bool FAutomationTestBase::IsExpectedMessage(
 
 	return false;
 }
+
+bool FAutomationTestBase::ShouldCaptureLogVerbosity(const ELogVerbosity::Type InLogVerbosity)
+{
+	return InLogVerbosity > ELogVerbosity::NoLogging && InLogVerbosity <= ExpectedMessageVerbosityThreshold;
+}
diff --git a/Engine/Source/Runtime/Core/Public/Misc/AutomationTest.h b/Engine/Source/Runtime/Core/Public/Misc/AutomationTest.h
--- a/Engine/Source/Runtime/Core/Public/Misc/AutomationTest.h
+++ b/Engine/Source/Runtime/Core/Public/Misc/AutomationTest.h
@@ -1919,6 +1919,13 @@ public:
 	 */
 	virtual bool ShouldCaptureLogCategory(const class FName& Category) const { return true; }
 
+	/**
+	* Whether or not to the provided log verbosity level should be captured and converted into test events
+	* Driven from ExpectedMessageVerbosityThreshold (defaults to ELogVerbosity::Display)
+	* ExpectedMessageVerbosityThreshold is updated to match highest verbosity level for expected messages
+	*/
+	CORE_API virtual bool ShouldCaptureLogVerbosity(const ELogVerbosity::Type InLogVerbosity);
+
 	/**
 	 * If returns true then logging with a level of Error will not be recorded in test results
 	 *
@@ -2421,7 +2428,7 @@ public:
 	}
 
 	/**
-	 * Logs an error if the specified pointer is NULL.
+	 * Logs an error if the specified pointer is a nullptr.
 	 *
 	 * @param What - Description text for the test.
 	 * @param Pointer - The pointer to test.
@@ -2691,6 +2698,9 @@ private:
 	/* Log messages to be expected while processing this test.*/
 	TSet<FAutomationExpectedMessage> ExpectedMessages;
 
+	/* Threshold verbosity level for capturing log messages converted into test events (defaults to ELogVerbosity::Display matching previous behaviour) */	
+	TAtomic<ELogVerbosity::Type> ExpectedMessageVerbosityThreshold = ELogVerbosity::Display;
+
 	/** Critical section lock */
 	FTransactionallySafeRWLock ActionCS;
 };
