Added UE::CCompleteType to throw an error at compile time when a concept would be incorrectly instantiated due to a semantic constraint not being matched when the type is incomplete.

UE::CConvertibleTo has a bespoke and partial implementation against the common cases - this is non-trivial as there are many convertibility tests that can correctly determined without needing complete types, but such cases can be added as they are identified.

diff --git a/Engine/Source/Runtime/Core/Public/Concepts/CompleteType.h b/Engine/Source/Runtime/Core/Public/Concepts/CompleteType.h
new file mode 100644
--- /dev/null
+++ b/Engine/Source/Runtime/Core/Public/Concepts/CompleteType.h
@@ -0,0 +1,31 @@
+// Copyright Epic Games, Inc. All Rights Reserved.
+
+#pragma once
+
+#include "Misc/StaticAssertCompleteType.h"
+
+namespace UE::Private::CompleteType
+{
+	template <typename T>
+	constexpr bool CheckCompleteType()
+	{
+		UE_STATIC_ASSERT_COMPLETE_TYPE(T, "Expected a complete type");
+		return true;
+	}
+}
+
+namespace UE
+{
+	/**
+	 * Concept which describes a complete type.
+	 *
+	 * Unlike most concepts, which will evaluate to false if the type doesn't model the concept,
+	 * CCompleteType will always evaluate to true, but cause a compile error if the type is incomplete.
+	 * The purpose of it is to be included as an additional constraint of other concepts and catch when they're
+	 * evaluated with an incomplete type, which could give a wrong result and result in an ill-defined
+	 * program.  This is not important for concepts which only operate on fundamental types like UE::CPointer<T>,
+	 * or do not depend upon a type's definition like UE::CSameAs<A, B>.
+	 */
+	template <typename T>
+	concept CCompleteType = UE::Private::CompleteType::CheckCompleteType<T>();
+}
\ No newline at end of file
diff --git a/Engine/Source/Runtime/Core/Public/Concepts/ConditionallyConvertibleToBool.h b/Engine/Source/Runtime/Core/Public/Concepts/ConditionallyConvertibleToBool.h
--- a/Engine/Source/Runtime/Core/Public/Concepts/ConditionallyConvertibleToBool.h
+++ b/Engine/Source/Runtime/Core/Public/Concepts/ConditionallyConvertibleToBool.h
@@ -2,6 +2,9 @@
 
 #pragma once
 
+#include "Concepts/CompleteType.h"
+#include <type_traits>
+
 namespace UE
 {
 	/**
@@ -9,6 +12,7 @@ namespace UE
 	 */
 	template <typename T>
 	concept CConditionallyConvertibleToBool =
+		UE::CCompleteType<std::remove_reference_t<T>> &&
 		requires (const T& Val)
 		{
 			Val ? 1 : 0;
diff --git a/Engine/Source/Runtime/Core/Public/Concepts/ContiguousRange.h b/Engine/Source/Runtime/Core/Public/Concepts/ContiguousRange.h
--- a/Engine/Source/Runtime/Core/Public/Concepts/ContiguousRange.h
+++ b/Engine/Source/Runtime/Core/Public/Concepts/ContiguousRange.h
@@ -2,6 +2,7 @@
 
 #pragma once
 
+#include "Concepts/CompleteType.h"
 #include "Traits/IsContiguousContainer.h"
 #include <type_traits>
 
@@ -13,6 +14,7 @@ namespace UE
 	template <typename T>
 	concept CContiguousRange =
 		!std::is_reference_v<T> && // TIsContiguousContainer wrongly identifies references to containers as containers
+		UE::CCompleteType<T> &&
 		TIsContiguousContainer_V<T>;
 
 	/**
diff --git a/Engine/Source/Runtime/Core/Public/Concepts/ConvertibleTo.h b/Engine/Source/Runtime/Core/Public/Concepts/ConvertibleTo.h
--- a/Engine/Source/Runtime/Core/Public/Concepts/ConvertibleTo.h
+++ b/Engine/Source/Runtime/Core/Public/Concepts/ConvertibleTo.h
@@ -2,8 +2,99 @@
 
 #pragma once
 
+#include "Misc/StaticAssertCompleteType.h"
 #include <type_traits>
 
+namespace UE::Private::ConvertibleTo
+{
+	// Removes all top level qualifiers, pointers and references, e.g. TRemoveAllCVPtrRef<int* const* volatile&>::Type == int
+	template <typename T>
+	struct TRemoveAllCVPtrRef
+	{
+		using Type = T;
+	};
+
+	template <typename T> struct TRemoveAllCVPtrRef<T*>               : TRemoveAllCVPtrRef<T> {};
+	template <typename T> struct TRemoveAllCVPtrRef<T&>               : TRemoveAllCVPtrRef<T> {};
+	template <typename T> struct TRemoveAllCVPtrRef<T&&>              : TRemoveAllCVPtrRef<T> {};
+	template <typename T> struct TRemoveAllCVPtrRef<const          T> : TRemoveAllCVPtrRef<T> {};
+	template <typename T> struct TRemoveAllCVPtrRef<      volatile T> : TRemoveAllCVPtrRef<T> {};
+	template <typename T> struct TRemoveAllCVPtrRef<const volatile T> : TRemoveAllCVPtrRef<T> {};
+
+	template <typename From, typename To>
+	constexpr bool IsConvertibleToComplete()
+	{
+		// If a UE_STATIC_ASSERT_COMPLETE_TYPE fires here on a conversion that should be possible to determine without needing full types,
+		// the function IsConvertibleTo() below should be updated to handle that case.
+		UE_STATIC_ASSERT_COMPLETE_TYPE(typename TRemoveAllCVPtrRef<From>::Type, "Trying to test convertibility against a type that's incomplete");
+		UE_STATIC_ASSERT_COMPLETE_TYPE(typename TRemoveAllCVPtrRef<To  >::Type, "Trying to test convertibility against a type that's incomplete");
+
+		return std::is_convertible_v<From, To>;
+	}
+
+	// This exists to try and catch incomplete types in the same way as UE::CCompleteType, but when it
+	// comes to conversions, a complete type checker is too blunt - there are many common conversion
+	// tests that can be accurately done without needing complete types, and we need to try to handle those.
+	template <typename From, typename To>
+	constexpr bool IsConvertibleTo()
+	{
+		// Conversion checks between pointer types are a common case which don't require full types
+		if constexpr (std::is_pointer_v<From> && std::is_pointer_v<To>)
+		{
+			using FromType = std::remove_pointer_t<From>;
+			using ToType   = std::remove_pointer_t<To>;
+
+			// Losing cv qualifiers is always a convertibility failure
+			if constexpr ((std::is_const_v<FromType> && !std::is_const_v<ToType>) || (std::is_volatile_v<FromType> && !std::is_volatile_v<ToType>))
+			{
+				return false;
+			}
+			// Adding qualifiers to the same type is always possible
+			else if constexpr (std::is_same_v<const volatile FromType, const volatile ToType>)
+			{
+				return true;
+			}
+			// Conversion to void is always possible unless it's a function pointer
+			else if (std::is_void_v<ToType>)
+			{
+				return !std::is_function_v<From>;
+			}
+			// Fall back to the regular convertibility test, checking for complete types
+			else 
+			{
+				return IsConvertibleToComplete<From, To>();
+			}
+		}
+		// Conversion checks between reference types are also a common case
+		else if constexpr (std::is_lvalue_reference_v<From> && std::is_lvalue_reference_v<To>)
+		{
+			using FromType = std::remove_reference_t<From>;
+			using ToType   = std::remove_reference_t<To>;
+
+			// Losing cv qualifiers is always a convertibility failure
+			if constexpr ((std::is_const_v<FromType> && !std::is_const_v<ToType>) || (std::is_volatile_v<FromType> && !std::is_volatile_v<ToType>))
+			{
+				return false;
+			}
+			// Adding qualifiers to the same type is always possible
+			else if constexpr (std::is_same_v<const volatile FromType, const volatile ToType>)
+			{
+				return true;
+			}
+			// Fall back to the regular convertibility test, checking for complete types
+			else 
+			{
+				return IsConvertibleToComplete<From, To>();
+			}
+		}
+		// Fall back to the regular convertibility test, checking for complete types
+		else
+		{
+			return IsConvertibleToComplete<From, To>();
+		}
+	}
+}
+
 namespace UE
 {
 
@@ -11,8 +102,13 @@ namespace UE
  * Concept which describes convertibility from one type to another.
  *
  * We use this instead of std::convertible_to because <concepts> isn't a well supported header yet.
+ *
+ * We want to use UE::CCompleteType here, but many convertibility tests can be done without full type
+ * definitions, so it's not simple to use it in a way that won't break lots of valid code.
+ * UE::Private::ConvertibleTo::IsConvertibleTo is not a complete implementation of what's necessary
+ * but handles common cases before doing a complete type check.
  */
 template <typename From, typename To>
-concept CConvertibleTo = std::is_convertible_v<From, To> && requires { static_cast<To>(std::declval<From>()); };
+concept CConvertibleTo = UE::Private::ConvertibleTo::IsConvertibleTo<From, To>() && requires { static_cast<To>(std::declval<From>()); };
 
 } // UE
diff --git a/Engine/Source/Runtime/Core/Public/Concepts/Dereferenceable.h b/Engine/Source/Runtime/Core/Public/Concepts/Dereferenceable.h
--- a/Engine/Source/Runtime/Core/Public/Concepts/Dereferenceable.h
+++ b/Engine/Source/Runtime/Core/Public/Concepts/Dereferenceable.h
@@ -2,7 +2,9 @@
 
 #pragma once
 
+#include "Concepts/CompleteType.h"
 #include "Concepts/LValueReference.h"
+#include <type_traits>
 
 namespace UE
 {
@@ -11,6 +13,7 @@ namespace UE
 	 */
 	template <typename T>
 	concept CDereferenceable =
+		UE::CCompleteType<std::remove_reference_t<T>> &&
 		requires(const T& Val)
 		{
 			{ *Val } -> UE::CLValueReference;
diff --git a/Engine/Source/Runtime/Core/Public/Concepts/DerivedFrom.h b/Engine/Source/Runtime/Core/Public/Concepts/DerivedFrom.h
--- a/Engine/Source/Runtime/Core/Public/Concepts/DerivedFrom.h
+++ b/Engine/Source/Runtime/Core/Public/Concepts/DerivedFrom.h
@@ -2,6 +2,7 @@
 
 #pragma once
 
+#include "Concepts/CompleteType.h"
 #include <type_traits>
 
 namespace UE
@@ -9,8 +10,12 @@ namespace UE
 	/**
 	 * Concept which is satisfied if DerivedType is an unambiguous public base class of BaseType.
 	 *
-	 * We use this instead of std::same_as because <concepts> isn't a well supported header yet.
+	 * We use this instead of std::derived_from because <concepts> isn't a well supported header yet.
 	 */
 	template <typename DerivedType, typename BaseType>
-	concept CDerivedFrom = std::is_base_of_v<BaseType, DerivedType> && std::is_convertible_v<const volatile DerivedType*, const volatile BaseType*>;
+	concept CDerivedFrom =
+		UE::CCompleteType<DerivedType> &&
+		UE::CCompleteType<BaseType> &&
+		std::is_base_of_v<BaseType, DerivedType> &&
+		std::is_convertible_v<const volatile DerivedType*, const volatile BaseType*>;
 }
diff --git a/Engine/Source/Runtime/Core/Public/Concepts/PointerEqualityComparableWith.h b/Engine/Source/Runtime/Core/Public/Concepts/PointerEqualityComparableWith.h
--- a/Engine/Source/Runtime/Core/Public/Concepts/PointerEqualityComparableWith.h
+++ b/Engine/Source/Runtime/Core/Public/Concepts/PointerEqualityComparableWith.h
@@ -2,6 +2,11 @@
 
 #pragma once
 
+#include "Concepts/CompleteType.h"
+#include "Concepts/SameAs.h"
+#include "Concepts/Void.h"
+#include <type_traits>
+
 namespace UE
 {
 	/**
@@ -9,8 +14,15 @@ namespace UE
 	 */
 	template <typename T, typename U>
 	concept CPointerEqualityComparableWith =
-		requires(T* Lhs, U* Rhs)
-		{
-			Lhs == Rhs;
-		};
+		UE::CSameAs<std::remove_cv_t<T>, std::remove_cv_t<U>> ||
+		UE::CVoid<T> ||
+		UE::CVoid<U> ||
+		(
+			UE::CCompleteType<T> &&
+			UE::CCompleteType<U> &&
+			requires(T* Lhs, U* Rhs)
+			{
+				Lhs == Rhs;
+			}
+		);
 }
