Made AlignmentTemplates use concepts instead of static asserting against traits.

diff --git a/Engine/Source/Runtime/Core/Public/Templates/AlignmentTemplates.h b/Engine/Source/Runtime/Core/Public/Templates/AlignmentTemplates.h
--- a/Engine/Source/Runtime/Core/Public/Templates/AlignmentTemplates.h
+++ b/Engine/Source/Runtime/Core/Public/Templates/AlignmentTemplates.h
@@ -3,8 +3,8 @@
 #pragma once
 
 #include "CoreTypes.h"
-#include "Templates/IsIntegral.h"
-#include "Templates/IsPointer.h"
+#include "Concepts/Integral.h"
+#include "Concepts/Pointer.h"
 
 /**
  * Aligns a value to the nearest higher multiple of 'Alignment', which must be a power of two.
@@ -15,10 +15,9 @@
  * @return The value aligned up to the specified alignment.
  */
 template <typename T>
+	requires UE::CIntegral<T> || UE::CPointer<T>
 inline constexpr T Align(T Val, uint64 Alignment)
 {
-	static_assert(TIsIntegral<T>::Value || TIsPointer<T>::Value, "Align expects an integer or pointer type");
-
 	return (T)(((uint64)Val + Alignment - 1) & ~(Alignment - 1));
 }
 
@@ -31,10 +30,9 @@ inline constexpr T Align(T Val, uint64 Alignment)
  * @return The value aligned down to the specified alignment.
  */
 template <typename T>
+	requires UE::CIntegral<T> || UE::CPointer<T>
 inline constexpr T AlignDown(T Val, uint64 Alignment)
 {
-	static_assert(TIsIntegral<T>::Value || TIsPointer<T>::Value, "AlignDown expects an integer or pointer type");
-
 	return (T)(((uint64)Val) & ~(Alignment - 1));
 }
 
@@ -47,10 +45,9 @@ inline constexpr T AlignDown(T Val, uint64 Alignment)
  * @return true if the pointer is aligned to the specified alignment, false otherwise.
  */
 template <typename T>
+	requires UE::CIntegral<T> || UE::CPointer<T>
 inline constexpr bool IsAligned(T Val, uint64 Alignment)
 {
-	static_assert(TIsIntegral<T>::Value || TIsPointer<T>::Value, "IsAligned expects an integer or pointer type");
-
 	return !((uint64)Val & (Alignment - 1));
 }
 
@@ -63,9 +60,13 @@ inline constexpr bool IsAligned(T Val, uint64 Alignment)
  * @return The value aligned up to the specified alignment.
  */
 template <typename T>
+	requires UE::CIntegral<T> || UE::CPointer<T>
 inline constexpr T AlignArbitrary(T Val, uint64 Alignment)
 {
-	static_assert(TIsIntegral<T>::Value || TIsPointer<T>::Value, "AlignArbitrary expects an integer or pointer type");
-
 	return (T)((((uint64)Val + Alignment - 1) / Alignment) * Alignment);
 }
+
+#if UE_ENABLE_INCLUDE_ORDER_DEPRECATED_IN_5_8
+#include "Templates/IsIntegral.h"
+#include "Templates/IsPointer.h"
+#endif
