Math: Optimized FMath::RoundUpToPowerOfTwo[64] for Clang

Reduced typical instruction count from 43 with 2 branches down to 9 with no branches.

#Synced-CL 48512380

diff --git a/Engine/Source/Runtime/Core/Private/Misc/StringBuilder.cpp b/Engine/Source/Runtime/Core/Private/Misc/StringBuilder.cpp
--- a/Engine/Source/Runtime/Core/Private/Misc/StringBuilder.cpp
+++ b/Engine/Source/Runtime/Core/Private/Misc/StringBuilder.cpp
@@ -7,18 +7,6 @@
 #include "Misc/CString.h"
 #include "Misc/ScopeExit.h"
 
-static inline uint64_t NextPowerOfTwo(uint64_t x)
-{
-	x -= 1;
-	x |= (x >> 1);
-	x |= (x >> 2);
-	x |= (x >> 4);
-	x |= (x >> 8);
-	x |= (x >> 16);
-
-	return x + 1;
-}
-
 template <typename C>
 TStringBuilderBase<C>::~TStringBuilderBase()
 {
@@ -36,7 +24,7 @@ void TStringBuilderBase<C>::Extend(SIZE_T ExtraCapacity)
 		SIZE_T(CurPos - Base), SIZE_T(End - Base), ExtraCapacity);
 
 	const SIZE_T OldCapacity = End - Base;
-	const SIZE_T NewCapacity = FPlatformMath::Max<SIZE_T>(NextPowerOfTwo(OldCapacity + ExtraCapacity), 32);
+	const SIZE_T NewCapacity = FPlatformMath::Max<SIZE_T>(FMath::RoundUpToPowerOfTwo64(OldCapacity + ExtraCapacity), 32);
 
 	C* NewBase = (C*)AllocBuffer(NewCapacity);
 
diff --git a/Engine/Source/Runtime/Core/Public/Clang/ClangPlatformMath.h b/Engine/Source/Runtime/Core/Public/Clang/ClangPlatformMath.h
--- a/Engine/Source/Runtime/Core/Public/Clang/ClangPlatformMath.h
+++ b/Engine/Source/Runtime/Core/Public/Clang/ClangPlatformMath.h
@@ -49,7 +49,7 @@ struct FClangPlatformMath : public FGenericPlatformMath
 	 *
 	 * @return the number of zeros before the first "on" bit
 	 */
-	static inline constexpr uint64 CountLeadingZeros64(uint64 Value)
+	static UE_FORCEINLINE_HINT constexpr uint64 CountLeadingZeros64(uint64 Value)
 	{
 		if (Value == 0)
 		{
@@ -66,16 +66,16 @@ struct FClangPlatformMath : public FGenericPlatformMath
 	 *
 	 * @return the number of zeros after the last "on" bit
 	 */
-	static inline constexpr uint32 CountTrailingZeros(uint32 Value)
+	static UE_FORCEINLINE_HINT constexpr uint32 CountTrailingZeros(uint32 Value)
 	{
 		if (Value == 0)
 		{
 			return 32;
 		}
-	
+
 		return (uint32)__builtin_ctz(Value);
 	}
-	
+
 	/**
 	 * Counts the number of trailing zeros in the bit representation of the value
 	 *
@@ -83,13 +83,13 @@ struct FClangPlatformMath : public FGenericPlatformMath
 	 *
 	 * @return the number of zeros after the last "on" bit
 	 */
-	static inline constexpr uint64 CountTrailingZeros64(uint64 Value)
+	static UE_FORCEINLINE_HINT constexpr uint64 CountTrailingZeros64(uint64 Value)
 	{
 		if (Value == 0)
 		{
 			return 64;
 		}
-	
+
 		return (uint64)__builtin_ctzll(Value);
 	}
 
@@ -113,6 +113,30 @@ struct FClangPlatformMath : public FGenericPlatformMath
 		return 63 - __builtin_clzll(Value);
 	}
 
+	static UE_FORCEINLINE_HINT constexpr uint32 CeilLogTwo(uint32 Arg)
+	{
+		// if Arg is 0, change it to 1 so that we return 0
+		Arg = Arg ? Arg : 1;
+		return 32 - CountLeadingZeros(Arg - 1);
+	}
+
+	static UE_FORCEINLINE_HINT constexpr uint64 CeilLogTwo64(uint64 Arg)
+	{
+		// if Arg is 0, change it to 1 so that we return 0
+		Arg = Arg ? Arg : 1;
+		return 64 - CountLeadingZeros64(Arg - 1);
+	}
+
+	static UE_FORCEINLINE_HINT constexpr uint32 RoundUpToPowerOfTwo(uint32 Arg)
+	{
+		return uint32(1) << CeilLogTwo(Arg);
+	}
+
+	static UE_FORCEINLINE_HINT constexpr uint64 RoundUpToPowerOfTwo64(uint64 Arg)
+	{
+		return uint64(1) << CeilLogTwo64(Arg);
+	}
+
 	/**
 	 * Adds two integers of any integer type, checking for overflow.
 	 * If there was overflow, it returns false, and OutResult may or may not be written.
