Fixing comment from review changes.

#rnx
#Synced-CL 48801413

diff --git a/Engine/Source/Runtime/Core/Internal/IO/IoStore.h b/Engine/Source/Runtime/Core/Internal/IO/IoStore.h
--- a/Engine/Source/Runtime/Core/Internal/IO/IoStore.h
+++ b/Engine/Source/Runtime/Core/Internal/IO/IoStore.h
@@ -320,7 +320,7 @@ public:
 	[[nodiscard]] virtual FIoStoreTocChunkInfo GetTocChunkInfo(int32 TocEntryIndex) const = 0;
 	CORE_API virtual ~IIoStoreTocReader();
 
-	// Creates an instance for accessing the TOC. OutTocReader is only valid when FIoStatus is FIoStatus::Ok. The TocFilePath should be a path to the
+	// Creates an instance for accessing the TOC. The TocFilePath should be a path to the
 	// .utoc file (not the ucas or undecorated). 
 	CORE_API static TIoStatusOr<TUniquePtr<IIoStoreTocReader>> ReadFromDisk(const TCHAR* TocFilePath, const TMap<FGuid, FAES::FAESKey>& DecryptionKeys);
 };
\ No newline at end of file
