Update check to call CA_ASSUME(true)

#rnx

#Synced-CL 49130599

diff --git a/Engine/Source/Runtime/Core/Public/Misc/AssertionMacros.h b/Engine/Source/Runtime/Core/Public/Misc/AssertionMacros.h
--- a/Engine/Source/Runtime/Core/Public/Misc/AssertionMacros.h
+++ b/Engine/Source/Runtime/Core/Public/Misc/AssertionMacros.h
@@ -245,6 +245,10 @@ RetType UE_COLD UE_DEBUG_SECTION DispatchCheckVerify (InnerType&& Inner, ArgType
 				} \
 				CA_ASSUME(false); \
 			} \
+			else \
+			{ \
+				CA_ASSUME(true); \
+			} \
 		}
 	
 	/**
@@ -269,6 +273,10 @@ RetType UE_COLD UE_DEBUG_SECTION DispatchCheckVerify (InnerType&& Inner, ArgType
 				} \
 				CA_ASSUME(false); \
 			} \
+			else \
+			{ \
+				CA_ASSUME(true); \
+			} \
 		}
 
 	/**
