[RaceDetector]
- Fix a bug where we could push a syncobject bank in the freelist while there was still a ref on it from another thread

#Synced-CL 49011649

diff --git a/Engine/Source/Runtime/Core/Internal/Sanitizer/RaceDetectorTypes.h b/Engine/Source/Runtime/Core/Internal/Sanitizer/RaceDetectorTypes.h
--- a/Engine/Source/Runtime/Core/Internal/Sanitizer/RaceDetectorTypes.h
+++ b/Engine/Source/Runtime/Core/Internal/Sanitizer/RaceDetectorTypes.h
@@ -954,7 +954,7 @@ namespace UE::Sanitizer::RaceDetector {
 			int32 NewRefCount = FPlatformAtomics::InterlockedDecrement(&RefCount);
 			if (NewRefCount == 0)
 			{
-				delete this;
+				ReleaseLastRef();
 			}
 			return NewRefCount;
 		}
@@ -978,7 +978,9 @@ namespace UE::Sanitizer::RaceDetector {
 
 		static volatile int64 ObjectCount;
 	private:
-		volatile int32 RefCount = 1;
+		INSTRUMENTATION_FUNCTION_ATTRIBUTES void ReleaseLastRef();
+
+		volatile int32 RefCount = 0;
 
 		// Save space by allocating the first sync object as part of the bank itself.
 		// Most of the time the sync object will be at offset 0.
diff --git a/Engine/Source/Runtime/Core/Private/Sanitizer/RaceDetector.cpp b/Engine/Source/Runtime/Core/Private/Sanitizer/RaceDetector.cpp
--- a/Engine/Source/Runtime/Core/Private/Sanitizer/RaceDetector.cpp
+++ b/Engine/Source/Runtime/Core/Private/Sanitizer/RaceDetector.cpp
@@ -394,7 +394,7 @@ namespace UE::Sanitizer::RaceDetector {
 
 	INSTRUMENTATION_FUNCTION_ATTRIBUTES void PushFreeObjectBank(FSyncObjectBank* SyncObjectBank)
 	{
-		RACEDETECTOR_CHECK(SyncObjectBank->GetRefCount() == 1);
+		RACEDETECTOR_CHECK(SyncObjectBank->GetRefCount() == 0);
 
 		// Cleanup all non-aligned objects, those are unlikely to be needed when the
 		// sync bank is reused for a different atomic address.
@@ -426,6 +426,7 @@ namespace UE::Sanitizer::RaceDetector {
 			{
 				FPlatformAtomics::InterlockedIncrement(&FSyncObjectBank::ObjectCount);
 				LocalSyncBankHead->Next = nullptr;
+				RACEDETECTOR_CHECK(LocalSyncBankHead->GetRefCount() == 0);
 				return LocalSyncBankHead;
 			}
 		}
@@ -453,7 +454,7 @@ namespace UE::Sanitizer::RaceDetector {
 				// When publishing in shadow memory, we need to make sure the refcount is not 0, since another thread could try to use
 				// the newly published bank from the shadow, perform an AddRef and then release it and we'd end up with
 				// an invalid NewObjectBank before even starting to use it.
-				RACEDETECTOR_CHECK(NewObjectBank->GetRefCount() == 1);
+				NewObjectBank->AddRef();
 
 				FSyncObjectBank* OldObjectBank = (FSyncObjectBank*)FPlatformAtomics::InterlockedCompareExchangePointer((void**)&ClockBankSlot->SyncObjectBank, NewObjectBank, nullptr);
 				if (OldObjectBank == nullptr)
@@ -474,10 +475,7 @@ namespace UE::Sanitizer::RaceDetector {
 				LocalBank = OldObjectBank;
 
 				// Should not happen very often, and at most once per thread taking part of the race
-				// Put it back in the list, freeing it now could cause a crash since that is fresh out a linked-list
-				// that another thread might still try to dereference.
-				PushFreeObjectBank(NewObjectBank);
-				RACEDETECTOR_CHECK(Result == 0);
+				NewObjectBank->Release();
 			}
 
 			Context.BankHazard = LocalBank;
@@ -488,6 +486,7 @@ namespace UE::Sanitizer::RaceDetector {
 			if (LocalBank == ClockBankSlot->SyncObjectBank)
 			{
 				FSyncObjectRef Result = FSyncObjectRef(LocalBank, LocalBank->GetSyncObject(Index));
+				RACEDETECTOR_CHECK(LocalBank->GetRefCount() >= 1);
 				Context.BankHazard = nullptr;
 				return MoveTemp(Result);
 			}
@@ -1067,6 +1066,13 @@ namespace UE::Sanitizer::RaceDetector {
 		}
 		return Hazards;
 	}
+	
+	INSTRUMENTATION_FUNCTION_ATTRIBUTES void FSyncObjectBank::ReleaseLastRef()
+	{
+		// Releasing the last ref will put it back in the free list, freeing it could cause a crash since it is part of a linked-list
+		// and another thread might need to dereference the Next pointer.
+		PushFreeObjectBank(this);
+	}
 
 	INSTRUMENTATION_FUNCTION_ATTRIBUTES void ResetShadow()
 	{
@@ -1118,7 +1124,7 @@ namespace UE::Sanitizer::RaceDetector {
 				}
 				else
 				{
-					PushFreeObjectBank(ToDelete);
+					ToDelete->Release();
 				}
 			}
 
