Re-enable AutoRTFM + ASan for VerseTestScriptCmd which someone (me) accidentally disabled.

#Synced-CL 49163358

diff --git a/Engine/Source/Runtime/AutoRTFM/Private/Context.cpp b/Engine/Source/Runtime/AutoRTFM/Private/Context.cpp
--- a/Engine/Source/Runtime/AutoRTFM/Private/Context.cpp
+++ b/Engine/Source/Runtime/AutoRTFM/Private/Context.cpp
@@ -324,11 +324,21 @@ AutoRTFM::FStackRange FContext::GetThreadStackRange()
 		Stack.Low = static_cast<char*>(Stack.High) - StackSize;
 #else
 		pthread_attr_t Attr{};
-		pthread_getattr_np(pthread_self(), &Attr);
+
+		{
+			const int Error = pthread_getattr_np(pthread_self(), &Attr);
+			AUTORTFM_ASSERT(0 == Error);
+		}
+
 		Stack.Low = 0;
 		size_t StackSize = 0;
 		pthread_attr_getstack(&Attr, &Stack.Low, &StackSize);
 		Stack.High = static_cast<char*>(Stack.Low) + StackSize;
+
+		{
+			const int Error = pthread_attr_destroy(&Attr);
+			AUTORTFM_ASSERT(0 == Error);
+		}
 #endif
 
 		AUTORTFM_ASSERT(Stack.High > Stack.Low);
diff --git a/Engine/Source/Runtime/AutoRTFM/Private/FunctionMap.cpp b/Engine/Source/Runtime/AutoRTFM/Private/FunctionMap.cpp
--- a/Engine/Source/Runtime/AutoRTFM/Private/FunctionMap.cpp
+++ b/Engine/Source/Runtime/AutoRTFM/Private/FunctionMap.cpp
@@ -8,6 +8,8 @@
 #include "AutoRTFM/ReportHazard.h"
 #include "BuildMacros.h"
 #include "FunctionMapInlines.h"
+#include "ContextInlines.h"
+#include "TransactionInlines.h"
 #include "Utils.h"
 
 #include <cinttypes>
@@ -310,14 +312,25 @@ AUTORTFM_DISABLE_NO_SANITIZE static void* FunctionMapReportError(void* Open, con
 	}
 #endif
 
-	std::string FunctionDescription = GetFunctionDescription(Open);
+	// This is gross - but if we are allowed to recover from missing functions (the eventual mode
+	// we intend to ship AutoRTFM in to users), then we will do an abort in the report error calls
+	// below. Which would mean we wouldn't free the allocation backing the string below! Get around
+	// it by putting it in an explicit heap allocation, and (ab)use our existing on-abort handling
+	// system to free the memory.
+	std::string* const FunctionDescription = new std::string(GetFunctionDescription(Open));
+
+	FContext::Get()->GetCurrentTransaction()->DeferUntilAbort([FunctionDescription]
+	{
+		delete FunctionDescription;
+	});
+
 	if (Where)
 	{
-		AUTORTFM_REPORT_ERROR("Could not find function %p '%s' where '%s'.", Open, FunctionDescription.c_str(), Where);
+		AUTORTFM_REPORT_ERROR("Could not find function %p '%s' where '%s'.", Open, FunctionDescription->c_str(), Where);
 	}
 	else
 	{
-		AUTORTFM_REPORT_ERROR("Could not find function %p '%s'.", Open, FunctionDescription.c_str());
+		AUTORTFM_REPORT_ERROR("Could not find function %p '%s'.", Open, FunctionDescription->c_str());
 	}
 	return nullptr;
 }
diff --git a/Engine/Source/Runtime/Core/Public/Async/TransactionallySafeMutex.h b/Engine/Source/Runtime/Core/Public/Async/TransactionallySafeMutex.h
--- a/Engine/Source/Runtime/Core/Public/Async/TransactionallySafeMutex.h
+++ b/Engine/Source/Runtime/Core/Public/Async/TransactionallySafeMutex.h
@@ -38,14 +38,11 @@ namespace UE
 // modifications we should make.
 struct FTransactionallySafeMutex final
 {
-	// Always open because the constructor will create the underlying mutex.
-	UE_AUTORTFM_ALWAYS_OPEN
 	FTransactionallySafeMutex() : State(MakeShared<FState>())
 	{
 	}
 
 	// Construct in a locked state.
-	UE_AUTORTFM_ALWAYS_OPEN
 	inline explicit FTransactionallySafeMutex(UE::FAcquireLock Tag) : State(MakeShared<FState>(Tag))
 	{
 	}
@@ -257,13 +254,11 @@ private:
 
 	struct FState final
 	{
-		UE_AUTORTFM_NOAUTORTFM
 		FState() : bTransactionallyLocked(false), bRegisteredCommitHandler(false) {}
 
-		UE_AUTORTFM_NOAUTORTFM
 		explicit FState(UE::FAcquireLock Tag) : Mutex(Tag), bTransactionallyLocked(false), bRegisteredCommitHandler(false) {}
 
-		UE_AUTORTFM_NOAUTORTFM
+		AUTORTFM_DISABLE
 		~FState()
 		{
 			ensure(TransactionalLockCount == 0);
@@ -271,7 +266,7 @@ private:
 			ensure(bRegisteredCommitHandler == false);
 		}
 
-		UE_AUTORTFM_NOAUTORTFM
+		AUTORTFM_DISABLE
 		void ResetTransactionState()
 		{
 			TransactionalLockCount = 0;
@@ -279,25 +274,25 @@ private:
 			bRegisteredCommitHandler = false;
 		}
 
-		UE_AUTORTFM_ALWAYS_OPEN
+		AUTORTFM_OPEN
 		void MutexLock()
 		{
 			Mutex.Lock();
 		}
 		
-		UE_AUTORTFM_ALWAYS_OPEN
+		AUTORTFM_OPEN
 		bool MutexTryLock()
 		{
 			return Mutex.TryLock();
 		}
 
-		UE_AUTORTFM_ALWAYS_OPEN
+		AUTORTFM_OPEN
 		void MutexUnlock()
 		{
 			Mutex.Unlock();
 		}
 
-		UE_AUTORTFM_ALWAYS_OPEN
+		AUTORTFM_OPEN
 		bool MutexIsLocked()
 		{
 			return Mutex.IsLocked();
