Add a `TTransactionallySafeMpscQueue` and make our `TTransactionallySafeSpscQueue` alias it (since we use locks its the same implementation for both).

#Synced-CL 49163605

diff --git a/Engine/Source/Programs/AutoRTFMTests/Private/QueueTests.cpp b/Engine/Source/Programs/AutoRTFMTests/Private/QueueTests.cpp
--- a/Engine/Source/Programs/AutoRTFMTests/Private/QueueTests.cpp
+++ b/Engine/Source/Programs/AutoRTFMTests/Private/QueueTests.cpp
@@ -3,46 +3,74 @@
 #include "AutoRTFM/Testing.h"
 #include "Catch2Includes.h"
 #include "Containers/TransactionallySafeSpscQueue.h"
+#include "Containers/TransactionallySafeMpscQueue.h"
 
 #include <thread>
 
 // Test loosely based on code from Containers/ConcurrentQueuesTest.cpp.
-template <typename QueueType>
-void TestSpscQueueCorrectness(int Num, bool bTransact)
+template <typename QueueType, unsigned NumProducerThreads = 1>
+void TestQueueCorrectness(int Num, bool bTransact, typename QueueType::ElementType (*Produce)(int Index))
 {
 	QueueType Queue;
-	uint32 NumProduced = 0;
-	uint32 NumConsumed = 0;
 	std::thread ConsumeTask;
+	std::vector<std::thread> ProduceTasks;
+	std::atomic<bool> BusyWait = true;
+	std::set<typename QueueType::ElementType> ElementsConsumed;
 
 	// Consumer
-	auto ConsumeLoop = [&Queue, &NumConsumed, Num]
+	auto ConsumeLoop = [&Queue, Num, &ElementsConsumed]
 	{
-		while (NumConsumed != Num)
+		while (ElementsConsumed.size() != (Num * NumProducerThreads))
 		{
-			if (Queue.Dequeue().IsSet())
+			TOptional<typename QueueType::ElementType> Element = Queue.Dequeue();
+
+			if (Element.IsSet())
 			{
-				++NumConsumed;
+				ElementsConsumed.insert(Element.GetValue());
 			}
 		}
 	};
 
 	// Producer
-	auto ProduceLoop = [&Queue, &NumProduced, Num]
+	auto ProduceLoop = [&Queue, Num, &BusyWait, Produce]
 	{
-		do
+		for (int Index = 0; Index < Num; Index++)
 		{
-			Queue.Enqueue(typename QueueType::ElementType{});
-		} 
-		while (++NumProduced != Num);
+			Queue.Enqueue(Produce(Index));
+
+			if (Index == 0)
+			{
+				// We signal that other threads can start producing
+				// after we've produced one element.
+				UE_AUTORTFM_OPEN
+				{
+					BusyWait = false;
+					BusyWait.notify_all();
+				};
+			}
+		}
 	};
-	
-	// Run the consumer loop on a helper thread.
-	UE_AUTORTFM_OPEN
+
+	// Threaded producer
+	auto ThreadedProduceLoop = [&Queue, Num, &BusyWait, Produce](unsigned ThreadId)
 	{
-		ConsumeTask = std::thread{ConsumeLoop};
+		BusyWait.wait(true);
+
+		for (int Index = 0; Index < Num; Index++)
+		{
+			Queue.Enqueue(Produce(ThreadId * Num + Index));
+		}
 	};
+	
+	// Run the consumer loop on a helper thread.
+	ConsumeTask = std::thread{ConsumeLoop};
 
+	// We start at 1 because our current thread will produce items too.
+	for (unsigned Thread = 1; Thread < NumProducerThreads; Thread++)
+	{
+		ProduceTasks.push_back(std::thread(ThreadedProduceLoop, Thread));
+	}
+	
 	// Run the producer loop within a transaction.
 	if (bTransact)
 	{
@@ -57,24 +85,66 @@ void TestSpscQueueCorrectness(int Num, bool bTransact)
 	}
 
 	// Join the consumer thread to allow it to finish.
-	UE_AUTORTFM_OPEN
+	ConsumeTask.join();
+
+	for (std::thread& ProduceTask : ProduceTasks)
 	{
-		ConsumeTask.join();
-	};
+		ProduceTask.join();
+	}
 
 	// We should have produced and consumed exactly Num items, and the queue should be empty.
-	REQUIRE(NumProduced == Num);
-	REQUIRE(NumConsumed == Num);
+	REQUIRE(ElementsConsumed.size() == (Num * NumProducerThreads));
 	REQUIRE(!Queue.Dequeue().IsSet());
+
+	for (int Index = 0; Index < (Num * NumProducerThreads); Index++)
+	{
+		REQUIRE(ElementsConsumed.contains(Produce(Index)));
+	}
+}
+
+template<template<typename> class QueueType> struct TQueueWithInt final
+{
+	using Queue = QueueType<int>;
+
+	static int Producer(int Index);
+};
+
+template<template<typename> class QueueType> int TQueueWithInt<QueueType>::Producer(int Index)
+{
+	return Index;
+}
+
+template<template<typename> class QueueType> struct TQueueWithString final
+{
+	using Queue = QueueType<FString>;
+
+	static FString Producer(int Index);
+};
+
+template<template<typename> class QueueType> FString TQueueWithString<QueueType>::Producer(int Index)
+{
+	return FString::Printf(TEXT("%d"), Index);
+}
+
+// Test with trivial and complex types
+TEMPLATE_TEST_CASE("TransactionallySafeSpscQueue", "",
+	TQueueWithInt<TTransactionallySafeSpscQueue>, TQueueWithString<TTransactionallySafeSpscQueue>,
+	TQueueWithInt<TTransactionallySafeMpscQueue>, TQueueWithString<TTransactionallySafeMpscQueue>)
+{
+	// Non-transactional
+	TestQueueCorrectness<typename TestType::Queue>(10000, /*bTransact=*/false, TestType::Producer);
+
+	// Transactional
+	TestQueueCorrectness<typename TestType::Queue>(10000, /*bTransact=*/true, TestType::Producer);
 }
 
-TEST_CASE("TransactionallySafeSpscQueue")
+// Test with trivial and complex types
+TEMPLATE_TEST_CASE("TransactionallySafeMpscQueue", "",
+	TQueueWithInt<TTransactionallySafeMpscQueue>, TQueueWithString<TTransactionallySafeMpscQueue>)
 {
-	// Non-transactional, trivial and complex types
-	TestSpscQueueCorrectness<TTransactionallySafeSpscQueue<int>>(10000, /*bTransact=*/false);
-	TestSpscQueueCorrectness<TTransactionallySafeSpscQueue<FString>>(10000, /*bTransact=*/false);
+	// Non-transactional
+	TestQueueCorrectness<typename TestType::Queue, 8>(10000, /*bTransact=*/false, TestType::Producer);
 
-	// Transactional, trivial and complex types
-	TestSpscQueueCorrectness<TTransactionallySafeSpscQueue<int>>(10000, /*bTransact=*/true);
-	TestSpscQueueCorrectness<TTransactionallySafeSpscQueue<FString>>(10000, /*bTransact=*/true);
+	// Transactional
+	TestQueueCorrectness<typename TestType::Queue, 8>(10000, /*bTransact=*/true, TestType::Producer);
 }
diff --git a/Engine/Source/Runtime/Core/Public/Containers/TransactionallySafeMpscQueue.h b/Engine/Source/Runtime/Core/Public/Containers/TransactionallySafeMpscQueue.h
new file mode 100644
--- /dev/null
+++ b/Engine/Source/Runtime/Core/Public/Containers/TransactionallySafeMpscQueue.h
@@ -0,0 +1,189 @@
+// Copyright Epic Games, Inc. All Rights Reserved.
+
+#pragma once
+
+#include "CoreTypes.h"
+#include "Async/TransactionallySafeMutex.h"
+#include "Misc/Optional.h"
+#include "Misc/ScopeLock.h"
+#include "Templates/MemoryOps.h"
+#include "Templates/TypeCompatibleBytes.h"
+#include "Templates/UnrealTemplate.h"
+
+/**
+ * Fast, transactionally-safe multi-producer/single-consumer unbounded concurrent queue. 
+ * Doesn't free memory until destruction but recycles consumed items.
+ * 
+ * The transactionally-safe queue uses a mutex to enforce thread safety instead of atomic 
+ * operations. The difference in performance should be negligible unless you are CPU-bound
+ * on constantly enqueueing and dequeueing objects as fast as possible.
+ * 
+ * It is not safe to spin-wait on Dequeue from within an AutoRTFM transaction!
+ * The other thread's Enqueue will be blocked on the mutex, so you will deadlock inside the 
+ * spin-wait.
+ */
+template<typename T, typename AllocatorType = FMemory>
+class TTransactionallySafeMpscQueue final
+{
+public:
+	using ElementType = T;
+
+	UE_NONCOPYABLE(TTransactionallySafeMpscQueue);
+
+	[[nodiscard]] TTransactionallySafeMpscQueue()
+	{
+		FNode* Node = ::new (AllocatorType::Malloc(sizeof(FNode), alignof(FNode))) FNode;
+		Head = First = Tail = TailCopy = Node;
+	}
+
+	~TTransactionallySafeMpscQueue()
+	{
+		// Nobody should have a reference to this class anymore, but we still take the mutex to
+		// guarantee that the queue isn't modified by another thread while a transaction is underway.
+		UE::TScopeLock Lock(Mutex);
+		FNode* Node = First;
+
+		// Delete all sentinel or unoccupied nodes.
+		bool bContinue = false;
+		do
+		{
+			FNode* Next = Node->Next;
+			bContinue = Node != Tail;
+			AllocatorType::Free(Node);
+			Node = Next;
+		} 
+		while (bContinue);
+
+		// Destroy and free all occupied nodes.
+		while (Node != nullptr)
+		{
+			FNode* Next = Node->Next;
+			DestructItem((ElementType*)&Node->Value);
+			AllocatorType::Free(Node);
+			Node = Next;
+		}
+	}
+
+	template <typename... ArgTypes>
+	void Enqueue(ArgTypes&&... Args)
+	{
+		FNode* Node = AllocNode();
+		::new ((void*)&Node->Value) ElementType(Forward<ArgTypes>(Args)...);
+
+		UE::TScopeLock Lock(Mutex);
+		Head->Next = Node;
+		Head = Node;
+	}
+
+	// Returns NullOpt if the queue is empty.
+	// Spin-waiting on Dequeue from within an AutoRTFM transaction is likely to deadlock,
+	// as the matching Enqueue will be waiting on Mutex.
+	TOptional<ElementType> Dequeue()
+	{
+		Mutex.Lock();
+		FNode* LocalTail = Tail;
+		FNode* LocalTailNext = LocalTail->Next;
+		Mutex.Unlock();
+
+		if (LocalTailNext == nullptr)
+		{
+			return NullOpt;
+		}
+
+		ElementType* TailNextValue = (ElementType*)&LocalTailNext->Value;
+		TOptional<ElementType> Value{ MoveTemp(*TailNextValue) };
+		DestructItem(TailNextValue);
+
+		Mutex.Lock();
+		Tail = LocalTailNext;
+		Mutex.Unlock();
+
+		return Value;
+	}
+
+	bool Dequeue(ElementType& OutElem)
+	{
+		TOptional<ElementType> LocalElement = Dequeue();
+		if (LocalElement.IsSet())
+		{
+			OutElem = MoveTempIfPossible(LocalElement.GetValue());
+			return true;
+		}
+		
+		return false;
+	}
+
+	[[nodiscard]] bool IsEmpty() const
+	{
+		UE::TScopeLock Lock(Mutex);
+		FNode* LocalTail = Tail;
+		FNode* LocalTailNext = LocalTail->Next;
+
+		return LocalTailNext == nullptr;
+	}
+
+	// As there can be only one consumer, a consumer can safely "peek" the tail of the queue.
+	// Returns a pointer to the tail if the queue is not empty, nullptr otherwise.
+	// There's no overload with TOptional as it doesn't support references.
+	[[nodiscard]] ElementType* Peek() const
+	{
+		UE::TScopeLock Lock(Mutex);
+		FNode* LocalTail = Tail;
+		FNode* LocalTailNext = LocalTail->Next;
+
+		return LocalTailNext
+			? (ElementType*)&LocalTailNext->Value
+			: nullptr;
+	}
+
+private:
+	struct FNode
+	{
+		FNode* Next = nullptr;
+		TTypeCompatibleBytes<ElementType> Value;
+	};
+
+private:
+	FNode* AllocFromCache()
+	{
+		// The mutex must _already be held_ when calling this function.
+		FNode* Node = First;
+		First = First->Next;
+		Node->Next = nullptr;
+		return Node;
+	}
+
+	FNode* AllocNode()
+	{
+		// First, we attempt to allocate a node from internal node cache.
+		{
+			UE::TScopeLock Lock(Mutex);
+			if (First != TailCopy)
+			{
+				return AllocFromCache();
+			}
+
+			TailCopy = Tail;
+			if (First != TailCopy)
+			{
+				return AllocFromCache();
+			}
+		}
+
+		// Our cache is empty; allocate via ::operator new() instead.
+		return ::new (AllocatorType::Malloc(sizeof(FNode), alignof(FNode))) FNode();
+	}
+
+private:
+	// This mutex guards all accesses to this structure or to FNode::Next.
+	mutable UE::FTransactionallySafeMutex Mutex;
+
+	// consumer part 
+	// accessed mainly by consumer, infrequently by producer 
+	FNode* Tail; // tail of the queue 
+	// producer part 
+	// accessed only by producer 
+	FNode* Head; // head of the queue
+	FNode* First; // last unused node (tail of node cache) 
+	FNode* TailCopy; // points to the initially-created tail
+};
diff --git a/Engine/Source/Runtime/Core/Public/Containers/TransactionallySafeSpscQueue.h b/Engine/Source/Runtime/Core/Public/Containers/TransactionallySafeSpscQueue.h
--- a/Engine/Source/Runtime/Core/Public/Containers/TransactionallySafeSpscQueue.h
+++ b/Engine/Source/Runtime/Core/Public/Containers/TransactionallySafeSpscQueue.h
@@ -2,191 +2,18 @@
 
 #pragma once
 
-#include "CoreTypes.h"
-#include "Async/TransactionallySafeMutex.h"
-#include "Misc/Optional.h"
-#include "Misc/ScopeLock.h"
-#include "Templates/MemoryOps.h"
-#include "Templates/TypeCompatibleBytes.h"
-#include "Templates/UnrealTemplate.h"
+#include "TransactionallySafeMpscQueue.h"
 
 /**
  * Fast, transactionally-safe single-producer/single-consumer unbounded concurrent queue. 
- * Doesn't free memory until destruction but recycles consumed items.
- * Based on TSpscQueue, which is itself based on http://www.1024cores.net/home/lock-free-algorithms/queues/unbounded-spsc-queue
+ * Doesn't free memory until destruction but recycles consumed items. Based on TSpscQueue.
  * 
- * The transactionally-safe queue uses a mutex to enforce thread safety instead of atomic 
- * operations. The difference in performance compared to a TSpscQueue should be negligible 
- * unless you are CPU-bound on constantly enqueueing and dequeueing objects as fast as 
- * possible.
+ * This is just an alias to `TTransactionallySafeMpscQueue`: because we use locks to enforce
+ * atomicity the queues are implementation wise the same.
  * 
  * It is not safe to spin-wait on Dequeue from within an AutoRTFM transaction!
  * The other thread's Enqueue will be blocked on the mutex, so you will deadlock inside the 
  * spin-wait. This class works best with the game thread as the producer, and a separate 
  * helper thread as the consumer.
  */
-template<typename T, typename AllocatorType = FMemory>
-class TTransactionallySafeSpscQueue final
-{
-public:
-	using ElementType = T;
-
-	UE_NONCOPYABLE(TTransactionallySafeSpscQueue);
-
-	[[nodiscard]] TTransactionallySafeSpscQueue()
-	{
-		FNode* Node = ::new(AllocatorType::Malloc(sizeof(FNode), alignof(FNode))) FNode;
-		Head = First = Tail = TailCopy = Node;
-	}
-
-	~TTransactionallySafeSpscQueue()
-	{
-		// Nobody should have a reference to this class anymore, but we still take the mutex to
-		// guarantee that the queue isn't modified by another thread while a transaction is underway.
-		UE::TScopeLock Lock(Mutex);
-		FNode* Node = First;
-
-		// Delete all sentinel or unoccupied nodes.
-		bool bContinue = false;
-		do
-		{
-			FNode* Next = Node->Next;
-			bContinue = Node != Tail;
-			AllocatorType::Free(Node);
-			Node = Next;
-		} 
-		while (bContinue);
-
-		// Destroy and free all occupied nodes.
-		while (Node != nullptr)
-		{
-			FNode* Next = Node->Next;
-			DestructItem((ElementType*)&Node->Value);
-			AllocatorType::Free(Node);
-			Node = Next;
-		}
-	}
-
-	template <typename... ArgTypes>
-	void Enqueue(ArgTypes&&... Args)
-	{
-		FNode* Node = AllocNode();
-		::new((void*)&Node->Value) ElementType(Forward<ArgTypes>(Args)...);
-
-		UE::TScopeLock Lock(Mutex);
-		Head->Next = Node;
-		Head = Node;
-	}
-
-	// Returns NullOpt if the queue is empty.
-	// Spin-waiting on Dequeue from within an AutoRTFM transaction is likely to deadlock,
-	// as the matching Enqueue will be waiting on Mutex.
-	TOptional<ElementType> Dequeue()
-	{
-		Mutex.Lock();
-		FNode* LocalTail = Tail;
-		FNode* LocalTailNext = LocalTail->Next;
-		Mutex.Unlock();
-
-		if (LocalTailNext == nullptr)
-		{
-			return NullOpt;
-		}
-
-		ElementType* TailNextValue = (ElementType*)&LocalTailNext->Value;
-		TOptional<ElementType> Value{ MoveTemp(*TailNextValue) };
-		DestructItem(TailNextValue);
-
-		Mutex.Lock();
-		Tail = LocalTailNext;
-		Mutex.Unlock();
-
-		return Value;
-	}
-
-	bool Dequeue(ElementType& OutElem)
-	{
-		TOptional<ElementType> LocalElement = Dequeue();
-		if (LocalElement.IsSet())
-		{
-			OutElem = MoveTempIfPossible(LocalElement.GetValue());
-			return true;
-		}
-		
-		return false;
-	}
-
-	[[nodiscard]] bool IsEmpty() const
-	{
-		UE::TScopeLock Lock(Mutex);
-		FNode* LocalTail = Tail;
-		FNode* LocalTailNext = LocalTail->Next;
-
-		return LocalTailNext == nullptr;
-	}
-
-	// As there can be only one consumer, a consumer can safely "peek" the tail of the queue.
-	// Returns a pointer to the tail if the queue is not empty, nullptr otherwise.
-	// There's no overload with TOptional as it doesn't support references.
-	[[nodiscard]] ElementType* Peek() const
-	{
-		UE::TScopeLock Lock(Mutex);
-		FNode* LocalTail = Tail;
-		FNode* LocalTailNext = LocalTail->Next;
-
-		return LocalTailNext
-			? (ElementType*)&LocalTailNext->Value
-			: nullptr;
-	}
-
-private:
-	struct FNode
-	{
-		FNode* Next = nullptr;
-		TTypeCompatibleBytes<ElementType> Value;
-	};
-
-private:
-	FNode* AllocFromCache()
-	{
-		// The mutex must _already be held_ when calling this function.
-		FNode* Node = First;
-		First = First->Next;
-		Node->Next = nullptr;
-		return Node;
-	}
-
-	FNode* AllocNode()
-	{
-		// First, we attempt to allocate a node from internal node cache.
-		{
-			UE::TScopeLock Lock(Mutex);
-			if (First != TailCopy)
-			{
-				return AllocFromCache();
-			}
-
-			TailCopy = Tail;
-			if (First != TailCopy)
-			{
-				return AllocFromCache();
-			}
-		}
-
-		// Our cache is empty; allocate via ::operator new() instead.
-		return ::new(AllocatorType::Malloc(sizeof(FNode), alignof(FNode))) FNode();
-	}
-
-private:
-	// This mutex guards all accesses to this structure or to FNode::Next.
-	mutable UE::FTransactionallySafeMutex Mutex;
-
-	// consumer part 
-	// accessed mainly by consumer, infrequently by producer 
-	FNode* Tail; // tail of the queue 
-	// producer part 
-	// accessed only by producer 
-	FNode* Head; // head of the queue
-	FNode* First; // last unused node (tail of node cache) 
-	FNode* TailCopy; // points to the initially-created tail
-};
\ No newline at end of file
+template<typename T, typename AllocatorType = FMemory> using TTransactionallySafeSpscQueue = TTransactionallySafeMpscQueue<T, AllocatorType>;
