[Insights Memory] Fixed assert for array overflow when attempting compression

#rnx
#Synced-CL 49130107

diff --git a/Engine/Source/Runtime/Core/Public/ProfilingDebugging/CallstackTrace.h b/Engine/Source/Runtime/Core/Public/ProfilingDebugging/CallstackTrace.h
--- a/Engine/Source/Runtime/Core/Public/ProfilingDebugging/CallstackTrace.h
+++ b/Engine/Source/Runtime/Core/Public/ProfilingDebugging/CallstackTrace.h
@@ -125,7 +125,7 @@ namespace FCallstackXORAndRLE
 {
 	inline uint32 Compress(TConstArrayView<uint64> Addresses, TArrayView<uint8> Compressed)
 	{
-		check(Addresses.Num() * 8 < Compressed.Num());
+		check(Addresses.Num() * 8 <= Compressed.Num());
 		uint32 OutIdx = 0;
 		uint32 CompressedBytes = 0;
 		uint64 LastAddress = 0;
