Fix constexpr construction of FLinearColor with FColor not compiling, because helper static was not itself constexpr

#Synced-CL 49102229

diff --git a/Engine/Source/Runtime/Core/Private/Math/Color.cpp b/Engine/Source/Runtime/Core/Private/Math/Color.cpp
--- a/Engine/Source/Runtime/Core/Private/Math/Color.cpp
+++ b/Engine/Source/Runtime/Core/Private/Math/Color.cpp
@@ -645,66 +645,3 @@ float FLinearColor::Pow22OneOver255Table[256] =
 	0.875137582365205f, 0.883179737672745f, 0.891262036813419f, 0.899384513046529f, 0.907547199521614f, 0.915750129279253f, 0.923993335251873f, 0.932276850264543f, 0.940600707035753f,
 	0.948964938178195f, 0.957369576199527f, 0.96581465350313f, 0.974300202388861f, 0.982826255053791f, 0.99139284359294f, 1.0f
 };
-
-
-/**
-* Table for fast FColor -> FLinearColor conversion.
-*
-* Color > 0.04045 ? pow( Color * (1.0 / 1.055) + 0.0521327, 2.4 ) : Color * (1.0 / 12.92);
-*/
-float FLinearColor::sRGBToLinearTable[256] =
-{
-	0.0f,
-	0.000303526983548838f, 0.000607053967097675f, 0.000910580950646512f, 0.00121410793419535f, 0.00151763491774419f,
-	0.00182116190129302f, 0.00212468888484186f, 0.0024282158683907f, 0.00273174285193954f, 0.00303526983548838f,
-	0.00334653564113713f, 0.00367650719436314f, 0.00402471688178252f, 0.00439144189356217f, 0.00477695332960869f,
-	0.005181516543916f, 0.00560539145834456f, 0.00604883284946662f, 0.00651209061157708f, 0.00699540999852809f,
-	0.00749903184667767f, 0.00802319278093555f, 0.0085681254056307f, 0.00913405848170623f, 0.00972121709156193f,
-	0.0103298227927056f, 0.0109600937612386f, 0.0116122449260844f, 0.012286488094766f, 0.0129830320714536f,
-	0.0137020827679224f, 0.0144438433080002f, 0.0152085141260192f, 0.0159962930597398f, 0.0168073754381669f,
-	0.0176419541646397f, 0.0185002197955389f, 0.0193823606149269f, 0.0202885627054049f, 0.0212190100154473f,
-	0.0221738844234532f, 0.02315336579873f, 0.0241576320596103f, 0.0251868592288862f, 0.0262412214867272f,
-	0.0273208912212394f, 0.0284260390768075f, 0.0295568340003534f, 0.0307134432856324f, 0.0318960326156814f,
-	0.0331047661035236f, 0.0343398063312275f, 0.0356013143874111f, 0.0368894499032755f, 0.0382043710872463f,
-	0.0395462347582974f, 0.0409151963780232f, 0.0423114100815264f, 0.0437350287071788f, 0.0451862038253117f,
-	0.0466650857658898f, 0.0481718236452158f, 0.049706565391714f, 0.0512694577708345f, 0.0528606464091205f,
-	0.0544802758174765f, 0.0561284894136735f, 0.0578054295441256f, 0.0595112375049707f, 0.0612460535624849f,
-	0.0630100169728596f, 0.0648032660013696f, 0.0666259379409563f, 0.0684781691302512f, 0.070360094971063f,
-	0.0722718499453493f, 0.0742135676316953f, 0.0761853807213167f, 0.0781874210336082f, 0.0802198195312533f,
-	0.0822827063349132f, 0.0843762107375113f, 0.0865004612181274f, 0.0886555854555171f, 0.0908417103412699f,
-	0.0930589619926197f, 0.0953074657649191f, 0.0975873462637915f, 0.0998987273569704f, 0.102241732185838f,
-	0.104616483176675f, 0.107023102051626f, 0.109461709839399f, 0.1119324268857f, 0.114435372863418f,
-	0.116970666782559f, 0.119538426999953f, 0.122138771228724f, 0.124771816547542f, 0.127437679409664f,
-	0.130136475651761f, 0.132868320502552f, 0.135633328591233f, 0.138431613955729f, 0.141263290050755f,
-	0.144128469755705f, 0.147027265382362f, 0.149959788682454f, 0.152926150855031f, 0.155926462553701f,
-	0.158960833893705f, 0.162029374458845f, 0.16513219330827f, 0.168269398983119f, 0.171441099513036f,
-	0.174647402422543f, 0.17788841473729f, 0.181164242990184f, 0.184474993227387f, 0.187820771014205f,
-	0.191201681440861f, 0.194617829128147f, 0.198069318232982f, 0.201556252453853f, 0.205078735036156f,
-	0.208636868777438f, 0.212230756032542f, 0.215860498718652f, 0.219526198320249f, 0.223227955893977f,
-	0.226965872073417f, 0.23074004707378f, 0.23455058069651f, 0.238397572333811f, 0.242281120973093f,
-	0.246201325201334f, 0.250158283209375f, 0.254152092796134f, 0.258182851372752f, 0.262250655966664f,
-	0.266355603225604f, 0.270497789421545f, 0.274677310454565f, 0.278894261856656f, 0.283148738795466f,
-	0.287440836077983f, 0.291770648154158f, 0.296138269120463f, 0.300543792723403f, 0.304987312362961f,
-	0.309468921095997f, 0.313988711639584f, 0.3185467763743f, 0.323143207347467f, 0.32777809627633f,
-	0.332451534551205f, 0.337163613238559f, 0.341914423084057f, 0.346704054515559f, 0.351532597646068f,
-	0.356400142276637f, 0.361306777899234f, 0.36625259369956f, 0.371237678559833f, 0.376262121061519f,
-	0.381326009488037f, 0.386429431827418f, 0.39157247577492f, 0.396755228735618f, 0.401977777826949f,
-	0.407240209881218f, 0.41254261144808f, 0.417885068796976f, 0.423267667919539f, 0.428690494531971f,
-	0.434153634077377f, 0.439657171728079f, 0.445201192387887f, 0.450785780694349f, 0.456411021020965f,
-	0.462076997479369f, 0.467783793921492f, 0.473531493941681f, 0.479320180878805f, 0.485149937818323f,
-	0.491020847594331f, 0.496932992791578f, 0.502886455747457f, 0.50888131855397f, 0.514917663059676f,
-	0.520995570871595f, 0.527115123357109f, 0.533276401645826f, 0.539479486631421f, 0.545724458973463f,
-	0.552011399099209f, 0.558340387205378f, 0.56471150325991f, 0.571124827003694f, 0.577580437952282f,
-	0.584078415397575f, 0.590618838409497f, 0.597201785837643f, 0.603827336312907f, 0.610495568249093f,
-	0.617206559844509f, 0.623960389083534f, 0.630757133738175f, 0.637596871369601f, 0.644479679329661f,
-	0.651405634762384f, 0.658374814605461f, 0.665387295591707f, 0.672443154250516f, 0.679542466909286f,
-	0.686685309694841f, 0.693871758534824f, 0.701101889159085f, 0.708375777101046f, 0.71569349769906f,
-	0.723055126097739f, 0.730460737249286f, 0.737910405914797f, 0.745404206665559f, 0.752942213884326f,
-	0.760524501766589f, 0.768151144321824f, 0.775822215374732f, 0.783537788566466f, 0.791297937355839f,
-	0.799102735020525f, 0.806952254658248f, 0.81484656918795f, 0.822785751350956f, 0.830769873712124f,
-	0.838799008660978f, 0.846873228412837f, 0.854992605009927f, 0.863157210322481f, 0.871367116049835f,
-	0.879622393721502f, 0.887923114698241f, 0.896269350173118f, 0.904661171172551f, 0.913098648557343f,
-	0.921581853023715f, 0.930110855104312f, 0.938685725169219f, 0.947306533426946f, 0.955973349925421f,
-	0.964686244552961f, 0.973445287039244f, 0.982250546956257f, 0.991102093719252f, 1.0f
-};
-
diff --git a/Engine/Source/Runtime/Core/Public/Math/Color.h b/Engine/Source/Runtime/Core/Public/Math/Color.h
--- a/Engine/Source/Runtime/Core/Public/Math/Color.h
+++ b/Engine/Source/Runtime/Core/Public/Math/Color.h
@@ -63,8 +63,65 @@ struct FLinearColor
 	/** Static lookup table used for FColor -> FLinearColor conversion. Pow(2.2) */
 	static float Pow22OneOver255Table[256];
 
-	/** Static lookup table used for FColor -> FLinearColor conversion. sRGB */
-	static CORE_API float sRGBToLinearTable[256];
+	/** Static lookup table used for fast FColor -> FLinearColor conversion. sRGB 
+	 * 
+	 *  Color > 0.04045 ? pow( Color * (1.0 / 1.055) + 0.0521327, 2.4 ) : Color * (1.0 / 12.92);
+	 */
+	static constexpr float sRGBToLinearTable[] =
+	{
+		0.0f,
+		0.000303526983548838f, 0.000607053967097675f, 0.000910580950646512f, 0.00121410793419535f, 0.00151763491774419f,
+		0.00182116190129302f, 0.00212468888484186f, 0.0024282158683907f, 0.00273174285193954f, 0.00303526983548838f,
+		0.00334653564113713f, 0.00367650719436314f, 0.00402471688178252f, 0.00439144189356217f, 0.00477695332960869f,
+		0.005181516543916f, 0.00560539145834456f, 0.00604883284946662f, 0.00651209061157708f, 0.00699540999852809f,
+		0.00749903184667767f, 0.00802319278093555f, 0.0085681254056307f, 0.00913405848170623f, 0.00972121709156193f,
+		0.0103298227927056f, 0.0109600937612386f, 0.0116122449260844f, 0.012286488094766f, 0.0129830320714536f,
+		0.0137020827679224f, 0.0144438433080002f, 0.0152085141260192f, 0.0159962930597398f, 0.0168073754381669f,
+		0.0176419541646397f, 0.0185002197955389f, 0.0193823606149269f, 0.0202885627054049f, 0.0212190100154473f,
+		0.0221738844234532f, 0.02315336579873f, 0.0241576320596103f, 0.0251868592288862f, 0.0262412214867272f,
+		0.0273208912212394f, 0.0284260390768075f, 0.0295568340003534f, 0.0307134432856324f, 0.0318960326156814f,
+		0.0331047661035236f, 0.0343398063312275f, 0.0356013143874111f, 0.0368894499032755f, 0.0382043710872463f,
+		0.0395462347582974f, 0.0409151963780232f, 0.0423114100815264f, 0.0437350287071788f, 0.0451862038253117f,
+		0.0466650857658898f, 0.0481718236452158f, 0.049706565391714f, 0.0512694577708345f, 0.0528606464091205f,
+		0.0544802758174765f, 0.0561284894136735f, 0.0578054295441256f, 0.0595112375049707f, 0.0612460535624849f,
+		0.0630100169728596f, 0.0648032660013696f, 0.0666259379409563f, 0.0684781691302512f, 0.070360094971063f,
+		0.0722718499453493f, 0.0742135676316953f, 0.0761853807213167f, 0.0781874210336082f, 0.0802198195312533f,
+		0.0822827063349132f, 0.0843762107375113f, 0.0865004612181274f, 0.0886555854555171f, 0.0908417103412699f,
+		0.0930589619926197f, 0.0953074657649191f, 0.0975873462637915f, 0.0998987273569704f, 0.102241732185838f,
+		0.104616483176675f, 0.107023102051626f, 0.109461709839399f, 0.1119324268857f, 0.114435372863418f,
+		0.116970666782559f, 0.119538426999953f, 0.122138771228724f, 0.124771816547542f, 0.127437679409664f,
+		0.130136475651761f, 0.132868320502552f, 0.135633328591233f, 0.138431613955729f, 0.141263290050755f,
+		0.144128469755705f, 0.147027265382362f, 0.149959788682454f, 0.152926150855031f, 0.155926462553701f,
+		0.158960833893705f, 0.162029374458845f, 0.16513219330827f, 0.168269398983119f, 0.171441099513036f,
+		0.174647402422543f, 0.17788841473729f, 0.181164242990184f, 0.184474993227387f, 0.187820771014205f,
+		0.191201681440861f, 0.194617829128147f, 0.198069318232982f, 0.201556252453853f, 0.205078735036156f,
+		0.208636868777438f, 0.212230756032542f, 0.215860498718652f, 0.219526198320249f, 0.223227955893977f,
+		0.226965872073417f, 0.23074004707378f, 0.23455058069651f, 0.238397572333811f, 0.242281120973093f,
+		0.246201325201334f, 0.250158283209375f, 0.254152092796134f, 0.258182851372752f, 0.262250655966664f,
+		0.266355603225604f, 0.270497789421545f, 0.274677310454565f, 0.278894261856656f, 0.283148738795466f,
+		0.287440836077983f, 0.291770648154158f, 0.296138269120463f, 0.300543792723403f, 0.304987312362961f,
+		0.309468921095997f, 0.313988711639584f, 0.3185467763743f, 0.323143207347467f, 0.32777809627633f,
+		0.332451534551205f, 0.337163613238559f, 0.341914423084057f, 0.346704054515559f, 0.351532597646068f,
+		0.356400142276637f, 0.361306777899234f, 0.36625259369956f, 0.371237678559833f, 0.376262121061519f,
+		0.381326009488037f, 0.386429431827418f, 0.39157247577492f, 0.396755228735618f, 0.401977777826949f,
+		0.407240209881218f, 0.41254261144808f, 0.417885068796976f, 0.423267667919539f, 0.428690494531971f,
+		0.434153634077377f, 0.439657171728079f, 0.445201192387887f, 0.450785780694349f, 0.456411021020965f,
+		0.462076997479369f, 0.467783793921492f, 0.473531493941681f, 0.479320180878805f, 0.485149937818323f,
+		0.491020847594331f, 0.496932992791578f, 0.502886455747457f, 0.50888131855397f, 0.514917663059676f,
+		0.520995570871595f, 0.527115123357109f, 0.533276401645826f, 0.539479486631421f, 0.545724458973463f,
+		0.552011399099209f, 0.558340387205378f, 0.56471150325991f, 0.571124827003694f, 0.577580437952282f,
+		0.584078415397575f, 0.590618838409497f, 0.597201785837643f, 0.603827336312907f, 0.610495568249093f,
+		0.617206559844509f, 0.623960389083534f, 0.630757133738175f, 0.637596871369601f, 0.644479679329661f,
+		0.651405634762384f, 0.658374814605461f, 0.665387295591707f, 0.672443154250516f, 0.679542466909286f,
+		0.686685309694841f, 0.693871758534824f, 0.701101889159085f, 0.708375777101046f, 0.71569349769906f,
+		0.723055126097739f, 0.730460737249286f, 0.737910405914797f, 0.745404206665559f, 0.752942213884326f,
+		0.760524501766589f, 0.768151144321824f, 0.775822215374732f, 0.783537788566466f, 0.791297937355839f,
+		0.799102735020525f, 0.806952254658248f, 0.81484656918795f, 0.822785751350956f, 0.830769873712124f,
+		0.838799008660978f, 0.846873228412837f, 0.854992605009927f, 0.863157210322481f, 0.871367116049835f,
+		0.879622393721502f, 0.887923114698241f, 0.896269350173118f, 0.904661171172551f, 0.913098648557343f,
+		0.921581853023715f, 0.930110855104312f, 0.938685725169219f, 0.947306533426946f, 0.955973349925421f,
+		0.964686244552961f, 0.973445287039244f, 0.982250546956257f, 0.991102093719252f, 1.0f
+	};
 
 	FLinearColor() = default;
 	UE_FORCEINLINE_HINT explicit FLinearColor(EForceInit)
