Resolve core.NullDereference in Strnistr, Strnstr, & Strrstr

#rnx
#Synced-CL 49130599

diff --git a/Engine/Source/Runtime/Core/Public/Misc/CString.h b/Engine/Source/Runtime/Core/Public/Misc/CString.h
--- a/Engine/Source/Runtime/Core/Public/Misc/CString.h
+++ b/Engine/Source/Runtime/Core/Public/Misc/CString.h
@@ -875,6 +875,10 @@ const typename TCString<T>::CharType* TCString<T>::Strnistr(const CharType* Str,
 		checkf(InStrLen >= 0, TEXT("Invalid InStrLen: %d"), InStrLen);
 		return nullptr;
 	}
+	if (Find == nullptr || Str == nullptr)
+	{
+		return nullptr;
+	}
 
 	// get upper-case first letter of the find string (to reduce the number of full strnicmps)
 	CharType FindInitial = TChar<CharType>::ToUpper(*Find);
@@ -915,6 +919,10 @@ const typename TCString<T>::CharType* TCString<T>::Strnstr(const CharType* Str,
 		checkf(InStrLen >= 0, TEXT("Invalid InStrLen: %d"), InStrLen);
 		return nullptr;
 	}
+	if (Find == nullptr || Str == nullptr)
+	{
+		return nullptr;
+	}
 
 	// get first letter of the find string (to reduce the number of full strncmps)
 	const CharType FindInitial = *Find;
@@ -1132,6 +1140,10 @@ const typename TCString<T>::CharType* TCString<T>::Strrstr( const CharType* Stri
 template <typename T> inline
 typename TCString<T>::CharType* TCString<T>::Strrstr( CharType* String, const CharType* Find )
 {
+	if (String == nullptr || Find == nullptr)
+	{
+		return nullptr;
+	}
 	if (*Find == (CharType)0)
 	{
 		return String + Strlen(String);
