[Apple] Custom implementation of FPlatformRecursiveMutex

Added FAppleRecursiveMutex, which is os_unfair_lock_ powered.
This is now used as platform recursive mutex on Apple platforms.

Some measurements from _a big game_ shown savings of:
~4-5% on RT, ~3-4% on GT, ~10% on map load time, ~4-5% on boot time

[RN]

diff --git a/Engine/Source/Runtime/Core/Public/Apple/AppleRecursiveMutex.h b/Engine/Source/Runtime/Core/Public/Apple/AppleRecursiveMutex.h
new file mode 100644
--- /dev/null
+++ b/Engine/Source/Runtime/Core/Public/Apple/AppleRecursiveMutex.h
@@ -0,0 +1,73 @@
+// Copyright Epic Games, Inc. All Rights Reserved.
+
+#pragma once
+
+#include "CoreTypes.h"
+
+#include <atomic>
+#include <type_traits>
+
+#include <mach/port.h>
+#include <pthread.h>
+#include <os/lock.h>
+
+namespace UE
+{
+
+class FAppleRecursiveMutex final
+{
+    os_unfair_lock Mutex {OS_UNFAIR_LOCK_INIT};
+
+	static_assert(std::is_same_v<uint32, mach_port_t>);
+	std::atomic<uint32_t> ThreadId = MACH_PORT_NULL;
+	uint32 LockCount = 0;
+
+public:
+    FAppleRecursiveMutex() = default;
+
+    FAppleRecursiveMutex(const FAppleRecursiveMutex&) = delete;
+    FAppleRecursiveMutex& operator=(const FAppleRecursiveMutex&) = delete;
+
+	UE_FORCEINLINE_HINT void Lock()
+    {
+		const uint32 CurrentThreadId = (uint32)pthread_mach_thread_np(pthread_self());
+
+        if (ThreadId.load(std::memory_order_relaxed) != CurrentThreadId)
+		{
+			os_unfair_lock_lock_with_flags(&Mutex, OS_UNFAIR_LOCK_FLAG_ADAPTIVE_SPIN);
+			ThreadId.store(CurrentThreadId, std::memory_order_relaxed);
+		}
+		++LockCount;
+    }
+
+	UE_FORCEINLINE_HINT bool TryLock()
+    {
+		const uint32 CurrentThreadId = (uint32)pthread_mach_thread_np(pthread_self());
+		
+		if (ThreadId.load(std::memory_order_relaxed) == CurrentThreadId)
+		{
+			++LockCount;
+			return true;
+		}
+
+		if (os_unfair_lock_trylock(&Mutex))
+		{
+			ThreadId.store(CurrentThreadId, std::memory_order_relaxed);
+			++LockCount;
+			return true;
+		}
+
+		return false;
+    }
+
+	UE_FORCEINLINE_HINT void Unlock()
+    {
+		if ((--LockCount) == 0)
+		{
+			ThreadId.store(0, std::memory_order_relaxed);
+			os_unfair_lock_unlock(&Mutex);
+		}
+    }
+};
+
+} // UE
diff --git a/Engine/Source/Runtime/Core/Public/IOS/IOSPlatformMutex.h b/Engine/Source/Runtime/Core/Public/IOS/IOSPlatformMutex.h
--- a/Engine/Source/Runtime/Core/Public/IOS/IOSPlatformMutex.h
+++ b/Engine/Source/Runtime/Core/Public/IOS/IOSPlatformMutex.h
@@ -2,14 +2,14 @@
 
 #pragma once
 
+#include "Apple/AppleRecursiveMutex.h"
 #include "GenericPlatform/GenericPlatformMutex.h"
-#include "HAL/PThreadsRecursiveMutex.h"
 #include "HAL/PThreadsSharedMutex.h"
 
 namespace UE
 {
 
-using FPlatformRecursiveMutex = FPThreadsRecursiveMutex;
+using FPlatformRecursiveMutex = FAppleRecursiveMutex;
 using FPlatformSharedMutex = FPThreadsSharedMutex;
 using FPlatformSystemWideMutex = FPlatformSystemWideMutexNotImplemented;
 
diff --git a/Engine/Source/Runtime/Core/Public/Mac/MacPlatformMutex.h b/Engine/Source/Runtime/Core/Public/Mac/MacPlatformMutex.h
--- a/Engine/Source/Runtime/Core/Public/Mac/MacPlatformMutex.h
+++ b/Engine/Source/Runtime/Core/Public/Mac/MacPlatformMutex.h
@@ -2,8 +2,8 @@
 
 #pragma once
 
+#include "Apple/AppleRecursiveMutex.h"
 #include "Containers/UnrealString.h"
-#include "HAL/PThreadsRecursiveMutex.h"
 #include "HAL/PThreadsSharedMutex.h"
 #include "Misc/Timespan.h"
 
@@ -37,7 +37,7 @@ private:
 	int32 FileHandle;
 };
 
-using FPlatformRecursiveMutex = FPThreadsRecursiveMutex;
+using FPlatformRecursiveMutex = FAppleRecursiveMutex;
 using FPlatformSharedMutex = FPThreadsSharedMutex;
 using FPlatformSystemWideMutex = FMacSystemWideMutex;
 
