Added TupleFwd.h.

diff --git a/Engine/Source/Developer/LowLevelTestsRunner/Public/TestHarness.h b/Engine/Source/Developer/LowLevelTestsRunner/Public/TestHarness.h
--- a/Engine/Source/Developer/LowLevelTestsRunner/Public/TestHarness.h
+++ b/Engine/Source/Developer/LowLevelTestsRunner/Public/TestHarness.h
@@ -10,6 +10,7 @@
 #include "CoreTypes.h"
 #include "Containers/ContainersFwd.h"
 #include "Containers/StringFwd.h"
+#include "Templates/TupleFwd.h"
 
 #if PLATFORM_MAC
 // Fix type redefinition error of FVector
@@ -108,8 +109,6 @@ THIRD_PARTY_INCLUDES_END
 #define DISABLED_SECTION(...) auto LLT_CONCAT(disabled_section_,__LINE__) = []()
 
 // Tell Catch how to print TTuple<KeyType, ValueType>
-template <typename... Types> struct TTuple;
-
 template <typename KeyType, typename ValueType>
 struct Catch::StringMaker<TTuple<KeyType, ValueType>>
 {
diff --git a/Engine/Source/Runtime/Core/Public/Misc/Crc.h b/Engine/Source/Runtime/Core/Public/Misc/Crc.h
--- a/Engine/Source/Runtime/Core/Public/Misc/Crc.h
+++ b/Engine/Source/Runtime/Core/Public/Misc/Crc.h
@@ -7,12 +7,10 @@
 #include "Misc/AssertionMacros.h"
 #include "Misc/CString.h"
 #include "Misc/Char.h"
+#include "Templates/TupleFwd.h"
 #include "Templates/UnrealTypeTraits.h"
 #include "Traits/IsCharType.h"
 
-template <typename... Types>
-struct TTuple;
-
 /** 
  * CRC hash generation for different types of input data
  **/
diff --git a/Engine/Source/Runtime/Core/Public/Templates/Tuple.h b/Engine/Source/Runtime/Core/Public/Templates/Tuple.h
--- a/Engine/Source/Runtime/Core/Public/Templates/Tuple.h
+++ b/Engine/Source/Runtime/Core/Public/Templates/Tuple.h
@@ -8,6 +8,7 @@
 #include "Templates/UnrealTemplate.h"
 #include "Delegates/IntegerSequence.h"
 #include "Templates/Invoke.h"
+#include "Templates/TupleFwd.h"
 #include "Serialization/StructuredArchive.h"
 #include "Serialization/MemoryLayout.h"
 #include "Templates/Requires.h"
@@ -28,9 +29,6 @@
 
 class FArchive;
 
-template <typename... Types>
-struct TTuple;
-
 template <typename T> constexpr bool TIsTuple_V = false;
 
 template <typename... Types> constexpr bool TIsTuple_V<               TTuple<Types...>> = true;
@@ -44,9 +42,6 @@ struct TIsTuple
 	enum { Value = TIsTuple_V<T> };
 };
 
-template <typename KeyType, typename ValueType>
-using TPair = TTuple<KeyType, ValueType>;
-
 template <typename... Types>
 constexpr TTuple<std::decay_t<Types>...> MakeTuple(Types&&... Args);
 
diff --git a/Engine/Source/Runtime/Core/Public/Templates/TupleFwd.h b/Engine/Source/Runtime/Core/Public/Templates/TupleFwd.h
new file mode 100644
--- /dev/null
+++ b/Engine/Source/Runtime/Core/Public/Templates/TupleFwd.h
@@ -0,0 +1,9 @@
+// Copyright Epic Games, Inc. All Rights Reserved.
+
+#pragma once
+
+template <typename... Types>
+struct TTuple;
+
+template <typename KeyType, typename ValueType>
+using TPair = TTuple<KeyType, ValueType>;
diff --git a/Engine/Source/Runtime/MovieScene/Public/EntitySystem/MovieSceneEntityRange.h b/Engine/Source/Runtime/MovieScene/Public/EntitySystem/MovieSceneEntityRange.h
--- a/Engine/Source/Runtime/MovieScene/Public/EntitySystem/MovieSceneEntityRange.h
+++ b/Engine/Source/Runtime/MovieScene/Public/EntitySystem/MovieSceneEntityRange.h
@@ -3,10 +3,9 @@
 #pragma once
 
 #include "Delegates/IntegerSequence.h"
+#include "Templates/TupleFwd.h"
 #include "Templates/UnrealTemplate.h"
 
-template<typename... T> struct TTuple;
-
 namespace UE
 {
 namespace MovieScene
diff --git a/Engine/Source/Runtime/MovieSceneTracks/Public/Sections/MovieSceneConsoleVariableTrackInterface.h b/Engine/Source/Runtime/MovieSceneTracks/Public/Sections/MovieSceneConsoleVariableTrackInterface.h
--- a/Engine/Source/Runtime/MovieSceneTracks/Public/Sections/MovieSceneConsoleVariableTrackInterface.h
+++ b/Engine/Source/Runtime/MovieSceneTracks/Public/Sections/MovieSceneConsoleVariableTrackInterface.h
@@ -3,15 +3,13 @@
 #pragma once
 
 #include "Containers/ContainersFwd.h"
+#include "Templates/TupleFwd.h"
 #include "UObject/ObjectMacros.h"
 #include "UObject/UObjectGlobals.h"
 #include "UObject/Interface.h"
 
 #include "MovieSceneConsoleVariableTrackInterface.generated.h"
 
-template<typename...> struct TTuple;
-class FString;
-
 UINTERFACE(MinimalAPI)
 class UMovieSceneConsoleVariableTrackInterface : public UInterface
 {
