Added new traits for getting the arity of a function and getting the member type of a pointer-to-member type.

diff --git a/Engine/Source/Runtime/Core/Private/Traits/FunctionArityTests.cpp b/Engine/Source/Runtime/Core/Private/Traits/FunctionArityTests.cpp
new file mode 100644
--- /dev/null
+++ b/Engine/Source/Runtime/Core/Private/Traits/FunctionArityTests.cpp
@@ -0,0 +1,96 @@
+// Copyright Epic Games, Inc. All Rights Reserved.
+
+#include "CoreTypes.h"
+#include "Containers/ContainersFwd.h"
+#include "Misc/OptionalFwd.h"
+#include "Templates/FunctionFwd.h"
+#include "Templates/TupleFwd.h"
+#include "Traits/FunctionArity.h"
+
+#if WITH_TESTS
+
+namespace
+{
+	// 0-arity
+	static_assert(TFunctionArity_V<void()                 > == 0);
+	static_assert(TFunctionArity_V<bool()                 > == 0);
+	static_assert(TFunctionArity_V<void() const           > == 0);
+	static_assert(TFunctionArity_V<bool() const           > == 0);
+	static_assert(TFunctionArity_V<void()       volatile  > == 0);
+	static_assert(TFunctionArity_V<bool()       volatile  > == 0);
+	static_assert(TFunctionArity_V<void() const volatile  > == 0);
+	static_assert(TFunctionArity_V<bool() const volatile  > == 0);
+	static_assert(TFunctionArity_V<void()               & > == 0);
+	static_assert(TFunctionArity_V<bool()               & > == 0);
+	static_assert(TFunctionArity_V<void() const         & > == 0);
+	static_assert(TFunctionArity_V<bool() const         & > == 0);
+	static_assert(TFunctionArity_V<void()       volatile& > == 0);
+	static_assert(TFunctionArity_V<bool()       volatile& > == 0);
+	static_assert(TFunctionArity_V<void() const volatile& > == 0);
+	static_assert(TFunctionArity_V<bool() const volatile& > == 0);
+	static_assert(TFunctionArity_V<void()               &&> == 0);
+	static_assert(TFunctionArity_V<bool()               &&> == 0);
+	static_assert(TFunctionArity_V<void() const         &&> == 0);
+	static_assert(TFunctionArity_V<bool() const         &&> == 0);
+	static_assert(TFunctionArity_V<void()       volatile&&> == 0);
+	static_assert(TFunctionArity_V<bool()       volatile&&> == 0);
+	static_assert(TFunctionArity_V<void() const volatile&&> == 0);
+	static_assert(TFunctionArity_V<bool() const volatile&&> == 0);
+
+	// 1-arity
+	static_assert(TFunctionArity_V<void(int)                 > == 1);
+	static_assert(TFunctionArity_V<bool(int)                 > == 1);
+	static_assert(TFunctionArity_V<void(int) const           > == 1);
+	static_assert(TFunctionArity_V<bool(int) const           > == 1);
+	static_assert(TFunctionArity_V<void(int)       volatile  > == 1);
+	static_assert(TFunctionArity_V<bool(int)       volatile  > == 1);
+	static_assert(TFunctionArity_V<void(int) const volatile  > == 1);
+	static_assert(TFunctionArity_V<bool(int) const volatile  > == 1);
+	static_assert(TFunctionArity_V<void(int)               & > == 1);
+	static_assert(TFunctionArity_V<bool(int)               & > == 1);
+	static_assert(TFunctionArity_V<void(int) const         & > == 1);
+	static_assert(TFunctionArity_V<bool(int) const         & > == 1);
+	static_assert(TFunctionArity_V<void(int)       volatile& > == 1);
+	static_assert(TFunctionArity_V<bool(int)       volatile& > == 1);
+	static_assert(TFunctionArity_V<void(int) const volatile& > == 1);
+	static_assert(TFunctionArity_V<bool(int) const volatile& > == 1);
+	static_assert(TFunctionArity_V<void(int)               &&> == 1);
+	static_assert(TFunctionArity_V<bool(int)               &&> == 1);
+	static_assert(TFunctionArity_V<void(int) const         &&> == 1);
+	static_assert(TFunctionArity_V<bool(int) const         &&> == 1);
+	static_assert(TFunctionArity_V<void(int)       volatile&&> == 1);
+	static_assert(TFunctionArity_V<bool(int)       volatile&&> == 1);
+	static_assert(TFunctionArity_V<void(int) const volatile&&> == 1);
+	static_assert(TFunctionArity_V<bool(int) const volatile&&> == 1);
+
+	// 2-arity
+	static_assert(TFunctionArity_V<void(int, float)                 > == 2);
+	static_assert(TFunctionArity_V<bool(int, float)                 > == 2);
+	static_assert(TFunctionArity_V<void(int, float) const           > == 2);
+	static_assert(TFunctionArity_V<bool(int, float) const           > == 2);
+	static_assert(TFunctionArity_V<void(int, float)       volatile  > == 2);
+	static_assert(TFunctionArity_V<bool(int, float)       volatile  > == 2);
+	static_assert(TFunctionArity_V<void(int, float) const volatile  > == 2);
+	static_assert(TFunctionArity_V<bool(int, float) const volatile  > == 2);
+	static_assert(TFunctionArity_V<void(int, float)               & > == 2);
+	static_assert(TFunctionArity_V<bool(int, float)               & > == 2);
+	static_assert(TFunctionArity_V<void(int, float) const         & > == 2);
+	static_assert(TFunctionArity_V<bool(int, float) const         & > == 2);
+	static_assert(TFunctionArity_V<void(int, float)       volatile& > == 2);
+	static_assert(TFunctionArity_V<bool(int, float)       volatile& > == 2);
+	static_assert(TFunctionArity_V<void(int, float) const volatile& > == 2);
+	static_assert(TFunctionArity_V<bool(int, float) const volatile& > == 2);
+	static_assert(TFunctionArity_V<void(int, float)               &&> == 2);
+	static_assert(TFunctionArity_V<bool(int, float)               &&> == 2);
+	static_assert(TFunctionArity_V<void(int, float) const         &&> == 2);
+	static_assert(TFunctionArity_V<bool(int, float) const         &&> == 2);
+	static_assert(TFunctionArity_V<void(int, float)       volatile&&> == 2);
+	static_assert(TFunctionArity_V<bool(int, float)       volatile&&> == 2);
+	static_assert(TFunctionArity_V<void(int, float) const volatile&&> == 2);
+	static_assert(TFunctionArity_V<bool(int, float) const volatile&&> == 2);
+
+	// Very complex function type
+	static_assert(TFunctionArity_V<TTuple<int, char, float>(bool, float, const char*, const TMap<int32, TArray<float>>, TOptional<TSet<char>>&&) const volatile&&> == 5);
+}
+
+#endif // WITH_TESTS
diff --git a/Engine/Source/Runtime/Core/Private/Traits/RemoveMemberPointerTests.cpp b/Engine/Source/Runtime/Core/Private/Traits/RemoveMemberPointerTests.cpp
new file mode 100644
--- /dev/null
+++ b/Engine/Source/Runtime/Core/Private/Traits/RemoveMemberPointerTests.cpp
@@ -0,0 +1,56 @@
+// Copyright Epic Games, Inc. All Rights Reserved.
+
+#include "CoreTypes.h"
+#include "Traits/RemoveMemberPtr.h"
+#include <type_traits>
+
+#if WITH_TESTS
+
+namespace
+{
+	struct FDummyType;
+
+	// Types which aren't pointer-to-members should return the same type
+	static_assert(std::is_same_v<TRemoveMemberPtr_T<void                   >, void                   >);
+	static_assert(std::is_same_v<TRemoveMemberPtr_T<int32                  >, int32                  >);
+	static_assert(std::is_same_v<TRemoveMemberPtr_T<const char*            >, const char*            >);
+	static_assert(std::is_same_v<TRemoveMemberPtr_T<bool(const char*) const>, bool(const char*) const>);
+
+	// Pointers to data members should return the data member type (with cv qualifiers)
+	static_assert(std::is_same_v<TRemoveMemberPtr_T<               int32       FDummyType::*>,                int32      >);
+	static_assert(std::is_same_v<TRemoveMemberPtr_T<const          int32       FDummyType::*>, const          int32      >);
+	static_assert(std::is_same_v<TRemoveMemberPtr_T<      volatile int32       FDummyType::*>,       volatile int32      >);
+	static_assert(std::is_same_v<TRemoveMemberPtr_T<const volatile int32       FDummyType::*>, const volatile int32      >);
+	static_assert(std::is_same_v<TRemoveMemberPtr_T<const char*                FDummyType::*>, const char*               >);
+	static_assert(std::is_same_v<TRemoveMemberPtr_T<const char* const          FDummyType::*>, const char* const         >);
+	static_assert(std::is_same_v<TRemoveMemberPtr_T<const char*       volatile FDummyType::*>, const char*       volatile>);
+	static_assert(std::is_same_v<TRemoveMemberPtr_T<const char* const volatile FDummyType::*>, const char* const volatile>);
+
+	// Pointers to member function types should return the function type (with cv- and ref-qualifiers
+	static_assert(std::is_same_v<TRemoveMemberPtr_T<void (FDummyType::*)()                                          >, void()                                          >);
+	static_assert(std::is_same_v<TRemoveMemberPtr_T<void (FDummyType::*)()                          const           >, void()                          const           >);
+	static_assert(std::is_same_v<TRemoveMemberPtr_T<void (FDummyType::*)()                                volatile  >, void()                                volatile  >);
+	static_assert(std::is_same_v<TRemoveMemberPtr_T<void (FDummyType::*)()                          const volatile  >, void()                          const volatile  >);
+	static_assert(std::is_same_v<TRemoveMemberPtr_T<void (FDummyType::*)()                                        & >, void()                                        & >);
+	static_assert(std::is_same_v<TRemoveMemberPtr_T<void (FDummyType::*)()                          const         & >, void()                          const         & >);
+	static_assert(std::is_same_v<TRemoveMemberPtr_T<void (FDummyType::*)()                                volatile& >, void()                                volatile& >);
+	static_assert(std::is_same_v<TRemoveMemberPtr_T<void (FDummyType::*)()                          const volatile& >, void()                          const volatile& >);
+	static_assert(std::is_same_v<TRemoveMemberPtr_T<void (FDummyType::*)()                                        &&>, void()                                        &&>);
+	static_assert(std::is_same_v<TRemoveMemberPtr_T<void (FDummyType::*)()                          const         &&>, void()                          const         &&>);
+	static_assert(std::is_same_v<TRemoveMemberPtr_T<void (FDummyType::*)()                                volatile&&>, void()                                volatile&&>);
+	static_assert(std::is_same_v<TRemoveMemberPtr_T<void (FDummyType::*)()                          const volatile&&>, void()                          const volatile&&>);
+	static_assert(std::is_same_v<TRemoveMemberPtr_T<bool (FDummyType::*)(float, int32, const char*)                 >, bool(float, int32, const char*)                 >);
+	static_assert(std::is_same_v<TRemoveMemberPtr_T<bool (FDummyType::*)(float, int32, const char*) const           >, bool(float, int32, const char*) const           >);
+	static_assert(std::is_same_v<TRemoveMemberPtr_T<bool (FDummyType::*)(float, int32, const char*)       volatile  >, bool(float, int32, const char*)       volatile  >);
+	static_assert(std::is_same_v<TRemoveMemberPtr_T<bool (FDummyType::*)(float, int32, const char*) const volatile  >, bool(float, int32, const char*) const volatile  >);
+	static_assert(std::is_same_v<TRemoveMemberPtr_T<bool (FDummyType::*)(float, int32, const char*)               & >, bool(float, int32, const char*)               & >);
+	static_assert(std::is_same_v<TRemoveMemberPtr_T<bool (FDummyType::*)(float, int32, const char*) const         & >, bool(float, int32, const char*) const         & >);
+	static_assert(std::is_same_v<TRemoveMemberPtr_T<bool (FDummyType::*)(float, int32, const char*)       volatile& >, bool(float, int32, const char*)       volatile& >);
+	static_assert(std::is_same_v<TRemoveMemberPtr_T<bool (FDummyType::*)(float, int32, const char*) const volatile& >, bool(float, int32, const char*) const volatile& >);
+	static_assert(std::is_same_v<TRemoveMemberPtr_T<bool (FDummyType::*)(float, int32, const char*)               &&>, bool(float, int32, const char*)               &&>);
+	static_assert(std::is_same_v<TRemoveMemberPtr_T<bool (FDummyType::*)(float, int32, const char*) const         &&>, bool(float, int32, const char*) const         &&>);
+	static_assert(std::is_same_v<TRemoveMemberPtr_T<bool (FDummyType::*)(float, int32, const char*)       volatile&&>, bool(float, int32, const char*)       volatile&&>);
+	static_assert(std::is_same_v<TRemoveMemberPtr_T<bool (FDummyType::*)(float, int32, const char*) const volatile&&>, bool(float, int32, const char*) const volatile&&>);
+}
+
+#endif // WITH_TESTS
diff --git a/Engine/Source/Runtime/Core/Public/Traits/FunctionArity.h b/Engine/Source/Runtime/Core/Public/Traits/FunctionArity.h
new file mode 100644
--- /dev/null
+++ b/Engine/Source/Runtime/Core/Public/Traits/FunctionArity.h
@@ -0,0 +1,34 @@
+// Copyright Epic Games, Inc. All Rights Reserved.
+
+#pragma once
+
+/**
+ * Type trait which yields the number of arguments in a function type, e.g.:
+ *
+ * TFunctionArity_V<int(float, bool, const char*)> == 3
+ *
+ * This doesn't work for pointers or references to functions, or member function pointer types  - use the correct one of these to get a pure function type:
+ * - std::remove_pointer_t
+ * - std::remove_reference_t
+ * - TRemoveMemberPtr_T
+ */
+template <typename FunctionType>
+struct TFunctionArity
+{
+};
+
+template <typename RetType, typename... ParamTypes> struct TFunctionArity<RetType(ParamTypes...)                 > { inline static constexpr unsigned int Value = sizeof...(ParamTypes); };
+template <typename RetType, typename... ParamTypes> struct TFunctionArity<RetType(ParamTypes...)               & > { inline static constexpr unsigned int Value = sizeof...(ParamTypes); };
+template <typename RetType, typename... ParamTypes> struct TFunctionArity<RetType(ParamTypes...)               &&> { inline static constexpr unsigned int Value = sizeof...(ParamTypes); };
+template <typename RetType, typename... ParamTypes> struct TFunctionArity<RetType(ParamTypes...) const           > { inline static constexpr unsigned int Value = sizeof...(ParamTypes); };
+template <typename RetType, typename... ParamTypes> struct TFunctionArity<RetType(ParamTypes...) const         & > { inline static constexpr unsigned int Value = sizeof...(ParamTypes); };
+template <typename RetType, typename... ParamTypes> struct TFunctionArity<RetType(ParamTypes...) const         &&> { inline static constexpr unsigned int Value = sizeof...(ParamTypes); };
+template <typename RetType, typename... ParamTypes> struct TFunctionArity<RetType(ParamTypes...)       volatile  > { inline static constexpr unsigned int Value = sizeof...(ParamTypes); };
+template <typename RetType, typename... ParamTypes> struct TFunctionArity<RetType(ParamTypes...)       volatile& > { inline static constexpr unsigned int Value = sizeof...(ParamTypes); };
+template <typename RetType, typename... ParamTypes> struct TFunctionArity<RetType(ParamTypes...)       volatile&&> { inline static constexpr unsigned int Value = sizeof...(ParamTypes); };
+template <typename RetType, typename... ParamTypes> struct TFunctionArity<RetType(ParamTypes...) const volatile  > { inline static constexpr unsigned int Value = sizeof...(ParamTypes); };
+template <typename RetType, typename... ParamTypes> struct TFunctionArity<RetType(ParamTypes...) const volatile& > { inline static constexpr unsigned int Value = sizeof...(ParamTypes); };
+template <typename RetType, typename... ParamTypes> struct TFunctionArity<RetType(ParamTypes...) const volatile&&> { inline static constexpr unsigned int Value = sizeof...(ParamTypes); };
+
+template <typename FunctionType>
+inline constexpr unsigned int TFunctionArity_V = TFunctionArity<FunctionType>::Value;
diff --git a/Engine/Source/Runtime/Core/Public/Traits/RemoveMemberPtr.h b/Engine/Source/Runtime/Core/Public/Traits/RemoveMemberPtr.h
new file mode 100644
--- /dev/null
+++ b/Engine/Source/Runtime/Core/Public/Traits/RemoveMemberPtr.h
@@ -0,0 +1,24 @@
+// Copyright Epic Games, Inc. All Rights Reserved.
+
+#pragma once
+
+/**
+ * Type trait which yields the type being pointed to of a member pointer type, e.g.:
+ *
+ * TRemoveMemberFunctionPtr_T<decltype(&FVector::X)>   == double
+ * TRemoveMemberFunctionPtr_T<decltype(&FVector::Dot)> == double(const FVector&)
+ */
+template <typename T>
+struct TRemoveMemberPtr
+{
+	using Type = T;
+};
+
+template <typename OuterType, typename MemberType>
+struct TRemoveMemberPtr<MemberType OuterType::*>
+{
+	using Type = MemberType;
+};
+
+template <typename T>
+using TRemoveMemberPtr_T = typename TRemoveMemberPtr<T>::Type;
