[IOS] os_log is back

48566790 still had NSLog for !Shipping builds, due to os_log not routing to stdout by default, which broke automation, etc.
This change brings back the initial state -- os_log is used in all configurations, with logs being private in Shipping.

TLDR: Set 'OS_ACTIVITY_DT_MODE=enable' on device to route os_logs to stdout.

#rnx

diff --git a/Engine/Source/Programs/AutomationTool/Gauntlet/Unreal/Base/Gauntlet.UnrealSession.cs b/Engine/Source/Programs/AutomationTool/Gauntlet/Unreal/Base/Gauntlet.UnrealSession.cs
--- a/Engine/Source/Programs/AutomationTool/Gauntlet/Unreal/Base/Gauntlet.UnrealSession.cs
+++ b/Engine/Source/Programs/AutomationTool/Gauntlet/Unreal/Base/Gauntlet.UnrealSession.cs
@@ -1932,7 +1932,7 @@ namespace Gauntlet
 						Instance.AppInstance.Kill();
 					}
 
-					string WarningMessage = $"Encountered {Ex.GetType()} when attempting to run install {Ex.Message}";
+					string WarningMessage = $"Encountered exception when attempting to run install: {Ex}";
 					if (Ex is DeviceException)
 					{
 						UnrealDeviceReservation.MarkProblemDevice(Install.Device, WarningMessage);
diff --git a/Engine/Source/Programs/AutomationTool/IOS/IOSPlatform.Automation.cs b/Engine/Source/Programs/AutomationTool/IOS/IOSPlatform.Automation.cs
--- a/Engine/Source/Programs/AutomationTool/IOS/IOSPlatform.Automation.cs
+++ b/Engine/Source/Programs/AutomationTool/IOS/IOSPlatform.Automation.cs
@@ -1239,6 +1239,7 @@ public class IOSPlatform : ApplePlatform
 				Arguments += " -e '{\\\"APP_DISTRIBUTOR_ID_OVERRIDE\\\": \\\"com.apple.AppStore\\\"}'";
 				Arguments += " --terminate-existing";	// if it's already running on device, kill it first
 				Arguments += " --console";	// attach to the console so we can get log output
+				Arguments += " --environment-variables '{\"OS_ACTIVITY_DT_MODE\":\"enable\"}'";	// ask os_log to go into stdout/stderr
 				Arguments += " --device " + Params.DeviceNames[0];
 				Arguments += " \"" + BundleIdentifier + "\" --";
 			}
diff --git a/Engine/Source/Runtime/Core/Private/Apple/ApplePlatformMisc.cpp b/Engine/Source/Runtime/Core/Private/Apple/ApplePlatformMisc.cpp
--- a/Engine/Source/Runtime/Core/Private/Apple/ApplePlatformMisc.cpp
+++ b/Engine/Source/Runtime/Core/Private/Apple/ApplePlatformMisc.cpp
@@ -77,13 +77,25 @@ FString FApplePlatformMisc::GetEnvironmentVariable(const TCHAR* VariableName)
 
 void FApplePlatformMisc::LocalPrint(const TCHAR* Message)
 {
+	//
+	// os_log doesn't go to stdout/stderr, by default, as NSLog does.
+	//
+	// It will be there if requested via one of environment variables:
+	//   OS_ACTIVITY_DT_MODE=enable
+	//   ACTIVITY_LOG_STDERR=1 (will also include frameworks)
+	//
+	// E.g. if app is run via devicectl:
+	//
+	//   xcrun devicectl device process launch \
+	//     --device <device> \
+	//     --console \
+	//     --environment-variables '{"OS_ACTIVITY_DT_MODE":"enable"}' \
+	//     <bundle-name>
+	//
 #if !UE_BUILD_SHIPPING
-	// This goes to all output consoles, not only system one, stdout included.
-	NSLog(@"[UE] %s", TCHAR_TO_UTF8(Message));
+	os_log(OS_LOG_DEFAULT, "[UE] %{public}s", TCHAR_TO_UTF8(Message));
 #else
-	// This only goes to the system console.
-	// Private by default, would be visible in Xcode.
-	os_log(OS_LOG_DEFAULT, "[UE] %s", TCHAR_TO_UTF8(Message));
+	os_log(OS_LOG_DEFAULT, "[UE] %s", TCHAR_TO_UTF8(Message)); // Private by default, would be visible in Xcode.
 #endif
 }
 
