Deprecated TScriptDelegate::GetUObjectRef() (which despite the name, doesn't actually return a reference to what GetUObject() returns).
Deprecated TMulticastScriptDelegate::GetAllObjectRefsEvenIfUnreachable().
Added TScriptDelegate::GetRemoteId().
Reordered the tests in FMulticastDelegateProperty::Identical() to make it less likely that a remote object pointer needs resolving.

diff --git a/Engine/Source/Runtime/Core/Public/UObject/ScriptDelegates.h b/Engine/Source/Runtime/Core/Public/UObject/ScriptDelegates.h
--- a/Engine/Source/Runtime/Core/Public/UObject/ScriptDelegates.h
+++ b/Engine/Source/Runtime/Core/Public/UObject/ScriptDelegates.h
@@ -298,16 +298,27 @@ public:
 		return static_cast< const UObject* >( Object.GetEvenIfUnreachable() );
 	}
 
+	UE_DEPRECATED(5.8, "GetUObjectRef has been deprecated.")
 	WeakPtrType& GetUObjectRef()
 	{
 		return Object;
 	}
 
+	UE_DEPRECATED(5.8, "GetUObjectRef has been deprecated.")
 	const WeakPtrType& GetUObjectRef() const
 	{
 		return Object;
 	}
 
+#if UE_WITH_REMOTE_OBJECT_HANDLE
+	auto GetRemoteId() const
+	{
+		FReadAccessScope ReadScope = GetReadAccessScope();
+
+		return Object.GetRemoteId();
+	}
+#endif
+
 	/**
 	 * Gets the name of the function to call on the bound object
 	 *
@@ -795,20 +806,26 @@ public:
 	 * For advanced uses only -- you should never need call this function in normal circumstances.
  	 * @return	List of objects bound to this delegate
 	 */
-	TArray< UObject* > GetAllObjectsEvenIfUnreachable() const
+	TArray<UObject*> GetAllObjectsEvenIfUnreachable() const
 	{
-		FReadAccessScope ReadScope = GetReadAccessScope();    
+		FReadAccessScope ReadScope = GetReadAccessScope();
+		using WeakPtrType = typename UnicastDelegateType::WeakPtrType;
 		TArray<UObject*> Result;
-		for (auto* Ref : GetAllObjectRefsEvenIfUnreachable())
+		for (typename InvocationListType::TIterator CurDelegate(InvocationList); CurDelegate; ++CurDelegate)
 		{
-			Result.Add(Ref->GetEvenIfUnreachable());
+			if (UObject* CurObjectPtr = CurDelegate->GetUObjectEvenIfUnreachable())
+			{
+				Result.Add(CurObjectPtr);
+			}
 		}
 		return Result;
 	}
-	
+
+PRAGMA_DISABLE_DEPRECATION_WARNINGS
+	UE_DEPRECATED(5.8, "GetAllObjectRefsEvenIfUnreachable has been deprecated.")
 	TArray< typename UnicastDelegateType::WeakPtrType* > GetAllObjectRefsEvenIfUnreachable() const
 	{
-		FReadAccessScope ReadScope = GetReadAccessScope();        
+		FReadAccessScope ReadScope = GetReadAccessScope();
 		using WeakPtrType = typename UnicastDelegateType::WeakPtrType;
 		TArray< WeakPtrType* > OutputList;
 		for( typename InvocationListType::TIterator CurDelegate( InvocationList ); CurDelegate; ++CurDelegate )
@@ -821,6 +838,7 @@ public:
 		}
 		return OutputList;
 	}
+PRAGMA_ENABLE_DEPRECATION_WARNINGS
 
 	/**
 	 * Returns the amount of memory allocated by this delegate's invocation list, not including sizeof(*this).
diff --git a/Engine/Source/Runtime/CoreUObject/Private/UObject/PropertyMulticastDelegate.cpp b/Engine/Source/Runtime/CoreUObject/Private/UObject/PropertyMulticastDelegate.cpp
--- a/Engine/Source/Runtime/CoreUObject/Private/UObject/PropertyMulticastDelegate.cpp
+++ b/Engine/Source/Runtime/CoreUObject/Private/UObject/PropertyMulticastDelegate.cpp
@@ -152,19 +152,19 @@ bool FMulticastDelegateProperty::Identical( const void* A, const void* B, uint32
 		if ((PortFlags & PPF_AvoidRemoteObjectMigration) != 0)
 		{
 			// Before comparing the actual object pointers, compare their remote ids to prevent dereferencing objects that may be remote
-			if (BindingA.GetUObjectRef().GetRemoteId() != BindingB.GetUObjectRef().GetRemoteId())
+			if (BindingA.GetRemoteId() != BindingB.GetRemoteId())
 			{
 				return false;
 			}
 		}
 		else
 #endif
-		if (BindingA.GetUObject() != BindingB.GetUObject())
+		if (BindingA.GetFunctionName() != BindingB.GetFunctionName())
 		{
 			return false;
 		}
 
-		if (BindingA.GetFunctionName() != BindingB.GetFunctionName())
+		if (BindingA.GetUObject() != BindingB.GetUObject())
 		{
 			return false;
 		}
