- Added some thread safety to FConfigCacheIni::ReadString (which the others like GetInt go through), SafeReload, and AddDynamicLayersToHierarchy

#Synced-CL 48811006

diff --git a/Engine/Source/Runtime/Core/Private/Misc/ConfigCacheIni.cpp b/Engine/Source/Runtime/Core/Private/Misc/ConfigCacheIni.cpp
--- a/Engine/Source/Runtime/Core/Private/Misc/ConfigCacheIni.cpp
+++ b/Engine/Source/Runtime/Core/Private/Misc/ConfigCacheIni.cpp
@@ -33,6 +33,8 @@
 #include "Serialization/LargeMemoryReader.h"
 #include "Async/Async.h"
 #include "Async/ParallelFor.h"
+#include "Async/SharedLock.h"
+#include "Async/UniqueLock.h"
 #include "Misc/OutputDeviceRedirector.h"
 #include "Logging/MessageLog.h"
 #include <limits>
@@ -3440,6 +3442,9 @@ bool FConfigBranch::AddDynamicLayersToHierarchy(const TArray<FDynamicLayerInfo>&
 	{
 		SafeReload();
 	}
+	
+	// don't allow a read while modifying it here
+	UE::TUniqueLock WriteLock(InMemoryFile.SharedMutex);
 
 	static bool bDumpIniLoadInfo = FParse::Param(FCommandLine::Get(), TEXT("dumpiniloads"));
 
@@ -3598,6 +3603,9 @@ bool FConfigBranch::AddDynamicLayerStringToHierarchy(const FString& Filename, co
 		SafeReload();
 	}
 
+	// don't allow a read while modifying it here
+	UE::TUniqueLock WriteLock(InMemoryFile.SharedMutex);
+
 	bool bInsertedAtEnd = false;
 
 	// calculate a patch so we don't lose in-memory changes
@@ -3807,6 +3815,9 @@ void FConfigBranch::SafeReload()
 {
 	LLM_SCOPE(ELLMTag::ConfigSystem);
 
+	// this can modify InMemoryFile, so we need to protect it 
+	UE::TSharedLock ReadLock(InMemoryFile.SharedMutex);
+
 	double StartTime = FPlatformTime::Seconds();
 
 	// read static layers back in from disk
@@ -4663,6 +4674,10 @@ bool FConfigCacheIni::GetString( const TCHAR* Section, const TCHAR* Key, FString
 	{
 		return false;
 	}
+	
+	// don't read while something may be modifying the File (generally a Branch's InMmeoryFile) 
+	UE::TSharedLock ReadLock(File->SharedMutex);
+	
 	const FConfigSection* Sec = File->FindSection( Section );
 	if( !Sec )
 	{
@@ -4691,6 +4706,10 @@ bool FConfigCacheIni::GetText( const TCHAR* Section, const TCHAR* Key, FText& Va
 	{
 		return false;
 	}
+	
+	// don't read while something may be modifying the File (generally a Branch's InMmeoryFile) 
+	UE::TSharedLock ReadLock(File->SharedMutex);
+
 	const FConfigSection* Sec = File->FindSection( Section );
 	if( !Sec )
 	{
diff --git a/Engine/Source/Runtime/Core/Public/Misc/ConfigCacheIni.h b/Engine/Source/Runtime/Core/Public/Misc/ConfigCacheIni.h
--- a/Engine/Source/Runtime/Core/Public/Misc/ConfigCacheIni.h
+++ b/Engine/Source/Runtime/Core/Public/Misc/ConfigCacheIni.h
@@ -7,6 +7,7 @@
 #pragma once
 
 #include "Algo/Reverse.h"
+#include "Async/TransactionallySafeSharedMutex.h"
 #include "Containers/Array.h"
 #include "Containers/Map.h"
 #include "Containers/Set.h"
@@ -605,6 +606,9 @@ private:
 	// if this is set, then we track changes made to sections for saving or replaying later (currently unused)
 	FConfigCommandStream* ChangeTracker = nullptr;
 	
+	// thread safety for reading while something is modifying this file
+	UE::FTransactionallySafeSharedMutex SharedMutex;
+	
 #if UE_WITH_CONFIG_TRACKING
 	mutable TRefCountPtr<UE::ConfigAccessTracking::FFile> FileAccess;
 #endif
@@ -1158,7 +1162,7 @@ public:
 	
 	// tracks runtime changes for optimal saving
 	FConfigCommandStream RuntimeChanges;
-	
+		
 	friend FArchive& operator<<(FArchive& Ar, FConfigBranch& ConfigBranch);
 	
 	/** Appends a new INI file to the SourceIniHierarchy and combines it with the current contents
