Trim channel name tokens.

#rnx
#Synced-CL 49081113

diff --git a/Engine/Source/Runtime/Core/Private/ProfilingDebugging/TraceAuxiliary.cpp b/Engine/Source/Runtime/Core/Private/ProfilingDebugging/TraceAuxiliary.cpp
--- a/Engine/Source/Runtime/Core/Private/ProfilingDebugging/TraceAuxiliary.cpp
+++ b/Engine/Source/Runtime/Core/Private/ProfilingDebugging/TraceAuxiliary.cpp
@@ -392,7 +392,7 @@ void FTraceAuxiliaryImpl::ForEachChannel(const TCHAR* ChannelList, bool bResolve
 	for (const FStringView& Channel : ExpandedChannels)
 	{
 		TCHAR Name[80];
-		const size_t ChannelNameSize = Channel.CopyString(Name, UE_ARRAY_COUNT(Name) - 1);
+		const size_t ChannelNameSize = Channel.TrimStartAndEnd().CopyString(Name, UE_ARRAY_COUNT(Name) - 1);
 		Name[ChannelNameSize] = '\0';
 
 		Invoke(Callable, this, Name, std::forward<Ts>(Args)...);
