[Apple] Temporary disabling unfair lock on Mac

It seems that changes to that header lead to few Mac apps built for x64 fail with non-related crashes when running on Apple silicon. This was discovered by CI.
Looks like it triggers bugs in Rosetta, the investigation is still ongoing.

IOS keeps using unfair locks.

#rnx

diff --git a/Engine/Source/Runtime/Core/Public/Mac/MacPlatformMutex.h b/Engine/Source/Runtime/Core/Public/Mac/MacPlatformMutex.h
--- a/Engine/Source/Runtime/Core/Public/Mac/MacPlatformMutex.h
+++ b/Engine/Source/Runtime/Core/Public/Mac/MacPlatformMutex.h
@@ -2,8 +2,8 @@
 
 #pragma once
 
-#include "Apple/AppleRecursiveMutex.h"
 #include "Containers/UnrealString.h"
+#include "HAL/PThreadsRecursiveMutex.h"
 #include "HAL/PThreadsSharedMutex.h"
 #include "Misc/Timespan.h"
 
@@ -37,7 +37,7 @@ private:
 	int32 FileHandle;
 };
 
-using FPlatformRecursiveMutex = FAppleRecursiveMutex;
+using FPlatformRecursiveMutex = FPThreadsRecursiveMutex;
 using FPlatformSharedMutex = FPThreadsSharedMutex;
 using FPlatformSystemWideMutex = FMacSystemWideMutex;
 
