[Apple] Unfair locks are back to Mac

Still disabled for Apple Intel builds, as that code seems to trigger Rosetta2 bugs (even without os_unfair_lock_ API usage...).

#rnx

diff --git a/Engine/Source/Runtime/Core/Public/Mac/MacPlatformMutex.h b/Engine/Source/Runtime/Core/Public/Mac/MacPlatformMutex.h
--- a/Engine/Source/Runtime/Core/Public/Mac/MacPlatformMutex.h
+++ b/Engine/Source/Runtime/Core/Public/Mac/MacPlatformMutex.h
@@ -2,8 +2,13 @@
 
 #pragma once
 
+#if PLATFORM_MAC_ARM64
+#	include "Apple/AppleRecursiveMutex.h"
+#else
+#	include "HAL/PThreadsRecursiveMutex.h"
+#endif
+
 #include "Containers/UnrealString.h"
-#include "HAL/PThreadsRecursiveMutex.h"
 #include "HAL/PThreadsSharedMutex.h"
 #include "Misc/Timespan.h"
 
@@ -37,7 +42,13 @@ private:
 	int32 FileHandle;
 };
 
+#if PLATFORM_MAC_ARM64
+using FPlatformRecursiveMutex = FAppleRecursiveMutex;
+#else
+// We don't use FAppleRecursiveMutex in Intel Apple builds as it seems
+// to trigger a bug in Rosetta2, leading to a very early segfault crash.
 using FPlatformRecursiveMutex = FPThreadsRecursiveMutex;
+#endif
 using FPlatformSharedMutex = FPThreadsSharedMutex;
 using FPlatformSystemWideMutex = FMacSystemWideMutex;
 
