Build fix

#rnx

diff --git a/Engine/Source/Runtime/Core/Public/Mac/MacPlatformMutex.h b/Engine/Source/Runtime/Core/Public/Mac/MacPlatformMutex.h
--- a/Engine/Source/Runtime/Core/Public/Mac/MacPlatformMutex.h
+++ b/Engine/Source/Runtime/Core/Public/Mac/MacPlatformMutex.h
@@ -2,16 +2,18 @@
 
 #pragma once
 
+#include "Containers/UnrealString.h"
+#include "HAL/PThreadsSharedMutex.h"
+#include "Misc/Timespan.h"
+
+#include "HAL/Platform.h"
+
 #if PLATFORM_MAC_ARM64
 #	include "Apple/AppleRecursiveMutex.h"
 #else
 #	include "HAL/PThreadsRecursiveMutex.h"
 #endif
 
-#include "Containers/UnrealString.h"
-#include "HAL/PThreadsSharedMutex.h"
-#include "Misc/Timespan.h"
-
 namespace UE
 {
 
