TMap::FindOrAdd() - Corrected comment

#Synced-CL 49180728

diff --git a/Engine/Source/Runtime/Core/Public/Containers/Map.h.inl b/Engine/Source/Runtime/Core/Public/Containers/Map.h.inl
--- a/Engine/Source/Runtime/Core/Public/Containers/Map.h.inl
+++ b/Engine/Source/Runtime/Core/Public/Containers/Map.h.inl
@@ -573,7 +573,7 @@ private:
 
 	/**
 	 * Find the value associated with a specified key, or if none exists,
-	 * adds the value
+	 * adds the value.
 	 *
 	 * @param Key The key to search for.
 	 * @param Value The value to associate with the key.
@@ -608,7 +608,7 @@ public:
 
 	/**
 	 * Find the value associated with a specified key, or if none exists, 
-	 * adds a value using the default constructor.
+	 * adds the value.
 	 *
 	 * @param Key The key to search for.
 	 * @param Value The value to associate with the key.
