Allow `CallstackTrace_GetCurrentId` to work under transactions (used from tracing).

#Synced-CL 49370461

diff --git a/Engine/Source/Programs/AutoRTFMTests/AutoRTFMTests.Target.cs b/Engine/Source/Programs/AutoRTFMTests/AutoRTFMTests.Target.cs
--- a/Engine/Source/Programs/AutoRTFMTests/AutoRTFMTests.Target.cs
+++ b/Engine/Source/Programs/AutoRTFMTests/AutoRTFMTests.Target.cs
@@ -86,6 +86,9 @@ public class AutoRTFMTestsTarget : TargetRules
 		GlobalDefinitions.Add("CSV_PROFILER=1");
 		GlobalDefinitions.Add("CSV_PROFILER_ENABLE_IN_SHIPPING=1");
 
+		// CallstackTrace tests rely on this being enabled!
+		GlobalDefinitions.Add("UE_CALLSTACK_TRACE_ENABLED=1");
+
 		// Need to force the heartbeat on so we can test with it!
 		GlobalDefinitions.Add("USE_HITCH_DETECTION=1");
 
diff --git a/Engine/Source/Programs/AutoRTFMTests/Private/InsightsTests.cpp b/Engine/Source/Programs/AutoRTFMTests/Private/InsightsTests.cpp
--- a/Engine/Source/Programs/AutoRTFMTests/Private/InsightsTests.cpp
+++ b/Engine/Source/Programs/AutoRTFMTests/Private/InsightsTests.cpp
@@ -4,6 +4,8 @@
 #include "AutoRTFM/Testing.h"
 #include "Catch2Includes.h"
 #include "CoreTypes.h"
+#include "ProfilingDebugging/CallstackTrace.h"
+#include "ProfilingDebugging/CallstackTracePrivate.h"
 #include "ProfilingDebugging/CpuProfilerTrace.h"
 #include "Trace/Trace.h"
 #include "Trace/Trace.inl"
@@ -536,4 +538,16 @@ TEST_CASE("Insights")
 		});
 	}
 #endif
+
+	SECTION("CallstackTrace")
+	{
+		AutoRTFMInsightsDeferer Callstack(CallstackChannel);
+		REQUIRE(UE_TRACE_CHANNELEXPR_IS_ENABLED(CallstackChannel));
+
+		const int OutwithTransactionId = CallstackTrace_GetCurrentId();
+		AutoRTFM::Testing::Commit([&]
+		{
+			REQUIRE(OutwithTransactionId != CallstackTrace_GetCurrentId());
+		});
+	}
 }
diff --git a/Engine/Source/Runtime/Core/Private/ProfilingDebugging/CallstackTracePrivate.h b/Engine/Source/Runtime/Core/Private/ProfilingDebugging/CallstackTracePrivate.h
--- a/Engine/Source/Runtime/Core/Private/ProfilingDebugging/CallstackTracePrivate.h
+++ b/Engine/Source/Runtime/Core/Private/ProfilingDebugging/CallstackTracePrivate.h
@@ -31,7 +31,7 @@
 #endif
 
 ////////////////////////////////////////////////////////////////////////////////
-UE_TRACE_CHANNEL_EXTERN(CallstackChannel)
+UE_TRACE_CHANNEL_EXTERN(CallstackChannel, CORE_API)
 
 UE_TRACE_EVENT_BEGIN_EXTERN(Memory, CallstackSpec, NoSync)
 	UE_TRACE_EVENT_FIELD(uint32, CallstackId)
diff --git a/Engine/Source/Runtime/Core/Public/ProfilingDebugging/CallstackTrace.h b/Engine/Source/Runtime/Core/Public/ProfilingDebugging/CallstackTrace.h
--- a/Engine/Source/Runtime/Core/Public/ProfilingDebugging/CallstackTrace.h
+++ b/Engine/Source/Runtime/Core/Public/ProfilingDebugging/CallstackTrace.h
@@ -1,5 +1,7 @@
 // Copyright Epic Games, Inc. All Rights Reserved.
 #pragma once
+
+#include "AutoRTFM/Defines.h"
 #include "HAL/Platform.h"
 #include "HAL/PlatformTLS.h"
 #include "Containers/ArrayView.h"
@@ -50,7 +52,7 @@ bool CallstackTrace_FilterModule(FStringView ModuleName);
  * can be used in trace events and be resolved in analysis.
  * @return Unique id identifying the current callstack.
  */
-CORE_API uint32 CallstackTrace_GetCurrentId();
+AUTORTFM_OPEN CORE_API uint32 CallstackTrace_GetCurrentId();
 
  /**
   * Callstack Trace Scoped Macro to avoid resolving the full callstack
@@ -76,7 +78,7 @@ inline void* CallstackTrace_GetFallbackPlatformReturnAddressData()
 /**
 * @return Needs full callstack resolve
 */
-inline bool  CallstackTrace_ResolveFullCallStack()
+inline bool CallstackTrace_ResolveFullCallStack()
 {
 	return CallstackTrace_GetFallbackPlatformReturnAddressData() == nullptr;
 }
@@ -108,7 +110,7 @@ public:
 
 inline void CallstackTrace_Create(class FMalloc* Malloc) {}
 inline void CallstackTrace_Initialize() {}
-inline uint32 CallstackTrace_GetCurrentId() { return 0; }
+AUTORTFM_OPEN inline uint32 CallstackTrace_GetCurrentId() { return 0; }
 inline void* CallstackTrace_GetCurrentReturnAddressData() { return nullptr; }
 inline void* CallstackTrace_GetFallbackPlatformReturnAddressData() { return nullptr; }
 inline bool  CallstackTrace_ResolveFullCallStack() { return true; }
