Added a list of justifications for using UEOps.

diff --git a/Engine/Source/Runtime/Core/Public/Misc/UEOps.h b/Engine/Source/Runtime/Core/Public/Misc/UEOps.h
--- a/Engine/Source/Runtime/Core/Public/Misc/UEOps.h
+++ b/Engine/Source/Runtime/Core/Public/Misc/UEOps.h
@@ -10,6 +10,19 @@
  * These ensure that the operators and their variants are defined correctly based on a single customization
  * point and should result in faster overload resolution.
  *
+ * Benefits:
+ * - They're automatically marked up with things like [[nodiscard]], constexpr and UE_REWRITE, so you don't have to.
+ * - It synthesises <, <=, >, and >= from a single UEOpLessThan, which C++20 does not (it only synthesises those from operator<=>).
+ * - It validates your definitions to ensure they are const and self-consistent.
+ * - It avoids overload resolution problems with multiple systems providing mutually-compatible templated operators,
+ *     e.g. A::operator==(const T&) and B::operator==(const U&), where T matches with B and U matches with A.
+ * - The more UEOps is used, the shorter error messages we get ("could be this overload, or this overload, or this overload...").
+ * - It avoids friendship issues.
+ * - It's syntactically a lot cleaner.
+ * - (Unproven) It should be a lot less work for the compiler to do overload resolution on.
+ * 
+ * Usage:
+ *
  * bool FMyType::UEOpEquals(const FMyType&) const;
  * bool FMyType::UEOpEquals(const OtherType&) const;
  *   Should be defined to opt into operator== and operator!=.
