Added UE::CFundamental and UE::CScalar concepts.

diff --git a/Engine/Source/Runtime/Core/Public/Concepts/Fundamental.h b/Engine/Source/Runtime/Core/Public/Concepts/Fundamental.h
new file mode 100644
--- /dev/null
+++ b/Engine/Source/Runtime/Core/Public/Concepts/Fundamental.h
@@ -0,0 +1,14 @@
+// Copyright Epic Games, Inc. All Rights Reserved.
+
+#pragma once
+
+#include <type_traits>
+
+namespace UE
+{
+	/**
+	 * Concept which describes a fundamental type.
+	 */
+	template <typename T>
+	concept CFundamental = std::is_fundamental_v<T>;
+}
diff --git a/Engine/Source/Runtime/Core/Public/Concepts/Scalar.h b/Engine/Source/Runtime/Core/Public/Concepts/Scalar.h
new file mode 100644
--- /dev/null
+++ b/Engine/Source/Runtime/Core/Public/Concepts/Scalar.h
@@ -0,0 +1,14 @@
+// Copyright Epic Games, Inc. All Rights Reserved.
+
+#pragma once
+
+#include <type_traits>
+
+namespace UE
+{
+	/**
+	 * Concept which describes a scalar type.
+	 */
+	template <typename T>
+	concept CScalar = std::is_scalar_v<T>;
+}
