[Core] TLruCache: Add SetMaxNumElements / Declare default move constructor/assignment and define proper copy/assignment as they are currently broken (with unittest)
[PCG] Cache: optimize cache grow and use SetMaxNumElements instead of the previous full copy that was happening

#Synced-CL 49409681

diff --git a/Engine/Plugins/PCG/Source/PCG/Private/Graph/PCGGraphCache.cpp b/Engine/Plugins/PCG/Source/PCG/Private/Graph/PCGGraphCache.cpp
--- a/Engine/Plugins/PCG/Source/PCG/Private/Graph/PCGGraphCache.cpp
+++ b/Engine/Plugins/PCG/Source/PCG/Private/Graph/PCGGraphCache.cpp
@@ -368,25 +368,8 @@ void FPCGGraphCache::GrowCache_Unsafe()
 {
 	TRACE_CPUPROFILER_EVENT_SCOPE(FPCGGraphCache::GrowCache_Unsafe);
 
-	TLruCache<FPCGCacheEntryKey, FPCGDataCollection> CacheDataCopy(CacheData.Max());
-
-	// Iteration begins from most recent, so this reverses the temporal order.
-	for (TLruCache<FPCGCacheEntryKey, FPCGDataCollection>::TIterator It(CacheData); It; ++It)
-	{
-		CacheDataCopy.Add(It.Key(), It.Value());
-	}
-
-	// Resize and flush
 	const int32 NewSize = (CacheData.Num() > 0) ? (2 * CacheData.Num()) : GPCGGraphCacheInitialCapacity;
-	ClearCacheInternal(NewSize, /*bClearMemory=*/false);
-
-	UE_LOG(LogPCG, Log, TEXT("Graph cache capacity increased to %d entries."), CacheData.Max());
-
-	// Copy back. Restore temporal order.
-	for (TLruCache<FPCGCacheEntryKey, FPCGDataCollection>::TIterator It(CacheDataCopy); It; ++It)
-	{
-		AddToCacheInternal(It.Key(), It.Value(), /*bAddToMemory=*/false);
-	}
+	CacheData.SetMaxNumElements(NewSize);
 }
 
 void FPCGGraphCache::AddDataToAccountedMemory(const FPCGDataCollection& InCollection)
diff --git a/Engine/Source/Runtime/Core/Public/Containers/LruCache.h b/Engine/Source/Runtime/Core/Public/Containers/LruCache.h
--- a/Engine/Source/Runtime/Core/Public/Containers/LruCache.h
+++ b/Engine/Source/Runtime/Core/Public/Containers/LruCache.h
@@ -3,6 +3,7 @@
 #pragma once
 
 #include "CoreTypes.h"
+#include "Containers/AllowShrinking.h"
 #include "Containers/Set.h"
 #include "Misc/AssertionMacros.h"
 
@@ -151,6 +152,29 @@ public:
 		Empty(InMaxNumElements);
 	}
 
+	TLruCache(const TLruCache& Other)
+	{
+		*this = Other;
+	}
+
+	TLruCache& operator=(const TLruCache& Other)
+	{
+		Empty(Other.MaxNumElements);
+
+		// Copy from least recent to most so that we keep the usage state through copy
+		const FCacheEntry* CurrentEntry = Other.LeastRecent;
+		while (CurrentEntry)
+		{
+			Add(CurrentEntry->Key, CurrentEntry->Value);
+			CurrentEntry = CurrentEntry->MoreRecent;
+		}
+
+		return *this;
+	}
+
+	TLruCache(TLruCache&& Other) = default;
+	TLruCache& operator=(TLruCache&& Other) = default;
+
 	/** Destructor. */
 	~TLruCache()
 	{
@@ -309,6 +333,31 @@ public:
 		LeastRecent = nullptr;
 	}
 
+	/**
+	 * Change max cache size.
+	 * 
+	 * @param InMaxNumElements The maximum number of elements this cache can hold
+	 * @param InAllowShrinking If InMaxNumElements is smaller than current MaxNumElements, indicates if we want to compact internal memory
+	 */
+	void SetMaxNumElements(int32 InMaxNumElements, EAllowShrinking InAllowShrinking = EAllowShrinking::Default)
+	{
+		if (InMaxNumElements < MaxNumElements)
+		{
+			while (LookupSet.Num() > InMaxNumElements)
+			{
+				RemoveLeastRecent();
+			}
+
+			if (InAllowShrinking == EAllowShrinking::Yes)
+			{
+				LookupSet.Shrink();
+			}
+		}
+
+		LookupSet.Reserve(InMaxNumElements);
+		MaxNumElements = InMaxNumElements;
+	}
+
 	/**
 	 * Filter the entries in the cache using a predicate.
 	 *
diff --git a/Engine/Source/Runtime/Core/Tests/Containers/LruCacheTest.cpp b/Engine/Source/Runtime/Core/Tests/Containers/LruCacheTest.cpp
new file mode 100644
--- /dev/null
+++ b/Engine/Source/Runtime/Core/Tests/Containers/LruCacheTest.cpp
@@ -0,0 +1,174 @@
+// Copyright Epic Games, Inc. All Rights Reserved.
+
+#if WITH_TESTS
+
+#include "Containers/LruCache.h"
+#include "Tests/TestHarnessAdapter.h"
+
+namespace Test
+{
+	template<typename KeyType, typename ValueType>
+	void TestCompareCache(const TLruCache<KeyType,ValueType>& CacheA, const TLruCache<KeyType,ValueType>& CacheB)
+	{
+		CHECK(CacheA.GetLeastRecentKey() == CacheB.GetLeastRecentKey());
+
+		TArray<KeyType> KeysA;
+		CacheA.GetKeys(KeysA);
+
+		TArray<KeyType> KeysB;
+		CacheB.GetKeys(KeysB);
+		CHECK(KeysA == KeysB);
+
+		for (auto Key : KeysA)
+		{
+			const ValueType* ValuePtr = CacheA.Find(Key);
+			const ValueType* ValueCopyPtr = CacheB.Find(Key);
+			CHECK(ValuePtr != ValueCopyPtr);
+			CHECK(*ValuePtr == *ValueCopyPtr);
+		}
+	};
+}
+
+TEST_CASE_NAMED(FLruCacheTest, "System::Core::Containers::LruCache", "[Core][Containers][LruCache]")
+{
+	SECTION("Basic")
+	{
+		TLruCache<int32, FString> Cache(5);
+		Cache.Add(0, TEXT("A"));
+		Cache.Add(1, TEXT("B"));
+		Cache.Add(2, TEXT("C"));
+		Cache.Add(3, TEXT("D"));
+		Cache.Add(4, TEXT("E"));
+
+		// Validate 0 is least recent key
+		int32 LeastRecentKey = Cache.GetLeastRecentKey();
+		CHECK(LeastRecentKey == 0);
+			
+		// Validate value matches inserted value
+		CHECK(*Cache.Find(LeastRecentKey) == TEXT("A"));
+		
+		// Move LeastRecentKey to most recent and compare with inserted value
+		CHECK(*Cache.FindAndTouch(LeastRecentKey) == TEXT("A"));
+
+		// Validate 1 is least recent key
+		LeastRecentKey = Cache.GetLeastRecentKey();
+		CHECK(LeastRecentKey == 1);
+
+		// Validate value matches inserted value
+		CHECK(*Cache.Find(LeastRecentKey) == TEXT("B"));
+
+		// Add new value (which should discard LeastRecentKey(1)
+		Cache.Add(5, TEXT("F"));
+
+		// Validate discard
+		CHECK(Cache.Find(LeastRecentKey) == nullptr);
+
+		// Validate 2 is least recent key
+		LeastRecentKey = Cache.GetLeastRecentKey();
+		CHECK(LeastRecentKey == 2);
+
+		// Validate value matches inserted value
+		CHECK(*Cache.Find(LeastRecentKey) == TEXT("C"));
+
+		// Remove LeastRecent entry(2)
+		Cache.RemoveLeastRecent();
+
+		// Validate discard
+		CHECK(Cache.Find(LeastRecentKey) == nullptr);
+
+		// Validate 3 is least recent key
+		LeastRecentKey = Cache.GetLeastRecentKey();
+		CHECK(LeastRecentKey == 3);
+
+		// Validate value matches inserted value
+		CHECK(*Cache.Find(LeastRecentKey) == TEXT("D"));
+
+		// Validate size
+		CHECK(Cache.Num() == 4);
+	}
+
+	SECTION("SetMaxNumElements")
+	{
+		TLruCache<int32, FString> Cache(5);
+		Cache.Add(0, TEXT("A"));
+		Cache.Add(1, TEXT("B"));
+		Cache.Add(2, TEXT("C"));
+		Cache.Add(3, TEXT("D"));
+		Cache.Add(4, TEXT("E"));
+
+		TLruCache<int32, FString> CacheCopy(Cache);
+
+		CacheCopy.SetMaxNumElements(10);
+
+		Test::TestCompareCache(Cache, CacheCopy);
+
+		Cache.SetMaxNumElements(3);
+		CacheCopy.SetMaxNumElements(3);
+
+		CHECK(Cache.Num() == 3);
+		CHECK(CacheCopy.Num() == 3);
+
+		Test::TestCompareCache(Cache, CacheCopy);
+	}
+
+	SECTION("Copy")
+	{
+		TLruCache<int32, FString> Cache(5);
+		Cache.Add(0, TEXT("A"));
+		Cache.Add(1, TEXT("B"));
+		Cache.Add(2, TEXT("C"));
+		Cache.Add(3, TEXT("D"));
+		Cache.Add(4, TEXT("E"));
+
+		TLruCache<int32, FString> CacheCopy(Cache);
+		Test::TestCompareCache(Cache, CacheCopy);
+	}
+
+	SECTION("CopyAssign")
+	{
+		TLruCache<int32, FString> Cache(5);
+		Cache.Add(0, TEXT("A"));
+		Cache.Add(1, TEXT("B"));
+		Cache.Add(2, TEXT("C"));
+		Cache.Add(3, TEXT("D"));
+		Cache.Add(4, TEXT("E"));
+
+		TLruCache<int32, FString> CacheCopy;
+		CacheCopy = Cache;
+		Test::TestCompareCache(Cache, CacheCopy);
+	}
+
+	SECTION("Move")
+	{
+		TLruCache<int32, FString> Cache(5);
+		Cache.Add(0, TEXT("A"));
+		Cache.Add(1, TEXT("B"));
+		Cache.Add(2, TEXT("C"));
+		Cache.Add(3, TEXT("D"));
+		Cache.Add(4, TEXT("E"));
+
+		TLruCache<int32, FString> CacheCopy(Cache);
+		TLruCache<int32, FString> MovedCache(MoveTemp(Cache));
+		CHECK(Cache.IsEmpty());
+		Test::TestCompareCache(CacheCopy, MovedCache);
+	}
+
+	SECTION("MoveAssign")
+	{
+		TLruCache<int32, FString> Cache(5);
+		Cache.Add(0, TEXT("A"));
+		Cache.Add(1, TEXT("B"));
+		Cache.Add(2, TEXT("C"));
+		Cache.Add(3, TEXT("D"));
+		Cache.Add(4, TEXT("E"));
+
+		TLruCache<int32, FString> CacheCopy(Cache);
+		TLruCache<int32, FString> MovedCache;
+		MovedCache = MoveTemp(Cache);
+		CHECK(Cache.IsEmpty());
+		Test::TestCompareCache(CacheCopy, MovedCache);
+	}
+	
+}
+
+#endif // WITH_TESTS
\ No newline at end of file
