[Apple] Unfair locks fix for older OS

Adaptive spinning is used on newer OS, where it is available, otherwise the usual unfair lock's locking is used.
This brings back unfair locks for Apple Intel.

[RN]

diff --git a/Engine/Source/Runtime/Core/Private/Apple/AppleRecursiveMutex.cpp b/Engine/Source/Runtime/Core/Private/Apple/AppleRecursiveMutex.cpp
new file mode 100644
--- /dev/null
+++ b/Engine/Source/Runtime/Core/Private/Apple/AppleRecursiveMutex.cpp
@@ -0,0 +1,32 @@
+// Copyright Epic Games, Inc. All Rights Reserved.
+
+#include "Apple/AppleRecursiveMutex.h"
+
+namespace UE::Private
+{
+
+static LockImplementation LockRuntimeDispatch;
+static LockImplementation LockFallback;
+
+CORE_API LockImplementation* LockWithFlags = &LockRuntimeDispatch;
+
+static void LockRuntimeDispatch(os_unfair_lock_t lock, os_unfair_lock_flags_t flags)
+{
+	//
+	// os_unfair_lock_lock_with_flags is available on: 
+	//   iOS 18.0+, iPadOS 18.0+, Mac Catalyst 18.0+, macOS 15.0+, tvOS 18.0+, visionOS 2.0+, watchOS 11.0+
+	//
+	// It is a weak import, so will be nullptr on older OS's. Hence, we can do a
+	// runtime dispatch on the first lock call.
+	//
+	LockWithFlags = os_unfair_lock_lock_with_flags ? os_unfair_lock_lock_with_flags : &LockFallback;
+	LockWithFlags(lock, flags);
+}
+
+static void LockFallback(os_unfair_lock_t lock, os_unfair_lock_flags_t flags)
+{
+	(void)flags;
+	os_unfair_lock_lock(lock);
+}
+
+} // namespace UE::Private
diff --git a/Engine/Source/Runtime/Core/Public/Apple/AppleRecursiveMutex.h b/Engine/Source/Runtime/Core/Public/Apple/AppleRecursiveMutex.h
--- a/Engine/Source/Runtime/Core/Public/Apple/AppleRecursiveMutex.h
+++ b/Engine/Source/Runtime/Core/Public/Apple/AppleRecursiveMutex.h
@@ -14,6 +14,15 @@
 namespace UE
 {
 
+namespace Private
+{
+
+typedef void LockImplementation(os_unfair_lock_t lock, os_unfair_lock_flags_t flags);
+
+CORE_API extern LockImplementation* LockWithFlags;
+
+} // namespace Private
+
 class FAppleRecursiveMutex final
 {
     os_unfair_lock Mutex {OS_UNFAIR_LOCK_INIT};
@@ -34,7 +43,7 @@ public:
 
         if (ThreadId.load(std::memory_order_relaxed) != CurrentThreadId)
 		{
-			os_unfair_lock_lock_with_flags(&Mutex, OS_UNFAIR_LOCK_FLAG_ADAPTIVE_SPIN);
+			Private::LockWithFlags(&Mutex, OS_UNFAIR_LOCK_FLAG_ADAPTIVE_SPIN);
 			ThreadId.store(CurrentThreadId, std::memory_order_relaxed);
 		}
 		++LockCount;
diff --git a/Engine/Source/Runtime/Core/Public/Mac/MacPlatformMutex.h b/Engine/Source/Runtime/Core/Public/Mac/MacPlatformMutex.h
--- a/Engine/Source/Runtime/Core/Public/Mac/MacPlatformMutex.h
+++ b/Engine/Source/Runtime/Core/Public/Mac/MacPlatformMutex.h
@@ -2,18 +2,11 @@
 
 #pragma once
 
+#include "Apple/AppleRecursiveMutex.h"
 #include "Containers/UnrealString.h"
 #include "HAL/PThreadsSharedMutex.h"
 #include "Misc/Timespan.h"
 
-#include "HAL/Platform.h"
-
-#if PLATFORM_MAC_ARM64
-#	include "Apple/AppleRecursiveMutex.h"
-#else
-#	include "HAL/PThreadsRecursiveMutex.h"
-#endif
-
 namespace UE
 {
 
@@ -44,14 +37,8 @@ private:
 	int32 FileHandle;
 };
 
-#if PLATFORM_MAC_ARM64
-using FPlatformRecursiveMutex = FAppleRecursiveMutex;
-#else
-// We don't use FAppleRecursiveMutex in Intel Apple builds as it seems
-// to trigger a bug in Rosetta2, leading to a very early segfault crash.
-using FPlatformRecursiveMutex = FPThreadsRecursiveMutex;
-#endif
-using FPlatformSharedMutex = FPThreadsSharedMutex;
+using FPlatformRecursiveMutex  = FAppleRecursiveMutex;
+using FPlatformSharedMutex     = FPThreadsSharedMutex;
 using FPlatformSystemWideMutex = FMacSystemWideMutex;
 
 } // UE
