C++20 doesn't allow arithemetic on two unrelated enums (see P1120R0)

#rnx

#Synced-CL 49321981

diff --git a/Engine/Source/Runtime/Core/Public/HAL/MallocBinned.h b/Engine/Source/Runtime/Core/Public/HAL/MallocBinned.h
--- a/Engine/Source/Runtime/Core/Public/HAL/MallocBinned.h
+++ b/Engine/Source/Runtime/Core/Public/HAL/MallocBinned.h
@@ -116,8 +116,11 @@ private:
 	enum { POOL_COUNT = 41 };
 
 	/** Maximum allocation for the pooled allocator */
-	enum { EXTENDED_PAGE_POOL_ALLOCATION_COUNT = 2 };
-	enum { MAX_POOLED_ALLOCATION_SIZE   = 32768+1 };
+	enum
+	{
+		EXTENDED_PAGE_POOL_ALLOCATION_COUNT = 2,
+		MAX_POOLED_ALLOCATION_SIZE = 32768 + 1,
+	};
 	
 	// Forward declares.
 	struct FFreeMem;
@@ -229,7 +232,7 @@ private:
 	FPoolTable  PoolTable[POOL_COUNT];
 	FPoolTable	OsTable;
 	FPoolTable	PagePoolTable[EXTENDED_PAGE_POOL_ALLOCATION_COUNT];
-	FPoolTable* MemSizeToPoolTable[MAX_POOLED_ALLOCATION_SIZE+EXTENDED_PAGE_POOL_ALLOCATION_COUNT];
+	FPoolTable* MemSizeToPoolTable[MAX_POOLED_ALLOCATION_SIZE + EXTENDED_PAGE_POOL_ALLOCATION_COUNT];
 	
 	PoolHashBucket* HashBuckets;
 	PoolHashBucket* HashBucketFreeList;
