Added FGuid::CombineValid that returns a GUID which is a combination of the two provided ones, ignoring any invalid ones in the process.

#Synced-CL 49407983

diff --git a/Engine/Source/Runtime/Core/Private/Misc/Guid.cpp b/Engine/Source/Runtime/Core/Private/Misc/Guid.cpp
--- a/Engine/Source/Runtime/Core/Private/Misc/Guid.cpp
+++ b/Engine/Source/Runtime/Core/Private/Misc/Guid.cpp
@@ -304,6 +304,21 @@ FGuid FGuid::Combine(const FGuid& GuidA, const FGuid& GuidB)
 	);
 }
 
+FGuid FGuid::CombineValid(const FGuid& GuidA, const FGuid& GuidB)
+{
+	if (GuidA.IsValid())
+	{
+		if (GuidB.IsValid())
+		{
+			return Combine(GuidA, GuidB);
+		}
+
+		return GuidA;
+	}
+
+	return GuidB;
+}
+
 bool FGuid::Parse(const TCHAR* GuidString, FGuid& OutGuid)
 {
 	return Parse(FStringView(GuidString), OutGuid);
diff --git a/Engine/Source/Runtime/Core/Public/Misc/Guid.h b/Engine/Source/Runtime/Core/Public/Misc/Guid.h
--- a/Engine/Source/Runtime/Core/Public/Misc/Guid.h
+++ b/Engine/Source/Runtime/Core/Public/Misc/Guid.h
@@ -383,12 +383,19 @@ public:
 	static CORE_API FGuid NewGuidFromHashBytes(const void* HashData, int64 DataLen);
 
 	/**
-	 * Returns a GUID which is a combinationof the two provided ones.
+	 * Returns a GUID which is a combination of the two provided ones.
 	 *
 	 * @return The combined GUID.
 	 */
 	static CORE_API FGuid Combine(const FGuid& GuidA, const FGuid& GuidB);
 
+	/**
+	 * Returns a GUID which is a combination of the two provided ones, ignoring any invalid ones in the process.
+	 *
+	 * @return The combined GUID, invalid if both inputs are invalid.
+	 */
+	static CORE_API FGuid CombineValid(const FGuid& GuidA, const FGuid& GuidB);
+
 	/**
 	 * Converts a string to a GUID.
 	 *
diff --git a/Engine/Source/Runtime/Core/Tests/Misc/GuidTest.cpp b/Engine/Source/Runtime/Core/Tests/Misc/GuidTest.cpp
--- a/Engine/Source/Runtime/Core/Tests/Misc/GuidTest.cpp
+++ b/Engine/Source/Runtime/Core/Tests/Misc/GuidTest.cpp
@@ -107,6 +107,18 @@ TEST_CASE_NAMED(FGuidTest, "System::Core::Misc::Guid", "[ApplicationContextMask]
 
 	FGuid g5_1(MD5HashToGuid(Hash));
 	CHECK_MESSAGE(TEXT("Constructing from MD5 hash must be valid"), g5_1.ToString() == TEXT("89d900e5bc77dba5ee796e69694a5f78"));
+
+	// GUID combines
+	const FGuid g6_0;
+	const FGuid g6_1(TEXT("e500d989a5db77bc696e79ee785f4a69"));
+	const FGuid g6_2(TEXT("89d900e5bc77dba5ee796e69694a5f78"));
+	const FGuid g6_3(TEXT("43be1e3368afb993e98aaa574a6d550f"));
+
+	CHECK_MESSAGE(TEXT("Combining 2 GUIDs"), FGuid::Combine(g6_1, g6_2) == g6_3);
+	CHECK_MESSAGE(TEXT("Combining 2 invalid GUIDs"), !FGuid::CombineValid(g6_0, g6_0).IsValid());
+	CHECK_MESSAGE(TEXT("Combining 1 valid GUID and 1 invalid GUID"), FGuid::CombineValid(g6_1, g6_0) == g6_1);
+	CHECK_MESSAGE(TEXT("Combining 1 invalid GUID and 1 valid GUID"), FGuid::CombineValid(g6_0, g6_1) == g6_1);
+	CHECK_MESSAGE(TEXT("Combining 2 valid GUIDs "), FGuid::CombineValid(g6_1, g6_2) == g6_3);
 }
 
 #endif //WITH_TESTS
