Moved delegate execution code from UHT-generated files to templates.

diff --git a/Engine/Plugins/Experimental/PythonScriptPlugin/Source/PythonScriptPlugin/Private/PyWrapperDelegate.cpp b/Engine/Plugins/Experimental/PythonScriptPlugin/Source/PythonScriptPlugin/Private/PyWrapperDelegate.cpp
--- a/Engine/Plugins/Experimental/PythonScriptPlugin/Source/PythonScriptPlugin/Private/PyWrapperDelegate.cpp
+++ b/Engine/Plugins/Experimental/PythonScriptPlugin/Source/PythonScriptPlugin/Private/PyWrapperDelegate.cpp
@@ -255,7 +255,7 @@ struct TPyDelegateInvocation<FMulticastScriptDelegate>
 
 	static void Call(const FMulticastScriptDelegate& Delegate, void* Params)
 	{
-		Delegate.ProcessMulticastDelegate<UObject>(Params);
+		Delegate.ProcessDelegate<UObject>(Params);
 	}
 };
 
diff --git a/Engine/Source/Programs/Shared/EpicGames.UHT/Exporters/CodeGen/UhtHeaderCodeGeneratorCppFile.cs b/Engine/Source/Programs/Shared/EpicGames.UHT/Exporters/CodeGen/UhtHeaderCodeGeneratorCppFile.cs
--- a/Engine/Source/Programs/Shared/EpicGames.UHT/Exporters/CodeGen/UhtHeaderCodeGeneratorCppFile.cs
+++ b/Engine/Source/Programs/Shared/EpicGames.UHT/Exporters/CodeGen/UhtHeaderCodeGeneratorCppFile.cs
@@ -1280,26 +1280,7 @@ namespace EpicGames.UHT.Exporters.CodeGen
 
 		private StringBuilder AppendDelegate(StringBuilder builder, UhtFunction function)
 		{
-			AppendFunction(builder, function, false, false);
-
-			int tabs = 0;
-			string strippedFunctionName = function.StrippedFunctionName;
-			string exportFunctionName = GetDelegateFunctionExportName(function);
-			string extraParameter = GetDelegateFunctionExtraParameter(function);
-
-			AppendNativeFunctionHeader(builder, function, UhtPropertyTextType.EventFunctionArgOrRetVal, false, exportFunctionName, extraParameter, UhtFunctionExportFlags.None, 0, "\r\n");
-			AppendEventFunctionPrologue(builder, function, strippedFunctionName, tabs, "\r\n", true);
-			builder
-				.Append('\t')
-				.Append(strippedFunctionName)
-				.Append('.')
-				.Append(function.FunctionFlags.HasAnyFlags(EFunctionFlags.MulticastDelegate) ? "ProcessMulticastDelegate" : "ProcessDelegate")
-				.Append("<UObject>(")
-				.Append(function.Children.Count > 0 ? "&Parms" : "NULL")
-				.Append(");\r\n");
-			AppendEventFunctionEpilogue(builder, function, tabs, "\r\n");
-
-			return builder;
+			return AppendFunction(builder, function, false, false);
 		}
 
 		private StringBuilder AppendFunction(StringBuilder builder, UhtFunction function, bool isNoExport, bool isVerseNativeCallable)
diff --git a/Engine/Source/Programs/Shared/EpicGames.UHT/Exporters/CodeGen/UhtHeaderCodeGeneratorHFile.cs b/Engine/Source/Programs/Shared/EpicGames.UHT/Exporters/CodeGen/UhtHeaderCodeGeneratorHFile.cs
--- a/Engine/Source/Programs/Shared/EpicGames.UHT/Exporters/CodeGen/UhtHeaderCodeGeneratorHFile.cs
+++ b/Engine/Source/Programs/Shared/EpicGames.UHT/Exporters/CodeGen/UhtHeaderCodeGeneratorHFile.cs
@@ -183,9 +183,6 @@ namespace EpicGames.UHT.Exporters.CodeGen
 					}
 					else if (field is UhtFunction function)
 					{
-						using UhtCodeBlockComment blockComment = new(builder, field);
-						using UhtMacroBlockEmitter macroBlockEmitter = new(builder, UhtDefineScopeNames.Standard, function.DefineScope);
-						AppendDelegate(builder, function);
 					}
 					else if (field is UhtClass classObj)
 					{
@@ -802,40 +799,6 @@ namespace EpicGames.UHT.Exporters.CodeGen
 			return builder;
 		}
 
-		private StringBuilder AppendDelegate(StringBuilder builder, UhtFunction function)
-		{
-			using (UhtMacroCreator macro = new(builder, this, function, DelegateMacroSuffix))
-			{
-				string exportFunctionName = GetDelegateFunctionExportName(function);
-				string extraParameter = GetDelegateFunctionExtraParameter(function);
-
-				bool addAPI = true;
-
-				UhtClass? outerClass = function.Outer as UhtClass;
-				if (outerClass != null)
-				{
-					builder.Append("static ");
-					if (outerClass.ClassFlags.HasFlag(EClassFlags.RequiredAPI))
-					{
-						addAPI = false;
-					}
-				}
-
-				//if (!function.FunctionFlags.HasAnyFlags(EFunctionFlags.RequiredAPI)) // TODO: This requires too much fixup for now
-				//{
-				//	addAPI = false
-				//}
-
-				if (addAPI)
-				{
-					builder.Append(Module.Api);
-				}
-
-				AppendNativeFunctionHeader(builder, function, UhtPropertyTextType.EventFunctionArgOrRetVal, true, exportFunctionName, extraParameter, UhtFunctionExportFlags.None, 0, "; \\\r\n");
-			}
-			return builder;
-		}
-
 		private StringBuilder AppendClass(StringBuilder builder, UhtClass classObj)
 		{
 			string api = classObj.ClassFlags.HasAnyFlags(EClassFlags.MinimalAPI) ? Module.Api : "NO_API ";
diff --git a/Engine/Source/Runtime/Core/Public/Delegates/Delegate.h b/Engine/Source/Runtime/Core/Public/Delegates/Delegate.h
--- a/Engine/Source/Runtime/Core/Public/Delegates/Delegate.h
+++ b/Engine/Source/Runtime/Core/Public/Delegates/Delegate.h
@@ -196,9 +196,6 @@
 /** This suffix is appended to all header exported delegates */
 #define HEADER_GENERATED_DELEGATE_SIGNATURE_SUFFIX TEXT("__DelegateSignature")
 
-/** Helper macro that enables passing comma-separated arguments as a single macro parameter */
-#define FUNC_CONCAT( ... ) __VA_ARGS__
-
 /**
  * Declares a delegate that can only bind to one native function at a time
  *
@@ -232,87 +229,19 @@
 	class EventName : public BaseTypeEvent { friend class OwningType; };
 
 /** Declare user's dynamic delegate, with wrapper proxy method for executing the delegate */
-#define FUNC_DECLARE_DYNAMIC_DELEGATE( DynamicDelegateClassName, ExecFunction, FuncParamList, FuncParamPassThru, ... ) \
+#define FUNC_DECLARE_DYNAMIC_DELEGATE( DynamicDelegateClassName, ... ) \
 	class DynamicDelegateClassName : public TBaseDynamicDelegate<FNotThreadSafeDelegateMode, __VA_ARGS__> \
 	{ \
 	public: \
-		/** Default constructor */ \
-		DynamicDelegateClassName() \
-		{ \
-		} \
-		\
-		/** Construction from an FScriptDelegate must be explicit.  This is really only used by UObject system internals. */ \
-		explicit DynamicDelegateClassName( const TScriptDelegate<>& InScriptDelegate ) \
-			: TBaseDynamicDelegate<FNotThreadSafeDelegateMode, __VA_ARGS__>( InScriptDelegate ) \
-		{ \
-		} \
-		\
-		/** Execute the delegate.  If the function pointer is not valid, an error will occur. */ \
-		inline void Execute( FuncParamList ) const \
-		{ \
-			/* Verify that the user object is still valid.  We only have a weak reference to it. */ \
-			checkSlow( IsBound() ); \
-			ExecFunction( FuncParamPassThru ); \
-		} \
-		/** Execute the delegate, but only if the function pointer is still valid */ \
-		inline bool ExecuteIfBound( FuncParamList ) const \
-		{ \
-			if( IsBound() ) \
-			{ \
-				ExecFunction( FuncParamPassThru ); \
-				return true; \
-			} \
-			return false; \
-		} \
-	};
-
-/** Declare user's dynamic delegate with return value, with wrapper proxy method for executing the delegate */
-#define FUNC_DECLARE_DYNAMIC_DELEGATE_RETVAL(DynamicDelegateRetValClassName, ExecFunction, RetValType, FuncParamList, FuncParamPassThru, ...) \
-	class DynamicDelegateRetValClassName : public TBaseDynamicDelegate<FNotThreadSafeDelegateMode, __VA_ARGS__> \
-	{ \
-	public: \
-		/** Default constructor */ \
-		DynamicDelegateRetValClassName() \
-		{ \
-		} \
-		\
-		/** Construction from an FScriptDelegate must be explicit.  This is really only used by UObject system internals. */ \
-		explicit DynamicDelegateRetValClassName( const TScriptDelegate<>& InScriptDelegate ) \
-			: TBaseDynamicDelegate<FNotThreadSafeDelegateMode, __VA_ARGS__>( InScriptDelegate ) \
-		{ \
-		} \
-		\
-		/** Execute the delegate.  If the function pointer is not valid, an error will occur. */ \
-		inline RetValType Execute( FuncParamList ) const \
-		{ \
-			/* Verify that the user object is still valid.  We only have a weak reference to it. */ \
-			checkSlow( IsBound() ); \
-			return ExecFunction( FuncParamPassThru ); \
-		} \
+		using TBaseDynamicDelegate<FNotThreadSafeDelegateMode, __VA_ARGS__>::TBaseDynamicDelegate; \
 	};
 
-
 /** Declare user's dynamic multi-cast delegate, with wrapper proxy method for executing the delegate */
-#define FUNC_DECLARE_DYNAMIC_MULTICAST_DELEGATE(DynamicMulticastDelegateClassName, ExecFunction, FuncParamList, FuncParamPassThru, ...) \
+#define FUNC_DECLARE_DYNAMIC_MULTICAST_DELEGATE(DynamicMulticastDelegateClassName, ...) \
 class DynamicMulticastDelegateClassName : public TBaseDynamicMulticastDelegate<FNotThreadSafeDelegateMode, __VA_ARGS__> \
 	{ \
 	public: \
-		/** Default constructor */ \
-		DynamicMulticastDelegateClassName() \
-		{ \
-		} \
-		\
-		/** Construction from an FMulticastScriptDelegate must be explicit.  This is really only used by UObject system internals. */ \
-		explicit DynamicMulticastDelegateClassName( const TMulticastScriptDelegate<>& InMulticastScriptDelegate ) \
-			: TBaseDynamicMulticastDelegate<FNotThreadSafeDelegateMode, __VA_ARGS__>( InMulticastScriptDelegate ) \
-		{ \
-		} \
-		\
-		/** Broadcasts this delegate to all bound objects, except to those that may have expired */ \
-		void Broadcast( FuncParamList ) const \
-		{ \
-			ExecFunction( FuncParamPassThru ); \
-		} \
+		using TBaseDynamicMulticastDelegate<FNotThreadSafeDelegateMode, __VA_ARGS__>::TBaseDynamicMulticastDelegate; \
 	};
 
 namespace UE::Core::Private
diff --git a/Engine/Source/Runtime/Core/Public/Delegates/DelegateCombinations.h b/Engine/Source/Runtime/Core/Public/Delegates/DelegateCombinations.h
--- a/Engine/Source/Runtime/Core/Public/Delegates/DelegateCombinations.h
+++ b/Engine/Source/Runtime/Core/Public/Delegates/DelegateCombinations.h
@@ -32,16 +32,16 @@
 #define DECLARE_EVENT( OwningType, EventName ) FUNC_DECLARE_EVENT( OwningType, EventName, void )
 
 /** Declares a blueprint-accessible delegate that can only bind to one UFUNCTION at a time */
-#define DECLARE_DYNAMIC_DELEGATE( DelegateName ) BODY_MACRO_COMBINE(CURRENT_FILE_ID,_,__LINE__,_DELEGATE) FUNC_DECLARE_DYNAMIC_DELEGATE( DelegateName, DelegateName##_DelegateWrapper, , FUNC_CONCAT( *this ), void )
+#define DECLARE_DYNAMIC_DELEGATE( DelegateName ) FUNC_DECLARE_DYNAMIC_DELEGATE( DelegateName, void )
 
 /** Declares a blueprint-accessible broadcast delegate that can bind to multiple native UFUNCTIONs simultaneously */
-#define DECLARE_DYNAMIC_MULTICAST_DELEGATE( DelegateName ) BODY_MACRO_COMBINE(CURRENT_FILE_ID,_,__LINE__,_DELEGATE) FUNC_DECLARE_DYNAMIC_MULTICAST_DELEGATE( DelegateName, DelegateName##_DelegateWrapper, , FUNC_CONCAT( *this ), void )
+#define DECLARE_DYNAMIC_MULTICAST_DELEGATE( DelegateName ) FUNC_DECLARE_DYNAMIC_MULTICAST_DELEGATE( DelegateName, void )
 
 /** Declares a delegate with return value that can only bind to one native function at a time */
 #define DECLARE_DELEGATE_RetVal( ReturnValueType, DelegateName ) FUNC_DECLARE_DELEGATE( DelegateName, ReturnValueType )
 
 /** Declares a blueprint-accessible delegate with return value that can only bind to one UNFUNCTION at a time */
-#define DECLARE_DYNAMIC_DELEGATE_RetVal( ReturnValueType, DelegateName ) BODY_MACRO_COMBINE(CURRENT_FILE_ID,_,__LINE__,_DELEGATE) FUNC_DECLARE_DYNAMIC_DELEGATE_RETVAL( DelegateName, DelegateName##_DelegateWrapper, ReturnValueType, , FUNC_CONCAT( *this ), ReturnValueType )
+#define DECLARE_DYNAMIC_DELEGATE_RetVal( ReturnValueType, DelegateName ) FUNC_DECLARE_DYNAMIC_DELEGATE( DelegateName, ReturnValueType )
 
 
 // Multiple-parameter versions of above delegate types:
@@ -49,79 +49,79 @@
 #define DECLARE_MULTICAST_DELEGATE_OneParam( DelegateName, Param1Type ) FUNC_DECLARE_MULTICAST_DELEGATE( DelegateName, void, Param1Type )
 #define DECLARE_TS_MULTICAST_DELEGATE_OneParam( DelegateName, Param1Type ) FUNC_DECLARE_TS_MULTICAST_DELEGATE( DelegateName, void, Param1Type )
 #define DECLARE_EVENT_OneParam( OwningType, EventName, Param1Type ) FUNC_DECLARE_EVENT( OwningType, EventName, void, Param1Type )
-#define DECLARE_DYNAMIC_DELEGATE_OneParam( DelegateName, Param1Type, Param1Name ) BODY_MACRO_COMBINE(CURRENT_FILE_ID,_,__LINE__,_DELEGATE) FUNC_DECLARE_DYNAMIC_DELEGATE( DelegateName, DelegateName##_DelegateWrapper, FUNC_CONCAT( Param1Type InParam1 ), FUNC_CONCAT( *this, InParam1 ), void, Param1Type )
-#define DECLARE_DYNAMIC_MULTICAST_DELEGATE_OneParam( DelegateName, Param1Type, Param1Name ) BODY_MACRO_COMBINE(CURRENT_FILE_ID,_,__LINE__,_DELEGATE) FUNC_DECLARE_DYNAMIC_MULTICAST_DELEGATE( DelegateName, DelegateName##_DelegateWrapper, FUNC_CONCAT( Param1Type InParam1 ), FUNC_CONCAT( *this, InParam1 ), void, Param1Type )
+#define DECLARE_DYNAMIC_DELEGATE_OneParam( DelegateName, Param1Type, Param1Name ) FUNC_DECLARE_DYNAMIC_DELEGATE( DelegateName, void, Param1Type )
+#define DECLARE_DYNAMIC_MULTICAST_DELEGATE_OneParam( DelegateName, Param1Type, Param1Name ) FUNC_DECLARE_DYNAMIC_MULTICAST_DELEGATE( DelegateName, void, Param1Type )
 #define DECLARE_DELEGATE_RetVal_OneParam( ReturnValueType, DelegateName, Param1Type ) FUNC_DECLARE_DELEGATE( DelegateName, ReturnValueType, Param1Type )
-#define DECLARE_DYNAMIC_DELEGATE_RetVal_OneParam( ReturnValueType, DelegateName, Param1Type, Param1Name ) BODY_MACRO_COMBINE(CURRENT_FILE_ID,_,__LINE__,_DELEGATE) FUNC_DECLARE_DYNAMIC_DELEGATE_RETVAL( DelegateName, DelegateName##_DelegateWrapper, ReturnValueType, FUNC_CONCAT( Param1Type InParam1 ), FUNC_CONCAT( *this, InParam1 ), ReturnValueType, Param1Type )
+#define DECLARE_DYNAMIC_DELEGATE_RetVal_OneParam( ReturnValueType, DelegateName, Param1Type, Param1Name ) FUNC_DECLARE_DYNAMIC_DELEGATE( DelegateName, ReturnValueType, Param1Type )
 
 #define DECLARE_DELEGATE_TwoParams( DelegateName, Param1Type, Param2Type ) FUNC_DECLARE_DELEGATE( DelegateName, void, Param1Type, Param2Type )
 #define DECLARE_MULTICAST_DELEGATE_TwoParams( DelegateName, Param1Type, Param2Type ) FUNC_DECLARE_MULTICAST_DELEGATE( DelegateName, void, Param1Type, Param2Type )
 #define DECLARE_TS_MULTICAST_DELEGATE_TwoParams( DelegateName, Param1Type, Param2Type ) FUNC_DECLARE_TS_MULTICAST_DELEGATE( DelegateName, void, Param1Type, Param2Type )
 #define DECLARE_EVENT_TwoParams( OwningType, EventName, Param1Type, Param2Type ) FUNC_DECLARE_EVENT( OwningType, EventName, void, Param1Type, Param2Type )
-#define DECLARE_DYNAMIC_DELEGATE_TwoParams( DelegateName, Param1Type, Param1Name, Param2Type, Param2Name ) BODY_MACRO_COMBINE(CURRENT_FILE_ID,_,__LINE__,_DELEGATE) FUNC_DECLARE_DYNAMIC_DELEGATE( DelegateName, DelegateName##_DelegateWrapper, FUNC_CONCAT( Param1Type InParam1, Param2Type InParam2 ), FUNC_CONCAT( *this, InParam1, InParam2 ), void, Param1Type, Param2Type )
-#define DECLARE_DYNAMIC_MULTICAST_DELEGATE_TwoParams( DelegateName, Param1Type, Param1Name, Param2Type, Param2Name ) BODY_MACRO_COMBINE(CURRENT_FILE_ID,_,__LINE__,_DELEGATE) FUNC_DECLARE_DYNAMIC_MULTICAST_DELEGATE( DelegateName, DelegateName##_DelegateWrapper, FUNC_CONCAT( Param1Type InParam1, Param2Type InParam2 ), FUNC_CONCAT( *this, InParam1, InParam2 ), void, Param1Type, Param2Type )
+#define DECLARE_DYNAMIC_DELEGATE_TwoParams( DelegateName, Param1Type, Param1Name, Param2Type, Param2Name ) FUNC_DECLARE_DYNAMIC_DELEGATE( DelegateName, void, Param1Type, Param2Type )
+#define DECLARE_DYNAMIC_MULTICAST_DELEGATE_TwoParams( DelegateName, Param1Type, Param1Name, Param2Type, Param2Name ) FUNC_DECLARE_DYNAMIC_MULTICAST_DELEGATE( DelegateName, void, Param1Type, Param2Type )
 #define DECLARE_DELEGATE_RetVal_TwoParams( ReturnValueType, DelegateName, Param1Type, Param2Type ) FUNC_DECLARE_DELEGATE( DelegateName, ReturnValueType, Param1Type, Param2Type )
-#define DECLARE_DYNAMIC_DELEGATE_RetVal_TwoParams( ReturnValueType, DelegateName, Param1Type, Param1Name, Param2Type, Param2Name ) BODY_MACRO_COMBINE(CURRENT_FILE_ID,_,__LINE__,_DELEGATE) FUNC_DECLARE_DYNAMIC_DELEGATE_RETVAL( DelegateName, DelegateName##_DelegateWrapper, ReturnValueType, FUNC_CONCAT( Param1Type InParam1, Param2Type InParam2 ), FUNC_CONCAT( *this, InParam1, InParam2 ), ReturnValueType, Param1Type, Param2Type )
+#define DECLARE_DYNAMIC_DELEGATE_RetVal_TwoParams( ReturnValueType, DelegateName, Param1Type, Param1Name, Param2Type, Param2Name ) FUNC_DECLARE_DYNAMIC_DELEGATE( DelegateName, ReturnValueType, Param1Type, Param2Type )
 
 #define DECLARE_DELEGATE_ThreeParams( DelegateName, Param1Type, Param2Type, Param3Type ) FUNC_DECLARE_DELEGATE( DelegateName, void, Param1Type, Param2Type, Param3Type )
 #define DECLARE_MULTICAST_DELEGATE_ThreeParams( DelegateName, Param1Type, Param2Type, Param3Type ) FUNC_DECLARE_MULTICAST_DELEGATE( DelegateName, void, Param1Type, Param2Type, Param3Type )
 #define DECLARE_TS_MULTICAST_DELEGATE_ThreeParams( DelegateName, Param1Type, Param2Type, Param3Type ) FUNC_DECLARE_TS_MULTICAST_DELEGATE( DelegateName, void, Param1Type, Param2Type, Param3Type )
 #define DECLARE_EVENT_ThreeParams( OwningType, EventName, Param1Type, Param2Type, Param3Type ) FUNC_DECLARE_EVENT( OwningType, EventName, void, Param1Type, Param2Type, Param3Type )
-#define DECLARE_DYNAMIC_DELEGATE_ThreeParams( DelegateName, Param1Type, Param1Name, Param2Type, Param2Name, Param3Type, Param3Name ) BODY_MACRO_COMBINE(CURRENT_FILE_ID,_,__LINE__,_DELEGATE) FUNC_DECLARE_DYNAMIC_DELEGATE( DelegateName, DelegateName##_DelegateWrapper, FUNC_CONCAT( Param1Type InParam1, Param2Type InParam2, Param3Type InParam3 ), FUNC_CONCAT( *this, InParam1, InParam2, InParam3 ), void, Param1Type, Param2Type, Param3Type )
-#define DECLARE_DYNAMIC_MULTICAST_DELEGATE_ThreeParams( DelegateName, Param1Type, Param1Name, Param2Type, Param2Name, Param3Type, Param3Name ) BODY_MACRO_COMBINE(CURRENT_FILE_ID,_,__LINE__,_DELEGATE) FUNC_DECLARE_DYNAMIC_MULTICAST_DELEGATE( DelegateName, DelegateName##_DelegateWrapper, FUNC_CONCAT( Param1Type InParam1, Param2Type InParam2, Param3Type InParam3 ), FUNC_CONCAT( *this, InParam1, InParam2, InParam3 ), void, Param1Type, Param2Type, Param3Type )
+#define DECLARE_DYNAMIC_DELEGATE_ThreeParams( DelegateName, Param1Type, Param1Name, Param2Type, Param2Name, Param3Type, Param3Name ) FUNC_DECLARE_DYNAMIC_DELEGATE( DelegateName, void, Param1Type, Param2Type, Param3Type )
+#define DECLARE_DYNAMIC_MULTICAST_DELEGATE_ThreeParams( DelegateName, Param1Type, Param1Name, Param2Type, Param2Name, Param3Type, Param3Name ) FUNC_DECLARE_DYNAMIC_MULTICAST_DELEGATE( DelegateName, void, Param1Type, Param2Type, Param3Type )
 #define DECLARE_DELEGATE_RetVal_ThreeParams( ReturnValueType, DelegateName, Param1Type, Param2Type, Param3Type ) FUNC_DECLARE_DELEGATE( DelegateName, ReturnValueType, Param1Type, Param2Type, Param3Type )
-#define DECLARE_DYNAMIC_DELEGATE_RetVal_ThreeParams( ReturnValueType, DelegateName, Param1Type, Param1Name, Param2Type, Param2Name, Param3Type, Param3Name ) BODY_MACRO_COMBINE(CURRENT_FILE_ID,_,__LINE__,_DELEGATE) FUNC_DECLARE_DYNAMIC_DELEGATE_RETVAL( DelegateName, DelegateName##_DelegateWrapper, ReturnValueType, FUNC_CONCAT( Param1Type InParam1, Param2Type InParam2, Param3Type InParam3 ), FUNC_CONCAT( *this, InParam1, InParam2, InParam3 ), ReturnValueType, Param1Type, Param2Type, Param3Type )
+#define DECLARE_DYNAMIC_DELEGATE_RetVal_ThreeParams( ReturnValueType, DelegateName, Param1Type, Param1Name, Param2Type, Param2Name, Param3Type, Param3Name ) FUNC_DECLARE_DYNAMIC_DELEGATE( DelegateName, ReturnValueType, Param1Type, Param2Type, Param3Type )
 
 #define DECLARE_DELEGATE_FourParams( DelegateName, Param1Type, Param2Type, Param3Type, Param4Type ) FUNC_DECLARE_DELEGATE( DelegateName, void, Param1Type, Param2Type, Param3Type, Param4Type )
 #define DECLARE_MULTICAST_DELEGATE_FourParams( DelegateName, Param1Type, Param2Type, Param3Type, Param4Type ) FUNC_DECLARE_MULTICAST_DELEGATE( DelegateName, void, Param1Type, Param2Type, Param3Type, Param4Type )
 #define DECLARE_TS_MULTICAST_DELEGATE_FourParams( DelegateName, Param1Type, Param2Type, Param3Type, Param4Type ) FUNC_DECLARE_TS_MULTICAST_DELEGATE( DelegateName, void, Param1Type, Param2Type, Param3Type, Param4Type )
 #define DECLARE_EVENT_FourParams( OwningType, EventName, Param1Type, Param2Type, Param3Type, Param4Type ) FUNC_DECLARE_EVENT( OwningType, EventName, void, Param1Type, Param2Type, Param3Type, Param4Type )
-#define DECLARE_DYNAMIC_DELEGATE_FourParams( DelegateName, Param1Type, Param1Name, Param2Type, Param2Name, Param3Type, Param3Name, Param4Type, Param4Name ) BODY_MACRO_COMBINE(CURRENT_FILE_ID,_,__LINE__,_DELEGATE) FUNC_DECLARE_DYNAMIC_DELEGATE( DelegateName, DelegateName##_DelegateWrapper, FUNC_CONCAT( Param1Type InParam1, Param2Type InParam2, Param3Type InParam3, Param4Type InParam4 ), FUNC_CONCAT( *this, InParam1, InParam2, InParam3, InParam4 ), void, Param1Type, Param2Type, Param3Type, Param4Type )
-#define DECLARE_DYNAMIC_MULTICAST_DELEGATE_FourParams( DelegateName, Param1Type, Param1Name, Param2Type, Param2Name, Param3Type, Param3Name, Param4Type, Param4Name ) BODY_MACRO_COMBINE(CURRENT_FILE_ID,_,__LINE__,_DELEGATE) FUNC_DECLARE_DYNAMIC_MULTICAST_DELEGATE( DelegateName, DelegateName##_DelegateWrapper, FUNC_CONCAT( Param1Type InParam1, Param2Type InParam2, Param3Type InParam3, Param4Type InParam4 ), FUNC_CONCAT( *this, InParam1, InParam2, InParam3, InParam4 ), void, Param1Type, Param2Type, Param3Type, Param4Type )
+#define DECLARE_DYNAMIC_DELEGATE_FourParams( DelegateName, Param1Type, Param1Name, Param2Type, Param2Name, Param3Type, Param3Name, Param4Type, Param4Name ) FUNC_DECLARE_DYNAMIC_DELEGATE( DelegateName, void, Param1Type, Param2Type, Param3Type, Param4Type )
+#define DECLARE_DYNAMIC_MULTICAST_DELEGATE_FourParams( DelegateName, Param1Type, Param1Name, Param2Type, Param2Name, Param3Type, Param3Name, Param4Type, Param4Name ) FUNC_DECLARE_DYNAMIC_MULTICAST_DELEGATE( DelegateName, void, Param1Type, Param2Type, Param3Type, Param4Type )
 #define DECLARE_DELEGATE_RetVal_FourParams( ReturnValueType, DelegateName, Param1Type, Param2Type, Param3Type, Param4Type ) FUNC_DECLARE_DELEGATE( DelegateName, ReturnValueType, Param1Type, Param2Type, Param3Type, Param4Type )
-#define DECLARE_DYNAMIC_DELEGATE_RetVal_FourParams( ReturnValueType, DelegateName, Param1Type, Param1Name, Param2Type, Param2Name, Param3Type, Param3Name, Param4Type, Param4Name ) BODY_MACRO_COMBINE(CURRENT_FILE_ID,_,__LINE__,_DELEGATE) FUNC_DECLARE_DYNAMIC_DELEGATE_RETVAL( DelegateName, DelegateName##_DelegateWrapper, ReturnValueType, FUNC_CONCAT( Param1Type InParam1, Param2Type InParam2, Param3Type InParam3, Param4Type InParam4 ), FUNC_CONCAT( *this, InParam1, InParam2, InParam3, InParam4 ), ReturnValueType, Param1Type, Param2Type, Param3Type, Param4Type )
+#define DECLARE_DYNAMIC_DELEGATE_RetVal_FourParams( ReturnValueType, DelegateName, Param1Type, Param1Name, Param2Type, Param2Name, Param3Type, Param3Name, Param4Type, Param4Name ) FUNC_DECLARE_DYNAMIC_DELEGATE( DelegateName, ReturnValueType, Param1Type, Param2Type, Param3Type, Param4Type )
 
 #define DECLARE_DELEGATE_FiveParams( DelegateName, Param1Type, Param2Type, Param3Type, Param4Type, Param5Type ) FUNC_DECLARE_DELEGATE( DelegateName, void, Param1Type, Param2Type, Param3Type, Param4Type, Param5Type )
 #define DECLARE_MULTICAST_DELEGATE_FiveParams( DelegateName, Param1Type, Param2Type, Param3Type, Param4Type, Param5Type ) FUNC_DECLARE_MULTICAST_DELEGATE( DelegateName, void, Param1Type, Param2Type, Param3Type, Param4Type, Param5Type )
 #define DECLARE_TS_MULTICAST_DELEGATE_FiveParams( DelegateName, Param1Type, Param2Type, Param3Type, Param4Type, Param5Type ) FUNC_DECLARE_TS_MULTICAST_DELEGATE( DelegateName, void, Param1Type, Param2Type, Param3Type, Param4Type, Param5Type )
 #define DECLARE_EVENT_FiveParams( OwningType, EventName, Param1Type, Param2Type, Param3Type, Param4Type, Param5Type ) FUNC_DECLARE_EVENT( OwningType, EventName, void, Param1Type, Param2Type, Param3Type, Param4Type, Param5Type )
-#define DECLARE_DYNAMIC_DELEGATE_FiveParams( DelegateName, Param1Type, Param1Name, Param2Type, Param2Name, Param3Type, Param3Name, Param4Type, Param4Name, Param5Type, Param5Name ) BODY_MACRO_COMBINE(CURRENT_FILE_ID,_,__LINE__,_DELEGATE) FUNC_DECLARE_DYNAMIC_DELEGATE( DelegateName, DelegateName##_DelegateWrapper, FUNC_CONCAT( Param1Type InParam1, Param2Type InParam2, Param3Type InParam3, Param4Type InParam4, Param5Type InParam5 ), FUNC_CONCAT( *this, InParam1, InParam2, InParam3, InParam4, InParam5 ), void, Param1Type, Param2Type, Param3Type, Param4Type, Param5Type )
-#define DECLARE_DYNAMIC_MULTICAST_DELEGATE_FiveParams( DelegateName, Param1Type, Param1Name, Param2Type, Param2Name, Param3Type, Param3Name, Param4Type, Param4Name, Param5Type, Param5Name ) BODY_MACRO_COMBINE(CURRENT_FILE_ID,_,__LINE__,_DELEGATE) FUNC_DECLARE_DYNAMIC_MULTICAST_DELEGATE( DelegateName, DelegateName##_DelegateWrapper, FUNC_CONCAT( Param1Type InParam1, Param2Type InParam2, Param3Type InParam3, Param4Type InParam4, Param5Type InParam5 ), FUNC_CONCAT( *this, InParam1, InParam2, InParam3, InParam4, InParam5 ), void, Param1Type, Param2Type, Param3Type, Param4Type, Param5Type )
+#define DECLARE_DYNAMIC_DELEGATE_FiveParams( DelegateName, Param1Type, Param1Name, Param2Type, Param2Name, Param3Type, Param3Name, Param4Type, Param4Name, Param5Type, Param5Name ) FUNC_DECLARE_DYNAMIC_DELEGATE( DelegateName, void, Param1Type, Param2Type, Param3Type, Param4Type, Param5Type )
+#define DECLARE_DYNAMIC_MULTICAST_DELEGATE_FiveParams( DelegateName, Param1Type, Param1Name, Param2Type, Param2Name, Param3Type, Param3Name, Param4Type, Param4Name, Param5Type, Param5Name ) FUNC_DECLARE_DYNAMIC_MULTICAST_DELEGATE( DelegateName, void, Param1Type, Param2Type, Param3Type, Param4Type, Param5Type )
 #define DECLARE_DELEGATE_RetVal_FiveParams( ReturnValueType, DelegateName, Param1Type, Param2Type, Param3Type, Param4Type, Param5Type ) FUNC_DECLARE_DELEGATE( DelegateName, ReturnValueType, Param1Type, Param2Type, Param3Type, Param4Type, Param5Type )
-#define DECLARE_DYNAMIC_DELEGATE_RetVal_FiveParams( ReturnValueType, DelegateName, Param1Type, Param1Name, Param2Type, Param2Name, Param3Type, Param3Name, Param4Type, Param4Name, Param5Type, Param5Name ) BODY_MACRO_COMBINE(CURRENT_FILE_ID,_,__LINE__,_DELEGATE) FUNC_DECLARE_DYNAMIC_DELEGATE_RETVAL( DelegateName, DelegateName##_DelegateWrapper, ReturnValueType, FUNC_CONCAT( Param1Type InParam1, Param2Type InParam2, Param3Type InParam3, Param4Type InParam4, Param5Type InParam5 ), FUNC_CONCAT( *this, InParam1, InParam2, InParam3, InParam4, InParam5 ), ReturnValueType, Param1Type, Param2Type, Param3Type, Param4Type, Param5Type )
+#define DECLARE_DYNAMIC_DELEGATE_RetVal_FiveParams( ReturnValueType, DelegateName, Param1Type, Param1Name, Param2Type, Param2Name, Param3Type, Param3Name, Param4Type, Param4Name, Param5Type, Param5Name ) FUNC_DECLARE_DYNAMIC_DELEGATE( DelegateName, ReturnValueType, Param1Type, Param2Type, Param3Type, Param4Type, Param5Type )
 
 #define DECLARE_DELEGATE_SixParams( DelegateName, Param1Type, Param2Type, Param3Type, Param4Type, Param5Type, Param6Type ) FUNC_DECLARE_DELEGATE( DelegateName, void, Param1Type, Param2Type, Param3Type, Param4Type, Param5Type, Param6Type )
 #define DECLARE_MULTICAST_DELEGATE_SixParams( DelegateName, Param1Type, Param2Type, Param3Type, Param4Type, Param5Type, Param6Type ) FUNC_DECLARE_MULTICAST_DELEGATE( DelegateName, void, Param1Type, Param2Type, Param3Type, Param4Type, Param5Type, Param6Type )
 #define DECLARE_TS_MULTICAST_DELEGATE_SixParams( DelegateName, Param1Type, Param2Type, Param3Type, Param4Type, Param5Type, Param6Type ) FUNC_DECLARE_TS_MULTICAST_DELEGATE( DelegateName, void, Param1Type, Param2Type, Param3Type, Param4Type, Param5Type, Param6Type )
 #define DECLARE_EVENT_SixParams( OwningType, EventName, Param1Type, Param2Type, Param3Type, Param4Type, Param5Type, Param6Type ) FUNC_DECLARE_EVENT( OwningType, EventName, void, Param1Type, Param2Type, Param3Type, Param4Type, Param5Type, Param6Type )
-#define DECLARE_DYNAMIC_DELEGATE_SixParams( DelegateName, Param1Type, Param1Name, Param2Type, Param2Name, Param3Type, Param3Name, Param4Type, Param4Name, Param5Type, Param5Name, Param6Type, Param6Name ) BODY_MACRO_COMBINE(CURRENT_FILE_ID,_,__LINE__,_DELEGATE) FUNC_DECLARE_DYNAMIC_DELEGATE( DelegateName, DelegateName##_DelegateWrapper, FUNC_CONCAT( Param1Type InParam1, Param2Type InParam2, Param3Type InParam3, Param4Type InParam4, Param5Type InParam5, Param6Type InParam6 ), FUNC_CONCAT( *this, InParam1, InParam2, InParam3, InParam4, InParam5, InParam6 ), void, Param1Type, Param2Type, Param3Type, Param4Type, Param5Type, Param6Type )
-#define DECLARE_DYNAMIC_MULTICAST_DELEGATE_SixParams( DelegateName, Param1Type, Param1Name, Param2Type, Param2Name, Param3Type, Param3Name, Param4Type, Param4Name, Param5Type, Param5Name, Param6Type, Param6Name ) BODY_MACRO_COMBINE(CURRENT_FILE_ID,_,__LINE__,_DELEGATE) FUNC_DECLARE_DYNAMIC_MULTICAST_DELEGATE( DelegateName, DelegateName##_DelegateWrapper, FUNC_CONCAT( Param1Type InParam1, Param2Type InParam2, Param3Type InParam3, Param4Type InParam4, Param5Type InParam5, Param6Type InParam6 ), FUNC_CONCAT( *this, InParam1, InParam2, InParam3, InParam4, InParam5, InParam6 ), void, Param1Type, Param2Type, Param3Type, Param4Type, Param5Type, Param6Type )
+#define DECLARE_DYNAMIC_DELEGATE_SixParams( DelegateName, Param1Type, Param1Name, Param2Type, Param2Name, Param3Type, Param3Name, Param4Type, Param4Name, Param5Type, Param5Name, Param6Type, Param6Name ) FUNC_DECLARE_DYNAMIC_DELEGATE( DelegateName, void, Param1Type, Param2Type, Param3Type, Param4Type, Param5Type, Param6Type )
+#define DECLARE_DYNAMIC_MULTICAST_DELEGATE_SixParams( DelegateName, Param1Type, Param1Name, Param2Type, Param2Name, Param3Type, Param3Name, Param4Type, Param4Name, Param5Type, Param5Name, Param6Type, Param6Name ) FUNC_DECLARE_DYNAMIC_MULTICAST_DELEGATE( DelegateName, void, Param1Type, Param2Type, Param3Type, Param4Type, Param5Type, Param6Type )
 #define DECLARE_DELEGATE_RetVal_SixParams( ReturnValueType, DelegateName, Param1Type, Param2Type, Param3Type, Param4Type, Param5Type, Param6Type ) FUNC_DECLARE_DELEGATE( DelegateName, ReturnValueType, Param1Type, Param2Type, Param3Type, Param4Type, Param5Type, Param6Type )
-#define DECLARE_DYNAMIC_DELEGATE_RetVal_SixParams( ReturnValueType, DelegateName, Param1Type, Param1Name, Param2Type, Param2Name, Param3Type, Param3Name, Param4Type, Param4Name, Param5Type, Param5Name, Param6Type, Param6Name ) BODY_MACRO_COMBINE(CURRENT_FILE_ID,_,__LINE__,_DELEGATE) FUNC_DECLARE_DYNAMIC_DELEGATE_RETVAL( DelegateName, DelegateName##_DelegateWrapper, ReturnValueType, FUNC_CONCAT( Param1Type InParam1, Param2Type InParam2, Param3Type InParam3, Param4Type InParam4, Param5Type InParam5, Param6Type InParam6 ), FUNC_CONCAT( *this, InParam1, InParam2, InParam3, InParam4, InParam5, InParam6 ), ReturnValueType, Param1Type, Param2Type, Param3Type, Param4Type, Param5Type, Param6Type )
+#define DECLARE_DYNAMIC_DELEGATE_RetVal_SixParams( ReturnValueType, DelegateName, Param1Type, Param1Name, Param2Type, Param2Name, Param3Type, Param3Name, Param4Type, Param4Name, Param5Type, Param5Name, Param6Type, Param6Name ) FUNC_DECLARE_DYNAMIC_DELEGATE( DelegateName, ReturnValueType, Param1Type, Param2Type, Param3Type, Param4Type, Param5Type, Param6Type )
 
 #define DECLARE_DELEGATE_SevenParams( DelegateName, Param1Type, Param2Type, Param3Type, Param4Type, Param5Type, Param6Type, Param7Type ) FUNC_DECLARE_DELEGATE( DelegateName, void, Param1Type, Param2Type, Param3Type, Param4Type, Param5Type, Param6Type, Param7Type )
 #define DECLARE_MULTICAST_DELEGATE_SevenParams( DelegateName, Param1Type, Param2Type, Param3Type, Param4Type, Param5Type, Param6Type, Param7Type ) FUNC_DECLARE_MULTICAST_DELEGATE( DelegateName, void, Param1Type, Param2Type, Param3Type, Param4Type, Param5Type, Param6Type, Param7Type )
 #define DECLARE_TS_MULTICAST_DELEGATE_SevenParams( DelegateName, Param1Type, Param2Type, Param3Type, Param4Type, Param5Type, Param6Type, Param7Type ) FUNC_DECLARE_TS_MULTICAST_DELEGATE( DelegateName, void, Param1Type, Param2Type, Param3Type, Param4Type, Param5Type, Param6Type, Param7Type )
 #define DECLARE_EVENT_SevenParams( OwningType, EventName, Param1Type, Param2Type, Param3Type, Param4Type, Param5Type, Param6Type, Param7Type ) FUNC_DECLARE_EVENT( OwningType, EventName, void, Param1Type, Param2Type, Param3Type, Param4Type, Param5Type, Param6Type, Param7Type )
-#define DECLARE_DYNAMIC_DELEGATE_SevenParams( DelegateName, Param1Type, Param1Name, Param2Type, Param2Name, Param3Type, Param3Name, Param4Type, Param4Name, Param5Type, Param5Name, Param6Type, Param6Name, Param7Type, Param7Name ) BODY_MACRO_COMBINE(CURRENT_FILE_ID,_,__LINE__,_DELEGATE) FUNC_DECLARE_DYNAMIC_DELEGATE( DelegateName, DelegateName##_DelegateWrapper, FUNC_CONCAT( Param1Type InParam1, Param2Type InParam2, Param3Type InParam3, Param4Type InParam4, Param5Type InParam5, Param6Type InParam6, Param7Type InParam7 ), FUNC_CONCAT( *this, InParam1, InParam2, InParam3, InParam4, InParam5, InParam6, InParam7 ), void, Param1Type, Param2Type, Param3Type, Param4Type, Param5Type, Param6Type, Param7Type )
-#define DECLARE_DYNAMIC_MULTICAST_DELEGATE_SevenParams( DelegateName, Param1Type, Param1Name, Param2Type, Param2Name, Param3Type, Param3Name, Param4Type, Param4Name, Param5Type, Param5Name, Param6Type, Param6Name, Param7Type, Param7Name ) BODY_MACRO_COMBINE(CURRENT_FILE_ID,_,__LINE__,_DELEGATE) FUNC_DECLARE_DYNAMIC_MULTICAST_DELEGATE( DelegateName, DelegateName##_DelegateWrapper, FUNC_CONCAT( Param1Type InParam1, Param2Type InParam2, Param3Type InParam3, Param4Type InParam4, Param5Type InParam5, Param6Type InParam6, Param7Type InParam7 ), FUNC_CONCAT( *this, InParam1, InParam2, InParam3, InParam4, InParam5, InParam6, InParam7 ), void, Param1Type, Param2Type, Param3Type, Param4Type, Param5Type, Param6Type, Param7Type )
+#define DECLARE_DYNAMIC_DELEGATE_SevenParams( DelegateName, Param1Type, Param1Name, Param2Type, Param2Name, Param3Type, Param3Name, Param4Type, Param4Name, Param5Type, Param5Name, Param6Type, Param6Name, Param7Type, Param7Name ) FUNC_DECLARE_DYNAMIC_DELEGATE( DelegateName, void, Param1Type, Param2Type, Param3Type, Param4Type, Param5Type, Param6Type, Param7Type )
+#define DECLARE_DYNAMIC_MULTICAST_DELEGATE_SevenParams( DelegateName, Param1Type, Param1Name, Param2Type, Param2Name, Param3Type, Param3Name, Param4Type, Param4Name, Param5Type, Param5Name, Param6Type, Param6Name, Param7Type, Param7Name ) FUNC_DECLARE_DYNAMIC_MULTICAST_DELEGATE( DelegateName, void, Param1Type, Param2Type, Param3Type, Param4Type, Param5Type, Param6Type, Param7Type )
 #define DECLARE_DELEGATE_RetVal_SevenParams( ReturnValueType, DelegateName, Param1Type, Param2Type, Param3Type, Param4Type, Param5Type, Param6Type, Param7Type ) FUNC_DECLARE_DELEGATE( DelegateName, ReturnValueType, Param1Type, Param2Type, Param3Type, Param4Type, Param5Type, Param6Type, Param7Type )
-#define DECLARE_DYNAMIC_DELEGATE_RetVal_SevenParams( ReturnValueType, DelegateName, Param1Type, Param1Name, Param2Type, Param2Name, Param3Type, Param3Name, Param4Type, Param4Name, Param5Type, Param5Name, Param6Type, Param6Name, Param7Type, Param7Name ) BODY_MACRO_COMBINE(CURRENT_FILE_ID,_,__LINE__,_DELEGATE) FUNC_DECLARE_DYNAMIC_DELEGATE_RETVAL( DelegateName, DelegateName##_DelegateWrapper, ReturnValueType, FUNC_CONCAT( Param1Type InParam1, Param2Type InParam2, Param3Type InParam3, Param4Type InParam4, Param5Type InParam5, Param6Type InParam6, Param7Type InParam7 ), FUNC_CONCAT( *this, InParam1, InParam2, InParam3, InParam4, InParam5, InParam6, InParam7 ), ReturnValueType, Param1Type, Param2Type, Param3Type, Param4Type, Param5Type, Param6Type, Param7Type )
+#define DECLARE_DYNAMIC_DELEGATE_RetVal_SevenParams( ReturnValueType, DelegateName, Param1Type, Param1Name, Param2Type, Param2Name, Param3Type, Param3Name, Param4Type, Param4Name, Param5Type, Param5Name, Param6Type, Param6Name, Param7Type, Param7Name ) FUNC_DECLARE_DYNAMIC_DELEGATE( DelegateName, ReturnValueType, Param1Type, Param2Type, Param3Type, Param4Type, Param5Type, Param6Type, Param7Type )
 
 #define DECLARE_DELEGATE_EightParams( DelegateName, Param1Type, Param2Type, Param3Type, Param4Type, Param5Type, Param6Type, Param7Type, Param8Type ) FUNC_DECLARE_DELEGATE( DelegateName, void, Param1Type, Param2Type, Param3Type, Param4Type, Param5Type, Param6Type, Param7Type, Param8Type )
 #define DECLARE_MULTICAST_DELEGATE_EightParams( DelegateName, Param1Type, Param2Type, Param3Type, Param4Type, Param5Type, Param6Type, Param7Type, Param8Type ) FUNC_DECLARE_MULTICAST_DELEGATE( DelegateName, void, Param1Type, Param2Type, Param3Type, Param4Type, Param5Type, Param6Type, Param7Type, Param8Type )
 #define DECLARE_TS_MULTICAST_DELEGATE_EightParams( DelegateName, Param1Type, Param2Type, Param3Type, Param4Type, Param5Type, Param6Type, Param7Type, Param8Type ) FUNC_DECLARE_TS_MULTICAST_DELEGATE( DelegateName, void, Param1Type, Param2Type, Param3Type, Param4Type, Param5Type, Param6Type, Param7Type, Param8Type )
 #define DECLARE_EVENT_EightParams( OwningType, EventName, Param1Type, Param2Type, Param3Type, Param4Type, Param5Type, Param6Type, Param7Type, Param8Type ) FUNC_DECLARE_EVENT( OwningType, EventName, void, Param1Type, Param2Type, Param3Type, Param4Type, Param5Type, Param6Type, Param7Type, Param8Type )
-#define DECLARE_DYNAMIC_DELEGATE_EightParams( DelegateName, Param1Type, Param1Name, Param2Type, Param2Name, Param3Type, Param3Name, Param4Type, Param4Name, Param5Type, Param5Name, Param6Type, Param6Name, Param7Type, Param7Name, Param8Type, Param8Name ) BODY_MACRO_COMBINE(CURRENT_FILE_ID,_,__LINE__,_DELEGATE) FUNC_DECLARE_DYNAMIC_DELEGATE( DelegateName, DelegateName##_DelegateWrapper, FUNC_CONCAT( Param1Type InParam1, Param2Type InParam2, Param3Type InParam3, Param4Type InParam4, Param5Type InParam5, Param6Type InParam6, Param7Type InParam7, Param8Type InParam8 ), FUNC_CONCAT( *this, InParam1, InParam2, InParam3, InParam4, InParam5, InParam6, InParam7, InParam8 ), void, Param1Type, Param2Type, Param3Type, Param4Type, Param5Type, Param6Type, Param7Type, Param8Type )
-#define DECLARE_DYNAMIC_MULTICAST_DELEGATE_EightParams( DelegateName, Param1Type, Param1Name, Param2Type, Param2Name, Param3Type, Param3Name, Param4Type, Param4Name, Param5Type, Param5Name, Param6Type, Param6Name, Param7Type, Param7Name, Param8Type, Param8Name ) BODY_MACRO_COMBINE(CURRENT_FILE_ID,_,__LINE__,_DELEGATE) FUNC_DECLARE_DYNAMIC_MULTICAST_DELEGATE( DelegateName, DelegateName##_DelegateWrapper, FUNC_CONCAT( Param1Type InParam1, Param2Type InParam2, Param3Type InParam3, Param4Type InParam4, Param5Type InParam5, Param6Type InParam6, Param7Type InParam7, Param8Type InParam8 ), FUNC_CONCAT( *this, InParam1, InParam2, InParam3, InParam4, InParam5, InParam6, InParam7, InParam8 ), void, Param1Type, Param2Type, Param3Type, Param4Type, Param5Type, Param6Type, Param7Type, Param8Type )
+#define DECLARE_DYNAMIC_DELEGATE_EightParams( DelegateName, Param1Type, Param1Name, Param2Type, Param2Name, Param3Type, Param3Name, Param4Type, Param4Name, Param5Type, Param5Name, Param6Type, Param6Name, Param7Type, Param7Name, Param8Type, Param8Name ) FUNC_DECLARE_DYNAMIC_DELEGATE( DelegateName, void, Param1Type, Param2Type, Param3Type, Param4Type, Param5Type, Param6Type, Param7Type, Param8Type )
+#define DECLARE_DYNAMIC_MULTICAST_DELEGATE_EightParams( DelegateName, Param1Type, Param1Name, Param2Type, Param2Name, Param3Type, Param3Name, Param4Type, Param4Name, Param5Type, Param5Name, Param6Type, Param6Name, Param7Type, Param7Name, Param8Type, Param8Name ) FUNC_DECLARE_DYNAMIC_MULTICAST_DELEGATE( DelegateName, void, Param1Type, Param2Type, Param3Type, Param4Type, Param5Type, Param6Type, Param7Type, Param8Type )
 #define DECLARE_DELEGATE_RetVal_EightParams( ReturnValueType, DelegateName, Param1Type, Param2Type, Param3Type, Param4Type, Param5Type, Param6Type, Param7Type, Param8Type ) FUNC_DECLARE_DELEGATE( DelegateName, ReturnValueType, Param1Type, Param2Type, Param3Type, Param4Type, Param5Type, Param6Type, Param7Type, Param8Type )
-#define DECLARE_DYNAMIC_DELEGATE_RetVal_EightParams( ReturnValueType, DelegateName, Param1Type, Param1Name, Param2Type, Param2Name, Param3Type, Param3Name, Param4Type, Param4Name, Param5Type, Param5Name, Param6Type, Param6Name, Param7Type, Param7Name, Param8Type, Param8Name ) BODY_MACRO_COMBINE(CURRENT_FILE_ID,_,__LINE__,_DELEGATE) FUNC_DECLARE_DYNAMIC_DELEGATE_RETVAL( DelegateName, DelegateName##_DelegateWrapper, ReturnValueType, FUNC_CONCAT( Param1Type InParam1, Param2Type InParam2, Param3Type InParam3, Param4Type InParam4, Param5Type InParam5, Param6Type InParam6, Param7Type InParam7, Param8Type InParam8 ), FUNC_CONCAT( *this, InParam1, InParam2, InParam3, InParam4, InParam5, InParam6, InParam7, InParam8 ), ReturnValueType, Param1Type, Param2Type, Param3Type, Param4Type, Param5Type, Param6Type, Param7Type, Param8Type )
+#define DECLARE_DYNAMIC_DELEGATE_RetVal_EightParams( ReturnValueType, DelegateName, Param1Type, Param1Name, Param2Type, Param2Name, Param3Type, Param3Name, Param4Type, Param4Name, Param5Type, Param5Name, Param6Type, Param6Name, Param7Type, Param7Name, Param8Type, Param8Name ) FUNC_DECLARE_DYNAMIC_DELEGATE( DelegateName, ReturnValueType, Param1Type, Param2Type, Param3Type, Param4Type, Param5Type, Param6Type, Param7Type, Param8Type )
 
 #define DECLARE_DELEGATE_NineParams( DelegateName, Param1Type, Param2Type, Param3Type, Param4Type, Param5Type, Param6Type, Param7Type, Param8Type, Param9Type ) FUNC_DECLARE_DELEGATE( DelegateName, void, Param1Type, Param2Type, Param3Type, Param4Type, Param5Type, Param6Type, Param7Type, Param8Type, Param9Type )
 #define DECLARE_MULTICAST_DELEGATE_NineParams( DelegateName, Param1Type, Param2Type, Param3Type, Param4Type, Param5Type, Param6Type, Param7Type, Param8Type, Param9Type ) FUNC_DECLARE_MULTICAST_DELEGATE( DelegateName, void, Param1Type, Param2Type, Param3Type, Param4Type, Param5Type, Param6Type, Param7Type, Param8Type, Param9Type )
 #define DECLARE_TS_MULTICAST_DELEGATE_NineParams( DelegateName, Param1Type, Param2Type, Param3Type, Param4Type, Param5Type, Param6Type, Param7Type, Param8Type, Param9Type ) FUNC_DECLARE_TS_MULTICAST_DELEGATE( DelegateName, void, Param1Type, Param2Type, Param3Type, Param4Type, Param5Type, Param6Type, Param7Type, Param8Type, Param9Type )
 #define DECLARE_EVENT_NineParams( OwningType, EventName, Param1Type, Param2Type, Param3Type, Param4Type, Param5Type, Param6Type, Param7Type, Param8Type, Param9Type ) FUNC_DECLARE_EVENT( OwningType, EventName, void, Param1Type, Param2Type, Param3Type, Param4Type, Param5Type, Param6Type, Param7Type, Param8Type, Param9Type )
-#define DECLARE_DYNAMIC_DELEGATE_NineParams( DelegateName, Param1Type, Param1Name, Param2Type, Param2Name, Param3Type, Param3Name, Param4Type, Param4Name, Param5Type, Param5Name, Param6Type, Param6Name, Param7Type, Param7Name, Param8Type, Param8Name, Param9Type, Param9Name ) BODY_MACRO_COMBINE(CURRENT_FILE_ID,_,__LINE__,_DELEGATE) FUNC_DECLARE_DYNAMIC_DELEGATE( DelegateName, DelegateName##_DelegateWrapper, FUNC_CONCAT( Param1Type InParam1, Param2Type InParam2, Param3Type InParam3, Param4Type InParam4, Param5Type InParam5, Param6Type InParam6, Param7Type InParam7, Param8Type InParam8, Param9Type InParam9 ), FUNC_CONCAT( *this, InParam1, InParam2, InParam3, InParam4, InParam5, InParam6, InParam7, InParam8, InParam9 ), void, Param1Type, Param2Type, Param3Type, Param4Type, Param5Type, Param6Type, Param7Type, Param8Type, Param9Type )
-#define DECLARE_DYNAMIC_MULTICAST_DELEGATE_NineParams( DelegateName, Param1Type, Param1Name, Param2Type, Param2Name, Param3Type, Param3Name, Param4Type, Param4Name, Param5Type, Param5Name, Param6Type, Param6Name, Param7Type, Param7Name, Param8Type, Param8Name, Param9Type, Param9Name ) BODY_MACRO_COMBINE(CURRENT_FILE_ID,_,__LINE__,_DELEGATE) FUNC_DECLARE_DYNAMIC_MULTICAST_DELEGATE( DelegateName, DelegateName##_DelegateWrapper, FUNC_CONCAT( Param1Type InParam1, Param2Type InParam2, Param3Type InParam3, Param4Type InParam4, Param5Type InParam5, Param6Type InParam6, Param7Type InParam7, Param8Type InParam8, Param9Type InParam9 ), FUNC_CONCAT( *this, InParam1, InParam2, InParam3, InParam4, InParam5, InParam6, InParam7, InParam8, InParam9 ), void, Param1Type, Param2Type, Param3Type, Param4Type, Param5Type, Param6Type, Param7Type, Param8Type, Param9Type )
+#define DECLARE_DYNAMIC_DELEGATE_NineParams( DelegateName, Param1Type, Param1Name, Param2Type, Param2Name, Param3Type, Param3Name, Param4Type, Param4Name, Param5Type, Param5Name, Param6Type, Param6Name, Param7Type, Param7Name, Param8Type, Param8Name, Param9Type, Param9Name ) FUNC_DECLARE_DYNAMIC_DELEGATE( DelegateName, void, Param1Type, Param2Type, Param3Type, Param4Type, Param5Type, Param6Type, Param7Type, Param8Type, Param9Type )
+#define DECLARE_DYNAMIC_MULTICAST_DELEGATE_NineParams( DelegateName, Param1Type, Param1Name, Param2Type, Param2Name, Param3Type, Param3Name, Param4Type, Param4Name, Param5Type, Param5Name, Param6Type, Param6Name, Param7Type, Param7Name, Param8Type, Param8Name, Param9Type, Param9Name ) FUNC_DECLARE_DYNAMIC_MULTICAST_DELEGATE( DelegateName, void, Param1Type, Param2Type, Param3Type, Param4Type, Param5Type, Param6Type, Param7Type, Param8Type, Param9Type )
 #define DECLARE_DELEGATE_RetVal_NineParams( ReturnValueType, DelegateName, Param1Type, Param2Type, Param3Type, Param4Type, Param5Type, Param6Type, Param7Type, Param8Type, Param9Type ) FUNC_DECLARE_DELEGATE( DelegateName, ReturnValueType, Param1Type, Param2Type, Param3Type, Param4Type, Param5Type, Param6Type, Param7Type, Param8Type, Param9Type )
-#define DECLARE_DYNAMIC_DELEGATE_RetVal_NineParams( ReturnValueType, DelegateName, Param1Type, Param1Name, Param2Type, Param2Name, Param3Type, Param3Name, Param4Type, Param4Name, Param5Type, Param5Name, Param6Type, Param6Name, Param7Type, Param7Name, Param8Type, Param8Name, Param9Type, Param9Name ) BODY_MACRO_COMBINE(CURRENT_FILE_ID,_,__LINE__,_DELEGATE) FUNC_DECLARE_DYNAMIC_DELEGATE_RETVAL( DelegateName, DelegateName##_DelegateWrapper, ReturnValueType, FUNC_CONCAT( Param1Type InParam1, Param2Type InParam2, Param3Type InParam3, Param4Type InParam4, Param5Type InParam5, Param6Type InParam6, Param7Type InParam7, Param8Type InParam8, Param9Type InParam9 ), FUNC_CONCAT( *this, InParam1, InParam2, InParam3, InParam4, InParam5, InParam6, InParam7, InParam8, InParam9 ), ReturnValueType, Param1Type, Param2Type, Param3Type, Param4Type, Param5Type, Param6Type, Param7Type, Param8Type, Param9Type )
+#define DECLARE_DYNAMIC_DELEGATE_RetVal_NineParams( ReturnValueType, DelegateName, Param1Type, Param1Name, Param2Type, Param2Name, Param3Type, Param3Name, Param4Type, Param4Name, Param5Type, Param5Name, Param6Type, Param6Name, Param7Type, Param7Name, Param8Type, Param8Name, Param9Type, Param9Name ) FUNC_DECLARE_DYNAMIC_DELEGATE( DelegateName, ReturnValueType, Param1Type, Param2Type, Param3Type, Param4Type, Param5Type, Param6Type, Param7Type, Param8Type, Param9Type )
diff --git a/Engine/Source/Runtime/Core/Public/Delegates/DelegateSignatureImpl.inl b/Engine/Source/Runtime/Core/Public/Delegates/DelegateSignatureImpl.inl
--- a/Engine/Source/Runtime/Core/Public/Delegates/DelegateSignatureImpl.inl
+++ b/Engine/Source/Runtime/Core/Public/Delegates/DelegateSignatureImpl.inl
@@ -10,6 +10,7 @@
 // HEADER_UNIT_SKIP - Not included directly
 
 #include "CoreTypes.h"
+#include "Concepts/Void.h"
 #include "Misc/AssertionMacros.h"
 #include "UObject/NameTypes.h"
 #include "UObject/ScriptDelegates.h"
@@ -1207,9 +1208,25 @@ public:
 		__Internal_BindDynamic(ToRawPtr(InUserObject), InMethodPtr, InFunctionName);
 	}
 
-	// NOTE:  Execute() method must be defined in derived classes
+	/** Execute the delegate.  If the function pointer is not valid, an error will occur. */
+	inline RetValType Execute(UE::Core::Private::TDelegateCallTraits_T<ParamTypes>... Params) const
+	{
+		/* Verify that the user object is still valid.  We only have a weak reference to it. */
+		checkSlow(this->IsBound());
+		return UE::Core::Private::TDelegateFunctionWrapper<RetValType, UE::Core::Private::TDelegateCallTraits_T<ParamTypes>...>::CallWrapper(*this, Params...);
+	}
 
-	// NOTE:  ExecuteIfBound() method must be defined in derived classes
+	/** Execute the delegate, but only if the function pointer is still valid */
+	inline bool ExecuteIfBound(UE::Core::Private::TDelegateCallTraits_T<ParamTypes>... Params) const
+		requires (UE::CVoid<RetValType>)
+	{
+		if (this->IsBound())
+		{
+			UE::Core::Private::TDelegateFunctionWrapper<RetValType, UE::Core::Private::TDelegateCallTraits_T<ParamTypes>...>::CallWrapper(*this, Params...);
+			return true;
+		}
+		return false;
+	}
 };
 
 
@@ -1347,7 +1364,11 @@ public:
 		__Internal_RemoveDynamic(ToRawPtr(InUserObject), InMethodPtr, InFunctionName);
 	}
 
-	// NOTE:  Broadcast() method must be defined in derived classes
+	/** Broadcasts this delegate to all bound objects, except to those that may have expired */
+	inline void Broadcast(UE::Core::Private::TDelegateCallTraits_T<ParamTypes>... Params) const
+	{
+		UE::Core::Private::TDelegateFunctionWrapper<RetValType, UE::Core::Private::TDelegateCallTraits_T<ParamTypes>...>::CallWrapper(*this, Params...);
+	}
 };
 
 #if UE_ENABLE_INCLUDE_ORDER_DEPRECATED_IN_5_7
diff --git a/Engine/Source/Runtime/Core/Public/UObject/ScriptDelegates.h b/Engine/Source/Runtime/Core/Public/UObject/ScriptDelegates.h
--- a/Engine/Source/Runtime/Core/Public/UObject/ScriptDelegates.h
+++ b/Engine/Source/Runtime/Core/Public/UObject/ScriptDelegates.h
@@ -2,6 +2,8 @@
 
 #pragma once
 
+#include "Concepts/CompleteType.h"
+#include "Concepts/Scalar.h"
 #include "Containers/Array.h"
 #include "Containers/ContainerAllocationPolicies.h"
 #include "Containers/UnrealString.h"
@@ -12,6 +14,7 @@
 #include "Templates/TypeHash.h"
 #include "Templates/UnrealTypeTraits.h"
 #include "UObject/NameTypes.h"
+#include <type_traits>
 
 namespace UE::Core::Private
 {
@@ -25,6 +28,143 @@ namespace UE::Core::Private
 		using ThreadSafetyMode = InThreadSafetyMode;
 		using UnicastThreadSafetyModeForMulticasts = FNotThreadSafeNotCheckedDelegateMode;
 	};
+
+	// Describes a type which is constructible with a ForceInit argument.
+	template <typename T>
+	concept CForceInitable = UE::CCompleteType<T> && std::is_constructible_v<T, EForceInit>;
+
+	// The code below effectively reimplements a minimal version of TTuple, tailored to the needs
+	// of delegates' param structs passed to ProcessEvent.
+	//
+	// It distinguishes a return element from parameter elements, giving it a different member name for
+	// easy access, forcing its default value with either value init or EForceInit depending on the type.
+	// 
+	// It also directly constructs the elements rather than default construct and copy assign, and uses
+	// perfect forwarding so that moves happen instead of copies where possible.
+
+	template <typename T>
+	struct TDelegateWrapperReturnBase
+	{
+		std::remove_cv_t<std::remove_reference_t<T>> ReturnValue = {};
+	};
+	template <CForceInitable T>
+	struct TDelegateWrapperReturnBase<T>
+	{
+		std::remove_cv_t<std::remove_reference_t<T>> ReturnValue = T(ForceInit);
+	};
+
+	template <typename T, uint32 Index>
+	struct TDelegateWrapperParmBase
+	{
+		using MemberType = std::remove_cv_t<std::remove_reference_t<T>>;
+
+		MemberType Member;
+
+		// This odd constraint is to allow selection of the correct MoveOut overload by specifying the
+		// index when they are merged into the derived class's overload set.
+		//
+		// The argument is a const reference because we want it to compile when passed a const param,
+		// but the move will only happen when the the original parameter was a non-const reference, so
+		// it's fine.
+		template <uint32 CallIndex>
+			requires (Index == CallIndex)
+		UE_REWRITE void MoveOut(const MemberType& Out)
+		{
+			if constexpr (std::is_reference_v<T> && !std::is_const_v<std::remove_reference_t<T>>)
+			{
+				const_cast<MemberType&>(Out) = MoveTemp(this->Member);
+			}
+		}
+	};
+
+	// This is the dedicated tuple type, but it takes a function type instead of a list of types,
+	// and it requires a TIntegerSequence parameter, which is hidden inside TTuple's
+	// implementation but provided explicitly here because this doesn't need to have the nicest API.
+	template <typename FuncType, typename IndexSequenceType>
+	struct TDelegateWrapperParms;
+
+	template <typename RetType, typename... ArgTypes, uint32... Indices>
+	struct TDelegateWrapperParms<RetType(ArgTypes...), TIntegerSequence<uint32, Indices...>> : TDelegateWrapperReturnBase<RetType>, TDelegateWrapperParmBase<ArgTypes, Indices>...
+	{
+		// Merge all parameters' MoveOut overloads into a common overload set.
+		using TDelegateWrapperParmBase<ArgTypes, Indices>::MoveOut...;
+
+		template <typename... ConstructorArgTypes>
+		[[nodiscard]] UE_REWRITE TDelegateWrapperParms(ConstructorArgTypes&&... Args)
+			: TDelegateWrapperParmBase<ArgTypes, Indices>{ Forward<ConstructorArgTypes>(Args) }...
+		{
+		}
+	};
+
+	template <typename... ArgTypes, uint32... Indices>
+	struct TDelegateWrapperParms<void(ArgTypes...), TIntegerSequence<uint32, Indices...>> : TDelegateWrapperParmBase<ArgTypes, Indices>...
+	{
+		// Merge all parameters' MoveOut overloads into a common overload set.
+		using TDelegateWrapperParmBase<ArgTypes, Indices>::MoveOut...;
+
+		template <typename... ConstructorArgTypes>
+		[[nodiscard]] UE_REWRITE TDelegateWrapperParms(ConstructorArgTypes&&... Args)
+			: TDelegateWrapperParmBase<ArgTypes, Indices>{ Forward<ConstructorArgTypes>(Args) }...
+		{
+		}
+	};
+
+	template <typename T>
+	struct TDelegateCallTraits
+	{
+		using Type = std::conditional_t<std::is_reference_v<T>, T, const T&>;
+	};
+
+	template <typename T>
+		requires UE::CScalar<std::remove_reference_t<T>>
+	struct TDelegateCallTraits<T>
+	{
+		using Type = std::remove_cv_t<std::remove_reference_t<T>>;
+	};
+
+	template <typename T>
+	using TDelegateCallTraits_T = typename TDelegateCallTraits<T>::Type;
+
+	template <typename RetType, typename... ParamTypes>
+	struct TDelegateFunctionWrapper
+	{
+	private:
+		template <typename DelegateType, uint32... Indices>
+		static RetType CallWrapperImpl(const DelegateType& Delegate, TIntegerSequence<uint32, Indices...>, ParamTypes... Args)
+		{
+			using TupleType = TDelegateWrapperParms<RetType(ParamTypes...), TIntegerSequence<uint32, Indices...>>;
+
+			TupleType Parms((ParamTypes&&)Args...);
+
+			// TMulticastScriptDelegate::ProcessMulticastDelegate was renamed to ProcessDelegate to allow this template to be
+			// be used for single cast and multicasts without needing a special construct to select the correct name.
+			Delegate.template ProcessDelegate<UObject>(&Parms);
+
+			// Copy out params only from the struct to the arguments.
+			(Parms.template MoveOut<Indices>(Args), ...);
+
+			// Copy return value from the struct if there is one.
+			if constexpr (!std::is_void_v<RetType>)
+			{
+				// Make sure uint32 -> bool is supported
+				if constexpr (std::is_same_v<const RetType, const bool>)
+				{
+					return !!Parms.ReturnValue;
+				}
+				else
+				{
+					return MoveTemp(Parms.ReturnValue);
+				}
+			}
+		}
+
+	public:
+		template <typename ScriptDelegateType>
+		UE_REWRITE static RetType CallWrapper(const ScriptDelegateType& Delegate, ParamTypes... Args)
+		{
+			return CallWrapperImpl(Delegate, TMakeIntegerSequence<uint32, sizeof...(ParamTypes)>{}, Args...);
+		}
+	};
 }
 
 /**
@@ -748,7 +888,7 @@ public:
 	 * @param	Params				Parameter structure
 	 */
 	template <class UObjectTemplate>
-	void ProcessMulticastDelegate(void* Parameters) const
+	void ProcessDelegate(void* Parameters) const
 	{
 		{
 			FReadAccessScope ReadScope = const_cast<TMulticastScriptDelegate*>(this)->GetReadAccessScope();
@@ -780,6 +920,13 @@ public:
 		}
 	}
 
+	template <class UObjectTemplate>
+	UE_DEPRECATED(5.8, "ProcessMulticastDelegate has been deprecated, please use ProcessDelegate instead.")
+	UE_REWRITE void ProcessMulticastDelegate(void* Parameters) const
+	{
+		ProcessDelegate(Parameters);
+	}
+
 	/**
 	 * Returns all objects associated with this multicast-delegate.
 	 * For advanced uses only -- you should never need call this function in normal circumstances.
diff --git a/Engine/Source/Runtime/CoreUObject/Private/UObject/ScriptCore.cpp b/Engine/Source/Runtime/CoreUObject/Private/UObject/ScriptCore.cpp
--- a/Engine/Source/Runtime/CoreUObject/Private/UObject/ScriptCore.cpp
+++ b/Engine/Source/Runtime/CoreUObject/Private/UObject/ScriptCore.cpp
@@ -3326,7 +3326,7 @@ public:
 		//Process delegate
 		if (DelegateAddr)
 		{
-			DelegateAddr->ProcessMulticastDelegate<UObject>(Parameters);
+			DelegateAddr->ProcessDelegate<UObject>(Parameters);
 		}
 		
 		//Clean parameters
diff --git a/Engine/Source/Runtime/CoreUObject/Public/UObject/SparseDelegate.h b/Engine/Source/Runtime/CoreUObject/Public/UObject/SparseDelegate.h
--- a/Engine/Source/Runtime/CoreUObject/Public/UObject/SparseDelegate.h
+++ b/Engine/Source/Runtime/CoreUObject/Public/UObject/SparseDelegate.h
@@ -428,8 +428,8 @@ private:
 	void __Internal_Clear(const UObject*, FName);
 };
 
-#define FUNC_DECLARE_DYNAMIC_MULTICAST_SPARSE_DELEGATE(SparseDelegateClassName, OwningClass, DelegateName, FuncParamList, FuncParamPassThru, ...) \
-	FUNC_DECLARE_DYNAMIC_MULTICAST_DELEGATE(SparseDelegateClassName##_MCSignature, SparseDelegateClassName##_DelegateWrapper, FUNC_CONCAT(FuncParamList), FUNC_CONCAT(FuncParamPassThru), __VA_ARGS__) \
+#define FUNC_DECLARE_DYNAMIC_MULTICAST_SPARSE_DELEGATE(SparseDelegateClassName, OwningClass, DelegateName, ...) \
+	FUNC_DECLARE_DYNAMIC_MULTICAST_DELEGATE(SparseDelegateClassName##_MCSignature, __VA_ARGS__) \
 	struct SparseDelegateClassName##InfoGetter \
 	{ \
 		static const char* GetDelegateName() { return #DelegateName; } \
@@ -441,14 +441,14 @@ private:
 	};
 
 /** Declares a sparse blueprint-accessible broadcast delegate that can bind to multiple native UFUNCTIONs simultaneously */
-#define DECLARE_DYNAMIC_MULTICAST_SPARSE_DELEGATE( SparseDelegateClass, OwningClass, DelegateName ) BODY_MACRO_COMBINE(CURRENT_FILE_ID,_,__LINE__,_DELEGATE) FUNC_DECLARE_DYNAMIC_MULTICAST_SPARSE_DELEGATE( SparseDelegateClass, OwningClass, DelegateName, , FUNC_CONCAT( *this ), void )
-
-#define DECLARE_DYNAMIC_MULTICAST_SPARSE_DELEGATE_OneParam( SparseDelegateClass, OwningClass, DelegateName, Param1Type, Param1Name ) BODY_MACRO_COMBINE(CURRENT_FILE_ID,_,__LINE__,_DELEGATE) FUNC_DECLARE_DYNAMIC_MULTICAST_SPARSE_DELEGATE( SparseDelegateClass, OwningClass, DelegateName, FUNC_CONCAT( Param1Type InParam1 ), FUNC_CONCAT( *this, InParam1 ), void, Param1Type )
-#define DECLARE_DYNAMIC_MULTICAST_SPARSE_DELEGATE_TwoParams( SparseDelegateClass, OwningClass, DelegateName, Param1Type, Param1Name, Param2Type, Param2Name ) BODY_MACRO_COMBINE(CURRENT_FILE_ID,_,__LINE__,_DELEGATE) FUNC_DECLARE_DYNAMIC_MULTICAST_SPARSE_DELEGATE( SparseDelegateClass, OwningClass, DelegateName, FUNC_CONCAT( Param1Type InParam1, Param2Type InParam2 ), FUNC_CONCAT( *this, InParam1, InParam2 ), void, Param1Type, Param2Type )
-#define DECLARE_DYNAMIC_MULTICAST_SPARSE_DELEGATE_ThreeParams( SparseDelegateClass, OwningClass, DelegateName, Param1Type, Param1Name, Param2Type, Param2Name, Param3Type, Param3Name ) BODY_MACRO_COMBINE(CURRENT_FILE_ID,_,__LINE__,_DELEGATE) FUNC_DECLARE_DYNAMIC_MULTICAST_SPARSE_DELEGATE( SparseDelegateClass, OwningClass, DelegateName, FUNC_CONCAT( Param1Type InParam1, Param2Type InParam2, Param3Type InParam3 ), FUNC_CONCAT( *this, InParam1, InParam2, InParam3 ), void, Param1Type, Param2Type, Param3Type )
-#define DECLARE_DYNAMIC_MULTICAST_SPARSE_DELEGATE_FourParams( SparseDelegateClass, OwningClass, DelegateName, Param1Type, Param1Name, Param2Type, Param2Name, Param3Type, Param3Name, Param4Type, Param4Name ) BODY_MACRO_COMBINE(CURRENT_FILE_ID,_,__LINE__,_DELEGATE) FUNC_DECLARE_DYNAMIC_MULTICAST_SPARSE_DELEGATE( SparseDelegateClass, OwningClass, DelegateName, FUNC_CONCAT( Param1Type InParam1, Param2Type InParam2, Param3Type InParam3, Param4Type InParam4 ), FUNC_CONCAT( *this, InParam1, InParam2, InParam3, InParam4 ), void, Param1Type, Param2Type, Param3Type, Param4Type )
-#define DECLARE_DYNAMIC_MULTICAST_SPARSE_DELEGATE_FiveParams( SparseDelegateClass, OwningClass, DelegateName, Param1Type, Param1Name, Param2Type, Param2Name, Param3Type, Param3Name, Param4Type, Param4Name, Param5Type, Param5Name ) BODY_MACRO_COMBINE(CURRENT_FILE_ID,_,__LINE__,_DELEGATE) FUNC_DECLARE_DYNAMIC_MULTICAST_SPARSE_DELEGATE( SparseDelegateClass, OwningClass, DelegateName, FUNC_CONCAT( Param1Type InParam1, Param2Type InParam2, Param3Type InParam3, Param4Type InParam4, Param5Type InParam5 ), FUNC_CONCAT( *this, InParam1, InParam2, InParam3, InParam4, InParam5 ), void, Param1Type, Param2Type, Param3Type, Param4Type, Param5Type )
-#define DECLARE_DYNAMIC_MULTICAST_SPARSE_DELEGATE_SixParams( SparseDelegateClass, OwningClass, DelegateName, Param1Type, Param1Name, Param2Type, Param2Name, Param3Type, Param3Name, Param4Type, Param4Name, Param5Type, Param5Name, Param6Type, Param6Name ) BODY_MACRO_COMBINE(CURRENT_FILE_ID,_,__LINE__,_DELEGATE) FUNC_DECLARE_DYNAMIC_MULTICAST_SPARSE_DELEGATE( SparseDelegateClass, OwningClass, DelegateName, FUNC_CONCAT( Param1Type InParam1, Param2Type InParam2, Param3Type InParam3, Param4Type InParam4, Param5Type InParam5, Param6Type InParam6 ), FUNC_CONCAT( *this, InParam1, InParam2, InParam3, InParam4, InParam5, InParam6 ), void, Param1Type, Param2Type, Param3Type, Param4Type, Param5Type, Param6Type )
-#define DECLARE_DYNAMIC_MULTICAST_SPARSE_DELEGATE_SevenParams( SparseDelegateClass, OwningClass, DelegateName, Param1Type, Param1Name, Param2Type, Param2Name, Param3Type, Param3Name, Param4Type, Param4Name, Param5Type, Param5Name, Param6Type, Param6Name, Param7Type, Param7Name ) BODY_MACRO_COMBINE(CURRENT_FILE_ID,_,__LINE__,_DELEGATE) FUNC_DECLARE_DYNAMIC_MULTICAST_SPARSE_DELEGATE( SparseDelegateClass, OwningClass, DelegateName, FUNC_CONCAT( Param1Type InParam1, Param2Type InParam2, Param3Type InParam3, Param4Type InParam4, Param5Type InParam5, Param6Type InParam6, Param7Type InParam7 ), FUNC_CONCAT( *this, InParam1, InParam2, InParam3, InParam4, InParam5, InParam6, InParam7 ), void, Param1Type, Param2Type, Param3Type, Param4Type, Param5Type, Param6Type, Param7Type )
-#define DECLARE_DYNAMIC_MULTICAST_SPARSE_DELEGATE_EightParams( SparseDelegateClass, OwningClass, DelegateName, Param1Type, Param1Name, Param2Type, Param2Name, Param3Type, Param3Name, Param4Type, Param4Name, Param5Type, Param5Name, Param6Type, Param6Name, Param7Type, Param7Name, Param8Type, Param8Name ) BODY_MACRO_COMBINE(CURRENT_FILE_ID,_,__LINE__,_DELEGATE) FUNC_DECLARE_DYNAMIC_MULTICAST_SPARSE_DELEGATE( SparseDelegateClass, OwningClass, DelegateName, FUNC_CONCAT( Param1Type InParam1, Param2Type InParam2, Param3Type InParam3, Param4Type InParam4, Param5Type InParam5, Param6Type InParam6, Param7Type InParam7, Param8Type InParam8 ), FUNC_CONCAT( *this, InParam1, InParam2, InParam3, InParam4, InParam5, InParam6, InParam7, InParam8 ), void, Param1Type, Param2Type, Param3Type, Param4Type, Param5Type, Param6Type, Param7Type, Param8Type )
-#define DECLARE_DYNAMIC_MULTICAST_SPARSE_DELEGATE_NineParams( SparseDelegateClass, OwningClass, DelegateName, Param1Type, Param1Name, Param2Type, Param2Name, Param3Type, Param3Name, Param4Type, Param4Name, Param5Type, Param5Name, Param6Type, Param6Name, Param7Type, Param7Name, Param8Type, Param8Name, Param9Type, Param9Name ) BODY_MACRO_COMBINE(CURRENT_FILE_ID,_,__LINE__,_DELEGATE) FUNC_DECLARE_DYNAMIC_MULTICAST_SPARSE_DELEGATE( SparseDelegateClass, OwningClass, DelegateName, FUNC_CONCAT( Param1Type InParam1, Param2Type InParam2, Param3Type InParam3, Param4Type InParam4, Param5Type InParam5, Param6Type InParam6, Param7Type InParam7, Param8Type InParam8, Param9Type InParam9 ), FUNC_CONCAT( *this, InParam1, InParam2, InParam3, InParam4, InParam5, InParam6, InParam7, InParam8, InParam9 ), void, Param1Type, Param2Type, Param3Type, Param4Type, Param5Type, Param6Type, Param7Type, Param8Type, Param9Type )
+#define DECLARE_DYNAMIC_MULTICAST_SPARSE_DELEGATE( SparseDelegateClass, OwningClass, DelegateName ) FUNC_DECLARE_DYNAMIC_MULTICAST_SPARSE_DELEGATE( SparseDelegateClass, OwningClass, DelegateName, void )
+
+#define DECLARE_DYNAMIC_MULTICAST_SPARSE_DELEGATE_OneParam( SparseDelegateClass, OwningClass, DelegateName, Param1Type, Param1Name ) FUNC_DECLARE_DYNAMIC_MULTICAST_SPARSE_DELEGATE( SparseDelegateClass, OwningClass, DelegateName, void, Param1Type )
+#define DECLARE_DYNAMIC_MULTICAST_SPARSE_DELEGATE_TwoParams( SparseDelegateClass, OwningClass, DelegateName, Param1Type, Param1Name, Param2Type, Param2Name ) FUNC_DECLARE_DYNAMIC_MULTICAST_SPARSE_DELEGATE( SparseDelegateClass, OwningClass, DelegateName, void, Param1Type, Param2Type )
+#define DECLARE_DYNAMIC_MULTICAST_SPARSE_DELEGATE_ThreeParams( SparseDelegateClass, OwningClass, DelegateName, Param1Type, Param1Name, Param2Type, Param2Name, Param3Type, Param3Name ) FUNC_DECLARE_DYNAMIC_MULTICAST_SPARSE_DELEGATE( SparseDelegateClass, OwningClass, DelegateName, void, Param1Type, Param2Type, Param3Type )
+#define DECLARE_DYNAMIC_MULTICAST_SPARSE_DELEGATE_FourParams( SparseDelegateClass, OwningClass, DelegateName, Param1Type, Param1Name, Param2Type, Param2Name, Param3Type, Param3Name, Param4Type, Param4Name ) FUNC_DECLARE_DYNAMIC_MULTICAST_SPARSE_DELEGATE( SparseDelegateClass, OwningClass, DelegateName, void, Param1Type, Param2Type, Param3Type, Param4Type )
+#define DECLARE_DYNAMIC_MULTICAST_SPARSE_DELEGATE_FiveParams( SparseDelegateClass, OwningClass, DelegateName, Param1Type, Param1Name, Param2Type, Param2Name, Param3Type, Param3Name, Param4Type, Param4Name, Param5Type, Param5Name ) FUNC_DECLARE_DYNAMIC_MULTICAST_SPARSE_DELEGATE( SparseDelegateClass, OwningClass, DelegateName, void, Param1Type, Param2Type, Param3Type, Param4Type, Param5Type )
+#define DECLARE_DYNAMIC_MULTICAST_SPARSE_DELEGATE_SixParams( SparseDelegateClass, OwningClass, DelegateName, Param1Type, Param1Name, Param2Type, Param2Name, Param3Type, Param3Name, Param4Type, Param4Name, Param5Type, Param5Name, Param6Type, Param6Name ) FUNC_DECLARE_DYNAMIC_MULTICAST_SPARSE_DELEGATE( SparseDelegateClass, OwningClass, DelegateName, void, Param1Type, Param2Type, Param3Type, Param4Type, Param5Type, Param6Type )
+#define DECLARE_DYNAMIC_MULTICAST_SPARSE_DELEGATE_SevenParams( SparseDelegateClass, OwningClass, DelegateName, Param1Type, Param1Name, Param2Type, Param2Name, Param3Type, Param3Name, Param4Type, Param4Name, Param5Type, Param5Name, Param6Type, Param6Name, Param7Type, Param7Name ) FUNC_DECLARE_DYNAMIC_MULTICAST_SPARSE_DELEGATE( SparseDelegateClass, OwningClass, DelegateName, void, Param1Type, Param2Type, Param3Type, Param4Type, Param5Type, Param6Type, Param7Type )
+#define DECLARE_DYNAMIC_MULTICAST_SPARSE_DELEGATE_EightParams( SparseDelegateClass, OwningClass, DelegateName, Param1Type, Param1Name, Param2Type, Param2Name, Param3Type, Param3Name, Param4Type, Param4Name, Param5Type, Param5Name, Param6Type, Param6Name, Param7Type, Param7Name, Param8Type, Param8Name ) FUNC_DECLARE_DYNAMIC_MULTICAST_SPARSE_DELEGATE( SparseDelegateClass, OwningClass, DelegateName, void, Param1Type, Param2Type, Param3Type, Param4Type, Param5Type, Param6Type, Param7Type, Param8Type )
+#define DECLARE_DYNAMIC_MULTICAST_SPARSE_DELEGATE_NineParams( SparseDelegateClass, OwningClass, DelegateName, Param1Type, Param1Name, Param2Type, Param2Name, Param3Type, Param3Name, Param4Type, Param4Name, Param5Type, Param5Name, Param6Type, Param6Name, Param7Type, Param7Name, Param8Type, Param8Name, Param9Type, Param9Name ) FUNC_DECLARE_DYNAMIC_MULTICAST_SPARSE_DELEGATE( SparseDelegateClass, OwningClass, DelegateName, void, Param1Type, Param2Type, Param3Type, Param4Type, Param5Type, Param6Type, Param7Type, Param8Type, Param9Type )
