IOS: selectively back out CL 49206596, which enabled os_log instead of NSLog in non-shipping builds. Os_log is still proving very flaky and breaking automation

#Synced-CL 49549262

diff --git a/Engine/Source/Runtime/Core/Private/Apple/ApplePlatformMisc.cpp b/Engine/Source/Runtime/Core/Private/Apple/ApplePlatformMisc.cpp
--- a/Engine/Source/Runtime/Core/Private/Apple/ApplePlatformMisc.cpp
+++ b/Engine/Source/Runtime/Core/Private/Apple/ApplePlatformMisc.cpp
@@ -77,25 +77,13 @@ FString FApplePlatformMisc::GetEnvironmentVariable(const TCHAR* VariableName)
 
 void FApplePlatformMisc::LocalPrint(const TCHAR* Message)
 {
-	//
-	// os_log doesn't go to stdout/stderr, by default, as NSLog does.
-	//
-	// It will be there if requested via one of environment variables:
-	//   OS_ACTIVITY_DT_MODE=enable
-	//   ACTIVITY_LOG_STDERR=1 (will also include frameworks)
-	//
-	// E.g. if app is run via devicectl:
-	//
-	//   xcrun devicectl device process launch \
-	//     --device <device> \
-	//     --console \
-	//     --environment-variables '{"OS_ACTIVITY_DT_MODE":"enable"}' \
-	//     <bundle-name>
-	//
 #if !UE_BUILD_SHIPPING
-	os_log(OS_LOG_DEFAULT, "[UE] %{public}s", TCHAR_TO_UTF8(Message));
+	// This goes to all output consoles, not only system one, stdout included.
+	NSLog(@"[UE] %s", TCHAR_TO_UTF8(Message));
 #else
-	os_log(OS_LOG_DEFAULT, "[UE] %s", TCHAR_TO_UTF8(Message)); // Private by default, would be visible in Xcode.
+	// This only goes to the system console.
+	// Private by default, would be visible in Xcode.
+	os_log(OS_LOG_DEFAULT, "[UE] %s", TCHAR_TO_UTF8(Message));
 #endif
 }
 
