Build fix: add back UE::Private::MurmurFinalize32 removed in 49601163 because it is still used in code (should be removed).

#Synced-CL 49598772

diff --git a/Engine/Source/Runtime/Core/Public/Templates/TypeHash.h b/Engine/Source/Runtime/Core/Public/Templates/TypeHash.h
--- a/Engine/Source/Runtime/Core/Public/Templates/TypeHash.h
+++ b/Engine/Source/Runtime/Core/Public/Templates/TypeHash.h
@@ -14,6 +14,16 @@ namespace UE
 {
 	namespace Private
 	{
+		[[nodiscard]] inline uint32 MurmurFinalize32(uint32 Hash)
+		{
+			Hash ^= Hash >> 16;
+			Hash *= 0x85ebca6b;
+			Hash ^= Hash >> 13;
+			Hash *= 0xc2b2ae35;
+			Hash ^= Hash >> 16;
+			return Hash;
+		}
+
 		[[nodiscard]] inline uint64 MurmurFinalize64(uint64 Hash)
 		{
 			Hash ^= Hash >> 33;
