Added TMemberPtrOuter trait for retrieving the struct/class type from a pointer-to-member type.
Deprecated TMemberFunctionPtrOuter which was doing the same thing, but only for pointer-to-member-function types.
Fixed a comment in RemoveMemberPtr.h.

#Synced-CL 49589526

diff --git a/Engine/Source/Runtime/Core/Public/Templates/Invoke.h b/Engine/Source/Runtime/Core/Public/Templates/Invoke.h
--- a/Engine/Source/Runtime/Core/Public/Templates/Invoke.h
+++ b/Engine/Source/Runtime/Core/Public/Templates/Invoke.h
@@ -4,7 +4,7 @@
 
 #include "CoreTypes.h"
 #include "AutoRTFM.h"
-#include "Traits/MemberFunctionPtrOuter.h"
+#include "Traits/MemberPtrOuter.h"
 #include <type_traits>
 
 
@@ -58,7 +58,7 @@ template <
 	typename    PtrMemFunType,
 	typename    TargetType,
 	typename... ArgTypes,
-	typename    ObjType = TMemberFunctionPtrOuter_T<PtrMemFunType>
+	typename    ObjType = TMemberPtrOuter_T<PtrMemFunType>
 >
 AUTORTFM_INFER constexpr auto Invoke(PtrMemFunType PtrMemFun, TargetType&& Target, ArgTypes&&... Args)
 	-> decltype((UE::Core::Private::DereferenceIfNecessary<ObjType>((TargetType&&)Target, &Target).*PtrMemFun)((ArgTypes&&)Args...))
@@ -137,3 +137,7 @@ using TInvokeResult_T = typename TInvokeResult<FuncType, ArgTypes...>::Type;
 #if UE_ENABLE_INCLUDE_ORDER_DEPRECATED_IN_5_7
 #include "Templates/UnrealTemplate.h"
 #endif
+
+#if UE_ENABLE_INCLUDE_ORDER_DEPRECATED_IN_5_8
+#include "Traits/MemberFunctionPtrOuter.h"
+#endif
diff --git a/Engine/Source/Runtime/Core/Public/Traits/MemberFunctionPtrOuter.h b/Engine/Source/Runtime/Core/Public/Traits/MemberFunctionPtrOuter.h
--- a/Engine/Source/Runtime/Core/Public/Traits/MemberFunctionPtrOuter.h
+++ b/Engine/Source/Runtime/Core/Public/Traits/MemberFunctionPtrOuter.h
@@ -2,26 +2,32 @@
 
 #pragma once
 
+#include "Misc/CoreMiscDefines.h"
+
+UE_DEPRECATED_HEADER(5.8, "MemberFunctionPtrOuter.h has been deprecated.")
+
 /**
  * Type trait which yields the type of the class given a pointer to a member function of that class, e.g.:
  *
  * TMemberFunctionPtrOuter_T<decltype(&FVector::Dot)> == FVector
  */
 template <typename T>
-struct TMemberFunctionPtrOuter;
+struct UE_DEPRECATED(5.8, "TMemberFunctionPtrOuter has been deprecated, please use TMemberPtrOuter instead.") TMemberFunctionPtrOuter;
 
-template <typename ReturnType, typename ObjectType, typename... ArgTypes> struct TMemberFunctionPtrOuter<ReturnType (ObjectType::*)(ArgTypes...)                 > { using Type = ObjectType; };
-template <typename ReturnType, typename ObjectType, typename... ArgTypes> struct TMemberFunctionPtrOuter<ReturnType (ObjectType::*)(ArgTypes...)               & > { using Type = ObjectType; };
-template <typename ReturnType, typename ObjectType, typename... ArgTypes> struct TMemberFunctionPtrOuter<ReturnType (ObjectType::*)(ArgTypes...)               &&> { using Type = ObjectType; };
-template <typename ReturnType, typename ObjectType, typename... ArgTypes> struct TMemberFunctionPtrOuter<ReturnType (ObjectType::*)(ArgTypes...) const           > { using Type = ObjectType; };
-template <typename ReturnType, typename ObjectType, typename... ArgTypes> struct TMemberFunctionPtrOuter<ReturnType (ObjectType::*)(ArgTypes...) const         & > { using Type = ObjectType; };
-template <typename ReturnType, typename ObjectType, typename... ArgTypes> struct TMemberFunctionPtrOuter<ReturnType (ObjectType::*)(ArgTypes...) const         &&> { using Type = ObjectType; };
-template <typename ReturnType, typename ObjectType, typename... ArgTypes> struct TMemberFunctionPtrOuter<ReturnType (ObjectType::*)(ArgTypes...)       volatile  > { using Type = ObjectType; };
-template <typename ReturnType, typename ObjectType, typename... ArgTypes> struct TMemberFunctionPtrOuter<ReturnType (ObjectType::*)(ArgTypes...)       volatile& > { using Type = ObjectType; };
-template <typename ReturnType, typename ObjectType, typename... ArgTypes> struct TMemberFunctionPtrOuter<ReturnType (ObjectType::*)(ArgTypes...)       volatile&&> { using Type = ObjectType; };
-template <typename ReturnType, typename ObjectType, typename... ArgTypes> struct TMemberFunctionPtrOuter<ReturnType (ObjectType::*)(ArgTypes...) const volatile  > { using Type = ObjectType; };
-template <typename ReturnType, typename ObjectType, typename... ArgTypes> struct TMemberFunctionPtrOuter<ReturnType (ObjectType::*)(ArgTypes...) const volatile& > { using Type = ObjectType; };
-template <typename ReturnType, typename ObjectType, typename... ArgTypes> struct TMemberFunctionPtrOuter<ReturnType (ObjectType::*)(ArgTypes...) const volatile&&> { using Type = ObjectType; };
+template <typename ReturnType, typename ObjectType, typename... ArgTypes> struct UE_DEPRECATED(5.8, "TMemberFunctionPtrOuter has been deprecated, use TMemberPtrOuter instead.") TMemberFunctionPtrOuter<ReturnType (ObjectType::*)(ArgTypes...)                 > { using Type = ObjectType; };
+template <typename ReturnType, typename ObjectType, typename... ArgTypes> struct UE_DEPRECATED(5.8, "TMemberFunctionPtrOuter has been deprecated, use TMemberPtrOuter instead.") TMemberFunctionPtrOuter<ReturnType (ObjectType::*)(ArgTypes...)               & > { using Type = ObjectType; };
+template <typename ReturnType, typename ObjectType, typename... ArgTypes> struct UE_DEPRECATED(5.8, "TMemberFunctionPtrOuter has been deprecated, use TMemberPtrOuter instead.") TMemberFunctionPtrOuter<ReturnType (ObjectType::*)(ArgTypes...)               &&> { using Type = ObjectType; };
+template <typename ReturnType, typename ObjectType, typename... ArgTypes> struct UE_DEPRECATED(5.8, "TMemberFunctionPtrOuter has been deprecated, use TMemberPtrOuter instead.") TMemberFunctionPtrOuter<ReturnType (ObjectType::*)(ArgTypes...) const           > { using Type = ObjectType; };
+template <typename ReturnType, typename ObjectType, typename... ArgTypes> struct UE_DEPRECATED(5.8, "TMemberFunctionPtrOuter has been deprecated, use TMemberPtrOuter instead.") TMemberFunctionPtrOuter<ReturnType (ObjectType::*)(ArgTypes...) const         & > { using Type = ObjectType; };
+template <typename ReturnType, typename ObjectType, typename... ArgTypes> struct UE_DEPRECATED(5.8, "TMemberFunctionPtrOuter has been deprecated, use TMemberPtrOuter instead.") TMemberFunctionPtrOuter<ReturnType (ObjectType::*)(ArgTypes...) const         &&> { using Type = ObjectType; };
+template <typename ReturnType, typename ObjectType, typename... ArgTypes> struct UE_DEPRECATED(5.8, "TMemberFunctionPtrOuter has been deprecated, use TMemberPtrOuter instead.") TMemberFunctionPtrOuter<ReturnType (ObjectType::*)(ArgTypes...)       volatile  > { using Type = ObjectType; };
+template <typename ReturnType, typename ObjectType, typename... ArgTypes> struct UE_DEPRECATED(5.8, "TMemberFunctionPtrOuter has been deprecated, use TMemberPtrOuter instead.") TMemberFunctionPtrOuter<ReturnType (ObjectType::*)(ArgTypes...)       volatile& > { using Type = ObjectType; };
+template <typename ReturnType, typename ObjectType, typename... ArgTypes> struct UE_DEPRECATED(5.8, "TMemberFunctionPtrOuter has been deprecated, use TMemberPtrOuter instead.") TMemberFunctionPtrOuter<ReturnType (ObjectType::*)(ArgTypes...)       volatile&&> { using Type = ObjectType; };
+template <typename ReturnType, typename ObjectType, typename... ArgTypes> struct UE_DEPRECATED(5.8, "TMemberFunctionPtrOuter has been deprecated, use TMemberPtrOuter instead.") TMemberFunctionPtrOuter<ReturnType (ObjectType::*)(ArgTypes...) const volatile  > { using Type = ObjectType; };
+template <typename ReturnType, typename ObjectType, typename... ArgTypes> struct UE_DEPRECATED(5.8, "TMemberFunctionPtrOuter has been deprecated, use TMemberPtrOuter instead.") TMemberFunctionPtrOuter<ReturnType (ObjectType::*)(ArgTypes...) const volatile& > { using Type = ObjectType; };
+template <typename ReturnType, typename ObjectType, typename... ArgTypes> struct UE_DEPRECATED(5.8, "TMemberFunctionPtrOuter has been deprecated, use TMemberPtrOuter instead.") TMemberFunctionPtrOuter<ReturnType (ObjectType::*)(ArgTypes...) const volatile&&> { using Type = ObjectType; };
 
+PRAGMA_DISABLE_DEPRECATION_WARNINGS
 template <typename T>
-using TMemberFunctionPtrOuter_T = typename TMemberFunctionPtrOuter<T>::Type;
+using TMemberFunctionPtrOuter_T UE_DEPRECATED(5.8, "TMemberFunctionPtrOuter_T has been deprecated, use TMemberPtrOuter_T instead.") = typename TMemberFunctionPtrOuter<T>::Type;
+PRAGMA_ENABLE_DEPRECATION_WARNINGS
diff --git a/Engine/Source/Runtime/Core/Public/Traits/MemberPtrOuter.h b/Engine/Source/Runtime/Core/Public/Traits/MemberPtrOuter.h
new file mode 100644
--- /dev/null
+++ b/Engine/Source/Runtime/Core/Public/Traits/MemberPtrOuter.h
@@ -0,0 +1,23 @@
+// Copyright Epic Games, Inc. All Rights Reserved.
+
+#pragma once
+
+/**
+ * Type trait which yields the type of the class given a pointer to a member of that class, e.g.:
+ *
+ * TMemberPtrOuter_T<decltype(&FVector::X)>   == FVector
+ * TMemberPtrOuter_T<decltype(&FVector::Dot)> == FVector
+ */
+template <typename T>
+struct TMemberPtrOuter
+{
+};
+
+template <typename OuterType, typename MemberType>
+struct TMemberPtrOuter<MemberType OuterType::*>
+{
+	using Type = OuterType;
+};
+
+template <typename T>
+using TMemberPtrOuter_T = typename TMemberPtrOuter<T>::Type;
diff --git a/Engine/Source/Runtime/Core/Public/Traits/RemoveMemberPtr.h b/Engine/Source/Runtime/Core/Public/Traits/RemoveMemberPtr.h
--- a/Engine/Source/Runtime/Core/Public/Traits/RemoveMemberPtr.h
+++ b/Engine/Source/Runtime/Core/Public/Traits/RemoveMemberPtr.h
@@ -5,8 +5,8 @@
 /**
  * Type trait which yields the type being pointed to of a member pointer type, e.g.:
  *
- * TRemoveMemberFunctionPtr_T<decltype(&FVector::X)>   == double
- * TRemoveMemberFunctionPtr_T<decltype(&FVector::Dot)> == double(const FVector&)
+ * TRemoveMemberPtr_T<decltype(&FVector::X)>   == double
+ * TRemoveMemberPtr_T<decltype(&FVector::Dot)> == double(const FVector&)
  */
 template <typename T>
 struct TRemoveMemberPtr
