LLM: Move LowLevelMemTracker's private files into a subdirectory for tidyness.
#rnx

#Synced-CL 49592695

diff --git a/Engine/Source/Runtime/Core/Private/HAL/LLM/LLM.cpp b/Engine/Source/Runtime/Core/Private/HAL/LLM/LLM.cpp
new file mode 100644
--- /dev/null
+++ b/Engine/Source/Runtime/Core/Private/HAL/LLM/LLM.cpp
@@ -0,0 +1,6834 @@
+// Copyright Epic Games, Inc. All Rights Reserved.
+
+#include "HAL/LowLevelMemTracker.h"
+#include "Containers/ContainerAllocationPolicies.h"
+
+#if ENABLE_LOW_LEVEL_MEM_TRACKER
+
+#include "Async/SharedLock.h"
+#include "Async/UniqueLock.h"
+#include "HAL/LowLevelMemStats.h"
+#include "HAL/FileManager.h"
+#include "HAL/IConsoleManager.h"
+#include "HAL/LLM/LLMPrivate.h"
+#include "HAL/ThreadHeartBeat.h"
+#include "HAL/PlatformMemory.h" // for page allocation association.
+#include "MemPro/MemProProfiler.h"
+#include "Math/NumericLimits.h"
+#include "Misc/CString.h"
+#include "Misc/CoreDelegates.h"
+#include "Misc/FileHelper.h"
+#include "Misc/Fork.h"
+#include "Misc/OutputDeviceRedirector.h"
+#include "Misc/Paths.h"
+#include "Misc/ScopeExit.h"
+#include "Misc/ScopeLock.h"
+#include "Misc/ScopeRWLock.h"
+#include "Misc/StringBuilder.h"
+#include "Misc/VarArgs.h"
+#include "ProfilingDebugging/CsvProfiler.h"
+#include "Templates/Atomic.h"
+#include "Trace/Trace.inl"
+
+static_assert(static_cast<uint32>(ELLMTagSet::Max) == 4, "You added a new TagSet, please update the #ifs bellow");
+#if LLM_ALLOW_ASSETS_TAGS
+#define LLM_TAGSET_MAX 3
+#else
+#define LLM_TAGSET_MAX 2
+#endif
+
+// Enable to allow memory tracking to be passed through to external tools without the extra overhead of LLM also tracking allocations.
+#define UE_ONLY_USE_PLATFORM_TRACKER 0
+
+#if UE_ENABLE_ARRAY_SLACK_TRACKING
+
+// Specifies whether to generate the whole log file in memory before writing.  Switch uses an async thread for file writing,
+// which does array allocations, deadlocking on the critical section.  Writing to memory first avoids a dependency on the
+// async thread, as we can write the memory to a file at the end, after the lock is released.
+#define ARRAY_SLACK_LOG_TO_MEMORY !PLATFORM_WINDOWS
+
+// Useful values to set in the debugger, to set a breakpoint on a particular allocation tag, element size, and max count.  Sometimes this can
+// give more context regarding an allocation than what you get from a stack trace alone.  As an example, the constructor for UAudioCaptureComponent,
+// which allocated a 1.83 MB array, just showed up as "UClass::CreateDefaultObject" in the stack trace.  The constructor for the specific subclass
+// was optimized out, and there was no way to tell from the slack report what it was actually related to.  Stopping on the allocation in the debugger
+// makes it immediately obvious, because you can see that the UClass in question is UAudioCaptureComponent.  These debug values also let you stop
+// on places where the array count changes without triggering an allocation (the tracking only grabs call stacks on reallocation).
+int32 GArraySlackTagToWatch = -1;		// (int32)ELLMTag::UObject;
+int32 GArraySlackSizeToWatch = 0;
+int32 GArraySlackMaxToWatch = 0;
+
+uint32 GArraySlackDumpIndex = 0;				// Incremented for each auto-generated report file
+bool GArraySlackInit = false;					// Set when we start tracking slack -- startup constructor allocations add a lot of noise (start this as true if you want these)
+bool GArraySlackFirstStackOnly = true;			// Only do stack trace on first stack for a given allocation -- faster, but could be useful to know last allocation
+bool GArraySlackGroupByTag = false;				// Group slack by tag when running with -llm
+bool GArraySlackDefaultVerbose = true;			// Whether to default to verbose output, can be overridden with -Verbose=[0,1]
+
+// We require a minimum number of total bytes for a group of allocations with the same stack trace to be reported.  A setting of 64 discards 80% of
+// allocations representing less than 0.2% of the slack memory.  If you want to look at aggregations of smaller allocations, you can use -Stack=N to
+// trim off more of the call stack, which will make more allocations alias to the same stack trace, and show up in the report (this is probably what
+// you would do anyway when investigating that scenario).  Or you could locally set this to zero to get everything.
+//
+// Console platforms have far slower symbol lookup than PC, making slack reports take a couple orders of magnitude longer to generate than on PC,
+// so we use more conservative settings.  Without some sort of trimming of the generated results, a slack report can take over an hour, which just
+// isn't useful (it still can take 15 minutes with these settings, versus around 20 seconds on Windows for a much larger report).
+#if PLATFORM_WINDOWS
+static const int32 GArraySlackThreshold = 64;				// Typically covers 99.8% of slack memory
+static const int32 GArraySlackFullStackNum = MAX_int32;		// Show full call stack context for all allocations
+static const int32 GArraySlackDefaultStackDepth = 9;		// Sort by this deep in the call stack -- matches max call stack depth in FArraySlackTrackingHeader structure
+#else
+static const int32 GArraySlackThreshold = 8192;				// Typically covers 95% of slack memory
+static const int32 GArraySlackFullStackNum = 150;			// Show full call stack context for this many allocations
+static const int32 GArraySlackDefaultStackDepth = 5;		// Sort by this deep in the call stack
+#endif
+
+std::atomic<int64> GArrayMaxByTag[256];
+std::atomic<int64> GArrayUsedByTag[256];
+int64 GArraySlackByTag[256];
+int32 GArrayCountByTag[256];
+
+// Doubly linked list of all tracked allocations, and critical section to protect it.  Critical section is a pointer so we can detect if it's constructed,
+// otherwise you get crashes in startup constructors which run before the constructor of the critical section is called.  Initialized in the LLM tracker,
+// which gets initialized by the first startup code that uses an LLM scope (even when -llm is disabled), which tends to be pretty early -- allocations
+// before that just won't be tracked.  If this became an issue, we could atomically initialize the lock on first access from any thread, but we already
+// defer tracking until engine PreInit anyway to factor out noise from the myriad static global FString constructors, so it doesn't matter as it stands now.
+FArraySlackTrackingHeader* GTrackArrayDetailedList;
+FCriticalSection* GTrackArrayDetailedLock;
+
+// Dummy function to set a breakpoint on where it's called, if you want to investigate code related to a certain allocation
+FORCENOINLINE void LlmTrackSetBreakpointHere()
+{
+	// Need something non-empty that won't compile out
+	static int32 Dummy = 0;
+	Dummy++;
+}
+
+void FArraySlackTrackingHeader::AddAllocation()
+{
+	if (ArrayNum != INDEX_NONE)
+	{
+		GArrayMaxByTag[Tag].fetch_add(ArrayMax * (int64)ElemSize, std::memory_order_relaxed);
+		GArrayUsedByTag[Tag].fetch_add(ArrayNum * (int64)ElemSize, std::memory_order_relaxed);
+
+		// This code is only reached for reallocations, since during the initial allocation, ArrayNum won't have been set yet.
+		ReallocCount++;
+	}
+
+	if (GArraySlackFirstStackOnly == false || NumStackFrames == 0)
+	{
+		// Skip the first 3 stack frames, which are tracking code (CaptureStackBackTrace, LlmTrackArrayAddAllocation, FArraySlackTrackingHeader::Realloc)
+		constexpr int8 SkipStackFrames = 3;
+		uint64 StackFrameTemp[UE_ARRAY_COUNT(StackFrames) + SkipStackFrames];
+		NumStackFrames = (int8)FPlatformStackWalk::CaptureStackBackTrace(StackFrameTemp, UE_ARRAY_COUNT(StackFrameTemp)) - SkipStackFrames;
+		if (NumStackFrames < 0)
+		{
+			NumStackFrames = 0;
+		}
+		for (int32 StackIndex = 0; StackIndex < NumStackFrames; StackIndex++)
+		{
+			StackFrames[StackIndex] = StackFrameTemp[SkipStackFrames + StackIndex];
+		}
+	}
+
+	if (GTrackArrayDetailedLock && GArraySlackInit)
+	{
+		// For detailed tracking, we add the array header to a doubly linked list
+		FScopeLock Lock(GTrackArrayDetailedLock);
+
+		if (GTrackArrayDetailedList)
+		{
+			GTrackArrayDetailedList->Prev = &Next;
+		}
+		Next = GTrackArrayDetailedList;
+		Prev = &GTrackArrayDetailedList;
+		GTrackArrayDetailedList = this;
+
+		if ((Tag == GArraySlackTagToWatch) &&
+			(ElemSize == GArraySlackSizeToWatch) &&
+			(!GArraySlackMaxToWatch || (ArrayMax == GArraySlackMaxToWatch)))
+		{
+			LlmTrackSetBreakpointHere();
+		}
+
+		GArrayCountByTag[Tag]++;
+	}
+}
+
+void FArraySlackTrackingHeader::RemoveAllocation()
+{
+	if (ArrayNum != INDEX_NONE)
+	{
+		GArrayUsedByTag[Tag].fetch_sub(ArrayNum * (int64)ElemSize, std::memory_order_relaxed);
+		GArrayMaxByTag[Tag].fetch_sub(ArrayMax * (int64)ElemSize, std::memory_order_release);
+	}
+
+	if (Prev)
+	{
+		FScopeLock Lock(GTrackArrayDetailedLock);
+
+		GArrayCountByTag[Tag]--;
+
+		if (Next)
+		{
+			Next->Prev = Prev;
+		}
+		(*Prev) = Next;
+
+		Next = nullptr;
+		Prev = nullptr;
+	}
+}
+
+void FArraySlackTrackingHeader::UpdateNumUsed(int64 NewNumUsed)
+{
+	check(NewNumUsed <= ArrayMax);
+
+	if ((Tag == GArraySlackTagToWatch) &&
+		(ElemSize == GArraySlackSizeToWatch) &&
+		(!GArraySlackMaxToWatch || (ArrayMax == GArraySlackMaxToWatch)))
+	{
+		LlmTrackSetBreakpointHere();
+	}
+
+	// Track the allocation in our totals when ArrayNum is first set to something other than INDEX_NONE.  This allows us to
+	// factor out container allocations that aren't arrays (mainly hash tables), which won't ever call "UpdateNumUsed".
+	if (ArrayNum == INDEX_NONE)
+	{
+		GArrayMaxByTag[Tag].fetch_add(ArrayMax * (int64)ElemSize, std::memory_order_relaxed);
+		ArrayNum = 0;
+		FirstAllocFrame = (uint32)GFrameCounter;
+	}
+	GArrayUsedByTag[Tag].fetch_add((NewNumUsed - ArrayNum) * (int64)ElemSize, std::memory_order_relaxed);
+	ArrayNum = NewNumUsed;
+	ArrayPeak = FMath::Max(ArrayPeak, (uint32)FMath::Min(NewNumUsed, 0xffffffffll));
+}
+
+FORCENOINLINE void* FArraySlackTrackingHeader::Realloc(void* Ptr, int64 Count, uint64 ElemSize, int32 Alignment)
+{
+	// Figure out how much padding we need under the allocation
+	int32 HeaderAlign = FPlatformMath::RoundUpToPowerOfTwo(sizeof(FArraySlackTrackingHeader));
+	int32 PaddingRequired = HeaderAlign > Alignment ? HeaderAlign : Alignment;
+
+	// Get the base pointer of the original allocation, and remove tracking for it
+	if (Ptr)
+	{
+		FArraySlackTrackingHeader* TrackingHeader = (FArraySlackTrackingHeader*)((uint8*)Ptr - sizeof(FArraySlackTrackingHeader));
+		TrackingHeader->RemoveAllocation();
+
+		Ptr = (uint8*)TrackingHeader - TrackingHeader->AllocOffset;
+	}
+
+	uint8* ResultPtr = nullptr;
+	if (Count)
+	{
+		ResultPtr = (uint8*)FMemory::Realloc(Ptr, Count * ElemSize + PaddingRequired, Alignment);
+		ResultPtr += PaddingRequired;
+		FArraySlackTrackingHeader* TrackingHeader = (FArraySlackTrackingHeader*)(ResultPtr - sizeof(FArraySlackTrackingHeader));
+
+		// Set the tag and other default information in the allocation if it's newly created
+		if (!Ptr)
+		{
+			// Note that we initially set the slack tracking ArrayNum to INDEX_NONE.  The container allocator is used by both arrays and
+			// other containers (Set / Map / Hash), and we don't know it's actually an array until "UpdateNumUsed" is called on it.
+			check(PaddingRequired <= 65536);
+			TrackingHeader->Next = nullptr;
+			TrackingHeader->Prev = nullptr;
+			TrackingHeader->AllocOffset = (uint16)(PaddingRequired - sizeof(FArraySlackTrackingHeader));
+			TrackingHeader->Tag = LlmGetActiveTag();
+			TrackingHeader->NumStackFrames = 0;			// Filled in later...
+			TrackingHeader->FirstAllocFrame = 0;		// Filled in later...
+			TrackingHeader->ReallocCount = 0;
+			TrackingHeader->ArrayPeak = 0;
+			TrackingHeader->ElemSize = ElemSize;
+			TrackingHeader->ArrayNum = INDEX_NONE;		// Set in UpdateNumUsed
+		}
+
+		// Update ArrayMax and re-register the allocation
+		TrackingHeader->ArrayMax = Count;
+		TrackingHeader->AddAllocation();
+	}
+	else
+	{
+		FMemory::Free(Ptr);
+	}
+
+	return ResultPtr;
+}
+
+struct FArraySlackSortItem
+{
+	FArraySlackTrackingHeader* Header;
+	FName CustomName;
+	uint64 StackTraceTotalSlack;			// Sum of slack for elements with the same stack trace
+	uint32 StackTraceRunLength;				// Run of items with the same stack trace
+	uint32 RunLength;						// Run of identical elements (elemsize, num, max the same)
+	int32 StackTraceIgnore;					// Number of stack trace items to ignore
+
+	bool EqualsTagStackTrace(const FArraySlackSortItem& Other, uint32 StackTraceDepth, bool bLlmEnabled) const
+	{
+		if (bLlmEnabled)
+		{
+			if (Header->Tag != Other.Header->Tag)
+			{
+				return false;
+			}
+			if (CustomName != Other.CustomName)
+			{
+				return false;
+			}
+		}
+
+		// If the stack depth is set to zero, and LLM is disabled, basically everything in the capture will get lumped into
+		// one bucket.  Comparing by ElemSize is a last resort to force some differentiation in the report in that case (or
+		// perhaps if we have a platform that doesn't support stack traces, or someone wants to locally disable them for
+		// performance).  In cases where a stack frame exists, the leaf stack frame is always some sort of template type
+		// specific resize (i.e. TArray<float>::ResizeTo), so this would be redundant, but also harmless.
+		if (Header->ElemSize != Other.Header->ElemSize)
+		{
+			return false;
+		}
+
+		uint32 NumStackFramesThis = FMath::Min((uint32)(Header->NumStackFrames - StackTraceIgnore), StackTraceDepth);
+		uint32 NumStackFramesOther = FMath::Min((uint32)(Other.Header->NumStackFrames - Other.StackTraceIgnore), StackTraceDepth);
+
+		if (NumStackFramesThis != NumStackFramesOther)
+		{
+			return false;
+		}
+
+		return FMemory::Memcmp(&Header->StackFrames[StackTraceIgnore], &Other.Header->StackFrames[Other.StackTraceIgnore], NumStackFramesThis * sizeof(Header->StackFrames[0])) == 0;
+	}
+
+	bool Compare(const FArraySlackSortItem& Other, uint32 StackTraceDepth, bool bLlmEnabled) const
+	{
+		// Order by decreasing stack trace slack total bytes
+		if (StackTraceTotalSlack != Other.StackTraceTotalSlack)
+		{
+			return StackTraceTotalSlack > Other.StackTraceTotalSlack;
+		}
+
+		// Order by tag
+		if (bLlmEnabled)
+		{
+			if (Header->Tag != Other.Header->Tag)
+			{
+				return Header->Tag < Other.Header->Tag;
+			}
+			if (CustomName != Other.CustomName)
+			{
+				return CustomName.GetComparisonIndex().CompareFast(Other.CustomName.GetComparisonIndex()) < 0;
+			}
+		}
+
+		// Order by increasing element size
+		if (Header->ElemSize != Other.Header->ElemSize)
+		{
+			return Header->ElemSize < Other.Header->ElemSize;
+		}
+
+		// Order by stack trace
+		uint32 NumStackFramesThis = FMath::Min((uint32)(Header->NumStackFrames - StackTraceIgnore), StackTraceDepth);
+		uint32 NumStackFramesOther = FMath::Min((uint32)(Other.Header->NumStackFrames - Other.StackTraceIgnore), StackTraceDepth);
+
+		if (NumStackFramesThis != NumStackFramesOther)
+		{
+			return NumStackFramesThis < NumStackFramesOther;
+		}
+		int32 StackFrameOrdinalCompare = memcmp(&Header->StackFrames[StackTraceIgnore], &Other.Header->StackFrames[Other.StackTraceIgnore], NumStackFramesThis * sizeof(Header->StackFrames[0]));
+		if (StackFrameOrdinalCompare)
+		{
+			return StackFrameOrdinalCompare < 0;
+		}
+
+		int64 SlackBytesA = Header->SlackSizeInBytes();
+		int64 SlackBytesB = Other.Header->SlackSizeInBytes();
+		int64 RunSlackBytesA = SlackBytesA * RunLength;
+		int64 RunSlackBytesB = SlackBytesB * Other.RunLength;
+
+		// Order by decreasing run slack total bytes
+		if (RunSlackBytesA != RunSlackBytesB)
+		{
+			return RunSlackBytesA > RunSlackBytesB;
+		}
+
+		// Order by decreasing individual item slack bytes
+		if (SlackBytesA != SlackBytesB)
+		{
+			return SlackBytesA > SlackBytesB;
+		}
+
+		// Order by increasing Max
+		return Header->ArrayMax < Other.Header->ArrayMax;
+	}
+};
+
+static void LlmTrackArrayDumpTag(FArchive* LogFile, FOutputDevice& Output, TAnsiStringBuilder<4096>& Builder, int32 TagIndex, uint32 StackTraceDepth, bool bVerbose, double StartTime)
+{
+	FLowLevelMemTracker& Tracker = FLowLevelMemTracker::Get();
+	bool bLlmEnabled = Tracker.IsEnabled();
+
+	FScopeLock Lock(GTrackArrayDetailedLock);
+
+	TArray<FArraySlackSortItem> ArraySlackSortArray;
+	int32 ReserveAmount = 0;
+	if (TagIndex == INDEX_NONE)
+	{
+		for (int32 CountByTag : GArrayCountByTag)
+		{
+			ReserveAmount += CountByTag;
+		}
+	}
+	else
+	{
+		ReserveAmount = GArrayCountByTag[TagIndex];
+	}
+	ArraySlackSortArray.Reserve(ReserveAmount);
+	ArraySlackSortArray.GetAllocatorInstance().DisableSlackTracking();
+
+	for (FArraySlackTrackingHeader* Current = GTrackArrayDetailedList; Current; Current = Current->Next)
+	{
+		// Don't bother dumping allocations with zero waste (or untracked where ArrayNum == INDEX_NONE)
+		if ((TagIndex == INDEX_NONE || Current->Tag == TagIndex) && (Current->ArrayNum != INDEX_NONE) && (Current->ArrayNum != Current->ArrayMax))
+		{
+			FArraySlackSortItem& SortItem = ArraySlackSortArray.AddDefaulted_GetRef();
+
+			SortItem.Header = Current;
+			if (bLlmEnabled && (Current->Tag == (uint8)ELLMTag::CustomName))
+			{
+				SortItem.CustomName = Tracker.Get().FindPtrDisplayName((uint8*)Current - Current->AllocOffset);
+			}
+			else
+			{
+				SortItem.CustomName = NAME_None;
+			}
+
+			// Filled in later
+			SortItem.StackTraceTotalSlack = 0;
+			SortItem.StackTraceRunLength = 1;
+			SortItem.RunLength = 1;
+			SortItem.StackTraceIgnore = 0;
+		}
+	}
+
+	// We want to ignore ResizeAllocation() if it's the first stack frame, as it's not interesting.  This stack frame will sometimes be there,
+	// and sometimes not, because the most common ResizeAllocation() template variations are tagged FORCENOINLINE to reduce code size.  We set
+	// StackTraceIgnore=1 to indicate where this is the first stack frame, indicating we can ignore it downstream.  It has to be filled in before
+	// the sort, to properly handle the StackTraceDepth setting.
+	//
+	// Assuming symbol lookups are expensive, we sort by leaf stack frame first so we only need to do symbol lookups for unique leaf stack frames.
+	Algo::SortBy(ArraySlackSortArray, [](const FArraySlackSortItem& Item) { return Item.Header->StackFrames[0]; });
+
+	for (int32 ItemIndex = 0; ItemIndex < ArraySlackSortArray.Num(); ItemIndex++)
+	{
+		int32 StackTraceIgnore = 0;
+		if (ItemIndex == 0 || ArraySlackSortArray[ItemIndex].Header->StackFrames[0] != ArraySlackSortArray[ItemIndex - 1].Header->StackFrames[0])
+		{
+			StackTraceIgnore = 0;
+			if (ArraySlackSortArray[ItemIndex].Header->NumStackFrames)
+			{
+				FProgramCounterSymbolInfo SymbolInfo;
+				FPlatformStackWalk::ProgramCounterToSymbolInfo(ArraySlackSortArray[ItemIndex].Header->StackFrames[0], SymbolInfo);
+				if (FCStringAnsi::Strstr(SymbolInfo.FunctionName, "::ResizeAllocation("))
+				{
+					StackTraceIgnore = 1;
+				}
+			}
+		}
+
+		ArraySlackSortArray[ItemIndex].StackTraceIgnore = StackTraceIgnore;
+	}
+
+	// First pass sort -- we haven't yet filled in totals for StackTraceTotalSlack and RunTotalSlack
+	Algo::Sort(ArraySlackSortArray, [StackTraceDepth, bLlmEnabled](const FArraySlackSortItem& A, const FArraySlackSortItem& B) { return A.Compare(B, StackTraceDepth, bLlmEnabled); });
+
+	int64 IgnoredGroups = 0;
+	int64 IgnoredSlack = 0;
+	{
+		// Compute slack associated with each stack trace and runs of identical elements, and store it on the sort elements
+		int32 StackTraceRun = 0;
+		int64 StackTraceTotal = 0;
+		uint32 ElementRun = 0;
+		for (int32 ItemIndex = 0; ItemIndex < ArraySlackSortArray.Num(); ItemIndex++)
+		{
+			// Add current item
+			int64 ElementSlack = ArraySlackSortArray[ItemIndex].Header->SlackSizeInBytes();
+			StackTraceRun++;
+			StackTraceTotal += ElementSlack;
+			ElementRun++;
+
+			// If the next item has a different stack trace, or it's the end of the array, echo the stack trace total slack to all the items
+			if ((ItemIndex == ArraySlackSortArray.Num() - 1) || !ArraySlackSortArray[ItemIndex].EqualsTagStackTrace(ArraySlackSortArray[ItemIndex + 1], StackTraceDepth, bLlmEnabled))
+			{
+				for (int32 RunIndex = ItemIndex - (StackTraceRun - 1); RunIndex <= ItemIndex; RunIndex++)
+				{
+					ArraySlackSortArray[RunIndex].StackTraceTotalSlack = StackTraceTotal;
+					ArraySlackSortArray[RunIndex].StackTraceRunLength = StackTraceRun;
+				}
+				if (StackTraceTotal < GArraySlackThreshold)
+				{
+					IgnoredGroups++;
+					IgnoredSlack += StackTraceTotal;
+				}
+				StackTraceTotal = 0;
+				StackTraceRun = 0;
+			}
+
+			// Check if the item is different at all, and echo run length to all the items
+			if ((ItemIndex == ArraySlackSortArray.Num() - 1) || ArraySlackSortArray[ItemIndex].Compare(ArraySlackSortArray[ItemIndex + 1], StackTraceDepth, bLlmEnabled))
+			{
+				for (int32 RunIndex = ItemIndex - (ElementRun - 1); RunIndex <= ItemIndex; RunIndex++)
+				{
+					ArraySlackSortArray[RunIndex].RunLength = ElementRun;
+				}
+				ElementRun = 0;
+			}
+		}
+	}
+
+	// Second pass, final sort
+	Algo::Sort(ArraySlackSortArray, [StackTraceDepth, bLlmEnabled](const FArraySlackSortItem& A, const FArraySlackSortItem& B) { return A.Compare(B, StackTraceDepth, bLlmEnabled); });
+
+	// Only include tag column if LLM is enabled
+	const TCHAR* TagColumnSeparator = bLlmEnabled ? TEXT("\t") : TEXT("");
+	const ANSICHAR* TagColumnSeparatorANSI = bLlmEnabled ? "\t" : "";
+
+	if (LogFile)
+	{
+		Builder.Reset();
+		Builder.Appendf("Ignored:\t%lld\tGroups,\t%lld\tBytes total\n", IgnoredGroups, IgnoredSlack);
+		Builder.Appendf("Under:\t%d\tGroup size threshold\n\n", GArraySlackThreshold);
+
+		Builder.Appendf("NumArrays\tReallocs\tLifetime\tPeakAvg\tPeak\tElemSize\tNum\tMax\t%s\tStackSlack%s%s\tStackTrace\n",
+			bVerbose ? "ItemSlack" : "LargestItem",
+			TagColumnSeparatorANSI,
+			bLlmEnabled ? "Tag" : "");
+		LogFile->Serialize(Builder.GetData(), Builder.Len());
+	}
+	else
+	{
+		FPlatformMisc::LowLevelOutputDebugStringf(TEXT("Ignored:\t%lld\tGroups,\t%lld\tBytes total\n"), IgnoredGroups, IgnoredSlack);
+		FPlatformMisc::LowLevelOutputDebugStringf(TEXT("Under:\t%d\tGroup size threshold\n\n"), GArraySlackThreshold);
+
+		// We prepend "SlackReport" to every line when outputting to the debug window, as this can help filtering out random
+		// log lines after the output is cut and pasted.  You can sort by the first column to accomplish that.
+		FPlatformMisc::LowLevelOutputDebugStringf(TEXT("SlackReport\tNumArrays\tReallocs\tLifetime\tPeakAvg\tPeak\tElemSize\tNum\tMax\t%s\tStackSlack%s%s\tStackTrace\n"),
+			bVerbose ? "ItemSlack" : "LargestItem",
+			TagColumnSeparator,
+			bLlmEnabled ? TEXT("Tag") : TEXT(""));
+	}
+
+	{
+		uint32 CurrentFrame = (uint32)GFrameCounter;
+		int32 ItemRun = 0;
+		double ItemReallocs = 0.0;
+		double ItemLifetime = 0.0;
+		double ItemPeakSum = 0.0;
+		uint32 ItemPeak = 0;
+		int32 FullStacksPrinted = 0;
+
+		for (int32 ItemIndex = 0; ItemIndex < ArraySlackSortArray.Num(); ItemIndex++)
+		{
+#if PLATFORM_WINDOWS
+			// Windows is a lot faster, so we don't need as much progress logging
+			constexpr int32 ProgressInterval = 25000;
+#else
+			constexpr int32 ProgressInterval = 1000;
+#endif
+			if ((ItemIndex % ProgressInterval) == 0)
+			{
+				Output.Logf(TEXT("Array Slack %d / %d allocs...  (%.2lf minutes, batch %lld bytes -> threshold %lld)"),
+					ItemIndex, ArraySlackSortArray.Num(), (FPlatformTime::Seconds() - StartTime) / 60.0, ArraySlackSortArray[ItemIndex].StackTraceTotalSlack, GArraySlackThreshold);
+			}
+
+			// Count current item
+			ItemRun++;
+			ItemReallocs += ArraySlackSortArray[ItemIndex].Header->ReallocCount;
+			ItemLifetime += CurrentFrame - ArraySlackSortArray[ItemIndex].Header->FirstAllocFrame;
+			ItemPeakSum += ArraySlackSortArray[ItemIndex].Header->ArrayPeak;
+			ItemPeak = FMath::Max(ItemPeak, ArraySlackSortArray[ItemIndex].Header->ArrayPeak);
+
+			// If this item is different than the next, or the last item, echo the count and item
+			if ((ItemIndex == ArraySlackSortArray.Num() - 1) || ArraySlackSortArray[ItemIndex].Compare(ArraySlackSortArray[ItemIndex + 1], StackTraceDepth, bLlmEnabled))
+			{
+				// Determine if this run has the same stack trace as the previous.  In verbose mode, we only print stack trace specific totals for unique
+				// stack traces, and when not verbose, we only print lines at all for unique stack traces.
+				int32 RunStart = ItemIndex - (ItemRun - 1);
+				bool bUniqueStackTrace = RunStart == 0 || !ArraySlackSortArray[RunStart - 1].EqualsTagStackTrace(ArraySlackSortArray[RunStart], StackTraceDepth, bLlmEnabled);
+
+				if ((bVerbose || bUniqueStackTrace) && (ArraySlackSortArray[ItemIndex].StackTraceTotalSlack >= GArraySlackThreshold))
+				{
+					if (LogFile)
+					{
+						Builder.Reset();
+						Builder.Appendf(
+							"%llu\t%.1lf\t%.1lf\t%.1lf\t%u\t%lld\t%lld\t%lld\t%lld\t",
+							bVerbose ? ItemRun : ArraySlackSortArray[ItemIndex].StackTraceRunLength,
+							ItemReallocs / ItemRun,
+							ItemLifetime / ItemRun,
+							ItemPeakSum / ItemRun,
+							ItemPeak,
+							ArraySlackSortArray[ItemIndex].Header->ElemSize,
+							ArraySlackSortArray[ItemIndex].Header->ArrayNum,
+							ArraySlackSortArray[ItemIndex].Header->ArrayMax,
+							ArraySlackSortArray[ItemIndex].Header->SlackSizeInBytes() * ItemRun);
+
+						// Stack trace total slack if it's unique
+						if (bUniqueStackTrace)
+						{
+							Builder.Appendf("%lld", ArraySlackSortArray[ItemIndex].StackTraceTotalSlack);
+						}
+
+						// Tag name
+						if (bLlmEnabled)
+						{
+							Builder.AppendChar('\t');
+							if (ArraySlackSortArray[ItemIndex].CustomName != NAME_None)
+							{
+								Builder.Append(*ArraySlackSortArray[ItemIndex].CustomName.ToString());
+							}
+							else
+							{
+								Builder.Append(Tracker.FindTagDisplayName(ArraySlackSortArray[ItemIndex].Header->Tag).ToString());
+							}
+						}
+					}
+					else
+					{
+						TStringBuilder<32> ByStackSlack;
+						if (bUniqueStackTrace)
+						{
+							ByStackSlack.Appendf(TEXT("%lld"), ArraySlackSortArray[ItemIndex].StackTraceTotalSlack);
+						}
+
+						FPlatformMisc::LowLevelOutputDebugStringf(
+							TEXT("SlackReport\t%llu\t%.1lf\t%.1lf\t%.1lf\t%lld\t%lld\t%lld\t%lld\t%lld\t%s%s%s"),
+							bVerbose ? ItemRun : ArraySlackSortArray[ItemIndex].StackTraceRunLength,
+							ItemReallocs / ItemRun,
+							ItemLifetime / ItemRun,
+							ItemPeakSum / ItemRun,
+							ItemPeak,
+							ArraySlackSortArray[ItemIndex].Header->ElemSize,
+							ArraySlackSortArray[ItemIndex].Header->ArrayNum,
+							ArraySlackSortArray[ItemIndex].Header->ArrayMax,
+							ArraySlackSortArray[ItemIndex].Header->SlackSizeInBytes() * ItemRun,
+							ByStackSlack.ToString(),
+							TagColumnSeparator,
+							!bLlmEnabled ? TEXT("") :
+							(ArraySlackSortArray[ItemIndex].CustomName != NAME_None ?
+								*ArraySlackSortArray[ItemIndex].CustomName.ToString() :
+								*Tracker.FindTagDisplayName(ArraySlackSortArray[ItemIndex].Header->Tag).ToString()));
+					}
+
+					// Only print stack trace if this is a unique stack trace.
+					if (bUniqueStackTrace)
+					{
+						int32 StackFirst = ArraySlackSortArray[ItemIndex].StackTraceIgnore;
+						int32 StackCount = ArraySlackSortArray[ItemIndex].Header->NumStackFrames - ArraySlackSortArray[ItemIndex].StackTraceIgnore;
+						if (FullStacksPrinted >= GArraySlackFullStackNum)
+						{
+							StackCount = FMath::Min(StackCount, (int32)StackTraceDepth);
+						}
+						else
+						{
+							FullStacksPrinted++;
+						}
+
+						for (int32 StackIndex = StackFirst; StackIndex < StackFirst + StackCount; StackIndex++)
+						{
+							FProgramCounterSymbolInfo SymbolInfo;
+							FPlatformStackWalk::ProgramCounterToSymbolInfo(ArraySlackSortArray[ItemIndex].Header->StackFrames[StackIndex], SymbolInfo);
+
+							if (LogFile)
+							{
+								Builder.AppendChar('\t');
+								Builder.Append(SymbolInfo.FunctionName[0] ? SymbolInfo.FunctionName : "UnknownFunction");
+
+								if (SymbolInfo.Filename[0] && SymbolInfo.LineNumber)
+								{
+									// Format " [Filename:Line]"
+									Builder.Append(" [");
+									Builder.Append(SymbolInfo.Filename);
+									Builder.Appendf(":%i]", SymbolInfo.LineNumber);
+								}
+								else
+								{
+									Builder.Append(" []");
+								}
+							}
+							else
+							{
+								TStringBuilder<512> SymbolName;
+								SymbolName.AppendChar(TEXT('\t'));
+								SymbolName.Append(SymbolInfo.FunctionName[0] ? SymbolInfo.FunctionName : "UnknownFunction");
+
+								if (SymbolInfo.Filename[0] && SymbolInfo.LineNumber)
+								{
+									// Format " [Filename:Line]"
+									SymbolName.Append(TEXT(" ["));
+									SymbolName.Append(SymbolInfo.Filename);
+									SymbolName.Appendf(TEXT(":%i]"), SymbolInfo.LineNumber);
+								}
+								else
+								{
+									SymbolName.Append(TEXT(" []"));
+								}
+
+								FPlatformMisc::LowLevelOutputDebugString(SymbolName.ToString());
+							}
+						}
+					}
+
+					if (LogFile)
+					{
+						Builder.AppendChar('\n');
+						LogFile->Serialize(Builder.GetData(), Builder.Len());
+					}
+					else
+					{
+						FPlatformMisc::LowLevelOutputDebugString(TEXT("\n"));
+					}
+				}
+
+				ItemRun = 0;
+				ItemReallocs = 0.0;
+				ItemLifetime = 0.0;
+				ItemPeakSum = 0.0;
+				ItemPeak = 0;
+			}
+		}
+	}
+}
+
+void ArraySlackTrackInit()
+{
+	// Any array allocations before this is called won't have array slack tracking, although subsequent reallocations of existing arrays
+	// will gain tracking if that occurs.  The goal is to filter out startup constructors which run before Main, which introduce a
+	// ton of noise into slack reports.  Especially the roughly 30,000 static FString constructors in the code base, each with a
+	// unique call stack, and all having a little bit of slack due to malloc bucket size rounding.
+	GArraySlackInit = true;
+}
+
+static void LlmTrackArrayTick()
+{
+	// Updating these every frame is handy, so you can see them in a watch window in the debugger or debug print them without running a capture.
+	// We could consider including these as stats, so you can track them in Insights, but the report is giving enough information for now.
+	for (int32 TagIndex = 0; TagIndex < 256; TagIndex++)
+	{
+		GArraySlackByTag[TagIndex] = GArrayMaxByTag[TagIndex] - GArrayUsedByTag[TagIndex];
+	}
+}
+
+void ArraySlackTrackGenerateReport(const TCHAR* Cmd, FOutputDevice& Output)
+{
+	Output.Logf(TEXT("Generating Array Slack report."));
+
+	double StartTime = FPlatformTime::Seconds();
+
+	// Make sure the slack by tag totals are up to date
+	LlmTrackArrayTick();
+
+	// Parse command -- tokens
+	FString LogFilename;
+	uint32 StackTraceDepth = GArraySlackDefaultStackDepth;
+	bool bVerbose = GArraySlackDefaultVerbose;
+	for (FString Arg = FParse::Token(Cmd, false); !Arg.IsEmpty(); Arg = FParse::Token(Cmd, false))
+	{
+		if (Arg[0] == TEXT('-'))
+		{
+			FStringView StackDepthSwitch(TEXTVIEW("-Stack="));
+			FStringView VerboseSwitch(TEXTVIEW("-Verbose="));
+			if ((Arg.Len() > StackDepthSwitch.Len()) && !FCString::Strnicmp(&Arg[0], StackDepthSwitch.GetData(), StackDepthSwitch.Len()))
+			{
+				StackTraceDepth = (uint32)FCString::Strtoui64(&Arg[StackDepthSwitch.Len()], nullptr, 10);
+			}
+			else if ((Arg.Len() > VerboseSwitch.Len()) && !FCString::Strnicmp(&Arg[0], VerboseSwitch.GetData(), VerboseSwitch.Len()))
+			{
+				bVerbose = Arg[VerboseSwitch.Len()] != TEXT('0');
+			}
+			else
+			{
+				Output.Logf(TEXT("Array Slack unsupported switch: \"%s\".  Valid switches:  -Stack=N, -Verbose=[0,1]"), *Arg);
+			}
+		}
+		else
+		{
+			LogFilename = Arg;
+		}
+	}
+
+	FArchive* LogFile = nullptr;
+
+#if ARRAY_SLACK_LOG_TO_MEMORY
+	TArray<uint8> LogFileMemory;
+	LogFileMemory.Reserve(4 * 1024 * 1024);
+	LogFileMemory.GetAllocatorInstance().DisableSlackTracking();
+#endif
+	FString LogFilenameWithPath;
+
+	// Special name "NOFILE" indicates to write to debug log instead of file.  Useful for debugging the system, as you can see the
+	// lines of text generated while debugging, as opposed to needing to wait until the file gets written to look at the output.
+	if (!LogFilename.Equals(TEXT("NOFILE")))
+	{
+		FString AbsoluteProjectLogDir = FPaths::ConvertRelativePathToFull(FPaths::ProjectLogDir());
+		FString SlackReportLogDir = FPaths::Combine(AbsoluteProjectLogDir, TEXT("SlackReport"));
+		IPlatformFile::GetPlatformPhysical().CreateDirectoryTree(*SlackReportLogDir);
+
+		if (LogFilename.IsEmpty())
+		{
+			LogFilename = FString::Printf(TEXT("SlackDump_%03u.tsv"), GArraySlackDumpIndex++);
+		}
+		else
+		{
+			if (!LogFilename.EndsWith(TEXT(".tsv")))
+			{
+				LogFilename.Append(TEXT(".tsv"));
+			}
+		}
+		LogFilenameWithPath = SlackReportLogDir / LogFilename;
+
+#if ARRAY_SLACK_LOG_TO_MEMORY
+		LogFile = new FMemoryWriter(LogFileMemory);
+#else
+		IFileManager* FileManager = &IFileManager::Get();
+		LogFile = FileManager->CreateFileWriter(*LogFilenameWithPath, 0);
+#endif
+	}
+
+	const FLowLevelMemTracker& Tracker = FLowLevelMemTracker::Get();
+	bool bLlmEnabled = Tracker.IsEnabled();
+
+	TAnsiStringBuilder<4096> Builder;
+
+	TArray<uint64> SortedTags;
+	SortedTags.SetNumZeroed(256);
+	SortedTags.GetAllocatorInstance().DisableSlackTracking();
+
+	// Tag summary report is only useful when -llm is enabled
+	if (bLlmEnabled)
+	{
+		if (LogFile)
+		{
+			Builder.Append("Tag\tSlack\tTotalMem\n");
+			LogFile->Serialize(Builder.GetData(), Builder.Len());
+		}
+		else
+		{
+			FPlatformMisc::LowLevelOutputDebugString(TEXT("SlackByTag\tTag\tSlack\tTotalMem\n"));
+		}
+
+		// Sort tags by descending memory.  Store tag index in low 8 bits, memory in high 56 bits.
+		for (int32 TagIndex = 0; TagIndex < 256; TagIndex++)
+		{
+			SortedTags[TagIndex] = (GArraySlackByTag[TagIndex] << 8) | TagIndex;
+		}
+		SortedTags.Sort(TGreater<uint64>());
+
+		for (int32 SortedIndex = 0; SortedIndex < 256; SortedIndex++)
+		{
+			if (SortedTags[SortedIndex] >= 256)
+			{
+				uint32 TagIndex = (uint32)SortedTags[SortedIndex] & 0xff;
+
+				if (LogFile)
+				{
+					Builder.Reset();
+					Builder.Append(Tracker.FindTagDisplayName(TagIndex).ToString());
+					Builder.Appendf("\t%lld\t%lld\n",
+						GArraySlackByTag[TagIndex],
+						GArrayMaxByTag[TagIndex].load(std::memory_order_relaxed));
+					LogFile->Serialize(Builder.GetData(), Builder.Len());
+				}
+				else
+				{
+					FPlatformMisc::LowLevelOutputDebugStringf(TEXT("SlackByTag\t%s\t%lld\t%lld\n"),
+						*Tracker.FindTagDisplayName(TagIndex).ToString(),
+						GArraySlackByTag[TagIndex],
+						GArrayMaxByTag[TagIndex].load(std::memory_order_relaxed));
+				}
+			}
+		}
+
+		if (LogFile)
+		{
+			char NewLines[] = "\n\n";
+			LogFile->Serialize(NewLines, 2);
+		}
+	}
+
+	if (LogFile)
+	{
+		// Append information about options and timing of dump
+		FArraySlackTrackingHeader Dummy;
+		Builder.Reset();
+		Builder.Appendf("Ran with:\t-Stack=%u -Verbose=%d", FMath::Min((uint32)UE_ARRAY_COUNT(Dummy.StackFrames), StackTraceDepth), bVerbose ? 1 : 0);
+		if (FLowLevelMemTracker::Get().IsEnabled())
+		{
+			Builder.Append(" -llm");
+		}
+		if (!bVerbose)
+		{
+			Builder.Append("\t\t\tFields besides NumArrays / StackSlack are for the largest slack bucket (unique Num / Max combo), run with -Verbose=1 to see all buckets");
+		}
+		Builder.Appendf("\nOn frame:\t%u\n\n", (int32)GFrameCounter);
+		LogFile->Serialize(Builder.GetData(), Builder.Len());
+	}
+
+	if (bLlmEnabled && GArraySlackGroupByTag)
+	{
+		// Original behavior grouped by tag, but all in one batch is generally preferable.  Could expose this with a switch in the future.
+		for (int32 SortedIndex = 0; SortedIndex < 256; SortedIndex++)
+		{
+			if (SortedTags[SortedIndex] > 256)
+			{
+				int32 TagIndex = (int32)SortedTags[SortedIndex] & 0xff;
+				if (GArraySlackByTag[TagIndex])
+				{
+#if ARRAY_SLACK_LOG_TO_MEMORY
+					// Disable tracking each loop iteration, in case the memory grew to the point where it was reallocated in the previous iteration.
+					LogFileMemory.GetAllocatorInstance().DisableSlackTracking();
+#endif
+
+					LlmTrackArrayDumpTag(LogFile, Output, Builder, TagIndex, StackTraceDepth, bVerbose, StartTime);
+				}
+			}
+		}
+	}
+	else
+	{
+		// INDEX_NONE == dump all tags in one batch
+		LlmTrackArrayDumpTag(LogFile, Output, Builder, INDEX_NONE, StackTraceDepth, bVerbose, StartTime);
+	}
+
+	if (LogFile)
+	{
+		delete LogFile;
+
+#if ARRAY_SLACK_LOG_TO_MEMORY
+		// Now create the actual file
+		IFileManager* FileManager = &IFileManager::Get();
+		LogFile = FileManager->CreateFileWriter(*LogFilenameWithPath, 0);
+		LogFile->Serialize(LogFileMemory.GetData(), LogFileMemory.Num());
+		delete LogFile;
+#endif
+	}
+
+	Output.Logf(TEXT("Finished generating Array Slack report to %s."), LogFile ? *LogFilename : TEXT("[Debug Output]"));
+}
+
+uint8 LlmGetActiveTag()
+{
+	FLowLevelMemTracker& Tracker = FLowLevelMemTracker::Get();
+	if (!Tracker.IsInitialized())
+	{
+		return 0;
+	}
+
+	const UE::LLMPrivate::FTagData* TagData = Tracker.GetActiveTagData(ELLMTracker::Default);
+	return TagData ? (uint8)TagData->GetEnumTag() : 0;
+}
+
+#endif  // UE_ENABLE_ARRAY_SLACK_TRACKING
+
+UE_TRACE_CHANNEL(MemTagChannel, "Memory overview")
+
+UE_TRACE_EVENT_BEGIN(LLM, TagsSpec, NoSync|Important)
+	UE_TRACE_EVENT_FIELD(uint64, TagId)
+	UE_TRACE_EVENT_FIELD(uint64, ParentId)
+	UE_TRACE_EVENT_FIELD(uint8, TagSetId)
+	UE_TRACE_EVENT_FIELD(UE::Trace::AnsiString, Name)
+UE_TRACE_EVENT_END()
+
+UE_TRACE_EVENT_BEGIN(LLM, TrackerSpec, NoSync|Important)
+	UE_TRACE_EVENT_FIELD(uint8, TrackerId)
+	UE_TRACE_EVENT_FIELD(UE::Trace::AnsiString, Name)
+UE_TRACE_EVENT_END()
+
+UE_TRACE_EVENT_BEGIN(LLM, TagSetSpec, NoSync|Important)
+	UE_TRACE_EVENT_FIELD(uint8, TagSetId)
+	UE_TRACE_EVENT_FIELD(UE::Trace::AnsiString, Name)
+UE_TRACE_EVENT_END()
+
+UE_TRACE_EVENT_BEGIN(LLM, TagValue)
+	UE_TRACE_EVENT_FIELD(uint8, TrackerId)
+	UE_TRACE_EVENT_FIELD(uint64, Cycle)
+	UE_TRACE_EVENT_FIELD(uint64[], Tags)
+	UE_TRACE_EVENT_FIELD(int64[], Values)
+UE_TRACE_EVENT_END()
+
+#define LLM_CSV_PROFILER_WRITER_ENABLED CSV_PROFILER_STATS
+
+#if LLM_CSV_PROFILER_WRITER_ENABLED
+	CSV_DEFINE_CATEGORY(LLM, true);
+	CSV_DEFINE_CATEGORY(LLMPlatform, true);
+#endif
+
+TAutoConsoleVariable<int32> CVarLLMTrackPeaks(TEXT("LLM.TrackPeaks"), 0,
+	TEXT("Track peak memory in each category since process start rather than current frame's value."));
+
+TAutoConsoleVariable<int32> CVarLLMWriteInterval(
+	TEXT("LLM.LLMWriteInterval"),
+	1,
+	TEXT("The number of seconds between each line in the LLM csv (zero to write every frame)")
+);
+
+TAutoConsoleVariable<int32> CVarLLMCsvFlushEveryRow(TEXT("LLM.CsvFlushEveryRow"), 1,
+	TEXT("Whether to flush the CSV with every row written. If disabled we only flush on a crash"));
+
+TAutoConsoleVariable<int32> CVarLLMHeaderMaxSize(
+	TEXT("LLM.LLMHeaderMaxSize"),
+#if LLM_ALLOW_ASSETS_TAGS
+	500000, // When using asset tags, you will have MANY more LLM titles since so many are auto generated.
+#elif LLM_ALLOW_UOBJECTCLASSES_TAGS
+	30000, // When using uobjects clases, you will have at least 20K+ LLM titles.
+#else
+	5000,
+#endif
+	TEXT("The maximum total number of characters allowed for all of the LLM titles")
+);
+
+FAutoConsoleCommand LLMSnapshot(
+	TEXT("LLMSnapshot"),
+	TEXT("Takes a single LLM Snapshot of one frame. This command requires the commandline -llmdisableautopublish"),
+	FConsoleCommandDelegate::CreateLambda([]()
+	{
+		FLowLevelMemTracker::Get().PublishDataSingleFrame();
+	}));
+
+FAutoConsoleCommand DumpLLM(
+	TEXT("DumpLLM"),
+	TEXT("Logs out the current and peak sizes of all tracked LLM tags"),
+	FConsoleCommandWithWorldArgsAndOutputDeviceDelegate::CreateLambda(
+		[](const TArray<FString>& Args, UWorld* InWorld, FOutputDevice& Ar)
+	{
+		FString Command = FString::Join(Args, TEXT(" "));
+
+		FLowLevelMemTracker::EDumpFormat DumpFormat = FLowLevelMemTracker::EDumpFormat::PlainText;
+		bool bCSV = FParse::Param(*Command, TEXT("CSV"));
+		if (bCSV)
+		{
+			DumpFormat = FLowLevelMemTracker::EDumpFormat::CSV;
+		}
+
+		UE::LLM::ESizeParams SizeParams = UE::LLM::ESizeParams::Default;
+		bool bSnapshot = FParse::Param(*Command, TEXT("SNAPSHOT"));
+		if (bSnapshot)
+		{
+			EnumAddFlags(SizeParams, UE::LLM::ESizeParams::RelativeToSnapshot);
+		}
+
+		ELLMTagSet TagSet = ELLMTagSet::None;
+#if LLM_ALLOW_ASSETS_TAGS
+		if (FParse::Param(*Command, TEXT("Assets")))
+		{
+			TagSet = ELLMTagSet::Assets;
+		}
+		else if (FParse::Param(*Command, TEXT("AssetClasses")))
+		{
+			TagSet = ELLMTagSet::AssetClasses;
+		}
+#elif LLM_ALLOW_UOBJECTCLASSES_TAGS
+		if (FParse::Param(*Command, TEXT("UObjectClasses")))
+		{
+			TagSet = ELLMTagSet::UObjectClasses;
+		}
+#endif // LLM_ALLOW_UOBJECTCLASSES_TAGS
+
+		FLowLevelMemTracker::Get().DumpToLog(DumpFormat, &Ar, SizeParams, TagSet);
+	}));
+
+#if !PLATFORM_HAS_MULTITHREADED_PREMAIN
+struct FLowLevelMemTracker::FEnableStateScopeLock
+{
+};
+bool FLowLevelMemTracker::TryEnterEnabled(FLowLevelMemTracker::FEnableStateScopeLock& ScopeLock)
+{
+	return IsEnabled();
+}
+#else // !PLATFORM_HAS_MULTITHREADED_PREMAIN
+struct FLowLevelMemTracker::FEnableStateScopeLock
+{
+	// Undefine copy/move constructor since FReadScopeLock does not support it.
+	FEnableStateScopeLock() = default;
+	FEnableStateScopeLock(const FEnableStateScopeLock&) = delete;
+	FEnableStateScopeLock(FEnableStateScopeLock&&) = delete;
+
+	TOptional<FReadScopeLock> Inner;
+};
+
+namespace UE::LLMPrivate
+{
+
+FRWLock& GetEnableStateLock()
+{
+	static FRWLock Lock;
+	return Lock;
+}
+
+} // namespace UE::LLMPrivate
+
+bool FLowLevelMemTracker::TryEnterEnabled(FLowLevelMemTracker::FEnableStateScopeLock& ScopeLock)
+{
+	switch (EnabledState)
+	{
+	case EEnabled::NotYetKnown:
+		ScopeLock.Inner.Emplace(UE::LLMPrivate::GetEnableStateLock());
+		// Evaluate EnabledState again since it may have changed under ProcessCommandLine's WriteLock.
+		return IsEnabled();
+	case EEnabled::Disabled:
+		return false;
+	case EEnabled::Enabled:
+		return true;
+	default:
+		checkNoEntry();
+		return false;
+	}
+}
+#endif // !PLATFORM_HAS_MULTITHREADED_PREMAIN
+
+namespace UE::LLM
+{
+
+FString GetLLMProfilingDir()
+{
+	return FPaths::ProfilingDir() + TEXT("LLM/");
+}
+
+} // namespace UE::LLM
+
+void FLowLevelMemTracker::DumpToLog(EDumpFormat DumpFormat, FOutputDevice* OutputDevice, UE::LLM::ESizeParams SizeParams, ELLMTagSet TagSet)
+{
+	if (!IsEnabled())
+	{
+		return;
+	}
+	const float InvToMb = 1.0 / (1024 * 1024);
+	FOutputDevice& Ar = *(OutputDevice ? OutputDevice : GLog);
+	if (DumpFormat == EDumpFormat::CSV)
+	{
+		Ar.Logf(TEXT(",TagName,SizeMB,PeakMB,Tracker,PathName"));
+	}
+	else
+	{
+		Ar.Logf(TEXT("%40s %12s %12s  Tracker PathName"), TEXT("TagName"), TEXT("SizeMB"), TEXT("PeakMB"));
+	}
+
+	for (ELLMTracker TrackerType : {ELLMTracker::Default, ELLMTracker::Platform})
+	{
+		const TCHAR* TrackerName = TrackerType == ELLMTracker::Default ? TEXT("Default") : TEXT("Platform");
+
+		struct FTagLine
+		{
+			FString TagName;
+			FString Line;
+			int64 Size;
+		};
+
+		UE::LLM::ESizeParams CurrentSizeParams = SizeParams;
+		UE::LLM::ESizeParams PeakSizeParams = CurrentSizeParams;
+		EnumAddFlags(PeakSizeParams, UE::LLM::ESizeParams::ReportPeak);
+
+		TArray<FTagLine> TagLines;
+		for (const UE::LLMPrivate::FTagData* TagData : GetTrackedTags(TrackerType, TagSet))
+		{
+			if (!TagData->IsReportable())
+			{
+				continue;
+			}
+			int64 CurrentAmount = GetTagAmountForTracker(TrackerType, TagData, SizeParams);
+			int64 PeakAmount = GetTagAmountForTracker(TrackerType, TagData, PeakSizeParams);
+			FString TagName = GetTagUniqueName(TagData).ToString();
+			FString TagLine;
+			if (DumpFormat == EDumpFormat::CSV)
+			{
+				TagLine = FString::Printf(TEXT(",%s,%.3f,%.3f,%s,%s"),
+					*TagName,
+					static_cast<float>(CurrentAmount) * InvToMb,
+					static_cast<float>(PeakAmount) * InvToMb,
+					TrackerName,
+					*GetTagDisplayPathName(TagData));
+			}
+			else
+			{
+				TagLine = FString::Printf(TEXT("%40s %12.3f %12.3f %8s %s"),
+					*TagName,
+					static_cast<float>(CurrentAmount) * InvToMb,
+					static_cast<float>(PeakAmount) * InvToMb,
+					TrackerName,
+					*GetTagDisplayPathName(TagData));
+			}
+			TagLines.Add(FTagLine{ MoveTemp(TagName), MoveTemp(TagLine), CurrentAmount });
+		}
+		TagLines.Sort([](const FTagLine& A, const FTagLine& B)
+			{
+				if (A.Size != B.Size)
+				{
+					return A.Size > B.Size;
+				}
+				return A.TagName < B.TagName;
+			});
+		for (FTagLine& TagLine : TagLines)
+		{
+			Ar.Logf(TEXT("%s"), *TagLine.Line);
+		}
+	}
+}
+
+DECLARE_LLM_MEMORY_STAT(TEXT("LLM Overhead"), STAT_LLMOverheadTotal, STATGROUP_LLMOverhead);
+
+// LLM Summary stats referenced by ELLMTagNames
+DEFINE_STAT(STAT_EngineSummaryLLM);
+DEFINE_STAT(STAT_ProjectSummaryLLM);
+DEFINE_STAT(STAT_NetworkingSummaryLLM);
+DEFINE_STAT(STAT_AudioSummaryLLM);
+DEFINE_STAT(STAT_TrackedTotalSummaryLLM);
+DEFINE_STAT(STAT_MeshesSummaryLLM);
+DEFINE_STAT(STAT_PhysicsSummaryLLM);
+DEFINE_STAT(STAT_PhysXSummaryLLM);
+DEFINE_STAT(STAT_ChaosSummaryLLM);
+DEFINE_STAT(STAT_UObjectSummaryLLM);
+DEFINE_STAT(STAT_AnimationSummaryLLM);
+DEFINE_STAT(STAT_StaticMeshSummaryLLM);
+DEFINE_STAT(STAT_MaterialsSummaryLLM);
+DEFINE_STAT(STAT_ParticlesSummaryLLM);
+DEFINE_STAT(STAT_NiagaraSummaryLLM);
+DEFINE_STAT(STAT_UISummaryLLM);
+DEFINE_STAT(STAT_NavigationSummaryLLM);
+DEFINE_STAT(STAT_TexturesSummaryLLM);
+DEFINE_STAT(STAT_MediaStreamingSummaryLLM);
+
+// LLM stats referenced by ELLMTagNames
+DECLARE_LLM_MEMORY_STAT(TEXT("Total"), STAT_TotalLLM, STATGROUP_LLMFULL);
+DECLARE_LLM_MEMORY_STAT(TEXT("Untracked"), STAT_UntrackedLLM, STATGROUP_LLMFULL);
+DECLARE_LLM_MEMORY_STAT(TEXT("Total"), STAT_PlatformTotalLLM, STATGROUP_LLMPlatform);
+DECLARE_LLM_MEMORY_STAT(TEXT("Tracked Total"), STAT_TrackedTotalLLM, STATGROUP_LLMFULL);
+DECLARE_LLM_MEMORY_STAT(TEXT("Untagged"), STAT_UntaggedTotalLLM, STATGROUP_LLMFULL);
+DECLARE_LLM_MEMORY_STAT(TEXT("WorkingSetSize"), STAT_WorkingSetSizeLLM, STATGROUP_LLMFULL);
+DECLARE_LLM_MEMORY_STAT(TEXT("PagefileUsed"), STAT_PagefileUsedLLM, STATGROUP_LLMFULL);
+DECLARE_LLM_MEMORY_STAT(TEXT("Tracked Total"), STAT_PlatformTrackedTotalLLM, STATGROUP_LLMPlatform);
+DECLARE_LLM_MEMORY_STAT(TEXT("Untagged"), STAT_PlatformUntaggedTotalLLM, STATGROUP_LLMPlatform);
+DECLARE_LLM_MEMORY_STAT(TEXT("Untracked"), STAT_PlatformUntrackedLLM, STATGROUP_LLMPlatform);
+DECLARE_LLM_MEMORY_STAT(TEXT("Overhead"), STAT_PlatformOverheadLLM, STATGROUP_LLMPlatform);
+DECLARE_LLM_MEMORY_STAT(TEXT("OS Available"), STAT_PlatformOSAvailableLLM, STATGROUP_LLMPlatform);
+DECLARE_LLM_MEMORY_STAT(TEXT("FMalloc"), STAT_FMallocLLM, STATGROUP_LLMPlatform);
+DECLARE_LLM_MEMORY_STAT(TEXT("FMalloc Unused"), STAT_FMallocUnusedLLM, STATGROUP_LLMFULL);
+DECLARE_LLM_MEMORY_STAT(TEXT("RHI Unused"), STAT_RHIUnusedLLM, STATGROUP_LLMFULL);
+DECLARE_LLM_MEMORY_STAT(TEXT("ThreadStack"), STAT_ThreadStackLLM, STATGROUP_LLMFULL);
+DECLARE_LLM_MEMORY_STAT(TEXT("ThreadStackPlatform"), STAT_ThreadStackPlatformLLM, STATGROUP_LLMPlatform);
+DECLARE_LLM_MEMORY_STAT(TEXT("Program Size"), STAT_ProgramSizePlatformLLM, STATGROUP_LLMPlatform);
+DECLARE_LLM_MEMORY_STAT(TEXT("Program Size"), STAT_ProgramSizeLLM, STATGROUP_LLMFULL);
+DECLARE_LLM_MEMORY_STAT(TEXT("OOM Backup Pool"), STAT_OOMBackupPoolPlatformLLM, STATGROUP_LLMPlatform);
+DECLARE_LLM_MEMORY_STAT(TEXT("OOM Backup Pool"), STAT_OOMBackupPoolLLM, STATGROUP_LLMFULL);
+DECLARE_LLM_MEMORY_STAT(TEXT("GenericPlatformMallocCrash"), STAT_GenericPlatformMallocCrashLLM, STATGROUP_LLMFULL);
+DECLARE_LLM_MEMORY_STAT(TEXT("GenericPlatformMallocCrash"), STAT_GenericPlatformMallocCrashPlatformLLM, STATGROUP_LLMPlatform);
+DECLARE_LLM_MEMORY_STAT(TEXT("Engine Misc"), STAT_EngineMiscLLM, STATGROUP_LLMFULL);
+DECLARE_LLM_MEMORY_STAT(TEXT("TaskGraph Misc Tasks"), STAT_TaskGraphTasksMiscLLM, STATGROUP_LLMFULL);
+DECLARE_LLM_MEMORY_STAT(TEXT("Linear Allocator"), STAT_LinearAllocatorLLM, STATGROUP_LLMFULL);
+DECLARE_LLM_MEMORY_STAT(TEXT("Audio"), STAT_AudioLLM, STATGROUP_LLMFULL);
+DECLARE_LLM_MEMORY_STAT(TEXT("AudioMisc"), STAT_AudioMiscLLM, STATGROUP_LLMFULL);
+DECLARE_LLM_MEMORY_STAT(TEXT("AudioSoundWaves"), STAT_AudioSoundWavesLLM, STATGROUP_LLMFULL);
+DECLARE_LLM_MEMORY_STAT(TEXT("AudioSoundWaveProxies"), STAT_AudioSoundWaveProxiesLLM, STATGROUP_LLMFULL);
+DECLARE_LLM_MEMORY_STAT(TEXT("AudioMixer"), STAT_AudioMixerLLM, STATGROUP_LLMFULL);
+DECLARE_LLM_MEMORY_STAT(TEXT("AudioMixerPlugins"), STAT_AudioMixerPluginsLLM, STATGROUP_LLMFULL);
+DECLARE_LLM_MEMORY_STAT(TEXT("AudioPrecache"), STAT_AudioPrecacheLLM, STATGROUP_LLMFULL);
+DECLARE_LLM_MEMORY_STAT(TEXT("AudioDecompress"), STAT_AudioDecompressLLM, STATGROUP_LLMFULL);
+DECLARE_LLM_MEMORY_STAT(TEXT("AudioRealtimePrecache"), STAT_AudioRealtimePrecacheLLM, STATGROUP_LLMFULL);
+DECLARE_LLM_MEMORY_STAT(TEXT("AudioFullDecompress"), STAT_AudioFullDecompressLLM, STATGROUP_LLMFULL);
+DECLARE_LLM_MEMORY_STAT(TEXT("AudioStreamCache"), STAT_AudioStreamCacheLLM, STATGROUP_LLMFULL);
+DECLARE_LLM_MEMORY_STAT(TEXT("AudioStreamCacheCompressedData"), STAT_AudioStreamCacheCompressedDataLLM, STATGROUP_LLMFULL);
+DECLARE_LLM_MEMORY_STAT(TEXT("AudioSynthesis"), STAT_AudioSynthesisLLM, STATGROUP_LLMFULL);
+DECLARE_LLM_MEMORY_STAT(TEXT("RealTimeCommunications"), STAT_RealTimeCommunicationsLLM, STATGROUP_LLMFULL);
+DECLARE_LLM_MEMORY_STAT(TEXT("FName"), STAT_FNameLLM, STATGROUP_LLMFULL);
+DECLARE_LLM_MEMORY_STAT(TEXT("Networking"), STAT_NetworkingLLM, STATGROUP_LLMFULL);
+DECLARE_LLM_MEMORY_STAT(TEXT("Meshes"), STAT_MeshesLLM, STATGROUP_LLMFULL);
+DECLARE_LLM_MEMORY_STAT(TEXT("Stats"), STAT_StatsLLM, STATGROUP_LLMFULL);
+DECLARE_LLM_MEMORY_STAT(TEXT("Shaders"), STAT_ShadersLLM, STATGROUP_LLMFULL);
+DECLARE_LLM_MEMORY_STAT(TEXT("PSO"), STAT_PSOLLM, STATGROUP_LLMFULL);
+DECLARE_LLM_MEMORY_STAT(TEXT("Textures"), STAT_TexturesLLM, STATGROUP_LLMFULL);
+DECLARE_LLM_MEMORY_STAT(TEXT("TextureMetaData"), STAT_TextureMetaDataLLM, STATGROUP_LLMFULL);
+DECLARE_LLM_MEMORY_STAT(TEXT("VirtualTextureSystem"), STAT_VirtualTextureSystemLLM, STATGROUP_LLMFULL);
+DECLARE_LLM_MEMORY_STAT(TEXT("Render Targets"), STAT_RenderTargetsLLM, STATGROUP_LLMFULL);
+DECLARE_LLM_MEMORY_STAT(TEXT("SceneRender"), STAT_SceneRenderLLM, STATGROUP_LLMFULL);
+DECLARE_LLM_MEMORY_STAT(TEXT("RHIMisc"), STAT_RHIMiscLLM, STATGROUP_LLMFULL);
+DECLARE_LLM_MEMORY_STAT(TEXT("AsyncLoading"), STAT_AsyncLoadingLLM, STATGROUP_LLMFULL);
+DECLARE_LLM_MEMORY_STAT(TEXT("UObject"), STAT_UObjectLLM, STATGROUP_LLMFULL);
+DECLARE_LLM_MEMORY_STAT(TEXT("Animation"), STAT_AnimationLLM, STATGROUP_LLMFULL);
+DECLARE_LLM_MEMORY_STAT(TEXT("StaticMesh"), STAT_StaticMeshLLM, STATGROUP_LLMFULL);
+DECLARE_LLM_MEMORY_STAT(TEXT("Materials"), STAT_MaterialsLLM, STATGROUP_LLMFULL);
+DECLARE_LLM_MEMORY_STAT(TEXT("Particles"), STAT_ParticlesLLM, STATGROUP_LLMFULL);
+DECLARE_LLM_MEMORY_STAT(TEXT("Niagara"), STAT_NiagaraLLM, STATGROUP_LLMFULL);
+DECLARE_LLM_MEMORY_STAT(TEXT("GPUSort"), STAT_GPUSortLLM, STATGROUP_LLMFULL);
+DECLARE_LLM_MEMORY_STAT(TEXT("GC"), STAT_GCLLM, STATGROUP_LLMFULL);
+DECLARE_LLM_MEMORY_STAT(TEXT("UI"), STAT_UILLM, STATGROUP_LLMFULL);
+DECLARE_LLM_MEMORY_STAT(TEXT("NavigationRecast"), STAT_NavigationRecastLLM, STATGROUP_LLMFULL);
+DECLARE_LLM_MEMORY_STAT(TEXT("Physics"), STAT_PhysicsLLM, STATGROUP_LLMFULL);
+DECLARE_LLM_MEMORY_STAT(TEXT("PhysX"), STAT_PhysXLLM, STATGROUP_LLMFULL);
+DECLARE_LLM_MEMORY_STAT(TEXT("PhysXGeometry"), STAT_PhysXGeometryLLM, STATGROUP_LLMFULL);
+DECLARE_LLM_MEMORY_STAT(TEXT("PhysXLandscape"), STAT_PhysXLandscapeLLM, STATGROUP_LLMFULL);
+DECLARE_LLM_MEMORY_STAT(TEXT("PhysXTrimesh"), STAT_PhysXTrimeshLLM, STATGROUP_LLMFULL);
+DECLARE_LLM_MEMORY_STAT(TEXT("PhysXConvex"), STAT_PhysXConvexLLM, STATGROUP_LLMFULL);
+DECLARE_LLM_MEMORY_STAT(TEXT("PhysXAllocator"), STAT_PhysXAllocatorLLM, STATGROUP_LLMFULL);
+DECLARE_LLM_MEMORY_STAT(TEXT("Chaos"), STAT_ChaosLLM, STATGROUP_LLMFULL);
+DECLARE_LLM_MEMORY_STAT(TEXT("ChaosGeometry"), STAT_ChaosGeometryLLM, STATGROUP_LLMFULL);
+DECLARE_LLM_MEMORY_STAT(TEXT("ChaosAcceleration"), STAT_ChaosAccelerationLLM, STATGROUP_LLMFULL);
+DECLARE_LLM_MEMORY_STAT(TEXT("ChaosParticles"), STAT_ChaosParticlesLLM, STATGROUP_LLMFULL);
+DECLARE_LLM_MEMORY_STAT(TEXT("ChaosLandscape"), STAT_ChaosLandscapeLLM, STATGROUP_LLMFULL);
+DECLARE_LLM_MEMORY_STAT(TEXT("ChaosTrimesh"), STAT_ChaosTrimeshLLM, STATGROUP_LLMFULL);
+DECLARE_LLM_MEMORY_STAT(TEXT("ChaosConvex"), STAT_ChaosConvexLLM, STATGROUP_LLMFULL);
+DECLARE_LLM_MEMORY_STAT(TEXT("ChaosScene"), STAT_ChaosSceneLLM, STATGROUP_LLMFULL);
+DECLARE_LLM_MEMORY_STAT(TEXT("ChaosUpdate"), STAT_ChaosUpdateLLM, STATGROUP_LLMFULL);
+DECLARE_LLM_MEMORY_STAT(TEXT("ChaosActor"), STAT_ChaosActorLLM, STATGROUP_LLMFULL);
+DECLARE_LLM_MEMORY_STAT(TEXT("ChaosBody"), STAT_ChaosBodyLLM, STATGROUP_LLMFULL);
+DECLARE_LLM_MEMORY_STAT(TEXT("ChaosConstraint"), STAT_ChaosConstraintLLM, STATGROUP_LLMFULL);
+DECLARE_LLM_MEMORY_STAT(TEXT("ChaosMaterial"), STAT_ChaosMaterialLLM, STATGROUP_LLMFULL);
+DECLARE_LLM_MEMORY_STAT(TEXT("EnginePreInit"), STAT_EnginePreInitLLM, STATGROUP_LLMFULL);
+DECLARE_LLM_MEMORY_STAT(TEXT("EngineInit"), STAT_EngineInitLLM, STATGROUP_LLMFULL);
+DECLARE_LLM_MEMORY_STAT(TEXT("Rendering Thread"), STAT_RenderingThreadLLM, STATGROUP_LLMFULL);
+DECLARE_LLM_MEMORY_STAT(TEXT("LoadMap Misc"), STAT_LoadMapMiscLLM, STATGROUP_LLMFULL);
+DECLARE_LLM_MEMORY_STAT(TEXT("StreamingManager"), STAT_StreamingManagerLLM, STATGROUP_LLMFULL);
+DECLARE_LLM_MEMORY_STAT(TEXT("Graphics"), STAT_GraphicsPlatformLLM, STATGROUP_LLMPlatform);
+DECLARE_LLM_MEMORY_STAT(TEXT("FileSystem"), STAT_FileSystemLLM, STATGROUP_LLMFULL);
+DECLARE_LLM_MEMORY_STAT(TEXT("Localization"), STAT_LocalizationLLM, STATGROUP_LLMFULL);
+DECLARE_LLM_MEMORY_STAT(TEXT("AssetRegistry"), STAT_AssetRegistryLLM, STATGROUP_LLMFULL);
+DECLARE_LLM_MEMORY_STAT(TEXT("ConfigSystem"), STAT_ConfigSystemLLM, STATGROUP_LLMFULL);
+DECLARE_LLM_MEMORY_STAT(TEXT("InitUObject"), STAT_InitUObjectLLM, STATGROUP_LLMFULL);
+DECLARE_LLM_MEMORY_STAT(TEXT("VideoRecording"), STAT_VideoRecordingLLM, STATGROUP_LLMFULL);
+DECLARE_LLM_MEMORY_STAT(TEXT("Replays"), STAT_ReplaysLLM, STATGROUP_LLMFULL);
+DECLARE_LLM_MEMORY_STAT(TEXT("CsvProfiler"), STAT_CsvProfilerLLM, STATGROUP_LLMFULL);
+DECLARE_LLM_MEMORY_STAT(TEXT("MaterialInstance"), STAT_MaterialInstanceLLM, STATGROUP_LLMFULL);
+DECLARE_LLM_MEMORY_STAT(TEXT("SkeletalMesh"), STAT_SkeletalMeshLLM, STATGROUP_LLMFULL);
+DECLARE_LLM_MEMORY_STAT(TEXT("InstancedMesh"), STAT_InstancedMeshLLM, STATGROUP_LLMFULL);
+DECLARE_LLM_MEMORY_STAT(TEXT("Landscape"), STAT_LandscapeLLM, STATGROUP_LLMFULL);
+DECLARE_LLM_MEMORY_STAT(TEXT("MediaStreaming"), STAT_MediaStreamingLLM, STATGROUP_LLMFULL);
+DECLARE_LLM_MEMORY_STAT(TEXT("ElectraPlayer"), STAT_ElectraPlayerLLM, STATGROUP_LLMFULL);
+DECLARE_LLM_MEMORY_STAT(TEXT("WMFPlayer"), STAT_WMFPlayerLLM, STATGROUP_LLMFULL);
+DECLARE_LLM_MEMORY_STAT(TEXT("MMIO"), STAT_PlatformMMIOLLM, STATGROUP_LLMPlatform);
+DECLARE_LLM_MEMORY_STAT(TEXT("VirtualMemory"), STAT_PlatformVMLLM, STATGROUP_LLMPlatform);
+DECLARE_LLM_MEMORY_STAT(TEXT("CustomName"), STAT_CustomName, STATGROUP_LLMFULL);
+
+namespace UE::LLMPrivate
+{
+
+/** FLLMCsvWriter: class for writing out the LLM tag sizes to a csv file every few seconds. */
+class FLLMCsvWriter
+{
+public:
+	FLLMCsvWriter();
+	~FLLMCsvWriter();
+
+	void SetTracker(ELLMTracker InTracker);
+	void Clear();
+
+	void Publish(FLowLevelMemTracker& LLMRef, const FTrackerTagSizeMap& TagSizes,
+		const FTagData* OverrideTrackedTotalTagData, const FTagData* OverrideUntaggedTagData, int64 TrackedTotal,
+		UE::LLM::ESizeParams SizeParams);
+
+	void OnPreFork();
+
+	void Flush(bool IsOnCrash);
+
+private:
+	void Write(FStringView Text);
+	static const TCHAR* GetTrackerCsvName(ELLMTracker InTracker);
+
+	bool CreateArchive();
+	bool UpdateColumns(const FTrackerTagSizeMap& TagSizes);
+	void WriteHeader(const FTagData* OverrideTrackedTotalTagData, const FTagData* OverrideUntaggedTagData);
+	void AddRow(FLowLevelMemTracker& LLMRef, const FTrackerTagSizeMap& TagSizes,
+		const FTagData* OverrideTrackedTotalName, const FTagData* OverrideUntaggedName, int64 TrackedTotal,
+		UE::LLM::ESizeParams SizeParams);
+
+	FConstTagDataArray Columns;
+	TLLMSet<const FTagData*> ExistingColumns;
+	FArchive* Archive;
+	double LastWriteTime;
+	int32 WriteCount;
+	ELLMTracker Tracker;
+	bool bRegisteredFlushDelegate;
+};
+
+/** Outputs the LLM tags and sizes to TraceLog events. */
+class FLLMTraceWriter
+{
+public:
+	FLLMTraceWriter();
+	void SetTracker(ELLMTracker InTracker);
+	void Clear();
+	void Publish(FLowLevelMemTracker& LLMRef, const FTrackerTagSizeMap& TagSizes,
+		const FTagData* OverrideTrackedTotalName, const FTagData* OverrideUntaggedName, int64 TrackedTotal,
+		UE::LLM::ESizeParams SizeParams);
+
+private:
+	static const void* GetTagId(const FTagData* TagData);
+	void SendTagDeclaration(const FTagData* TagData);
+
+	ELLMTracker Tracker;
+	TLLMSet<const FTagData*> DeclaredTags;
+	bool bTrackerSpecSent = false;
+	bool bTagSetSpecSent = false;
+};
+
+/** FLLMCsvWriter: class for writing out LLM stats to the Csv Profiler. */
+class FLLMCsvProfilerWriter
+{
+public:
+	FLLMCsvProfilerWriter();
+	void Clear();
+	void SetTracker(ELLMTracker InTracker);
+	void Publish(FLowLevelMemTracker& LLMRef, const FTrackerTagSizeMap& TagSizes,
+		const FTagData* OverrideTrackedTotalName, const FTagData* OverrideUntaggedName, int64 TrackedTotal,
+		UE::LLM::ESizeParams SizeParams);
+protected:
+	void RecordTagToCsv(int32 CsvCategoryIndex, const FTagData* TagData, int64 TagSize);
+#if LLM_CSV_PROFILER_WRITER_ENABLED
+	ELLMTracker Tracker;
+	TLLMMap<const FTagData*, FName> TagDataToCsvStatName;
+#endif
+};
+
+/** Per-thread state in an LLMTracker. */
+class FLLMThreadState
+{
+public:
+	FLLMThreadState();
+
+	void Clear();
+
+	void PushTag(const FTagData* TagData, ELLMTagSet TagSet);
+	void PopTag(ELLMTagSet TagSet);
+	const FTagData* GetTopTag(ELLMTagSet TagSet);
+	void TrackAllocation(const void* Ptr, int64 Size, ELLMTracker Tracker, ELLMAllocType AllocType,
+		const FTagData* TagData, const FTagData* AssetTagData, const FTagData* AssetClassTagData, const FTagData* UObjectClassesTagData, bool bTrackInMemPro);
+	void TrackFree(const void* Ptr, int64 Size, ELLMTracker Tracker, ELLMAllocType AllocType,
+		const FTagData* TagData, const FTagData* AssetTagData, const FTagData* AssetClassTagData, const FTagData* UObjectClassesTagData, bool bTrackInMemPro);
+	void TrackMemory(int64 Amount, ELLMTracker Tracker, ELLMAllocType AllocType,
+		const FTagData* TagData);
+	void TrackMoved(const void* Dest, const void* Source, int64 Size, ELLMTracker Tracker, const FTagData* TagData);
+	void IncrTag(const FTagData* Tag, int64 Amount);
+
+	void PropagateChildSizesToParents();
+	void OnTagsResorted(FTagDataArray& OldTagDatas);
+	void LockTags(bool bLock);
+	void FetchAndClearTagSizes(FTrackerTagSizeMap& TagSizes, int64* InAllocTypeAmounts, bool bTrimAllocations);
+
+	void ClearAllocTypeAmounts();
+
+	// AllocationsLock must be held around read/write of Allocations and AllocTypeAmounts, because these are read/write
+	// from the Update thread. All other data on FLLMThreadState are accessed only from the thread that owns the
+	// FLLMThreadState and do not require a lock.
+	// AllocationsLock is also used to guarantee our contract provided to FLLMTracker::LockAllThreadTags: We must not
+	// dereference any FTagData* held by the FLLMThreadState unless we have entered AllocationsLock.
+	FConstTagDataArray TagStack[static_cast<int32>(ELLMTagSet::Max)];
+	FThreadTagSizeMap Allocations;
+	FCriticalSection AllocationsLock;
+
+	int8 PausedCounter[(int32)ELLMAllocType::Count];
+	int64 AllocTypeAmounts[(int32)ELLMAllocType::Count];
+};
+
+/** The main LLM implementation class. It owns the thread state objects. */
+class FLLMTracker
+{
+public:
+
+	FLLMTracker(FLowLevelMemTracker& InLLM);
+	~FLLMTracker();
+
+	void Initialise(ELLMTracker InTracker, FLLMAllocator* InAllocator);
+
+	void PushTag(ELLMTag EnumTag, ELLMTagSet TagSet);
+	void PushTag(FName Tag, bool bInIsStatTag, ELLMTagSet TagSet);
+	void PushTag(const FTagData* TagData, ELLMTagSet TagSet);
+	void PopTag(ELLMTagSet TagSet);
+	void TrackAllocation(const void* Ptr, int64 Size, FName DefaultTag, ELLMAllocType AllocType, bool bTrackInMemPro);
+	void TrackAllocation(const void* Ptr, int64 Size, ELLMTag DefaultTag, ELLMAllocType AllocType, bool bTrackInMemPro);
+	void TrackFree(const void* Ptr, ELLMAllocType AllocType, bool bTrackInMemPro);
+	void TrackMemoryOfActiveTag(int64 Amount, FName DefaultTag, ELLMAllocType AllocType);
+	void TrackMemoryOfActiveTag(int64 Amount, ELLMTag DefaultTag, ELLMAllocType AllocType);
+	void OnAllocMoved(const void* Dest, const void* Source, ELLMAllocType AllocType);
+
+	void TrackMemory(ELLMTag EnumTag, int64 Amount, ELLMAllocType AllocType);
+	void TrackMemory(FName TagName, ELLMTagSet TagSet, int64 Amount, ELLMAllocType AllocType);
+	void TrackMemory(const FTagData* TagData, int64 Amount, ELLMAllocType AllocType);
+
+	// This will pause/unpause tracking, and also manually increment a given tag
+	void PauseAndTrackMemory(FName TagName, ELLMTagSet TagSet, bool bInIsStatTag, int64 Amount, ELLMAllocType AllocType);
+	void PauseAndTrackMemory(ELLMTag EnumTag, int64 Amount, ELLMAllocType AllocType);
+	void PauseAndTrackMemory(const FTagData* TagData, int64 Amount, ELLMAllocType AllocType);
+	void Pause(ELLMAllocType AllocType);
+	void Unpause(ELLMAllocType AllocType);
+	bool IsPaused(ELLMAllocType AllocType);
+
+	void Clear();
+
+	// Dump the allocation count and size for each tag, along with the proportion that is private/shared/unreferenced
+	// on linux machines. This will dump to a CSV in the project Saved/LLM directory with a separate file for each
+	// tracker and forked child. Integers will be appended to the filename to prevent overwrites. Returns false
+	// if the memory information can't be retrieved (likely because it's not on a supported platform)
+	bool DumpForkedAllocationInfo();
+
+	// Dump the list of all registered LLM tags and their properties to a CSV file, in the LLM's profiling directory
+	// (Saved/Profiling/LLM/); see UE::LLM::GetLLMProfilingDir().
+	bool DumpTags();
+
+	void PublishStats(UE::LLM::ESizeParams SizeParam);
+	void PublishCsv(UE::LLM::ESizeParams SizeParam);
+	void PublishTrace(UE::LLM::ESizeParams SizeParam);
+	void PublishCsvProfiler(UE::LLM::ESizeParams SizeParam);
+
+	void OnPreFork();
+
+	struct FLowLevelAllocInfo
+	{
+	public:
+		FLowLevelAllocInfo()
+		{
+#if LLM_ENABLED_FULL_TAGS && (LLM_ALLOW_ASSETS_TAGS || LLM_ALLOW_UOBJECTCLASSES_TAGS)
+			FMemory::Memset(Tag, InvalidCompressedTagValue);
+#endif
+		}
+
+		FORCEINLINE static int32 GetIndexFromTagSet(ELLMTagSet InTagSet)
+		{
+			static_assert(static_cast<int32>(ELLMTagSet::Max) == 4,
+				"You added a new tagset, please update the code bellow");
+
+#if LLM_ALLOW_ASSETS_TAGS
+			switch (InTagSet)
+			{
+				case ELLMTagSet::None: return 0;
+				case ELLMTagSet::Assets: return 1;
+				case ELLMTagSet::AssetClasses: return 2;
+				default:
+					LLMCheck(false);
+					return 0;
+			}
+#elif LLM_ALLOW_UOBJECTCLASSES_TAGS
+			switch (InTagSet)
+			{
+				case ELLMTagSet::None: return 0;
+				case ELLMTagSet::UObjectClasses: return 1;
+				default:
+					LLMCheck(false);
+					return 0;
+			}
+#endif
+		}
+
+		void SetTag(const FTagData* InTag, FLowLevelMemTracker& InLLMRef
+#if LLM_ALLOW_ASSETS_TAGS || LLM_ALLOW_UOBJECTCLASSES_TAGS
+		, ELLMTagSet InTagSet = ELLMTagSet::None
+#endif
+		)
+		{
+#if LLM_ENABLED_FULL_TAGS
+#if !LLM_ALLOW_ASSETS_TAGS && !LLM_ALLOW_UOBJECTCLASSES_TAGS
+			Tag = InTag->GetIndex();
+#else
+			int32 TagSetIndex = GetIndexFromTagSet(InTagSet);
+			Tag[TagSetIndex] = InTag ? InTag->GetIndex() : InvalidCompressedTagValue;
+#endif //!LLM_ALLOW_ASSETS_TAGS
+#else
+			Tag = InTag->GetEnumTag();
+#endif
+		}
+
+		const FTagData* GetTag(FLowLevelMemTracker& InLLMRef
+#if LLM_ALLOW_ASSETS_TAGS || LLM_ALLOW_UOBJECTCLASSES_TAGS
+		, ELLMTagSet InTagSet = ELLMTagSet::None
+#endif
+		) const
+		{
+#if LLM_ENABLED_FULL_TAGS
+			UE::TSharedLock TagDataScopeLock(InLLMRef.TagDataMutex);
+#if !LLM_ALLOW_ASSETS_TAGS && !LLM_ALLOW_UOBJECTCLASSES_TAGS
+			return (*InLLMRef.TagDatas)[Tag];
+#else
+			int32 TagSetIndex = GetIndexFromTagSet(InTagSet);
+			return Tag[TagSetIndex] >= 0 ? (*InLLMRef.TagDatas)[Tag[TagSetIndex]] : nullptr;
+#endif
+#else
+			return InLLMRef.FindTagData(Tag);
+#endif
+		}
+
+#if LLM_ENABLED_FULL_TAGS
+		int32 GetCompressedTag(ELLMTagSet InTagSet = ELLMTagSet::None)
+		{
+#if !LLM_ALLOW_ASSETS_TAGS && !LLM_ALLOW_UOBJECTCLASSES_TAGS
+			return (InTagSet == ELLMTagSet::None) ? Tag : InvalidCompressedTagValue;
+#else
+			int32 TagSetIndex = GetIndexFromTagSet(InTagSet);
+			return Tag[TagSetIndex];
+#endif
+		}
+
+		void SetCompressedTag(int32 InTag, ELLMTagSet InTagSet = ELLMTagSet::None)
+		{
+#if !LLM_ALLOW_ASSETS_TAGS && !LLM_ALLOW_UOBJECTCLASSES_TAGS
+			if (InTagSet == ELLMTagSet::None)
+			{
+				Tag = InTag;
+			}
+#else
+			int32 TagSetIndex = GetIndexFromTagSet(InTagSet);
+			Tag[TagSetIndex] = InTag;
+#endif
+		}
+#else
+		ELLMTag GetCompressedTag() { return Tag; }
+		void SetCompressedTag(ELLMTag InTag) { Tag = InTag; }
+#endif
+
+	private:
+#if LLM_ENABLED_FULL_TAGS
+		// Even with arbitrary tags we are still partially compressed - the allocation records the tag's index
+		// (4 bytes) rather than the full tag pointer (8 bytes) or the tag's name (12 bytes).
+#if !LLM_ALLOW_ASSETS_TAGS && !LLM_ALLOW_UOBJECTCLASSES_TAGS
+		int32 Tag = 0;
+#else
+		int32 Tag[LLM_TAGSET_MAX];
+
+		static_assert(LLM_TAGSET_MAX <= 3,
+			"Incrementing the maximum amount of TagSet enums that can be defined at the same time adds 4 bytes of LLM overhead memory FOR EVERY ACTIVE ALLOCATION (e.g. 10 million * 4 == 40 million bytes). "
+			"Do not add them without measuring the cost of the overhead in the cook and runtime of large projects.");
+#endif // LLM_ALLOW_ASSETS_TAGS
+#else
+		ELLMTag Tag = ELLMTag::Untagged;
+#endif
+		static const int32 InvalidCompressedTagValue = -1;
+	};
+
+	typedef LLMMap<PointerKey, uint32, FLowLevelAllocInfo, LLMNumAllocsType> FLLMAllocMap;  // pointer, size, info, Capacity SizeType
+
+	void SetTotalTags(const FTagData* InOverrideUntaggedTagData, const FTagData* InOverrideTrackedTotalTagData);
+	void Update();
+	void UpdateThreads();
+	void OnTagsResorted(FTagDataArray& OldTagDatas);
+	void LockAllThreadTags(bool bLock);
+
+	void CaptureTagSnapshot();
+	void ClearTagSnapshot();
+
+	int64 GetTagAmount(const FTagData* TagData, UE::LLM::ESizeParams SizeParams = UE::LLM::ESizeParams::Default) const;
+	void SetTagAmountExternal(const FTagData* TagData, int64 Amount, bool bAddToTotal);
+	void SetTagAmountInUpdate(const FTagData* TagData, int64 Amount, bool bAddToTotal);
+	const FTagData* GetActiveTagData(ELLMTagSet TagSet = ELLMTagSet::None);
+	TArray<const FTagData*> GetTagDatas(ELLMTagSet TagSet = ELLMTagSet::None);
+
+	void GetTagsNamesWithAmount(TMap<FName, uint64>& OutTagsNamesWithAmount, ELLMTagSet TagSet = ELLMTagSet::None);
+	void GetTagsNamesWithAmountFiltered(TMap<FName, uint64>& OutTagsNamesWithAmount, ELLMTagSet TagSet, TArray<FLLMTagSetAllocationFilter>& Filters);
+
+	bool FindTagsForPtr(void* InPtr, TArray<const FTagData*, TInlineAllocator<static_cast<int32>(ELLMTagSet::Max)>>& OutTags) const;
+
+	int64 GetAllocTypeAmount(ELLMAllocType AllocType);
+
+	int64 GetTrackedTotal(UE::LLM::ESizeParams SizeParams = UE::LLM::ESizeParams::Default) const
+	{
+		if (EnumHasAnyFlags(SizeParams, UE::LLM::ESizeParams::RelativeToSnapshot))
+		{
+			return FMath::Clamp<int64>(TrackedTotal - TrackedTotalInSnapshot, 0, INT64_MAX); //-V569
+		}
+		return TrackedTotal;
+	}
+
+protected:
+	FLLMThreadState* GetOrCreateState();
+	FLLMThreadState* GetState();
+	void TrackAllocation(const void* Ptr, int64 Size, const FTagData* ActiveTagData, ELLMAllocType AllocType, FLLMThreadState* State, bool bTrackInMemPro);
+	void TrackMemoryOfActiveTag(int64 Amount, const FTagData* TagData, ELLMAllocType AllocType, FLLMThreadState* State);
+
+	FLowLevelMemTracker& LLMRef;
+
+	ELLMTracker Tracker;
+
+	uint32 TlsSlot;
+
+	TArray<FLLMThreadState*, FDefaultLLMAllocator> ThreadStates;
+
+	FCriticalSection PendingThreadStatesGuard;
+	TArray<FLLMThreadState*, FDefaultLLMAllocator> PendingThreadStates;
+	/**
+	 * Backup map from thread to threadstate. The primary lookup method is FPlatformTLS, but that is unavailable
+	 * during thread termination on some platforms. When unavailable, we use this TMap (under PendingThreadStatesGuard
+	 * critical section) to find the state.
+	 */
+	TMap<uint32, FLLMThreadState*, FDefaultSetLLMAllocator> ThreadIdToThreadState;
+
+	/** Sum of memory from all tracked tags. Duplicated in separate storage to make it instantly available at any time of frame without waiting for accumulation from threads during update.
+		GCC_ALIGN is required because it is modified in FPlatformAtomics::InterlockedAdd, which requires aligned values. */
+	int64 TrackedTotal GCC_ALIGN(8);
+
+	// The total tracked memory when the snapshot was taken
+	int64 TrackedTotalInSnapshot;
+
+#if !UE_ONLY_USE_PLATFORM_TRACKER
+	FLLMAllocMap AllocationMap;
+#endif
+
+	FTrackerTagSizeMap TagSizes;
+
+	const FTagData* OverrideUntaggedTagData;
+	const FTagData* OverrideTrackedTotalTagData;
+
+	FLLMCsvWriter CsvWriter;
+	FLLMTraceWriter TraceWriter;
+	FLLMCsvProfilerWriter CsvProfilerWriter;
+
+	double LastTrimTime;
+
+	int64 AllocTypeAmounts[(int32)ELLMAllocType::Count];
+};
+
+const TCHAR* ToString(ETagReferenceSource ReferenceSource);
+void SetMemoryStatByFName(FName Name, int64 Amount);
+void ValidateUniqueName(FStringView UniqueName);
+
+typedef TArray<TPair<FLLMInitialisedCallback, UPTRINT>, TInlineAllocator<1>> FInitializedCallbacksArray;
+FInitializedCallbacksArray& GetInitialisedCallbacks()
+{
+	static FInitializedCallbacksArray Array;
+	return Array;
+}
+
+typedef TArray<TPair<FTagCreationCallback, UPTRINT>, TInlineAllocator<1>> FTagCreationCallbacksArray;
+FTagCreationCallbacksArray& GetTagCreationCallbacks()
+{
+	static FTagCreationCallbacksArray Array;
+	return Array;
+}
+
+void FPrivateCallbacks::AddInitialisedCallback(FLLMInitialisedCallback Callback, UPTRINT UserData)
+{
+	FLowLevelMemTracker::Get().BootstrapInitialise();
+
+	FInitializedCallbacksArray& Callbacks = UE::LLMPrivate::GetInitialisedCallbacks();
+	Callbacks.Add(TPair<FLLMInitialisedCallback, UPTRINT> { Callback, UserData });
+}
+
+void FPrivateCallbacks::AddTagCreationCallback(FTagCreationCallback Callback, UPTRINT UserData)
+{
+	FLowLevelMemTracker::Get().BootstrapInitialise();
+
+	FTagCreationCallbacksArray& Callbacks = UE::LLMPrivate::GetTagCreationCallbacks();
+	Callbacks.Add(TPair<FTagCreationCallback, UPTRINT> { Callback, UserData });
+}
+
+void FPrivateCallbacks::RemoveTagCreationCallback(FTagCreationCallback Callback)
+{
+	FLowLevelMemTracker::Get().BootstrapInitialise();
+
+	FTagCreationCallbacksArray& Callbacks = UE::LLMPrivate::GetTagCreationCallbacks();
+	Callbacks.RemoveAll([Callback](TPair<FTagCreationCallback, UPTRINT>& Pair) { return Pair.Key == Callback; });
+}
+
+} // namespace UE::LLMPrivate
+
+static FName GetTagName_CustomName()
+{
+	// This FName can be read before c++ global constructors are complete, by LLM_SCOPE_BY_BOOTSTRAP_TAG
+	// so it needs to be a function static rather than a global.
+	static FName TagName_CustomName("CustomName");
+	return TagName_CustomName;
+}
+static FName TagName_Untagged("Untagged");
+static FName TagName_UntaggedAsset("UntaggedAsset");
+static FName TagName_UntaggedAssetClass("UntaggedAssetClass");
+static FName TagName_UntaggedUObjectClasses("UntaggedUObjectClasses");
+
+FName LLMGetTagUniqueName(ELLMTag Tag)
+{
+#define LLM_TAG_NAME_ARRAY(Enum,Str,Stat,Group,ParentTag) FName(TEXT(#Enum)),
+	static const FName UniqueNames[] = { LLM_ENUM_GENERIC_TAGS(LLM_TAG_NAME_ARRAY) };
+#undef LLM_TAG_NAME_ARRAY
+
+	int32 Index = static_cast<int32>(Tag);
+	if (Index < 0)
+	{
+		return NAME_None;
+	}
+	else if (Index < UE_ARRAY_COUNT(UniqueNames))
+	{
+		return UniqueNames[Index];
+	}
+	else if (Index < LLM_CUSTOM_TAG_START)
+	{
+		return NAME_None;
+	}
+	else if (Index <= LLM_CUSTOM_TAG_END)
+	{
+		static FName CustomNames[LLM_CUSTOM_TAG_COUNT];
+		static bool bCustomNamesInitialized = false;
+		if (!bCustomNamesInitialized)
+		{
+			TStringBuilder<256> UniqueNameBuffer;
+
+			for (int32 CreateIndex = LLM_CUSTOM_TAG_START; CreateIndex <= LLM_CUSTOM_TAG_END; ++CreateIndex)
+			{
+				UniqueNameBuffer.Reset();
+				UniqueNameBuffer.Appendf(TEXT("ELLMTag%d"), CreateIndex);
+				CustomNames[CreateIndex - LLM_CUSTOM_TAG_START] = FName(UniqueNameBuffer);
+			}
+		}
+		return CustomNames[Index - LLM_CUSTOM_TAG_START];
+	}
+	else
+	{
+		return NAME_None;
+	}
+}
+
+extern const TCHAR* LLMGetTagName(ELLMTag Tag)
+{
+#define LLM_TAG_NAME_ARRAY(Enum,Str,Stat,Group,ParentTag) TEXT(Str),
+	static TCHAR const* Names[] = { LLM_ENUM_GENERIC_TAGS(LLM_TAG_NAME_ARRAY) };
+#undef LLM_TAG_NAME_ARRAY
+
+	int32 Index = static_cast<int32>(Tag);
+	if (Index >= 0 && Index < UE_ARRAY_COUNT(Names))
+	{
+		return Names[Index];
+	}
+	else
+	{
+		return nullptr;
+	}
+}
+
+extern const FName LLMGetUntaggedTagName(ELLMTagSet TagSet)
+{
+	switch (TagSet)
+	{
+		case ELLMTagSet::None:
+		{
+			return TagName_Untagged;
+		}
+#if LLM_ALLOW_ASSETS_TAGS
+		case ELLMTagSet::Assets:
+		{
+			return TagName_UntaggedAsset;
+		}
+		case ELLMTagSet::AssetClasses:
+		{
+			return TagName_UntaggedAssetClass;
+		}
+#elif LLM_ALLOW_UOBJECTCLASSES_TAGS
+		case ELLMTagSet::UObjectClasses:
+		{
+			return TagName_UntaggedUObjectClasses;
+		}
+#endif // LLM_ALLOW_UOBJECTCLASSES_TAGS
+		default:
+		{
+			return NAME_None;
+		}
+	}
+}
+
+extern const ANSICHAR* LLMGetTagNameANSI(ELLMTag Tag)
+{
+#define LLM_TAG_NAME_ARRAY(Enum,Str,Stat,Group,ParentTag) Str,
+	static ANSICHAR const* Names[] = { LLM_ENUM_GENERIC_TAGS(LLM_TAG_NAME_ARRAY) };
+#undef LLM_TAG_NAME_ARRAY
+
+	int32 Index = static_cast<int32>(Tag);
+	if (Index >= 0 && Index < UE_ARRAY_COUNT(Names))
+	{
+		return Names[Index];
+	}
+	else
+	{
+		return nullptr;
+	}
+}
+
+FLowLevelMemTracker& FLowLevelMemTracker::Construct()
+{
+	static FLowLevelMemTracker Tracker;
+	TrackerInstance = &Tracker;
+	return Tracker;
+}
+
+FLowLevelMemTracker* FLowLevelMemTracker::TrackerInstance = nullptr;
+// LLM must start off enabled because allocations happen before the command line enables/disables us
+bool FLowLevelMemTracker::bIsDisabled = false;
+FLowLevelMemTracker::EEnabled FLowLevelMemTracker::EnabledState = EEnabled::NotYetKnown;
+
+static const TCHAR* InvalidLLMTagName = TEXT("?");
+
+FLowLevelMemTracker::FLowLevelMemTracker()
+	: TagDatas(nullptr)
+	, TagDataNameMap(nullptr)
+	, TagDataEnumMap(nullptr)
+	, TraceTagFilter(nullptr)
+	, ProgramSize(0)
+	, MemoryUsageCurrentOverhead(0)
+	, MemoryUsagePlatformTotalUntracked(0)
+	, bFirstTimeUpdating(true)
+	, bCanEnable(true)
+	, bCsvWriterEnabled(false)
+	, bTraceWriterEnabled(false)
+	, bInitialisedTracking(false)
+	, bIsBootstrapping(false)
+	, bFullyInitialised(false)
+	, bConfigurationComplete(false)
+	, bTagAdded(false)
+	, bAutoPublish(true)
+	, bPublishSingleFrame(false)
+	, bCapturedSizeSnapshot(false)
+{
+	using namespace UE::LLMPrivate;
+
+	// set the alloc functions
+	LLMAllocFunction PlatformLLMAlloc = NULL;
+	LLMFreeFunction PlatformLLMFree = NULL;
+	int32 Alignment = 0;
+	if (!FPlatformMemory::GetLLMAllocFunctions(PlatformLLMAlloc, PlatformLLMFree, Alignment))
+	{
+		EnabledState = EEnabled::Disabled;
+		PRAGMA_DISABLE_DEPRECATION_WARNINGS;
+		bIsDisabled = true;
+		PRAGMA_ENABLE_DEPRECATION_WARNINGS;
+		bCanEnable = false;
+		bConfigurationComplete = true;
+		return;
+	}
+	LLMCheck(FMath::IsPowerOfTwo(Alignment));
+
+	Allocator.Initialise(PlatformLLMAlloc, PlatformLLMFree, Alignment);
+	FLLMAllocator::Get() = &Allocator;
+
+	// only None is on by default
+	for (int32 Index = 0; Index < static_cast<int32>(ELLMTagSet::Max); Index++)
+	{
+		ActiveSets[Index] = Index == static_cast<int32>(ELLMTagSet::None);
+	}
+}
+
+FLowLevelMemTracker::~FLowLevelMemTracker()
+{
+	using namespace UE::LLMPrivate;
+
+	EnabledState = EEnabled::Disabled;
+	PRAGMA_DISABLE_DEPRECATION_WARNINGS;
+	bIsDisabled = true;
+	PRAGMA_ENABLE_DEPRECATION_WARNINGS;
+	Clear();
+	FLLMAllocator::Get() = nullptr;
+}
+
+void FLowLevelMemTracker::BootstrapInitialise()
+{
+	using namespace UE::LLMPrivate;
+	if (bInitialisedTracking)
+	{
+		return;
+	}
+	UE::TUniqueLock BootstrapLock(BootstrapMutex);
+	if (bInitialisedTracking)
+	{
+		return;
+	}
+	// just make sure this part isn't reentrant.
+	if (bBeganInitialisedTracking != false)
+	{
+		UE_LOG(LogInit, Fatal, TEXT("FLowLevelMemTracker::BootstrapInitialise is not reentrant"));
+	}
+	bBeganInitialisedTracking = true;
+
+	for (int32 TrackerIndex = 0; TrackerIndex < static_cast<int32>(ELLMTracker::Max); ++TrackerIndex)
+	{
+		FLLMTracker* Tracker = Allocator.New<FLLMTracker>(*this);
+
+		Trackers[TrackerIndex] = Tracker;
+
+		Tracker->Initialise(static_cast<ELLMTracker>(TrackerIndex), &Allocator);
+	}
+
+	BootstrapTagDatas();
+	static_assert((uint8)ELLMTracker::Max == 2,
+		"You added a tracker, without updating FLowLevelMemTracker::BootstrapInitialise (and probably need to update macros)");
+	GetTracker(ELLMTracker::Platform)->SetTotalTags(FindOrAddTagData(ELLMTag::PlatformUntaggedTotal),
+		FindOrAddTagData(ELLMTag::PlatformTrackedTotal));
+	GetTracker(ELLMTracker::Default)->SetTotalTags(FindOrAddTagData(ELLMTag::UntaggedTotal),
+		FindOrAddTagData(ELLMTag::TrackedTotal));
+
+
+	// calculate program size early on... the platform can call SetProgramSize later if it sees fit
+	InitialiseProgramSize();
+
+	FPlatformMisc::MemoryBarrier();
+
+	bInitialisedTracking = true;
+}
+
+void FLowLevelMemTracker::Clear()
+{
+	using namespace UE::LLMPrivate;
+
+	if (!bInitialisedTracking)
+	{
+		return;
+	}
+
+	LLMCheck(!IsEnabled()); // tracking must be stopped at this point or it will crash while tracking its own destruction
+	for (int32 TrackerIndex = 0; TrackerIndex < static_cast<int32>(ELLMTracker::Max); TrackerIndex++)
+	{
+		GetTracker((ELLMTracker)TrackerIndex)->Clear();
+		Allocator.Delete(Trackers[TrackerIndex]);
+		Trackers[TrackerIndex] = nullptr;
+	}
+
+	if (TraceTagFilter)
+	{
+		Allocator.Delete(TraceTagFilter);
+		TraceTagFilter = nullptr;
+	}
+
+	ClearTagDatas();
+	Allocator.Clear();
+	bFullyInitialised = false;
+	bInitialisedTracking = false;
+}
+
+void FLowLevelMemTracker::OnPreFork()
+{
+	using namespace UE::LLMPrivate;
+	if (IsEnabled())
+	{
+		FLLMTracker& DefaultTracker = *GetTracker(ELLMTracker::Default);
+		FLLMTracker& PlatformTracker = *GetTracker(ELLMTracker::Platform);
+
+		DefaultTracker.OnPreFork();
+		PlatformTracker.OnPreFork();
+	}
+}
+
+void FLowLevelMemTracker::UpdateStatsPerFrame(const TCHAR* LogName)
+{
+	CSV_SCOPED_TIMING_STAT_EXCLUSIVE(LLM);
+#if UE_ENABLE_ARRAY_SLACK_TRACKING
+	// Slack tracking, when compiled in, can run even when regular LLM tracking is disabled
+	LlmTrackArrayTick();
+#endif
+
+	if (!IsEnabled())
+	{
+		if (bFirstTimeUpdating)
+		{
+			// UpdateStatsPerFrame is usually called from the game thread, but can sometimes be called from the
+			// async loading thread, so enter a lock for it
+			UE::TUniqueLock UpdateLock(UpdateMutex);
+			if (bFirstTimeUpdating)
+			{
+				// Write the saved overhead value to the stats system; this allows us to see the overhead that is
+				// always there even when disabled (unless the #define completely removes support, of course)
+				bFirstTimeUpdating = false;
+				// Don't call Update; Trackers no longer exist. But do publish the recorded values.
+				PublishDataPerFrame(LogName);
+			}
+		}
+		return;
+	}
+
+	// UpdateStatsPerFrame is usually called from the game thread, but can sometimes be called from the
+	// async loading thread, so enter a lock for it
+	UE::TUniqueLock UpdateLock(UpdateMutex);
+	BootstrapInitialise();
+
+	if (bFirstTimeUpdating)
+	{
+		// Nothing needed here yet
+		UE_LOG(LogInit, Log, TEXT("First time updating LLM stats..."));
+		bFirstTimeUpdating = false;
+	}
+	TickInternal();
+
+	if (bAutoPublish || bPublishSingleFrame)
+	{
+		PublishDataPerFrame(LogName);
+		bPublishSingleFrame = false;
+	}
+}
+
+void FLowLevelMemTracker::Tick()
+{
+	if (!IsEnabled())
+	{
+		return;
+	}
+	// TickInternal is usually called from the game thread, but can sometimes be called from the async loading thread,
+	// so enter a lock for it.
+	UE::TUniqueLock UpdateLock(UpdateMutex);
+	TickInternal();
+}
+
+void FLowLevelMemTracker::TickInternal()
+{
+	using namespace UE::LLMPrivate;
+
+	if (bFullyInitialised)
+	{
+		// We call tick when not fully initialised to get the overhead when disabled. When not initialised, we have to
+		// avoid the portion of the tick that uses tags.
+
+		// Get the platform to update any custom tags; this must be done before the aggregation that occurs in
+		// GetTracker()->Update.
+		FPlatformMemory::UpdateCustomLLMTags();
+
+		FCoreDelegates::OnLowLevelMemTrackerTick.Broadcast();
+
+		UpdateTags();
+
+		// update the trackers
+		for (int32 TrackerIndex = 0; TrackerIndex < static_cast<int32>(ELLMTracker::Max); TrackerIndex++)
+		{
+			GetTracker((ELLMTracker)TrackerIndex)->Update();
+		}
+	}
+	FLLMTracker& DefaultTracker = *Trackers[static_cast<int32>(ELLMTracker::Default)];
+	FLLMTracker& PlatformTracker = *Trackers[static_cast<int32>(ELLMTracker::Platform)];
+
+	const int64 TrackedTotal = DefaultTracker.GetTrackedTotal();
+
+	// Cache the amount of memory used early, since some of these functions (FindOrAddTagData) can
+	// cause allocations, which will throw the numbers off slightly.
+	FPlatformMemoryStats PlatformStats = FPlatformMemory::GetStats();
+#if PLATFORM_DESKTOP
+	// virtual is working set + paged out memory.
+	const int64 PlatformProcessMemory = static_cast<int64>(PlatformStats.UsedVirtual);
+#elif PLATFORM_ANDROID
+	// On some Android devices total used mem (VmRss + VmSwap) does not include GPU memory allocations
+	// and so it's going to be lower than TrackedTotal
+	const int64 PlatformProcessMemory = FMath::Max((int64)PlatformStats.UsedPhysical, TrackedTotal);
+
+	PlatformTracker.SetTagAmountInUpdate(FindOrAddTagData(ELLMTag::PlatformMMIO), (int64)PlatformStats.RssFile, false);
+#else
+	const int64 PlatformProcessMemory = static_cast<int64>(PlatformStats.UsedPhysical);
+#endif
+	// Update how much overhead LLM is using. Note we used to also added sizeof(FLowLevelMemTracker), but this
+	// isn't allocated and is in the data segment. So adding it throws the numbers off.
+	MemoryUsageCurrentOverhead = Allocator.GetTotal();
+	PlatformTracker.SetTagAmountInUpdate(FindOrAddTagData(ELLMTag::PlatformOverhead), MemoryUsageCurrentOverhead, !FPlatformMemory::TracksLLMAllocations());
+
+	// Calculate FMalloc unused stat and set it in the Default tracker.
+	const int64 FMallocAmount = DefaultTracker.GetAllocTypeAmount(ELLMAllocType::FMalloc);
+	const int64 FMallocPlatformAmount = PlatformTracker.GetTagAmount(FindOrAddTagData(ELLMTag::FMalloc));
+	const int64 CachedFreeMemory = UE::Private::GMalloc ? UE::Private::GMalloc->GetTotalFreeCachedMemorySize() : 0;
+	int64 FMallocUnused = FMallocPlatformAmount - FMallocAmount - CachedFreeMemory;
+	if (FMallocPlatformAmount == 0)
+	{
+		// We do not have instrumentation for this allocator, and so can not calculate how much memory it is using
+		// internally. Set unused to 0 for this case.
+		FMallocUnused = 0;
+	}
+	DefaultTracker.SetTagAmountInUpdate(FindOrAddTagData(ELLMTag::FMallocUnused), FMallocUnused, true);
+
+	// Determine the UnusedRHI amount by finding the difference between the Platform and Default tracker.
+	int64 PlatformRHIAmount = PlatformTracker.GetAllocTypeAmount(ELLMAllocType::RHI);
+	int64 DefaultRHIAmount = DefaultTracker.GetAllocTypeAmount(ELLMAllocType::RHI);
+	int64 UnusedRHIAmount = PlatformRHIAmount - DefaultRHIAmount;
+	DefaultTracker.SetTagAmountInUpdate(FindOrAddTagData(ELLMTag::RHIUnused), UnusedRHIAmount, true);
+
+	// Compare memory the platform thinks we have allocated to what we have tracked, including the program memory
+	const int64 PlatformTrackedTotal = PlatformTracker.GetTrackedTotal();
+	MemoryUsagePlatformTotalUntracked = FMath::Max<int64>(0, PlatformProcessMemory - PlatformTrackedTotal);
+
+	PlatformTracker.SetTagAmountInUpdate(FindOrAddTagData(ELLMTag::PlatformTotal), PlatformProcessMemory, false);
+	PlatformTracker.SetTagAmountInUpdate(FindOrAddTagData(ELLMTag::PlatformUntracked), MemoryUsagePlatformTotalUntracked, false);
+	PlatformTracker.SetTagAmountInUpdate(FindOrAddTagData(ELLMTag::PlatformOSAvailable), PlatformStats.AvailablePhysical, false);
+
+	// remove the MemoryUsageCurrentOverhead from the "Total" for the default LLM as it's not something anyone needs
+	// to investigate when finding what to reduce the platform LLM will have the info
+	const int64 DefaultProcessMemory = PlatformProcessMemory - (!FPlatformMemory::TracksLLMAllocations() ? MemoryUsageCurrentOverhead : 0);
+	const int64 ImmediatelyFreeableCachedMemory = UE::Private::GMalloc ? UE::Private::GMalloc->GetImmediatelyFreeableCachedMemorySize() : 0;
+	// we need to subtract diff between total cached free mem and immediately freeable cached memory as it is not accounted anywhere in the LLM and would artificially inflate Untracked.
+	// while any new allocations from that cached, but not immediately freeable memory would be properly accounted by the LLM
+	const int64 DefaultUntracked = FMath::Max<int64>(0, DefaultProcessMemory - TrackedTotal - CachedFreeMemory + ImmediatelyFreeableCachedMemory);
+	DefaultTracker.SetTagAmountInUpdate(FindOrAddTagData(ELLMTag::Total), DefaultProcessMemory, false);
+	DefaultTracker.SetTagAmountInUpdate(FindOrAddTagData(ELLMTag::Untracked), DefaultUntracked, false);
+
+#if PLATFORM_WINDOWS
+	DefaultTracker.SetTagAmountInUpdate(FindOrAddTagData(ELLMTag::WorkingSetSize), PlatformStats.UsedPhysical, false);
+	DefaultTracker.SetTagAmountInUpdate(FindOrAddTagData(ELLMTag::PagefileUsed), PlatformStats.UsedVirtual, false);
+#endif
+}
+
+void FLowLevelMemTracker::UpdateTags()
+{
+	using namespace UE::LLMPrivate;
+
+	if (!bTagAdded)
+	{
+		return;
+	}
+
+	bTagAdded = false;
+	bool bNeedsResort = false;
+	{
+		UE::TSharedLock Lock(TagDataMutex);
+		// Cannot use a ranged-for because FinishConstruct can drop our lock and add elements to TagDatas
+		// Check the valid-index condition on every loop iteration
+		for (int32 TagIndex = 0; TagIndex < TagDatas->Num(); ++TagIndex)
+		{
+			FTagData* TagData = (*TagDatas)[TagIndex];
+			FinishConstruct(TagData, UE::LLMPrivate::ETagReferenceSource::FunctionAPI);
+			const FTagData* Parent = TagData->GetParent();
+			if (Parent && Parent->GetIndex() > TagData->GetIndex())
+			{
+				bNeedsResort = true;
+			}
+		}
+	}
+	if (bNeedsResort)
+	{
+		// Prevent threads from reading their tags while we are mutating tags
+		for (int32 TrackerIndex = 0; TrackerIndex < static_cast<int32>(ELLMTracker::Max); TrackerIndex++)
+		{
+			GetTracker(static_cast<ELLMTracker>(TrackerIndex))->LockAllThreadTags(true);
+		}
+
+		FTagDataArray* OldTagDatas;
+		UE::TUniqueLock Lock(TagDataMutex);
+		SortTags(OldTagDatas);
+
+		for (int32 TrackerIndex = 0; TrackerIndex < static_cast<int32>(ELLMTracker::Max); TrackerIndex++)
+		{
+			GetTracker(static_cast<ELLMTracker>(TrackerIndex))->OnTagsResorted(*OldTagDatas);
+			GetTracker(static_cast<ELLMTracker>(TrackerIndex))->LockAllThreadTags(false);
+		}
+
+		Allocator.Delete(OldTagDatas);
+	}
+}
+
+void FLowLevelMemTracker::SortTags(UE::LLMPrivate::FTagDataArray*& OutOldTagDatas)
+{
+	using namespace UE::LLMPrivate;
+
+	// Caller is responsible for holding a UniqueLock on TagDataMutex
+	OutOldTagDatas = TagDatas;
+	TagDatas = Allocator.New<FTagDataArray>();
+	FTagDataArray& LocalTagDatas = *TagDatas;
+	LocalTagDatas.Reserve(OutOldTagDatas->Num());
+	for (FTagData* TagData : *OutOldTagDatas)
+	{
+		LocalTagDatas.Add(TagData);
+	}
+
+	auto GetEdges = [&LocalTagDatas](int32 Vertex, int* Edges, int& NumEdges)
+	{
+		NumEdges = 0;
+		const FTagData* Parent = LocalTagDatas[Vertex]->GetParent();
+		if (Parent)
+		{
+			Edges[NumEdges++] = Parent->GetIndex();
+		}
+	};
+
+	LLMAlgo::TopologicalSortLeafToRoot(LocalTagDatas, GetEdges);
+
+	// Set each tag's index to its new position in the sort order
+	for (int32 n = 0; n < LocalTagDatas.Num(); ++n)
+	{
+		LocalTagDatas[n]->SetIndex(n);
+	}
+}
+
+void FLowLevelMemTracker::PublishDataPerFrame(const TCHAR* LogName)
+{
+	using namespace UE::LLMPrivate;
+
+	// set overhead stats
+	SET_MEMORY_STAT(STAT_LLMOverheadTotal, MemoryUsageCurrentOverhead);
+	if (IsEnabled())
+	{
+		FLLMTracker& DefaultTracker = *GetTracker(ELLMTracker::Default);
+		FLLMTracker& PlatformTracker = *GetTracker(ELLMTracker::Platform);
+
+		bool bTrackPeaks = CVarLLMTrackPeaks.GetValueOnAnyThread() != 0;
+
+		UE::LLM::ESizeParams SizeParams(UE::LLM::ESizeParams::Default);
+		if (bTrackPeaks)
+		{
+			EnumAddFlags(SizeParams, UE::LLM::ESizeParams::ReportPeak);
+		}
+
+		if (bCapturedSizeSnapshot)
+		{
+			EnumAddFlags(SizeParams, UE::LLM::ESizeParams::RelativeToSnapshot);
+		}
+
+#if !LLM_ENABLED_TRACK_PEAK_MEMORY
+		if (bTrackPeaks)
+		{
+			static bool bWarningGiven = false;
+			if (!bWarningGiven)
+			{
+				UE_LOG(LogHAL, Warning,
+					TEXT("Attempted to enable LLM.TrackPeaks, but LLM_ENABLED_TRACK_PEAK_MEMORY is not defined to 1. You will need to enable the define"));
+				bWarningGiven = true;
+			}
+		}
+#endif
+
+		DefaultTracker.PublishStats(SizeParams);
+		PlatformTracker.PublishStats(SizeParams);
+
+		if (bCsvWriterEnabled)
+		{
+			DefaultTracker.PublishCsv(SizeParams);
+			PlatformTracker.PublishCsv(SizeParams);
+		}
+
+		if (bTraceWriterEnabled)
+		{
+			DefaultTracker.PublishTrace(SizeParams);
+			PlatformTracker.PublishTrace(SizeParams);
+		}
+
+#if LLM_CSV_PROFILER_WRITER_ENABLED
+		if (FCsvProfiler::Get()->IsCapturing())
+		{
+			DefaultTracker.PublishCsvProfiler(SizeParams);
+			PlatformTracker.PublishCsvProfiler(SizeParams);
+		}
+#endif
+	}
+
+	if (LogName != nullptr)
+	{
+		FPlatformMisc::LowLevelOutputDebugStringf(TEXT("---> Untracked memory at %s = %.2f mb\n"),
+			LogName, (double)MemoryUsagePlatformTotalUntracked / (1024.0 * 1024.0));
+	}
+}
+
+void FLowLevelMemTracker::InitialiseProgramSize()
+{
+	// On Android we report ProgramSize via AndroidMemoryProbeTick
+#if !PLATFORM_ANDROID
+	if (!ProgramSize)
+	{
+		FPlatformMemoryStats Stats = FPlatformMemory::GetStats();
+		ProgramSize = Stats.TotalPhysical - Stats.AvailablePhysical;
+
+		GetTracker(ELLMTracker::Platform)->TrackMemory(ELLMTag::ProgramSizePlatform, ProgramSize, ELLMAllocType::System);
+		GetTracker(ELLMTracker::Default)->TrackMemory(ELLMTag::ProgramSize, ProgramSize, ELLMAllocType::System);
+	}
+#endif
+}
+
+void FLowLevelMemTracker::SetProgramSize(uint64 InProgramSize)
+{
+	if (!IsEnabled())
+	{
+		return;
+	}
+	BootstrapInitialise();
+
+	int64 ProgramSizeDiff = static_cast<int64>(InProgramSize) - ProgramSize;
+
+	ProgramSize = static_cast<int64>(InProgramSize);
+
+	GetTracker(ELLMTracker::Platform)->TrackMemory(ELLMTag::ProgramSizePlatform, ProgramSizeDiff, ELLMAllocType::System);
+	GetTracker(ELLMTracker::Default)->TrackMemory(ELLMTag::ProgramSize, ProgramSizeDiff, ELLMAllocType::System);
+}
+
+void FLowLevelMemTracker::ProcessCommandLine(const TCHAR* CmdLine)
+{
+	CSV_METADATA(TEXT("LLM"), TEXT("0"));
+
+#if LLM_AUTO_ENABLE
+	// LLM is always on, regardless of command line
+	bool bShouldDisable = false;
+#elif LLM_COMMANDLINE_ENABLES_FUNCTIONALITY
+	// if we require commandline to enable it, then we are disabled if it's not there
+	bool bShouldDisable = FParse::Param(CmdLine, TEXT("LLM")) == false;
+#else
+	// if we allow commandline to disable us, then we are disabled if it's there
+	bool bShouldDisable = FParse::Param(CmdLine, TEXT("NOLLM")) == true;
+#endif
+
+	bool bLocalCsvWriterEnabled = FParse::Param(CmdLine, TEXT("LLMCSV"));
+	bool bLocalTraceWriterEnabled = UE_TRACE_CHANNELEXPR_IS_ENABLED(MemTagChannel);
+	// automatically enable LLM if only csv or trace output is active
+	if (bLocalCsvWriterEnabled || bLocalTraceWriterEnabled)
+	{
+		if (bLocalTraceWriterEnabled)
+		{
+			UE_LOG(LogInit, Log, TEXT("LLM enabled due to UE_TRACE MemTagChannel being enabled"));
+		}
+		bShouldDisable = false;
+	}
+
+	bAutoPublish = FParse::Param(CmdLine, TEXT("LLMDISABLEAUTOPUBLISH")) == false;
+
+	if (!bCanEnable)
+	{
+		LLMCheck(EnabledState == EEnabled::Disabled);
+		PRAGMA_DISABLE_DEPRECATION_WARNINGS;
+		LLMCheck(bIsDisabled);
+		PRAGMA_ENABLE_DEPRECATION_WARNINGS;
+		if (!bShouldDisable)
+		{
+			UE_LOG(LogInit, Log,
+				TEXT("LLM - Ignoring request to enable LLM; it is not available on the current platform"));
+		}
+		return;
+	}
+
+	ON_SCOPE_EXIT
+	{
+		bConfigurationComplete = true;
+	};
+
+	if (bShouldDisable)
+	{
+		// Before we shutdown, update once so we can publish the overhead-when-disabled later during the first
+		// call to UpdateStatsPerFrame.
+		if (IsEnabled())
+		{
+			Tick();
+		}
+
+		{
+#if PLATFORM_HAS_MULTITHREADED_PREMAIN
+			// The EnableStateLock must be limited in scope because other code in the function
+			// allocates memory and would block on the readlock. Trying to take it around the write of the
+			// EnabledState is sufficient; this will cause us to wait until threads already in OnLowLevelAlloc exit the
+			// function and clear their readlock, and will cause new calls to those functions to block while we're waiting
+			// and call IsEnabled again after we release the lock.
+			FWriteScopeLock EnableStateLock(UE::LLMPrivate::GetEnableStateLock());
+#endif
+
+			EnabledState = EEnabled::Disabled;
+			PRAGMA_DISABLE_DEPRECATION_WARNINGS;
+			bIsDisabled = true;
+			PRAGMA_ENABLE_DEPRECATION_WARNINGS;
+		}
+		bCsvWriterEnabled = false;
+		bTraceWriterEnabled = false;
+		bCanEnable = false; // Reenabling after a clear is not implemented
+		Clear();
+		return;
+	}
+	CSV_METADATA(TEXT("LLM"), TEXT("1"));
+
+	// PLATFORM_HAS_MULTITHREADED_PREMAIN: No need for a Write lock because we are not changing state to disabled.
+	// The other data we modify in this function is synchronized using other methods.
+	EnabledState = EEnabled::Enabled;
+	PRAGMA_DISABLE_DEPRECATION_WARNINGS;
+	bIsDisabled = false;
+	PRAGMA_ENABLE_DEPRECATION_WARNINGS;
+	bCsvWriterEnabled = bLocalCsvWriterEnabled;
+	bTraceWriterEnabled = bLocalTraceWriterEnabled;
+	FinishInitialise();
+
+	// activate tag sets (we ignore None set, it's always on)
+	FString SetList;
+	static_assert(static_cast<uint8>(ELLMTagSet::Max) == 4, "You added a tagset, without updating FLowLevelMemTracker::ProcessCommandLine");
+	if (FParse::Value(CmdLine, TEXT("LLMTAGSETS="), SetList, false /* bShouldStopOnSeparator */))
+	{
+		TArray<FString> Sets;
+		SetList.ParseIntoArray(Sets, TEXT(","), true);
+		for (FString& Set : Sets)
+		{
+#if LLM_ALLOW_ASSETS_TAGS
+			if (Set == TEXT("Assets"))
+			{
+				ActiveSets[static_cast<int32>(ELLMTagSet::Assets)] = true;
+			}
+			else if (Set == TEXT("AssetClasses"))
+			{
+				ActiveSets[static_cast<int32>(ELLMTagSet::AssetClasses)] = true;
+			}
+#elif LLM_ALLOW_UOBJECTCLASSES_TAGS
+			if (Set == TEXT("UObjectClasses"))
+			{
+				ActiveSets[static_cast<int32>(ELLMTagSet::UObjectClasses)] = true;
+			}
+#else
+				// Asset tracking has a per-thread memory overhead, so we have a #define to completely disable it.
+				// Warn if we try to toggle it on at runtime but it is not defined.
+				UE_LOG(LogInit, Warning,
+					TEXT("Attempted to use LLM to track assets, but LLM_ALLOW_ASSETS_TAGS is not defined to 1. You will need to enable the define"));
+#endif
+		}
+	}
+
+	if (FParse::Value(CmdLine, TEXT("LLMTraceTagFilter="), SetList, false /* bShouldStopOnSeparator */))
+	{
+		LLMCheck(TraceTagFilter == nullptr); // ProcessCommandLine() is expected to be called only once
+		TArray<FString> TagNames;
+		SetList.ParseIntoArray(TagNames, TEXT(","), true);
+		if (!TagNames.IsEmpty())
+		{
+			if (!TraceTagFilter)
+			{
+				TraceTagFilter = Allocator.New<UE::LLMPrivate::FTagFilter>();
+			}
+			for (FString& TagName : TagNames)
+			{
+				TraceTagFilter->Add(FName(TagName));
+			}
+			// Update bIsTraceReportable flag for already created tags.
+			// TagData is supposed to be immutable once constructed, but we use bIsTraceReportable flag
+			// as a mutable member to quickly filter tags for the trace purposes.
+			UE::TSharedLock Lock(TagDataMutex);
+			for (UE::LLMPrivate::FTagData* TagData : (*TagDatas))
+			{
+				if (!TraceTagFilter->Contains(TagData->GetName()))
+				{
+					TagData->SetIsTraceReportable(false);
+				}
+			}
+		}
+	}
+
+	// Commandline overrides for console variables
+	int TrackPeaks = 0;
+	if (FParse::Value(CmdLine, TEXT("LLMTrackPeaks="), TrackPeaks))
+	{
+		CVarLLMTrackPeaks->Set(TrackPeaks);
+	}
+
+	UE_LOG(LogInit, Log, TEXT("LLM enabled CsvWriter: %s TraceWriter: %s"),
+		bCsvWriterEnabled ? TEXT("on") : TEXT("off"), bTraceWriterEnabled ? TEXT("on") : TEXT("off"));
+
+	// Disable hitchdetector because LLM is already slow enough.
+	if (!FParse::Param(CmdLine, TEXT("DetectHitchesWithLLM")))
+	{
+		// Schedule disabling when it's safe to init HitchHeartBeat
+		FCoreDelegates::OnPostEngineInit.AddLambda([]{
+			// Calling Get() will instantiate FGameThreadHitchHeartBeat if it is not already.
+			FGameThreadHitchHeartBeat::Get().Stop();
+		});
+	}
+}
+
+// Return the total amount of memory being tracked
+uint64 FLowLevelMemTracker::GetTotalTrackedMemory(ELLMTracker Tracker)
+{
+	if (!IsEnabled())
+	{
+		return 0;
+	}
+	BootstrapInitialise();
+
+	return static_cast<uint64>(GetTracker(Tracker)->GetTrackedTotal());
+}
+
+void FLowLevelMemTracker::OnLowLevelAlloc(ELLMTracker Tracker, const void* Ptr, uint64 Size, ELLMTag DefaultTag,
+	ELLMAllocType AllocType, bool bTrackInMemPro)
+{
+	FEnableStateScopeLock EnableScopeLock;
+	if (!TryEnterEnabled(EnableScopeLock))
+	{
+		return;
+	}
+
+	BootstrapInitialise();
+
+	GetTracker(Tracker)->TrackAllocation(Ptr, static_cast<int64>(Size), DefaultTag, AllocType, bTrackInMemPro);
+}
+
+void FLowLevelMemTracker::OnLowLevelAlloc(ELLMTracker Tracker, const void* Ptr, uint64 Size, FName DefaultTag,
+	ELLMAllocType AllocType, bool bTrackInMemPro)
+{
+	FEnableStateScopeLock EnableScopeLock;
+	if (!TryEnterEnabled(EnableScopeLock))
+	{
+		return;
+	}
+
+	BootstrapInitialise();
+
+	GetTracker(Tracker)->TrackAllocation(Ptr, static_cast<int64>(Size), DefaultTag, AllocType, bTrackInMemPro);
+}
+
+void FLowLevelMemTracker::OnLowLevelFree(ELLMTracker Tracker, const void* Ptr,
+	ELLMAllocType AllocType, bool bTrackInMemPro)
+{
+	FEnableStateScopeLock EnableScopeLock;
+	if (!TryEnterEnabled(EnableScopeLock))
+	{
+		return;
+	}
+
+	BootstrapInitialise();
+
+	if (Ptr != nullptr)
+	{
+		GetTracker(Tracker)->TrackFree(Ptr, AllocType, bTrackInMemPro);
+	}
+}
+
+void FLowLevelMemTracker::OnLowLevelChangeInMemoryUse(ELLMTracker Tracker, int64 DeltaMemory, ELLMTag DefaultTag,
+	ELLMAllocType AllocType)
+{
+	FEnableStateScopeLock EnableScopeLock;
+	if (!TryEnterEnabled(EnableScopeLock))
+	{
+		return;
+	}
+
+	BootstrapInitialise();
+	GetTracker(Tracker)->TrackMemoryOfActiveTag(DeltaMemory, DefaultTag, AllocType);
+}
+
+void FLowLevelMemTracker::OnLowLevelChangeInMemoryUse(ELLMTracker Tracker, int64 DeltaMemory, FName DefaultTag,
+	ELLMAllocType AllocType)
+{
+	FEnableStateScopeLock EnableScopeLock;
+	if (!TryEnterEnabled(EnableScopeLock))
+	{
+		return;
+	}
+
+	BootstrapInitialise();
+	GetTracker(Tracker)->TrackMemoryOfActiveTag(DeltaMemory, DefaultTag, AllocType);
+}
+
+void FLowLevelMemTracker::OnLowLevelAllocMoved(ELLMTracker Tracker, const void* Dest, const void* Source,
+	ELLMAllocType AllocType)
+{
+	FEnableStateScopeLock EnableScopeLock;
+	if (!TryEnterEnabled(EnableScopeLock))
+	{
+		return;
+	}
+
+	BootstrapInitialise();
+
+	//update the allocation map
+	GetTracker(Tracker)->OnAllocMoved(Dest, Source, AllocType);
+}
+
+UE::LLMPrivate::FLLMTracker* FLowLevelMemTracker::GetTracker(ELLMTracker Tracker)
+{
+	return Trackers[static_cast<int32>(Tracker)];
+}
+
+const UE::LLMPrivate::FLLMTracker* FLowLevelMemTracker::GetTracker(ELLMTracker Tracker) const
+{
+	return Trackers[static_cast<int32>(Tracker)];
+}
+
+bool FLowLevelMemTracker::Exec(const TCHAR* Cmd, FOutputDevice& Ar)
+{
+	if (!IsEnabled())
+	{
+		if (FParse::Command(&Cmd, TEXT("LLMEM")))
+		{
+			UE_LOG(LogHAL, Warning, TEXT("Cannot execute LLMEM command! LLM is not enabled."));
+			return true;
+		}
+		return false;
+	}
+
+	if (FParse::Command(&Cmd, TEXT("LLMEM")))
+	{
+		BootstrapInitialise();
+		if (FParse::Command(&Cmd, TEXT("SPAMALLOC")))
+		{
+			int32 NumAllocs = 128;
+			int64 MaxSize = FCString::Atoi(Cmd);
+			if (MaxSize == 0)
+			{
+				MaxSize = 128 * 1024;
+			}
+
+			UpdateStatsPerFrame(TEXT("Before spam"));
+			FPlatformMisc::LowLevelOutputDebugStringf(TEXT("----> Spamming %d allocations, from %d..%d bytes\n"),
+				NumAllocs, MaxSize/2, MaxSize);
+
+			TArray<void*> Spam;
+			Spam.Reserve(NumAllocs);
+			SIZE_T TotalSize = 0;
+			for (int32 Index = 0; Index < NumAllocs; Index++)
+			{
+				SIZE_T Size = (FPlatformMath::Rand() % MaxSize / 2) + MaxSize / 2;
+				TotalSize += Size;
+				Spam.Add(FMemory::Malloc(Size));
+			}
+			FPlatformMisc::LowLevelOutputDebugStringf(TEXT("----> Allocated %d total bytes\n"), TotalSize);
+
+			UpdateStatsPerFrame(TEXT("After spam"));
+
+			for (int32 Index = 0; Index < Spam.Num(); Index++)
+			{
+				FMemory::Free(Spam[Index]);
+			}
+
+			UpdateStatsPerFrame(TEXT("After cleanup"));
+		}
+		else if (FParse::Command(&Cmd, TEXT("DUMPPRIVATESHARED")))
+		{
+			for (int32 TrackerIndex = 0; TrackerIndex < static_cast<int32>(ELLMTracker::Max); TrackerIndex++)
+			{
+				if (GetTracker((ELLMTracker)TrackerIndex)->DumpForkedAllocationInfo() == false)
+				{
+					FPlatformMisc::LowLevelOutputDebugString(TEXT("Failed to dumping forked allocation info (check platform supports it?)\n"));
+					return true;
+				}
+			}
+		}
+		else if (FParse::Command(&Cmd, TEXT("REMEMBER")))
+		{
+			CSV_EVENT_GLOBAL(TEXT("LLM_REMEMBER"));
+
+			UE::TUniqueLock UpdateLock(UpdateMutex);
+			// Go through all the tags and REMEMBER the current recorded size values
+			for (int32 TrackerIndex = 0; TrackerIndex < static_cast<int32>(ELLMTracker::Max); TrackerIndex++)
+			{
+				GetTracker((ELLMTracker)TrackerIndex)->CaptureTagSnapshot();
+			}
+			bCapturedSizeSnapshot = true;
+		}
+		else if (FParse::Command(&Cmd, TEXT("FORGET")))
+		{
+			CSV_EVENT_GLOBAL(TEXT("LLM_FORGET"));
+
+			UE::TUniqueLock UpdateLock(UpdateMutex);
+			// Go through all tags and make them forget their current size values
+			for (int32 TrackerIndex = 0; TrackerIndex < static_cast<int32>(ELLMTracker::Max); TrackerIndex++)
+			{
+				GetTracker((ELLMTracker)TrackerIndex)->ClearTagSnapshot();
+			}
+			bCapturedSizeSnapshot = false;
+		}
+		else if (FParse::Command(&Cmd, TEXT("DUMPTAGS")))
+		{
+			if (GetTracker(ELLMTracker::Default)->DumpTags() == false)
+			{
+				FPlatformMisc::LowLevelOutputDebugString(TEXT("Failed to dump tags info\n"));
+				return true;
+			}
+		}
+		return true;
+	}
+
+	return false;
+}
+
+bool FLowLevelMemTracker::IsTagSetActive(ELLMTagSet Set)
+{
+	if (!IsEnabled())
+	{
+		return false;
+	}
+	BootstrapInitialise();
+
+	return ActiveSets[static_cast<int32>(Set)];
+}
+
+bool FLowLevelMemTracker::ShouldReduceThreads()
+{
+#if LLM_ENABLED_REDUCE_THREADS
+	if (!IsEnabled())
+	{
+		return false;
+	}
+	BootstrapInitialise();
+	LLMCheckf(bConfigurationComplete,
+		TEXT("ShouldReduceThreads has been called too early, before we processed the configuration settings required for it."));
+
+	return IsTagSetActive(ELLMTagSet::Assets) || IsTagSetActive(ELLMTagSet::AssetClasses);
+#else
+	return false;
+#endif
+}
+
+#if LLM_ALLOW_ASSETS_TAGS
+
+static bool IsAssetTagForAssets(ELLMTagSet Set)
+{
+	return Set == ELLMTagSet::Assets || Set == ELLMTagSet::AssetClasses;
+}
+
+#endif // LLM_ALLOW_ASSETS_TAGS
+
+void FLowLevelMemTracker::RegisterCustomTagInternal(int32 Tag, ELLMTagSet TagSet, const TCHAR* InDisplayName, FName StatName,
+	FName SummaryStatName, int32 ParentTag)
+{
+	using namespace UE::LLMPrivate;
+
+	LLMCheckf(Tag >= LLM_CUSTOM_TAG_START && Tag <= LLM_CUSTOM_TAG_END, TEXT("Tag %d out of range"), Tag);
+	LLMCheckf(InDisplayName != nullptr, TEXT("Tag %d has no name"), Tag);
+	LLMCheckf(ParentTag == -1 || ParentTag < LLM_TAG_COUNT, TEXT("Parent tag %d out of range"), ParentTag);
+
+	FName DisplayName = InDisplayName ? InDisplayName : InvalidLLMTagName;
+	ELLMTag EnumTag = static_cast<ELLMTag>(Tag);
+	FName ParentName = ParentTag >= 0 ? LLMGetTagUniqueName(static_cast<ELLMTag>(ParentTag)) : NAME_None;
+
+	RegisterTagData(LLMGetTagUniqueName(EnumTag), DisplayName, ParentName, StatName, SummaryStatName, true,
+		EnumTag, false, ETagReferenceSource::CustomEnumTag, TagSet);
+}
+
+void FLowLevelMemTracker::RegisterPlatformTag(int32 Tag, const TCHAR* Name, FName StatName, FName SummaryStatName,
+	int32 ParentTag)
+{
+	MemoryTrace_AnnounceCustomTag(Tag, ParentTag, Name);
+
+	if (!IsEnabled())
+	{
+		return;
+	}
+	BootstrapInitialise();
+
+	LLMCheck(Tag >= static_cast<int32>(ELLMTag::PlatformTagStart) &&
+		Tag <= static_cast<int32>(ELLMTag::PlatformTagEnd));
+	RegisterCustomTagInternal(Tag, ELLMTagSet::None, Name, StatName, SummaryStatName, ParentTag);
+}
+
+void FLowLevelMemTracker::RegisterProjectTag(int32 Tag, const TCHAR* Name, FName StatName, FName SummaryStatName,
+	int32 ParentTag)
+{
+	MemoryTrace_AnnounceCustomTag(Tag, ParentTag, Name);
+
+	if (!IsEnabled())
+	{
+		return;
+	}
+	BootstrapInitialise();
+
+	LLMCheck(Tag >= static_cast<int32>(ELLMTag::ProjectTagStart) && Tag <= static_cast<int32>(ELLMTag::ProjectTagEnd));
+	RegisterCustomTagInternal(Tag, ELLMTagSet::None, Name, StatName, SummaryStatName, ParentTag);
+}
+
+void GlobalRegisterTagDeclaration(FLLMTagDeclaration& TagDeclaration)
+{
+	if (!FLowLevelMemTracker::IsEnabled())
+	{
+		return;
+	}
+	FLowLevelMemTracker::Get().RegisterTagDeclaration(TagDeclaration);
+}
+
+void FLowLevelMemTracker::BootstrapTagDatas()
+{
+	using namespace UE::LLMPrivate;
+
+	// While bootstrapping we are not allowed to construct any FNames because the FName system may not yet have been
+	// constructed. Construct not-fully-initialized TagDatas for the central list of ELLMTags.
+	{
+		UE::TUniqueLock Lock(TagDataMutex);
+		bIsBootstrapping = true;
+
+		TagDatas = Allocator.New<FTagDataArray>();
+		TagDataNameMap = Allocator.New<FTagDataNameMap>();
+		TagDataEnumMap = reinterpret_cast<FTagData**>(Allocator.Alloc(sizeof(FTagData*) * LLM_TAG_COUNT));
+		FMemory::Memset(TagDataEnumMap, 0, sizeof(FTagData*) * LLM_TAG_COUNT);
+
+#define REGISTER_ELLMTAG(Enum,Str,Stat,Group,ParentTag) \
+		{ \
+			ELLMTag EnumTag = ELLMTag::Enum; \
+			int32 Index = static_cast<int32>(EnumTag); \
+			LLMCheck(0 <= Index && Index < LLM_TAG_COUNT); \
+			FTagData* TagData = Allocator.New<FTagData>(EnumTag); \
+			TagData->SetIndex(TagDatas->Num()); \
+			TagDatas->Add(TagData); \
+			LLMCheck(TagDataEnumMap[Index] == nullptr); \
+			TagDataEnumMap[Index] = TagData; \
+		}
+		LLM_ENUM_GENERIC_TAGS(REGISTER_ELLMTAG);
+#undef REGISTER_ELLMTAG
+
+#if LLM_ALLOW_NAMES_TAGS
+		// The CustomName tag is an adapter for connecting LLM_SCOPE_BYNAME tags with platforms that used the
+		// ELLMTag-based reporting. We want to hide this procedural tag in systems that use FName-based reporting;
+		// if it is displayed it confusingly just displays a sum of every LLM_SCOPE_BYNAME tag.
+		TagDataEnumMap[static_cast<int32>(ELLMTag::CustomName)]->SetIsReportable(false);
+#endif
+	}
+}
+
+void FLowLevelMemTracker::FinishInitialise()
+{
+	if (bFullyInitialised)
+	{
+		return;
+	}
+	BootstrapInitialise();
+
+#if UE_ENABLE_ARRAY_SLACK_TRACKING
+	GTrackArrayDetailedLock = new FCriticalSection();
+#endif
+	bFullyInitialised = true;
+	// Make sure that FNames and Malloc have already been initialised, since we will use them during InitialiseTagDatas
+	// We force this by calling LLMGetTagUniqueName, which initializes FNames internally, and will therein trigger
+	// FName system construction, which will itself trigger Malloc construction.
+	(void)LLMGetTagUniqueName(ELLMTag::Untagged);
+	InitialiseTagDatas();
+
+	UE::LLMPrivate::FInitializedCallbacksArray& Callbacks = UE::LLMPrivate::GetInitialisedCallbacks();
+	for (const TPair<UE::LLMPrivate::FLLMInitialisedCallback, UPTRINT>& Callback : Callbacks)
+	{
+		Callback.Key(Callback.Value);
+	}
+	Callbacks.Empty();
+}
+
+void FLowLevelMemTracker::InitialiseTagDatas_SetLLMTagNames()
+{
+	using namespace UE::LLMPrivate;
+
+	TStringBuilder<256> NameBuffer;
+	// Load all the names for the central list of ELLMTags (recording the allocations the FName system makes for the
+	// construction of the names).
+#define SET_ELLMTAG_NAMES(Enum,Str,Stat,Group,ParentTag) \
+	{ \
+		ELLMTag EnumTag = ELLMTag::Enum; \
+		int32 Index = static_cast<int32>(EnumTag); \
+		FTagData* TagData = TagDataEnumMap[Index]; \
+		FName TagName = LLMGetTagUniqueName(EnumTag); \
+		TagName.ToString(NameBuffer); \
+		ValidateUniqueName(NameBuffer); \
+		TagData->SetName(LLMGetTagUniqueName(EnumTag)); \
+		TagData->SetDisplayName(Str); \
+		TagData->SetStatName(Stat); \
+		TagData->SetSummaryStatName(Group); \
+		TagData->SetParentName(static_cast<int32>(ParentTag) == -1 ? \
+			NAME_None : LLMGetTagUniqueName(static_cast<ELLMTag>(ParentTag))); \
+	}
+	LLM_ENUM_GENERIC_TAGS(SET_ELLMTAG_NAMES);
+#undef SET_ELLMTAG_NAMES
+}
+
+void FLowLevelMemTracker::InitialiseTagDatas_FinishRegister()
+{
+	using namespace UE::LLMPrivate;
+
+	// Record the central list of ELLMTags in TagDataNameMap, and mark that bootstrapping is complete
+	{
+		UE::TUniqueLock Lock(TagDataMutex);
+
+#define FINISH_REGISTER(Enum,Str,Stat,Group,ParentTag) \
+		{ \
+			ELLMTag EnumTag = ELLMTag::Enum; \
+			int32 Index = static_cast<int32>(EnumTag); \
+			FTagData* TagData = TagDataEnumMap[Index]; \
+			FTagData*& ExistingTagData = TagDataNameMap->FindOrAdd(FTagDataNameKey(TagData->GetName(), TagData->GetTagSet()), nullptr); \
+			if (ExistingTagData != nullptr) \
+			{ \
+				ReportDuplicateTagName(ExistingTagData, ETagReferenceSource::EnumTag); \
+			} \
+			ExistingTagData = TagData; \
+		}
+		LLM_ENUM_GENERIC_TAGS(FINISH_REGISTER);
+#undef FINISH_REGISTER
+
+		bIsBootstrapping = false;
+	}
+}
+
+void FLowLevelMemTracker::InitialiseTagDatas()
+{
+	using namespace UE::LLMPrivate;
+
+	InitialiseTagDatas_SetLLMTagNames();
+
+	InitialiseTagDatas_FinishRegister();
+
+	// Construct the remaining startup tags; allocations when constructing these tags are known to consist only of the
+	// central list of ELLMTags so we do not need to bootstrap.
+	{
+		// Construct LLM_DECLARE_TAGs
+		FLLMTagDeclaration* List = FLLMTagDeclaration::GetList();
+		while (List)
+		{
+			RegisterTagDeclaration(*List);
+			List = List->Next;
+		}
+		FLLMTagDeclaration::AddCreationCallback(GlobalRegisterTagDeclaration);
+	}
+
+	// now let the platform add any custom tags
+	FPlatformMemory::RegisterCustomLLMTags();
+
+	// All parents in the ELLMTags and the initial modules' list of LLM_DEFINE_TAG must be contained within that same
+	// set, so we can FinishConstruct them now, which we do in UpdateTags.
+	bTagAdded = true;
+	UpdateTags();
+}
+
+void FLowLevelMemTracker::ClearTagDatas()
+{
+	using namespace UE::LLMPrivate;
+
+	UE::TUniqueLock Lock(TagDataMutex);
+	FLLMTagDeclaration::ClearCreationCallbacks();
+
+	Allocator.Free(TagDataEnumMap, sizeof(FTagData*) * LLM_TAG_COUNT);
+	TagDataEnumMap = nullptr;
+	Allocator.Delete(TagDataNameMap);
+	TagDataNameMap = nullptr;
+	for (FTagData* TagData : *TagDatas)
+	{
+		Allocator.Delete(TagData);
+	}
+	Allocator.Delete(TagDatas);
+	TagDatas = nullptr;
+}
+
+void FLowLevelMemTracker::RegisterTagDeclaration(FLLMTagDeclaration& TagDeclaration)
+{
+	TagDeclaration.ConstructUniqueName();
+	RegisterTagData(TagDeclaration.UniqueName, TagDeclaration.DisplayName, TagDeclaration.ParentTagName,
+		TagDeclaration.StatName, TagDeclaration.SummaryStatName, false, ELLMTag::CustomName, false,
+		UE::LLMPrivate::ETagReferenceSource::Declare, TagDeclaration.TagSet);
+}
+
+UE::LLMPrivate::FTagData& FLowLevelMemTracker::RegisterTagData(FName Name, FName DisplayName, FName ParentName,
+	FName StatName, FName SummaryStatName, bool bHasEnumTag, ELLMTag EnumTag, bool bIsStatTag,
+	UE::LLMPrivate::ETagReferenceSource ReferenceSource, ELLMTagSet TagSet/* = ELLMTagSet::None*/)
+{
+	LLMCheckf(!bIsBootstrapping,
+		TEXT("A tag outside of LLM_ENUM_GENERIC_TAGS was requested from LLM_SCOPE or allocation while bootstrapping the names for LLM_ENUM_GENERIC_TAGS, this is not supported."));
+
+	using namespace UE::LLMPrivate;
+
+	// If this allocates, that is okay. Set it to something small-as-possible-to-avoid-normally-allocating to prevent
+	// adding a lot of stack space in the calling LLM_SCOPE code.
+	TStringBuilder<256> NameBuffer;
+	Name.ToString(NameBuffer);
+
+	if (bHasEnumTag)
+	{
+		ValidateUniqueName(NameBuffer);
+		// EnumTags can specify DisplayName (if they are central or if CustomTag registration provided it);
+		// if not, they set DisplayName = UniqueName.
+		// Enum tags only specify ParentName explicitly; if no ParentName is provided, they have no parent.
+		if (DisplayName.IsNone())
+		{
+			DisplayName = Name;
+		}
+	}
+	else if (TagSet != ELLMTagSet::None)
+	{
+		// Tags not in the None TagSet do not have parents and do not require name validation
+		DisplayName = Name;
+	}
+	else if (bIsStatTag)
+	{
+		// We set LLM UniqueName = <TheEntireString> and LLM DisplayName = StatDisplayName.
+		// Stat tags do not specify their parent, and their parent is set to the CustomName aggregator.
+		LLMCheck(DisplayName.IsNone());
+		LLMCheck(ParentName.IsNone());
+		DisplayName = Name;
+		ParentName = GetTagName_CustomName();
+	}
+	else
+	{
+		ValidateUniqueName(NameBuffer);
+		// Normal defined-by-name tags supply unique names of the form Grandparent/.../Parent/Name.
+		// ParentName and  DisplayName can be provided.
+
+		// If both ParentName and /Parent/ are supplied, it is an error if they do not match.
+		// All custom name tags have to be children of an ELLMTag. If no parent is set, it defaults to the proxy
+		// parent CustomName.
+		const TCHAR* LeafStart = NameBuffer.ToString();
+		while (true)
+		{
+			const TCHAR* NextDivider = FCString::Strstr(LeafStart, TEXT("/"));
+			if (!NextDivider)
+			{
+				break;
+			}
+			LeafStart = NextDivider + 1;
+		}
+		LLMCheckf(LeafStart[0] != '\0', TEXT("Invalid LLM custom name tag '%s'. Tag names must not end with /."),
+			NameBuffer.ToString());
+		if (LeafStart != NameBuffer.ToString())
+		{
+			FName ParsedParentName = FName(FStringView(NameBuffer.ToString(),
+				UE_PTRDIFF_TO_INT32(LeafStart - 1 - NameBuffer.ToString())));
+			if (!ParentName.IsNone() && ParentName != ParsedParentName)
+			{
+				TStringBuilder<128> ParentBuffer;
+				ParentName.ToString(ParentBuffer);
+				LLMCheckf(false,
+					TEXT("Invalid LLM tag: parent supplied in tag declaration is '%s', which does not match the parent parsed from the tag unique name '%s'"),
+					ParentBuffer.ToString(), NameBuffer.ToString());
+			}
+			ParentName = ParsedParentName;
+		}
+		else if (ParentName.IsNone())
+		{
+			ParentName = GetTagName_CustomName();
+		}
+
+		// Display name is set to the leaf portion of the unique name, and is overridden if DisplayName is provided.
+		if (DisplayName.IsNone())
+		{
+			DisplayName = FName(LeafStart);
+		}
+	}
+
+	UE::TUniqueLock Lock(TagDataMutex);
+	FTagData*& TagDataForName = TagDataNameMap->FindOrAdd(FTagDataNameKey(Name, TagSet), nullptr);
+	if (TagDataForName)
+	{
+		// The tag already exists; this can happen because two formal registrations have tried to add the same name,
+		// but it can also happen due to a race condition when auto-adding a tag from LLM_SCOPE, if the tag is added
+		// by another thread in between FindOrAddTagData's check of TagDataNameMap->Find and now. Note that it is not
+		// valid for an LLM_SCOPE to be called before a formal registration (e.g. LLM_DECLARE_TAG). If a formal
+		// registration exists for a tag, it must precede its use in any LLM_SCOPE calls.
+		if (ParentName != TagDataForName->GetParentNameSafeBeforeFinishConstruct() ||
+			StatName != TagDataForName->GetStatName() ||
+			SummaryStatName != TagDataForName->GetSummaryStatName() ||
+			bHasEnumTag != TagDataForName->HasEnumTag() ||
+			(bHasEnumTag && EnumTag != TagDataForName->GetEnumTag()))
+		{
+			if (ReferenceSource != UE::LLMPrivate::ETagReferenceSource::Scope &&
+				ReferenceSource != UE::LLMPrivate::ETagReferenceSource::ImplicitParent &&
+				ReferenceSource != UE::LLMPrivate::ETagReferenceSource::FunctionAPI)
+			{
+				ReportDuplicateTagName(TagDataForName, ReferenceSource);
+			}
+		}
+
+		// Abandon the string work we've done and return the version that was added by the other source.
+		return *TagDataForName;
+	}
+
+	FTagData* ParentData = nullptr;
+	if (!ParentName.IsNone())
+	{
+		FTagData** ParentPtr = TagDataNameMap->Find(FTagDataNameKey(ParentName, TagSet));
+		if (ParentPtr)
+		{
+			ParentData = *ParentPtr;
+		}
+	}
+
+	FTagData* TagData;
+	if (ParentName.IsNone() || ParentData)
+	{
+		TagData = Allocator.New<FTagData>(Name, TagSet, DisplayName, ParentData, StatName, SummaryStatName, bHasEnumTag,
+			EnumTag, ReferenceSource);
+	}
+	else
+	{
+		TagData = Allocator.New<FTagData>(Name, TagSet, DisplayName, ParentName, StatName, SummaryStatName, bHasEnumTag,
+			EnumTag, ReferenceSource);
+	}
+	TagData->SetIndex(TagDatas->Num());
+	TagDatas->Add(TagData);
+
+	TagDataForName = TagData;
+
+	if (bHasEnumTag)
+	{
+		int32 Index = static_cast<int32>(EnumTag);
+		LLMCheck(0 <= Index && Index < LLM_TAG_COUNT);
+		FTagData*& TagDataForEnum = TagDataEnumMap[Index];
+		if (TagDataForEnum != nullptr)
+		{
+			LLMCheckf(false, TEXT("LLM Error: Duplicate copies of enumtag %d"), Index);
+		}
+		TagDataForEnum = TagData;
+	}
+
+	bTagAdded = true;
+	return *TagData;
+}
+
+void FLowLevelMemTracker::ReportDuplicateTagName(UE::LLMPrivate::FTagData* TagData,
+	UE::LLMPrivate::ETagReferenceSource ReferenceSource)
+{
+	using namespace UE::LLMPrivate;
+
+	// We're inside the LLM lock, so do not allow these GetName calls to allocate memory from FMalloc.
+	if (ReferenceSource == ETagReferenceSource::FunctionAPI || ReferenceSource == ETagReferenceSource::Scope ||
+		ReferenceSource == ETagReferenceSource::ImplicitParent)
+	{
+		LLMCheckf(false,
+			TEXT("LLM Error: Unexpected call to RegisterTagData(%s) from LLM_SCOPE or function call when the tag already exists."),
+			*WriteToString<FName::StringBufferSize>(TagData->GetName()));
+	}
+	else if (TagData->GetReferenceSource() == ETagReferenceSource::FunctionAPI ||
+		TagData->GetReferenceSource() == ETagReferenceSource::Scope)
+	{
+		LLMCheckf(false, TEXT("LLM Error: Tag %s parsed from %s after it was already used in an LLM_SCOPE or LLM api call."),
+			*WriteToString<FName::StringBufferSize>(TagData->GetName()), ToString(ReferenceSource));
+	}
+	else if (TagData->GetReferenceSource() == ETagReferenceSource::ImplicitParent)
+	{
+		LLMCheckf(false,
+			TEXT("LLM Error: Tag %s parsed from %s after it was already used as an implicit parent in another tag. ")
+			TEXT("Add LLM_DEFINE_TAG(% s) in cpp, or move it to the same module as the child tag using it, so that it will be defined before the child tag tries to use it."),
+			*WriteToString<FName::StringBufferSize>(TagData->GetName()), ToString(ReferenceSource),
+			*WriteToString<FName::StringBufferSize>(TagData->GetName()));
+	}
+	else
+	{
+		LLMCheckf(false, TEXT("LLM Error: Multiple occurrences of a unique tag name %s in ELLMTag or LLM_DEFINE_TAG. ")
+			TEXT("First occurrence : % s.Second occurrence : % s."),
+			*WriteToString<FName::StringBufferSize>(TagData->GetName()), ToString(TagData->GetReferenceSource()),
+			ToString(ReferenceSource));
+	}
+}
+
+void FLowLevelMemTracker::FinishConstruct(UE::LLMPrivate::FTagData* TagData,
+	UE::LLMPrivate::ETagReferenceSource ReferenceSource)
+{
+	using namespace UE::LLMPrivate;
+	// Caller is responsible for holding a SharedLock (NOT a UniqueLock) on TagDataMutex
+
+	LLMCheck(TagData);
+	if (TagData->IsFinishConstructed())
+	{
+		return;
+	}
+	if (bIsBootstrapping)
+	{
+		// Can't access Names yet; run the FinishConstruct later
+		return;
+	}
+
+	if (!TagData->IsParentConstructed())
+	{
+		FName ParentName = TagData->GetParentName();
+		if (ParentName.IsNone())
+		{
+			TagData->SetParent(nullptr);
+		}
+		else
+		{
+			FTagData** ParentDataPtr = TagDataNameMap->Find(FTagDataNameKey(ParentName, TagData->GetTagSet()));
+			if (!ParentDataPtr)
+			{
+				// We have to drop the SharedLock to call RegisterTagData, which takes a UniqueLock.
+				TagDataMutex.UnlockShared();
+				FTagData* ParentTagData = &RegisterTagData(ParentName, NAME_None, NAME_None, NAME_None, NAME_None,
+					false, ELLMTag::CustomName,false, ETagReferenceSource::ImplicitParent, TagData->GetTagSet());
+				TagDataMutex.LockShared();
+				if (TagData->IsFinishConstructed())
+				{
+					// Another thread got in and finished construction while we were outside of the lock.
+					return;
+				}
+				ParentDataPtr = TagDataNameMap->Find(FTagDataNameKey(ParentName, TagData->GetTagSet()));
+				LLMCheck(ParentDataPtr);
+			}
+			FTagData* ParentData = *ParentDataPtr;
+			TagData->SetParent(ParentData);
+		}
+	}
+
+	FTagData* ParentData = const_cast<FTagData*>(TagData->GetParent());
+	if (ParentData)
+	{
+		// Make sure the parent chain is FinishConstructed as well, since GetContainingEnum or GetDisplayPath will be
+		// called and walk up the parent chain.
+		FinishConstruct(ParentData, ReferenceSource);
+	}
+
+	if (TraceTagFilter && !TraceTagFilter->Contains(TagData->GetName()))
+	{
+		TagData->SetIsTraceReportable(false);
+	}
+
+	TagData->SetFinishConstructed();
+
+	// Broadcast the tag creation, except for generic tags which are constructed before any subscriber could
+	// possibly have registered. Subscribers must instead read those from GetTrackedTags.
+	if (!TagData->HasEnumTag() || TagData->GetEnumTag() >= ELLMTag::GenericTagCount)
+	{
+		// Leave the shared mutex while calling the callback, since the callback could be arbitrary
+		// code that calls back into LLM. The TagData pointer is immutable so we do not have to worry about
+		// it disappearing out from under us.
+		TagDataMutex.UnlockShared();
+		for (const TPair<UE::LLMPrivate::FTagCreationCallback, UPTRINT>& Callback :
+			UE::LLMPrivate::GetTagCreationCallbacks())
+		{
+			Callback.Key(TagData, Callback.Value);
+		}
+		TagDataMutex.LockShared();
+	}
+}
+
+TArray<const UE::LLMPrivate::FTagData*> FLowLevelMemTracker::GetTrackedTags(ELLMTagSet TagSet)
+{
+	using namespace UE::LLMPrivate;
+
+	if (!IsEnabled())
+	{
+		return TArray<const FTagData*>();
+	}
+
+	BootstrapInitialise();
+
+	UE::TSharedLock Lock(TagDataMutex);
+
+	TArray<const FTagData*> FoundResults(TagDatas->FilterByPredicate([TagSet](const FTagData* InTagData)
+	{
+		return InTagData != nullptr && InTagData->GetTagSet() == TagSet;
+	}));
+
+	return FoundResults;
+}
+
+TArray<const UE::LLMPrivate::FTagData*> FLowLevelMemTracker::GetTrackedTags(ELLMTracker Tracker, ELLMTagSet TagSet /* = ELLMTagSet::None */)
+{
+	using namespace UE::LLMPrivate;
+
+	if (!IsEnabled())
+	{
+		return TArray<const FTagData*>();
+	}
+
+	BootstrapInitialise();
+
+	UE::TUniqueLock UpdateLock(UpdateMutex); // uses of TagSizes are guarded by the UpdateLock
+	return GetTracker(Tracker)->GetTagDatas(TagSet);
+}
+
+void FLowLevelMemTracker::GetTrackedTagsNamesWithAmount(TMap<FName, uint64>& TagsNamesWithAmount, ELLMTracker Tracker, ELLMTagSet TagSet)
+{
+	using namespace UE::LLMPrivate;
+
+	if (!IsEnabled())
+	{
+		return;
+	}
+
+	BootstrapInitialise();
+
+	UE::TUniqueLock UpdateLock(UpdateMutex);
+	GetTracker(Tracker)->GetTagsNamesWithAmount(TagsNamesWithAmount, TagSet);
+}
+
+void FLowLevelMemTracker::GetTrackedTagsNamesWithAmountFiltered(TMap<FName, uint64>& TagsNamesWithAmount, ELLMTracker Tracker, ELLMTagSet TagSet, TArray<FLLMTagSetAllocationFilter>& Filters)
+{
+	using namespace UE::LLMPrivate;
+
+	if (!IsEnabled())
+	{
+		return;
+	}
+
+	BootstrapInitialise();
+
+	UE::TUniqueLock UpdateLock(UpdateMutex);
+	GetTracker(Tracker)->GetTagsNamesWithAmountFiltered(TagsNamesWithAmount, TagSet, Filters);
+}
+
+bool FLowLevelMemTracker::FindTagByName( const TCHAR* Name, uint64& OutTag, ELLMTagSet InTagSet /*= ELLMTagSet::None*/ ) const
+{
+	using namespace UE::LLMPrivate;
+
+	if (!IsEnabled())
+	{
+		return false;
+	}
+	// Cannot call BootstrapInitialise and FinishInitialise without shenanigans because this function is const.
+	LLMCheck(bFullyInitialised);
+
+	if (Name != nullptr)
+	{
+		UE::TSharedLock Lock(TagDataMutex);
+
+		// Search by Name
+		FName SearchName(Name);
+		FTagData** TagDataPtr = TagDataNameMap->Find(FTagDataNameKey(SearchName, InTagSet));
+		if (TagDataPtr)
+		{
+			const FTagData* TagData = *TagDataPtr;
+			OutTag = static_cast<uint64>(TagData->GetContainingEnum());
+			return true;
+		}
+
+		// Search by ELLMTag's DisplayName
+		for (int32 Index = 0; Index < LLM_TAG_COUNT; ++Index)
+		{
+			const FTagData* TagData = TagDataEnumMap[Index];
+			if (!TagData)
+			{
+				continue;
+			}
+			if (TCString<TCHAR>::Stricmp(*TagData->GetDisplayName().ToString(), Name))
+			{
+				OutTag = static_cast<uint64>(TagData->GetContainingEnum());
+				return true;
+			}
+		}
+	}
+
+	return false;
+}
+
+const TCHAR* FLowLevelMemTracker::FindTagName(uint64 Tag) const
+{
+	if (!IsEnabled())
+	{
+		return nullptr;
+	}
+	// Cannot call BootstrapInitialise without shenanigans because this function is const.
+	LLMCheck(bInitialisedTracking);
+
+	static TMap<uint64, FString> FoundTags;
+	FString* Cached = FoundTags.Find(Tag);
+	if (Cached)
+	{
+		return **Cached;
+	}
+
+	FName DisplayName = FindTagDisplayName(Tag);
+	if (DisplayName.IsNone())
+	{
+		return nullptr;
+	}
+
+	FString& AddedCache = FoundTags.Add(Tag);
+	AddedCache = DisplayName.ToString();
+	return *AddedCache;
+}
+
+FName FLowLevelMemTracker::FindTagDisplayName(uint64 Tag) const
+{
+	using namespace UE::LLMPrivate;
+
+	if (!IsEnabled())
+	{
+		return NAME_None;
+	}
+	// Cannot call BootstrapInitialise without shenanigans because this function is const.
+	LLMCheck(bInitialisedTracking);
+
+	int32 Index = static_cast<int32>(Tag);
+	if (0 <= Index && Index < LLM_TAG_COUNT)
+	{
+		const FTagData* TagData = TagDataEnumMap[Index];
+		if (TagData)
+		{
+			return TagData->GetDisplayName();
+		}
+	}
+	return NAME_None;
+}
+
+FName FLowLevelMemTracker::FindPtrDisplayName(void* Ptr) const
+{
+	using namespace UE::LLMPrivate;
+
+	if (!IsEnabled() || !bFullyInitialised)
+	{
+		return NAME_None;
+	}
+
+	const FLLMTracker* TrackerData = GetTracker(ELLMTracker::Default);
+	TArray<const FTagData*, TInlineAllocator<static_cast<int32>(ELLMTagSet::Max)>> Tags;
+	TrackerData->FindTagsForPtr(Ptr, Tags);
+
+	return Tags.Num() ? Tags[0]->GetDisplayName() : NAME_None;
+}
+
+FName FLowLevelMemTracker::GetTagDisplayName(const UE::LLMPrivate::FTagData* TagData) const
+{
+	return TagData->GetDisplayName();
+}
+
+FString FLowLevelMemTracker::GetTagDisplayPathName(const UE::LLMPrivate::FTagData* TagData) const
+{
+	TStringBuilder<FName::StringBufferSize> Buffer;
+	GetTagDisplayPathName(TagData, Buffer);
+	return FString(Buffer);
+}
+
+void FLowLevelMemTracker::GetTagDisplayPathName(const UE::LLMPrivate::FTagData* TagData,
+	FStringBuilderBase& OutPathName, int32 MaxLen) const
+{
+	TagData->GetDisplayPath(OutPathName, MaxLen);
+}
+
+FName FLowLevelMemTracker::GetTagUniqueName(const UE::LLMPrivate::FTagData* TagData) const
+{
+	return TagData->GetName();
+}
+
+const UE::LLMPrivate::FTagData* FLowLevelMemTracker::GetTagParent(const UE::LLMPrivate::FTagData* TagData) const
+{
+	return TagData->GetParent();
+}
+
+bool FLowLevelMemTracker::GetTagIsEnumTag(const UE::LLMPrivate::FTagData* TagData) const
+{
+	return TagData->HasEnumTag();
+}
+
+ELLMTag FLowLevelMemTracker::GetTagClosestEnumTag(const UE::LLMPrivate::FTagData* TagData) const
+{
+	return TagData->GetEnumTag();
+}
+
+int64 FLowLevelMemTracker::GetTagAmountForTracker(ELLMTracker Tracker, const UE::LLMPrivate::FTagData* TagData,
+	UE::LLM::ESizeParams SizeParams)
+{
+	using namespace UE::LLMPrivate;
+
+	if (!IsEnabled())
+	{
+		return 0;
+	}
+
+	BootstrapInitialise();
+
+	if (TagData == nullptr)
+	{
+		return 0;
+	}
+
+	UE::TUniqueLock UpdateLock(UpdateMutex); // uses of TagSizes are guarded by the UpdateLock
+	return GetTracker(Tracker)->GetTagAmount(TagData, SizeParams);
+}
+
+int64 FLowLevelMemTracker::GetTagAmountForTracker(ELLMTracker Tracker, FName Tag, ELLMTagSet TagSet,
+	UE::LLM::ESizeParams SizeParams)
+{
+	if (!IsEnabled())
+	{
+		return 0;
+	}
+
+	BootstrapInitialise();
+
+	return GetTagAmountForTracker(Tracker, FindTagData(Tag, TagSet, UE::LLMPrivate::ETagReferenceSource::FunctionAPI), SizeParams);
+}
+
+int64 FLowLevelMemTracker::GetTagAmountForTracker(ELLMTracker Tracker, ELLMTag Tag, UE::LLM::ESizeParams SizeParams)
+{
+	if (!IsEnabled())
+	{
+		return 0;
+	}
+
+	BootstrapInitialise();
+
+	return GetTagAmountForTracker(Tracker, FindTagData(Tag), SizeParams);
+}
+
+void FLowLevelMemTracker::SetTagAmountForTracker(ELLMTracker Tracker, ELLMTag Tag, int64 Amount, bool bAddToTotal)
+{
+	using namespace UE::LLMPrivate;
+
+	if (!IsEnabled())
+	{
+		return;
+	}
+	BootstrapInitialise();
+	const FTagData* TagData = FindOrAddTagData(Tag);
+
+	UE::TUniqueLock UpdateLock(UpdateMutex); // uses of TagSizes are guarded by the UpdateLock
+	GetTracker(Tracker)->SetTagAmountExternal(TagData, Amount, bAddToTotal);
+}
+
+void FLowLevelMemTracker::SetTagAmountForTracker(ELLMTracker Tracker, FName Tag, ELLMTagSet TagSet, int64 Amount, bool bAddToTotal)
+{
+	using namespace UE::LLMPrivate;
+
+	if (!IsEnabled())
+	{
+		return;
+	}
+	BootstrapInitialise();
+	const FTagData* TagData = FindOrAddTagData(Tag, TagSet);
+
+	UE::TUniqueLock UpdateLock(UpdateMutex); // uses of TagSizes are guarded by the UpdateLock
+	GetTracker(Tracker)->SetTagAmountExternal(TagData, Amount, bAddToTotal);
+}
+
+const UE::LLMPrivate::FTagData* FLowLevelMemTracker::GetActiveTagData(ELLMTracker Tracker, ELLMTagSet TagSet /*= ELLMTagSet::None*/)
+{
+	if (!IsEnabled())
+	{
+		return nullptr;
+	}
+	BootstrapInitialise();
+
+	return GetTracker(Tracker)->GetActiveTagData(TagSet);
+}
+
+uint64 FLowLevelMemTracker::DumpTag( ELLMTracker Tracker, const char* FileName, int LineNumber )
+{
+	using namespace UE::LLMPrivate;
+	if (!IsEnabled())
+	{
+		return static_cast<int64>(ELLMTag::Untagged);
+	}
+	BootstrapInitialise();
+
+	const FTagData* TagData = GetActiveTagData(Tracker);
+	if (TagData)
+	{
+		FPlatformMisc::LowLevelOutputDebugStringf(TEXT("LLM TAG: %s (%lld) @ %s:%d\n"),
+			*TagData->GetDisplayName().ToString(), TagData->GetContainingEnum(),
+			FileName ? ANSI_TO_TCHAR(FileName) : TEXT("?"), LineNumber);
+		return static_cast<uint64>(TagData->GetContainingEnum());
+	}
+	else
+	{
+		FPlatformMisc::LowLevelOutputDebugStringf(TEXT("LLM TAG: No Active Tag"));
+		return static_cast<uint64>(ELLMTag::Untagged);
+	}
+}
+
+void FLowLevelMemTracker::PublishDataSingleFrame()
+{
+	if (bAutoPublish)
+	{
+		UE_LOG(LogHAL, Error, TEXT("Command must be used with the -llmdisableautopublish command line parameter."));
+	}
+	else
+	{
+		bPublishSingleFrame = true;
+	}
+}
+
+const UE::LLMPrivate::FTagData* FLowLevelMemTracker::FindOrAddTagData(ELLMTag EnumTag,
+	UE::LLMPrivate::ETagReferenceSource ReferenceSource)
+{
+	using namespace UE::LLMPrivate;
+
+	int32 Index = static_cast<int32>(EnumTag);
+	LLMCheckf(0 <= Index && Index < LLM_TAG_COUNT, TEXT("Out of range ELLMTag %d"), Index);
+
+	{
+		UE::TSharedLock Lock(TagDataMutex);
+		FTagData* TagData = TagDataEnumMap[Index];
+		if (TagData)
+		{
+			FinishConstruct(TagData, ReferenceSource);
+			return TagData;
+		}
+	}
+
+	// If we have not initialized tags yet initialize now to potentially create the custom ELLMTag we are reading.
+	if (!bFullyInitialised)
+	{
+		FinishInitialise();
+		// Reenter this function so that we retry the find above; we avoid infinite recursion because
+		// bFullyInitialised is now true.
+		return FindOrAddTagData(EnumTag, ReferenceSource);
+	}
+	LLMCheckf(!bIsBootstrapping, TEXT("LLM Error: Invalid use of custom ELLMTag when initialising tags."));
+
+	// Add the new Tag
+	FName TagName = LLMGetTagUniqueName(EnumTag);
+	{
+		FTagData* TagData = &RegisterTagData(TagName, NAME_None, NAME_None, NAME_None, NAME_None, true, EnumTag,
+			false, ReferenceSource);
+		UE::TSharedLock Lock(TagDataMutex);
+		FinishConstruct(TagData, ReferenceSource);
+		return TagData;
+	}
+}
+
+const UE::LLMPrivate::FTagData* FLowLevelMemTracker::FindOrAddTagData(FName TagName, ELLMTagSet TagSet, bool bIsStatTag,
+	UE::LLMPrivate::ETagReferenceSource ReferenceSource)
+{
+	return FindOrAddTagData(TagName, TagSet, bIsStatTag ? TagName : NAME_None, ReferenceSource);
+}
+
+const UE::LLMPrivate::FTagData* FLowLevelMemTracker::FindOrAddTagData(FName TagName, ELLMTagSet TagSet, FName StatName,
+	UE::LLMPrivate::ETagReferenceSource ReferenceSource)
+{
+	using namespace UE::LLMPrivate;
+
+	{
+		UE::TSharedLock Lock(TagDataMutex);
+		FTagData** TagDataPtr = TagDataNameMap->Find(FTagDataNameKey(TagName, TagSet));
+		if (TagDataPtr)
+		{
+			FTagData* TagData = *TagDataPtr;
+			FinishConstruct(TagData, ReferenceSource);
+			return TagData;
+		}
+	}
+
+	// If we have not initialized tags yet initialize now to potentially create the custom ELLMTag we are reading.
+	if (!bFullyInitialised)
+	{
+		FinishInitialise();
+		// Reenter this function so that we retry the find above; note we avoid infinite recursion because
+		// bFullyInitialised is now true.
+		return FindOrAddTagData(TagName, TagSet, StatName, ReferenceSource);
+	}
+	LLMCheckf(!bIsBootstrapping, TEXT("LLM Error: Invalid use of FName tag when initialising tags."));
+
+	// Add the new Tag
+	bool bIsStatTag = !StatName.IsNone() && StatName == TagName;
+	FTagData* TagData = &RegisterTagData(TagName, NAME_None, NAME_None, StatName, NAME_None, false,
+		ELLMTag::CustomName, bIsStatTag, ReferenceSource, TagSet);
+	{
+		UE::TSharedLock Lock(TagDataMutex);
+		FinishConstruct(TagData, ReferenceSource);
+		return TagData;
+	}
+}
+
+const UE::LLMPrivate::FTagData* FLowLevelMemTracker::FindTagData(ELLMTag EnumTag,
+	UE::LLMPrivate::ETagReferenceSource ReferenceSource)
+{
+	using namespace UE::LLMPrivate;
+
+	int32 Index = static_cast<int32>(EnumTag);
+	LLMCheckf(0 <= Index && Index < LLM_TAG_COUNT, TEXT("Out of range ELLMTag %d"), Index);
+
+	UE::TSharedLock Lock(TagDataMutex);
+	FTagData* TagData = TagDataEnumMap[Index];
+	if (TagData)
+	{
+		FinishConstruct(TagData, ReferenceSource);
+		return TagData;
+	}
+	else
+	{
+		return nullptr;
+	}
+}
+
+const UE::LLMPrivate::FTagData* FLowLevelMemTracker::FindTagData(FName TagName, ELLMTagSet TagSet,
+	UE::LLMPrivate::ETagReferenceSource ReferenceSource)
+{
+	using namespace UE::LLMPrivate;
+
+	UE::TSharedLock Lock(TagDataMutex);
+
+	FTagData** TagDataPtr = TagDataNameMap->Find(FTagDataNameKey(TagName, TagSet));
+	if (TagDataPtr)
+	{
+		FTagData* TagData = *TagDataPtr;
+		FinishConstruct(TagData, ReferenceSource);
+		return TagData;
+	}
+	else
+	{
+		return nullptr;
+	}
+}
+
+FLLMScope::FLLMScope(ELLMTag TagEnum, bool bInIsStatTag, ELLMTagSet InTagSet, ELLMTracker InTracker, bool bOverride)
+{
+	if (!FLowLevelMemTracker::IsEnabled())
+	{
+		return;
+	}
+
+	// We run the init in the open, because we want to track LLM even in transactions.
+	UE_AUTORTFM_OPEN
+	{
+		// This scope does not support TagSets (use FLLMScopeDynamic instead), except for the special case of
+		// anonymized allocations, which are indicated with a TagSet and ELLMTag::EngineMisc
+		LLMCheck((!bInIsStatTag && InTagSet == ELLMTagSet::None) || (TagEnum == ELLMTag::EngineMisc));
+		// ELLMTag::FMalloc is a special tag expected to only be used by the Platform Tracker (see header where defined).
+		LLMCheck((TagEnum != ELLMTag::FMalloc) || (ELLMTracker::Platform == InTracker));
+
+		FLowLevelMemTracker& LLMRef = FLowLevelMemTracker::Get();
+		// We have to check IsEnabled() again after calling Get, because the constructor is called
+		// from Get, and will set EnabledState=Disabled if the platform doesn't support it.
+		if (FLowLevelMemTracker::IsEnabled())
+		{
+			LLMRef.BootstrapInitialise();
+
+			if (bOverride || LLMRef.GetTracker(InTracker)->GetActiveTagData(InTagSet) == nullptr)
+			{
+				bEnabled = true;
+				Tracker = InTracker;
+				TagSet = InTagSet;
+				LLMRef.GetTracker(Tracker)->PushTag(TagEnum, InTagSet);
+			}
+		}
+	};
+
+	// But remember that if we abort while we hold the scope, we need to destroy the scope.
+	AutoRTFM::PushOnAbortHandler(this, [this] { if (bEnabled) { Destruct(); } });
+}
+
+FLLMScope::FLLMScope(FName TagName, bool bInIsStatTag, ELLMTagSet InTagSet, ELLMTracker InTracker, bool bOverride)
+{
+	if (!FLowLevelMemTracker::IsEnabled())
+	{
+		return;
+	}
+
+	// We run the init in the open, because we want to track LLM even in transactions.
+	UE_AUTORTFM_OPEN
+	{
+		FLowLevelMemTracker& LLMRef = FLowLevelMemTracker::Get();
+		// We have to check IsEnabled() again after calling Get, because the constructor is called
+		// from Get, and will set EnabledStae=Disabled if the platform doesn't support it.
+		if (FLowLevelMemTracker::IsEnabled())
+		{
+			LLMRef.BootstrapInitialise();
+			if (LLMRef.IsTagSetActive(InTagSet))
+			{
+				if (bOverride || LLMRef.GetTracker(InTracker)->GetActiveTagData(InTagSet) == nullptr)
+				{
+					bEnabled = true;
+					Tracker = InTracker;
+					TagSet = InTagSet;
+
+					LLMRef.GetTracker(Tracker)->PushTag(TagName, bInIsStatTag, InTagSet);
+				}
+			}
+		}
+	};
+
+	// But remember that if we abort while we hold the scope, we need to destroy the scope.
+	AutoRTFM::PushOnAbortHandler(this, [this] { if (bEnabled) { Destruct(); } });
+}
+
+FLLMScope::FLLMScope(const UE::LLMPrivate::FTagData* TagData, bool bInIsStatTag, ELLMTagSet InTagSet,
+	ELLMTracker InTracker, bool bOverride)
+{
+	if (!FLowLevelMemTracker::IsEnabled())
+	{
+		return;
+	}
+
+	// We run the init in the open, because we want to track LLM even in transactions.
+	UE_AUTORTFM_OPEN
+	{
+		FLowLevelMemTracker& LLMRef = FLowLevelMemTracker::Get();
+		// We have to check IsEnabled() again after calling Get, because the constructor is called
+		// from Get, and will set EnabledState=Disabled if the platform doesn't support it.
+		if (FLowLevelMemTracker::IsEnabled())
+		{
+			LLMRef.BootstrapInitialise();
+
+			if (bOverride || LLMRef.GetTracker(InTracker)->GetActiveTagData(InTagSet) == nullptr)
+			{
+				bEnabled = true;
+				Tracker = InTracker;
+				TagSet = InTagSet;
+				LLMRef.GetTracker(Tracker)->PushTag(TagData, InTagSet);
+			}
+		}
+	};
+
+	// But remember that if we abort while we hold the scope, we need to destroy the scope.
+	AutoRTFM::PushOnAbortHandler(this, [this] { if (bEnabled) { Destruct(); } });
+}
+
+void FLLMScope::Destruct()
+{
+	FLowLevelMemTracker& LLMRef = FLowLevelMemTracker::Get();
+	LLMRef.GetTracker(Tracker)->PopTag(TagSet);
+}
+
+FLLMScope::~FLLMScope()
+{
+	if (bEnabled)
+	{
+		// We run the destroy in the open, because we want to track LLM even in transactions.
+		UE_AUTORTFM_OPEN{ Destruct(); };
+		// But remember to pop our on-abort handler (so that we don't try and double destroy).
+		AutoRTFM::PopOnAbortHandler(this);
+	}
+}
+
+void FLLMScopeDynamic::Init(ELLMTracker InTracker, ELLMTagSet InTagSet)
+{
+	FLowLevelMemTracker& LLMRef = FLowLevelMemTracker::Get();
+	// We have to check IsEnabled() again after calling Get, because the constructor is called
+	// from Get, and will set EnabledState=Disabled if the platform doesn't support it.
+	if (!FLowLevelMemTracker::IsEnabled())
+	{
+		bEnabled = false;
+		return;
+	}
+	LLMRef.BootstrapInitialise();
+	if (!LLMRef.IsTagSetActive(InTagSet))
+	{
+		bEnabled = false;
+		return;
+	}
+	bEnabled = true;
+	TagData = nullptr;
+	Tracker = InTracker;
+	TagSet = InTagSet;
+}
+
+bool FLLMScopeDynamic::TryFindTag(FName TagName)
+{
+	FLowLevelMemTracker& LLMRef = FLowLevelMemTracker::Get();
+	TagData = LLMRef.FindTagData(TagName, TagSet, UE::LLMPrivate::ETagReferenceSource::Scope);
+	return TagData != nullptr;
+}
+
+bool FLLMScopeDynamic::TryAddTagAndActivate(FName TagName, const ILLMDynamicTagConstructor& Constructor)
+{
+	FLowLevelMemTracker& LLMRef = FLowLevelMemTracker::Get();
+	FName StatFullName = NAME_None;
+#if LLM_ENABLED_STAT_TAGS && STATS
+	FString StatConstructorName = Constructor.GetStatName();
+	if (!StatConstructorName.IsEmpty())
+	{
+		if (Constructor.NeedsStatConstruction())
+		{
+			switch (TagSet)
+			{
+			case ELLMTagSet::None:
+				StatFullName = FDynamicStats::CreateMemoryStatId<FStatGroup_STATGROUP_LLMFULL>(
+					StatConstructorName).GetName();
+				break;
+			case ELLMTagSet::Assets:
+				StatFullName = FDynamicStats::CreateMemoryStatId<FStatGroup_STATGROUP_LLMAssets>(
+					StatConstructorName).GetName();
+				break;
+			case ELLMTagSet::AssetClasses:
+				StatFullName = FDynamicStats::CreateMemoryStatId<FStatGroup_STATGROUP_LLMAssets>(
+					StatConstructorName).GetName();
+				break;
+			case ELLMTagSet::UObjectClasses:
+				StatFullName = FDynamicStats::CreateMemoryStatId<FStatGroup_STATGROUP_LLMUObjects>(
+					StatConstructorName).GetName();
+			default:
+				checkNoEntry();
+				break;
+			}
+		}
+		else
+		{
+			StatFullName = FName(*StatConstructorName);
+		}
+	}
+#endif
+	TagData = LLMRef.FindOrAddTagData(TagName, TagSet, StatFullName, UE::LLMPrivate::ETagReferenceSource::Scope);
+	LLMRef.GetTracker(Tracker)->PushTag(TagData, TagSet);
+	return true;
+}
+
+void FLLMScopeDynamic::Activate()
+{
+	FLowLevelMemTracker& LLMRef = FLowLevelMemTracker::Get();
+	LLMRef.GetTracker(Tracker)->PushTag(TagData, TagSet);
+}
+
+void FLLMScopeDynamic::Destruct()
+{
+	if (TagData)
+	{
+		FLowLevelMemTracker& LLMRef = FLowLevelMemTracker::Get();
+		LLMRef.GetTracker(Tracker)->PopTag(TagSet);
+	}
+}
+
+
+FLLMPauseScope::FLLMPauseScope(FName TagName, bool bIsStatTag, uint64 Amount, ELLMTracker TrackerToPause,
+	ELLMAllocType InAllocType)
+{
+	if (!FLowLevelMemTracker::IsEnabled())
+	{
+		bEnabled = false;
+		return;
+	}
+	Init(TagName, ELLMTag::Untagged, false, bIsStatTag, Amount, TrackerToPause, InAllocType);
+}
+
+FLLMPauseScope::FLLMPauseScope(ELLMTag TagEnum, bool bIsStatTag, uint64 Amount, ELLMTracker TrackerToPause,
+	ELLMAllocType InAllocType)
+{
+	if (!FLowLevelMemTracker::IsEnabled())
+	{
+		bEnabled = false;
+		return;
+	}
+	LLMCheck(!bIsStatTag);
+	Init(NAME_None, TagEnum, true, false, Amount, TrackerToPause, InAllocType);
+}
+
+void FLLMPauseScope::Init(FName TagName, ELLMTag EnumTag, bool bIsEnumTag, bool bIsStatTag, uint64 Amount,
+	ELLMTracker TrackerToPause, ELLMAllocType InAllocType)
+{
+	FLowLevelMemTracker& LLMRef = FLowLevelMemTracker::Get();
+	// We have to check IsEnabled() again after calling Get, because the constructor is called
+	// from Get, and will set EnabledState=Disabled if the platform doesn't support it.
+	if (!FLowLevelMemTracker::IsEnabled())
+	{
+		bEnabled = false;
+		return;
+	}
+	LLMRef.BootstrapInitialise();
+	if (!LLMRef.IsTagSetActive(ELLMTagSet::None))
+	{
+		bEnabled = false;
+		return;
+	}
+
+	bEnabled = true;
+	PausedTracker = TrackerToPause;
+	AllocType = InAllocType;
+
+	for (int32 TrackerIndex = 0; TrackerIndex < static_cast<int32>(ELLMTracker::Max); TrackerIndex++)
+	{
+		ELLMTracker Tracker = (ELLMTracker)TrackerIndex;
+
+		if (PausedTracker == ELLMTracker::Max || PausedTracker == Tracker)
+		{
+			if (Amount == 0)
+			{
+				LLMRef.GetTracker(Tracker)->Pause(InAllocType);
+			}
+			else
+			{
+				if (bIsEnumTag)
+				{
+					LLMRef.GetTracker(Tracker)->PauseAndTrackMemory(EnumTag, static_cast<int64>(Amount), InAllocType);
+				}
+				else
+				{
+					LLMRef.GetTracker(Tracker)->PauseAndTrackMemory(TagName, ELLMTagSet::None, bIsStatTag,
+						static_cast<int64>(Amount), InAllocType);
+				}
+			}
+		}
+	}
+}
+
+FLLMPauseScope::~FLLMPauseScope()
+{
+	if (!bEnabled)
+	{
+		return;
+	}
+	FLowLevelMemTracker& LLMRef = FLowLevelMemTracker::Get();
+
+	for (int32 TrackerIndex = 0; TrackerIndex < static_cast<int32>(ELLMTracker::Max); TrackerIndex++)
+	{
+		ELLMTracker Tracker = static_cast<ELLMTracker>(TrackerIndex);
+
+		if (PausedTracker == ELLMTracker::Max || Tracker == PausedTracker)
+		{
+			LLMRef.GetTracker(Tracker)->Unpause(AllocType);
+		}
+	}
+}
+
+FLLMClearScope::FLLMClearScope(ELLMTagSet InTagSet, ELLMTracker InTracker)
+	: FLLMScope(nullptr, false, InTagSet, InTracker)
+{
+}
+
+FLLMScopeFromPtr::FLLMScopeFromPtr(void* Ptr, ELLMTracker InTracker)
+{
+	using namespace UE::LLMPrivate;
+	if (!FLowLevelMemTracker::IsEnabled())
+	{
+		DisableAll();
+		return;
+	}
+	if (Ptr == nullptr)
+	{
+		DisableAll();
+		return;
+	}
+
+	FLowLevelMemTracker& LLMRef = FLowLevelMemTracker::Get();
+	// We have to check IsEnabled() again after calling Get, because the constructor is called
+	// from Get, and will set EnabledState=Disabled if the platform doesn't support it.
+	if (!FLowLevelMemTracker::IsEnabled())
+	{
+		DisableAll();
+		return;
+	}
+	LLMRef.BootstrapInitialise();
+
+	FLLMTracker* TrackerData = LLMRef.GetTracker(InTracker);
+
+	TArray<const FTagData*, TInlineAllocator<static_cast<int32>(ELLMTagSet::Max)>> Tags;
+	if (!TrackerData->FindTagsForPtr(Ptr, Tags))
+	{
+		DisableAll();
+		return;
+	}
+
+	Tracker = InTracker;
+
+	for (int32 TagSetAsInteger = 0; TagSetAsInteger < static_cast<int32>(ELLMTagSet::Max); TagSetAsInteger++)
+	{
+		if (!LLMRef.IsTagSetActive(static_cast<ELLMTagSet>(TagSetAsInteger)))
+		{
+			bEnabled[TagSetAsInteger] = false;
+
+			continue;
+		}
+
+		if (Tags[TagSetAsInteger])
+		{
+			bEnabled[TagSetAsInteger] = true;
+			TrackerData->PushTag(Tags[TagSetAsInteger], static_cast<ELLMTagSet>(TagSetAsInteger));
+		}
+		else
+		{
+			bEnabled[TagSetAsInteger] = false;
+		}
+	}
+}
+
+FLLMScopeFromPtr::~FLLMScopeFromPtr()
+{
+	FLowLevelMemTracker& LLMRef = FLowLevelMemTracker::Get();
+	for (int32 TagSetAsInteger = 0; TagSetAsInteger < static_cast<int32>(ELLMTagSet::Max); TagSetAsInteger++)
+	{
+		if (bEnabled[TagSetAsInteger])
+		{
+			LLMRef.GetTracker(Tracker)->PopTag(static_cast<ELLMTagSet>(TagSetAsInteger));
+		}
+	}
+}
+
+FLLMTagDeclaration::FLLMTagDeclaration(const TCHAR* InCPPName, FName InDisplayName, FName InParentTagName, FName InStatName, FName InSummaryStatName, ELLMTagSet InTagSet)
+	:CPPName(InCPPName), UniqueName(NAME_None), DisplayName(InDisplayName), ParentTagName(InParentTagName), StatName(InStatName), SummaryStatName(InSummaryStatName)
+	, TagSet(InTagSet)
+{
+	Register();
+}
+
+void FLLMTagDeclaration::ConstructUniqueName()
+{
+	FString NameBuffer(CPPName);
+	NameBuffer.ReplaceCharInline(TEXT('_'), TEXT('/'));
+	UniqueName = FName(*NameBuffer);
+}
+
+namespace UE::LLMPrivate::LLMTagDeclarationInternal
+{
+
+TArray<FLLMTagDeclaration::FCreationCallback, TInlineAllocator<1>>& GetCreationCallbacks()
+{
+	static TArray<FLLMTagDeclaration::FCreationCallback, TInlineAllocator<1>> CreationCallbacks;
+	return CreationCallbacks;
+}
+
+FLLMTagDeclaration*& GetList()
+{
+	static FLLMTagDeclaration* List = nullptr;
+	return List;
+}
+
+} // namespace UE::LLMPrivate::LLMTagDeclarationInternal
+
+void FLLMTagDeclaration::AddCreationCallback(FCreationCallback InCallback)
+{
+	TArray<FCreationCallback, TInlineAllocator<1>>& CreationCallbacks =
+		UE::LLMPrivate::LLMTagDeclarationInternal::GetCreationCallbacks();
+	if (CreationCallbacks.Num() >= CreationCallbacks.Max())
+	{
+		check(false); // We are not allowed to allocate memory here; you need to increase the allocation size.
+	}
+	else
+	{
+		CreationCallbacks.Add(InCallback);
+	}
+}
+
+void FLLMTagDeclaration::ClearCreationCallbacks()
+{
+	TArray<FCreationCallback, TInlineAllocator<1>>& CreationCallbacks =
+		UE::LLMPrivate::LLMTagDeclarationInternal::GetCreationCallbacks();
+	CreationCallbacks.Empty();
+}
+
+TArrayView<FLLMTagDeclaration::FCreationCallback> FLLMTagDeclaration::GetCreationCallbacks()
+{
+	return UE::LLMPrivate::LLMTagDeclarationInternal::GetCreationCallbacks();
+}
+
+FLLMTagDeclaration* FLLMTagDeclaration::GetList()
+{
+	return UE::LLMPrivate::LLMTagDeclarationInternal::GetList();
+}
+
+void FLLMTagDeclaration::Register()
+{
+	for (FCreationCallback CreationCallback : GetCreationCallbacks())
+	{
+		CreationCallback(*this);
+	}
+	FLLMTagDeclaration*& List = UE::LLMPrivate::LLMTagDeclarationInternal::GetList();
+	Next = List;
+	List = this;
+}
+
+namespace UE::LLMPrivate
+{
+#if DO_CHECK
+
+bool HandleAssert(bool bLog, const TCHAR* Format, ...)
+{
+	if (bLog)
+	{
+		TCHAR DescriptionString[4096];
+		GET_TYPED_VARARGS(TCHAR, DescriptionString, UE_ARRAY_COUNT(DescriptionString), UE_ARRAY_COUNT(DescriptionString) - 1,
+			Format, Format);
+
+		FPlatformMisc::LowLevelOutputDebugString(DescriptionString);
+
+		if (FPlatformMisc::IsDebuggerPresent())
+			FPlatformMisc::PromptForRemoteDebugging(true);
+
+		UE_DEBUG_BREAK();
+	}
+	return false;
+}
+
+#endif
+
+const TCHAR* ToString(UE::LLMPrivate::ETagReferenceSource ReferenceSource)
+{
+	switch (ReferenceSource)
+	{
+	case ETagReferenceSource::Scope:
+		return TEXT("LLM_SCOPE");
+	case ETagReferenceSource::Declare:
+		return TEXT("LLM_DEFINE_TAG");
+	case ETagReferenceSource::EnumTag:
+		return TEXT("LLM_ENUM_GENERIC_TAGS");
+	case ETagReferenceSource::CustomEnumTag:
+		return TEXT("RegisterPlatformTag/RegisterProjectTag");
+	case ETagReferenceSource::FunctionAPI:
+		return TEXT("DefaultName/InternalCall");
+	case ETagReferenceSource::ImplicitParent:
+		return TEXT("ImplicitParent");
+	default:
+		return TEXT("Invalid");
+	}
+}
+
+void ValidateUniqueName(FStringView UniqueName)
+{
+// Characters that are invalid for c++ names are invalid (other than /), since we use uniquenames
+// (with / replaced by _) as part of the name of the auto-constructed FLLMTagDeclaration variables
+// _ is invalid since we use an _ to indicate a / in FLLMTagDeclaration.
+// So only Alnum characters or / are allowed, and the first character can not be a number.
+	if (UniqueName.Len() == 0)
+	{
+		LLMCheckf(false, TEXT("Invalid length-zero Tag Unique Name"));
+	}
+	else
+	{
+	LLMCheckf(!TChar<TCHAR>::IsDigit(UniqueName[0]),
+		TEXT("Invalid first character is digit in Tag Unique Name '%.*s'"),
+		UniqueName.Len(), UniqueName.GetData());
+	}
+	for (TCHAR c : UniqueName)
+	{
+		if (!TChar<TCHAR>::IsAlnum(c) && c != TEXT('/'))
+		{
+		LLMCheckf(false, TEXT("Invalid character %c in Tag Unique Name '%.*s'"), c,
+			UniqueName.Len(), UniqueName.GetData());
+		}
+	}
+}
+
+
+namespace AllocatorPrivate
+{
+
+/**
+ * When a Page is allocated, it splits the memory of the page up into blocks, and creates an FAlloc at the start of
+ * each block. All the FAllocs are joined together in a FreeList linked list.
+ * When a Page allocates memory, it takes an FAlloc from the freelist and gives it to the caller, and forgets about it.
+ * When the caller returns a pointer, the Page restores the FAlloc at the beginning of the block and puts it back on
+ * the FreeList.
+ */
+struct FAlloc
+{
+	FAlloc* Next;
+};
+
+/**
+ * An FPage holds a single page of memory received from the OS; all pages are of the same size.
+ * FPages are owned by FBins, and the FPages for an FBin divide the page up into blocks of the FBin's size.
+ * An FPage keeps track of the blocks it has not yet given out so it can allocate, and keeps track of how many blocks
+ * it has given out, so that it can be freed when no longer used.
+ * Pages that are not free or empty are available for allocating from and are kept in a doubly-linked list on the FBin.
+ */
+struct FPage
+{
+	FPage(int32 PageSize, int32 BinSize);
+	void* Allocate();
+	void Free(void* Ptr);
+	bool IsFull() const;
+	bool IsEmpty() const;
+	void AddToList(FPage*& Head);
+	void RemoveFromList(FPage*& Head);
+
+	FAlloc* FreeList;
+	FPage* Prev;
+	FPage* Next;
+	int32 UsedCount;
+};
+
+/**
+ * An FBin handles all allocations that fit in its size range. Its size is the power of two at the top of that range.
+ * The FBin allocates one FPage at a time from the OS; the FPage gets split up into blocks and handles providing a
+ * block for callers requesting a pointer.
+ * The FBin has a doubly-linked list of pages in use but not yet full. It provides allocations from these pages.
+ * When an FPage gets full, the FBin forgets about it, counting on the caller to give the pointer to the page back
+ * when it frees the pointer and the page becomes non-full again.
+ * When an FBin has no more non-full pages and needs to satisfy an alloc, it allocates a new page.
+ * When a page becomes unused due to a free, the FBin frees the page, returning it to the OS.
+ */
+struct FBin
+{
+	FBin(int32 InBinSize);
+	void* Allocate(FLLMAllocator& Allocator);
+	void Free(void* Ptr, FLLMAllocator& Allocator);
+
+	FPage* FreePages;
+	int32 UsedCount;
+	int32 BinSize;
+};
+
+} // namespace AllocatorPrivate
+
+FLLMAllocator*& FLLMAllocator::Get()
+{
+	static FLLMAllocator* Allocator = nullptr;
+	return Allocator;
+}
+
+FLLMAllocator::FLLMAllocator()
+	: PlatformAlloc(nullptr)
+	, PlatformFree(nullptr)
+	, Bins(nullptr)
+	, Total(0)
+	, PageSize(0)
+	, NumBins(0)
+{
+}
+
+FLLMAllocator::~FLLMAllocator()
+{
+	Clear();
+}
+
+void FLLMAllocator::Initialise(LLMAllocFunction InAlloc, LLMFreeFunction InFree, int32 InPageSize)
+{
+	using namespace UE::LLMPrivate::AllocatorPrivate;
+
+	PlatformAlloc = InAlloc;
+	PlatformFree = InFree;
+	PageSize = InPageSize;
+
+	if (PlatformAlloc)
+	{
+		constexpr int32 MinBinSizeForAlignment = 16;
+		constexpr int32 MinBinSizeForAllocStorage = static_cast<int32>(sizeof(FAlloc));
+		constexpr int32 MultiplierBetweenBins = 2;
+		// Setting MultiplierAfterLastBin=2 would be useless because the PageSize/2 bin would only get a single
+		// allocation out of each page due to the FPage data taking up the first half
+		// TODO: For bins >= 4*FPage size, allocate FPages in a separate list rather than embedding them.
+		// This will require allocating extra space in each allocation to store its page pointer.
+		constexpr int32 MultiplierAfterLastBin = 4;
+
+		int32 MinBinSize = FMath::Max(MinBinSizeForAllocStorage, MinBinSizeForAlignment);
+		int32 MaxBinSize = InPageSize / MultiplierAfterLastBin;
+		int32 BinSize = MinBinSize;
+		while (BinSize <= MaxBinSize)
+		{
+			BinSize *= MultiplierBetweenBins;
+			++NumBins;
+		}
+		if (NumBins > 0)
+		{
+			Bins = reinterpret_cast<FBin*>(AllocPages(NumBins * sizeof(FBin)));
+			BinSize = MinBinSize;
+			for (int32 BinIndex = 0; BinIndex < NumBins; ++BinIndex)
+			{
+				new (&Bins[BinIndex]) FBin(BinSize);
+				BinSize *= MultiplierBetweenBins;
+			}
+		}
+	}
+}
+
+void FLLMAllocator::Clear()
+{
+	using namespace UE::LLMPrivate::AllocatorPrivate;
+
+	if (NumBins)
+	{
+		for (int32 BinIndex = 0; BinIndex < NumBins; ++BinIndex)
+		{
+			LLMCheck(Bins[BinIndex].UsedCount == 0);
+			Bins[BinIndex].~FBin();
+		}
+		FreePages(Bins, NumBins * sizeof(FBin));
+		Bins = nullptr;
+		NumBins = 0;
+	}
+}
+
+void* FLLMAllocator::Malloc(size_t Size)
+{
+	return Alloc(Size);
+}
+
+void* FLLMAllocator::Alloc(size_t Size)
+{
+	using namespace UE::LLMPrivate::AllocatorPrivate;
+
+	if (Size == 0)
+	{
+		return nullptr;
+	}
+	int32 BinIndex = GetBinIndex(Size);
+	UE::TUniqueLock Lock(Mutex);
+	if (BinIndex == NumBins)
+	{
+		return AllocPages(Size);
+	}
+	return Bins[BinIndex].Allocate(*this);
+}
+
+void FLLMAllocator::Free(void* Ptr, size_t Size)
+{
+	using namespace UE::LLMPrivate::AllocatorPrivate;
+
+	if (Ptr != nullptr)
+	{
+		int32 BinIndex = GetBinIndex(Size);
+		UE::TUniqueLock Lock(Mutex);
+		if (BinIndex == NumBins)
+		{
+			FreePages(Ptr, Size);
+		}
+		else
+		{
+			Bins[BinIndex].Free(Ptr, *this);
+		}
+	}
+}
+
+void* FLLMAllocator::Realloc(void* Ptr, size_t OldSize, size_t NewSize)
+{
+	void* NewPtr;
+	if (NewSize)
+	{
+		NewPtr = Alloc(NewSize);
+		if (OldSize)
+		{
+			size_t CopySize = FMath::Min(OldSize, NewSize);
+			FMemory::Memcpy(NewPtr, Ptr, CopySize);
+		}
+	}
+	else
+	{
+		NewPtr = nullptr;
+	}
+	Free(Ptr, OldSize);
+	return NewPtr;
+}
+
+int64 FLLMAllocator::GetTotal() const
+{
+	UE::TUniqueLock Lock(const_cast<UE::FPlatformRecursiveMutex&>(Mutex));
+	return Total;
+}
+
+void* FLLMAllocator::AllocPages(size_t Size)
+{
+	Size = Align(Size, PageSize);
+	void* Ptr = PlatformAlloc(Size);
+	LLMCheck(Ptr);
+	LLMCheck((reinterpret_cast<intptr_t>(Ptr) & (PageSize - 1)) == 0);
+	Total += Size;
+	return Ptr;
+}
+
+void FLLMAllocator::FreePages(void* Ptr, size_t Size)
+{
+	Size = Align(Size, PageSize);
+	PlatformFree(Ptr, Size);
+	Total -= Size;
+}
+
+int32 FLLMAllocator::GetBinIndex(size_t Size) const
+{
+	int BinIndex = 0;
+	while (BinIndex < NumBins && static_cast<size_t>(Bins[BinIndex].BinSize) < Size)
+	{
+		++BinIndex;
+	}
+	return BinIndex;
+}
+
+namespace AllocatorPrivate
+{
+
+FPage::FPage(int32 PageSize, int32 BinSize)
+{
+	Next = Prev = nullptr;
+	UsedCount = 0;
+	int32 NumHeaderBins = (FMath::Max(static_cast<int32>(sizeof(FPage)), BinSize) + BinSize - 1) / BinSize;
+	int32 FreeCount = PageSize / BinSize - NumHeaderBins;
+
+	// Divide the rest of the page after this header into FAllocs, and add all the FAllocs into the free list.
+	FreeList = reinterpret_cast<FAlloc*>(reinterpret_cast<intptr_t>(this) + NumHeaderBins*BinSize);
+	FAlloc* EndAlloc = reinterpret_cast<FAlloc*>(
+	reinterpret_cast<intptr_t>(FreeList) + (FreeCount-1) * BinSize);
+	FAlloc* Alloc = FreeList;
+	while (Alloc != EndAlloc)
+	{
+		Alloc->Next = reinterpret_cast<FAlloc*>(reinterpret_cast<intptr_t>(Alloc) + BinSize);
+		Alloc = Alloc->Next;
+	}
+	EndAlloc->Next = nullptr;
+}
+
+void* FPage::Allocate()
+{
+	LLMCheck(FreeList);
+	FAlloc* Alloc = FreeList;
+	FreeList = Alloc->Next;
+	++UsedCount;
+	return Alloc;
+}
+
+void FPage::Free(void* Ptr)
+{
+	LLMCheck(UsedCount > 0);
+	FAlloc* Alloc = reinterpret_cast<FAlloc*>(Ptr);
+	Alloc->Next = FreeList;
+	FreeList = Alloc;
+	--UsedCount;
+}
+
+bool FPage::IsFull() const
+{
+	return FreeList == nullptr;
+}
+
+bool FPage::IsEmpty() const
+{
+	return UsedCount == 0;
+}
+
+void FPage::AddToList(FPage*& Head)
+{
+	Next = Head;
+	Prev = nullptr;
+	Head = this;
+	if (Next)
+	{
+		Next->Prev = this;
+	}
+}
+
+void FPage::RemoveFromList(FPage*& Head)
+{
+	if (Prev)
+	{
+		Prev->Next = Next;
+		if (Next)
+		{
+			Next->Prev = Prev;
+		}
+	}
+	else
+	{
+		Head = Next;
+		if (Next)
+		{
+			Next->Prev = nullptr;
+		}
+	}
+	Next = Prev = nullptr;
+}
+
+FBin::FBin(int32 InBinSize)
+{
+	FreePages = nullptr;
+	UsedCount = 0;
+	BinSize = InBinSize;
+}
+
+void* FBin::Allocate(FLLMAllocator& Allocator)
+{
+	if (!FreePages)
+	{
+		FPage* Page = reinterpret_cast<FPage*>(Allocator.AllocPages(Allocator.PageSize));
+		++UsedCount;
+		LLMCheck(Page);
+		// The FPage is at the beginning of the array of PageSize bytes.
+		new (Page) FPage(Allocator.PageSize, BinSize);
+		Page->AddToList(FreePages);
+	}
+
+	void* Result = FreePages->Allocate();
+	if (FreePages->IsFull())
+	{
+		FreePages->RemoveFromList(FreePages); //-V678
+	}
+	return Result;
+}
+
+void FBin::Free(void* Ptr, FLLMAllocator& Allocator)
+{
+	FPage* Page = reinterpret_cast<FPage*>(
+		reinterpret_cast<intptr_t>(Ptr) & ~(static_cast<intptr_t>(Allocator.PageSize) - 1));
+	if (Page->IsFull())
+	{
+		Page->AddToList(FreePages);
+	}
+	Page->Free(Ptr);
+	if (Page->IsEmpty())
+	{
+		Page->RemoveFromList(FreePages);
+		--UsedCount;
+		Allocator.FreePages(Page, Allocator.PageSize);
+	}
+}
+
+} // namespace AllocatorPrivate
+
+FTagData::FTagData(FName InName, ELLMTagSet InTagSet, FName InDisplayName, FName InParentName, FName InStatName,
+	FName InSummaryStatName, bool bInHasEnumTag, ELLMTag InEnumTag, ETagReferenceSource InReferenceSource)
+	: Name(InName), DisplayName(InDisplayName), ParentName(InParentName), StatName(InStatName)
+	, SummaryStatName(InSummaryStatName), EnumTag(InEnumTag), ReferenceSource(InReferenceSource)
+	, TagSet(InTagSet), bIsFinishConstructed(false), bParentIsName(true), bHasEnumTag(bInHasEnumTag)
+	, bIsReportable(true), bIsTraceReportable(true)
+{
+}
+
+FTagData::FTagData(FName InName, ELLMTagSet InTagSet, FName InDisplayName, const FTagData* InParent, FName InStatName,
+	FName InSummaryStatName, bool bInHasEnumTag, ELLMTag InEnumTag, ETagReferenceSource InReferenceSource)
+	: FTagData(InName, InTagSet, InDisplayName, NAME_None, InStatName, InSummaryStatName, bInHasEnumTag, InEnumTag,
+		InReferenceSource)
+{
+	SetParent(InParent);
+}
+
+FTagData::FTagData(ELLMTag InEnumTag)
+	: FTagData(NAME_None, ELLMTagSet::None, NAME_None, NAME_None, NAME_None, NAME_None, true, InEnumTag,
+		ETagReferenceSource::EnumTag)
+{
+}
+
+FTagData::~FTagData()
+{
+	if (bParentIsName)
+	{
+		ParentName.~FName();
+		bParentIsName = false;
+	}
+}
+
+bool FTagData::IsFinishConstructed() const
+{
+	return bIsFinishConstructed;
+}
+
+bool FTagData::IsParentConstructed() const
+{
+	return !bParentIsName;
+}
+
+FName FTagData::GetName() const
+{
+	return Name;
+}
+
+FName FTagData::GetDisplayName() const
+{
+	return DisplayName;
+}
+
+void FTagData::GetDisplayPath(FStringBuilderBase& Result, int32 MaxLen) const
+{
+	Result.Reset();
+	AppendDisplayPath(Result, MaxLen);
+}
+
+void FTagData::AppendDisplayPath(FStringBuilderBase& Result, int32 MaxLen) const
+{
+	if (Parent && Parent->IsUsedAsDisplayParent())
+	{
+		Parent->AppendDisplayPath(Result, MaxLen);
+		if (MaxLen >= 0 && Result.Len() + 1 >= MaxLen)
+		{
+			return;
+		}
+		Result << TEXT("/");
+	}
+	if (MaxLen >= 0)
+	{
+		int32 MaxRemainingLen = MaxLen - Result.Len();
+		if (static_cast<int32>(DisplayName.GetStringLength() + 1) > MaxRemainingLen)
+		{
+			if (MaxRemainingLen > 1)
+			{
+				TStringBuilder<FName::StringBufferSize> Buffer;
+				DisplayName.AppendString(Buffer);
+				Result << Buffer.ToView().Left(MaxRemainingLen - 1);
+			}
+			return;
+		}
+	}
+	DisplayName.AppendString(Result);
+}
+
+const FTagData* FTagData::GetParent() const
+{
+	LLMCheckf(!bParentIsName, TEXT("GetParent called on TagData %s before SetParent was called"),
+		*WriteToString<FName::StringBufferSize>(Name));
+	return Parent;
+}
+
+FName FTagData::GetParentName() const
+{
+	LLMCheckf(bParentIsName, TEXT("GetParentName called on TagData %s after SetParent was called"),
+		*WriteToString<FName::StringBufferSize>(Name));
+	return ParentName;
+}
+
+FName FTagData::GetParentNameSafeBeforeFinishConstruct() const
+{
+	if (bParentIsName)
+	{
+		return ParentName;
+	}
+	else
+	{
+		return Parent ? Parent->GetName() : NAME_None;
+	}
+}
+
+FName FTagData::GetStatName() const
+{
+	return StatName;
+}
+
+FName FTagData::GetSummaryStatName() const
+{
+	return SummaryStatName;
+}
+
+ELLMTag FTagData::GetEnumTag() const
+{
+	return EnumTag;
+}
+
+ELLMTagSet FTagData::GetTagSet() const
+{
+	return TagSet;
+}
+
+bool FTagData::HasEnumTag() const
+{
+	return bHasEnumTag;
+}
+
+const FTagData* FTagData::GetContainingEnumTagData() const
+{
+	const FTagData* TagData = this;
+	do
+	{
+		if (TagData->bHasEnumTag)
+		{
+			return TagData;
+		}
+		TagData = TagData->GetParent();
+	} while (TagData);
+	LLMCheckf(false, TEXT("TagData is not a descendant of an ELLMTag TagData. ")
+		TEXT("All TagDatas must be descendants of ELLMTag::CustomName if they are not descendants of any other ELLMTag"));
+		return this;
+}
+
+ELLMTag FTagData::GetContainingEnum() const
+{
+	const FTagData* TagData = GetContainingEnumTagData();
+	return TagData->EnumTag;
+}
+
+ETagReferenceSource FTagData::GetReferenceSource() const
+{
+	return ReferenceSource;
+}
+
+int32 FTagData::GetIndex() const
+{
+	return Index;
+}
+
+bool FTagData::IsReportable() const
+{
+	return bIsReportable && TagSet == ELLMTagSet::None;
+}
+
+bool FTagData::IsStatsReportable() const
+{
+	return bIsReportable;
+}
+
+bool FTagData::IsTraceReportable() const
+{
+	return bIsTraceReportable && bIsReportable;
+}
+
+void FTagData::SetParent(const FTagData* InParent)
+{
+	if (bParentIsName)
+	{
+		ParentName.~FName();
+		bParentIsName = false;
+	}
+	Parent = InParent;
+}
+
+void FTagData::SetName(FName InName)
+{
+	Name = InName;
+}
+
+void FTagData::SetDisplayName(FName InDisplayName)
+{
+	DisplayName = InDisplayName;
+}
+
+void FTagData::SetStatName(FName InStatName)
+{
+	StatName = InStatName;
+}
+
+void FTagData::SetSummaryStatName(FName InSummaryStatName)
+{
+	SummaryStatName = InSummaryStatName;
+}
+
+void FTagData::SetParentName(FName InParentName)
+{
+	LLMCheck(bParentIsName);
+	ParentName = InParentName;
+}
+
+void FTagData::SetFinishConstructed()
+{
+	bIsFinishConstructed = true;
+}
+
+void FTagData::SetIndex(int32 InIndex)
+{
+	Index = InIndex;
+}
+
+void FTagData::SetIsReportable(bool bInIsReportable)
+{
+	bIsReportable = bInIsReportable;
+}
+
+void FTagData::SetIsTraceReportable(bool bInIsTraceReportable)
+{
+	bIsTraceReportable = bInIsTraceReportable;
+}
+
+bool FTagData::IsUsedAsDisplayParent() const
+{
+	// All Tags but one are UsedAsDisplayParent - their name is prepended during GetDisplayPath.
+	// ELLMTag::CustomName is the exception. It is set for FName tags that do not have a real parent to provide a
+	// containing ELLMTag for them to provide to systems that do not support FName tags. When FName tags without a real
+	// parent are displayed, their path should display as parentless despite having the CustomName tag as their parent.
+	return !(bHasEnumTag && EnumTag == ELLMTag::CustomName);
+}
+
+FLLMTracker::FLLMTracker(FLowLevelMemTracker& InLLM)
+	: LLMRef(InLLM)
+	, Tracker(ELLMTracker::Max)
+	, TrackedTotal(0)
+	, OverrideUntaggedTagData(nullptr)
+	, OverrideTrackedTotalTagData(nullptr)
+	, LastTrimTime(0.0)
+{
+	TlsSlot = FPlatformTLS::AllocTlsSlot();
+
+	for (int32 Index = 0; Index < static_cast<int32>(ELLMAllocType::Count); ++Index)
+	{
+		AllocTypeAmounts[Index] = 0;
+	}
+}
+
+FLLMTracker::~FLLMTracker()
+{
+	Clear();
+
+	FPlatformTLS::FreeTlsSlot(TlsSlot);
+}
+
+void FLLMTracker::Initialise(
+	ELLMTracker InTracker,
+	FLLMAllocator* InAllocator)
+{
+	check(InTracker != ELLMTracker::Max);
+	Tracker = InTracker;
+	CsvWriter.SetTracker(InTracker);
+	TraceWriter.SetTracker(InTracker);
+	CsvProfilerWriter.SetTracker(InTracker);
+
+#if !UE_ONLY_USE_PLATFORM_TRACKER
+	AllocationMap.SetAllocator(InAllocator);
+#endif
+}
+
+FLLMThreadState* FLLMTracker::GetOrCreateState()
+{
+	// look for already allocated thread state
+	FLLMThreadState* State = GetState();
+	// Create one if needed
+	if (State == nullptr)
+	{
+		State = LLMRef.Allocator.New<FLLMThreadState>();
+		LLMCheckf(State != nullptr, TEXT("LLMRef.Allocator.New returned nullptr."));
+
+		// Add to pending thread states, (these will be consumed on the main thread and transferred to ThreadStates,
+		// which is only read/write on main thread). Also add to our backup map from thread id to thread state.
+		uint32 ThreadId = FPlatformTLS::GetCurrentThreadId();
+		{
+			FScopeLock Lock(&PendingThreadStatesGuard);
+			PendingThreadStates.Add(State);
+			ThreadIdToThreadState.Add(ThreadId, State);
+		}
+
+		// push to Tls
+		FPlatformTLS::SetTlsValue(TlsSlot, State);
+	}
+	return State;
+}
+
+FLLMThreadState* FLLMTracker::GetState()
+{
+	FLLMThreadState* State = (FLLMThreadState*)FPlatformTLS::GetTlsValue(TlsSlot);
+	if (!State)
+	{
+		// GetTlsValue might return null even if we previously set it, if called during thread termination
+		// Check our backup mapping from thread id to thread state before concluding the state does not exist.
+		uint32 ThreadId = FPlatformTLS::GetCurrentThreadId();
+		{
+			FScopeLock Lock(&PendingThreadStatesGuard);
+			FLLMThreadState** ExistingState = ThreadIdToThreadState.Find(ThreadId);
+			if (ExistingState)
+			{
+				State = *ExistingState;
+			}
+		}
+	}
+	return State; // Can be nullptr if not yet created
+}
+
+void FLLMTracker::PushTag(ELLMTag EnumTag, ELLMTagSet TagSet)
+{
+	const FTagData* TagData = LLMRef.FindOrAddTagData(EnumTag, ETagReferenceSource::Scope);
+
+	// pass along to the state object
+	GetOrCreateState()->PushTag(TagData, TagSet);
+}
+
+void FLLMTracker::PushTag(FName Tag, bool bInIsStatData, ELLMTagSet TagSet)
+{
+	const FTagData* TagData = LLMRef.FindOrAddTagData(Tag, TagSet, bInIsStatData, ETagReferenceSource::Scope);
+
+	// pass along to the state object
+	GetOrCreateState()->PushTag(TagData, TagSet);
+}
+
+void FLLMTracker::PushTag(const FTagData* TagData, ELLMTagSet TagSet)
+{
+	// pass along to the state object
+	GetOrCreateState()->PushTag(TagData, TagSet);
+}
+
+void FLLMTracker::PopTag(ELLMTagSet TagSet)
+{
+	// look for already allocated thread state
+	FLLMThreadState* State = GetState();
+
+	LLMCheckf(State != nullptr, TEXT("Called PopTag but PushTag was never called!"));
+
+	State->PopTag(TagSet);
+}
+
+void FLLMTracker::TrackAllocation(const void* Ptr, int64 Size, ELLMTag DefaultTag, ELLMAllocType AllocType, bool bTrackInMemPro)
+{
+	FLLMThreadState* State = GetOrCreateState();
+	const FTagData* TagData = State->GetTopTag(ELLMTagSet::None);
+	if (!TagData)
+	{
+		TagData = LLMRef.FindOrAddTagData(DefaultTag);
+	}
+	TrackAllocation(Ptr, Size, TagData, AllocType, State, bTrackInMemPro);
+}
+
+void FLLMTracker::TrackAllocation(const void* Ptr, int64 Size, FName DefaultTag, ELLMAllocType AllocType, bool bTrackInMemPro)
+{
+	FLLMThreadState* State = GetOrCreateState();
+	const FTagData* TagData = State->GetTopTag(ELLMTagSet::None);
+	if (!TagData)
+	{
+		TagData = LLMRef.FindOrAddTagData(DefaultTag, ELLMTagSet::None);
+	}
+	TrackAllocation(Ptr, Size, TagData, AllocType, State, bTrackInMemPro);
+}
+
+void FLLMTracker::TrackAllocation(const void* Ptr, int64 Size, const FTagData* ActiveTagData, ELLMAllocType AllocType, FLLMThreadState* State, bool bTrackInMemPro)
+{
+	if (IsPaused(AllocType))
+	{
+		// When Paused, we do not track any new allocations and we do not update the counters for the memory
+		// they use; the code that triggered the pause is responsible for updating those counters. Since we do not
+		// track the allocations, TrackFree will likewise not update the counters when those allocations are freed.
+		return;
+	}
+
+	// track the total quickly
+	FPlatformAtomics::InterlockedAdd(&TrackedTotal, Size);
+
+#if !LLM_ENABLED_FULL_TAGS
+	// When full tags are disabled, we instead store the top-level enumtag parent of the allocation's tag.
+	ActiveTagData = ActiveTagData->GetContainingEnumTagData();
+#endif
+
+#if LLM_ALLOW_ASSETS_TAGS
+	const FTagData* AssetTagData = State->GetTopTag(ELLMTagSet::Assets);
+	const FTagData* AssetClassTagData = State->GetTopTag(ELLMTagSet::AssetClasses);
+
+	if (AssetTagData == nullptr && !LLMRef.IsBootstrapping() && LLMRef.IsTagSetActive(ELLMTagSet::Assets))
+	{
+		AssetTagData = LLMRef.FindOrAddTagData(TagName_UntaggedAsset, ELLMTagSet::Assets, false, UE::LLMPrivate::ETagReferenceSource::FunctionAPI);
+	}
+	if (AssetClassTagData == nullptr && !LLMRef.IsBootstrapping() && LLMRef.IsTagSetActive(ELLMTagSet::AssetClasses))
+	{
+		AssetClassTagData = LLMRef.FindOrAddTagData(TagName_UntaggedAssetClass, ELLMTagSet::AssetClasses, false, UE::LLMPrivate::ETagReferenceSource::FunctionAPI);
+	}
+#else
+	const FTagData* AssetTagData = nullptr;
+	const FTagData* AssetClassTagData = nullptr;
+#endif // LLM_ALLOW_ASSETS_TAGS
+
+#if LLM_ALLOW_UOBJECTCLASSES_TAGS
+	const FTagData* UObjectClassesTagData = State->GetTopTag(ELLMTagSet::UObjectClasses);
+#else
+	const FTagData* UObjectClassesTagData = nullptr;
+#endif
+
+	// track on the thread state
+	State->TrackAllocation(Ptr, Size, Tracker, AllocType, ActiveTagData, AssetTagData, AssetClassTagData, UObjectClassesTagData, bTrackInMemPro);
+
+#if !UE_ONLY_USE_PLATFORM_TRACKER
+	// tracking a nullptr with a Size is allowed, but we don't need to remember it, since we can't free it ever.
+	if (Ptr != nullptr)
+	{
+		// remember the size and tag info
+		FLLMTracker::FLowLevelAllocInfo AllocInfo;
+		AllocInfo.SetTag(ActiveTagData, LLMRef);
+#if LLM_ALLOW_ASSETS_TAGS
+		AllocInfo.SetTag(AssetTagData, LLMRef, ELLMTagSet::Assets);
+		AllocInfo.SetTag(AssetClassTagData, LLMRef, ELLMTagSet::AssetClasses);
+#endif
+#if LLM_ALLOW_UOBJECTCLASSES_TAGS
+		// If there is no scoped tag for UObjectClasses, don't set it
+		if (UObjectClassesTagData != nullptr)
+		{
+			AllocInfo.SetTag(UObjectClassesTagData, LLMRef, ELLMTagSet::UObjectClasses);
+		}
+#endif
+		LLMCheck(Size <= 0x0000'ffff'ffff'ffff);
+		uint32 SizeLow = uint32(Size);
+		uint16 SizeHigh = uint16(Size >> 32ull);
+		PointerKey Key(Ptr, SizeHigh);
+		AllocationMap.Add(Key, SizeLow, AllocInfo);
+	}
+#endif
+}
+
+void FLLMTracker::TrackFree(const void* Ptr, ELLMAllocType AllocType, bool bTrackInMemPro)
+{
+#if !UE_ONLY_USE_PLATFORM_TRACKER
+	// look up the pointer in the tracking map
+	FLLMAllocMap::Values Values;
+	{
+		if (!AllocationMap.Remove(PointerKey(Ptr), Values))
+		{
+			return;
+		}
+	}
+
+	if (IsPaused(AllocType))
+	{
+		// When Paused, we remove our data for any freed allocations, but we do not update the counters for the
+		// memory they used; the code that triggered the pause is responsible for updating those counters.
+		return;
+	}
+
+	int64 SizeLow = static_cast<int64>(Values.Value1);
+	int64 SizeHigh = Values.Key.GetExtraData();
+	int64 Size = (SizeHigh << 32ull) | SizeLow;
+	FLLMTracker::FLowLevelAllocInfo& AllocInfo = Values.Value2;
+
+	// track the total quickly
+	FPlatformAtomics::InterlockedAdd(&TrackedTotal, -Size);
+
+	FLLMThreadState* State = GetOrCreateState();
+	const FTagData* TagData = AllocInfo.GetTag(LLMRef);
+#if LLM_ALLOW_ASSETS_TAGS
+	const FTagData* AssetTagData = AllocInfo.GetTag(LLMRef, ELLMTagSet::Assets);
+	const FTagData* AssetClassTagData = AllocInfo.GetTag(LLMRef, ELLMTagSet::AssetClasses);
+#else
+	const FTagData* AssetTagData = nullptr;
+	const FTagData* AssetClassTagData = nullptr;
+#endif // LLM_ALLOW_ASSETS_TAGS
+
+#if LLM_ALLOW_UOBJECTCLASSES_TAGS
+	const FTagData* UObjectClassesTagData = AllocInfo.GetTag(LLMRef, ELLMTagSet::UObjectClasses);
+#else
+	const FTagData* UObjectClassesTagData = nullptr;
+#endif // LLM_ALLOW_UOBJECTCLASSES_TAGS
+
+#else
+
+	FLLMThreadState* State = GetOrCreateState();
+	int64 Size = 0;
+	const FTagData* TagData = nullptr;
+	const FTagData* AssetTagData = nullptr;
+	const FTagData* AssetClassTagData = nullptr;
+	const FTagData* UObjectClassesTagData = nullptr;
+
+#endif // !UE_ONLY_USE_PLATFORM_TRACKER
+
+	State->TrackFree(Ptr, Size, Tracker, AllocType, TagData, AssetTagData, AssetClassTagData, UObjectClassesTagData, bTrackInMemPro);
+}
+
+void FLLMTracker::OnAllocMoved(const void* Dest, const void* Source, ELLMAllocType AllocType)
+{
+#if !UE_ONLY_USE_PLATFORM_TRACKER
+	FLLMAllocMap::Values Values;
+	{
+		if (!AllocationMap.Remove(PointerKey(Source), Values))
+		{
+			return;
+		}
+
+		PointerKey Key(Dest, uint16(Values.Key.GetExtraData()));
+		AllocationMap.Add(Key, Values.Value1, Values.Value2);
+	}
+
+	if (IsPaused(AllocType))
+	{
+		// When Paused, don't update counters in case any of the external tracking systems are not available.
+		return;
+	}
+
+	int64 SizeLow = static_cast<int64>(Values.Value1);
+	int64 SizeHigh = Values.Key.GetExtraData();
+	int64 Size = (SizeHigh << 32ull) | SizeLow;
+	const FLLMTracker::FLowLevelAllocInfo& AllocInfo = Values.Value2;
+	const FTagData* TagData = AllocInfo.GetTag(LLMRef);
+#else
+	// TODO: This loses the memory size of an allocation moved within a larger allocation.
+	int64 Size = 0;
+	const FTagData* TagData = nullptr;
+#endif
+	FLLMThreadState* State = GetOrCreateState();
+	State->TrackMoved(Dest, Source, Size, Tracker, TagData);
+}
+
+void FLLMTracker::TrackMemoryOfActiveTag(int64 Amount, FName DefaultTag, ELLMAllocType AllocType)
+{
+	FLLMThreadState* State = GetOrCreateState();
+	const FTagData* TagData = State->GetTopTag(ELLMTagSet::None);
+	if (!TagData)
+	{
+		TagData = LLMRef.FindOrAddTagData(DefaultTag, ELLMTagSet::None);
+	}
+	TrackMemoryOfActiveTag(Amount, TagData, AllocType, State);
+}
+
+void FLLMTracker::TrackMemoryOfActiveTag(int64 Amount, ELLMTag DefaultTag, ELLMAllocType AllocType)
+{
+	FLLMThreadState* State = GetOrCreateState();
+	const FTagData* TagData = State->GetTopTag(ELLMTagSet::None);
+	if (!TagData)
+	{
+		TagData = LLMRef.FindOrAddTagData(DefaultTag);
+	}
+	TrackMemoryOfActiveTag(Amount, TagData, AllocType, State);
+}
+
+void FLLMTracker::TrackMemoryOfActiveTag(int64 Amount, const FTagData* TagData, ELLMAllocType AllocType, FLLMThreadState* State)
+{
+	if (IsPaused(AllocType))
+	{
+		// When Paused, we do not track any delta memory; the code that triggered the pause is responsible for updating the delta memory
+		return;
+	}
+
+	// track the total quickly
+	FPlatformAtomics::InterlockedAdd(&TrackedTotal, Amount);
+
+#if !LLM_ENABLED_FULL_TAGS
+	// When full tags are disabled, we instead store the top-level enumtag parent of the tag used by each allocation
+	TagData = TagData->GetContainingEnumTagData();
+#endif
+
+	// track on the thread state
+	State->TrackMemory(Amount, Tracker, AllocType, TagData);
+}
+
+void FLLMTracker::TrackMemory(ELLMTag Tag, int64 Amount, ELLMAllocType AllocType)
+{
+	TrackMemory(LLMRef.FindOrAddTagData(Tag), Amount, AllocType);
+}
+
+void FLLMTracker::TrackMemory(FName Tag, ELLMTagSet TagSet, int64 Amount, ELLMAllocType AllocType)
+{
+	TrackMemory(LLMRef.FindOrAddTagData(Tag, TagSet, false, UE::LLMPrivate::ETagReferenceSource::FunctionAPI), Amount, AllocType);
+}
+
+void FLLMTracker::TrackMemory(const FTagData* TagData, int64 Amount, ELLMAllocType AllocType)
+{
+	FLLMThreadState* State = GetOrCreateState();
+	FPlatformAtomics::InterlockedAdd(&TrackedTotal, Amount);
+	State->TrackMemory(Amount, Tracker, AllocType, TagData);
+}
+
+void FLLMTracker::PauseAndTrackMemory(FName TagName, ELLMTagSet TagSet, bool bInIsStatTag, int64 Amount, ELLMAllocType AllocType)
+{
+	const FTagData* TagData = LLMRef.FindOrAddTagData(TagName, TagSet, bInIsStatTag, UE::LLMPrivate::ETagReferenceSource::FunctionAPI);
+	PauseAndTrackMemory(TagData, Amount, AllocType);
+}
+
+void FLLMTracker::PauseAndTrackMemory(ELLMTag EnumTag, int64 Amount, ELLMAllocType AllocType)
+{
+	const FTagData* TagData = LLMRef.FindOrAddTagData(EnumTag);
+	PauseAndTrackMemory(TagData, Amount, AllocType);
+}
+
+// This will pause/unpause tracking, and also manually increment a given tag.
+void FLLMTracker::PauseAndTrackMemory(const FTagData* TagData, int64 Amount, ELLMAllocType AllocType)
+{
+	FLLMThreadState* State = GetOrCreateState();
+	FPlatformAtomics::InterlockedAdd(&TrackedTotal, Amount);
+	State->TrackMemory(Amount, Tracker, AllocType, TagData);
+	LLMCheck(State->PausedCounter[static_cast<int32>(AllocType)] < 127);
+	State->PausedCounter[static_cast<int32>(AllocType)]++;
+}
+
+void FLLMTracker::Pause(ELLMAllocType AllocType)
+{
+	FLLMThreadState* State = GetOrCreateState();
+	LLMCheck(State->PausedCounter[static_cast<int32>(AllocType)] < 127);
+	State->PausedCounter[static_cast<int32>(AllocType)]++;
+}
+
+void FLLMTracker::Unpause(ELLMAllocType AllocType)
+{
+	FLLMThreadState* State = GetOrCreateState();
+	LLMCheck(State->PausedCounter[static_cast<int32>(AllocType)] > 0);
+	State->PausedCounter[static_cast<int32>(AllocType)]--;
+}
+
+bool FLLMTracker::IsPaused(ELLMAllocType AllocType)
+{
+	FLLMThreadState* State = GetState();
+	// pause during shutdown, as the external trackers might not be able to robustly handle tracking.
+	return IsEngineExitRequested() || (State == nullptr ? false :
+		((State->PausedCounter[static_cast<int32>(ELLMAllocType::None)]>0) ||
+			(State->PausedCounter[static_cast<int32>(AllocType)])>0));
+}
+
+void FLLMTracker::Clear()
+{
+	{
+		FScopeLock Lock(&PendingThreadStatesGuard);
+		for (FLLMThreadState* ThreadState : PendingThreadStates)
+		{
+			LLMRef.Allocator.Delete(ThreadState);
+		}
+		PendingThreadStates.Empty();
+		ThreadIdToThreadState.Empty();
+	}
+
+	for (FLLMThreadState* ThreadState : ThreadStates)
+	{
+		LLMRef.Allocator.Delete(ThreadState);
+	}
+	ThreadStates.Empty();
+
+#if !UE_ONLY_USE_PLATFORM_TRACKER
+	AllocationMap.Clear();
+#endif
+
+	CsvWriter.Clear();
+	TraceWriter.Clear();
+	CsvProfilerWriter.Clear();
+}
+
+void FLLMTracker::OnPreFork()
+{
+	CsvWriter.OnPreFork();
+}
+
+void FLLMTracker::SetTotalTags(const FTagData* InOverrideUntaggedTagData,
+	const FTagData* InOverrideTrackedTotalTagData)
+{
+	OverrideUntaggedTagData = InOverrideUntaggedTagData;
+	OverrideTrackedTotalTagData = InOverrideTrackedTotalTagData;
+}
+
+void FLLMTracker::CaptureTagSnapshot()
+{
+	UpdateThreads();
+	for (FLLMThreadState* ThreadState : ThreadStates)
+	{
+		ThreadState->PropagateChildSizesToParents();
+		ThreadState->FetchAndClearTagSizes(TagSizes, AllocTypeAmounts, false);
+	}
+
+	for (TPair<const FTagData*, FTrackerTagSizeData>& It : TagSizes)
+	{
+		FTrackerTagSizeData& TrackerInfo = It.Value;
+		TrackerInfo.CaptureSnapshot();
+	}
+
+	TrackedTotalInSnapshot = TrackedTotal;
+}
+
+void FLLMTracker::ClearTagSnapshot()
+{
+	for (TPair<const FTagData*, FTrackerTagSizeData>& It : TagSizes)
+	{
+		FTrackerTagSizeData& TrackerInfo = It.Value;
+		TrackerInfo.ClearSnapshot();
+	}
+
+	TrackedTotalInSnapshot = 0;
+}
+
+void FLLMTracker::Update()
+{
+	UpdateThreads();
+	double CurrentTime = FPlatformTime::Seconds();
+	constexpr double UpdateTrimPeriod = 10.0;
+	bool bTrimAllocations = CurrentTime - LastTrimTime > UpdateTrimPeriod;
+	if (bTrimAllocations)
+	{
+		LastTrimTime = CurrentTime;
+		{
+#if !UE_ONLY_USE_PLATFORM_TRACKER
+			AllocationMap.Trim();
+#endif
+		}
+	}
+
+	// Add the values from each thread to the central repository.
+	bool bTrimThreads = bTrimAllocations;
+	for (FLLMThreadState* ThreadState : ThreadStates)
+	{
+		ThreadState->PropagateChildSizesToParents();
+		ThreadState->FetchAndClearTagSizes(TagSizes, AllocTypeAmounts, bTrimThreads);
+	}
+
+	// Update peak sizes and external sizes in the central repository.
+	for (TPair<const FTagData*, FTrackerTagSizeData>& It : TagSizes)
+	{
+		FTrackerTagSizeData& AllocationData = It.Value;
+
+		// Update external amount
+		if (AllocationData.bExternalValid)
+		{
+			if (AllocationData.bExternalAddToTotal)
+			{
+				FPlatformAtomics::InterlockedAdd(&TrackedTotal, AllocationData.ExternalAmount - AllocationData.Size);
+			}
+			AllocationData.Size = AllocationData.ExternalAmount;
+			AllocationData.bExternalValid = false;
+		}
+
+		// Calculate peaks.
+#if LLM_ENABLED_TRACK_PEAK_MEMORY
+		// TODO: we should keep track of the intra-frame memory peak for the total tracked memory. For now we will
+		// use the memory at the time the update happens since there are threading implications to being accurate.
+		AllocationData.PeakSize = FMath::Max(AllocationData.PeakSize, AllocationData.Size);
+#endif
+	}
+}
+
+void FLLMTracker::UpdateThreads()
+{
+	// Consume pending thread states. We must be careful to do all allocations outside of the
+	// PendingThreadStatesGuard guard as that can lead to a deadlock due to contention with
+	// PendingThreadStatesGuard & Locks inside the underlying allocator (i.e. MallocBinned2 -> Mutex).
+	{
+		PendingThreadStatesGuard.Lock();
+		const int32 NumPendingThreadStatesToConsume = PendingThreadStates.Num();
+		if (NumPendingThreadStatesToConsume > 0)
+		{
+			PendingThreadStatesGuard.Unlock();
+			ThreadStates.Reserve(ThreadStates.Num() + NumPendingThreadStatesToConsume);
+			PendingThreadStatesGuard.Lock();
+
+			for (int32 i = 0; i < NumPendingThreadStatesToConsume; ++i)
+			{
+				ThreadStates.Add(PendingThreadStates.Pop(EAllowShrinking::No));
+			}
+		}
+		PendingThreadStatesGuard.Unlock();
+	}
+}
+
+void FLLMTracker::PublishStats(UE::LLM::ESizeParams SizeParams)
+{
+	if (OverrideTrackedTotalTagData)
+	{
+		SetMemoryStatByFName(OverrideTrackedTotalTagData->GetStatName(), TrackedTotal);
+		SetMemoryStatByFName(OverrideTrackedTotalTagData->GetSummaryStatName(), TrackedTotal);
+	}
+
+	if (OverrideUntaggedTagData)
+	{
+		const FTagData* TagData = LLMRef.FindTagData(TagName_Untagged, ELLMTagSet::None);
+		const FTrackerTagSizeData* AllocationData = TagData ? TagSizes.Find(TagData) : nullptr;
+		SetMemoryStatByFName(OverrideUntaggedTagData->GetStatName(),
+			AllocationData ? AllocationData->GetSize(SizeParams) : 0);
+		SetMemoryStatByFName(OverrideUntaggedTagData->GetSummaryStatName(),
+			AllocationData ? AllocationData->GetSize(SizeParams) : 0);
+	}
+
+	for (const TPair<const FTagData*, FTrackerTagSizeData>& It : TagSizes)
+	{
+		const FTagData* TagData = It.Key;
+		if (!TagData->IsStatsReportable())
+		{
+			continue;
+		}
+		if (OverrideUntaggedTagData && TagData->GetName() == TagName_Untagged)
+		{
+			// Handled separately by OverrideUntaggedTagData.
+			continue;
+		}
+		const FTrackerTagSizeData& AllocationData = It.Value;
+		int64 Amount = AllocationData.GetSize(SizeParams);
+
+		SetMemoryStatByFName(TagData->GetStatName(), Amount);
+		SetMemoryStatByFName(TagData->GetSummaryStatName(), Amount);
+	}
+}
+
+static bool PageLessThan(const FForkedPageAllocation& LHS, const FForkedPageAllocation& RHS)
+{
+	return LHS.PageStart < RHS.PageStart;
+}
+static bool PageFinder(const FForkedPageAllocation& LHS, const uint64 AddressRHS)
+{
+	return LHS.PageStart < AddressRHS;
+}
+
+bool FLLMTracker::DumpForkedAllocationInfo()
+{
+	// Try to associate the allocations with a page range so we can determine whether it's
+	// in unique or shared memory. Then print out the set for each tag.
+	TArray<FForkedPageAllocation> Pages;
+	if (FGenericPlatformMemory::GetForkedPageAllocationInfo(Pages) == false)
+	{
+		// Create a Placeholder page so we can test on platforms that don't support forkedpages
+		FForkedPageAllocation& Page = Pages.Emplace_GetRef();
+		Page.PageStart = 0;
+		Page.PageEnd = MAX_uint64;
+		Page.SharedCleanKiB = 1024 * 100;
+		Page.SharedDirtyKiB = 1024 * 110;
+		Page.PrivateCleanKiB = 1024 * 120;
+		Page.PrivateDirtyKiB = 1024 * 130;
+	}
+
+	Algo::Sort(Pages, PageLessThan);
+
+	enum EClassification
+	{
+		Class_Split,
+		Class_Unreferenced,
+		Class_Private,
+		Class_Shared,
+		Class_COUNT
+	};
+
+	struct FCounts
+	{
+		uint64 TotalAllocations = 0;
+		uint64 CrossPageAllocationCount = 0;
+
+		uint64 AllocCount[Class_COUNT] = {};
+
+		uint64 TotalBytes = 0;
+		uint64 ByteCount[Class_COUNT] = {};
+	};
+
+
+	// We cannot call GetTag to find the Tag from the Allocation when we are within AllocationMap, because GetTag
+	// locks the TagDataMutex and TagDataMutex has to be entered before entering AllocationMap.LockAll.
+	// So for the TagIdentifier use the CompressedTag rather than calling GetTag.
+#if LLM_ENABLED_FULL_TAGS
+	TMap<int32, FCounts> CountsPerTag;
+#else
+	TMap<ELLMTag, FCounts> CountsPerTag;
+#endif
+	int32 NumTags;
+	{
+		UE::TSharedLock Lock(LLMRef.TagDataMutex);
+		NumTags = LLMRef.TagDatas->Num();
+	}
+
+	CountsPerTag.Reserve(NumTags);
+#if !UE_ONLY_USE_PLATFORM_TRACKER
+	AllocationMap.LockAll();
+	for (const FLLMAllocMap::FTuple& Tuple : AllocationMap)
+	{
+		void* Ptr = Tuple.Key.GetPointer();
+
+		int FirstPageEqualOrAfterPtr = Algo::LowerBound(Pages, (uint64)Ptr, PageFinder);
+		int ContainingAllocationIndex;
+		if (FirstPageEqualOrAfterPtr < Pages.Num() && Pages[FirstPageEqualOrAfterPtr].PageStart == (uint64)Ptr)
+		{
+			ContainingAllocationIndex = FirstPageEqualOrAfterPtr;
+		}
+		else
+		{
+			ContainingAllocationIndex = FirstPageEqualOrAfterPtr - 1;
+		}
+		if (ContainingAllocationIndex < 0 || Pages[ContainingAllocationIndex].PageEnd <= (uint64)Ptr)
+		{
+			UE_LOG(LogHAL, Error, TEXT("Can't find allocation 0x%llx in the pages list!"), (uint64)Ptr);
+			continue;
+		}
+
+		// We unfortunately don't know exactly which of the details apply because we don't know which pages
+		// in the allocation section are shared/unique etc, so we just keep some generalities.
+		uint64 SizeHigh = Tuple.Key.GetExtraData() << 32;
+		uint64 SizeLow = Tuple.Value1;
+		uint64 Size = SizeHigh | SizeLow;
+
+		// Number of page allocations for the ue allocation
+		uint32 PageAllocationCount = 0;
+
+		// Track the classification for each page allocation so we can try to classify the ue allocation
+		uint32 SplitCount = 0;
+		uint32 SharedCount = 0;
+		uint32 PrivateCount = 0;
+		uint32 UnreferencedCount = 0;
+
+		// Walk the page allocations until we get them all if we aren't all in the same one.
+		uint64 Remaining = Size;
+		while (Remaining)
+		{
+			if (ContainingAllocationIndex >= Pages.Num())
+			{
+				UE_LOG(LogHAL, Error, TEXT("Allocation 0x%llu extended beyond the pages!"), (uint64)Ptr);
+				break;
+			}
+
+			PageAllocationCount++;
+			FForkedPageAllocation& Page = Pages[ContainingAllocationIndex];
+
+			uint64 PageKiB = (Page.PageEnd - Page.PageStart) / 1024;
+
+			uint64 OffsetInPage = (uint64)Ptr - Page.PageStart;
+			uint64 RemainingInPage = Page.PageEnd - (uint64)Ptr;
+
+			uint64 AmountInThisPage = Remaining;
+			if (AmountInThisPage > RemainingInPage)
+			{
+				AmountInThisPage = RemainingInPage;
+			}
+
+			uint64 SharedKiB = Page.SharedCleanKiB + Page.SharedDirtyKiB;
+			uint64 PrivateKiB = Page.PrivateCleanKiB + Page.PrivateDirtyKiB;
+			uint64 UnreferencedKiB = PageKiB - SharedKiB - PrivateKiB;
+
+			// We can't classify the page if we have more than one of any.
+			uint64 TypeCount = !!SharedKiB + !!PrivateKiB + !!UnreferencedKiB;
+			bool bIsSplit = TypeCount > 1;
+
+			if (bIsSplit)
+			{
+				// The page is split and we have no way to know which is ours.
+				SplitCount++;
+			}
+			else if (SharedKiB)
+			{
+				SharedCount++;
+			}
+			else if (PrivateKiB)
+			{
+				PrivateCount++;
+			}
+			else
+			{
+				UnreferencedCount++;
+			}
+
+			Remaining -= AmountInThisPage;
+			ContainingAllocationIndex++;
+		}
+
+		if (Remaining)
+		{
+			// We error'd out, ignore this allocation.
+			continue;
+		}
+
+		// Classify the allocation - is it entirely private, entirely shared, or what.
+		uint64 ClassCount = !!SharedCount + !!PrivateCount + !!UnreferencedCount;
+		EClassification AllocClass = Class_Split;
+		if (ClassCount > 1)
+		{
+			AllocClass = Class_Split;
+		}
+		else if (SharedCount)
+		{
+			AllocClass = Class_Shared;
+		}
+		else if (PrivateCount)
+		{
+			AllocClass = Class_Private;
+		}
+		else
+		{
+			AllocClass = Class_Unreferenced;
+		}
+
+		// Associate the tag and the page stats.
+		FCounts& Counts = CountsPerTag.FindOrAdd(Tuple.Value2.GetCompressedTag());
+		Counts.AllocCount[AllocClass]++;
+		Counts.ByteCount[AllocClass] += Size;
+		Counts.TotalAllocations++;
+		Counts.TotalBytes += Size;
+		if (PageAllocationCount > 1)
+		{
+			Counts.CrossPageAllocationCount++;
+		}
+	}
+	AllocationMap.UnlockAll();
+#endif // !UE_ONLY_USE_PLATFORM_TRACKER
+
+	CountsPerTag.ValueSort([](const FCounts& A, const FCounts& B)
+	{
+		return A.TotalBytes > B.TotalBytes;
+	});
+
+	TArray<FString> Lines;
+	Lines.Reserve(CountsPerTag.Num() * 2 + 1);
+
+	Lines.Add(TEXT("Tag,SharedKib,PrivateKib,SplitKib,UnrefKib,TotalKib,SharedCount,PrivateCount,SplitCount,UnrefCount,TotalCount,CrossCount"));
+
+
+#if LLM_ENABLED_FULL_TAGS
+	for (TPair<int32, FCounts>& P : CountsPerTag)
+#else
+	for (TPair<ELLMTag, FCounts>& P : CountsPerTag)
+#endif
+	{
+		FLowLevelAllocInfo AllocInfoPlaceholder;
+		AllocInfoPlaceholder.SetCompressedTag(P.Key);
+		const FTagData* Tag = AllocInfoPlaceholder.GetTag(LLMRef);
+		if (!Tag)
+		{
+			continue;
+		}
+		Lines.Add(
+			FString::Printf(TEXT("%s,%llu,%llu,%llu,%llu,%llu,%llu,%llu,%llu,%llu,%llu,%llu"),
+				*Tag->GetDisplayName().ToString(),
+				P.Value.ByteCount[Class_Shared] / (1024),
+				P.Value.ByteCount[Class_Private] / (1024),
+				P.Value.ByteCount[Class_Split] / (1024),
+				P.Value.ByteCount[Class_Unreferenced] / (1024),
+				P.Value.TotalBytes / (1024),
+				P.Value.AllocCount[Class_Shared],
+				P.Value.AllocCount[Class_Private],
+				P.Value.AllocCount[Class_Split],
+				P.Value.AllocCount[Class_Unreferenced],
+				P.Value.TotalAllocations,
+				P.Value.CrossPageAllocationCount)
+			);
+	}
+
+	// This should be UE::LLM::GetLLMProfilingDir(), i.e. "Saved/Profiling/LLM", but kept to "Saved/LLM" for legacy reasons.
+	FString SavedDir = FPaths::ProjectSavedDir() / TEXT("LLM");
+
+	uint16 ForkId = FForkProcessHelper::GetForkedChildProcessIndex();
+	FString ForkString = ForkId == 0 ? FString("Parent") : FString::Printf(TEXT("Child%d"), ForkId);
+
+	const TCHAR* TrackerString = TEXT("Default");
+	if (Tracker == ELLMTracker::Platform)
+	{
+		TrackerString = TEXT("Platform");
+	}
+	static_assert((int)ELLMTracker::Max == 2, "Add other tracker type strings here");
+
+	FString FileName = FString::Printf(TEXT("LLM_PrivateShared_Tracker%s_%s"), TrackerString, *ForkString);
+
+	FString UniqueFilename = SavedDir / FString::Printf(TEXT("%s.csv"), *FileName);
+	uint32 Counter = 1;
+	while (IFileManager::Get().FileSize(*UniqueFilename) >= 0)
+	{
+		UniqueFilename = SavedDir / FString::Printf(TEXT("%s_%d.csv"), *FileName, Counter);
+		Counter++;
+	}
+
+	if (FFileHelper::SaveStringArrayToFile(Lines, *UniqueFilename) == false)
+	{
+		UE_LOG(LogHAL, Error, TEXT("Failed to write LLM Private/Shared CSV file: %s\n"), *UniqueFilename);
+		return false;
+	}
+
+	UE_LOG(LogHAL, Log, TEXT("Wrote LLM Private/Shared CSV file: %s"), *UniqueFilename);
+	return true;
+}
+
+bool FLLMTracker::DumpTags()
+{
+	TArray<FString> Lines;
+
+	{
+		UE::TSharedLock Lock(LLMRef.TagDataMutex);
+
+		int32 NumTags = LLMRef.TagDatas->Num();
+		Lines.Reserve(NumTags + 1);
+
+		Lines.Add(TEXT("Index,Id,Name,DisplayName,Parent,StatName,SummaryStatName,ReferenceSource,EnumTag,bHasEnumTag,TagSet,bIsFinishConstructed,bIsReportable,bIsTraceReportable"));
+
+		for (FTagData* TagData : (*LLMRef.TagDatas))
+		{
+			Lines.Add(
+				FString::Printf(TEXT("%d,0x%llX,%s,%s,%s,%s,%s,%s,%u,%s,%u,%s,%s,%s"),
+					TagData->GetIndex(),
+					uint64(TagData),
+					*TagData->GetName().ToString(),
+					*TagData->GetDisplayName().ToString(),
+					TagData->IsParentConstructed() ?
+						(TagData->GetParent() ? *FString::Printf(TEXT("0x%llX"), uint64(TagData->GetParent())) : TEXT("")) :
+						*TagData->GetParentName().ToString(),
+					*TagData->GetStatName().ToString(),
+					*TagData->GetSummaryStatName().ToString(),
+					ToString(TagData->GetReferenceSource()),
+					uint32(TagData->GetEnumTag()),
+					TagData->HasEnumTag() ? TEXT("true") : TEXT("false"),
+					uint32(TagData->GetTagSet()),
+					TagData->IsFinishConstructed() ? TEXT("true") : TEXT("false"),
+					TagData->IsStatsReportable() ? TEXT("true") : TEXT("false"),
+					TagData->IsTraceReportable() ? TEXT("true") : TEXT("false"))
+				);
+		}
+	}
+
+	FString SavedDir = UE::LLM::GetLLMProfilingDir();
+
+	FString FileName("LLM_Tags");
+
+	FString UniqueFilename = SavedDir / FString::Printf(TEXT("%s.csv"), *FileName);
+	uint32 Counter = 1;
+	while (IFileManager::Get().FileSize(*UniqueFilename) >= 0)
+	{
+		UniqueFilename = SavedDir / FString::Printf(TEXT("%s_%d.csv"), *FileName, Counter);
+		Counter++;
+	}
+
+	if (FFileHelper::SaveStringArrayToFile(Lines, *UniqueFilename) == false)
+	{
+		UE_LOG(LogHAL, Error, TEXT("Failed to write LLM Tags CSV file: %s\n"), *UniqueFilename);
+		return false;
+	}
+
+	UE_LOG(LogHAL, Log, TEXT("Wrote LLM Tags CSV file: %s"), *UniqueFilename);
+	return true;
+}
+
+void FLLMTracker::PublishCsv(UE::LLM::ESizeParams SizeParams)
+{
+	CsvWriter.Publish(LLMRef, TagSizes, OverrideTrackedTotalTagData, OverrideUntaggedTagData, GetTrackedTotal(SizeParams), SizeParams);
+}
+
+void FLLMTracker::PublishTrace(UE::LLM::ESizeParams SizeParams)
+{
+	// Trace tracker does not support the handling of Snapshot currently
+	EnumRemoveFlags(SizeParams, UE::LLM::ESizeParams::RelativeToSnapshot);
+
+	TraceWriter.Publish(LLMRef, TagSizes, OverrideTrackedTotalTagData, OverrideUntaggedTagData, GetTrackedTotal(SizeParams), SizeParams);
+}
+
+void FLLMTracker::PublishCsvProfiler(UE::LLM::ESizeParams SizeParams)
+{
+	CsvProfilerWriter.Publish(LLMRef, TagSizes, OverrideTrackedTotalTagData, OverrideUntaggedTagData, GetTrackedTotal(SizeParams), SizeParams);
+}
+
+void FLLMTracker::OnTagsResorted(FTagDataArray& OldTagDatas)
+{
+#if LLM_ENABLED_FULL_TAGS && !UE_ONLY_USE_PLATFORM_TRACKER
+	{
+		// Each allocation references the tag by its index, which we have just remapped.
+		// Remap each allocation's tag index to the new index for the tag.
+		AllocationMap.LockAll();
+		for (const FLLMAllocMap::FTuple& Tuple : AllocationMap)
+		{
+			for (int32 TagSetAsInteger = 0; TagSetAsInteger < static_cast<int32>(ELLMTagSet::Max); TagSetAsInteger++)
+			{
+				if (!LLMRef.IsTagSetActive(static_cast<ELLMTagSet>(TagSetAsInteger)))
+				{
+					continue;
+				}
+
+				const int32 OldTagIndex = Tuple.Value2.GetCompressedTag(static_cast<ELLMTagSet>(TagSetAsInteger));
+				if (OldTagIndex >= 0)
+				{
+					Tuple.Value2.SetCompressedTag(OldTagDatas[OldTagIndex]->GetIndex(), static_cast<ELLMTagSet>(TagSetAsInteger));
+				}
+			}
+		}
+		AllocationMap.UnlockAll();
+	}
+#else
+	// Values in AllocationMap are ELLMTags, and don't depend on the Index of the tagdatas.
+#endif
+
+	// Update the uses of Index in the ThreadStates.
+	for (FLLMThreadState* ThreadState : ThreadStates)
+	{
+		ThreadState->OnTagsResorted(OldTagDatas);
+	}
+}
+
+void FLLMTracker::LockAllThreadTags(bool bLock)
+{
+	if (bLock)
+	{
+		UpdateThreads();
+		PendingThreadStatesGuard.Lock();
+	}
+
+	for (FLLMThreadState* ThreadState : ThreadStates)
+	{
+		ThreadState->LockTags(bLock);
+	}
+
+	if (!bLock)
+	{
+		PendingThreadStatesGuard.Unlock();
+	}
+}
+
+const FTagData* FLLMTracker::GetActiveTagData(ELLMTagSet TagSet)
+{
+	FLLMThreadState* State = GetOrCreateState();
+	return State->GetTopTag(TagSet);
+}
+
+TArray<const FTagData*> FLLMTracker::GetTagDatas(ELLMTagSet TagSet)
+{
+	TArray<const FTagData*> FoundTagDatas;
+	TagSizes.GetKeys(FoundTagDatas);
+
+	return FoundTagDatas.FilterByPredicate([TagSet](const FTagData* InTagData) {
+		return InTagData != nullptr && InTagData->GetTagSet() == TagSet;
+	});
+}
+
+void FLLMTracker::GetTagsNamesWithAmount(TMap<FName, uint64>& OutTagsNamesWithAmount, ELLMTagSet TagSet /* = ELLMTagSet::None */)
+{
+	for (const TPair<const FTagData*, FTrackerTagSizeData>& It : TagSizes)
+	{
+		const FTagData* TagData = It.Key;
+		if (TagData->GetTagSet() == TagSet)
+		{
+			OutTagsNamesWithAmount.Add(TagData->GetName(), It.Value.GetSize(UE::LLM::ESizeParams::Default));
+		}
+	}
+}
+
+void FLLMTracker::GetTagsNamesWithAmountFiltered(TMap<FName, uint64>& OutTagsNamesWithAmount, ELLMTagSet TagSet /* = ELLMTagSet::None */, TArray<FLLMTagSetAllocationFilter>& Filters)
+{
+	// We cannot call GetTag to find the Tag from the Allocation when we are within AllocationMap, because GetTag
+	// locks the TagDataMutex and TagDataMutex has to be entered before entering AllocationMap.LockAll.
+	// So for the TagIdentifier use the CompressedTag rather than calling GetTag.
+	struct FLocalTagData
+	{
+		FName Name;
+		uint64 Size = 0;
+	};
+	TArray<FLocalTagData> TagDataCached;
+	{
+		UE::TSharedLock Lock(LLMRef.TagDataMutex);
+		const int32 NumTags = LLMRef.TagDatas->Num();
+		TagDataCached.Init(FLocalTagData{}, NumTags);
+		for (FTagData* TagData : (*LLMRef.TagDatas))
+		{
+			FLLMTracker::FLowLevelAllocInfo AllocInfo;
+			AllocInfo.SetTag(TagData, LLMRef);
+			int32 CompressedTagIndex = static_cast<int32>(AllocInfo.GetCompressedTag());
+			if(ensure(TagDataCached.IsValidIndex(CompressedTagIndex)))
+			{
+				FLocalTagData& Data = TagDataCached[CompressedTagIndex];
+				if (Data.Name == NAME_None)
+				{
+					Data.Name = TagData->GetName();
+				}
+			}
+		}
+	}
+#if !UE_ONLY_USE_PLATFORM_TRACKER
+	AllocationMap.LockAll();
+	for (const FLLMAllocMap::FTuple& Tuple : AllocationMap)
+	{
+		bool bIncludeAllocation = true;
+		for (const FLLMTagSetAllocationFilter& Filter : Filters)
+		{
+			FLocalTagData* Data = nullptr;
+#if LLM_ALLOW_ASSETS_TAGS || LLM_ALLOW_UOBJECTCLASSES_TAGS
+			int32 CompressedTagIndex = static_cast<int32>(Tuple.Value2.GetCompressedTag(Filter.TagSet));
+#else
+			int32 CompressedTagIndex = static_cast<int32>(Tuple.Value2.GetCompressedTag());
+#endif
+			if(TagDataCached.IsValidIndex(CompressedTagIndex))
+			{
+#if LLM_ALLOW_ASSETS_TAGS || LLM_ALLOW_UOBJECTCLASSES_TAGS
+				Data = &TagDataCached[CompressedTagIndex];
+#else
+				if (Filter.TagSet == ELLMTagSet::None)
+				{
+					Data = &TagDataCached[CompressedTagIndex];
+				}
+#endif
+			}
+			if (!Data || Data->Name != Filter.Name)
+			{
+				bIncludeAllocation = false;
+				break;
+			}
+		}
+
+		if (bIncludeAllocation)
+		{
+#if LLM_ALLOW_ASSETS_TAGS || LLM_ALLOW_UOBJECTCLASSES_TAGS
+			int32 CompressedTagIndex = Tuple.Value2.GetCompressedTag(TagSet);
+#else
+			int32 CompressedTagIndex = static_cast<int32>(Tuple.Value2.GetCompressedTag());
+#endif
+			FLocalTagData* Data = TagDataCached.IsValidIndex(CompressedTagIndex) ? &TagDataCached[CompressedTagIndex] : nullptr;
+			if (Data)
+			{
+				Data->Size += Tuple.Value1;
+			}
+		}
+	}
+	AllocationMap.UnlockAll();
+#endif // !UE_ONLY_USE_PLATFORM_TRACKER
+
+	for (const FLocalTagData& LocalTagData : TagDataCached)
+	{
+		if (LocalTagData.Size != 0)
+		{
+			OutTagsNamesWithAmount.FindOrAdd(LocalTagData.Name, 0) += LocalTagData.Size;
+		}
+	}
+}
+
+bool FLLMTracker::FindTagsForPtr(void* InPtr, TArray<const FTagData *, TInlineAllocator<static_cast<int32>(ELLMTagSet::Max)>>& OutTags) const
+{
+#if !UE_ONLY_USE_PLATFORM_TRACKER
+	uint32 Size;
+	FLowLevelAllocInfo AllocInfoPtr;
+	PointerKey FoundKey = AllocationMap.Find(PointerKey(InPtr), Size, AllocInfoPtr);
+	if (!FoundKey)
+	{
+		return false;
+	}
+
+	OutTags.SetNumUninitialized(static_cast<int32>(ELLMTagSet::Max));
+#if LLM_ALLOW_ASSETS_TAGS || LLM_ALLOW_UOBJECTCLASSES_TAGS
+	for (int32 TagSetAsInteger = 0; TagSetAsInteger < static_cast<int32>(ELLMTagSet::Max); TagSetAsInteger++)
+	{
+		// Calling GetTag with an ELLMTagSet that is compiled-out is illegal and will cause an assertion.
+		// We rely on IsTagSetActive to be false for compiled-out ELLMTagSet, and use that to avoid calling GetTag.
+		if (!LLMRef.IsTagSetActive(static_cast<ELLMTagSet>(TagSetAsInteger)))
+		{
+			continue;
+		}
+
+		OutTags[TagSetAsInteger] = AllocInfoPtr.GetTag(LLMRef, static_cast<ELLMTagSet>(TagSetAsInteger));
+	}
+#else
+	OutTags[static_cast<int32>(ELLMTagSet::None)] = AllocInfoPtr.GetTag(LLMRef);
+#endif
+
+	return true;
+#else
+	return false;
+#endif // !UE_ONLY_USE_PLATFORM_TRACKER
+}
+
+int64 FLLMTracker::GetTagAmount(const FTagData* TagData, UE::LLM::ESizeParams SizeParams) const
+{
+	const FTrackerTagSizeData* AllocationData = TagSizes.Find(TagData);
+	if (AllocationData)
+	{
+		return AllocationData->GetSize(SizeParams);
+	}
+	else
+	{
+		return 0;
+	}
+}
+
+void FLLMTracker::SetTagAmountExternal(const FTagData* TagData, int64 Amount, bool bAddToTotal)
+{
+	FTrackerTagSizeData& AllocationData = TagSizes.FindOrAdd(TagData);
+	AllocationData.bExternalValid = true;
+	AllocationData.bExternalAddToTotal = bAddToTotal;
+	AllocationData.ExternalAmount = Amount;
+}
+
+void FLLMTracker::SetTagAmountInUpdate(const FTagData* TagData, int64 Amount, bool bAddToTotal)
+{
+	FTrackerTagSizeData& AllocationData = TagSizes.FindOrAdd(TagData);
+	if (bAddToTotal)
+	{
+		FPlatformAtomics::InterlockedAdd(&TrackedTotal, Amount - AllocationData.Size);
+	}
+	AllocationData.Size = Amount;
+#if LLM_ENABLED_TRACK_PEAK_MEMORY
+	AllocationData.PeakSize = FMath::Max(AllocationData.PeakSize, AllocationData.Size);
+#endif
+}
+
+int64 FLLMTracker::GetAllocTypeAmount(ELLMAllocType AllocType)
+{
+	return AllocTypeAmounts[static_cast<int32>(AllocType)];
+}
+
+FLLMThreadState::FLLMThreadState()
+{
+	for (int32 Index = 0; Index < static_cast<int32>(ELLMAllocType::Count); ++Index)
+	{
+		PausedCounter[Index] = 0;
+	}
+
+	// By contract AllocationsLock must be entered around ClearAllocTypeAmounts, but multithreading contracts of that type
+	// are not required in constructors, which complete before the object can be used from multiple threads.
+	ClearAllocTypeAmounts();
+}
+
+void FLLMThreadState::Clear()
+{
+	FScopeLock Lock(&AllocationsLock);
+
+	for (int32 TagSetAsInteger = 0; TagSetAsInteger < static_cast<int32>(ELLMTagSet::Max); TagSetAsInteger++)
+	{
+		TagStack[TagSetAsInteger].Empty();
+	}
+	Allocations.Empty();
+	ClearAllocTypeAmounts();
+}
+
+void FLLMThreadState::PushTag(const FTagData* TagData, ELLMTagSet TagSet)
+{
+	FScopeLock Lock(&AllocationsLock);
+
+	// Push a tag.
+	TagStack[static_cast<int32>(TagSet)].Add(TagData);
+}
+
+void FLLMThreadState::PopTag(ELLMTagSet TagSet)
+{
+	FScopeLock Lock(&AllocationsLock);
+
+	LLMCheckf(TagStack[static_cast<int32>(TagSet)].Num() > 0,
+		TEXT("Called FLLMThreadState::PopTag without a matching Push (stack was empty on pop)"));
+	TagStack[static_cast<int32>(TagSet)].Pop(EAllowShrinking::No);
+}
+
+const FTagData* FLLMThreadState::GetTopTag(ELLMTagSet TagSet)
+{
+	if (TagStack[static_cast<int32>(TagSet)].Num() == 0)
+	{
+		return nullptr;
+	}
+
+	return TagStack[static_cast<int32>(TagSet)].Last();
+}
+
+void FLLMThreadState::IncrTag(const FTagData* TagData, int64 Amount)
+{
+	// Caller is responsible for holding a lock on AllocationsLock.
+	FThreadTagSizeData& AllocationSize = Allocations.FindOrAdd(TagData->GetIndex());
+	AllocationSize.TagData = TagData;
+	AllocationSize.Size += Amount;
+}
+
+void FLLMThreadState::TrackAllocation(const void* Ptr, int64 Size, ELLMTracker Tracker, ELLMAllocType AllocType,
+	const FTagData* TagData, const FTagData* AssetTagData, const FTagData* AssetClassTagData, const FTagData* UObjectClassesTagData, bool bTrackInMemPro)
+{
+	FScopeLock Lock(&AllocationsLock);
+
+	AllocTypeAmounts[static_cast<int32>(AllocType)] += Size;
+
+	if (TagData)
+	{
+		IncrTag(TagData, Size);
+	}
+#if LLM_ALLOW_ASSETS_TAGS
+	if (AssetTagData)
+	{
+		IncrTag(AssetTagData, Size);
+	}
+	if (AssetClassTagData)
+	{
+		IncrTag(AssetClassTagData, Size);
+	}
+#endif // LLM_ALLOW_ASSETS_TAGS
+
+#if LLM_ALLOW_UOBJECTCLASSES_TAGS
+	if (UObjectClassesTagData)
+	{
+		IncrTag(UObjectClassesTagData, Size);
+	}
+#endif // LLM_ALLOW_UOBJECTCLASSES_TAGS
+
+	ELLMTag EnumTag = TagData ? TagData->GetContainingEnum() : ELLMTag(0);
+	if (Tracker == ELLMTracker::Default)
+	{
+		FPlatformMemory::OnLowLevelMemory_Alloc(Ptr, static_cast<uint64>(Size), static_cast<uint64>(EnumTag));
+	}
+
+#if MEMPRO_ENABLED
+	if (FMemProProfiler::IsTrackingTag(EnumTag) && bTrackInMemPro)
+	{
+		MEMPRO_TRACK_ALLOC(const_cast<void*>(Ptr), static_cast<size_t>(Size));
+	}
+#endif
+}
+
+void FLLMThreadState::TrackFree(const void* Ptr, int64 Size, ELLMTracker Tracker, ELLMAllocType AllocType,
+	const FTagData* TagData, const FTagData* AssetTagData, const FTagData* AssetClassTagData, const FTagData* UObjectClassesTagData, bool bTrackInMemPro)
+{
+	FScopeLock Lock(&AllocationsLock);
+
+	AllocTypeAmounts[static_cast<int32>(AllocType)] -= Size;
+	if (TagData)
+	{
+		IncrTag(TagData, -Size);
+	}
+#if LLM_ALLOW_ASSETS_TAGS
+	if (AssetTagData)
+	{
+		IncrTag(AssetTagData, -Size);
+	}
+	if (AssetClassTagData)
+	{
+		IncrTag(AssetClassTagData, -Size);
+	}
+#endif // LLM_ALLOW_ASSETS_TAGS
+
+#if LLM_ALLOW_UOBJECTCLASSES_TAGS
+	if (UObjectClassesTagData)
+	{
+		IncrTag(UObjectClassesTagData, -Size);
+	}
+#endif // LLM_ALLOW_UOBJECTCLASSES_TAGS
+
+	ELLMTag EnumTag = TagData ? TagData->GetContainingEnum() : ELLMTag(0);
+	if (Tracker == ELLMTracker::Default)
+	{
+		FPlatformMemory::OnLowLevelMemory_Free(Ptr, static_cast<uint64>(Size), static_cast<uint64>(EnumTag));
+	}
+
+#if MEMPRO_ENABLED
+	if (FMemProProfiler::IsTrackingTag(EnumTag) && bTrackInMemPro)
+	{
+		MEMPRO_TRACK_FREE(const_cast<void*>(Ptr));
+	}
+#endif
+}
+
+void FLLMThreadState::TrackMemory(int64 Amount, ELLMTracker Tracker, ELLMAllocType AllocType, const FTagData* TagData)
+{
+	FScopeLock Lock(&AllocationsLock);
+	AllocTypeAmounts[static_cast<int32>(AllocType)] += Amount;
+	IncrTag(TagData, Amount);
+
+	// TODO: Need to expose TrackMemory to Platform-specific trackers and FMemPropProfiler
+}
+
+void FLLMThreadState::TrackMoved(const void* Dest, const void* Source, int64 Size, ELLMTracker Tracker,
+	const FTagData* TagData)
+{
+	// Update external memory trackers (ideally would want a proper 'move' option on these).
+	ELLMTag EnumTag = TagData ? TagData->GetContainingEnum() : ELLMTag(0);
+	if (Tracker == ELLMTracker::Default)
+	{
+		FPlatformMemory::OnLowLevelMemory_Free(Source, static_cast<uint64>(Size), static_cast<uint64>(EnumTag));
+		FPlatformMemory::OnLowLevelMemory_Alloc(Dest, static_cast<uint64>(Size), static_cast<uint64>(EnumTag));
+	}
+
+#if MEMPRO_ENABLED
+	if (FMemProProfiler::IsTrackingTag(EnumTag))
+	{
+		MEMPRO_TRACK_FREE(const_cast<void*>(Source));
+		MEMPRO_TRACK_ALLOC(const_cast<void*>(Dest), static_cast<size_t>(Size));
+	}
+#endif
+}
+
+void FLLMThreadState::PropagateChildSizesToParents()
+{
+	FScopeLock Lock(&AllocationsLock);
+
+	// Make sure all parents of any TagDatas in the Allocations are also present.
+	FConstTagDataArray ParentsToAdd;
+	for (const TPair<int32, FThreadTagSizeData>& It : Allocations)
+	{
+		const FTagData* TagData = It.Value.TagData;
+		const FTagData* ParentData = TagData->GetParent();
+		while (ParentData && !Allocations.Contains(ParentData->GetIndex()))
+		{
+			ParentsToAdd.Add(ParentData);
+			ParentData = ParentData->GetParent();
+		}
+	}
+	for (const FTagData* TagData : ParentsToAdd)
+	{
+		FThreadTagSizeData& Info = Allocations.FindOrAdd(TagData->GetIndex());
+		Info.TagData = TagData;
+	}
+
+	// Tags are sorted topologically from parent to child, so we can accumulate children into parents recursively
+	// by reverse iterating the map.
+	for (FThreadTagSizeMap::TConstReverseIterator It(Allocations); It; ++It)
+	{
+		const FThreadTagSizeData& Info = It->Value;
+		const FTagData* ParentData = Info.TagData->GetParent();
+		if (Info.Size && ParentData)
+		{
+			Allocations.FindChecked(ParentData->GetIndex()).Size += Info.Size;
+		}
+	}
+}
+
+void FLLMThreadState::OnTagsResorted(FTagDataArray& OldTagDatas)
+{
+	FScopeLock Lock(&AllocationsLock);
+	TArray<FThreadTagSizeData, FDefaultLLMAllocator> AllocationDatas;
+	AllocationDatas.Reserve(Allocations.Num());
+	for (const TPair<int32, FThreadTagSizeData>& It : Allocations)
+	{
+		AllocationDatas.Add(It.Value);
+	}
+	Allocations.Reset();
+	for (const FThreadTagSizeData& AllocationData : AllocationDatas)
+	{
+		Allocations.Add(AllocationData.TagData->GetIndex(), AllocationData);
+	}
+}
+
+void FLLMThreadState::LockTags(bool bLock)
+{
+	if (bLock)
+	{
+		AllocationsLock.Lock();
+	}
+	else
+	{
+		AllocationsLock.Unlock();
+	}
+}
+
+void FLLMThreadState::FetchAndClearTagSizes(FTrackerTagSizeMap& TagSizes, int64* InAllocTypeAmounts,
+	bool bTrimAllocations)
+{
+	FScopeLock Lock(&AllocationsLock);
+	for (TPair<int32, FThreadTagSizeData>& Item : Allocations)
+	{
+		FThreadTagSizeData& ThreadInfo = Item.Value;
+		if (ThreadInfo.Size)
+		{
+			const FTagData* TagData = ThreadInfo.TagData;
+			FTrackerTagSizeData& TrackerInfo = TagSizes.FindOrAdd(TagData);
+			TrackerInfo.Size += ThreadInfo.Size;
+			ThreadInfo.Size = 0;
+		}
+	}
+	if (bTrimAllocations)
+	{
+		Allocations.Empty();
+	}
+
+	for (int32 Index = 0; Index < static_cast<int32>(ELLMAllocType::Count); ++Index)
+	{
+		InAllocTypeAmounts[Index] += AllocTypeAmounts[Index];
+		AllocTypeAmounts[Index] = 0;
+	}
+}
+
+void SetMemoryStatByFName(FName Name, int64 Amount)
+{
+	if (Name != NAME_None)
+	{
+		SET_MEMORY_STAT_FName(Name, Amount);
+	}
+}
+
+void FLLMThreadState::ClearAllocTypeAmounts()
+{
+	// Caller is responsible for holding a lock on AllocationsLock.
+
+	for (int32 Index = 0; Index < static_cast<int32>(ELLMAllocType::Count); ++Index)
+	{
+		AllocTypeAmounts[Index] = 0;
+	}
+}
+
+// FLLMCsvWriter implementation.
+FLLMCsvWriter::FLLMCsvWriter()
+	: Archive(nullptr)
+	, LastWriteTime(FPlatformTime::Seconds())
+	, WriteCount(0)
+	, bRegisteredFlushDelegate(false)
+{
+}
+
+FLLMCsvWriter::~FLLMCsvWriter()
+{
+	delete Archive;
+}
+
+void FLLMCsvWriter::Clear()
+{
+	Columns.Empty();
+	ExistingColumns.Empty();
+}
+
+void FLLMCsvWriter::OnPreFork()
+{
+	if (Archive)
+	{
+		Archive->Flush();
+		delete Archive;
+		Archive = nullptr;
+	}
+}
+
+void FLLMCsvWriter::Flush(bool IsOnCrash)
+{
+	if (Archive)
+	{
+		if (IsOnCrash)
+		{
+			FPlatformMisc::LowLevelOutputDebugStringf(TEXT("Flushing LLM CSV on crash\n"));
+		}
+
+		Archive->Flush();
+	}
+}
+
+void FLLMCsvWriter::SetTracker(ELLMTracker InTracker)
+{
+	Tracker = InTracker;
+}
+
+void FLLMCsvWriter::Publish(FLowLevelMemTracker& LLMRef, const FTrackerTagSizeMap& TagSizes,
+	const FTagData* OverrideTrackedTotalTagData, const FTagData* OverrideUntaggedTagData, int64 TrackedTotal,
+	UE::LLM::ESizeParams SizeParams)
+{
+	const double Now = FPlatformTime::Seconds();
+
+	if ((FLowLevelMemTracker::Get().bPublishSingleFrame == false) &&
+		(Now - LastWriteTime < (double)CVarLLMWriteInterval.GetValueOnAnyThread()))
+	{
+		return;
+	}
+
+	LastWriteTime = Now;
+
+	const bool bCreatedArchive = CreateArchive();
+	const bool bColumnsUpdated = UpdateColumns(TagSizes);
+
+	if (bCreatedArchive || bColumnsUpdated)
+	{
+		// The column names are written at the start of the archive; when they change we seek back to the start of
+		// the file and rewrite the column names.
+		WriteHeader(OverrideTrackedTotalTagData, OverrideUntaggedTagData);
+	}
+
+	AddRow(LLMRef, TagSizes, OverrideTrackedTotalTagData, OverrideUntaggedTagData, TrackedTotal, SizeParams);
+}
+
+const TCHAR* FLLMCsvWriter::GetTrackerCsvName(ELLMTracker InTracker)
+{
+	switch (InTracker)
+	{
+		case ELLMTracker::Default: return TEXT("LLM");
+		case ELLMTracker::Platform: return TEXT("LLMPlatform");
+		default: LLMCheck(false); return TEXT("");
+	}
+}
+
+// Archive is a binary stream, so we can't just serialize an FString using <<.
+void FLLMCsvWriter::Write(FStringView Text)
+{
+	Archive->Serialize(const_cast<ANSICHAR*>(
+		StringCast<ANSICHAR>(Text.GetData(), Text.Len()).Get()), Text.Len() * sizeof(ANSICHAR));
+}
+
+bool FLLMCsvWriter::CreateArchive()
+{
+	if (Archive)
+	{
+		return false;
+	}
+
+	// Create the csv file.
+	FString Directory = UE::LLM::GetLLMProfilingDir();
+	IFileManager::Get().MakeDirectory(*Directory, true);
+
+	const TCHAR* TrackerName = GetTrackerCsvName(Tracker);
+	const FDateTime FileDate = FDateTime::Now();
+#if PLATFORM_DESKTOP
+	FString PlatformName = FPlatformProperties::PlatformName();
+#else // Use the CPU for consoles so we can differentiate things like different SKUs of a console generation.
+	FString PlatformName = FPlatformMisc::GetCPUBrand().TrimStartAndEnd();
+#endif
+	PlatformName.ReplaceCharInline(' ', '_');
+	PlatformName = FPaths::MakeValidFileName(PlatformName);
+#if WITH_SERVER_CODE
+	FString Filename = FString::Printf(TEXT("%s/%s_Pid%d_%s_%s.csv"), *Directory, TrackerName,
+		FPlatformProcess::GetCurrentProcessId(), *FileDate.ToString(), *PlatformName);
+#else
+	FString Filename = FString::Printf(TEXT("%s/%s_%s_%s.csv"), *Directory, TrackerName, *FileDate.ToString(),
+		*PlatformName);
+#endif
+	Archive = IFileManager::Get().CreateFileWriter(*Filename, FILEWRITE_AllowRead | FILEWRITE_NoFail);
+	LLMCheck(Archive);
+
+	// Create space for column titles that are filled in as we get them.
+	Write(FString::ChrN(CVarLLMHeaderMaxSize.GetValueOnAnyThread(), ' '));
+	Write(TEXT("\n"));
+
+	return true;
+}
+
+bool FLLMCsvWriter::UpdateColumns(const FTrackerTagSizeMap& TagSizes)
+{
+	bool bUpdated = false;
+
+	for (const TPair<const FTagData*, FTrackerTagSizeData>& It : TagSizes)
+	{
+		const FTagData* TagData = It.Key;
+		if (!TagData->IsReportable())
+		{
+			continue;
+		}
+		if (TagData->GetName() == TagName_Untagged)
+		{
+			continue; // Handled by OverrideUntaggedName.
+		}
+		if (ExistingColumns.Contains(TagData))
+		{
+			continue;
+		}
+
+		ExistingColumns.Add(TagData);
+		Columns.Add(TagData);
+		bUpdated = true;
+	}
+	return bUpdated;
+}
+
+void FLLMCsvWriter::WriteHeader(const FTagData* OverrideTrackedTotalTagData, const FTagData* OverrideUntaggedTagData)
+{
+	int64 OriginalOffset = Archive->Tell();
+	Archive->Seek(0);
+
+	TStringBuilder<256> NameBuffer;
+	auto WriteTagData = [this, &NameBuffer](const FTagData* TagData)
+	{
+		if (!TagData)
+		{
+			return;
+		}
+
+		NameBuffer.Reset();
+		TagData->AppendDisplayPath(NameBuffer);
+		NameBuffer << TEXT(",");
+		Write(NameBuffer);
+	};
+
+	WriteTagData(OverrideTrackedTotalTagData);
+	WriteTagData(OverrideUntaggedTagData);
+	for (const FTagData* TagData : Columns)
+	{
+		WriteTagData(TagData);
+	}
+
+	int64 ColumnTitleTotalSize = Archive->Tell();
+	if (ColumnTitleTotalSize >= CVarLLMHeaderMaxSize.GetValueOnAnyThread())
+	{
+		UE_LOG(LogHAL, Error,
+			TEXT("LLM column titles have overflowed, LLM CSM data will be corrupted. Increase CVarLLMHeaderMaxSize > %" INT64_FMT),
+			ColumnTitleTotalSize);
+	}
+
+	Archive->Seek(OriginalOffset);
+}
+
+void FLLMCsvWriter::AddRow(FLowLevelMemTracker& LLMRef, const FTrackerTagSizeMap& TagSizes,
+	const FTagData* OverrideTrackedTotalTagData, const FTagData* OverrideUntaggedTagData, int64 TrackedTotal,
+	UE::LLM::ESizeParams SizeParams)
+{
+	TStringBuilder<256> TextBuffer;
+	auto WriteValue = [this, &TextBuffer](int64 Value)
+	{
+		TextBuffer.Reset();
+		TextBuffer.Appendf(TEXT("%0.2f,"), (float)Value / 1024.0f / 1024.0f);
+		Write(TextBuffer);
+	};
+	auto WriteTag = [&WriteValue, &TagSizes, SizeParams](const FTagData* TagData)
+	{
+		if (!TagData)
+		{
+			WriteValue(0);
+		}
+		else
+		{
+			const FTrackerTagSizeData* AllocationData = TagSizes.Find(TagData);
+			if (!AllocationData)
+			{
+				WriteValue(0);
+			}
+			else
+			{
+				WriteValue(AllocationData->GetSize(SizeParams));
+			}
+		}
+	};
+
+	if (OverrideTrackedTotalTagData)
+	{
+		WriteValue(TrackedTotal);
+	}
+	if (OverrideUntaggedTagData)
+	{
+		WriteTag(LLMRef.FindTagData(TagName_Untagged, ELLMTagSet::None));
+	}
+
+	for (const FTagData* TagData : Columns)
+	{
+		WriteTag(TagData);
+	}
+	Write(TEXTVIEW("\n"));
+
+	WriteCount++;
+
+	if (CVarLLMWriteInterval.GetValueOnAnyThread())
+	{
+		UE_LOG(LogHAL, Log, TEXT("Wrote LLM csv line %d"), WriteCount);
+	}
+
+	if (CVarLLMCsvFlushEveryRow.GetValueOnAnyThread())
+	{
+		Archive->Flush();
+	}
+	else if (!bRegisteredFlushDelegate)
+	{
+		// If we're not flushing every row, lazily register flush delegates to ensure we flush on a crash
+		// Note: we intentionally leak this since we can't clean it up safely
+		FCoreDelegates::OnHandleSystemError.AddRaw(this, &FLLMCsvWriter::Flush, true);
+		FCoreDelegates::GetOutOfMemoryDelegate().AddRaw(this, &FLLMCsvWriter::Flush, true);
+		FCoreDelegates::ApplicationWillEnterBackgroundDelegate.AddRaw(this, &FLLMCsvWriter::Flush, false); // When you close the app on some platforms it is suspended rather than terminated
+		FCoreDelegates::GetApplicationWillTerminateDelegate().AddRaw(this, &FLLMCsvWriter::Flush, false);
+		bRegisteredFlushDelegate = true;
+	}
+
+}
+
+// FLLMTraceWriter implementation.
+FLLMTraceWriter::FLLMTraceWriter()
+{
+}
+
+inline void FLLMTraceWriter::SetTracker(ELLMTracker InTracker)
+{
+	Tracker = InTracker;
+}
+
+void FLLMTraceWriter::Clear()
+{
+	DeclaredTags.Empty();
+}
+
+const void* FLLMTraceWriter::GetTagId(const FTagData* TagData)
+{
+	if (!TagData)
+	{
+		return nullptr;
+	}
+	return reinterpret_cast<const void*>(TagData);
+}
+
+void FLLMTraceWriter::Publish(FLowLevelMemTracker& LLMRef, const FTrackerTagSizeMap& TagSizes,
+	const FTagData* OverrideTrackedTotalTagData, const FTagData* OverrideUntaggedTagData, int64 TrackedTotal,
+	UE::LLM::ESizeParams SizeParams)
+{
+	if (!UE_TRACE_CHANNELEXPR_IS_ENABLED(MemTagChannel))
+	{
+		return;
+	}
+
+	if (!bTrackerSpecSent)
+	{
+		bTrackerSpecSent = true;
+		FAnsiStringView TrackerNames[] = { "Platform", "Default" };
+		static_assert((int)ELLMTracker::Platform == 0, "");
+		static_assert((int)ELLMTracker::Default == 1, "");
+		static_assert(UE_ARRAY_COUNT(TrackerNames) == (int)ELLMTracker::Max, "");
+		FAnsiStringView TrackerName = TrackerNames[(uint8)Tracker];
+		UE_TRACE_LOG(LLM, TrackerSpec, MemTagChannel, TrackerName.Len() * sizeof(ANSICHAR))
+			<< TrackerSpec.TrackerId((uint8)Tracker)
+			<< TrackerSpec.Name(TrackerName.GetData(), TrackerName.Len());
+	}
+
+	if (!bTagSetSpecSent && Tracker == ELLMTracker::Default)
+	{
+		bTagSetSpecSent = true;
+		FAnsiStringView TagSetNames[] = {
+			"System",
+			"Assets",
+			"AssetClasses",
+			"UObjectClasses"
+		};
+		static_assert((int)ELLMTagSet::Assets == 1, "");
+		static_assert((int)ELLMTagSet::AssetClasses == 2, "");
+		static_assert((int)ELLMTagSet::UObjectClasses == 3, "");
+		static_assert(UE_ARRAY_COUNT(TagSetNames) == (int)ELLMTagSet::Max, "");
+		for (uint8 TagSet = 0; TagSet < (uint8)ELLMTagSet::Max; ++TagSet)
+		{
+			if (LLMRef.IsTagSetActive((ELLMTagSet)TagSet))
+			{
+				FAnsiStringView TagSetName = TagSetNames[TagSet];
+				UE_TRACE_LOG(LLM, TagSetSpec, MemTagChannel, TagSetName.Len() * sizeof(ANSICHAR))
+					<< TagSetSpec.TagSetId(TagSet)
+					<< TagSetSpec.Name(TagSetName.GetData(), TagSetName.Len());
+			}
+		}
+	}
+
+	if (OverrideTrackedTotalTagData && !OverrideTrackedTotalTagData->IsTraceReportable())
+	{
+		OverrideTrackedTotalTagData = nullptr;
+	}
+	if (OverrideUntaggedTagData && !OverrideUntaggedTagData->IsTraceReportable())
+	{
+		OverrideUntaggedTagData = nullptr;
+	}
+
+	if (OverrideTrackedTotalTagData)
+	{
+		SendTagDeclaration(OverrideTrackedTotalTagData);
+	}
+	if (OverrideUntaggedTagData)
+	{
+		SendTagDeclaration(OverrideUntaggedTagData);
+	}
+	for (const TPair<const FTagData*, FTrackerTagSizeData>& It : TagSizes)
+	{
+		const FTagData* TagData = It.Key;
+		if (!TagData->IsTraceReportable() || !LLMRef.IsTagSetActive(TagData->GetTagSet()))
+		{
+			continue;
+		}
+		if (OverrideUntaggedTagData &&
+			TagData->GetTagSet() == ELLMTagSet::None &&
+			TagData->GetName() == TagName_Untagged)
+		{
+			continue; // Handled by OverrideUntaggedTagData.
+		}
+		SendTagDeclaration(TagData);
+	}
+
+	TArray<uint64, FDefaultLLMAllocator> TagIds;
+	TArray<int64, FDefaultLLMAllocator> TagValues;
+	TagIds.Reserve(TagSizes.Num() + 2);
+	TagValues.Reserve(TagSizes.Num() + 2);
+	auto AddValue = [&TagIds, &TagValues](const FTagData* TagData, int64 Value)
+	{
+		const uint64 TagId = uint64(GetTagId(TagData));
+		TagIds.Add(TagId);
+		TagValues.Add(Value);
+	};
+
+	if (OverrideTrackedTotalTagData)
+	{
+		AddValue(OverrideTrackedTotalTagData, TrackedTotal);
+	}
+	if (OverrideUntaggedTagData)
+	{
+		const FTagData* TagData = LLMRef.FindTagData(TagName_Untagged, ELLMTagSet::None);
+		if (!TagData)
+		{
+			AddValue(OverrideUntaggedTagData, 0);
+		}
+		else
+		{
+			const FTrackerTagSizeData* AllocationData = TagSizes.Find(TagData);
+			AddValue(OverrideUntaggedTagData, AllocationData ? AllocationData->GetSize(SizeParams) : 0);
+		}
+	}
+
+	for (const TPair<const FTagData*, FTrackerTagSizeData>& It : TagSizes)
+	{
+		const FTagData* TagData = It.Key;
+		if (!TagData->IsTraceReportable() || !LLMRef.IsTagSetActive(TagData->GetTagSet()))
+		{
+			continue;
+		}
+		if (OverrideUntaggedTagData &&
+			TagData->GetTagSet() == ELLMTagSet::None &&
+			TagData->GetName() == TagName_Untagged)
+		{
+			continue; // Handled by OverrideUntaggedTagData.
+		}
+		AddValue(TagData, It.Value.GetSize(SizeParams));
+	}
+
+	LLMCheck(TagIds.Num() == TagValues.Num());
+	const uint64 Cycle = FPlatformTime::Cycles64();
+	UE_TRACE_LOG(LLM, TagValue, MemTagChannel)
+		<< TagValue.TrackerId((uint8)Tracker)
+		<< TagValue.Cycle(Cycle)
+		<< TagValue.Tags(TagIds.GetData(), TagIds.Num())
+		<< TagValue.Values(TagValues.GetData(), TagValues.Num());
+}
+
+void FLLMTraceWriter::SendTagDeclaration(const FTagData* TagData)
+{
+	LLMCheck(TagData);
+
+	bool bAlreadyInSet;
+	DeclaredTags.Add(TagData, &bAlreadyInSet);
+	if (bAlreadyInSet)
+	{
+		return;
+	}
+
+	const FTagData* Parent = TagData->GetParent();
+	if (Parent)
+	{
+		SendTagDeclaration(Parent);
+	}
+
+	TStringBuilder<1024> NameBuffer;
+	TagData->AppendDisplayPath(NameBuffer);
+	UE_TRACE_LOG(LLM, TagsSpec, MemTagChannel, NameBuffer.Len() * sizeof(ANSICHAR))
+		<< TagsSpec.TagId(uint64(GetTagId(TagData)))
+		<< TagsSpec.ParentId(uint64(GetTagId(Parent)))
+		<< TagsSpec.TagSetId(static_cast<uint8>(TagData->GetTagSet()))
+		<< TagsSpec.Name(*NameBuffer, NameBuffer.Len());
+};
+
+// FLLMCsvProfilerWriter implementation.
+FLLMCsvProfilerWriter::FLLMCsvProfilerWriter()
+{
+}
+
+void FLLMCsvProfilerWriter::Clear()
+{
+#if LLM_CSV_PROFILER_WRITER_ENABLED
+	TagDataToCsvStatName.Empty();
+#endif
+}
+
+void FLLMCsvProfilerWriter::SetTracker(ELLMTracker InTracker)
+{
+#if LLM_CSV_PROFILER_WRITER_ENABLED
+	check(InTracker == ELLMTracker::Platform || InTracker == ELLMTracker::Default);
+	Tracker = InTracker;
+#endif
+}
+
+void FLLMCsvProfilerWriter::Publish(FLowLevelMemTracker& LLMRef, const FTrackerTagSizeMap& TagSizes,
+	const FTagData* OverrideTrackedTotalTagData, const FTagData* OverrideUntaggedTagData, int64 TrackedTotal,
+	UE::LLM::ESizeParams SizeParams)
+{
+#if LLM_CSV_PROFILER_WRITER_ENABLED
+	int32 CsvCategoryIndex = (Tracker == ELLMTracker::Platform) ?
+		CSV_CATEGORY_INDEX(LLMPlatform) : CSV_CATEGORY_INDEX(LLM);
+
+	if (OverrideTrackedTotalTagData)
+	{
+		RecordTagToCsv(CsvCategoryIndex, OverrideTrackedTotalTagData, TrackedTotal);
+	}
+
+	if (OverrideUntaggedTagData)
+	{
+		const FTagData* TagData = LLMRef.FindTagData(TagName_Untagged, ELLMTagSet::None);
+		const FTrackerTagSizeData* AllocationData = TagData ? TagSizes.Find(TagData) : nullptr;
+		RecordTagToCsv(CsvCategoryIndex, OverrideUntaggedTagData,
+			AllocationData ? AllocationData->GetSize(SizeParams) : 0);
+	}
+
+	for (const TPair<const FTagData*, FTrackerTagSizeData>& It : TagSizes)
+	{
+		const FTagData* TagData = It.Key;
+		if (!TagData->IsReportable())
+		{
+			continue;
+		}
+		if (OverrideUntaggedTagData && TagData->GetName() == TagName_Untagged)
+		{
+			// Handled separately by OverrideUntaggedTagData.
+			continue;
+		}
+		const FTrackerTagSizeData* AllocationData = TagSizes.Find(TagData);
+		RecordTagToCsv(CsvCategoryIndex, TagData, AllocationData ? AllocationData->GetSize(SizeParams) : 0);
+	}
+#endif
+}
+
+void FLLMCsvProfilerWriter::RecordTagToCsv(int32 CsvCategoryIndex, const FTagData* TagData, int64 TagSize)
+{
+#if LLM_CSV_PROFILER_WRITER_ENABLED
+	FName NewCsvStatName;
+	FName* CsvStatNamePtr = TagDataToCsvStatName.Find(TagData);
+	if (CsvStatNamePtr == nullptr)
+	{
+		TStringBuilder<FName::StringBufferSize> DisplayPath;
+		TagData->GetDisplayPath(DisplayPath);
+		NewCsvStatName = FName(DisplayPath);
+		TagDataToCsvStatName.Add(TagData, NewCsvStatName);
+		CsvStatNamePtr = &NewCsvStatName;
+	}
+	FCsvProfiler::RecordCustomStat(*CsvStatNamePtr, CsvCategoryIndex, (float)((double)TagSize / (1024.0 * 1024.0)),
+		ECsvCustomStatOp::Set);
+#endif
+}
+
+} // namespace UE::LLMPrivate
+
+#else // #if ENABLE_LOW_LEVEL_MEM_TRACKER
+
+// We need to stub some functions so things link when the UE_ENABLE_ARRAY_SLACK_TRACKING debug feature is enabled in builds
+// where LLM is disabled.  Compiling out the slack tracking code completely is difficult due to include order issues.
+// Slack tracking is in a header that must be included before LLM, so it can't access the ENABLE_LOW_LEVEL_MEM_TRACKER
+// define.  And moving that define leads to a chain reaction of other include order issues.  It's just not worth it for
+// a rarely enabled debug feature to go to all that trouble, when stubbing functions works fine...
+#if UE_ENABLE_ARRAY_SLACK_TRACKING
+uint8 LlmGetActiveTag() { return 0; }
+void ArraySlackTrackInit() {}
+void ArraySlackTrackGenerateReport(const TCHAR* Cmd, FOutputDevice& Ar) {}
+void FArraySlackTrackingHeader::AddAllocation() {}
+void FArraySlackTrackingHeader::RemoveAllocation() {}
+void FArraySlackTrackingHeader::UpdateNumUsed(int64 NewNumUsed) {}
+
+FORCENOINLINE void* FArraySlackTrackingHeader::Realloc(void* Ptr, int64 Count, uint64 ElemSize, int32 Alignment)
+{
+	return FMemory::Realloc(Ptr, Count * ElemSize, Alignment);
+}
+#endif
+
+#endif  // #else .. #if ENABLE_LOW_LEVEL_MEM_TRACKER
diff --git a/Engine/Source/Runtime/Core/Private/HAL/LLM/LLMPrivate.h b/Engine/Source/Runtime/Core/Private/HAL/LLM/LLMPrivate.h
new file mode 100644
--- /dev/null
+++ b/Engine/Source/Runtime/Core/Private/HAL/LLM/LLMPrivate.h
@@ -0,0 +1,246 @@
+// Copyright Epic Games, Inc. All Rights Reserved.
+
+#pragma once
+
+#include "HAL/LowLevelMemTracker.h"
+
+#if ENABLE_LOW_LEVEL_MEM_TRACKER
+#include "Containers/Array.h"
+#include "Containers/Map.h"
+#include "Containers/Set.h"
+#include "Containers/SortedMap.h"
+#include "HAL/LLM/MemoryUtils.h"
+
+// Private defines configuring LLM; see also public defines in LowLevelMemTracker.h
+// LLM_ALLOW_NAMES_TAGS: Set to 1 to allow arbitrary FNames to be used as tags, at the cost of more LLM memory usage per allocation. Set to 0 to store only the top-level ELLMTag for each allocation.
+// LLM_SCOPES always use FNames for the definition of their tag
+// Storing these tags on each allocation however requires 4 bytes per allocation
+// To reduce the memory overhead of LLM, this can be reduced to 1 byte per allocation at the cost of showing only the top-level 256 Tags; tags are replaced with their containing toplevel tag during allocation.
+// This setting is ignored if features requiring more information per allocation (LLM_ALLOW_STATS, LLM_ALLOW_ASSETS_TAGS) are enabled
+#ifndef LLM_ALLOW_NAMES_TAGS
+	#define LLM_ALLOW_NAMES_TAGS 1
+#endif
+
+// Enable storing the full tag for each alloc if (1) arbitrary names tags are allowed or (2) Stats are allowed (stats use names tags) or (3) Asset tags are allowed (assets use names tags)
+#define LLM_ENABLED_FULL_TAGS (LLM_ALLOW_NAMES_TAGS || LLM_ALLOW_STATS || LLM_ALLOW_ASSETS_TAGS)
+
+// Whether to enable running with reduced threads. This is currently enabled because the engine crashes with -norenderthread
+#define LLM_ENABLED_REDUCE_THREADS 0
+
+// this controls if the commandline is used to enable tracking, or to disable it. If LLM_COMMANDLINE_ENABLES_FUNCTIONALITY is true, 
+// then tracking will only happen through Engine::Init(), at which point it will be disabled unless the commandline tells 
+// it to keep going (with -llm). If LLM_COMMANDLINE_ENABLES_FUNCTIONALITY is false, then tracking will be on unless the commandline
+// disables it (with -nollm)
+#ifndef LLM_COMMANDLINE_ENABLES_FUNCTIONALITY
+	#define LLM_COMMANDLINE_ENABLES_FUNCTIONALITY 1
+#endif 
+
+// when set to 1, forces LLM to be enabled without having to parse the command line.
+#ifndef LLM_AUTO_ENABLE
+	#define LLM_AUTO_ENABLE 0
+#endif
+
+// There is a little memory and cpu overhead in tracking peak memory but it is generally more useful than current memory.
+// Disable if you need a little more memory or speed
+#define LLM_ENABLED_TRACK_PEAK_MEMORY 1
+
+namespace UE
+{
+namespace LLMPrivate
+{
+	struct FTagDataNameKey
+	{
+		FName Name;
+		ELLMTagSet TagSet;
+		
+		FTagDataNameKey(FName InName, ELLMTagSet InTagSet) :
+			Name(InName),
+			TagSet(InTagSet)
+		{
+		}
+
+		friend uint32 GetTypeHash(const FTagDataNameKey& Key)
+		{
+			constexpr uint32 HashPrime = 101;
+			return GetTypeHash(Key.Name) + HashPrime * static_cast<uint32>(Key.TagSet);
+		}
+
+		friend bool operator==(const FTagDataNameKey& A, const FTagDataNameKey& B)
+		{
+			return A.Name == B.Name && A.TagSet == B.TagSet;
+		}
+
+		friend bool operator!=(const FTagDataNameKey& A, const FTagDataNameKey& B)
+		{
+			return (A.Name != B.Name || A.TagSet != B.TagSet);
+		}
+
+		FTagDataNameKey() = delete;
+	};
+
+	/**
+	 * FTagData: Description of the properties of a Tag that can be used in LLM_SCOPE
+	 */
+	class FTagData
+	{
+	public:
+
+		FTagData(FName InName, ELLMTagSet InTagSet, FName InDisplayName, FName InParentName, FName InStatName,
+			FName InSummaryStatName, bool bInHasEnumTag, ELLMTag InEnumTag, ETagReferenceSource InReferenceSource);
+		FTagData(FName InName, ELLMTagSet InTagSet, FName InDisplayName, const FTagData* InParent, FName InStatName,
+			FName InSummaryStatName, bool bInHasEnumTag, ELLMTag InEnumTag, ETagReferenceSource InReferenceSource);
+		~FTagData();
+
+		bool IsParentConstructed() const;
+		bool IsFinishConstructed() const;
+		FName GetName() const;
+		FName GetDisplayName() const;
+		void GetDisplayPath(FStringBuilderBase& Result, int32 MaxLen=-1) const;
+		void AppendDisplayPath(FStringBuilderBase& Result, int32 MaxLen=-1) const;
+		const FTagData* GetParent() const;
+		FName GetParentName() const;
+		FName GetParentNameSafeBeforeFinishConstruct() const;
+		FName GetStatName() const;
+		FName GetSummaryStatName() const;
+		ELLMTag GetEnumTag() const;
+		ELLMTagSet GetTagSet() const;
+		bool HasEnumTag() const;
+		const FTagData* GetContainingEnumTagData() const;
+		ELLMTag GetContainingEnum() const;
+		ETagReferenceSource GetReferenceSource() const;
+		int32 GetIndex() const;
+		bool IsReportable() const;
+		bool IsStatsReportable() const;
+		bool IsTraceReportable() const;
+
+		void SetParent(const FTagData* InParent);
+		void SetIndex(int32 InIndex);
+		void SetIsReportable(bool bInReportable);
+		void SetIsTraceReportable(bool bInTraceReportable);
+		void SetFinishConstructed();
+
+		// These functions are normally invalid - these properties should be immutable - but are called for EnumTags during bootstrapping
+		FTagData(ELLMTag InEnumTag);
+		void SetName(FName InName);
+		void SetDisplayName(FName InDisplayName);
+		void SetStatName(FName InStatName);
+		void SetSummaryStatName(FName InSummaryStatName);
+		void SetParentName(FName InParentName);
+
+	private:
+		bool IsUsedAsDisplayParent() const;
+
+		FName Name;
+		FName DisplayName;
+		union
+		{
+			const FTagData* Parent;
+			FName ParentName;
+		};
+		FName StatName;
+		FName SummaryStatName;
+		int32 Index;
+		ELLMTag EnumTag;
+		ETagReferenceSource ReferenceSource;
+		ELLMTagSet TagSet;
+		bool bIsFinishConstructed;
+		bool bParentIsName;
+		bool bHasEnumTag;
+		bool bIsReportable;
+		bool bIsTraceReportable;
+	};
+
+	// TagData container types that use FLLMAllocator
+	class FTagDataNameMap : public TMap<FTagDataNameKey, FTagData*, FDefaultSetLLMAllocator>
+	{
+		using TMap<FTagDataNameKey, FTagData*, FDefaultSetLLMAllocator>::TMap;
+	};
+	class FConstTagDataArray : public TArray<const FTagData*, FDefaultLLMAllocator>
+	{
+		using TArray<const FTagData*, FDefaultLLMAllocator>::TArray;
+	};
+	class FTagDataArray : public TArray<FTagData*, FDefaultLLMAllocator>
+	{
+		using TArray<FTagData*, FDefaultLLMAllocator>::TArray;
+	};
+
+	class FTagFilter : public TLLMSet<FName>
+	{
+	};
+}
+}
+
+// Interface for Algo::TopologicalSort for FTagDataArray
+inline UE::LLMPrivate::FTagData** GetData(UE::LLMPrivate::FTagDataArray& Array)
+{
+	return Array.GetData();
+}
+inline UE::LLMPrivate::FTagDataArray::SizeType GetNum(UE::LLMPrivate::FTagDataArray& Array)
+{
+	return Array.Num();
+}
+
+namespace UE
+{
+namespace LLMPrivate
+{
+	/** Size information stored on the tracker for a tag; includes amounts aggregated from threadstates and from external api users */
+	struct FTrackerTagSizeData
+	{
+		int64 Size = 0;
+#if LLM_ENABLED_TRACK_PEAK_MEMORY
+		int64 PeakSize = 0;
+#endif
+		int64 SizeInSnapshot = 0;
+		int64 ExternalAmount = 0;
+		bool bExternalValid = false;
+		bool bExternalAddToTotal = false;
+
+		int64 GetSize(UE::LLM::ESizeParams SizeParams) const
+		{
+			int64 CurrentSize = Size;
+
+#if LLM_ENABLED_TRACK_PEAK_MEMORY
+			if (EnumHasAnyFlags(SizeParams, UE::LLM::ESizeParams::ReportPeak))
+			{
+				CurrentSize = PeakSize;
+			}
+#endif
+			// Note, this will also subtract the snapshotted size from PeakSize if that flag is enabled
+			if (EnumHasAnyFlags(SizeParams, UE::LLM::ESizeParams::RelativeToSnapshot))
+			{
+				CurrentSize = FMath::Clamp<int64>(CurrentSize - SizeInSnapshot, 0, INT64_MAX); //-V569
+			}
+
+			return CurrentSize;
+		};
+
+		void CaptureSnapshot()
+		{
+			SizeInSnapshot = Size;
+		}
+
+		void ClearSnapshot()
+		{
+			SizeInSnapshot = 0;
+		}
+	};
+	typedef TLLMMap<const FTagData*, FTrackerTagSizeData> FTrackerTagSizeMap;
+
+	/**
+	 * Size information stored on each threadstate for a tag
+	 * TagSizes are sorted by Index instead of by pointer in the ThreadTagSizeMap to enforce the constraint that Parents come before children
+	 */
+	struct FThreadTagSizeData
+	{
+		const FTagData* TagData = nullptr;
+		int64 Size = 0;
+	};
+	// 
+	typedef TSortedMap<int32, FThreadTagSizeData, FDefaultLLMAllocator> FThreadTagSizeMap;
+}
+}
+
+FName LLMGetTagUniqueName(ELLMTag Tag);
+
+#endif		// #if ENABLE_LOW_LEVEL_MEM_TRACKER
diff --git a/Engine/Source/Runtime/Core/Private/HAL/LLM/LLMTest.cpp b/Engine/Source/Runtime/Core/Private/HAL/LLM/LLMTest.cpp
new file mode 100644
--- /dev/null
+++ b/Engine/Source/Runtime/Core/Private/HAL/LLM/LLMTest.cpp
@@ -0,0 +1,178 @@
+// Copyright Epic Games, Inc. All Rights Reserved.
+#include "HAL/LowLevelMemTracker.h"
+
+#if ENABLE_LOW_LEVEL_MEM_TRACKER
+
+#include "HAL/LLM/LLMPrivate.h"
+
+#include "Containers/Set.h"
+#include "Misc/AutomationTest.h"
+
+LLM_DEFINE_TAG(LLMTestParentTag_LLMTestChildTag1, NAME_None, TEXT("LLMTestParentTag")); // A test tag to verify that child tags can be declared before parent tags
+LLM_DEFINE_TAG(LLMTestParentTag, NAME_None, NAME_None);
+LLM_DEFINE_TAG(LLMTestParentTag_LLMTestChildTag2, NAME_None, TEXT("LLMTestParentTag")); // A test tag to verify that child tags can be declared after parent tags
+LLM_DEFINE_TAG(LLMTestParentTag_LLMTestChildTag3); // A test tag to verify that child tags parse their parent tags
+
+#if WITH_DEV_AUTOMATION_TESTS
+
+IMPLEMENT_SIMPLE_AUTOMATION_TEST(FLLMScopeTest, "System.Core.LLM.Scope", EAutomationTestFlags_ApplicationContextMask | EAutomationTestFlags::EngineFilter);
+
+bool FLLMScopeTest::RunTest(const FString& Parameters)
+{
+	int* AllocationByNameExisting = nullptr;
+	int* AllocationByNameNew = nullptr;
+	int* AllocationByTag = nullptr;
+	int* AllocationByTag3 = nullptr;
+	FLowLevelMemTracker& LocalLLM = FLowLevelMemTracker::Get();
+	bool bIsEnabled = FLowLevelMemTracker::IsEnabled();
+	{
+		LLM_SCOPE_BYNAME(TEXT("LLMTestParentTag/LLMTestChildTag1"));
+		AllocationByNameExisting = new int[10];
+		if (bIsEnabled)
+		{
+			TestEqual(TEXT("NameExisting has proper name"), LocalLLM.GetActiveTagData(ELLMTracker::Default)->GetName(), FName("LLMTestParentTag_LLMTestChildTag1"));
+			TestEqual(TEXT("NameExisting has proper parent"), LocalLLM.GetActiveTagData(ELLMTracker::Default)->GetParent()->GetName(), FName("LLMTestParentTag"));
+		}
+	}
+	{
+		LLM_SCOPE_BYNAME(TEXT("LLMTestParentTag/LLMTestChildTag4"));
+		AllocationByNameNew = new int[10];
+		if (bIsEnabled)
+		{
+			TestEqual(TEXT("NameNew has proper name"), LocalLLM.GetActiveTagData(ELLMTracker::Default)->GetName(), FName("LLMTestParentTag_LLMTestChildTag4"));
+			TestEqual(TEXT("NameNew has proper parent"), LocalLLM.GetActiveTagData(ELLMTracker::Default)->GetParent()->GetName(), FName("LLMTestParentTag"));
+		}
+	}
+	{
+		LLM_SCOPE_BYTAG(LLMTestParentTag_LLMTestChildTag2);
+		AllocationByTag = new int[10];
+		if (bIsEnabled)
+		{
+			TestEqual(TEXT("ByTag has proper name"), LocalLLM.GetActiveTagData(ELLMTracker::Default)->GetName(), FName("LLMTestParentTag_LLMTestChildTag2"));
+			TestEqual(TEXT("ByTag has proper parent"), LocalLLM.GetActiveTagData(ELLMTracker::Default)->GetParent()->GetName(), FName("LLMTestParentTag"));
+		}
+	}
+	{
+		LLM_SCOPE_BYTAG(LLMTestParentTag_LLMTestChildTag3);
+		AllocationByTag3 = new int[10];
+		if (bIsEnabled)
+		{
+			TestEqual(TEXT("Parsed parent tag has proper name"), LocalLLM.GetActiveTagData(ELLMTracker::Default)->GetName(), FName("LLMTestParentTag_LLMTestChildTag3"));
+			TestEqual(TEXT("Parsed parent tag has proper parent"), LocalLLM.GetActiveTagData(ELLMTracker::Default)->GetParent()->GetName(), FName("LLMTestParentTag"));
+		}
+	}
+	delete[] AllocationByNameExisting;
+	delete[] AllocationByNameNew;
+	delete[] AllocationByTag;
+	delete[] AllocationByTag3;
+
+	return true;
+}
+
+IMPLEMENT_SIMPLE_AUTOMATION_TEST(FLLMMiscTest, "System.Core.LLM.Misc", EAutomationTestFlags_ApplicationContextMask | EAutomationTestFlags::EngineFilter);
+
+bool FLLMMiscTest::RunTest(const FString& Parameters)
+{
+	// Test LLMGetTagUniqueNames
+	TSet<FName> UniqueNames;
+	for (int32 Tag = 0; Tag < static_cast<int32>(LLM_TAG_COUNT); ++Tag)
+	{
+		FName Name = LLMGetTagUniqueName(static_cast<ELLMTag>(Tag));
+		if ((0 <= Tag && Tag < static_cast<int32>(ELLMTag::GenericTagCount)) || (LLM_CUSTOM_TAG_START <= Tag && Tag <= LLM_CUSTOM_TAG_END))
+		{
+			bool bAlreadySet;
+			UniqueNames.Add(Name, &bAlreadySet);
+			TestTrue(TEXT("LLMGetTagUniqueNames are unique"), !Name.IsNone() && !bAlreadySet);
+		}
+		else
+		{
+			TestTrue(TEXT("LLMGetTagUniqueName returns None for invalid tags"), Name.IsNone());
+		}
+	}
+	return true;
+}
+
+IMPLEMENT_SIMPLE_AUTOMATION_TEST(FTopologicalSortTest, "System.Core.LLM.TopologicalSort", EAutomationTestFlags_ApplicationContextMask | EAutomationTestFlags::EngineFilter);
+
+bool FTopologicalSortTest::RunTest(const FString& Parameters)
+{
+	TArray<TArray<int32>> Edges;
+	auto AppendEdges = [&Edges](int32 InVertex, int32* OutEdges, int32& OutNumEdges)
+	{
+		TArray<int32>& VertexEdges = Edges[InVertex];
+		OutNumEdges = VertexEdges.Num();
+		if (OutNumEdges)
+		{
+			FMemory::Memcpy(OutEdges, VertexEdges.GetData(), OutNumEdges*sizeof(VertexEdges[0]));
+		}
+	};
+
+	// Simple case:
+	//							0
+	//						|		|
+	//						v		v
+	//						1  ->	2
+	{
+		TArray<TArray<int32>> RootToLeafEdges;
+		RootToLeafEdges.Add(TArray<int32>{ 1, 2});
+		RootToLeafEdges.Add(TArray<int32>{ 2 });
+		RootToLeafEdges.AddDefaulted();
+
+		Edges = RootToLeafEdges;
+		LLMAlgo::TopologicalSortRootToLeaf(Edges, AppendEdges);
+		TestTrue(TEXT("SimpleCase pre-sorted RootToLeaf remains stable sorted"), Edges == RootToLeafEdges);
+
+		TArray<TArray<int32>> ReverseEdges;
+		for (int n = RootToLeafEdges.Num() - 1; n >= 0; --n)
+		{
+			ReverseEdges.Add(RootToLeafEdges[n]);
+		}
+		LLMAlgo::TopologicalSortLeafToRoot(Edges, AppendEdges);
+		TestTrue(TEXT("SimpleCase reverse-sorted RootToLeaf sorts into reverse order"), Edges == ReverseEdges);
+		
+		TArray<TArray<int32>> LeafToRootEdges;
+		LeafToRootEdges.AddDefaulted();
+		LeafToRootEdges.Add(TArray<int32>{ 0});
+		LeafToRootEdges.Add(TArray<int32>{ 1,0 });
+		Edges = LeafToRootEdges;
+		LLMAlgo::TopologicalSortLeafToRoot(Edges, AppendEdges);
+		TestTrue(TEXT("SimpleCase pre-sorted LeafToRoof remains stable sorted"), Edges == LeafToRootEdges);
+
+		ReverseEdges.Reset();
+		for (int n = LeafToRootEdges.Num() - 1; n >= 0; --n)
+		{
+			ReverseEdges.Add(LeafToRootEdges[n]);
+		}
+		LLMAlgo::TopologicalSortRootToLeaf(Edges, AppendEdges);
+		TestTrue(TEXT("SimpleCase reverse-sorted LeafToRoof sorts into reverse order"), Edges == ReverseEdges);
+	}
+
+	// Simple cycle
+	//							0
+	//						|		|
+	//						v		v
+	//						1	->	2
+	//							<-	|
+	//								v
+	//								3								
+	{
+		TArray<TArray<int32>> CycleEdges;
+		CycleEdges.Add(TArray<int32>{ 1, 2});
+		CycleEdges.Add(TArray<int32>{ 2 });
+		CycleEdges.Add(TArray<int32>{ 1,3 });
+		CycleEdges.AddDefaulted();
+
+		Edges = CycleEdges;
+		LLMAlgo::TopologicalSortLeafToRoot(Edges, AppendEdges);
+		TestTrue(TEXT("Cycle correctly sorts target of cycle first in LeafToRoot order"), Edges[0] == CycleEdges[3]);
+		TestTrue(TEXT("Cycle correctly sorts referencer of cycle last in LeafToRoot order"), Edges[3] == CycleEdges[0]);
+		TestTrue(TEXT("Cycle correctly sorts vertex1 of cycle at an unspecified location into the middle"), Edges[1] == CycleEdges[1] || Edges[1] == CycleEdges[2]);
+		TestTrue(TEXT("Cycle correctly sorts vertex2 of cycle at an unspecified location into the middle"), (Edges[2] == CycleEdges[1] || Edges[2] == CycleEdges[2]) && Edges[2] != Edges[1]);
+	}
+
+	return !HasAnyErrors();
+}
+
+#endif // WITH_DEV_AUTOMATION_TESTS
+
+#endif // ENABLE_LOW_LEVEL_MEM_TRACKER
\ No newline at end of file
diff --git a/Engine/Source/Runtime/Core/Private/HAL/LLM/MemoryUtils.cpp b/Engine/Source/Runtime/Core/Private/HAL/LLM/MemoryUtils.cpp
new file mode 100644
--- /dev/null
+++ b/Engine/Source/Runtime/Core/Private/HAL/LLM/MemoryUtils.cpp
@@ -0,0 +1,18 @@
+// Copyright Epic Games, Inc. All Rights Reserved.
+
+#include "HAL/LLM/MemoryUtils.h"
+
+#include "CoreGlobals.h"
+
+#if ENABLE_LOW_LEVEL_MEM_TRACKER
+
+FORCENOINLINE void UE::Core::Private::OnInvalidLLMAllocatorNum(int32 IndexSize, int64 NewNum, SIZE_T NumBytesPerElement)
+{
+	UE_LOG(LogCore, Fatal, TEXT("Trying to resize TSizedLLMAllocator<%d> to an invalid size of % " INT64_FMT " with element size %" SIZE_T_FMT), IndexSize, NewNum, NumBytesPerElement);
+	for (;;);
+}
+
+template void TSizedLLMAllocator<32>::ForAnyElementType::ResizeAllocation(SizeType CurrentNum, SizeType NewMax, SIZE_T NumBytesPerElement);
+template void TSizedLLMAllocator<64>::ForAnyElementType::ResizeAllocation(SizeType CurrentNum, SizeType NewMax, SIZE_T NumBytesPerElement);
+
+#endif
diff --git a/Engine/Source/Runtime/Core/Private/HAL/LLM/MemoryUtils.h b/Engine/Source/Runtime/Core/Private/HAL/LLM/MemoryUtils.h
new file mode 100644
--- /dev/null
+++ b/Engine/Source/Runtime/Core/Private/HAL/LLM/MemoryUtils.h
@@ -0,0 +1,1243 @@
+// Copyright Epic Games, Inc. All Rights Reserved.
+
+#pragma once
+
+#include "HAL/LowLevelMemTracker.h"
+
+#if ENABLE_LOW_LEVEL_MEM_TRACKER
+
+#include "Algo/BinarySearch.h"
+#include "Async/UniqueLock.h"
+#include "Containers/ContainerAllocationPolicies.h"
+#include "Containers/Map.h"
+#include "Containers/Set.h"
+#include "HAL/PlatformMutex.h"
+#include "Math/NumericLimits.h"
+#include "Templates/Tuple.h"
+#include <type_traits>
+
+#define LLM_PAGE_SIZE (16*1024)
+
+#if WITH_EDITOR
+// When cooking, the number of simultaneous allocations can reach the danger zone of tens of millions, and our margin*capacity calculation ~ 100*capacity will rise over MAX_uint32
+typedef uint64 LLMNumAllocsType;
+#else
+// Even in our 64 bit runtimes, the number of simultaneous allocations we have never gets over a few million, so we don't reach the danger zone of 100*capacity > MAX_uInt32
+typedef uint32 LLMNumAllocsType;
+#endif
+
+// Trivially movable types without destructors only; memcpy is used and destructors are not called
+template<typename T, typename SizeType=int32>
+class FLLMArray
+{
+public:
+	FLLMArray()
+		: Array(StaticArray)
+		, Count(0)
+		, Capacity(StaticArrayCapacity)
+		, Allocator(nullptr)
+	{
+	}
+
+	~FLLMArray()
+	{
+		Clear(true);
+	}
+
+	void SetAllocator(UE::LLMPrivate::FLLMAllocator* InAllocator)
+	{
+		Allocator = InAllocator;
+	}
+
+	SizeType Num() const
+	{
+		return Count;
+	}
+
+	void Clear(bool ReleaseMemory = false)
+	{
+		if (ReleaseMemory)
+		{
+			if (Array != StaticArray)
+			{
+				Allocator->Free(Array, Capacity * sizeof(T));
+				Array = StaticArray;
+			}
+			Capacity = StaticArrayCapacity;
+		}
+		Count = 0;
+	}
+
+	void Add(const T& Item)
+	{
+		if (Count == Capacity)
+		{
+			SizeType NewCapacity = DefaultCapacity;
+			if (Capacity)
+			{
+				NewCapacity = Capacity + (Capacity / 2);
+				ensureMsgf(NewCapacity > Capacity, TEXT("Unsigned integer overflow."));
+			}
+			Reserve(NewCapacity);
+		}
+
+		Array[Count] = Item;
+		++Count;
+	}
+
+	T RemoveLast()
+	{
+		LLMCheck(Count > 0);
+		--Count;
+		T Last = Array[Count];
+
+		return Last;
+	}
+
+	T& operator[](SizeType Index)
+	{
+		LLMCheck(Index >= 0 && Index < Count);
+		return Array[Index];
+	}
+
+	const T& operator[](SizeType Index) const
+	{
+		LLMCheck(Index >= 0 && Index < Count);
+		return Array[Index];
+	}
+
+	const T* GetData() const {
+		return Array;
+	}
+
+	T& GetLast()
+	{
+		LLMCheck(Count > 0);
+		return Array[Count - 1];
+	}
+
+	void Reserve(SizeType NewCapacity)
+	{
+		if (NewCapacity == Capacity)
+		{
+			return;
+		}
+
+		if (NewCapacity <= StaticArrayCapacity)
+		{
+			if (Array != StaticArray)
+			{
+				if (Count)
+				{
+					memcpy(StaticArray, Array, Count * sizeof(T));
+				}
+				Allocator->Free(Array, Capacity * sizeof(T));
+
+				Array = StaticArray;
+				Capacity = StaticArrayCapacity;
+			}
+		}
+		else
+		{
+			NewCapacity = AlignArbitrary(NewCapacity, ItemsPerPage);
+
+			T* NewArray = (T*)Allocator->Alloc(NewCapacity * sizeof(T));
+
+			if (Count)
+			{
+				memcpy(NewArray, Array, Count * sizeof(T));
+			}
+			if (Array != StaticArray)
+			{
+				Allocator->Free(Array, Capacity * sizeof(T));
+			}
+
+			Array = NewArray;
+			Capacity = NewCapacity;
+		}
+	}
+
+	void operator=(const FLLMArray<T>& Other)
+	{
+		Clear();
+		Reserve(Other.Count);
+		memcpy(Array, Other.Array, Other.Count * sizeof(T));
+		Count = Other.Count;
+	}
+
+	void Trim()
+	{
+		// Trim if usage has dropped below 3/4 of the total capacity
+		if (Array != StaticArray && Count < (Capacity - (Capacity / 4)))
+		{
+			Reserve(Count);
+		} 
+	}
+
+private:
+	T* Array;
+	SizeType Count;
+	SizeType Capacity;
+
+	UE::LLMPrivate::FLLMAllocator* Allocator;
+
+	static const int StaticArrayCapacity = 64;	// because the default size is so large this actually saves memory
+	T StaticArray[StaticArrayCapacity];
+
+	static const int ItemsPerPage = LLM_PAGE_SIZE / sizeof(T);
+	static const int DefaultCapacity = ItemsPerPage;
+};
+
+/*
+* hash map
+*/
+template<typename TKey, typename TValue1, typename TValue2, typename SizeType=int32>
+class LLMMap
+{
+public:
+	typedef LLMMap<TKey, TValue1, TValue2, SizeType> ThisType;
+
+	struct Values
+	{
+		TKey Key;
+		TValue1 Value1;
+		TValue2 Value2;
+	};
+
+	LLMMap()
+		: bAllStripesLocked(false)
+	{
+	}
+
+	~LLMMap()
+	{
+		Clear();
+	}
+
+	void SetAllocator(UE::LLMPrivate::FLLMAllocator* InAllocator, SizeType InDefaultCapacity = DefaultCapacity)
+	{
+		for (StripeData& Stripe : MapStripes)
+		{
+			Stripe.Allocator = InAllocator;
+
+			Stripe.Keys.SetAllocator(InAllocator);
+			Stripe.KeyHashes.SetAllocator(InAllocator);
+			Stripe.Values1.SetAllocator(InAllocator);
+			Stripe.Values2.SetAllocator(InAllocator);
+			Stripe.FreeKeyIndices.SetAllocator(InAllocator);
+
+			Stripe.Reserve(InDefaultCapacity / StripeCount);
+		}
+	}
+
+	void Clear()
+	{
+		for (StripeData& Stripe : MapStripes)
+		{
+			UE::TUniqueLock AllocationScopeLock(Stripe.Mutex);
+
+			Stripe.Keys.Clear(true);
+			Stripe.KeyHashes.Clear(true);
+			Stripe.Values1.Clear(true);
+			Stripe.Values2.Clear(true);
+			Stripe.FreeKeyIndices.Clear(true);
+
+			Stripe.Allocator->Free(Stripe.Map, Stripe.Capacity * sizeof(SizeType));
+			Stripe.Map = NULL;
+			Stripe.Count = 0;
+			Stripe.Capacity = 0;
+		}
+	}
+
+	// Add a value to this set.
+	// If this set already contains the value does nothing.
+	void Add(const TKey& Key, const TValue1& Value1, const TValue2& Value2)
+	{
+		SizeType KeyHash = Key.GetHashCode();
+		StripeData& Stripe = MapStripes[GetStripeIndex(KeyHash)];
+
+		UE::TUniqueLock AllocationScopeLock(Stripe.Mutex);
+
+		LLMCheck(Stripe.Map);
+
+		SizeType MapIndex = Stripe.GetMapIndex(Key, KeyHash);
+		SizeType KeyIndex = Stripe.Map[MapIndex];
+
+		if (KeyIndex != InvalidIndex)
+		{
+			static bool ShownWarning = false;
+			if (!ShownWarning)
+			{
+				FPlatformMisc::LowLevelOutputDebugString(TEXT("LLM WARNING: Replacing allocation in tracking map. Alloc/Free Mismatch.\n"));
+				ShownWarning = true;
+			}
+
+			Stripe.Values1[KeyIndex] = Value1;
+			Stripe.Values2[KeyIndex] = Value2;
+		}
+		else
+		{
+			// Be careful with the multiplication to avoid overflow, while still using only integer arithmetic for speed
+			// Margin/256U is the actual float we want to multiply by. The exact number doesn't matter, so use one order of operations
+			// when Capacity is low, and another when it is high
+			SizeType MaxCount = (Stripe.Capacity >= 256U * 256U ? (Stripe.Capacity / 256U) * Margin : (Margin * Stripe.Capacity) / 256U);
+			if (Stripe.Count >= MaxCount)
+			{
+				if (Stripe.Capacity * 2 < Stripe.Capacity)
+				{
+					// Trying to issue a check statement here will cause reentry into this function, use PLATFORM_BREAK directly instead
+					FPlatformMisc::LowLevelOutputDebugString(TEXT("LLM Error: Integer overflow in LLMap::Add, Capacity has reached its maximum size.\n"));
+					PLATFORM_BREAK();
+				}
+				if (Stripe.Count > MaxCount)
+				{
+					// Trying to issue a check statement here will cause reentry into this function, use PLATFORM_BREAK directly instead
+					// This shouldn't happen, because Count is only incremented here, and Capacity is only changed here, and Margin does not change, so Count should equal MaxCount before it goes over it
+					FPlatformMisc::LowLevelOutputDebugString(TEXT("LLM Assertion failure: Count > MaxCount.\n"));
+					PLATFORM_BREAK();
+				}
+				Stripe.Grow();
+				MapIndex = Stripe.GetMapIndex(Key, KeyHash);
+			}
+
+			if (Stripe.FreeKeyIndices.Num())
+			{
+				SizeType FreeIndex = Stripe.FreeKeyIndices.RemoveLast();
+				Stripe.Map[MapIndex] = FreeIndex;
+				Stripe.Keys[FreeIndex] = Key;
+				Stripe.KeyHashes[FreeIndex] = KeyHash;
+				Stripe.Values1[FreeIndex] = Value1;
+				Stripe.Values2[FreeIndex] = Value2;
+			}
+			else
+			{
+				Stripe.Map[MapIndex] = Stripe.Keys.Num();
+				Stripe.Keys.Add(Key);
+				Stripe.KeyHashes.Add(KeyHash);
+				Stripe.Values1.Add(Value1);
+				Stripe.Values2.Add(Value2);
+			}
+
+			++Stripe.Count;
+		}
+	}
+
+	Values GetValue(const TKey& Key)
+	{
+		Values RetValues;
+		RetValues.Key = Find(Key, RetValues.Value1, RetValues.Value2);
+		LLMEnsure(RetValues.Key);
+		return RetValues;
+	}
+
+	TKey Find(const TKey& Key, TValue1& OutValue1, TValue2& OutValue2) const
+	{
+		SizeType KeyHash = Key.GetHashCode();
+		const StripeData& Stripe = MapStripes[GetStripeIndex(KeyHash)];
+
+		UE::TUniqueLock AllocationScopeLock(Stripe.Mutex);
+
+		SizeType MapIndex = Stripe.GetMapIndex(Key, KeyHash);
+
+		SizeType KeyIndex = Stripe.Map[MapIndex];
+		if (KeyIndex == InvalidIndex)
+		{
+			return TKey();
+		}
+
+		OutValue1 = Stripe.Values1[KeyIndex];
+		OutValue2 = Stripe.Values2[KeyIndex];
+
+		return Stripe.Keys[KeyIndex];
+	}
+
+	bool Remove(const TKey& Key, Values& OutValues)
+	{
+		SizeType KeyHash = Key.GetHashCode();
+		StripeData& Stripe = MapStripes[GetStripeIndex(KeyHash)];
+
+		UE::TUniqueLock AllocationScopeLock(Stripe.Mutex);
+
+		LLMCheck(Stripe.Map);
+
+		SizeType MapIndex = Stripe.GetMapIndex(Key, KeyHash);
+		if (!Stripe.IsItemInUse(MapIndex))
+		{
+			return false;
+		}
+
+		SizeType KeyIndex = Stripe.Map[MapIndex];
+
+		OutValues.Key = Stripe.Keys[KeyIndex];
+		OutValues.Value1 = Stripe.Values1[KeyIndex];
+		OutValues.Value2 = Stripe.Values2[KeyIndex];
+
+		if (KeyIndex == Stripe.Keys.Num() - 1)
+		{
+			Stripe.Keys.RemoveLast();
+			Stripe.KeyHashes.RemoveLast();
+			Stripe.Values1.RemoveLast();
+			Stripe.Values2.RemoveLast();
+		}
+		else
+		{
+			Stripe.FreeKeyIndices.Add(KeyIndex);
+		}
+
+		// find first index in this array
+		SizeType IndexIter = MapIndex;
+		SizeType FirstIndex = MapIndex;
+		if (!IndexIter)
+		{
+			IndexIter = Stripe.Capacity;
+		}
+		--IndexIter;
+		while (Stripe.IsItemInUse(IndexIter))
+		{
+			FirstIndex = IndexIter;
+			if (!IndexIter)
+			{
+				IndexIter = Stripe.Capacity;
+			}
+			--IndexIter;
+		}
+
+		bool Found = false;
+		for (;;)
+		{
+			// find the last item in the array that can replace the item being removed
+			SizeType IndexIter2 = (MapIndex + 1) & (Stripe.Capacity - 1);
+
+			SizeType SwapIndex = InvalidIndex;
+			while (Stripe.IsItemInUse(IndexIter2))
+			{
+				SizeType SearchKeyIndex = Stripe.Map[IndexIter2];
+				const SizeType SearchHashCode = Stripe.KeyHashes[SearchKeyIndex];
+				const SizeType SearchInsertIndex = SearchHashCode & (Stripe.Capacity - 1);
+
+				if (Stripe.InRange(SearchInsertIndex, FirstIndex, MapIndex))
+				{
+					SwapIndex = IndexIter2;
+					Found = true;
+				}
+
+				IndexIter2 = (IndexIter2 + 1) & (Stripe.Capacity - 1);
+			}
+
+			// swap the item
+			if (Found)
+			{
+				Stripe.Map[MapIndex] = Stripe.Map[SwapIndex];
+				MapIndex = SwapIndex;
+				Found = false;
+			}
+			else
+			{
+				break;
+			}
+		}
+
+		// remove the last item
+		Stripe.Map[MapIndex] = InvalidIndex;
+
+		--Stripe.Count;
+
+		return true;
+	}
+
+	SizeType Num() const
+	{
+		SizeType TotalCount = 0;
+		for (StripeData& Stripe : MapStripes)
+		{
+			UE::TUniqueLock AllocationScopeLock(Stripe.Mutex);
+			TotalCount += Stripe.Count;
+		}
+		return TotalCount;
+	}
+
+	bool HasKey(const TKey& Key)
+	{
+		SizeType KeyHash = Key.GetHashCode();
+		StripeData& Stripe = MapStripes[GetStripeIndex(KeyHash)];
+
+		UE::TUniqueLock AllocationScopeLock(Stripe.Mutex);
+
+		SizeType MapIndex = Stripe.GetMapIndex(Key, KeyHash);
+		return Stripe.IsItemInUse(MapIndex);
+	}
+
+	void Trim()
+	{
+		for (StripeData& Stripe : MapStripes)
+		{
+			UE::TUniqueLock AllocationScopeLock(Stripe.Mutex);
+
+			Stripe.Keys.Trim();
+			Stripe.KeyHashes.Trim();
+			Stripe.Values1.Trim();
+			Stripe.Values2.Trim();
+			Stripe.FreeKeyIndices.Trim();
+		}
+	}
+
+	struct FBaseIterator
+	{
+	public:
+		FBaseIterator& operator++()
+		{
+			++MapIndex;
+			while (StripeIndex < StripeCount)
+			{
+				StripeData& Stripe = MapRef.MapStripes[StripeIndex];
+				while (MapIndex < Stripe.Capacity)
+				{
+					if (Stripe.IsItemInUse(MapIndex))
+					{
+						return *this;
+					}
+					MapIndex++;
+				}
+
+				++StripeIndex;
+				MapIndex = 0;
+			}
+
+			return *this;
+		}
+
+
+	protected:
+		FBaseIterator(ThisType& InMap, bool bEnd)
+			: MapRef(InMap), MapIndex(0)
+		{
+			if (bEnd)
+			{
+				StripeIndex = StripeCount;
+				MapIndex = 0;
+			}
+			else
+			{
+				StripeIndex = 0;
+				MapIndex = -1;
+				++(*this);
+			}
+		}
+
+		ThisType& MapRef;
+		int32 StripeIndex;
+		SizeType MapIndex;
+	};
+
+	struct FIterator;
+	struct FTuple
+	{
+		const TKey& Key;
+		TValue1& Value1;
+		TValue2& Value2;
+
+	private:
+		FTuple(const TKey& InKey, TValue1& InValue1, TValue2& InValue2)
+			:Key(InKey), Value1(InValue1), Value2(InValue2)
+		{
+		}
+
+		friend struct ThisType::FIterator;
+	};
+
+	struct FIterator : public FBaseIterator
+	{
+		FIterator(ThisType& InMap, bool bEnd)
+			: FBaseIterator(InMap, bEnd)
+		{
+		}
+
+		bool operator!=(const FIterator& Other) const
+		{
+			return FBaseIterator::StripeIndex != Other.FBaseIterator::StripeIndex || 
+				FBaseIterator::MapIndex != Other.FBaseIterator::MapIndex;
+		}
+
+		FTuple operator*() const
+		{
+			ThisType& LocalMap = FBaseIterator::MapRef;
+			StripeData& Stripe = LocalMap.MapStripes[FBaseIterator::StripeIndex];
+			SizeType KeyIndex = Stripe.Map[FBaseIterator::MapIndex];
+			return FTuple(Stripe.Keys[KeyIndex], Stripe.Values1[KeyIndex], Stripe.Values2[KeyIndex]);
+		}
+	};
+
+	struct FConstIterator;
+	struct FConstTuple
+	{
+		const TKey& Key;
+		const TValue1& Value1;
+		const TValue2& Value2;
+
+	private:
+		FConstTuple(const TKey& InKey, const TValue1& InValue1, const TValue2& InValue2)
+			:Key(InKey), Value1(InValue1), Value2(InValue2)
+		{
+		}
+
+		friend struct ThisType::FConstIterator;
+	};
+
+	struct FConstIterator : public FBaseIterator
+	{
+		FConstIterator(const ThisType& InMap, bool bEnd)
+			: FBaseIterator(const_cast<ThisType&>(InMap), bEnd)
+		{
+		}
+
+		bool operator!=(const FConstIterator& Other) const
+		{
+			return FBaseIterator::StripeIndex != Other.FBaseIterator::StripeIndex || 
+				FBaseIterator::MapIndex != Other.FBaseIterator::MapIndex;
+		}
+
+		FConstTuple operator*() const
+		{
+			ThisType& LocalMap = FBaseIterator::MapRef;
+			StripeData& Stripe = LocalMap.MapStripes[FBaseIterator::StripeIndex];
+			SizeType KeyIndex = Stripe.Map[FBaseIterator::MapIndex];
+			return FConstTuple(Stripe.Keys[KeyIndex], Stripe.Values1[KeyIndex], Stripe.Values2[KeyIndex]);
+		}
+	};
+
+	FIterator begin()
+	{
+		LLMCheck(bAllStripesLocked);
+		return FIterator(*this, false);
+	}
+
+	FIterator end()
+	{
+		return FIterator(*this, true);
+	}
+
+	FConstIterator begin() const
+	{
+		LLMCheck(bAllStripesLocked);
+		return FConstIterator(*this, false);
+	}
+
+	FConstIterator  end() const
+	{
+		return FConstIterator(*this, true);
+	}
+
+	void LockAll()
+	{
+		// Locking all stripes is required to create an iterator
+		AllStripesMutex.Lock();
+		LLMCheck(bAllStripesLocked == false);
+		bAllStripesLocked = true;
+
+		for (StripeData& Stripe : MapStripes)
+		{
+			Stripe.Mutex.Lock();
+		}
+	}
+
+	void UnlockAll()
+	{
+		for (StripeData& Stripe : MapStripes)
+		{
+			Stripe.Mutex.Unlock();
+		}
+
+		LLMCheck(bAllStripesLocked == true);
+		bAllStripesLocked = false;
+		AllStripesMutex.Unlock();
+	}
+
+	// data
+private:
+	static const int32 StripeCountLog2 = 4;		// 16 stripes
+	enum { StripeCount = 1 << StripeCountLog2 };
+	enum { DefaultCapacity = 1024 * 1024 };
+	enum { InvalidIndex = -1 };
+	static const SizeType Margin = (30 * 256) / 100;
+
+	static inline uint32 GetStripeIndex(uint32 KeyHash)
+	{
+		// FMath::Max prevents "shift too large" compile error in otherwise unreachable code when StripeCount == 1
+		return (StripeCount == 1) ? 0 : KeyHash >> (32 - FMath::Max(StripeCountLog2, 1));
+	}
+
+	static inline uint32 GetStripeIndex(uint64 KeyHash)
+	{
+		// FMath::Max prevents "shift too large" compile error in otherwise unreachable code when StripeCount == 1
+		return static_cast<uint32>((StripeCount == 1) ? 0 : KeyHash >> (64 - FMath::Max(StripeCountLog2, 1)));
+	}
+
+	struct StripeData
+	{
+		StripeData()
+			: Allocator(nullptr)
+			, Map(nullptr)
+			, Count(0)
+			, Capacity(0)
+#ifdef PROFILE_LLMMAP
+			, IterAcc(0)
+			, IterCount(0)
+#endif
+		{}
+
+		void Reserve(SizeType NewCapacity)
+		{
+			NewCapacity = FMath::Max(NewCapacity, (SizeType)1);
+			NewCapacity = static_cast<SizeType>(FPlatformMath::RoundUpToPowerOfTwo64(static_cast<uint64>(NewCapacity)));
+
+			// keep a copy of the old map
+			SizeType* OldMap = Map;
+			SizeType OldCapacity = Capacity;
+
+			LLMCheck(NewCapacity > OldCapacity);
+
+			// allocate the new table
+			Capacity = NewCapacity;
+			Map = (SizeType*)Allocator->Alloc(NewCapacity * sizeof(SizeType));
+
+			for (SizeType Index = 0; Index < NewCapacity; ++Index)
+				Map[Index] = InvalidIndex;
+
+			// copy the values from the old to the new table
+			SizeType* OldItem = OldMap;
+			for (SizeType Index = 0; Index < OldCapacity; ++Index, ++OldItem)
+			{
+				SizeType KeyIndex = *OldItem;
+
+				if (KeyIndex != InvalidIndex)
+				{
+					SizeType MapIndex = GetMapIndex(Keys[KeyIndex], KeyHashes[KeyIndex]);
+					Map[MapIndex] = KeyIndex;
+				}
+			}
+
+			Allocator->Free(OldMap, OldCapacity * sizeof(SizeType));
+		}
+
+		bool IsItemInUse(SizeType MapIndex) const
+		{
+			return Map[MapIndex] != InvalidIndex;
+		}
+
+		SizeType GetMapIndex(const TKey& Key, SizeType Hash) const
+		{
+			SizeType Mask = Capacity - 1;
+			SizeType MapIndex = Hash & Mask;
+			SizeType KeyIndex = Map[MapIndex];
+
+			while (KeyIndex != InvalidIndex && !(Keys[KeyIndex] == Key))
+			{
+				MapIndex = (MapIndex + 1) & Mask;
+				KeyIndex = Map[MapIndex];
+#ifdef PROFILE_LLMMAP
+				++IterAcc;
+#endif
+			}
+
+#ifdef PROFILE_LLMMAP
+			++IterCount;
+			double Average = IterAcc / (double)IterCount;
+			if (Average > 2.0)
+			{
+				static double LastWriteTime = 0.0;
+				double Now = FPlatformTime::Seconds();
+				if (Now - LastWriteTime > 5)
+				{
+					LastWriteTime = Now;
+					UE_LOG(LogStats, Log, TEXT("WARNING: LLMMap average: %f\n"), (float)Average);
+				}
+			}
+#endif
+			return MapIndex;
+		}
+
+		// Increase the capacity of the map
+		void Grow()
+		{
+			SizeType NewCapacity = Capacity ? 2 * Capacity : DefaultCapacity / StripeCount;
+			Reserve(NewCapacity);
+		}
+
+		static bool InRange(
+			const SizeType Index,
+			const SizeType StartIndex,
+			const SizeType EndIndex)
+		{
+			return (StartIndex <= EndIndex) ?
+				Index >= StartIndex && Index <= EndIndex :
+				Index >= StartIndex || Index <= EndIndex;
+		}
+
+		mutable UE::FPlatformRecursiveMutex Mutex;
+
+		UE::LLMPrivate::FLLMAllocator* Allocator;
+
+		SizeType* Map;
+		SizeType Count;
+		SizeType Capacity;
+
+		// all these arrays must be kept in sync and are accessed by MapIndex
+		FLLMArray<TKey, SizeType> Keys;
+		FLLMArray<SizeType, SizeType> KeyHashes;
+		FLLMArray<TValue1, SizeType> Values1;
+		FLLMArray<TValue2, SizeType> Values2;
+
+		FLLMArray<SizeType, SizeType> FreeKeyIndices;
+
+#ifdef PROFILE_LLMMAP
+		mutable int64 IterAcc;
+		mutable int64 IterCount;
+#endif
+	};
+
+	// We divide map items into multiple stripes by the high bits of the hash key, to reduce lock contention.
+	// Each stripe has its own critical section lock, and given even distribution of hash codes (due to the
+	// mix function applied to the key), multiple threads will usually end up with their own lock, and not
+	// need to wait on each other.
+	StripeData MapStripes[StripeCount];
+	UE::FPlatformRecursiveMutex AllStripesMutex;
+	std::atomic_bool bAllStripesLocked;
+};
+
+// Pointer key for hash map
+struct PointerKey
+{
+	explicit PointerKey() = default;
+	explicit PointerKey(const void* InPointer, uint16 ExtraData=0)
+	: Pointer((UPTRINT(InPointer) << 16) | UPTRINT(ExtraData))
+	{
+		static_assert(sizeof(InPointer) == 8, "LLM only supports 64 bit platforms");
+		LLMEnsure(UPTRINT(InPointer) && UPTRINT(InPointer) <= 0x0000'ffff'ffff'ffff);
+	}
+	LLMNumAllocsType GetHashCode() const
+	{
+		return GetHashCodeImpl<sizeof(LLMNumAllocsType), sizeof(void*)>();
+	}
+	bool operator==(const PointerKey& other) const { return GetPointer() == other.GetPointer(); }
+	explicit operator bool () const { return !!Pointer; }
+	void* GetPointer() const { return (void*)(Pointer >> 16); };
+	uint64 GetExtraData() const { return uint64(Pointer & 0xffff); }
+
+private:
+	UPTRINT Pointer = 0;
+
+	template <int HashSize, int PointerSize>
+	LLMNumAllocsType GetHashCodeImpl() const
+	{
+		static_assert(HashSize == 0 && PointerSize == 0, "Converting void* to a LLMNumAllocsType - sized hashkey is not implemented for the current sizes.");
+		return (LLMNumAllocsType)(UPTRINT)GetPointer();
+	}
+
+	template <>
+	LLMNumAllocsType GetHashCodeImpl<8,8>() const
+	{
+		// 64 bit pointer to 64 bit hash
+		uint64 Key = (uint64)GetPointer();
+		Key = (~Key) + (Key << 21);
+		Key = Key ^ (Key >> 24);
+		Key = Key * 265;
+		Key = Key ^ (Key >> 14);
+		Key = Key * 21;
+		Key = Key ^ (Key >> 28);
+		Key = Key + (Key << 31);
+		return (LLMNumAllocsType)Key;
+	}
+
+	template <>
+	LLMNumAllocsType GetHashCodeImpl<4, 8>() const
+	{
+		// 64 bit pointer to 32 bit hash
+		uint64 Key = (uint64)GetPointer();
+		Key = (~Key) + (Key << 21);
+		Key = Key ^ (Key >> 24);
+		Key = Key * 265;
+		Key = Key ^ (Key >> 14);
+		Key = Key * 21;
+		Key = Key ^ (Key >> 28);
+		Key = Key + (Key << 31);
+		return (LLMNumAllocsType)Key;
+	}
+
+	template <>
+	LLMNumAllocsType GetHashCodeImpl<4, 4>() const
+	{
+		// 32 bit pointer to 32 bit Hash
+		uint64 Key = (uint64)GetPointer();
+		Key = (~Key) + (Key << 18);
+		Key = Key ^ (Key >> 31);
+		Key = Key * 21;
+		Key = Key ^ (Key >> 11);
+		Key = Key + (Key << 6);
+		Key = Key ^ (Key >> 22);
+		return (LLMNumAllocsType)Key;
+	}
+};
+
+namespace UE::Core::Private
+{
+	[[noreturn]] CORE_API void OnInvalidLLMAllocatorNum(int32 IndexSize, int64 NewNum, SIZE_T NumBytesPerElement);
+}
+
+/**
+ * An allocator usable in Core containers. It is based on TSizedHeapAllocator, but instead of allocating from FMemory it allocates from FLLMAllocator.
+ * Because FLLMAllocator::Free requires the Size, this allocator also has a Size field that TSizedHeapAllocator does not.
+ */
+template<int IndexSize>
+class TSizedLLMAllocator
+{
+public:
+	using SizeType = typename TBitsToSizeType<IndexSize>::Type;
+
+private:
+	using USizeType = std::make_unsigned_t<SizeType>;
+
+public:
+	enum { NeedsElementType = false };
+	enum { RequireRangeCheck = true };
+
+	class ForAnyElementType
+	{
+		template <int>
+		friend class TSizedLLMAllocator;
+
+	public:
+		/** Default constructor. */
+		ForAnyElementType()
+			: Data(nullptr)
+			, Size(0)
+		{}
+
+		/**
+		 * Moves the state of another allocator into this one.  The allocator can be different.
+		 *
+		 * Assumes that the allocator is currently empty, i.e. memory may be allocated but any existing elements have already been destructed (if necessary).
+		 * @param Other - The allocator to move the state from.  This allocator should be left in a valid empty state.
+		 */
+		template <typename OtherAllocator>
+		FORCEINLINE void MoveToEmptyFromOtherAllocator(typename OtherAllocator::ForAnyElementType& Other)
+		{
+			checkSlow((void*)this != (void*)&Other);
+
+			if (Data)
+			{
+				UE::LLMPrivate::FLLMAllocator::Get()->Free(Data, Size);
+			}
+
+			Data = Other.Data;
+			Size = Other.Size;
+			Other.Data = nullptr;
+			Other.Size = 0;
+		}
+
+		/**
+		 * Moves the state of another allocator into this one.
+		 * Moves the state of another allocator into this one.  The allocator can be different.
+		 *
+		 * Assumes that the allocator is currently empty, i.e. memory may be allocated but any existing elements have already been destructed (if necessary).
+		 * @param Other - The allocator to move the state from.  This allocator should be left in a valid empty state.
+		 */
+		FORCEINLINE void MoveToEmpty(ForAnyElementType& Other)
+		{
+			this->MoveToEmptyFromOtherAllocator<TSizedLLMAllocator>(Other);
+		}
+
+		/** Destructor. */
+		FORCEINLINE ~ForAnyElementType()
+		{
+			if (Data)
+			{
+				UE::LLMPrivate::FLLMAllocator::Get()->Free(Data, Size);
+			}
+		}
+
+		FORCEINLINE void* GetAllocation() const
+		{
+			return Data;
+		}
+
+		FORCEINLINE void ResizeAllocation(SizeType CurrentNum, SizeType NewMax, SIZE_T NumBytesPerElement)
+		{
+			// Avoid calling FMemory::Realloc( nullptr, 0 ) as ANSI C mandates returning a valid pointer which is not what we want.
+			if (Data || NewMax)
+			{
+				static_assert(sizeof(SizeType) <= sizeof(SIZE_T), "SIZE_T is expected to handle all possible sizes");
+
+				// Check for under/overflow
+				bool bInvalidResize = NewMax < 0 || NumBytesPerElement < 1 || NumBytesPerElement > (SIZE_T)MAX_int32;
+				if constexpr (sizeof(SizeType) == sizeof(SIZE_T))
+				{
+					bInvalidResize = bInvalidResize || (SIZE_T)(USizeType)NewMax > (SIZE_T)TNumericLimits<SizeType>::Max() / NumBytesPerElement;
+				}
+				if (UNLIKELY(bInvalidResize))
+				{
+					UE::Core::Private::OnInvalidLLMAllocatorNum(IndexSize, NewMax, NumBytesPerElement);
+				}
+
+				//checkSlow(((uint64)NewMax*(uint64)ElementTypeInfo.GetSize() < (uint64)INT_MAX));
+				size_t NewSize = NewMax * NumBytesPerElement;
+				Data = UE::LLMPrivate::FLLMAllocator::Get()->Realloc(Data, Size, NewSize);
+				Size = NewSize;
+			}
+		}
+		FORCEINLINE SizeType CalculateSlackReserve(SizeType NewMax, SIZE_T NumBytesPerElement) const
+		{
+			return DefaultCalculateSlackReserve(NewMax, NumBytesPerElement, true);
+		}
+		FORCEINLINE SizeType CalculateSlackShrink(SizeType NewMax, SizeType CurrentMax, SIZE_T NumBytesPerElement) const
+		{
+			return DefaultCalculateSlackShrink(NewMax, CurrentMax, NumBytesPerElement, true);
+		}
+		FORCEINLINE SizeType CalculateSlackGrow(SizeType NewMax, SizeType CurrentMax, SIZE_T NumBytesPerElement) const
+		{
+			return DefaultCalculateSlackGrow(NewMax, CurrentMax, NumBytesPerElement, true);
+		}
+
+		SIZE_T GetAllocatedSize(SizeType CurrentMax, SIZE_T NumBytesPerElement) const
+		{
+			return CurrentMax * NumBytesPerElement;
+		}
+
+		bool HasAllocation() const
+		{
+			return !!Data;
+		}
+
+		SizeType GetInitialCapacity() const
+		{
+			return 0;
+		}
+
+	private:
+		ForAnyElementType(const ForAnyElementType&);
+		ForAnyElementType& operator=(const ForAnyElementType&);
+
+		/** A pointer to the container's elements. */
+		void* Data;
+		/** The allocation size of Data */
+		size_t Size;
+	};
+
+	template<typename ElementType>
+	class ForElementType : public ForAnyElementType
+	{
+	public:
+
+		/** Default constructor. */
+		ForElementType()
+		{}
+
+		FORCEINLINE ElementType* GetAllocation() const
+		{
+			return (ElementType*)ForAnyElementType::GetAllocation();
+		}
+	};
+};
+
+// Define the ResizeAllocation functions with the regular allocators as exported to avoid bloat
+extern template CORE_API FORCENOINLINE void TSizedLLMAllocator<32>::ForAnyElementType::ResizeAllocation(SizeType CurrentNum, SizeType NewMax, SIZE_T NumBytesPerElement);
+extern template CORE_API FORCENOINLINE void TSizedLLMAllocator<64>::ForAnyElementType::ResizeAllocation(SizeType CurrentNum, SizeType NewMax, SIZE_T NumBytesPerElement);
+
+// The standard container-specific allocators based on TSizedLLMAllocator; these are copied from ContainerAllocationPolicies.h
+using FDefaultLLMAllocator = TSizedLLMAllocator<32>;
+using FDefaultLLMAllocator64 = TSizedLLMAllocator<64>;
+
+class FDefaultBitArrayLLMAllocator : public TInlineAllocator<4, FDefaultLLMAllocator> { public: typedef TInlineAllocator<4, FDefaultLLMAllocator>     Typedef; };
+class FDefaultSparseArrayLLMAllocator : public TSparseArrayAllocator<FDefaultLLMAllocator, FDefaultBitArrayLLMAllocator> { public: typedef TSparseArrayAllocator<FDefaultLLMAllocator, FDefaultBitArrayLLMAllocator> Typedef; };
+class FDefaultSetLLMAllocator : public TSetAllocator<FDefaultSparseArrayLLMAllocator, TInlineAllocator<1, FDefaultLLMAllocator>> { public: typedef TSetAllocator<FDefaultSparseArrayLLMAllocator, TInlineAllocator<1, FDefaultLLMAllocator>>         Typedef; };
+
+template<typename KeyType, typename ValueType>
+class TLLMMap : public TMap<KeyType, ValueType, FDefaultSetLLMAllocator>
+{};
+
+template<typename KeyType>
+class TLLMSet : public TSet<KeyType, DefaultKeyFuncs<KeyType>, FDefaultSetLLMAllocator>
+{};
+
+
+
+// Some algorithms used by LLM that require scratch-space internal allocation. To avoid polluting
+// the Algo namespace with an Allocator parameter, we've copied (or privately implemented the nonexisting ones) those
+// algorithms here
+namespace LLMAlgoImpl
+{
+	enum ETopologicalSortOrder
+	{
+		RootToLeaf,
+		LeafToRoot
+	};
+
+	template <typename T, typename GetEdgesType, typename SizeType>
+	void TopologicalSort(T* VertexData, SizeType NumVertices, GetEdgesType GetEdges, ETopologicalSortOrder SortOrder);
+}
+
+namespace LLMAlgo
+{
+	/**
+	 * Sort a range of vertices topologically from root to leaf by the directed edges given by GetEdges.
+	 * Vertices in cycles will be sorted in an arbitrary order relative to each other.
+	 * The sort is stable.
+	 *
+	 * @param Range			The range of vertices to sort.
+	 * @param GetEdges		A callable GetEdges(SizeType InVertex, SizeType* OutEdgeBuffer, Sizetype& OutNumEdges) that provides the indexes into VertexData of the edges for InVertex.
+	 *						Edges should be written into OutEdgeBuffer; this buffer only has room for GetNum(Range) elements.
+	 *						InVertex is the index in Range of the vertex.
+	 *						OutEdgeBuffer should be populated with the indexes in Range of the target vertices.
+	 */
+	template <typename RangeType, typename GetEdgesType>
+	void TopologicalSortRootToLeaf(RangeType&& Range, GetEdgesType GetEdges)
+	{
+		TopologicalSort(GetData(Range), GetNum(Range), GetEdges, LLMAlgoImpl::ETopologicalSortOrder::RootToLeaf);
+	}
+
+	/**
+	 * Sort a range of vertices topologically from leaf to root by the directed edges given by GetEdges.
+	 * Vertices in cycles will be sorted in an arbitrary order relative to each other.
+	 * The sort is stable.
+	 *
+	 * @param Range			The range of vertices to sort.
+	 * @param GetEdges		A callable GetEdges(SizeType InVertex, SizeType* OutEdgeBuffer, Sizetype& OutNumEdges) that provides the indexes into VertexData of the edges for InVertex.
+	 *						Edges should be written into OutEdgeBuffer; this buffer only has room for GetNum(Range) elements.
+	 *						InVertex is the index in Range of the vertex.
+	 *						OutEdgeBuffer should be populated with the indexes in Range of the target vertices.
+	 */
+	template <typename RangeType, typename GetEdgesType>
+	void TopologicalSortLeafToRoot(RangeType&& Range, GetEdgesType GetEdges)
+	{
+		using namespace LLMAlgoImpl;
+		TopologicalSort(GetData(Range), GetNum(Range), GetEdges, ETopologicalSortOrder::LeafToRoot);
+	}
+}
+
+namespace LLMAlgoImpl
+{
+
+	/**
+	 * Sort an array of vertices topologically by the directed edges given by GetEdges.
+	 * Vertices in cycles will be sorted in an arbitrary order relative to each other.
+	 * The sort is stable.
+	 *
+	 * @param VertexData	Pointer to the vertices to sort.
+	 * @param NumVertices	Number of vertices.
+	 * @param GetEdges		A callable GetEdges(SizeType InVertex, SizeType* OutEdgeBuffer, Sizetype& OutNumEdges) that provides the indexes into VertexData of the edges for InVertex.
+	 *						Edges should be written into OutEdgeBuffer; this buffer only has room for NumVertices elements.
+	 *						InVertex is the index in Range of the vertex.
+	 *						OutEdgeBuffer should be populated with the indexes in Range of the target vertices.
+	 * @param SortOrder		Whether to order the in/out VertexData by sources before targets (RootToLeaf) or targets before sources (LeafToRoot).
+	 */
+	template <typename T, typename GetEdgesType, typename SizeType>
+	void TopologicalSort(T* VertexData, SizeType NumVertices, GetEdgesType GetEdges, ETopologicalSortOrder SortOrder)
+	{
+		if (NumVertices == 0)
+		{
+			return;
+		}
+		LLMCheckf(NumVertices > 0, TEXT("Invalid number of vertices"));
+
+		// In our traversal, we write vertices LeafToRoot. To make a stable sort, we need to iterate the input list from LeafToRoot as well.
+		SizeType TraversalStart, TraversalEnd, TraversalDir;
+		if (SortOrder == ETopologicalSortOrder::RootToLeaf)
+		{
+			TraversalStart = NumVertices - 1;
+			TraversalEnd = -1;
+			TraversalDir = -1;
+		}
+		else
+		{
+			TraversalStart = 0;
+			TraversalEnd = NumVertices;
+			TraversalDir = 1;
+		}
+
+		TArray<SizeType, FDefaultLLMAllocator> LeafToRootOrder;
+		LeafToRootOrder.Reserve(NumVertices);
+
+		{
+			constexpr SizeType ExpectedMaxNumEdges = 16;
+			struct FVisitData
+			{
+				SizeType Vertex;
+				SizeType NextEdge;
+				SizeType EdgeStart;
+			};
+
+			TArray<FVisitData, FDefaultLLMAllocator> Stack;
+			Stack.Reserve(NumVertices);
+
+			TArray<SizeType, FDefaultLLMAllocator> EdgeBuffer;
+			EdgeBuffer.AddUninitialized(NumVertices);
+			SizeType* EdgeBufferData = EdgeBuffer.GetData();
+
+			TArray<SizeType, FDefaultLLMAllocator> EdgesOnStack;
+			EdgesOnStack.Reserve(NumVertices * ExpectedMaxNumEdges);
+
+			TArray<bool, FDefaultLLMAllocator> Visited;
+			Visited.AddDefaulted(NumVertices);
+
+			auto PushVertex = [&EdgesOnStack, &Stack, &GetEdges, &Visited, EdgeBufferData, NumVertices](SizeType Vertex)
+			{
+				Visited[Vertex] = true;
+				Stack.Add(FVisitData{ Vertex, EdgesOnStack.Num(), EdgesOnStack.Num() });
+				SizeType NumEdges = static_cast<SizeType>(-1);
+				GetEdges(Vertex, EdgeBufferData, NumEdges);
+				LLMCheckf(0 <= NumEdges && NumEdges <= NumVertices, TEXT("GetEdges function passed into TopologicalSort did not write OutNumEdges, or wrote an invalid value"));
+				EdgesOnStack.Append(EdgeBufferData, NumEdges);
+			};
+
+			for (SizeType RootVertex = TraversalStart; RootVertex != TraversalEnd; RootVertex += TraversalDir)
+			{
+				if (Visited[RootVertex])
+				{
+					continue;
+				}
+
+				PushVertex(RootVertex);
+				while (Stack.Num() > 0)
+				{
+					FVisitData& VisitData = Stack.Last();
+					bool bPushed = false;
+					SizeType NumEdgesOnStack = EdgesOnStack.Num();
+					for (; VisitData.NextEdge < NumEdgesOnStack; ++VisitData.NextEdge)
+					{
+						SizeType TargetVertex = EdgesOnStack[VisitData.NextEdge];
+						if (!Visited[TargetVertex])
+						{
+							++VisitData.NextEdge;
+							PushVertex(TargetVertex);
+							bPushed = true;
+							break;
+						}
+					}
+					if (!bPushed)
+					{
+						LeafToRootOrder.Add(VisitData.Vertex);
+						EdgesOnStack.SetNum(VisitData.EdgeStart, EAllowShrinking::No);
+						Stack.Pop(EAllowShrinking::No);
+					}
+				}
+			}
+		}
+		LLMCheck(LeafToRootOrder.Num() == NumVertices); // This could only fail due to an internal logic error; all vertices should have been visited and added
+
+		TArray<T, FDefaultLLMAllocator> Original;
+		Original.Reserve(NumVertices);
+		for (SizeType n = 0; n < NumVertices; ++n)
+		{
+			Original.Add(MoveTemp(VertexData[n]));
+		}
+
+		for (SizeType LeafToRootIndex = 0; LeafToRootIndex < NumVertices; ++LeafToRootIndex)
+		{
+			SizeType WriteIndex = TraversalStart + TraversalDir * LeafToRootIndex;
+			SizeType ReadIndex = LeafToRootOrder[LeafToRootIndex];
+			VertexData[WriteIndex] = MoveTemp(Original[ReadIndex]);
+		}
+	}
+}
+
+#endif		// #if ENABLE_LOW_LEVEL_MEM_TRACKER
+
diff --git a/Engine/Source/Runtime/Core/Private/HAL/LowLevelMemTracker.cpp b/Engine/Source/Runtime/Core/Private/HAL/LowLevelMemTracker.cpp
deleted file mode 100644
diff --git a/Engine/Source/Runtime/Core/Private/HAL/LowLevelMemTrackerPrivate.h b/Engine/Source/Runtime/Core/Private/HAL/LowLevelMemTrackerPrivate.h
deleted file mode 100644
diff --git a/Engine/Source/Runtime/Core/Private/HAL/LowLevelMemTrackerTest.cpp b/Engine/Source/Runtime/Core/Private/HAL/LowLevelMemTrackerTest.cpp
deleted file mode 100644
diff --git a/Engine/Source/Runtime/Core/Private/HAL/LowLevelMemoryUtils.cpp b/Engine/Source/Runtime/Core/Private/HAL/LowLevelMemoryUtils.cpp
deleted file mode 100644
diff --git a/Engine/Source/Runtime/Core/Private/HAL/LowLevelMemoryUtils.h b/Engine/Source/Runtime/Core/Private/HAL/LowLevelMemoryUtils.h
deleted file mode 100644
diff --git a/Engine/Source/Runtime/Core/Private/ProfilingDebugging/TagTrace.cpp b/Engine/Source/Runtime/Core/Private/ProfilingDebugging/TagTrace.cpp
--- a/Engine/Source/Runtime/Core/Private/ProfilingDebugging/TagTrace.cpp
+++ b/Engine/Source/Runtime/Core/Private/ProfilingDebugging/TagTrace.cpp
@@ -5,8 +5,8 @@
 #include "Containers/StringView.h"
 #include "CoreTypes.h"
 #include "ProfilingDebugging/MemoryTrace.h"
+#include "HAL/LLM/LLMPrivate.h"
 #include "HAL/LowLevelMemTracker.h"
-#include "HAL/LowLevelMemTrackerPrivate.h"
 #include "HAL/PlatformString.h"
 #include "Misc/Optional.h"
 #include "Misc/StringBuilder.h"
