LLM: Whitespace change: types within a namespace should be unindented in core code.
#rnx

#Synced-CL 49633983

diff --git a/Engine/Source/Runtime/Core/Private/HAL/LLM/LLMPrivate.h b/Engine/Source/Runtime/Core/Private/HAL/LLM/LLMPrivate.h
--- a/Engine/Source/Runtime/Core/Private/HAL/LLM/LLMPrivate.h
+++ b/Engine/Source/Runtime/Core/Private/HAL/LLM/LLMPrivate.h
@@ -44,131 +44,130 @@
 // Disable if you need a little more memory or speed
 #define LLM_ENABLED_TRACK_PEAK_MEMORY 1
 
-namespace UE
+namespace UE::LLMPrivate
 {
-namespace LLMPrivate
+
+struct FTagDataNameKey
 {
-	struct FTagDataNameKey
-	{
-		FName Name;
-		ELLMTagSet TagSet;
+	FName Name;
+	ELLMTagSet TagSet;
 		
-		FTagDataNameKey(FName InName, ELLMTagSet InTagSet) :
-			Name(InName),
-			TagSet(InTagSet)
-		{
-		}
+	FTagDataNameKey(FName InName, ELLMTagSet InTagSet) :
+		Name(InName),
+		TagSet(InTagSet)
+	{
+	}
 
-		friend uint32 GetTypeHash(const FTagDataNameKey& Key)
-		{
-			constexpr uint32 HashPrime = 101;
-			return GetTypeHash(Key.Name) + HashPrime * static_cast<uint32>(Key.TagSet);
-		}
+	friend uint32 GetTypeHash(const FTagDataNameKey& Key)
+	{
+		constexpr uint32 HashPrime = 101;
+		return GetTypeHash(Key.Name) + HashPrime * static_cast<uint32>(Key.TagSet);
+	}
 
-		friend bool operator==(const FTagDataNameKey& A, const FTagDataNameKey& B)
-		{
-			return A.Name == B.Name && A.TagSet == B.TagSet;
-		}
+	friend bool operator==(const FTagDataNameKey& A, const FTagDataNameKey& B)
+	{
+		return A.Name == B.Name && A.TagSet == B.TagSet;
+	}
 
-		friend bool operator!=(const FTagDataNameKey& A, const FTagDataNameKey& B)
-		{
-			return (A.Name != B.Name || A.TagSet != B.TagSet);
-		}
+	friend bool operator!=(const FTagDataNameKey& A, const FTagDataNameKey& B)
+	{
+		return (A.Name != B.Name || A.TagSet != B.TagSet);
+	}
 
-		FTagDataNameKey() = delete;
-	};
+	FTagDataNameKey() = delete;
+};
 
-	/**
-	 * FTagData: Description of the properties of a Tag that can be used in LLM_SCOPE
-	 */
-	class FTagData
-	{
-	public:
-
-		FTagData(FName InName, ELLMTagSet InTagSet, FName InDisplayName, FName InParentName, FName InStatName,
-			FName InSummaryStatName, bool bInHasEnumTag, ELLMTag InEnumTag, ETagReferenceSource InReferenceSource);
-		FTagData(FName InName, ELLMTagSet InTagSet, FName InDisplayName, const FTagData* InParent, FName InStatName,
-			FName InSummaryStatName, bool bInHasEnumTag, ELLMTag InEnumTag, ETagReferenceSource InReferenceSource);
-		~FTagData();
-
-		bool IsParentConstructed() const;
-		bool IsFinishConstructed() const;
-		FName GetName() const;
-		FName GetDisplayName() const;
-		void GetDisplayPath(FStringBuilderBase& Result, int32 MaxLen=-1) const;
-		void AppendDisplayPath(FStringBuilderBase& Result, int32 MaxLen=-1) const;
-		const FTagData* GetParent() const;
-		FName GetParentName() const;
-		FName GetParentNameSafeBeforeFinishConstruct() const;
-		FName GetStatName() const;
-		FName GetSummaryStatName() const;
-		ELLMTag GetEnumTag() const;
-		ELLMTagSet GetTagSet() const;
-		bool HasEnumTag() const;
-		const FTagData* GetContainingEnumTagData() const;
-		ELLMTag GetContainingEnum() const;
-		ETagReferenceSource GetReferenceSource() const;
-		int32 GetIndex() const;
-		bool IsReportable() const;
-		bool IsStatsReportable() const;
-		bool IsTraceReportable() const;
-
-		void SetParent(const FTagData* InParent);
-		void SetIndex(int32 InIndex);
-		void SetIsReportable(bool bInReportable);
+/**
+ * FTagData: Description of the properties of a Tag that can be used in LLM_SCOPE
+ */
+class FTagData
+{
+public:
+
+	FTagData(FName InName, ELLMTagSet InTagSet, FName InDisplayName, FName InParentName, FName InStatName,
+		FName InSummaryStatName, bool bInHasEnumTag, ELLMTag InEnumTag, ETagReferenceSource InReferenceSource);
+	FTagData(FName InName, ELLMTagSet InTagSet, FName InDisplayName, const FTagData* InParent, FName InStatName,
+		FName InSummaryStatName, bool bInHasEnumTag, ELLMTag InEnumTag, ETagReferenceSource InReferenceSource);
+	~FTagData();
+
+	bool IsParentConstructed() const;
+	bool IsFinishConstructed() const;
+	FName GetName() const;
+	FName GetDisplayName() const;
+	void GetDisplayPath(FStringBuilderBase& Result, int32 MaxLen=-1) const;
+	void AppendDisplayPath(FStringBuilderBase& Result, int32 MaxLen=-1) const;
+	const FTagData* GetParent() const;
+	FName GetParentName() const;
+	FName GetParentNameSafeBeforeFinishConstruct() const;
+	FName GetStatName() const;
+	FName GetSummaryStatName() const;
+	ELLMTag GetEnumTag() const;
+	ELLMTagSet GetTagSet() const;
+	bool HasEnumTag() const;
+	const FTagData* GetContainingEnumTagData() const;
+	ELLMTag GetContainingEnum() const;
+	ETagReferenceSource GetReferenceSource() const;
+	int32 GetIndex() const;
+	bool IsReportable() const;
+	bool IsStatsReportable() const;
+	bool IsTraceReportable() const;
+
+	void SetParent(const FTagData* InParent);
+	void SetIndex(int32 InIndex);
+	void SetIsReportable(bool bInReportable);
 		void SetIsTraceReportable(bool bInTraceReportable);
-		void SetFinishConstructed();
-
-		// These functions are normally invalid - these properties should be immutable - but are called for EnumTags during bootstrapping
-		FTagData(ELLMTag InEnumTag);
-		void SetName(FName InName);
-		void SetDisplayName(FName InDisplayName);
-		void SetStatName(FName InStatName);
-		void SetSummaryStatName(FName InSummaryStatName);
-		void SetParentName(FName InParentName);
-
-	private:
-		bool IsUsedAsDisplayParent() const;
-
-		FName Name;
-		FName DisplayName;
-		union
-		{
-			const FTagData* Parent;
-			FName ParentName;
-		};
-		FName StatName;
-		FName SummaryStatName;
-		int32 Index;
-		ELLMTag EnumTag;
-		ETagReferenceSource ReferenceSource;
-		ELLMTagSet TagSet;
-		bool bIsFinishConstructed;
-		bool bParentIsName;
-		bool bHasEnumTag;
-		bool bIsReportable;
-		bool bIsTraceReportable;
-	};
-
-	// TagData container types that use FLLMAllocator
-	class FTagDataNameMap : public TMap<FTagDataNameKey, FTagData*, FDefaultSetLLMAllocator>
-	{
-		using TMap<FTagDataNameKey, FTagData*, FDefaultSetLLMAllocator>::TMap;
-	};
-	class FConstTagDataArray : public TArray<const FTagData*, FDefaultLLMAllocator>
+	void SetFinishConstructed();
+
+	// These functions are normally invalid - these properties should be immutable - but are called for EnumTags during bootstrapping
+	FTagData(ELLMTag InEnumTag);
+	void SetName(FName InName);
+	void SetDisplayName(FName InDisplayName);
+	void SetStatName(FName InStatName);
+	void SetSummaryStatName(FName InSummaryStatName);
+	void SetParentName(FName InParentName);
+
+private:
+	bool IsUsedAsDisplayParent() const;
+
+	FName Name;
+	FName DisplayName;
+	union
 	{
-		using TArray<const FTagData*, FDefaultLLMAllocator>::TArray;
-	};
-	class FTagDataArray : public TArray<FTagData*, FDefaultLLMAllocator>
-	{
-		using TArray<FTagData*, FDefaultLLMAllocator>::TArray;
+		const FTagData* Parent;
+		FName ParentName;
 	};
+	FName StatName;
+	FName SummaryStatName;
+	int32 Index;
+	ELLMTag EnumTag;
+	ETagReferenceSource ReferenceSource;
+	ELLMTagSet TagSet;
+	bool bIsFinishConstructed;
+	bool bParentIsName;
+	bool bHasEnumTag;
+	bool bIsReportable;
+	bool bIsTraceReportable;
+};
+
+// TagData container types that use FLLMAllocator
+class FTagDataNameMap : public TMap<FTagDataNameKey, FTagData*, FDefaultSetLLMAllocator>
+{
+	using TMap<FTagDataNameKey, FTagData*, FDefaultSetLLMAllocator>::TMap;
+};
+class FConstTagDataArray : public TArray<const FTagData*, FDefaultLLMAllocator>
+{
+	using TArray<const FTagData*, FDefaultLLMAllocator>::TArray;
+};
+class FTagDataArray : public TArray<FTagData*, FDefaultLLMAllocator>
+{
+	using TArray<FTagData*, FDefaultLLMAllocator>::TArray;
+};
 
-	class FTagFilter : public TLLMSet<FName>
-	{
-	};
-}
-}
+class FTagFilter : public TLLMSet<FName>
+{
+};
+
+} // namespace UE::LLMPrivate
 
 // Interface for Algo::TopologicalSort for FTagDataArray
 inline UE::LLMPrivate::FTagData** GetData(UE::LLMPrivate::FTagDataArray& Array)
@@ -180,66 +179,64 @@ inline UE::LLMPrivate::FTagDataArray::SizeType GetNum(UE::LLMPrivate::FTagDataAr
 	return Array.Num();
 }
 
-namespace UE
+namespace UE::LLMPrivate
 {
-namespace LLMPrivate
+
+/** Size information stored on the tracker for a tag; includes amounts aggregated from threadstates and from external api users */
+struct FTrackerTagSizeData
 {
-	/** Size information stored on the tracker for a tag; includes amounts aggregated from threadstates and from external api users */
-	struct FTrackerTagSizeData
-	{
-		int64 Size = 0;
+	int64 Size = 0;
 #if LLM_ENABLED_TRACK_PEAK_MEMORY
-		int64 PeakSize = 0;
+	int64 PeakSize = 0;
 #endif
-		int64 SizeInSnapshot = 0;
-		int64 ExternalAmount = 0;
-		bool bExternalValid = false;
-		bool bExternalAddToTotal = false;
+	int64 SizeInSnapshot = 0;
+	int64 ExternalAmount = 0;
+	bool bExternalValid = false;
+	bool bExternalAddToTotal = false;
 
-		int64 GetSize(UE::LLM::ESizeParams SizeParams) const
-		{
-			int64 CurrentSize = Size;
+	int64 GetSize(UE::LLM::ESizeParams SizeParams) const
+	{
+		int64 CurrentSize = Size;
 
 #if LLM_ENABLED_TRACK_PEAK_MEMORY
-			if (EnumHasAnyFlags(SizeParams, UE::LLM::ESizeParams::ReportPeak))
-			{
-				CurrentSize = PeakSize;
-			}
-#endif
-			// Note, this will also subtract the snapshotted size from PeakSize if that flag is enabled
-			if (EnumHasAnyFlags(SizeParams, UE::LLM::ESizeParams::RelativeToSnapshot))
-			{
-				CurrentSize = FMath::Clamp<int64>(CurrentSize - SizeInSnapshot, 0, INT64_MAX); //-V569
-			}
-
-			return CurrentSize;
-		};
-
-		void CaptureSnapshot()
+		if (EnumHasAnyFlags(SizeParams, UE::LLM::ESizeParams::ReportPeak))
 		{
-			SizeInSnapshot = Size;
+			CurrentSize = PeakSize;
 		}
-
-		void ClearSnapshot()
+#endif
+		// Note, this will also subtract the snapshotted size from PeakSize if that flag is enabled
+		if (EnumHasAnyFlags(SizeParams, UE::LLM::ESizeParams::RelativeToSnapshot))
 		{
-			SizeInSnapshot = 0;
+			CurrentSize = FMath::Clamp<int64>(CurrentSize - SizeInSnapshot, 0, INT64_MAX); //-V569
 		}
+
+		return CurrentSize;
 	};
-	typedef TLLMMap<const FTagData*, FTrackerTagSizeData> FTrackerTagSizeMap;
 
-	/**
-	 * Size information stored on each threadstate for a tag
-	 * TagSizes are sorted by Index instead of by pointer in the ThreadTagSizeMap to enforce the constraint that Parents come before children
-	 */
-	struct FThreadTagSizeData
+	void CaptureSnapshot()
 	{
-		const FTagData* TagData = nullptr;
-		int64 Size = 0;
-	};
-	// 
-	typedef TSortedMap<int32, FThreadTagSizeData, FDefaultLLMAllocator> FThreadTagSizeMap;
-}
-}
+		SizeInSnapshot = Size;
+	}
+
+	void ClearSnapshot()
+	{
+		SizeInSnapshot = 0;
+	}
+};
+typedef TLLMMap<const FTagData*, FTrackerTagSizeData> FTrackerTagSizeMap;
+
+/**
+	* Size information stored on each threadstate for a tag
+	* TagSizes are sorted by Index instead of by pointer in the ThreadTagSizeMap to enforce the constraint that Parents come before children
+	*/
+struct FThreadTagSizeData
+{
+	const FTagData* TagData = nullptr;
+	int64 Size = 0;
+};
+typedef TSortedMap<int32, FThreadTagSizeData, FDefaultLLMAllocator> FThreadTagSizeMap;
+
+} // namespace UE::LLMPrivate
 
 FName LLMGetTagUniqueName(ELLMTag Tag);
 
