LLM: Move definitions of LLM's private classes from .cpp into LLMPrivate.h so we can reference them from multiple LLM cpp files.
#rnx
#Synced-CL 49633983

diff --git a/Engine/Source/Runtime/Core/Private/HAL/LLM/LLM.cpp b/Engine/Source/Runtime/Core/Private/HAL/LLM/LLM.cpp
--- a/Engine/Source/Runtime/Core/Private/HAL/LLM/LLM.cpp
+++ b/Engine/Source/Runtime/Core/Private/HAL/LLM/LLM.cpp
@@ -26,20 +26,9 @@
 #include "Misc/ScopeRWLock.h"
 #include "Misc/StringBuilder.h"
 #include "Misc/VarArgs.h"
-#include "ProfilingDebugging/CsvProfiler.h"
 #include "Templates/Atomic.h"
 #include "Trace/Trace.inl"
 
-static_assert(static_cast<uint32>(ELLMTagSet::Max) == 4, "You added a new TagSet, please update the #ifs bellow");
-#if LLM_ALLOW_ASSETS_TAGS
-#define LLM_TAGSET_MAX 3
-#else
-#define LLM_TAGSET_MAX 2
-#endif
-
-// Enable to allow memory tracking to be passed through to external tools without the extra overhead of LLM also tracking allocations.
-#define UE_ONLY_USE_PLATFORM_TRACKER 0
-
 #if UE_ENABLE_ARRAY_SLACK_TRACKING
 
 // Specifies whether to generate the whole log file in memory before writing.  Switch uses an async thread for file writing,
@@ -954,8 +943,6 @@ UE_TRACE_EVENT_BEGIN(LLM, TagValue)
 	UE_TRACE_EVENT_FIELD(int64[], Values)
 UE_TRACE_EVENT_END()
 
-#define LLM_CSV_PROFILER_WRITER_ENABLED CSV_PROFILER_STATS
-
 #if LLM_CSV_PROFILER_WRITER_ENABLED
 	CSV_DEFINE_CATEGORY(LLM, true);
 	CSV_DEFINE_CATEGORY(LLMPlatform, true);
@@ -1307,376 +1294,6 @@ DECLARE_LLM_MEMORY_STAT(TEXT("CustomName"), STAT_CustomName, STATGROUP_LLMFULL);
 namespace UE::LLMPrivate
 {
 
-/** FLLMCsvWriter: class for writing out the LLM tag sizes to a csv file every few seconds. */
-class FLLMCsvWriter
-{
-public:
-	FLLMCsvWriter();
-	~FLLMCsvWriter();
-
-	void SetTracker(ELLMTracker InTracker);
-	void Clear();
-
-	void Publish(FLowLevelMemTracker& LLMRef, const FTrackerTagSizeMap& TagSizes,
-		const FTagData* OverrideTrackedTotalTagData, const FTagData* OverrideUntaggedTagData, int64 TrackedTotal,
-		UE::LLM::ESizeParams SizeParams);
-
-	void OnPreFork();
-
-	void Flush(bool IsOnCrash);
-
-private:
-	void Write(FStringView Text);
-	static const TCHAR* GetTrackerCsvName(ELLMTracker InTracker);
-
-	bool CreateArchive();
-	bool UpdateColumns(const FTrackerTagSizeMap& TagSizes);
-	void WriteHeader(const FTagData* OverrideTrackedTotalTagData, const FTagData* OverrideUntaggedTagData);
-	void AddRow(FLowLevelMemTracker& LLMRef, const FTrackerTagSizeMap& TagSizes,
-		const FTagData* OverrideTrackedTotalName, const FTagData* OverrideUntaggedName, int64 TrackedTotal,
-		UE::LLM::ESizeParams SizeParams);
-
-	FConstTagDataArray Columns;
-	TLLMSet<const FTagData*> ExistingColumns;
-	FArchive* Archive;
-	double LastWriteTime;
-	int32 WriteCount;
-	ELLMTracker Tracker;
-	bool bRegisteredFlushDelegate;
-};
-
-/** Outputs the LLM tags and sizes to TraceLog events. */
-class FLLMTraceWriter
-{
-public:
-	FLLMTraceWriter();
-	void SetTracker(ELLMTracker InTracker);
-	void Clear();
-	void Publish(FLowLevelMemTracker& LLMRef, const FTrackerTagSizeMap& TagSizes,
-		const FTagData* OverrideTrackedTotalName, const FTagData* OverrideUntaggedName, int64 TrackedTotal,
-		UE::LLM::ESizeParams SizeParams);
-
-private:
-	static const void* GetTagId(const FTagData* TagData);
-	void SendTagDeclaration(const FTagData* TagData);
-
-	ELLMTracker Tracker;
-	TLLMSet<const FTagData*> DeclaredTags;
-	bool bTrackerSpecSent = false;
-	bool bTagSetSpecSent = false;
-};
-
-/** FLLMCsvWriter: class for writing out LLM stats to the Csv Profiler. */
-class FLLMCsvProfilerWriter
-{
-public:
-	FLLMCsvProfilerWriter();
-	void Clear();
-	void SetTracker(ELLMTracker InTracker);
-	void Publish(FLowLevelMemTracker& LLMRef, const FTrackerTagSizeMap& TagSizes,
-		const FTagData* OverrideTrackedTotalName, const FTagData* OverrideUntaggedName, int64 TrackedTotal,
-		UE::LLM::ESizeParams SizeParams);
-protected:
-	void RecordTagToCsv(int32 CsvCategoryIndex, const FTagData* TagData, int64 TagSize);
-#if LLM_CSV_PROFILER_WRITER_ENABLED
-	ELLMTracker Tracker;
-	TLLMMap<const FTagData*, FName> TagDataToCsvStatName;
-#endif
-};
-
-/** Per-thread state in an LLMTracker. */
-class FLLMThreadState
-{
-public:
-	FLLMThreadState();
-
-	void Clear();
-
-	void PushTag(const FTagData* TagData, ELLMTagSet TagSet);
-	void PopTag(ELLMTagSet TagSet);
-	const FTagData* GetTopTag(ELLMTagSet TagSet);
-	void TrackAllocation(const void* Ptr, int64 Size, ELLMTracker Tracker, ELLMAllocType AllocType,
-		const FTagData* TagData, const FTagData* AssetTagData, const FTagData* AssetClassTagData, const FTagData* UObjectClassesTagData, bool bTrackInMemPro);
-	void TrackFree(const void* Ptr, int64 Size, ELLMTracker Tracker, ELLMAllocType AllocType,
-		const FTagData* TagData, const FTagData* AssetTagData, const FTagData* AssetClassTagData, const FTagData* UObjectClassesTagData, bool bTrackInMemPro);
-	void TrackMemory(int64 Amount, ELLMTracker Tracker, ELLMAllocType AllocType,
-		const FTagData* TagData);
-	void TrackMoved(const void* Dest, const void* Source, int64 Size, ELLMTracker Tracker, const FTagData* TagData);
-	void IncrTag(const FTagData* Tag, int64 Amount);
-
-	void PropagateChildSizesToParents();
-	void OnTagsResorted(FTagDataArray& OldTagDatas);
-	void LockTags(bool bLock);
-	void FetchAndClearTagSizes(FTrackerTagSizeMap& TagSizes, int64* InAllocTypeAmounts, bool bTrimAllocations);
-
-	void ClearAllocTypeAmounts();
-
-	// AllocationsLock must be held around read/write of Allocations and AllocTypeAmounts, because these are read/write
-	// from the Update thread. All other data on FLLMThreadState are accessed only from the thread that owns the
-	// FLLMThreadState and do not require a lock.
-	// AllocationsLock is also used to guarantee our contract provided to FLLMTracker::LockAllThreadTags: We must not
-	// dereference any FTagData* held by the FLLMThreadState unless we have entered AllocationsLock.
-	FConstTagDataArray TagStack[static_cast<int32>(ELLMTagSet::Max)];
-	FThreadTagSizeMap Allocations;
-	FCriticalSection AllocationsLock;
-
-	int8 PausedCounter[(int32)ELLMAllocType::Count];
-	int64 AllocTypeAmounts[(int32)ELLMAllocType::Count];
-};
-
-/** The main LLM implementation class. It owns the thread state objects. */
-class FLLMTracker
-{
-public:
-
-	FLLMTracker(FLowLevelMemTracker& InLLM);
-	~FLLMTracker();
-
-	void Initialise(ELLMTracker InTracker, FLLMAllocator* InAllocator);
-
-	void PushTag(ELLMTag EnumTag, ELLMTagSet TagSet);
-	void PushTag(FName Tag, bool bInIsStatTag, ELLMTagSet TagSet);
-	void PushTag(const FTagData* TagData, ELLMTagSet TagSet);
-	void PopTag(ELLMTagSet TagSet);
-	void TrackAllocation(const void* Ptr, int64 Size, FName DefaultTag, ELLMAllocType AllocType, bool bTrackInMemPro);
-	void TrackAllocation(const void* Ptr, int64 Size, ELLMTag DefaultTag, ELLMAllocType AllocType, bool bTrackInMemPro);
-	void TrackFree(const void* Ptr, ELLMAllocType AllocType, bool bTrackInMemPro);
-	void TrackMemoryOfActiveTag(int64 Amount, FName DefaultTag, ELLMAllocType AllocType);
-	void TrackMemoryOfActiveTag(int64 Amount, ELLMTag DefaultTag, ELLMAllocType AllocType);
-	void OnAllocMoved(const void* Dest, const void* Source, ELLMAllocType AllocType);
-
-	void TrackMemory(ELLMTag EnumTag, int64 Amount, ELLMAllocType AllocType);
-	void TrackMemory(FName TagName, ELLMTagSet TagSet, int64 Amount, ELLMAllocType AllocType);
-	void TrackMemory(const FTagData* TagData, int64 Amount, ELLMAllocType AllocType);
-
-	// This will pause/unpause tracking, and also manually increment a given tag
-	void PauseAndTrackMemory(FName TagName, ELLMTagSet TagSet, bool bInIsStatTag, int64 Amount, ELLMAllocType AllocType);
-	void PauseAndTrackMemory(ELLMTag EnumTag, int64 Amount, ELLMAllocType AllocType);
-	void PauseAndTrackMemory(const FTagData* TagData, int64 Amount, ELLMAllocType AllocType);
-	void Pause(ELLMAllocType AllocType);
-	void Unpause(ELLMAllocType AllocType);
-	bool IsPaused(ELLMAllocType AllocType);
-
-	void Clear();
-
-	// Dump the allocation count and size for each tag, along with the proportion that is private/shared/unreferenced
-	// on linux machines. This will dump to a CSV in the project Saved/LLM directory with a separate file for each
-	// tracker and forked child. Integers will be appended to the filename to prevent overwrites. Returns false
-	// if the memory information can't be retrieved (likely because it's not on a supported platform)
-	bool DumpForkedAllocationInfo();
-
-	// Dump the list of all registered LLM tags and their properties to a CSV file, in the LLM's profiling directory
-	// (Saved/Profiling/LLM/); see UE::LLM::GetLLMProfilingDir().
-	bool DumpTags();
-
-	void PublishStats(UE::LLM::ESizeParams SizeParam);
-	void PublishCsv(UE::LLM::ESizeParams SizeParam);
-	void PublishTrace(UE::LLM::ESizeParams SizeParam);
-	void PublishCsvProfiler(UE::LLM::ESizeParams SizeParam);
-
-	void OnPreFork();
-
-	struct FLowLevelAllocInfo
-	{
-	public:
-		FLowLevelAllocInfo()
-		{
-#if LLM_ENABLED_FULL_TAGS && (LLM_ALLOW_ASSETS_TAGS || LLM_ALLOW_UOBJECTCLASSES_TAGS)
-			FMemory::Memset(Tag, InvalidCompressedTagValue);
-#endif
-		}
-
-		FORCEINLINE static int32 GetIndexFromTagSet(ELLMTagSet InTagSet)
-		{
-			static_assert(static_cast<int32>(ELLMTagSet::Max) == 4,
-				"You added a new tagset, please update the code bellow");
-
-#if LLM_ALLOW_ASSETS_TAGS
-			switch (InTagSet)
-			{
-				case ELLMTagSet::None: return 0;
-				case ELLMTagSet::Assets: return 1;
-				case ELLMTagSet::AssetClasses: return 2;
-				default:
-					LLMCheck(false);
-					return 0;
-			}
-#elif LLM_ALLOW_UOBJECTCLASSES_TAGS
-			switch (InTagSet)
-			{
-				case ELLMTagSet::None: return 0;
-				case ELLMTagSet::UObjectClasses: return 1;
-				default:
-					LLMCheck(false);
-					return 0;
-			}
-#endif
-		}
-
-		void SetTag(const FTagData* InTag, FLowLevelMemTracker& InLLMRef
-#if LLM_ALLOW_ASSETS_TAGS || LLM_ALLOW_UOBJECTCLASSES_TAGS
-		, ELLMTagSet InTagSet = ELLMTagSet::None
-#endif
-		)
-		{
-#if LLM_ENABLED_FULL_TAGS
-#if !LLM_ALLOW_ASSETS_TAGS && !LLM_ALLOW_UOBJECTCLASSES_TAGS
-			Tag = InTag->GetIndex();
-#else
-			int32 TagSetIndex = GetIndexFromTagSet(InTagSet);
-			Tag[TagSetIndex] = InTag ? InTag->GetIndex() : InvalidCompressedTagValue;
-#endif //!LLM_ALLOW_ASSETS_TAGS
-#else
-			Tag = InTag->GetEnumTag();
-#endif
-		}
-
-		const FTagData* GetTag(FLowLevelMemTracker& InLLMRef
-#if LLM_ALLOW_ASSETS_TAGS || LLM_ALLOW_UOBJECTCLASSES_TAGS
-		, ELLMTagSet InTagSet = ELLMTagSet::None
-#endif
-		) const
-		{
-#if LLM_ENABLED_FULL_TAGS
-			UE::TSharedLock TagDataScopeLock(InLLMRef.TagDataMutex);
-#if !LLM_ALLOW_ASSETS_TAGS && !LLM_ALLOW_UOBJECTCLASSES_TAGS
-			return (*InLLMRef.TagDatas)[Tag];
-#else
-			int32 TagSetIndex = GetIndexFromTagSet(InTagSet);
-			return Tag[TagSetIndex] >= 0 ? (*InLLMRef.TagDatas)[Tag[TagSetIndex]] : nullptr;
-#endif
-#else
-			return InLLMRef.FindTagData(Tag);
-#endif
-		}
-
-#if LLM_ENABLED_FULL_TAGS
-		int32 GetCompressedTag(ELLMTagSet InTagSet = ELLMTagSet::None)
-		{
-#if !LLM_ALLOW_ASSETS_TAGS && !LLM_ALLOW_UOBJECTCLASSES_TAGS
-			return (InTagSet == ELLMTagSet::None) ? Tag : InvalidCompressedTagValue;
-#else
-			int32 TagSetIndex = GetIndexFromTagSet(InTagSet);
-			return Tag[TagSetIndex];
-#endif
-		}
-
-		void SetCompressedTag(int32 InTag, ELLMTagSet InTagSet = ELLMTagSet::None)
-		{
-#if !LLM_ALLOW_ASSETS_TAGS && !LLM_ALLOW_UOBJECTCLASSES_TAGS
-			if (InTagSet == ELLMTagSet::None)
-			{
-				Tag = InTag;
-			}
-#else
-			int32 TagSetIndex = GetIndexFromTagSet(InTagSet);
-			Tag[TagSetIndex] = InTag;
-#endif
-		}
-#else
-		ELLMTag GetCompressedTag() { return Tag; }
-		void SetCompressedTag(ELLMTag InTag) { Tag = InTag; }
-#endif
-
-	private:
-#if LLM_ENABLED_FULL_TAGS
-		// Even with arbitrary tags we are still partially compressed - the allocation records the tag's index
-		// (4 bytes) rather than the full tag pointer (8 bytes) or the tag's name (12 bytes).
-#if !LLM_ALLOW_ASSETS_TAGS && !LLM_ALLOW_UOBJECTCLASSES_TAGS
-		int32 Tag = 0;
-#else
-		int32 Tag[LLM_TAGSET_MAX];
-
-		static_assert(LLM_TAGSET_MAX <= 3,
-			"Incrementing the maximum amount of TagSet enums that can be defined at the same time adds 4 bytes of LLM overhead memory FOR EVERY ACTIVE ALLOCATION (e.g. 10 million * 4 == 40 million bytes). "
-			"Do not add them without measuring the cost of the overhead in the cook and runtime of large projects.");
-#endif // LLM_ALLOW_ASSETS_TAGS
-#else
-		ELLMTag Tag = ELLMTag::Untagged;
-#endif
-		static const int32 InvalidCompressedTagValue = -1;
-	};
-
-	typedef LLMMap<PointerKey, uint32, FLowLevelAllocInfo, LLMNumAllocsType> FLLMAllocMap;  // pointer, size, info, Capacity SizeType
-
-	void SetTotalTags(const FTagData* InOverrideUntaggedTagData, const FTagData* InOverrideTrackedTotalTagData);
-	void Update();
-	void UpdateThreads();
-	void OnTagsResorted(FTagDataArray& OldTagDatas);
-	void LockAllThreadTags(bool bLock);
-
-	void CaptureTagSnapshot();
-	void ClearTagSnapshot();
-
-	int64 GetTagAmount(const FTagData* TagData, UE::LLM::ESizeParams SizeParams = UE::LLM::ESizeParams::Default) const;
-	void SetTagAmountExternal(const FTagData* TagData, int64 Amount, bool bAddToTotal);
-	void SetTagAmountInUpdate(const FTagData* TagData, int64 Amount, bool bAddToTotal);
-	const FTagData* GetActiveTagData(ELLMTagSet TagSet = ELLMTagSet::None);
-	TArray<const FTagData*> GetTagDatas(ELLMTagSet TagSet = ELLMTagSet::None);
-
-	void GetTagsNamesWithAmount(TMap<FName, uint64>& OutTagsNamesWithAmount, ELLMTagSet TagSet = ELLMTagSet::None);
-	void GetTagsNamesWithAmountFiltered(TMap<FName, uint64>& OutTagsNamesWithAmount, ELLMTagSet TagSet, TArray<FLLMTagSetAllocationFilter>& Filters);
-
-	bool FindTagsForPtr(void* InPtr, TArray<const FTagData*, TInlineAllocator<static_cast<int32>(ELLMTagSet::Max)>>& OutTags) const;
-
-	int64 GetAllocTypeAmount(ELLMAllocType AllocType);
-
-	int64 GetTrackedTotal(UE::LLM::ESizeParams SizeParams = UE::LLM::ESizeParams::Default) const
-	{
-		if (EnumHasAnyFlags(SizeParams, UE::LLM::ESizeParams::RelativeToSnapshot))
-		{
-			return FMath::Clamp<int64>(TrackedTotal - TrackedTotalInSnapshot, 0, INT64_MAX); //-V569
-		}
-		return TrackedTotal;
-	}
-
-protected:
-	FLLMThreadState* GetOrCreateState();
-	FLLMThreadState* GetState();
-	void TrackAllocation(const void* Ptr, int64 Size, const FTagData* ActiveTagData, ELLMAllocType AllocType, FLLMThreadState* State, bool bTrackInMemPro);
-	void TrackMemoryOfActiveTag(int64 Amount, const FTagData* TagData, ELLMAllocType AllocType, FLLMThreadState* State);
-
-	FLowLevelMemTracker& LLMRef;
-
-	ELLMTracker Tracker;
-
-	uint32 TlsSlot;
-
-	TArray<FLLMThreadState*, FDefaultLLMAllocator> ThreadStates;
-
-	FCriticalSection PendingThreadStatesGuard;
-	TArray<FLLMThreadState*, FDefaultLLMAllocator> PendingThreadStates;
-	/**
-	 * Backup map from thread to threadstate. The primary lookup method is FPlatformTLS, but that is unavailable
-	 * during thread termination on some platforms. When unavailable, we use this TMap (under PendingThreadStatesGuard
-	 * critical section) to find the state.
-	 */
-	TMap<uint32, FLLMThreadState*, FDefaultSetLLMAllocator> ThreadIdToThreadState;
-
-	/** Sum of memory from all tracked tags. Duplicated in separate storage to make it instantly available at any time of frame without waiting for accumulation from threads during update.
-		GCC_ALIGN is required because it is modified in FPlatformAtomics::InterlockedAdd, which requires aligned values. */
-	int64 TrackedTotal GCC_ALIGN(8);
-
-	// The total tracked memory when the snapshot was taken
-	int64 TrackedTotalInSnapshot;
-
-#if !UE_ONLY_USE_PLATFORM_TRACKER
-	FLLMAllocMap AllocationMap;
-#endif
-
-	FTrackerTagSizeMap TagSizes;
-
-	const FTagData* OverrideUntaggedTagData;
-	const FTagData* OverrideTrackedTotalTagData;
-
-	FLLMCsvWriter CsvWriter;
-	FLLMTraceWriter TraceWriter;
-	FLLMCsvProfilerWriter CsvProfilerWriter;
-
-	double LastTrimTime;
-
-	int64 AllocTypeAmounts[(int32)ELLMAllocType::Count];
-};
 
 const TCHAR* ToString(ETagReferenceSource ReferenceSource);
 void SetMemoryStatByFName(FName Name, int64 Amount);
@@ -4834,6 +4451,116 @@ bool FTagData::IsUsedAsDisplayParent() const
 	return !(bHasEnumTag && EnumTag == ELLMTag::CustomName);
 }
 
+FLLMTracker::FLowLevelAllocInfo::FLowLevelAllocInfo()
+{
+#if LLM_ENABLED_FULL_TAGS && (LLM_ALLOW_ASSETS_TAGS || LLM_ALLOW_UOBJECTCLASSES_TAGS)
+	FMemory::Memset(Tag, InvalidCompressedTagValue);
+#endif
+}
+
+int32 FLLMTracker::FLowLevelAllocInfo::GetIndexFromTagSet(ELLMTagSet InTagSet)
+{
+	static_assert(static_cast<int32>(ELLMTagSet::Max) == 4,
+		"You added a new tagset, please update the code bellow");
+
+#if LLM_ALLOW_ASSETS_TAGS
+	switch (InTagSet)
+	{
+	case ELLMTagSet::None: return 0;
+	case ELLMTagSet::Assets: return 1;
+	case ELLMTagSet::AssetClasses: return 2;
+	default:
+		LLMCheck(false);
+		return 0;
+	}
+#elif LLM_ALLOW_UOBJECTCLASSES_TAGS
+	switch (InTagSet)
+	{
+	case ELLMTagSet::None: return 0;
+	case ELLMTagSet::UObjectClasses: return 1;
+	default:
+		LLMCheck(false);
+		return 0;
+	}
+#endif
+}
+
+void FLLMTracker::FLowLevelAllocInfo::SetTag(const FTagData* InTag, FLowLevelMemTracker& InLLMRef
+#if LLM_ALLOW_ASSETS_TAGS || LLM_ALLOW_UOBJECTCLASSES_TAGS
+	, ELLMTagSet InTagSet
+#endif
+)
+{
+#if LLM_ENABLED_FULL_TAGS
+#if !LLM_ALLOW_ASSETS_TAGS && !LLM_ALLOW_UOBJECTCLASSES_TAGS
+	Tag = InTag->GetIndex();
+#else
+	int32 TagSetIndex = GetIndexFromTagSet(InTagSet);
+	Tag[TagSetIndex] = InTag ? InTag->GetIndex() : InvalidCompressedTagValue;
+#endif //!LLM_ALLOW_ASSETS_TAGS
+#else
+	Tag = InTag->GetEnumTag();
+#endif
+}
+
+const FTagData* FLLMTracker::FLowLevelAllocInfo::GetTag(FLowLevelMemTracker& InLLMRef
+#if LLM_ALLOW_ASSETS_TAGS || LLM_ALLOW_UOBJECTCLASSES_TAGS
+	, ELLMTagSet InTagSet
+#endif
+) const
+{
+#if LLM_ENABLED_FULL_TAGS
+	UE::TSharedLock TagDataScopeLock(InLLMRef.TagDataMutex);
+#if !LLM_ALLOW_ASSETS_TAGS && !LLM_ALLOW_UOBJECTCLASSES_TAGS
+	return (*InLLMRef.TagDatas)[Tag];
+#else
+	int32 TagSetIndex = GetIndexFromTagSet(InTagSet);
+	return Tag[TagSetIndex] >= 0 ? (*InLLMRef.TagDatas)[Tag[TagSetIndex]] : nullptr;
+#endif
+#else
+	return InLLMRef.FindTagData(Tag);
+#endif
+}
+
+#if LLM_ENABLED_FULL_TAGS
+
+int32 FLLMTracker::FLowLevelAllocInfo::GetCompressedTag(ELLMTagSet InTagSet)
+{
+#if !LLM_ALLOW_ASSETS_TAGS && !LLM_ALLOW_UOBJECTCLASSES_TAGS
+	return (InTagSet == ELLMTagSet::None) ? Tag : InvalidCompressedTagValue;
+#else
+	int32 TagSetIndex = GetIndexFromTagSet(InTagSet);
+	return Tag[TagSetIndex];
+#endif
+}
+
+void FLLMTracker::FLowLevelAllocInfo::SetCompressedTag(int32 InTag, ELLMTagSet InTagSet)
+{
+#if !LLM_ALLOW_ASSETS_TAGS && !LLM_ALLOW_UOBJECTCLASSES_TAGS
+	if (InTagSet == ELLMTagSet::None)
+	{
+		Tag = InTag;
+	}
+#else
+	int32 TagSetIndex = GetIndexFromTagSet(InTagSet);
+	Tag[TagSetIndex] = InTag;
+#endif
+}
+
+#else // LLM_ENABLED_FULL_TAGS
+
+ELLMTag FLLMTracker::FLowLevelAllocInfo::GetCompressedTag()
+{
+	return Tag;
+}
+
+void FLLMTracker::FLowLevelAllocInfo::SetCompressedTag(ELLMTag InTag)
+{
+	Tag = InTag;
+}
+
+#endif // else LLM_ENABLED_FULL_TAGS
+
 FLLMTracker::FLLMTracker(FLowLevelMemTracker& InLLM)
 	: LLMRef(InLLM)
 	, Tracker(ELLMTracker::Max)
diff --git a/Engine/Source/Runtime/Core/Private/HAL/LLM/LLMPrivate.h b/Engine/Source/Runtime/Core/Private/HAL/LLM/LLMPrivate.h
--- a/Engine/Source/Runtime/Core/Private/HAL/LLM/LLMPrivate.h
+++ b/Engine/Source/Runtime/Core/Private/HAL/LLM/LLMPrivate.h
@@ -10,6 +10,7 @@
 #include "Containers/Set.h"
 #include "Containers/SortedMap.h"
 #include "HAL/LLM/MemoryUtils.h"
+#include "ProfilingDebugging/CsvProfiler.h"
 
 // Private defines configuring LLM; see also public defines in LowLevelMemTracker.h
 // LLM_ALLOW_NAMES_TAGS: Set to 1 to allow arbitrary FNames to be used as tags, at the cost of more LLM memory usage per allocation. Set to 0 to store only the top-level ELLMTag for each allocation.
@@ -44,6 +45,18 @@
 // Disable if you need a little more memory or speed
 #define LLM_ENABLED_TRACK_PEAK_MEMORY 1
 
+// Enable to allow memory tracking to be passed through to external tools without the extra overhead of LLM also tracking allocations.
+#define UE_ONLY_USE_PLATFORM_TRACKER 0
+
+#define LLM_CSV_PROFILER_WRITER_ENABLED CSV_PROFILER_STATS
+
+static_assert(static_cast<uint32>(ELLMTagSet::Max) == 4, "You added a new TagSet, please update the #ifs below");
+#if LLM_ALLOW_ASSETS_TAGS
+#define LLM_TAGSET_MAX 3
+#else
+#define LLM_TAGSET_MAX 2
+#endif
+
 namespace UE::LLMPrivate
 {
 
@@ -226,9 +239,9 @@ struct FTrackerTagSizeData
 typedef TLLMMap<const FTagData*, FTrackerTagSizeData> FTrackerTagSizeMap;
 
 /**
-	* Size information stored on each threadstate for a tag
-	* TagSizes are sorted by Index instead of by pointer in the ThreadTagSizeMap to enforce the constraint that Parents come before children
-	*/
+ * Size information stored on each threadstate for a tag
+ * TagSizes are sorted by Index instead of by pointer in the ThreadTagSizeMap to enforce the constraint that Parents come before children
+ */
 struct FThreadTagSizeData
 {
 	const FTagData* TagData = nullptr;
@@ -240,4 +253,305 @@ typedef TSortedMap<int32, FThreadTagSizeData, FDefaultLLMAllocator> FThreadTagSi
 
 FName LLMGetTagUniqueName(ELLMTag Tag);
 
+namespace UE::LLMPrivate
+{
+
+/** FLLMCsvWriter: class for writing out the LLM tag sizes to a csv file every few seconds. */
+class FLLMCsvWriter
+{
+public:
+	FLLMCsvWriter();
+	~FLLMCsvWriter();
+
+	void SetTracker(ELLMTracker InTracker);
+	void Clear();
+
+	void Publish(FLowLevelMemTracker& LLMRef, const FTrackerTagSizeMap& TagSizes,
+		const FTagData* OverrideTrackedTotalTagData, const FTagData* OverrideUntaggedTagData, int64 TrackedTotal,
+		UE::LLM::ESizeParams SizeParams);
+
+	void OnPreFork();
+
+	void Flush(bool IsOnCrash);
+
+private:
+	void Write(FStringView Text);
+	static const TCHAR* GetTrackerCsvName(ELLMTracker InTracker);
+
+	bool CreateArchive();
+	bool UpdateColumns(const FTrackerTagSizeMap& TagSizes);
+	void WriteHeader(const FTagData* OverrideTrackedTotalTagData, const FTagData* OverrideUntaggedTagData);
+	void AddRow(FLowLevelMemTracker& LLMRef, const FTrackerTagSizeMap& TagSizes,
+		const FTagData* OverrideTrackedTotalName, const FTagData* OverrideUntaggedName, int64 TrackedTotal,
+		UE::LLM::ESizeParams SizeParams);
+
+	FConstTagDataArray Columns;
+	TLLMSet<const FTagData*> ExistingColumns;
+	FArchive* Archive;
+	double LastWriteTime;
+	int32 WriteCount;
+	ELLMTracker Tracker;
+	bool bRegisteredFlushDelegate;
+};
+
+/** Outputs the LLM tags and sizes to TraceLog events. */
+class FLLMTraceWriter
+{
+public:
+	FLLMTraceWriter();
+	void SetTracker(ELLMTracker InTracker);
+	void Clear();
+	void Publish(FLowLevelMemTracker& LLMRef, const FTrackerTagSizeMap& TagSizes,
+		const FTagData* OverrideTrackedTotalName, const FTagData* OverrideUntaggedName, int64 TrackedTotal,
+		UE::LLM::ESizeParams SizeParams);
+
+private:
+	static const void* GetTagId(const FTagData* TagData);
+	void SendTagDeclaration(const FTagData* TagData);
+
+	ELLMTracker Tracker;
+	TLLMSet<const FTagData*> DeclaredTags;
+	bool bTrackerSpecSent = false;
+	bool bTagSetSpecSent = false;
+};
+
+/** FLLMCsvWriter: class for writing out LLM stats to the Csv Profiler. */
+class FLLMCsvProfilerWriter
+{
+public:
+	FLLMCsvProfilerWriter();
+	void Clear();
+	void SetTracker(ELLMTracker InTracker);
+	void Publish(FLowLevelMemTracker& LLMRef, const FTrackerTagSizeMap& TagSizes,
+		const FTagData* OverrideTrackedTotalName, const FTagData* OverrideUntaggedName, int64 TrackedTotal,
+		UE::LLM::ESizeParams SizeParams);
+protected:
+	void RecordTagToCsv(int32 CsvCategoryIndex, const FTagData* TagData, int64 TagSize);
+#if LLM_CSV_PROFILER_WRITER_ENABLED
+	ELLMTracker Tracker;
+	TLLMMap<const FTagData*, FName> TagDataToCsvStatName;
+#endif
+};
+
+/** Per-thread state in an LLMTracker. */
+class FLLMThreadState
+{
+public:
+	FLLMThreadState();
+
+	void Clear();
+
+	void PushTag(const FTagData* TagData, ELLMTagSet TagSet);
+	void PopTag(ELLMTagSet TagSet);
+	const FTagData* GetTopTag(ELLMTagSet TagSet);
+	void TrackAllocation(const void* Ptr, int64 Size, ELLMTracker Tracker, ELLMAllocType AllocType,
+		const FTagData* TagData, const FTagData* AssetTagData, const FTagData* AssetClassTagData, const FTagData* UObjectClassesTagData, bool bTrackInMemPro);
+	void TrackFree(const void* Ptr, int64 Size, ELLMTracker Tracker, ELLMAllocType AllocType,
+		const FTagData* TagData, const FTagData* AssetTagData, const FTagData* AssetClassTagData, const FTagData* UObjectClassesTagData, bool bTrackInMemPro);
+	void TrackMemory(int64 Amount, ELLMTracker Tracker, ELLMAllocType AllocType,
+		const FTagData* TagData);
+	void TrackMoved(const void* Dest, const void* Source, int64 Size, ELLMTracker Tracker, const FTagData* TagData);
+	void IncrTag(const FTagData* Tag, int64 Amount);
+
+	void PropagateChildSizesToParents();
+	void OnTagsResorted(FTagDataArray& OldTagDatas);
+	void LockTags(bool bLock);
+	void FetchAndClearTagSizes(FTrackerTagSizeMap& TagSizes, int64* InAllocTypeAmounts, bool bTrimAllocations);
+
+	void ClearAllocTypeAmounts();
+
+	// AllocationsLock must be held around read/write of Allocations and AllocTypeAmounts, because these are read/write
+	// from the Update thread. All other data on FLLMThreadState are accessed only from the thread that owns the
+	// FLLMThreadState and do not require a lock.
+	// AllocationsLock is also used to guarantee our contract provided to FLLMTracker::LockAllThreadTags: We must not
+	// dereference any FTagData* held by the FLLMThreadState unless we have entered AllocationsLock.
+	FConstTagDataArray TagStack[static_cast<int32>(ELLMTagSet::Max)];
+	FThreadTagSizeMap Allocations;
+	FCriticalSection AllocationsLock;
+
+	int8 PausedCounter[(int32)ELLMAllocType::Count];
+	int64 AllocTypeAmounts[(int32)ELLMAllocType::Count];
+};
+
+/** The main LLM implementation class. It owns the thread state objects. */
+class FLLMTracker
+{
+public:
+
+	FLLMTracker(FLowLevelMemTracker& InLLM);
+	~FLLMTracker();
+
+	void Initialise(ELLMTracker InTracker, FLLMAllocator* InAllocator);
+
+	void PushTag(ELLMTag EnumTag, ELLMTagSet TagSet);
+	void PushTag(FName Tag, bool bInIsStatTag, ELLMTagSet TagSet);
+	void PushTag(const FTagData* TagData, ELLMTagSet TagSet);
+	void PopTag(ELLMTagSet TagSet);
+	void TrackAllocation(const void* Ptr, int64 Size, FName DefaultTag, ELLMAllocType AllocType, bool bTrackInMemPro);
+	void TrackAllocation(const void* Ptr, int64 Size, ELLMTag DefaultTag, ELLMAllocType AllocType, bool bTrackInMemPro);
+	void TrackFree(const void* Ptr, ELLMAllocType AllocType, bool bTrackInMemPro);
+	void TrackMemoryOfActiveTag(int64 Amount, FName DefaultTag, ELLMAllocType AllocType);
+	void TrackMemoryOfActiveTag(int64 Amount, ELLMTag DefaultTag, ELLMAllocType AllocType);
+	void OnAllocMoved(const void* Dest, const void* Source, ELLMAllocType AllocType);
+
+	void TrackMemory(ELLMTag EnumTag, int64 Amount, ELLMAllocType AllocType);
+	void TrackMemory(FName TagName, ELLMTagSet TagSet, int64 Amount, ELLMAllocType AllocType);
+	void TrackMemory(const FTagData* TagData, int64 Amount, ELLMAllocType AllocType);
+
+	// This will pause/unpause tracking, and also manually increment a given tag
+	void PauseAndTrackMemory(FName TagName, ELLMTagSet TagSet, bool bInIsStatTag, int64 Amount, ELLMAllocType AllocType);
+	void PauseAndTrackMemory(ELLMTag EnumTag, int64 Amount, ELLMAllocType AllocType);
+	void PauseAndTrackMemory(const FTagData* TagData, int64 Amount, ELLMAllocType AllocType);
+	void Pause(ELLMAllocType AllocType);
+	void Unpause(ELLMAllocType AllocType);
+	bool IsPaused(ELLMAllocType AllocType);
+
+	void Clear();
+
+	// Dump the allocation count and size for each tag, along with the proportion that is private/shared/unreferenced
+	// on linux machines. This will dump to a CSV in the project Saved/LLM directory with a separate file for each
+	// tracker and forked child. Integers will be appended to the filename to prevent overwrites. Returns false
+	// if the memory information can't be retrieved (likely because it's not on a supported platform)
+	bool DumpForkedAllocationInfo();
+
+	// Dump the list of all registered LLM tags and their properties to a CSV file, in the LLM's profiling directory
+	// (Saved/Profiling/LLM/); see UE::LLM::GetLLMProfilingDir().
+	bool DumpTags();
+
+	void PublishStats(UE::LLM::ESizeParams SizeParam);
+	void PublishCsv(UE::LLM::ESizeParams SizeParam);
+	void PublishTrace(UE::LLM::ESizeParams SizeParam);
+	void PublishCsvProfiler(UE::LLM::ESizeParams SizeParam);
+
+	void OnPreFork();
+
+	struct FLowLevelAllocInfo
+	{
+	public:
+		FLowLevelAllocInfo();
+
+		static int32 GetIndexFromTagSet(ELLMTagSet InTagSet);
+
+		void SetTag(const FTagData* InTag, FLowLevelMemTracker& InLLMRef
+#if LLM_ALLOW_ASSETS_TAGS || LLM_ALLOW_UOBJECTCLASSES_TAGS
+			, ELLMTagSet InTagSet = ELLMTagSet::None
+#endif
+		);
+
+		const FTagData* GetTag(FLowLevelMemTracker& InLLMRef
+#if LLM_ALLOW_ASSETS_TAGS || LLM_ALLOW_UOBJECTCLASSES_TAGS
+			, ELLMTagSet InTagSet = ELLMTagSet::None
+#endif
+		) const;
+
+#if LLM_ENABLED_FULL_TAGS
+		int32 GetCompressedTag(ELLMTagSet InTagSet = ELLMTagSet::None);
+		void SetCompressedTag(int32 InTag, ELLMTagSet InTagSet = ELLMTagSet::None);
+#else
+		ELLMTag GetCompressedTag();
+		void SetCompressedTag(ELLMTag InTag);
+#endif
+
+	private:
+#if LLM_ENABLED_FULL_TAGS
+		// Even with arbitrary tags we are still partially compressed - the allocation records the tag's index
+		// (4 bytes) rather than the full tag pointer (8 bytes) or the tag's name (12 bytes).
+#if !LLM_ALLOW_ASSETS_TAGS && !LLM_ALLOW_UOBJECTCLASSES_TAGS
+		int32 Tag = 0;
+#else
+		int32 Tag[LLM_TAGSET_MAX];
+
+		static_assert(LLM_TAGSET_MAX <= 3,
+			"Incrementing the maximum amount of TagSet enums that can be defined at the same time adds 4 bytes of LLM overhead memory FOR EVERY ACTIVE ALLOCATION (e.g. 10 million * 4 == 40 million bytes). "
+			"Do not add them without measuring the cost of the overhead in the cook and runtime of large projects.");
+#endif // LLM_ALLOW_ASSETS_TAGS
+#else
+		ELLMTag Tag = ELLMTag::Untagged;
+#endif
+		static const int32 InvalidCompressedTagValue = -1;
+	};
+
+	typedef LLMMap<PointerKey, uint32, FLowLevelAllocInfo, LLMNumAllocsType> FLLMAllocMap;  // pointer, size, info, Capacity SizeType
+
+	void SetTotalTags(const FTagData* InOverrideUntaggedTagData, const FTagData* InOverrideTrackedTotalTagData);
+	void Update();
+	void UpdateThreads();
+	void OnTagsResorted(FTagDataArray& OldTagDatas);
+	void LockAllThreadTags(bool bLock);
+
+	void CaptureTagSnapshot();
+	void ClearTagSnapshot();
+
+	int64 GetTagAmount(const FTagData* TagData, UE::LLM::ESizeParams SizeParams = UE::LLM::ESizeParams::Default) const;
+	void SetTagAmountExternal(const FTagData* TagData, int64 Amount, bool bAddToTotal);
+	void SetTagAmountInUpdate(const FTagData* TagData, int64 Amount, bool bAddToTotal);
+	const FTagData* GetActiveTagData(ELLMTagSet TagSet = ELLMTagSet::None);
+	TArray<const FTagData*> GetTagDatas(ELLMTagSet TagSet = ELLMTagSet::None);
+
+	void GetTagsNamesWithAmount(TMap<FName, uint64>& OutTagsNamesWithAmount, ELLMTagSet TagSet = ELLMTagSet::None);
+	void GetTagsNamesWithAmountFiltered(TMap<FName, uint64>& OutTagsNamesWithAmount, ELLMTagSet TagSet, TArray<FLLMTagSetAllocationFilter>& Filters);
+
+	bool FindTagsForPtr(void* InPtr, TArray<const FTagData*, TInlineAllocator<static_cast<int32>(ELLMTagSet::Max)>>& OutTags) const;
+
+	int64 GetAllocTypeAmount(ELLMAllocType AllocType);
+
+	int64 GetTrackedTotal(UE::LLM::ESizeParams SizeParams = UE::LLM::ESizeParams::Default) const
+	{
+		if (EnumHasAnyFlags(SizeParams, UE::LLM::ESizeParams::RelativeToSnapshot))
+		{
+			return FMath::Clamp<int64>(TrackedTotal - TrackedTotalInSnapshot, 0, INT64_MAX); //-V569
+		}
+		return TrackedTotal;
+	}
+
+protected:
+	FLLMThreadState* GetOrCreateState();
+	FLLMThreadState* GetState();
+	void TrackAllocation(const void* Ptr, int64 Size, const FTagData* ActiveTagData, ELLMAllocType AllocType, FLLMThreadState* State, bool bTrackInMemPro);
+	void TrackMemoryOfActiveTag(int64 Amount, const FTagData* TagData, ELLMAllocType AllocType, FLLMThreadState* State);
+
+	FLowLevelMemTracker& LLMRef;
+
+	ELLMTracker Tracker;
+
+	uint32 TlsSlot;
+
+	TArray<FLLMThreadState*, FDefaultLLMAllocator> ThreadStates;
+
+	FCriticalSection PendingThreadStatesGuard;
+	TArray<FLLMThreadState*, FDefaultLLMAllocator> PendingThreadStates;
+	/**
+	 * Backup map from thread to threadstate. The primary lookup method is FPlatformTLS, but that is unavailable
+	 * during thread termination on some platforms. When unavailable, we use this TMap (under PendingThreadStatesGuard
+	 * critical section) to find the state.
+	 */
+	TMap<uint32, FLLMThreadState*, FDefaultSetLLMAllocator> ThreadIdToThreadState;
+
+	/** Sum of memory from all tracked tags. Duplicated in separate storage to make it instantly available at any time of frame without waiting for accumulation from threads during update.
+		GCC_ALIGN is required because it is modified in FPlatformAtomics::InterlockedAdd, which requires aligned values. */
+	int64 TrackedTotal GCC_ALIGN(8);
+
+	// The total tracked memory when the snapshot was taken
+	int64 TrackedTotalInSnapshot;
+
+#if !UE_ONLY_USE_PLATFORM_TRACKER
+	FLLMAllocMap AllocationMap;
+#endif
+
+	FTrackerTagSizeMap TagSizes;
+
+	const FTagData* OverrideUntaggedTagData;
+	const FTagData* OverrideTrackedTotalTagData;
+
+	FLLMCsvWriter CsvWriter;
+	FLLMTraceWriter TraceWriter;
+	FLLMCsvProfilerWriter CsvProfilerWriter;
+
+	double LastTrimTime;
+
+	int64 AllocTypeAmounts[(int32)ELLMAllocType::Count];
+};
+
+} // namespace UE::LLMPrivate
+
 #endif		// #if ENABLE_LOW_LEVEL_MEM_TRACKER
