- Disable stall detector when running with a debugger attached

#Synced-CL 49561835

diff --git a/Engine/Source/Runtime/Core/Private/ProfilingDebugging/StallDetector.cpp b/Engine/Source/Runtime/Core/Private/ProfilingDebugging/StallDetector.cpp
--- a/Engine/Source/Runtime/Core/Private/ProfilingDebugging/StallDetector.cpp
+++ b/Engine/Source/Runtime/Core/Private/ProfilingDebugging/StallDetector.cpp
@@ -161,7 +161,7 @@ uint32 UE::FStallDetectorRunnable::Run()
 				}
 
 				double OverageSeconds = Detector->Check(false, Seconds);
-				if (OverageSeconds > 0.0)
+				if (OverageSeconds > 0.0 && !FPlatformMisc::IsDebuggerPresent())
 				{
 					TArray<uint64> Backtrace;
 					Backtrace.AddUninitialized(StallDetectorMaxBacktraceEntries);
