LLM: Hide the implementation in a private header so that modifying implementation details does not cause a full recompile.
Update spelling of Initialise for private fields to match the coding standard spelling.
#rnx

#Synced-CL 49633983

diff --git a/Engine/Source/Runtime/Core/Private/HAL/LLM/LLM.cpp b/Engine/Source/Runtime/Core/Private/HAL/LLM/LLM.cpp
--- a/Engine/Source/Runtime/Core/Private/HAL/LLM/LLM.cpp
+++ b/Engine/Source/Runtime/Core/Private/HAL/LLM/LLM.cpp
@@ -352,7 +352,7 @@ struct FArraySlackSortItem
 
 static void LlmTrackArrayDumpTag(FArchive* LogFile, FOutputDevice& Output, TAnsiStringBuilder<4096>& Builder, int32 TagIndex, uint32 StackTraceDepth, bool bVerbose, double StartTime)
 {
-	FLowLevelMemTracker& Tracker = FLowLevelMemTracker::Get();
+	FLLMGlobals& Tracker = FLLMGlobals::GetInner();
 	bool bLlmEnabled = Tracker.IsEnabled();
 
 	FScopeLock Lock(GTrackArrayDetailedLock);
@@ -782,7 +782,7 @@ void ArraySlackTrackGenerateReport(const TCHAR* Cmd, FOutputDevice& Output)
 #endif
 	}
 
-	const FLowLevelMemTracker& Tracker = FLowLevelMemTracker::Get();
+	const FLLMGlobals& Tracker = FLLMGlobals::GetInner();
 	bool bLlmEnabled = Tracker.IsEnabled();
 
 	TAnsiStringBuilder<4096> Builder;
@@ -905,7 +905,7 @@ void ArraySlackTrackGenerateReport(const TCHAR* Cmd, FOutputDevice& Output)
 
 uint8 LlmGetActiveTag()
 {
-	FLowLevelMemTracker& Tracker = FLowLevelMemTracker::Get();
+	FLLMGlobals& Tracker = FLLMGlobals::GetInner();
 	if (!Tracker.IsInitialized())
 	{
 		return 0;
@@ -1023,23 +1023,11 @@ FAutoConsoleCommand DumpLLM(
 	}));
 
 #if !PLATFORM_HAS_MULTITHREADED_PREMAIN
-struct FLowLevelMemTracker::FEnableStateScopeLock
-{
-};
-bool FLowLevelMemTracker::TryEnterEnabled(FLowLevelMemTracker::FEnableStateScopeLock& ScopeLock)
+bool FLowLevelMemTracker::TryEnterEnabled(UE::LLMPrivate::FEnableStateScopeLock& ScopeLock)
 {
 	return IsEnabled();
 }
 #else // !PLATFORM_HAS_MULTITHREADED_PREMAIN
-struct FLowLevelMemTracker::FEnableStateScopeLock
-{
-	// Undefine copy/move constructor since FReadScopeLock does not support it.
-	FEnableStateScopeLock() = default;
-	FEnableStateScopeLock(const FEnableStateScopeLock&) = delete;
-	FEnableStateScopeLock(FEnableStateScopeLock&&) = delete;
-
-	TOptional<FReadScopeLock> Inner;
-};
 
 namespace UE::LLMPrivate
 {
@@ -1052,8 +1040,10 @@ FRWLock& GetEnableStateLock()
 
 } // namespace UE::LLMPrivate
 
-bool FLowLevelMemTracker::TryEnterEnabled(FLowLevelMemTracker::FEnableStateScopeLock& ScopeLock)
+bool FLowLevelMemTracker::TryEnterEnabled(UE::LLMPrivate::FEnableStateScopeLock& ScopeLock)
 {
+	using namespace UE::LLMPrivate;
+
 	switch (EnabledState)
 	{
 	case EEnabled::NotYetKnown:
@@ -1087,6 +1077,14 @@ void FLowLevelMemTracker::DumpToLog(EDumpFormat DumpFormat, FOutputDevice* Outpu
 	{
 		return;
 	}
+	GetGlobals()->DumpToLogInner(DumpFormat, OutputDevice, SizeParams, TagSet);
+}
+
+namespace UE::LLMPrivate
+{
+
+void FLLMGlobals::DumpToLogInner(EDumpFormat DumpFormat, FOutputDevice* OutputDevice, UE::LLM::ESizeParams SizeParams, ELLMTagSet TagSet)
+{
 	const float InvToMb = 1.0 / (1024 * 1024);
 	FOutputDevice& Ar = *(OutputDevice ? OutputDevice : GLog);
 	if (DumpFormat == EDumpFormat::CSV)
@@ -1159,6 +1157,8 @@ void FLowLevelMemTracker::DumpToLog(EDumpFormat DumpFormat, FOutputDevice* Outpu
 	}
 }
 
+} // namespace UE::LLMPrivate
+
 DECLARE_LLM_MEMORY_STAT(TEXT("LLM Overhead"), STAT_LLMOverheadTotal, STATGROUP_LLMOverhead);
 
 // LLM Summary stats referenced by ELLMTagNames
@@ -1294,13 +1294,12 @@ DECLARE_LLM_MEMORY_STAT(TEXT("CustomName"), STAT_CustomName, STATGROUP_LLMFULL);
 namespace UE::LLMPrivate
 {
 
-
 const TCHAR* ToString(ETagReferenceSource ReferenceSource);
 void SetMemoryStatByFName(FName Name, int64 Amount);
 void ValidateUniqueName(FStringView UniqueName);
 
 typedef TArray<TPair<FLLMInitialisedCallback, UPTRINT>, TInlineAllocator<1>> FInitializedCallbacksArray;
-FInitializedCallbacksArray& GetInitialisedCallbacks()
+FInitializedCallbacksArray& GetInitializedCallbacks()
 {
 	static FInitializedCallbacksArray Array;
 	return Array;
@@ -1317,7 +1316,7 @@ void FPrivateCallbacks::AddInitialisedCallback(FLLMInitialisedCallback Callback,
 {
 	FLowLevelMemTracker::Get().BootstrapInitialise();
 
-	FInitializedCallbacksArray& Callbacks = UE::LLMPrivate::GetInitialisedCallbacks();
+	FInitializedCallbacksArray& Callbacks = UE::LLMPrivate::GetInitializedCallbacks();
 	Callbacks.Add(TPair<FLLMInitialisedCallback, UPTRINT> { Callback, UserData });
 }
 
@@ -1459,7 +1458,7 @@ extern const ANSICHAR* LLMGetTagNameANSI(ELLMTag Tag)
 
 FLowLevelMemTracker& FLowLevelMemTracker::Construct()
 {
-	static FLowLevelMemTracker Tracker;
+	static UE::LLMPrivate::FLLMGlobals Tracker;
 	TrackerInstance = &Tracker;
 	return Tracker;
 }
@@ -1467,11 +1466,418 @@ FLowLevelMemTracker& FLowLevelMemTracker::Construct()
 FLowLevelMemTracker* FLowLevelMemTracker::TrackerInstance = nullptr;
 // LLM must start off enabled because allocations happen before the command line enables/disables us
 bool FLowLevelMemTracker::bIsDisabled = false;
-FLowLevelMemTracker::EEnabled FLowLevelMemTracker::EnabledState = EEnabled::NotYetKnown;
+UE::LLMPrivate::EEnabled FLowLevelMemTracker::EnabledState = UE::LLMPrivate::EEnabled::NotYetKnown;
 
 static const TCHAR* InvalidLLMTagName = TEXT("?");
 
-FLowLevelMemTracker::FLowLevelMemTracker()
+bool FLowLevelMemTracker::IsInitialized() const
+{
+	// IsEnabled is not checked here because Inner can be called even if not enabled.
+	return GetGlobals()->IsInitializedInner();
+}
+
+bool FLowLevelMemTracker::IsConfigured() const
+{
+	// IsEnabled is not checked here because Inner can be called even if not enabled.
+	return GetGlobals()->IsConfiguredInner();
+}
+
+void FLowLevelMemTracker::BootstrapInitialise()
+{
+	// IsEnabled is not checked here because there is some work to do in Inner function even if not enabled.
+	GetGlobals()->BootstrapInitializeInner();
+}
+
+void FLowLevelMemTracker::OnPreFork()
+{
+	if (!IsEnabled())
+	{
+		return;
+	}
+	GetGlobals()->OnPreForkInner();
+}
+
+void FLowLevelMemTracker::UpdateStatsPerFrame(const TCHAR* LogName)
+{
+	// IsEnabled is not checked here because there is some work to do in Inner function even if not enabled.
+	GetGlobals()->UpdateStatsPerFrameInner(LogName);
+}
+
+void FLowLevelMemTracker::Tick()
+{
+	if (!IsEnabled())
+	{
+		return;
+	}
+	GetGlobals()->TickInner();
+}
+
+void FLowLevelMemTracker::SetProgramSize(uint64 InProgramSize)
+{
+	if (!IsEnabled())
+	{
+		return;
+	}
+	GetGlobals()->SetProgramSizeInner(InProgramSize);
+}
+
+void FLowLevelMemTracker::ProcessCommandLine(const TCHAR* CmdLine)
+{
+	// IsEnabled is not checked here because there is some work to do in Inner function even if not enabled.
+	GetGlobals()->ProcessCommandLineInner(CmdLine);
+}
+
+uint64 FLowLevelMemTracker::GetTotalTrackedMemory(ELLMTracker Tracker)
+{
+	if (!IsEnabled())
+	{
+		return 0;
+	}
+	return GetGlobals()->GetTotalTrackedMemoryInner(Tracker);
+}
+
+void FLowLevelMemTracker::OnLowLevelAlloc(ELLMTracker Tracker, const void* Ptr, uint64 Size, ELLMTag DefaultTag,
+	ELLMAllocType AllocType, bool bTrackInMemPro)
+{
+	UE::LLMPrivate::FEnableStateScopeLock EnableScopeLock;
+	if (!TryEnterEnabled(EnableScopeLock))
+	{
+		return;
+	}
+	GetGlobals()->OnLowLevelAllocInner(Tracker, Ptr, Size, DefaultTag, AllocType, bTrackInMemPro);
+}
+
+void FLowLevelMemTracker::OnLowLevelAlloc(ELLMTracker Tracker, const void* Ptr, uint64 Size, FName DefaultTag,
+	ELLMAllocType AllocType, bool bTrackInMemPro)
+{
+	UE::LLMPrivate::FEnableStateScopeLock EnableScopeLock;
+	if (!TryEnterEnabled(EnableScopeLock))
+	{
+		return;
+	}
+	GetGlobals()->OnLowLevelAllocInner(Tracker, Ptr, Size, DefaultTag, AllocType, bTrackInMemPro);
+}
+
+void FLowLevelMemTracker::OnLowLevelFree(ELLMTracker Tracker, const void* Ptr,
+	ELLMAllocType AllocType, bool bTrackInMemPro)
+{
+	UE::LLMPrivate::FEnableStateScopeLock EnableScopeLock;
+	if (!TryEnterEnabled(EnableScopeLock))
+	{
+		return;
+	}
+	GetGlobals()->OnLowLevelFreeInner(Tracker, Ptr, AllocType, bTrackInMemPro);
+}
+
+void FLowLevelMemTracker::OnLowLevelChangeInMemoryUse(ELLMTracker Tracker, int64 DeltaMemory, ELLMTag DefaultTag,
+	ELLMAllocType AllocType)
+{
+	UE::LLMPrivate::FEnableStateScopeLock EnableScopeLock;
+	if (!TryEnterEnabled(EnableScopeLock))
+	{
+		return;
+	}
+	GetGlobals()->OnLowLevelChangeInMemoryUseInner(Tracker, DeltaMemory, DefaultTag, AllocType);
+}
+
+void FLowLevelMemTracker::OnLowLevelChangeInMemoryUse(ELLMTracker Tracker, int64 DeltaMemory, FName DefaultTag,
+	ELLMAllocType AllocType)
+{
+	UE::LLMPrivate::FEnableStateScopeLock EnableScopeLock;
+	if (!TryEnterEnabled(EnableScopeLock))
+	{
+		return;
+	}
+	GetGlobals()->OnLowLevelChangeInMemoryUseInner(Tracker, DeltaMemory, DefaultTag, AllocType);
+}
+
+void FLowLevelMemTracker::OnLowLevelAllocMoved(ELLMTracker Tracker, const void* Dest, const void* Source,
+	ELLMAllocType AllocType)
+{
+	UE::LLMPrivate::FEnableStateScopeLock EnableScopeLock;
+	if (!TryEnterEnabled(EnableScopeLock))
+	{
+		return;
+	}
+	GetGlobals()->OnLowLevelAllocMovedInner(Tracker, Dest, Source, AllocType);
+}
+
+bool FLowLevelMemTracker::Exec(const TCHAR* Cmd, FOutputDevice& Ar)
+{
+	if (!IsEnabled())
+	{
+		if (FParse::Command(&Cmd, TEXT("LLMEM")))
+		{
+			UE_LOG(LogHAL, Warning, TEXT("Cannot execute LLMEM command! LLM is not enabled."));
+			return true;
+		}
+		return false;
+	}
+	return GetGlobals()->ExecInner(Cmd, Ar);
+}
+
+bool FLowLevelMemTracker::IsTagSetActive(ELLMTagSet Set)
+{
+	if (!IsEnabled())
+	{
+		return false;
+	}
+	return GetGlobals()->IsTagSetActiveInner(Set);
+}
+
+bool FLowLevelMemTracker::ShouldReduceThreads()
+{
+#if LLM_ENABLED_REDUCE_THREADS
+	if (!IsEnabled())
+	{
+		return false;
+	}
+	return GetGlobals()->ShouldReduceThreadsInner();
+#else
+	return false;
+#endif
+}
+
+void FLowLevelMemTracker::RegisterPlatformTag(int32 Tag, const TCHAR* Name, FName StatName, FName SummaryStatName,
+	int32 ParentTag)
+{
+	// IsEnabled is not checked here because there is some work to do in Inner function even if not enabled.
+	GetGlobals()->RegisterPlatformTagInner(Tag, Name, StatName, SummaryStatName, ParentTag);
+}
+
+void FLowLevelMemTracker::RegisterProjectTag(int32 Tag, const TCHAR* Name, FName StatName, FName SummaryStatName,
+	int32 ParentTag)
+{
+	// IsEnabled is not checked here because there is some work to do in Inner function even if not enabled.
+	GetGlobals()->RegisterProjectTagInner(Tag, Name, StatName, SummaryStatName, ParentTag);
+}
+
+void GlobalRegisterTagDeclaration(FLLMTagDeclaration& TagDeclaration)
+{
+	if (!FLowLevelMemTracker::IsEnabled())
+	{
+		return;
+	}
+	UE::LLMPrivate::FLLMGlobals::GetInner().RegisterTagDeclaration(TagDeclaration);
+}
+
+void FLowLevelMemTracker::FinishInitialise()
+{
+	// IsEnabled is not checked here because there is some work to do in Inner function even if not enabled.
+	GetGlobals()->FinishInitializeInner();
+}
+
+TArray<const UE::LLMPrivate::FTagData*> FLowLevelMemTracker::GetTrackedTags(ELLMTagSet TagSet)
+{
+	if (!IsEnabled())
+	{
+		return TArray<const UE::LLMPrivate::FTagData*>();
+	}
+	return GetGlobals()->GetTrackedTagsInner(TagSet);
+}
+
+TArray<const UE::LLMPrivate::FTagData*> FLowLevelMemTracker::GetTrackedTags(ELLMTracker Tracker, ELLMTagSet TagSet /* = ELLMTagSet::None */)
+{
+	if (!IsEnabled())
+	{
+		return TArray<const UE::LLMPrivate::FTagData*>();
+	}
+	return GetGlobals()->GetTrackedTagsInner(Tracker, TagSet);
+}
+
+void FLowLevelMemTracker::GetTrackedTagsNamesWithAmount(TMap<FName, uint64>& TagsNamesWithAmount, ELLMTracker Tracker, ELLMTagSet TagSet)
+{
+	if (!IsEnabled())
+	{
+		return;
+	}
+	GetGlobals()->GetTrackedTagsNamesWithAmountInner(TagsNamesWithAmount, Tracker, TagSet);
+}
+
+void FLowLevelMemTracker::GetTrackedTagsNamesWithAmountFiltered(TMap<FName, uint64>& TagsNamesWithAmount, ELLMTracker Tracker, ELLMTagSet TagSet, TArray<FLLMTagSetAllocationFilter>& Filters)
+{
+	if (!IsEnabled())
+	{
+		return;
+	}
+	GetGlobals()->GetTrackedTagsNamesWithAmountFilteredInner(TagsNamesWithAmount, Tracker, TagSet, Filters);
+}
+
+bool FLowLevelMemTracker::FindTagByName(const TCHAR* Name, uint64& OutTag, ELLMTagSet InTagSet /*= ELLMTagSet::None*/) const
+{
+	using namespace UE::LLMPrivate;
+
+	if (!IsEnabled())
+	{
+		return false;
+	}
+	return GetGlobals()->FindTagByNameInner(Name, OutTag, InTagSet);
+}
+
+const TCHAR* FLowLevelMemTracker::FindTagName(uint64 Tag) const
+{
+	// This function does not call an Inner function on GetGlobals because it is a deprecated function.
+
+	if (!IsEnabled())
+	{
+		return nullptr;
+	}
+	// Cannot call BootstrapInitialise without shenanigans because this function is const.
+	LLMCheck(GetGlobals()->bInitializedTracking);
+
+	static TMap<uint64, FString> FoundTags;
+	FString* Cached = FoundTags.Find(Tag);
+	if (Cached)
+	{
+		return **Cached;
+	}
+
+	FName DisplayName = FindTagDisplayName(Tag);
+	if (DisplayName.IsNone())
+	{
+		return nullptr;
+	}
+
+	FString& AddedCache = FoundTags.Add(Tag);
+	AddedCache = DisplayName.ToString();
+	return *AddedCache;
+}
+
+FName FLowLevelMemTracker::FindTagDisplayName(uint64 Tag) const
+{
+	if (!IsEnabled())
+	{
+		return NAME_None;
+	}
+	return GetGlobals()->FindTagDisplayNameInner(Tag);
+}
+
+FName FLowLevelMemTracker::FindPtrDisplayName(void* Ptr) const
+{
+	if (!IsEnabled())
+	{
+		return NAME_None;
+	}
+	return GetGlobals()->FindPtrDisplayNameInner(Ptr);
+}
+
+FName FLowLevelMemTracker::GetTagDisplayName(const UE::LLMPrivate::FTagData* TagData) const
+{
+	return GetGlobals()->GetTagDisplayNameInner(TagData);
+}
+
+FString FLowLevelMemTracker::GetTagDisplayPathName(const UE::LLMPrivate::FTagData* TagData) const
+{
+	// IsEnabled is not checked here because Inner is valid to call even if not enabled.
+	return GetGlobals()->GetTagDisplayPathNameInner(TagData);
+}
+
+void FLowLevelMemTracker::GetTagDisplayPathName(const UE::LLMPrivate::FTagData* TagData,
+	FStringBuilderBase& OutPathName, int32 MaxLen) const
+{
+	// IsEnabled is not checked here because Inner is valid to call even if not enabled.
+	return GetGlobals()->GetTagDisplayPathNameInner(TagData, OutPathName, MaxLen);
+}
+
+FName FLowLevelMemTracker::GetTagUniqueName(const UE::LLMPrivate::FTagData* TagData) const
+{
+	// IsEnabled is not checked here because Inner is valid to call even if not enabled.
+	return GetGlobals()->GetTagUniqueNameInner(TagData);
+}
+
+const UE::LLMPrivate::FTagData* FLowLevelMemTracker::GetTagParent(const UE::LLMPrivate::FTagData* TagData) const
+{
+	// IsEnabled is not checked here because Inner is valid to call even if not enabled.
+	return GetGlobals()->GetTagParentInner(TagData);
+}
+
+bool FLowLevelMemTracker::GetTagIsEnumTag(const UE::LLMPrivate::FTagData* TagData) const
+{
+	// IsEnabled is not checked here because Inner is valid to call even if not enabled.
+	return GetGlobals()->GetTagIsEnumTagInner(TagData);
+}
+
+ELLMTag FLowLevelMemTracker::GetTagClosestEnumTag(const UE::LLMPrivate::FTagData* TagData) const
+{
+	// IsEnabled is not checked here because Inner is valid to call even if not enabled.
+	return GetGlobals()->GetTagClosestEnumTagInner(TagData);
+}
+
+int64 FLowLevelMemTracker::GetTagAmountForTracker(ELLMTracker Tracker, const UE::LLMPrivate::FTagData* TagData,
+	UE::LLM::ESizeParams SizeParams)
+{
+	if (!IsEnabled())
+	{
+		return 0;
+	}
+	return GetGlobals()->GetTagAmountForTrackerInner(Tracker, TagData, SizeParams);
+}
+
+int64 FLowLevelMemTracker::GetTagAmountForTracker(ELLMTracker Tracker, FName Tag, ELLMTagSet TagSet,
+	UE::LLM::ESizeParams SizeParams)
+{
+	if (!IsEnabled())
+	{
+		return 0;
+	}
+	return GetGlobals()->GetTagAmountForTrackerInner(Tracker, Tag, TagSet, SizeParams);
+}
+
+int64 FLowLevelMemTracker::GetTagAmountForTracker(ELLMTracker Tracker, ELLMTag Tag, UE::LLM::ESizeParams SizeParams)
+{
+	if (!IsEnabled())
+	{
+		return 0;
+	}
+	return GetGlobals()->GetTagAmountForTrackerInner(Tracker, Tag, SizeParams);
+}
+
+void FLowLevelMemTracker::SetTagAmountForTracker(ELLMTracker Tracker, ELLMTag Tag, int64 Amount, bool bAddToTotal)
+{
+	if (!IsEnabled())
+	{
+		return;
+	}
+	GetGlobals()->SetTagAmountForTrackerInner(Tracker, Tag, Amount, bAddToTotal);
+}
+
+void FLowLevelMemTracker::SetTagAmountForTracker(ELLMTracker Tracker, FName Tag, ELLMTagSet TagSet, int64 Amount, bool bAddToTotal)
+{
+	if (!IsEnabled())
+	{
+		return;
+	}
+	GetGlobals()->SetTagAmountForTrackerInner(Tracker, Tag, TagSet, Amount, bAddToTotal);
+}
+
+const UE::LLMPrivate::FTagData* FLowLevelMemTracker::GetActiveTagData(ELLMTracker Tracker, ELLMTagSet TagSet /*= ELLMTagSet::None*/)
+{
+	if (!IsEnabled())
+	{
+		return nullptr;
+	}
+	return GetGlobals()->GetActiveTagDataInner(Tracker, TagSet);
+}
+
+uint64 FLowLevelMemTracker::DumpTag(ELLMTracker Tracker, const char* FileName, int LineNumber)
+{
+	if (!IsEnabled())
+	{
+		return static_cast<int64>(ELLMTag::Untagged);
+	}
+	return GetGlobals()->DumpTagInner(Tracker, FileName, LineNumber);
+}
+
+void FLowLevelMemTracker::PublishDataSingleFrame()
+{
+	// IsEnabled is not checked here because Inner is valid to call even if not enabled.
+	GetGlobals()->PublishDataSingleFrameInner();
+}
+
+namespace UE::LLMPrivate
+{
+
+FLLMGlobals::FLLMGlobals()
 	: TagDatas(nullptr)
 	, TagDataNameMap(nullptr)
 	, TagDataEnumMap(nullptr)
@@ -1483,17 +1889,15 @@ FLowLevelMemTracker::FLowLevelMemTracker()
 	, bCanEnable(true)
 	, bCsvWriterEnabled(false)
 	, bTraceWriterEnabled(false)
-	, bInitialisedTracking(false)
+	, bInitializedTracking(false)
 	, bIsBootstrapping(false)
-	, bFullyInitialised(false)
+	, bFullyInitialized(false)
 	, bConfigurationComplete(false)
 	, bTagAdded(false)
 	, bAutoPublish(true)
 	, bPublishSingleFrame(false)
 	, bCapturedSizeSnapshot(false)
 {
-	using namespace UE::LLMPrivate;
-
 	// set the alloc functions
 	LLMAllocFunction PlatformLLMAlloc = NULL;
 	LLMFreeFunction PlatformLLMFree = NULL;
@@ -1510,7 +1914,7 @@ FLowLevelMemTracker::FLowLevelMemTracker()
 	}
 	LLMCheck(FMath::IsPowerOfTwo(Alignment));
 
-	Allocator.Initialise(PlatformLLMAlloc, PlatformLLMFree, Alignment);
+	Allocator.Initialize(PlatformLLMAlloc, PlatformLLMFree, Alignment);
 	FLLMAllocator::Get() = &Allocator;
 
 	// only None is on by default
@@ -1520,10 +1924,8 @@ FLowLevelMemTracker::FLowLevelMemTracker()
 	}
 }
 
-FLowLevelMemTracker::~FLowLevelMemTracker()
-{
-	using namespace UE::LLMPrivate;
-
+FLLMGlobals::~FLLMGlobals()
+{
 	EnabledState = EEnabled::Disabled;
 	PRAGMA_DISABLE_DEPRECATION_WARNINGS;
 	bIsDisabled = true;
@@ -1532,24 +1934,23 @@ FLowLevelMemTracker::~FLowLevelMemTracker()
 	FLLMAllocator::Get() = nullptr;
 }
 
-void FLowLevelMemTracker::BootstrapInitialise()
+void FLLMGlobals::BootstrapInitializeInner()
 {
-	using namespace UE::LLMPrivate;
-	if (bInitialisedTracking)
+	if (bInitializedTracking)
 	{
 		return;
 	}
 	UE::TUniqueLock BootstrapLock(BootstrapMutex);
-	if (bInitialisedTracking)
+	if (bInitializedTracking)
 	{
 		return;
 	}
 	// just make sure this part isn't reentrant.
-	if (bBeganInitialisedTracking != false)
+	if (bBeganInitializedTracking != false)
 	{
-		UE_LOG(LogInit, Fatal, TEXT("FLowLevelMemTracker::BootstrapInitialise is not reentrant"));
+		UE_LOG(LogInit, Fatal, TEXT("FLowLevelMemTracker::BootstrapInitialize is not reentrant"));
 	}
-	bBeganInitialisedTracking = true;
+	bBeganInitializedTracking = true;
 
 	for (int32 TrackerIndex = 0; TrackerIndex < static_cast<int32>(ELLMTracker::Max); ++TrackerIndex)
 	{
@@ -1557,12 +1958,12 @@ void FLowLevelMemTracker::BootstrapInitialise()
 
 		Trackers[TrackerIndex] = Tracker;
 
-		Tracker->Initialise(static_cast<ELLMTracker>(TrackerIndex), &Allocator);
+		Tracker->Initialize(static_cast<ELLMTracker>(TrackerIndex), &Allocator);
 	}
 
 	BootstrapTagDatas();
 	static_assert((uint8)ELLMTracker::Max == 2,
-		"You added a tracker, without updating FLowLevelMemTracker::BootstrapInitialise (and probably need to update macros)");
+		"You added a tracker, without updating FLowLevelMemTracker::BootstrapInitialize (and probably need to update macros)");
 	GetTracker(ELLMTracker::Platform)->SetTotalTags(FindOrAddTagData(ELLMTag::PlatformUntaggedTotal),
 		FindOrAddTagData(ELLMTag::PlatformTrackedTotal));
 	GetTracker(ELLMTracker::Default)->SetTotalTags(FindOrAddTagData(ELLMTag::UntaggedTotal),
@@ -1570,18 +1971,16 @@ void FLowLevelMemTracker::BootstrapInitialise()
 
 
 	// calculate program size early on... the platform can call SetProgramSize later if it sees fit
-	InitialiseProgramSize();
+	InitializeProgramSize();
 
 	FPlatformMisc::MemoryBarrier();
 
-	bInitialisedTracking = true;
+	bInitializedTracking = true;
 }
 
-void FLowLevelMemTracker::Clear()
+void FLLMGlobals::Clear()
 {
-	using namespace UE::LLMPrivate;
-
-	if (!bInitialisedTracking)
+	if (!bInitializedTracking)
 	{
 		return;
 	}
@@ -1602,24 +2001,20 @@ void FLowLevelMemTracker::Clear()
 
 	ClearTagDatas();
 	Allocator.Clear();
-	bFullyInitialised = false;
-	bInitialisedTracking = false;
+	bFullyInitialized = false;
+	bInitializedTracking = false;
 }
 
-void FLowLevelMemTracker::OnPreFork()
+void FLLMGlobals::OnPreForkInner()
 {
-	using namespace UE::LLMPrivate;
-	if (IsEnabled())
-	{
-		FLLMTracker& DefaultTracker = *GetTracker(ELLMTracker::Default);
-		FLLMTracker& PlatformTracker = *GetTracker(ELLMTracker::Platform);
+	FLLMTracker& DefaultTracker = *GetTracker(ELLMTracker::Default);
+	FLLMTracker& PlatformTracker = *GetTracker(ELLMTracker::Platform);
 
-		DefaultTracker.OnPreFork();
-		PlatformTracker.OnPreFork();
-	}
+	DefaultTracker.OnPreFork();
+	PlatformTracker.OnPreFork();
 }
 
-void FLowLevelMemTracker::UpdateStatsPerFrame(const TCHAR* LogName)
+void FLLMGlobals::UpdateStatsPerFrameInner(const TCHAR* LogName)
 {
 	CSV_SCOPED_TIMING_STAT_EXCLUSIVE(LLM);
 #if UE_ENABLE_ARRAY_SLACK_TRACKING
@@ -1657,7 +2052,7 @@ void FLowLevelMemTracker::UpdateStatsPerFrame(const TCHAR* LogName)
 		UE_LOG(LogInit, Log, TEXT("First time updating LLM stats..."));
 		bFirstTimeUpdating = false;
 	}
-	TickInternal();
+	TickInnerNoLock();
 
 	if (bAutoPublish || bPublishSingleFrame)
 	{
@@ -1666,23 +2061,17 @@ void FLowLevelMemTracker::UpdateStatsPerFrame(const TCHAR* LogName)
 	}
 }
 
-void FLowLevelMemTracker::Tick()
+void FLLMGlobals::TickInner()
 {
-	if (!IsEnabled())
-	{
-		return;
-	}
-	// TickInternal is usually called from the game thread, but can sometimes be called from the async loading thread,
+	// TickInnerNoLock is usually called from the game thread, but can sometimes be called from the async loading thread,
 	// so enter a lock for it.
 	UE::TUniqueLock UpdateLock(UpdateMutex);
-	TickInternal();
+	TickInnerNoLock();
 }
 
-void FLowLevelMemTracker::TickInternal()
+void FLLMGlobals::TickInnerNoLock()
 {
-	using namespace UE::LLMPrivate;
-
-	if (bFullyInitialised)
+	if (bFullyInitialized)
 	{
 		// We call tick when not fully initialised to get the overhead when disabled. When not initialised, we have to
 		// avoid the portion of the tick that uses tags.
@@ -1769,10 +2158,8 @@ void FLowLevelMemTracker::TickInternal()
 #endif
 }
 
-void FLowLevelMemTracker::UpdateTags()
+void FLLMGlobals::UpdateTags()
 {
-	using namespace UE::LLMPrivate;
-
 	if (!bTagAdded)
 	{
 		return;
@@ -1787,7 +2174,7 @@ void FLowLevelMemTracker::UpdateTags()
 		for (int32 TagIndex = 0; TagIndex < TagDatas->Num(); ++TagIndex)
 		{
 			FTagData* TagData = (*TagDatas)[TagIndex];
-			FinishConstruct(TagData, UE::LLMPrivate::ETagReferenceSource::FunctionAPI);
+			FinishConstruct(TagData, ETagReferenceSource::FunctionAPI);
 			const FTagData* Parent = TagData->GetParent();
 			if (Parent && Parent->GetIndex() > TagData->GetIndex())
 			{
@@ -1817,10 +2204,8 @@ void FLowLevelMemTracker::UpdateTags()
 	}
 }
 
-void FLowLevelMemTracker::SortTags(UE::LLMPrivate::FTagDataArray*& OutOldTagDatas)
+void FLLMGlobals::SortTags(FTagDataArray*& OutOldTagDatas)
 {
-	using namespace UE::LLMPrivate;
-
 	// Caller is responsible for holding a UniqueLock on TagDataMutex
 	OutOldTagDatas = TagDatas;
 	TagDatas = Allocator.New<FTagDataArray>();
@@ -1850,10 +2235,8 @@ void FLowLevelMemTracker::SortTags(UE::LLMPrivate::FTagDataArray*& OutOldTagData
 	}
 }
 
-void FLowLevelMemTracker::PublishDataPerFrame(const TCHAR* LogName)
+void FLLMGlobals::PublishDataPerFrame(const TCHAR* LogName)
 {
-	using namespace UE::LLMPrivate;
-
 	// set overhead stats
 	SET_MEMORY_STAT(STAT_LLMOverheadTotal, MemoryUsageCurrentOverhead);
 	if (IsEnabled())
@@ -1918,7 +2301,7 @@ void FLowLevelMemTracker::PublishDataPerFrame(const TCHAR* LogName)
 	}
 }
 
-void FLowLevelMemTracker::InitialiseProgramSize()
+void FLLMGlobals::InitializeProgramSize()
 {
 	// On Android we report ProgramSize via AndroidMemoryProbeTick
 #if !PLATFORM_ANDROID
@@ -1933,7 +2316,7 @@ void FLowLevelMemTracker::InitialiseProgramSize()
 #endif
 }
 
-void FLowLevelMemTracker::SetProgramSize(uint64 InProgramSize)
+void FLLMGlobals::SetProgramSizeInner(uint64 InProgramSize)
 {
 	if (!IsEnabled())
 	{
@@ -1949,7 +2332,7 @@ void FLowLevelMemTracker::SetProgramSize(uint64 InProgramSize)
 	GetTracker(ELLMTracker::Default)->TrackMemory(ELLMTag::ProgramSize, ProgramSizeDiff, ELLMAllocType::System);
 }
 
-void FLowLevelMemTracker::ProcessCommandLine(const TCHAR* CmdLine)
+void FLLMGlobals::ProcessCommandLineInner(const TCHAR* CmdLine)
 {
 	CSV_METADATA(TEXT("LLM"), TEXT("0"));
 
@@ -2013,7 +2396,7 @@ void FLowLevelMemTracker::ProcessCommandLine(const TCHAR* CmdLine)
 			// EnabledState is sufficient; this will cause us to wait until threads already in OnLowLevelAlloc exit the
 			// function and clear their readlock, and will cause new calls to those functions to block while we're waiting
 			// and call IsEnabled again after we release the lock.
-			FWriteScopeLock EnableStateLock(UE::LLMPrivate::GetEnableStateLock());
+			FWriteScopeLock EnableStateLock(GetEnableStateLock());
 #endif
 
 			EnabledState = EEnabled::Disabled;
@@ -2121,55 +2504,32 @@ void FLowLevelMemTracker::ProcessCommandLine(const TCHAR* CmdLine)
 	}
 }
 
-// Return the total amount of memory being tracked
-uint64 FLowLevelMemTracker::GetTotalTrackedMemory(ELLMTracker Tracker)
+uint64 FLLMGlobals::GetTotalTrackedMemoryInner(ELLMTracker Tracker)
 {
-	if (!IsEnabled())
-	{
-		return 0;
-	}
 	BootstrapInitialise();
 
 	return static_cast<uint64>(GetTracker(Tracker)->GetTrackedTotal());
 }
 
-void FLowLevelMemTracker::OnLowLevelAlloc(ELLMTracker Tracker, const void* Ptr, uint64 Size, ELLMTag DefaultTag,
+void FLLMGlobals::OnLowLevelAllocInner(ELLMTracker Tracker, const void* Ptr, uint64 Size, ELLMTag DefaultTag,
 	ELLMAllocType AllocType, bool bTrackInMemPro)
 {
-	FEnableStateScopeLock EnableScopeLock;
-	if (!TryEnterEnabled(EnableScopeLock))
-	{
-		return;
-	}
-
 	BootstrapInitialise();
 
 	GetTracker(Tracker)->TrackAllocation(Ptr, static_cast<int64>(Size), DefaultTag, AllocType, bTrackInMemPro);
 }
 
-void FLowLevelMemTracker::OnLowLevelAlloc(ELLMTracker Tracker, const void* Ptr, uint64 Size, FName DefaultTag,
+void FLLMGlobals::OnLowLevelAllocInner(ELLMTracker Tracker, const void* Ptr, uint64 Size, FName DefaultTag,
 	ELLMAllocType AllocType, bool bTrackInMemPro)
 {
-	FEnableStateScopeLock EnableScopeLock;
-	if (!TryEnterEnabled(EnableScopeLock))
-	{
-		return;
-	}
-
 	BootstrapInitialise();
 
 	GetTracker(Tracker)->TrackAllocation(Ptr, static_cast<int64>(Size), DefaultTag, AllocType, bTrackInMemPro);
 }
 
-void FLowLevelMemTracker::OnLowLevelFree(ELLMTracker Tracker, const void* Ptr,
+void FLLMGlobals::OnLowLevelFreeInner(ELLMTracker Tracker, const void* Ptr,
 	ELLMAllocType AllocType, bool bTrackInMemPro)
 {
-	FEnableStateScopeLock EnableScopeLock;
-	if (!TryEnterEnabled(EnableScopeLock))
-	{
-		return;
-	}
-
 	BootstrapInitialise();
 
 	if (Ptr != nullptr)
@@ -2178,69 +2538,41 @@ void FLowLevelMemTracker::OnLowLevelFree(ELLMTracker Tracker, const void* Ptr,
 	}
 }
 
-void FLowLevelMemTracker::OnLowLevelChangeInMemoryUse(ELLMTracker Tracker, int64 DeltaMemory, ELLMTag DefaultTag,
+void FLLMGlobals::OnLowLevelChangeInMemoryUseInner(ELLMTracker Tracker, int64 DeltaMemory, ELLMTag DefaultTag,
 	ELLMAllocType AllocType)
 {
-	FEnableStateScopeLock EnableScopeLock;
-	if (!TryEnterEnabled(EnableScopeLock))
-	{
-		return;
-	}
-
 	BootstrapInitialise();
 	GetTracker(Tracker)->TrackMemoryOfActiveTag(DeltaMemory, DefaultTag, AllocType);
 }
 
-void FLowLevelMemTracker::OnLowLevelChangeInMemoryUse(ELLMTracker Tracker, int64 DeltaMemory, FName DefaultTag,
+void FLLMGlobals::OnLowLevelChangeInMemoryUseInner(ELLMTracker Tracker, int64 DeltaMemory, FName DefaultTag,
 	ELLMAllocType AllocType)
 {
-	FEnableStateScopeLock EnableScopeLock;
-	if (!TryEnterEnabled(EnableScopeLock))
-	{
-		return;
-	}
-
 	BootstrapInitialise();
 	GetTracker(Tracker)->TrackMemoryOfActiveTag(DeltaMemory, DefaultTag, AllocType);
 }
 
-void FLowLevelMemTracker::OnLowLevelAllocMoved(ELLMTracker Tracker, const void* Dest, const void* Source,
+void FLLMGlobals::OnLowLevelAllocMovedInner(ELLMTracker Tracker, const void* Dest, const void* Source,
 	ELLMAllocType AllocType)
 {
-	FEnableStateScopeLock EnableScopeLock;
-	if (!TryEnterEnabled(EnableScopeLock))
-	{
-		return;
-	}
-
 	BootstrapInitialise();
 
 	//update the allocation map
 	GetTracker(Tracker)->OnAllocMoved(Dest, Source, AllocType);
 }
 
-UE::LLMPrivate::FLLMTracker* FLowLevelMemTracker::GetTracker(ELLMTracker Tracker)
+FLLMTracker* FLLMGlobals::GetTracker(ELLMTracker Tracker)
 {
 	return Trackers[static_cast<int32>(Tracker)];
 }
 
-const UE::LLMPrivate::FLLMTracker* FLowLevelMemTracker::GetTracker(ELLMTracker Tracker) const
+const FLLMTracker* FLLMGlobals::GetTracker(ELLMTracker Tracker) const
 {
 	return Trackers[static_cast<int32>(Tracker)];
 }
 
-bool FLowLevelMemTracker::Exec(const TCHAR* Cmd, FOutputDevice& Ar)
+bool FLLMGlobals::ExecInner(const TCHAR* Cmd, FOutputDevice& Ar)
 {
-	if (!IsEnabled())
-	{
-		if (FParse::Command(&Cmd, TEXT("LLMEM")))
-		{
-			UE_LOG(LogHAL, Warning, TEXT("Cannot execute LLMEM command! LLM is not enabled."));
-			return true;
-		}
-		return false;
-	}
-
 	if (FParse::Command(&Cmd, TEXT("LLMEM")))
 	{
 		BootstrapInitialise();
@@ -2326,24 +2658,16 @@ bool FLowLevelMemTracker::Exec(const TCHAR* Cmd, FOutputDevice& Ar)
 	return false;
 }
 
-bool FLowLevelMemTracker::IsTagSetActive(ELLMTagSet Set)
+bool FLLMGlobals::IsTagSetActiveInner(ELLMTagSet Set)
 {
-	if (!IsEnabled())
-	{
-		return false;
-	}
 	BootstrapInitialise();
 
 	return ActiveSets[static_cast<int32>(Set)];
 }
 
-bool FLowLevelMemTracker::ShouldReduceThreads()
+bool FLLMGlobals::ShouldReduceThreadsInner()
 {
 #if LLM_ENABLED_REDUCE_THREADS
-	if (!IsEnabled())
-	{
-		return false;
-	}
 	BootstrapInitialise();
 	LLMCheckf(bConfigurationComplete,
 		TEXT("ShouldReduceThreads has been called too early, before we processed the configuration settings required for it."));
@@ -2363,11 +2687,9 @@ static bool IsAssetTagForAssets(ELLMTagSet Set)
 
 #endif // LLM_ALLOW_ASSETS_TAGS
 
-void FLowLevelMemTracker::RegisterCustomTagInternal(int32 Tag, ELLMTagSet TagSet, const TCHAR* InDisplayName, FName StatName,
+void FLLMGlobals::RegisterCustomTag(int32 Tag, ELLMTagSet TagSet, const TCHAR* InDisplayName, FName StatName,
 	FName SummaryStatName, int32 ParentTag)
 {
-	using namespace UE::LLMPrivate;
-
 	LLMCheckf(Tag >= LLM_CUSTOM_TAG_START && Tag <= LLM_CUSTOM_TAG_END, TEXT("Tag %d out of range"), Tag);
 	LLMCheckf(InDisplayName != nullptr, TEXT("Tag %d has no name"), Tag);
 	LLMCheckf(ParentTag == -1 || ParentTag < LLM_TAG_COUNT, TEXT("Parent tag %d out of range"), ParentTag);
@@ -2380,11 +2702,12 @@ void FLowLevelMemTracker::RegisterCustomTagInternal(int32 Tag, ELLMTagSet TagSet
 		EnumTag, false, ETagReferenceSource::CustomEnumTag, TagSet);
 }
 
-void FLowLevelMemTracker::RegisterPlatformTag(int32 Tag, const TCHAR* Name, FName StatName, FName SummaryStatName,
+void FLLMGlobals::RegisterPlatformTagInner(int32 Tag, const TCHAR* Name, FName StatName, FName SummaryStatName,
 	int32 ParentTag)
 {
 	MemoryTrace_AnnounceCustomTag(Tag, ParentTag, Name);
 
+	// IsEnabled is checked in this Inner function because there is some work to do even if not enabled.
 	if (!IsEnabled())
 	{
 		return;
@@ -2393,14 +2716,15 @@ void FLowLevelMemTracker::RegisterPlatformTag(int32 Tag, const TCHAR* Name, FNam
 
 	LLMCheck(Tag >= static_cast<int32>(ELLMTag::PlatformTagStart) &&
 		Tag <= static_cast<int32>(ELLMTag::PlatformTagEnd));
-	RegisterCustomTagInternal(Tag, ELLMTagSet::None, Name, StatName, SummaryStatName, ParentTag);
+	RegisterCustomTag(Tag, ELLMTagSet::None, Name, StatName, SummaryStatName, ParentTag);
 }
 
-void FLowLevelMemTracker::RegisterProjectTag(int32 Tag, const TCHAR* Name, FName StatName, FName SummaryStatName,
+void FLLMGlobals::RegisterProjectTagInner(int32 Tag, const TCHAR* Name, FName StatName, FName SummaryStatName,
 	int32 ParentTag)
 {
 	MemoryTrace_AnnounceCustomTag(Tag, ParentTag, Name);
 
+	// IsEnabled is checked in this Inner function because there is some work to do even if not enabled.
 	if (!IsEnabled())
 	{
 		return;
@@ -2408,22 +2732,11 @@ void FLowLevelMemTracker::RegisterProjectTag(int32 Tag, const TCHAR* Name, FName
 	BootstrapInitialise();
 
 	LLMCheck(Tag >= static_cast<int32>(ELLMTag::ProjectTagStart) && Tag <= static_cast<int32>(ELLMTag::ProjectTagEnd));
-	RegisterCustomTagInternal(Tag, ELLMTagSet::None, Name, StatName, SummaryStatName, ParentTag);
-}
-
-void GlobalRegisterTagDeclaration(FLLMTagDeclaration& TagDeclaration)
-{
-	if (!FLowLevelMemTracker::IsEnabled())
-	{
-		return;
-	}
-	FLowLevelMemTracker::Get().RegisterTagDeclaration(TagDeclaration);
+	RegisterCustomTag(Tag, ELLMTagSet::None, Name, StatName, SummaryStatName, ParentTag);
 }
 
-void FLowLevelMemTracker::BootstrapTagDatas()
+void FLLMGlobals::BootstrapTagDatas()
 {
-	using namespace UE::LLMPrivate;
-
 	// While bootstrapping we are not allowed to construct any FNames because the FName system may not yet have been
 	// constructed. Construct not-fully-initialized TagDatas for the central list of ELLMTags.
 	{
@@ -2458,9 +2771,9 @@ void FLowLevelMemTracker::BootstrapTagDatas()
 	}
 }
 
-void FLowLevelMemTracker::FinishInitialise()
+void FLLMGlobals::FinishInitializeInner()
 {
-	if (bFullyInitialised)
+	if (bFullyInitialized)
 	{
 		return;
 	}
@@ -2469,25 +2782,23 @@ void FLowLevelMemTracker::FinishInitialise()
 #if UE_ENABLE_ARRAY_SLACK_TRACKING
 	GTrackArrayDetailedLock = new FCriticalSection();
 #endif
-	bFullyInitialised = true;
-	// Make sure that FNames and Malloc have already been initialised, since we will use them during InitialiseTagDatas
+	bFullyInitialized = true;
+	// Make sure that FNames and Malloc have already been initialised, since we will use them during InitializeTagDatas
 	// We force this by calling LLMGetTagUniqueName, which initializes FNames internally, and will therein trigger
 	// FName system construction, which will itself trigger Malloc construction.
 	(void)LLMGetTagUniqueName(ELLMTag::Untagged);
-	InitialiseTagDatas();
+	InitializeTagDatas();
 
-	UE::LLMPrivate::FInitializedCallbacksArray& Callbacks = UE::LLMPrivate::GetInitialisedCallbacks();
-	for (const TPair<UE::LLMPrivate::FLLMInitialisedCallback, UPTRINT>& Callback : Callbacks)
+	FInitializedCallbacksArray& Callbacks = GetInitializedCallbacks();
+	for (const TPair<FLLMInitialisedCallback, UPTRINT>& Callback : Callbacks)
 	{
 		Callback.Key(Callback.Value);
 	}
 	Callbacks.Empty();
 }
 
-void FLowLevelMemTracker::InitialiseTagDatas_SetLLMTagNames()
+void FLLMGlobals::InitializeTagDatas_SetLLMTagNames()
 {
-	using namespace UE::LLMPrivate;
-
 	TStringBuilder<256> NameBuffer;
 	// Load all the names for the central list of ELLMTags (recording the allocations the FName system makes for the
 	// construction of the names).
@@ -2510,10 +2821,8 @@ void FLowLevelMemTracker::InitialiseTagDatas_SetLLMTagNames()
 #undef SET_ELLMTAG_NAMES
 }
 
-void FLowLevelMemTracker::InitialiseTagDatas_FinishRegister()
+void FLLMGlobals::InitializeTagDatas_FinishRegister()
 {
-	using namespace UE::LLMPrivate;
-
 	// Record the central list of ELLMTags in TagDataNameMap, and mark that bootstrapping is complete
 	{
 		UE::TUniqueLock Lock(TagDataMutex);
@@ -2537,13 +2846,11 @@ void FLowLevelMemTracker::InitialiseTagDatas_FinishRegister()
 	}
 }
 
-void FLowLevelMemTracker::InitialiseTagDatas()
+void FLLMGlobals::InitializeTagDatas()
 {
-	using namespace UE::LLMPrivate;
+	InitializeTagDatas_SetLLMTagNames();
 
-	InitialiseTagDatas_SetLLMTagNames();
-
-	InitialiseTagDatas_FinishRegister();
+	InitializeTagDatas_FinishRegister();
 
 	// Construct the remaining startup tags; allocations when constructing these tags are known to consist only of the
 	// central list of ELLMTags so we do not need to bootstrap.
@@ -2567,10 +2874,8 @@ void FLowLevelMemTracker::InitialiseTagDatas()
 	UpdateTags();
 }
 
-void FLowLevelMemTracker::ClearTagDatas()
+void FLLMGlobals::ClearTagDatas()
 {
-	using namespace UE::LLMPrivate;
-
 	UE::TUniqueLock Lock(TagDataMutex);
 	FLLMTagDeclaration::ClearCreationCallbacks();
 
@@ -2586,23 +2891,21 @@ void FLowLevelMemTracker::ClearTagDatas()
 	TagDatas = nullptr;
 }
 
-void FLowLevelMemTracker::RegisterTagDeclaration(FLLMTagDeclaration& TagDeclaration)
+void FLLMGlobals::RegisterTagDeclaration(FLLMTagDeclaration& TagDeclaration)
 {
 	TagDeclaration.ConstructUniqueName();
 	RegisterTagData(TagDeclaration.UniqueName, TagDeclaration.DisplayName, TagDeclaration.ParentTagName,
 		TagDeclaration.StatName, TagDeclaration.SummaryStatName, false, ELLMTag::CustomName, false,
-		UE::LLMPrivate::ETagReferenceSource::Declare, TagDeclaration.TagSet);
+		ETagReferenceSource::Declare, TagDeclaration.TagSet);
 }
 
-UE::LLMPrivate::FTagData& FLowLevelMemTracker::RegisterTagData(FName Name, FName DisplayName, FName ParentName,
+FTagData& FLLMGlobals::RegisterTagData(FName Name, FName DisplayName, FName ParentName,
 	FName StatName, FName SummaryStatName, bool bHasEnumTag, ELLMTag EnumTag, bool bIsStatTag,
-	UE::LLMPrivate::ETagReferenceSource ReferenceSource, ELLMTagSet TagSet/* = ELLMTagSet::None*/)
+	ETagReferenceSource ReferenceSource, ELLMTagSet TagSet/* = ELLMTagSet::None*/)
 {
 	LLMCheckf(!bIsBootstrapping,
 		TEXT("A tag outside of LLM_ENUM_GENERIC_TAGS was requested from LLM_SCOPE or allocation while bootstrapping the names for LLM_ENUM_GENERIC_TAGS, this is not supported."));
 
-	using namespace UE::LLMPrivate;
-
 	// If this allocates, that is okay. Set it to something small-as-possible-to-avoid-normally-allocating to prevent
 	// adding a lot of stack space in the calling LLM_SCOPE code.
 	TStringBuilder<256> NameBuffer;
@@ -2695,9 +2998,9 @@ UE::LLMPrivate::FTagData& FLowLevelMemTracker::RegisterTagData(FName Name, FName
 			bHasEnumTag != TagDataForName->HasEnumTag() ||
 			(bHasEnumTag && EnumTag != TagDataForName->GetEnumTag()))
 		{
-			if (ReferenceSource != UE::LLMPrivate::ETagReferenceSource::Scope &&
-				ReferenceSource != UE::LLMPrivate::ETagReferenceSource::ImplicitParent &&
-				ReferenceSource != UE::LLMPrivate::ETagReferenceSource::FunctionAPI)
+			if (ReferenceSource != ETagReferenceSource::Scope &&
+				ReferenceSource != ETagReferenceSource::ImplicitParent &&
+				ReferenceSource != ETagReferenceSource::FunctionAPI)
 			{
 				ReportDuplicateTagName(TagDataForName, ReferenceSource);
 			}
@@ -2749,11 +3052,8 @@ UE::LLMPrivate::FTagData& FLowLevelMemTracker::RegisterTagData(FName Name, FName
 	return *TagData;
 }
 
-void FLowLevelMemTracker::ReportDuplicateTagName(UE::LLMPrivate::FTagData* TagData,
-	UE::LLMPrivate::ETagReferenceSource ReferenceSource)
+void FLLMGlobals::ReportDuplicateTagName(FTagData* TagData, ETagReferenceSource ReferenceSource)
 {
-	using namespace UE::LLMPrivate;
-
 	// We're inside the LLM lock, so do not allow these GetName calls to allocate memory from FMalloc.
 	if (ReferenceSource == ETagReferenceSource::FunctionAPI || ReferenceSource == ETagReferenceSource::Scope ||
 		ReferenceSource == ETagReferenceSource::ImplicitParent)
@@ -2785,10 +3085,8 @@ void FLowLevelMemTracker::ReportDuplicateTagName(UE::LLMPrivate::FTagData* TagDa
 	}
 }
 
-void FLowLevelMemTracker::FinishConstruct(UE::LLMPrivate::FTagData* TagData,
-	UE::LLMPrivate::ETagReferenceSource ReferenceSource)
+void FLLMGlobals::FinishConstruct(FTagData* TagData, ETagReferenceSource ReferenceSource)
 {
-	using namespace UE::LLMPrivate;
 	// Caller is responsible for holding a SharedLock (NOT a UniqueLock) on TagDataMutex
 
 	LLMCheck(TagData);
@@ -2855,8 +3153,7 @@ void FLowLevelMemTracker::FinishConstruct(UE::LLMPrivate::FTagData* TagData,
 		// code that calls back into LLM. The TagData pointer is immutable so we do not have to worry about
 		// it disappearing out from under us.
 		TagDataMutex.UnlockShared();
-		for (const TPair<UE::LLMPrivate::FTagCreationCallback, UPTRINT>& Callback :
-			UE::LLMPrivate::GetTagCreationCallbacks())
+		for (const TPair<FTagCreationCallback, UPTRINT>& Callback : GetTagCreationCallbacks())
 		{
 			Callback.Key(TagData, Callback.Value);
 		}
@@ -2864,15 +3161,8 @@ void FLowLevelMemTracker::FinishConstruct(UE::LLMPrivate::FTagData* TagData,
 	}
 }
 
-TArray<const UE::LLMPrivate::FTagData*> FLowLevelMemTracker::GetTrackedTags(ELLMTagSet TagSet)
+TArray<const FTagData*> FLLMGlobals::GetTrackedTagsInner(ELLMTagSet TagSet)
 {
-	using namespace UE::LLMPrivate;
-
-	if (!IsEnabled())
-	{
-		return TArray<const FTagData*>();
-	}
-
 	BootstrapInitialise();
 
 	UE::TSharedLock Lock(TagDataMutex);
@@ -2885,61 +3175,34 @@ TArray<const UE::LLMPrivate::FTagData*> FLowLevelMemTracker::GetTrackedTags(ELLM
 	return FoundResults;
 }
 
-TArray<const UE::LLMPrivate::FTagData*> FLowLevelMemTracker::GetTrackedTags(ELLMTracker Tracker, ELLMTagSet TagSet /* = ELLMTagSet::None */)
+TArray<const FTagData*> FLLMGlobals::GetTrackedTagsInner(ELLMTracker Tracker, ELLMTagSet TagSet /* = ELLMTagSet::None */)
 {
-	using namespace UE::LLMPrivate;
-
-	if (!IsEnabled())
-	{
-		return TArray<const FTagData*>();
-	}
-
 	BootstrapInitialise();
 
 	UE::TUniqueLock UpdateLock(UpdateMutex); // uses of TagSizes are guarded by the UpdateLock
 	return GetTracker(Tracker)->GetTagDatas(TagSet);
 }
 
-void FLowLevelMemTracker::GetTrackedTagsNamesWithAmount(TMap<FName, uint64>& TagsNamesWithAmount, ELLMTracker Tracker, ELLMTagSet TagSet)
+void FLLMGlobals::GetTrackedTagsNamesWithAmountInner(TMap<FName, uint64>& TagsNamesWithAmount, ELLMTracker Tracker, ELLMTagSet TagSet)
 {
-	using namespace UE::LLMPrivate;
-
-	if (!IsEnabled())
-	{
-		return;
-	}
-
 	BootstrapInitialise();
 
 	UE::TUniqueLock UpdateLock(UpdateMutex);
 	GetTracker(Tracker)->GetTagsNamesWithAmount(TagsNamesWithAmount, TagSet);
 }
 
-void FLowLevelMemTracker::GetTrackedTagsNamesWithAmountFiltered(TMap<FName, uint64>& TagsNamesWithAmount, ELLMTracker Tracker, ELLMTagSet TagSet, TArray<FLLMTagSetAllocationFilter>& Filters)
+void FLLMGlobals::GetTrackedTagsNamesWithAmountFilteredInner(TMap<FName, uint64>& TagsNamesWithAmount, ELLMTracker Tracker, ELLMTagSet TagSet, TArray<FLLMTagSetAllocationFilter>& Filters)
 {
-	using namespace UE::LLMPrivate;
-
-	if (!IsEnabled())
-	{
-		return;
-	}
-
 	BootstrapInitialise();
 
 	UE::TUniqueLock UpdateLock(UpdateMutex);
 	GetTracker(Tracker)->GetTagsNamesWithAmountFiltered(TagsNamesWithAmount, TagSet, Filters);
 }
 
-bool FLowLevelMemTracker::FindTagByName( const TCHAR* Name, uint64& OutTag, ELLMTagSet InTagSet /*= ELLMTagSet::None*/ ) const
+bool FLLMGlobals::FindTagByNameInner(const TCHAR* Name, uint64& OutTag, ELLMTagSet InTagSet /*= ELLMTagSet::None*/) const
 {
-	using namespace UE::LLMPrivate;
-
-	if (!IsEnabled())
-	{
-		return false;
-	}
 	// Cannot call BootstrapInitialise and FinishInitialise without shenanigans because this function is const.
-	LLMCheck(bFullyInitialised);
+	LLMCheck(bFullyInitialized);
 
 	if (Name != nullptr)
 	{
@@ -2974,43 +3237,10 @@ bool FLowLevelMemTracker::FindTagByName( const TCHAR* Name, uint64& OutTag, ELLM
 	return false;
 }
 
-const TCHAR* FLowLevelMemTracker::FindTagName(uint64 Tag) const
-{
-	if (!IsEnabled())
-	{
-		return nullptr;
-	}
-	// Cannot call BootstrapInitialise without shenanigans because this function is const.
-	LLMCheck(bInitialisedTracking);
-
-	static TMap<uint64, FString> FoundTags;
-	FString* Cached = FoundTags.Find(Tag);
-	if (Cached)
-	{
-		return **Cached;
-	}
-
-	FName DisplayName = FindTagDisplayName(Tag);
-	if (DisplayName.IsNone())
-	{
-		return nullptr;
-	}
-
-	FString& AddedCache = FoundTags.Add(Tag);
-	AddedCache = DisplayName.ToString();
-	return *AddedCache;
-}
-
-FName FLowLevelMemTracker::FindTagDisplayName(uint64 Tag) const
+FName FLLMGlobals::FindTagDisplayNameInner(uint64 Tag) const
 {
-	using namespace UE::LLMPrivate;
-
-	if (!IsEnabled())
-	{
-		return NAME_None;
-	}
 	// Cannot call BootstrapInitialise without shenanigans because this function is const.
-	LLMCheck(bInitialisedTracking);
+	LLMCheck(bInitializedTracking);
 
 	int32 Index = static_cast<int32>(Tag);
 	if (0 <= Index && Index < LLM_TAG_COUNT)
@@ -3024,11 +3254,9 @@ FName FLowLevelMemTracker::FindTagDisplayName(uint64 Tag) const
 	return NAME_None;
 }
 
-FName FLowLevelMemTracker::FindPtrDisplayName(void* Ptr) const
+FName FLLMGlobals::FindPtrDisplayNameInner(void* Ptr) const
 {
-	using namespace UE::LLMPrivate;
-
-	if (!IsEnabled() || !bFullyInitialised)
+	if (!bFullyInitialized)
 	{
 		return NAME_None;
 	}
@@ -3040,54 +3268,47 @@ FName FLowLevelMemTracker::FindPtrDisplayName(void* Ptr) const
 	return Tags.Num() ? Tags[0]->GetDisplayName() : NAME_None;
 }
 
-FName FLowLevelMemTracker::GetTagDisplayName(const UE::LLMPrivate::FTagData* TagData) const
+FName FLLMGlobals::GetTagDisplayNameInner(const FTagData* TagData) const
 {
 	return TagData->GetDisplayName();
 }
 
-FString FLowLevelMemTracker::GetTagDisplayPathName(const UE::LLMPrivate::FTagData* TagData) const
+FString FLLMGlobals::GetTagDisplayPathNameInner(const FTagData * TagData) const
 {
 	TStringBuilder<FName::StringBufferSize> Buffer;
 	GetTagDisplayPathName(TagData, Buffer);
 	return FString(Buffer);
 }
 
-void FLowLevelMemTracker::GetTagDisplayPathName(const UE::LLMPrivate::FTagData* TagData,
+void FLLMGlobals::GetTagDisplayPathNameInner(const FTagData* TagData,
 	FStringBuilderBase& OutPathName, int32 MaxLen) const
 {
 	TagData->GetDisplayPath(OutPathName, MaxLen);
 }
 
-FName FLowLevelMemTracker::GetTagUniqueName(const UE::LLMPrivate::FTagData* TagData) const
+FName FLLMGlobals::GetTagUniqueNameInner(const FTagData* TagData) const
 {
 	return TagData->GetName();
 }
 
-const UE::LLMPrivate::FTagData* FLowLevelMemTracker::GetTagParent(const UE::LLMPrivate::FTagData* TagData) const
+const FTagData* FLLMGlobals::GetTagParentInner(const FTagData* TagData) const
 {
 	return TagData->GetParent();
 }
 
-bool FLowLevelMemTracker::GetTagIsEnumTag(const UE::LLMPrivate::FTagData* TagData) const
+bool FLLMGlobals::GetTagIsEnumTagInner(const FTagData* TagData) const
 {
 	return TagData->HasEnumTag();
 }
 
-ELLMTag FLowLevelMemTracker::GetTagClosestEnumTag(const UE::LLMPrivate::FTagData* TagData) const
+ELLMTag FLLMGlobals::GetTagClosestEnumTagInner(const FTagData* TagData) const
 {
 	return TagData->GetEnumTag();
 }
 
-int64 FLowLevelMemTracker::GetTagAmountForTracker(ELLMTracker Tracker, const UE::LLMPrivate::FTagData* TagData,
+int64 FLLMGlobals::GetTagAmountForTrackerInner(ELLMTracker Tracker, const FTagData* TagData,
 	UE::LLM::ESizeParams SizeParams)
 {
-	using namespace UE::LLMPrivate;
-
-	if (!IsEnabled())
-	{
-		return 0;
-	}
-
 	BootstrapInitialise();
 
 	if (TagData == nullptr)
@@ -3099,39 +3320,23 @@ int64 FLowLevelMemTracker::GetTagAmountForTracker(ELLMTracker Tracker, const UE:
 	return GetTracker(Tracker)->GetTagAmount(TagData, SizeParams);
 }
 
-int64 FLowLevelMemTracker::GetTagAmountForTracker(ELLMTracker Tracker, FName Tag, ELLMTagSet TagSet,
+int64 FLLMGlobals::GetTagAmountForTrackerInner(ELLMTracker Tracker, FName Tag, ELLMTagSet TagSet,
 	UE::LLM::ESizeParams SizeParams)
 {
-	if (!IsEnabled())
-	{
-		return 0;
-	}
-
 	BootstrapInitialise();
 
-	return GetTagAmountForTracker(Tracker, FindTagData(Tag, TagSet, UE::LLMPrivate::ETagReferenceSource::FunctionAPI), SizeParams);
+	return GetTagAmountForTracker(Tracker, FindTagData(Tag, TagSet, ETagReferenceSource::FunctionAPI), SizeParams);
 }
 
-int64 FLowLevelMemTracker::GetTagAmountForTracker(ELLMTracker Tracker, ELLMTag Tag, UE::LLM::ESizeParams SizeParams)
+int64 FLLMGlobals::GetTagAmountForTrackerInner(ELLMTracker Tracker, ELLMTag Tag, UE::LLM::ESizeParams SizeParams)
 {
-	if (!IsEnabled())
-	{
-		return 0;
-	}
-
 	BootstrapInitialise();
 
 	return GetTagAmountForTracker(Tracker, FindTagData(Tag), SizeParams);
 }
 
-void FLowLevelMemTracker::SetTagAmountForTracker(ELLMTracker Tracker, ELLMTag Tag, int64 Amount, bool bAddToTotal)
+void FLLMGlobals::SetTagAmountForTrackerInner(ELLMTracker Tracker, ELLMTag Tag, int64 Amount, bool bAddToTotal)
 {
-	using namespace UE::LLMPrivate;
-
-	if (!IsEnabled())
-	{
-		return;
-	}
 	BootstrapInitialise();
 	const FTagData* TagData = FindOrAddTagData(Tag);
 
@@ -3139,14 +3344,8 @@ void FLowLevelMemTracker::SetTagAmountForTracker(ELLMTracker Tracker, ELLMTag Ta
 	GetTracker(Tracker)->SetTagAmountExternal(TagData, Amount, bAddToTotal);
 }
 
-void FLowLevelMemTracker::SetTagAmountForTracker(ELLMTracker Tracker, FName Tag, ELLMTagSet TagSet, int64 Amount, bool bAddToTotal)
+void FLLMGlobals::SetTagAmountForTrackerInner(ELLMTracker Tracker, FName Tag, ELLMTagSet TagSet, int64 Amount, bool bAddToTotal)
 {
-	using namespace UE::LLMPrivate;
-
-	if (!IsEnabled())
-	{
-		return;
-	}
 	BootstrapInitialise();
 	const FTagData* TagData = FindOrAddTagData(Tag, TagSet);
 
@@ -3154,24 +3353,15 @@ void FLowLevelMemTracker::SetTagAmountForTracker(ELLMTracker Tracker, FName Tag,
 	GetTracker(Tracker)->SetTagAmountExternal(TagData, Amount, bAddToTotal);
 }
 
-const UE::LLMPrivate::FTagData* FLowLevelMemTracker::GetActiveTagData(ELLMTracker Tracker, ELLMTagSet TagSet /*= ELLMTagSet::None*/)
+const FTagData* FLLMGlobals::GetActiveTagDataInner(ELLMTracker Tracker, ELLMTagSet TagSet /*= ELLMTagSet::None*/)
 {
-	if (!IsEnabled())
-	{
-		return nullptr;
-	}
 	BootstrapInitialise();
 
 	return GetTracker(Tracker)->GetActiveTagData(TagSet);
 }
 
-uint64 FLowLevelMemTracker::DumpTag( ELLMTracker Tracker, const char* FileName, int LineNumber )
+uint64 FLLMGlobals::DumpTagInner( ELLMTracker Tracker, const char* FileName, int LineNumber )
 {
-	using namespace UE::LLMPrivate;
-	if (!IsEnabled())
-	{
-		return static_cast<int64>(ELLMTag::Untagged);
-	}
 	BootstrapInitialise();
 
 	const FTagData* TagData = GetActiveTagData(Tracker);
@@ -3189,7 +3379,7 @@ uint64 FLowLevelMemTracker::DumpTag( ELLMTracker Tracker, const char* FileName,
 	}
 }
 
-void FLowLevelMemTracker::PublishDataSingleFrame()
+void FLLMGlobals::PublishDataSingleFrameInner()
 {
 	if (bAutoPublish)
 	{
@@ -3201,11 +3391,8 @@ void FLowLevelMemTracker::PublishDataSingleFrame()
 	}
 }
 
-const UE::LLMPrivate::FTagData* FLowLevelMemTracker::FindOrAddTagData(ELLMTag EnumTag,
-	UE::LLMPrivate::ETagReferenceSource ReferenceSource)
+const FTagData* FLLMGlobals::FindOrAddTagData(ELLMTag EnumTag, ETagReferenceSource ReferenceSource)
 {
-	using namespace UE::LLMPrivate;
-
 	int32 Index = static_cast<int32>(EnumTag);
 	LLMCheckf(0 <= Index && Index < LLM_TAG_COUNT, TEXT("Out of range ELLMTag %d"), Index);
 
@@ -3220,11 +3407,11 @@ const UE::LLMPrivate::FTagData* FLowLevelMemTracker::FindOrAddTagData(ELLMTag En
 	}
 
 	// If we have not initialized tags yet initialize now to potentially create the custom ELLMTag we are reading.
-	if (!bFullyInitialised)
+	if (!bFullyInitialized)
 	{
 		FinishInitialise();
 		// Reenter this function so that we retry the find above; we avoid infinite recursion because
-		// bFullyInitialised is now true.
+		// bFullyInitialized is now true.
 		return FindOrAddTagData(EnumTag, ReferenceSource);
 	}
 	LLMCheckf(!bIsBootstrapping, TEXT("LLM Error: Invalid use of custom ELLMTag when initialising tags."));
@@ -3240,17 +3427,15 @@ const UE::LLMPrivate::FTagData* FLowLevelMemTracker::FindOrAddTagData(ELLMTag En
 	}
 }
 
-const UE::LLMPrivate::FTagData* FLowLevelMemTracker::FindOrAddTagData(FName TagName, ELLMTagSet TagSet, bool bIsStatTag,
-	UE::LLMPrivate::ETagReferenceSource ReferenceSource)
+const FTagData* FLLMGlobals::FindOrAddTagData(FName TagName, ELLMTagSet TagSet, bool bIsStatTag,
+	ETagReferenceSource ReferenceSource)
 {
 	return FindOrAddTagData(TagName, TagSet, bIsStatTag ? TagName : NAME_None, ReferenceSource);
 }
 
-const UE::LLMPrivate::FTagData* FLowLevelMemTracker::FindOrAddTagData(FName TagName, ELLMTagSet TagSet, FName StatName,
-	UE::LLMPrivate::ETagReferenceSource ReferenceSource)
+const FTagData* FLLMGlobals::FindOrAddTagData(FName TagName, ELLMTagSet TagSet, FName StatName,
+	ETagReferenceSource ReferenceSource)
 {
-	using namespace UE::LLMPrivate;
-
 	{
 		UE::TSharedLock Lock(TagDataMutex);
 		FTagData** TagDataPtr = TagDataNameMap->Find(FTagDataNameKey(TagName, TagSet));
@@ -3263,11 +3448,11 @@ const UE::LLMPrivate::FTagData* FLowLevelMemTracker::FindOrAddTagData(FName TagN
 	}
 
 	// If we have not initialized tags yet initialize now to potentially create the custom ELLMTag we are reading.
-	if (!bFullyInitialised)
+	if (!bFullyInitialized)
 	{
 		FinishInitialise();
 		// Reenter this function so that we retry the find above; note we avoid infinite recursion because
-		// bFullyInitialised is now true.
+		// bFullyInitialized is now true.
 		return FindOrAddTagData(TagName, TagSet, StatName, ReferenceSource);
 	}
 	LLMCheckf(!bIsBootstrapping, TEXT("LLM Error: Invalid use of FName tag when initialising tags."));
@@ -3283,11 +3468,8 @@ const UE::LLMPrivate::FTagData* FLowLevelMemTracker::FindOrAddTagData(FName TagN
 	}
 }
 
-const UE::LLMPrivate::FTagData* FLowLevelMemTracker::FindTagData(ELLMTag EnumTag,
-	UE::LLMPrivate::ETagReferenceSource ReferenceSource)
+const FTagData* FLLMGlobals::FindTagData(ELLMTag EnumTag, ETagReferenceSource ReferenceSource)
 {
-	using namespace UE::LLMPrivate;
-
 	int32 Index = static_cast<int32>(EnumTag);
 	LLMCheckf(0 <= Index && Index < LLM_TAG_COUNT, TEXT("Out of range ELLMTag %d"), Index);
 
@@ -3304,11 +3486,8 @@ const UE::LLMPrivate::FTagData* FLowLevelMemTracker::FindTagData(ELLMTag EnumTag
 	}
 }
 
-const UE::LLMPrivate::FTagData* FLowLevelMemTracker::FindTagData(FName TagName, ELLMTagSet TagSet,
-	UE::LLMPrivate::ETagReferenceSource ReferenceSource)
+const FTagData* FLLMGlobals::FindTagData(FName TagName, ELLMTagSet TagSet, ETagReferenceSource ReferenceSource)
 {
-	using namespace UE::LLMPrivate;
-
 	UE::TSharedLock Lock(TagDataMutex);
 
 	FTagData** TagDataPtr = TagDataNameMap->Find(FTagDataNameKey(TagName, TagSet));
@@ -3324,6 +3503,8 @@ const UE::LLMPrivate::FTagData* FLowLevelMemTracker::FindTagData(FName TagName,
 	}
 }
 
+} // namespace UE::LLMPrivate
+
 FLLMScope::FLLMScope(ELLMTag TagEnum, bool bInIsStatTag, ELLMTagSet InTagSet, ELLMTracker InTracker, bool bOverride)
 {
 	if (!FLowLevelMemTracker::IsEnabled())
@@ -3340,7 +3521,7 @@ FLLMScope::FLLMScope(ELLMTag TagEnum, bool bInIsStatTag, ELLMTagSet InTagSet, EL
 		// ELLMTag::FMalloc is a special tag expected to only be used by the Platform Tracker (see header where defined).
 		LLMCheck((TagEnum != ELLMTag::FMalloc) || (ELLMTracker::Platform == InTracker));
 
-		FLowLevelMemTracker& LLMRef = FLowLevelMemTracker::Get();
+		UE::LLMPrivate::FLLMGlobals& LLMRef = UE::LLMPrivate::FLLMGlobals::GetInner();
 		// We have to check IsEnabled() again after calling Get, because the constructor is called
 		// from Get, and will set EnabledState=Disabled if the platform doesn't support it.
 		if (FLowLevelMemTracker::IsEnabled())
@@ -3371,7 +3552,7 @@ FLLMScope::FLLMScope(FName TagName, bool bInIsStatTag, ELLMTagSet InTagSet, ELLM
 	// We run the init in the open, because we want to track LLM even in transactions.
 	UE_AUTORTFM_OPEN
 	{
-		FLowLevelMemTracker& LLMRef = FLowLevelMemTracker::Get();
+		UE::LLMPrivate::FLLMGlobals& LLMRef = UE::LLMPrivate::FLLMGlobals::GetInner();
 		// We have to check IsEnabled() again after calling Get, because the constructor is called
 		// from Get, and will set EnabledStae=Disabled if the platform doesn't support it.
 		if (FLowLevelMemTracker::IsEnabled())
@@ -3406,7 +3587,7 @@ FLLMScope::FLLMScope(const UE::LLMPrivate::FTagData* TagData, bool bInIsStatTag,
 	// We run the init in the open, because we want to track LLM even in transactions.
 	UE_AUTORTFM_OPEN
 	{
-		FLowLevelMemTracker& LLMRef = FLowLevelMemTracker::Get();
+		UE::LLMPrivate::FLLMGlobals& LLMRef = UE::LLMPrivate::FLLMGlobals::GetInner();
 		// We have to check IsEnabled() again after calling Get, because the constructor is called
 		// from Get, and will set EnabledState=Disabled if the platform doesn't support it.
 		if (FLowLevelMemTracker::IsEnabled())
@@ -3429,7 +3610,7 @@ FLLMScope::FLLMScope(const UE::LLMPrivate::FTagData* TagData, bool bInIsStatTag,
 
 void FLLMScope::Destruct()
 {
-	FLowLevelMemTracker& LLMRef = FLowLevelMemTracker::Get();
+	UE::LLMPrivate::FLLMGlobals& LLMRef = UE::LLMPrivate::FLLMGlobals::GetInner();
 	LLMRef.GetTracker(Tracker)->PopTag(TagSet);
 }
 
@@ -3446,7 +3627,7 @@ FLLMScope::~FLLMScope()
 
 void FLLMScopeDynamic::Init(ELLMTracker InTracker, ELLMTagSet InTagSet)
 {
-	FLowLevelMemTracker& LLMRef = FLowLevelMemTracker::Get();
+	UE::LLMPrivate::FLLMGlobals& LLMRef = UE::LLMPrivate::FLLMGlobals::GetInner();
 	// We have to check IsEnabled() again after calling Get, because the constructor is called
 	// from Get, and will set EnabledState=Disabled if the platform doesn't support it.
 	if (!FLowLevelMemTracker::IsEnabled())
@@ -3468,14 +3649,14 @@ void FLLMScopeDynamic::Init(ELLMTracker InTracker, ELLMTagSet InTagSet)
 
 bool FLLMScopeDynamic::TryFindTag(FName TagName)
 {
-	FLowLevelMemTracker& LLMRef = FLowLevelMemTracker::Get();
+	UE::LLMPrivate::FLLMGlobals& LLMRef = UE::LLMPrivate::FLLMGlobals::GetInner();
 	TagData = LLMRef.FindTagData(TagName, TagSet, UE::LLMPrivate::ETagReferenceSource::Scope);
 	return TagData != nullptr;
 }
 
 bool FLLMScopeDynamic::TryAddTagAndActivate(FName TagName, const ILLMDynamicTagConstructor& Constructor)
 {
-	FLowLevelMemTracker& LLMRef = FLowLevelMemTracker::Get();
+	UE::LLMPrivate::FLLMGlobals& LLMRef = UE::LLMPrivate::FLLMGlobals::GetInner();
 	FName StatFullName = NAME_None;
 #if LLM_ENABLED_STAT_TAGS && STATS
 	FString StatConstructorName = Constructor.GetStatName();
@@ -3518,7 +3699,7 @@ bool FLLMScopeDynamic::TryAddTagAndActivate(FName TagName, const ILLMDynamicTagC
 
 void FLLMScopeDynamic::Activate()
 {
-	FLowLevelMemTracker& LLMRef = FLowLevelMemTracker::Get();
+	UE::LLMPrivate::FLLMGlobals& LLMRef = UE::LLMPrivate::FLLMGlobals::GetInner();
 	LLMRef.GetTracker(Tracker)->PushTag(TagData, TagSet);
 }
 
@@ -3526,12 +3707,11 @@ void FLLMScopeDynamic::Destruct()
 {
 	if (TagData)
 	{
-		FLowLevelMemTracker& LLMRef = FLowLevelMemTracker::Get();
+		UE::LLMPrivate::FLLMGlobals& LLMRef = UE::LLMPrivate::FLLMGlobals::GetInner();
 		LLMRef.GetTracker(Tracker)->PopTag(TagSet);
 	}
 }
 
-
 FLLMPauseScope::FLLMPauseScope(FName TagName, bool bIsStatTag, uint64 Amount, ELLMTracker TrackerToPause,
 	ELLMAllocType InAllocType)
 {
@@ -3558,7 +3738,7 @@ FLLMPauseScope::FLLMPauseScope(ELLMTag TagEnum, bool bIsStatTag, uint64 Amount,
 void FLLMPauseScope::Init(FName TagName, ELLMTag EnumTag, bool bIsEnumTag, bool bIsStatTag, uint64 Amount,
 	ELLMTracker TrackerToPause, ELLMAllocType InAllocType)
 {
-	FLowLevelMemTracker& LLMRef = FLowLevelMemTracker::Get();
+	UE::LLMPrivate::FLLMGlobals& LLMRef = UE::LLMPrivate::FLLMGlobals::GetInner();
 	// We have to check IsEnabled() again after calling Get, because the constructor is called
 	// from Get, and will set EnabledState=Disabled if the platform doesn't support it.
 	if (!FLowLevelMemTracker::IsEnabled())
@@ -3609,7 +3789,7 @@ FLLMPauseScope::~FLLMPauseScope()
 	{
 		return;
 	}
-	FLowLevelMemTracker& LLMRef = FLowLevelMemTracker::Get();
+	UE::LLMPrivate::FLLMGlobals& LLMRef = UE::LLMPrivate::FLLMGlobals::GetInner();
 
 	for (int32 TrackerIndex = 0; TrackerIndex < static_cast<int32>(ELLMTracker::Max); TrackerIndex++)
 	{
@@ -3641,7 +3821,7 @@ FLLMScopeFromPtr::FLLMScopeFromPtr(void* Ptr, ELLMTracker InTracker)
 		return;
 	}
 
-	FLowLevelMemTracker& LLMRef = FLowLevelMemTracker::Get();
+	FLLMGlobals& LLMRef = FLLMGlobals::GetInner();
 	// We have to check IsEnabled() again after calling Get, because the constructor is called
 	// from Get, and will set EnabledState=Disabled if the platform doesn't support it.
 	if (!FLowLevelMemTracker::IsEnabled())
@@ -3685,7 +3865,7 @@ FLLMScopeFromPtr::FLLMScopeFromPtr(void* Ptr, ELLMTracker InTracker)
 
 FLLMScopeFromPtr::~FLLMScopeFromPtr()
 {
-	FLowLevelMemTracker& LLMRef = FLowLevelMemTracker::Get();
+	UE::LLMPrivate::FLLMGlobals& LLMRef = UE::LLMPrivate::FLLMGlobals::GetInner();
 	for (int32 TagSetAsInteger = 0; TagSetAsInteger < static_cast<int32>(ELLMTagSet::Max); TagSetAsInteger++)
 	{
 		if (bEnabled[TagSetAsInteger])
@@ -3815,31 +3995,30 @@ const TCHAR* ToString(UE::LLMPrivate::ETagReferenceSource ReferenceSource)
 
 void ValidateUniqueName(FStringView UniqueName)
 {
-// Characters that are invalid for c++ names are invalid (other than /), since we use uniquenames
-// (with / replaced by _) as part of the name of the auto-constructed FLLMTagDeclaration variables
-// _ is invalid since we use an _ to indicate a / in FLLMTagDeclaration.
-// So only Alnum characters or / are allowed, and the first character can not be a number.
+	// Characters that are invalid for c++ names are invalid (other than /), since we use uniquenames
+	// (with / replaced by _) as part of the name of the auto-constructed FLLMTagDeclaration variables
+	// _ is invalid since we use an _ to indicate a / in FLLMTagDeclaration.
+	// So only Alnum characters or / are allowed, and the first character can not be a number.
 	if (UniqueName.Len() == 0)
 	{
 		LLMCheckf(false, TEXT("Invalid length-zero Tag Unique Name"));
 	}
 	else
 	{
-	LLMCheckf(!TChar<TCHAR>::IsDigit(UniqueName[0]),
-		TEXT("Invalid first character is digit in Tag Unique Name '%.*s'"),
-		UniqueName.Len(), UniqueName.GetData());
+		LLMCheckf(!TChar<TCHAR>::IsDigit(UniqueName[0]),
+			TEXT("Invalid first character is digit in Tag Unique Name '%.*s'"),
+			UniqueName.Len(), UniqueName.GetData());
 	}
 	for (TCHAR c : UniqueName)
 	{
 		if (!TChar<TCHAR>::IsAlnum(c) && c != TEXT('/'))
 		{
-		LLMCheckf(false, TEXT("Invalid character %c in Tag Unique Name '%.*s'"), c,
-			UniqueName.Len(), UniqueName.GetData());
+			LLMCheckf(false, TEXT("Invalid character %c in Tag Unique Name '%.*s'"), c,
+				UniqueName.Len(), UniqueName.GetData());
 		}
 	}
 }
 
-
 namespace AllocatorPrivate
 {
 
@@ -3922,7 +4101,7 @@ FLLMAllocator::~FLLMAllocator()
 	Clear();
 }
 
-void FLLMAllocator::Initialise(LLMAllocFunction InAlloc, LLMFreeFunction InFree, int32 InPageSize)
+void FLLMAllocator::Initialize(LLMAllocFunction InAlloc, LLMFreeFunction InFree, int32 InPageSize)
 {
 	using namespace UE::LLMPrivate::AllocatorPrivate;
 
@@ -4485,7 +4664,7 @@ int32 FLLMTracker::FLowLevelAllocInfo::GetIndexFromTagSet(ELLMTagSet InTagSet)
 #endif
 }
 
-void FLLMTracker::FLowLevelAllocInfo::SetTag(const FTagData* InTag, FLowLevelMemTracker& InLLMRef
+void FLLMTracker::FLowLevelAllocInfo::SetTag(const FTagData* InTag, FLLMGlobals& InLLMRef
 #if LLM_ALLOW_ASSETS_TAGS || LLM_ALLOW_UOBJECTCLASSES_TAGS
 	, ELLMTagSet InTagSet
 #endif
@@ -4503,7 +4682,7 @@ void FLLMTracker::FLowLevelAllocInfo::SetTag(const FTagData* InTag, FLowLevelMem
 #endif
 }
 
-const FTagData* FLLMTracker::FLowLevelAllocInfo::GetTag(FLowLevelMemTracker& InLLMRef
+const FTagData* FLLMTracker::FLowLevelAllocInfo::GetTag(FLLMGlobals& InLLMRef
 #if LLM_ALLOW_ASSETS_TAGS || LLM_ALLOW_UOBJECTCLASSES_TAGS
 	, ELLMTagSet InTagSet
 #endif
@@ -4561,7 +4740,7 @@ void FLLMTracker::FLowLevelAllocInfo::SetCompressedTag(ELLMTag InTag)
 
 #endif // else LLM_ENABLED_FULL_TAGS
 
-FLLMTracker::FLLMTracker(FLowLevelMemTracker& InLLM)
+FLLMTracker::FLLMTracker(FLLMGlobals& InLLM)
 	: LLMRef(InLLM)
 	, Tracker(ELLMTracker::Max)
 	, TrackedTotal(0)
@@ -4584,7 +4763,7 @@ FLLMTracker::~FLLMTracker()
 	FPlatformTLS::FreeTlsSlot(TlsSlot);
 }
 
-void FLLMTracker::Initialise(
+void FLLMTracker::Initialize(
 	ELLMTracker InTracker,
 	FLLMAllocator* InAllocator)
 {
@@ -4722,11 +4901,11 @@ void FLLMTracker::TrackAllocation(const void* Ptr, int64 Size, const FTagData* A
 
 	if (AssetTagData == nullptr && !LLMRef.IsBootstrapping() && LLMRef.IsTagSetActive(ELLMTagSet::Assets))
 	{
-		AssetTagData = LLMRef.FindOrAddTagData(TagName_UntaggedAsset, ELLMTagSet::Assets, false, UE::LLMPrivate::ETagReferenceSource::FunctionAPI);
+		AssetTagData = LLMRef.FindOrAddTagData(TagName_UntaggedAsset, ELLMTagSet::Assets, false, ETagReferenceSource::FunctionAPI);
 	}
 	if (AssetClassTagData == nullptr && !LLMRef.IsBootstrapping() && LLMRef.IsTagSetActive(ELLMTagSet::AssetClasses))
 	{
-		AssetClassTagData = LLMRef.FindOrAddTagData(TagName_UntaggedAssetClass, ELLMTagSet::AssetClasses, false, UE::LLMPrivate::ETagReferenceSource::FunctionAPI);
+		AssetClassTagData = LLMRef.FindOrAddTagData(TagName_UntaggedAssetClass, ELLMTagSet::AssetClasses, false, ETagReferenceSource::FunctionAPI);
 	}
 #else
 	const FTagData* AssetTagData = nullptr;
@@ -4909,7 +5088,7 @@ void FLLMTracker::TrackMemory(ELLMTag Tag, int64 Amount, ELLMAllocType AllocType
 
 void FLLMTracker::TrackMemory(FName Tag, ELLMTagSet TagSet, int64 Amount, ELLMAllocType AllocType)
 {
-	TrackMemory(LLMRef.FindOrAddTagData(Tag, TagSet, false, UE::LLMPrivate::ETagReferenceSource::FunctionAPI), Amount, AllocType);
+	TrackMemory(LLMRef.FindOrAddTagData(Tag, TagSet, false, ETagReferenceSource::FunctionAPI), Amount, AllocType);
 }
 
 void FLLMTracker::TrackMemory(const FTagData* TagData, int64 Amount, ELLMAllocType AllocType)
@@ -4921,7 +5100,7 @@ void FLLMTracker::TrackMemory(const FTagData* TagData, int64 Amount, ELLMAllocTy
 
 void FLLMTracker::PauseAndTrackMemory(FName TagName, ELLMTagSet TagSet, bool bInIsStatTag, int64 Amount, ELLMAllocType AllocType)
 {
-	const FTagData* TagData = LLMRef.FindOrAddTagData(TagName, TagSet, bInIsStatTag, UE::LLMPrivate::ETagReferenceSource::FunctionAPI);
+	const FTagData* TagData = LLMRef.FindOrAddTagData(TagName, TagSet, bInIsStatTag, ETagReferenceSource::FunctionAPI);
 	PauseAndTrackMemory(TagData, Amount, AllocType);
 }
 
@@ -6059,13 +6238,13 @@ void FLLMCsvWriter::SetTracker(ELLMTracker InTracker)
 	Tracker = InTracker;
 }
 
-void FLLMCsvWriter::Publish(FLowLevelMemTracker& LLMRef, const FTrackerTagSizeMap& TagSizes,
+void FLLMCsvWriter::Publish(FLLMGlobals& LLMRef, const FTrackerTagSizeMap& TagSizes,
 	const FTagData* OverrideTrackedTotalTagData, const FTagData* OverrideUntaggedTagData, int64 TrackedTotal,
 	UE::LLM::ESizeParams SizeParams)
 {
 	const double Now = FPlatformTime::Seconds();
 
-	if ((FLowLevelMemTracker::Get().bPublishSingleFrame == false) &&
+	if ((LLMRef.bPublishSingleFrame == false) &&
 		(Now - LastWriteTime < (double)CVarLLMWriteInterval.GetValueOnAnyThread()))
 	{
 		return;
@@ -6204,7 +6383,7 @@ void FLLMCsvWriter::WriteHeader(const FTagData* OverrideTrackedTotalTagData, con
 	Archive->Seek(OriginalOffset);
 }
 
-void FLLMCsvWriter::AddRow(FLowLevelMemTracker& LLMRef, const FTrackerTagSizeMap& TagSizes,
+void FLLMCsvWriter::AddRow(FLLMGlobals& LLMRef, const FTrackerTagSizeMap& TagSizes,
 	const FTagData* OverrideTrackedTotalTagData, const FTagData* OverrideUntaggedTagData, int64 TrackedTotal,
 	UE::LLM::ESizeParams SizeParams)
 {
@@ -6298,7 +6477,7 @@ const void* FLLMTraceWriter::GetTagId(const FTagData* TagData)
 	return reinterpret_cast<const void*>(TagData);
 }
 
-void FLLMTraceWriter::Publish(FLowLevelMemTracker& LLMRef, const FTrackerTagSizeMap& TagSizes,
+void FLLMTraceWriter::Publish(FLLMGlobals& LLMRef, const FTrackerTagSizeMap& TagSizes,
 	const FTagData* OverrideTrackedTotalTagData, const FTagData* OverrideUntaggedTagData, int64 TrackedTotal,
 	UE::LLM::ESizeParams SizeParams)
 {
@@ -6478,7 +6657,7 @@ void FLLMCsvProfilerWriter::SetTracker(ELLMTracker InTracker)
 #endif
 }
 
-void FLLMCsvProfilerWriter::Publish(FLowLevelMemTracker& LLMRef, const FTrackerTagSizeMap& TagSizes,
+void FLLMCsvProfilerWriter::Publish(FLLMGlobals& LLMRef, const FTrackerTagSizeMap& TagSizes,
 	const FTagData* OverrideTrackedTotalTagData, const FTagData* OverrideUntaggedTagData, int64 TrackedTotal,
 	UE::LLM::ESizeParams SizeParams)
 {
diff --git a/Engine/Source/Runtime/Core/Private/HAL/LLM/LLMAllocator.h b/Engine/Source/Runtime/Core/Private/HAL/LLM/LLMAllocator.h
new file mode 100644
--- /dev/null
+++ b/Engine/Source/Runtime/Core/Private/HAL/LLM/LLMAllocator.h
@@ -0,0 +1,84 @@
+// Copyright Epic Games, Inc. All Rights Reserved.
+
+#pragma once
+
+#include "HAL/LowLevelMemTracker.h"
+
+#if ENABLE_LOW_LEVEL_MEM_TRACKER
+#include "HAL/PlatformMutex.h"
+
+namespace UE::LLMPrivate
+{
+
+namespace AllocatorPrivate
+{
+	struct FBin;
+	struct FPage;
+}
+/**
+ * The allocator LLM uses to allocate internal memory. Uses platform defined
+ * allocation functions to grab memory directly from the OS.
+ */
+class FLLMAllocator
+{
+public:
+	FLLMAllocator();
+	~FLLMAllocator();
+
+	static FLLMAllocator*& Get();
+
+	void Initialize(LLMAllocFunction InAlloc, LLMFreeFunction InFree, int32 InPageSize);
+	void Clear();
+	void* Alloc(size_t Size);
+	void* Malloc(size_t Size);
+	void Free(void* Ptr, size_t Size);
+	void* Realloc(void* Ptr, size_t OldSize, size_t NewSize);
+	int64 GetTotal() const;
+
+	template <typename T, typename... ArgsType>
+	T* New(ArgsType&&... Args)
+	{
+		T* Ptr = reinterpret_cast<T*>(Alloc(sizeof(T)));
+		new (Ptr) T(Forward<ArgsType>(Args)...);
+		return Ptr;
+	}
+
+	template <typename T>
+	void Delete(T* Ptr)
+	{
+		if (Ptr)
+		{
+			Ptr->~T();
+			Free(Ptr, sizeof(T));
+		}
+	}
+
+	LLMAllocFunction GetPlatformAlloc()
+	{
+		return PlatformAlloc;
+	}
+	LLMFreeFunction GetPlatformFree()
+	{
+		return PlatformFree;
+	}
+
+private:
+	void* AllocPages(size_t Size);
+	void FreePages(void* Ptr, size_t Size);
+	int32 GetBinIndex(size_t Size) const;
+
+	UE::FPlatformRecursiveMutex Mutex;
+	LLMAllocFunction PlatformAlloc;
+	LLMFreeFunction PlatformFree;
+	AllocatorPrivate::FBin* Bins;
+	int64 Total;
+	int32 PageSize;
+	int32 NumBins;
+
+	friend struct UE::LLMPrivate::AllocatorPrivate::FPage;
+	friend struct UE::LLMPrivate::AllocatorPrivate::FBin;
+};
+
+} // namespace UE::LLMPrivate
+
+#endif		// #if ENABLE_LOW_LEVEL_MEM_TRACKER
diff --git a/Engine/Source/Runtime/Core/Private/HAL/LLM/LLMPrivate.h b/Engine/Source/Runtime/Core/Private/HAL/LLM/LLMPrivate.h
--- a/Engine/Source/Runtime/Core/Private/HAL/LLM/LLMPrivate.h
+++ b/Engine/Source/Runtime/Core/Private/HAL/LLM/LLMPrivate.h
@@ -9,7 +9,11 @@
 #include "Containers/Map.h"
 #include "Containers/Set.h"
 #include "Containers/SortedMap.h"
+#include "HAL/CriticalSection.h"
 #include "HAL/LLM/MemoryUtils.h"
+#include "HAL/PlatformMutex.h"
+#include "Misc/Optional.h"
+#include "Misc/ScopeRWLock.h"
 #include "ProfilingDebugging/CsvProfiler.h"
 
 // Private defines configuring LLM; see also public defines in LowLevelMemTracker.h
@@ -57,9 +61,21 @@ static_assert(static_cast<uint32>(ELLMTagSet::Max) == 4, "You added a new TagSet
 #define LLM_TAGSET_MAX 2
 #endif
 
+void GlobalRegisterTagDeclaration(FLLMTagDeclaration& TagDeclaration);
+
 namespace UE::LLMPrivate
 {
 
+enum class ETagReferenceSource : uint8
+{
+	Scope,
+	Declare,
+	EnumTag,
+	CustomEnumTag,
+	FunctionAPI,
+	ImplicitParent
+};
+
 struct FTagDataNameKey
 {
 	FName Name;
@@ -256,6 +272,26 @@ FName LLMGetTagUniqueName(ELLMTag Tag);
 namespace UE::LLMPrivate
 {
 
+#if !PLATFORM_HAS_MULTITHREADED_PREMAIN
+
+struct FEnableStateScopeLock
+{
+};
+
+#else // !PLATFORM_HAS_MULTITHREADED_PREMAIN
+
+struct FEnableStateScopeLock
+{
+	// Undefine copy/move constructor since FReadScopeLock does not support it.
+	FEnableStateScopeLock() = default;
+	FEnableStateScopeLock(const FEnableStateScopeLock&) = delete;
+	FEnableStateScopeLock(FEnableStateScopeLock&&) = delete;
+
+	TOptional<FReadScopeLock> Inner;
+};
+
+#endif // else !PLATFORM_HAS_MULTITHREADED_PREMAIN
+
 /** FLLMCsvWriter: class for writing out the LLM tag sizes to a csv file every few seconds. */
 class FLLMCsvWriter
 {
@@ -266,7 +302,7 @@ public:
 	void SetTracker(ELLMTracker InTracker);
 	void Clear();
 
-	void Publish(FLowLevelMemTracker& LLMRef, const FTrackerTagSizeMap& TagSizes,
+	void Publish(FLLMGlobals& LLMRef, const FTrackerTagSizeMap& TagSizes,
 		const FTagData* OverrideTrackedTotalTagData, const FTagData* OverrideUntaggedTagData, int64 TrackedTotal,
 		UE::LLM::ESizeParams SizeParams);
 
@@ -281,7 +317,7 @@ private:
 	bool CreateArchive();
 	bool UpdateColumns(const FTrackerTagSizeMap& TagSizes);
 	void WriteHeader(const FTagData* OverrideTrackedTotalTagData, const FTagData* OverrideUntaggedTagData);
-	void AddRow(FLowLevelMemTracker& LLMRef, const FTrackerTagSizeMap& TagSizes,
+	void AddRow(FLLMGlobals& LLMRef, const FTrackerTagSizeMap& TagSizes,
 		const FTagData* OverrideTrackedTotalName, const FTagData* OverrideUntaggedName, int64 TrackedTotal,
 		UE::LLM::ESizeParams SizeParams);
 
@@ -301,7 +337,7 @@ public:
 	FLLMTraceWriter();
 	void SetTracker(ELLMTracker InTracker);
 	void Clear();
-	void Publish(FLowLevelMemTracker& LLMRef, const FTrackerTagSizeMap& TagSizes,
+	void Publish(FLLMGlobals& LLMRef, const FTrackerTagSizeMap& TagSizes,
 		const FTagData* OverrideTrackedTotalName, const FTagData* OverrideUntaggedName, int64 TrackedTotal,
 		UE::LLM::ESizeParams SizeParams);
 
@@ -322,7 +358,7 @@ public:
 	FLLMCsvProfilerWriter();
 	void Clear();
 	void SetTracker(ELLMTracker InTracker);
-	void Publish(FLowLevelMemTracker& LLMRef, const FTrackerTagSizeMap& TagSizes,
+	void Publish(FLLMGlobals& LLMRef, const FTrackerTagSizeMap& TagSizes,
 		const FTagData* OverrideTrackedTotalName, const FTagData* OverrideUntaggedName, int64 TrackedTotal,
 		UE::LLM::ESizeParams SizeParams);
 protected:
@@ -378,10 +414,10 @@ class FLLMTracker
 {
 public:
 
-	FLLMTracker(FLowLevelMemTracker& InLLM);
+	FLLMTracker(FLLMGlobals& InLLM);
 	~FLLMTracker();
 
-	void Initialise(ELLMTracker InTracker, FLLMAllocator* InAllocator);
+	void Initialize(ELLMTracker InTracker, FLLMAllocator* InAllocator);
 
 	void PushTag(ELLMTag EnumTag, ELLMTagSet TagSet);
 	void PushTag(FName Tag, bool bInIsStatTag, ELLMTagSet TagSet);
@@ -432,13 +468,13 @@ public:
 
 		static int32 GetIndexFromTagSet(ELLMTagSet InTagSet);
 
-		void SetTag(const FTagData* InTag, FLowLevelMemTracker& InLLMRef
+		void SetTag(const FTagData* InTag, FLLMGlobals& InLLMRef
 #if LLM_ALLOW_ASSETS_TAGS || LLM_ALLOW_UOBJECTCLASSES_TAGS
 			, ELLMTagSet InTagSet = ELLMTagSet::None
 #endif
 		);
 
-		const FTagData* GetTag(FLowLevelMemTracker& InLLMRef
+		const FTagData* GetTag(FLLMGlobals& InLLMRef
 #if LLM_ALLOW_ASSETS_TAGS || LLM_ALLOW_UOBJECTCLASSES_TAGS
 			, ELLMTagSet InTagSet = ELLMTagSet::None
 #endif
@@ -510,7 +546,7 @@ protected:
 	void TrackAllocation(const void* Ptr, int64 Size, const FTagData* ActiveTagData, ELLMAllocType AllocType, FLLMThreadState* State, bool bTrackInMemPro);
 	void TrackMemoryOfActiveTag(int64 Amount, const FTagData* TagData, ELLMAllocType AllocType, FLLMThreadState* State);
 
-	FLowLevelMemTracker& LLMRef;
+	FLLMGlobals& LLMRef;
 
 	ELLMTracker Tracker;
 
@@ -552,6 +588,198 @@ protected:
 	int64 AllocTypeAmounts[(int32)ELLMAllocType::Count];
 };
 
+/**
+ * The singleton implementation LLM class. It contains multiple FLLMTrackers, and contains storage for data shared
+ * between trackers, such as FTagDatas.
+ */
+class FLLMGlobals : public FLowLevelMemTracker
+{
+public:
+	// Functions Forwarded to from FLowLevelMemTracker
+	void ProcessCommandLineInner(const TCHAR* CmdLine);
+	uint64 GetTotalTrackedMemoryInner(ELLMTracker Tracker);
+	void OnLowLevelAllocInner(ELLMTracker Tracker, const void* Ptr, uint64 Size, ELLMTag DefaultTag = ELLMTag::Untagged,
+		ELLMAllocType AllocType = ELLMAllocType::None, bool bTrackInMemPro = true);
+	void OnLowLevelAllocInner(ELLMTracker Tracker, const void* Ptr, uint64 Size, FName DefaultTag,
+		ELLMAllocType AllocType = ELLMAllocType::None, bool bTrackInMemPro = true);
+	void OnLowLevelFreeInner(ELLMTracker Tracker, const void* Ptr,
+		ELLMAllocType AllocType = ELLMAllocType::None, bool bTrackInMemPro = true);
+	void OnLowLevelChangeInMemoryUseInner(ELLMTracker Tracker, int64 DeltaMemory, ELLMTag DefaultTag = ELLMTag::Untagged, ELLMAllocType AllocType = ELLMAllocType::None);
+	void OnLowLevelChangeInMemoryUseInner(ELLMTracker Tracker, int64 DeltaMemory, FName DefaultTag, ELLMAllocType AllocType = ELLMAllocType::None);
+	void OnLowLevelAllocMovedInner(ELLMTracker Tracker, const void* Dest, const void* Source,
+		ELLMAllocType AllocType = ELLMAllocType::None);
+	void UpdateStatsPerFrameInner(const TCHAR* LogName = nullptr);
+	void TickInner();
+	void SetProgramSizeInner(uint64 InProgramSize);
+	bool ExecInner(const TCHAR* Cmd, FOutputDevice& Ar);
+	bool IsTagSetActiveInner(ELLMTagSet Set);
+	bool ShouldReduceThreadsInner();
+	const FTagData* GetActiveTagDataInner(ELLMTracker Tracker, ELLMTagSet TagSet = ELLMTagSet::None);
+	void RegisterPlatformTagInner(int32 Tag, const TCHAR* Name, FName StatName, FName SummaryStatName, int32 ParentTag = -1);
+	void RegisterProjectTagInner(int32 Tag, const TCHAR* Name, FName StatName, FName SummaryStatName, int32 ParentTag = -1);
+	TArray<const FTagData*> GetTrackedTagsInner(ELLMTagSet TagSet = ELLMTagSet::None);
+	TArray<const FTagData*> GetTrackedTagsInner(ELLMTracker Tracker, ELLMTagSet TagSet = ELLMTagSet::None);
+	void GetTrackedTagsNamesWithAmountInner(TMap<FName, uint64>& TagsNamesWithAmount, ELLMTracker Tracker, ELLMTagSet TagSet);
+	void GetTrackedTagsNamesWithAmountFilteredInner(TMap<FName, uint64>& TagsNamesWithAmount, ELLMTracker Tracker, ELLMTagSet TagSet, TArray<FLLMTagSetAllocationFilter>& Filters);
+	bool FindTagByNameInner(const TCHAR* Name, uint64& OutTag, ELLMTagSet InTagSet = ELLMTagSet::None) const;
+	FName FindTagDisplayNameInner(uint64 Tag) const;
+	FName FindPtrDisplayNameInner(void* Ptr) const;
+	FName GetTagDisplayNameInner(const FTagData* TagData) const;
+	FString GetTagDisplayPathNameInner(const FTagData* TagData) const;
+	void GetTagDisplayPathNameInner(const FTagData* TagData,
+		FStringBuilderBase& OutPathName, int32 MaxLen = -1) const;
+	FName GetTagUniqueNameInner(const FTagData* TagData) const;
+	const FTagData* GetTagParentInner(const FTagData* TagData) const;
+	bool GetTagIsEnumTagInner(const FTagData* TagData) const;
+	ELLMTag GetTagClosestEnumTagInner(const FTagData* TagData) const;
+	int64 GetTagAmountForTrackerInner(ELLMTracker Tracker, ELLMTag Tag, UE::LLM::ESizeParams SizeParams = UE::LLM::ESizeParams::Default);
+	int64 GetTagAmountForTrackerInner(ELLMTracker Tracker, const FTagData* TagData, UE::LLM::ESizeParams SizeParams = UE::LLM::ESizeParams::Default);
+	int64 GetTagAmountForTrackerInner(ELLMTracker Tracker, FName Tag, ELLMTagSet TagSet, UE::LLM::ESizeParams SizeParams = UE::LLM::ESizeParams::Default);
+	void SetTagAmountForTrackerInner(ELLMTracker Tracker, ELLMTag Tag, int64 Amount, bool bAddToTotal);
+	void SetTagAmountForTrackerInner(ELLMTracker Tracker, FName Tag, ELLMTagSet TagSet, int64 Amount, bool bAddToTotal);
+	uint64 DumpTagInner(ELLMTracker Tracker, const char* FileName, int LineNumber);
+	void PublishDataSingleFrameInner();
+	void DumpToLogInner(EDumpFormat DumpFormat = EDumpFormat::PlainText, FOutputDevice* OutputDevice = nullptr, UE::LLM::ESizeParams SizeParams = UE::LLM::ESizeParams::Default, ELLMTagSet TagSet = ELLMTagSet::None);
+	void OnPreForkInner();
+	bool IsInitializedInner() const;
+	bool IsConfiguredInner() const;
+	void BootstrapInitializeInner();
+	void FinishInitializeInner();
+
+	// Functions for use by types within LLMPrivate
+	static FLLMGlobals& GetInner();
+
+private:
+	FLLMGlobals();
+	~FLLMGlobals();
+
+	bool IsBootstrapping() const;
+
+	/** Free all memory. This will put the tracker into a permanently disabled state. */
+	void Clear();
+	void InitializeProgramSize();
+
+	class FLLMTracker* GetTracker(ELLMTracker Tracker);
+	const class FLLMTracker* GetTracker(ELLMTracker Tracker) const;
+
+	void TickInnerNoLock();
+	void UpdateTags();
+	void SortTags(FTagDataArray*& OutOldTagDatas);
+	void PublishDataPerFrame(const TCHAR* LogName);
+
+	void RegisterCustomTag(int32 Tag, ELLMTagSet TagSet, const TCHAR* Name, FName StatName, FName SummaryStatName, int32 ParentTag = -1);
+	/**
+	* Called during C++ global static initialization when GMalloc is available (and hence FNames are unavailable)
+	* Creates the subset of tags necessary to record allocations during GMalloc and FName construction
+	*/
+	void BootstrapTagDatas();
+	void InitializeTagDatas_SetLLMTagNames();
+	void InitializeTagDatas_FinishRegister();
+	void InitializeTagDatas();
+	void ClearTagDatas();
+	void RegisterTagDeclaration(FLLMTagDeclaration& TagDeclaration);
+	FTagData& RegisterTagData(FName Name, FName DisplayName, FName ParentName, FName StatName, FName SummaryStatName, bool bHasEnumTag, ELLMTag EnumTag, bool bIsStatTag, ETagReferenceSource ReferenceSource, ELLMTagSet TagSet = ELLMTagSet::None);
+	/** Construct if not yet done the data on the given FTagData that relies on the presence of other TagDatas */
+	void FinishConstruct(FTagData* TagData, ETagReferenceSource ReferenceSource);
+	void ReportDuplicateTagName(FTagData* TagData, ETagReferenceSource ReferenceSource);
+
+	const FTagData* FindOrAddTagData(ELLMTag EnumTag, ETagReferenceSource ReferenceSource = ETagReferenceSource::FunctionAPI);
+	const FTagData* FindOrAddTagData(FName Name, ELLMTagSet TagSet, bool bIsStatData = false, ETagReferenceSource ReferenceSource = ETagReferenceSource::FunctionAPI);
+	const FTagData* FindOrAddTagData(FName Name, ELLMTagSet TagSet, FName StatName, ETagReferenceSource ReferenceSource = ETagReferenceSource::FunctionAPI);
+	const FTagData* FindTagData(ELLMTag EnumTag, ETagReferenceSource ReferenceSource = ETagReferenceSource::FunctionAPI);
+	const FTagData* FindTagData(FName Name, ELLMTagSet TagSet, ETagReferenceSource ReferenceSource = ETagReferenceSource::FunctionAPI);
+
+	friend class ::FLLMPauseScope;
+	friend class ::FLLMScope;
+	friend class ::FLLMScopeDynamic;
+	friend class ::FLLMScopeFromPtr;
+	friend class ::FLowLevelMemTracker;
+	friend class UE::LLMPrivate::FLLMCsvWriter;
+	friend class UE::LLMPrivate::FLLMTracker;
+	friend class UE::LLMPrivate::FLLMThreadState;
+	friend class UE::LLMPrivate::FLLMTraceWriter;
+	friend class UE::LLMPrivate::FLLMCsvProfilerWriter;
+	friend void ::GlobalRegisterTagDeclaration(FLLMTagDeclaration& TagDeclaration);
+
+private:
+	FLLMAllocator Allocator;
+	/** All TagDatas that have been constructed, in an array sorted by TagData->GetIndex() */
+	FTagDataArray* TagDatas;
+	/** Map from TagData->GetName() to TagData for all names, used to handle LLM_SCOPE with FName */
+	FTagDataNameMap* TagDataNameMap;
+	/** Array to Map from ELLMTag to the TagData for that tag, used to handle LLM_SCOPE with ELLMTag */
+	FTagData** TagDataEnumMap;
+
+	UE::LLMPrivate::FTagFilter* TraceTagFilter;
+
+	FLLMTracker* Trackers[static_cast<int32>(ELLMTracker::Max)];
+
+	mutable UE::FPlatformSharedMutex TagDataMutex;
+	UE::FPlatformRecursiveMutex UpdateMutex;
+	UE::FPlatformRecursiveMutex BootstrapMutex;
+
+	uint64 ProgramSize;
+	int64 MemoryUsageCurrentOverhead;
+	int64 MemoryUsagePlatformTotalUntracked;
+
+	bool ActiveSets[(int32)ELLMTagSet::Max];
+
+	bool bFirstTimeUpdating;
+	bool bCanEnable;
+	bool bCsvWriterEnabled;
+	bool bTraceWriterEnabled;
+	bool bInitializedTracking;
+	bool bIsBootstrapping;
+	bool bBeganInitializedTracking = false;
+	bool bFullyInitialized;
+	std::atomic<bool> bConfigurationComplete;
+	bool bTagAdded;
+	bool bAutoPublish;
+	bool bPublishSingleFrame;
+	bool bCapturedSizeSnapshot;
+};
+
+} // namespace UE::LLMPrivate
+
+
+///////////////////////////////////////////////////////
+// Inline implementations
+///////////////////////////////////////////////////////
+
+
+inline UE::LLMPrivate::FLLMGlobals* FLowLevelMemTracker::GetGlobals()
+{
+	return static_cast<UE::LLMPrivate::FLLMGlobals*>(this);
+}
+
+inline const UE::LLMPrivate::FLLMGlobals* FLowLevelMemTracker::GetGlobals() const
+{
+	return static_cast<const UE::LLMPrivate::FLLMGlobals*>(this);
+}
+
+namespace UE::LLMPrivate
+{
+
+inline bool FLLMGlobals::IsInitializedInner() const
+{
+	return bFullyInitialized;
+}
+
+inline bool FLLMGlobals::IsConfiguredInner() const
+{
+	return bConfigurationComplete;
+}
+
+inline bool FLLMGlobals::IsBootstrapping() const
+{
+	return bIsBootstrapping;
+}
+
+inline FLLMGlobals& FLLMGlobals::GetInner()
+{
+	return *FLowLevelMemTracker::Get().GetGlobals();
+}
+
 } // namespace UE::LLMPrivate
 
 #endif		// #if ENABLE_LOW_LEVEL_MEM_TRACKER
diff --git a/Engine/Source/Runtime/Core/Private/HAL/LLM/MemoryUtils.h b/Engine/Source/Runtime/Core/Private/HAL/LLM/MemoryUtils.h
--- a/Engine/Source/Runtime/Core/Private/HAL/LLM/MemoryUtils.h
+++ b/Engine/Source/Runtime/Core/Private/HAL/LLM/MemoryUtils.h
@@ -11,9 +11,13 @@
 #include "Containers/ContainerAllocationPolicies.h"
 #include "Containers/Map.h"
 #include "Containers/Set.h"
+#include "HAL/LLM/LLMAllocator.h"
 #include "HAL/PlatformMutex.h"
 #include "Math/NumericLimits.h"
+#include "Templates/AlignmentTemplates.h"
 #include "Templates/Tuple.h"
+
+#include <atomic>
 #include <type_traits>
 
 #define LLM_PAGE_SIZE (16*1024)
diff --git a/Engine/Source/Runtime/Core/Public/HAL/LowLevelMemTracker.h b/Engine/Source/Runtime/Core/Public/HAL/LowLevelMemTracker.h
--- a/Engine/Source/Runtime/Core/Public/HAL/LowLevelMemTracker.h
+++ b/Engine/Source/Runtime/Core/Public/HAL/LowLevelMemTracker.h
@@ -3,7 +3,6 @@
 #pragma once
 
 #include "CoreTypes.h"
-#include "HAL/PlatformMutex.h"
 #include "LowLevelMemTrackerDefines.h" // LLM_ENABLED_IN_CONFIG
 #include "ProfilingDebugging/TagTrace.h"
 #include "AutoRTFM.h"
@@ -53,16 +52,12 @@
 #include "Containers/Array.h"
 #include "Containers/ArrayView.h"
 #include "Containers/StringFwd.h"
-#include "HAL/CriticalSection.h"
 #include "HAL/PlatformCrt.h"
 #include "HAL/PlatformMisc.h"
-#include "Templates/AlignmentTemplates.h"
 #include "Templates/UnrealTemplate.h"
 #include "UObject/NameTypes.h"
 #include "UObject/UnrealNames.h"
 
-#include <atomic>
-
 class FTagTrace;
 
 #if DO_CHECK
@@ -487,93 +482,15 @@ struct FLLMTagSetAllocationFilter;
 namespace UE::LLMPrivate
 {
 
-class FLLMCsvProfilerWriter;
-class FLLMCsvWriter;
-class FLLMThreadState;
-class FLLMTraceWriter;
-class FLLMTracker;
+class FLLMGlobals;
 class FTagData;
-class FTagDataArray;
-class FTagDataNameMap;
-class FTagFilter;
+struct FEnableStateScopeLock;
 
-namespace AllocatorPrivate
-{
-	struct FBin;
-	struct FPage;
-}
-/**
- * The allocator LLM uses to allocate internal memory. Uses platform defined
- * allocation functions to grab memory directly from the OS.
- */
-class FLLMAllocator
+enum class EEnabled : uint8
 {
-public:
-	FLLMAllocator();
-	~FLLMAllocator();
-
-	static FLLMAllocator*& Get();
-
-	void Initialise(LLMAllocFunction InAlloc, LLMFreeFunction InFree, int32 InPageSize);
-	void Clear();
-	void* Alloc(size_t Size);
-	void* Malloc(size_t Size);
-	void Free(void* Ptr, size_t Size);
-	void* Realloc(void* Ptr, size_t OldSize, size_t NewSize);
-	int64 GetTotal() const;
-
-	template <typename T, typename... ArgsType>
-	T* New(ArgsType&&... Args)
-	{
-		T* Ptr = reinterpret_cast<T*>(Alloc(sizeof(T)));
-		new (Ptr) T(Forward<ArgsType>(Args)...);
-		return Ptr;
-	}
-
-	template <typename T>
-	void Delete(T* Ptr)
-	{
-		if (Ptr)
-		{
-			Ptr->~T();
-			Free(Ptr, sizeof(T));
-		}
-	}
-
-	LLMAllocFunction GetPlatformAlloc()
-	{
-		return PlatformAlloc;
-	}
-	LLMFreeFunction GetPlatformFree()
-	{
-		return PlatformFree;
-	}
-
-private:
-	void* AllocPages(size_t Size);
-	void FreePages(void* Ptr, size_t Size);
-	int32 GetBinIndex(size_t Size) const;
-
-	UE::FPlatformRecursiveMutex Mutex;
-	LLMAllocFunction PlatformAlloc;
-	LLMFreeFunction PlatformFree;
-	AllocatorPrivate::FBin* Bins;
-	int64 Total;
-	int32 PageSize;
-	int32 NumBins;
-
-	friend struct UE::LLMPrivate::AllocatorPrivate::FPage;
-	friend struct UE::LLMPrivate::AllocatorPrivate::FBin;
-};
-
-enum class ETagReferenceSource : uint8
-{
-	Scope,
-	Declare,
-	EnumTag,
-	CustomEnumTag,
-	FunctionAPI,
-	ImplicitParent
+	NotYetKnown = 0,
+	Disabled,
+	Enabled,
 };
 
 typedef void (*FLLMInitialisedCallback)(UPTRINT UserData);
@@ -614,7 +531,7 @@ struct FLLMTagInfo
 	int32 ParentTag = -1;
 };
 
-/** The main LLM tracker class. */
+/** The public LLM interface. */
 class FLowLevelMemTracker
 {
 public:
@@ -784,15 +701,8 @@ public:
 
 	CORE_API void OnPreFork();
 
-	UE_FORCEINLINE_HINT bool IsInitialized() const
-	{
-		return bFullyInitialised;
-	}
-
-	UE_FORCEINLINE_HINT bool IsConfigured() const
-	{
-		return bConfigurationComplete;
-	}
+	CORE_API bool IsInitialized() const;
+	CORE_API bool IsConfigured() const;
 
 	/**
 	 * Initialize the data required for tracking allocations if it has not already been done.
@@ -812,55 +722,8 @@ public:
 	CORE_API void FinishInitialise();
 
 private:
-	enum class EEnabled : uint8
-	{
-		NotYetKnown = 0,
-		Disabled,
-		Enabled,
-	};
-	struct FEnableStateScopeLock;
-
-private:
-	FLowLevelMemTracker();
-
-	~FLowLevelMemTracker();
-
-	bool IsBootstrapping() const { return bIsBootstrapping; }
-
-	/** Free all memory. This will put the tracker into a permanently disabled state. */
-	void Clear();
-	void InitialiseProgramSize();
-
-	class UE::LLMPrivate::FLLMTracker* GetTracker(ELLMTracker Tracker);
-	const class UE::LLMPrivate::FLLMTracker* GetTracker(ELLMTracker Tracker) const;
-
-	void TickInternal();
-	void UpdateTags();
-	void SortTags(UE::LLMPrivate::FTagDataArray*& OutOldTagDatas);
-	void PublishDataPerFrame(const TCHAR* LogName);
-
-	void RegisterCustomTagInternal(int32 Tag, ELLMTagSet TagSet, const TCHAR* Name, FName StatName, FName SummaryStatName, int32 ParentTag = -1);
-	/**
-	* Called during C++ global static initialization when GMalloc is available (and hence FNames are unavailable)
-	* Creates the subset of tags necessary to record allocations during GMalloc and FName construction
-	*/
-	void BootstrapTagDatas();
-	void InitialiseTagDatas_SetLLMTagNames();
-	void InitialiseTagDatas_FinishRegister();
-	void InitialiseTagDatas();
-	void ClearTagDatas();
-	void RegisterTagDeclaration(FLLMTagDeclaration& TagDeclaration);
-	UE::LLMPrivate::FTagData& RegisterTagData(FName Name, FName DisplayName, FName ParentName, FName StatName, FName SummaryStatName, bool bHasEnumTag, ELLMTag EnumTag, bool bIsStatTag, UE::LLMPrivate::ETagReferenceSource ReferenceSource, ELLMTagSet TagSet = ELLMTagSet::None);
-	/** Construct if not yet done the data on the given FTagData that relies on the presence of other TagDatas */
-	void FinishConstruct(UE::LLMPrivate::FTagData* TagData, UE::LLMPrivate::ETagReferenceSource ReferenceSource);
-	void ReportDuplicateTagName(UE::LLMPrivate::FTagData* TagData, UE::LLMPrivate::ETagReferenceSource ReferenceSource);
-
-	const UE::LLMPrivate::FTagData* FindOrAddTagData(ELLMTag EnumTag, UE::LLMPrivate::ETagReferenceSource ReferenceSource = UE::LLMPrivate::ETagReferenceSource::FunctionAPI);
-	const UE::LLMPrivate::FTagData* FindOrAddTagData(FName Name, ELLMTagSet TagSet, bool bIsStatData=false, UE::LLMPrivate::ETagReferenceSource ReferenceSource = UE::LLMPrivate::ETagReferenceSource::FunctionAPI);
-	const UE::LLMPrivate::FTagData* FindOrAddTagData(FName Name, ELLMTagSet TagSet, FName StatName, UE::LLMPrivate::ETagReferenceSource ReferenceSource = UE::LLMPrivate::ETagReferenceSource::FunctionAPI);
-	const UE::LLMPrivate::FTagData* FindTagData(ELLMTag EnumTag, UE::LLMPrivate::ETagReferenceSource ReferenceSource = UE::LLMPrivate::ETagReferenceSource::FunctionAPI);
-	const UE::LLMPrivate::FTagData* FindTagData(FName Name, ELLMTagSet TagSet, UE::LLMPrivate::ETagReferenceSource ReferenceSource = UE::LLMPrivate::ETagReferenceSource::FunctionAPI);
-
+	UE::LLMPrivate::FLLMGlobals* GetGlobals();
+	const UE::LLMPrivate::FLLMGlobals* GetGlobals() const;
 	/**
 	 * A function that returns IsEnabled(), but also enters a lock if current state is NotYetKnown and therefore might
 	 * change if the main thread calls ProcessCommandLine and commandline specifies that LLM should disable itself.
@@ -869,58 +732,13 @@ private:
 	 * and delete (OnLowLevelAlloc, etc), since those are the only functions that can be called from other threads
 	 * before ProcessCommandLine.
 	 */
-	inline static bool TryEnterEnabled(FEnableStateScopeLock& ScopeLock);
-
-	friend class FLLMPauseScope;
-	friend class FLLMScope;
-	friend class FLLMScopeDynamic;
-	friend class FLLMScopeFromPtr;
-	friend class UE::LLMPrivate::FLLMCsvWriter;
-	friend class UE::LLMPrivate::FLLMTracker;
-	friend class UE::LLMPrivate::FLLMThreadState;
-	friend class UE::LLMPrivate::FLLMTraceWriter;
-	friend class UE::LLMPrivate::FLLMCsvProfilerWriter;
-	friend void GlobalRegisterTagDeclaration(FLLMTagDeclaration& TagDeclaration);
-
-private:
-	UE::LLMPrivate::FLLMAllocator Allocator;
-	/** All TagDatas that have been constructed, in an array sorted by TagData->GetIndex() */
-	UE::LLMPrivate::FTagDataArray* TagDatas;
-	/** Map from TagData->GetName() to TagData for all names, used to handle LLM_SCOPE with FName */
-	UE::LLMPrivate::FTagDataNameMap* TagDataNameMap;
-	/** Array to Map from ELLMTag to the TagData for that tag, used to handle LLM_SCOPE with ELLMTag */
-	UE::LLMPrivate::FTagData** TagDataEnumMap;
-
-	UE::LLMPrivate::FTagFilter* TraceTagFilter;
-
-	UE::LLMPrivate::FLLMTracker* Trackers[static_cast<int32>(ELLMTracker::Max)];
-
-	mutable UE::FPlatformSharedMutex TagDataMutex;
-	UE::FPlatformRecursiveMutex UpdateMutex;
-	UE::FPlatformRecursiveMutex BootstrapMutex;
-
-	uint64 ProgramSize;
-	int64 MemoryUsageCurrentOverhead;
-	int64 MemoryUsagePlatformTotalUntracked;
-
-	bool ActiveSets[(int32)ELLMTagSet::Max];
-
-	bool bFirstTimeUpdating;
-	bool bCanEnable;
-	bool bCsvWriterEnabled;
-	bool bTraceWriterEnabled;
-	bool bInitialisedTracking;
-	bool bIsBootstrapping;
-	bool bBeganInitialisedTracking = false;
-	bool bFullyInitialised;
-	std::atomic<bool> bConfigurationComplete;
-	bool bTagAdded;
-	bool bAutoPublish;
-	bool bPublishSingleFrame;
-	bool bCapturedSizeSnapshot;
+	inline static bool TryEnterEnabled(UE::LLMPrivate::FEnableStateScopeLock& ScopeLock);
 
 	static CORE_API FLowLevelMemTracker* TrackerInstance;
-	static CORE_API EEnabled EnabledState;
+	static CORE_API UE::LLMPrivate::EEnabled EnabledState;
+
+	friend UE::LLMPrivate::FLLMGlobals;
+
 public:
 	UE_DEPRECATED(5.5, "Use IsEnabled instead.")
 	static CORE_API bool bIsDisabled;
@@ -1079,7 +897,7 @@ protected:
 	ELLMTagSet TagSet;
 	FLLMTagDeclaration* Next = nullptr;
 
-	friend class FLowLevelMemTracker;
+	friend class UE::LLMPrivate::FLLMGlobals;
 };
 
 /** Used to filter Tag allocations specifying a Name that matches in a TagSet, or TagSet alone or just by Name. */
@@ -1093,7 +911,7 @@ struct FLLMTagSetAllocationFilter
 
 inline bool FLowLevelMemTracker::IsEnabled()
 {
-	return EnabledState != EEnabled::Disabled;
+	return EnabledState != UE::LLMPrivate::EEnabled::Disabled;
 }
 
 #else
