Core: Fixed TFunctionWithContext to support function pointers

#Synced-CL 49572071

diff --git a/Engine/Source/Runtime/Core/Public/Templates/FunctionWithContext.h b/Engine/Source/Runtime/Core/Public/Templates/FunctionWithContext.h
--- a/Engine/Source/Runtime/Core/Public/Templates/FunctionWithContext.h
+++ b/Engine/Source/Runtime/Core/Public/Templates/FunctionWithContext.h
@@ -43,23 +43,23 @@ public:
 	using FunctionType = ReturnType (void*, ArgTypes...);
 
 	/** Construct from a lambda or a function pointer. */
-	template <typename LambdaType>
-	inline TFunctionWithContext(LambdaType&& Lambda UE_LIFETIMEBOUND)
-		requires (!std::is_same_v<std::decay_t<LambdaType>, TFunctionWithContext>) &&
-			std::is_invocable_r_v<ReturnType, std::decay_t<LambdaType>, ArgTypes...>
-		: Function(&Call<std::decay_t<LambdaType>>)
-		, Context((void*)&Lambda)
+	template <typename CallableType>
+	inline TFunctionWithContext(CallableType&& Callable UE_LIFETIMEBOUND)
+		requires (!std::is_same_v<std::decay_t<CallableType>, TFunctionWithContext>) &&
+			std::is_invocable_r_v<ReturnType, std::decay_t<CallableType>, ArgTypes...>
+		: Function(&Call<std::remove_pointer_t<std::decay_t<CallableType>>>)
+		, Context((void*)&Callable)
 	{
 	}
 
 	/** Assign from a lambda or a function pointer. */
-	template <typename LambdaType>
-	inline TFunctionWithContext& operator=(LambdaType&& Lambda UE_LIFETIMEBOUND)
-		requires (!std::is_same_v<std::decay_t<LambdaType>, TFunctionWithContext>) &&
-			std::is_invocable_r_v<ReturnType, std::decay_t<LambdaType>, ArgTypes...>
+	template <typename CallableType>
+	inline TFunctionWithContext& operator=(CallableType&& Callable UE_LIFETIMEBOUND)
+		requires (!std::is_same_v<std::decay_t<CallableType>, TFunctionWithContext>) &&
+			std::is_invocable_r_v<ReturnType, std::decay_t<CallableType>, ArgTypes...>
 	{
-		Function = &Call<std::decay_t<LambdaType>>;
-		Context = (void*)&Lambda;
+		Function = &Call<std::remove_pointer_t<std::decay_t<CallableType>>>;
+		Context = (void*)&Callable;
 		return *this;
 	}
 
@@ -98,10 +98,10 @@ public:
 	}
 
 private:
-	template <typename LambdaType>
-	static ReturnType Call(void* Lambda, ArgTypes... Args)
+	template <typename CallableType>
+	static ReturnType Call(void* Callable, ArgTypes... Args)
 	{
-		return Invoke(*(LambdaType*)Lambda, (ArgTypes&&)Args...);
+		return Invoke(*(CallableType*)Callable, (ArgTypes&&)Args...);
 	}
 
 	FunctionType* Function = nullptr;
