- Added some thread safety to FConfigCacheIni::ReadString (which the others like GetInt go through), SafeReload, and AddDynamicLayersToHierarchy

#Synced-CL 48811006

#Synced-CL 49171745

#was-backout 49234874

#Synced-CL 49298013

diff --git a/Engine/Source/Runtime/Core/Private/Misc/ConfigCacheIni.cpp b/Engine/Source/Runtime/Core/Private/Misc/ConfigCacheIni.cpp
--- a/Engine/Source/Runtime/Core/Private/Misc/ConfigCacheIni.cpp
+++ b/Engine/Source/Runtime/Core/Private/Misc/ConfigCacheIni.cpp
@@ -33,10 +33,21 @@
 #include "Serialization/LargeMemoryReader.h"
 #include "Async/Async.h"
 #include "Async/ParallelFor.h"
+#include "Async/SharedLock.h"
+#include "Async/UniqueLock.h"
 #include "Misc/OutputDeviceRedirector.h"
 #include "Logging/MessageLog.h"
 #include <limits>
 
+// if set to 1, this will make sure FCOnfigCacheIni::GetString (and friends) is thread safe from dynamic layer loading, and SafeReload
+// if set to 0, this will behave as 2025 and before (no safety for GetString)
+// note that if you bypass FConfigCache::Get* and use FCOnfigFile::Get*, then it will not be threadsafe
+#define UE_THREAD_SAFETY_PROTECT_GETSTRING 1
+
+// if set to 1, GetString (GetInt, etc) will be slower, but SafeUnload will unload more memory
+// if set to 0, SafeUnload will free less memory, but Get* will be faster
+#define UE_THREAD_SAFETY_USE_WRITE_LOCKS 1
+
 namespace
 {
 	const FString CurrentIniVersionStr = "CurrentIniVersion";
@@ -3435,12 +3446,17 @@ bool FConfigBranch::AddDynamicLayersToHierarchy(const TArray<FString>& Filenames
 
 bool FConfigBranch::AddDynamicLayersToHierarchy(const TArray<FDynamicLayerInfo>& Layers, FConfigModificationTracker* ModificationTracker, const TSet<FString>* GlobalConfigFileCache, const TSet<FString>* PluginConfigFileCache, bool bForceFullDynamicLayerUpdate)
 {
+#if UE_THREAD_SAFETY_PROTECT_GETSTRING
+	// don't allow a read while modifying it here
+	UE::TUniqueLock FileWriteLock(InMemoryFile.SharedMutex);
+#endif
+
 	LLM_SCOPE(ELLMTag::ConfigSystem);
 	if (bIsSafeUnloaded)
 	{
 		SafeReload();
 	}
-
+	
 	static bool bDumpIniLoadInfo = FParse::Param(FCommandLine::Get(), TEXT("dumpiniloads"));
 
 	bool bFoundAFile = false;
@@ -3592,6 +3608,11 @@ bool FConfigBranch::AddDynamicLayersToHierarchy(const TArray<FDynamicLayerInfo>&
 
 bool FConfigBranch::AddDynamicLayerStringToHierarchy(const FString& Filename, const FString& Contents, FName Tag, DynamicLayerPriority Priority, FConfigModificationTracker* ModificationTracker)
 {
+#if UE_THREAD_SAFETY_PROTECT_GETSTRING
+	// don't allow a read while modifying it here
+	UE::TUniqueLock FileWriteLock(InMemoryFile.SharedMutex);
+#endif
+
 	LLM_SCOPE(ELLMTag::ConfigSystem);
 	if (bIsSafeUnloaded)
 	{
@@ -3704,6 +3725,11 @@ bool FConfigBranch::RemoveDynamicLayersFromHierarchy(const TArray<FString>& File
 {
 	LLM_SCOPE(ELLMTag::ConfigSystem);
 
+#if UE_THREAD_SAFETY_PROTECT_GETSTRING
+	// this can modify InMemoryFile, so we need to protect it 
+	UE::TUniqueLock WriteFileLock(InMemoryFile.SharedMutex);
+#endif
+
 	if (bIsSafeUnloaded)
 	{
 		SafeReload();
@@ -3763,8 +3789,15 @@ bool FConfigBranch::RemoveDynamicLayersFromHierarchy(const TArray<FString>& File
 
 void FConfigBranch::RemoveTagsFromHierarchy(const TArray<FName>& Tags, FConfigModificationTracker* ModificationTracker)
 {
-	// gather tagged layers
+#if UE_THREAD_SAFETY_PROTECT_GETSTRING
+	// this can modify InMemoryFile, so we need to protect it 
+	UE::TUniqueLock WriteFileLock(InMemoryFile.SharedMutex);
+#endif
+
 	TArray<FString> LayersToRemove;
+
+	
+	// gather tagged layers
 	for (DynamicLayerList::TIterator Node(DynamicLayers.GetHead()); Node; ++Node)
 	{
 		if (Tags.Contains(Node->Tag))
@@ -3784,12 +3817,26 @@ void FConfigBranch::RemoveTagsFromHierarchy(const TArray<FName>& Tags, FConfigMo
 
 void FConfigBranch::SafeUnload()
 {
+#if UE_THREAD_SAFETY_PROTECT_GETSTRING
+	// this can modify InMemoryFile, so we need to protect it 
+	UE::TUniqueLock WriteLock(InMemoryFile.SharedMutex);
+#endif
+	
 	bIsSafeUnloaded = true;
 
+#if UE_THREAD_SAFETY_USE_WRITE_LOCKS
+	// if we are not using full Write Locks, then we we no longer unload the final combined data
+	// this makes thread safety much simpler at the expense of some memory. because we don't clear 
+	// the InMemoryFile, we don't actually need to SafeReload to get the final value, only when we 
+	// add/remove layers or similar. 
+	// This block is the low-memory, slower threading, mode
+	// advanced behavior
 	InMemoryFile.Cleanup();
+#endif
+	
 	CombinedStaticLayers.Cleanup();
 	FinalCombinedLayers.Cleanup();
-
+	
 	// empty the command streams for the static and dynamic layers, but leave any other streams alone
 	// note that we keep the dynamic list around, but without the section data, because we use the 
 	// dynamic layer filename, tag, and priority to load again
@@ -3809,21 +3856,32 @@ void FConfigBranch::SafeReload()
 
 	double StartTime = FPlatformTime::Seconds();
 
+	// read dynamic layers back in from disk
+	TArray<FDynamicLayerInfo> ReloadInfos;
+	
+#if UE_THREAD_SAFETY_PROTECT_GETSTRING
+	UE::TUniqueLock FileWriteLock(InMemoryFile.SharedMutex);
+#endif
+
+	// in case another thread reloaded while we were blocked, skip this
+	if (!bIsSafeUnloaded)
+	{
+		return;
+	}
+	
 	// read static layers back in from disk
 	// @todo make sure we only Unload from GConfig
 	FConfigContext Context = FConfigContext::ReadIntoConfigSystem(GConfig, Platform.ToString());
 	Context.Branch = this;
 	Context.DestIniFilename = IniPath;
+	Context.bIsSafeReloading = true;
 	Context.Load(*IniName.ToString());
-
-	// read dynamic layers back in from disk
-	TArray<FDynamicLayerInfo> ReloadInfos;
-
+			
 	FConfigBranch::DynamicLayerList::TDoubleLinkedListNode* CurrentNode = DynamicLayers.GetHead(), * NextNode;
 	while (CurrentNode != nullptr)
 	{
 		NextNode = CurrentNode->GetNextNode();
-
+		
 		// any never unload layers we just leave in the stream
 		FConfigCommandStream& S = *CurrentNode->GetValue();
 		if (!S.bNeverSafeUnload)
@@ -3848,6 +3906,10 @@ void FConfigBranch::SafeReload()
 
 void FConfigBranch::ReapplyLayers()
 {
+#if UE_THREAD_SAFETY_PROTECT_GETSTRING
+	UE::TUniqueLock FileWriteLock(InMemoryFile.SharedMutex);
+#endif
+	
 	// calculate a patch so we don't lose in-memory changes
 	FConfigCommandStream Patch = CalculateDiff(FinalCombinedLayers, InMemoryFile);
 
@@ -4116,7 +4178,7 @@ void FConfigCacheIni::Tick(float DeltaSeconds)
 FConfigBranch* FConfigCacheIni::FindBranchWithNoReload(FName BaseIniName, const FString& Filename)
 {
 	// look for a known file, if there's no ini extension
-	FConfigBranch* Branch = KnownFiles.GetBranch(BaseIniName);
+	FConfigBranch* Branch = BaseIniName == NAME_None ? nullptr : KnownFiles.GetBranch(BaseIniName);
 
 	if (Branch == nullptr)
 	{
@@ -4152,16 +4214,19 @@ FConfigBranch* FConfigCacheIni::FindBranch(FName BaseIniName, const FString& Fil
 {
 	FConfigBranch* Branch = FindBranchWithNoReload(BaseIniName, Filename);
 
-	if (Branch && Branch->bIsSafeUnloaded)
-	{
-		Branch->SafeReload();
-	}
-
-	// track that this branch is being used, so re-set the time
-	if (Branch && Branch->InactiveTimer >= 0 && GTimeToUnloadConfig > 0)
+	if (Branch)
 	{
-		Branch->InactiveTimer = FPlatformTime::Seconds();
-		UE_LOG(LogConfig, Verbose, TEXT("Resetting InactiveTimer for %s"), *Branch->IniName.ToString());
+		// track that this branch is being used, so re-set the time
+		if (Branch->InactiveTimer >= 0 && GTimeToUnloadConfig > 0)
+		{
+			Branch->InactiveTimer = FPlatformTime::Seconds();
+			UE_LOG(LogConfig, Verbose, TEXT("Resetting InactiveTimer for %s"), *Branch->IniName.ToString());
+		}
+		
+		if (Branch->bIsSafeUnloaded)
+		{
+			Branch->SafeReload();
+		}
 	}
 
 	return Branch;
@@ -4655,19 +4720,65 @@ void FConfigCacheIni::Detach(const FString& Filename)
 		File->NoSave = 1;
 }
 
+static FConfigBranch* FindOrLoadNoSafeReload(FConfigCacheIni* ConfigSystem, const FString& Filename)
+{
+	// check for non-filenames
+	if(Filename.Len() == 0)
+	{
+		return nullptr;
+	}
+
+	// we need to get the branch without reloading here so that we can put the Reload and Get() be protected
+	// by the same block, which is held buy the the Branch/File we are looking up
+	FConfigBranch* Branch = ConfigSystem->FindBranchWithNoReload(NAME_None, Filename);
+	if (Branch == nullptr)
+	{
+		// this will load if possible, with Normalized path (and log a warning if it wasn't match)
+		ConfigSystem->Find(*Filename);
+		
+		// need to find with Normalized path
+		Branch = ConfigSystem->FindBranchWithNoReload(NAME_None, ConfigSystem->NormalizeConfigIniPath(Filename));
+	}
+
+	return Branch;
+}
+
 bool FConfigCacheIni::GetString( const TCHAR* Section, const TCHAR* Key, FString& Value, const FString& Filename )
 {
-	FRemoteConfig::Get()->FinishRead(*Filename); // Ensure the remote file has been loaded and processed
-	FConfigFile* File = Find(Filename);
-	if( !File )
+	// get the file without a SafeReload possible, so no lock is needed
+	FConfigBranch* Branch = FindOrLoadNoSafeReload(this, Filename);
+	if (Branch == nullptr)
 	{
 		return false;
 	}
-	const FConfigSection* Sec = File->FindSection( Section );
+	
+	// due to some rare-but-possible edge cases, we need to either keep InMemoryFile around when we 
+	// SafeUnload, and _not_ call SafeReload here, or we need a writelock around both the Unload and the
+	// GetString functionality (or we would Unload between the Reload call and the Get). A ReadLock isn't
+	// enough, and was causing deadlocks
+#if UE_THREAD_SAFETY_USE_WRITE_LOCKS
+	// don't read while something may be modifying the File (generally a Branch's InMmeoryFile)
+	// frustratingly we need 
+	UE::TUniqueLock WriteLock(Branch->InMemoryFile.SharedMutex);
+	
+	if (Branch->bIsSafeUnloaded)
+	{
+		Branch->SafeReload();
+	}
+#else
+	// we do need a ReadLock because dynamiclayers can still modify InMemoryFile
+	// (if we did a SafeReload, we would need a WriteLock which makes this function ~2-4x slower!)
+	UE::TSharedLock ReadLock(Branch->InMemoryFile.SharedMutex);	
+#endif
+	
+	FRemoteConfig::Get()->FinishRead(*Filename); // Ensure the remote file has been loaded and processed
+
+	
+	const FConfigSection* Sec = Branch->InMemoryFile.FindSection( Section );
 	if( !Sec )
 	{
 #if !UE_BUILD_SHIPPING
-		CheckLongSectionNames( Section, File );
+		CheckLongSectionNames( Section, &Branch->InMemoryFile );
 #endif
 		return false;
 	}
@@ -4685,17 +4796,39 @@ bool FConfigCacheIni::GetString( const TCHAR* Section, const TCHAR* Key, FString
 
 bool FConfigCacheIni::GetText( const TCHAR* Section, const TCHAR* Key, FText& Value, const FString& Filename )
 {
-	FRemoteConfig::Get()->FinishRead(*Filename); // Ensure the remote file has been loaded and processed
-	FConfigFile* File = Find(Filename);
-	if( !File )
+	// get the file without a SafeReload possible, so no lock is needed
+	FConfigBranch* Branch = FindOrLoadNoSafeReload(this, Filename);	
+	if (Branch == nullptr)
 	{
 		return false;
 	}
-	const FConfigSection* Sec = File->FindSection( Section );
+
+	// due to some rare-but-possible edge cases, we need to either keep InMemoryFile around when we 
+	// SafeUnload, and _not_ call SafeReload here, or we need a writelock around both the Unload and the
+	// GetString functionality (or we would Unload between the Reload call and the Get). A ReadLock isn't
+	// enough, and was causing deadlocks
+#if UE_THREAD_SAFETY_USE_WRITE_LOCKS
+	// don't read while something may be modifying the File (generally a Branch's InMmeoryFile)
+	// frustratingly we need 
+	UE::TUniqueLock WriteLock(Branch->InMemoryFile.SharedMutex);
+	
+	if (Branch->bIsSafeUnloaded)
+	{
+		Branch->SafeReload();
+	}
+#else
+	// we do need a ReadLock because dynamiclayers can still modify InMemoryFile
+	// (if we did a SafeReload, we would need a WriteLock which makes this function ~2-4x slower!)
+	UE::TSharedLock ReadLock(Branch->InMemoryFile.SharedMutex);	
+#endif
+
+	FRemoteConfig::Get()->FinishRead(*Filename); // Ensure the remote file has been loaded and processed
+	
+	const FConfigSection* Sec = Branch->InMemoryFile.FindSection( Section );
 	if( !Sec )
 	{
 #if !UE_BUILD_SHIPPING
-		CheckLongSectionNames( Section, File );
+		CheckLongSectionNames( Section, &Branch->InMemoryFile );
 #endif
 		return false;
 	}
diff --git a/Engine/Source/Runtime/Core/Private/Misc/ConfigContext.cpp b/Engine/Source/Runtime/Core/Private/Misc/ConfigContext.cpp
--- a/Engine/Source/Runtime/Core/Private/Misc/ConfigContext.cpp
+++ b/Engine/Source/Runtime/Core/Private/Misc/ConfigContext.cpp
@@ -1206,8 +1206,9 @@ bool FConfigContext::LoadIniFileHierarchy()
 
 	Branch->FinalCombinedLayers = Branch->InMemoryFile;
 
-	// look for some standard dynamic layers that we know to look for now
-	if (!bIsForPlugin && !ProjectConfigDir.IsEmpty())
+	// look for some standard dynamic layers that we know to look for now (and not safe reloading, which will already handle this
+	// since it will reload all dynamic layers already)
+	if (!bIsForPlugin && !ProjectConfigDir.IsEmpty() && !bIsSafeReloading)
 	{
 		FString ProjectPluginsOverride = FPaths::Combine(ProjectConfigDir, FString::Printf(TEXT("PluginOverride%s.ini"), *BaseIniName));
 
diff --git a/Engine/Source/Runtime/Core/Public/Misc/ConfigCacheIni.h b/Engine/Source/Runtime/Core/Public/Misc/ConfigCacheIni.h
--- a/Engine/Source/Runtime/Core/Public/Misc/ConfigCacheIni.h
+++ b/Engine/Source/Runtime/Core/Public/Misc/ConfigCacheIni.h
@@ -7,6 +7,10 @@
 #pragma once
 
 #include "Algo/Reverse.h"
+#include "Async/TransactionallySafeSharedMutex.h"
+#include "Async/SharedMutex.h"
+#include "Async/TransactionallySafeSharedRecursiveMutex.h"
+#include "Async/SharedRecursiveMutex.h"
 #include "Containers/Array.h"
 #include "Containers/Map.h"
 #include "Containers/Set.h"
@@ -54,6 +58,8 @@ CORE_API DECLARE_LOG_CATEGORY_EXTERN(LogConfig, Log, All);
 #define CONFIG_CAN_SAVE_COMMENTS (WITH_EDITOR)
 
 
+using FConfigSystemMutex = UE::FTransactionallySafeSharedRecursiveMutex;
+
 ///////////////////////////////////////////////////////////////////////////////
 // Info about the deprecation of functions returning non-const FConfigSections:
 //   In a future change, we will be tracking operations done to config files (via GConfig, etc) for improved saving and allowing for plugin unloading.
@@ -605,6 +611,9 @@ private:
 	// if this is set, then we track changes made to sections for saving or replaying later (currently unused)
 	FConfigCommandStream* ChangeTracker = nullptr;
 	
+	// thread safety for reading while something is modifying this file
+	FConfigSystemMutex SharedMutex;
+	
 #if UE_WITH_CONFIG_TRACKING
 	mutable TRefCountPtr<UE::ConfigAccessTracking::FFile> FileAccess;
 #endif
diff --git a/Engine/Source/Runtime/Core/Public/Misc/ConfigContext.h b/Engine/Source/Runtime/Core/Public/Misc/ConfigContext.h
--- a/Engine/Source/Runtime/Core/Public/Misc/ConfigContext.h
+++ b/Engine/Source/Runtime/Core/Public/Misc/ConfigContext.h
@@ -251,6 +251,7 @@ public:
 	bool bIsMakingBinaryConfig = false;
 	bool bIsFixingUpAfterBinaryConfig = false;
 	bool bCanRemoveBranchIfEmpty = false;
+	bool bIsSafeReloading = false;
 
 	DynamicLayerPriority PluginModificationPriority;
 
