[Core] Apply AUTORTFM_INFER to TFunctionWithContext::Call()

Fixes a "Cannot call AutoRTFM-disabled function from an AutoRTFM-enabled function" error when compiling AutoRTFMTests Linux Debug.

#Synced-CL 49690413

diff --git a/Engine/Source/Runtime/Core/Public/Templates/FunctionWithContext.h b/Engine/Source/Runtime/Core/Public/Templates/FunctionWithContext.h
--- a/Engine/Source/Runtime/Core/Public/Templates/FunctionWithContext.h
+++ b/Engine/Source/Runtime/Core/Public/Templates/FunctionWithContext.h
@@ -99,7 +99,7 @@ public:
 
 private:
 	template <typename CallableType>
-	static ReturnType Call(void* Callable, ArgTypes... Args)
+	AUTORTFM_INFER static ReturnType Call(void* Callable, ArgTypes... Args)
 	{
 		return Invoke(*(CallableType*)Callable, (ArgTypes&&)Args...);
 	}
