Make `FConcurrentLinearBulkObjectAllocator` work under transactions.

#Synced-CL 49722996

diff --git a/Engine/Source/Programs/AutoRTFMTests/Private/UECoreTests.cpp b/Engine/Source/Programs/AutoRTFMTests/Private/UECoreTests.cpp
--- a/Engine/Source/Programs/AutoRTFMTests/Private/UECoreTests.cpp
+++ b/Engine/Source/Programs/AutoRTFMTests/Private/UECoreTests.cpp
@@ -2446,3 +2446,23 @@ TEST_CASE("UECore.FConcurrentLinearAllocator")
 		FConcurrentLinearAllocator::Free(Allocation);
 	}
 }
+
+TEST_CASE("UECore.FConcurrentLinearBulkObjectAllocator")
+{
+	SECTION("Malloc")
+	{
+		FConcurrentLinearBulkObjectAllocator Allocator;
+		int32* Allocation = nullptr;
+
+		AutoRTFM::Testing::Commit([&]
+		{
+			Allocation = static_cast<int32*>(Allocator.Malloc(sizeof(int32), alignof(int32)));
+			*Allocation = 42;
+		});
+
+		REQUIRE(Allocation);
+		REQUIRE(*Allocation == 42);
+	}
+}
+
+
diff --git a/Engine/Source/Runtime/Core/Public/Experimental/ConcurrentLinearAllocator.h b/Engine/Source/Runtime/Core/Public/Experimental/ConcurrentLinearAllocator.h
--- a/Engine/Source/Runtime/Core/Public/Experimental/ConcurrentLinearAllocator.h
+++ b/Engine/Source/Runtime/Core/Public/Experimental/ConcurrentLinearAllocator.h
@@ -837,7 +837,9 @@ public:
 		}
 	}
 
-	inline void* Malloc(SIZE_T Size, uint32 Alignment)
+	// In a transaction because the bulk allocator only has bulk delete, on abort we're just wasting the memory.
+	// So we run this code in the open (non-transactionally) always.
+	AUTORTFM_OPEN inline void* Malloc(SIZE_T Size, uint32 Alignment)
 	{
 		SIZE_T TotalSize = Align(sizeof(FAllocation), Alignment) + Size;
 		FAllocation* Allocation = new (TConcurrentLinearAllocator<BlockAllocationTag>::Malloc(TotalSize, FMath::Max(static_cast<uint32>(alignof(FAllocation)), Alignment))) FAllocation();
