CIS fix.

diff --git a/Engine/Source/Runtime/Core/Public/Templates/Invoke.h b/Engine/Source/Runtime/Core/Public/Templates/Invoke.h
--- a/Engine/Source/Runtime/Core/Public/Templates/Invoke.h
+++ b/Engine/Source/Runtime/Core/Public/Templates/Invoke.h
@@ -139,5 +139,5 @@ using TInvokeResult_T = typename TInvokeResult<FuncType, ArgTypes...>::Type;
 #endif
 
 #if UE_ENABLE_INCLUDE_ORDER_DEPRECATED_IN_5_8
-#include "Traits/MemberFunctionPtrOuter.h"
+#include "Traits/MemberFunctionPtrOuterDeprecated.h"
 #endif
diff --git a/Engine/Source/Runtime/Core/Public/Traits/MemberFunctionPtrOuter.h b/Engine/Source/Runtime/Core/Public/Traits/MemberFunctionPtrOuter.h
--- a/Engine/Source/Runtime/Core/Public/Traits/MemberFunctionPtrOuter.h
+++ b/Engine/Source/Runtime/Core/Public/Traits/MemberFunctionPtrOuter.h
@@ -5,29 +5,4 @@
 #include "Misc/CoreMiscDefines.h"
 
 UE_DEPRECATED_HEADER(5.8, "MemberFunctionPtrOuter.h has been deprecated.")
-
-/**
- * Type trait which yields the type of the class given a pointer to a member function of that class, e.g.:
- *
- * TMemberFunctionPtrOuter_T<decltype(&FVector::Dot)> == FVector
- */
-template <typename T>
-struct UE_DEPRECATED(5.8, "TMemberFunctionPtrOuter has been deprecated, please use TMemberPtrOuter instead.") TMemberFunctionPtrOuter;
-
-template <typename ReturnType, typename ObjectType, typename... ArgTypes> struct UE_DEPRECATED(5.8, "TMemberFunctionPtrOuter has been deprecated, use TMemberPtrOuter instead.") TMemberFunctionPtrOuter<ReturnType (ObjectType::*)(ArgTypes...)                 > { using Type = ObjectType; };
-template <typename ReturnType, typename ObjectType, typename... ArgTypes> struct UE_DEPRECATED(5.8, "TMemberFunctionPtrOuter has been deprecated, use TMemberPtrOuter instead.") TMemberFunctionPtrOuter<ReturnType (ObjectType::*)(ArgTypes...)               & > { using Type = ObjectType; };
-template <typename ReturnType, typename ObjectType, typename... ArgTypes> struct UE_DEPRECATED(5.8, "TMemberFunctionPtrOuter has been deprecated, use TMemberPtrOuter instead.") TMemberFunctionPtrOuter<ReturnType (ObjectType::*)(ArgTypes...)               &&> { using Type = ObjectType; };
-template <typename ReturnType, typename ObjectType, typename... ArgTypes> struct UE_DEPRECATED(5.8, "TMemberFunctionPtrOuter has been deprecated, use TMemberPtrOuter instead.") TMemberFunctionPtrOuter<ReturnType (ObjectType::*)(ArgTypes...) const           > { using Type = ObjectType; };
-template <typename ReturnType, typename ObjectType, typename... ArgTypes> struct UE_DEPRECATED(5.8, "TMemberFunctionPtrOuter has been deprecated, use TMemberPtrOuter instead.") TMemberFunctionPtrOuter<ReturnType (ObjectType::*)(ArgTypes...) const         & > { using Type = ObjectType; };
-template <typename ReturnType, typename ObjectType, typename... ArgTypes> struct UE_DEPRECATED(5.8, "TMemberFunctionPtrOuter has been deprecated, use TMemberPtrOuter instead.") TMemberFunctionPtrOuter<ReturnType (ObjectType::*)(ArgTypes...) const         &&> { using Type = ObjectType; };
-template <typename ReturnType, typename ObjectType, typename... ArgTypes> struct UE_DEPRECATED(5.8, "TMemberFunctionPtrOuter has been deprecated, use TMemberPtrOuter instead.") TMemberFunctionPtrOuter<ReturnType (ObjectType::*)(ArgTypes...)       volatile  > { using Type = ObjectType; };
-template <typename ReturnType, typename ObjectType, typename... ArgTypes> struct UE_DEPRECATED(5.8, "TMemberFunctionPtrOuter has been deprecated, use TMemberPtrOuter instead.") TMemberFunctionPtrOuter<ReturnType (ObjectType::*)(ArgTypes...)       volatile& > { using Type = ObjectType; };
-template <typename ReturnType, typename ObjectType, typename... ArgTypes> struct UE_DEPRECATED(5.8, "TMemberFunctionPtrOuter has been deprecated, use TMemberPtrOuter instead.") TMemberFunctionPtrOuter<ReturnType (ObjectType::*)(ArgTypes...)       volatile&&> { using Type = ObjectType; };
-template <typename ReturnType, typename ObjectType, typename... ArgTypes> struct UE_DEPRECATED(5.8, "TMemberFunctionPtrOuter has been deprecated, use TMemberPtrOuter instead.") TMemberFunctionPtrOuter<ReturnType (ObjectType::*)(ArgTypes...) const volatile  > { using Type = ObjectType; };
-template <typename ReturnType, typename ObjectType, typename... ArgTypes> struct UE_DEPRECATED(5.8, "TMemberFunctionPtrOuter has been deprecated, use TMemberPtrOuter instead.") TMemberFunctionPtrOuter<ReturnType (ObjectType::*)(ArgTypes...) const volatile& > { using Type = ObjectType; };
-template <typename ReturnType, typename ObjectType, typename... ArgTypes> struct UE_DEPRECATED(5.8, "TMemberFunctionPtrOuter has been deprecated, use TMemberPtrOuter instead.") TMemberFunctionPtrOuter<ReturnType (ObjectType::*)(ArgTypes...) const volatile&&> { using Type = ObjectType; };
-
-PRAGMA_DISABLE_DEPRECATION_WARNINGS
-template <typename T>
-using TMemberFunctionPtrOuter_T UE_DEPRECATED(5.8, "TMemberFunctionPtrOuter_T has been deprecated, use TMemberPtrOuter_T instead.") = typename TMemberFunctionPtrOuter<T>::Type;
-PRAGMA_ENABLE_DEPRECATION_WARNINGS
+#include "Traits/MemberFunctionPtrOuterDeprecated.h"
diff --git a/Engine/Source/Runtime/Core/Public/Traits/MemberFunctionPtrOuterDeprecated.h b/Engine/Source/Runtime/Core/Public/Traits/MemberFunctionPtrOuterDeprecated.h
new file mode 100644
--- /dev/null
+++ b/Engine/Source/Runtime/Core/Public/Traits/MemberFunctionPtrOuterDeprecated.h
@@ -0,0 +1,36 @@
+// Copyright Epic Games, Inc. All Rights Reserved.
+
+#pragma once
+
+#include "Misc/CoreMiscDefines.h"
+
+// THIS HEADER IS DEPRECATD and shouldn't be included directly.
+//
+// It is only intended to be included by other deprecated contexts.
+// When those deprecated contexts are removed, this header can be removed too.
+
+/**
+ * Type trait which yields the type of the class given a pointer to a member function of that class, e.g.:
+ *
+ * TMemberFunctionPtrOuter_T<decltype(&FVector::Dot)> == FVector
+ */
+template <typename T>
+struct UE_DEPRECATED(5.8, "TMemberFunctionPtrOuter has been deprecated, please use TMemberPtrOuter instead.") TMemberFunctionPtrOuter;
+
+template <typename ReturnType, typename ObjectType, typename... ArgTypes> struct UE_DEPRECATED(5.8, "TMemberFunctionPtrOuter has been deprecated, use TMemberPtrOuter instead.") TMemberFunctionPtrOuter<ReturnType (ObjectType::*)(ArgTypes...)                 > { using Type = ObjectType; };
+template <typename ReturnType, typename ObjectType, typename... ArgTypes> struct UE_DEPRECATED(5.8, "TMemberFunctionPtrOuter has been deprecated, use TMemberPtrOuter instead.") TMemberFunctionPtrOuter<ReturnType (ObjectType::*)(ArgTypes...)               & > { using Type = ObjectType; };
+template <typename ReturnType, typename ObjectType, typename... ArgTypes> struct UE_DEPRECATED(5.8, "TMemberFunctionPtrOuter has been deprecated, use TMemberPtrOuter instead.") TMemberFunctionPtrOuter<ReturnType (ObjectType::*)(ArgTypes...)               &&> { using Type = ObjectType; };
+template <typename ReturnType, typename ObjectType, typename... ArgTypes> struct UE_DEPRECATED(5.8, "TMemberFunctionPtrOuter has been deprecated, use TMemberPtrOuter instead.") TMemberFunctionPtrOuter<ReturnType (ObjectType::*)(ArgTypes...) const           > { using Type = ObjectType; };
+template <typename ReturnType, typename ObjectType, typename... ArgTypes> struct UE_DEPRECATED(5.8, "TMemberFunctionPtrOuter has been deprecated, use TMemberPtrOuter instead.") TMemberFunctionPtrOuter<ReturnType (ObjectType::*)(ArgTypes...) const         & > { using Type = ObjectType; };
+template <typename ReturnType, typename ObjectType, typename... ArgTypes> struct UE_DEPRECATED(5.8, "TMemberFunctionPtrOuter has been deprecated, use TMemberPtrOuter instead.") TMemberFunctionPtrOuter<ReturnType (ObjectType::*)(ArgTypes...) const         &&> { using Type = ObjectType; };
+template <typename ReturnType, typename ObjectType, typename... ArgTypes> struct UE_DEPRECATED(5.8, "TMemberFunctionPtrOuter has been deprecated, use TMemberPtrOuter instead.") TMemberFunctionPtrOuter<ReturnType (ObjectType::*)(ArgTypes...)       volatile  > { using Type = ObjectType; };
+template <typename ReturnType, typename ObjectType, typename... ArgTypes> struct UE_DEPRECATED(5.8, "TMemberFunctionPtrOuter has been deprecated, use TMemberPtrOuter instead.") TMemberFunctionPtrOuter<ReturnType (ObjectType::*)(ArgTypes...)       volatile& > { using Type = ObjectType; };
+template <typename ReturnType, typename ObjectType, typename... ArgTypes> struct UE_DEPRECATED(5.8, "TMemberFunctionPtrOuter has been deprecated, use TMemberPtrOuter instead.") TMemberFunctionPtrOuter<ReturnType (ObjectType::*)(ArgTypes...)       volatile&&> { using Type = ObjectType; };
+template <typename ReturnType, typename ObjectType, typename... ArgTypes> struct UE_DEPRECATED(5.8, "TMemberFunctionPtrOuter has been deprecated, use TMemberPtrOuter instead.") TMemberFunctionPtrOuter<ReturnType (ObjectType::*)(ArgTypes...) const volatile  > { using Type = ObjectType; };
+template <typename ReturnType, typename ObjectType, typename... ArgTypes> struct UE_DEPRECATED(5.8, "TMemberFunctionPtrOuter has been deprecated, use TMemberPtrOuter instead.") TMemberFunctionPtrOuter<ReturnType (ObjectType::*)(ArgTypes...) const volatile& > { using Type = ObjectType; };
+template <typename ReturnType, typename ObjectType, typename... ArgTypes> struct UE_DEPRECATED(5.8, "TMemberFunctionPtrOuter has been deprecated, use TMemberPtrOuter instead.") TMemberFunctionPtrOuter<ReturnType (ObjectType::*)(ArgTypes...) const volatile&&> { using Type = ObjectType; };
+
+PRAGMA_DISABLE_DEPRECATION_WARNINGS
+template <typename T>
+using TMemberFunctionPtrOuter_T UE_DEPRECATED(5.8, "TMemberFunctionPtrOuter_T has been deprecated, use TMemberPtrOuter_T instead.") = typename TMemberFunctionPtrOuter<T>::Type;
+PRAGMA_ENABLE_DEPRECATION_WARNINGS
