Ensure the OS locale is canonized before querying AllAvailableCulturesMap

#rnx
#Synced-CL 49717872

diff --git a/Engine/Source/Runtime/Core/Private/Internationalization/ICUInternationalization.cpp b/Engine/Source/Runtime/Core/Private/Internationalization/ICUInternationalization.cpp
--- a/Engine/Source/Runtime/Core/Private/Internationalization/ICUInternationalization.cpp
+++ b/Engine/Source/Runtime/Core/Private/Internationalization/ICUInternationalization.cpp
@@ -174,10 +174,10 @@ bool FICUInternationalization::Initialize()
 		I18N->InvariantCulture = FindOrMakeCanonizedCulture(FString(), EAllowDefaultCultureFallback::Yes);
 	}
 	I18N->DefaultLanguage = FindOrMakeCulture(FPlatformMisc::GetDefaultLanguage(), EAllowDefaultCultureFallback::Yes);
-	if (const FString& DefaultLocaleName = FPlatformMisc::GetDefaultLocale();
+	if (const FString& DefaultLocaleName = FCultureImplementation::GetCanonicalName(FPlatformMisc::GetDefaultLocale(), *I18N);
 		AllAvailableCulturesMap.Contains(DefaultLocaleName))
 	{
-		I18N->DefaultLocale = FindOrMakeCulture(DefaultLocaleName, EAllowDefaultCultureFallback::Yes);
+		I18N->DefaultLocale = FindOrMakeCanonizedCulture(DefaultLocaleName, EAllowDefaultCultureFallback::Yes);
 	}
 	else
 	{
