LLM: Move defines used in LLM.cpp from LLMPrivate.h into a separate header so they can be read from new headers coming in the upcoming Allocations change. Also add some minor syntax and whitespace changes.
#rnx
#Synced-CL 49725844

diff --git a/Engine/Source/Runtime/Core/Private/HAL/LLM/LLM.cpp b/Engine/Source/Runtime/Core/Private/HAL/LLM/LLM.cpp
--- a/Engine/Source/Runtime/Core/Private/HAL/LLM/LLM.cpp
+++ b/Engine/Source/Runtime/Core/Private/HAL/LLM/LLM.cpp
@@ -7,10 +7,10 @@
 
 #include "Async/SharedLock.h"
 #include "Async/UniqueLock.h"
-#include "HAL/LowLevelMemStats.h"
 #include "HAL/FileManager.h"
 #include "HAL/IConsoleManager.h"
 #include "HAL/LLM/LLMPrivate.h"
+#include "HAL/LowLevelMemStats.h"
 #include "HAL/ThreadHeartBeat.h"
 #include "HAL/PlatformMemory.h" // for page allocation association.
 #include "MemPro/MemProProfiler.h"
@@ -4703,7 +4703,7 @@ const FTagData* FLLMTracker::FLowLevelAllocInfo::GetTag(FLLMGlobals& InLLMRef
 
 #if LLM_ENABLED_FULL_TAGS
 
-int32 FLLMTracker::FLowLevelAllocInfo::GetCompressedTag(ELLMTagSet InTagSet)
+int32 FLLMTracker::FLowLevelAllocInfo::GetCompressedTag(ELLMTagSet InTagSet) const
 {
 #if !LLM_ALLOW_ASSETS_TAGS && !LLM_ALLOW_UOBJECTCLASSES_TAGS
 	return (InTagSet == ELLMTagSet::None) ? Tag : InvalidCompressedTagValue;
@@ -4728,7 +4728,7 @@ void FLLMTracker::FLowLevelAllocInfo::SetCompressedTag(int32 InTag, ELLMTagSet I
 
 #else // LLM_ENABLED_FULL_TAGS
 
-ELLMTag FLLMTracker::FLowLevelAllocInfo::GetCompressedTag()
+ELLMTag FLLMTracker::FLowLevelAllocInfo::GetCompressedTag() const
 {
 	return Tag;
 }
@@ -4855,7 +4855,8 @@ void FLLMTracker::PopTag(ELLMTagSet TagSet)
 	State->PopTag(TagSet);
 }
 
-void FLLMTracker::TrackAllocation(const void* Ptr, int64 Size, ELLMTag DefaultTag, ELLMAllocType AllocType, bool bTrackInMemPro)
+void FLLMTracker::TrackAllocation(const void* Ptr, int64 Size, ELLMTag DefaultTag, ELLMAllocType AllocType,
+	bool bTrackInMemPro)
 {
 	FLLMThreadState* State = GetOrCreateState();
 	const FTagData* TagData = State->GetTopTag(ELLMTagSet::None);
@@ -4866,7 +4867,8 @@ void FLLMTracker::TrackAllocation(const void* Ptr, int64 Size, ELLMTag DefaultTa
 	TrackAllocation(Ptr, Size, TagData, AllocType, State, bTrackInMemPro);
 }
 
-void FLLMTracker::TrackAllocation(const void* Ptr, int64 Size, FName DefaultTag, ELLMAllocType AllocType, bool bTrackInMemPro)
+void FLLMTracker::TrackAllocation(const void* Ptr, int64 Size, FName DefaultTag, ELLMAllocType AllocType,
+	bool bTrackInMemPro)
 {
 	FLLMThreadState* State = GetOrCreateState();
 	const FTagData* TagData = State->GetTopTag(ELLMTagSet::None);
@@ -5528,9 +5530,9 @@ bool FLLMTracker::DumpForkedAllocationInfo()
 
 
 #if LLM_ENABLED_FULL_TAGS
-	for (TPair<int32, FCounts>& P : CountsPerTag)
+	for (const TPair<int32, FCounts>& P : CountsPerTag)
 #else
-	for (TPair<ELLMTag, FCounts>& P : CountsPerTag)
+	for (const TPair<ELLMTag, FCounts>& P : CountsPerTag)
 #endif
 	{
 		FLowLevelAllocInfo AllocInfoPlaceholder;
@@ -5737,7 +5739,8 @@ TArray<const FTagData*> FLLMTracker::GetTagDatas(ELLMTagSet TagSet)
 	});
 }
 
-void FLLMTracker::GetTagsNamesWithAmount(TMap<FName, uint64>& OutTagsNamesWithAmount, ELLMTagSet TagSet /* = ELLMTagSet::None */)
+void FLLMTracker::GetTagsNamesWithAmount(TMap<FName, uint64>& OutTagsNamesWithAmount,
+	ELLMTagSet TagSet /* = ELLMTagSet::None */)
 {
 	for (const TPair<const FTagData*, FTrackerTagSizeData>& It : TagSizes)
 	{
@@ -5950,7 +5953,7 @@ void FLLMThreadState::PopTag(ELLMTagSet TagSet)
 	TagStack[static_cast<int32>(TagSet)].Pop(EAllowShrinking::No);
 }
 
-const FTagData* FLLMThreadState::GetTopTag(ELLMTagSet TagSet)
+const FTagData* FLLMThreadState::GetTopTag(ELLMTagSet TagSet) const
 {
 	if (TagStack[static_cast<int32>(TagSet)].Num() == 0)
 	{
diff --git a/Engine/Source/Runtime/Core/Private/HAL/LLM/LLMPrivate.h b/Engine/Source/Runtime/Core/Private/HAL/LLM/LLMPrivate.h
--- a/Engine/Source/Runtime/Core/Private/HAL/LLM/LLMPrivate.h
+++ b/Engine/Source/Runtime/Core/Private/HAL/LLM/LLMPrivate.h
@@ -10,72 +10,18 @@
 #include "Containers/Set.h"
 #include "Containers/SortedMap.h"
 #include "HAL/CriticalSection.h"
+#include "HAL/LLM/LLMPrivateDefines.h"
 #include "HAL/LLM/MemoryUtils.h"
 #include "HAL/PlatformMutex.h"
 #include "Misc/Optional.h"
 #include "Misc/ScopeRWLock.h"
 #include "ProfilingDebugging/CsvProfiler.h"
 
-// Private defines configuring LLM; see also public defines in LowLevelMemTracker.h
-// LLM_ALLOW_NAMES_TAGS: Set to 1 to allow arbitrary FNames to be used as tags, at the cost of more LLM memory usage per allocation. Set to 0 to store only the top-level ELLMTag for each allocation.
-// LLM_SCOPES always use FNames for the definition of their tag
-// Storing these tags on each allocation however requires 4 bytes per allocation
-// To reduce the memory overhead of LLM, this can be reduced to 1 byte per allocation at the cost of showing only the top-level 256 Tags; tags are replaced with their containing toplevel tag during allocation.
-// This setting is ignored if features requiring more information per allocation (LLM_ALLOW_STATS, LLM_ALLOW_ASSETS_TAGS) are enabled
-#ifndef LLM_ALLOW_NAMES_TAGS
-	#define LLM_ALLOW_NAMES_TAGS 1
-#endif
-
-// Enable storing the full tag for each alloc if (1) arbitrary names tags are allowed or (2) Stats are allowed (stats use names tags) or (3) Asset tags are allowed (assets use names tags)
-#define LLM_ENABLED_FULL_TAGS (LLM_ALLOW_NAMES_TAGS || LLM_ALLOW_STATS || LLM_ALLOW_ASSETS_TAGS)
-
-// Whether to enable running with reduced threads. This is currently enabled because the engine crashes with -norenderthread
-#define LLM_ENABLED_REDUCE_THREADS 0
-
-// this controls if the commandline is used to enable tracking, or to disable it. If LLM_COMMANDLINE_ENABLES_FUNCTIONALITY is true, 
-// then tracking will only happen through Engine::Init(), at which point it will be disabled unless the commandline tells 
-// it to keep going (with -llm). If LLM_COMMANDLINE_ENABLES_FUNCTIONALITY is false, then tracking will be on unless the commandline
-// disables it (with -nollm)
-#ifndef LLM_COMMANDLINE_ENABLES_FUNCTIONALITY
-	#define LLM_COMMANDLINE_ENABLES_FUNCTIONALITY 1
-#endif 
-
-// when set to 1, forces LLM to be enabled without having to parse the command line.
-#ifndef LLM_AUTO_ENABLE
-	#define LLM_AUTO_ENABLE 0
-#endif
-
-// There is a little memory and cpu overhead in tracking peak memory but it is generally more useful than current memory.
-// Disable if you need a little more memory or speed
-#define LLM_ENABLED_TRACK_PEAK_MEMORY 1
-
-// Enable to allow memory tracking to be passed through to external tools without the extra overhead of LLM also tracking allocations.
-#define UE_ONLY_USE_PLATFORM_TRACKER 0
-
-#define LLM_CSV_PROFILER_WRITER_ENABLED CSV_PROFILER_STATS
-
-static_assert(static_cast<uint32>(ELLMTagSet::Max) == 4, "You added a new TagSet, please update the #ifs below");
-#if LLM_ALLOW_ASSETS_TAGS
-#define LLM_TAGSET_MAX 3
-#else
-#define LLM_TAGSET_MAX 2
-#endif
-
 void GlobalRegisterTagDeclaration(FLLMTagDeclaration& TagDeclaration);
 
 namespace UE::LLMPrivate
 {
 
-enum class ETagReferenceSource : uint8
-{
-	Scope,
-	Declare,
-	EnumTag,
-	CustomEnumTag,
-	FunctionAPI,
-	ImplicitParent
-};
-
 struct FTagDataNameKey
 {
 	FName Name;
@@ -379,7 +325,7 @@ public:
 
 	void PushTag(const FTagData* TagData, ELLMTagSet TagSet);
 	void PopTag(ELLMTagSet TagSet);
-	const FTagData* GetTopTag(ELLMTagSet TagSet);
+	const FTagData* GetTopTag(ELLMTagSet TagSet) const;
 	void TrackAllocation(const void* Ptr, int64 Size, ELLMTracker Tracker, ELLMAllocType AllocType,
 		const FTagData* TagData, const FTagData* AssetTagData, const FTagData* AssetClassTagData, const FTagData* UObjectClassesTagData, bool bTrackInMemPro);
 	void TrackFree(const void* Ptr, int64 Size, ELLMTracker Tracker, ELLMAllocType AllocType,
@@ -481,7 +427,7 @@ public:
 		) const;
 
 #if LLM_ENABLED_FULL_TAGS
-		int32 GetCompressedTag(ELLMTagSet InTagSet = ELLMTagSet::None);
+		int32 GetCompressedTag(ELLMTagSet InTagSet = ELLMTagSet::None) const;
 		void SetCompressedTag(int32 InTag, ELLMTagSet InTagSet = ELLMTagSet::None);
 #else
 		ELLMTag GetCompressedTag();
@@ -683,11 +629,16 @@ private:
 	void FinishConstruct(FTagData* TagData, ETagReferenceSource ReferenceSource);
 	void ReportDuplicateTagName(FTagData* TagData, ETagReferenceSource ReferenceSource);
 
-	const FTagData* FindOrAddTagData(ELLMTag EnumTag, ETagReferenceSource ReferenceSource = ETagReferenceSource::FunctionAPI);
-	const FTagData* FindOrAddTagData(FName Name, ELLMTagSet TagSet, bool bIsStatData = false, ETagReferenceSource ReferenceSource = ETagReferenceSource::FunctionAPI);
-	const FTagData* FindOrAddTagData(FName Name, ELLMTagSet TagSet, FName StatName, ETagReferenceSource ReferenceSource = ETagReferenceSource::FunctionAPI);
-	const FTagData* FindTagData(ELLMTag EnumTag, ETagReferenceSource ReferenceSource = ETagReferenceSource::FunctionAPI);
-	const FTagData* FindTagData(FName Name, ELLMTagSet TagSet, ETagReferenceSource ReferenceSource = ETagReferenceSource::FunctionAPI);
+	const FTagData* FindOrAddTagData(ELLMTag EnumTag,
+		ETagReferenceSource ReferenceSource = ETagReferenceSource::FunctionAPI);
+	const FTagData* FindOrAddTagData(FName Name, ELLMTagSet TagSet, bool bIsStatData = false,
+		ETagReferenceSource ReferenceSource = ETagReferenceSource::FunctionAPI);
+	const FTagData* FindOrAddTagData(FName Name, ELLMTagSet TagSet, FName StatName,
+		ETagReferenceSource ReferenceSource = ETagReferenceSource::FunctionAPI);
+	const FTagData* FindTagData(ELLMTag EnumTag,
+		ETagReferenceSource ReferenceSource = ETagReferenceSource::FunctionAPI);
+	const FTagData* FindTagData(FName Name, ELLMTagSet TagSet,
+		ETagReferenceSource ReferenceSource = ETagReferenceSource::FunctionAPI);
 
 	friend class ::FLLMPauseScope;
 	friend class ::FLLMScope;
diff --git a/Engine/Source/Runtime/Core/Private/HAL/LLM/LLMPrivateDefines.h b/Engine/Source/Runtime/Core/Private/HAL/LLM/LLMPrivateDefines.h
new file mode 100644
--- /dev/null
+++ b/Engine/Source/Runtime/Core/Private/HAL/LLM/LLMPrivateDefines.h
@@ -0,0 +1,69 @@
+// Copyright Epic Games, Inc. All Rights Reserved.
+
+#pragma once
+
+#include "HAL/LowLevelMemTracker.h"
+
+#if ENABLE_LOW_LEVEL_MEM_TRACKER
+
+// Private defines configuring LLM; see also public defines in LowLevelMemTracker.h
+// LLM_ALLOW_NAMES_TAGS: Set to 1 to allow arbitrary FNames to be used as tags, at the cost of more LLM memory usage per allocation. Set to 0 to store only the top-level ELLMTag for each allocation.
+// LLM_SCOPES always use FNames for the definition of their tag
+// Storing these tags on each allocation however requires 4 bytes per allocation
+// To reduce the memory overhead of LLM, this can be reduced to 1 byte per allocation at the cost of showing only the top-level 256 Tags; tags are replaced with their containing toplevel tag during allocation.
+// This setting is ignored if features requiring more information per allocation (LLM_ALLOW_STATS, LLM_ALLOW_ASSETS_TAGS) are enabled
+#ifndef LLM_ALLOW_NAMES_TAGS
+#define LLM_ALLOW_NAMES_TAGS 1
+#endif
+
+// Enable storing the full tag for each alloc if (1) arbitrary names tags are allowed or (2) Stats are allowed (stats use names tags) or (3) Asset tags are allowed (assets use names tags)
+#define LLM_ENABLED_FULL_TAGS (LLM_ALLOW_NAMES_TAGS || LLM_ALLOW_STATS || LLM_ALLOW_ASSETS_TAGS)
+
+// Whether to enable running with reduced threads. This is currently enabled because the engine crashes with -norenderthread
+#define LLM_ENABLED_REDUCE_THREADS 0
+
+// this controls if the commandline is used to enable tracking, or to disable it. If LLM_COMMANDLINE_ENABLES_FUNCTIONALITY is true, 
+// then tracking will only happen through Engine::Init(), at which point it will be disabled unless the commandline tells 
+// it to keep going (with -llm). If LLM_COMMANDLINE_ENABLES_FUNCTIONALITY is false, then tracking will be on unless the commandline
+// disables it (with -nollm)
+#ifndef LLM_COMMANDLINE_ENABLES_FUNCTIONALITY
+#define LLM_COMMANDLINE_ENABLES_FUNCTIONALITY 1
+#endif 
+
+// when set to 1, forces LLM to be enabled without having to parse the command line.
+#ifndef LLM_AUTO_ENABLE
+#define LLM_AUTO_ENABLE 0
+#endif
+
+// There is a little memory and cpu overhead in tracking peak memory but it is generally more useful than current memory.
+// Disable if you need a little more memory or speed
+#define LLM_ENABLED_TRACK_PEAK_MEMORY 1
+
+// Enable to allow memory tracking to be passed through to external tools without the extra overhead of LLM also tracking allocations.
+#define UE_ONLY_USE_PLATFORM_TRACKER 0
+
+#define LLM_CSV_PROFILER_WRITER_ENABLED CSV_PROFILER_STATS
+
+static_assert(static_cast<uint32>(ELLMTagSet::Max) == 4, "You added a new TagSet, please update the #ifs below");
+#if LLM_ALLOW_ASSETS_TAGS
+#define LLM_TAGSET_MAX 3
+#else
+#define LLM_TAGSET_MAX 2
+#endif
+
+namespace UE::LLMPrivate
+{
+
+enum class ETagReferenceSource : uint8
+{
+	Scope,
+	Declare,
+	EnumTag,
+	CustomEnumTag,
+	FunctionAPI,
+	ImplicitParent
+};
+
+} // namespace UE::LLMPrivate
+
+#endif		// #if ENABLE_LOW_LEVEL_MEM_TRACKER
