[Trace Services] ZigZag encoding with Variable length compression for better Callstack memory consumption

#rnx
#Synced-CL 49728215

diff --git a/Engine/Source/Developer/TraceServices/Private/Analyzers/CallstacksAnalysis.cpp b/Engine/Source/Developer/TraceServices/Private/Analyzers/CallstacksAnalysis.cpp
--- a/Engine/Source/Developer/TraceServices/Private/Analyzers/CallstacksAnalysis.cpp
+++ b/Engine/Source/Developer/TraceServices/Private/Analyzers/CallstacksAnalysis.cpp
@@ -15,7 +15,8 @@ namespace TraceServices
 enum ERoutes
 {
 	RouteId_Callstack,
-	RouteId_CallstackXORAndRLE
+	RouteId_CallstackXORAndRLE,
+	RouteId_CallstackZigZagAndVLE
 };
 	
 ////////////////////////////////////////////////////////////////////////////////
@@ -32,6 +33,7 @@ void FCallstacksAnalyzer::OnAnalysisBegin(const FOnAnalysisContext& Context)
 	FInterfaceBuilder& Builder = Context.InterfaceBuilder;
 	Builder.RouteEvent(RouteId_Callstack, "Memory", "CallstackSpec");
 	Builder.RouteEvent(RouteId_CallstackXORAndRLE, "Memory", "CallstackSpecXORAndRLE");
+	Builder.RouteEvent(RouteId_CallstackZigZagAndVLE, "Memory", "CallstackSpecZigZagAndVLE");
 }
 
 ////////////////////////////////////////////////////////////////////////////////
@@ -79,6 +81,21 @@ bool FCallstacksAnalyzer::OnEvent(uint16 RouteId, EStyle Style, const FOnEventCo
 			}
 		}
 		break;
+
+		case RouteId_CallstackZigZagAndVLE:
+		{
+			const TArrayReader<uint8>& CompressedFrames = Context.EventData.GetArray<uint8>("CompressedFrames");
+			TStaticArray<uint64, UE_CALLSTACK_TRACE_MAX_FRAMES> Frames;
+			const uint32 FrameCount = FCallstackZigZagAndVLE::Uncompress(
+				MakeArrayView(CompressedFrames.GetData(), CompressedFrames.Num()),
+				Frames
+			);
+			if (const uint32 Id = Context.EventData.GetValue<uint32>("CallstackId"))
+			{
+				Provider->AddCallstack(Id, Frames.GetData(), uint8(FrameCount));
+			}
+		}
+		break;
 	}
 
 	return true;
diff --git a/Engine/Source/Runtime/Core/Private/ProfilingDebugging/CallstackTrace.cpp b/Engine/Source/Runtime/Core/Private/ProfilingDebugging/CallstackTrace.cpp
--- a/Engine/Source/Runtime/Core/Private/ProfilingDebugging/CallstackTrace.cpp
+++ b/Engine/Source/Runtime/Core/Private/ProfilingDebugging/CallstackTrace.cpp
@@ -15,6 +15,7 @@ void	CallstackTrace_InitializeInternal();
 UE_TRACE_CHANNEL_DEFINE(CallstackChannel)
 UE_TRACE_EVENT_DEFINE(Memory, CallstackSpec)
 UE_TRACE_EVENT_DEFINE(Memory, CallstackSpecXORAndRLE)
+UE_TRACE_EVENT_DEFINE(Memory, CallstackSpecZigZagAndVLE)
 
 uint32 GCallStackTracingTlsSlotIndex = FPlatformTLS::InvalidTlsSlot;
 
@@ -48,7 +49,7 @@ void CallstackTrace_Initialize()
 
 void Test_Compression()
 {
-	auto Check = [](TConstArrayView<uint64> Input)
+	auto CheckXORAndRLE = [](TConstArrayView<uint64> Input)
 	{
 		uint8 Compressed[UE_CALLSTACK_TRACE_MAX_FRAMES*sizeof(uint64)];
 		const uint32 CompressedSize = FCallstackXORAndRLE::Compress(Input, Compressed);
@@ -60,6 +61,19 @@ void Test_Compression()
 
 		check(Algo::Compare(UncompressedView, Input));
 	};
+	
+	auto CheckZigZagAndVLE = [](TConstArrayView<uint64> Input)
+	{
+		uint8 Compressed[UE_CALLSTACK_TRACE_MAX_FRAMES*sizeof(uint64)];
+		const uint32 CompressedSize = FCallstackZigZagAndVLE::Compress(Input, Compressed);
+		const auto CompressedView = MakeArrayView(Compressed, CompressedSize);
+
+		TStaticArray<uint64, UE_CALLSTACK_TRACE_MAX_FRAMES> Uncompressed;
+		const uint32 UncompressedCount = FCallstackZigZagAndVLE::Uncompress(CompressedView, Uncompressed);
+		auto UncompressedView = MakeArrayView(Uncompressed.GetData(), UncompressedCount);
+
+		check(Algo::Compare(UncompressedView, Input));
+	};
 
 	uint64 Frames[] = {
 		0x00007ff693ab7b09,
@@ -72,7 +86,8 @@ void Test_Compression()
 		0x00007ffd639226b1,
 	};
 	uint8 FramesExpected[] = { 0x0 };
-	Check(Frames);
+	CheckXORAndRLE(Frames);
+	CheckZigZagAndVLE(Frames);
 	
 	uint64 Zeroes[] = {
 		0x0, 0x0, 0x0, 0x0, 0x0
@@ -80,7 +95,8 @@ void Test_Compression()
 	uint8 ZeroesExpected[] = {
 		64, 64, 64, 64, 64
 	};
-	Check(Zeroes);
+	CheckXORAndRLE(Zeroes);
+    CheckZigZagAndVLE(Zeroes);
 	
 	uint64 SingleOne[] = {
 		0x1,
@@ -88,7 +104,8 @@ void Test_Compression()
 	uint8 SingleOneExpected [] = {
 		63, 0x1
 	};
-	Check(SingleOne);
+	CheckXORAndRLE(SingleOne);
+	CheckZigZagAndVLE(SingleOne);
 	
 	uint64 SingleNormal[] = {
 		0x2d4f763d,
@@ -96,7 +113,8 @@ void Test_Compression()
 	uint8 SingleNormalExpected [] = {
 		34, 0x3d, 0x76, 0x4f, 0x2d
 	};
-	Check(SingleNormal);
+	CheckXORAndRLE(SingleNormal);
+	CheckZigZagAndVLE(SingleNormal);
 	
 	uint64 Mixed[] = {
 		0x000000002d4f829c,
@@ -118,7 +136,8 @@ void Test_Compression()
 		64,
 		34, 0x9c, 0x82, 0x4f, 0x2d,
 	};
-	Check(Mixed);
+	CheckXORAndRLE(Mixed);
+	CheckZigZagAndVLE(Mixed);
 
 	uint64 Random[] = {
 		0x85EA1249376518D5,
@@ -130,7 +149,8 @@ void Test_Compression()
 		9, 0x57, 0x89, 0x69, 0x12, 0x94, 0x57, 0x6B,
 		1, 0x19, 0xB7, 0xD4, 0xF7, 0x48, 0x5F, 0x50, 0x7B
 	};
-	Check(Random);
+	CheckXORAndRLE(Random);
+	CheckZigZagAndVLE(Random);
 }
 
 static bool bCompressionTestsWasRun = []()
diff --git a/Engine/Source/Runtime/Core/Private/ProfilingDebugging/CallstackTracePrivate.h b/Engine/Source/Runtime/Core/Private/ProfilingDebugging/CallstackTracePrivate.h
--- a/Engine/Source/Runtime/Core/Private/ProfilingDebugging/CallstackTracePrivate.h
+++ b/Engine/Source/Runtime/Core/Private/ProfilingDebugging/CallstackTracePrivate.h
@@ -43,6 +43,11 @@ UE_TRACE_EVENT_BEGIN_EXTERN(Memory, CallstackSpecXORAndRLE, NoSync)
 	UE_TRACE_EVENT_FIELD(uint8[], CompressedFrames)
 UE_TRACE_EVENT_END()
 
+UE_TRACE_EVENT_BEGIN_EXTERN(Memory, CallstackSpecZigZagAndVLE, NoSync)
+	UE_TRACE_EVENT_FIELD(uint32, CallstackId)
+	UE_TRACE_EVENT_FIELD(uint8[], CompressedFrames)
+UE_TRACE_EVENT_END()
+
 ////////////////////////////////////////////////////////////////////////////////
 
 	class FCallstackTracer
@@ -93,14 +98,14 @@ UE_TRACE_EVENT_END()
 					<< CallstackSpec.Frames(Entry.Frames, Entry.FrameCount);
 #else
 				uint8 Compressed[UE_CALLSTACK_TRACE_MAX_FRAMES*sizeof(uint64)];
-				const uint32 CompressedSize = FCallstackXORAndRLE::Compress(
+				const uint32 CompressedSize = FCallstackZigZagAndVLE::Compress(
 					MakeArrayView(Entry.Frames, Entry.FrameCount),
 					Compressed
 				);
 				
-				UE_TRACE_LOG(Memory, CallstackSpecXORAndRLE, CallstackChannel)
-					<< CallstackSpecXORAndRLE.CallstackId(Id)
-					<< CallstackSpecXORAndRLE.CompressedFrames(Compressed, CompressedSize);
+				UE_TRACE_LOG(Memory, CallstackSpecZigZagAndVLE, CallstackChannel)
+					<< CallstackSpecZigZagAndVLE.CallstackId(Id)
+					<< CallstackSpecZigZagAndVLE.CompressedFrames(Compressed, CompressedSize);
 #endif //UE_CALLSTACK_TRACE_UNCOMPRESSED
 			}
 
diff --git a/Engine/Source/Runtime/Core/Public/ProfilingDebugging/CallstackTrace.h b/Engine/Source/Runtime/Core/Public/ProfilingDebugging/CallstackTrace.h
--- a/Engine/Source/Runtime/Core/Public/ProfilingDebugging/CallstackTrace.h
+++ b/Engine/Source/Runtime/Core/Public/ProfilingDebugging/CallstackTrace.h
@@ -187,4 +187,61 @@ namespace FCallstackXORAndRLE
 		}
 		return AddressCount;
 	}
-}
+} // namespace FCallstackXORAndRLE
+
+/**
+ * ZigZag encoding with Variable-length integer compression of arrays of addresses (e.g stack frames).
+ * Needs to be available regardless if tracing is enabled to allow analysis to decompress.
+ */
+namespace FCallstackZigZagAndVLE
+{
+	inline uint32 Compress(TConstArrayView<uint64> Addresses, TArrayView<uint8> Compressed)
+	{
+		check(Addresses.Num() * 8 <= Compressed.Num());
+		uint32 CompressedBytes = 0;
+		uint64 LastAddress = 0;
+		for(uint8 Index = 0; Index < Addresses.Num(); ++Index)
+		{
+			const uint64 Address = Addresses[Index];
+			const int64 Delta = int64(Address - LastAddress);
+			LastAddress = Address;
+
+			uint64 EncodedAddress = (uint64(Delta) << 1) ^ uint64(Delta >> 63);
+			do
+			{
+				uint8 Byte = EncodedAddress & 0x7F;
+				EncodedAddress >>= 7;
+				Compressed[CompressedBytes++] = Byte | uint8((EncodedAddress > 0) << 7);
+			}
+			while (EncodedAddress > 0);
+		}
+		return CompressedBytes;
+	}
+
+	inline uint32 Uncompress(TConstArrayView<uint8> Compressed, TArrayView<uint64> Addresses)
+	{
+		const int32 Bytes = Compressed.Num();
+		int32 ByteIndex = 0;
+		uint32 AddressCount = 0;
+		uint64 LastAddress = 0;
+		while (ByteIndex < Bytes)
+		{
+			uint64 EncodedAddress = 0;
+			int Shift = 0;
+			uint8 Byte = 0;
+			do
+			{
+				Byte = Compressed[ByteIndex++];
+				EncodedAddress |= uint64(Byte & 0x7F) << Shift;
+				Shift += 7;
+			} while (Byte & 0x80);
+			
+			int64 Delta = (EncodedAddress >> 1) ^ -(int64)(EncodedAddress & 1);
+			uint64 Address = LastAddress + Delta;
+			LastAddress = Address;
+			
+			Addresses[AddressCount++] = Address;
+		}
+		return AddressCount;
+	}
+} // namespace FCallstackZigZag
