Make automation tests use single threaded `TQueue`s for things that are used within transactions.

#Synced-CL 49723230

diff --git a/Engine/Source/Runtime/Core/Public/Misc/AutomationTest.h b/Engine/Source/Runtime/Core/Public/Misc/AutomationTest.h
--- a/Engine/Source/Runtime/Core/Public/Misc/AutomationTest.h
+++ b/Engine/Source/Runtime/Core/Public/Misc/AutomationTest.h
@@ -1495,10 +1495,10 @@ private:
 	TMap<FString, TSet<FString>> ImmutableTags;
 
 	/** Queue of deferred commands */
-	TQueue< TSharedPtr<IAutomationLatentCommand> > LatentCommands;
+	TQueue<TSharedPtr<IAutomationLatentCommand>, EQueueMode::SingleThreaded> LatentCommands;
 
 	/** Queue of deferred commands */
-	TQueue< TSharedPtr<IAutomationNetworkCommand> > NetworkCommands;
+	TQueue<TSharedPtr<IAutomationNetworkCommand>, EQueueMode::SingleThreaded> NetworkCommands;
 
 	/** Whether we are currently executing smoke tests for startup/commandlet to minimize log spam */
 	EAutomationTestFlags RequestedTestFilter;
