Substrate
- Making parts common between the legacy and new translator (Substrate tree topology, validation, simplification, hlsl codegen).
- All within FSubstrateTranslatorData.
- Added some Substrate IR

TONs of follow up but everyone can help now.

#tests PC editor dx12 with new translator on some materials.
#rnx
#Synced-CL 49731272

diff --git a/Engine/Shaders/Private/MaterialTemplate.ush b/Engine/Shaders/Private/MaterialTemplate.ush
--- a/Engine/Shaders/Private/MaterialTemplate.ush
+++ b/Engine/Shaders/Private/MaterialTemplate.ush
@@ -617,7 +617,8 @@ struct FMaterialPixelParameters
 	FSharedLocalBases SharedLocalBases;
 	FSubstrateTree SubstrateTree;
 
-#if SUBSTRATE_USE_FULLYSIMPLIFIED_MATERIAL == 1
+#if (SUBSTRATE_USE_FULLYSIMPLIFIED_MATERIAL == 1) || ENABLE_NEW_HLSL_GENERATOR
+	// The new translator always add those members in. This is because we do not hide fully simplified instructions behind a "#if SUBSTRATE_USE_FULLYSIMPLIFIED_MATERIAL == 1" for now.
 	FSharedLocalBases SharedLocalBasesFullySimplified;
 	FSubstrateTree SubstrateTreeFullySimplified;
 #endif
diff --git a/Engine/Source/Runtime/Core/Private/UObject/FortniteMainBranchObjectVersion.cpp b/Engine/Source/Runtime/Core/Private/UObject/FortniteMainBranchObjectVersion.cpp
--- a/Engine/Source/Runtime/Core/Private/UObject/FortniteMainBranchObjectVersion.cpp
+++ b/Engine/Source/Runtime/Core/Private/UObject/FortniteMainBranchObjectVersion.cpp
@@ -13,7 +13,7 @@ TMap<FGuid, FGuid> FFortniteMainBranchObjectVersion::GetSystemGuids()
 	SystemGuids.Add(DevGuids.MATERIALSHADERMAP_DERIVEDDATA_VER, FGuid("1A77CBCDFBDB4C12BF8563EDBB7D8888"));
 	SystemGuids.Add(DevGuids.NIAGARASHADERMAP_DERIVEDDATA_VER, FGuid("CC229D919A214E59B79D01AC2E5BF981"));
 	SystemGuids.Add(DevGuids.Niagara_LatestScriptCompileVersion, FGuid("FB15C2870A4E4AE8B5F58566F0E3825E"));
-	SystemGuids.Add(DevGuids.MaterialTranslationDDCVersion, FGuid("EA008827093C4A2BBFAC9F4EA3D92BE9"));
+	SystemGuids.Add(DevGuids.MaterialTranslationDDCVersion, FGuid("5D6FEADBA5AB4D129B967241B347EFEA"));
 
 	// These GUIDs are in AppendRestrictedFortniteMainStreamObjectVersionGuids as they should rarely be changed.
 	//SystemGuids.Add(DevGuids.NANITE_DERIVEDDATA_VER, FGuid("ACF98184978742A592537DA223A1E6BE"));
