Added support for const member function bindings to dynamic delegates.

diff --git a/Engine/Source/Runtime/Core/Public/Delegates/DelegateInstanceInterface.h b/Engine/Source/Runtime/Core/Public/Delegates/DelegateInstanceInterface.h
--- a/Engine/Source/Runtime/Core/Public/Delegates/DelegateInstanceInterface.h
+++ b/Engine/Source/Runtime/Core/Public/Delegates/DelegateInstanceInterface.h
@@ -50,7 +50,7 @@ struct IBaseDelegateInstance<RetType(ArgTypes...), UserPolicy> : public UserPoli
 // Temp workaround for deprecation warnings being emitted by this template with VS2022 17.9. Reported fixed in 17.10
 PRAGMA_DISABLE_DEPRECATION_WARNINGS
 
-template <bool Const, typename Class, typename FuncType>
+template <bool bConst, typename Class, typename FuncType>
 struct TMemFunPtrType;
 
 template <typename Class, typename RetType, typename... ArgTypes>
diff --git a/Engine/Source/Runtime/Core/Public/Delegates/DelegateSignatureImpl.inl b/Engine/Source/Runtime/Core/Public/Delegates/DelegateSignatureImpl.inl
--- a/Engine/Source/Runtime/Core/Public/Delegates/DelegateSignatureImpl.inl
+++ b/Engine/Source/Runtime/Core/Public/Delegates/DelegateSignatureImpl.inl
@@ -1143,6 +1143,12 @@ template <typename DelegateSignature>
 using TTSMulticastDelegate = TMulticastDelegate<DelegateSignature, FDefaultTSDelegateUserPolicy>;
 
 
+template <typename ThreadSafetyMode, typename RetValType, typename... ParamTypes>
+class TBaseDynamicMulticastDelegate;
+
+template <typename MulticastDelegate, typename OwningClass, typename DelegateInfoClass>
+struct TSparseDynamicDelegate;
+
 /**
  * Dynamic delegate template class (UObject-based, serializable).  You'll use the various DECLARE_DYNAMIC_DELEGATE
  * macros to create the actual delegate type, templated to the function signature the delegate is compatible with.
@@ -1151,6 +1157,13 @@ using TTSMulticastDelegate = TMulticastDelegate<DelegateSignature, FDefaultTSDel
 template <typename ThreadSafetyMode, typename RetValType, typename... ParamTypes>
 class TBaseDynamicDelegate : public TScriptDelegate<ThreadSafetyMode>
 {
+	template <typename, typename, typename...>
+	friend class TBaseDynamicMulticastDelegate;
+
+	/** Sparse version of TBaseDynamicDelegate */
+	template <typename, typename, typename>
+	friend struct TSparseDynamicDelegate;
+
 public:
 	/**
 	 * Default constructor
@@ -1169,13 +1182,33 @@ public:
 	/**
 	 * Templated helper class to define a typedef for user's method pointer, then used below
 	 */
-	template< class UserClass >
-	class TMethodPtrResolver
-	{
-	public:
-		typedef RetValType (UserClass::*FMethodPtr)(ParamTypes... Params);
+	template <bool bConst, typename UserClass>
+	struct TMethodPtrResolver
+	{
+		using FMethodPtr = std::conditional_t<
+			bConst,
+			RetValType(UserClass::*)(ParamTypes... Params) const,
+			RetValType(UserClass::*)(ParamTypes... Params)
+		>;
 	};
 
+private:
+	template <typename UserClass, typename... VarTypes>
+	void BindDynamicImpl(const UserClass* InUserObject, FName InFunctionName)
+	{
+		check(InUserObject);
+
+		// NOTE: If you hit a compile error on the following line, it means you're trying to use a non-UObject type
+		//       with this delegate, which is not supported
+		this->Object = const_cast<UObject*>(Cast<const UObject>(InUserObject));
+
+		// Store the function name.  The incoming function name was generated by a macro and includes the method's class name.
+		this->FunctionName = InFunctionName;
+
+		ensureMsgf(this->IsBound(), TEXT("Unable to bind delegate to '%s' (function might not be marked as a UFUNCTION or object may be pending kill)"), *InFunctionName.ToString());
+	}
+
+public:
 	/**
 	 * Binds a UObject instance and a UObject method address to this delegate.
 	 *
@@ -1185,27 +1218,30 @@ public:
 	 *
 	 * NOTE:  Do not call this function directly.  Instead, call BindDynamic() which is a macro proxy function that
 	 *        automatically sets the function name string for the caller.
+	 *
+	 * NOTE:  We're not actually using the incoming method pointer.  We simply require it for type-safety reasons.
 	 */
-	template< class UserClass >
-	void __Internal_BindDynamic( UserClass* InUserObject, typename TMethodPtrResolver< UserClass >::FMethodPtr InMethodPtr, FName InFunctionName )
+	template <typename UserClass>
+	UE_REWRITE void __Internal_BindDynamic(UserClass* InUserObject, typename TMethodPtrResolver<false, UserClass>::FMethodPtr InMethodPtr, FName InFunctionName)
 	{
-		check( InUserObject != nullptr && InMethodPtr != nullptr );
-
-		// NOTE: We're not actually storing the incoming method pointer or calling it.  We simply require it for type-safety reasons.
-
-		// NOTE: If you hit a compile error on the following line, it means you're trying to use a non-UObject type
-		//       with this delegate, which is not supported
-		this->Object = Cast<UObject>(InUserObject);
-
-		// Store the function name.  The incoming function name was generated by a macro and includes the method's class name.
-		this->FunctionName = InFunctionName;
-
-		ensureMsgf(this->IsBound(), TEXT("Unable to bind delegate to '%s' (function might not be marked as a UFUNCTION or object may be pending kill)"), *InFunctionName.ToString());
+		UE_STATIC_DEPRECATE(5.8, std::is_const_v<UserClass>, "Binding a delegate with a const object pointer and non-const function is deprecated.");
+		this->BindDynamicImpl(InUserObject, InFunctionName);
+	}
+	template <typename UserClass>
+	UE_REWRITE void __Internal_BindDynamic(TObjectPtr<UserClass> InUserObject, typename TMethodPtrResolver<false, UserClass>::FMethodPtr InMethodPtr, FName InFunctionName)
+	{
+		UE_STATIC_DEPRECATE(5.8, std::is_const_v<UserClass>, "Binding a delegate with a const object pointer and non-const function is deprecated.");
+		this->BindDynamicImpl(ToRawPtr(InUserObject), InFunctionName);
+	}
+	template <typename UserClass>
+	UE_REWRITE void __Internal_BindDynamic(UserClass* InUserObject, typename TMethodPtrResolver<true, UserClass>::FMethodPtr InMethodPtr, FName InFunctionName)
+	{
+		this->BindDynamicImpl(InUserObject, InFunctionName);
 	}
-	template< class UserClass >
-	void __Internal_BindDynamic( TObjectPtr<UserClass> InUserObject, typename TMethodPtrResolver< UserClass >::FMethodPtr InMethodPtr, FName InFunctionName )
+	template <typename UserClass>
+	UE_REWRITE void __Internal_BindDynamic(TObjectPtr<UserClass> InUserObject, typename TMethodPtrResolver<true, UserClass>::FMethodPtr InMethodPtr, FName InFunctionName)
 	{
-		__Internal_BindDynamic(ToRawPtr(InUserObject), InMethodPtr, InFunctionName);
+		this->BindDynamicImpl(ToRawPtr(InUserObject), InFunctionName);
 	}
 
 	/** Execute the delegate.  If the function pointer is not valid, an error will occur. */
@@ -1257,6 +1293,15 @@ public:
 		: TMulticastScriptDelegate<ThreadSafetyMode>( InMulticastScriptDelegate )
 	{ }
 
+private:
+	template <typename UserClass>
+	bool IsAlreadyBoundImpl(UserClass* InUserObject, FName InFunctionName) const
+	{
+		check(InUserObject);
+		return this->Contains(InUserObject, InFunctionName);
+	}
+
+public:
 	/**
 	 * Tests if a UObject instance and a UObject method address pair are already bound to this multi-cast delegate.
 	 *
@@ -1267,22 +1312,43 @@ public:
 	 *
 	 * NOTE:  Do not call this function directly.  Instead, call IsAlreadyBound() which is a macro proxy function that
 	 *        automatically sets the function name string for the caller.
+	 *
+	 * NOTE:  We're not actually using the incoming method pointer.  We simply require it for type-safety reasons.
 	 */
-	template< class UserClass >
-	bool __Internal_IsAlreadyBound( UserClass* InUserObject, typename FDelegate::template TMethodPtrResolver< UserClass >::FMethodPtr InMethodPtr, FName InFunctionName ) const
+	template <typename UserClass>
+	UE_REWRITE bool __Internal_IsAlreadyBound(UserClass* InUserObject, typename FDelegate::template TMethodPtrResolver<false, UserClass>::FMethodPtr InMethodPtr, FName InFunctionName) const
 	{
-		check( InUserObject != nullptr && InMethodPtr != nullptr );
-
-		// NOTE: We're not actually using the incoming method pointer or calling it.  We simply require it for type-safety reasons.
-
-		return this->Contains( InUserObject, InFunctionName );
+		return this->IsAlreadyBoundImpl(InUserObject, InFunctionName);
 	}
-	template< class UserClass >
-	bool __Internal_IsAlreadyBound( TObjectPtr<UserClass> InUserObject, typename FDelegate::template TMethodPtrResolver< UserClass >::FMethodPtr InMethodPtr, FName InFunctionName ) const
+	template <typename UserClass>
+	UE_REWRITE bool __Internal_IsAlreadyBound(TObjectPtr<UserClass> InUserObject, typename FDelegate::template TMethodPtrResolver<false, UserClass>::FMethodPtr InMethodPtr, FName InFunctionName) const
 	{
-		return __Internal_IsAlreadyBound(ToRawPtr(InUserObject), InMethodPtr, InFunctionName);
+		return this->IsAlreadyBoundImpl(ToRawPtr(InUserObject), InFunctionName);
 	}
+	template <typename UserClass>
+	UE_REWRITE bool __Internal_IsAlreadyBound(UserClass* InUserObject, typename FDelegate::template TMethodPtrResolver<true, UserClass>::FMethodPtr InMethodPtr, FName InFunctionName) const
+	{
+		return this->IsAlreadyBoundImpl(InUserObject, InFunctionName);
+	}
+	template <typename UserClass>
+	UE_REWRITE bool __Internal_IsAlreadyBound(TObjectPtr<UserClass> InUserObject, typename FDelegate::template TMethodPtrResolver<true, UserClass>::FMethodPtr InMethodPtr, FName InFunctionName) const
+	{
+		return this->IsAlreadyBoundImpl(ToRawPtr(InUserObject), InFunctionName);
+	}
+
+private:
+	template <typename UserClass>
+	void AddDynamicImpl(const UserClass* InUserObject, FName InFunctionName)
+	{
+		check(InUserObject);
+
+		FDelegate NewDelegate;
+		NewDelegate.BindDynamicImpl(InUserObject, InFunctionName);
 
+		this->Add(NewDelegate);
+	}
+
+public:
 	/**
 	 * Binds a UObject instance and a UObject method address to this multi-cast delegate.
 	 *
@@ -1292,25 +1358,45 @@ public:
 	 *
 	 * NOTE:  Do not call this function directly.  Instead, call AddDynamic() which is a macro proxy function that
 	 *        automatically sets the function name string for the caller.
+	 *
+	 * NOTE:  We're not actually using the incoming method pointer.  We simply require it for type-safety reasons.
 	 */
-	template< class UserClass >
-	void __Internal_AddDynamic( UserClass* InUserObject, typename FDelegate::template TMethodPtrResolver< UserClass >::FMethodPtr InMethodPtr, FName InFunctionName )
+	template <typename UserClass>
+	UE_REWRITE void __Internal_AddDynamic(UserClass* InUserObject, typename FDelegate::template TMethodPtrResolver<false, UserClass>::FMethodPtr InMethodPtr, FName InFunctionName)
+	{
+		UE_STATIC_DEPRECATE(5.8, std::is_const_v<UserClass>, "Binding a delegate with a const object pointer and non-const function is deprecated.");
+		this->AddDynamicImpl(InUserObject, InFunctionName);
+	}
+	template <typename UserClass>
+	UE_REWRITE void __Internal_AddDynamic(TObjectPtr<UserClass> InUserObject, typename FDelegate::template TMethodPtrResolver<false, UserClass>::FMethodPtr InMethodPtr, FName InFunctionName)
+	{
+		UE_STATIC_DEPRECATE(5.8, std::is_const_v<UserClass>, "Binding a delegate with a const object pointer and non-const function is deprecated.");
+		this->AddDynamicImpl(ToRawPtr(InUserObject), InFunctionName);
+	}
+	template <typename UserClass>
+	UE_REWRITE void __Internal_AddDynamic(UserClass* InUserObject, typename FDelegate::template TMethodPtrResolver<true, UserClass>::FMethodPtr InMethodPtr, FName InFunctionName)
 	{
-		check( InUserObject != nullptr && InMethodPtr != nullptr );
+		this->AddDynamicImpl(InUserObject, InFunctionName);
+	}
+	template <typename UserClass>
+	UE_REWRITE void __Internal_AddDynamic(TObjectPtr<UserClass> InUserObject, typename FDelegate::template TMethodPtrResolver<true, UserClass>::FMethodPtr InMethodPtr, FName InFunctionName)
+	{
+		this->AddDynamicImpl(ToRawPtr(InUserObject), InFunctionName);
+	}
 
-		// NOTE: We're not actually storing the incoming method pointer or calling it.  We simply require it for type-safety reasons.
+private:
+	template <typename UserClass>
+	void AddUniqueDynamicImpl(const UserClass* InUserObject, FName InFunctionName)
+	{
+		check(InUserObject);
 
 		FDelegate NewDelegate;
-		NewDelegate.__Internal_BindDynamic( InUserObject, InMethodPtr, InFunctionName );
+		NewDelegate.BindDynamicImpl(InUserObject, InFunctionName);
 
-		this->Add( NewDelegate );
-	}
-	template< class UserClass >
-	void __Internal_AddDynamic( TObjectPtr<UserClass> InUserObject, typename FDelegate::template TMethodPtrResolver< UserClass >::FMethodPtr InMethodPtr, FName InFunctionName )
-	{
-		__Internal_AddDynamic(ToRawPtr(InUserObject), InMethodPtr, InFunctionName);
+		this->AddUnique(NewDelegate);
 	}
 
+public:
 	/**
 	 * Binds a UObject instance and a UObject method address to this multi-cast delegate, but only if it hasn't been bound before.
 	 *
@@ -1320,25 +1406,42 @@ public:
 	 *
 	 * NOTE:  Do not call this function directly.  Instead, call AddUniqueDynamic() which is a macro proxy function that
 	 *        automatically sets the function name string for the caller.
+	 *
+	 * NOTE:  We're not actually using the incoming method pointer.  We simply require it for type-safety reasons.
 	 */
-	template< class UserClass >
-	void __Internal_AddUniqueDynamic( UserClass* InUserObject, typename FDelegate::template TMethodPtrResolver< UserClass >::FMethodPtr InMethodPtr, FName InFunctionName )
+	template <typename UserClass>
+	UE_REWRITE void __Internal_AddUniqueDynamic(UserClass* InUserObject, typename FDelegate::template TMethodPtrResolver<false, UserClass>::FMethodPtr InMethodPtr, FName InFunctionName)
 	{
-		check( InUserObject != nullptr && InMethodPtr != nullptr );
-
-		// NOTE: We're not actually storing the incoming method pointer or calling it.  We simply require it for type-safety reasons.
-
-		FDelegate NewDelegate;
-		NewDelegate.__Internal_BindDynamic( InUserObject, InMethodPtr, InFunctionName );
-
-		this->AddUnique( NewDelegate );
+		UE_STATIC_DEPRECATE(5.8, std::is_const_v<UserClass>, "Binding a delegate with a const object pointer and non-const function is deprecated.");
+		this->AddUniqueDynamicImpl(InUserObject, InFunctionName);
+	}
+	template <typename UserClass>
+	UE_REWRITE void __Internal_AddUniqueDynamic(TObjectPtr<UserClass> InUserObject, typename FDelegate::template TMethodPtrResolver<false, UserClass>::FMethodPtr InMethodPtr, FName InFunctionName)
+	{
+		UE_STATIC_DEPRECATE(5.8, std::is_const_v<UserClass>, "Binding a delegate with a const object pointer and non-const function is deprecated.");
+		this->AddUniqueDynamicImpl(ToRawPtr(InUserObject), InFunctionName);
+	}
+	template <typename UserClass>
+	UE_REWRITE void __Internal_AddUniqueDynamic(UserClass* InUserObject, typename FDelegate::template TMethodPtrResolver<true, UserClass>::FMethodPtr InMethodPtr, FName InFunctionName)
+	{
+		this->AddUniqueDynamicImpl(InUserObject, InFunctionName);
+	}
+	template <typename UserClass>
+	UE_REWRITE void __Internal_AddUniqueDynamic(TObjectPtr<UserClass> InUserObject, typename FDelegate::template TMethodPtrResolver<true, UserClass>::FMethodPtr InMethodPtr, FName InFunctionName)
+	{
+		this->AddUniqueDynamicImpl(ToRawPtr(InUserObject), InFunctionName);
 	}
-	template< class UserClass >
-	void __Internal_AddUniqueDynamic( TObjectPtr<UserClass> InUserObject, typename FDelegate::template TMethodPtrResolver< UserClass >::FMethodPtr InMethodPtr, FName InFunctionName )
+
+private:
+	template <typename UserClass>
+	void RemoveDynamicImpl(const UserClass* InUserObject, FName InFunctionName)
 	{
-		__Internal_AddUniqueDynamic(ToRawPtr(InUserObject), InMethodPtr, InFunctionName);
+		check(InUserObject);
+
+		this->Remove(InUserObject, InFunctionName);
 	}
 
+public:
 	/**
 	 * Unbinds a UObject instance and a UObject method address from this multi-cast delegate.
 	 *
@@ -1348,20 +1451,28 @@ public:
 	 *
 	 * NOTE:  Do not call this function directly.  Instead, call RemoveDynamic() which is a macro proxy function that
 	 *        automatically sets the function name string for the caller.
+	 *
+	 * NOTE:  We're not actually using the incoming method pointer.  We simply require it for type-safety reasons.
 	 */
-	template< class UserClass >
-	void __Internal_RemoveDynamic( UserClass* InUserObject, typename FDelegate::template TMethodPtrResolver< UserClass >::FMethodPtr InMethodPtr, FName InFunctionName )
+	template <typename UserClass>
+	UE_REWRITE void __Internal_RemoveDynamic(UserClass* InUserObject, typename FDelegate::template TMethodPtrResolver<false, UserClass>::FMethodPtr InMethodPtr, FName InFunctionName)
 	{
-		check( InUserObject != nullptr && InMethodPtr != nullptr );
-
-		// NOTE: We're not actually storing the incoming method pointer or calling it.  We simply require it for type-safety reasons.
-
-		this->Remove( InUserObject, InFunctionName );
+		this->RemoveDynamicImpl(InUserObject, InFunctionName);
+	}
+	template <typename UserClass>
+	UE_REWRITE void __Internal_RemoveDynamic(TObjectPtr<UserClass> InUserObject, typename FDelegate::template TMethodPtrResolver<false, UserClass>::FMethodPtr InMethodPtr, FName InFunctionName)
+	{
+		this->RemoveDynamicImpl(ToRawPtr(InUserObject), InFunctionName);
+	}
+	template <typename UserClass>
+	UE_REWRITE void __Internal_RemoveDynamic(UserClass* InUserObject, typename FDelegate::template TMethodPtrResolver<true, UserClass>::FMethodPtr InMethodPtr, FName InFunctionName)
+	{
+		this->RemoveDynamicImpl(InUserObject, InFunctionName);
 	}
-	template< class UserClass >
-	void __Internal_RemoveDynamic( TObjectPtr<UserClass> InUserObject, typename FDelegate::template TMethodPtrResolver< UserClass >::FMethodPtr InMethodPtr, FName InFunctionName )
+	template <typename UserClass>
+	UE_REWRITE void __Internal_RemoveDynamic(TObjectPtr<UserClass> InUserObject, typename FDelegate::template TMethodPtrResolver<true, UserClass>::FMethodPtr InMethodPtr, FName InFunctionName)
 	{
-		__Internal_RemoveDynamic(ToRawPtr(InUserObject), InMethodPtr, InFunctionName);
+		this->RemoveDynamicImpl(ToRawPtr(InUserObject), InFunctionName);
 	}
 
 	/** Broadcasts this delegate to all bound objects, except to those that may have expired */
diff --git a/Engine/Source/Runtime/CoreUObject/Public/UObject/SparseDelegate.h b/Engine/Source/Runtime/CoreUObject/Public/UObject/SparseDelegate.h
--- a/Engine/Source/Runtime/CoreUObject/Public/UObject/SparseDelegate.h
+++ b/Engine/Source/Runtime/CoreUObject/Public/UObject/SparseDelegate.h
@@ -334,6 +334,16 @@ public:
 		}
 	}
 
+private:
+	template <typename UserClass>
+	bool IsAlreadyBoundImpl(const UserClass* InUserObject, FName InFunctionName) const
+	{
+		check(InUserObject);
+
+		return this->Contains(InUserObject, InFunctionName);
+	}
+
+public:
 	/**
 	* Tests if a UObject instance and a UObject method address pair are already bound to this multi-cast delegate.
 	*
@@ -344,17 +354,43 @@ public:
 	*
 	* NOTE:  Do not call this function directly.  Instead, call IsAlreadyBound() which is a macro proxy function that
 	*        automatically sets the function name string for the caller.
+	*
+	* NOTE:  We're not actually using the incoming method pointer.  We simply require it for type-safety reasons.
 	*/
-	template< class UserClass >
-	bool __Internal_IsAlreadyBound( UserClass* InUserObject, typename FDelegate::template TMethodPtrResolver< UserClass >::FMethodPtr InMethodPtr, FName InFunctionName ) const
+	template <typename UserClass>
+	UE_REWRITE bool __Internal_IsAlreadyBound(UserClass* InUserObject, typename FDelegate::template TMethodPtrResolver<false, UserClass>::FMethodPtr InMethodPtr, FName InFunctionName) const
+	{
+		return this->IsAlreadyBoundImpl(InUserObject, InFunctionName);
+	}
+	template <typename UserClass>
+	UE_REWRITE bool __Internal_IsAlreadyBound(TObjectPtr<UserClass> InUserObject, typename FDelegate::template TMethodPtrResolver<false, UserClass>::FMethodPtr InMethodPtr, FName InFunctionName) const
+	{
+		return this->IsAlreadyBoundImpl(ToRawPtr(InUserObject), InFunctionName);
+	}
+	template <typename UserClass>
+	UE_REWRITE bool __Internal_IsAlreadyBound(UserClass* InUserObject, typename FDelegate::template TMethodPtrResolver<true, UserClass>::FMethodPtr InMethodPtr, FName InFunctionName) const
+	{
+		return this->IsAlreadyBoundImpl(InUserObject, InFunctionName);
+	}
+	template <typename UserClass>
+	UE_REWRITE bool __Internal_IsAlreadyBound(TObjectPtr<UserClass> InUserObject, typename FDelegate::template TMethodPtrResolver<true, UserClass>::FMethodPtr InMethodPtr, FName InFunctionName) const
 	{
-		check( InUserObject != nullptr && InMethodPtr != nullptr );
+		return this->IsAlreadyBoundImpl(ToRawPtr(InUserObject), InFunctionName);
+	}
 
-		// NOTE: We're not actually using the incoming method pointer or calling it.  We simply require it for type-safety reasons.
+private:
+	template <typename UserClass>
+	void AddDynamicImpl(const UserClass* InUserObject, FName InFunctionName)
+	{
+		check(InUserObject);
 
-		return this->Contains( InUserObject, InFunctionName );
+		FDelegate NewDelegate;
+		NewDelegate.BindDynamicImpl(InUserObject, InFunctionName);
+
+		this->Add(NewDelegate);
 	}
 
+public:
 	/**
 	* Binds a UObject instance and a UObject method address to this multi-cast delegate.
 	*
@@ -364,20 +400,45 @@ public:
 	*
 	* NOTE:  Do not call this function directly.  Instead, call AddDynamic() which is a macro proxy function that
 	*        automatically sets the function name string for the caller.
+	*
+	* NOTE:  We're not actually using the incoming method pointer.  We simply require it for type-safety reasons.
 	*/
-	template< class UserClass >
-	void __Internal_AddDynamic( UserClass* InUserObject, typename FDelegate::template TMethodPtrResolver< UserClass >::FMethodPtr InMethodPtr, FName InFunctionName )
+	template <typename UserClass>
+	UE_REWRITE void __Internal_AddDynamic(UserClass* InUserObject, typename FDelegate::template TMethodPtrResolver<false, UserClass>::FMethodPtr InMethodPtr, FName InFunctionName)
+	{
+		UE_STATIC_DEPRECATE(5.8, std::is_const_v<UserClass>, "Binding a delegate with a const object pointer and non-const function is deprecated.");
+		this->AddDynamicImpl(InUserObject, InFunctionName);
+	}
+	template <typename UserClass>
+	UE_REWRITE void __Internal_AddDynamic(TObjectPtr<UserClass> InUserObject, typename FDelegate::template TMethodPtrResolver<false, UserClass>::FMethodPtr InMethodPtr, FName InFunctionName)
+	{
+		UE_STATIC_DEPRECATE(5.8, std::is_const_v<UserClass>, "Binding a delegate with a const object pointer and non-const function is deprecated.");
+		this->AddDynamicImpl(ToRawPtr(InUserObject), InFunctionName);
+	}
+	template <typename UserClass>
+	UE_REWRITE void __Internal_AddDynamic(UserClass* InUserObject, typename FDelegate::template TMethodPtrResolver<true, UserClass>::FMethodPtr InMethodPtr, FName InFunctionName)
+	{
+		this->AddDynamicImpl(InUserObject, InFunctionName);
+	}
+	template <typename UserClass>
+	UE_REWRITE void __Internal_AddDynamic(TObjectPtr<UserClass> InUserObject, typename FDelegate::template TMethodPtrResolver<true, UserClass>::FMethodPtr InMethodPtr, FName InFunctionName)
 	{
-		check( InUserObject != nullptr && InMethodPtr != nullptr );
+		this->AddDynamicImpl(ToRawPtr(InUserObject), InFunctionName);
+	}
 
-		// NOTE: We're not actually storing the incoming method pointer or calling it.  We simply require it for type-safety reasons.
+private:
+	template <typename UserClass>
+	void AddUniqueDynamicImpl(const UserClass* InUserObject, FName InFunctionName)
+	{
+		check(InUserObject);
 
 		FDelegate NewDelegate;
-		NewDelegate.__Internal_BindDynamic( InUserObject, InMethodPtr, InFunctionName );
+		NewDelegate.BindDynamicImpl(InUserObject, InFunctionName);
 
-		this->Add( NewDelegate );
+		this->AddUnique(NewDelegate);
 	}
 
+public:
 	/**
 	* Binds a UObject instance and a UObject method address to this multi-cast delegate, but only if it hasn't been bound before.
 	*
@@ -387,20 +448,42 @@ public:
 	*
 	* NOTE:  Do not call this function directly.  Instead, call AddUniqueDynamic() which is a macro proxy function that
 	*        automatically sets the function name string for the caller.
+	*
+	* NOTE:  We're not actually using the incoming method pointer.  We simply require it for type-safety reasons.
 	*/
-	template< class UserClass >
-	void __Internal_AddUniqueDynamic( UserClass* InUserObject, typename FDelegate::template TMethodPtrResolver< UserClass >::FMethodPtr InMethodPtr, FName InFunctionName )
+	template <typename UserClass>
+	UE_REWRITE void __Internal_AddUniqueDynamic(UserClass* InUserObject, typename FDelegate::template TMethodPtrResolver<false, UserClass>::FMethodPtr InMethodPtr, FName InFunctionName)
 	{
-		check( InUserObject != nullptr && InMethodPtr != nullptr );
-
-		// NOTE: We're not actually storing the incoming method pointer or calling it.  We simply require it for type-safety reasons.
+		UE_STATIC_DEPRECATE(5.8, std::is_const_v<UserClass>, "Binding a delegate with a const object pointer and non-const function is deprecated.");
+		this->AddUniqueDynamicImpl(InUserObject, InFunctionName);
+	}
+	template <typename UserClass>
+	UE_REWRITE void __Internal_AddUniqueDynamic(TObjectPtr<UserClass> InUserObject, typename FDelegate::template TMethodPtrResolver<false, UserClass>::FMethodPtr InMethodPtr, FName InFunctionName)
+	{
+		UE_STATIC_DEPRECATE(5.8, std::is_const_v<UserClass>, "Binding a delegate with a const object pointer and non-const function is deprecated.");
+		this->AddUniqueDynamicImpl(ToRawPtr(InUserObject), InFunctionName);
+	}
+	template <typename UserClass>
+	UE_REWRITE void __Internal_AddUniqueDynamic(UserClass* InUserObject, typename FDelegate::template TMethodPtrResolver<true, UserClass>::FMethodPtr InMethodPtr, FName InFunctionName)
+	{
+		this->AddUniqueDynamicImpl(InUserObject, InFunctionName);
+	}
+	template <typename UserClass>
+	UE_REWRITE void __Internal_AddUniqueDynamic(TObjectPtr<UserClass> InUserObject, typename FDelegate::template TMethodPtrResolver<true, UserClass>::FMethodPtr InMethodPtr, FName InFunctionName)
+	{
+		this->AddUniqueDynamicImpl(ToRawPtr(InUserObject), InFunctionName);
+	}
 
-		FDelegate NewDelegate;
-		NewDelegate.__Internal_BindDynamic( InUserObject, InMethodPtr, InFunctionName );
+private:
+	template <typename UserClass>
+	void RemoveDynamicImpl(const UserClass* InUserObject, FName InFunctionName)
+	{
+		check(InUserObject);
 
-		this->AddUnique( NewDelegate );
+		this->Remove(InUserObject, InFunctionName);
 	}
 
+public:
 	/**
 	* Unbinds a UObject instance and a UObject method address from this multi-cast delegate.
 	*
@@ -410,15 +493,28 @@ public:
 	*
 	* NOTE:  Do not call this function directly.  Instead, call RemoveDynamic() which is a macro proxy function that
 	*        automatically sets the function name string for the caller.
+	*
+	* NOTE:  We're not actually using the incoming method pointer.  We simply require it for type-safety reasons.
 	*/
-	template< class UserClass >
-	void __Internal_RemoveDynamic( UserClass* InUserObject, typename FDelegate::template TMethodPtrResolver< UserClass >::FMethodPtr InMethodPtr, FName InFunctionName )
+	template <typename UserClass>
+	UE_REWRITE void __Internal_RemoveDynamic(UserClass* InUserObject, typename FDelegate::template TMethodPtrResolver<false, UserClass>::FMethodPtr InMethodPtr, FName InFunctionName)
 	{
-		check( InUserObject != nullptr && InMethodPtr != nullptr );
-
-		// NOTE: We're not actually storing the incoming method pointer or calling it.  We simply require it for type-safety reasons.
-
-		this->Remove( InUserObject, InFunctionName );
+		this->RemoveDynamicImpl(InUserObject, InFunctionName);
+	}
+	template <typename UserClass>
+	UE_REWRITE void __Internal_RemoveDynamic(TObjectPtr<UserClass> InUserObject, typename FDelegate::template TMethodPtrResolver<false, UserClass>::FMethodPtr InMethodPtr, FName InFunctionName)
+	{
+		this->RemoveDynamicImpl(ToRawPtr(InUserObject), InFunctionName);
+	}
+	template <typename UserClass>
+	UE_REWRITE void __Internal_RemoveDynamic(UserClass* InUserObject, typename FDelegate::template TMethodPtrResolver<true, UserClass>::FMethodPtr InMethodPtr, FName InFunctionName)
+	{
+		this->RemoveDynamicImpl(InUserObject, InFunctionName);
+	}
+	template <typename UserClass>
+	UE_REWRITE void __Internal_RemoveDynamic(TObjectPtr<UserClass> InUserObject, typename FDelegate::template TMethodPtrResolver<true, UserClass>::FMethodPtr InMethodPtr, FName InFunctionName)
+	{
+		this->RemoveDynamicImpl(ToRawPtr(InUserObject), InFunctionName);
 	}
 
 private:
