Fixed improper reuse of an UErrorCode that was/maybe set to an error code
Log failures from ICU's CreateCollator API.

diff --git a/Engine/Source/Runtime/Core/Private/Internationalization/ICUCulture.cpp b/Engine/Source/Runtime/Core/Private/Internationalization/ICUCulture.cpp
--- a/Engine/Source/Runtime/Core/Private/Internationalization/ICUCulture.cpp
+++ b/Engine/Source/Runtime/Core/Private/Internationalization/ICUCulture.cpp
@@ -7,6 +7,9 @@
 #include "Internationalization/FastDecimalFormat.h"
 
 #if UE_ENABLE_ICU
+
+DEFINE_LOG_CATEGORY_STATIC(LogICUCulture, Log, All);
+
 #include "Internationalization/ICUUtilities.h"
 #include "HAL/IConsoleManager.h"
 
@@ -77,6 +80,7 @@ namespace
 		TSharedPtr<const icu::BreakIterator> Ptr = MakeShareable(FactoryFunction(ICULocale, ICUStatus));
 		if (!ensureAlwaysMsgf(Ptr, TEXT("Creating a break iterator object failed using locale %s. Perhaps this locale has no data."), StringCast<TCHAR>(ICULocale.getName()).Get()))
 		{
+			ICUStatus = U_ZERO_ERROR;
 			Ptr = MakeShareable(FactoryFunction(GetInvariantLocale(), ICUStatus));
 			check(Ptr);
 		}
@@ -89,7 +93,11 @@ namespace
 		TSharedPtr<const icu::Collator, ESPMode::ThreadSafe> Ptr = MakeShareable(icu::Collator::createInstance(ICULocale, ICUStatus));
 		if (!ensureAlwaysMsgf(Ptr, TEXT("Creating a collator object failed using locale %s. Perhaps this locale has no data."), StringCast<TCHAR>(ICULocale.getName()).Get()))
 		{
+			UE_LOG(LogICUCulture, Warning, TEXT("ICU collator instance creation failed locale '%s', %s (%d)"), StringCast<TCHAR>(ICULocale.getName()).Get(), StringCast<TCHAR>(u_errorName(ICUStatus)).Get(), ICUStatus);
+			// Reset the ICUStatus value, this must be non-error or the icu api will immediately fail.
+			ICUStatus = U_ZERO_ERROR;
 			Ptr = MakeShareable(icu::Collator::createInstance(GetInvariantLocale(), ICUStatus));
+			UE_CLOG(!Ptr, LogICUCulture, Fatal, TEXT("ICU collator instance creation failed to create invariant locale %s (%d)"), StringCast<TCHAR>(u_errorName(ICUStatus)).Get(), ICUStatus);
 			check(Ptr);
 		}
 		return Ptr.ToSharedRef();
@@ -207,6 +215,7 @@ FICUCultureImplementation::FICUCultureImplementation(const FString& LocaleName)
 		ICUCardinalPluralRules = icu::PluralRules::forLocale(ICULocale, UPLURAL_TYPE_CARDINAL, ICUStatus);
 		if (!ensureAlwaysMsgf(U_SUCCESS(ICUStatus) && ICUCardinalPluralRules, TEXT("Creating a cardinal plural rules object failed using locale %s. Perhaps this locale has no data."), *LocaleName))
 		{
+			ICUStatus = U_ZERO_ERROR;
 			ICUCardinalPluralRules = icu::PluralRules::forLocale(GetInvariantLocale(), UPLURAL_TYPE_CARDINAL, ICUStatus);
 			check(ICUCardinalPluralRules);
 		}
@@ -217,6 +226,7 @@ FICUCultureImplementation::FICUCultureImplementation(const FString& LocaleName)
 		ICUOrdinalPluralRules = icu::PluralRules::forLocale(ICULocale, UPLURAL_TYPE_ORDINAL, ICUStatus);
 		if (!ensureAlwaysMsgf(U_SUCCESS(ICUStatus) && ICUOrdinalPluralRules, TEXT("Creating an ordinal plural rules object failed using locale %s. Perhaps this locale has no data."), *LocaleName))
 		{
+			ICUStatus = U_ZERO_ERROR;
 			ICUOrdinalPluralRules = icu::PluralRules::forLocale(GetInvariantLocale(), UPLURAL_TYPE_ORDINAL, ICUStatus);
 			check(ICUOrdinalPluralRules);
 		}
@@ -401,7 +411,6 @@ TSharedRef<const icu::Collator, ESPMode::ThreadSafe> FICUCultureImplementation::
 		ICUCollator = CreateCollator( ICULocale );
 	}
 
-	UErrorCode ICUStatus = U_ZERO_ERROR;
 	const bool bIsDefault = (ComparisonLevel == ETextComparisonLevel::Default);
 	const TSharedRef<const icu::Collator, ESPMode::ThreadSafe> DefaultCollator( ICUCollator.ToSharedRef() );
 	if(bIsDefault)
@@ -411,7 +420,10 @@ TSharedRef<const icu::Collator, ESPMode::ThreadSafe> FICUCultureImplementation::
 	else
 	{
 		const TSharedRef<icu::Collator, ESPMode::ThreadSafe> Collator( DefaultCollator->clone() );
+
+		UErrorCode ICUStatus = U_ZERO_ERROR;
 		Collator->setAttribute(UColAttribute::UCOL_STRENGTH, UEToICU(ComparisonLevel), ICUStatus);
+		
 		return Collator;
 	}
 }
@@ -427,6 +439,7 @@ TSharedRef<const icu::DecimalFormat, ESPMode::ThreadSafe> FICUCultureImplementat
 			DecimalFormatterForCulture = MakeShareable(static_cast<icu::DecimalFormat*>(icu::NumberFormat::createInstance(ICULocale, ICUStatus)));
 			if (!ensureAlwaysMsgf(DecimalFormatterForCulture, TEXT("Creating a decimal format object failed using locale %s. Perhaps this locale has no data."), StringCast<TCHAR>(ICULocale.getName()).Get()))
 			{
+				ICUStatus = U_ZERO_ERROR;
 				DecimalFormatterForCulture = MakeShareable(static_cast<icu::DecimalFormat*>(icu::NumberFormat::createInstance(GetInvariantLocale(), ICUStatus)));
 				check(DecimalFormatterForCulture);
 			}
@@ -480,16 +493,20 @@ TSharedRef<const icu::DateFormat, ESPMode::ThreadSafe> FICUCultureImplementation
 	bool bIsDefaultTimeZone = SanitizedTimezoneCode.IsEmpty();
 	if( !bIsDefaultTimeZone )
 	{
-		UErrorCode ICUStatus = U_ZERO_ERROR;
-
 		icu::UnicodeString CanonicalInputTimeZoneID;
-		icu::TimeZone::getCanonicalID(InputTimeZoneID, CanonicalInputTimeZoneID, ICUStatus);
+		{
+			UErrorCode ICUStatus = U_ZERO_ERROR;
+			icu::TimeZone::getCanonicalID(InputTimeZoneID, CanonicalInputTimeZoneID, ICUStatus);
+		}
 
 		icu::UnicodeString DefaultTimeZoneID;
 		DefaultFormatter->getTimeZone().getID(DefaultTimeZoneID);
 
 		icu::UnicodeString CanonicalDefaultTimeZoneID;
-		icu::TimeZone::getCanonicalID(DefaultTimeZoneID, CanonicalDefaultTimeZoneID, ICUStatus);
+		{
+			UErrorCode ICUStatus = U_ZERO_ERROR;
+			icu::TimeZone::getCanonicalID(DefaultTimeZoneID, CanonicalDefaultTimeZoneID, ICUStatus);
+		}
 
 		bIsDefaultTimeZone = (CanonicalInputTimeZoneID == CanonicalDefaultTimeZoneID ? true : false);
 	}
@@ -532,16 +549,20 @@ TSharedRef<const icu::DateFormat, ESPMode::ThreadSafe> FICUCultureImplementation
 	bool bIsDefaultTimeZone = SanitizedTimezoneCode.IsEmpty();
 	if( !bIsDefaultTimeZone )
 	{
-		UErrorCode ICUStatus = U_ZERO_ERROR;
-
 		icu::UnicodeString CanonicalInputTimeZoneID;
-		icu::TimeZone::getCanonicalID(InputTimeZoneID, CanonicalInputTimeZoneID, ICUStatus);
+		{
+			UErrorCode ICUStatus = U_ZERO_ERROR;
+			icu::TimeZone::getCanonicalID(InputTimeZoneID, CanonicalInputTimeZoneID, ICUStatus);
+		}
 
 		icu::UnicodeString DefaultTimeZoneID;
 		DefaultFormatter->getTimeZone().getID(DefaultTimeZoneID);
 
 		icu::UnicodeString CanonicalDefaultTimeZoneID;
-		icu::TimeZone::getCanonicalID(DefaultTimeZoneID, CanonicalDefaultTimeZoneID, ICUStatus);
+		{
+			UErrorCode ICUStatus = U_ZERO_ERROR;
+			icu::TimeZone::getCanonicalID(DefaultTimeZoneID, CanonicalDefaultTimeZoneID, ICUStatus);
+		}
 
 		bIsDefaultTimeZone = (CanonicalInputTimeZoneID == CanonicalDefaultTimeZoneID ? true : false);
 	}
@@ -585,16 +606,20 @@ TSharedRef<const icu::DateFormat, ESPMode::ThreadSafe> FICUCultureImplementation
 	bool bIsDefaultTimeZone = SanitizedTimezoneCode.IsEmpty();
 	if( !bIsDefaultTimeZone )
 	{
-		UErrorCode ICUStatus = U_ZERO_ERROR;
-
 		icu::UnicodeString CanonicalInputTimeZoneID;
-		icu::TimeZone::getCanonicalID(InputTimeZoneID, CanonicalInputTimeZoneID, ICUStatus);
+		{
+			UErrorCode ICUStatus = U_ZERO_ERROR;
+			icu::TimeZone::getCanonicalID(InputTimeZoneID, CanonicalInputTimeZoneID, ICUStatus);
+		}
 
 		icu::UnicodeString DefaultTimeZoneID;
 		DefaultFormatter->getTimeZone().getID(DefaultTimeZoneID);
 
 		icu::UnicodeString CanonicalDefaultTimeZoneID;
-		icu::TimeZone::getCanonicalID(DefaultTimeZoneID, CanonicalDefaultTimeZoneID, ICUStatus);
+		{
+			UErrorCode ICUStatus = U_ZERO_ERROR;
+			icu::TimeZone::getCanonicalID(DefaultTimeZoneID, CanonicalDefaultTimeZoneID, ICUStatus);
+		}
 
 		bIsDefaultTimeZone = (CanonicalInputTimeZoneID == CanonicalDefaultTimeZoneID ? true : false);
 	}
@@ -672,13 +697,15 @@ TSharedRef<const icu::DateFormat, ESPMode::ThreadSafe> FICUCultureImplementation
 
 		const FString DateTimeSkeleton = DateTimePatternToICUSkeleton(*CustomPattern);
 
-		UErrorCode ICUStatus = U_ZERO_ERROR;
 		TSharedPtr<icu::DateFormat, ESPMode::ThreadSafe> Formatter;
-		
-		Formatter = MakeShareable(icu::DateFormat::createInstanceForSkeleton(ICUUtilities::ConvertString(DateTimeSkeleton), ICULocale, ICUStatus));
-		if (!Formatter)
 		{
-			Formatter = MakeShareable(icu::DateFormat::createInstanceForSkeleton(ICUUtilities::ConvertString(DateTimeSkeleton), GetInvariantLocale(), ICUStatus));
+			UErrorCode ICUStatus = U_ZERO_ERROR;
+			Formatter = MakeShareable(icu::DateFormat::createInstanceForSkeleton(ICUUtilities::ConvertString(DateTimeSkeleton), ICULocale, ICUStatus));
+			if (!Formatter)
+			{
+				ICUStatus = U_ZERO_ERROR;
+				Formatter = MakeShareable(icu::DateFormat::createInstanceForSkeleton(ICUUtilities::ConvertString(DateTimeSkeleton), GetInvariantLocale(), ICUStatus));
+			}
 		}
 
 		if (Formatter)
@@ -817,6 +844,7 @@ const FDecimalNumberFormattingRules& FICUCultureImplementation::GetDecimalNumber
 		DecimalFormatterForCulture = MakeShareable(static_cast<icu::DecimalFormat*>(icu::NumberFormat::createInstance(ICULocale, ICUStatus)));
 		if (!ensureAlwaysMsgf(DecimalFormatterForCulture, TEXT("Creating a decimal format object failed using locale %s. Perhaps this locale has no data."), StringCast<TCHAR>(ICULocale.getName()).Get()))
 		{
+			ICUStatus = U_ZERO_ERROR;
 			DecimalFormatterForCulture = MakeShareable(static_cast<icu::DecimalFormat*>(icu::NumberFormat::createInstance(GetInvariantLocale(), ICUStatus)));
 			check(DecimalFormatterForCulture);
 		}
@@ -851,6 +879,7 @@ const FDecimalNumberFormattingRules& FICUCultureImplementation::GetPercentFormat
 		PercentFormatterForCulture = MakeShareable(static_cast<icu::DecimalFormat*>(icu::NumberFormat::createPercentInstance(ICULocale, ICUStatus)));
 		if (!ensureAlwaysMsgf(PercentFormatterForCulture, TEXT("Creating a percent format object failed using locale %s. Perhaps this locale has no data."), StringCast<TCHAR>(ICULocale.getName()).Get()))
 		{
+			ICUStatus = U_ZERO_ERROR;
 			PercentFormatterForCulture = MakeShareable(static_cast<icu::DecimalFormat*>(icu::NumberFormat::createPercentInstance(GetInvariantLocale(), ICUStatus)));
 			check(PercentFormatterForCulture);
 		}
@@ -901,6 +930,7 @@ const FDecimalNumberFormattingRules& FICUCultureImplementation::GetCurrencyForma
 		CurrencyFormatterForCulture = MakeShareable(static_cast<icu::DecimalFormat*>(icu::NumberFormat::createCurrencyInstance(ICULocale, ICUStatus)));
 		if (!ensureAlwaysMsgf(CurrencyFormatterForCulture, TEXT("Creating a currency format object failed using locale %s. Perhaps this locale has no data."), StringCast<TCHAR>(ICULocale.getName()).Get()))
 		{
+			ICUStatus = U_ZERO_ERROR;
 			CurrencyFormatterForCulture = MakeShareable(static_cast<icu::DecimalFormat*>(icu::NumberFormat::createCurrencyInstance(GetInvariantLocale(), ICUStatus)));
 			check(CurrencyFormatterForCulture);
 		}
