[AutoRTFM] Handle "true" and "false" when parsing dpcvar flags

#Synced-CL 49717872

diff --git a/Engine/Source/Runtime/Core/Private/AutoRTFM/AutoRTFMUE.cpp b/Engine/Source/Runtime/Core/Private/AutoRTFM/AutoRTFMUE.cpp
--- a/Engine/Source/Runtime/Core/Private/AutoRTFM/AutoRTFMUE.cpp
+++ b/Engine/Source/Runtime/Core/Private/AutoRTFM/AutoRTFMUE.cpp
@@ -30,12 +30,12 @@ namespace
 
 static bool IsDisable(const FString& String)
 {
-	return String == "disable" || String == "disabled" || String == "off" || String == "0";
+	return String == "disable" || String == "disabled" || String == "false" || String == "off" || String == "0";
 }
 
 static bool IsEnable(const FString& String)
 {
-	return String == "enable" || String == "enabled" || String == "on" || String == "1";
+	return String == "enable" || String == "enabled" || String == "true" || String == "on" || String == "1";
 }
 
 void OnAutoRTFMRuntimeEnabledChanged()
