Run CollectStorageInformation asynchronously on windows. Saves between 0.5-1s of boot time when collecting extended storage information.

#Synced-CL 49720840

diff --git a/Engine/Source/Runtime/Core/Private/Windows/WindowsPlatformMisc.cpp b/Engine/Source/Runtime/Core/Private/Windows/WindowsPlatformMisc.cpp
--- a/Engine/Source/Runtime/Core/Private/Windows/WindowsPlatformMisc.cpp
+++ b/Engine/Source/Runtime/Core/Private/Windows/WindowsPlatformMisc.cpp
@@ -91,6 +91,14 @@ THIRD_PARTY_INCLUDES_END
 #define ALLOW_COLLECT_STORAGE_INFORMATION 0
 #endif
 
+// Collect extended storage information asynchronously. A Sync call can take in the order of 1 second on the calling thread
+bool GCollectStorageInformationAsync = false;
+static FAutoConsoleVariableRef CVarCollectStorageInformationAsync(
+	TEXT("g.CollectStorageInformationAsync"),
+	GCollectStorageInformationAsync,
+	TEXT("Collects extended storage information async"),
+	ECVF_ReadOnly);
+
 // This might not be defined by Windows when maintaining backwards-compatibility to pre-Win8 builds
 #ifndef SM_CONVERTIBLESLATEMODE
 #define SM_CONVERTIBLESLATEMODE			0x2003
@@ -418,6 +426,8 @@ namespace
 		}
 	};
 
+	// NOTE: This is populated async (if GCollectStorageInformationAsync)
+	// Usage must be synchronized with WaitOnCollectStorageInformationAsync
 	TArray<FStorageDevice> StorageDevices;
 
 	static void LogStorageInformationWarning(HRESULT HRes,  const TCHAR* message)
@@ -444,31 +454,30 @@ namespace
 			UE_LOG(LogWindows, Log, TEXT("%s [error code %x]"), message, HRes);
 		}
 	}
-
-	static bool CollectStorageInformation()
-	{
+	
 #if ALLOW_COLLECT_STORAGE_INFORMATION
-
+	static void CollectStorageInformationImpl()
+	{
 		IWbemLocator* WbemLocator = nullptr;
 		IWbemServices* WbemServices = nullptr;
 
 		if (!FWindowsPlatformMisc::CoInitialize())
 		{
-			return false;
+			return;
 		}
 		HRESULT hres = CoInitializeSecurity(NULL, -1, NULL, NULL, RPC_C_AUTHN_LEVEL_DEFAULT, RPC_C_IMP_LEVEL_IMPERSONATE, NULL, EOAC_NONE, NULL);
 		if (FAILED(hres))
 		{
 			LogStorageInformationWarning(hres, TEXT("Error initializing COM"));
 			FWindowsPlatformMisc::CoUninitialize();
-			return false;
+			return;
 		}
 		hres = CoCreateInstance(CLSID_WbemLocator, 0, CLSCTX_INPROC_SERVER, IID_IWbemLocator, (LPVOID*)&WbemLocator);
 		if (FAILED(hres))
 		{
 			LogStorageInformationWarning(hres, TEXT("Error creating Wbem instance"));
 			FWindowsPlatformMisc::CoUninitialize();
-			return false;
+			return;
 		}
 		hres = WbemLocator->ConnectServer(_bstr_t(L"ROOT\\microsoft\\windows\\storage"), NULL, NULL, 0, NULL, 0, 0, &WbemServices);
 		if (FAILED(hres))
@@ -476,7 +485,7 @@ namespace
 			LogStorageInformationWarning(hres, TEXT("Error connecting to storage service"));
 			WbemLocator->Release();
 			FWindowsPlatformMisc::CoUninitialize();
-			return false;
+			return;
 		}
 		hres = CoSetProxyBlanket(
 			WbemServices,
@@ -494,7 +503,7 @@ namespace
 			WbemServices->Release();
 			WbemLocator->Release();
 			FWindowsPlatformMisc::CoUninitialize();
-			return false;
+			return;
 		}
 
 
@@ -512,7 +521,7 @@ namespace
 			WbemServices->Release();
 			WbemLocator->Release();
 			FWindowsPlatformMisc::CoUninitialize();
-			return false;
+			return;
 		}
 
 		// Enumerate all partitions, and store disk information
@@ -613,7 +622,7 @@ namespace
 			WbemServices->Release();
 			WbemLocator->Release();
 			FWindowsPlatformMisc::CoUninitialize();
-			return false;
+			return;
 		}
 
 		while (StorageEnumerator)
@@ -687,10 +696,72 @@ namespace
 		WbemLocator->Release();
 		FWindowsPlatformMisc::CoUninitialize();
 		FWindowsPlatformMisc::UpdateDriveFreeSpace();
+	}
+
+	struct FCollectStorageInformationAsyncThread
+	{
+		HANDLE HThread;
+		DWORD HThreadId;
 
+		FCollectStorageInformationAsyncThread()
+			: HThread(NULL)
+			, HThreadId(0)
+		{
+		}
+
+		~FCollectStorageInformationAsyncThread()
+		{
+			if (HThread)
+			{
+				if (::WaitForSingleObject(HThread, 1000) == WAIT_OBJECT_0)
+				{
+					::CloseHandle(HThread);
+				}
+				HThread = NULL;
+			}
+		}
+	};
+	static FCollectStorageInformationAsyncThread GCollectStorageInformationAsyncThread;
+	
+	static DWORD WINAPI CollectStorageInformationAsyncProc(LPVOID Arg)
+	{
+		CollectStorageInformationImpl();
+		return 0;
+	}
+
+	static void CollectStorageInformationAsync()
+	{
+		GCollectStorageInformationAsyncThread.HThread = ::CreateThread(NULL, 0, CollectStorageInformationAsyncProc, NULL, 0, &GCollectStorageInformationAsyncThread.HThreadId);
+		if (GCollectStorageInformationAsyncThread.HThread == NULL)
+		{
+			UE_LOG(LogWindows, Error, TEXT("Unable to create thread for CollectStorageInformationAsync, running synchronous..."));
+			CollectStorageInformationImpl();
+		}
+	}
 #endif // ALLOW_COLLECT_STORAGE_INFORMATION
 
-		return true;
+	static void WaitOnCollectStorageInformationAsync()
+	{
+#if ALLOW_COLLECT_STORAGE_INFORMATION
+		if (GCollectStorageInformationAsync && ::GetCurrentThreadId() != GCollectStorageInformationAsyncThread.HThreadId)
+		{
+			::WaitForSingleObject(GCollectStorageInformationAsyncThread.HThread, INFINITE);
+		}
+#endif
+	}
+
+	static void CollectStorageInformation()
+	{
+#if ALLOW_COLLECT_STORAGE_INFORMATION
+		if (GCollectStorageInformationAsync)
+		{
+			CollectStorageInformationAsync();
+		}
+		else
+		{
+			CollectStorageInformationImpl();
+		}
+#endif
 	}
 }
 
@@ -949,7 +1020,6 @@ void FWindowsPlatformMisc::PlatformPreInit()
 #endif
 }
 
-
 void FWindowsPlatformMisc::PlatformInit()
 {
 	FGenericPlatformMisc::LogNameEventStatsInit();
@@ -4394,6 +4464,7 @@ int32 FWindowsPlatformMisc::GetMaxRefreshRate()
 
 void FWindowsPlatformMisc::UpdateDriveFreeSpace()
 {
+	WaitOnCollectStorageInformationAsync();
 	for (auto& StorageDevice : StorageDevices)
 	{
 		ULARGE_INTEGER TotalNumberOfBytes, TotalNumberOfFreeBytes;
@@ -4408,6 +4479,7 @@ void FWindowsPlatformMisc::UpdateDriveFreeSpace()
 
 const FPlatformDriveStats* FWindowsPlatformMisc::GetDriveStats(WIDECHAR DriveLetter)
 {
+	WaitOnCollectStorageInformationAsync();
 	for (auto& StorageDevice : StorageDevices)
 	{
 		if (StorageDevice.Stats.DriveName == DriveLetter)
diff --git a/Engine/Source/Runtime/Core/Public/Windows/WindowsPlatformMisc.h b/Engine/Source/Runtime/Core/Public/Windows/WindowsPlatformMisc.h
--- a/Engine/Source/Runtime/Core/Public/Windows/WindowsPlatformMisc.h
+++ b/Engine/Source/Runtime/Core/Public/Windows/WindowsPlatformMisc.h
@@ -385,10 +385,16 @@ struct FWindowsPlatformMisc
 
 	static CORE_API int32 GetMaxRefreshRate();
 
-	/** Update statistics of free/used bytes on all drives. */
+	/** 
+	 * Update statistics of free/used bytes on all drives.
+	 * NOTE: This call will block (if GCollectStorageInformationAsync is enabled) until stats are available
+	 */
 	static CORE_API void UpdateDriveFreeSpace();
 
-	/** Retrieve information about a drive, or nullptr if no information is available. */
+	/**
+	 * Retrieve information about a drive, or nullptr if no information is available.
+	 * NOTE: This call will block (if GCollectStorageInformationAsync is enabled) until stats are available
+	 */
 	static CORE_API const FPlatformDriveStats* GetDriveStats(WIDECHAR DriveLetter);
 };
 
