Fixed a config reading short circuit issue

#Synced-CL 49298013

diff --git a/Engine/Source/Runtime/Core/Private/HAL/ThreadHeartBeat.cpp b/Engine/Source/Runtime/Core/Private/HAL/ThreadHeartBeat.cpp
--- a/Engine/Source/Runtime/Core/Private/HAL/ThreadHeartBeat.cpp
+++ b/Engine/Source/Runtime/Core/Private/HAL/ThreadHeartBeat.cpp
@@ -1023,10 +1023,9 @@ void FGameThreadHitchHeartBeatThreaded::InitSettingsInternal()
 		bool bReadFromConfig = false;
 		if (GConfig)
 		{
-			static bool bLocalReadFromConfig = 
-				GConfig->GetFloat(TEXT("Core.System"), TEXT("GameThreadHeartBeatHitchDuration"), Config_Duration, GEngineIni) ||
-				GConfig->GetBool(TEXT("Core.System"), TEXT("GameThreadHeartBeatStackWalk"), Config_StackWalk, GEngineIni);
-			bReadFromConfig = bLocalReadFromConfig;
+			static bool bDidReadDurationsetting = GConfig->GetFloat(TEXT("Core.System"), TEXT("GameThreadHeartBeatHitchDuration"), Config_Duration, GEngineIni);
+			static bool bDidReadStackwalkSetting = GConfig->GetBool(TEXT("Core.System"), TEXT("GameThreadHeartBeatStackWalk"), Config_StackWalk, GEngineIni);
+			bReadFromConfig = bDidReadDurationsetting || bDidReadStackwalkSetting;
 		}
 
 		if (bReadFromConfig)
