Undo changelist 49764452

diff --git a/Engine/Source/Runtime/Core/Public/UObject/ScriptDelegates.h b/Engine/Source/Runtime/Core/Public/UObject/ScriptDelegates.h
--- a/Engine/Source/Runtime/Core/Public/UObject/ScriptDelegates.h
+++ b/Engine/Source/Runtime/Core/Public/UObject/ScriptDelegates.h
@@ -16,16 +16,6 @@
 #include "UObject/NameTypes.h"
 #include <type_traits>
 
-/* Specifies whether or not dynamic delegates can have payloads. */
-#ifndef UE_USE_DYNAMIC_DELEGATE_PAYLOADS
-	#define UE_USE_DYNAMIC_DELEGATE_PAYLOADS 0
-
-	/* Dynamic delegates with payloads and return values need special handling. */
-	#ifndef UE_USE_DYNAMIC_DELEGATE_PAYLOADS_RETVAL
-		#define UE_USE_DYNAMIC_DELEGATE_PAYLOADS_RETVAL 0
-	#endif
-#endif
-
 namespace UE::Core::Private
 {
 	template <typename InThreadSafetyMode>
diff --git a/Engine/Source/Runtime/CoreUObject/Public/UObject/DynamicDelegateTestObject.cpp b/Engine/Source/Runtime/CoreUObject/Public/UObject/DynamicDelegateTestObject.cpp
deleted file mode 100644
diff --git a/Engine/Source/Runtime/CoreUObject/Public/UObject/DynamicDelegateTestObject.h b/Engine/Source/Runtime/CoreUObject/Public/UObject/DynamicDelegateTestObject.h
deleted file mode 100644
