PR #14003: Fix template type mismatch for Heap functions

#14003

diff --git a/Engine/Source/Runtime/Core/Public/Containers/Array.h b/Engine/Source/Runtime/Core/Public/Containers/Array.h
--- a/Engine/Source/Runtime/Core/Public/Containers/Array.h
+++ b/Engine/Source/Runtime/Core/Public/Containers/Array.h
@@ -3751,7 +3751,7 @@ public:
 	 */
 	UE_NODEBUG void Heapify()
 	{
-		Heapify(TLess<ElementType>());
+		Heapify(TLess<>());
 	}
 
 	/** 
@@ -3811,7 +3811,7 @@ public:
 	 */
 	UE_NODEBUG SizeType HeapPush(ElementType&& InItem)
 	{
-		return HeapPush(MoveTempIfPossible(InItem), TLess<ElementType>());
+		return HeapPush(MoveTempIfPossible(InItem), TLess<>());
 	}
 
 	/** 
@@ -3827,7 +3827,7 @@ public:
 	 */
 	UE_NODEBUG SizeType HeapPush(const ElementType& InItem)
 	{
-		return HeapPush(InItem, TLess<ElementType>());
+		return HeapPush(InItem, TLess<>());
 	}
 
 	/** 
@@ -3866,7 +3866,7 @@ public:
 	 */
 	UE_NODEBUG void HeapPop(ElementType& OutItem, EAllowShrinking AllowShrinking = UE::Core::Private::AllowShrinkingByDefault<AllocatorType>())
 	{
-		HeapPop(OutItem, TLess<ElementType>(), AllowShrinking);
+		HeapPop(OutItem, TLess<>(), AllowShrinking);
 	}
 
 	/**
@@ -3911,7 +3911,7 @@ public:
 	 */
 	UE_NODEBUG void HeapPopDiscard(EAllowShrinking AllowShrinking = UE::Core::Private::AllowShrinkingByDefault<AllocatorType>())
 	{
-		HeapPopDiscard(TLess<ElementType>(), AllowShrinking);
+		HeapPopDiscard(TLess<>(), AllowShrinking);
 	}
 
 	/** 
@@ -3971,7 +3971,7 @@ public:
 	 */
 	UE_NODEBUG void HeapRemoveAt(SizeType Index, EAllowShrinking AllowShrinking = UE::Core::Private::AllowShrinkingByDefault<AllocatorType>())
 	{
-		HeapRemoveAt(Index, TLess< ElementType >(), AllowShrinking);
+		HeapRemoveAt(Index, TLess<>(), AllowShrinking);
 	}
 
 	/**
@@ -4000,7 +4000,7 @@ public:
 	 */
 	UE_NODEBUG void HeapSort()
 	{
-		HeapSort(TLess<ElementType>());
+		HeapSort(TLess<>());
 	}
 
 	UE_NODEBUG [[nodiscard]] const ElementAllocatorType& GetAllocatorInstance() const
